CompactBinary: Moved validation functionality into new files

#rnx

diff --git a/Engine/Source/Runtime/Core/Private/Serialization/CompactBinary.cpp b/Engine/Source/Runtime/Core/Private/Serialization/CompactBinary.cpp
--- a/Engine/Source/Runtime/Core/Private/Serialization/CompactBinary.cpp
+++ b/Engine/Source/Runtime/Core/Private/Serialization/CompactBinary.cpp
@@ -11,22 +11,24 @@
 
 ///////////////////////////////////////////////////////////////////////////////////////////////////////////////////////
 
-namespace CbObjectPrivate
+namespace CompactBinaryPrivate
 {
-	static constexpr const uint8 EmptyObjectPayload[] = { 0x00 };
-	static constexpr const uint8 EmptyArrayPayload[] = { 0x01, 0x00 };
 
-	template <typename T>
-	static constexpr FORCEINLINE T ReadUnaligned(const void* const Memory)
-	{
-	#if PLATFORM_SUPPORTS_UNALIGNED_LOADS
-		return *static_cast<const T*>(Memory);
-	#else
-		T Value;
-		FMemory::Memcpy(&Value, Memory, sizeof(Value));
-		return Value;
-	#endif
-	}
+static constexpr const uint8 GEmptyObjectPayload[] = { 0x00 };
+static constexpr const uint8 GEmptyArrayPayload[] = { 0x01, 0x00 };
+
+template <typename T>
+static constexpr FORCEINLINE T ReadUnaligned(const void* const Memory)
+{
+#if PLATFORM_SUPPORTS_UNALIGNED_LOADS
+	return *static_cast<const T*>(Memory);
+#else
+	T Value;
+	FMemory::Memcpy(&Value, Memory, sizeof(Value));
+	return Value;
+#endif
+}
+
 }
 
 ///////////////////////////////////////////////////////////////////////////////////////////////////////////////////////
@@ -259,7 +261,7 @@ float FCbField::AsFloat(const float Default)
 	case ECbFieldType::Float32:
 	{
 		Error = ECbFieldError::None;
-		const uint32 Value = NETWORK_ORDER32(CbObjectPrivate::ReadUnaligned<uint32>(Payload));
+		const uint32 Value = NETWORK_ORDER32(CompactBinaryPrivate::ReadUnaligned<uint32>(Payload));
 		return reinterpret_cast<const float&>(Value);
 	}
 	case ECbFieldType::Float64:
@@ -290,13 +292,13 @@ double FCbField::AsDouble(const double Default)
 	case ECbFieldType::Float32:
 	{
 		Error = ECbFieldError::None;
-		const uint32 Value = NETWORK_ORDER32(CbObjectPrivate::ReadUnaligned<uint32>(Payload));
+		const uint32 Value = NETWORK_ORDER32(CompactBinaryPrivate::ReadUnaligned<uint32>(Payload));
 		return reinterpret_cast<const float&>(Value);
 	}
 	case ECbFieldType::Float64:
 	{
 		Error = ECbFieldError::None;
-		const uint64 Value = NETWORK_ORDER64(CbObjectPrivate::ReadUnaligned<uint64>(Payload));
+		const uint64 Value = NETWORK_ORDER64(CompactBinaryPrivate::ReadUnaligned<uint64>(Payload));
 		return reinterpret_cast<const double&>(Value);
 	}
 	default:
@@ -385,7 +387,7 @@ int64 FCbField::AsDateTimeTicks(const int64 Default)
 	if (FCbFieldType::IsDateTime(Type))
 	{
 		Error = ECbFieldError::None;
-		return NETWORK_ORDER64(CbObjectPrivate::ReadUnaligned<int64>(Payload));
+		return NETWORK_ORDER64(CompactBinaryPrivate::ReadUnaligned<int64>(Payload));
 	}
 	else
 	{
@@ -409,7 +411,7 @@ int64 FCbField::AsTimeSpanTicks(const int64 Default)
 	if (FCbFieldType::IsTimeSpan(Type))
 	{
 		Error = ECbFieldError::None;
-		return NETWORK_ORDER64(CbObjectPrivate::ReadUnaligned<int64>(Payload));
+		return NETWORK_ORDER64(CompactBinaryPrivate::ReadUnaligned<int64>(Payload));
 	}
 	else
 	{
@@ -437,7 +439,7 @@ void FCbField::CopyTo(const FMutableMemoryView Buffer) const
 ///////////////////////////////////////////////////////////////////////////////////////////////////////////////////////
 
 FCbArray::FCbArray()
-	: Payload(CbObjectPrivate::EmptyArrayPayload)
+	: Payload(CompactBinaryPrivate::GEmptyArrayPayload)
 	, Type(ECbFieldType::Array)
 {
 }
@@ -507,7 +509,7 @@ void FCbArray::CopyTo(const FMutableMemoryView Buffer) const
 ///////////////////////////////////////////////////////////////////////////////////////////////////////////////////////
 
 FCbObject::FCbObject()
-	: Payload(CbObjectPrivate::EmptyObjectPayload)
+	: Payload(CompactBinaryPrivate::GEmptyObjectPayload)
 	, Type(ECbFieldType::Object)
 {
 }
@@ -591,376 +593,6 @@ void FCbObject::CopyTo(const FMutableMemoryView Buffer) const
 
 ///////////////////////////////////////////////////////////////////////////////////////////////////////////////////////
 
-/**
- * Adds the given error(s) to the error mask.
- *
- * This function exists to make validation errors easier to debug by providing one location to set a breakpoint.
- */
-FORCENOINLINE static void AddError(ECbValidateError& OutError, const ECbValidateError InError)
-{
-	OutError |= InError;
-}
-
-/**
- * Validate and read a field type from the view.
- *
- * A type argument with the HasFieldType flag indicates that the type will not be read from the view.
- */
-static ECbFieldType ValidateCbFieldType(FConstMemoryView& View, ECbValidateMode Mode, ECbValidateError& Error, ECbFieldType Type = ECbFieldType::HasFieldType)
-{
-	if (FCbFieldType::HasFieldType(Type))
-	{
-		if (View.GetSize() >= 1)
-		{
-			Type = *static_cast<const ECbFieldType*>(View.GetData());
-			View += 1;
-			if (FCbFieldType::HasFieldType(Type))
-			{
-				AddError(Error, ECbValidateError::InvalidType);
-			}
-		}
-		else
-		{
-			AddError(Error, ECbValidateError::OutOfBounds);
-			View.Reset();
-			return ECbFieldType::None;
-		}
-	}
-
-	if (FCbFieldType::GetType(Type) > ECbFieldType::TimeSpan)
-	{
-		AddError(Error, ECbValidateError::InvalidType);
-		View.Reset();
-	}
-
-	return Type;
-}
-
-/**
- * Validate and read an unsigned integer from the view.
- *
- * Modifies the view to start at the end of the value, and adds error flags if applicable.
- */
-static uint64 ValidateCbUInt(FConstMemoryView& View, ECbValidateMode Mode, ECbValidateError& Error)
-{
-	if (View.GetSize() > 0 && View.GetSize() >= MeasureVarUInt(View.GetData()))
-	{
-		uint32 ValueByteCount;
-		const uint64 Value = ReadVarUInt(View.GetData(), ValueByteCount);
-		if (EnumHasAnyFlags(Mode, ECbValidateMode::Format) && ValueByteCount > MeasureVarUInt(Value))
-		{
-			AddError(Error, ECbValidateError::InvalidInteger);
-		}
-		View += ValueByteCount;
-		return Value;
-	}
-	else
-	{
-		AddError(Error, ECbValidateError::OutOfBounds);
-		View.Reset();
-		return 0;
-	}
-}
-
-/**
- * Validate a 64-bit floating point value from the view.
- *
- * Modifies the view to start at the end of the value, and adds error flags if applicable.
- */
-static void ValidateCbFloat64(FConstMemoryView& View, ECbValidateMode Mode, ECbValidateError& Error)
-{
-	if (View.GetSize() >= sizeof(double))
-	{
-		if (EnumHasAnyFlags(Mode, ECbValidateMode::Format))
-		{
-			const uint64 RawValue = NETWORK_ORDER64(CbObjectPrivate::ReadUnaligned<uint64>(View.GetData()));
-			const double Value = reinterpret_cast<const double&>(RawValue);
-			if (Value == double(float(Value)))
-			{
-				AddError(Error, ECbValidateError::InvalidFloat);
-			}
-		}
-		View += sizeof(double);
-	}
-	else
-	{
-		AddError(Error, ECbValidateError::OutOfBounds);
-		View.Reset();
-	}
-}
-
-/**
- * Validate and read a string from the view.
- *
- * Modifies the view to start at the end of the string, and adds error flags if applicable.
- */
-static FAnsiStringView ValidateCbString(FConstMemoryView& View, ECbValidateMode Mode, ECbValidateError& Error)
-{
-	const uint64 NameSize = ValidateCbUInt(View, Mode, Error);
-	if (View.GetSize() >= NameSize)
-	{
-		const FAnsiStringView Name(static_cast<const ANSICHAR*>(View.GetData()), static_cast<int32>(NameSize));
-		View += NameSize;
-		return Name;
-	}
-	else
-	{
-		AddError(Error, ECbValidateError::OutOfBounds);
-		View.Reset();
-		return FAnsiStringView();
-	}
-}
-
-static FCbField ValidateCbField(FConstMemoryView& View, ECbValidateMode Mode, ECbValidateError& Error, ECbFieldType ExternalType);
-
-/** A type that checks whether all validated fields are of the same type. */
-class FCbUniformFieldsValidator
-{
-public:
-	inline explicit FCbUniformFieldsValidator(ECbFieldType InExternalType)
-		: ExternalType(InExternalType)
-	{
-	}
-
-	inline FCbField ValidateField(FConstMemoryView& View, ECbValidateMode Mode, ECbValidateError& Error)
-	{
-		const void* const FieldData = View.GetData();
-		if (FCbField Field = ValidateCbField(View, Mode, Error, ExternalType))
-		{
-			++FieldCount;
-			if (FCbFieldType::HasFieldType(ExternalType))
-			{
-				const ECbFieldType FieldType = *static_cast<const ECbFieldType*>(FieldData);
-				if (FieldCount == 1)
-				{
-					FirstType = FieldType;
-				}
-				else if (FieldType != FirstType)
-				{
-					bUniform = false;
-				}
-			}
-			return Field;
-		}
-
-		// It may not safe to check for uniformity if the field was invalid.
-		bUniform = false;
-		return FCbField();
-	}
-
-	inline bool IsUniform() const { return FieldCount > 0 && bUniform; }
-
-private:
-	uint32 FieldCount = 0;
-	bool bUniform = true;
-	ECbFieldType FirstType = ECbFieldType::None;
-	ECbFieldType ExternalType;
-};
-
-static void ValidateCbObject(FConstMemoryView& View, ECbValidateMode Mode, ECbValidateError& Error, ECbFieldType ObjectType)
-{
-	const uint64 Size = ValidateCbUInt(View, Mode, Error);
-	FConstMemoryView ObjectView = View.Left(Size);
-	View += Size;
-
-	if (Size > 0)
-	{
-		TArray<FAnsiStringView, TInlineAllocator<16>> Names;
-
-		const bool bUniformObject = FCbFieldType::GetType(ObjectType) == ECbFieldType::UniformObject;
-		const ECbFieldType ExternalType = bUniformObject ? ValidateCbFieldType(ObjectView, Mode, Error) : ECbFieldType::HasFieldType;
-		FCbUniformFieldsValidator UniformValidator(ExternalType);
-		do
-		{
-			if (FCbField Field = UniformValidator.ValidateField(ObjectView, Mode, Error))
-			{
-				if (EnumHasAnyFlags(Mode, ECbValidateMode::Names))
-				{
-					if (Field.HasName())
-					{
-						Names.Add(Field.GetName());
-					}
-					else
-					{
-						AddError(Error, ECbValidateError::MissingName);
-					}
-				}
-			}
-		}
-		while (!ObjectView.IsEmpty());
-
-		if (EnumHasAnyFlags(Mode, ECbValidateMode::Names) && Names.Num() > 1)
-		{
-			Algo::Sort(Names, [](FAnsiStringView L, FAnsiStringView R) { return L.Compare(R) < 0; });
-			for (const FAnsiStringView* NamesIt = Names.GetData(), *NamesEnd = NamesIt + Names.Num() - 1; NamesIt != NamesEnd; ++NamesIt)
-			{
-				if (NamesIt[0].Equals(NamesIt[1]))
-				{
-					AddError(Error, ECbValidateError::DuplicateName);
-					break;
-				}
-			}
-		}
-
-		if (!bUniformObject && EnumHasAnyFlags(Mode, ECbValidateMode::Format) && UniformValidator.IsUniform())
-		{
-			AddError(Error, ECbValidateError::NonUniformObject);
-		}
-	}
-}
-
-static void ValidateCbArray(FConstMemoryView& View, ECbValidateMode Mode, ECbValidateError& Error, ECbFieldType ArrayType)
-{
-	const uint64 Size = ValidateCbUInt(View, Mode, Error);
-	FConstMemoryView ArrayView = View.Left(Size);
-	View += Size;
-
-	const uint64 Count = ValidateCbUInt(ArrayView, Mode, Error);
-	const uint64 FieldsSize = ArrayView.GetSize();
-	const bool bUniformArray = FCbFieldType::GetType(ArrayType) == ECbFieldType::UniformArray;
-	const ECbFieldType ExternalType = bUniformArray ? ValidateCbFieldType(ArrayView, Mode, Error) : ECbFieldType::HasFieldType;
-	FCbUniformFieldsValidator UniformValidator(ExternalType);
-
-	for (uint64 Index = 0; Index < Count; ++Index)
-	{
-		if (FCbField Field = UniformValidator.ValidateField(ArrayView, Mode, Error))
-		{
-			if (Field.HasName() && EnumHasAnyFlags(Mode, ECbValidateMode::Names))
-			{
-				AddError(Error, ECbValidateError::ArrayName);
-			}
-		}
-	}
-
-	if (!bUniformArray && EnumHasAnyFlags(Mode, ECbValidateMode::Format) && UniformValidator.IsUniform() && FieldsSize > Count)
-	{
-		AddError(Error, ECbValidateError::NonUniformArray);
-	}
-}
-
-static FCbField ValidateCbField(FConstMemoryView& View, ECbValidateMode Mode, ECbValidateError& Error, const ECbFieldType ExternalType = ECbFieldType::HasFieldType)
-{
-	const FConstMemoryView FieldView = View;
-	const ECbFieldType Type = ValidateCbFieldType(View, Mode, Error, ExternalType);
-	const FAnsiStringView Name = FCbFieldType::HasFieldName(Type) ? ValidateCbString(View, Mode, Error) : FAnsiStringView();
-
-	auto ValidateFixedPayload = [&View, &Error](uint32 PayloadSize)
-	{
-		if (View.GetSize() >= PayloadSize)
-		{
-			View += PayloadSize;
-		}
-		else
-		{
-			AddError(Error, ECbValidateError::OutOfBounds);
-			View.Reset();
-		}
-	};
-
-	if (EnumHasAnyFlags(Error, ECbValidateError::OutOfBounds | ECbValidateError::InvalidType))
-	{
-		return FCbField();
-	}
-
-	switch (FCbFieldType::GetType(Type))
-	{
-	default:
-	case ECbFieldType::None:
-		AddError(Error, ECbValidateError::InvalidType);
-		View.Reset();
-		break;
-	case ECbFieldType::Null:
-	case ECbFieldType::BoolFalse:
-	case ECbFieldType::BoolTrue:
-		if (FieldView == View)
-		{
-			// Reset the view because a zero-sized field can cause infinite field iteration.
-			AddError(Error, ECbValidateError::InvalidType);
-			View.Reset();
-		}
-		break;
-	case ECbFieldType::Object:
-	case ECbFieldType::UniformObject:
-		ValidateCbObject(View, Mode, Error, FCbFieldType::GetType(Type));
-		break;
-	case ECbFieldType::Array:
-	case ECbFieldType::UniformArray:
-		ValidateCbArray(View, Mode, Error, FCbFieldType::GetType(Type));
-		break;
-	case ECbFieldType::Binary:
-	{
-		const uint64 ValueSize = ValidateCbUInt(View, Mode, Error);
-		if (View.GetSize() < ValueSize)
-		{
-			AddError(Error, ECbValidateError::OutOfBounds);
-			View.Reset();
-		}
-		else
-		{
-			View += ValueSize;
-		}
-		break;
-	}
-	case ECbFieldType::String:
-		ValidateCbString(View, Mode, Error);
-		break;
-	case ECbFieldType::IntegerPositive:
-		ValidateCbUInt(View, Mode, Error);
-		break;
-	case ECbFieldType::IntegerNegative:
-		ValidateCbUInt(View, Mode, Error);
-		break;
-	case ECbFieldType::Float32:
-		ValidateFixedPayload(4);
-		break;
-	case ECbFieldType::Float64:
-		ValidateCbFloat64(View, Mode, Error);
-		break;
-	case ECbFieldType::Reference:
-	case ECbFieldType::BinaryReference:
-	case ECbFieldType::Hash:
-		ValidateFixedPayload(32);
-		break;
-	case ECbFieldType::Uuid:
-		ValidateFixedPayload(16);
-		break;
-	case ECbFieldType::DateTime:
-	case ECbFieldType::TimeSpan:
-		ValidateFixedPayload(8);
-		break;
-	}
-
-	if (EnumHasAnyFlags(Error, ECbValidateError::OutOfBounds | ECbValidateError::InvalidType))
-	{
-		return FCbField();
-	}
-
-	return FCbField(FieldView.GetData(), ExternalType);
-}
-
-ECbValidateError ValidateCompactBinary(FConstMemoryView View, ECbValidateMode Mode, ECbFieldType Type)
-{
-	ECbValidateError Error = ECbValidateError::None;
-	ValidateCbField(View, Mode, Error, Type);
-	if (!View.IsEmpty() && EnumHasAnyFlags(Mode, ECbValidateMode::Padding))
-	{
-		AddError(Error, ECbValidateError::Padding);
-	}
-	return Error;
-}
-
-ECbValidateError ValidateCompactBinaryRange(FConstMemoryView View, ECbValidateMode Mode)
-{
-	ECbValidateError Error = ECbValidateError::None;
-	while (!View.IsEmpty())
-	{
-		ValidateCbField(View, Mode, Error);
-	}
-	return Error;
-}
-
-///////////////////////////////////////////////////////////////////////////////////////////////////////////////////////
-
 uint64 MeasureCompactBinary(FConstMemoryView View, ECbFieldType Type)
 {
 	uint64 Size = 0;
diff --git a/Engine/Source/Runtime/Core/Private/Serialization/CompactBinaryValidation.cpp b/Engine/Source/Runtime/Core/Private/Serialization/CompactBinaryValidation.cpp
new file mode 100644
--- /dev/null
+++ b/Engine/Source/Runtime/Core/Private/Serialization/CompactBinaryValidation.cpp
@@ -0,0 +1,398 @@
+// Copyright Epic Games, Inc. All Rights Reserved.
+
+#include "Serialization/CompactBinaryValidation.h"
+
+#include "Misc/ByteSwap.h"
+#include "Serialization/CompactBinary.h"
+#include "Serialization/VarInt.h"
+
+///////////////////////////////////////////////////////////////////////////////////////////////////////////////////////
+
+namespace CompactBinaryValidationPrivate
+{
+
+template <typename T>
+static constexpr FORCEINLINE T ReadUnaligned(const void* const Memory)
+{
+#if PLATFORM_SUPPORTS_UNALIGNED_LOADS
+	return *static_cast<const T*>(Memory);
+#else
+	T Value;
+	FMemory::Memcpy(&Value, Memory, sizeof(Value));
+	return Value;
+#endif
+}
+
+}
+
+///////////////////////////////////////////////////////////////////////////////////////////////////////////////////////
+
+/**
+ * Adds the given error(s) to the error mask.
+ *
+ * This function exists to make validation errors easier to debug by providing one location to set a breakpoint.
+ */
+FORCENOINLINE static void AddError(ECbValidateError& OutError, const ECbValidateError InError)
+{
+	OutError |= InError;
+}
+
+/**
+ * Validate and read a field type from the view.
+ *
+ * A type argument with the HasFieldType flag indicates that the type will not be read from the view.
+ */
+static ECbFieldType ValidateCbFieldType(FConstMemoryView& View, ECbValidateMode Mode, ECbValidateError& Error, ECbFieldType Type = ECbFieldType::HasFieldType)
+{
+	if (FCbFieldType::HasFieldType(Type))
+	{
+		if (View.GetSize() >= 1)
+		{
+			Type = *static_cast<const ECbFieldType*>(View.GetData());
+			View += 1;
+			if (FCbFieldType::HasFieldType(Type))
+			{
+				AddError(Error, ECbValidateError::InvalidType);
+			}
+		}
+		else
+		{
+			AddError(Error, ECbValidateError::OutOfBounds);
+			View.Reset();
+			return ECbFieldType::None;
+		}
+	}
+
+	if (FCbFieldType::GetType(Type) > ECbFieldType::TimeSpan)
+	{
+		AddError(Error, ECbValidateError::InvalidType);
+		View.Reset();
+	}
+
+	return Type;
+}
+
+/**
+ * Validate and read an unsigned integer from the view.
+ *
+ * Modifies the view to start at the end of the value, and adds error flags if applicable.
+ */
+static uint64 ValidateCbUInt(FConstMemoryView& View, ECbValidateMode Mode, ECbValidateError& Error)
+{
+	if (View.GetSize() > 0 && View.GetSize() >= MeasureVarUInt(View.GetData()))
+	{
+		uint32 ValueByteCount;
+		const uint64 Value = ReadVarUInt(View.GetData(), ValueByteCount);
+		if (EnumHasAnyFlags(Mode, ECbValidateMode::Format) && ValueByteCount > MeasureVarUInt(Value))
+		{
+			AddError(Error, ECbValidateError::InvalidInteger);
+		}
+		View += ValueByteCount;
+		return Value;
+	}
+	else
+	{
+		AddError(Error, ECbValidateError::OutOfBounds);
+		View.Reset();
+		return 0;
+	}
+}
+
+/**
+ * Validate a 64-bit floating point value from the view.
+ *
+ * Modifies the view to start at the end of the value, and adds error flags if applicable.
+ */
+static void ValidateCbFloat64(FConstMemoryView& View, ECbValidateMode Mode, ECbValidateError& Error)
+{
+	if (View.GetSize() >= sizeof(double))
+	{
+		if (EnumHasAnyFlags(Mode, ECbValidateMode::Format))
+		{
+			const uint64 RawValue = NETWORK_ORDER64(CompactBinaryValidationPrivate::ReadUnaligned<uint64>(View.GetData()));
+			const double Value = reinterpret_cast<const double&>(RawValue);
+			if (Value == double(float(Value)))
+			{
+				AddError(Error, ECbValidateError::InvalidFloat);
+			}
+		}
+		View += sizeof(double);
+	}
+	else
+	{
+		AddError(Error, ECbValidateError::OutOfBounds);
+		View.Reset();
+	}
+}
+
+/**
+ * Validate and read a string from the view.
+ *
+ * Modifies the view to start at the end of the string, and adds error flags if applicable.
+ */
+static FAnsiStringView ValidateCbString(FConstMemoryView& View, ECbValidateMode Mode, ECbValidateError& Error)
+{
+	const uint64 NameSize = ValidateCbUInt(View, Mode, Error);
+	if (View.GetSize() >= NameSize)
+	{
+		const FAnsiStringView Name(static_cast<const ANSICHAR*>(View.GetData()), static_cast<int32>(NameSize));
+		View += NameSize;
+		return Name;
+	}
+	else
+	{
+		AddError(Error, ECbValidateError::OutOfBounds);
+		View.Reset();
+		return FAnsiStringView();
+	}
+}
+
+static FCbField ValidateCbField(FConstMemoryView& View, ECbValidateMode Mode, ECbValidateError& Error, ECbFieldType ExternalType);
+
+/** A type that checks whether all validated fields are of the same type. */
+class FCbUniformFieldsValidator
+{
+public:
+	inline explicit FCbUniformFieldsValidator(ECbFieldType InExternalType)
+		: ExternalType(InExternalType)
+	{
+	}
+
+	inline FCbField ValidateField(FConstMemoryView& View, ECbValidateMode Mode, ECbValidateError& Error)
+	{
+		const void* const FieldData = View.GetData();
+		if (FCbField Field = ValidateCbField(View, Mode, Error, ExternalType))
+		{
+			++FieldCount;
+			if (FCbFieldType::HasFieldType(ExternalType))
+			{
+				const ECbFieldType FieldType = *static_cast<const ECbFieldType*>(FieldData);
+				if (FieldCount == 1)
+				{
+					FirstType = FieldType;
+				}
+				else if (FieldType != FirstType)
+				{
+					bUniform = false;
+				}
+			}
+			return Field;
+		}
+
+		// It may not safe to check for uniformity if the field was invalid.
+		bUniform = false;
+		return FCbField();
+	}
+
+	inline bool IsUniform() const { return FieldCount > 0 && bUniform; }
+
+private:
+	uint32 FieldCount = 0;
+	bool bUniform = true;
+	ECbFieldType FirstType = ECbFieldType::None;
+	ECbFieldType ExternalType;
+};
+
+static void ValidateCbObject(FConstMemoryView& View, ECbValidateMode Mode, ECbValidateError& Error, ECbFieldType ObjectType)
+{
+	const uint64 Size = ValidateCbUInt(View, Mode, Error);
+	FConstMemoryView ObjectView = View.Left(Size);
+	View += Size;
+
+	if (Size > 0)
+	{
+		TArray<FAnsiStringView, TInlineAllocator<16>> Names;
+
+		const bool bUniformObject = FCbFieldType::GetType(ObjectType) == ECbFieldType::UniformObject;
+		const ECbFieldType ExternalType = bUniformObject ? ValidateCbFieldType(ObjectView, Mode, Error) : ECbFieldType::HasFieldType;
+		FCbUniformFieldsValidator UniformValidator(ExternalType);
+		do
+		{
+			if (FCbField Field = UniformValidator.ValidateField(ObjectView, Mode, Error))
+			{
+				if (EnumHasAnyFlags(Mode, ECbValidateMode::Names))
+				{
+					if (Field.HasName())
+					{
+						Names.Add(Field.GetName());
+					}
+					else
+					{
+						AddError(Error, ECbValidateError::MissingName);
+					}
+				}
+			}
+		}
+		while (!ObjectView.IsEmpty());
+
+		if (EnumHasAnyFlags(Mode, ECbValidateMode::Names) && Names.Num() > 1)
+		{
+			Algo::Sort(Names, [](FAnsiStringView L, FAnsiStringView R) { return L.Compare(R) < 0; });
+			for (const FAnsiStringView* NamesIt = Names.GetData(), *NamesEnd = NamesIt + Names.Num() - 1; NamesIt != NamesEnd; ++NamesIt)
+			{
+				if (NamesIt[0].Equals(NamesIt[1]))
+				{
+					AddError(Error, ECbValidateError::DuplicateName);
+					break;
+				}
+			}
+		}
+
+		if (!bUniformObject && EnumHasAnyFlags(Mode, ECbValidateMode::Format) && UniformValidator.IsUniform())
+		{
+			AddError(Error, ECbValidateError::NonUniformObject);
+		}
+	}
+}
+
+static void ValidateCbArray(FConstMemoryView& View, ECbValidateMode Mode, ECbValidateError& Error, ECbFieldType ArrayType)
+{
+	const uint64 Size = ValidateCbUInt(View, Mode, Error);
+	FConstMemoryView ArrayView = View.Left(Size);
+	View += Size;
+
+	const uint64 Count = ValidateCbUInt(ArrayView, Mode, Error);
+	const uint64 FieldsSize = ArrayView.GetSize();
+	const bool bUniformArray = FCbFieldType::GetType(ArrayType) == ECbFieldType::UniformArray;
+	const ECbFieldType ExternalType = bUniformArray ? ValidateCbFieldType(ArrayView, Mode, Error) : ECbFieldType::HasFieldType;
+	FCbUniformFieldsValidator UniformValidator(ExternalType);
+
+	for (uint64 Index = 0; Index < Count; ++Index)
+	{
+		if (FCbField Field = UniformValidator.ValidateField(ArrayView, Mode, Error))
+		{
+			if (Field.HasName() && EnumHasAnyFlags(Mode, ECbValidateMode::Names))
+			{
+				AddError(Error, ECbValidateError::ArrayName);
+			}
+		}
+	}
+
+	if (!bUniformArray && EnumHasAnyFlags(Mode, ECbValidateMode::Format) && UniformValidator.IsUniform() && FieldsSize > Count)
+	{
+		AddError(Error, ECbValidateError::NonUniformArray);
+	}
+}
+
+static FCbField ValidateCbField(FConstMemoryView& View, ECbValidateMode Mode, ECbValidateError& Error, const ECbFieldType ExternalType = ECbFieldType::HasFieldType)
+{
+	const FConstMemoryView FieldView = View;
+	const ECbFieldType Type = ValidateCbFieldType(View, Mode, Error, ExternalType);
+	const FAnsiStringView Name = FCbFieldType::HasFieldName(Type) ? ValidateCbString(View, Mode, Error) : FAnsiStringView();
+
+	auto ValidateFixedPayload = [&View, &Error](uint32 PayloadSize)
+	{
+		if (View.GetSize() >= PayloadSize)
+		{
+			View += PayloadSize;
+		}
+		else
+		{
+			AddError(Error, ECbValidateError::OutOfBounds);
+			View.Reset();
+		}
+	};
+
+	if (EnumHasAnyFlags(Error, ECbValidateError::OutOfBounds | ECbValidateError::InvalidType))
+	{
+		return FCbField();
+	}
+
+	switch (FCbFieldType::GetType(Type))
+	{
+	default:
+	case ECbFieldType::None:
+		AddError(Error, ECbValidateError::InvalidType);
+		View.Reset();
+		break;
+	case ECbFieldType::Null:
+	case ECbFieldType::BoolFalse:
+	case ECbFieldType::BoolTrue:
+		if (FieldView == View)
+		{
+			// Reset the view because a zero-sized field can cause infinite field iteration.
+			AddError(Error, ECbValidateError::InvalidType);
+			View.Reset();
+		}
+		break;
+	case ECbFieldType::Object:
+	case ECbFieldType::UniformObject:
+		ValidateCbObject(View, Mode, Error, FCbFieldType::GetType(Type));
+		break;
+	case ECbFieldType::Array:
+	case ECbFieldType::UniformArray:
+		ValidateCbArray(View, Mode, Error, FCbFieldType::GetType(Type));
+		break;
+	case ECbFieldType::Binary:
+	{
+		const uint64 ValueSize = ValidateCbUInt(View, Mode, Error);
+		if (View.GetSize() < ValueSize)
+		{
+			AddError(Error, ECbValidateError::OutOfBounds);
+			View.Reset();
+		}
+		else
+		{
+			View += ValueSize;
+		}
+		break;
+	}
+	case ECbFieldType::String:
+		ValidateCbString(View, Mode, Error);
+		break;
+	case ECbFieldType::IntegerPositive:
+		ValidateCbUInt(View, Mode, Error);
+		break;
+	case ECbFieldType::IntegerNegative:
+		ValidateCbUInt(View, Mode, Error);
+		break;
+	case ECbFieldType::Float32:
+		ValidateFixedPayload(4);
+		break;
+	case ECbFieldType::Float64:
+		ValidateCbFloat64(View, Mode, Error);
+		break;
+	case ECbFieldType::Reference:
+	case ECbFieldType::BinaryReference:
+	case ECbFieldType::Hash:
+		ValidateFixedPayload(32);
+		break;
+	case ECbFieldType::Uuid:
+		ValidateFixedPayload(16);
+		break;
+	case ECbFieldType::DateTime:
+	case ECbFieldType::TimeSpan:
+		ValidateFixedPayload(8);
+		break;
+	}
+
+	if (EnumHasAnyFlags(Error, ECbValidateError::OutOfBounds | ECbValidateError::InvalidType))
+	{
+		return FCbField();
+	}
+
+	return FCbField(FieldView.GetData(), ExternalType);
+}
+
+ECbValidateError ValidateCompactBinary(FConstMemoryView View, ECbValidateMode Mode, ECbFieldType Type)
+{
+	ECbValidateError Error = ECbValidateError::None;
+	ValidateCbField(View, Mode, Error, Type);
+	if (!View.IsEmpty() && EnumHasAnyFlags(Mode, ECbValidateMode::Padding))
+	{
+		AddError(Error, ECbValidateError::Padding);
+	}
+	return Error;
+}
+
+ECbValidateError ValidateCompactBinaryRange(FConstMemoryView View, ECbValidateMode Mode)
+{
+	ECbValidateError Error = ECbValidateError::None;
+	while (!View.IsEmpty())
+	{
+		ValidateCbField(View, Mode, Error);
+	}
+	return Error;
+}
+
+///////////////////////////////////////////////////////////////////////////////////////////////////////////////////////
diff --git a/Engine/Source/Runtime/Core/Private/Tests/Serialization/CompactBinaryTest.cpp b/Engine/Source/Runtime/Core/Private/Tests/Serialization/CompactBinaryTest.cpp
--- a/Engine/Source/Runtime/Core/Private/Tests/Serialization/CompactBinaryTest.cpp
+++ b/Engine/Source/Runtime/Core/Private/Tests/Serialization/CompactBinaryTest.cpp
@@ -13,7 +13,8 @@
 
 ///////////////////////////////////////////////////////////////////////////////////////////////////////////////////////
 
-static constexpr EAutomationTestFlags::Type CompactBinaryTestFlags = EAutomationTestFlags::Type(EAutomationTestFlags::ApplicationContextMask | EAutomationTestFlags::SmokeFilter);
+static constexpr EAutomationTestFlags::Type CompactBinaryTestFlags =
+	EAutomationTestFlags::Type(EAutomationTestFlags::ApplicationContextMask | EAutomationTestFlags::SmokeFilter);
 
 ///////////////////////////////////////////////////////////////////////////////////////////////////////////////////////
 
@@ -1214,266 +1215,6 @@ bool FCbObjectRefTest::RunTest(const FString& Parameters)
 
 ///////////////////////////////////////////////////////////////////////////////////////////////////////////////////////
 
-IMPLEMENT_SIMPLE_AUTOMATION_TEST(FCbValidateTest, "System.Core.Serialization.ValidateCompactBinary", CompactBinaryTestFlags)
-bool FCbValidateTest::RunTest(const FString& Parameters)
-{
-	auto Validate = [this](std::initializer_list<uint8> Data, ECbFieldType Type = ECbFieldType::HasFieldType) -> ECbValidateError
-	{
-		return ValidateCompactBinary(MakeMemoryView(Data), ECbValidateMode::All, Type);
-	};
-	auto ValidateMode = [this](std::initializer_list<uint8> Data, ECbValidateMode Mode, ECbFieldType Type = ECbFieldType::HasFieldType) -> ECbValidateError
-	{
-		return ValidateCompactBinary(MakeMemoryView(Data), Mode, Type);
-	};
-
-	auto AddName = [](ECbFieldType Type) -> uint8 { return uint8(Type | ECbFieldType::HasFieldName); };
-
-	constexpr uint8 NullNoName = uint8(ECbFieldType::Null);
-	constexpr uint8 NullWithName = uint8(ECbFieldType::Null | ECbFieldType::HasFieldName);
-	constexpr uint8 IntNoName = uint8(ECbFieldType::IntegerPositive);
-	constexpr uint8 IntWithName = uint8(ECbFieldType::IntegerPositive | ECbFieldType::HasFieldName);
-
-	// Test OutOfBounds
-	TestEqual(TEXT("ValidateCompactBinary(OutOfBounds, Empty)"), Validate({}), ECbValidateError::OutOfBounds);
-
-	TestEqual(TEXT("ValidateCompactBinary(Valid, Null)"), Validate({NullNoName}), ECbValidateError::None);
-	TestEqual(TEXT("ValidateCompactBinary(Valid, Null, Name)"), Validate({NullWithName, 1, 'N'}), ECbValidateError::None);
-	TestEqual(TEXT("ValidateCompactBinary(OutOfBounds, Null, Name)"), Validate({NullWithName}), ECbValidateError::OutOfBounds);
-	TestEqual(TEXT("ValidateCompactBinary(OutOfBounds, Null, Name)"), Validate({NullWithName, 1}), ECbValidateError::OutOfBounds);
-	TestEqual(TEXT("ValidateCompactBinary(OutOfBounds, Null, Name)"), Validate({NullWithName, 0x80}), ECbValidateError::OutOfBounds);
-	TestEqual(TEXT("ValidateCompactBinary(OutOfBounds, Null, Name)"), Validate({NullWithName, 0x80, 128}), ECbValidateError::OutOfBounds);
-
-	TestEqual(TEXT("ValidateCompactBinary(Valid, Object, Empty)"), Validate({uint8(ECbFieldType::Object), 0}), ECbValidateError::None);
-	TestEqual(TEXT("ValidateCompactBinary(Valid, Object, Empty, NoType)"), Validate({0}, ECbFieldType::Object), ECbValidateError::None);
-	TestEqual(TEXT("ValidateCompactBinary(Valid, Object, Field)"), Validate({uint8(ECbFieldType::Object), 7, NullWithName, 1, 'N', IntWithName, 1, 'I', 0}), ECbValidateError::None);
-	TestEqual(TEXT("ValidateCompactBinary(Valid, Object, Field, NoType)"), Validate({7, NullWithName, 1, 'N', IntWithName, 1, 'I', 0}, ECbFieldType::Object), ECbValidateError::None);
-	TestEqual(TEXT("ValidateCompactBinary(OutOfBounds, Object)"), Validate({uint8(ECbFieldType::Object)}), ECbValidateError::OutOfBounds);
-	TestEqual(TEXT("ValidateCompactBinary(OutOfBounds, Object, NoType)"), Validate({}, ECbFieldType::Object), ECbValidateError::OutOfBounds);
-	TestEqual(TEXT("ValidateCompactBinary(OutOfBounds, Object)"), Validate({uint8(ECbFieldType::Object), 1}), ECbValidateError::OutOfBounds);
-	TestEqual(TEXT("ValidateCompactBinary(OutOfBounds, Object, NoType)"), Validate({1}, ECbFieldType::Object), ECbValidateError::OutOfBounds);
-	TestEqual(TEXT("ValidateCompactBinary(OutOfBounds, Object, OOB Field)"), Validate({uint8(ECbFieldType::Object), 3, AddName(ECbFieldType::Float32), 1, 'N'}), ECbValidateError::OutOfBounds);
-	TestEqual(TEXT("ValidateCompactBinary(OutOfBounds, Object, OOB Field, NoType)"), Validate({3, AddName(ECbFieldType::Float32), 1, 'N'}, ECbFieldType::Object), ECbValidateError::OutOfBounds);
-
-	TestEqual(TEXT("ValidateCompactBinary(Valid, UniformObject, Field)"), Validate({uint8(ECbFieldType::UniformObject), 3, NullWithName, 1, 'N'}), ECbValidateError::None);
-	TestEqual(TEXT("ValidateCompactBinary(Valid, UniformObject, Field, NoType)"), Validate({3, NullWithName, 1, 'N'}, ECbFieldType::UniformObject), ECbValidateError::None);
-	TestEqual(TEXT("ValidateCompactBinary(OutOfBounds, UniformObject)"), Validate({uint8(ECbFieldType::UniformObject)}), ECbValidateError::OutOfBounds);
-	TestEqual(TEXT("ValidateCompactBinary(OutOfBounds, UniformObject, NoType)"), Validate({}, ECbFieldType::UniformObject), ECbValidateError::OutOfBounds);
-	TestEqual(TEXT("ValidateCompactBinary(OutOfBounds, UniformObject)"), Validate({uint8(ECbFieldType::UniformObject), 1}), ECbValidateError::OutOfBounds);
-	TestEqual(TEXT("ValidateCompactBinary(OutOfBounds, UniformObject, NoType)"), Validate({1}, ECbFieldType::UniformObject), ECbValidateError::OutOfBounds);
-	TestEqual(TEXT("ValidateCompactBinary(OutOfBounds, UniformObject, OOB Field)"), Validate({uint8(ECbFieldType::UniformObject), 3, AddName(ECbFieldType::Float32), 1, 'N'}), ECbValidateError::OutOfBounds);
-	TestEqual(TEXT("ValidateCompactBinary(OutOfBounds, UniformObject, OOB Field, NoType)"), Validate({3, AddName(ECbFieldType::Float32), 1, 'N'}, ECbFieldType::UniformObject), ECbValidateError::OutOfBounds);
-
-	TestEqual(TEXT("ValidateCompactBinary(Valid, Array, Empty)"), Validate({uint8(ECbFieldType::Array), 1, 0}), ECbValidateError::None);
-	TestEqual(TEXT("ValidateCompactBinary(Valid, Array, Empty, NoType)"), Validate({1, 0}, ECbFieldType::Array), ECbValidateError::None);
-	TestEqual(TEXT("ValidateCompactBinary(Valid, Array, Field)"), Validate({uint8(ECbFieldType::Array), 4, 2, NullNoName, IntNoName, 0}), ECbValidateError::None);
-	TestEqual(TEXT("ValidateCompactBinary(Valid, Array, Field, NoType)"), Validate({4, 2, NullNoName, IntNoName, 0}, ECbFieldType::Array), ECbValidateError::None);
-	TestEqual(TEXT("ValidateCompactBinary(OutOfBounds, Array)"), Validate({uint8(ECbFieldType::Array)}), ECbValidateError::OutOfBounds);
-	TestEqual(TEXT("ValidateCompactBinary(OutOfBounds, Array, NoType)"), Validate({}, ECbFieldType::Array), ECbValidateError::OutOfBounds);
-	TestEqual(TEXT("ValidateCompactBinary(OutOfBounds, Array)"), Validate({uint8(ECbFieldType::Array), 1}), ECbValidateError::OutOfBounds);
-	TestEqual(TEXT("ValidateCompactBinary(OutOfBounds, Array, NoType)"), Validate({1}, ECbFieldType::Array), ECbValidateError::OutOfBounds);
-	TestEqual(TEXT("ValidateCompactBinary(OutOfBounds, Array, OOB Field)"), Validate({uint8(ECbFieldType::Array), 2, 1, uint8(ECbFieldType::Float32)}), ECbValidateError::OutOfBounds);
-	TestEqual(TEXT("ValidateCompactBinary(OutOfBounds, Array, OOB Field, NoType)"), Validate({2, 1, uint8(ECbFieldType::Float32)}, ECbFieldType::Array), ECbValidateError::OutOfBounds);
-
-	TestEqual(TEXT("ValidateCompactBinary(Valid, UniformArray, Field)"), Validate({uint8(ECbFieldType::UniformArray), 3, 1, IntNoName, 0}), ECbValidateError::None);
-	TestEqual(TEXT("ValidateCompactBinary(Valid, UniformArray, Field, NoType)"), Validate({3, 1, IntNoName, 0}, ECbFieldType::UniformArray), ECbValidateError::None);
-	TestEqual(TEXT("ValidateCompactBinary(OutOfBounds, UniformArray)"), Validate({uint8(ECbFieldType::UniformArray)}), ECbValidateError::OutOfBounds);
-	TestEqual(TEXT("ValidateCompactBinary(OutOfBounds, UniformArray, NoType)"), Validate({}, ECbFieldType::UniformArray), ECbValidateError::OutOfBounds);
-	TestEqual(TEXT("ValidateCompactBinary(OutOfBounds, UniformArray)"), Validate({uint8(ECbFieldType::UniformArray), 1}), ECbValidateError::OutOfBounds);
-	TestEqual(TEXT("ValidateCompactBinary(OutOfBounds, UniformArray, NoType)"), Validate({1}, ECbFieldType::UniformArray), ECbValidateError::OutOfBounds);
-	TestEqual(TEXT("ValidateCompactBinary(OutOfBounds, UniformArray, OOB Field)"), Validate({uint8(ECbFieldType::UniformArray), 2, 1, uint8(ECbFieldType::Float32)}), ECbValidateError::OutOfBounds);
-	TestEqual(TEXT("ValidateCompactBinary(OutOfBounds, UniformArray, OOB Field, NoType)"), Validate({2, 1, uint8(ECbFieldType::Float32)}, ECbFieldType::UniformArray), ECbValidateError::OutOfBounds);
-
-	TestEqual(TEXT("ValidateCompactBinary(Valid, Binary, Empty)"), Validate({uint8(ECbFieldType::Binary), 0}), ECbValidateError::None);
-	TestEqual(TEXT("ValidateCompactBinary(Valid, Binary, Empty, NoType)"), Validate({0}, ECbFieldType::Binary), ECbValidateError::None);
-	TestEqual(TEXT("ValidateCompactBinary(Valid, Binary, Field)"), Validate({uint8(ECbFieldType::Binary), 1, 0}), ECbValidateError::None);
-	TestEqual(TEXT("ValidateCompactBinary(Valid, Binary, Field, NoType)"), Validate({1, 0}, ECbFieldType::Binary), ECbValidateError::None);
-	TestEqual(TEXT("ValidateCompactBinary(OutOfBounds, Binary)"), Validate({uint8(ECbFieldType::Binary)}), ECbValidateError::OutOfBounds);
-	TestEqual(TEXT("ValidateCompactBinary(OutOfBounds, Binary, NoType)"), Validate({}, ECbFieldType::Binary), ECbValidateError::OutOfBounds);
-	TestEqual(TEXT("ValidateCompactBinary(OutOfBounds, Binary)"), Validate({uint8(ECbFieldType::Binary), 1}), ECbValidateError::OutOfBounds);
-	TestEqual(TEXT("ValidateCompactBinary(OutOfBounds, Binary, NoType)"), Validate({1}, ECbFieldType::Binary), ECbValidateError::OutOfBounds);
-
-	TestEqual(TEXT("ValidateCompactBinary(Valid, String, Empty)"), Validate({uint8(ECbFieldType::String), 0}), ECbValidateError::None);
-	TestEqual(TEXT("ValidateCompactBinary(Valid, String, Empty, NoType)"), Validate({0}, ECbFieldType::String), ECbValidateError::None);
-	TestEqual(TEXT("ValidateCompactBinary(Valid, String, Field)"), Validate({uint8(ECbFieldType::String), 1, 'S'}), ECbValidateError::None);
-	TestEqual(TEXT("ValidateCompactBinary(Valid, String, Field, NoType)"), Validate({1, 'S'}, ECbFieldType::String), ECbValidateError::None);
-	TestEqual(TEXT("ValidateCompactBinary(OutOfBounds, String)"), Validate({uint8(ECbFieldType::String)}), ECbValidateError::OutOfBounds);
-	TestEqual(TEXT("ValidateCompactBinary(OutOfBounds, String, NoType)"), Validate({}, ECbFieldType::String), ECbValidateError::OutOfBounds);
-	TestEqual(TEXT("ValidateCompactBinary(OutOfBounds, String)"), Validate({uint8(ECbFieldType::String), 1}), ECbValidateError::OutOfBounds);
-	TestEqual(TEXT("ValidateCompactBinary(OutOfBounds, String, NoType)"), Validate({1}, ECbFieldType::String), ECbValidateError::OutOfBounds);
-
-	TestEqual(TEXT("ValidateCompactBinary(Valid, IntegerPositive, 1-byte)"), Validate({uint8(ECbFieldType::IntegerPositive), 0}), ECbValidateError::None);
-	TestEqual(TEXT("ValidateCompactBinary(Valid, IntegerPositive, 1-byte, NoType)"), Validate({0}, ECbFieldType::IntegerPositive), ECbValidateError::None);
-	TestEqual(TEXT("ValidateCompactBinary(Valid, IntegerPositive, 2-byte)"), Validate({uint8(ECbFieldType::IntegerPositive), 0x80, 0x80}), ECbValidateError::None);
-	TestEqual(TEXT("ValidateCompactBinary(Valid, IntegerPositive, 2-byte, NoType)"), Validate({0x80, 0x80}, ECbFieldType::IntegerPositive), ECbValidateError::None);
-	TestEqual(TEXT("ValidateCompactBinary(OutOfBounds, IntegerPositive, 1-byte)"), Validate({uint8(ECbFieldType::IntegerPositive)}), ECbValidateError::OutOfBounds);
-	TestEqual(TEXT("ValidateCompactBinary(OutOfBounds, IntegerPositive, 1-byte, NoType)"), Validate({}, ECbFieldType::IntegerPositive), ECbValidateError::OutOfBounds);
-	TestEqual(TEXT("ValidateCompactBinary(OutOfBounds, IntegerPositive, 2-byte)"), Validate({uint8(ECbFieldType::IntegerPositive), 0x80}), ECbValidateError::OutOfBounds);
-	TestEqual(TEXT("ValidateCompactBinary(OutOfBounds, IntegerPositive, 2-byte, NoType)"), Validate({0x80}, ECbFieldType::IntegerPositive), ECbValidateError::OutOfBounds);
-	TestEqual(TEXT("ValidateCompactBinary(OutOfBounds, IntegerPositive, 9-byte)"), Validate({uint8(ECbFieldType::IntegerPositive), 0xff, 0, 0, 0, 0, 0, 0, 0}), ECbValidateError::OutOfBounds);
-	TestEqual(TEXT("ValidateCompactBinary(OutOfBounds, IntegerPositive, 9-byte, NoType)"), Validate({0xff, 0, 0, 0, 0, 0, 0, 0}, ECbFieldType::IntegerPositive), ECbValidateError::OutOfBounds);
-
-	TestEqual(TEXT("ValidateCompactBinary(Valid, IntegerNegative, 1-byte)"), Validate({uint8(ECbFieldType::IntegerNegative), 0}), ECbValidateError::None);
-	TestEqual(TEXT("ValidateCompactBinary(Valid, IntegerNegative, 1-byte, NoType)"), Validate({0}, ECbFieldType::IntegerNegative), ECbValidateError::None);
-	TestEqual(TEXT("ValidateCompactBinary(Valid, IntegerNegative, 2-byte)"), Validate({uint8(ECbFieldType::IntegerNegative), 0x80, 0x80}), ECbValidateError::None);
-	TestEqual(TEXT("ValidateCompactBinary(Valid, IntegerNegative, 2-byte, NoType)"), Validate({0x80, 0x80}, ECbFieldType::IntegerNegative), ECbValidateError::None);
-	TestEqual(TEXT("ValidateCompactBinary(OutOfBounds, IntegerNegative, 1-byte)"), Validate({uint8(ECbFieldType::IntegerNegative)}), ECbValidateError::OutOfBounds);
-	TestEqual(TEXT("ValidateCompactBinary(OutOfBounds, IntegerNegative, 1-byte, NoType)"), Validate({}, ECbFieldType::IntegerNegative), ECbValidateError::OutOfBounds);
-	TestEqual(TEXT("ValidateCompactBinary(OutOfBounds, IntegerNegative, 2-byte)"), Validate({uint8(ECbFieldType::IntegerNegative), 0x80}), ECbValidateError::OutOfBounds);
-	TestEqual(TEXT("ValidateCompactBinary(OutOfBounds, IntegerNegative, 2-byte, NoType)"), Validate({0x80}, ECbFieldType::IntegerNegative), ECbValidateError::OutOfBounds);
-	TestEqual(TEXT("ValidateCompactBinary(OutOfBounds, IntegerNegative, 9-byte)"), Validate({uint8(ECbFieldType::IntegerNegative), 0xff, 0, 0, 0, 0, 0, 0, 0}), ECbValidateError::OutOfBounds);
-	TestEqual(TEXT("ValidateCompactBinary(OutOfBounds, IntegerNegative, 9-byte, NoType)"), Validate({0xff, 0, 0, 0, 0, 0, 0, 0}, ECbFieldType::IntegerNegative), ECbValidateError::OutOfBounds);
-
-	TestEqual(TEXT("ValidateCompactBinary(Valid, Float32)"), Validate({uint8(ECbFieldType::Float32), 0, 0, 0, 0}), ECbValidateError::None);
-	TestEqual(TEXT("ValidateCompactBinary(Valid, Float32, NoType)"), Validate({0, 0, 0, 0}, ECbFieldType::Float32), ECbValidateError::None);
-	TestEqual(TEXT("ValidateCompactBinary(OutOfBounds, Float32)"), Validate({uint8(ECbFieldType::Float32), 0, 0, 0}), ECbValidateError::OutOfBounds);
-	TestEqual(TEXT("ValidateCompactBinary(OutOfBounds, Float32, NoType)"), Validate({0, 0, 0}, ECbFieldType::Float32), ECbValidateError::OutOfBounds);
-
-	TestEqual(TEXT("ValidateCompactBinary(Valid, Float64)"), Validate({uint8(ECbFieldType::Float64), 0x3f, 0xff, 0xff, 0xff, 0xf0, 0x00, 0x00, 0x00}), ECbValidateError::None);
-	TestEqual(TEXT("ValidateCompactBinary(Valid, Float64, NoType)"), Validate({0x3f, 0xff, 0xff, 0xff, 0xf0, 0x00, 0x00, 0x00}, ECbFieldType::Float64), ECbValidateError::None);
-	TestEqual(TEXT("ValidateCompactBinary(OutOfBounds, Float64)"), Validate({uint8(ECbFieldType::Float64), 0, 0, 0, 0, 0, 0, 0}), ECbValidateError::OutOfBounds);
-	TestEqual(TEXT("ValidateCompactBinary(OutOfBounds, Float64, NoType)"), Validate({0, 0, 0, 0, 0, 0, 0}, ECbFieldType::Float64), ECbValidateError::OutOfBounds);
-
-	TestEqual(TEXT("ValidateCompactBinary(Valid, BoolFalse)"), Validate({uint8(ECbFieldType::BoolFalse)}), ECbValidateError::None);
-	TestEqual(TEXT("ValidateCompactBinary(Valid, BoolTrue)"), Validate({uint8(ECbFieldType::BoolTrue)}), ECbValidateError::None);
-
-	TestEqual(TEXT("ValidateCompactBinary(Valid, Reference)"), Validate({uint8(ECbFieldType::Reference), 0, 0, 0, 0, 0, 0, 0, 0, 0, 0, 0, 0, 0, 0, 0, 0, 0, 0, 0, 0, 0, 0, 0, 0, 0, 0, 0, 0, 0, 0, 0, 0}), ECbValidateError::None);
-	TestEqual(TEXT("ValidateCompactBinary(Valid, Reference, NoType)"), Validate({0, 0, 0, 0, 0, 0, 0, 0, 0, 0, 0, 0, 0, 0, 0, 0, 0, 0, 0, 0, 0, 0, 0, 0, 0, 0, 0, 0, 0, 0, 0, 0}, ECbFieldType::Reference), ECbValidateError::None);
-	TestEqual(TEXT("ValidateCompactBinary(OutOfBounds, Reference)"), Validate({uint8(ECbFieldType::Reference), 0, 0, 0, 0, 0, 0, 0, 0, 0, 0, 0, 0, 0, 0, 0, 0, 0, 0, 0, 0, 0, 0, 0, 0, 0, 0, 0, 0, 0, 0, 0}), ECbValidateError::OutOfBounds);
-	TestEqual(TEXT("ValidateCompactBinary(OutOfBounds, Reference, NoType)"), Validate({0, 0, 0, 0, 0, 0, 0, 0, 0, 0, 0, 0, 0, 0, 0, 0, 0, 0, 0, 0, 0, 0, 0, 0, 0, 0, 0, 0, 0, 0, 0}, ECbFieldType::Reference), ECbValidateError::OutOfBounds);
-
-	TestEqual(TEXT("ValidateCompactBinary(Valid, BinaryReference)"), Validate({uint8(ECbFieldType::BinaryReference), 0, 0, 0, 0, 0, 0, 0, 0, 0, 0, 0, 0, 0, 0, 0, 0, 0, 0, 0, 0, 0, 0, 0, 0, 0, 0, 0, 0, 0, 0, 0, 0}), ECbValidateError::None);
-	TestEqual(TEXT("ValidateCompactBinary(Valid, BinaryReference, NoType)"), Validate({0, 0, 0, 0, 0, 0, 0, 0, 0, 0, 0, 0, 0, 0, 0, 0, 0, 0, 0, 0, 0, 0, 0, 0, 0, 0, 0, 0, 0, 0, 0, 0}, ECbFieldType::BinaryReference), ECbValidateError::None);
-	TestEqual(TEXT("ValidateCompactBinary(OutOfBounds, BinaryReference)"), Validate({uint8(ECbFieldType::BinaryReference), 0, 0, 0, 0, 0, 0, 0, 0, 0, 0, 0, 0, 0, 0, 0, 0, 0, 0, 0, 0, 0, 0, 0, 0, 0, 0, 0, 0, 0, 0, 0}), ECbValidateError::OutOfBounds);
-	TestEqual(TEXT("ValidateCompactBinary(OutOfBounds, BinaryReference, NoType)"), Validate({0, 0, 0, 0, 0, 0, 0, 0, 0, 0, 0, 0, 0, 0, 0, 0, 0, 0, 0, 0, 0, 0, 0, 0, 0, 0, 0, 0, 0, 0, 0}, ECbFieldType::BinaryReference), ECbValidateError::OutOfBounds);
-
-	TestEqual(TEXT("ValidateCompactBinary(Valid, Hash)"), Validate({uint8(ECbFieldType::Hash), 0, 0, 0, 0, 0, 0, 0, 0, 0, 0, 0, 0, 0, 0, 0, 0, 0, 0, 0, 0, 0, 0, 0, 0, 0, 0, 0, 0, 0, 0, 0, 0}), ECbValidateError::None);
-	TestEqual(TEXT("ValidateCompactBinary(Valid, Hash, NoType)"), Validate({0, 0, 0, 0, 0, 0, 0, 0, 0, 0, 0, 0, 0, 0, 0, 0, 0, 0, 0, 0, 0, 0, 0, 0, 0, 0, 0, 0, 0, 0, 0, 0}, ECbFieldType::Hash), ECbValidateError::None);
-	TestEqual(TEXT("ValidateCompactBinary(OutOfBounds, Hash)"), Validate({uint8(ECbFieldType::Hash), 0, 0, 0, 0, 0, 0, 0, 0, 0, 0, 0, 0, 0, 0, 0, 0, 0, 0, 0, 0, 0, 0, 0, 0, 0, 0, 0, 0, 0, 0, 0}), ECbValidateError::OutOfBounds);
-	TestEqual(TEXT("ValidateCompactBinary(OutOfBounds, Hash, NoType)"), Validate({0, 0, 0, 0, 0, 0, 0, 0, 0, 0, 0, 0, 0, 0, 0, 0, 0, 0, 0, 0, 0, 0, 0, 0, 0, 0, 0, 0, 0, 0, 0}, ECbFieldType::Hash), ECbValidateError::OutOfBounds);
-
-	TestEqual(TEXT("ValidateCompactBinary(Valid, Uuid)"), Validate({uint8(ECbFieldType::Uuid), 0, 0, 0, 0, 0, 0, 0, 0, 0, 0, 0, 0, 0, 0, 0, 0}), ECbValidateError::None);
-	TestEqual(TEXT("ValidateCompactBinary(Valid, Uuid, NoType)"), Validate({0, 0, 0, 0, 0, 0, 0, 0, 0, 0, 0, 0, 0, 0, 0, 0}, ECbFieldType::Uuid), ECbValidateError::None);
-	TestEqual(TEXT("ValidateCompactBinary(OutOfBounds, Uuid)"), Validate({uint8(ECbFieldType::Uuid), 0, 0, 0, 0, 0, 0, 0, 0, 0, 0, 0, 0, 0, 0, 0}), ECbValidateError::OutOfBounds);
-	TestEqual(TEXT("ValidateCompactBinary(OutOfBounds, Uuid, NoType)"), Validate({0, 0, 0, 0, 0, 0, 0, 0, 0, 0, 0, 0, 0, 0, 0}, ECbFieldType::Uuid), ECbValidateError::OutOfBounds);
-
-	TestEqual(TEXT("ValidateCompactBinary(Valid, DateTime)"), Validate({uint8(ECbFieldType::DateTime), 0, 0, 0, 0, 0, 0, 0, 0}), ECbValidateError::None);
-	TestEqual(TEXT("ValidateCompactBinary(Valid, DateTime, NoType)"), Validate({0, 0, 0, 0, 0, 0, 0, 0}, ECbFieldType::DateTime), ECbValidateError::None);
-	TestEqual(TEXT("ValidateCompactBinary(OutOfBounds, DateTime)"), Validate({uint8(ECbFieldType::DateTime), 0, 0, 0, 0, 0, 0, 0}), ECbValidateError::OutOfBounds);
-	TestEqual(TEXT("ValidateCompactBinary(OutOfBounds, DateTime, NoType)"), Validate({0, 0, 0, 0, 0, 0, 0}, ECbFieldType::DateTime), ECbValidateError::OutOfBounds);
-
-	TestEqual(TEXT("ValidateCompactBinary(Valid, TimeSpan)"), Validate({uint8(ECbFieldType::TimeSpan), 0, 0, 0, 0, 0, 0, 0, 0}), ECbValidateError::None);
-	TestEqual(TEXT("ValidateCompactBinary(Valid, TimeSpan, NoType)"), Validate({0, 0, 0, 0, 0, 0, 0, 0}, ECbFieldType::TimeSpan), ECbValidateError::None);
-	TestEqual(TEXT("ValidateCompactBinary(OutOfBounds, TimeSpan)"), Validate({uint8(ECbFieldType::TimeSpan), 0, 0, 0, 0, 0, 0, 0}), ECbValidateError::OutOfBounds);
-	TestEqual(TEXT("ValidateCompactBinary(OutOfBounds, TimeSpan, NoType)"), Validate({0, 0, 0, 0, 0, 0, 0}, ECbFieldType::TimeSpan), ECbValidateError::OutOfBounds);
-
-	// Test InvalidType
-	TestEqual(TEXT("ValidateCompactBinary(InvalidType, Unknown)"), Validate({uint8(ECbFieldType::TimeSpan) + 1}), ECbValidateError::InvalidType);
-	TestEqual(TEXT("ValidateCompactBinary(InvalidType, Unknown)"), Validate({}, ECbFieldType(uint8(ECbFieldType::TimeSpan) + 1)), ECbValidateError::InvalidType);
-	TestEqual(TEXT("ValidateCompactBinary(InvalidType, HasFieldType)"), Validate({uint8(ECbFieldType::Null | ECbFieldType::HasFieldType)}), ECbValidateError::InvalidType);
-
-	TestEqual(TEXT("ValidateCompactBinary(InvalidType, ZeroSizeField)"), Validate({}, ECbFieldType::Null), ECbValidateError::InvalidType);
-	TestEqual(TEXT("ValidateCompactBinary(InvalidType, ZeroSizeField, BoolFalse)"), Validate({}, ECbFieldType::BoolFalse), ECbValidateError::InvalidType);
-	TestEqual(TEXT("ValidateCompactBinary(InvalidType, ZeroSizeField, BoolTrue)"), Validate({}, ECbFieldType::BoolTrue), ECbValidateError::InvalidType);
-
-	TestEqual(TEXT("ValidateCompactBinary(InvalidType, ZeroSizeField, Array)"), Validate({uint8(ECbFieldType::UniformArray), 2, 2, NullNoName}), ECbValidateError::InvalidType);
-	TestEqual(TEXT("ValidateCompactBinary(InvalidType, ZeroSizeField, Object)"), Validate({uint8(ECbFieldType::UniformObject), 2, NullNoName, 0}), ECbValidateError::InvalidType);
-
-	// Test DuplicateName
-	TestEqual(TEXT("ValidateCompactBinary(DuplicateName)"), Validate({uint8(ECbFieldType::UniformObject), 7, NullWithName, 1, 'A', 1, 'B', 1, 'A'}), ECbValidateError::DuplicateName);
-	TestEqual(TEXT("ValidateCompactBinary(DuplicateName, CaseSensitive)"), Validate({uint8(ECbFieldType::UniformObject), 7, NullWithName, 1, 'A', 1, 'B', 1, 'a'}), ECbValidateError::None);
-	TestEqual(TEXT("ValidateCompactBinary(DuplicateName, Mode)"), ValidateMode({uint8(ECbFieldType::UniformObject), 7, NullWithName, 1, 'A', 1, 'B', 1, 'A'}, ECbValidateMode::All & ~ECbValidateMode::Names), ECbValidateError::None);
-
-	// Test MissingName
-	TestEqual(TEXT("ValidateCompactBinary(MissingName)"), Validate({uint8(ECbFieldType::Object), 3, NullNoName, IntNoName, 0}), ECbValidateError::MissingName);
-	TestEqual(TEXT("ValidateCompactBinary(MissingName, Uniform)"), Validate({uint8(ECbFieldType::UniformObject), 3, IntNoName, 0, 0}), ECbValidateError::MissingName);
-	TestEqual(TEXT("ValidateCompactBinary(MissingName, Mode)"), ValidateMode({uint8(ECbFieldType::Object), 3, NullNoName, IntNoName, 0}, ECbValidateMode::All & ~ECbValidateMode::Names), ECbValidateError::None);
-	TestEqual(TEXT("ValidateCompactBinary(MissingName, Uniform, Mode)"), ValidateMode({uint8(ECbFieldType::UniformObject), 3, IntNoName, 0, 0}, ECbValidateMode::All & ~ECbValidateMode::Names), ECbValidateError::None);
-
-	// Test ArrayName
-	TestEqual(TEXT("ValidateCompactBinary(ArrayName)"), Validate({uint8(ECbFieldType::Array), 5, 2, NullNoName, NullWithName, 1, 'F'}), ECbValidateError::ArrayName);
-	TestEqual(TEXT("ValidateCompactBinary(ArrayName, Uniform)"), Validate({uint8(ECbFieldType::UniformArray), 4, 1, NullWithName, 1, 'F'}), ECbValidateError::ArrayName);
-	TestEqual(TEXT("ValidateCompactBinary(ArrayName, Mode)"), ValidateMode({uint8(ECbFieldType::Array), 5, 2, NullNoName, NullWithName, 1, 'F'}, ECbValidateMode::All & ~ECbValidateMode::Names), ECbValidateError::None);
-	TestEqual(TEXT("ValidateCompactBinary(ArrayName, Uniform, Mode)"), ValidateMode({uint8(ECbFieldType::UniformArray), 4, 1, NullWithName, 1, 'F'}, ECbValidateMode::All & ~ECbValidateMode::Names), ECbValidateError::None);
-
-	// Test InvalidString
-	// Not tested or implemented yet because the engine does not provide enough UTF-8 functionality.
-
-	// Test InvalidInteger
-	TestEqual(TEXT("ValidateCompactBinary(InvalidInteger, NameSize)"), Validate({NullWithName, 0x80, 1, 'N'}), ECbValidateError::InvalidInteger);
-	TestEqual(TEXT("ValidateCompactBinary(InvalidInteger, ObjectSize)"), Validate({uint8(ECbFieldType::Object), 0xc0, 0, 0}), ECbValidateError::InvalidInteger);
-	TestEqual(TEXT("ValidateCompactBinary(InvalidInteger, ArraySize)"), Validate({uint8(ECbFieldType::Array), 0xe0, 0, 0, 1, 0}), ECbValidateError::InvalidInteger);
-	TestEqual(TEXT("ValidateCompactBinary(InvalidInteger, ArrayCount)"), Validate({uint8(ECbFieldType::Array), 5, 0xf0, 0, 0, 0, 0}), ECbValidateError::InvalidInteger);
-	TestEqual(TEXT("ValidateCompactBinary(InvalidInteger, BinarySize)"), Validate({uint8(ECbFieldType::Binary), 0xf8, 0, 0, 0, 0, 0}), ECbValidateError::InvalidInteger);
-	TestEqual(TEXT("ValidateCompactBinary(InvalidInteger, StringSize)"), Validate({uint8(ECbFieldType::String), 0xfc, 0, 0, 0, 0, 0, 0}), ECbValidateError::InvalidInteger);
-	TestEqual(TEXT("ValidateCompactBinary(InvalidInteger, IntegerPositive)"), Validate({uint8(ECbFieldType::IntegerPositive), 0xfe, 0, 0, 0, 0, 0, 0, 0}), ECbValidateError::InvalidInteger);
-	TestEqual(TEXT("ValidateCompactBinary(InvalidInteger, IntegerNegative)"), Validate({uint8(ECbFieldType::IntegerNegative), 0xff, 0, 0, 0, 0, 0, 0, 0, 0}), ECbValidateError::InvalidInteger);
-	TestEqual(TEXT("ValidateCompactBinary(InvalidInteger, ArraySize)"), Validate({uint8(ECbFieldType::Array), 0x80, 1, 0}), ECbValidateError::InvalidInteger);
-	TestEqual(TEXT("ValidateCompactBinary(InvalidInteger, ArrayCount)"), Validate({uint8(ECbFieldType::Array), 3, 0xc0, 0, 0}), ECbValidateError::InvalidInteger);
-	TestEqual(TEXT("ValidateCompactBinary(InvalidInteger, ObjectSize)"), Validate({uint8(ECbFieldType::Object), 0xe0, 0, 0, 0}), ECbValidateError::InvalidInteger);
-
-	TestEqual(TEXT("ValidateCompactBinary(InvalidInteger, NameSize, Mode)"), ValidateMode({NullWithName, 0x80, 1, 'N'}, ECbValidateMode::All & ~ECbValidateMode::Format), ECbValidateError::None);
-	TestEqual(TEXT("ValidateCompactBinary(InvalidInteger, ArraySize, Mode)"), ValidateMode({uint8(ECbFieldType::Array), 0xc0, 0, 1, 0}, ECbValidateMode::All & ~ECbValidateMode::Format), ECbValidateError::None);
-	TestEqual(TEXT("ValidateCompactBinary(InvalidInteger, ObjectSize, Mode)"), ValidateMode({uint8(ECbFieldType::Object), 0xe0, 0, 0, 0}, ECbValidateMode::All & ~ECbValidateMode::Format), ECbValidateError::None);
-
-	// Test InvalidFloat
-	TestEqual(TEXT("ValidateCompactBinary(InvalidFloat, MaxSignificant+1)"), Validate({uint8(ECbFieldType::Float64), 0x3f, 0xff, 0xff, 0xff, 0xf0, 0x00, 0x00, 0x00}), ECbValidateError::None); // 1.9999999403953552
-	TestEqual(TEXT("ValidateCompactBinary(InvalidFloat, MaxExponent+1)"), Validate({uint8(ECbFieldType::Float64), 0x47, 0xff, 0xff, 0xff, 0xe0, 0x00, 0x00, 0x00}), ECbValidateError::None); // 6.8056469327705771e38
-	TestEqual(TEXT("ValidateCompactBinary(InvalidFloat, MaxSignificand)"), Validate({uint8(ECbFieldType::Float64), 0x3f, 0xff, 0xff, 0xff, 0xe0, 0x00, 0x00, 0x00}), ECbValidateError::InvalidFloat); // 1.9999998807907104
-	TestEqual(TEXT("ValidateCompactBinary(InvalidFloat, MaxExponent)"), Validate({uint8(ECbFieldType::Float64), 0x47, 0xef, 0xff, 0xff, 0xe0, 0x00, 0x00, 0x00}), ECbValidateError::InvalidFloat); // 3.4028234663852886e38
-	TestEqual(TEXT("ValidateCompactBinary(InvalidFloat, MaxSignificand, Mode)"), ValidateMode({uint8(ECbFieldType::Float64), 0x3f, 0xff, 0xff, 0xff, 0xe0, 0x00, 0x00, 0x00}, ECbValidateMode::All & ~ECbValidateMode::Format), ECbValidateError::None); // 1.9999998807907104
-	TestEqual(TEXT("ValidateCompactBinary(InvalidFloat, MaxExponent, Mode)"), ValidateMode({uint8(ECbFieldType::Float64), 0x47, 0xef, 0xff, 0xff, 0xe0, 0x00, 0x00, 0x00}, ECbValidateMode::All & ~ECbValidateMode::Format), ECbValidateError::None); // 3.4028234663852886e38
-
-	// Test NonUniformObject
-	TestEqual(TEXT("ValidateCompactBinary(NonUniformObject)"), Validate({uint8(ECbFieldType::Object), 3, NullWithName, 1, 'A'}), ECbValidateError::NonUniformObject);
-	TestEqual(TEXT("ValidateCompactBinary(NonUniformObject)"), Validate({uint8(ECbFieldType::Object), 6, NullWithName, 1, 'A', NullWithName, 1, 'B'}), ECbValidateError::NonUniformObject);
-	TestEqual(TEXT("ValidateCompactBinary(NonUniformObject, Mode)"), ValidateMode({uint8(ECbFieldType::Object), 3, NullWithName, 1, 'A'}, ECbValidateMode::All & ~ECbValidateMode::Format), ECbValidateError::None);
-	TestEqual(TEXT("ValidateCompactBinary(NonUniformObject, Mode)"), ValidateMode({uint8(ECbFieldType::Object), 6, NullWithName, 1, 'A', NullWithName, 1, 'B'}, ECbValidateMode::All & ~ECbValidateMode::Format), ECbValidateError::None);
-
-	// Test NonUniformArray
-	TestEqual(TEXT("ValidateCompactBinary(NonUniformArray)"), Validate({uint8(ECbFieldType::Array), 3, 1, IntNoName, 0}), ECbValidateError::NonUniformArray);
-	TestEqual(TEXT("ValidateCompactBinary(NonUniformArray)"), Validate({uint8(ECbFieldType::Array), 5, 2, IntNoName, 1, IntNoName, 2}), ECbValidateError::NonUniformArray);
-	TestEqual(TEXT("ValidateCompactBinary(NonUniformArray, Null)"), Validate({uint8(ECbFieldType::Array), 3, 2, NullNoName, NullNoName}), ECbValidateError::None);
-	TestEqual(TEXT("ValidateCompactBinary(NonUniformArray, Bool)"), Validate({uint8(ECbFieldType::Array), 3, 2, uint8(ECbFieldType::BoolFalse), uint8(ECbFieldType::BoolFalse)}), ECbValidateError::None);
-	TestEqual(TEXT("ValidateCompactBinary(NonUniformArray, Bool)"), Validate({uint8(ECbFieldType::Array), 3, 2, uint8(ECbFieldType::BoolTrue), uint8(ECbFieldType::BoolTrue)}), ECbValidateError::None);
-	TestEqual(TEXT("ValidateCompactBinary(NonUniformArray, Mode)"), ValidateMode({uint8(ECbFieldType::Array), 3, 1, IntNoName, 0}, ECbValidateMode::All & ~ECbValidateMode::Format), ECbValidateError::None);
-	TestEqual(TEXT("ValidateCompactBinary(NonUniformArray, Mode)"), ValidateMode({uint8(ECbFieldType::Array), 5, 2, IntNoName, 1, IntNoName, 2}, ECbValidateMode::All & ~ECbValidateMode::Format), ECbValidateError::None);
-
-	// Test Padding
-	TestEqual(TEXT("ValidateCompactBinary(Padding)"), Validate({NullNoName, 0}), ECbValidateError::Padding);
-	TestEqual(TEXT("ValidateCompactBinary(Padding)"), Validate({uint8(ECbFieldType::Array), 1, 0, 0}), ECbValidateError::Padding);
-	TestEqual(TEXT("ValidateCompactBinary(Padding)"), Validate({uint8(ECbFieldType::Object), 0, 0}), ECbValidateError::Padding);
-	TestEqual(TEXT("ValidateCompactBinary(Padding, Mode)"), ValidateMode({NullNoName, 0}, ECbValidateMode::All & ~ECbValidateMode::Padding), ECbValidateError::None);
-	TestEqual(TEXT("ValidateCompactBinary(Padding, Mode)"), ValidateMode({uint8(ECbFieldType::Array), 1, 0, 0}, ECbValidateMode::All & ~ECbValidateMode::Padding), ECbValidateError::None);
-	TestEqual(TEXT("ValidateCompactBinary(Padding, Mode)"), ValidateMode({uint8(ECbFieldType::Object), 0, 0}, ECbValidateMode::All & ~ECbValidateMode::Padding), ECbValidateError::None);
-
-	return true;
-}
-
-///////////////////////////////////////////////////////////////////////////////////////////////////////////////////////
-
-IMPLEMENT_SIMPLE_AUTOMATION_TEST(FCbValidateRangeTest, "System.Core.Serialization.ValidateCompactBinaryRange", CompactBinaryTestFlags)
-bool FCbValidateRangeTest::RunTest(const FString& Parameters)
-{
-	auto Validate = [this](std::initializer_list<uint8> Data) -> ECbValidateError
-	{
-		return ValidateCompactBinaryRange(MakeMemoryView(Data), ECbValidateMode::All);
-	};
-
-	// Test Empty
-	TestEqual(TEXT("ValidateCompactBinaryRange(Empty)"), Validate({}), ECbValidateError::None);
-
-	// Test Valid
-	TestEqual(TEXT("ValidateCompactBinaryRange(Null x2)"), Validate({uint8(ECbFieldType::Null), uint8(ECbFieldType::Null)}), ECbValidateError::None);
-
-	// Test Padding
-	TestEqual(TEXT("ValidateCompactBinaryRange(Padding InvalidType)"), Validate({uint8(ECbFieldType::Null), 0}), ECbValidateError::InvalidType);
-	TestEqual(TEXT("ValidateCompactBinaryRange(Padding OutOfBounds)"), Validate({uint8(ECbFieldType::Null), uint8(ECbFieldType::Binary)}), ECbValidateError::OutOfBounds);
-
-	return true;
-}
-
-///////////////////////////////////////////////////////////////////////////////////////////////////////////////////////
-
 IMPLEMENT_SIMPLE_AUTOMATION_TEST(FCbMeasureTest, "System.Core.Serialization.MeasureCompactBinary", CompactBinaryTestFlags)
 bool FCbMeasureTest::RunTest(const FString& Parameters)
 {
diff --git a/Engine/Source/Runtime/Core/Private/Tests/Serialization/CompactBinaryValidationTest.cpp b/Engine/Source/Runtime/Core/Private/Tests/Serialization/CompactBinaryValidationTest.cpp
new file mode 100644
--- /dev/null
+++ b/Engine/Source/Runtime/Core/Private/Tests/Serialization/CompactBinaryValidationTest.cpp
@@ -0,0 +1,276 @@
+// Copyright Epic Games, Inc. All Rights Reserved.
+
+#include "Serialization/CompactBinaryValidation.h"
+
+#include "Misc/AutomationTest.h"
+
+#if WITH_DEV_AUTOMATION_TESTS
+
+///////////////////////////////////////////////////////////////////////////////////////////////////////////////////////
+
+static constexpr EAutomationTestFlags::Type CompactBinaryValidationTestFlags =
+	EAutomationTestFlags::Type(EAutomationTestFlags::ApplicationContextMask | EAutomationTestFlags::SmokeFilter);
+
+///////////////////////////////////////////////////////////////////////////////////////////////////////////////////////
+
+IMPLEMENT_SIMPLE_AUTOMATION_TEST(FCbValidateTest, "System.Core.Serialization.ValidateCompactBinary", CompactBinaryValidationTestFlags)
+bool FCbValidateTest::RunTest(const FString& Parameters)
+{
+	auto Validate = [this](std::initializer_list<uint8> Data, ECbFieldType Type = ECbFieldType::HasFieldType) -> ECbValidateError
+	{
+		return ValidateCompactBinary(MakeMemoryView(Data), ECbValidateMode::All, Type);
+	};
+	auto ValidateMode = [this](std::initializer_list<uint8> Data, ECbValidateMode Mode, ECbFieldType Type = ECbFieldType::HasFieldType) -> ECbValidateError
+	{
+		return ValidateCompactBinary(MakeMemoryView(Data), Mode, Type);
+	};
+
+	auto AddName = [](ECbFieldType Type) -> uint8 { return uint8(Type | ECbFieldType::HasFieldName); };
+
+	constexpr uint8 NullNoName = uint8(ECbFieldType::Null);
+	constexpr uint8 NullWithName = uint8(ECbFieldType::Null | ECbFieldType::HasFieldName);
+	constexpr uint8 IntNoName = uint8(ECbFieldType::IntegerPositive);
+	constexpr uint8 IntWithName = uint8(ECbFieldType::IntegerPositive | ECbFieldType::HasFieldName);
+
+	// Test OutOfBounds
+	TestEqual(TEXT("ValidateCompactBinary(OutOfBounds, Empty)"), Validate({}), ECbValidateError::OutOfBounds);
+
+	TestEqual(TEXT("ValidateCompactBinary(Valid, Null)"), Validate({NullNoName}), ECbValidateError::None);
+	TestEqual(TEXT("ValidateCompactBinary(Valid, Null, Name)"), Validate({NullWithName, 1, 'N'}), ECbValidateError::None);
+	TestEqual(TEXT("ValidateCompactBinary(OutOfBounds, Null, Name)"), Validate({NullWithName}), ECbValidateError::OutOfBounds);
+	TestEqual(TEXT("ValidateCompactBinary(OutOfBounds, Null, Name)"), Validate({NullWithName, 1}), ECbValidateError::OutOfBounds);
+	TestEqual(TEXT("ValidateCompactBinary(OutOfBounds, Null, Name)"), Validate({NullWithName, 0x80}), ECbValidateError::OutOfBounds);
+	TestEqual(TEXT("ValidateCompactBinary(OutOfBounds, Null, Name)"), Validate({NullWithName, 0x80, 128}), ECbValidateError::OutOfBounds);
+
+	TestEqual(TEXT("ValidateCompactBinary(Valid, Object, Empty)"), Validate({uint8(ECbFieldType::Object), 0}), ECbValidateError::None);
+	TestEqual(TEXT("ValidateCompactBinary(Valid, Object, Empty, NoType)"), Validate({0}, ECbFieldType::Object), ECbValidateError::None);
+	TestEqual(TEXT("ValidateCompactBinary(Valid, Object, Field)"), Validate({uint8(ECbFieldType::Object), 7, NullWithName, 1, 'N', IntWithName, 1, 'I', 0}), ECbValidateError::None);
+	TestEqual(TEXT("ValidateCompactBinary(Valid, Object, Field, NoType)"), Validate({7, NullWithName, 1, 'N', IntWithName, 1, 'I', 0}, ECbFieldType::Object), ECbValidateError::None);
+	TestEqual(TEXT("ValidateCompactBinary(OutOfBounds, Object)"), Validate({uint8(ECbFieldType::Object)}), ECbValidateError::OutOfBounds);
+	TestEqual(TEXT("ValidateCompactBinary(OutOfBounds, Object, NoType)"), Validate({}, ECbFieldType::Object), ECbValidateError::OutOfBounds);
+	TestEqual(TEXT("ValidateCompactBinary(OutOfBounds, Object)"), Validate({uint8(ECbFieldType::Object), 1}), ECbValidateError::OutOfBounds);
+	TestEqual(TEXT("ValidateCompactBinary(OutOfBounds, Object, NoType)"), Validate({1}, ECbFieldType::Object), ECbValidateError::OutOfBounds);
+	TestEqual(TEXT("ValidateCompactBinary(OutOfBounds, Object, OOB Field)"), Validate({uint8(ECbFieldType::Object), 3, AddName(ECbFieldType::Float32), 1, 'N'}), ECbValidateError::OutOfBounds);
+	TestEqual(TEXT("ValidateCompactBinary(OutOfBounds, Object, OOB Field, NoType)"), Validate({3, AddName(ECbFieldType::Float32), 1, 'N'}, ECbFieldType::Object), ECbValidateError::OutOfBounds);
+
+	TestEqual(TEXT("ValidateCompactBinary(Valid, UniformObject, Field)"), Validate({uint8(ECbFieldType::UniformObject), 3, NullWithName, 1, 'N'}), ECbValidateError::None);
+	TestEqual(TEXT("ValidateCompactBinary(Valid, UniformObject, Field, NoType)"), Validate({3, NullWithName, 1, 'N'}, ECbFieldType::UniformObject), ECbValidateError::None);
+	TestEqual(TEXT("ValidateCompactBinary(OutOfBounds, UniformObject)"), Validate({uint8(ECbFieldType::UniformObject)}), ECbValidateError::OutOfBounds);
+	TestEqual(TEXT("ValidateCompactBinary(OutOfBounds, UniformObject, NoType)"), Validate({}, ECbFieldType::UniformObject), ECbValidateError::OutOfBounds);
+	TestEqual(TEXT("ValidateCompactBinary(OutOfBounds, UniformObject)"), Validate({uint8(ECbFieldType::UniformObject), 1}), ECbValidateError::OutOfBounds);
+	TestEqual(TEXT("ValidateCompactBinary(OutOfBounds, UniformObject, NoType)"), Validate({1}, ECbFieldType::UniformObject), ECbValidateError::OutOfBounds);
+	TestEqual(TEXT("ValidateCompactBinary(OutOfBounds, UniformObject, OOB Field)"), Validate({uint8(ECbFieldType::UniformObject), 3, AddName(ECbFieldType::Float32), 1, 'N'}), ECbValidateError::OutOfBounds);
+	TestEqual(TEXT("ValidateCompactBinary(OutOfBounds, UniformObject, OOB Field, NoType)"), Validate({3, AddName(ECbFieldType::Float32), 1, 'N'}, ECbFieldType::UniformObject), ECbValidateError::OutOfBounds);
+
+	TestEqual(TEXT("ValidateCompactBinary(Valid, Array, Empty)"), Validate({uint8(ECbFieldType::Array), 1, 0}), ECbValidateError::None);
+	TestEqual(TEXT("ValidateCompactBinary(Valid, Array, Empty, NoType)"), Validate({1, 0}, ECbFieldType::Array), ECbValidateError::None);
+	TestEqual(TEXT("ValidateCompactBinary(Valid, Array, Field)"), Validate({uint8(ECbFieldType::Array), 4, 2, NullNoName, IntNoName, 0}), ECbValidateError::None);
+	TestEqual(TEXT("ValidateCompactBinary(Valid, Array, Field, NoType)"), Validate({4, 2, NullNoName, IntNoName, 0}, ECbFieldType::Array), ECbValidateError::None);
+	TestEqual(TEXT("ValidateCompactBinary(OutOfBounds, Array)"), Validate({uint8(ECbFieldType::Array)}), ECbValidateError::OutOfBounds);
+	TestEqual(TEXT("ValidateCompactBinary(OutOfBounds, Array, NoType)"), Validate({}, ECbFieldType::Array), ECbValidateError::OutOfBounds);
+	TestEqual(TEXT("ValidateCompactBinary(OutOfBounds, Array)"), Validate({uint8(ECbFieldType::Array), 1}), ECbValidateError::OutOfBounds);
+	TestEqual(TEXT("ValidateCompactBinary(OutOfBounds, Array, NoType)"), Validate({1}, ECbFieldType::Array), ECbValidateError::OutOfBounds);
+	TestEqual(TEXT("ValidateCompactBinary(OutOfBounds, Array, OOB Field)"), Validate({uint8(ECbFieldType::Array), 2, 1, uint8(ECbFieldType::Float32)}), ECbValidateError::OutOfBounds);
+	TestEqual(TEXT("ValidateCompactBinary(OutOfBounds, Array, OOB Field, NoType)"), Validate({2, 1, uint8(ECbFieldType::Float32)}, ECbFieldType::Array), ECbValidateError::OutOfBounds);
+
+	TestEqual(TEXT("ValidateCompactBinary(Valid, UniformArray, Field)"), Validate({uint8(ECbFieldType::UniformArray), 3, 1, IntNoName, 0}), ECbValidateError::None);
+	TestEqual(TEXT("ValidateCompactBinary(Valid, UniformArray, Field, NoType)"), Validate({3, 1, IntNoName, 0}, ECbFieldType::UniformArray), ECbValidateError::None);
+	TestEqual(TEXT("ValidateCompactBinary(OutOfBounds, UniformArray)"), Validate({uint8(ECbFieldType::UniformArray)}), ECbValidateError::OutOfBounds);
+	TestEqual(TEXT("ValidateCompactBinary(OutOfBounds, UniformArray, NoType)"), Validate({}, ECbFieldType::UniformArray), ECbValidateError::OutOfBounds);
+	TestEqual(TEXT("ValidateCompactBinary(OutOfBounds, UniformArray)"), Validate({uint8(ECbFieldType::UniformArray), 1}), ECbValidateError::OutOfBounds);
+	TestEqual(TEXT("ValidateCompactBinary(OutOfBounds, UniformArray, NoType)"), Validate({1}, ECbFieldType::UniformArray), ECbValidateError::OutOfBounds);
+	TestEqual(TEXT("ValidateCompactBinary(OutOfBounds, UniformArray, OOB Field)"), Validate({uint8(ECbFieldType::UniformArray), 2, 1, uint8(ECbFieldType::Float32)}), ECbValidateError::OutOfBounds);
+	TestEqual(TEXT("ValidateCompactBinary(OutOfBounds, UniformArray, OOB Field, NoType)"), Validate({2, 1, uint8(ECbFieldType::Float32)}, ECbFieldType::UniformArray), ECbValidateError::OutOfBounds);
+
+	TestEqual(TEXT("ValidateCompactBinary(Valid, Binary, Empty)"), Validate({uint8(ECbFieldType::Binary), 0}), ECbValidateError::None);
+	TestEqual(TEXT("ValidateCompactBinary(Valid, Binary, Empty, NoType)"), Validate({0}, ECbFieldType::Binary), ECbValidateError::None);
+	TestEqual(TEXT("ValidateCompactBinary(Valid, Binary, Field)"), Validate({uint8(ECbFieldType::Binary), 1, 0}), ECbValidateError::None);
+	TestEqual(TEXT("ValidateCompactBinary(Valid, Binary, Field, NoType)"), Validate({1, 0}, ECbFieldType::Binary), ECbValidateError::None);
+	TestEqual(TEXT("ValidateCompactBinary(OutOfBounds, Binary)"), Validate({uint8(ECbFieldType::Binary)}), ECbValidateError::OutOfBounds);
+	TestEqual(TEXT("ValidateCompactBinary(OutOfBounds, Binary, NoType)"), Validate({}, ECbFieldType::Binary), ECbValidateError::OutOfBounds);
+	TestEqual(TEXT("ValidateCompactBinary(OutOfBounds, Binary)"), Validate({uint8(ECbFieldType::Binary), 1}), ECbValidateError::OutOfBounds);
+	TestEqual(TEXT("ValidateCompactBinary(OutOfBounds, Binary, NoType)"), Validate({1}, ECbFieldType::Binary), ECbValidateError::OutOfBounds);
+
+	TestEqual(TEXT("ValidateCompactBinary(Valid, String, Empty)"), Validate({uint8(ECbFieldType::String), 0}), ECbValidateError::None);
+	TestEqual(TEXT("ValidateCompactBinary(Valid, String, Empty, NoType)"), Validate({0}, ECbFieldType::String), ECbValidateError::None);
+	TestEqual(TEXT("ValidateCompactBinary(Valid, String, Field)"), Validate({uint8(ECbFieldType::String), 1, 'S'}), ECbValidateError::None);
+	TestEqual(TEXT("ValidateCompactBinary(Valid, String, Field, NoType)"), Validate({1, 'S'}, ECbFieldType::String), ECbValidateError::None);
+	TestEqual(TEXT("ValidateCompactBinary(OutOfBounds, String)"), Validate({uint8(ECbFieldType::String)}), ECbValidateError::OutOfBounds);
+	TestEqual(TEXT("ValidateCompactBinary(OutOfBounds, String, NoType)"), Validate({}, ECbFieldType::String), ECbValidateError::OutOfBounds);
+	TestEqual(TEXT("ValidateCompactBinary(OutOfBounds, String)"), Validate({uint8(ECbFieldType::String), 1}), ECbValidateError::OutOfBounds);
+	TestEqual(TEXT("ValidateCompactBinary(OutOfBounds, String, NoType)"), Validate({1}, ECbFieldType::String), ECbValidateError::OutOfBounds);
+
+	TestEqual(TEXT("ValidateCompactBinary(Valid, IntegerPositive, 1-byte)"), Validate({uint8(ECbFieldType::IntegerPositive), 0}), ECbValidateError::None);
+	TestEqual(TEXT("ValidateCompactBinary(Valid, IntegerPositive, 1-byte, NoType)"), Validate({0}, ECbFieldType::IntegerPositive), ECbValidateError::None);
+	TestEqual(TEXT("ValidateCompactBinary(Valid, IntegerPositive, 2-byte)"), Validate({uint8(ECbFieldType::IntegerPositive), 0x80, 0x80}), ECbValidateError::None);
+	TestEqual(TEXT("ValidateCompactBinary(Valid, IntegerPositive, 2-byte, NoType)"), Validate({0x80, 0x80}, ECbFieldType::IntegerPositive), ECbValidateError::None);
+	TestEqual(TEXT("ValidateCompactBinary(OutOfBounds, IntegerPositive, 1-byte)"), Validate({uint8(ECbFieldType::IntegerPositive)}), ECbValidateError::OutOfBounds);
+	TestEqual(TEXT("ValidateCompactBinary(OutOfBounds, IntegerPositive, 1-byte, NoType)"), Validate({}, ECbFieldType::IntegerPositive), ECbValidateError::OutOfBounds);
+	TestEqual(TEXT("ValidateCompactBinary(OutOfBounds, IntegerPositive, 2-byte)"), Validate({uint8(ECbFieldType::IntegerPositive), 0x80}), ECbValidateError::OutOfBounds);
+	TestEqual(TEXT("ValidateCompactBinary(OutOfBounds, IntegerPositive, 2-byte, NoType)"), Validate({0x80}, ECbFieldType::IntegerPositive), ECbValidateError::OutOfBounds);
+	TestEqual(TEXT("ValidateCompactBinary(OutOfBounds, IntegerPositive, 9-byte)"), Validate({uint8(ECbFieldType::IntegerPositive), 0xff, 0, 0, 0, 0, 0, 0, 0}), ECbValidateError::OutOfBounds);
+	TestEqual(TEXT("ValidateCompactBinary(OutOfBounds, IntegerPositive, 9-byte, NoType)"), Validate({0xff, 0, 0, 0, 0, 0, 0, 0}, ECbFieldType::IntegerPositive), ECbValidateError::OutOfBounds);
+
+	TestEqual(TEXT("ValidateCompactBinary(Valid, IntegerNegative, 1-byte)"), Validate({uint8(ECbFieldType::IntegerNegative), 0}), ECbValidateError::None);
+	TestEqual(TEXT("ValidateCompactBinary(Valid, IntegerNegative, 1-byte, NoType)"), Validate({0}, ECbFieldType::IntegerNegative), ECbValidateError::None);
+	TestEqual(TEXT("ValidateCompactBinary(Valid, IntegerNegative, 2-byte)"), Validate({uint8(ECbFieldType::IntegerNegative), 0x80, 0x80}), ECbValidateError::None);
+	TestEqual(TEXT("ValidateCompactBinary(Valid, IntegerNegative, 2-byte, NoType)"), Validate({0x80, 0x80}, ECbFieldType::IntegerNegative), ECbValidateError::None);
+	TestEqual(TEXT("ValidateCompactBinary(OutOfBounds, IntegerNegative, 1-byte)"), Validate({uint8(ECbFieldType::IntegerNegative)}), ECbValidateError::OutOfBounds);
+	TestEqual(TEXT("ValidateCompactBinary(OutOfBounds, IntegerNegative, 1-byte, NoType)"), Validate({}, ECbFieldType::IntegerNegative), ECbValidateError::OutOfBounds);
+	TestEqual(TEXT("ValidateCompactBinary(OutOfBounds, IntegerNegative, 2-byte)"), Validate({uint8(ECbFieldType::IntegerNegative), 0x80}), ECbValidateError::OutOfBounds);
+	TestEqual(TEXT("ValidateCompactBinary(OutOfBounds, IntegerNegative, 2-byte, NoType)"), Validate({0x80}, ECbFieldType::IntegerNegative), ECbValidateError::OutOfBounds);
+	TestEqual(TEXT("ValidateCompactBinary(OutOfBounds, IntegerNegative, 9-byte)"), Validate({uint8(ECbFieldType::IntegerNegative), 0xff, 0, 0, 0, 0, 0, 0, 0}), ECbValidateError::OutOfBounds);
+	TestEqual(TEXT("ValidateCompactBinary(OutOfBounds, IntegerNegative, 9-byte, NoType)"), Validate({0xff, 0, 0, 0, 0, 0, 0, 0}, ECbFieldType::IntegerNegative), ECbValidateError::OutOfBounds);
+
+	TestEqual(TEXT("ValidateCompactBinary(Valid, Float32)"), Validate({uint8(ECbFieldType::Float32), 0, 0, 0, 0}), ECbValidateError::None);
+	TestEqual(TEXT("ValidateCompactBinary(Valid, Float32, NoType)"), Validate({0, 0, 0, 0}, ECbFieldType::Float32), ECbValidateError::None);
+	TestEqual(TEXT("ValidateCompactBinary(OutOfBounds, Float32)"), Validate({uint8(ECbFieldType::Float32), 0, 0, 0}), ECbValidateError::OutOfBounds);
+	TestEqual(TEXT("ValidateCompactBinary(OutOfBounds, Float32, NoType)"), Validate({0, 0, 0}, ECbFieldType::Float32), ECbValidateError::OutOfBounds);
+
+	TestEqual(TEXT("ValidateCompactBinary(Valid, Float64)"), Validate({uint8(ECbFieldType::Float64), 0x3f, 0xff, 0xff, 0xff, 0xf0, 0x00, 0x00, 0x00}), ECbValidateError::None);
+	TestEqual(TEXT("ValidateCompactBinary(Valid, Float64, NoType)"), Validate({0x3f, 0xff, 0xff, 0xff, 0xf0, 0x00, 0x00, 0x00}, ECbFieldType::Float64), ECbValidateError::None);
+	TestEqual(TEXT("ValidateCompactBinary(OutOfBounds, Float64)"), Validate({uint8(ECbFieldType::Float64), 0, 0, 0, 0, 0, 0, 0}), ECbValidateError::OutOfBounds);
+	TestEqual(TEXT("ValidateCompactBinary(OutOfBounds, Float64, NoType)"), Validate({0, 0, 0, 0, 0, 0, 0}, ECbFieldType::Float64), ECbValidateError::OutOfBounds);
+
+	TestEqual(TEXT("ValidateCompactBinary(Valid, BoolFalse)"), Validate({uint8(ECbFieldType::BoolFalse)}), ECbValidateError::None);
+	TestEqual(TEXT("ValidateCompactBinary(Valid, BoolTrue)"), Validate({uint8(ECbFieldType::BoolTrue)}), ECbValidateError::None);
+
+	TestEqual(TEXT("ValidateCompactBinary(Valid, Reference)"), Validate({uint8(ECbFieldType::Reference), 0, 0, 0, 0, 0, 0, 0, 0, 0, 0, 0, 0, 0, 0, 0, 0, 0, 0, 0, 0, 0, 0, 0, 0, 0, 0, 0, 0, 0, 0, 0, 0}), ECbValidateError::None);
+	TestEqual(TEXT("ValidateCompactBinary(Valid, Reference, NoType)"), Validate({0, 0, 0, 0, 0, 0, 0, 0, 0, 0, 0, 0, 0, 0, 0, 0, 0, 0, 0, 0, 0, 0, 0, 0, 0, 0, 0, 0, 0, 0, 0, 0}, ECbFieldType::Reference), ECbValidateError::None);
+	TestEqual(TEXT("ValidateCompactBinary(OutOfBounds, Reference)"), Validate({uint8(ECbFieldType::Reference), 0, 0, 0, 0, 0, 0, 0, 0, 0, 0, 0, 0, 0, 0, 0, 0, 0, 0, 0, 0, 0, 0, 0, 0, 0, 0, 0, 0, 0, 0, 0}), ECbValidateError::OutOfBounds);
+	TestEqual(TEXT("ValidateCompactBinary(OutOfBounds, Reference, NoType)"), Validate({0, 0, 0, 0, 0, 0, 0, 0, 0, 0, 0, 0, 0, 0, 0, 0, 0, 0, 0, 0, 0, 0, 0, 0, 0, 0, 0, 0, 0, 0, 0}, ECbFieldType::Reference), ECbValidateError::OutOfBounds);
+
+	TestEqual(TEXT("ValidateCompactBinary(Valid, BinaryReference)"), Validate({uint8(ECbFieldType::BinaryReference), 0, 0, 0, 0, 0, 0, 0, 0, 0, 0, 0, 0, 0, 0, 0, 0, 0, 0, 0, 0, 0, 0, 0, 0, 0, 0, 0, 0, 0, 0, 0, 0}), ECbValidateError::None);
+	TestEqual(TEXT("ValidateCompactBinary(Valid, BinaryReference, NoType)"), Validate({0, 0, 0, 0, 0, 0, 0, 0, 0, 0, 0, 0, 0, 0, 0, 0, 0, 0, 0, 0, 0, 0, 0, 0, 0, 0, 0, 0, 0, 0, 0, 0}, ECbFieldType::BinaryReference), ECbValidateError::None);
+	TestEqual(TEXT("ValidateCompactBinary(OutOfBounds, BinaryReference)"), Validate({uint8(ECbFieldType::BinaryReference), 0, 0, 0, 0, 0, 0, 0, 0, 0, 0, 0, 0, 0, 0, 0, 0, 0, 0, 0, 0, 0, 0, 0, 0, 0, 0, 0, 0, 0, 0, 0}), ECbValidateError::OutOfBounds);
+	TestEqual(TEXT("ValidateCompactBinary(OutOfBounds, BinaryReference, NoType)"), Validate({0, 0, 0, 0, 0, 0, 0, 0, 0, 0, 0, 0, 0, 0, 0, 0, 0, 0, 0, 0, 0, 0, 0, 0, 0, 0, 0, 0, 0, 0, 0}, ECbFieldType::BinaryReference), ECbValidateError::OutOfBounds);
+
+	TestEqual(TEXT("ValidateCompactBinary(Valid, Hash)"), Validate({uint8(ECbFieldType::Hash), 0, 0, 0, 0, 0, 0, 0, 0, 0, 0, 0, 0, 0, 0, 0, 0, 0, 0, 0, 0, 0, 0, 0, 0, 0, 0, 0, 0, 0, 0, 0, 0}), ECbValidateError::None);
+	TestEqual(TEXT("ValidateCompactBinary(Valid, Hash, NoType)"), Validate({0, 0, 0, 0, 0, 0, 0, 0, 0, 0, 0, 0, 0, 0, 0, 0, 0, 0, 0, 0, 0, 0, 0, 0, 0, 0, 0, 0, 0, 0, 0, 0}, ECbFieldType::Hash), ECbValidateError::None);
+	TestEqual(TEXT("ValidateCompactBinary(OutOfBounds, Hash)"), Validate({uint8(ECbFieldType::Hash), 0, 0, 0, 0, 0, 0, 0, 0, 0, 0, 0, 0, 0, 0, 0, 0, 0, 0, 0, 0, 0, 0, 0, 0, 0, 0, 0, 0, 0, 0, 0}), ECbValidateError::OutOfBounds);
+	TestEqual(TEXT("ValidateCompactBinary(OutOfBounds, Hash, NoType)"), Validate({0, 0, 0, 0, 0, 0, 0, 0, 0, 0, 0, 0, 0, 0, 0, 0, 0, 0, 0, 0, 0, 0, 0, 0, 0, 0, 0, 0, 0, 0, 0}, ECbFieldType::Hash), ECbValidateError::OutOfBounds);
+
+	TestEqual(TEXT("ValidateCompactBinary(Valid, Uuid)"), Validate({uint8(ECbFieldType::Uuid), 0, 0, 0, 0, 0, 0, 0, 0, 0, 0, 0, 0, 0, 0, 0, 0}), ECbValidateError::None);
+	TestEqual(TEXT("ValidateCompactBinary(Valid, Uuid, NoType)"), Validate({0, 0, 0, 0, 0, 0, 0, 0, 0, 0, 0, 0, 0, 0, 0, 0}, ECbFieldType::Uuid), ECbValidateError::None);
+	TestEqual(TEXT("ValidateCompactBinary(OutOfBounds, Uuid)"), Validate({uint8(ECbFieldType::Uuid), 0, 0, 0, 0, 0, 0, 0, 0, 0, 0, 0, 0, 0, 0, 0}), ECbValidateError::OutOfBounds);
+	TestEqual(TEXT("ValidateCompactBinary(OutOfBounds, Uuid, NoType)"), Validate({0, 0, 0, 0, 0, 0, 0, 0, 0, 0, 0, 0, 0, 0, 0}, ECbFieldType::Uuid), ECbValidateError::OutOfBounds);
+
+	TestEqual(TEXT("ValidateCompactBinary(Valid, DateTime)"), Validate({uint8(ECbFieldType::DateTime), 0, 0, 0, 0, 0, 0, 0, 0}), ECbValidateError::None);
+	TestEqual(TEXT("ValidateCompactBinary(Valid, DateTime, NoType)"), Validate({0, 0, 0, 0, 0, 0, 0, 0}, ECbFieldType::DateTime), ECbValidateError::None);
+	TestEqual(TEXT("ValidateCompactBinary(OutOfBounds, DateTime)"), Validate({uint8(ECbFieldType::DateTime), 0, 0, 0, 0, 0, 0, 0}), ECbValidateError::OutOfBounds);
+	TestEqual(TEXT("ValidateCompactBinary(OutOfBounds, DateTime, NoType)"), Validate({0, 0, 0, 0, 0, 0, 0}, ECbFieldType::DateTime), ECbValidateError::OutOfBounds);
+
+	TestEqual(TEXT("ValidateCompactBinary(Valid, TimeSpan)"), Validate({uint8(ECbFieldType::TimeSpan), 0, 0, 0, 0, 0, 0, 0, 0}), ECbValidateError::None);
+	TestEqual(TEXT("ValidateCompactBinary(Valid, TimeSpan, NoType)"), Validate({0, 0, 0, 0, 0, 0, 0, 0}, ECbFieldType::TimeSpan), ECbValidateError::None);
+	TestEqual(TEXT("ValidateCompactBinary(OutOfBounds, TimeSpan)"), Validate({uint8(ECbFieldType::TimeSpan), 0, 0, 0, 0, 0, 0, 0}), ECbValidateError::OutOfBounds);
+	TestEqual(TEXT("ValidateCompactBinary(OutOfBounds, TimeSpan, NoType)"), Validate({0, 0, 0, 0, 0, 0, 0}, ECbFieldType::TimeSpan), ECbValidateError::OutOfBounds);
+
+	// Test InvalidType
+	TestEqual(TEXT("ValidateCompactBinary(InvalidType, Unknown)"), Validate({uint8(ECbFieldType::TimeSpan) + 1}), ECbValidateError::InvalidType);
+	TestEqual(TEXT("ValidateCompactBinary(InvalidType, Unknown)"), Validate({}, ECbFieldType(uint8(ECbFieldType::TimeSpan) + 1)), ECbValidateError::InvalidType);
+	TestEqual(TEXT("ValidateCompactBinary(InvalidType, HasFieldType)"), Validate({uint8(ECbFieldType::Null | ECbFieldType::HasFieldType)}), ECbValidateError::InvalidType);
+
+	TestEqual(TEXT("ValidateCompactBinary(InvalidType, ZeroSizeField)"), Validate({}, ECbFieldType::Null), ECbValidateError::InvalidType);
+	TestEqual(TEXT("ValidateCompactBinary(InvalidType, ZeroSizeField, BoolFalse)"), Validate({}, ECbFieldType::BoolFalse), ECbValidateError::InvalidType);
+	TestEqual(TEXT("ValidateCompactBinary(InvalidType, ZeroSizeField, BoolTrue)"), Validate({}, ECbFieldType::BoolTrue), ECbValidateError::InvalidType);
+
+	TestEqual(TEXT("ValidateCompactBinary(InvalidType, ZeroSizeField, Array)"), Validate({uint8(ECbFieldType::UniformArray), 2, 2, NullNoName}), ECbValidateError::InvalidType);
+	TestEqual(TEXT("ValidateCompactBinary(InvalidType, ZeroSizeField, Object)"), Validate({uint8(ECbFieldType::UniformObject), 2, NullNoName, 0}), ECbValidateError::InvalidType);
+
+	// Test DuplicateName
+	TestEqual(TEXT("ValidateCompactBinary(DuplicateName)"), Validate({uint8(ECbFieldType::UniformObject), 7, NullWithName, 1, 'A', 1, 'B', 1, 'A'}), ECbValidateError::DuplicateName);
+	TestEqual(TEXT("ValidateCompactBinary(DuplicateName, CaseSensitive)"), Validate({uint8(ECbFieldType::UniformObject), 7, NullWithName, 1, 'A', 1, 'B', 1, 'a'}), ECbValidateError::None);
+	TestEqual(TEXT("ValidateCompactBinary(DuplicateName, Mode)"), ValidateMode({uint8(ECbFieldType::UniformObject), 7, NullWithName, 1, 'A', 1, 'B', 1, 'A'}, ECbValidateMode::All & ~ECbValidateMode::Names), ECbValidateError::None);
+
+	// Test MissingName
+	TestEqual(TEXT("ValidateCompactBinary(MissingName)"), Validate({uint8(ECbFieldType::Object), 3, NullNoName, IntNoName, 0}), ECbValidateError::MissingName);
+	TestEqual(TEXT("ValidateCompactBinary(MissingName, Uniform)"), Validate({uint8(ECbFieldType::UniformObject), 3, IntNoName, 0, 0}), ECbValidateError::MissingName);
+	TestEqual(TEXT("ValidateCompactBinary(MissingName, Mode)"), ValidateMode({uint8(ECbFieldType::Object), 3, NullNoName, IntNoName, 0}, ECbValidateMode::All & ~ECbValidateMode::Names), ECbValidateError::None);
+	TestEqual(TEXT("ValidateCompactBinary(MissingName, Uniform, Mode)"), ValidateMode({uint8(ECbFieldType::UniformObject), 3, IntNoName, 0, 0}, ECbValidateMode::All & ~ECbValidateMode::Names), ECbValidateError::None);
+
+	// Test ArrayName
+	TestEqual(TEXT("ValidateCompactBinary(ArrayName)"), Validate({uint8(ECbFieldType::Array), 5, 2, NullNoName, NullWithName, 1, 'F'}), ECbValidateError::ArrayName);
+	TestEqual(TEXT("ValidateCompactBinary(ArrayName, Uniform)"), Validate({uint8(ECbFieldType::UniformArray), 4, 1, NullWithName, 1, 'F'}), ECbValidateError::ArrayName);
+	TestEqual(TEXT("ValidateCompactBinary(ArrayName, Mode)"), ValidateMode({uint8(ECbFieldType::Array), 5, 2, NullNoName, NullWithName, 1, 'F'}, ECbValidateMode::All & ~ECbValidateMode::Names), ECbValidateError::None);
+	TestEqual(TEXT("ValidateCompactBinary(ArrayName, Uniform, Mode)"), ValidateMode({uint8(ECbFieldType::UniformArray), 4, 1, NullWithName, 1, 'F'}, ECbValidateMode::All & ~ECbValidateMode::Names), ECbValidateError::None);
+
+	// Test InvalidString
+	// Not tested or implemented yet because the engine does not provide enough UTF-8 functionality.
+
+	// Test InvalidInteger
+	TestEqual(TEXT("ValidateCompactBinary(InvalidInteger, NameSize)"), Validate({NullWithName, 0x80, 1, 'N'}), ECbValidateError::InvalidInteger);
+	TestEqual(TEXT("ValidateCompactBinary(InvalidInteger, ObjectSize)"), Validate({uint8(ECbFieldType::Object), 0xc0, 0, 0}), ECbValidateError::InvalidInteger);
+	TestEqual(TEXT("ValidateCompactBinary(InvalidInteger, ArraySize)"), Validate({uint8(ECbFieldType::Array), 0xe0, 0, 0, 1, 0}), ECbValidateError::InvalidInteger);
+	TestEqual(TEXT("ValidateCompactBinary(InvalidInteger, ArrayCount)"), Validate({uint8(ECbFieldType::Array), 5, 0xf0, 0, 0, 0, 0}), ECbValidateError::InvalidInteger);
+	TestEqual(TEXT("ValidateCompactBinary(InvalidInteger, BinarySize)"), Validate({uint8(ECbFieldType::Binary), 0xf8, 0, 0, 0, 0, 0}), ECbValidateError::InvalidInteger);
+	TestEqual(TEXT("ValidateCompactBinary(InvalidInteger, StringSize)"), Validate({uint8(ECbFieldType::String), 0xfc, 0, 0, 0, 0, 0, 0}), ECbValidateError::InvalidInteger);
+	TestEqual(TEXT("ValidateCompactBinary(InvalidInteger, IntegerPositive)"), Validate({uint8(ECbFieldType::IntegerPositive), 0xfe, 0, 0, 0, 0, 0, 0, 0}), ECbValidateError::InvalidInteger);
+	TestEqual(TEXT("ValidateCompactBinary(InvalidInteger, IntegerNegative)"), Validate({uint8(ECbFieldType::IntegerNegative), 0xff, 0, 0, 0, 0, 0, 0, 0, 0}), ECbValidateError::InvalidInteger);
+	TestEqual(TEXT("ValidateCompactBinary(InvalidInteger, ArraySize)"), Validate({uint8(ECbFieldType::Array), 0x80, 1, 0}), ECbValidateError::InvalidInteger);
+	TestEqual(TEXT("ValidateCompactBinary(InvalidInteger, ArrayCount)"), Validate({uint8(ECbFieldType::Array), 3, 0xc0, 0, 0}), ECbValidateError::InvalidInteger);
+	TestEqual(TEXT("ValidateCompactBinary(InvalidInteger, ObjectSize)"), Validate({uint8(ECbFieldType::Object), 0xe0, 0, 0, 0}), ECbValidateError::InvalidInteger);
+
+	TestEqual(TEXT("ValidateCompactBinary(InvalidInteger, NameSize, Mode)"), ValidateMode({NullWithName, 0x80, 1, 'N'}, ECbValidateMode::All & ~ECbValidateMode::Format), ECbValidateError::None);
+	TestEqual(TEXT("ValidateCompactBinary(InvalidInteger, ArraySize, Mode)"), ValidateMode({uint8(ECbFieldType::Array), 0xc0, 0, 1, 0}, ECbValidateMode::All & ~ECbValidateMode::Format), ECbValidateError::None);
+	TestEqual(TEXT("ValidateCompactBinary(InvalidInteger, ObjectSize, Mode)"), ValidateMode({uint8(ECbFieldType::Object), 0xe0, 0, 0, 0}, ECbValidateMode::All & ~ECbValidateMode::Format), ECbValidateError::None);
+
+	// Test InvalidFloat
+	TestEqual(TEXT("ValidateCompactBinary(InvalidFloat, MaxSignificant+1)"), Validate({uint8(ECbFieldType::Float64), 0x3f, 0xff, 0xff, 0xff, 0xf0, 0x00, 0x00, 0x00}), ECbValidateError::None); // 1.9999999403953552
+	TestEqual(TEXT("ValidateCompactBinary(InvalidFloat, MaxExponent+1)"), Validate({uint8(ECbFieldType::Float64), 0x47, 0xff, 0xff, 0xff, 0xe0, 0x00, 0x00, 0x00}), ECbValidateError::None); // 6.8056469327705771e38
+	TestEqual(TEXT("ValidateCompactBinary(InvalidFloat, MaxSignificand)"), Validate({uint8(ECbFieldType::Float64), 0x3f, 0xff, 0xff, 0xff, 0xe0, 0x00, 0x00, 0x00}), ECbValidateError::InvalidFloat); // 1.9999998807907104
+	TestEqual(TEXT("ValidateCompactBinary(InvalidFloat, MaxExponent)"), Validate({uint8(ECbFieldType::Float64), 0x47, 0xef, 0xff, 0xff, 0xe0, 0x00, 0x00, 0x00}), ECbValidateError::InvalidFloat); // 3.4028234663852886e38
+	TestEqual(TEXT("ValidateCompactBinary(InvalidFloat, MaxSignificand, Mode)"), ValidateMode({uint8(ECbFieldType::Float64), 0x3f, 0xff, 0xff, 0xff, 0xe0, 0x00, 0x00, 0x00}, ECbValidateMode::All & ~ECbValidateMode::Format), ECbValidateError::None); // 1.9999998807907104
+	TestEqual(TEXT("ValidateCompactBinary(InvalidFloat, MaxExponent, Mode)"), ValidateMode({uint8(ECbFieldType::Float64), 0x47, 0xef, 0xff, 0xff, 0xe0, 0x00, 0x00, 0x00}, ECbValidateMode::All & ~ECbValidateMode::Format), ECbValidateError::None); // 3.4028234663852886e38
+
+	// Test NonUniformObject
+	TestEqual(TEXT("ValidateCompactBinary(NonUniformObject)"), Validate({uint8(ECbFieldType::Object), 3, NullWithName, 1, 'A'}), ECbValidateError::NonUniformObject);
+	TestEqual(TEXT("ValidateCompactBinary(NonUniformObject)"), Validate({uint8(ECbFieldType::Object), 6, NullWithName, 1, 'A', NullWithName, 1, 'B'}), ECbValidateError::NonUniformObject);
+	TestEqual(TEXT("ValidateCompactBinary(NonUniformObject, Mode)"), ValidateMode({uint8(ECbFieldType::Object), 3, NullWithName, 1, 'A'}, ECbValidateMode::All & ~ECbValidateMode::Format), ECbValidateError::None);
+	TestEqual(TEXT("ValidateCompactBinary(NonUniformObject, Mode)"), ValidateMode({uint8(ECbFieldType::Object), 6, NullWithName, 1, 'A', NullWithName, 1, 'B'}, ECbValidateMode::All & ~ECbValidateMode::Format), ECbValidateError::None);
+
+	// Test NonUniformArray
+	TestEqual(TEXT("ValidateCompactBinary(NonUniformArray)"), Validate({uint8(ECbFieldType::Array), 3, 1, IntNoName, 0}), ECbValidateError::NonUniformArray);
+	TestEqual(TEXT("ValidateCompactBinary(NonUniformArray)"), Validate({uint8(ECbFieldType::Array), 5, 2, IntNoName, 1, IntNoName, 2}), ECbValidateError::NonUniformArray);
+	TestEqual(TEXT("ValidateCompactBinary(NonUniformArray, Null)"), Validate({uint8(ECbFieldType::Array), 3, 2, NullNoName, NullNoName}), ECbValidateError::None);
+	TestEqual(TEXT("ValidateCompactBinary(NonUniformArray, Bool)"), Validate({uint8(ECbFieldType::Array), 3, 2, uint8(ECbFieldType::BoolFalse), uint8(ECbFieldType::BoolFalse)}), ECbValidateError::None);
+	TestEqual(TEXT("ValidateCompactBinary(NonUniformArray, Bool)"), Validate({uint8(ECbFieldType::Array), 3, 2, uint8(ECbFieldType::BoolTrue), uint8(ECbFieldType::BoolTrue)}), ECbValidateError::None);
+	TestEqual(TEXT("ValidateCompactBinary(NonUniformArray, Mode)"), ValidateMode({uint8(ECbFieldType::Array), 3, 1, IntNoName, 0}, ECbValidateMode::All & ~ECbValidateMode::Format), ECbValidateError::None);
+	TestEqual(TEXT("ValidateCompactBinary(NonUniformArray, Mode)"), ValidateMode({uint8(ECbFieldType::Array), 5, 2, IntNoName, 1, IntNoName, 2}, ECbValidateMode::All & ~ECbValidateMode::Format), ECbValidateError::None);
+
+	// Test Padding
+	TestEqual(TEXT("ValidateCompactBinary(Padding)"), Validate({NullNoName, 0}), ECbValidateError::Padding);
+	TestEqual(TEXT("ValidateCompactBinary(Padding)"), Validate({uint8(ECbFieldType::Array), 1, 0, 0}), ECbValidateError::Padding);
+	TestEqual(TEXT("ValidateCompactBinary(Padding)"), Validate({uint8(ECbFieldType::Object), 0, 0}), ECbValidateError::Padding);
+	TestEqual(TEXT("ValidateCompactBinary(Padding, Mode)"), ValidateMode({NullNoName, 0}, ECbValidateMode::All & ~ECbValidateMode::Padding), ECbValidateError::None);
+	TestEqual(TEXT("ValidateCompactBinary(Padding, Mode)"), ValidateMode({uint8(ECbFieldType::Array), 1, 0, 0}, ECbValidateMode::All & ~ECbValidateMode::Padding), ECbValidateError::None);
+	TestEqual(TEXT("ValidateCompactBinary(Padding, Mode)"), ValidateMode({uint8(ECbFieldType::Object), 0, 0}, ECbValidateMode::All & ~ECbValidateMode::Padding), ECbValidateError::None);
+
+	return true;
+}
+
+///////////////////////////////////////////////////////////////////////////////////////////////////////////////////////
+
+IMPLEMENT_SIMPLE_AUTOMATION_TEST(FCbValidateRangeTest, "System.Core.Serialization.ValidateCompactBinaryRange", CompactBinaryValidationTestFlags)
+bool FCbValidateRangeTest::RunTest(const FString& Parameters)
+{
+	auto Validate = [this](std::initializer_list<uint8> Data) -> ECbValidateError
+	{
+		return ValidateCompactBinaryRange(MakeMemoryView(Data), ECbValidateMode::All);
+	};
+
+	// Test Empty
+	TestEqual(TEXT("ValidateCompactBinaryRange(Empty)"), Validate({}), ECbValidateError::None);
+
+	// Test Valid
+	TestEqual(TEXT("ValidateCompactBinaryRange(Null x2)"), Validate({uint8(ECbFieldType::Null), uint8(ECbFieldType::Null)}), ECbValidateError::None);
+
+	// Test Padding
+	TestEqual(TEXT("ValidateCompactBinaryRange(Padding InvalidType)"), Validate({uint8(ECbFieldType::Null), 0}), ECbValidateError::InvalidType);
+	TestEqual(TEXT("ValidateCompactBinaryRange(Padding OutOfBounds)"), Validate({uint8(ECbFieldType::Null), uint8(ECbFieldType::Binary)}), ECbValidateError::OutOfBounds);
+
+	return true;
+}
+
+///////////////////////////////////////////////////////////////////////////////////////////////////////////////////////
+
+#endif // WITH_DEV_AUTOMATION_TESTS
diff --git a/Engine/Source/Runtime/Core/Private/Tests/Serialization/CompactBinaryWriterTest.cpp b/Engine/Source/Runtime/Core/Private/Tests/Serialization/CompactBinaryWriterTest.cpp
--- a/Engine/Source/Runtime/Core/Private/Tests/Serialization/CompactBinaryWriterTest.cpp
+++ b/Engine/Source/Runtime/Core/Private/Tests/Serialization/CompactBinaryWriterTest.cpp
@@ -8,6 +8,7 @@
 #include "Misc/Guid.h"
 #include "Misc/Timespan.h"
 #include "Misc/StringBuilder.h"
+#include "Serialization/CompactBinaryValidation.h"
 
 #if WITH_DEV_AUTOMATION_TESTS
 
diff --git a/Engine/Source/Runtime/Core/Public/Serialization/CompactBinary.h b/Engine/Source/Runtime/Core/Public/Serialization/CompactBinary.h
--- a/Engine/Source/Runtime/Core/Public/Serialization/CompactBinary.h
+++ b/Engine/Source/Runtime/Core/Public/Serialization/CompactBinary.h
@@ -1088,122 +1088,6 @@ inline FCbArrayRef FCbFieldRef::AsArrayRef() &&
 
 ///////////////////////////////////////////////////////////////////////////////////////////////////////////////////////
 
-/** Flags for validating compact binary data. */
-enum class ECbValidateMode : uint32
-{
-	/**
-	 * Validate that the value can be read and stays inside the bounds of the memory view.
-	 *
-	 * This is the minimum level of validation required to be able to safely read a field, array,
-	 * or object without the risk of crashing or reading out of bounds.
-	 */
-	Default                 = 0,
-
-	/**
-	 * Validate that object fields have unique non-empty names and array fields have no names.
-	 *
-	 * Name validation failures typically do not inhibit reading the input, but duplicated fields
-	 * cannot be looked up by name other than the first, and converting to other data formats can
-	 * fail in the presence of naming issues.
-	 */
-	Names                   = 1 << 0,
-
-	/**
-	 * Validate that fields are serialized in the canonical format.
-	 *
-	 * Format validation failures typically do not inhibit reading the input. Values that fail in
-	 * this mode require more memory than in the canonical format, and comparisons of such values
-	 * for equality are not reliable. Examples of failures include uniform arrays or objects that
-	 * were not encoded uniformly, variable-length integers that could be encoded in fewer bytes,
-	 * or 64-bit floats that could be encoded in 32 bits without loss of precision.
-	 */
-	Format                  = 1 << 1,
-
-	/**
-	 * Validate that there is no padding after the value before the end of the memory view.
-	 *
-	 * Padding validation failures have no impact on the ability to read the input, but are using
-	 * more memory than necessary.
-	 */
-	Padding                 = 1 << 2,
-
-	/** Perform all validation described above. */
-	All                     = Default | Names | Format | Padding,
-};
-
-/** Flags for compact binary validation errors. Multiple flags may be combined. */
-enum class ECbValidateError : uint32
-{
-	/** The input had no validation errors. */
-	None                    = 0,
-
-	// Mode: Default
-
-	/** The input cannot be read without reading out of bounds. */
-	OutOfBounds             = 1 << 0,
-	/** The input has a field with an unrecognized or invalid type. */
-	InvalidType             = 1 << 1,
-
-	// Mode: Names
-
-	/** An object had more than one field with the same name. */
-	DuplicateName           = 1 << 2,
-	/** An object had a field with no name. */
-	MissingName             = 1 << 3,
-	/** An array field had a name. */
-	ArrayName               = 1 << 4,
-
-	// Mode: Format
-
-	/** A name or string payload is not valid UTF-8. */
-	InvalidString           = 1 << 5,
-	/** A size or integer payload can be encoded in fewer bytes. */
-	InvalidInteger          = 1 << 6,
-	/** A float64 payload can be encoded as a float32 without loss of precision. */
-	InvalidFloat            = 1 << 7,
-	/** An object has the same type for every field but is not uniform. */
-	NonUniformObject        = 1 << 8,
-	/** An array has the same type for every field and non-empty payloads but is not uniform. */
-	NonUniformArray         = 1 << 9,
-
-	// Mode: Padding
-
-	/** A value did not use the entire memory view given for validation. */
-	Padding                 = 1 << 10,
-};
-
-ENUM_CLASS_FLAGS(ECbValidateMode);
-ENUM_CLASS_FLAGS(ECbValidateError);
-
-/**
- * Validate the compact binary data for one field in the view as specified by the mode flags.
- *
- * Only one top-level field is processed from the view, and validation recurses into any array or
- * object within that field. To validate multiple consecutive top-level fields, call the function
- * once for each top-level field. If the given view might contain multiple top-level fields, then
- * either exclude the Padding flag from the Mode or use MeasureCompactBinary to break up the view
- * into its constituent fields before validating.
- *
- * @param View A memory view containing at least one top-level field.
- * @param Mode A combination of the flags for the types of validation to perform.
- * @param Type HasFieldType means that View contains the type. Otherwise, use the given type.
- * @return None on success, otherwise the flags for the types of errors that were detected.
- */
-CORE_API ECbValidateError ValidateCompactBinary(FConstMemoryView View, ECbValidateMode Mode, ECbFieldType Type = ECbFieldType::HasFieldType);
-
-/**
- * Validate the compact binary data for every field in the view as specified by the mode flags.
- *
- * This function expects the entire view to contain fields. Any trailing region of the view which
- * does not contain a valid field will produce an OutOfBounds or InvalidType error instead of the
- * Padding error that would be produced by the single field validation function.
- *
- * \see \ref ValidateCompactBinary for more detail.
- */
-CORE_API ECbValidateError ValidateCompactBinaryRange(FConstMemoryView View, ECbValidateMode Mode);
-
-///////////////////////////////////////////////////////////////////////////////////////////////////////////////////////
-
 /**
  * Determine the size in bytes of a one compact binary field at the start of the view.
  *
diff --git a/Engine/Source/Runtime/Core/Public/Serialization/CompactBinaryValidation.h b/Engine/Source/Runtime/Core/Public/Serialization/CompactBinaryValidation.h
new file mode 100644
--- /dev/null
+++ b/Engine/Source/Runtime/Core/Public/Serialization/CompactBinaryValidation.h
@@ -0,0 +1,131 @@
+// Copyright Epic Games, Inc. All Rights Reserved.
+
+#pragma once
+
+#include "CoreTypes.h"
+#include "Memory/MemoryView.h"
+#include "Misc/EnumClassFlags.h"
+#include "Serialization/CompactBinary.h"
+
+///////////////////////////////////////////////////////////////////////////////////////////////////////////////////////
+
+/** Flags for validating compact binary data. */
+enum class ECbValidateMode : uint32
+{
+	/**
+	 * Validate that the value can be read and stays inside the bounds of the memory view.
+	 *
+	 * This is the minimum level of validation required to be able to safely read a field, array,
+	 * or object without the risk of crashing or reading out of bounds.
+	 */
+	Default                 = 0,
+
+	/**
+	 * Validate that object fields have unique non-empty names and array fields have no names.
+	 *
+	 * Name validation failures typically do not inhibit reading the input, but duplicated fields
+	 * cannot be looked up by name other than the first, and converting to other data formats can
+	 * fail in the presence of naming issues.
+	 */
+	Names                   = 1 << 0,
+
+	/**
+	 * Validate that fields are serialized in the canonical format.
+	 *
+	 * Format validation failures typically do not inhibit reading the input. Values that fail in
+	 * this mode require more memory than in the canonical format, and comparisons of such values
+	 * for equality are not reliable. Examples of failures include uniform arrays or objects that
+	 * were not encoded uniformly, variable-length integers that could be encoded in fewer bytes,
+	 * or 64-bit floats that could be encoded in 32 bits without loss of precision.
+	 */
+	Format                  = 1 << 1,
+
+	/**
+	 * Validate that there is no padding after the value before the end of the memory view.
+	 *
+	 * Padding validation failures have no impact on the ability to read the input, but are using
+	 * more memory than necessary.
+	 */
+	Padding                 = 1 << 2,
+
+	/** Perform all validation described above. */
+	All                     = Default | Names | Format | Padding,
+};
+
+ENUM_CLASS_FLAGS(ECbValidateMode);
+
+///////////////////////////////////////////////////////////////////////////////////////////////////////////////////////
+
+/** Flags for compact binary validation errors. Multiple flags may be combined. */
+enum class ECbValidateError : uint32
+{
+	/** The input had no validation errors. */
+	None                    = 0,
+
+	// Mode: Default
+
+	/** The input cannot be read without reading out of bounds. */
+	OutOfBounds             = 1 << 0,
+	/** The input has a field with an unrecognized or invalid type. */
+	InvalidType             = 1 << 1,
+
+	// Mode: Names
+
+	/** An object had more than one field with the same name. */
+	DuplicateName           = 1 << 2,
+	/** An object had a field with no name. */
+	MissingName             = 1 << 3,
+	/** An array field had a name. */
+	ArrayName               = 1 << 4,
+
+	// Mode: Format
+
+	/** A name or string payload is not valid UTF-8. */
+	InvalidString           = 1 << 5,
+	/** A size or integer payload can be encoded in fewer bytes. */
+	InvalidInteger          = 1 << 6,
+	/** A float64 payload can be encoded as a float32 without loss of precision. */
+	InvalidFloat            = 1 << 7,
+	/** An object has the same type for every field but is not uniform. */
+	NonUniformObject        = 1 << 8,
+	/** An array has the same type for every field and non-empty payloads but is not uniform. */
+	NonUniformArray         = 1 << 9,
+
+	// Mode: Padding
+
+	/** A value did not use the entire memory view given for validation. */
+	Padding                 = 1 << 10,
+};
+
+ENUM_CLASS_FLAGS(ECbValidateError);
+
+///////////////////////////////////////////////////////////////////////////////////////////////////////////////////////
+
+/**
+ * Validate the compact binary data for one field in the view as specified by the mode flags.
+ *
+ * Only one top-level field is processed from the view, and validation recurses into any array or
+ * object within that field. To validate multiple consecutive top-level fields, call the function
+ * once for each top-level field. If the given view might contain multiple top-level fields, then
+ * either exclude the Padding flag from the Mode or use MeasureCompactBinary to break up the view
+ * into its constituent fields before validating.
+ *
+ * @param View A memory view containing at least one top-level field.
+ * @param Mode A combination of the flags for the types of validation to perform.
+ * @param Type HasFieldType means that View contains the type. Otherwise, use the given type.
+ * @return None on success, otherwise the flags for the types of errors that were detected.
+ */
+CORE_API ECbValidateError ValidateCompactBinary(FConstMemoryView View, ECbValidateMode Mode, ECbFieldType Type = ECbFieldType::HasFieldType);
+
+/**
+ * Validate the compact binary data for every field in the view as specified by the mode flags.
+ *
+ * This function expects the entire view to contain fields. Any trailing region of the view which
+ * does not contain a valid field will produce an OutOfBounds or InvalidType error instead of the
+ * Padding error that would be produced by the single field validation function.
+ *
+ * \see \ref ValidateCompactBinary for more detail.
+ */
+CORE_API ECbValidateError ValidateCompactBinaryRange(FConstMemoryView View, ECbValidateMode Mode);
+
+///////////////////////////////////////////////////////////////////////////////////////////////////////////////////////
