Extended TMemoryView with a GetDataEnd accessor and a constructor that takes a pair of pointers

#rnx

diff --git a/Engine/Source/Runtime/Core/Private/Tests/Memory/MemoryViewTest.cpp b/Engine/Source/Runtime/Core/Private/Tests/Memory/MemoryViewTest.cpp
--- a/Engine/Source/Runtime/Core/Private/Tests/Memory/MemoryViewTest.cpp
+++ b/Engine/Source/Runtime/Core/Private/Tests/Memory/MemoryViewTest.cpp
@@ -35,6 +35,7 @@ bool FMemoryViewTest::RunTest(const FString& Parameters)
 	auto TestMutableMemoryView = [this](const FMutableMemoryView& View, void* Data, uint64 Size)
 	{
 		TestEqual(TEXT("MemoryView.GetData()"), View.GetData(), Data);
+		TestEqual(TEXT("MemoryView.GetDataEnd()"), View.GetDataEnd(), static_cast<void*>(static_cast<uint8*>(Data) + Size));
 		TestEqual(TEXT("MemoryView.GetSize()"), View.GetSize(), Size);
 		TestEqual(TEXT("MemoryView.IsEmpty()"), View.IsEmpty(), Size == 0);
 	};
@@ -42,6 +43,7 @@ bool FMemoryViewTest::RunTest(const FString& Parameters)
 	auto TestConstMemoryView = [this](const FConstMemoryView& View, const void* Data, uint64 Size)
 	{
 		TestEqual(TEXT("MemoryView.GetData()"), View.GetData(), Data);
+		TestEqual(TEXT("MemoryView.GetDataEnd()"), View.GetDataEnd(), static_cast<const void*>(static_cast<const uint8*>(Data) + Size));
 		TestEqual(TEXT("MemoryView.GetSize()"), View.GetSize(), Size);
 		TestEqual(TEXT("MemoryView.IsEmpty()"), View.IsEmpty(), Size == 0);
 	};
@@ -61,13 +63,15 @@ bool FMemoryViewTest::RunTest(const FString& Parameters)
 	TestConstMemoryView(FMutableMemoryView(), nullptr, 0);
 	TestConstMemoryView(FConstMemoryView(), nullptr, 0);
 
-	// Test Construction from Type[], TArrayView, (Type*, uint64)
+	// Test Construction from Type[], TArrayView, (Type*, uint64), (Type*, Type*)
 	TestMutableMemoryView(MakeMemoryView(IntArray), IntArray, sizeof(IntArray));
 	TestMutableMemoryView(MakeMemoryView(MakeArrayView(IntArray)), IntArray, sizeof(IntArray));
 	TestMutableMemoryView(MakeMemoryView(IntArray, sizeof(IntArray)), IntArray, sizeof(IntArray));
+	TestMutableMemoryView(MakeMemoryView(IntArray, IntArray + 6), IntArray, sizeof(*IntArray) * 6);
 	TestConstMemoryView(MakeMemoryView(AsConst(IntArray)), IntArray, sizeof(IntArray));
 	TestConstMemoryView(MakeMemoryView(MakeArrayView(AsConst(IntArray))), IntArray, sizeof(IntArray));
 	TestConstMemoryView(MakeMemoryView(AsConst(IntArray), sizeof(IntArray)), IntArray, sizeof(IntArray));
+	TestConstMemoryView(MakeMemoryView(AsConst(IntArray), AsConst(IntArray) + 6), IntArray, sizeof(*IntArray) * 6);
 
 	// Test Construction from std::initializer_list
 	//MakeMemoryView({1, 2, 3}); // fail because the type must be deduced
diff --git a/Engine/Source/Runtime/Core/Public/Memory/MemoryView.h b/Engine/Source/Runtime/Core/Public/Memory/MemoryView.h
--- a/Engine/Source/Runtime/Core/Public/Memory/MemoryView.h
+++ b/Engine/Source/Runtime/Core/Public/Memory/MemoryView.h
@@ -21,28 +21,52 @@
 /**
  * A non-owning view of a contiguous region of memory.
  *
- * Prefer to use the aliases \ref FConstMemoryView or \ref FMutableMemoryView over this type.
+ * Prefer to use the aliases FConstMemoryView or FMutableMemoryView over this type.
  *
  * Functions that modify a view clamp sizes and offsets to always return a sub-view of the input.
  */
 template <typename DataType>
 class TMemoryView
 {
-	static_assert(TIsVoidType<DataType>::Value, "DataMust be cv-qualified void");
+	static_assert(TIsVoidType<DataType>::Value, "DataType must be cv-qualified void");
+
+	using ByteType = typename TChooseClass<TIsConst<DataType>::Value, const uint8, uint8>::Result;
 
 public:
 	/** Construct an empty view. */
 	constexpr TMemoryView() = default;
 
+	/** Construct a view of by copying a view with compatible const/volatile qualifiers. */
+	template <typename OtherDataType,
+		typename TEnableIf<TPointerIsConvertibleFromTo<OtherDataType, DataType>::Value>::Type* = nullptr>
+	constexpr inline TMemoryView(const TMemoryView<OtherDataType>& InView)
+		: Data(InView.Data)
+		, Size(InView.Size)
+	{
+	}
+
 	/** Construct a view of InSize bytes starting at InData. */
-	constexpr inline TMemoryView(DataType* InData, uint64 InSize) : Data(InData), Size(InSize) {}
+	constexpr inline TMemoryView(DataType* InData, uint64 InSize)
+		: Data(InData)
+		, Size(InSize)
+	{
+	}
 
-	template <typename OtherDataType, typename = typename TEnableIf<TPointerIsConvertibleFromTo<OtherDataType, DataType>::Value>::Type>
-	constexpr inline TMemoryView(const TMemoryView<OtherDataType>& InView) : Data(InView.Data), Size(InView.Size) {}
+	/** Construct a view starting at InData and ending at InDataEnd. */
+	template <typename DataEndType,
+		decltype(ImplicitConv<DataType*>(DeclVal<DataEndType*>()))* = nullptr>
+	inline TMemoryView(DataType* InData, DataEndType* InDataEnd)
+		: Data(InData)
+		, Size(static_cast<uint64>(static_cast<ByteType*>(ImplicitConv<DataType*>(InDataEnd)) - static_cast<ByteType*>(InData)))
+	{
+	}
 
 	/** Returns a pointer to the start of the view. */
 	constexpr inline DataType* GetData() const { return Data; }
 
+	/** Returns a pointer to the end of the view. */
+	inline DataType* GetDataEnd() const { return GetDataAtOffsetNoCheck(Size); }
+
 	/** Returns the number of bytes in the view. */
 	constexpr inline uint64 GetSize() const { return Size; }
 
@@ -189,7 +213,6 @@ private:
 	/** Returns the data pointer advanced by an offset in bytes. */
 	inline DataType* GetDataAtOffsetNoCheck(uint64 InOffset) const
 	{
-		using ByteType = typename TChooseClass<TIsConst<DataType>::Value, const uint8, uint8>::Result;
 		return reinterpret_cast<ByteType*>(Data) + InOffset;
 	}
 
@@ -227,6 +250,22 @@ constexpr inline TMemoryView<const void> MakeMemoryView(const void* Data, uint64
 	return TMemoryView<const void>(Data, Size);
 }
 
+/** Make a non-owning mutable view starting at Data and ending at DataEnd. */
+template <typename DataEndType,
+	decltype(ImplicitConv<void*>(DeclVal<DataEndType*>()))* = nullptr>
+inline TMemoryView<void> MakeMemoryView(void* Data, DataEndType* DataEnd)
+{
+	return TMemoryView<void>(Data, DataEnd);
+}
+
+/** Make a non-owning const view starting at Data and ending at DataEnd. */
+template <typename DataEndType,
+	decltype(ImplicitConv<const void*>(DeclVal<DataEndType*>()))* = nullptr>
+inline TMemoryView<const void> MakeMemoryView(const void* Data, DataEndType* DataEnd)
+{
+	return TMemoryView<const void>(Data, DataEnd);
+}
+
 /**
  * Make a non-owning view of the memory of the initializer list.
  *
@@ -239,7 +278,8 @@ constexpr inline TMemoryView<const void> MakeMemoryView(std::initializer_list<ty
 }
 
 /** Make a non-owning view of the memory of the contiguous container. */
-template <typename ContainerType, typename TEnableIf<TIsContiguousContainer<ContainerType>::Value>::Type* = nullptr>
+template <typename ContainerType,
+	typename TEnableIf<TIsContiguousContainer<ContainerType>::Value>::Type* = nullptr>
 constexpr inline auto MakeMemoryView(ContainerType&& Container)
 {
 	using ElementType = typename TRemovePointer<decltype(GetData(DeclVal<ContainerType>()))>::Type;
