CompactBinary: Added field iterator constructors from memory views, shared buffers, and compatible iterators

#rnx

diff --git a/Engine/Source/Runtime/Core/Private/Serialization/CompactBinaryWriter.cpp b/Engine/Source/Runtime/Core/Private/Serialization/CompactBinaryWriter.cpp
--- a/Engine/Source/Runtime/Core/Private/Serialization/CompactBinaryWriter.cpp
+++ b/Engine/Source/Runtime/Core/Private/Serialization/CompactBinaryWriter.cpp
@@ -72,7 +72,7 @@ FCbFieldRefIterator FCbWriter::Save() const
 	const uint64 Size = SaveSize();
 	FSharedBufferPtr Buffer = FSharedBuffer::Alloc(Size);
 	const FCbFieldIterator Output = SaveToMemory(Buffer->GetView());
-	return FCbFieldRefIterator(FCbFieldRef(Output, FSharedBuffer::MakeReadOnly(MoveTemp(Buffer))), Output);
+	return FCbFieldRefIterator(Output, FSharedBuffer::MakeReadOnly(MoveTemp(Buffer)));
 }
 
 FCbFieldIterator FCbWriter::SaveToMemory(const FMutableMemoryView Buffer) const
@@ -83,7 +83,7 @@ FCbFieldIterator FCbWriter::SaveToMemory(const FMutableMemoryView Buffer) const
 	checkf(Buffer.GetSize() == Data.Num(),
 		TEXT("Buffer is %" UINT64_FMT " bytes but %" INT64_FMT " is required."), Buffer.GetSize(), Data.Num());
 	FMemory::Memcpy(Buffer.GetData(), Data.GetData(), Data.Num());
-	return FCbFieldIterator(FCbField(Buffer.GetData()), Buffer.RightChop(Buffer.GetSize()).GetData());
+	return FCbFieldIterator(FCbField(Buffer.GetData()), Buffer.GetDataEnd());
 }
 
 uint64 FCbWriter::SaveSize() const
diff --git a/Engine/Source/Runtime/Core/Private/Tests/Serialization/CompactBinaryTest.cpp b/Engine/Source/Runtime/Core/Private/Tests/Serialization/CompactBinaryTest.cpp
--- a/Engine/Source/Runtime/Core/Private/Tests/Serialization/CompactBinaryTest.cpp
+++ b/Engine/Source/Runtime/Core/Private/Tests/Serialization/CompactBinaryTest.cpp
@@ -1215,6 +1215,119 @@ bool FCbObjectRefTest::RunTest(const FString& Parameters)
 
 ///////////////////////////////////////////////////////////////////////////////////////////////////////////////////////
 
+IMPLEMENT_SIMPLE_AUTOMATION_TEST(FCbFieldRefIteratorTest, "System.Core.Serialization.CbFieldRefIterator", CompactBinaryTestFlags)
+bool FCbFieldRefIteratorTest::RunTest(const FString& Parameters)
+{
+	static_assert(std::is_constructible<FCbFieldIterator, const FCbField&>::value, "Missing constructor for FCbFieldIterator");
+	static_assert(std::is_constructible<FCbFieldIterator, const FCbFieldRef&>::value, "Missing constructor for FCbFieldIterator");
+	static_assert(std::is_constructible<FCbFieldIterator, const FCbFieldRefIterator&>::value, "Missing constructor for FCbFieldIterator");
+	static_assert(std::is_constructible<FCbFieldIterator, FCbField&&>::value, "Missing constructor for FCbFieldIterator");
+	static_assert(std::is_constructible<FCbFieldIterator, FCbFieldRef&&>::value, "Missing constructor for FCbFieldIterator");
+	static_assert(std::is_constructible<FCbFieldIterator, FCbFieldRefIterator&&>::value, "Missing constructor for FCbFieldIterator");
+	static_assert(std::is_constructible<FCbFieldIterator, FCbField&&, const void*>::value, "Missing constructor for FCbFieldIterator");
+
+	static_assert(std::is_constructible<FCbFieldRefIterator, const FCbFieldRef&>::value, "Missing constructor for FCbFieldRefIterator");
+	static_assert(std::is_constructible<FCbFieldRefIterator, const FCbFieldRefIterator&>::value, "Missing constructor for FCbFieldRefIterator");
+	static_assert(std::is_constructible<FCbFieldRefIterator, FCbFieldRef&&>::value, "Missing constructor for FCbFieldRefIterator");
+	static_assert(std::is_constructible<FCbFieldRefIterator, FCbFieldRefIterator&&>::value, "Missing constructor for FCbFieldRefIterator");
+	static_assert(std::is_constructible<FCbFieldRefIterator, FCbFieldRef&&, const void*>::value, "Missing constructor for FCbFieldRefIterator");
+
+	static_assert(std::is_constructible<FCbFieldRefIterator, const FCbFieldIterator&, const FSharedBufferRef&>::value, "Missing constructor for FCbFieldRefIterator");
+	static_assert(std::is_constructible<FCbFieldRefIterator, const FCbFieldIterator&, const FSharedBufferPtr&>::value, "Missing constructor for FCbFieldRefIterator");
+	static_assert(std::is_constructible<FCbFieldRefIterator, const FCbFieldIterator&, const FSharedBufferConstRef&>::value, "Missing constructor for FCbFieldRefIterator");
+	static_assert(std::is_constructible<FCbFieldRefIterator, const FCbFieldIterator&, const FSharedBufferConstPtr&>::value, "Missing constructor for FCbFieldRefIterator");
+
+	static_assert(std::is_constructible<FCbFieldRefIterator, const FCbFieldIterator&, FSharedBufferRef&&>::value, "Missing constructor for FCbFieldRefIterator");
+	static_assert(std::is_constructible<FCbFieldRefIterator, const FCbFieldIterator&, FSharedBufferPtr&&>::value, "Missing constructor for FCbFieldRefIterator");
+	static_assert(std::is_constructible<FCbFieldRefIterator, const FCbFieldIterator&, FSharedBufferConstRef&&>::value, "Missing constructor for FCbFieldRefIterator");
+	static_assert(std::is_constructible<FCbFieldRefIterator, const FCbFieldIterator&, FSharedBufferConstPtr&&>::value, "Missing constructor for FCbFieldRefIterator");
+
+	const auto GetCount = [](auto It) -> uint32
+	{
+		uint32 Count = 0;
+		for (; It; ++It)
+		{
+			++Count;
+		}
+		return Count;
+	};
+
+	// Test FCbField[Ref]Iterator()
+	{
+		TestEqual(TEXT("FCbFieldIterator()"), GetCount(FCbFieldIterator()), 0);
+		TestEqual(TEXT("FCbFieldRefIterator()"), GetCount(FCbFieldRefIterator()), 0);
+	}
+
+	// Test FCbField[Ref]Iterator(Range)
+	{
+		constexpr uint8 T = uint8(ECbFieldType::IntegerPositive);
+		const uint8 Payload[] = { T, 0, T, 1, T, 2, T, 3 };
+		const uint8* const PayloadEnd = Payload + UE_ARRAY_COUNT(Payload);
+
+		const FSharedBufferConstPtr View = FSharedBuffer::MakeView(MakeMemoryView(Payload));
+		const void* const ViewEnd = View->GetView().GetDataEnd();
+		const FSharedBufferConstPtr Clone = FSharedBuffer::Clone(MakeMemoryView(Payload));
+		const void* const CloneEnd = Clone->GetView().GetDataEnd();
+
+		const FConstMemoryView EmptyView;
+		const FSharedBufferConstPtr NullBuffer;
+
+		const FCbFieldIterator FieldIt(FCbField(View->GetData()), ViewEnd);
+		const FCbFieldRefIterator FieldRefIt(FCbFieldRef(View), ViewEnd);
+
+		TestEqual(TEXT("FCbFieldIterator(EmptyView)"), GetCount(FCbFieldIterator(EmptyView)), 0);
+		TestEqual(TEXT("FCbFieldRefIterator(BufferNullL)"), GetCount(FCbFieldRefIterator(NullBuffer)), 0);
+		TestEqual(TEXT("FCbFieldRefIterator(BufferNullR)"), GetCount(FCbFieldRefIterator(FSharedBufferConstPtr(NullBuffer))), 0);
+
+		TestEqual(TEXT("FCbFieldIterator(View)"), GetCount(FCbFieldIterator(MakeMemoryView(Payload))), 4);
+		TestEqual(TEXT("FCbFieldRefIterator(BufferCloneL)"), GetCount(FCbFieldRefIterator(Clone)), 4);
+		TestEqual(TEXT("FCbFieldRefIterator(BufferCloneR)"), GetCount(FCbFieldRefIterator(FSharedBufferConstPtr(Clone))), 4);
+
+		TestEqual(TEXT("FCbFieldIterator(Field, End)"), GetCount(FCbFieldIterator(FCbField(Payload), PayloadEnd)), 4);
+		TestEqual(TEXT("FCbFieldRefIterator(FieldRef, End)"), GetCount(FCbFieldRefIterator(FCbFieldRef(View), PayloadEnd)), 4);
+
+		TestEqual(TEXT("FCbFieldRefIterator(FieldIt, BufferNullL)"), GetCount(FCbFieldRefIterator(FCbFieldIterator(FCbField(View->GetData()), ViewEnd), NullBuffer)), 4);
+		TestEqual(TEXT("FCbFieldRefIterator(FieldIt, BufferNullR)"), GetCount(FCbFieldRefIterator(FCbFieldIterator(FCbField(View->GetData()), ViewEnd), FSharedBufferConstPtr(NullBuffer))), 4);
+		TestEqual(TEXT("FCbFieldRefIterator(FieldIt, BufferViewL)"), GetCount(FCbFieldRefIterator(FCbFieldIterator(FCbField(View->GetData()), ViewEnd), View)), 4);
+		TestEqual(TEXT("FCbFieldRefIterator(FieldIt, BufferViewR)"), GetCount(FCbFieldRefIterator(FCbFieldIterator(FCbField(View->GetData()), ViewEnd), FSharedBufferConstPtr(View))), 4);
+		TestEqual(TEXT("FCbFieldRefIterator(FieldIt, BufferCloneL)"), GetCount(FCbFieldRefIterator(FCbFieldIterator(FCbField(Clone->GetData()), CloneEnd), Clone)), 4);
+		TestEqual(TEXT("FCbFieldRefIterator(FieldIt, BufferCloneR)"), GetCount(FCbFieldRefIterator(FCbFieldIterator(FCbField(Clone->GetData()), CloneEnd), FSharedBufferConstPtr(Clone))), 4);
+
+		TestEqual(TEXT("FCbFieldIterator(FieldRefItL)"), GetCount(FCbFieldIterator(FieldRefIt)), 4);
+		TestEqual(TEXT("FCbFieldIterator(FieldRefItR)"), GetCount(FCbFieldIterator(FCbFieldRefIterator(FieldRefIt))), 4);
+
+		// These lines are expected to assert when uncommented.
+		//const FSharedBufferConstPtr ShortView = FSharedBuffer::MakeView(MakeMemoryView(Payload).LeftChop(2));
+		//TestEqual(TEXT("FCbFieldRefIterator(FieldIt, InvalidBufferL)"), GetCount(FCbFieldRefIterator(FCbFieldIterator(FCbField(View->GetData()), ViewEnd), ShortView)), 4);
+		//TestEqual(TEXT("FCbFieldRefIterator(FieldIt, InvalidBufferR)"), GetCount(FCbFieldRefIterator(FCbFieldIterator(FCbField(View->GetData()), ViewEnd), FSharedBufferConstPtr(ShortView))), 4);
+	}
+
+	// Test FCbField[Ref]Iterator(Scalar)
+	{
+		constexpr uint8 T = uint8(ECbFieldType::IntegerPositive);
+		const uint8 Payload[] = { T, 0 };
+		const uint8* const PayloadEnd = Payload + UE_ARRAY_COUNT(Payload);
+
+		const FSharedBufferConstPtr View = FSharedBuffer::MakeView(MakeMemoryView(Payload));
+		const void* const ViewEnd = View->GetView().GetDataEnd();
+		const FSharedBufferConstPtr Clone = FSharedBuffer::Clone(MakeMemoryView(Payload));
+		const void* const CloneEnd = Clone->GetView().GetDataEnd();
+		const FSharedBufferConstPtr NullBuffer;
+
+		const FCbField Field(Payload);
+		const FCbFieldRef FieldRef(View);
+
+		TestEqual(TEXT("FCbFieldIterator(FieldL)"), GetCount(FCbFieldIterator(Field)), 1);
+		TestEqual(TEXT("FCbFieldIterator(FieldR)"), GetCount(FCbFieldIterator(FCbField(Field))), 1);
+		TestEqual(TEXT("FCbFieldRefIterator(FieldRefL)"), GetCount(FCbFieldRefIterator(FieldRef)), 1);
+		TestEqual(TEXT("FCbFieldRefIterator(FieldRefR)"), GetCount(FCbFieldRefIterator(FCbFieldRef(FieldRef))), 1);
+	}
+
+	return true;
+}
+
+///////////////////////////////////////////////////////////////////////////////////////////////////////////////////////
+
 IMPLEMENT_SIMPLE_AUTOMATION_TEST(FCbMeasureTest, "System.Core.Serialization.MeasureCompactBinary", CompactBinaryTestFlags)
 bool FCbMeasureTest::RunTest(const FString& Parameters)
 {
diff --git a/Engine/Source/Runtime/Core/Public/Serialization/CompactBinary.h b/Engine/Source/Runtime/Core/Public/Serialization/CompactBinary.h
--- a/Engine/Source/Runtime/Core/Public/Serialization/CompactBinary.h
+++ b/Engine/Source/Runtime/Core/Public/Serialization/CompactBinary.h
@@ -276,22 +276,9 @@ public:
 	 * @param InField The first field, or the default field if there are no fields.
 	 * @param InFieldsEnd A pointer to the end of the payload of the last field, or null.
 	 */
-	constexpr inline TCbFieldIterator(const FieldType& InField, const void* InFieldsEnd)
-		: FieldType(InField)
-		, FieldsEnd(InFieldsEnd)
-	{
-	}
-
-	/**
-	 * Construct a field iterator.
-	 *
-	 * @param InField The first field, or the default field if there are no fields.
-	 * @param InIterator An iterator from which to take the pointer to the end of the payload of the last field.
-	 */
-	template <typename OtherType>
-	constexpr inline TCbFieldIterator(FieldType&& InField, const TCbFieldIterator<OtherType>& InIterator)
+	constexpr inline TCbFieldIterator(FieldType&& InField, const void* InFieldsEnd)
 		: FieldType(MoveTemp(InField))
-		, FieldsEnd(InIterator.FieldsEnd)
+		, FieldsEnd(InFieldsEnd)
 	{
 	}
 
@@ -332,14 +319,30 @@ public:
 	constexpr inline FieldType& operator*() { return *this; }
 	constexpr inline FieldType* operator->() { return this; }
 
-	constexpr inline bool operator!() const { return !bool(*this); }
-
 	constexpr inline bool operator==(const TCbFieldIterator& Other) const
 	{
 		return FieldType::GetPayload() == Other.FieldType::GetPayload();
 	}
 
-	constexpr inline bool operator!=(const TCbFieldIterator& Other) const { return !(*this == Other); }
+	constexpr inline bool operator!=(const TCbFieldIterator& Other) const
+	{
+		return !(*this == Other);
+	}
+
+protected:
+	/** Returns the end of the last field, or null for an iterator at the end. */
+	template <typename OtherFieldType>
+	static inline const void* GetFieldsEnd(const TCbFieldIterator<OtherFieldType>& It)
+	{
+		return It.FieldsEnd;
+	}
+
+	/** Create a view of the field iterator. */
+	template <typename OtherFieldType>
+	static inline FConstMemoryView AsMemoryView(const TCbFieldIterator<OtherFieldType>& It)
+	{
+		return MakeMemoryView(It.OtherFieldType::AsMemoryView().GetData(), It.FieldsEnd);
+	}
 
 private:
 	friend inline TCbFieldIterator begin(const TCbFieldIterator& Iterator) { return Iterator; }
@@ -609,10 +612,25 @@ public:
 
 	FCbFieldIterator() = default;
 
-	inline FCbFieldIterator(TCbFieldIterator<FCbField>&& Field)
-		: TCbFieldIterator(MoveTemp(Field))
+	/** Construct an iterator from another iterator. */
+	template <typename OtherFieldType>
+	inline FCbFieldIterator(const TCbFieldIterator<OtherFieldType>& It)
+		: TCbFieldIterator(ImplicitConv<FCbField>(It), GetFieldsEnd(It))
+	{
+	}
+
+	/** Construct an iterator for the fields in the view. */
+	inline explicit FCbFieldIterator(FConstMemoryView View)
+		: FCbFieldIterator(MakeFromView(View))
 	{
 	}
+
+private:
+	/** Construct an iterator for the fields in the view. */
+	static inline FCbFieldIterator MakeFromView(FConstMemoryView View)
+	{
+		return !View.IsEmpty() ? FCbFieldIterator(FCbField(View.GetData()), View.GetDataEnd()) : FCbFieldIterator();
+	}
 };
 
 ///////////////////////////////////////////////////////////////////////////////////////////////////////////////////////
@@ -986,6 +1004,51 @@ class FCbFieldRefIterator : public TCbFieldIterator<FCbFieldRef>
 {
 public:
 	using TCbFieldIterator::TCbFieldIterator;
+
+	/** Construct a field iterator that has no fields to iterate. */
+	constexpr FCbFieldRefIterator() = default;
+
+	/** Construct an iterator that holds a reference to the buffer that contains its fields. */
+	inline FCbFieldRefIterator(const FCbFieldIterator& It, const FSharedBufferConstPtr& ValueBuffer)
+		: TCbFieldIterator(MakeFieldRef(It, ValueBuffer), GetFieldsEnd(It))
+	{
+	}
+
+	/** Construct an iterator that holds a reference to the buffer that contains its fields. */
+	inline FCbFieldRefIterator(const FCbFieldIterator& It, FSharedBufferConstPtr&& ValueBuffer)
+		: TCbFieldIterator(MakeFieldRef(It, MoveTemp(ValueBuffer)), GetFieldsEnd(It))
+	{
+	}
+
+	/** Construct an iterator and hold a reference to the buffer that contains its fields. */
+	inline explicit FCbFieldRefIterator(const FSharedBufferConstPtr& ValueBuffer)
+		: FCbFieldRefIterator(MakeFieldIterator(ValueBuffer), ValueBuffer)
+	{
+	}
+
+	/** Construct an iterator and hold a reference to the buffer that contains its fields. */
+	inline explicit FCbFieldRefIterator(FSharedBufferConstPtr&& ValueBuffer)
+		: FCbFieldRefIterator(MakeFieldIterator(ValueBuffer), MoveTemp(ValueBuffer))
+	{
+	}
+
+private:
+	/** Make a field ref from an iterator and the buffer that contains its fields. */
+	template <typename BufferType>
+	static inline FCbFieldRef MakeFieldRef(const FCbFieldIterator& It, BufferType&& ValueBuffer)
+	{
+		if (ValueBuffer)
+		{
+			check(ValueBuffer->GetView().Contains(AsMemoryView(It)));
+		}
+		return FCbFieldRef(It, Forward<BufferType>(ValueBuffer));
+	}
+
+	/** Make a field iterator from the possibly-null buffer that contains its fields. */
+	static inline FCbFieldIterator MakeFieldIterator(const FSharedBufferConstPtr& ValueBuffer)
+	{
+		return ValueBuffer ? FCbFieldIterator(ValueBuffer->GetView()) : FCbFieldIterator();
+	}
 };
 
 ///////////////////////////////////////////////////////////////////////////////////////////////////////////////////////
@@ -999,8 +1062,7 @@ public:
 	/** Create an iterator for the fields of this array. */
 	inline FCbFieldRefIterator CreateRefIterator() const
 	{
-		FCbFieldIterator It = CreateIterator();
-		return FCbFieldRefIterator(FCbFieldRef(It, GetBuffer()), It);
+		return FCbFieldRefIterator(CreateIterator(), GetBuffer());
 	}
 
 	/** Access the array as an array field. */
@@ -1027,8 +1089,7 @@ public:
 	/** Create an iterator for the fields of this object. */
 	inline FCbFieldRefIterator CreateRefIterator() const
 	{
-		FCbFieldIterator It = CreateIterator();
-		return FCbFieldRefIterator(FCbFieldRef(It, GetBuffer()), It);
+		return FCbFieldRefIterator(CreateIterator(), GetBuffer());
 	}
 
 	/** Find a field by case-sensitive name comparison. */
