CompactBinary: Fixed FCbArray::CreateIterator to compare the payload size against the correct byte count

#rnx

diff --git a/Engine/Source/Runtime/Core/Private/Serialization/CompactBinary.cpp b/Engine/Source/Runtime/Core/Private/Serialization/CompactBinary.cpp
--- a/Engine/Source/Runtime/Core/Private/Serialization/CompactBinary.cpp
+++ b/Engine/Source/Runtime/Core/Private/Serialization/CompactBinary.cpp
@@ -482,10 +482,11 @@ FCbFieldIterator FCbArray::CreateIterator() const
 	uint32 PayloadSizeByteCount;
 	const uint64 PayloadSize = ReadVarUInt(PayloadBytes, PayloadSizeByteCount);
 	PayloadBytes += PayloadSizeByteCount;
-	if (PayloadSize > PayloadSizeByteCount)
+	const uint64 NumByteCount = MeasureVarUInt(PayloadBytes);
+	if (PayloadSize > NumByteCount)
 	{
 		const void* const PayloadEnd = PayloadBytes + PayloadSize;
-		PayloadBytes += MeasureVarUInt(PayloadBytes);
+		PayloadBytes += NumByteCount;
 		const ECbFieldType UniformType = FCbFieldType::GetType(Type) == ECbFieldType::UniformArray ?
 			ECbFieldType(*PayloadBytes++) : ECbFieldType::HasFieldType;
 		return FCbFieldIterator(FCbField(PayloadBytes, UniformType), PayloadEnd);
