CompactBinary: Changed array and object to inherit from field

#rnx

diff --git a/Engine/Source/Runtime/Core/Private/Serialization/CompactBinary.cpp b/Engine/Source/Runtime/Core/Private/Serialization/CompactBinary.cpp
--- a/Engine/Source/Runtime/Core/Private/Serialization/CompactBinary.cpp
+++ b/Engine/Source/Runtime/Core/Private/Serialization/CompactBinary.cpp
@@ -84,7 +84,7 @@ FCbField::FCbField(const void* const InData, const ECbFieldType InType)
 	Payload = Bytes;
 }
 
-FConstMemoryView FCbField::AsMemoryView() const
+FConstMemoryView FCbField::GetFieldView() const
 {
 	const uint32 TypeSize = FCbFieldType::HasFieldType(Type) ? sizeof(ECbFieldType) : 0;
 	const uint32 NameSize = FCbFieldType::HasFieldName(Type) ? NameLen + MeasureVarUInt(NameLen) : 0;
@@ -94,7 +94,7 @@ FConstMemoryView FCbField::AsMemoryView() const
 
 uint64 FCbField::GetSize() const
 {
-	return AsMemoryView().GetSize();
+	return GetFieldView().GetSize();
 }
 
 uint64 FCbField::GetPayloadSize() const
@@ -141,7 +141,7 @@ uint64 FCbField::GetPayloadSize() const
 
 bool FCbField::Equals(const FCbField& Other) const
 {
-	return Type == Other.Type && AsMemoryView().EqualBytes(Other.AsMemoryView());
+	return Type == Other.Type && GetFieldView().EqualBytes(Other.GetFieldView());
 }
 
 FCbObject FCbField::AsObject()
@@ -150,7 +150,7 @@ FCbObject FCbField::AsObject()
 	if (FCbFieldType::IsObject(LocalType))
 	{
 		Error = ECbFieldError::None;
-		return FCbObject(Payload, FCbFieldType::GetType(LocalType));
+		return FCbObject::FromField(*this);
 	}
 	else
 	{
@@ -165,7 +165,7 @@ FCbArray FCbField::AsArray()
 	if (FCbFieldType::IsArray(LocalType))
 	{
 		Error = ECbFieldError::None;
-		return FCbArray(Payload, FCbFieldType::GetType(LocalType));
+		return FCbArray::FromField(*this);
 	}
 	else
 	{
@@ -432,45 +432,25 @@ FTimespan FCbField::AsTimeSpan(FTimespan Default)
 
 void FCbField::CopyTo(const FMutableMemoryView Buffer) const
 {
-	const FConstMemoryView Source = AsMemoryView();
+	const FConstMemoryView Source = GetFieldView();
 	FMemory::Memcpy(Buffer.GetData(), Source.GetData(), FMath::Min(Buffer.GetSize(), Source.GetSize()));
 }
 
 ///////////////////////////////////////////////////////////////////////////////////////////////////////////////////////
 
 FCbArray::FCbArray()
-	: Payload(CompactBinaryPrivate::GEmptyArrayPayload)
-	, Type(ECbFieldType::Array)
+	: FCbField(CompactBinaryPrivate::GEmptyArrayPayload, ECbFieldType::Array)
 {
 }
 
-FCbArray::FCbArray(const void* const InData, const ECbFieldType InType)
-{
-	const uint8* Bytes = static_cast<const uint8*>(InData);
-	const ECbFieldType LocalType = FCbFieldType::HasFieldType(InType) ? (ECbFieldType(*Bytes++) | ECbFieldType::HasFieldType) : InType;
-	uint32 NameLenByteCount = 0;
-	const uint64 NameLen = FCbFieldType::HasFieldName(LocalType) ? ReadVarUInt(Bytes, NameLenByteCount) : 0;
-	Bytes += NameLen + NameLenByteCount;
-
-	Type = LocalType;
-	Payload = Bytes;
-}
-
-FConstMemoryView FCbArray::AsMemoryView() const
-{
-	uint32 PayloadSizeByteCount;
-	const uint64 PayloadSize = ReadVarUInt(Payload, PayloadSizeByteCount);
-	return MakeMemoryView(Payload, PayloadSizeByteCount + PayloadSize);
-}
-
 uint64 FCbArray::GetSize() const
 {
-	return sizeof(ECbFieldType) + AsMemoryView().GetSize();
+	return sizeof(ECbFieldType) + GetPayloadSize();
 }
 
 uint64 FCbArray::Num() const
 {
-	const uint8* PayloadBytes = static_cast<const uint8*>(Payload);
+	const uint8* PayloadBytes = static_cast<const uint8*>(GetPayload());
 	PayloadBytes += MeasureVarUInt(PayloadBytes);
 	uint32 NumByteCount;
 	return ReadVarUInt(PayloadBytes, NumByteCount);
@@ -478,7 +458,7 @@ uint64 FCbArray::Num() const
 
 FCbFieldIterator FCbArray::CreateIterator() const
 {
-	const uint8* PayloadBytes = static_cast<const uint8*>(Payload);
+	const uint8* PayloadBytes = static_cast<const uint8*>(GetPayload());
 	uint32 PayloadSizeByteCount;
 	const uint64 PayloadSize = ReadVarUInt(PayloadBytes, PayloadSizeByteCount);
 	PayloadBytes += PayloadSizeByteCount;
@@ -487,7 +467,7 @@ FCbFieldIterator FCbArray::CreateIterator() const
 	{
 		const void* const PayloadEnd = PayloadBytes + PayloadSize;
 		PayloadBytes += NumByteCount;
-		const ECbFieldType UniformType = FCbFieldType::GetType(Type) == ECbFieldType::UniformArray ?
+		const ECbFieldType UniformType = FCbFieldType::GetType(GetType()) == ECbFieldType::UniformArray ?
 			ECbFieldType(*PayloadBytes++) : ECbFieldType::HasFieldType;
 		return FCbFieldIterator(FCbField(PayloadBytes, UniformType), PayloadEnd);
 	}
@@ -496,59 +476,40 @@ FCbFieldIterator FCbArray::CreateIterator() const
 
 bool FCbArray::Equals(const FCbArray& Other) const
 {
-	return FCbFieldType::GetType(Type) == FCbFieldType::GetType(Other.Type) && AsMemoryView().EqualBytes(Other.AsMemoryView());
+	return FCbFieldType::GetType(GetType()) == FCbFieldType::GetType(Other.GetType()) &&
+		GetPayloadView().EqualBytes(Other.GetPayloadView());
 }
 
 void FCbArray::CopyTo(const FMutableMemoryView Buffer) const
 {
 	uint8* BufferBytes = static_cast<uint8*>(Buffer.GetData());
-	*BufferBytes++ = uint8(FCbFieldType::GetType(Type));
-	const FConstMemoryView Source = AsMemoryView();
+	*BufferBytes++ = uint8(FCbFieldType::GetType(GetType()));
+	const FConstMemoryView Source = GetPayloadView();
 	FMemory::Memcpy(BufferBytes, Source.GetData(), FMath::Min(Buffer.GetSize() - 1, Source.GetSize()));
 }
 
 ///////////////////////////////////////////////////////////////////////////////////////////////////////////////////////
 
 FCbObject::FCbObject()
-	: Payload(CompactBinaryPrivate::GEmptyObjectPayload)
-	, Type(ECbFieldType::Object)
+	: FCbField(CompactBinaryPrivate::GEmptyObjectPayload, ECbFieldType::Object)
 {
 }
 
-FCbObject::FCbObject(const void* const InData, const ECbFieldType InType)
-{
-	const uint8* Bytes = static_cast<const uint8*>(InData);
-	const ECbFieldType LocalType = FCbFieldType::HasFieldType(InType) ? (ECbFieldType(*Bytes++) | ECbFieldType::HasFieldType) : InType;
-	uint32 NameLenByteCount = 0;
-	const uint64 NameLen = FCbFieldType::HasFieldName(LocalType) ? ReadVarUInt(Bytes, NameLenByteCount) : 0;
-	Bytes += NameLen + NameLenByteCount;
-
-	Type = LocalType;
-	Payload = Bytes;
-}
-
-FConstMemoryView FCbObject::AsMemoryView() const
-{
-	uint32 PayloadSizeByteCount;
-	const uint64 PayloadSize = ReadVarUInt(Payload, PayloadSizeByteCount);
-	return MakeMemoryView(Payload, PayloadSizeByteCount + PayloadSize);
-}
-
 uint64 FCbObject::GetSize() const
 {
-	return sizeof(ECbFieldType) + AsMemoryView().GetSize();
+	return sizeof(ECbFieldType) + GetPayloadSize();
 }
 
 FCbFieldIterator FCbObject::CreateIterator() const
 {
-	const uint8* PayloadBytes = static_cast<const uint8*>(Payload);
+	const uint8* PayloadBytes = static_cast<const uint8*>(GetPayload());
 	uint32 PayloadSizeByteCount;
 	const uint64 PayloadSize = ReadVarUInt(PayloadBytes, PayloadSizeByteCount);
 	PayloadBytes += PayloadSizeByteCount;
 	if (PayloadSize)
 	{
 		const void* const PayloadEnd = PayloadBytes + PayloadSize;
-		const ECbFieldType UniformType = FCbFieldType::GetType(Type) == ECbFieldType::UniformObject ?
+		const ECbFieldType UniformType = FCbFieldType::GetType(GetType()) == ECbFieldType::UniformObject ?
 			ECbFieldType(*PayloadBytes++) : ECbFieldType::HasFieldType;
 		return FCbFieldIterator(FCbField(PayloadBytes, UniformType), PayloadEnd);
 	}
@@ -557,7 +518,8 @@ FCbFieldIterator FCbObject::CreateIterator() const
 
 bool FCbObject::Equals(const FCbObject& Other) const
 {
-	return FCbFieldType::GetType(Type) == FCbFieldType::GetType(Other.Type) && AsMemoryView().EqualBytes(Other.AsMemoryView());
+	return FCbFieldType::GetType(GetType()) == FCbFieldType::GetType(Other.GetType()) &&
+		GetPayloadView().EqualBytes(Other.GetPayloadView());
 }
 
 FCbField FCbObject::Find(const FAnsiStringView Name) const
@@ -587,8 +549,8 @@ FCbField FCbObject::FindIgnoreCase(const FAnsiStringView Name) const
 void FCbObject::CopyTo(const FMutableMemoryView Buffer) const
 {
 	uint8* BufferBytes = static_cast<uint8*>(Buffer.GetData());
-	*BufferBytes++ = uint8(FCbFieldType::GetType(Type));
-	const FConstMemoryView Source = AsMemoryView();
+	*BufferBytes++ = uint8(FCbFieldType::GetType(GetType()));
+	const FConstMemoryView Source = GetPayloadView();
 	FMemory::Memcpy(BufferBytes, Source.GetData(), FMath::Min(Buffer.GetSize() - 1, Source.GetSize()));
 }
 
diff --git a/Engine/Source/Runtime/Core/Private/Serialization/CompactBinaryWriter.cpp b/Engine/Source/Runtime/Core/Private/Serialization/CompactBinaryWriter.cpp
--- a/Engine/Source/Runtime/Core/Private/Serialization/CompactBinaryWriter.cpp
+++ b/Engine/Source/Runtime/Core/Private/Serialization/CompactBinaryWriter.cpp
@@ -45,6 +45,22 @@ static constexpr bool IsUniformType(const ECbFieldType Type)
 	}
 }
 
+/** Append the payload from the compact binary value to the array and return its type. */
+template <typename T>
+static inline ECbFieldType AppendCompactBinary(const T& Value, TArray64<uint8>& OutData)
+{
+	struct FCopy : public T
+	{
+		using T::GetCopyType;
+		using T::GetPayloadView;
+	};
+	const FCopy& ValueCopy = static_cast<const FCopy&>(Value);
+	const FConstMemoryView SourceView = ValueCopy.GetPayloadView();
+	const int64 TargetOffset = OutData.AddUninitialized(SourceView.GetSize());
+	FMemory::Memcpy(OutData.GetData() + TargetOffset, SourceView.GetData(), SourceView.GetSize());
+	return FCbFieldType::GetType(ValueCopy.GetCopyType());
+}
+
 FCbWriter::FCbWriter()
 {
 	States.Emplace();
@@ -209,18 +225,7 @@ void FCbWriter::Field(const FCbField& Value)
 {
 	checkf(Value.HasValue(), TEXT("It is invalid to write a field with no value."));
 	BeginField();
-	struct FCbFieldCopy : public FCbField
-	{
-		using FCbField::GetCopyType;
-		using FCbField::GetPayload;
-		using FCbField::GetPayloadEnd;
-	};
-	const FCbFieldCopy& ValueCopy = static_cast<const FCbFieldCopy&>(Value);
-	const void* const SourceData = ValueCopy.GetPayload();
-	const int64 SourceSize = PTRINT(ValueCopy.GetPayloadEnd()) - PTRINT(SourceData);
-	const int64 TargetOffset = Data.AddUninitialized(SourceSize);
-	FMemory::Memcpy(Data.GetData() + TargetOffset, SourceData, SourceSize);
-	EndField(FCbFieldType::GetType(ValueCopy.GetCopyType()));
+	EndField(AppendCompactBinary(Value, Data));
 }
 
 void FCbWriter::BeginObject()
@@ -268,16 +273,7 @@ void FCbWriter::EndObject()
 void FCbWriter::Object(const FCbObject& Value)
 {
 	BeginField();
-	struct FCbObjectCopy : public FCbObject
-	{
-		using FCbObject::GetCopyType;
-		using FCbObject::AsMemoryView;
-	};
-	const FCbObjectCopy& ValueCopy = static_cast<const FCbObjectCopy&>(Value);
-	const FConstMemoryView SourceView = ValueCopy.AsMemoryView();
-	const int64 TargetOffset = Data.AddUninitialized(SourceView.GetSize());
-	FMemory::Memcpy(Data.GetData() + TargetOffset, SourceView.GetData(), SourceView.GetSize());
-	EndField(FCbFieldType::GetType(ValueCopy.GetCopyType()));
+	EndField(AppendCompactBinary(Value, Data));
 }
 
 void FCbWriter::BeginArray()
@@ -327,16 +323,7 @@ void FCbWriter::EndArray()
 void FCbWriter::Array(const FCbArray& Value)
 {
 	BeginField();
-	struct FCbArrayCopy : public FCbArray
-	{
-		using FCbArray::GetCopyType;
-		using FCbArray::AsMemoryView;
-	};
-	const FCbArrayCopy& ValueCopy = static_cast<const FCbArrayCopy&>(Value);
-	const FConstMemoryView SourceView = ValueCopy.AsMemoryView();
-	const int64 TargetOffset = Data.AddUninitialized(SourceView.GetSize());
-	FMemory::Memcpy(Data.GetData() + TargetOffset, SourceView.GetData(), SourceView.GetSize());
-	EndField(FCbFieldType::GetType(ValueCopy.GetCopyType()));
+	EndField(AppendCompactBinary(Value, Data));
 }
 
 void FCbWriter::Null()
diff --git a/Engine/Source/Runtime/Core/Private/Tests/Serialization/CompactBinaryTest.cpp b/Engine/Source/Runtime/Core/Private/Tests/Serialization/CompactBinaryTest.cpp
--- a/Engine/Source/Runtime/Core/Private/Tests/Serialization/CompactBinaryTest.cpp
+++ b/Engine/Source/Runtime/Core/Private/Tests/Serialization/CompactBinaryTest.cpp
@@ -263,6 +263,11 @@ bool FCbFieldNullTest::RunTest(const FString& Parameters)
 IMPLEMENT_CUSTOM_SIMPLE_AUTOMATION_TEST(FCbFieldObjectTest, FCbFieldTestBase, "System.Core.Serialization.CbField.Object", CompactBinaryTestFlags)
 bool FCbFieldObjectTest::RunTest(const FString& Parameters)
 {
+	static_assert(!std::is_constructible<FCbField, const FCbObject&>::value, "Invalid constructor for FCbField");
+	static_assert(!std::is_assignable<FCbField, const FCbObject&>::value, "Invalid assignment for FCbField");
+	static_assert(!std::is_convertible<FCbField, FCbObject>::value, "Invalid conversion to FCbObject");
+	static_assert(!std::is_assignable<FCbObject, const FCbField&>::value, "Invalid assignment for FCbObject");
+
 	auto TestIntObject = [this](const FCbObject& Object, int32 ExpectedNum, uint64 ExpectedPayloadSize)
 	{
 		TestEqual(TEXT("FCbField(Object)::AsObject().GetSize()"), Object.GetSize(), ExpectedPayloadSize + sizeof(ECbFieldType));
@@ -358,8 +363,8 @@ bool FCbFieldObjectTest::RunTest(const FString& Parameters)
 		const uint8 NamedPayload[] = { 1, 'O', 10, IntType, 1, 'A', 1, 1, 'B', 2, 1, 'C', 3 };
 		FCbField NamedField(NamedPayload, ECbFieldType::UniformObject | ECbFieldType::HasFieldName);
 		TestTrue(TEXT("FCbObject::Equals()"), Object.Equals(NamedField.AsObject()));
-		TestTrue(TEXT("FCbObject::AsField().Equals()"), Field.Equals(NamedField.AsObject().AsField()));
-		TestFalse(TEXT("FCbObject::AsField().Equals()"), NamedField.Equals(NamedField.AsObject().AsField()));
+		TestTrue(TEXT("FCbObject::AsField().Equals()"), Field.Equals(Field.AsObject().AsField()));
+		TestTrue(TEXT("FCbObject::AsField().Equals()"), NamedField.Equals(NamedField.AsObject().AsField()));
 	}
 
 	// Test FCbField(None) as Object
@@ -421,6 +426,11 @@ bool FCbFieldObjectTest::RunTest(const FString& Parameters)
 IMPLEMENT_CUSTOM_SIMPLE_AUTOMATION_TEST(FCbFieldArrayTest, FCbFieldTestBase, "System.Core.Serialization.CbField.Array", CompactBinaryTestFlags)
 bool FCbFieldArrayTest::RunTest(const FString& Parameters)
 {
+	static_assert(!std::is_constructible<FCbField, const FCbArray&>::value, "Invalid constructor for FCbField");
+	static_assert(!std::is_assignable<FCbField, const FCbArray&>::value, "Invalid assignment for FCbField");
+	static_assert(!std::is_convertible<FCbField, FCbArray>::value, "Invalid conversion to FCbArray");
+	static_assert(!std::is_assignable<FCbArray, const FCbField&>::value, "Invalid assignment for FCbArray");
+
 	auto TestIntArray = [this](FCbArray Array, int32 ExpectedNum, uint64 ExpectedPayloadSize)
 	{
 		TestEqual(TEXT("FCbField(Array)::AsArray().GetSize()"), Array.GetSize(), ExpectedPayloadSize + sizeof(ECbFieldType));
@@ -494,8 +504,8 @@ bool FCbFieldArrayTest::RunTest(const FString& Parameters)
 		const uint8 NamedPayload[] = { 1, 'A', 5, 3, IntType, 1, 2, 3 };
 		FCbField NamedField(NamedPayload, ECbFieldType::UniformArray | ECbFieldType::HasFieldName);
 		TestTrue(TEXT("FCbArray::Equals()"), Array.Equals(NamedField.AsArray()));
-		TestTrue(TEXT("FCbArray::AsField().Equals()"), Field.Equals(NamedField.AsArray().AsField()));
-		TestFalse(TEXT("FCbArray::AsField().Equals()"), NamedField.Equals(NamedField.AsArray().AsField()));
+		TestTrue(TEXT("FCbArray::AsField().Equals()"), Field.Equals(Field.AsArray().AsField()));
+		TestTrue(TEXT("FCbArray::AsField().Equals()"), NamedField.Equals(NamedField.AsArray().AsField()));
 	}
 
 	// Test FCbField(None) as Array
diff --git a/Engine/Source/Runtime/Core/Public/Serialization/CompactBinary.h b/Engine/Source/Runtime/Core/Public/Serialization/CompactBinary.h
--- a/Engine/Source/Runtime/Core/Public/Serialization/CompactBinary.h
+++ b/Engine/Source/Runtime/Core/Public/Serialization/CompactBinary.h
@@ -339,9 +339,9 @@ protected:
 
 	/** Create a view of the field iterator. */
 	template <typename OtherFieldType>
-	static inline FConstMemoryView AsMemoryView(const TCbFieldIterator<OtherFieldType>& It)
+	static inline FConstMemoryView GetFieldRangeView(const TCbFieldIterator<OtherFieldType>& It)
 	{
-		return MakeMemoryView(It.OtherFieldType::AsMemoryView().GetData(), It.FieldsEnd);
+		return MakeMemoryView(It.OtherFieldType::GetFieldView().GetData(), It.FieldsEnd);
 	}
 
 private:
@@ -535,11 +535,25 @@ public:
 	constexpr inline ECbFieldError GetError() const { return Error; }
 
 protected:
+	/** Returns the type of the field including any relevant flags. */
 	constexpr inline ECbFieldType GetType() const { return Type; }
+
+	/** Returns the type to use when constructing a new field after using CopyTo. */
 	constexpr inline ECbFieldType GetCopyType() const { return Type; }
+
+	/** Size of the field payload in bytes. That is the field excluding its type and name. */
+	CORE_API uint64 GetPayloadSize() const;
+
+	/** Returns the start of the value payload. */
 	constexpr inline const void* GetPayload() const { return Payload; }
+
+	/** Returns the end of the value payload. */
 	inline const void* GetPayloadEnd() const { return static_cast<const uint8*>(Payload) + GetPayloadSize(); }
 
+	/** Returns a view of the value payload, which excludes the type and name. */
+	inline FConstMemoryView GetPayloadView() const { return MakeMemoryView(Payload, GetPayloadSize()); }
+
+	/** Assign a field from a pointer to its data and an optional externally-provided type. */
 	inline void Assign(const void* InData, const ECbFieldType InType)
 	{
 		static_assert(TIsTriviallyDestructible<FCbField>::Value,
@@ -555,12 +569,9 @@ protected:
 	CORE_API void CopyTo(FMutableMemoryView Buffer) const;
 
 	/** Create a view of the field, including the type and name when present. */
-	CORE_API FConstMemoryView AsMemoryView() const;
+	CORE_API FConstMemoryView GetFieldView() const;
 
 private:
-	/** Size of the field payload in bytes. That is the field excluding its type and name. */
-	CORE_API uint64 GetPayloadSize() const;
-
 	/** Parameters for converting to an integer. */
 	struct FIntegerParams
 	{
@@ -644,7 +655,7 @@ private:
  * This type only provides a view into memory and does not perform any memory management itself.
  * Use \ref FCbArrayRef to hold a reference to the underlying memory when necessary.
  */
-class FCbArray
+class FCbArray : protected FCbField
 {
 public:
 	/** Construct an array with no fields. */
@@ -656,7 +667,10 @@ public:
 	 * @param Data Pointer to the start of the array data.
 	 * @param Type HasFieldType means that Data contains the type. Otherwise, use the given type.
 	 */
-	CORE_API explicit FCbArray(const void* Data, ECbFieldType Type = ECbFieldType::HasFieldType);
+	inline explicit FCbArray(const void* Data, ECbFieldType Type = ECbFieldType::HasFieldType)
+		: FCbField(Data, Type)
+	{
+	}
 
 	/** Size of the array in bytes if serialized by itself with no name. */
 	CORE_API uint64 GetSize() const;
@@ -679,27 +693,26 @@ public:
 	CORE_API bool Equals(const FCbArray& Other) const;
 
 	/** Access the array as an array field. */
-	inline FCbField AsField() const
-	{
-		return FCbField(Payload, FCbFieldType::GetType(Type));
-	}
+	inline FCbField AsField() const { return static_cast<const FCbField&>(*this); }
+
+	/** Construct an array from an array field. No type check is performed! */
+	static FCbArray FromField(const FCbField& Field) { return FCbArray(Field); }
 
 protected:
-	inline ECbFieldType GetCopyType() const { return FCbFieldType::GetType(Type) | ECbFieldType::HasFieldType; }
+	/** Returns the type to use when constructing a new array after using CopyTo. */
+	constexpr inline ECbFieldType GetCopyType() const
+	{
+		return FCbFieldType::GetType(GetType()) | ECbFieldType::HasFieldType;
+	}
 
 	/** Copy the array into a buffer of at least GetSize() bytes. */
 	CORE_API void CopyTo(FMutableMemoryView Buffer) const;
 
-	/** Create a view of the payload, excluding the type or the name of the outer field. */
-	CORE_API FConstMemoryView AsMemoryView() const;
-
 private:
 	friend inline FCbFieldIterator begin(const FCbArray& Array) { return Array.CreateIterator(); }
 	friend inline FCbFieldIterator end(const FCbArray&) { return FCbFieldIterator(); }
 
-private:
-	const void* Payload;
-	ECbFieldType Type;
+	inline explicit FCbArray(const FCbField& Field) : FCbField(Field) {}
 };
 
 ///////////////////////////////////////////////////////////////////////////////////////////////////////////////////////
@@ -715,7 +728,7 @@ private:
  * This type only provides a view into memory and does not perform any memory management itself.
  * Use \ref FCbObjectRef to hold a reference to the underlying memory when necessary.
  */
-class FCbObject
+class FCbObject : protected FCbField
 {
 public:
 	/** Construct an object with no fields. */
@@ -727,7 +740,10 @@ public:
 	 * @param Data Pointer to the start of the object data.
 	 * @param Type HasFieldType means that Data contains the type. Otherwise, use the given type.
 	 */
-	CORE_API explicit FCbObject(const void* Data, ECbFieldType Type = ECbFieldType::HasFieldType);
+	inline explicit FCbObject(const void* Data, ECbFieldType Type = ECbFieldType::HasFieldType)
+		: FCbField(Data, Type)
+	{
+	}
 
 	/** Size of the object in bytes if serialized by itself with no name. */
 	CORE_API uint64 GetSize() const;
@@ -764,27 +780,26 @@ public:
 	inline FCbField operator[](FAnsiStringView Name) const { return Find(Name); }
 
 	/** Access the object as an object field. */
-	inline FCbField AsField() const
-	{
-		return FCbField(Payload, FCbFieldType::GetType(Type));
-	}
+	inline FCbField AsField() const { return static_cast<const FCbField&>(*this); }
+
+	/** Construct an object from an object field. No type check is performed! */
+	static FCbObject FromField(const FCbField& Field) { return FCbObject(Field); }
 
 protected:
-	inline ECbFieldType GetCopyType() const { return FCbFieldType::GetType(Type) | ECbFieldType::HasFieldType; }
+	/** Returns the type to use when constructing a new object after using CopyTo. */
+	constexpr inline ECbFieldType GetCopyType() const
+	{
+		return FCbFieldType::GetType(GetType()) | ECbFieldType::HasFieldType;
+	}
 
 	/** Copy the object into a buffer of at least GetSize() bytes. */
 	CORE_API void CopyTo(FMutableMemoryView Buffer) const;
 
-	/** Create a view of the payload, excluding the type or the name of the outer field. */
-	CORE_API FConstMemoryView AsMemoryView() const;
-
 private:
 	friend inline FCbFieldIterator begin(const FCbObject& Object) { return Object.CreateIterator(); }
 	friend inline FCbFieldIterator end(const FCbObject&) { return FCbFieldIterator(); }
 
-private:
-	const void* Payload;
-	ECbFieldType Type;
+	inline explicit FCbObject(const FCbField& Field) : FCbField(Field) {}
 };
 
 ///////////////////////////////////////////////////////////////////////////////////////////////////////////////////////
@@ -851,7 +866,7 @@ public:
 	{
 		if (ValueBuffer)
 		{
-			check(ValueBuffer->GetView().Contains(BaseType::AsMemoryView()));
+			check(ValueBuffer->GetView().Contains(BaseType::GetFieldView()));
 			if (ValueBuffer->IsOwned())
 			{
 				Buffer = ValueBuffer;
@@ -865,7 +880,7 @@ public:
 	{
 		if (ValueBuffer)
 		{
-			check(ValueBuffer->GetView().Contains(BaseType::AsMemoryView()));
+			check(ValueBuffer->GetView().Contains(BaseType::GetFieldView()));
 			if (ValueBuffer->IsOwned())
 			{
 				Buffer = MoveTemp(ValueBuffer);
@@ -1039,7 +1054,7 @@ private:
 	{
 		if (ValueBuffer)
 		{
-			check(ValueBuffer->GetView().Contains(AsMemoryView(It)));
+			check(ValueBuffer->GetView().Contains(GetFieldRangeView(It)));
 		}
 		return FCbFieldRef(It, Forward<BufferType>(ValueBuffer));
 	}
@@ -1095,7 +1110,7 @@ public:
 	/** Find a field by case-sensitive name comparison. */
 	inline FCbFieldRef FindRef(FAnsiStringView Name) const
 	{
-		if (FCbField Field = Find(Name))
+		if (::FCbField Field = Find(Name))
 		{
 			return FCbFieldRef(Field, *this);
 		}
@@ -1105,7 +1120,7 @@ public:
 	/** Find a field by case-insensitive name comparison. */
 	inline FCbFieldRef FindRefIgnoreCase(FAnsiStringView Name) const
 	{
-		if (FCbField Field = FindIgnoreCase(Name))
+		if (::FCbField Field = FindIgnoreCase(Name))
 		{
 			return FCbFieldRef(Field, *this);
 		}
