CompactBinary: Added functions to calculate the hash of compact binary values

#rnx

diff --git a/Engine/Source/Runtime/Core/Private/Serialization/CompactBinary.cpp b/Engine/Source/Runtime/Core/Private/Serialization/CompactBinary.cpp
--- a/Engine/Source/Runtime/Core/Private/Serialization/CompactBinary.cpp
+++ b/Engine/Source/Runtime/Core/Private/Serialization/CompactBinary.cpp
@@ -480,6 +480,20 @@ bool FCbArray::Equals(const FCbArray& Other) const
 		GetPayloadView().EqualBytes(Other.GetPayloadView());
 }
 
+FBlake3Hash FCbArray::GetHash() const
+{
+	FBlake3 Hash;
+	GetHash(Hash);
+	return Hash.Finalize();
+}
+
+void FCbArray::GetHash(FBlake3& Hash) const
+{
+	const uint8 LocalType = uint8(FCbFieldType::GetType(GetType()));
+	Hash.Update(&LocalType, sizeof(LocalType));
+	Hash.Update(GetPayloadView());
+}
+
 void FCbArray::CopyTo(const FMutableMemoryView Buffer) const
 {
 	uint8* BufferBytes = static_cast<uint8*>(Buffer.GetData());
@@ -546,6 +560,20 @@ FCbField FCbObject::FindIgnoreCase(const FAnsiStringView Name) const
 	return FCbField();
 }
 
+FBlake3Hash FCbObject::GetHash() const
+{
+	FBlake3 Hash;
+	GetHash(Hash);
+	return Hash.Finalize();
+}
+
+void FCbObject::GetHash(FBlake3& Hash) const
+{
+	const uint8 LocalType = uint8(FCbFieldType::GetType(GetType()));
+	Hash.Update(&LocalType, sizeof(LocalType));
+	Hash.Update(GetPayloadView());
+}
+
 void FCbObject::CopyTo(const FMutableMemoryView Buffer) const
 {
 	uint8* BufferBytes = static_cast<uint8*>(Buffer.GetData());
diff --git a/Engine/Source/Runtime/Core/Private/Tests/Serialization/CompactBinaryTest.cpp b/Engine/Source/Runtime/Core/Private/Tests/Serialization/CompactBinaryTest.cpp
--- a/Engine/Source/Runtime/Core/Private/Tests/Serialization/CompactBinaryTest.cpp
+++ b/Engine/Source/Runtime/Core/Private/Tests/Serialization/CompactBinaryTest.cpp
@@ -189,6 +189,7 @@ bool FCbFieldNoneTest::RunTest(const FString& Parameters)
 		TestFalse(TEXT("!FCbField()::HasValue()"), DefaultField.HasValue());
 		TestFalse(TEXT("!FCbField()::HasError()"), DefaultField.HasError());
 		TestEqual(TEXT("FCbField()::GetError() == None"), DefaultField.GetError(), ECbFieldError::None);
+		TestEqual(TEXT("FCbField()::GetHash()"), DefaultField.GetHash(), FBlake3::HashBuffer(FConstMemoryView()));
 	}
 
 	// Test FCbField(None)
@@ -200,6 +201,7 @@ bool FCbFieldNoneTest::RunTest(const FString& Parameters)
 		TestFalse(TEXT("!FCbField(None)::HasValue()"), NoneField.HasValue());
 		TestFalse(TEXT("!FCbField(None)::HasError()"), NoneField.HasError());
 		TestEqual(TEXT("FCbField(None)::GetError() == None"), NoneField.GetError(), ECbFieldError::None);
+		TestEqual(TEXT("FCbField(None)::GetHash()"), NoneField.GetHash(), FCbField().GetHash());
 	}
 
 	// Test FCbField(None|Type|Name)
@@ -211,6 +213,7 @@ bool FCbFieldNoneTest::RunTest(const FString& Parameters)
 		TestEqual(TEXT("FCbField(None|Type|Name)::GetName()"), NoneField.GetName(), "Name"_ASV);
 		TestTrue(TEXT("FCbField(None|Type|Name)::HasName()"), NoneField.HasName());
 		TestFalse(TEXT("!FCbField(None|Type|Name)::HasValue()"), NoneField.HasValue());
+		TestNotEqual(TEXT("FCbField(None|Type|Name)::GetHash()"), NoneField.GetHash(), FCbField().GetHash());
 	}
 
 	// Test FCbField(None|Type)
@@ -222,6 +225,7 @@ bool FCbFieldNoneTest::RunTest(const FString& Parameters)
 		TestEqual(TEXT("FCbField(None|Type)::GetName()"), NoneField.GetName().Len(), 0);
 		TestFalse(TEXT("FCbField(None|Type)::HasName()"), NoneField.HasName());
 		TestFalse(TEXT("!FCbField(None|Type)::HasValue()"), NoneField.HasValue());
+		TestNotEqual(TEXT("FCbField(None|Type)::GetHash()"), NoneField.GetHash(), FCbField().GetHash());
 	}
 
 	// Test FCbField(None|Name)
@@ -233,6 +237,7 @@ bool FCbFieldNoneTest::RunTest(const FString& Parameters)
 		TestEqual(TEXT("FCbField(None|Name)::GetName()"), NoneField.GetName(), "Name"_ASV);
 		TestTrue(TEXT("FCbField(None|Name)::HasName()"), NoneField.HasName());
 		TestFalse(TEXT("!FCbField(None|Name)::HasValue()"), NoneField.HasValue());
+		TestNotEqual(TEXT("FCbField(None|Name)::GetHash()"), NoneField.GetHash(), FCbField().GetHash());
 	}
 
 	return true;
@@ -249,6 +254,7 @@ bool FCbFieldNullTest::RunTest(const FString& Parameters)
 		TestTrue(TEXT("FCbField(Null)::HasValue()"), NullField.HasValue());
 		TestFalse(TEXT("!FCbField(Null)::HasError()"), NullField.HasError());
 		TestEqual(TEXT("FCbField(Null)::GetError() == None"), NullField.GetError(), ECbFieldError::None);
+		TestEqual(TEXT("FCbField(Null)::GetHash()"), NullField.GetHash(), FBlake3::HashBuffer(FConstMemoryView()));
 	}
 
 	// Test FCbField(None) as Null
@@ -383,6 +389,7 @@ bool FCbFieldObjectTest::RunTest(const FString& Parameters)
 		const FCbObjectRef ObjectClone = FCbObjectRef::Clone(Object);
 		TestEqual(TEXT("FCbObjectRef(ObjectWithName)::GetSize()"), ObjectClone.GetSize(), uint64(6));
 		TestTrue(TEXT("FCbObject::Equals()"), Object.Equals(ObjectClone));
+		TestEqual(TEXT("FCbObject::GetHash()"), ObjectClone.GetHash(), Object.GetHash());
 		for (FCbFieldRefIterator It = ObjectClone.CreateRefIterator(); It; ++It)
 		{
 			FCbFieldRef Field = *It;
@@ -524,6 +531,7 @@ bool FCbFieldArrayTest::RunTest(const FString& Parameters)
 		const FCbArrayRef ArrayClone = FCbArrayRef::Clone(Array);
 		TestEqual(TEXT("FCbArrayRef(ArrayWithName)::GetSize()"), ArrayClone.GetSize(), uint64(5));
 		TestTrue(TEXT("FCbArray::Equals()"), Array.Equals(ArrayClone));
+		TestEqual(TEXT("FCbArray::GetHash()"), ArrayClone.GetHash(), Array.GetHash());
 		for (FCbFieldRefIterator It = ArrayClone.CreateRefIterator(); It; ++It)
 		{
 			FCbFieldRef Field = *It;
@@ -1285,6 +1293,9 @@ bool FCbFieldRefIteratorTest::RunTest(const FString& Parameters)
 		const FCbFieldIterator FieldIt(FCbField(View->GetData()), ViewEnd);
 		const FCbFieldRefIterator FieldRefIt(FCbFieldRef(View), ViewEnd);
 
+		TestEqual(TEXT("FCbFieldIterator::GetHash()"), FieldIt.GetHash(), FBlake3::HashBuffer(*View));
+		TestEqual(TEXT("FCbFieldRefIterator::GetHash()"), FieldRefIt.GetHash(), FBlake3::HashBuffer(*View));
+
 		TestEqual(TEXT("FCbFieldIterator(EmptyView)"), GetCount(FCbFieldIterator(EmptyView)), 0);
 		TestEqual(TEXT("FCbFieldRefIterator(BufferNullL)"), GetCount(FCbFieldRefIterator(NullBuffer)), 0);
 		TestEqual(TEXT("FCbFieldRefIterator(BufferNullR)"), GetCount(FCbFieldRefIterator(FSharedBufferConstPtr(NullBuffer))), 0);
diff --git a/Engine/Source/Runtime/Core/Public/Misc/Blake3.h b/Engine/Source/Runtime/Core/Public/Misc/Blake3.h
--- a/Engine/Source/Runtime/Core/Public/Misc/Blake3.h
+++ b/Engine/Source/Runtime/Core/Public/Misc/Blake3.h
@@ -4,6 +4,7 @@
 
 #include "Containers/StringFwd.h"
 #include "HAL/UnrealMemory.h"
+#include "Memory/MemoryView.h"
 #include "String/BytesToHex.h"
 #include "Templates/TypeCompatibleBytes.h"
 
@@ -48,12 +49,24 @@ public:
 	/** Add the data as input to the hash. May be called any number of times. */
 	CORE_API void Update(const void* Data, uint64 Size);
 
+	/** Add the view as input to the hash. May be called any number of times. */
+	inline void Update(FConstMemoryView View)
+	{
+		Update(View.GetData(), View.GetSize());
+	}
+
 	/** Finalize the hash of the input data. May be called any number of times, and more input may be added after. */
 	CORE_API FBlake3Hash Finalize() const;
 
 	/** Calculate the hash of the input data. */
 	CORE_API static FBlake3Hash HashBuffer(const void* Data, uint64 Size);
 
+	/** Calculate the hash of the input view. */
+	static inline FBlake3Hash HashBuffer(FConstMemoryView View)
+	{
+		return HashBuffer(View.GetData(), View.GetSize());
+	}
+
 private:
 	TAlignedBytes<1912, 8> HasherBytes;
 };
diff --git a/Engine/Source/Runtime/Core/Public/Serialization/CompactBinary.h b/Engine/Source/Runtime/Core/Public/Serialization/CompactBinary.h
--- a/Engine/Source/Runtime/Core/Public/Serialization/CompactBinary.h
+++ b/Engine/Source/Runtime/Core/Public/Serialization/CompactBinary.h
@@ -319,16 +319,31 @@ public:
 	constexpr inline FieldType& operator*() { return *this; }
 	constexpr inline FieldType* operator->() { return this; }
 
-	constexpr inline bool operator==(const TCbFieldIterator& Other) const
+	using FieldType::Equals;
+
+	template <typename OtherFieldType>
+	constexpr inline bool Equals(const TCbFieldIterator<OtherFieldType>& Other) const
 	{
 		return FieldType::GetPayload() == Other.FieldType::GetPayload();
 	}
 
-	constexpr inline bool operator!=(const TCbFieldIterator& Other) const
+	template <typename OtherFieldType>
+	constexpr inline bool operator==(const TCbFieldIterator<OtherFieldType>& Other) const
 	{
-		return !(*this == Other);
+		return Equals(Other);
 	}
 
+	template <typename OtherFieldType>
+	constexpr inline bool operator!=(const TCbFieldIterator<OtherFieldType>& Other) const
+	{
+		return !Equals(Other);
+	}
+
+	/** Calculate the hash of the fields by hashing their underlying memory. */
+	inline FBlake3Hash GetHash() const { return FBlake3::HashBuffer(GetFieldRangeView(*this)); }
+	/** Calculate the hash of the fields by hashing their underlying memory. */
+	inline void GetHash(FBlake3& Hash) const { Hash.Update(GetFieldRangeView(*this)); }
+
 protected:
 	/** Returns the end of the last field, or null for an iterator at the end. */
 	template <typename OtherFieldType>
@@ -425,6 +440,11 @@ public:
 	 */
 	CORE_API bool Equals(const FCbField& Other) const;
 
+	/** Calculate the hash of the field by hashing its underlying memory. */
+	inline FBlake3Hash GetHash() const { return FBlake3::HashBuffer(GetFieldView()); }
+	/** Calculate the hash of the field by hashing its underlying memory. */
+	inline void GetHash(FBlake3& Hash) const { Hash.Update(GetFieldView()); }
+
 	/** Access the field as an object. Defaults to an empty object on error. */
 	CORE_API FCbObject AsObject();
 
@@ -698,6 +718,11 @@ public:
 	/** Construct an array from an array field. No type check is performed! */
 	static FCbArray FromField(const FCbField& Field) { return FCbArray(Field); }
 
+	/** Calculate the hash of the array if serialized by itself with no name. */
+	CORE_API FBlake3Hash GetHash() const;
+	/** Calculate the hash of the array if serialized by itself with no name. */
+	CORE_API void GetHash(FBlake3& Hash) const;
+
 protected:
 	/** Returns the type to use when constructing a new array after using CopyTo. */
 	constexpr inline ECbFieldType GetCopyType() const
@@ -785,6 +810,11 @@ public:
 	/** Construct an object from an object field. No type check is performed! */
 	static FCbObject FromField(const FCbField& Field) { return FCbObject(Field); }
 
+	/** Calculate the hash of the object if serialized by itself with no name. */
+	CORE_API FBlake3Hash GetHash() const;
+	/** Calculate the hash of the object if serialized by itself with no name. */
+	CORE_API void GetHash(FBlake3& Hash) const;
+
 protected:
 	/** Returns the type to use when constructing a new object after using CopyTo. */
 	constexpr inline ECbFieldType GetCopyType() const
