CompactBinary: Fixed the compact binary writer to work with string literals

#rnx

diff --git a/Engine/Source/Runtime/Core/Private/Tests/Serialization/CompactBinaryWriterTest.cpp b/Engine/Source/Runtime/Core/Private/Tests/Serialization/CompactBinaryWriterTest.cpp
--- a/Engine/Source/Runtime/Core/Private/Tests/Serialization/CompactBinaryWriterTest.cpp
+++ b/Engine/Source/Runtime/Core/Private/Tests/Serialization/CompactBinaryWriterTest.cpp
@@ -972,6 +972,8 @@ bool FCbWriterStreamTest::RunTest(const FString& Parameters)
 		Writer << "DateTime"_ASV << FDateTime(2020, 5, 13, 15, 10);
 		Writer << "TimeSpan"_ASV << FTimespan(1, 2, 4, 8);
 
+		Writer << "LiteralName" << nullptr;
+
 		Writer.EndObject();
 		Object = Writer.Save().AsObjectRef();
 	}
diff --git a/Engine/Source/Runtime/Core/Public/Serialization/CompactBinaryWriter.h b/Engine/Source/Runtime/Core/Public/Serialization/CompactBinaryWriter.h
--- a/Engine/Source/Runtime/Core/Public/Serialization/CompactBinaryWriter.h
+++ b/Engine/Source/Runtime/Core/Public/Serialization/CompactBinaryWriter.h
@@ -250,6 +250,12 @@ inline FCbWriter& operator<<(FCbWriter& Writer, FAnsiStringView NameOrValue)
 	return Writer;
 }
 
+/** Write the field name if valid in this state, otherwise write the string value. */
+inline FCbWriter& operator<<(FCbWriter& Writer, const ANSICHAR* NameOrValue)
+{
+	return Writer << FAnsiStringView(NameOrValue);
+}
+
 inline FCbWriter& operator<<(FCbWriter& Writer, const FCbField& Value)
 {
 	Writer.Field(Value);
