CompactBinary: Added Iterate[Range]References to compact binary to easily enumerate reference fields

#rnx

diff --git a/Engine/Source/Runtime/Core/Private/Serialization/CompactBinary.cpp b/Engine/Source/Runtime/Core/Private/Serialization/CompactBinary.cpp
--- a/Engine/Source/Runtime/Core/Private/Serialization/CompactBinary.cpp
+++ b/Engine/Source/Runtime/Core/Private/Serialization/CompactBinary.cpp
@@ -144,6 +144,24 @@ bool FCbField::Equals(const FCbField& Other) const
 	return Type == Other.Type && GetFieldView().EqualBytes(Other.GetFieldView());
 }
 
+void FCbField::IterateReferences(FCbFieldVisitor Visitor) const
+{
+	switch (FCbFieldType::GetType(Type))
+	{
+	case ECbFieldType::Object:
+	case ECbFieldType::UniformObject:
+		return FCbObject::FromField(*this).IterateReferences(Visitor);
+	case ECbFieldType::Array:
+	case ECbFieldType::UniformArray:
+		return FCbArray::FromField(*this).IterateReferences(Visitor);
+	case ECbFieldType::Reference:
+	case ECbFieldType::BinaryReference:
+		return Visitor(*this);
+	default:
+		return;
+	}
+}
+
 FCbObject FCbField::AsObject()
 {
 	const ECbFieldType LocalType = Type;
diff --git a/Engine/Source/Runtime/Core/Private/Tests/Serialization/CompactBinaryTest.cpp b/Engine/Source/Runtime/Core/Private/Tests/Serialization/CompactBinaryTest.cpp
--- a/Engine/Source/Runtime/Core/Private/Tests/Serialization/CompactBinaryTest.cpp
+++ b/Engine/Source/Runtime/Core/Private/Tests/Serialization/CompactBinaryTest.cpp
@@ -7,6 +7,8 @@
 #include "Misc/DateTime.h"
 #include "Misc/Guid.h"
 #include "Misc/Timespan.h"
+#include "Serialization/CompactBinaryValidation.h"
+#include "Serialization/CompactBinaryWriter.h"
 #include "Serialization/VarInt.h"
 
 #if WITH_DEV_AUTOMATION_TESTS
@@ -994,6 +996,180 @@ bool FCbFieldTimeSpanTest::RunTest(const FString& Parameters)
 	return true;
 }
 
+IMPLEMENT_CUSTOM_SIMPLE_AUTOMATION_TEST(FCbFieldIterateReferencesTest, FCbFieldTestBase, "System.Core.Serialization.CbField.IterateReferences", CompactBinaryTestFlags)
+bool FCbFieldIterateReferencesTest::RunTest(const FString& Parameters)
+{
+	const auto MakeTestHash = [](uint32 Index) { return FBlake3::HashBuffer(&Index, sizeof(Index)); };
+
+	FCbFieldRefIterator Fields;
+	{
+		FCbWriter Writer;
+
+		Writer.Name("IgnoredTypeInRoot").Hash(MakeTestHash(100));
+		Writer.Reference(MakeTestHash(0));
+		Writer.BinaryReference(MakeTestHash(1));
+		Writer.Name("RefInRoot").Reference(MakeTestHash(2));
+		Writer.Name("BinRefInRoot").BinaryReference(MakeTestHash(3));
+
+		// Uniform array of type to ignore.
+		Writer.BeginArray();
+		{
+			Writer << 1;
+			Writer << 2;
+		}
+		Writer.EndArray();
+		// Uniform array of binary references.
+		Writer.BeginArray();
+		{
+			Writer.BinaryReference(MakeTestHash(4));
+			Writer.BinaryReference(MakeTestHash(5));
+		}
+		Writer.EndArray();
+		// Uniform array of uniform arrays.
+		Writer.BeginArray();
+		{
+			Writer.BeginArray();
+			Writer.BinaryReference(MakeTestHash(6));
+			Writer.BinaryReference(MakeTestHash(7));
+			Writer.EndArray();
+			Writer.BeginArray();
+			Writer.BinaryReference(MakeTestHash(8));
+			Writer.BinaryReference(MakeTestHash(9));
+			Writer.EndArray();
+		}
+		Writer.EndArray();
+		// Uniform array of non-uniform arrays.
+		Writer.BeginArray();
+		{
+			Writer.BeginArray();
+			Writer << 0;
+			Writer << false;
+			Writer.EndArray();
+			Writer.BeginArray();
+			Writer.Reference(MakeTestHash(10));
+			Writer << false;
+			Writer.EndArray();
+		}
+		Writer.EndArray();
+		// Uniform array of uniform objects.
+		Writer.BeginArray();
+		{
+			Writer.BeginObject();
+			Writer.Name("RefInUniObjInUniObj1").Reference(MakeTestHash(11));
+			Writer.Name("RefInUniObjInUniObj2").Reference(MakeTestHash(12));
+			Writer.EndObject();
+			Writer.BeginObject();
+			Writer.Name("RefInUniObjInUniObj3").Reference(MakeTestHash(13));
+			Writer.Name("RefInUniObjInUniObj4").Reference(MakeTestHash(14));
+			Writer.EndObject();
+		}
+		Writer.EndArray();
+		// Uniform array of non-uniform objects.
+		Writer.BeginArray();
+		{
+			Writer.BeginObject();
+			Writer << "Int" << 0;
+			Writer << "Bool" << false;
+			Writer.EndObject();
+			Writer.BeginObject();
+			Writer.Name("RefInNonUniObjInUniObj").Reference(MakeTestHash(15));
+			Writer << "Bool" << false;
+			Writer.EndObject();
+		}
+		Writer.EndArray();
+
+		// Uniform object of type to ignore.
+		Writer.BeginObject();
+		{
+			Writer << "Int1" << 1;
+			Writer << "Int2" << 2;
+		}
+		Writer.EndObject();
+		// Uniform object of binary references.
+		Writer.BeginObject();
+		{
+			Writer.Name("BinRefInUniObj1").BinaryReference(MakeTestHash(16));
+			Writer.Name("BinRefInUniObj2").BinaryReference(MakeTestHash(17));
+		}
+		Writer.EndObject();
+		// Uniform object of uniform arrays.
+		Writer.BeginObject();
+		{
+			Writer.Name("Array1");
+			Writer.BeginArray();
+			Writer.BinaryReference(MakeTestHash(18));
+			Writer.BinaryReference(MakeTestHash(19));
+			Writer.EndArray();
+			Writer.Name("Array2");
+			Writer.BeginArray();
+			Writer.BinaryReference(MakeTestHash(20));
+			Writer.BinaryReference(MakeTestHash(21));
+			Writer.EndArray();
+		}
+		Writer.EndObject();
+		// Uniform object of non-uniform arrays.
+		Writer.BeginObject();
+		{
+			Writer.Name("Array1");
+			Writer.BeginArray();
+			Writer << 0;
+			Writer << false;
+			Writer.EndArray();
+			Writer.Name("Array2");
+			Writer.BeginArray();
+			Writer.Reference(MakeTestHash(22));
+			Writer << false;
+			Writer.EndArray();
+		}
+		Writer.EndObject();
+		// Uniform object of uniform objects.
+		Writer.BeginObject();
+		{
+			Writer.Name("Object1");
+			Writer.BeginObject();
+			Writer.Name("RefInUniObjInUniObj1").Reference(MakeTestHash(23));
+			Writer.Name("RefInUniObjInUniObj2").Reference(MakeTestHash(24));
+			Writer.EndObject();
+			Writer.Name("Object2");
+			Writer.BeginObject();
+			Writer.Name("RefInUniObjInUniObj3").Reference(MakeTestHash(25));
+			Writer.Name("RefInUniObjInUniObj4").Reference(MakeTestHash(26));
+			Writer.EndObject();
+		}
+		Writer.EndObject();
+		// Uniform object of non-uniform objects.
+		Writer.BeginObject();
+		{
+			Writer.Name("Object1");
+			Writer.BeginObject();
+			Writer << "Int" << 0;
+			Writer << "Bool" << false;
+			Writer.EndObject();
+			Writer.Name("Object2");
+			Writer.BeginObject();
+			Writer.Name("RefInNonUniObjInUniObj").Reference(MakeTestHash(27));
+			Writer << "Bool" << false;
+			Writer.EndObject();
+		}
+		Writer.EndObject();
+
+		Fields = Writer.Save();
+	}
+
+	TestEqual(TEXT("FCbField::IterateReferences Validate"), ValidateCompactBinaryRange(*Fields.GetBuffer(), ECbValidateMode::All), ECbValidateError::None);
+
+	uint32 ReferenceIndex = 0;
+	Fields.IterateRangeReferences([this, &ReferenceIndex, &MakeTestHash](FCbField Field)
+		{
+			TestTrue(FString::Printf(TEXT("FCbField::IterateReferences(%u)::IsAnyReference"), ReferenceIndex), Field.IsAnyReference());
+			TestEqual(FString::Printf(TEXT("FCbField::IterateReferences(%u)"), ReferenceIndex), Field.AsReference(Field.AsBinaryReference()), MakeTestHash(ReferenceIndex));
+			++ReferenceIndex;
+		});
+	TestEqual(TEXT("FCbField::IterateReferences"), ReferenceIndex, 28);
+
+	return true;
+}
+
 ///////////////////////////////////////////////////////////////////////////////////////////////////////////////////////
 
 IMPLEMENT_CUSTOM_SIMPLE_AUTOMATION_TEST(FCbFieldRefTest, FCbFieldTestBase, "System.Core.Serialization.CbFieldRef", CompactBinaryTestFlags)
@@ -1293,8 +1469,8 @@ bool FCbFieldRefIteratorTest::RunTest(const FString& Parameters)
 		const FCbFieldIterator FieldIt(FCbField(View->GetData()), ViewEnd);
 		const FCbFieldRefIterator FieldRefIt(FCbFieldRef(View), ViewEnd);
 
-		TestEqual(TEXT("FCbFieldIterator::GetHash()"), FieldIt.GetHash(), FBlake3::HashBuffer(*View));
-		TestEqual(TEXT("FCbFieldRefIterator::GetHash()"), FieldRefIt.GetHash(), FBlake3::HashBuffer(*View));
+		TestEqual(TEXT("FCbFieldIterator::GetRangeHash()"), FieldIt.GetRangeHash(), FBlake3::HashBuffer(*View));
+		TestEqual(TEXT("FCbFieldRefIterator::GetRangeHash()"), FieldRefIt.GetRangeHash(), FBlake3::HashBuffer(*View));
 
 		TestEqual(TEXT("FCbFieldIterator(EmptyView)"), GetCount(FCbFieldIterator(EmptyView)), 0);
 		TestEqual(TEXT("FCbFieldRefIterator(BufferNullL)"), GetCount(FCbFieldRefIterator(NullBuffer)), 0);
diff --git a/Engine/Source/Runtime/Core/Public/Serialization/CompactBinary.h b/Engine/Source/Runtime/Core/Public/Serialization/CompactBinary.h
--- a/Engine/Source/Runtime/Core/Public/Serialization/CompactBinary.h
+++ b/Engine/Source/Runtime/Core/Public/Serialization/CompactBinary.h
@@ -68,10 +68,12 @@
  */
 
 class FCbArray;
+class FCbField;
 class FCbObject;
 struct FDateTime;
 struct FGuid;
 struct FTimespan;
+template <typename FuncType> class TFunctionRef;
 
 ///////////////////////////////////////////////////////////////////////////////////////////////////////////////////////
 
@@ -253,10 +255,21 @@ public:
 
 	static constexpr inline bool IsDateTime(ECbFieldType Type)   { return GetType(Type) == ECbFieldType::DateTime; }
 	static constexpr inline bool IsTimeSpan(ECbFieldType Type)   { return GetType(Type) == ECbFieldType::TimeSpan; }
+
+	/** Whether the type is or may contain fields of any reference type. */
+	static constexpr inline bool MayContainReferences(ECbFieldType Type)
+	{
+		return IsObject(Type) | IsArray(Type) | IsAnyReference(Type);
+	}
 };
 
 ///////////////////////////////////////////////////////////////////////////////////////////////////////////////////////
 
+/** A reference to a function that is used to visit fields. */
+using FCbFieldVisitor = TFunctionRef<void (FCbField)>;
+
+///////////////////////////////////////////////////////////////////////////////////////////////////////////////////////
+
 /**
  * Iterator for \ref FCbField[Ref] that can operate on any contiguous range of fields.
  *
@@ -339,10 +352,37 @@ public:
 		return !Equals(Other);
 	}
 
-	/** Calculate the hash of the fields by hashing their underlying memory. */
-	inline FBlake3Hash GetHash() const { return FBlake3::HashBuffer(GetFieldRangeView(*this)); }
-	/** Calculate the hash of the fields by hashing their underlying memory. */
-	inline void GetHash(FBlake3& Hash) const { Hash.Update(GetFieldRangeView(*this)); }
+	/** Calculate the hash of the field range by hashing their underlying memory. */
+	inline FBlake3Hash GetRangeHash() const { return FBlake3::HashBuffer(GetFieldRangeView(*this)); }
+	/** Calculate the hash of the field range by hashing their underlying memory. */
+	inline void GetRangeHash(FBlake3& Hash) const { Hash.Update(GetFieldRangeView(*this)); }
+
+	/** Invoke the visitor for every reference or binary reference in the field range. */
+	inline void IterateRangeReferences(FCbFieldVisitor Visitor) const
+	{
+		// Always iterate over non-uniform ranges because we do not know if they contain a reference.
+		if (FCbFieldType::HasFieldType(FieldType::GetType()))
+		{
+			for (TCbFieldIterator It(*this); It; ++It)
+			{
+				if (FCbFieldType::MayContainReferences(It.GetType()))
+				{
+					It.IterateReferences(Visitor);
+				}
+			}
+		}
+		// Only iterate over uniform ranges if the uniform type may contain a reference.
+		else
+		{
+			if (FCbFieldType::MayContainReferences(FieldType::GetType()))
+			{
+				for (TCbFieldIterator It(*this); It; ++It)
+				{
+					It.IterateReferences(Visitor);
+				}
+			}
+		}
+	}
 
 protected:
 	/** Returns the end of the last field, or null for an iterator at the end. */
@@ -445,6 +485,9 @@ public:
 	/** Calculate the hash of the field by hashing its underlying memory. */
 	inline void GetHash(FBlake3& Hash) const { Hash.Update(GetFieldView()); }
 
+	/** Invoke the visitor for every reference or binary reference in the field. */
+	CORE_API void IterateReferences(FCbFieldVisitor Visitor) const;
+
 	/** Access the field as an object. Defaults to an empty object on error. */
 	CORE_API FCbObject AsObject();
 
@@ -723,6 +766,9 @@ public:
 	/** Calculate the hash of the array if serialized by itself with no name. */
 	CORE_API void GetHash(FBlake3& Hash) const;
 
+	/** Invoke the visitor for every reference or binary reference in the array. */
+	inline void IterateReferences(FCbFieldVisitor Visitor) const { CreateIterator().IterateRangeReferences(Visitor); }
+
 protected:
 	/** Returns the type to use when constructing a new array after using CopyTo. */
 	constexpr inline ECbFieldType GetCopyType() const
@@ -815,6 +861,9 @@ public:
 	/** Calculate the hash of the object if serialized by itself with no name. */
 	CORE_API void GetHash(FBlake3& Hash) const;
 
+	/** Invoke the visitor for every reference or binary reference in the object. */
+	inline void IterateReferences(FCbFieldVisitor Visitor) const { CreateIterator().IterateRangeReferences(Visitor); }
+
 protected:
 	/** Returns the type to use when constructing a new object after using CopyTo. */
 	constexpr inline ECbFieldType GetCopyType() const
