Add CA_* macro support for the clang static analyzer in WindowsPlatformCodeAnalysis.h
Fixed the condition in Tuple.h to not include the MSVC static analyzer workaround when using the clang static analyzer.
Added a special flag in UBT to instruct clang to properly set the __clang_analyzer__ macro during static analysis.

#rnx

diff --git a/Engine/Source/Programs/UnrealBuildTool/Platform/Windows/VCToolChain.cs b/Engine/Source/Programs/UnrealBuildTool/Platform/Windows/VCToolChain.cs
--- a/Engine/Source/Programs/UnrealBuildTool/Platform/Windows/VCToolChain.cs
+++ b/Engine/Source/Programs/UnrealBuildTool/Platform/Windows/VCToolChain.cs
@@ -183,6 +183,9 @@ namespace UnrealBuildTool
 					// Needed for some of the C++ checkers.
 					Arguments.Add("-Xclang -analyzer-config -Xclang aggressive-binary-operation-simplification=true");
 
+					// Ensure the compiler sets the __clang_analyzer__ macro correctly.
+					Arguments.Add("-Xclang -setup-static-analyzer");
+
 					// Enable only specific checkers of families of checkers
 					// See https://clang.llvm.org/docs/analyzer/checkers.html for a full list. Or run:
 					//    'clang -Xclang -analyzer-checker-help' 
diff --git a/Engine/Source/Runtime/Core/Public/Templates/Tuple.h b/Engine/Source/Runtime/Core/Public/Templates/Tuple.h
--- a/Engine/Source/Runtime/Core/Public/Templates/Tuple.h
+++ b/Engine/Source/Runtime/Core/Public/Templates/Tuple.h
@@ -32,7 +32,7 @@
 // - Key is a pointer type.
 // - Key and Value are members of a base class.
 // - Dereferencing is done as part of a compound boolean expression (removing '&& Pair.Value != 15' removes the warning)
-#if defined(_MSC_VER) && USING_CODE_ANALYSIS
+#if (defined(_MSC_VER) && !defined(__clang__)) && USING_CODE_ANALYSIS
 	#define UE_TUPLE_STATIC_ANALYSIS_WORKAROUND 1
 #else
 	#define UE_TUPLE_STATIC_ANALYSIS_WORKAROUND 0
diff --git a/Engine/Source/Runtime/Core/Public/Windows/WindowsPlatformCodeAnalysis.h b/Engine/Source/Runtime/Core/Public/Windows/WindowsPlatformCodeAnalysis.h
--- a/Engine/Source/Runtime/Core/Public/Windows/WindowsPlatformCodeAnalysis.h
+++ b/Engine/Source/Runtime/Core/Public/Windows/WindowsPlatformCodeAnalysis.h
@@ -4,7 +4,7 @@
 
 
 // Code analysis features
-#if defined( _PREFAST_ ) || defined( PVS_STUDIO )
+#if defined( _PREFAST_ ) || defined( PVS_STUDIO ) || defined(__clang_analyzer__)
 	#define USING_CODE_ANALYSIS 1
 #else
 	#define USING_CODE_ANALYSIS 0
@@ -35,6 +35,8 @@
 
 #if USING_CODE_ANALYSIS
 
+#if !defined(__clang_analyzer__)
+
 	// Input argument
 	// Example:  void SetValue( CA_IN bool bReadable );
 	#define CA_IN __in
@@ -230,4 +232,53 @@
 	//-V:self:678
 	//-V:AccumulateParentID:678
 	//-V:Resource:623
+
+#else // defined(__clang_analyzer__)
+
+// A fake function marked with noreturn that acts as a marker for CA_ASSUME to ensure the
+// static analyzer doesn't take an analysis path that is assumed not to be navigable.
+__declspec(dllimport, noreturn) void CA_AssumeNoReturn();
+
+// Input argument
+// Example:  void SetValue( CA_IN bool bReadable );
+#define CA_IN
+
+// Output argument
+// Example:  void FillValue( CA_OUT bool& bWriteable );
+#define CA_OUT
+
+// Specifies that a function parameter may only be read from, never written.
+// NOTE: CA_READ_ONLY is inferred automatically if your parameter is has a const qualifier.
+// Example:  void SetValue( CA_READ_ONLY bool bReadable );
+#define CA_READ_ONLY
+
+// Specifies that a function parameter may only be written to, never read.
+// Example:  void FillValue( CA_WRITE_ONLY bool& bWriteable );
+#define CA_WRITE_ONLY
+
+// Incoming pointer parameter must not be NULL and must point to a valid location in memory.
+// Place before a function parameter's type name.
+// Example:  void SetPointer( CA_VALID_POINTER void* Pointer );
+#define CA_VALID_POINTER
+
+// Caller must check the return value.  Place before the return value in a function declaration.
+// Example:  CA_CHECK_RETVAL int32 GetNumber();
+#define CA_CHECK_RETVAL
+
+// Function is expected to never return
+#define CA_NO_RETURN __declspec(noreturn)
+
+// Suppresses a warning for a single occurrence.  Should be used only for code analysis warnings on Windows platform!
+#define CA_SUPPRESS( WarningNumber )
+
+// Tells the code analysis engine to assume the statement to be true.  Useful for suppressing false positive warnings.
+#define CA_ASSUME( Expr )  (__builtin_expect(!bool(Expr), 0) ? CA_AssumeNoReturn() : (void)0)
+
+// Does a simple 'if (Condition)', but disables warnings about using constants in the condition.  Helps with some macro expansions.
+#define CA_CONSTANT_IF(Condition) if (Condition)
+
+
+#endif
+
+
 #endif
