Editgrate 14757227
Fix Resume Log Parsing is not propagated to CIS log interfering with CIS log parser

#rnx

diff --git a/Engine/Source/Runtime/Core/Private/Misc/AutomationTest.cpp b/Engine/Source/Runtime/Core/Private/Misc/AutomationTest.cpp
--- a/Engine/Source/Runtime/Core/Private/Misc/AutomationTest.cpp
+++ b/Engine/Source/Runtime/Core/Private/Misc/AutomationTest.cpp
@@ -826,16 +826,19 @@ bool FAutomationTestFramework::InternalStopTest(FAutomationTestExecutionInfo& Ou
 		UE_LOG(LogAutomationTest, Log, TEXT("%s %s ran in %f"), *CurrentTest->GetBeautifiedTestName(), *Parameters, TimeForTest);
 	}
 
-	// Disassociate the test from the output device and feedback context
-	AutomationTestOutputDevice.SetCurrentAutomationTest(nullptr);
-	AutomationTestMessageFilter.SetCurrentAutomationTest(nullptr);
-
 	// Determine if the test was successful based on three criteria:
 	// 1) Did the test itself report success?
 	// 2) Did any errors occur and were logged by the feedback context during execution?++----
 	// 3) Did we meet any errors that were expected with this test
 	bTestSuccessful = bTestSuccessful && !CurrentTest->HasAnyErrors() && CurrentTest->HasMetExpectedErrors();
 
+	// Re-enable log parsing if it was disabled and empty the expected errors list
+	if (CurrentTest->ExpectedErrors.Num())
+	{
+		GLog->Logf(ELogVerbosity::Display, TEXT("<-- Resume Log Parsing -->"));
+	}
+	CurrentTest->ExpectedErrors.Empty();
+
 	// Set the success state of the test based on the above criteria
 	CurrentTest->SetSuccessState( bTestSuccessful );
 
@@ -845,12 +848,9 @@ bool FAutomationTestFramework::InternalStopTest(FAutomationTestExecutionInfo& Ou
 	// Save off timing for the test
 	OutExecutionInfo.Duration = TimeForTest;
 
-	// Re-enable log parsing if it was disabled and empty the expected errors list
-	if (CurrentTest->ExpectedErrors.Num())
-	{
-		GLog->Logf(ELogVerbosity::Display, TEXT("<-- Resume Log Parsing -->"));
-	}
-	CurrentTest->ExpectedErrors.Empty();
+	// Disassociate the test from the output device and feedback context
+	AutomationTestOutputDevice.SetCurrentAutomationTest(nullptr);
+	AutomationTestMessageFilter.SetCurrentAutomationTest(nullptr);
 
 	// Release pointers to now-invalid data
 	CurrentTest = NULL;
