CompactBinary: Added functions to serialize compact binary to/from archives

#rnx

diff --git a/Engine/Source/Runtime/Core/Private/Serialization/CompactBinary.cpp b/Engine/Source/Runtime/Core/Private/Serialization/CompactBinary.cpp
--- a/Engine/Source/Runtime/Core/Private/Serialization/CompactBinary.cpp
+++ b/Engine/Source/Runtime/Core/Private/Serialization/CompactBinary.cpp
@@ -7,6 +7,7 @@
 #include "Misc/DateTime.h"
 #include "Misc/Guid.h"
 #include "Misc/Timespan.h"
+#include "Serialization/Archive.h"
 #include "Serialization/VarInt.h"
 
 ///////////////////////////////////////////////////////////////////////////////////////////////////////////////////////
@@ -454,6 +455,13 @@ void FCbField::CopyTo(const FMutableMemoryView Buffer) const
 	FMemory::Memcpy(Buffer.GetData(), Source.GetData(), FMath::Min(Buffer.GetSize(), Source.GetSize()));
 }
 
+void FCbField::CopyTo(FArchive& Ar) const
+{
+	check(Ar.IsSaving());
+	const FConstMemoryView Source = GetFieldView();
+	Ar.Serialize(const_cast<void*>(Source.GetData()), static_cast<int64>(Source.GetSize()));
+}
+
 ///////////////////////////////////////////////////////////////////////////////////////////////////////////////////////
 
 FCbArray::FCbArray()
@@ -520,6 +528,15 @@ void FCbArray::CopyTo(const FMutableMemoryView Buffer) const
 	FMemory::Memcpy(BufferBytes, Source.GetData(), FMath::Min(Buffer.GetSize() - 1, Source.GetSize()));
 }
 
+void FCbArray::CopyTo(FArchive& Ar) const
+{
+	check(Ar.IsSaving());
+	uint8 TypeByte = uint8(FCbFieldType::GetType(GetType()));
+	const FConstMemoryView Source = GetPayloadView();
+	Ar.Serialize(&TypeByte, sizeof(TypeByte));
+	Ar.Serialize(const_cast<void*>(Source.GetData()), static_cast<int64>(Source.GetSize()));
+}
+
 ///////////////////////////////////////////////////////////////////////////////////////////////////////////////////////
 
 FCbObject::FCbObject()
@@ -600,113 +617,13 @@ void FCbObject::CopyTo(const FMutableMemoryView Buffer) const
 	FMemory::Memcpy(BufferBytes, Source.GetData(), FMath::Min(Buffer.GetSize() - 1, Source.GetSize()));
 }
 
-///////////////////////////////////////////////////////////////////////////////////////////////////////////////////////
-
-uint64 MeasureCompactBinary(FConstMemoryView View, ECbFieldType Type)
+void FCbObject::CopyTo(FArchive& Ar) const
 {
-	uint64 Size = 0;
-
-	if (FCbFieldType::HasFieldType(Type))
-	{
-		if (View.GetSize() == 0)
-		{
-			return 0;
-		}
-
-		Type = *static_cast<const ECbFieldType*>(View.GetData());
-		View += 1;
-		Size += 1;
-	}
-
-	bool bVariableSize = false;
-	switch (FCbFieldType::GetType(Type))
-	{
-	case ECbFieldType::None:
-	case ECbFieldType::Null:
-		break;
-	case ECbFieldType::Object:
-	case ECbFieldType::UniformObject:
-	case ECbFieldType::Array:
-	case ECbFieldType::UniformArray:
-	case ECbFieldType::Binary:
-	case ECbFieldType::String:
-	case ECbFieldType::IntegerPositive:
-	case ECbFieldType::IntegerNegative:
-		bVariableSize = true;
-		break;
-	case ECbFieldType::Float32:
-		Size += 4;
-		break;
-	case ECbFieldType::Float64:
-		Size += 8;
-		break;
-	case ECbFieldType::BoolFalse:
-	case ECbFieldType::BoolTrue:
-		break;
-	case ECbFieldType::Reference:
-	case ECbFieldType::BinaryReference:
-	case ECbFieldType::Hash:
-		Size += 32;
-		break;
-	case ECbFieldType::Uuid:
-		Size += 16;
-		break;
-	case ECbFieldType::DateTime:
-	case ECbFieldType::TimeSpan:
-		Size += 8;
-		break;
-	default:
-		return 0;
-	}
-
-	if (FCbFieldType::HasFieldName(Type))
-	{
-		if (View.GetSize() == 0 || View.GetSize() < MeasureVarUInt(View.GetData()))
-		{
-			return 0;
-		}
-
-		uint32 NameLenByteCount;
-		const uint64 NameLen = ReadVarUInt(View.GetData(), NameLenByteCount);
-		const uint64 NameSize = NameLen + NameLenByteCount;
-
-		if (bVariableSize && View.GetSize() < NameSize)
-		{
-			return 0;
-		}
-
-		View += NameSize;
-		Size += NameSize;
-	}
-
-	switch (FCbFieldType::GetType(Type))
-	{
-	case ECbFieldType::Object:
-	case ECbFieldType::UniformObject:
-	case ECbFieldType::Array:
-	case ECbFieldType::UniformArray:
-	case ECbFieldType::Binary:
-	case ECbFieldType::String:
-		if (View.GetSize() == 0 || View.GetSize() < MeasureVarUInt(View.GetData()))
-		{
-			return 0;
-		}
-		else
-		{
-			uint32 PayloadSizeByteCount;
-			const uint64 PayloadSize = ReadVarUInt(View.GetData(), PayloadSizeByteCount);
-			return Size + PayloadSize + PayloadSizeByteCount;
-		}
-	case ECbFieldType::IntegerPositive:
-	case ECbFieldType::IntegerNegative:
-		if (View.GetSize() == 0)
-		{
-			return 0;
-		}
-		return Size + MeasureVarUInt(View.GetData());
-	default:
-		return Size;
-	}
+	check(Ar.IsSaving());
+	uint8 TypeByte = uint8(FCbFieldType::GetType(GetType()));
+	const FConstMemoryView Source = GetPayloadView();
+	Ar.Serialize(&TypeByte, sizeof(TypeByte));
+	Ar.Serialize(const_cast<void*>(Source.GetData()), static_cast<int64>(Source.GetSize()));
 }
 
 ///////////////////////////////////////////////////////////////////////////////////////////////////////////////////////
diff --git a/Engine/Source/Runtime/Core/Private/Serialization/CompactBinarySerialization.cpp b/Engine/Source/Runtime/Core/Private/Serialization/CompactBinarySerialization.cpp
new file mode 100644
--- /dev/null
+++ b/Engine/Source/Runtime/Core/Private/Serialization/CompactBinarySerialization.cpp
@@ -0,0 +1,248 @@
+// Copyright Epic Games, Inc. All Rights Reserved.
+
+#include "Serialization/CompactBinarySerialization.h"
+
+#include "Serialization/VarInt.h"
+#include "Templates/IdentityFunctor.h"
+#include "Templates/Invoke.h"
+
+///////////////////////////////////////////////////////////////////////////////////////////////////////////////////////
+
+uint64 MeasureCompactBinary(FConstMemoryView View, ECbFieldType Type)
+{
+	uint64 Size;
+	return TryMeasureCompactBinary(View, Type, Size, Type) ? Size : 0;
+}
+
+///////////////////////////////////////////////////////////////////////////////////////////////////////////////////////
+
+bool TryMeasureCompactBinary(FConstMemoryView View, ECbFieldType& OutType, uint64& OutSize, ECbFieldType Type)
+{
+	uint64 Size = 0;
+
+	if (FCbFieldType::HasFieldType(Type))
+	{
+		if (View.GetSize() == 0)
+		{
+			OutType = ECbFieldType::None;
+			OutSize = 1;
+			return false;
+		}
+
+		Type = *static_cast<const ECbFieldType*>(View.GetData());
+		View += 1;
+		Size += 1;
+	}
+
+	bool bDynamicSize = false;
+	uint64 FixedSize = 0;
+	switch (FCbFieldType::GetType(Type))
+	{
+	case ECbFieldType::Null:
+		break;
+	case ECbFieldType::Object:
+	case ECbFieldType::UniformObject:
+	case ECbFieldType::Array:
+	case ECbFieldType::UniformArray:
+	case ECbFieldType::Binary:
+	case ECbFieldType::String:
+	case ECbFieldType::IntegerPositive:
+	case ECbFieldType::IntegerNegative:
+		bDynamicSize = true;
+		break;
+	case ECbFieldType::Float32:
+		FixedSize = 4;
+		break;
+	case ECbFieldType::Float64:
+		FixedSize = 8;
+		break;
+	case ECbFieldType::BoolFalse:
+	case ECbFieldType::BoolTrue:
+		break;
+	case ECbFieldType::Reference:
+	case ECbFieldType::BinaryReference:
+	case ECbFieldType::Hash:
+		FixedSize = 32;
+		break;
+	case ECbFieldType::Uuid:
+		FixedSize = 16;
+		break;
+	case ECbFieldType::DateTime:
+	case ECbFieldType::TimeSpan:
+		FixedSize = 8;
+		break;
+	case ECbFieldType::None:
+	default:
+		OutType = ECbFieldType::None;
+		OutSize = 0;
+		return false;
+	}
+
+	OutType = Type;
+
+	if (FCbFieldType::HasFieldName(Type))
+	{
+		if (View.GetSize() == 0)
+		{
+			OutSize = Size + 1;
+			return false;
+		}
+
+		uint32 NameLenByteCount = MeasureVarUInt(View.GetData());
+		if (View.GetSize() < NameLenByteCount)
+		{
+			OutSize = Size + NameLenByteCount;
+			return false;
+		}
+
+		const uint64 NameLen = ReadVarUInt(View.GetData(), NameLenByteCount);
+		const uint64 NameSize = NameLen + NameLenByteCount;
+
+		if (bDynamicSize && View.GetSize() < NameSize)
+		{
+			OutSize = Size + NameSize;
+			return false;
+		}
+
+		View += NameSize;
+		Size += NameSize;
+	}
+
+	switch (FCbFieldType::GetType(Type))
+	{
+	case ECbFieldType::Object:
+	case ECbFieldType::UniformObject:
+	case ECbFieldType::Array:
+	case ECbFieldType::UniformArray:
+	case ECbFieldType::Binary:
+	case ECbFieldType::String:
+		if (View.GetSize() == 0)
+		{
+			OutSize = Size + 1;
+			return false;
+		}
+		else
+		{
+			uint32 PayloadSizeByteCount = MeasureVarUInt(View.GetData());
+			if (View.GetSize() < PayloadSizeByteCount)
+			{
+				OutSize = Size + PayloadSizeByteCount;
+				return false;
+			}
+			const uint64 PayloadSize = ReadVarUInt(View.GetData(), PayloadSizeByteCount);
+			OutSize = Size + PayloadSize + PayloadSizeByteCount;
+			return true;
+		}
+	case ECbFieldType::IntegerPositive:
+	case ECbFieldType::IntegerNegative:
+		if (View.GetSize() == 0)
+		{
+			OutSize = Size + 1;
+			return false;
+		}
+		OutSize = Size + MeasureVarUInt(View.GetData());
+		return true;
+	default:
+		OutSize = Size + FixedSize;
+		return true;
+	}
+}
+
+///////////////////////////////////////////////////////////////////////////////////////////////////////////////////////
+
+FCbFieldRef LoadCompactBinary(FArchive& Ar, FCbBufferAllocator Allocator)
+{
+	TArray<uint8, TInlineAllocator<64>> HeaderBytes;
+	ECbFieldType FieldType;
+	uint64 FieldSize = 1;
+
+	// Read in small increments until the total field size is known, to avoid reading too far.
+	for (;;)
+	{
+		const int32 ReadSize = int32(FieldSize - HeaderBytes.Num());
+		const int32 ReadOffset = HeaderBytes.AddUninitialized(ReadSize);
+		Ar.Serialize(HeaderBytes.GetData() + ReadOffset, ReadSize);
+		if (TryMeasureCompactBinary(MakeMemoryView(HeaderBytes), FieldType, FieldSize))
+		{
+			break;
+		}
+		checkf(FieldSize > 0, TEXT("Failed to load from invalid compact binary data."), FieldSize);
+	}
+
+	// Allocate the buffer, copy the header, and read the remainder of the field.
+	FSharedBufferPtr Buffer = Allocator(FieldType, FieldSize);
+	FMutableMemoryView View = Buffer->GetView();
+	FMemory::Memcpy(View.GetData(), HeaderBytes.GetData(), HeaderBytes.Num());
+	View += HeaderBytes.Num();
+	if (!View.IsEmpty())
+	{
+		Ar.Serialize(View.GetData(), static_cast<int64>(View.GetSize()));
+	}
+	return FCbFieldRef(MoveTemp(Buffer));
+}
+
+///////////////////////////////////////////////////////////////////////////////////////////////////////////////////////
+
+template <typename T>
+static void SaveCompactBinaryToArchive(FArchive& Ar, const T& Value)
+{
+	check(Ar.IsSaving());
+	struct FCopy : public T
+	{
+		using T::GetCopyType;
+		using T::CopyTo;
+	};
+	const FCopy& ValueCopy = static_cast<const FCopy&>(Value);
+	const ECbFieldType CopyType = ValueCopy.GetCopyType();
+	if (!EnumHasAnyFlags(CopyType, ECbFieldType::HasFieldType))
+	{
+		uint8 TypeByte = uint8(CopyType);
+		Ar.Serialize(&TypeByte, sizeof(TypeByte));
+	}
+	ValueCopy.CopyTo(Ar);
+}
+
+void SaveCompactBinary(FArchive& Ar, const FCbField& Field) { SaveCompactBinaryToArchive(Ar, Field); }
+void SaveCompactBinary(FArchive& Ar, const FCbArray& Array) { SaveCompactBinaryToArchive(Ar, Array); }
+void SaveCompactBinary(FArchive& Ar, const FCbObject& Object) { SaveCompactBinaryToArchive(Ar, Object); }
+
+///////////////////////////////////////////////////////////////////////////////////////////////////////////////////////
+
+template <typename T, typename ConvertType>
+static FArchive& SerializeCompactBinary(FArchive& Ar, T& Value, ConvertType&& Convert)
+{
+	if (Ar.IsLoading())
+	{
+		Value = Invoke(Forward<ConvertType>(Convert),
+			LoadCompactBinary(Ar, [](ECbFieldType Type, uint64 Size)
+			{
+				return FSharedBuffer::Alloc(Size);
+			}));
+	}
+	else if (Ar.IsSaving())
+	{
+		SaveCompactBinary(Ar, Value);
+	}
+	else
+	{
+		checkNoEntry();
+	}
+	return Ar;
+}
+
+FArchive& operator<<(FArchive& Ar, FCbFieldRef& Field)
+{
+	return SerializeCompactBinary(Ar, Field, FIdentityFunctor());
+}
+
+FArchive& operator<<(FArchive& Ar, FCbArrayRef& Array)
+{
+	return SerializeCompactBinary(Ar, Array, [](FCbFieldRef&& Field) { return MoveTemp(Field).AsArrayRef(); });
+}
+
+FArchive& operator<<(FArchive& Ar, FCbObjectRef& Object)
+{
+	return SerializeCompactBinary(Ar, Object, [](FCbFieldRef&& Field) { return MoveTemp(Field).AsObjectRef(); });
+}
+
+///////////////////////////////////////////////////////////////////////////////////////////////////////////////////////
diff --git a/Engine/Source/Runtime/Core/Private/Serialization/CompactBinaryWriter.cpp b/Engine/Source/Runtime/Core/Private/Serialization/CompactBinaryWriter.cpp
--- a/Engine/Source/Runtime/Core/Private/Serialization/CompactBinaryWriter.cpp
+++ b/Engine/Source/Runtime/Core/Private/Serialization/CompactBinaryWriter.cpp
@@ -8,6 +8,7 @@
 #include "Misc/DateTime.h"
 #include "Misc/Guid.h"
 #include "Misc/Timespan.h"
+#include "Serialization/CompactBinarySerialization.h"
 #include "Serialization/VarInt.h"
 
 enum class FCbWriter::EStateFlags : uint8
diff --git a/Engine/Source/Runtime/Core/Private/Tests/Serialization/CompactBinarySerializationTest.cpp b/Engine/Source/Runtime/Core/Private/Tests/Serialization/CompactBinarySerializationTest.cpp
new file mode 100644
--- /dev/null
+++ b/Engine/Source/Runtime/Core/Private/Tests/Serialization/CompactBinarySerializationTest.cpp
@@ -0,0 +1,246 @@
+// Copyright Epic Games, Inc. All Rights Reserved.
+
+#include "Serialization/CompactBinarySerialization.h"
+
+#include "Misc/AutomationTest.h"
+#include "Serialization/BufferArchive.h"
+#include "Serialization/MemoryReader.h"
+
+#if WITH_DEV_AUTOMATION_TESTS
+
+///////////////////////////////////////////////////////////////////////////////////////////////////////////////////////
+
+static constexpr EAutomationTestFlags::Type CompactBinarySerializationTestFlags =
+	EAutomationTestFlags::Type(EAutomationTestFlags::ApplicationContextMask | EAutomationTestFlags::SmokeFilter);
+
+///////////////////////////////////////////////////////////////////////////////////////////////////////////////////////
+
+IMPLEMENT_SIMPLE_AUTOMATION_TEST(FCbMeasureTest, "System.Core.Serialization.MeasureCompactBinary", CompactBinarySerializationTestFlags)
+bool FCbMeasureTest::RunTest(const FString& Parameters)
+{
+	const auto TestMeasure = [this](
+		const TCHAR* Test,
+		std::initializer_list<uint8> Data,
+		bool bExpectedValue,
+		uint64 ExpectedSize,
+		ECbFieldType ExpectedType,
+		ECbFieldType ExternalType = ECbFieldType::HasFieldType) -> void
+	{
+		ECbFieldType ActualType = ECbFieldType(255);
+		uint64 ActualSize = ~uint64(0);
+		const bool bActualValue = TryMeasureCompactBinary(MakeMemoryView(Data), ActualType, ActualSize, ExternalType);
+		TestEqual(FString::Printf(TEXT("TryMeasureCompactBinary(%s)"), Test), bActualValue, bExpectedValue);
+		TestEqual(FString::Printf(TEXT("TryMeasureCompactBinary(%s)->Type"), Test), ActualType, ExpectedType);
+		TestEqual(FString::Printf(TEXT("TryMeasureCompactBinary(%s)->Size"), Test), ActualSize, ExpectedSize);
+		if (ActualSize == ExpectedSize)
+		{
+			const uint64 MeasureSize = MeasureCompactBinary(MakeMemoryView(Data), ExternalType);
+			TestEqual(FString::Printf(TEXT("MeasureCompactBinary(%s)"), Test), MeasureSize, bExpectedValue ? ExpectedSize : 0);
+		}
+	};
+
+	using EType = ECbFieldType;
+
+	TestMeasure(TEXT("Empty"), {}, false, 1, EType::None);
+
+	TestMeasure(TEXT("None"), {uint8(EType::None)}, false, 0, EType::None);
+	TestMeasure(TEXT("None, NoType"), {}, false, 0, EType::None, EType::None);
+	TestMeasure(TEXT("None, NoType, Name"), {}, false, 0, EType::None, EType::None | EType::HasFieldName);
+
+	TestMeasure(TEXT("Null"), {uint8(EType::Null)}, true, 1, EType::Null);
+	TestMeasure(TEXT("Null, NameSize1B"), {uint8(EType::Null | EType::HasFieldName), 30}, true, 32, EType::Null | EType::HasFieldName);
+	TestMeasure(TEXT("Null, NameSize2B"), {uint8(EType::Null | EType::HasFieldName), 0x80, 0x80}, true, 131, EType::Null | EType::HasFieldName);
+	TestMeasure(TEXT("Null, NameSize2BShort"), {uint8(EType::Null | EType::HasFieldName), 0x80}, false, 3, EType::Null | EType::HasFieldName);
+	TestMeasure(TEXT("Null, NameSize3BShort"), {uint8(EType::Null | EType::HasFieldName), 0xc0}, false, 4, EType::Null | EType::HasFieldName);
+	TestMeasure(TEXT("Null, MissingName"), {uint8(EType::Null | EType::HasFieldName)}, false, 2, EType::Null | EType::HasFieldName);
+
+	TestMeasure(TEXT("Null, NoType"), {}, true, 0, EType::Null, EType::Null);
+	TestMeasure(TEXT("Null, NoType, NameSize1B"), {30}, true, 31, EType::Null | EType::HasFieldName, EType::Null | EType::HasFieldName);
+	TestMeasure(TEXT("Null, NoType, NameSize2B"), {0x80, 0x80}, true, 130, EType::Null | EType::HasFieldName, EType::Null | EType::HasFieldName);
+	TestMeasure(TEXT("Null, NoType, NameSize2BShort"), {0x80}, false, 2, EType::Null | EType::HasFieldName, EType::Null | EType::HasFieldName);
+	TestMeasure(TEXT("Null, NoType, NameSize3BShort"), {0xc0}, false, 3, EType::Null | EType::HasFieldName, EType::Null | EType::HasFieldName);
+	TestMeasure(TEXT("Null, NoType, MissingName"), {}, false, 1, EType::Null | EType::HasFieldName, EType::Null | EType::HasFieldName);
+
+	TestMeasure(TEXT("Object, NoSize"), {uint8(EType::Object)}, false, 2, EType::Object);
+	TestMeasure(TEXT("Object, Size1B"), {uint8(EType::Object), 30}, true, 32, EType::Object);
+	TestMeasure(TEXT("UniformObject, NoSize"), {uint8(EType::UniformObject)}, false, 2, EType::UniformObject);
+	TestMeasure(TEXT("UniformObject, Size1B"), {uint8(EType::UniformObject), 30}, true, 32, EType::UniformObject);
+
+	TestMeasure(TEXT("Array, NoSize"), {uint8(EType::Array)}, false, 2, EType::Array);
+	TestMeasure(TEXT("Array, Size1B"), {uint8(EType::Array), 30}, true, 32, EType::Array);
+	TestMeasure(TEXT("UniformArray, NoSize"), {uint8(EType::UniformArray)}, false, 2, EType::UniformArray);
+	TestMeasure(TEXT("UniformArray, Size1B"), {uint8(EType::UniformArray), 30}, true, 32, EType::UniformArray);
+
+	TestMeasure(TEXT("Binary, NoSize"), {uint8(EType::Binary)}, false, 2, EType::Binary);
+	TestMeasure(TEXT("Binary, Size1B"), {uint8(EType::Binary), 30}, true, 32, EType::Binary);
+
+	TestMeasure(TEXT("String, NoSize"), {uint8(EType::String)}, false, 2, EType::String);
+	TestMeasure(TEXT("String, Size1B"), {uint8(EType::String), 30}, true, 32, EType::String);
+	TestMeasure(TEXT("String, Size2B"), {uint8(EType::String), 0x80, 0x80}, true, 131, EType::String);
+	TestMeasure(TEXT("String, Size2BShort"), {uint8(EType::String), 0x80}, false, 3, EType::String);
+	TestMeasure(TEXT("String, Size3BShort"), {uint8(EType::String), 0xc0}, false, 4, EType::String);
+
+	TestMeasure(TEXT("String, NameNoSize"), {uint8(EType::String | EType::HasFieldName)}, false, 2, EType::String | EType::HasFieldName);
+	TestMeasure(TEXT("String, NameSize1BMissingText"), {uint8(EType::String | EType::HasFieldName), 1}, false, 3, EType::String | EType::HasFieldName);
+	TestMeasure(TEXT("String, NameSize1B, NoSize"), {uint8(EType::String | EType::HasFieldName), 1, 'A'}, false, 4, EType::String | EType::HasFieldName);
+	TestMeasure(TEXT("String, NameSize1B, Size1B"), {uint8(EType::String | EType::HasFieldName), 1, 'A', 30}, true, 34, EType::String | EType::HasFieldName);
+	TestMeasure(TEXT("String, NameSize1B, Size2B"), {uint8(EType::String | EType::HasFieldName), 1, 'A', 0x80, 0x80}, true, 133, EType::String | EType::HasFieldName);
+	TestMeasure(TEXT("String, NameSize1B, Size2BShort"), {uint8(EType::String | EType::HasFieldName), 1, 'A', 0x80}, false, 5, EType::String | EType::HasFieldName);
+	TestMeasure(TEXT("String, NameSize1B, Size3BShort"), {uint8(EType::String | EType::HasFieldName), 1, 'A', 0xc0}, false, 6, EType::String | EType::HasFieldName);
+	TestMeasure(TEXT("String, NameSize2BShort"), {uint8(EType::String | EType::HasFieldName), 0x80}, false, 3, EType::String | EType::HasFieldName);
+	TestMeasure(TEXT("String, NameSize3BShort"), {uint8(EType::String | EType::HasFieldName), 0xc0}, false, 4, EType::String | EType::HasFieldName);
+
+	TestMeasure(TEXT("IntegerPositive, NoValue"), {uint8(EType::IntegerPositive)}, false, 2, EType::IntegerPositive);
+	TestMeasure(TEXT("IntegerPositive, Value1B"), {uint8(EType::IntegerPositive), 0x7f}, true, 2, EType::IntegerPositive);
+	TestMeasure(TEXT("IntegerPositive, Value2B"), {uint8(EType::IntegerPositive), 0x80}, true, 3, EType::IntegerPositive);
+
+	TestMeasure(TEXT("IntegerNegative, NoValue"), {uint8(EType::IntegerNegative)}, false, 2, EType::IntegerNegative);
+	TestMeasure(TEXT("IntegerNegative, Value1B"), {uint8(EType::IntegerNegative), 0x7f}, true, 2, EType::IntegerNegative);
+	TestMeasure(TEXT("IntegerNegative, Value2B"), {uint8(EType::IntegerNegative), 0x80}, true, 3, EType::IntegerNegative);
+
+	TestMeasure(TEXT("Float32"), {uint8(EType::Float32)}, true, 5, EType::Float32);
+	TestMeasure(TEXT("Float32, NameSize1B"), {uint8(EType::Float32 | EType::HasFieldName), 30}, true, 36, EType::Float32 | EType::HasFieldName);
+	TestMeasure(TEXT("Float32, NameSize2B"), {uint8(EType::Float32 | EType::HasFieldName), 0x80, 0x80}, true, 135, EType::Float32 | EType::HasFieldName);
+	TestMeasure(TEXT("Float32, NameSize2BShort"), {uint8(EType::Float32 | EType::HasFieldName), 0x80}, false, 3, EType::Float32 | EType::HasFieldName);
+	TestMeasure(TEXT("Float32, MissingName"), {uint8(EType::Float32 | EType::HasFieldName)}, false, 2, EType::Float32 | EType::HasFieldName);
+
+	TestMeasure(TEXT("Float32, NoType"), {}, true, 4, EType::Float32, EType::Float32);
+	TestMeasure(TEXT("Float32, NoType, NameSize1B"), {30}, true, 35, EType::Float32 | EType::HasFieldName, EType::Float32 | EType::HasFieldName);
+	TestMeasure(TEXT("Float32, NoType, NameSize2B"), {0x80, 0x80}, true, 134, EType::Float32 | EType::HasFieldName, EType::Float32 | EType::HasFieldName);
+	TestMeasure(TEXT("Float32, NoType, NameSize2BShort"), {0x80}, false, 2, EType::Float32 | EType::HasFieldName, EType::Float32 | EType::HasFieldName);
+	TestMeasure(TEXT("Float32, NoType, MissingName"), {}, false, 1, EType::Float32 | EType::HasFieldName, EType::Float32 | EType::HasFieldName);
+
+	TestMeasure(TEXT("Float64"), {uint8(EType::Float64)}, true, 9, EType::Float64);
+
+	TestMeasure(TEXT("BoolFalse"), {uint8(EType::BoolFalse)}, true, 1, EType::BoolFalse);
+	TestMeasure(TEXT("BoolTrue"), {uint8(EType::BoolTrue)}, true, 1, EType::BoolTrue);
+
+	TestMeasure(TEXT("Reference"), {uint8(EType::Reference)}, true, 33, EType::Reference);
+	TestMeasure(TEXT("BinaryReference"), {uint8(EType::BinaryReference)}, true, 33, EType::BinaryReference);
+
+	TestMeasure(TEXT("Hash"), {uint8(EType::Hash)}, true, 33, EType::Hash);
+	TestMeasure(TEXT("Uuid"), {uint8(EType::Uuid)}, true, 17, EType::Uuid);
+
+	TestMeasure(TEXT("DateTime"), {uint8(EType::DateTime)}, true, 9, EType::DateTime);
+	TestMeasure(TEXT("TimeSpan"), {uint8(EType::TimeSpan)}, true, 9, EType::TimeSpan);
+
+	return true;
+}
+
+///////////////////////////////////////////////////////////////////////////////////////////////////////////////////////
+
+IMPLEMENT_SIMPLE_AUTOMATION_TEST(FCbSaveTest, "System.Core.Serialization.SaveCompactBinary", CompactBinarySerializationTestFlags)
+bool FCbSaveTest::RunTest(const FString& Parameters)
+{
+	const auto TestSave = [this](const TCHAR* Test, auto Value, FConstMemoryView ExpectedData)
+	{
+		{
+			FBufferArchive WriteAr;
+			SaveCompactBinary(WriteAr, Value);
+			if (TestEqual(FString::Printf(TEXT("SaveCompactBinary(%s)->Size"), Test), uint64(WriteAr.Num()), ExpectedData.GetSize()) &&
+				TestTrue(FString::Printf(TEXT("SaveCompactBinary(%s)->EqualBytes"), Test), ExpectedData.EqualBytes(MakeMemoryView(WriteAr))))
+			{
+				struct FView : public FCbFieldRef
+				{
+					using FCbFieldRef::GetFieldView;
+				};
+				FMemoryReader ReadAr(WriteAr);
+				FCbFieldRef Field = LoadCompactBinary(ReadAr, [](ECbFieldType Type, uint64 Size) { return FSharedBuffer::Alloc(Size); });
+				TestTrue(FString::Printf(TEXT("LoadCompactBinary(%s)->EqualBytes"), Test), ExpectedData.EqualBytes(static_cast<FView&>(Field).GetFieldView()));
+			}
+		}
+		{
+			FBufferArchive WriteAr;
+			WriteAr << Value;
+			if (TestEqual(FString::Printf(TEXT("Ar << CompactBinary Save(%s)->Size"), Test), uint64(WriteAr.Num()), ExpectedData.GetSize()) &&
+				TestTrue(FString::Printf(TEXT("Ar << CompactBinary Save(%s)->EqualBytes"), Test), ExpectedData.EqualBytes(MakeMemoryView(WriteAr))))
+			{
+				struct FView : public decltype(Value)
+				{
+					using decltype(Value)::GetFieldView;
+				};
+				FMemoryReader ReadAr(WriteAr);
+				ReadAr << Value;
+				TestTrue(FString::Printf(TEXT("Ar << CompactBinary Load(%s)->EqualBytes"), Test), ExpectedData.EqualBytes(static_cast<FView&>(Value).GetFieldView()));
+			}
+		}
+	};
+
+	// Field
+	{
+		const uint8 Payload[] = { uint8(ECbFieldType::IntegerPositive), 42 };
+		TestSave(TEXT("Field"), FCbFieldRef(FCbField(Payload), FSharedBufferConstPtr()), MakeMemoryView(Payload));
+	}
+	{
+		const uint8 Payload[] = { uint8(ECbFieldType::IntegerPositive | ECbFieldType::HasFieldName), 1, 'A', 42 };
+		TestSave(TEXT("Field, Name"), FCbFieldRef(FCbField(Payload), FSharedBufferConstPtr()), MakeMemoryView(Payload));
+	}
+	{
+		const uint8 Payload[] = { 42 };
+		TestSave(TEXT("Field, NoType"), FCbFieldRef(FCbField(Payload, ECbFieldType::IntegerPositive), FSharedBufferConstPtr()),
+			MakeMemoryView<uint8>({ uint8(ECbFieldType::IntegerPositive), 42 }));
+	}
+	{
+		const uint8 Payload[] = { 1, 'I', 42 };
+		TestSave(TEXT("Field, NoType, Name"), FCbFieldRef(FCbField(Payload, ECbFieldType::IntegerPositive | ECbFieldType::HasFieldName), FSharedBufferConstPtr()),
+			MakeMemoryView<uint8>({ uint8(ECbFieldType::IntegerPositive | ECbFieldType::HasFieldName), 1, 'I', 42 }));
+	}
+
+	// Array
+	{
+		constexpr uint8 IntType = uint8(ECbFieldType::IntegerPositive);
+		const uint8 Payload[] = { uint8(ECbFieldType::UniformArray), 5, 3, IntType, 1, 2, 3 };
+		TestSave(TEXT("Array"), FCbArrayRef(FCbArray(Payload), FSharedBufferConstPtr()),
+			MakeMemoryView(Payload));
+	}
+	{
+		constexpr uint8 IntType = uint8(ECbFieldType::IntegerPositive);
+		const uint8 Payload[] = { uint8(ECbFieldType::UniformArray | ECbFieldType::HasFieldName), 1, 'A', 5, 3, IntType, 1, 2, 3 };
+		TestSave(TEXT("Array, Name"), FCbArrayRef(FCbArray(Payload), FSharedBufferConstPtr()),
+			MakeMemoryView<uint8>({ uint8(ECbFieldType::UniformArray), 5, 3, IntType, 1, 2, 3 }));
+	}
+	{
+		constexpr uint8 IntType = uint8(ECbFieldType::IntegerPositive);
+		const uint8 Payload[] = { 5, 3, IntType, 1, 2, 3 };
+		TestSave(TEXT("Array, NoType"), FCbArrayRef(FCbArray(Payload, ECbFieldType::UniformArray), FSharedBufferConstPtr()),
+			MakeMemoryView<uint8>({ uint8(ECbFieldType::UniformArray), 5, 3, IntType, 1, 2, 3 }));
+	}
+	{
+		constexpr uint8 IntType = uint8(ECbFieldType::IntegerPositive);
+		const uint8 Payload[] = { 1, 'A', 5, 3, IntType, 1, 2, 3 };
+		TestSave(TEXT("Array, NoType, Name"), FCbArrayRef(FCbArray(Payload, ECbFieldType::UniformArray | ECbFieldType::HasFieldName), FSharedBufferConstPtr()),
+			MakeMemoryView<uint8>({ uint8(ECbFieldType::UniformArray), 5, 3, IntType, 1, 2, 3 }));
+	}
+
+	// Object
+	{
+		constexpr uint8 IntType = uint8(ECbFieldType::HasFieldName | ECbFieldType::IntegerPositive);
+		const uint8 Payload[] = { uint8(ECbFieldType::UniformObject), 10, IntType, 1, 'A', 1, 1, 'B', 2, 1, 'C', 3 };
+		TestSave(TEXT("Object"), FCbObjectRef(FCbObject(Payload), FSharedBufferConstPtr()),
+			MakeMemoryView(Payload));
+	}
+	{
+		constexpr uint8 IntType = uint8(ECbFieldType::HasFieldName | ECbFieldType::IntegerPositive);
+		const uint8 Payload[] = { uint8(ECbFieldType::UniformObject | ECbFieldType::HasFieldName), 1, 'O', 10, IntType, 1, 'A', 1, 1, 'B', 2, 1, 'C', 3 };
+		TestSave(TEXT("Object, Name"), FCbObjectRef(FCbObject(Payload), FSharedBufferConstPtr()),
+			MakeMemoryView<uint8>({ uint8(ECbFieldType::UniformObject), 10, IntType, 1, 'A', 1, 1, 'B', 2, 1, 'C', 3 }));
+	}
+	{
+		constexpr uint8 IntType = uint8(ECbFieldType::HasFieldName | ECbFieldType::IntegerPositive);
+		const uint8 Payload[] = { 10, IntType, 1, 'A', 1, 1, 'B', 2, 1, 'C', 3 };
+		TestSave(TEXT("Object, NoType"), FCbObjectRef(FCbObject(Payload, ECbFieldType::UniformObject), FSharedBufferConstPtr()),
+			MakeMemoryView<uint8>({ uint8(ECbFieldType::UniformObject), 10, IntType, 1, 'A', 1, 1, 'B', 2, 1, 'C', 3 }));
+	}
+	{
+		constexpr uint8 IntType = uint8(ECbFieldType::HasFieldName | ECbFieldType::IntegerPositive);
+		const uint8 Payload[] = { 1, 'O', 10, IntType, 1, 'A', 1, 1, 'B', 2, 1, 'C', 3 };
+		TestSave(TEXT("Object, NoType, Name"), FCbObjectRef(FCbObject(Payload, ECbFieldType::UniformObject | ECbFieldType::HasFieldName), FSharedBufferConstPtr()),
+			MakeMemoryView<uint8>({ uint8(ECbFieldType::UniformObject), 10, IntType, 1, 'A', 1, 1, 'B', 2, 1, 'C', 3 }));
+	}
+
+	return true;
+}
+
+///////////////////////////////////////////////////////////////////////////////////////////////////////////////////////
+
+#endif // WITH_DEV_AUTOMATION_TESTS
diff --git a/Engine/Source/Runtime/Core/Private/Tests/Serialization/CompactBinaryTest.cpp b/Engine/Source/Runtime/Core/Private/Tests/Serialization/CompactBinaryTest.cpp
--- a/Engine/Source/Runtime/Core/Private/Tests/Serialization/CompactBinaryTest.cpp
+++ b/Engine/Source/Runtime/Core/Private/Tests/Serialization/CompactBinaryTest.cpp
@@ -1525,87 +1525,6 @@ bool FCbFieldRefIteratorTest::RunTest(const FString& Parameters)
 
 ///////////////////////////////////////////////////////////////////////////////////////////////////////////////////////
 
-IMPLEMENT_SIMPLE_AUTOMATION_TEST(FCbMeasureTest, "System.Core.Serialization.MeasureCompactBinary", CompactBinaryTestFlags)
-bool FCbMeasureTest::RunTest(const FString& Parameters)
-{
-	auto Measure = [this](std::initializer_list<uint8> Data, ECbFieldType Type = ECbFieldType::HasFieldType) -> uint64
-	{
-		return MeasureCompactBinary(MakeMemoryView(Data), Type);
-	};
-
-	TestEqual(TEXT("MeasureCompactBinary(Empty)"), Measure({}), uint64(0));
-
-	TestEqual(TEXT("MeasureCompactBinary(Null, NoType)"), Measure({}, ECbFieldType::Null), uint64(0));
-	TestEqual(TEXT("MeasureCompactBinary(Null, NameSize1B, NoType)"), Measure({30}, ECbFieldType::Null | ECbFieldType::HasFieldName), uint64(31));
-	TestEqual(TEXT("MeasureCompactBinary(Null, NameSize2B, NoType)"), Measure({0x80, 0x80}, ECbFieldType::Null | ECbFieldType::HasFieldName), uint64(130));
-	TestEqual(TEXT("MeasureCompactBinary(Null, NameSize2BShort, NoType)"), Measure({0x80}, ECbFieldType::Null | ECbFieldType::HasFieldName), uint64(0));
-	TestEqual(TEXT("MeasureCompactBinary(Null, MissingName, NoType)"), Measure({}, ECbFieldType::Null | ECbFieldType::HasFieldName), uint64(0));
-
-	TestEqual(TEXT("MeasureCompactBinary(Null)"), Measure({uint8(ECbFieldType::Null)}), uint64(1));
-	TestEqual(TEXT("MeasureCompactBinary(Null, NameSize1B)"), Measure({uint8(ECbFieldType::Null | ECbFieldType::HasFieldName), 30}), uint64(32));
-	TestEqual(TEXT("MeasureCompactBinary(Null, NameSize2B)"), Measure({uint8(ECbFieldType::Null | ECbFieldType::HasFieldName), 0x80, 0x80}), uint64(131));
-	TestEqual(TEXT("MeasureCompactBinary(Null, NameSize2BShort)"), Measure({uint8(ECbFieldType::Null | ECbFieldType::HasFieldName), 0x80}), uint64(0));
-	TestEqual(TEXT("MeasureCompactBinary(Null, MissingName)"), Measure({uint8(ECbFieldType::Null | ECbFieldType::HasFieldName)}), uint64(0));
-
-	TestEqual(TEXT("MeasureCompactBinary(Object, NoSize)"), Measure({uint8(ECbFieldType::Object)}), uint64(0));
-	TestEqual(TEXT("MeasureCompactBinary(Object, Size1B)"), Measure({uint8(ECbFieldType::Object), 30}), uint64(32));
-	TestEqual(TEXT("MeasureCompactBinary(UniformObject, NoSize)"), Measure({uint8(ECbFieldType::UniformObject)}), uint64(0));
-	TestEqual(TEXT("MeasureCompactBinary(UniformObject, Size1B)"), Measure({uint8(ECbFieldType::UniformObject), 30}), uint64(32));
-
-	TestEqual(TEXT("MeasureCompactBinary(Array, NoSize)"), Measure({uint8(ECbFieldType::Array)}), uint64(0));
-	TestEqual(TEXT("MeasureCompactBinary(Array, Size1B)"), Measure({uint8(ECbFieldType::Array), 30}), uint64(32));
-	TestEqual(TEXT("MeasureCompactBinary(UniformArray, NoSize)"), Measure({uint8(ECbFieldType::UniformArray)}), uint64(0));
-	TestEqual(TEXT("MeasureCompactBinary(UniformArray, Size1B)"), Measure({uint8(ECbFieldType::UniformArray), 30}), uint64(32));
-
-	TestEqual(TEXT("MeasureCompactBinary(Binary, NoSize)"), Measure({uint8(ECbFieldType::Binary)}), uint64(0));
-	TestEqual(TEXT("MeasureCompactBinary(Binary, Size1B)"), Measure({uint8(ECbFieldType::Binary), 30}), uint64(32));
-
-	TestEqual(TEXT("MeasureCompactBinary(String, NoSize)"), Measure({uint8(ECbFieldType::String)}), uint64(0));
-	TestEqual(TEXT("MeasureCompactBinary(String, Size1B)"), Measure({uint8(ECbFieldType::String), 30}), uint64(32));
-	TestEqual(TEXT("MeasureCompactBinary(String, Size2B)"), Measure({uint8(ECbFieldType::String), 0x80, 0x80}), uint64(131));
-	TestEqual(TEXT("MeasureCompactBinary(String, Size2BShort)"), Measure({uint8(ECbFieldType::String), 0x80}), uint64(0));
-	TestEqual(TEXT("MeasureCompactBinary(String, MissingNameSize)"), Measure({uint8(ECbFieldType::String | ECbFieldType::HasFieldName)}), uint64(0));
-	TestEqual(TEXT("MeasureCompactBinary(String, MissingName)"), Measure({uint8(ECbFieldType::String | ECbFieldType::HasFieldName), 1}), uint64(0));
-
-	TestEqual(TEXT("MeasureCompactBinary(IntegerPositive, NoValue)"), Measure({uint8(ECbFieldType::IntegerPositive)}), uint64(0));
-	TestEqual(TEXT("MeasureCompactBinary(IntegerPositive, Value1B)"), Measure({uint8(ECbFieldType::IntegerPositive), 0x7f}), uint64(2));
-	TestEqual(TEXT("MeasureCompactBinary(IntegerPositive, Value2B)"), Measure({uint8(ECbFieldType::IntegerPositive), 0x80}), uint64(3));
-
-	TestEqual(TEXT("MeasureCompactBinary(IntegerNegative, NoValue)"), Measure({uint8(ECbFieldType::IntegerNegative)}), uint64(0));
-	TestEqual(TEXT("MeasureCompactBinary(IntegerNegative, Value1B)"), Measure({uint8(ECbFieldType::IntegerNegative), 0x7f}), uint64(2));
-	TestEqual(TEXT("MeasureCompactBinary(IntegerNegative, Value2B)"), Measure({uint8(ECbFieldType::IntegerNegative), 0x80}), uint64(3));
-
-	TestEqual(TEXT("MeasureCompactBinary(Float32, NoType)"), Measure({}, ECbFieldType::Float32), uint64(4));
-	TestEqual(TEXT("MeasureCompactBinary(Float32, NameSize1B, NoType)"), Measure({30}, ECbFieldType::Float32 | ECbFieldType::HasFieldName), uint64(35));
-	TestEqual(TEXT("MeasureCompactBinary(Float32, NameSize2B, NoType)"), Measure({0x80, 0x80}, ECbFieldType::Float32 | ECbFieldType::HasFieldName), uint64(134));
-	TestEqual(TEXT("MeasureCompactBinary(Float32, NameSize2BShort, NoType)"), Measure({0x80}, ECbFieldType::Float32 | ECbFieldType::HasFieldName), uint64(0));
-	TestEqual(TEXT("MeasureCompactBinary(Float32, MissingName, NoType)"), Measure({}, ECbFieldType::Float32 | ECbFieldType::HasFieldName), uint64(0));
-
-	TestEqual(TEXT("MeasureCompactBinary(Float32)"), Measure({uint8(ECbFieldType::Float32)}), uint64(5));
-	TestEqual(TEXT("MeasureCompactBinary(Float32, NameSize1B)"), Measure({uint8(ECbFieldType::Float32 | ECbFieldType::HasFieldName), 30}), uint64(36));
-	TestEqual(TEXT("MeasureCompactBinary(Float32, NameSize2B)"), Measure({uint8(ECbFieldType::Float32 | ECbFieldType::HasFieldName), 0x80, 0x80}), uint64(135));
-	TestEqual(TEXT("MeasureCompactBinary(Float32, NameSize2BShort)"), Measure({uint8(ECbFieldType::Float32 | ECbFieldType::HasFieldName), 0x80}), uint64(0));
-	TestEqual(TEXT("MeasureCompactBinary(Float32, MissingName)"), Measure({uint8(ECbFieldType::Float32 | ECbFieldType::HasFieldName)}), uint64(0));
-
-	TestEqual(TEXT("MeasureCompactBinary(Float64)"), Measure({uint8(ECbFieldType::Float64)}), uint64(9));
-
-	TestEqual(TEXT("MeasureCompactBinary(BoolFalse)"), Measure({uint8(ECbFieldType::BoolFalse)}), uint64(1));
-	TestEqual(TEXT("MeasureCompactBinary(BoolTrue)"), Measure({uint8(ECbFieldType::BoolTrue)}), uint64(1));
-
-	TestEqual(TEXT("MeasureCompactBinary(Reference)"), Measure({uint8(ECbFieldType::Reference)}), uint64(33));
-	TestEqual(TEXT("MeasureCompactBinary(BinaryReference)"), Measure({uint8(ECbFieldType::BinaryReference)}), uint64(33));
-
-	TestEqual(TEXT("MeasureCompactBinary(Hash)"), Measure({uint8(ECbFieldType::Hash)}), uint64(33));
-	TestEqual(TEXT("MeasureCompactBinary(Uuid)"), Measure({uint8(ECbFieldType::Uuid)}), uint64(17));
-
-	TestEqual(TEXT("MeasureCompactBinary(DateTime)"), Measure({uint8(ECbFieldType::DateTime)}), uint64(9));
-	TestEqual(TEXT("MeasureCompactBinary(TimeSpan)"), Measure({uint8(ECbFieldType::TimeSpan)}), uint64(9));
-
-	return true;
-}
-
-///////////////////////////////////////////////////////////////////////////////////////////////////////////////////////
-
 IMPLEMENT_SIMPLE_AUTOMATION_TEST(FCbFieldParseTest, "System.Core.Serialization.CbFieldParseTest", CompactBinaryTestFlags)
 bool FCbFieldParseTest::RunTest(const FString& Parameters)
 {
diff --git a/Engine/Source/Runtime/Core/Public/Serialization/CompactBinary.h b/Engine/Source/Runtime/Core/Public/Serialization/CompactBinary.h
--- a/Engine/Source/Runtime/Core/Public/Serialization/CompactBinary.h
+++ b/Engine/Source/Runtime/Core/Public/Serialization/CompactBinary.h
@@ -67,6 +67,7 @@
  * \endcode
  */
 
+class FArchive;
 class FCbArray;
 class FCbField;
 class FCbObject;
@@ -631,6 +632,13 @@ protected:
 	 */
 	CORE_API void CopyTo(FMutableMemoryView Buffer) const;
 
+	/**
+	 * Copy the field into an archive. This will write GetSize() bytes.
+	 *
+	 * The field type will only be copied if HasFieldType is set.
+	 */
+	CORE_API void CopyTo(FArchive& Ar) const;
+
 	/** Create a view of the field, including the type and name when present. */
 	CORE_API FConstMemoryView GetFieldView() const;
 
@@ -779,6 +787,9 @@ protected:
 	/** Copy the array into a buffer of at least GetSize() bytes. */
 	CORE_API void CopyTo(FMutableMemoryView Buffer) const;
 
+	/** Copy the array into an archive. This will write GetSize() bytes. */
+	CORE_API void CopyTo(FArchive& Ar) const;
+
 private:
 	friend inline FCbFieldIterator begin(const FCbArray& Array) { return Array.CreateIterator(); }
 	friend inline FCbFieldIterator end(const FCbArray&) { return FCbFieldIterator(); }
@@ -874,6 +885,9 @@ protected:
 	/** Copy the object into a buffer of at least GetSize() bytes. */
 	CORE_API void CopyTo(FMutableMemoryView Buffer) const;
 
+	/** Copy the object into an archive. This will write GetSize() bytes. */
+	CORE_API void CopyTo(FArchive& Ar) const;
+
 private:
 	friend inline FCbFieldIterator begin(const FCbObject& Object) { return Object.CreateIterator(); }
 	friend inline FCbFieldIterator end(const FCbObject&) { return FCbFieldIterator(); }
@@ -883,6 +897,11 @@ private:
 
 ///////////////////////////////////////////////////////////////////////////////////////////////////////////////////////
 
+/** A reference to a function that is used to allocate buffers for compact binary data. */
+using FCbBufferAllocator = TFunctionRef<FSharedBufferPtr (ECbFieldType Type, uint64 Size)>;
+
+///////////////////////////////////////////////////////////////////////////////////////////////////////////////////////
+
 /**
  * A wrapper that can hold a reference to the memory that contains its value.
  *
@@ -1242,21 +1261,3 @@ inline FCbArrayRef FCbFieldRef::AsArrayRef() &&
 }
 
 ///////////////////////////////////////////////////////////////////////////////////////////////////////////////////////
-
-/**
- * Determine the size in bytes of a one compact binary field at the start of the view.
- *
- * This may be called on an incomplete or invalid field, in which case the returned size is zero.
- * A size can always be extracted from a valid field with no name if a view of at least the first
- * 10 bytes is provided, regardless of field size. For fields with names, the size of view needed
- * to calculate a size is at most 10 + MaxNameLen + MeasureVarUInt(MaxNameLen).
- *
- * This function can be used when streaming a field, for example, to determine the size of buffer
- * to fill before attempting to construct a field from it.
- *
- * @param View A memory view that may contain the start of a field.
- * @param Type HasFieldType means that View contains the type. Otherwise, use the given type.
- */
-CORE_API uint64 MeasureCompactBinary(FConstMemoryView View, ECbFieldType Type = ECbFieldType::HasFieldType);
-
-///////////////////////////////////////////////////////////////////////////////////////////////////////////////////////
diff --git a/Engine/Source/Runtime/Core/Public/Serialization/CompactBinarySerialization.h b/Engine/Source/Runtime/Core/Public/Serialization/CompactBinarySerialization.h
new file mode 100644
--- /dev/null
+++ b/Engine/Source/Runtime/Core/Public/Serialization/CompactBinarySerialization.h
@@ -0,0 +1,76 @@
+// Copyright Epic Games, Inc. All Rights Reserved.
+
+#pragma once
+
+#include "CoreTypes.h"
+#include "Memory/MemoryView.h"
+#include "Serialization/CompactBinary.h"
+
+class FArchive;
+
+///////////////////////////////////////////////////////////////////////////////////////////////////////////////////////
+
+/**
+ * Determine the size in bytes of the compact binary field at the start of the view.
+ *
+ * This may be called on an incomplete or invalid field, in which case the returned size is zero.
+ * A size can always be extracted from a valid field with no name if a view of at least the first
+ * 10 bytes is provided, regardless of field size. For fields with names, the size of view needed
+ * to calculate a size is at most 10 + MaxNameLen + MeasureVarUInt(MaxNameLen).
+ *
+ * This function can be used when streaming a field, for example, to determine the size of buffer
+ * to fill before attempting to construct a field from it.
+ *
+ * @param View A memory view that may contain the start of a field.
+ * @param Type HasFieldType means that View contains the type. Otherwise, use the given type.
+ */
+CORE_API uint64 MeasureCompactBinary(FConstMemoryView View, ECbFieldType Type = ECbFieldType::HasFieldType);
+
+/**
+ * Try to determine the type and size of the compact binary field at the start of the view.
+ *
+ * This may be called on an incomplete or invalid field, in which case it will return false, with
+ * OutSize being 0 for invalid fields, otherwise the minimum view size necessary to make progress
+ * in measuring the field on the next call to this function.
+ *
+ * @note A return of true from this function does not indicate that the entire field is valid.
+ *
+ * @param InView A memory view that may contain the start of a field.
+ * @param OutType The type (with flags) of the field. None is written until a value is available.
+ * @param OutSize The total field size for a return of true, 0 for invalid fields, or the size to
+ *                make progress in measuring the field on the next call to this function.
+ * @param InType HasFieldType means that InView contains the type. Otherwise, use the given type.
+ * @return true if the size of the field was determined, otherwise false.
+ */
+CORE_API bool TryMeasureCompactBinary(
+	FConstMemoryView InView,
+	ECbFieldType& OutType,
+	uint64& OutSize,
+	ECbFieldType InType = ECbFieldType::HasFieldType);
+
+///////////////////////////////////////////////////////////////////////////////////////////////////////////////////////
+
+/**
+ * Load a compact binary field from an archive.
+ *
+ * The field may be an array or an object, which the caller can convert to by using AsArrayRef or
+ * AsObjectRef as appropriate. The buffer allocator is called to provide the buffer for the field
+ * to load into once its size has been determined.
+ *
+ * @param Ar Archive to read the field from.
+ * @param Allocator Allocator for the buffer that the field is loaded into.
+ * @return A field with a reference to the provided buffer if it is owned.
+ */
+CORE_API FCbFieldRef LoadCompactBinary(FArchive& Ar, FCbBufferAllocator Allocator);
+
+/** Save a compact binary value to an archive. */
+CORE_API void SaveCompactBinary(FArchive& Ar, const FCbField& Field);
+CORE_API void SaveCompactBinary(FArchive& Ar, const FCbArray& Array);
+CORE_API void SaveCompactBinary(FArchive& Ar, const FCbObject& Object);
+
+/** Serialize a compact binary value to/from an archive. */
+CORE_API FArchive& operator<<(FArchive& Ar, FCbFieldRef& Field);
+CORE_API FArchive& operator<<(FArchive& Ar, FCbArrayRef& Array);
+CORE_API FArchive& operator<<(FArchive& Ar, FCbObjectRef& Object);
+
+///////////////////////////////////////////////////////////////////////////////////////////////////////////////////////
