CompactBinary: Fixed a static analysis warning

#rnx

diff --git a/Engine/Source/Runtime/Core/Public/Serialization/CompactBinary.h b/Engine/Source/Runtime/Core/Public/Serialization/CompactBinary.h
--- a/Engine/Source/Runtime/Core/Public/Serialization/CompactBinary.h
+++ b/Engine/Source/Runtime/Core/Public/Serialization/CompactBinary.h
@@ -260,7 +260,8 @@ public:
 	/** Whether the type is or may contain fields of any reference type. */
 	static constexpr inline bool MayContainReferences(ECbFieldType Type)
 	{
-		return IsObject(Type) | IsArray(Type) | IsAnyReference(Type);
+		// The use of !! will suppress V792 from static analysis. Using //-V792 did not work.
+		return !!IsObject(Type) | !!IsArray(Type) | !!IsAnyReference(Type);
 	}
 };
 
