Removed Waste field as it was wrong. Replaced it with the lower bits of the allocation size instead and packed it in spare alignment bits

#rnx

diff --git a/Engine/Source/Runtime/Core/Private/ProfilingDebugging/Windows/MemoryTrace.cpp b/Engine/Source/Runtime/Core/Private/ProfilingDebugging/Windows/MemoryTrace.cpp
--- a/Engine/Source/Runtime/Core/Private/ProfilingDebugging/Windows/MemoryTrace.cpp
+++ b/Engine/Source/Runtime/Core/Private/ProfilingDebugging/Windows/MemoryTrace.cpp
@@ -106,8 +106,7 @@ UE_TRACE_EVENT_BEGIN(Memory, Alloc)
 	UE_TRACE_EVENT_FIELD(void*, Owner)
 	UE_TRACE_EVENT_FIELD(void*, Address)
 	UE_TRACE_EVENT_FIELD(uint32, Size)
-	UE_TRACE_EVENT_FIELD(uint8, Alignment)
-	UE_TRACE_EVENT_FIELD(uint8, Waste)
+	UE_TRACE_EVENT_FIELD(uint8, Alignment_SizeLower)
 UE_TRACE_EVENT_END()
 
 UE_TRACE_EVENT_BEGIN(Memory, Realloc)
@@ -115,8 +114,7 @@ UE_TRACE_EVENT_BEGIN(Memory, Realloc)
 	UE_TRACE_EVENT_FIELD(void*, FreeAddress)	// compatible with both Alloc
 	UE_TRACE_EVENT_FIELD(void*, Address)		// and Free events.
 	UE_TRACE_EVENT_FIELD(uint32, Size)
-	UE_TRACE_EVENT_FIELD(uint8, Alignment)
-	UE_TRACE_EVENT_FIELD(uint8, Waste)
+	UE_TRACE_EVENT_FIELD(uint8, Alignment_SizeLower)
 UE_TRACE_EVENT_END()
 
 UE_TRACE_EVENT_BEGIN(Memory, Free)
@@ -310,27 +308,25 @@ void FAllocationTrace::CoreRemove(void* Base, size_t Size, void* Owner)
 void FAllocationTrace::Alloc(void* Address, size_t Size, uint32 Alignment, void* Owner)
 {
 	uint32 ActualAlignment = Alignment > uint32(MIN_ALIGNMENT) ? Alignment : uint32(MIN_ALIGNMENT);
-	uint32 Waste = ActualAlignment - (Size & (ActualAlignment - 1));
+	uint32 Alignment_SizeLower = ActualAlignment | (Size & ((1 << 3) - 1));
 
 	UE_TRACE_LOG(Memory, Alloc, MemAllocChannel)
 		<< Alloc.Address(Address)
 		<< Alloc.Size(uint32(Size >> 3))
-		<< Alloc.Alignment(uint8(Alignment))
-		<< Alloc.Waste(uint8(Waste));
+		<< Alloc.Alignment_SizeLower(uint8(Alignment_SizeLower));
 }
 
 ////////////////////////////////////////////////////////////////////////////////
 void FAllocationTrace::Realloc(void* PrevAddress, void* Address, size_t NewSize, uint32 Alignment, void* Owner)
 {
 	uint32 ActualAlignment = Alignment > uint32(MIN_ALIGNMENT) ? Alignment : uint32(MIN_ALIGNMENT);
-	uint32 Waste = ActualAlignment - (NewSize & (ActualAlignment - 1));
+	uint32 Alignment_SizeLower = ActualAlignment | (NewSize & ((1 << 3) - 1));
 
 	UE_TRACE_LOG(Memory, Realloc, MemAllocChannel)
 		<< Realloc.FreeAddress(PrevAddress)
 		<< Realloc.Address(Address)
 		<< Realloc.Size(uint32(NewSize >> 3))
-		<< Realloc.Alignment(uint8(Alignment))
-		<< Realloc.Waste(uint8(Waste));
+		<< Realloc.Alignment_SizeLower(uint8(Alignment_SizeLower));
 }
 
 ////////////////////////////////////////////////////////////////////////////////
