Fixed missing owner fields. Remove owner field from Free trace events to keep the size down (can be reinstated if need be, although preferably avoided as its very expensive).

#rnx

diff --git a/Engine/Source/Runtime/Core/Private/ProfilingDebugging/Windows/MemoryTrace.cpp b/Engine/Source/Runtime/Core/Private/ProfilingDebugging/Windows/MemoryTrace.cpp
--- a/Engine/Source/Runtime/Core/Private/ProfilingDebugging/Windows/MemoryTrace.cpp
+++ b/Engine/Source/Runtime/Core/Private/ProfilingDebugging/Windows/MemoryTrace.cpp
@@ -91,26 +91,26 @@ UE_TRACE_EVENT_BEGIN(Memory, Summary)
 UE_TRACE_EVENT_END()
 
 UE_TRACE_EVENT_BEGIN(Memory, CoreAdd)
-	UE_TRACE_EVENT_FIELD(void*, Owner)
+	UE_TRACE_EVENT_FIELD(uint64, Owner)
 	UE_TRACE_EVENT_FIELD(void*, Base)
 	UE_TRACE_EVENT_FIELD(uint32, Size)
 UE_TRACE_EVENT_END()
 
 UE_TRACE_EVENT_BEGIN(Memory, CoreRemove)
-	UE_TRACE_EVENT_FIELD(void*, Owner)
+	UE_TRACE_EVENT_FIELD(uint64, Owner)
 	UE_TRACE_EVENT_FIELD(void*, Base)
 	UE_TRACE_EVENT_FIELD(uint32, Size)
 UE_TRACE_EVENT_END()
 
 UE_TRACE_EVENT_BEGIN(Memory, Alloc)
-	UE_TRACE_EVENT_FIELD(void*, Owner)
+	UE_TRACE_EVENT_FIELD(uint64, Owner)
 	UE_TRACE_EVENT_FIELD(void*, Address)
 	UE_TRACE_EVENT_FIELD(uint32, Size)
 	UE_TRACE_EVENT_FIELD(uint8, Alignment_SizeLower)
 UE_TRACE_EVENT_END()
 
 UE_TRACE_EVENT_BEGIN(Memory, Realloc)
-	UE_TRACE_EVENT_FIELD(void*, Owner)
+	UE_TRACE_EVENT_FIELD(uint64, Owner)
 	UE_TRACE_EVENT_FIELD(void*, FreeAddress)	// compatible with both Alloc
 	UE_TRACE_EVENT_FIELD(void*, Address)		// and Free events.
 	UE_TRACE_EVENT_FIELD(uint32, Size)
@@ -118,7 +118,6 @@ UE_TRACE_EVENT_BEGIN(Memory, Realloc)
 UE_TRACE_EVENT_END()
 
 UE_TRACE_EVENT_BEGIN(Memory, Free)
-	UE_TRACE_EVENT_FIELD(uint64, Owner)
 	UE_TRACE_EVENT_FIELD(void*, FreeAddress)
 UE_TRACE_EVENT_END()
 
@@ -292,6 +291,7 @@ void FAllocationTrace::Initialize()
 void FAllocationTrace::CoreAdd(void* Base, size_t Size, void* Owner)
 {
 	UE_TRACE_LOG(Memory, CoreAdd, MemAllocChannel)
+		<< CoreAdd.Owner(uint64(Owner))
 		<< CoreAdd.Base(Base)
 		<< CoreAdd.Size(uint32(Size >> SizeShift));
 }
@@ -300,6 +300,7 @@ void FAllocationTrace::CoreAdd(void* Base, size_t Size, void* Owner)
 void FAllocationTrace::CoreRemove(void* Base, size_t Size, void* Owner)
 {
 	UE_TRACE_LOG(Memory, CoreRemove, MemAllocChannel)
+		<< CoreRemove.Owner(uint64(Owner))
 		<< CoreRemove.Base(Base)
 		<< CoreRemove.Size(uint32(Size >> SizeShift));
 }
@@ -311,6 +312,7 @@ void FAllocationTrace::Alloc(void* Address, size_t Size, uint32 Alignment, void*
 	uint32 Alignment_SizeLower = ActualAlignment | (Size & ((1 << 3) - 1));
 
 	UE_TRACE_LOG(Memory, Alloc, MemAllocChannel)
+		<< Alloc.Owner(uint64(Owner))
 		<< Alloc.Address(Address)
 		<< Alloc.Size(uint32(Size >> 3))
 		<< Alloc.Alignment_SizeLower(uint8(Alignment_SizeLower));
@@ -323,6 +325,7 @@ void FAllocationTrace::Realloc(void* PrevAddress, void* Address, size_t NewSize,
 	uint32 Alignment_SizeLower = ActualAlignment | (NewSize & ((1 << 3) - 1));
 
 	UE_TRACE_LOG(Memory, Realloc, MemAllocChannel)
+		<< Realloc.Owner(uint64(Owner))
 		<< Realloc.FreeAddress(PrevAddress)
 		<< Realloc.Address(Address)
 		<< Realloc.Size(uint32(NewSize >> 3))
@@ -333,8 +336,7 @@ void FAllocationTrace::Realloc(void* PrevAddress, void* Address, size_t NewSize,
 void FAllocationTrace::Free(void* Address, void* Owner)
 {
 	UE_TRACE_LOG(Memory, Free, MemAllocChannel)
-		<< Free.FreeAddress(Address)
-		<< Free.Owner(uint64(Owner));
+		<< Free.FreeAddress(Address);
 }
 
 
