Reallocs need to be two separate free and alloc events. Previously another thread could reuse the address freed by the Realloc() call before the single realloc trace event was traced. This causes dissonance in the trace event stream that become impossible to reason.

#rnx

diff --git a/Engine/Source/Runtime/Core/Private/ProfilingDebugging/Windows/MemoryTrace.cpp b/Engine/Source/Runtime/Core/Private/ProfilingDebugging/Windows/MemoryTrace.cpp
--- a/Engine/Source/Runtime/Core/Private/ProfilingDebugging/Windows/MemoryTrace.cpp
+++ b/Engine/Source/Runtime/Core/Private/ProfilingDebugging/Windows/MemoryTrace.cpp
@@ -109,16 +109,19 @@ UE_TRACE_EVENT_BEGIN(Memory, Alloc)
 	UE_TRACE_EVENT_FIELD(uint8, Alignment_SizeLower)
 UE_TRACE_EVENT_END()
 
-UE_TRACE_EVENT_BEGIN(Memory, Realloc)
+UE_TRACE_EVENT_BEGIN(Memory, Free)
+	UE_TRACE_EVENT_FIELD(void*, Address)
+UE_TRACE_EVENT_END()
+
+UE_TRACE_EVENT_BEGIN(Memory, ReallocAlloc)
 	UE_TRACE_EVENT_FIELD(uint64, Owner)
-	UE_TRACE_EVENT_FIELD(void*, FreeAddress)	// compatible with both Alloc
-	UE_TRACE_EVENT_FIELD(void*, Address)		// and Free events.
+	UE_TRACE_EVENT_FIELD(void*, Address)
 	UE_TRACE_EVENT_FIELD(uint32, Size)
 	UE_TRACE_EVENT_FIELD(uint8, Alignment_SizeLower)
 UE_TRACE_EVENT_END()
 
-UE_TRACE_EVENT_BEGIN(Memory, Free)
-	UE_TRACE_EVENT_FIELD(void*, FreeAddress)
+UE_TRACE_EVENT_BEGIN(Memory, ReallocFree)
+	UE_TRACE_EVENT_FIELD(void*, Address)
 UE_TRACE_EVENT_END()
 
 
@@ -265,8 +268,9 @@ public:
 	void	CoreAdd(void* Base, size_t Size, void* Owner);
 	void	CoreRemove(void* Base, size_t Size, void* Owner);
 	void	Alloc(void* Address, size_t Size, uint32 Alignment, void* Owner);
-	void	Realloc(void* PrevAddress, void* Address, size_t NewSize, uint32 Alignment, void* Owner);
-	void	Free(void* Address, void* Owner);
+	void	Free(void* Address);
+	void	ReallocAlloc(void* Address, size_t Size, uint32 Alignment, void* Owner);
+	void	ReallocFree(void* Address);
 
 private:
 	static const uint32 SizeShift = 3;
@@ -319,24 +323,30 @@ void FAllocationTrace::Alloc(void* Address, size_t Size, uint32 Alignment, void*
 }
 
 ////////////////////////////////////////////////////////////////////////////////
-void FAllocationTrace::Realloc(void* PrevAddress, void* Address, size_t NewSize, uint32 Alignment, void* Owner)
+void FAllocationTrace::Free(void* Address)
+{
+	UE_TRACE_LOG(Memory, Free, MemAllocChannel)
+		<< Free.Address(Address);
+}
+
+////////////////////////////////////////////////////////////////////////////////
+void FAllocationTrace::ReallocAlloc(void* Address, size_t Size, uint32 Alignment, void* Owner)
 {
 	uint32 ActualAlignment = Alignment > uint32(MIN_ALIGNMENT) ? Alignment : uint32(MIN_ALIGNMENT);
-	uint32 Alignment_SizeLower = ActualAlignment | (NewSize & ((1 << 3) - 1));
+	uint32 Alignment_SizeLower = ActualAlignment | (Size & ((1 << 3) - 1));
 
-	UE_TRACE_LOG(Memory, Realloc, MemAllocChannel)
-		<< Realloc.Owner(uint64(Owner))
-		<< Realloc.FreeAddress(PrevAddress)
-		<< Realloc.Address(Address)
-		<< Realloc.Size(uint32(NewSize >> 3))
-		<< Realloc.Alignment_SizeLower(uint8(Alignment_SizeLower));
+	UE_TRACE_LOG(Memory, ReallocAlloc, MemAllocChannel)
+		<< ReallocAlloc.Owner(uint64(Owner))
+		<< ReallocAlloc.Address(Address)
+		<< ReallocAlloc.Size(uint32(Size >> 3))
+		<< ReallocAlloc.Alignment_SizeLower(uint8(Alignment_SizeLower));
 }
 
 ////////////////////////////////////////////////////////////////////////////////
-void FAllocationTrace::Free(void* Address, void* Owner)
+void FAllocationTrace::ReallocFree(void* Address)
 {
-	UE_TRACE_LOG(Memory, Free, MemAllocChannel)
-		<< Free.FreeAddress(Address);
+	UE_TRACE_LOG(Memory, ReallocFree, MemAllocChannel)
+		<< ReallocFree.Address(Address);
 }
 
 
@@ -432,6 +442,7 @@ void* FMallocWrapper::Realloc(void* PrevAddress, SIZE_T NewSize, uint32 Alignmen
 		return nullptr;
 	}
 
+	// Track the block that will (or might) get freed
 	uint32 HeaderSize = 0;
 	void* InnerAddress = PrevAddress;
 #if USE_OVERVIEW_TRACE
@@ -441,9 +452,13 @@ void* FMallocWrapper::Realloc(void* PrevAddress, SIZE_T NewSize, uint32 Alignmen
 	GSummaryTrace->Free(Cookie->Tag, Cookie->Size);
 #endif
 
+	GAllocationTrace->ReallocFree(PrevAddress);
+
+	// Do the actual malloc
 	SIZE_T InnerSize = NewSize + HeaderSize;
 	void* RetAddress = InnerMalloc->Realloc(InnerAddress, InnerSize, Alignment);
 
+	// Track the block that was allocated
 #if USE_OVERVIEW_TRACE
 	RetAddress = (uint8*)RetAddress + HeaderSize;
 	Cookie = (FCookie*)RetAddress - 1;
@@ -453,7 +468,7 @@ void* FMallocWrapper::Realloc(void* PrevAddress, SIZE_T NewSize, uint32 Alignmen
 
 	void* Owner = GetOwner(bLight);
 	Alignment = GetActualAlignment(NewSize, Alignment);
-	GAllocationTrace->Realloc(PrevAddress, RetAddress, NewSize, Alignment, Owner);
+	GAllocationTrace->ReallocAlloc(RetAddress, NewSize, Alignment, Owner);
 
 	return RetAddress;
 }
@@ -466,8 +481,7 @@ void FMallocWrapper::Free(void* Address)
 		return;
 	}
 
-	void* Owner = GetOwner(bLight);
-	GAllocationTrace->Free(Address, Owner);
+	GAllocationTrace->Free(Address);
 
 	void* InnerAddress = Address;
 #if USE_OVERVIEW_TRACE
