[Insights] Fixed usage of UE::Trace namespace when (EXPERIMENTAL_)STATSTRACE_ENABLED is 1.

diff --git a/Engine/Source/Runtime/Core/Private/Stats/StatsTrace.cpp b/Engine/Source/Runtime/Core/Private/Stats/StatsTrace.cpp
--- a/Engine/Source/Runtime/Core/Private/Stats/StatsTrace.cpp
+++ b/Engine/Source/Runtime/Core/Private/Stats/StatsTrace.cpp
@@ -17,8 +17,8 @@ UE_TRACE_EVENT_BEGIN(Stats, Spec, NoSync|Important)
 	UE_TRACE_EVENT_FIELD(bool, IsFloatingPoint)
 	UE_TRACE_EVENT_FIELD(bool, IsMemory)
 	UE_TRACE_EVENT_FIELD(bool, ShouldClearEveryFrame)
-	UE_TRACE_EVENT_FIELD(Trace::AnsiString, Name)
-	UE_TRACE_EVENT_FIELD(Trace::AnsiString, Description)
+	UE_TRACE_EVENT_FIELD(UE::Trace::AnsiString, Name)
+	UE_TRACE_EVENT_FIELD(UE::Trace::AnsiString, Description)
 UE_TRACE_EVENT_END()
 
 UE_TRACE_EVENT_BEGIN(Stats, EventBatch)
