CompactBinary: Fixed writer wide string support, allow writer saving to an archive, and added support for future writer optimizations that avoid copying shared buffers

#rnx

diff --git a/Engine/Source/Runtime/Core/Private/Serialization/CompactBinaryWriter.cpp b/Engine/Source/Runtime/Core/Private/Serialization/CompactBinaryWriter.cpp
--- a/Engine/Source/Runtime/Core/Private/Serialization/CompactBinaryWriter.cpp
+++ b/Engine/Source/Runtime/Core/Private/Serialization/CompactBinaryWriter.cpp
@@ -8,6 +8,7 @@
 #include "Misc/DateTime.h"
 #include "Misc/Guid.h"
 #include "Misc/Timespan.h"
+#include "Serialization/Archive.h"
 #include "Serialization/CompactBinarySerialization.h"
 #include "Serialization/VarInt.h"
 
@@ -86,7 +87,7 @@ void FCbWriter::Reset()
 
 FCbFieldRefIterator FCbWriter::Save() const
 {
-	const uint64 Size = SaveSize();
+	const uint64 Size = GetSaveSize();
 	FSharedBufferPtr Buffer = FSharedBuffer::Alloc(Size);
 	const FCbFieldIterator Output = SaveToMemory(Buffer->GetView());
 	return FCbFieldRefIterator(Output, FSharedBuffer::MakeReadOnly(MoveTemp(Buffer)));
@@ -103,7 +104,16 @@ FCbFieldIterator FCbWriter::SaveToMemory(const FMutableMemoryView Buffer) const
 	return FCbFieldIterator(FCbField(Buffer.GetData()), Buffer.GetDataEnd());
 }
 
-uint64 FCbWriter::SaveSize() const
+void FCbWriter::SaveToArchive(FArchive& Ar) const
+{
+	check(Ar.IsSaving());
+	checkf(States.Num() == 1 && States.Last().Flags == EStateFlags::None,
+		TEXT("It is invalid to save while there are incomplete write operations."));
+	checkf(Data.Num() > 0, TEXT("It is invalid to save when nothing has been written."));
+	Ar.Serialize(const_cast<uint8*>(Data.GetData()), Data.Num());
+}
+
+uint64 FCbWriter::GetSaveSize() const
 {
 	return Data.Num();
 }
@@ -229,6 +239,11 @@ void FCbWriter::Field(const FCbField& Value)
 	EndField(AppendCompactBinary(Value, Data));
 }
 
+void FCbWriter::Field(const FCbFieldRef& Value)
+{
+	Field(FCbField(Value));
+}
+
 void FCbWriter::BeginObject()
 {
 	BeginField();
@@ -277,6 +292,11 @@ void FCbWriter::Object(const FCbObject& Value)
 	EndField(AppendCompactBinary(Value, Data));
 }
 
+void FCbWriter::Object(const FCbObjectRef& Value)
+{
+	Object(FCbObject(Value));
+}
+
 void FCbWriter::BeginArray()
 {
 	BeginField();
@@ -327,6 +347,11 @@ void FCbWriter::Array(const FCbArray& Value)
 	EndField(AppendCompactBinary(Value, Data));
 }
 
+void FCbWriter::Array(const FCbArrayRef& Value)
+{
+	Array(FCbArray(Value));
+}
+
 void FCbWriter::Null()
 {
 	BeginField();
@@ -343,6 +368,14 @@ void FCbWriter::Binary(const void* const Value, const uint64 Size)
 	EndField(ECbFieldType::Binary);
 }
 
+void FCbWriter::Binary(const FSharedBufferConstPtr& Buffer)
+{
+	if (Buffer)
+	{
+		Binary(Buffer->GetData(), Buffer->GetSize());
+	}
+}
+
 void FCbWriter::String(const FAnsiStringView Value)
 {
 	BeginField();
diff --git a/Engine/Source/Runtime/Core/Private/Tests/Serialization/CompactBinaryWriterTest.cpp b/Engine/Source/Runtime/Core/Private/Tests/Serialization/CompactBinaryWriterTest.cpp
--- a/Engine/Source/Runtime/Core/Private/Tests/Serialization/CompactBinaryWriterTest.cpp
+++ b/Engine/Source/Runtime/Core/Private/Tests/Serialization/CompactBinaryWriterTest.cpp
@@ -8,6 +8,7 @@
 #include "Misc/Guid.h"
 #include "Misc/Timespan.h"
 #include "Misc/StringBuilder.h"
+#include "Serialization/BufferArchive.h"
 #include "Serialization/CompactBinaryValidation.h"
 
 #if WITH_DEV_AUTOMATION_TESTS
@@ -737,21 +738,25 @@ IMPLEMENT_SIMPLE_AUTOMATION_TEST(FCbWriterComplexTest, "System.Core.Serializatio
 bool FCbWriterComplexTest::RunTest(const FString& Parameters)
 {
 	FCbObjectRef Object;
+	FBufferArchive Archive;
 	{
 		FCbWriter Writer;
 		Writer.BeginObject();
 
 		const uint8 LocalField[] = { uint8(ECbFieldType::IntegerPositive | ECbFieldType::HasFieldName), 1, 'I', 42 };
 		Writer.Name("FieldCopy"_ASV).Field(FCbField(LocalField));
+		Writer.Name("FieldRefCopy"_ASV).Field(FCbFieldRef(FSharedBuffer::Clone(MakeMemoryView(LocalField))));
 
 		const uint8 LocalObject[] = { uint8(ECbFieldType::Object | ECbFieldType::HasFieldName), 1, 'O', 7,
 			uint8(ECbFieldType::IntegerPositive | ECbFieldType::HasFieldName), 1, 'I', 42,
 			uint8(ECbFieldType::Null | ECbFieldType::HasFieldName), 1, 'N' };
 		Writer.Name("ObjectCopy"_ASV).Object(FCbObject(LocalObject));
+		Writer.Name("ObjectRefCopy"_ASV).Object(FCbObjectRef(FSharedBuffer::Clone(MakeMemoryView(LocalObject))));
 
 		const uint8 LocalArray[] = { uint8(ECbFieldType::UniformArray | ECbFieldType::HasFieldName), 1, 'A', 4, 2,
 			uint8(ECbFieldType::IntegerPositive), 42, 21 };
 		Writer.Name("ArrayCopy"_ASV).Array(FCbArray(LocalArray));
+		Writer.Name("ArrayRefCopy"_ASV).Array(FCbArrayRef(FSharedBuffer::Clone(MakeMemoryView(LocalArray))));
 
 		Writer.Name("Null"_ASV).Null();
 
@@ -761,6 +766,7 @@ bool FCbWriterComplexTest::RunTest(const FString& Parameters)
 			Writer.Name("Empty"_ASV).Binary(FConstMemoryView());
 			Writer.Name("Value"_ASV).Binary(MakeMemoryView("BinaryValue"));
 			Writer.Name("LargeValue"_ASV).Binary(MakeMemoryView(FString::ChrN(256, TEXT('.'))));
+			Writer.Name("LargeRefValue"_ASV).Binary(FSharedBuffer::Clone(MakeMemoryView(FString::ChrN(256, TEXT('!')))));
 		}
 		Writer.EndObject();
 
@@ -771,6 +777,8 @@ bool FCbWriterComplexTest::RunTest(const FString& Parameters)
 			Writer.Name("WideString"_ASV).String(FString::ChrN(256, TEXT('.')));
 			Writer.Name("EmptyAnsiString"_ASV).String(FAnsiStringView());
 			Writer.Name("EmptyWideString"_ASV).String(FWideStringView());
+			Writer.Name("AnsiStringLiteral").String("AnsiValue");
+			Writer.Name("WideStringLiteral").String(TEXT("AnsiValue"));
 		}
 		Writer.EndObject();
 
@@ -903,9 +911,16 @@ bool FCbWriterComplexTest::RunTest(const FString& Parameters)
 
 		Writer.EndObject();
 		Object = Writer.Save().AsObjectRef();
+
+		Writer.SaveToArchive(Archive);
+		TestEqual(TEXT("FCbWriter(Complex).SaveToArchive()->Num()"), uint64(Archive.Num()), Writer.GetSaveSize());
 	}
 
-	TestEqual(TEXT("FCbWriter(Complex) Validate"), ValidateCompactBinary(*Object.GetBuffer(), ECbValidateMode::All), ECbValidateError::None);
+	TestEqual(TEXT("FCbWriter(Complex).Save()->Validate"),
+		ValidateCompactBinary(*Object.GetBuffer(), ECbValidateMode::All), ECbValidateError::None);
+
+	TestEqual(TEXT("FCbWriter(Complex).SaveToArchive()->Validate"),
+		ValidateCompactBinary(MakeMemoryView(Archive), ECbValidateMode::All), ECbValidateError::None);
 
 	return true;
 }
@@ -953,7 +968,9 @@ bool FCbWriterStreamTest::RunTest(const FString& Parameters)
 		Writer.BeginObject();
 		Writer
 			<< "AnsiString"_ASV << "AnsiValue"_ASV
-			<< "WideString"_ASV << TEXT("WideValue"_SV);
+			<< "AnsiStringLiteral"_ASV << "AnsiValue"
+			<< "WideString"_ASV << TEXT("WideValue"_SV)
+			<< "WideStringLiteral"_ASV << TEXT("WideValue");
 		Writer.EndObject();
 
 		Writer << "Integers"_ASV;
diff --git a/Engine/Source/Runtime/Core/Public/Serialization/CompactBinaryWriter.h b/Engine/Source/Runtime/Core/Public/Serialization/CompactBinaryWriter.h
--- a/Engine/Source/Runtime/Core/Public/Serialization/CompactBinaryWriter.h
+++ b/Engine/Source/Runtime/Core/Public/Serialization/CompactBinaryWriter.h
@@ -8,6 +8,8 @@
 
 ///////////////////////////////////////////////////////////////////////////////////////////////////////////////////////
 
+class FArchive;
+class FSharedBufferConstPtr;
 struct FBlake3Hash;
 struct FDateTime;
 struct FGuid;
@@ -26,12 +28,13 @@ struct FTimespan;
  *
  * It is most convenient to use the streaming API for the writer, as demonstrated in the example.
  *
- * When writing a small amount of compact binary data, \see \ref TCbWriter can be more efficient,
- * as it uses a fixed-size stack buffer for storage before spilling onto the heap.
+ * When writing a small amount of compact binary data, TCbWriter can be more efficient as it uses
+ * a fixed-size stack buffer for storage before spilling onto the heap.
+ *
+ * @see TCbWriter
  *
  * Example:
  *
- * \code
  * FCbObjectRef WriteObject()
  * {
  *     TCbWriter<256> Writer;
@@ -48,7 +51,6 @@ struct FTimespan;
  *     Writer.EndObject();
  *     return Writer.Save().AsObjectRef();
  * }
- * \endcode
  */
 class FCbWriter
 {
@@ -76,17 +78,27 @@ public:
 	 * It is not valid to call this function in the middle of writing an object, array, or field.
 	 * The writer remains valid for further use when this function returns.
 	 *
-	 * @param Buffer A mutable memory view to write to. Must be exactly SaveSize() bytes.
+	 * @param Buffer A mutable memory view to write to. Must be exactly GetSaveSize() bytes.
 	 * @return An iterator for the field(s) written to the buffer.
 	 */
 	CORE_API FCbFieldIterator SaveToMemory(FMutableMemoryView Buffer) const;
 
+	/**
+	 * Serialize the field(s) to an archive.
+	 *
+	 * It is not valid to call this function in the middle of writing an object, array, or field.
+	 * The writer remains valid for further use when this function returns.
+	 *
+	 * @param Ar An archive to write to. Exactly GetSaveSize() bytes will be written.
+	 */
+	CORE_API void SaveToArchive(FArchive& Ar) const;
+
 	/**
 	 * The size of buffer (in bytes) required to serialize the fields that have been written.
 	 *
 	 * It is not valid to call this function in the middle of writing an object, array, or field.
 	 */
-	CORE_API uint64 SaveSize() const;
+	CORE_API uint64 GetSaveSize() const;
 
 	/**
 	 * Write a field name.
@@ -98,22 +110,28 @@ public:
 
 	/** Copy the value (not the name) of an existing field. */
 	CORE_API void Field(const FCbField& Value);
+	/** Copy the value (not the name) of an existing field. Holds a reference if owned. */
+	CORE_API void Field(const FCbFieldRef& Value);
 
-	/** Begin a new object. Must have a matching call to \ref EndObject. */
+	/** Begin a new object. Must have a matching call to EndObject. */
 	CORE_API void BeginObject();
 	/** End an object after its fields have been written. */
 	CORE_API void EndObject();
 
 	/** Copy the value (not the name) of an existing object. */
 	CORE_API void Object(const FCbObject& Value);
+	/** Copy the value (not the name) of an existing object. Holds a reference if owned. */
+	CORE_API void Object(const FCbObjectRef& Value);
 
-	/** Begin a new array. Must have a matching call to \ref EndArray. */
+	/** Begin a new array. Must have a matching call to EndArray. */
 	CORE_API void BeginArray();
 	/** End an array after its fields have been written. */
 	CORE_API void EndArray();
 
 	/** Copy the value (not the name) of an existing array. */
 	CORE_API void Array(const FCbArray& Value);
+	/** Copy the value (not the name) of an existing array. Holds a reference if owned. */
+	CORE_API void Array(const FCbArrayRef& Value);
 
 	/** Write a null field. */
 	CORE_API void Null();
@@ -122,6 +140,8 @@ public:
 	CORE_API void Binary(const void* Value, uint64 Size);
 	/** Write a binary field by copying the view. */
 	inline void Binary(FConstMemoryView Value) { Binary(Value.GetData(), Value.GetSize()); }
+	/** Write a binary field by copying the buffer. Holds a reference if owned. */
+	CORE_API void Binary(const FSharedBufferConstPtr& Buffer);
 
 	/** Write a string field by copying the UTF-8 value. */
 	CORE_API void String(FAnsiStringView Value);
@@ -221,7 +241,7 @@ private:
 /**
  * A writer for compact binary object, arrays, and fields that uses a fixed-size stack buffer.
  *
- * \see \ref FCbWriter
+ * @see FCbWriter
  */
 template <uint32 InlineBufferSize>
 class TCbWriter : public FCbWriter
@@ -262,18 +282,36 @@ inline FCbWriter& operator<<(FCbWriter& Writer, const FCbField& Value)
 	return Writer;
 }
 
+inline FCbWriter& operator<<(FCbWriter& Writer, const FCbFieldRef& Value)
+{
+	Writer.Field(Value);
+	return Writer;
+}
+
 inline FCbWriter& operator<<(FCbWriter& Writer, const FCbObject& Value)
 {
 	Writer.Object(Value);
 	return Writer;
 }
 
+inline FCbWriter& operator<<(FCbWriter& Writer, const FCbObjectRef& Value)
+{
+	Writer.Object(Value);
+	return Writer;
+}
+
 inline FCbWriter& operator<<(FCbWriter& Writer, const FCbArray& Value)
 {
 	Writer.Array(Value);
 	return Writer;
 }
 
+inline FCbWriter& operator<<(FCbWriter& Writer, const FCbArrayRef& Value)
+{
+	Writer.Array(Value);
+	return Writer;
+}
+
 inline FCbWriter& operator<<(FCbWriter& Writer, nullptr_t)
 {
 	Writer.Null();
@@ -286,6 +324,12 @@ inline FCbWriter& operator<<(FCbWriter& Writer, FWideStringView Value)
 	return Writer;
 }
 
+inline FCbWriter& operator<<(FCbWriter& Writer, const WIDECHAR* Value)
+{
+	Writer.String(Value);
+	return Writer;
+}
+
 inline FCbWriter& operator<<(FCbWriter& Writer, int32 Value)
 {
 	Writer.Integer(Value);
