FPathViews::IterateComponents
#rnx

diff --git a/Engine/Source/Runtime/Core/Private/Misc/PathViews.cpp b/Engine/Source/Runtime/Core/Private/Misc/PathViews.cpp
--- a/Engine/Source/Runtime/Core/Private/Misc/PathViews.cpp
+++ b/Engine/Source/Runtime/Core/Private/Misc/PathViews.cpp
@@ -6,6 +6,8 @@
 #include "Containers/UnrealString.h"
 #include "Containers/StringView.h"
 #include "Misc/StringBuilder.h"
+#include "String/ParseTokens.h"
+#include "Templates/Function.h"
 
 namespace UE4PathViews_Private
 {
@@ -69,6 +71,11 @@ FStringView FPathViews::GetPathLeaf(const FStringView& InPath)
 	return FStringView();
 }
 
+void FPathViews::IterateComponents(FStringView InPath, TFunctionRef<void(FStringView)> ComponentVisitor)
+{
+	UE::String::ParseTokensMultiple(InPath, { TEXT('/'), TEXT('\\') }, ComponentVisitor);
+}
+
 void FPathViews::Split(const FStringView& InPath, FStringView& OutPath, FStringView& OutName, FStringView& OutExt)
 {
 	const FStringView CleanName = GetCleanFilename(InPath);
diff --git a/Engine/Source/Runtime/Core/Public/Misc/PathViews.h b/Engine/Source/Runtime/Core/Public/Misc/PathViews.h
--- a/Engine/Source/Runtime/Core/Public/Misc/PathViews.h
+++ b/Engine/Source/Runtime/Core/Public/Misc/PathViews.h
@@ -6,6 +6,7 @@
 #include "Containers/StringFwd.h"
 
 class FString;
+template <typename FuncType> class TFunctionRef;
 
 class CORE_API FPathViews
 {
@@ -99,6 +100,16 @@ public:
 	 */
 	static FStringView GetPathLeaf(const FStringView& InPath);
 
+	/**
+	 * Splits InPath into individual directory components, and calls ComponentVisitor on each.
+	 *
+	 * Examples:
+	 * "A/B.C" -> {"A", "B.C"}
+	 * "A/B/C" -> {"A", "B", "C"}
+	 * "../../A/B/C.D" -> {"..", "..", "A", "B", "C.D" }
+	 */
+	static void IterateComponents(FStringView InPath, TFunctionRef<void(FStringView)> ComponentVisitor);
+
 	/**
 	 * Splits a path into three parts, any of which may be empty: the path, the clean name, and the extension.
 	 *
