Allow FTextStringHelper::WriteToBuffer to strip the package namespace from NSLOCTEXT entries (eg, when writing cooked data)

diff --git a/Engine/Source/Runtime/Core/Private/Internationalization/Text.cpp b/Engine/Source/Runtime/Core/Private/Internationalization/Text.cpp
--- a/Engine/Source/Runtime/Core/Private/Internationalization/Text.cpp
+++ b/Engine/Source/Runtime/Core/Private/Internationalization/Text.cpp
@@ -1348,14 +1348,14 @@ void FFormatArgumentValue::ToFormattedString(const bool bInRebuildText, const bo
 	}
 }
 
-FString FFormatArgumentValue::ToExportedString() const
+FString FFormatArgumentValue::ToExportedString(const bool bStripPackageNamespace) const
 {
 	FString Result;
-	ToExportedString(Result);
+	ToExportedString(Result, bStripPackageNamespace);
 	return Result;
 }
 
-void FFormatArgumentValue::ToExportedString(FString& OutResult) const
+void FFormatArgumentValue::ToExportedString(FString& OutResult, const bool bStripPackageNamespace) const
 {
 	switch (Type)
 	{
@@ -1374,7 +1374,7 @@ void FFormatArgumentValue::ToExportedString(FString& OutResult) const
 		OutResult += LexToString(DoubleValue);
 		break;
 	case EFormatArgumentType::Text:
-		FTextStringHelper::WriteToBuffer(OutResult, GetTextValue(), true);
+		FTextStringHelper::WriteToBuffer(OutResult, GetTextValue(), /*bRequiresQuotes*/true, bStripPackageNamespace);
 		break;
 	case EFormatArgumentType::Gender:
 		TextStringificationUtil::WriteScopedEnumToBuffer(OutResult, TEXT("ETextGender::"), GetGenderValue());
@@ -1781,7 +1781,7 @@ bool FTextStringHelper::ReadFromString(const TCHAR* Buffer, FText& OutValue, con
 	return false;
 }
 
-void FTextStringHelper::WriteToBuffer(FString& Buffer, const FText& Value, const bool bRequiresQuotes)
+void FTextStringHelper::WriteToBuffer(FString& Buffer, const FText& Value, const bool bRequiresQuotes, const bool bStripPackageNamespace)
 {
 	const FTextHistory& TextHistory = Value.TextData->GetTextHistory();
 	const FString& StringValue = FTextInspector::GetDisplayString(Value);
@@ -1797,7 +1797,7 @@ void FTextStringHelper::WriteToBuffer(FString& Buffer, const FText& Value, const
 #undef LOC_DEFINE_REGION
 	}
 	// Is this text that should be written via its text history?
-	else if (TextHistory.WriteToBuffer(Buffer, Value.TextData->GetLocalizedString()))
+	else if (TextHistory.WriteToBuffer(Buffer, Value.TextData->GetLocalizedString(), bStripPackageNamespace))
 	{
 	}
 	// This isn't special text, so write as a raw string (potentially quoted)
diff --git a/Engine/Source/Runtime/Core/Private/Internationalization/TextHistory.cpp b/Engine/Source/Runtime/Core/Private/Internationalization/TextHistory.cpp
--- a/Engine/Source/Runtime/Core/Private/Internationalization/TextHistory.cpp
+++ b/Engine/Source/Runtime/Core/Private/Internationalization/TextHistory.cpp
@@ -347,7 +347,7 @@ const TCHAR* ReadNumberFormattingOptionsFromBuffer(const TCHAR* Buffer, FNumberF
 	return Buffer;
 }
 
-void WriteNumberOrPercentToBuffer(FString& Buffer, const TCHAR* TokenMarker, const FFormatArgumentValue& SourceValue, const TOptional<FNumberFormattingOptions>& FormatOptions, FCulturePtr TargetCulture)
+void WriteNumberOrPercentToBuffer(FString& Buffer, const TCHAR* TokenMarker, const FFormatArgumentValue& SourceValue, const TOptional<FNumberFormattingOptions>& FormatOptions, FCulturePtr TargetCulture, const bool bStripPackageNamespace)
 {
 	FString Suffix;
 	FString CustomOptions;
@@ -376,7 +376,7 @@ void WriteNumberOrPercentToBuffer(FString& Buffer, const TCHAR* TokenMarker, con
 	Buffer += TokenMarker;
 	Buffer += Suffix;
 	Buffer += TEXT("(");
-	SourceValue.ToExportedString(Buffer);
+	SourceValue.ToExportedString(Buffer, bStripPackageNamespace);
 	if (Suffix == CustomSuffix)
 	{
 		Buffer += TEXT(", ");
@@ -455,7 +455,7 @@ const TCHAR* ReadNumberOrPercentFromBuffer(const TCHAR* Buffer, const FString& T
 	return nullptr;
 }
 
-void WriteDateTimeToBuffer(FString& Buffer, const TCHAR* TokenMarker, const FDateTime& DateTime, const EDateTimeStyle::Type* DateStylePtr, const EDateTimeStyle::Type* TimeStylePtr, const FString& TimeZone, FCulturePtr TargetCulture)
+void WriteDateTimeToBuffer(FString& Buffer, const TCHAR* TokenMarker, const FDateTime& DateTime, const EDateTimeStyle::Type* DateStylePtr, const EDateTimeStyle::Type* TimeStylePtr, const FString& TimeZone, FCulturePtr TargetCulture, const bool bStripPackageNamespace)
 {
 	auto WriteDateTimeStyle = [](FString& OutValueBuffer, const EDateTimeStyle::Type& InValue)
 	{
@@ -478,7 +478,7 @@ void WriteDateTimeToBuffer(FString& Buffer, const TCHAR* TokenMarker, const FDat
 	Buffer += TokenMarker;
 	Buffer += Suffix;
 	Buffer += TEXT("(");
-	FFormatArgumentValue(DateTime.ToUnixTimestamp()).ToExportedString(Buffer);
+	FFormatArgumentValue(DateTime.ToUnixTimestamp()).ToExportedString(Buffer, bStripPackageNamespace);
 	if (DateStylePtr)
 	{
 		Buffer += TEXT(", ");
@@ -608,13 +608,13 @@ const TCHAR* ReadDateTimeFromBuffer(const TCHAR* Buffer, const FString& TokenMar
 }
 
 typedef TFunctionRef<void(const FString*, const FFormatArgumentValue&)> FTextFormatArgumentEnumeratorCallback;
-void WriteTextFormatToBuffer(FString& Buffer, const FString& TokenMarker, const FTextFormat& SourceFmt, TFunctionRef<void(FTextFormatArgumentEnumeratorCallback)> ArgumentEnumerator)
+void WriteTextFormatToBuffer(FString& Buffer, const FString& TokenMarker, const FTextFormat& SourceFmt, const bool bStripPackageNamespace, TFunctionRef<void(FTextFormatArgumentEnumeratorCallback)> ArgumentEnumerator)
 {
 	// Produces LOCGEN_FORMAT_NAMED(..., [...]) or LOCGEN_FORMAT_ORDERED(..., [...])
 	Buffer += TokenMarker;
 	Buffer += TEXT("(");
-	FTextStringHelper::WriteToBuffer(Buffer, SourceFmt.GetSourceText(), true);
-	ArgumentEnumerator([&Buffer](const FString* InKey, const FFormatArgumentValue& InValue)
+	FTextStringHelper::WriteToBuffer(Buffer, SourceFmt.GetSourceText(), /*bRequireQuotes*/true, bStripPackageNamespace);
+	ArgumentEnumerator([&Buffer, bStripPackageNamespace](const FString* InKey, const FFormatArgumentValue& InValue)
 	{
 		if (InKey)
 		{
@@ -624,7 +624,7 @@ void WriteTextFormatToBuffer(FString& Buffer, const FString& TokenMarker, const
 		}
 
 		Buffer += TEXT(", ");
-		InValue.ToExportedString(Buffer);
+		InValue.ToExportedString(Buffer, bStripPackageNamespace);
 	});
 	Buffer += TEXT(")");
 }
@@ -703,7 +703,7 @@ const TCHAR* FTextHistory::ReadFromBuffer(const TCHAR* Buffer, const TCHAR* Text
 	return nullptr;
 }
 
-bool FTextHistory::WriteToBuffer(FString& Buffer, FTextDisplayStringPtr DisplayString) const
+bool FTextHistory::WriteToBuffer(FString& Buffer, FTextDisplayStringPtr DisplayString, const bool bStripPackageNamespace) const
 {
 	return false;
 }
@@ -1020,7 +1020,7 @@ const TCHAR* FTextHistory_Base::ReadFromBuffer(const TCHAR* Buffer, const TCHAR*
 	return nullptr;
 }
 
-bool FTextHistory_Base::WriteToBuffer(FString& Buffer, FTextDisplayStringPtr DisplayString) const
+bool FTextHistory_Base::WriteToBuffer(FString& Buffer, FTextDisplayStringPtr DisplayString, const bool bStripPackageNamespace) const
 {
 	FString Namespace;
 	FString Key;
@@ -1028,6 +1028,11 @@ bool FTextHistory_Base::WriteToBuffer(FString& Buffer, FTextDisplayStringPtr Dis
 
 	if (bFoundNamespaceAndKey)
 	{
+		if (bStripPackageNamespace)
+		{
+			TextNamespaceUtil::StripPackageNamespaceInline(Namespace);
+		}
+
 #define LOC_DEFINE_REGION
 		// Produces NSLOCTEXT("...", "...", "...")
 		Buffer += TEXT("NSLOCTEXT(\"");
@@ -1194,9 +1199,9 @@ const TCHAR* FTextHistory_NamedFormat::ReadFromBuffer(const TCHAR* Buffer, const
 	return nullptr;
 }
 
-bool FTextHistory_NamedFormat::WriteToBuffer(FString& Buffer, FTextDisplayStringPtr DisplayString) const
+bool FTextHistory_NamedFormat::WriteToBuffer(FString& Buffer, FTextDisplayStringPtr DisplayString, const bool bStripPackageNamespace) const
 {
-	TextStringificationUtil::WriteTextFormatToBuffer(Buffer, TextStringificationUtil::LocGenFormatNamedMarker, SourceFmt, [this](TextStringificationUtil::FTextFormatArgumentEnumeratorCallback Callback)
+	TextStringificationUtil::WriteTextFormatToBuffer(Buffer, TextStringificationUtil::LocGenFormatNamedMarker, SourceFmt, bStripPackageNamespace, [this](TextStringificationUtil::FTextFormatArgumentEnumeratorCallback Callback)
 	{
 		for (const auto& ArgumentPair : Arguments)
 		{
@@ -1361,9 +1366,9 @@ const TCHAR* FTextHistory_OrderedFormat::ReadFromBuffer(const TCHAR* Buffer, con
 	return nullptr;
 }
 
-bool FTextHistory_OrderedFormat::WriteToBuffer(FString& Buffer, FTextDisplayStringPtr DisplayString) const
+bool FTextHistory_OrderedFormat::WriteToBuffer(FString& Buffer, FTextDisplayStringPtr DisplayString, const bool bStripPackageNamespace) const
 {
-	TextStringificationUtil::WriteTextFormatToBuffer(Buffer, TextStringificationUtil::LocGenFormatOrderedMarker, SourceFmt, [this](TextStringificationUtil::FTextFormatArgumentEnumeratorCallback Callback)
+	TextStringificationUtil::WriteTextFormatToBuffer(Buffer, TextStringificationUtil::LocGenFormatOrderedMarker, SourceFmt, bStripPackageNamespace, [this](TextStringificationUtil::FTextFormatArgumentEnumeratorCallback Callback)
 	{
 		for (const FFormatArgumentValue& ArgumentValue : Arguments)
 		{
@@ -1493,9 +1498,9 @@ const TCHAR* FTextHistory_ArgumentDataFormat::ReadFromBuffer(const TCHAR* Buffer
 	return nullptr;
 }
 
-bool FTextHistory_ArgumentDataFormat::WriteToBuffer(FString& Buffer, FTextDisplayStringPtr DisplayString) const
+bool FTextHistory_ArgumentDataFormat::WriteToBuffer(FString& Buffer, FTextDisplayStringPtr DisplayString, const bool bStripPackageNamespace) const
 {
-	TextStringificationUtil::WriteTextFormatToBuffer(Buffer, TextStringificationUtil::LocGenFormatNamedMarker, SourceFmt, [this](TextStringificationUtil::FTextFormatArgumentEnumeratorCallback Callback)
+	TextStringificationUtil::WriteTextFormatToBuffer(Buffer, TextStringificationUtil::LocGenFormatNamedMarker, SourceFmt, bStripPackageNamespace, [this](TextStringificationUtil::FTextFormatArgumentEnumeratorCallback Callback)
 	{
 		for (const FFormatArgumentData& Argument : Arguments)
 		{
@@ -1719,9 +1724,9 @@ const TCHAR* FTextHistory_AsNumber::ReadFromBuffer(const TCHAR* Buffer, const TC
 	return Buffer;
 }
 
-bool FTextHistory_AsNumber::WriteToBuffer(FString& Buffer, FTextDisplayStringPtr DisplayString) const
+bool FTextHistory_AsNumber::WriteToBuffer(FString& Buffer, FTextDisplayStringPtr DisplayString, const bool bStripPackageNamespace) const
 {
-	TextStringificationUtil::WriteNumberOrPercentToBuffer(Buffer, TextStringificationUtil::LocGenNumberMarker, SourceValue, FormatOptions, TargetCulture);
+	TextStringificationUtil::WriteNumberOrPercentToBuffer(Buffer, TextStringificationUtil::LocGenNumberMarker, SourceValue, FormatOptions, TargetCulture, bStripPackageNamespace);
 	return true;
 }
 
@@ -1796,9 +1801,9 @@ const TCHAR* FTextHistory_AsPercent::ReadFromBuffer(const TCHAR* Buffer, const T
 	return Buffer;
 }
 
-bool FTextHistory_AsPercent::WriteToBuffer(FString& Buffer, FTextDisplayStringPtr DisplayString) const
+bool FTextHistory_AsPercent::WriteToBuffer(FString& Buffer, FTextDisplayStringPtr DisplayString, const bool bStripPackageNamespace) const
 {
-	TextStringificationUtil::WriteNumberOrPercentToBuffer(Buffer, TextStringificationUtil::LocGenPercentMarker, SourceValue, FormatOptions, TargetCulture);
+	TextStringificationUtil::WriteNumberOrPercentToBuffer(Buffer, TextStringificationUtil::LocGenPercentMarker, SourceValue, FormatOptions, TargetCulture, bStripPackageNamespace);
 	return true;
 }
 
@@ -1947,7 +1952,7 @@ const TCHAR* FTextHistory_AsCurrency::ReadFromBuffer(const TCHAR* Buffer, const
 	return nullptr;
 }
 
-bool FTextHistory_AsCurrency::WriteToBuffer(FString& Buffer, FTextDisplayStringPtr DisplayString) const
+bool FTextHistory_AsCurrency::WriteToBuffer(FString& Buffer, FTextDisplayStringPtr DisplayString, const bool bStripPackageNamespace) const
 {
 	FInternationalization& I18N = FInternationalization::Get();
 	checkf(I18N.IsInitialized() == true, TEXT("FInternationalization is not initialized. An FText formatting method was likely used in static object initialization - this is not supported."));
@@ -1980,7 +1985,7 @@ bool FTextHistory_AsCurrency::WriteToBuffer(FString& Buffer, FTextDisplayStringP
 
 	// Produces LOCGEN_CURRENCY(..., "...", "...")
 	Buffer += TEXT("LOCGEN_CURRENCY(");
-	FFormatArgumentValue(BaseVal).ToExportedString(Buffer);
+	FFormatArgumentValue(BaseVal).ToExportedString(Buffer, bStripPackageNamespace);
 	Buffer += TEXT(", \"");
 	Buffer += CurrencyCode.ReplaceCharWithEscapedChar();
 	Buffer += TEXT("\", \"");
@@ -2077,9 +2082,9 @@ const TCHAR* FTextHistory_AsDate::ReadFromBuffer(const TCHAR* Buffer, const TCHA
 	return Buffer;
 }
 
-bool FTextHistory_AsDate::WriteToBuffer(FString& Buffer, FTextDisplayStringPtr DisplayString) const
+bool FTextHistory_AsDate::WriteToBuffer(FString& Buffer, FTextDisplayStringPtr DisplayString, const bool bStripPackageNamespace) const
 {
-	TextStringificationUtil::WriteDateTimeToBuffer(Buffer, TextStringificationUtil::LocGenDateMarker, SourceDateTime, &DateStyle, nullptr, TimeZone, TargetCulture);
+	TextStringificationUtil::WriteDateTimeToBuffer(Buffer, TextStringificationUtil::LocGenDateMarker, SourceDateTime, &DateStyle, nullptr, TimeZone, TargetCulture, bStripPackageNamespace);
 	return true;
 }
 
@@ -2192,9 +2197,9 @@ const TCHAR* FTextHistory_AsTime::ReadFromBuffer(const TCHAR* Buffer, const TCHA
 	return Buffer;
 }
 
-bool FTextHistory_AsTime::WriteToBuffer(FString& Buffer, FTextDisplayStringPtr DisplayString) const
+bool FTextHistory_AsTime::WriteToBuffer(FString& Buffer, FTextDisplayStringPtr DisplayString, const bool bStripPackageNamespace) const
 {
-	TextStringificationUtil::WriteDateTimeToBuffer(Buffer, TextStringificationUtil::LocGenTimeMarker, SourceDateTime, nullptr, &TimeStyle, TimeZone, TargetCulture);
+	TextStringificationUtil::WriteDateTimeToBuffer(Buffer, TextStringificationUtil::LocGenTimeMarker, SourceDateTime, nullptr, &TimeStyle, TimeZone, TargetCulture, bStripPackageNamespace);
 	return true;
 }
 
@@ -2314,9 +2319,9 @@ const TCHAR* FTextHistory_AsDateTime::ReadFromBuffer(const TCHAR* Buffer, const
 	return Buffer;
 }
 
-bool FTextHistory_AsDateTime::WriteToBuffer(FString& Buffer, FTextDisplayStringPtr DisplayString) const
+bool FTextHistory_AsDateTime::WriteToBuffer(FString& Buffer, FTextDisplayStringPtr DisplayString, const bool bStripPackageNamespace) const
 {
-	TextStringificationUtil::WriteDateTimeToBuffer(Buffer, TextStringificationUtil::LocGenDateTimeMarker, SourceDateTime, &DateStyle, &TimeStyle, TimeZone, TargetCulture);
+	TextStringificationUtil::WriteDateTimeToBuffer(Buffer, TextStringificationUtil::LocGenDateTimeMarker, SourceDateTime, &DateStyle, &TimeStyle, TimeZone, TargetCulture, bStripPackageNamespace);
 	return true;
 }
 
@@ -2430,7 +2435,7 @@ const TCHAR* FTextHistory_Transform::ReadFromBuffer(const TCHAR* Buffer, const T
 	return Buffer;
 }
 
-bool FTextHistory_Transform::WriteToBuffer(FString& Buffer, FTextDisplayStringPtr DisplayString) const
+bool FTextHistory_Transform::WriteToBuffer(FString& Buffer, FTextDisplayStringPtr DisplayString, const bool bStripPackageNamespace) const
 {
 	// Produces LOCGEN_TOLOWER(...) or LOCGEN_TOUPPER
 	switch (TransformType)
@@ -2444,7 +2449,7 @@ bool FTextHistory_Transform::WriteToBuffer(FString& Buffer, FTextDisplayStringPt
 	default:
 		break;
 	}
-	FTextStringHelper::WriteToBuffer(Buffer, SourceText, true);
+	FTextStringHelper::WriteToBuffer(Buffer, SourceText, /*bRequireQuotes*/true, bStripPackageNamespace);
 	Buffer += TEXT(")");
 
 	return true;
@@ -2687,7 +2692,7 @@ const TCHAR* FTextHistory_StringTableEntry::ReadFromBuffer(const TCHAR* Buffer,
 	return nullptr;
 }
 
-bool FTextHistory_StringTableEntry::WriteToBuffer(FString& Buffer, FTextDisplayStringPtr DisplayString) const
+bool FTextHistory_StringTableEntry::WriteToBuffer(FString& Buffer, FTextDisplayStringPtr DisplayString, const bool bStripPackageNamespace) const
 {
 	if (StringTableReferenceData)
 	{
diff --git a/Engine/Source/Runtime/Core/Private/Internationalization/TextHistory.h b/Engine/Source/Runtime/Core/Private/Internationalization/TextHistory.h
--- a/Engine/Source/Runtime/Core/Private/Internationalization/TextHistory.h
+++ b/Engine/Source/Runtime/Core/Private/Internationalization/TextHistory.h
@@ -37,7 +37,7 @@ enum class ETextHistoryType : int8
 	static  bool StaticShouldReadFromBuffer(const TCHAR* Buffer);																											\
 	virtual bool ShouldReadFromBuffer(const TCHAR* Buffer) const override { return StaticShouldReadFromBuffer(Buffer); }													\
 	virtual const TCHAR* ReadFromBuffer(const TCHAR* Buffer, const TCHAR* TextNamespace, const TCHAR* PackageNamespace, FTextDisplayStringPtr& OutDisplayString) override;	\
-	virtual bool WriteToBuffer(FString& Buffer, FTextDisplayStringPtr DisplayString) const override;
+	virtual bool WriteToBuffer(FString& Buffer, FTextDisplayStringPtr DisplayString, const bool bStripPackageNamespace) const override;
 
 /** Utilities for stringifying text */
 namespace TextStringificationUtil
@@ -184,12 +184,13 @@ public:
 	/**
 	 * Write this text history to a stream of text
 	 *
-	 * @param Buffer			The buffer of text to write to.
-	 * @param DisplayString		The display string associated with the text being written
+	 * @param Buffer				 The buffer of text to write to.
+	 * @param DisplayString			 The display string associated with the text being written
+	 * @param bStripPackageNamespace True to strip the package namespace from the written NSLOCTEXT value (eg, when saving cooked data)
 	 *
 	 * @return True if we wrote valid data into Buffer, false otherwise
 	 */
-	virtual bool WriteToBuffer(FString& Buffer, FTextDisplayStringPtr DisplayString) const;
+	virtual bool WriteToBuffer(FString& Buffer, FTextDisplayStringPtr DisplayString, const bool bStripPackageNamespace) const;
 
 	/** Returns TRUE if the Revision is out of date */
 	virtual bool IsOutOfDate() const;
diff --git a/Engine/Source/Runtime/Core/Public/Internationalization/Text.h b/Engine/Source/Runtime/Core/Public/Internationalization/Text.h
--- a/Engine/Source/Runtime/Core/Public/Internationalization/Text.h
+++ b/Engine/Source/Runtime/Core/Public/Internationalization/Text.h
@@ -860,8 +860,8 @@ public:
 	FString ToFormattedString(const bool bInRebuildText, const bool bInRebuildAsSource) const;
 	void ToFormattedString(const bool bInRebuildText, const bool bInRebuildAsSource, FString& OutResult) const;
 
-	FString ToExportedString() const;
-	void ToExportedString(FString& OutResult) const;
+	FString ToExportedString(const bool bStripPackageNamespace = false) const;
+	void ToExportedString(FString& OutResult, const bool bStripPackageNamespace = false) const;
 	const TCHAR* FromExportedString(const TCHAR* InBuffer);
 
 	FORCEINLINE EFormatArgumentType::Type GetType() const
@@ -1149,11 +1149,12 @@ public:
 	/**
 	 * Write the given FText instance to a stream of text
 	 *
-	 * @param Buffer			The buffer of text to write to.
-	 * @param Value				The text value to write into the buffer.
-	 * @param bRequiresQuotes	True if the written text literal must be surrounded by quotes (eg, when saving as a delimited list)
+	 * @param Buffer				 The buffer of text to write to.
+	 * @param Value					 The text value to write into the buffer.
+	 * @param bRequiresQuotes		 True if the written text literal must be surrounded by quotes (eg, when saving as a delimited list)
+	 * @param bStripPackageNamespace True to strip the package namespace from the written NSLOCTEXT value (eg, when saving cooked data)
 	 */
-	static void WriteToBuffer(FString& Buffer, const FText& Value, const bool bRequiresQuotes = false);
+	static void WriteToBuffer(FString& Buffer, const FText& Value, const bool bRequiresQuotes = false, const bool bStripPackageNamespace = false);
 	
 	UE_DEPRECATED(4.22, "FTextStringHelper::WriteToString is deprecated. Use FTextStringHelper::WriteToBuffer instead.")
 	static bool WriteToString(FString& Buffer, const FText& Value, const bool bRequiresQuotes = false);
