Wrapped a few pieces of code in #if/defs that qwew related only to the disabled memory overview tracing.

#rnx

diff --git a/Engine/Source/Runtime/Core/Private/ProfilingDebugging/Windows/MemoryTrace.cpp b/Engine/Source/Runtime/Core/Private/ProfilingDebugging/Windows/MemoryTrace.cpp
--- a/Engine/Source/Runtime/Core/Private/ProfilingDebugging/Windows/MemoryTrace.cpp
+++ b/Engine/Source/Runtime/Core/Private/ProfilingDebugging/Windows/MemoryTrace.cpp
@@ -18,6 +18,8 @@
 #include <atomic>
 #include <intrin.h>
 
+#define USE_OVERVIEW_TRACE 0 // Disabled for now, as it occasionally conflicts with MallocBinned2's canary poison logic.
+
 ////////////////////////////////////////////////////////////////////////////////
 void	Backtracer_Create(FMalloc*);
 void	Backtracer_Initialize();
@@ -78,10 +80,13 @@ UE_TRACE_CHANNEL_DEFINE(MemAllocChannel)
 UE_TRACE_EVENT_BEGIN(Memory, Init)
 	UE_TRACE_EVENT_FIELD(uint8, MinAlignment)
 	UE_TRACE_EVENT_FIELD(uint8, SizeShift)
+#if USE_OVERVIEW_TRACE
 	UE_TRACE_EVENT_FIELD(uint8, SummarySizeShift)
+#endif
 	UE_TRACE_EVENT_FIELD(uint8, Mode)
 UE_TRACE_EVENT_END()
 
+#if USE_OVERVIEW_TRACE
 UE_TRACE_EVENT_BEGIN(Memory, Summary)
 	UE_TRACE_EVENT_FIELD(uint32, Bytes)
 	UE_TRACE_EVENT_FIELD(uint32, ActiveAllocs)
@@ -89,6 +94,7 @@ UE_TRACE_EVENT_BEGIN(Memory, Summary)
 	UE_TRACE_EVENT_FIELD(uint32, TotalReallocs)
 	UE_TRACE_EVENT_FIELD(uint32, TotalFrees)
 UE_TRACE_EVENT_END()
+#endif
 
 UE_TRACE_EVENT_BEGIN(Memory, CoreAdd)
 	UE_TRACE_EVENT_FIELD(uint64, Owner)
@@ -126,7 +132,6 @@ UE_TRACE_EVENT_END()
 
 
 
-#define USE_OVERVIEW_TRACE 0 // Disabled for now, as it occasionally conflicts with MallocBinned2's canary poison logic.
 #if USE_OVERVIEW_TRACE
 
 ////////////////////////////////////////////////////////////////////////////////
