Replaced some magic constants with something semantic and asserted the expected properties.

#rnx

diff --git a/Engine/Source/Runtime/Core/Private/ProfilingDebugging/Windows/MemoryTrace.cpp b/Engine/Source/Runtime/Core/Private/ProfilingDebugging/Windows/MemoryTrace.cpp
--- a/Engine/Source/Runtime/Core/Private/ProfilingDebugging/Windows/MemoryTrace.cpp
+++ b/Engine/Source/Runtime/Core/Private/ProfilingDebugging/Windows/MemoryTrace.cpp
@@ -294,6 +294,8 @@ void FAllocationTrace::Initialize()
 		<< Init.SummarySizeShift(uint8(FSummaryTrace::SizeShift))
 #endif
 		<< Init.SizeShift(uint8(SizeShift));
+
+	static_assert((1 < SizeShift) - 1 <= MIN_ALIGNMENT, "Not enough bits to pack size fields");
 }
 
 ////////////////////////////////////////////////////////////////////////////////
@@ -318,12 +320,12 @@ void FAllocationTrace::CoreRemove(void* Base, size_t Size, void* Owner)
 void FAllocationTrace::Alloc(void* Address, size_t Size, uint32 Alignment, void* Owner)
 {
 	uint32 ActualAlignment = Alignment > uint32(MIN_ALIGNMENT) ? Alignment : uint32(MIN_ALIGNMENT);
-	uint32 Alignment_SizeLower = ActualAlignment | (Size & ((1 << 3) - 1));
+	uint32 Alignment_SizeLower = ActualAlignment | (Size & ((1 << SizeShift) - 1));
 
 	UE_TRACE_LOG(Memory, Alloc, MemAllocChannel)
 		<< Alloc.Owner(uint64(Owner))
 		<< Alloc.Address(Address)
-		<< Alloc.Size(uint32(Size >> 3))
+		<< Alloc.Size(uint32(Size >> SizeShift))
 		<< Alloc.Alignment_SizeLower(uint8(Alignment_SizeLower));
 }
 
@@ -338,12 +340,12 @@ void FAllocationTrace::Free(void* Address)
 void FAllocationTrace::ReallocAlloc(void* Address, size_t Size, uint32 Alignment, void* Owner)
 {
 	uint32 ActualAlignment = Alignment > uint32(MIN_ALIGNMENT) ? Alignment : uint32(MIN_ALIGNMENT);
-	uint32 Alignment_SizeLower = ActualAlignment | (Size & ((1 << 3) - 1));
+	uint32 Alignment_SizeLower = ActualAlignment | (Size & ((1 << SizeShift) - 1));
 
 	UE_TRACE_LOG(Memory, ReallocAlloc, MemAllocChannel)
 		<< ReallocAlloc.Owner(uint64(Owner))
 		<< ReallocAlloc.Address(Address)
-		<< ReallocAlloc.Size(uint32(Size >> 3))
+		<< ReallocAlloc.Size(uint32(Size >> SizeShift))
 		<< ReallocAlloc.Alignment_SizeLower(uint8(Alignment_SizeLower));
 }
 
