Capture the very last tail of alloc/free events

#rnx

diff --git a/Engine/Source/Runtime/Core/Private/ProfilingDebugging/Windows/MemoryTrace.cpp b/Engine/Source/Runtime/Core/Private/ProfilingDebugging/Windows/MemoryTrace.cpp
--- a/Engine/Source/Runtime/Core/Private/ProfilingDebugging/Windows/MemoryTrace.cpp
+++ b/Engine/Source/Runtime/Core/Private/ProfilingDebugging/Windows/MemoryTrace.cpp
@@ -20,6 +20,15 @@
 
 #define USE_OVERVIEW_TRACE 0 // Disabled for now, as it occasionally conflicts with MallocBinned2's canary poison logic.
 
+////////////////////////////////////////////////////////////////////////////////
+namespace UE {
+namespace Trace {
+
+TRACELOG_API void Update();
+
+} // namespace Trace
+} // namespace UE
+
 ////////////////////////////////////////////////////////////////////////////////
 void	Backtracer_Create(FMalloc*);
 void	Backtracer_Initialize();
@@ -270,6 +279,7 @@ class FAllocationTrace
 {
 public:
 	void	Initialize();
+	void	EnableTracePump();
 	void	CoreAdd(void* Base, size_t Size, void* Owner);
 	void	CoreRemove(void* Base, size_t Size, void* Owner);
 	void	Alloc(void* Address, size_t Size, uint32 Alignment, void* Owner);
@@ -278,6 +288,8 @@ public:
 	void	ReallocFree(void* Address);
 
 private:
+	void				Update();
+	bool				bPumpTrace = false;
 	static const uint32 SizeShift = 3;
 	static_assert(MIN_ALIGNMENT >= (1 << SizeShift), "");
 };
@@ -298,6 +310,21 @@ void FAllocationTrace::Initialize()
 	static_assert((1 < SizeShift) - 1 <= MIN_ALIGNMENT, "Not enough bits to pack size fields");
 }
 
+////////////////////////////////////////////////////////////////////////////////
+void FAllocationTrace::EnableTracePump()
+{
+	bPumpTrace = true;
+}
+
+////////////////////////////////////////////////////////////////////////////////
+void FAllocationTrace::Update()
+{
+	if (bPumpTrace)
+	{
+		UE::Trace::Update();
+	}
+}
+
 ////////////////////////////////////////////////////////////////////////////////
 void FAllocationTrace::CoreAdd(void* Base, size_t Size, void* Owner)
 {
@@ -305,6 +332,8 @@ void FAllocationTrace::CoreAdd(void* Base, size_t Size, void* Owner)
 		<< CoreAdd.Owner(uint64(Owner))
 		<< CoreAdd.Base(Base)
 		<< CoreAdd.Size(uint32(Size >> SizeShift));
+
+	Update();
 }
 
 ////////////////////////////////////////////////////////////////////////////////
@@ -314,6 +343,8 @@ void FAllocationTrace::CoreRemove(void* Base, size_t Size, void* Owner)
 		<< CoreRemove.Owner(uint64(Owner))
 		<< CoreRemove.Base(Base)
 		<< CoreRemove.Size(uint32(Size >> SizeShift));
+
+	Update();
 }
 
 ////////////////////////////////////////////////////////////////////////////////
@@ -327,6 +358,8 @@ void FAllocationTrace::Alloc(void* Address, size_t Size, uint32 Alignment, void*
 		<< Alloc.Address(Address)
 		<< Alloc.Size(uint32(Size >> SizeShift))
 		<< Alloc.Alignment_SizeLower(uint8(Alignment_SizeLower));
+
+	Update();
 }
 
 ////////////////////////////////////////////////////////////////////////////////
@@ -334,6 +367,8 @@ void FAllocationTrace::Free(void* Address)
 {
 	UE_TRACE_LOG(Memory, Free, MemAllocChannel)
 		<< Free.Address(Address);
+
+	Update();
 }
 
 ////////////////////////////////////////////////////////////////////////////////
@@ -347,6 +382,8 @@ void FAllocationTrace::ReallocAlloc(void* Address, size_t Size, uint32 Alignment
 		<< ReallocAlloc.Address(Address)
 		<< ReallocAlloc.Size(uint32(Size >> SizeShift))
 		<< ReallocAlloc.Alignment_SizeLower(uint8(Alignment_SizeLower));
+
+	Update();
 }
 
 ////////////////////////////////////////////////////////////////////////////////
@@ -354,6 +391,8 @@ void FAllocationTrace::ReallocFree(void* Address)
 {
 	UE_TRACE_LOG(Memory, ReallocFree, MemAllocChannel)
 		<< ReallocFree.Address(Address);
+
+	Update();
 }
 
 
@@ -750,6 +789,15 @@ FMalloc* MemoryTrace_Create(FMalloc* InMalloc)
 
 	if (Mode > 0)
 	{
+		// Some OSes (i.e. Windows) will terminate all threads except the main
+		// one as part of static deinit. However we may receive more memory
+		// trace events that would get lost as Trace's worker thread has been
+		// terminated. So flush the last remaining memory events trace needs
+		// to be updated which we will do that in response to to memory events.
+		// We'll use an atexit can to know when Trace is probably no longer
+		// getting ticked.
+		atexit([] () { GAllocationTrace->EnableTracePump(); });
+
 		GAllocationTrace->Initialize();
 #if USE_OVERVIEW_TRACE
 		GSummaryTrace->Initialize();
