Add tags to memory tracing.

Instrument LLM scopes (LLM_SCOPE and LLM_REALLOC_SCOPE) with hooks to memory tracing. Added new trace events for the scopes and for specifying tags. Only available on platforms where  allocation tracing is activated.

diff --git a/Engine/Source/Runtime/Core/Private/HAL/LowLevelMemTracker.cpp b/Engine/Source/Runtime/Core/Private/HAL/LowLevelMemTracker.cpp
--- a/Engine/Source/Runtime/Core/Private/HAL/LowLevelMemTracker.cpp
+++ b/Engine/Source/Runtime/Core/Private/HAL/LowLevelMemTracker.cpp
@@ -1185,6 +1185,8 @@ void FLowLevelMemTracker::RegisterCustomTagInternal(int32 Tag, const TCHAR* InDi
 
 void FLowLevelMemTracker::RegisterPlatformTag(int32 Tag, const TCHAR* Name, FName StatName, FName SummaryStatName, int32 ParentTag)
 {
+	MemoryTrace_AnnounceCustomTag(Tag, ParentTag, Name);
+
 	if (bIsDisabled)
 	{
 		return;
@@ -1197,6 +1199,8 @@ void FLowLevelMemTracker::RegisterPlatformTag(int32 Tag, const TCHAR* Name, FNam
 
 void FLowLevelMemTracker::RegisterProjectTag(int32 Tag, const TCHAR* Name, FName StatName, FName SummaryStatName, int32 ParentTag)
 {
+	MemoryTrace_AnnounceCustomTag(Tag, ParentTag, Name);
+
 	if (bIsDisabled)
 	{
 		return;
diff --git a/Engine/Source/Runtime/Core/Private/ProfilingDebugging/TagTrace.cpp b/Engine/Source/Runtime/Core/Private/ProfilingDebugging/TagTrace.cpp
new file mode 100644
--- /dev/null
+++ b/Engine/Source/Runtime/Core/Private/ProfilingDebugging/TagTrace.cpp
@@ -0,0 +1,302 @@
+// Copyright Epic Games, Inc. All Rights Reserved.
+#include "ProfilingDebugging/TagTrace.h"
+
+#include "CoreTypes.h"
+#include "MemoryTrace.h"
+#include "HAL/LowLevelMemTracker.h"
+#include "Trace/Trace.inl"
+
+#if USE_MEMORY_TRACE_TAGS
+
+////////////////////////////////////////////////////////////////////////////////
+
+UE_TRACE_CHANNEL_EXTERN(MemAllocChannel);
+
+UE_TRACE_EVENT_BEGIN(Memory, TagSpec)
+	UE_TRACE_EVENT_FIELD(int32, Tag)
+	UE_TRACE_EVENT_FIELD(int32, Parent)
+	UE_TRACE_EVENT_FIELD(UE::Trace::AnsiString, Display)
+UE_TRACE_EVENT_END()
+
+UE_TRACE_EVENT_BEGIN(Memory, MemoryScope)
+	UE_TRACE_EVENT_FIELD(int32, Tag)
+UE_TRACE_EVENT_END()
+
+UE_TRACE_EVENT_BEGIN(Memory, MemoryScopeRealloc)
+	UE_TRACE_EVENT_FIELD(uint64, Ptr)
+UE_TRACE_EVENT_END()
+
+
+
+////////////////////////////////////////////////////////////////////////////////
+// Per thread active tag, i.e. the top level FMemScope
+thread_local int32 GActiveTag;
+
+////////////////////////////////////////////////////////////////////////////////
+FMemScope::FMemScope(int32 InTag)
+{
+	if (UE_TRACE_CHANNELEXPR_IS_ENABLED(MemAllocChannel))
+	{
+		ActivateScope(InTag);
+	}
+}
+
+////////////////////////////////////////////////////////////////////////////////
+FMemScope::FMemScope(ELLMTag InTag)
+{
+	if (UE_TRACE_CHANNELEXPR_IS_ENABLED(MemAllocChannel))
+	{
+		ActivateScope(static_cast<int32>(InTag));
+	}
+}
+
+////////////////////////////////////////////////////////////////////////////////
+FMemScope::FMemScope(FName InName)
+{
+	if (UE_TRACE_CHANNELEXPR_IS_ENABLED(MemAllocChannel))
+	{
+		ActivateScope(MemoryTrace_AnnounceFNameTag(InName));
+	}
+}
+
+////////////////////////////////////////////////////////////////////////////////
+FMemScope::FMemScope(const UE::LLMPrivate::FTagData* TagData)
+{
+	// TagData is opaque so we cant really use the input, additionally we 
+	// cannot count on LLM being active. Instead we have inserted an explicit 
+	// Trace scope directly following the LLM scope. 
+}
+
+////////////////////////////////////////////////////////////////////////////////
+void FMemScope::ActivateScope(int32 InTag)
+{
+	if (auto LogScope = FMemoryMemoryScopeFields::LogScopeType::ScopedEnter<FMemoryMemoryScopeFields>())
+	{
+		if (const auto& __restrict MemoryScope = *(FMemoryMemoryScopeFields*)(&LogScope))
+		{
+			Inner.SetActive();
+			LogScope += LogScope << MemoryScope.Tag(InTag);
+			PrevTag = GActiveTag;
+			GActiveTag = InTag;
+		}
+	}
+}
+
+////////////////////////////////////////////////////////////////////////////////
+FMemScope::~FMemScope()
+{
+	if (Inner.bActive)
+	{
+		GActiveTag = PrevTag;
+	}
+}
+
+////////////////////////////////////////////////////////////////////////////////
+FMemScopeRealloc::FMemScopeRealloc(uint64 InPtr)
+{
+	if (InPtr != 0 && TRACE_PRIVATE_CHANNELEXPR_IS_ENABLED(MemAllocChannel))
+	{
+		if (auto LogScope = FMemoryMemoryScopeReallocFields::LogScopeType::ScopedEnter<FMemoryMemoryScopeReallocFields>())
+		{
+			if (const auto& __restrict MemoryScope = *(FMemoryMemoryScopeReallocFields*)(&LogScope))
+			{
+				Inner.SetActive(), LogScope += LogScope << MemoryScope.Ptr(InPtr);
+			}
+		}
+	}
+}
+
+////////////////////////////////////////////////////////////////////////////////
+FMemScopeRealloc::~FMemScopeRealloc()
+{
+}
+
+
+
+/////////////////////////////////////////////////////////////////////////////////
+
+/**
+ * Utility class that manages tracing the specification of unique LLM tags 
+ * and custom name based tags.
+ */
+class FTagTrace
+{
+public:
+					FTagTrace(FMalloc* InMalloc);
+	void			AnnounceGenericTags();
+	void 			AnnounceTagDeclarations();
+	static void 	OnAnnounceTagDeclaration(FLLMTagDeclaration& TagDeclaration);
+	int32			AnnounceCustomTag(int32 Tag, int32 ParentTag, const ANSICHAR* Display);
+	int32 			AnnounceFNameTag(FName TagName);
+
+private:
+	/**
+	 * Allocator used to allocate from the untracked FMalloc instance.
+	 */
+	class FTagAllocator : public TSizedHeapAllocator<32>
+	{
+	public:
+
+		class ForAnyElementType : public TSizedHeapAllocator<32>::ForAnyElementType
+		{
+		public:
+			ForAnyElementType()
+				: Data(nullptr)
+			{}
+
+			FORCEINLINE FScriptContainerElement* GetAllocation() const { return Data; }
+			FORCEINLINE bool HasAllocation() const { return !!Data; }
+			void ResizeAllocation(SizeType PreviousNumElements, SizeType NumElements, SIZE_T NumBytesPerElement)
+			{
+				check(FTagTrace::Malloc);
+				if (Data || NumElements)
+				{
+					Data = (FScriptContainerElement*)FTagTrace::Malloc->Realloc(Data, NumElements * NumBytesPerElement);
+				}
+			}
+
+		private:
+			FScriptContainerElement* Data;
+		};
+
+		template<typename ElementType>
+		class ForElementType : public ForAnyElementType
+		{
+		public:
+			FORCEINLINE ElementType* GetAllocation() const
+			{
+				return (ElementType*)ForAnyElementType::GetAllocation();
+			}
+		};
+
+	};
+	typedef TSparseArrayAllocator<FTagAllocator> FTagSparseArrayAllocator;
+	typedef TInlineAllocator<1,FTagAllocator> FTagHashAllocator; 
+	typedef TSetAllocator<FTagSparseArrayAllocator, FTagHashAllocator> FTagSetAllocator;
+	typedef TSet<int32, DefaultKeyFuncs<int32>, FTagSetAllocator> FTagSet;
+
+	FCriticalSection		Cs;
+	FTagSet 				AnnouncedNames;
+	static FMalloc* 		Malloc;
+};
+
+FMalloc* FTagTrace::Malloc = nullptr;
+static FTagTrace* GTagTrace = nullptr;
+
+////////////////////////////////////////////////////////////////////////////////
+FTagTrace::FTagTrace(FMalloc* InMalloc) 
+{
+	Malloc = InMalloc;
+	AnnouncedNames.Reserve(1024);
+	AnnounceGenericTags();
+	AnnounceTagDeclarations();
+}
+
+////////////////////////////////////////////////////////////////////////////////
+void FTagTrace::AnnounceGenericTags()
+{
+	#define TRACE_TAG_SPEC(Enum,Str,Stat,Group,ParentTag)\
+	{\
+		const uint32 DisplayLen = FCStringAnsi::Strlen(Str);\
+		UE_TRACE_LOG(Memory, TagSpec, MemAllocChannel, DisplayLen * sizeof(ANSICHAR))\
+			<< TagSpec.Tag((int32) ELLMTag::Enum)\
+			<< TagSpec.Parent((int32) ParentTag)\
+			<< TagSpec.Display(Str, DisplayLen);\
+	}
+	LLM_ENUM_GENERIC_TAGS(TRACE_TAG_SPEC);
+	#undef TRACE_TAG_SPEC
+}
+
+////////////////////////////////////////////////////////////////////////////////
+void FTagTrace::AnnounceTagDeclarations()
+{
+	FLLMTagDeclaration*& List = FLLMTagDeclaration::GetList();
+	while (List)
+	{
+		OnAnnounceTagDeclaration(*List);
+		List = List->Next;
+	}
+	FLLMTagDeclaration::SetCreationCallback(FTagTrace::OnAnnounceTagDeclaration);
+}
+
+////////////////////////////////////////////////////////////////////////////////
+void FTagTrace::OnAnnounceTagDeclaration(FLLMTagDeclaration& TagDeclaration)
+{
+	TagDeclaration.ConstructUniqueName();
+	GTagTrace->AnnounceFNameTag(TagDeclaration.GetUniqueName());
+}
+
+////////////////////////////////////////////////////////////////////////////////
+int32 FTagTrace::AnnounceFNameTag(FName Name)
+{
+	FScopeLock _(&Cs);
+	const int32 NameIndex = Name.GetDisplayIndex().ToUnstableInt();
+	if (AnnouncedNames.Contains(NameIndex))
+	{
+		return NameIndex;
+	}
+
+	AnnouncedNames.Emplace(NameIndex);
+	ANSICHAR NameString[NAME_SIZE];
+	Name.GetPlainANSIString(NameString);
+	return AnnounceCustomTag(NameIndex, -1, NameString);
+}
+
+////////////////////////////////////////////////////////////////////////////////
+int32 FTagTrace::AnnounceCustomTag(int32 Tag, int32 ParentTag, const ANSICHAR* Display)
+{		
+	const uint32 DisplayLen = FCStringAnsi::Strlen(Display);
+	UE_TRACE_LOG(Memory, TagSpec, MemAllocChannel, DisplayLen * sizeof(ANSICHAR))
+		<< TagSpec.Tag(Tag)
+		<< TagSpec.Parent(ParentTag)
+		<< TagSpec.Display(Display, DisplayLen);
+	return Tag;
+}
+
+#endif //USE_MEMORY_TRACE_TAGS
+
+
+////////////////////////////////////////////////////////////////////////////////
+void MemoryTrace_InitTags(FMalloc* InMalloc)
+{
+#if USE_MEMORY_TRACE_TAGS
+	GTagTrace = (FTagTrace*)InMalloc->Malloc(sizeof(FTagTrace), alignof(FTagTrace));
+	new (GTagTrace) FTagTrace(InMalloc);
+#endif
+}
+
+////////////////////////////////////////////////////////////////////////////////
+int32 MemoryTrace_AnnounceCustomTag(int32 Tag, int32 ParentTag, const TCHAR* Display)
+{
+#if USE_MEMORY_TRACE_TAGS
+	//todo: How do we check if tag trace is active?
+	if (GTagTrace)
+	{
+		return GTagTrace->AnnounceCustomTag(Tag, ParentTag, TCHAR_TO_ANSI(Display));
+	}
+#endif
+	return -1;
+}
+
+////////////////////////////////////////////////////////////////////////////////
+int32 MemoryTrace_AnnounceFNameTag(FName TagName)
+{
+#if USE_MEMORY_TRACE_TAGS
+	if (GTagTrace)
+	{
+		return GTagTrace->AnnounceFNameTag(TagName);
+	}
+#endif
+	return -1;
+}
+
+////////////////////////////////////////////////////////////////////////////////
+int32 MemoryTrace_GetActiveTag()
+{
+#if USE_MEMORY_TRACE_TAGS
+	return GActiveTag;
+#else
+	return -1;
+#endif
+}
+
diff --git a/Engine/Source/Runtime/Core/Private/ProfilingDebugging/Windows/MemoryTrace.cpp b/Engine/Source/Runtime/Core/Private/ProfilingDebugging/Windows/MemoryTrace.cpp
--- a/Engine/Source/Runtime/Core/Private/ProfilingDebugging/Windows/MemoryTrace.cpp
+++ b/Engine/Source/Runtime/Core/Private/ProfilingDebugging/Windows/MemoryTrace.cpp
@@ -33,6 +33,7 @@ TRACELOG_API void Update();
 void	Backtracer_Create(FMalloc*);
 void	Backtracer_Initialize();
 void*	Backtracer_GetBacktraceId(void*);
+void 	MemoryTrace_InitTags(FMalloc*);
 
 ////////////////////////////////////////////////////////////////////////////////
 #if defined(_MSC_VER)
@@ -798,8 +799,12 @@ FMalloc* MemoryTrace_Create(FMalloc* InMalloc)
 		// getting ticked.
 		atexit([] () { GAllocationTrace->EnableTracePump(); });
 
+		GAllocationTrace.Construct();
 		GAllocationTrace->Initialize();
+
+		MemoryTrace_InitTags(InMalloc);
 #if USE_OVERVIEW_TRACE
+		GSummaryTrace.Construct();
 		GSummaryTrace->Initialize();
 #endif
 
@@ -815,11 +820,6 @@ FMalloc* MemoryTrace_Create(FMalloc* InMalloc)
 		static FUndestructed<FMallocWrapper> MemoryTrace;
 		MemoryTrace.Construct(InMalloc, bLight);
 
-		GAllocationTrace.Construct();
-#if USE_OVERVIEW_TRACE
-		GSummaryTrace.Construct();
-#endif
-
 		return &MemoryTrace;
 	}
 
diff --git a/Engine/Source/Runtime/Core/Public/Async/AsyncWork.h b/Engine/Source/Runtime/Core/Public/Async/AsyncWork.h
--- a/Engine/Source/Runtime/Core/Public/Async/AsyncWork.h
+++ b/Engine/Source/Runtime/Core/Public/Async/AsyncWork.h
@@ -63,6 +63,10 @@ class FAutoDeleteAsyncTask
 	TTask Task;
 	/** optional LLM tag */
 	LLM(const UE::LLMPrivate::FTagData* InheritedLLMTag);
+#if USE_MEMORY_TRACE_TAGS
+	/** optional Trace tag */
+	int32 InheritedTraceTag;
+#endif
 
 	/* Generic start function, not called directly
 	 * @param bForceSynchronous if true, this job will be started synchronously, now, on this thread
@@ -70,6 +74,9 @@ class FAutoDeleteAsyncTask
 	void Start(bool bForceSynchronous, FQueuedThreadPool* InQueuedPool)
 	{
 		LLM(InheritedLLMTag = FLowLevelMemTracker::bIsDisabled ? nullptr : FLowLevelMemTracker::Get().GetActiveTagData(ELLMTracker::Default));
+#if USE_MEMORY_TRACE_TAGS
+		InheritedTraceTag = MemoryTrace_GetActiveTag();
+#endif
 
 		FPlatformMisc::MemoryBarrier();
 		FQueuedThreadPool* QueuedPool = InQueuedPool;
@@ -94,6 +101,7 @@ class FAutoDeleteAsyncTask
 	void DoWork()
 	{
 		LLM_SCOPE(InheritedLLMTag);
+		UE_MEMSCOPE(InheritedTraceTag, ELLMTracker::Default);
 		FScopeCycleCounter Scope(Task.GetStatId(), true);
 
 		Task.DoWork();
diff --git a/Engine/Source/Runtime/Core/Public/Async/TaskGraphInterfaces.h b/Engine/Source/Runtime/Core/Public/Async/TaskGraphInterfaces.h
--- a/Engine/Source/Runtime/Core/Public/Async/TaskGraphInterfaces.h
+++ b/Engine/Source/Runtime/Core/Public/Async/TaskGraphInterfaces.h
@@ -446,6 +446,9 @@ protected:
 		, NumberOfPrerequistitesOutstanding(InNumberOfPrerequistitesOutstanding + 1) // + 1 is not a prerequisite, it is a lock to prevent it from executing while it is getting prerequisites, one it is safe to execute, call PrerequisitesComplete
 	{
 		checkThreadGraph(LifeStage.Increment() == int32(LS_Contructed));
+#if USE_MEMORY_TRACE_TAGS
+		InheritedTraceTag = MemoryTrace_GetActiveTag();
+#endif
 		LLM(InheritedLLMTag = FLowLevelMemTracker::bIsDisabled ? nullptr : FLowLevelMemTracker::Get().GetActiveTagData(ELLMTracker::Default));
 	}
 	/** 
@@ -523,6 +526,7 @@ private:
 	FORCEINLINE void Execute(TArray<FBaseGraphTask*>& NewTasks, ENamedThreads::Type CurrentThread)
 	{
 		LLM_SCOPE(InheritedLLMTag);
+		UE_MEMSCOPE(InheritedTraceTag, ELLMTracker::Default);
 		checkThreadGraph(LifeStage.Increment() == int32(LS_Executing));
 		ExecuteTask(NewTasks, CurrentThread);
 	}
@@ -561,7 +565,9 @@ private:
 	FThreadSafeCounter			LifeStage;
 
 #endif
-
+#if USE_MEMORY_TRACE_TAGS
+	int32 InheritedTraceTag;
+#endif
 	LLM(const UE::LLMPrivate::FTagData* InheritedLLMTag);
 };
 
diff --git a/Engine/Source/Runtime/Core/Public/HAL/LowLevelMemTracker.h b/Engine/Source/Runtime/Core/Public/HAL/LowLevelMemTracker.h
--- a/Engine/Source/Runtime/Core/Public/HAL/LowLevelMemTracker.h
+++ b/Engine/Source/Runtime/Core/Public/HAL/LowLevelMemTracker.h
@@ -3,6 +3,7 @@
 #pragma once
 
 #include "CoreTypes.h"
+#include "ProfilingDebugging/TagTrace.h"
 
 #ifndef ALLOW_LOW_LEVEL_MEM_TRACKER_IN_TEST
 	#define ALLOW_LOW_LEVEL_MEM_TRACKER_IN_TEST 0
@@ -283,27 +284,31 @@ extern FName LLMGetTagStat(ELLMTag Tag);
 /**
  * LLM scope macros
  */
-#define LLM_SCOPE(Tag)												FLLMScope SCOPE_NAME(Tag, false /* bIsStatTag */, ELLMTagSet::None, ELLMTracker::Default);
-#define LLM_SCOPE_BYNAME(Tag) static FName PREPROCESSOR_JOIN(LLMScope_Name,__LINE__)(Tag);	\
-																	FLLMScope SCOPE_NAME(PREPROCESSOR_JOIN(LLMScope_Name,__LINE__), false /* bIsStatTag */, ELLMTagSet::None, ELLMTracker::Default);
-#define LLM_SCOPE_BYTAG(TagDeclName)								FLLMScope SCOPE_NAME(PREPROCESSOR_JOIN(LLMTagDeclaration_, TagDeclName).GetUniqueName(), false /* bIsStatTag */, ELLMTagSet::None, ELLMTracker::Default);
+#define LLM_SCOPE(Tag)												FLLMScope SCOPE_NAME(Tag, false /* bIsStatTag */, ELLMTagSet::None, ELLMTracker::Default);\
+																	UE_MEMSCOPE(Tag, ELLMTracker::Default /* Tracker */);
+#define LLM_SCOPE_BYNAME(Tag) 										static FName PREPROCESSOR_JOIN(LLMScope_Name,__LINE__)(Tag);\
+																	FLLMScope SCOPE_NAME(PREPROCESSOR_JOIN(LLMScope_Name,__LINE__), false /* bIsStatTag */, ELLMTagSet::None, ELLMTracker::Default);\
+																	UE_MEMSCOPE(PREPROCESSOR_JOIN(LLMScope_Name,__LINE__), ELLMTracker::Default /* Tracker */);
+#define LLM_SCOPE_BYTAG(TagDeclName)								FLLMScope SCOPE_NAME(PREPROCESSOR_JOIN(LLMTagDeclaration_, TagDeclName).GetUniqueName(), false /* bIsStatTag */, ELLMTagSet::None, ELLMTracker::Default);\
+																	UE_MEMSCOPE(PREPROCESSOR_JOIN(LLMTagDeclaration_,TagDeclName).GetUniqueName(), ELLMTracker::Default /* Tracker */);
 #define LLM_PLATFORM_SCOPE(Tag)										FLLMScope SCOPE_NAME(Tag, false /* bIsStatTag */, ELLMTagSet::None, ELLMTracker::Platform);
-#define LLM_PLATFORM_SCOPE_BYNAME(Tag) static FName PREPROCESSOR_JOIN(LLMScope_Name,__LINE__)(Tag); \
+#define LLM_PLATFORM_SCOPE_BYNAME(Tag) 								static FName PREPROCESSOR_JOIN(LLMScope_Name,__LINE__)(Tag);\
 																	FLLMScope SCOPE_NAME(PREPROCESSOR_JOIN(LLMLLMScope_NameScope,__LINE__), false /* bIsStatTag */, ELLMTagSet::None, ELLMTracker::Platform);
 #define LLM_PLATFORM_SCOPE_BYTAG(TagDeclName)						FLLMScope SCOPE_NAME(PREPROCESSOR_JOIN(LLMTagDeclaration_, TagDeclName).GetUniqueName(), false /* bIsStatTag */, ELLMTagSet::None, ELLMTracker::Platform);
 
  /**
  * LLM Pause scope macros
  */
-#define LLM_SCOPED_PAUSE_TRACKING(AllocType) FLLMPauseScope SCOPE_NAME(ELLMTag::Untagged, false /* bIsStatTag */, 0, ELLMTracker::Max, AllocType);
-#define LLM_SCOPED_PAUSE_TRACKING_FOR_TRACKER(Tracker, AllocType) FLLMPauseScope SCOPE_NAME(ELLMTag::Untagged, false /* bIsStatTag */, 0, Tracker, AllocType);
+#define LLM_SCOPED_PAUSE_TRACKING(AllocType) 						FLLMPauseScope SCOPE_NAME(ELLMTag::Untagged, false /* bIsStatTag */, 0, ELLMTracker::Max, AllocType);
+#define LLM_SCOPED_PAUSE_TRACKING_FOR_TRACKER(Tracker, AllocType) 	FLLMPauseScope SCOPE_NAME(ELLMTag::Untagged, false /* bIsStatTag */, 0, Tracker, AllocType);
 #define LLM_SCOPED_PAUSE_TRACKING_WITH_ENUM_AND_AMOUNT(Tag, Amount, Tracker, AllocType) FLLMPauseScope SCOPE_NAME(Tag, false /* bIsStatTag */, Amount, Tracker, AllocType);
 
 /**
  * LLM realloc scope macros. Used when reallocating a pointer and you wish to retain the tagging from the source pointer
  */
-#define LLM_REALLOC_SCOPE(Ptr) FLLMScopeFromPtr SCOPE_NAME(Ptr, ELLMTracker::Default);
-#define LLM_REALLOC_PLATFORM_SCOPE(Ptr) FLLMScopeFromPtr SCOPE_NAME(Ptr, ELLMTracker::Platform);
+#define LLM_REALLOC_SCOPE(Ptr)										FLLMScopeFromPtr SCOPE_NAME(Ptr, ELLMTracker::Default);\
+																	UE_MEMSCOPE_REALLOC(Ptr, ELLMTracker::Default);
+#define LLM_REALLOC_PLATFORM_SCOPE(Ptr)								FLLMScopeFromPtr SCOPE_NAME(Ptr, ELLMTracker::Platform);
 
 /**
  * LLM tag dumping, to help with identifying mis-tagged items. Probably don't want to check in with these in use!
@@ -719,6 +724,7 @@ protected:
 	FLLMTagDeclaration* Next = nullptr;
 
 	friend class FLowLevelMemTracker;
+	friend class FTagTrace;
 };
 
 
diff --git a/Engine/Source/Runtime/Core/Public/ProfilingDebugging/TagTrace.h b/Engine/Source/Runtime/Core/Public/ProfilingDebugging/TagTrace.h
new file mode 100644
--- /dev/null
+++ b/Engine/Source/Runtime/Core/Public/ProfilingDebugging/TagTrace.h
@@ -0,0 +1,86 @@
+// Copyright Epic Games, Inc. All Rights Reserved.
+#pragma once
+#include "CoreTypes.h"
+#include "Trace/Trace.h"
+#include "Trace/Detail/LogScope.h"
+
+////////////////////////////////////////////////////////////////////////////////
+// Fwd declare ELLMTag
+enum class ELLMTag : uint8;
+// Fwd declare LLM private tag data
+namespace UE {
+	namespace LLMPrivate {
+		class FTagData;
+	}
+}
+
+////////////////////////////////////////////////////////////////////////////////
+CORE_API int32	MemoryTrace_AnnounceCustomTag(int32 Tag, int32 ParentTag, const TCHAR* Display);
+CORE_API int32	MemoryTrace_AnnounceFNameTag(FName TagName);
+CORE_API int32	MemoryTrace_GetActiveTag();
+
+////////////////////////////////////////////////////////////////////////////////
+
+#if !defined(USE_MEMORY_TRACE_TAGS) && UE_TRACE_ENABLED
+	#if PLATFORM_WINDOWS
+		#define USE_MEMORY_TRACE_TAGS 1
+	#endif
+#endif
+
+#if !defined(USE_MEMORY_TRACE_TAGS)
+	#define USE_MEMORY_TRACE_TAGS 0
+#endif
+
+#if USE_MEMORY_TRACE_TAGS
+
+////////////////////////////////////////////////////////////////////////////////
+
+/**
+  * Used to associate any allocation within this scope to a given tag.
+  *
+  * We need to be able to convert the three types of inputs to LLM scopes:
+  * - ELLMTag, an uint8 with fixed categories. There are three sub ranges
+	  Generic tags, platform and project tags.
+  * - FName, free form string, for example a specific asset.
+  * - TagData, an opaque pointer from LLM.
+  *
+  */
+class FMemScope
+{
+public:
+	CORE_API FMemScope(int32 InTag);
+	CORE_API FMemScope(ELLMTag InTag);
+	CORE_API FMemScope(FName InName);
+	CORE_API FMemScope(const UE::LLMPrivate::FTagData* TagData);
+	CORE_API ~FMemScope();
+private:
+	void ActivateScope(int32 InTag);
+	UE::Trace::Private::FScopedLogScope Inner;
+	int32 PrevTag;
+};
+
+/**
+ * Used order to keep the tag for memory that is being reallocated.
+ */
+class FMemScopeRealloc
+{
+public:
+	CORE_API FMemScopeRealloc(uint64 InPtr);
+	CORE_API ~FMemScopeRealloc();
+private:
+	UE::Trace::Private::FScopedLogScope Inner;
+};
+
+////////////////////////////////////////////////////////////////////////////////
+
+#define UE_MEMSCOPE(InTag, InTracker)				FMemScope PREPROCESSOR_JOIN(MemScope,__LINE__)(InTag);
+#define UE_MEMSCOPE_REALLOC(InPtr, InTracker)		FMemScopeRealloc PREPROCESSOR_JOIN(MemReallocScope,__LINE__)((uint64)InPtr);
+
+#else // USE_MEMORY_TRACE_TAGS
+
+////////////////////////////////////////////////////////////////////////////////
+#define UE_MEMSCOPE(...)
+#define UE_MEMSCOPE_REALLOC(...)
+
+#endif // USE_MEMORY_TRACE_TAGS
+
