Add -tracefiletrunc cmdline option to allow for overwriting a trace named on the commane line.
This is desired for data center builds where we only care about one trace at a time, and uniquely naming or rotating the file name just makes for extra complexity and bug risk in calling code. Also, for incremental worker configurations we can leak disk consumption if we use a different file name every time.

Win/Mac/Linux all have the truncate file open flag set, so the platform file code should all allow for this.

Tested on windows.

diff --git a/Engine/Source/Runtime/Core/Private/ProfilingDebugging/TraceAuxiliary.cpp b/Engine/Source/Runtime/Core/Private/ProfilingDebugging/TraceAuxiliary.cpp
--- a/Engine/Source/Runtime/Core/Private/ProfilingDebugging/TraceAuxiliary.cpp
+++ b/Engine/Source/Runtime/Core/Private/ProfilingDebugging/TraceAuxiliary.cpp
@@ -51,6 +51,7 @@ public:
 	bool					Stop();
 	void					EnableChannels();
 	void					DisableChannels();
+	void					SetTruncateFile(bool bTruncateFile);
 
 private:
 	enum class EState : uint8
@@ -74,6 +75,7 @@ private:
 	TMap<uint32, FChannel>	Channels;
 	FString					TraceDest;
 	EState					State = EState::Stopped;
+	bool					bTruncateFile = false;
 };
 
 static FTraceAuxiliaryImpl GTraceAuxiliary;
@@ -229,6 +231,12 @@ void FTraceAuxiliaryImpl::DisableChannels()
 	}
 }
 
+////////////////////////////////////////////////////////////////////////////////
+void FTraceAuxiliaryImpl::SetTruncateFile(bool bNewTruncateFileState)
+{
+	bTruncateFile = bNewTruncateFileState;
+}
+
 ////////////////////////////////////////////////////////////////////////////////
 bool FTraceAuxiliaryImpl::SendToHost(const TCHAR* Host)
 {
@@ -280,7 +288,7 @@ bool FTraceAuxiliaryImpl::WriteToFile(const TCHAR* Path)
 		return false;
 	}
 
-	if (FileManager.FileExists(*WritePath))
+	if (!bTruncateFile && FileManager.FileExists(*WritePath))
 	{
 		UE_LOG(LogCore, Warning, TEXT("Trace file '%s' already exists"), *WritePath);
 		return false;
@@ -427,6 +435,7 @@ void FTraceAuxiliary::Initialize(const TCHAR* CommandLine)
 	}
 	else if (FParse::Value(CommandLine, TEXT("-tracefile="), Parameter))
 	{
+		GTraceAuxiliary.SetTruncateFile(FParse::Param(CommandLine, TEXT("tracefiletrunc")));
 		GTraceAuxiliary.Connect(ETraceConnectType::File, *Parameter);
 	}
 	else if (FParse::Param(CommandLine, TEXT("tracefile")))
