A few no-unity include fixes.

#rnx

diff --git a/Engine/Source/Runtime/Core/Private/ProfilingDebugging/TagTrace.cpp b/Engine/Source/Runtime/Core/Private/ProfilingDebugging/TagTrace.cpp
--- a/Engine/Source/Runtime/Core/Private/ProfilingDebugging/TagTrace.cpp
+++ b/Engine/Source/Runtime/Core/Private/ProfilingDebugging/TagTrace.cpp
@@ -1,8 +1,10 @@
 // Copyright Epic Games, Inc. All Rights Reserved.
 #include "ProfilingDebugging/TagTrace.h"
 
+#include "Containers/Set.h"
 #include "CoreTypes.h"
 #include "MemoryTrace.h"
+#include "Misc/ScopeLock.h"
 #include "HAL/LowLevelMemTracker.h"
 #include "Trace/Trace.inl"
 
diff --git a/Engine/Source/Runtime/Core/Public/ProfilingDebugging/TagTrace.h b/Engine/Source/Runtime/Core/Public/ProfilingDebugging/TagTrace.h
--- a/Engine/Source/Runtime/Core/Public/ProfilingDebugging/TagTrace.h
+++ b/Engine/Source/Runtime/Core/Public/ProfilingDebugging/TagTrace.h
@@ -1,6 +1,7 @@
 // Copyright Epic Games, Inc. All Rights Reserved.
 #pragma once
 #include "CoreTypes.h"
+#include "UObject/NameTypes.h"
 #include "Trace/Trace.h"
 #include "Trace/Detail/LogScope.h"
 
