Fix application teardown asserts on macOS.

- Add FTaskTagScope::SetTagStaticInit()
- Reset application main thread state to static init in applicationWillTerminate:. This enables all the Is*Thread() checks to function properly during global (or local static) object destruction.

#rnx

diff --git a/Engine/Source/Runtime/Core/Private/HAL/ThreadingBase.cpp b/Engine/Source/Runtime/Core/Private/HAL/ThreadingBase.cpp
--- a/Engine/Source/Runtime/Core/Private/HAL/ThreadingBase.cpp
+++ b/Engine/Source/Runtime/Core/Private/HAL/ThreadingBase.cpp
@@ -57,6 +57,11 @@ void FTaskTagScope::SetTagNone()
 	ActiveTaskTag = ETaskTag::ENone;
 }
 
+void FTaskTagScope::SetTagStaticInit()
+{
+	ActiveTaskTag = ETaskTag::EStaticInit;
+}
+
 FTaskTagScope::FTaskTagScope(bool InTagOnlyIfNone, ETaskTag InTag) : Tag(InTag), TagOnlyIfNone(InTagOnlyIfNone)
 {
 	checkf(Tag != ETaskTag::ENone, TEXT("None cannot be used as a Tag"));
diff --git a/Engine/Source/Runtime/Core/Public/CoreGlobals.h b/Engine/Source/Runtime/Core/Public/CoreGlobals.h
--- a/Engine/Source/Runtime/Core/Public/CoreGlobals.h
+++ b/Engine/Source/Runtime/Core/Public/CoreGlobals.h
@@ -530,6 +530,13 @@ public:
 	 */
 	static void CORE_API SetTagNone();
 
+	/**
+	 * Restore the ETaskTag::StaticInit tag so that the destructors of global
+	 * (or local static) C++ objects function properly when checking thread
+	 * state (from functions like IsInGameThread()).
+	 */
+	static void CORE_API SetTagStaticInit();
+
 protected:
 	CORE_API FTaskTagScope(bool InTagOnlyIfNone, ETaskTag InTag);
 
diff --git a/Engine/Source/Runtime/Launch/Private/Mac/LaunchMac.cpp b/Engine/Source/Runtime/Launch/Private/Mac/LaunchMac.cpp
--- a/Engine/Source/Runtime/Launch/Private/Mac/LaunchMac.cpp
+++ b/Engine/Source/Runtime/Launch/Private/Mac/LaunchMac.cpp
@@ -201,6 +201,11 @@ static int32 MacOSVersionCompare(const NSOperatingSystemVersion& VersionA, const
 	}
 }
 
+- (void) applicationWillTerminate:(NSNotification*)notification
+{
+	FTaskTagScope::SetTagStaticInit();
+}
+
 - (void) runGameThread:(id)Arg
 {
 	bool bIsBuildMachine = false;
