Changed IsRunningCookCommandlet to return value initialized when parsing the command line.

diff --git a/Engine/Source/Runtime/Core/Private/Misc/CoreGlobals.cpp b/Engine/Source/Runtime/Core/Private/Misc/CoreGlobals.cpp
--- a/Engine/Source/Runtime/Core/Private/Misc/CoreGlobals.cpp
+++ b/Engine/Source/Runtime/Core/Private/Misc/CoreGlobals.cpp
@@ -160,6 +160,10 @@ bool GIsReinstancing = false;
  */
 bool GIsRunningUnattendedScript = false;
 
+#if WITH_EDITOR
+bool					PRIVATE_GIsRunningCookCommandlet	= false;				/** Whether this executable is running the cook commandlet */
+#endif
+
 #if WITH_ENGINE
 bool					PRIVATE_GIsRunningCommandlet		= false;				/** Whether this executable is running a commandlet (custom command-line processing code) */
 bool					PRIVATE_GAllowCommandletRendering	= false;				/** If true, initialise RHI and set up scene for rendering even when running a commandlet. */
@@ -423,13 +427,6 @@ void CORE_API EnableEmitDrawEventsOnlyOnCommandlist()
 	GEmitDrawEvents = true;
 }
 
-bool CORE_API IsRunningCookCommandlet()
-{
-	FString Commandline = FCommandLine::Get();
-	const bool bIsCookCommandlet = IsRunningCommandlet() && Commandline.Contains(TEXT("run=cook"));
-	return bIsCookCommandlet;
-}
-
 void ToggleGDebugPUCrashedFlag(const TArray<FString>& Args)
 {
 	GIsGPUCrashed = !GIsGPUCrashed;
diff --git a/Engine/Source/Runtime/Core/Public/CoreGlobals.h b/Engine/Source/Runtime/Core/Public/CoreGlobals.h
--- a/Engine/Source/Runtime/Core/Public/CoreGlobals.h
+++ b/Engine/Source/Runtime/Core/Public/CoreGlobals.h
@@ -195,6 +195,9 @@ extern CORE_API bool GFirstFrameIntraFrameDebugging;
 
 #endif // WITH_EDITORONLY_DATA
 
+#if WITH_EDITOR
+extern CORE_API bool PRIVATE_GIsRunningCookCommandlet;
+#endif
 
 /**
 * Check to see if this executable is running a commandlet (custom command-line processing code in an editor-like environment)
@@ -211,7 +214,14 @@ FORCEINLINE bool IsRunningCommandlet()
 /**
 * Check to see if this executable is running the cookcommandlet
 */
-CORE_API bool IsRunningCookCommandlet();
+FORCEINLINE bool IsRunningCookCommandlet()
+{
+#if WITH_EDITOR
+	return PRIVATE_GIsRunningCookCommandlet;
+#else
+	return false;
+#endif
+}
 
 /**
  * Check to see if we should initialise RHI and set up scene for rendering even when running a commandlet.
diff --git a/Engine/Source/Runtime/Engine/Private/WorldPartition/DataLayer/WorldDataLayers.cpp b/Engine/Source/Runtime/Engine/Private/WorldPartition/DataLayer/WorldDataLayers.cpp
--- a/Engine/Source/Runtime/Engine/Private/WorldPartition/DataLayer/WorldDataLayers.cpp
+++ b/Engine/Source/Runtime/Engine/Private/WorldPartition/DataLayer/WorldDataLayers.cpp
@@ -25,9 +25,7 @@ const AWorldDataLayers* AWorldDataLayers::Get(UWorld* World)
 		// Prepare flags for actor iterator. Don't use default Flags because it uses EActorIteratorFlags::OnlyActiveLevels 
 		// which will make this code return no actor when cooking (because world is not initialized)
 		EActorIteratorFlags Flags = EActorIteratorFlags::SkipPendingKill;
-#if WITH_EDITOR
 		if (!IsRunningCookCommandlet())
-#endif
 		{
 			Flags |= EActorIteratorFlags::OnlyActiveLevels;
 		}
diff --git a/Engine/Source/Runtime/Engine/Private/WorldPartition/WorldPartition.cpp b/Engine/Source/Runtime/Engine/Private/WorldPartition/WorldPartition.cpp
--- a/Engine/Source/Runtime/Engine/Private/WorldPartition/WorldPartition.cpp
+++ b/Engine/Source/Runtime/Engine/Private/WorldPartition/WorldPartition.cpp
@@ -281,9 +281,7 @@ void UWorldPartition::Initialize(UWorld* InWorld, const FTransform& InTransform)
 
 	InitState = EWorldPartitionInitState::Initialized;
 
-#if WITH_EDITOR
 	if (!IsRunningCookCommandlet())
-#endif
 	{
 		UWorldPartitionSubsystem* WorldPartitionSubsystem = World->GetSubsystem<UWorldPartitionSubsystem>();
 		WorldPartitionSubsystem->RegisterWorldPartition(this);
@@ -348,9 +346,7 @@ void UWorldPartition::Uninitialize()
 		EditorHash = nullptr;
 #endif		
 
-#if WITH_EDITOR
 		if (!IsRunningCookCommandlet())
-#endif
 		{
 			UWorldPartitionSubsystem* WorldPartitionSubsystem = World->GetSubsystem<UWorldPartitionSubsystem>();
 			WorldPartitionSubsystem->UnregisterWorldPartition(this);
diff --git a/Engine/Source/Runtime/Launch/Private/LaunchEngineLoop.cpp b/Engine/Source/Runtime/Launch/Private/LaunchEngineLoop.cpp
--- a/Engine/Source/Runtime/Launch/Private/LaunchEngineLoop.cpp
+++ b/Engine/Source/Runtime/Launch/Private/LaunchEngineLoop.cpp
@@ -2575,6 +2575,8 @@ PRAGMA_ENABLE_DEPRECATION_WARNINGS
 	bool bForceEnableHighDPI = false;
 #if WITH_EDITOR
 	bForceEnableHighDPI = FPIEPreviewDeviceModule::IsRequestingPreviewDevice();
+	// Update PRIVATE_GIsRunningCookCommandlet, we assume that from here, PRIVATE_GIsRunningCommandlet won't change
+	PRIVATE_GIsRunningCookCommandlet = IsRunningCommandlet() && bIsCook;
 #endif
 
 	// This must be called before any window (including the splash screen is created
