Fixed compilation errors if LLM is off but Trace is on.

#rnx

diff --git a/Engine/Source/Runtime/Core/Public/ProfilingDebugging/TagTrace.h b/Engine/Source/Runtime/Core/Public/ProfilingDebugging/TagTrace.h
--- a/Engine/Source/Runtime/Core/Public/ProfilingDebugging/TagTrace.h
+++ b/Engine/Source/Runtime/Core/Public/ProfilingDebugging/TagTrace.h
@@ -1,6 +1,7 @@
 // Copyright Epic Games, Inc. All Rights Reserved.
 #pragma once
 #include "CoreTypes.h"
+#include "HAL/LowLevelMemTracker.h"
 #include "UObject/NameTypes.h"
 #include "Trace/Trace.h"
 #include "Trace/Detail/LogScope.h"
@@ -22,7 +23,7 @@ CORE_API int32	MemoryTrace_GetActiveTag();
 
 ////////////////////////////////////////////////////////////////////////////////
 
-#if !defined(USE_MEMORY_TRACE_TAGS) && UE_TRACE_ENABLED
+#if !defined(USE_MEMORY_TRACE_TAGS) && UE_TRACE_ENABLED && ENABLE_LOW_LEVEL_MEM_TRACKER
 	#if PLATFORM_WINDOWS
 		#define USE_MEMORY_TRACE_TAGS 1
 	#endif
