Rebuild the package localization cache for all current cultures when the culture changes

Previously it would only rebuild the cache for the current language, which could leave the asset group cultures uncached and lead to the incorrect localized package being used.

diff --git a/Engine/Source/Runtime/Core/Private/Internationalization/Internationalization.cpp b/Engine/Source/Runtime/Core/Private/Internationalization/Internationalization.cpp
--- a/Engine/Source/Runtime/Core/Private/Internationalization/Internationalization.cpp
+++ b/Engine/Source/Runtime/Core/Private/Internationalization/Internationalization.cpp
@@ -158,7 +158,7 @@ bool FInternationalization::SetCurrentAssetGroupCulture(const FName& InAssetGrou
 
 FCultureRef FInternationalization::GetCurrentAssetGroupCulture(const FName& InAssetGroupName) const
 {
-	for (const auto& CurrentAssetGroupCulturePair : CurrentAssetGroupCultures)
+	for (const TTuple<FName, FCulturePtr>& CurrentAssetGroupCulturePair : CurrentAssetGroupCultures)
 	{
 		if (CurrentAssetGroupCulturePair.Key == InAssetGroupName)
 		{
@@ -176,13 +176,38 @@ void FInternationalization::ClearCurrentAssetGroupCulture(const FName& InAssetGr
 	});
 }
 
+TArray<FCultureRef> FInternationalization::GetCurrentCultures(const bool bIncludeLanguage, const bool bIncludeLocale, const bool bIncludeAssetGroups) const
+{
+	TArray<FCultureRef> CurrentCultures;
+
+	if (bIncludeLanguage)
+	{
+		CurrentCultures.AddUnique(CurrentLanguage.ToSharedRef());
+	}
+
+	if (bIncludeLocale)
+	{
+		CurrentCultures.AddUnique(CurrentLocale.ToSharedRef());
+	}
+
+	if (bIncludeAssetGroups)
+	{
+		for (const TTuple<FName, FCulturePtr>& CurrentAssetGroupCulturePair : CurrentAssetGroupCultures)
+		{
+			CurrentCultures.AddUnique(CurrentAssetGroupCulturePair.Value.ToSharedRef());
+		}
+	}
+
+	return CurrentCultures;
+}
+
 void FInternationalization::BackupCultureState(FCultureStateSnapshot& OutSnapshot) const
 {
 	OutSnapshot.Language = CurrentLanguage->GetName();
 	OutSnapshot.Locale = CurrentLocale->GetName();
 
 	OutSnapshot.AssetGroups.Reset(CurrentAssetGroupCultures.Num());
-	for (const auto& CurrentAssetGroupCulturePair : CurrentAssetGroupCultures)
+	for (const TTuple<FName, FCulturePtr>& CurrentAssetGroupCulturePair : CurrentAssetGroupCultures)
 	{
 		OutSnapshot.AssetGroups.Add(MakeTuple(CurrentAssetGroupCulturePair.Key, CurrentAssetGroupCulturePair.Value->GetName()));
 	}
diff --git a/Engine/Source/Runtime/Core/Public/Internationalization/Internationalization.h b/Engine/Source/Runtime/Core/Public/Internationalization/Internationalization.h
--- a/Engine/Source/Runtime/Core/Public/Internationalization/Internationalization.h
+++ b/Engine/Source/Runtime/Core/Public/Internationalization/Internationalization.h
@@ -147,6 +147,11 @@ public:
 		return InvariantCulture.ToSharedRef();
 	}
 
+	/**
+	 * Get the current cultures in use, optionally including the current language, locale, and any asset groups.
+	 */
+	CORE_API TArray<FCultureRef> GetCurrentCultures(const bool bIncludeLanguage, const bool bIncludeLocale, const bool bIncludeAssetGroups) const;
+
 	/**
 	 * Backup the current culture state to the given snapshot struct.
 	 */
diff --git a/Engine/Source/Runtime/CoreUObject/Private/Internationalization/PackageLocalizationCache.cpp b/Engine/Source/Runtime/CoreUObject/Private/Internationalization/PackageLocalizationCache.cpp
--- a/Engine/Source/Runtime/CoreUObject/Private/Internationalization/PackageLocalizationCache.cpp
+++ b/Engine/Source/Runtime/CoreUObject/Private/Internationalization/PackageLocalizationCache.cpp
@@ -359,11 +359,15 @@ void FPackageLocalizationCache::HandleCultureChanged()
 	const FString CurrentCultureName = FInternationalization::Get().GetCurrentLanguage()->GetName();
 	CurrentCultureCache = FindOrAddCacheForCulture_NoLock(CurrentCultureName);
 
-	if (CurrentCultureCache.IsValid())
+	// We expect culture changes to happen on the game thread, so update the cache now while it is likely safe to do so
+	// (ConditionalUpdateCache will internally check that this is currently the game thread before allowing the update)
+	const TArray<FCultureRef> CurrentCultures = FInternationalization::Get().GetCurrentCultures(/*bIncludeLanguage*/true, /*bIncludeLocale*/false, /*bIncludeAssetGroups*/true);
+	for (const FCultureRef& CurrentCulture : CurrentCultures)
 	{
-		// We expect culture changes to happen on the game thread, so update the cache now while it is likely safe to do so
-		// (ConditionalUpdateCache will internally check that this is currently the game thread before allowing the update)
-		CurrentCultureCache->ConditionalUpdateCache();
+		if (TSharedPtr<FPackageLocalizationCultureCache> CultureCache = FindOrAddCacheForCulture_NoLock(CurrentCulture->GetName()))
+		{
+			CultureCache->ConditionalUpdateCache();
+		}
 	}
 
 	ConditionalUpdatePackageNameToAssetGroupCache_NoLock();
