CompactBinary: Add context to the tests for different field types

#rnx

diff --git a/Engine/Source/Runtime/Core/Private/Tests/Serialization/CompactBinaryTest.cpp b/Engine/Source/Runtime/Core/Private/Tests/Serialization/CompactBinaryTest.cpp
--- a/Engine/Source/Runtime/Core/Private/Tests/Serialization/CompactBinaryTest.cpp
+++ b/Engine/Source/Runtime/Core/Private/Tests/Serialization/CompactBinaryTest.cpp
@@ -81,94 +81,94 @@ protected:
 	}
 
 	template <typename T, typename Default>
-	void TestFieldAsTypeNoClone(FCbField& Field, T (FCbField::*AsTypeFn)(Default), T ExpectedValue = T(), T DefaultValue = T(), ECbFieldError ExpectedError = ECbFieldError::None)
+	void TestFieldAsTypeNoClone(const TCHAR* What, FCbField& Field, T (FCbField::*AsTypeFn)(Default), T ExpectedValue = T(), T DefaultValue = T(), ECbFieldError ExpectedError = ECbFieldError::None)
 	{
-		TestTypeError(Field, ExpectedError);
-		TestEqual(TEXT("FCbField::As[Type]()"), (Field.*AsTypeFn)(DefaultValue), ExpectedValue);
-		TestEqual(TEXT("FCbField::As[Type]() -> HasError()"), Field.HasError(), ExpectedError != ECbFieldError::None);
-		TestEqual(TEXT("FCbField::As[Type]() -> GetError()"), Field.GetError(), ExpectedError);
+		TestTypeError(What, Field, ExpectedError);
+		TestEqual(FString::Printf(TEXT("FCbField::As[Type](%s) -> Equal"), What), (Field.*AsTypeFn)(DefaultValue), ExpectedValue);
+		TestEqual(FString::Printf(TEXT("FCbField::As[Type](%s) -> HasError()"), What), Field.HasError(), ExpectedError != ECbFieldError::None);
+		TestEqual(FString::Printf(TEXT("FCbField::As[Type](%s) -> GetError()"), What), Field.GetError(), ExpectedError);
 	}
 
 	template <typename T, typename Default>
-	void TestFieldAsType(FCbField& Field, T (FCbField::*AsTypeFn)(Default), T ExpectedValue = T(), T DefaultValue = T(), ECbFieldError ExpectedError = ECbFieldError::None)
+	void TestFieldAsType(const TCHAR* What, FCbField& Field, T (FCbField::*AsTypeFn)(Default), T ExpectedValue = T(), T DefaultValue = T(), ECbFieldError ExpectedError = ECbFieldError::None)
 	{
-		TestFieldAsTypeNoClone(Field, AsTypeFn, ExpectedValue, DefaultValue, ExpectedError);
+		TestFieldAsTypeNoClone(What, Field, AsTypeFn, ExpectedValue, DefaultValue, ExpectedError);
 		FCbFieldRef FieldClone = FCbFieldRef::Clone(Field);
-		TestFieldAsTypeNoClone(FieldClone, AsTypeFn, ExpectedValue, DefaultValue, ExpectedError);
-		TestTrue(TEXT("FCbField::Equals()"), Field.Equals(FieldClone));
+		TestFieldAsTypeNoClone(*FString::Printf(TEXT("%s, Clone"), What), FieldClone, AsTypeFn, ExpectedValue, DefaultValue, ExpectedError);
+		TestTrue(FString::Printf(TEXT("FCbField::Equals(%s)"), What), Field.Equals(FieldClone));
 	}
 
 	template <typename T>
-	void TestFieldAsTypeNoClone(FCbField& Field, T (FCbField::*AsTypeFn)(), T ExpectedValue = T(), T DefaultValue = T(), ECbFieldError ExpectedError = ECbFieldError::None)
+	void TestFieldAsTypeNoClone(const TCHAR* What, FCbField& Field, T (FCbField::*AsTypeFn)(), T ExpectedValue = T(), T DefaultValue = T(), ECbFieldError ExpectedError = ECbFieldError::None)
 	{
-		TestTypeError(Field, ExpectedError);
+		TestTypeError(What, Field, ExpectedError);
 		(Field.*AsTypeFn)();
-		TestEqual(TEXT("FCbField::As[Type]() -> HasError()"), Field.HasError(), ExpectedError != ECbFieldError::None);
-		TestEqual(TEXT("FCbField::As[Type]() -> GetError()"), Field.GetError(), ExpectedError);
+		TestEqual(FString::Printf(TEXT("FCbField::As[Type](%s) -> HasError()"), What), Field.HasError(), ExpectedError != ECbFieldError::None);
+		TestEqual(FString::Printf(TEXT("FCbField::As[Type](%s) -> GetError()"), What), Field.GetError(), ExpectedError);
 	}
 
 	template <typename T>
-	void TestFieldAsType(FCbField& Field, T (FCbField::*AsTypeFn)(), T ExpectedValue = T(), T DefaultValue = T(), ECbFieldError ExpectedError = ECbFieldError::None)
+	void TestFieldAsType(const TCHAR* What, FCbField& Field, T (FCbField::*AsTypeFn)(), T ExpectedValue = T(), T DefaultValue = T(), ECbFieldError ExpectedError = ECbFieldError::None)
 	{
-		TestFieldAsTypeNoClone(Field, AsTypeFn, ExpectedValue, DefaultValue, ExpectedError);
+		TestFieldAsTypeNoClone(What, Field, AsTypeFn, ExpectedValue, DefaultValue, ExpectedError);
 		FCbFieldRef FieldClone = FCbFieldRef::Clone(Field);
-		TestFieldAsTypeNoClone(FieldClone, AsTypeFn, ExpectedValue, DefaultValue, ExpectedError);
-		TestTrue(TEXT("FCbField::Equals()"), Field.Equals(FieldClone));
+		TestFieldAsTypeNoClone(*FString::Printf(TEXT("%s, Clone"), What), FieldClone, AsTypeFn, ExpectedValue, DefaultValue, ExpectedError);
+		TestTrue(FString::Printf(TEXT("FCbField::Equals(%s)"), What), Field.Equals(FieldClone));
 	}
 
 	template <ECbFieldType FieldType, typename T = TCbFieldValueType<FieldType>>
-	void TestField(FCbField& Field, T ExpectedValue = T(), T DefaultValue = T(), ECbFieldError ExpectedError = ECbFieldError::None)
+	void TestField(const TCHAR* What, FCbField& Field, T ExpectedValue = T(), T DefaultValue = T(), ECbFieldError ExpectedError = ECbFieldError::None)
 	{
-		TestTrue(TEXT("FCbField::Is[Type]()"), (Field.*TCbFieldIsTypeFn<FieldType>)());
-		TestFieldAsType(Field, TCbFieldAsTypeFn<FieldType>, ExpectedValue, DefaultValue, ExpectedError);
+		TestTrue(FString::Printf(TEXT("FCbField::Is[Type](%s)"), What), (Field.*TCbFieldIsTypeFn<FieldType>)());
+		TestFieldAsType(What, Field, TCbFieldAsTypeFn<FieldType>, ExpectedValue, DefaultValue, ExpectedError);
 	}
 
 	template <ECbFieldType FieldType, typename T = TCbFieldValueType<FieldType>>
-	void TestField(TArrayView<const uint8> Payload, T ExpectedValue = T(), T DefaultValue = T())
+	void TestField(const TCHAR* What, TArrayView<const uint8> Payload, T ExpectedValue = T(), T DefaultValue = T())
 	{
 		FCbField Field(Payload.GetData(), FieldType);
-		TestEqual(TEXT("FCbField::GetSize()"), Field.GetSize(), uint64(Payload.Num()));
-		TestTrue(TEXT("FCbField::HasValue()"), Field.HasValue());
-		TestFalse(TEXT("FCbField::HasError() == false"), Field.HasError());
-		TestEqual(TEXT("FCbField::GetError() == None"), Field.GetError(), ECbFieldError::None);
-		TestField<FieldType>(Field, ExpectedValue, DefaultValue);
+		TestEqual(FString::Printf(TEXT("FCbField::GetSize(%s)"), What), Field.GetSize(), uint64(Payload.Num()));
+		TestTrue(FString::Printf(TEXT("FCbField::HasValue(%s)"), What), Field.HasValue());
+		TestFalse(FString::Printf(TEXT("FCbField::HasError(%s) == false"), What), Field.HasError());
+		TestEqual(FString::Printf(TEXT("FCbField::GetError(%s) == None"), What), Field.GetError(), ECbFieldError::None);
+		TestField<FieldType>(What, Field, ExpectedValue, DefaultValue);
 	}
 
 	template <typename T, typename Default>
-	void TestFieldAsTypeError(FCbField& Field, T (FCbField::*AsTypeFn)(Default), ECbFieldError ExpectedError, T ExpectedValue = T())
+	void TestFieldAsTypeError(const TCHAR* What, FCbField& Field, T (FCbField::*AsTypeFn)(Default), ECbFieldError ExpectedError, T ExpectedValue = T())
 	{
-		TestFieldAsTypeNoClone(Field, AsTypeFn, ExpectedValue, ExpectedValue, ExpectedError);
+		TestFieldAsTypeNoClone(What, Field, AsTypeFn, ExpectedValue, ExpectedValue, ExpectedError);
 	}
 
 	template <typename T>
-	void TestFieldAsTypeError(FCbField& Field, T (FCbField::*AsTypeFn)(), ECbFieldError ExpectedError, T ExpectedValue = T())
+	void TestFieldAsTypeError(const TCHAR* What, FCbField& Field, T (FCbField::*AsTypeFn)(), ECbFieldError ExpectedError, T ExpectedValue = T())
 	{
-		TestFieldAsTypeNoClone(Field, AsTypeFn, ExpectedValue, ExpectedValue, ExpectedError);
+		TestFieldAsTypeNoClone(What, Field, AsTypeFn, ExpectedValue, ExpectedValue, ExpectedError);
 	}
 
 	template <ECbFieldType FieldType, typename T = TCbFieldValueType<FieldType>>
-	void TestFieldError(FCbField& Field, ECbFieldError ExpectedError, T ExpectedValue = T())
+	void TestFieldError(const TCHAR* What, FCbField& Field, ECbFieldError ExpectedError, T ExpectedValue = T())
 	{
-		TestEqual(TEXT("FCbField::Is[Type]()"), (Field.*TCbFieldIsTypeFn<FieldType>)(), ExpectedError != ECbFieldError::TypeError);
-		TestFieldAsTypeError(Field, TCbFieldAsTypeFn<FieldType>, ExpectedError, ExpectedValue);
+		TestEqual(FString::Printf(TEXT("FCbField::Is[Type](%s)"), What), (Field.*TCbFieldIsTypeFn<FieldType>)(), ExpectedError != ECbFieldError::TypeError);
+		TestFieldAsTypeError(What, Field, TCbFieldAsTypeFn<FieldType>, ExpectedError, ExpectedValue);
 	}
 
 	template <ECbFieldType FieldType, typename T = TCbFieldValueType<FieldType>>
-	void TestFieldError(TArrayView<const uint8> Payload, ECbFieldError ExpectedError, T ExpectedValue = T())
+	void TestFieldError(const TCHAR* What, TArrayView<const uint8> Payload, ECbFieldError ExpectedError, T ExpectedValue = T())
 	{
 		FCbField Field(Payload.GetData(), FieldType);
-		TestFieldError<FieldType>(Field, ExpectedError, ExpectedValue);
+		TestFieldError<FieldType>(What, Field, ExpectedError, ExpectedValue);
 	}
 
 private:
-	void TestTypeError(FCbField& Field, ECbFieldError ExpectedError)
+	void TestTypeError(const TCHAR* What, FCbField& Field, ECbFieldError ExpectedError)
 	{
 		if (ExpectedError == ECbFieldError::None && !Field.IsBool())
 		{
-			TestFalse(TEXT("FCbField::IsBool() == false"), Field.IsBool());
-			TestFalse(TEXT("FCbField::AsBool() == false"), Field.AsBool());
-			TestTrue(TEXT("FCbField::AsBool() -> HasError()"), Field.HasError());
-			TestEqual(TEXT("FCbField::AsBool() -> GetError() == TypeError"), Field.GetError(), ECbFieldError::TypeError);
+			TestFalse(FString::Printf(TEXT("FCbField::IsBool(%s) == false"), What), Field.IsBool());
+			TestFalse(FString::Printf(TEXT("FCbField::AsBool(%s) == false"), What), Field.AsBool());
+			TestTrue(FString::Printf(TEXT("FCbField::AsBool(%s) -> HasError()"), What), Field.HasError());
+			TestEqual(FString::Printf(TEXT("FCbField::AsBool(%s) -> GetError() == TypeError"), What), Field.GetError(), ECbFieldError::TypeError);
 		}
 	}
 };
@@ -300,7 +300,7 @@ bool FCbFieldObjectTest::RunTest(const FString& Parameters)
 	};
 
 	// Test FCbField(Object, Empty)
-	TestField<ECbFieldType::Object>({0});
+	TestField<ECbFieldType::Object>(TEXT("Object, Empty"), {0});
 
 	// Test FCbField(Object, Empty)
 	{
@@ -330,7 +330,7 @@ bool FCbFieldObjectTest::RunTest(const FString& Parameters)
 		constexpr uint8 IntType = uint8(ECbFieldType::HasFieldName | ECbFieldType::IntegerPositive);
 		const uint8 Payload[] = { 12, IntType, 1, 'A', 1, IntType, 1, 'B', 2, IntType, 1, 'C', 3 };
 		FCbField Field(Payload, ECbFieldType::Object);
-		TestFieldAsType(Field, &FCbField::AsObject);
+		TestFieldAsType(TEXT("Object, NotEmpty"), Field, &FCbField::AsObject);
 		FCbObjectRef Object = FCbObjectRef::Clone(Field.AsObject());
 		TestIntObject(Object, 3, sizeof(Payload));
 		TestIntObject(Field.AsObject(), 3, sizeof(Payload));
@@ -348,7 +348,7 @@ bool FCbFieldObjectTest::RunTest(const FString& Parameters)
 		constexpr uint8 IntType = uint8(ECbFieldType::HasFieldName | ECbFieldType::IntegerPositive);
 		const uint8 Payload[] = { 10, IntType, 1, 'A', 1, 1, 'B', 2, 1, 'C', 3 };
 		FCbField Field(Payload, ECbFieldType::UniformObject);
-		TestFieldAsType(Field, &FCbField::AsObject);
+		TestFieldAsType(TEXT("UniformObject, NotEmpty"), Field, &FCbField::AsObject);
 		FCbObjectRef Object = FCbObjectRef::Clone(Field.AsObject());
 		TestIntObject(Object, 3, sizeof(Payload));
 		TestIntObject(Field.AsObject(), 3, sizeof(Payload));
@@ -378,7 +378,7 @@ bool FCbFieldObjectTest::RunTest(const FString& Parameters)
 	// Test FCbField(None) as Object
 	{
 		FCbField Field;
-		TestFieldError<ECbFieldType::Object>(Field, ECbFieldError::TypeError);
+		TestFieldError<ECbFieldType::Object>(TEXT("Object, None"), Field, ECbFieldError::TypeError);
 		FCbFieldRef::MakeView(Field).AsObjectRef();
 	}
 
@@ -463,7 +463,7 @@ bool FCbFieldArrayTest::RunTest(const FString& Parameters)
 	};
 
 	// Test FCbField(Array, Empty)
-	TestField<ECbFieldType::Array>({1, 0});
+	TestField<ECbFieldType::Array>(TEXT("Array, Empty"), {1, 0});
 
 	// Test FCbField(Array, Empty)
 	{
@@ -488,7 +488,7 @@ bool FCbFieldArrayTest::RunTest(const FString& Parameters)
 		constexpr uint8 IntType = uint8(ECbFieldType::IntegerPositive);
 		const uint8 Payload[] = { 7, 3, IntType, 1, IntType, 2, IntType, 3 };
 		FCbField Field(Payload, ECbFieldType::Array);
-		TestFieldAsType(Field, &FCbField::AsArray);
+		TestFieldAsType(TEXT("Array, NotEmpty"), Field, &FCbField::AsArray);
 		FCbArrayRef Array = FCbArrayRef::Clone(Field.AsArray());
 		TestIntArray(Array, 3, sizeof(Payload));
 		TestIntArray(Field.AsArray(), 3, sizeof(Payload));
@@ -500,7 +500,7 @@ bool FCbFieldArrayTest::RunTest(const FString& Parameters)
 		constexpr uint8 IntType = uint8(ECbFieldType::IntegerPositive);
 		const uint8 Payload[] = { 5, 3, IntType, 1, 2, 3 };
 		FCbField Field(Payload, ECbFieldType::UniformArray);
-		TestFieldAsType(Field, &FCbField::AsArray);
+		TestFieldAsType(TEXT("UniformArray"), Field, &FCbField::AsArray);
 		FCbArrayRef Array = FCbArrayRef::Clone(Field.AsArray());
 		TestIntArray(Array, 3, sizeof(Payload));
 		TestIntArray(Field.AsArray(), 3, sizeof(Payload));
@@ -520,7 +520,7 @@ bool FCbFieldArrayTest::RunTest(const FString& Parameters)
 	// Test FCbField(None) as Array
 	{
 		FCbField Field;
-		TestFieldError<ECbFieldType::Array>(Field, ECbFieldError::TypeError);
+		TestFieldError<ECbFieldType::Array>(TEXT("Array, None"), Field, ECbFieldError::TypeError);
 		FCbFieldRef::MakeView(Field).AsArrayRef();
 	}
 
@@ -529,7 +529,7 @@ bool FCbFieldArrayTest::RunTest(const FString& Parameters)
 		const uint8 ArrayType = uint8(ECbFieldType::Array | ECbFieldType::HasFieldName);
 		const uint8 Buffer[] = { ArrayType, 3, 'K', 'e', 'y', 3, 1, uint8(ECbFieldType::IntegerPositive), 8 };
 		const FCbArray Array(Buffer);
-		TestEqual(TEXT("Array(ArrayWithName)::GetSize()"), Array.GetSize(), uint64(5));
+		TestEqual(TEXT("FCbArray(ArrayWithName)::GetSize()"), Array.GetSize(), uint64(5));
 		const FCbArrayRef ArrayClone = FCbArrayRef::Clone(Array);
 		TestEqual(TEXT("FCbArrayRef(ArrayWithName)::GetSize()"), ArrayClone.GetSize(), uint64(5));
 		TestTrue(TEXT("FCbArray::Equals()"), Array.Equals(ArrayClone));
@@ -577,20 +577,20 @@ IMPLEMENT_CUSTOM_SIMPLE_AUTOMATION_TEST(FCbFieldBinaryTest, FCbFieldTestBase, "S
 bool FCbFieldBinaryTest::RunTest(const FString& Parameters)
 {
 	// Test FCbField(Binary, Empty)
-	TestField<ECbFieldType::Binary>({0});
+	TestField<ECbFieldType::Binary>(TEXT("Binary, Empty"), {0});
 
 	// Test FCbField(Binary, Value)
 	{
 		const uint8 Payload[] = { 3, 4, 5, 6 }; // Size: 3, Data: 4/5/6
 		FCbField Field(Payload, ECbFieldType::Binary);
-		TestFieldAsTypeNoClone(Field, &FCbField::AsBinary, MakeMemoryView(Payload + 1, 3));
+		TestFieldAsTypeNoClone(TEXT("Binary, Value"), Field, &FCbField::AsBinary, MakeMemoryView(Payload + 1, 3));
 	}
 
 	// Test FCbField(None) as Binary
 	{
 		FCbField Field;
 		const uint8 Default[] = { 1, 2, 3 };
-		TestFieldError<ECbFieldType::Binary>(Field, ECbFieldError::TypeError, MakeMemoryView(Default));
+		TestFieldError<ECbFieldType::Binary>(TEXT("Binary, None"), Field, ECbFieldError::TypeError, MakeMemoryView(Default));
 	}
 
 	return true;
@@ -600,25 +600,25 @@ IMPLEMENT_CUSTOM_SIMPLE_AUTOMATION_TEST(FCbFieldStringTest, FCbFieldTestBase, "S
 bool FCbFieldStringTest::RunTest(const FString& Parameters)
 {
 	// Test FCbField(String, Empty)
-	TestField<ECbFieldType::String>({0});
+	TestField<ECbFieldType::String>(TEXT("String, Empty"), {0});
 
 	// Test FCbField(String, Value)
 	{
 		const uint8 Payload[] = { 3, 'A', 'B', 'C' }; // Size: 3, Data: ABC
-		TestField<ECbFieldType::String>(Payload, FAnsiStringView(reinterpret_cast<const ANSICHAR*>(Payload) + 1, 3));
+		TestField<ECbFieldType::String>(TEXT("String, Value"), Payload, FAnsiStringView(reinterpret_cast<const ANSICHAR*>(Payload) + 1, 3));
 	}
 
 	// Test FCbField(String, OutOfRangeSize)
 	{
 		uint8 Payload[9];
 		WriteVarUInt(uint64(1) << 31, Payload);
-		TestFieldError<ECbFieldType::String>(Payload, ECbFieldError::RangeError, "ABC"_ASV);
+		TestFieldError<ECbFieldType::String>(TEXT("String, OutOfRangeSize"), Payload, ECbFieldError::RangeError, "ABC"_ASV);
 	}
 
 	// Test FCbField(None) as String
 	{
 		FCbField Field;
-		TestFieldError<ECbFieldType::String>(Field, ECbFieldError::TypeError, "ABC"_ASV);
+		TestFieldError<ECbFieldType::String>(TEXT("String, None"), Field, ECbFieldError::TypeError, "ABC"_ASV);
 	}
 
 	return true;
@@ -664,21 +664,21 @@ protected:
 		constexpr uint64 DefaultValue = 8;
 		const uint64 ExpectedValue = Negative ? uint64(-int64(Magnitude)) : Magnitude;
 		FCbField Field(Payload, FieldType);
-		TestFieldAsType(Field, &FCbField::AsInt8, int8(EnumHasAnyFlags(ExpectedMask, EIntType::Int8) ? ExpectedValue : DefaultValue),
+		TestFieldAsType(TEXT("Int8"), Field, &FCbField::AsInt8, int8(EnumHasAnyFlags(ExpectedMask, EIntType::Int8) ? ExpectedValue : DefaultValue),
 			int8(DefaultValue), EnumHasAnyFlags(ExpectedMask, EIntType::Int8) ? ECbFieldError::None : ECbFieldError::RangeError);
-		TestFieldAsType(Field, &FCbField::AsInt16, int16(EnumHasAnyFlags(ExpectedMask, EIntType::Int16) ? ExpectedValue : DefaultValue),
+		TestFieldAsType(TEXT("Int16"), Field, &FCbField::AsInt16, int16(EnumHasAnyFlags(ExpectedMask, EIntType::Int16) ? ExpectedValue : DefaultValue),
 			int16(DefaultValue), EnumHasAnyFlags(ExpectedMask, EIntType::Int16) ? ECbFieldError::None : ECbFieldError::RangeError);
-		TestFieldAsType(Field, &FCbField::AsInt32, int32(EnumHasAnyFlags(ExpectedMask, EIntType::Int32) ? ExpectedValue : DefaultValue),
+		TestFieldAsType(TEXT("Int32"), Field, &FCbField::AsInt32, int32(EnumHasAnyFlags(ExpectedMask, EIntType::Int32) ? ExpectedValue : DefaultValue),
 			int32(DefaultValue), EnumHasAnyFlags(ExpectedMask, EIntType::Int32) ? ECbFieldError::None : ECbFieldError::RangeError);
-		TestFieldAsType(Field, &FCbField::AsInt64, int64(EnumHasAnyFlags(ExpectedMask, EIntType::Int64) ? ExpectedValue : DefaultValue),
+		TestFieldAsType(TEXT("Int64"), Field, &FCbField::AsInt64, int64(EnumHasAnyFlags(ExpectedMask, EIntType::Int64) ? ExpectedValue : DefaultValue),
 			int64(DefaultValue), EnumHasAnyFlags(ExpectedMask, EIntType::Int64) ? ECbFieldError::None : ECbFieldError::RangeError);
-		TestFieldAsType(Field, &FCbField::AsUInt8, uint8(EnumHasAnyFlags(ExpectedMask, EIntType::UInt8) ? ExpectedValue : DefaultValue),
+		TestFieldAsType(TEXT("UInt8"), Field, &FCbField::AsUInt8, uint8(EnumHasAnyFlags(ExpectedMask, EIntType::UInt8) ? ExpectedValue : DefaultValue),
 			uint8(DefaultValue), EnumHasAnyFlags(ExpectedMask, EIntType::UInt8) ? ECbFieldError::None : ECbFieldError::RangeError);
-		TestFieldAsType(Field, &FCbField::AsUInt16, uint16(EnumHasAnyFlags(ExpectedMask, EIntType::UInt16) ? ExpectedValue : DefaultValue),
+		TestFieldAsType(TEXT("UInt16"), Field, &FCbField::AsUInt16, uint16(EnumHasAnyFlags(ExpectedMask, EIntType::UInt16) ? ExpectedValue : DefaultValue),
 			uint16(DefaultValue), EnumHasAnyFlags(ExpectedMask, EIntType::UInt16) ? ECbFieldError::None : ECbFieldError::RangeError);
-		TestFieldAsType(Field, &FCbField::AsUInt32, uint32(EnumHasAnyFlags(ExpectedMask, EIntType::UInt32) ? ExpectedValue : DefaultValue),
+		TestFieldAsType(TEXT("UInt32"), Field, &FCbField::AsUInt32, uint32(EnumHasAnyFlags(ExpectedMask, EIntType::UInt32) ? ExpectedValue : DefaultValue),
 			uint32(DefaultValue), EnumHasAnyFlags(ExpectedMask, EIntType::UInt32) ? ECbFieldError::None : ECbFieldError::RangeError);
-		TestFieldAsType(Field, &FCbField::AsUInt64, uint64(EnumHasAnyFlags(ExpectedMask, EIntType::UInt64) ? ExpectedValue : DefaultValue),
+		TestFieldAsType(TEXT("UInt64"), Field, &FCbField::AsUInt64, uint64(EnumHasAnyFlags(ExpectedMask, EIntType::UInt64) ? ExpectedValue : DefaultValue),
 			uint64(DefaultValue), EnumHasAnyFlags(ExpectedMask, EIntType::UInt64) ? ECbFieldError::None : ECbFieldError::RangeError);
 	}
 };
@@ -719,8 +719,8 @@ bool FCbFieldIntegerTest::RunTest(const FString& Parameters)
 	// Test FCbField(None) as Integer
 	{
 		FCbField Field;
-		TestFieldError<ECbFieldType::IntegerPositive>(Field, ECbFieldError::TypeError, uint64(8));
-		TestFieldError<ECbFieldType::IntegerNegative>(Field, ECbFieldError::TypeError, int64(8));
+		TestFieldError<ECbFieldType::IntegerPositive>(TEXT("Integer+, None"), Field, ECbFieldError::TypeError, uint64(8));
+		TestFieldError<ECbFieldType::IntegerNegative>(TEXT("Integer-, None"), Field, ECbFieldError::TypeError, int64(8));
 	}
 
 	return true;
@@ -732,19 +732,19 @@ bool FCbFieldFloatTest::RunTest(const FString& Parameters)
 	// Test FCbField(Float, 32-bit)
 	{
 		const uint8 Payload[] = { 0xc0, 0x12, 0x34, 0x56 }; // -2.28444433f
-		TestField<ECbFieldType::Float32>(Payload, -2.28444433f);
+		TestField<ECbFieldType::Float32>(TEXT("Float32"), Payload, -2.28444433f);
 
 		FCbField Field(Payload, ECbFieldType::Float32);
-		TestFieldAsType(Field, &FCbField::AsDouble, -2.28444433);
+		TestFieldAsType(TEXT("Float32, AsDouble"), Field, &FCbField::AsDouble, -2.28444433);
 	}
 
 	// Test FCbField(Float, 64-bit)
 	{
 		const uint8 Payload[] = { 0xc1, 0x23, 0x45, 0x67, 0x89, 0xab, 0xcd, 0xef }; // -631475.76888888876
-		TestField<ECbFieldType::Float64>(Payload, -631475.76888888876);
+		TestField<ECbFieldType::Float64>(TEXT("Float64"), Payload, -631475.76888888876);
 
 		FCbField Field(Payload, ECbFieldType::Float64);
-		TestFieldAsTypeError(Field, &FCbField::AsFloat, ECbFieldError::RangeError, 8.0f);
+		TestFieldAsTypeError(TEXT("Float64, AsFloat"), Field, &FCbField::AsFloat, ECbFieldError::RangeError, 8.0f);
 	}
 
 	// Test FCbField(Integer+, MaxBinary32) as Float
@@ -752,8 +752,8 @@ bool FCbFieldFloatTest::RunTest(const FString& Parameters)
 		uint8 Payload[9];
 		WriteVarUInt((uint64(1) << 24) - 1, Payload); // 16,777,215
 		FCbField Field(Payload, ECbFieldType::IntegerPositive);
-		TestField<ECbFieldType::Float32>(Field, 16'777'215.0f);
-		TestField<ECbFieldType::Float64>(Field, 16'777'215.0);
+		TestField<ECbFieldType::Float32>(TEXT("Integer+, MaxBinary32, AsFloat"), Field, 16'777'215.0f);
+		TestField<ECbFieldType::Float64>(TEXT("Integer+, MaxBinary32, AsDouble"), Field, 16'777'215.0);
 	}
 
 	// Test FCbField(Integer+, MaxBinary32+1) as Float
@@ -761,8 +761,8 @@ bool FCbFieldFloatTest::RunTest(const FString& Parameters)
 		uint8 Payload[9];
 		WriteVarUInt(uint64(1) << 24, Payload); // 16,777,216
 		FCbField Field(Payload, ECbFieldType::IntegerPositive);
-		TestFieldError<ECbFieldType::Float32>(Field, ECbFieldError::RangeError, 8.0f);
-		TestField<ECbFieldType::Float64>(Field, 16'777'216.0);
+		TestFieldError<ECbFieldType::Float32>(TEXT("Integer+, MaxBinary32+1, AsFloat"), Field, ECbFieldError::RangeError, 8.0f);
+		TestField<ECbFieldType::Float64>(TEXT("Integer+, MaxBinary32+1, AsDouble"), Field, 16'777'216.0);
 	}
 
 	// Test FCbField(Integer+, MaxBinary64) as Float
@@ -770,8 +770,8 @@ bool FCbFieldFloatTest::RunTest(const FString& Parameters)
 		uint8 Payload[9];
 		WriteVarUInt((uint64(1) << 53) - 1, Payload); // 9,007,199,254,740,991
 		FCbField Field(Payload, ECbFieldType::IntegerPositive);
-		TestFieldError<ECbFieldType::Float32>(Field, ECbFieldError::RangeError, 8.0f);
-		TestField<ECbFieldType::Float64>(Field, 9'007'199'254'740'991.0);
+		TestFieldError<ECbFieldType::Float32>(TEXT("Integer+, MaxBinary64, AsFloat"), Field, ECbFieldError::RangeError, 8.0f);
+		TestField<ECbFieldType::Float64>(TEXT("Integer+, MaxBinary64, AsDouble"), Field, 9'007'199'254'740'991.0);
 	}
 
 	// Test FCbField(Integer+, MaxBinary64+1) as Float
@@ -779,8 +779,8 @@ bool FCbFieldFloatTest::RunTest(const FString& Parameters)
 		uint8 Payload[9];
 		WriteVarUInt(uint64(1) << 53, Payload); // 9,007,199,254,740,992
 		FCbField Field(Payload, ECbFieldType::IntegerPositive);
-		TestFieldError<ECbFieldType::Float32>(Field, ECbFieldError::RangeError, 8.0f);
-		TestFieldError<ECbFieldType::Float64>(Field, ECbFieldError::RangeError, 8.0);
+		TestFieldError<ECbFieldType::Float32>(TEXT("Integer+, MaxBinary64+1, AsFloat"), Field, ECbFieldError::RangeError, 8.0f);
+		TestFieldError<ECbFieldType::Float64>(TEXT("Integer+, MaxBinary64+1, AsDouble"), Field, ECbFieldError::RangeError, 8.0);
 	}
 
 	// Test FCbField(Integer+, MaxUInt64) as Float
@@ -788,8 +788,8 @@ bool FCbFieldFloatTest::RunTest(const FString& Parameters)
 		uint8 Payload[9];
 		WriteVarUInt(uint64(-1), Payload); // Max uint64
 		FCbField Field(Payload, ECbFieldType::IntegerPositive);
-		TestFieldError<ECbFieldType::Float32>(Field, ECbFieldError::RangeError, 8.0f);
-		TestFieldError<ECbFieldType::Float64>(Field, ECbFieldError::RangeError, 8.0);
+		TestFieldError<ECbFieldType::Float32>(TEXT("Integer+, MaxUInt64, AsFloat"), Field, ECbFieldError::RangeError, 8.0f);
+		TestFieldError<ECbFieldType::Float64>(TEXT("Integer+, MaxUInt64, AsDouble"), Field, ECbFieldError::RangeError, 8.0);
 	}
 
 	// Test FCbField(Integer-, MaxBinary32) as Float
@@ -797,8 +797,8 @@ bool FCbFieldFloatTest::RunTest(const FString& Parameters)
 		uint8 Payload[9];
 		WriteVarUInt((uint64(1) << 24) - 2, Payload); // -16,777,215
 		FCbField Field(Payload, ECbFieldType::IntegerNegative);
-		TestField<ECbFieldType::Float32>(Field, -16'777'215.0f);
-		TestField<ECbFieldType::Float64>(Field, -16'777'215.0);
+		TestField<ECbFieldType::Float32>(TEXT("Integer-, MaxBinary32, AsFloat"), Field, -16'777'215.0f);
+		TestField<ECbFieldType::Float64>(TEXT("Integer-, MaxBinary32, AsDouble"), Field, -16'777'215.0);
 	}
 
 	// Test FCbField(Integer-, MaxBinary32+1) as Float
@@ -806,8 +806,8 @@ bool FCbFieldFloatTest::RunTest(const FString& Parameters)
 		uint8 Payload[9];
 		WriteVarUInt((uint64(1) << 24) - 1, Payload); // -16,777,216
 		FCbField Field(Payload, ECbFieldType::IntegerNegative);
-		TestFieldError<ECbFieldType::Float32>(Field, ECbFieldError::RangeError, 8.0f);
-		TestField<ECbFieldType::Float64>(Field, -16'777'216.0);
+		TestFieldError<ECbFieldType::Float32>(TEXT("Integer-, MaxBinary32+1, AsFloat"), Field, ECbFieldError::RangeError, 8.0f);
+		TestField<ECbFieldType::Float64>(TEXT("Integer-, MaxBinary32+1, AsDouble"), Field, -16'777'216.0);
 	}
 
 	// Test FCbField(Integer-, MaxBinary64) as Float
@@ -815,8 +815,8 @@ bool FCbFieldFloatTest::RunTest(const FString& Parameters)
 		uint8 Payload[9];
 		WriteVarUInt((uint64(1) << 53) - 2, Payload); // -9,007,199,254,740,991
 		FCbField Field(Payload, ECbFieldType::IntegerNegative);
-		TestFieldError<ECbFieldType::Float32>(Field, ECbFieldError::RangeError, 8.0f);
-		TestField<ECbFieldType::Float64>(Field, -9'007'199'254'740'991.0);
+		TestFieldError<ECbFieldType::Float32>(TEXT("Integer-, MaxBinary64, AsFloat"), Field, ECbFieldError::RangeError, 8.0f);
+		TestField<ECbFieldType::Float64>(TEXT("Integer-, MaxBinary64, AsDouble"), Field, -9'007'199'254'740'991.0);
 	}
 
 	// Test FCbField(Integer-, MaxBinary64+1) as Float
@@ -824,15 +824,15 @@ bool FCbFieldFloatTest::RunTest(const FString& Parameters)
 		uint8 Payload[9];
 		WriteVarUInt((uint64(1) << 53) - 1, Payload); // -9,007,199,254,740,992
 		FCbField Field(Payload, ECbFieldType::IntegerNegative);
-		TestFieldError<ECbFieldType::Float32>(Field, ECbFieldError::RangeError, 8.0f);
-		TestFieldError<ECbFieldType::Float64>(Field, ECbFieldError::RangeError, 8.0);
+		TestFieldError<ECbFieldType::Float32>(TEXT("Integer-, MaxBinary64+1, AsFloat"), Field, ECbFieldError::RangeError, 8.0f);
+		TestFieldError<ECbFieldType::Float64>(TEXT("Integer-, MaxBinary64+1, AsDouble"), Field, ECbFieldError::RangeError, 8.0);
 	}
 
 	// Test FCbField(None) as Float
 	{
 		FCbField Field;
-		TestFieldError<ECbFieldType::Float32>(Field, ECbFieldError::TypeError, 8.0f);
-		TestFieldError<ECbFieldType::Float64>(Field, ECbFieldError::TypeError, 8.0);
+		TestFieldError<ECbFieldType::Float32>(TEXT("None, AsFloat"), Field, ECbFieldError::TypeError, 8.0f);
+		TestFieldError<ECbFieldType::Float64>(TEXT("None, AsDouble"), Field, ECbFieldError::TypeError, 8.0);
 	}
 
 	return true;
@@ -842,16 +842,16 @@ IMPLEMENT_CUSTOM_SIMPLE_AUTOMATION_TEST(FCbFieldBoolTest, FCbFieldTestBase, "Sys
 bool FCbFieldBoolTest::RunTest(const FString& Parameters)
 {
 	// Test FCbField(Bool, False)
-	TestField<ECbFieldType::BoolFalse>({}, false, true);
+	TestField<ECbFieldType::BoolFalse>(TEXT("Bool, False"), {}, false, true);
 
 	// Test FCbField(Bool, True)
-	TestField<ECbFieldType::BoolTrue>({}, true, false);
+	TestField<ECbFieldType::BoolTrue>(TEXT("Bool, True"), {}, true, false);
 
 	// Test FCbField(None) as Bool
 	{
 		FCbField DefaultField;
-		TestFieldError<ECbFieldType::BoolFalse>(DefaultField, ECbFieldError::TypeError, false);
-		TestFieldError<ECbFieldType::BoolTrue>(DefaultField, ECbFieldError::TypeError, true);
+		TestFieldError<ECbFieldType::BoolFalse>(TEXT("Bool, False, None"), DefaultField, ECbFieldError::TypeError, false);
+		TestFieldError<ECbFieldType::BoolTrue>(TEXT("Bool, True, None"), DefaultField, ECbFieldError::TypeError, true);
 	}
 
 	return true;
@@ -864,15 +864,15 @@ bool FCbFieldReferenceTest::RunTest(const FString& Parameters)
 	const FBlake3Hash::ByteArray SequentialBytes{1, 2, 3, 4, 5, 6, 7, 8, 9, 10, 11, 12, 13, 14, 15, 16, 17, 18, 19, 20, 21, 22, 23, 24, 25, 26, 27, 28, 29, 30, 31, 32};
 
 	// Test FCbField(Reference, Zero)
-	TestField<ECbFieldType::Reference>(ZeroBytes);
+	TestField<ECbFieldType::Reference>(TEXT("Reference, Zero"), ZeroBytes);
 
 	// Test FCbField(Reference, NonZero)
-	TestField<ECbFieldType::Reference>(SequentialBytes, FBlake3Hash(SequentialBytes));
+	TestField<ECbFieldType::Reference>(TEXT("Reference, NonZero"), SequentialBytes, FBlake3Hash(SequentialBytes));
 
 	// Test FCbField(None) as Reference
 	{
 		FCbField DefaultField;
-		TestFieldError<ECbFieldType::Reference>(DefaultField, ECbFieldError::TypeError, FBlake3Hash(SequentialBytes));
+		TestFieldError<ECbFieldType::Reference>(TEXT("Reference, None"), DefaultField, ECbFieldError::TypeError, FBlake3Hash(SequentialBytes));
 	}
 
 	return true;
@@ -885,15 +885,15 @@ bool FCbFieldBinaryReferenceTest::RunTest(const FString& Parameters)
 	const FBlake3Hash::ByteArray SequentialBytes{1, 2, 3, 4, 5, 6, 7, 8, 9, 10, 11, 12, 13, 14, 15, 16, 17, 18, 19, 20, 21, 22, 23, 24, 25, 26, 27, 28, 29, 30, 31, 32};
 
 	// Test FCbField(BinaryReference, Zero)
-	TestField<ECbFieldType::BinaryReference>(ZeroBytes);
+	TestField<ECbFieldType::BinaryReference>(TEXT("BinaryReference, Zero"), ZeroBytes);
 
 	// Test FCbField(BinaryReference, NonZero)
-	TestField<ECbFieldType::BinaryReference>(SequentialBytes, FBlake3Hash(SequentialBytes));
+	TestField<ECbFieldType::BinaryReference>(TEXT("BinaryReference, NonZero"), SequentialBytes, FBlake3Hash(SequentialBytes));
 
 	// Test FCbField(None) as BinaryReference
 	{
 		FCbField DefaultField;
-		TestFieldError<ECbFieldType::BinaryReference>(DefaultField, ECbFieldError::TypeError, FBlake3Hash(SequentialBytes));
+		TestFieldError<ECbFieldType::BinaryReference>(TEXT("BinaryReference, None"), DefaultField, ECbFieldError::TypeError, FBlake3Hash(SequentialBytes));
 	}
 
 	return true;
@@ -906,15 +906,15 @@ bool FCbFieldHashTest::RunTest(const FString& Parameters)
 	const FBlake3Hash::ByteArray SequentialBytes{1, 2, 3, 4, 5, 6, 7, 8, 9, 10, 11, 12, 13, 14, 15, 16, 17, 18, 19, 20, 21, 22, 23, 24, 25, 26, 27, 28, 29, 30, 31, 32};
 
 	// Test FCbField(Hash, Zero)
-	TestField<ECbFieldType::Hash>(ZeroBytes);
+	TestField<ECbFieldType::Hash>(TEXT("Hash, Zero"), ZeroBytes);
 
 	// Test FCbField(Hash, NonZero)
-	TestField<ECbFieldType::Hash>(SequentialBytes, FBlake3Hash(SequentialBytes));
+	TestField<ECbFieldType::Hash>(TEXT("Hash, NonZero"), SequentialBytes, FBlake3Hash(SequentialBytes));
 
 	// Test FCbField(None) as Hash
 	{
 		FCbField DefaultField;
-		TestFieldError<ECbFieldType::Hash>(DefaultField, ECbFieldError::TypeError, FBlake3Hash(SequentialBytes));
+		TestFieldError<ECbFieldType::Hash>(TEXT("Hash, None"), DefaultField, ECbFieldError::TypeError, FBlake3Hash(SequentialBytes));
 	}
 
 	return true;
@@ -928,15 +928,15 @@ bool FCbFieldUuidTest::RunTest(const FString& Parameters)
 	const FGuid SequentialGuid(TEXT("00010203-0405-0607-0809-0a0b0c0d0e0f"));
 
 	// Test FCbField(Uuid, Zero)
-	TestField<ECbFieldType::Uuid>(ZeroBytes, FGuid(), SequentialGuid);
+	TestField<ECbFieldType::Uuid>(TEXT("Uuid, Zero"), ZeroBytes, FGuid(), SequentialGuid);
 
 	// Test FCbField(Uuid, NonZero)
-	TestField<ECbFieldType::Uuid>(SequentialBytes, SequentialGuid, FGuid());
+	TestField<ECbFieldType::Uuid>(TEXT("Uuid, NonZero"), SequentialBytes, SequentialGuid, FGuid());
 
 	// Test FCbField(None) as Uuid
 	{
 		FCbField DefaultField;
-		TestFieldError<ECbFieldType::Uuid>(DefaultField, ECbFieldError::TypeError, FGuid::NewGuid());
+		TestFieldError<ECbFieldType::Uuid>(TEXT("Uuid, None"), DefaultField, ECbFieldError::TypeError, FGuid::NewGuid());
 	}
 
 	return true;
@@ -946,10 +946,10 @@ IMPLEMENT_CUSTOM_SIMPLE_AUTOMATION_TEST(FCbFieldDateTimeTest, FCbFieldTestBase,
 bool FCbFieldDateTimeTest::RunTest(const FString& Parameters)
 {
 	// Test FCbField(DateTime, Zero)
-	TestField<ECbFieldType::DateTime>({0, 0, 0, 0, 0, 0, 0, 0});
+	TestField<ECbFieldType::DateTime>(TEXT("DateTime, Zero"), {0, 0, 0, 0, 0, 0, 0, 0});
 
 	// Test FCbField(DateTime, 0x1020'3040'5060'7080)
-	TestField<ECbFieldType::DateTime>({0x10, 0x20, 0x30, 0x40, 0x50, 0x60, 0x70, 0x80}, int64(0x1020'3040'5060'7080));
+	TestField<ECbFieldType::DateTime>(TEXT("DateTime, NonZero"), {0x10, 0x20, 0x30, 0x40, 0x50, 0x60, 0x70, 0x80}, int64(0x1020'3040'5060'7080));
 
 	// Test FCbField(DateTime, Zero) as FDateTime
 	{
@@ -961,7 +961,7 @@ bool FCbFieldDateTimeTest::RunTest(const FString& Parameters)
 	// Test FCbField(None) as DateTime
 	{
 		FCbField DefaultField;
-		TestFieldError<ECbFieldType::DateTime>(DefaultField, ECbFieldError::TypeError);
+		TestFieldError<ECbFieldType::DateTime>(TEXT("DateTime, None"), DefaultField, ECbFieldError::TypeError);
 		const FDateTime DefaultValue(0x1020'3040'5060'7080);
 		TestEqual(TEXT("FCbField()::AsDateTime()"), DefaultField.AsDateTime(DefaultValue), DefaultValue);
 	}
@@ -973,10 +973,10 @@ IMPLEMENT_CUSTOM_SIMPLE_AUTOMATION_TEST(FCbFieldTimeSpanTest, FCbFieldTestBase,
 bool FCbFieldTimeSpanTest::RunTest(const FString& Parameters)
 {
 	// Test FCbField(TimeSpan, Zero)
-	TestField<ECbFieldType::TimeSpan>({0, 0, 0, 0, 0, 0, 0, 0});
+	TestField<ECbFieldType::TimeSpan>(TEXT("TimeSpan, Zero"), {0, 0, 0, 0, 0, 0, 0, 0});
 
 	// Test FCbField(TimeSpan, 0x1020'3040'5060'7080)
-	TestField<ECbFieldType::TimeSpan>({0x10, 0x20, 0x30, 0x40, 0x50, 0x60, 0x70, 0x80}, int64(0x1020'3040'5060'7080));
+	TestField<ECbFieldType::TimeSpan>(TEXT("TimeSpan, NonZero"), {0x10, 0x20, 0x30, 0x40, 0x50, 0x60, 0x70, 0x80}, int64(0x1020'3040'5060'7080));
 
 	// Test FCbField(TimeSpan, Zero) as FTimeSpan
 	{
@@ -988,7 +988,7 @@ bool FCbFieldTimeSpanTest::RunTest(const FString& Parameters)
 	// Test FCbField(None) as TimeSpan
 	{
 		FCbField DefaultField;
-		TestFieldError<ECbFieldType::TimeSpan>(DefaultField, ECbFieldError::TypeError);
+		TestFieldError<ECbFieldType::TimeSpan>(TEXT("TimeSpan, None"), DefaultField, ECbFieldError::TypeError);
 		const FTimespan DefaultValue(0x1020'3040'5060'7080);
 		TestEqual(TEXT("FCbField()::AsTimeSpan()"), DefaultField.AsTimeSpan(DefaultValue), DefaultValue);
 	}
