Emit periodic alloc/free timestamp trace events.

diff --git a/Engine/Source/Runtime/Core/Private/ProfilingDebugging/Windows/MemoryTrace.cpp b/Engine/Source/Runtime/Core/Private/ProfilingDebugging/Windows/MemoryTrace.cpp
--- a/Engine/Source/Runtime/Core/Private/ProfilingDebugging/Windows/MemoryTrace.cpp
+++ b/Engine/Source/Runtime/Core/Private/ProfilingDebugging/Windows/MemoryTrace.cpp
@@ -8,6 +8,7 @@
 #include "CoreTypes.h"
 #include "HAL/MemoryBase.h"
 #include "HAL/Platform.h"
+#include "HAL/PlatformTime.h"
 #include "Misc/CoreDelegates.h"
 #include "Misc/CString.h"
 #include "Trace/Trace.inl"
@@ -88,6 +89,8 @@ UE_TRACE_CHANNEL(MemSummaryChannel)
 UE_TRACE_CHANNEL_DEFINE(MemAllocChannel)
 
 UE_TRACE_EVENT_BEGIN(Memory, Init)
+	UE_TRACE_EVENT_FIELD(uint64, BaseCycle)
+	UE_TRACE_EVENT_FIELD(uint32, MarkerPeriod)
 	UE_TRACE_EVENT_FIELD(uint8, MinAlignment)
 	UE_TRACE_EVENT_FIELD(uint8, SizeShift)
 #if USE_OVERVIEW_TRACE
@@ -106,6 +109,10 @@ UE_TRACE_EVENT_BEGIN(Memory, Summary)
 UE_TRACE_EVENT_END()
 #endif
 
+UE_TRACE_EVENT_BEGIN(Memory, Marker)
+	UE_TRACE_EVENT_FIELD(uint32, Cycle)
+UE_TRACE_EVENT_END()
+
 UE_TRACE_EVENT_BEGIN(Memory, CoreAdd)
 	UE_TRACE_EVENT_FIELD(uint64, Owner)
 	UE_TRACE_EVENT_FIELD(void*, Base)
@@ -290,7 +297,10 @@ public:
 
 private:
 	void				Update();
+	uint64				BaseCycle;
+	std::atomic<uint32>	MarkerCounter;
 	bool				bPumpTrace = false;
+	static const uint32 MarkerSamplePeriod	= (4 << 10) - 1;
 	static const uint32 SizeShift = 3;
 	static_assert(MIN_ALIGNMENT >= (1 << SizeShift), "");
 };
@@ -301,7 +311,11 @@ static FUndestructed<FAllocationTrace> GAllocationTrace;
 ////////////////////////////////////////////////////////////////////////////////
 void FAllocationTrace::Initialize()
 {
+	BaseCycle = FPlatformTime::Cycles64();
+
 	UE_TRACE_LOG(Memory, Init, MemAllocChannel)
+		<< Init.BaseCycle(BaseCycle)
+		<< Init.MarkerPeriod(MarkerSamplePeriod + 1)
 		<< Init.MinAlignment(uint8(MIN_ALIGNMENT))
 #if USE_OVERVIEW_TRACE
 		<< Init.SummarySizeShift(uint8(FSummaryTrace::SizeShift))
@@ -320,6 +334,14 @@ void FAllocationTrace::EnableTracePump()
 ////////////////////////////////////////////////////////////////////////////////
 void FAllocationTrace::Update()
 {
+	uint32 TheCount = MarkerCounter.fetch_add(1, std::memory_order_relaxed);
+	if ((TheCount & MarkerSamplePeriod) == 0)
+	{
+		uint64 Cycle = FPlatformTime::Cycles64();
+		UE_TRACE_LOG(Memory, Marker, MemAllocChannel)
+			<< Marker.Cycle(uint32(Cycle - BaseCycle));
+	}
+
 	if (bPumpTrace)
 	{
 		UE::Trace::Update();
