Fixed an error in a static_assert's logic and reduced some static assertion redundancy.

#rnx

diff --git a/Engine/Source/Runtime/Core/Private/ProfilingDebugging/Windows/MemoryTrace.cpp b/Engine/Source/Runtime/Core/Private/ProfilingDebugging/Windows/MemoryTrace.cpp
--- a/Engine/Source/Runtime/Core/Private/ProfilingDebugging/Windows/MemoryTrace.cpp
+++ b/Engine/Source/Runtime/Core/Private/ProfilingDebugging/Windows/MemoryTrace.cpp
@@ -302,7 +302,6 @@ private:
 	bool				bPumpTrace = false;
 	static const uint32 MarkerSamplePeriod	= (4 << 10) - 1;
 	static const uint32 SizeShift = 3;
-	static_assert(MIN_ALIGNMENT >= (1 << SizeShift), "");
 };
 
 ////////////////////////////////////////////////////////////////////////////////
@@ -322,7 +321,7 @@ void FAllocationTrace::Initialize()
 #endif
 		<< Init.SizeShift(uint8(SizeShift));
 
-	static_assert((1 < SizeShift) - 1 <= MIN_ALIGNMENT, "Not enough bits to pack size fields");
+	static_assert((1 << SizeShift) - 1 <= MIN_ALIGNMENT, "Not enough bits to pack size fields");
 }
 
 ////////////////////////////////////////////////////////////////////////////////
