Assume that callers of FAllocationTrace::[Re]Alloc() will provide the precise alignment values that they wish to trace.

#rnx

diff --git a/Engine/Source/Runtime/Core/Private/ProfilingDebugging/Windows/MemoryTrace.cpp b/Engine/Source/Runtime/Core/Private/ProfilingDebugging/Windows/MemoryTrace.cpp
--- a/Engine/Source/Runtime/Core/Private/ProfilingDebugging/Windows/MemoryTrace.cpp
+++ b/Engine/Source/Runtime/Core/Private/ProfilingDebugging/Windows/MemoryTrace.cpp
@@ -372,8 +372,7 @@ void FAllocationTrace::CoreRemove(void* Base, size_t Size, void* Owner)
 ////////////////////////////////////////////////////////////////////////////////
 void FAllocationTrace::Alloc(void* Address, size_t Size, uint32 Alignment, void* Owner)
 {
-	uint32 ActualAlignment = Alignment > uint32(MIN_ALIGNMENT) ? Alignment : uint32(MIN_ALIGNMENT);
-	uint32 Alignment_SizeLower = ActualAlignment | (Size & ((1 << SizeShift) - 1));
+	uint32 Alignment_SizeLower = Alignment | (Size & ((1 << SizeShift) - 1));
 
 	UE_TRACE_LOG(Memory, Alloc, MemAllocChannel)
 		<< Alloc.Owner(uint64(Owner))
@@ -396,8 +395,7 @@ void FAllocationTrace::Free(void* Address)
 ////////////////////////////////////////////////////////////////////////////////
 void FAllocationTrace::ReallocAlloc(void* Address, size_t Size, uint32 Alignment, void* Owner)
 {
-	uint32 ActualAlignment = Alignment > uint32(MIN_ALIGNMENT) ? Alignment : uint32(MIN_ALIGNMENT);
-	uint32 Alignment_SizeLower = ActualAlignment | (Size & ((1 << SizeShift) - 1));
+	uint32 Alignment_SizeLower = Alignment | (Size & ((1 << SizeShift) - 1));
 
 	UE_TRACE_LOG(Memory, ReallocAlloc, MemAllocChannel)
 		<< ReallocAlloc.Owner(uint64(Owner))
