Remove very noisy profiler scope

#rnx

diff --git a/Engine/Source/Runtime/Core/Private/ProfilingDebugging/Windows/CallstackTrace.cpp b/Engine/Source/Runtime/Core/Private/ProfilingDebugging/Windows/CallstackTrace.cpp
--- a/Engine/Source/Runtime/Core/Private/ProfilingDebugging/Windows/CallstackTrace.cpp
+++ b/Engine/Source/Runtime/Core/Private/ProfilingDebugging/Windows/CallstackTrace.cpp
@@ -655,7 +655,6 @@ uint32 FCallstackProcWorker::Run()
 {
 	while (bRun)
 	{
-		TRACE_CPUPROFILER_EVENT_SCOPE("ProfilerThreadWork")
 		while (!Queue.was_empty()) 
 		{
 			FBacktraceEntry Entry = Queue.pop();
