Fixed IncludeTool error.

diff --git a/Engine/Source/Runtime/Core/Public/ProfilingDebugging/TagTrace.h b/Engine/Source/Runtime/Core/Public/ProfilingDebugging/TagTrace.h
--- a/Engine/Source/Runtime/Core/Public/ProfilingDebugging/TagTrace.h
+++ b/Engine/Source/Runtime/Core/Public/ProfilingDebugging/TagTrace.h
@@ -1,7 +1,6 @@
 // Copyright Epic Games, Inc. All Rights Reserved.
 #pragma once
 #include "CoreTypes.h"
-#include "HAL/LowLevelMemTracker.h"
 #include "UObject/NameTypes.h"
 #include "Trace/Trace.h"
 #include "Trace/Detail/LogScope.h"
