On Windows, replaced the vectored exception filter installed to capture static init crashes with a unhandled exception filter to prevent capturing and ending the engine process for exceptions that would otherwise be caught and cleared by __try __except executing later on.

diff --git a/Engine/Source/Runtime/Core/Private/Windows/WindowsPlatformCrashContext.cpp b/Engine/Source/Runtime/Core/Private/Windows/WindowsPlatformCrashContext.cpp
--- a/Engine/Source/Runtime/Core/Private/Windows/WindowsPlatformCrashContext.cpp
+++ b/Engine/Source/Runtime/Core/Private/Windows/WindowsPlatformCrashContext.cpp
@@ -1059,8 +1059,6 @@ private:
 	DWORD CrashingThreadId;
 	/** Handle to crashed thread.*/
 	HANDLE CrashingThreadHandle;
-	/** Handle used to remove vectored exception handler. */
-	HANDLE VectoredExceptionHandle;
 
 	/** Process handle to crash reporter client */
 	FProcHandle CrashClientHandle;
@@ -1088,14 +1086,7 @@ private:
 		__try
 #endif
 		{
-			// Removed vectored exception handler, we are now guaranteed to
-			// be able to catch unhandled exception in the main try/catch block.
-#if _WIN32_WINNT >= 0x0500
-			if (!FPlatformMisc::IsDebuggerPresent())
-			{
-				RemoveVectoredExceptionHandler(VectoredExceptionHandle);
-			}
-#endif
+
 			while (StopTaskCounter.GetValue() == 0)
 			{
 				if (WaitForSingleObject(CrashEvent, 500) == WAIT_OBJECT_0)
@@ -1149,7 +1140,6 @@ public:
 		, ExceptionInfo(nullptr)
 		, CrashingThreadId(0)
 		, CrashingThreadHandle(nullptr)
-		, VectoredExceptionHandle(INVALID_HANDLE_VALUE)
 		, CrashMonitorWritePipe(nullptr)
 		, CrashMonitorReadPipe(nullptr)
 		, CrashMonitorPid(0)
@@ -1158,15 +1148,12 @@ public:
 		CrashEvent = CreateEvent(nullptr, true, 0, nullptr);
 		CrashHandledEvent = CreateEvent(nullptr, true, 0, nullptr);
 
-		// Add an exception handler to catch issues during static initialization. This
-		// is removed once the crash reporter thread is started.
-#if _WIN32_WINNT >= 0x0500
+		// Add an exception handler to catch issues during static initialization. It is replaced by the engine handler once
+		// the guarded main is entered.
 		if (!FPlatformMisc::IsDebuggerPresent())
 		{
-			constexpr ULONG CallFirst = 1;
-			VectoredExceptionHandle = AddVectoredExceptionHandler(CallFirst, UnhandledStaticInitException);
+			::SetUnhandledExceptionFilter(UnhandledStaticInitException);
 		}
-#endif
 
 #if USE_CRASH_REPORTER_MONITOR
 		if (!FPlatformProperties::IsServerOnly())
@@ -1467,23 +1454,19 @@ TOptional<FCrashReportingThread> GCrashReportingThread(InPlace);
 LONG WINAPI UnhandledStaticInitException(LPEXCEPTION_POINTERS ExceptionInfo)
 {
 #if !NOINITCRASHREPORTER
-	// Top bit in exception code is fatal exceptions. Report those but not other types.
-	if ((ExceptionInfo->ExceptionRecord->ExceptionCode & 0x80000000L) != 0)
-	{
-		// If we get an exception during static init we hope that the crash reporting thread
-		// object has been created, otherwise we cannot handle the exception. This will hopefully 
-		// work even if there is a stack overflow. See 
-		// https://peteronprogramming.wordpress.com/2016/08/10/crashes-you-cant-handle-easily-2-stack-overflows-on-windows/
-		// @note: Even if the object has been created, the actual thread has not been started yet,
-		// (that happens after static init) so we must bypass that and report directly from this thread. 
-		// 
-		if (GCrashReportingThread.IsSet())
-		{
-			return GCrashReportingThread->OnCrashDuringStaticInit(ExceptionInfo);
-		}
+	// If we get an exception during static init we hope that the crash reporting thread
+	// object has been created, otherwise we cannot handle the exception. This will hopefully 
+	// work even if there is a stack overflow. See 
+	// https://peteronprogramming.wordpress.com/2016/08/10/crashes-you-cant-handle-easily-2-stack-overflows-on-windows/
+	// @note: Even if the object has been created, the actual thread has not been started yet,
+	// (that happens after static init) so we must bypass that and report directly from this thread. 
+	// 
+	if (GCrashReportingThread.IsSet())
+	{
+		return GCrashReportingThread->OnCrashDuringStaticInit(ExceptionInfo);
 	}
 #endif
-	
+
 	return EXCEPTION_CONTINUE_SEARCH;
 }
 
