CompactBinary: Refactor to share code and provide more functionality on iterators

- Changed GetSize, GetHash, CopyTo to always include the type.
- Changed MakeView to take an optional outer buffer parameter to mirror the shared buffer interface more closely.
- Changed GetBuffer to always return a buffer that exactly matches the view, and added GetOuterBuffer to return the outer buffer if available.
- Added TryGetSerializedView to field, array, and object to access the entire serialized value when available.
- Added an optional type parameter to the ref type constructors to make it easier to work with an external type.
- Added factory functions for iterators: MakeSingle, MakeRange, MakeRangeView, CloneRange.
- Added GetRangeSize, GetRangeHash, CopyRangeTo, GetRangeView, TryGetSerializedRangeView, GetRangeBuffer to iterators to mirror the functions on the field.

#rnx

diff --git a/Engine/Source/Runtime/Core/Private/Serialization/CompactBinary.cpp b/Engine/Source/Runtime/Core/Private/Serialization/CompactBinary.cpp
--- a/Engine/Source/Runtime/Core/Private/Serialization/CompactBinary.cpp
+++ b/Engine/Source/Runtime/Core/Private/Serialization/CompactBinary.cpp
@@ -15,8 +15,8 @@
 namespace CompactBinaryPrivate
 {
 
-static constexpr const uint8 GEmptyObjectPayload[] = { 0x00 };
-static constexpr const uint8 GEmptyArrayPayload[] = { 0x01, 0x00 };
+static constexpr const uint8 GEmptyObjectPayload[] = { uint8(ECbFieldType::Object), 0x00 };
+static constexpr const uint8 GEmptyArrayPayload[] = { uint8(ECbFieldType::Array), 0x01, 0x00 };
 
 template <typename T>
 static constexpr FORCEINLINE T ReadUnaligned(const void* const Memory)
@@ -85,88 +85,9 @@
 	Payload = Bytes;
 }
 
-FConstMemoryView FCbField::GetFieldView() const
-{
-	const uint32 TypeSize = FCbFieldType::HasFieldType(Type) ? sizeof(ECbFieldType) : 0;
-	const uint32 NameSize = FCbFieldType::HasFieldName(Type) ? NameLen + MeasureVarUInt(NameLen) : 0;
-	const uint64 LocalPayloadSize = GetPayloadSize();
-	return MakeMemoryView(static_cast<const uint8*>(Payload) - TypeSize - NameSize, TypeSize + NameSize + GetPayloadSize());
-}
-
-uint64 FCbField::GetSize() const
-{
-	return GetFieldView().GetSize();
-}
-
-uint64 FCbField::GetPayloadSize() const
-{
-	switch (FCbFieldType::GetType(Type))
-	{
-	case ECbFieldType::None:
-	case ECbFieldType::Null:
-		return 0;
-	case ECbFieldType::Object:
-	case ECbFieldType::UniformObject:
-	case ECbFieldType::Array:
-	case ECbFieldType::UniformArray:
-	case ECbFieldType::Binary:
-	case ECbFieldType::String:
-	{
-		uint32 PayloadSizeByteCount;
-		const uint64 PayloadSize = ReadVarUInt(Payload, PayloadSizeByteCount);
-		return PayloadSize + PayloadSizeByteCount;
-	}
-	case ECbFieldType::IntegerPositive:
-	case ECbFieldType::IntegerNegative:
-		return MeasureVarUInt(Payload);
-	case ECbFieldType::Float32:
-		return 4;
-	case ECbFieldType::Float64:
-		return 8;
-	case ECbFieldType::BoolFalse:
-	case ECbFieldType::BoolTrue:
-		return 0;
-	case ECbFieldType::Reference:
-	case ECbFieldType::BinaryReference:
-	case ECbFieldType::Hash:
-		return 32;
-	case ECbFieldType::Uuid:
-		return 16;
-	case ECbFieldType::DateTime:
-	case ECbFieldType::TimeSpan:
-		return 8;
-	default:
-		return 0;
-	}
-}
-
-bool FCbField::Equals(const FCbField& Other) const
-{
-	return Type == Other.Type && GetFieldView().EqualBytes(Other.GetFieldView());
-}
-
-void FCbField::IterateReferences(FCbFieldVisitor Visitor) const
-{
-	switch (FCbFieldType::GetType(Type))
-	{
-	case ECbFieldType::Object:
-	case ECbFieldType::UniformObject:
-		return FCbObject::FromField(*this).IterateReferences(Visitor);
-	case ECbFieldType::Array:
-	case ECbFieldType::UniformArray:
-		return FCbArray::FromField(*this).IterateReferences(Visitor);
-	case ECbFieldType::Reference:
-	case ECbFieldType::BinaryReference:
-		return Visitor(*this);
-	default:
-		return;
-	}
-}
-
 FCbObject FCbField::AsObject()
 {
-	const ECbFieldType LocalType = Type;
-	if (FCbFieldType::IsObject(LocalType))
+	if (FCbFieldType::IsObject(Type))
 	{
 		Error = ECbFieldError::None;
 		return FCbObject::FromField(*this);
@@ -180,8 +101,7 @@
 
 FCbArray FCbField::AsArray()
 {
-	const ECbFieldType LocalType = Type;
-	if (FCbFieldType::IsArray(LocalType))
+	if (FCbFieldType::IsArray(Type))
 	{
 		Error = ECbFieldError::None;
 		return FCbArray::FromField(*this);
@@ -237,8 +157,7 @@
 
 uint64 FCbField::AsInteger(const uint64 Default, const FIntegerParams Params)
 {
-	const ECbFieldType LocalType = Type;
-	if (FCbFieldType::IsInteger(LocalType))
+	if (FCbFieldType::IsInteger(Type))
 	{
 		// A shift of a 64-bit value by 64 is undefined so shift by one less because magnitude is never zero.
 		const uint64 OutOfRangeMask = uint64(-2) << (Params.MagnitudeBits - 1);
@@ -449,29 +368,131 @@
 	return FTimespan(AsTimeSpanTicks(Default.GetTicks()));
 }
 
-void FCbField::CopyTo(const FMutableMemoryView Buffer) const
+uint64 FCbField::GetSize() const
 {
-	const FConstMemoryView Source = GetFieldView();
-	FMemory::Memcpy(Buffer.GetData(), Source.GetData(), FMath::Min(Buffer.GetSize(), Source.GetSize()));
+	return sizeof(ECbFieldType) + GetViewNoType().GetSize();
 }
 
+uint64 FCbField::GetPayloadSize() const
+{
+	switch (FCbFieldType::GetType(Type))
+	{
+	case ECbFieldType::None:
+	case ECbFieldType::Null:
+		return 0;
+	case ECbFieldType::Object:
+	case ECbFieldType::UniformObject:
+	case ECbFieldType::Array:
+	case ECbFieldType::UniformArray:
+	case ECbFieldType::Binary:
+	case ECbFieldType::String:
+	{
+		uint32 PayloadSizeByteCount;
+		const uint64 PayloadSize = ReadVarUInt(Payload, PayloadSizeByteCount);
+		return PayloadSize + PayloadSizeByteCount;
+	}
+	case ECbFieldType::IntegerPositive:
+	case ECbFieldType::IntegerNegative:
+		return MeasureVarUInt(Payload);
+	case ECbFieldType::Float32:
+		return 4;
+	case ECbFieldType::Float64:
+		return 8;
+	case ECbFieldType::BoolFalse:
+	case ECbFieldType::BoolTrue:
+		return 0;
+	case ECbFieldType::Reference:
+	case ECbFieldType::BinaryReference:
+	case ECbFieldType::Hash:
+		return 32;
+	case ECbFieldType::Uuid:
+		return 16;
+	case ECbFieldType::DateTime:
+	case ECbFieldType::TimeSpan:
+		return 8;
+	default:
+		return 0;
+	}
+}
+
+FBlake3Hash FCbField::GetHash() const
+{
+	FBlake3 Hash;
+	GetHash(Hash);
+	return Hash.Finalize();
+}
+
+void FCbField::GetHash(FBlake3& Hash) const
+{
+	const ECbFieldType SerializedType = FCbFieldType::GetSerializedType(Type);
+	Hash.Update(&SerializedType, sizeof(SerializedType));
+	Hash.Update(GetViewNoType());
+}
+
+bool FCbField::Equals(const FCbField& Other) const
+{
+	return FCbFieldType::GetSerializedType(Type) == FCbFieldType::GetSerializedType(Other.Type) &&
+		GetViewNoType().EqualBytes(Other.GetViewNoType());
+}
+
+void FCbField::CopyTo(FMutableMemoryView Buffer) const
+{
+	const FConstMemoryView Source = GetViewNoType();
+	checkf(Buffer.GetSize() == sizeof(ECbFieldType) + Source.GetSize(),
+		TEXT("A buffer of %" UINT64_FMT " bytes was provided when %" UINT64_FMT " bytes are required"),
+		Buffer.GetSize(), sizeof(ECbFieldType) + Source.GetSize());
+	*static_cast<ECbFieldType*>(Buffer.GetData()) = FCbFieldType::GetSerializedType(Type);
+	Buffer.RightChopInline(sizeof(ECbFieldType));
+	FMemory::Memcpy(Buffer.GetData(), Source.GetData(), Source.GetSize());
+}
+
 void FCbField::CopyTo(FArchive& Ar) const
 {
 	check(Ar.IsSaving());
-	const FConstMemoryView Source = GetFieldView();
+	const FConstMemoryView Source = GetViewNoType();
+	ECbFieldType SerializedType = FCbFieldType::GetSerializedType(Type);
+	Ar.Serialize(&SerializedType, sizeof(SerializedType));
 	Ar.Serialize(const_cast<void*>(Source.GetData()), static_cast<int64>(Source.GetSize()));
 }
 
+void FCbField::IterateReferences(FCbFieldVisitor Visitor) const
+{
+	switch (FCbFieldType::GetType(Type))
+	{
+	case ECbFieldType::Object:
+	case ECbFieldType::UniformObject:
+		return FCbObject::FromField(*this).IterateReferences(Visitor);
+	case ECbFieldType::Array:
+	case ECbFieldType::UniformArray:
+		return FCbArray::FromField(*this).IterateReferences(Visitor);
+	case ECbFieldType::Reference:
+	case ECbFieldType::BinaryReference:
+		return Visitor(*this);
+	default:
+		return;
+	}
+}
+
+FConstMemoryView FCbField::GetView() const
+{
+	const uint32 TypeSize = FCbFieldType::HasFieldType(Type) ? sizeof(ECbFieldType) : 0;
+	const uint32 NameSize = FCbFieldType::HasFieldName(Type) ? NameLen + MeasureVarUInt(NameLen) : 0;
+	const uint64 PayloadSize = GetPayloadSize();
+	return MakeMemoryView(static_cast<const uint8*>(Payload) - TypeSize - NameSize, TypeSize + NameSize + PayloadSize);
+}
+
+FConstMemoryView FCbField::GetViewNoType() const
+{
+	const uint32 NameSize = FCbFieldType::HasFieldName(Type) ? NameLen + MeasureVarUInt(NameLen) : 0;
+	const uint64 PayloadSize = GetPayloadSize();
+	return MakeMemoryView(static_cast<const uint8*>(Payload) - NameSize, NameSize + PayloadSize);
+}
+
 ///////////////////////////////////////////////////////////////////////////////////////////////////////////////////////
 
 FCbArray::FCbArray()
-	: FCbField(CompactBinaryPrivate::GEmptyArrayPayload, ECbFieldType::Array)
+	: FCbField(CompactBinaryPrivate::GEmptyArrayPayload)
 {
 }
 
-uint64 FCbArray::GetSize() const
-{
-	return sizeof(ECbFieldType) + GetPayloadSize();
-}
-
 uint64 FCbArray::Num() const
@@ -495,15 +516,14 @@
 		PayloadBytes += NumByteCount;
 		const ECbFieldType UniformType = FCbFieldType::GetType(GetType()) == ECbFieldType::UniformArray ?
 			ECbFieldType(*PayloadBytes++) : ECbFieldType::HasFieldType;
-		return FCbFieldIterator(FCbField(PayloadBytes, UniformType), PayloadEnd);
+		return FCbFieldIterator::MakeRange(MakeMemoryView(PayloadBytes, PayloadEnd), UniformType);
 	}
 	return FCbFieldIterator();
 }
 
-bool FCbArray::Equals(const FCbArray& Other) const
+uint64 FCbArray::GetSize() const
 {
-	return FCbFieldType::GetType(GetType()) == FCbFieldType::GetType(Other.GetType()) &&
-		GetPayloadView().EqualBytes(Other.GetPayloadView());
+	return sizeof(ECbFieldType) + GetPayloadSize();
 }
 
 FBlake3Hash FCbArray::GetHash() const
@@ -515,40 +535,44 @@
 
 void FCbArray::GetHash(FBlake3& Hash) const
 {
-	const uint8 LocalType = uint8(FCbFieldType::GetType(GetType()));
-	Hash.Update(&LocalType, sizeof(LocalType));
+	const ECbFieldType SerializedType = FCbFieldType::GetType(GetType());
+	Hash.Update(&SerializedType, sizeof(SerializedType));
 	Hash.Update(GetPayloadView());
 }
 
-void FCbArray::CopyTo(const FMutableMemoryView Buffer) const
+bool FCbArray::Equals(const FCbArray& Other) const
 {
-	uint8* BufferBytes = static_cast<uint8*>(Buffer.GetData());
-	*BufferBytes++ = uint8(FCbFieldType::GetType(GetType()));
-	const FConstMemoryView Source = GetPayloadView();
-	FMemory::Memcpy(BufferBytes, Source.GetData(), FMath::Min(Buffer.GetSize() - 1, Source.GetSize()));
+	return FCbFieldType::GetType(GetType()) == FCbFieldType::GetType(Other.GetType()) &&
+		GetPayloadView().EqualBytes(Other.GetPayloadView());
 }
 
+void FCbArray::CopyTo(FMutableMemoryView Buffer) const
+{
+	const FConstMemoryView Source = GetPayloadView();
+	checkf(Buffer.GetSize() == sizeof(ECbFieldType) + Source.GetSize(),
+		TEXT("Buffer is %" UINT64_FMT " bytes but %" UINT64_FMT " is required."),
+		Buffer.GetSize(), sizeof(ECbFieldType) + Source.GetSize());
+	*static_cast<ECbFieldType*>(Buffer.GetData()) = FCbFieldType::GetType(GetType());
+	Buffer.RightChopInline(sizeof(ECbFieldType));
+	FMemory::Memcpy(Buffer.GetData(), Source.GetData(), Source.GetSize());
+}
+
 void FCbArray::CopyTo(FArchive& Ar) const
 {
 	check(Ar.IsSaving());
-	uint8 TypeByte = uint8(FCbFieldType::GetType(GetType()));
 	const FConstMemoryView Source = GetPayloadView();
-	Ar.Serialize(&TypeByte, sizeof(TypeByte));
+	ECbFieldType SerializedType = FCbFieldType::GetType(GetType());
+	Ar.Serialize(&SerializedType, sizeof(SerializedType));
 	Ar.Serialize(const_cast<void*>(Source.GetData()), static_cast<int64>(Source.GetSize()));
 }
 
 ///////////////////////////////////////////////////////////////////////////////////////////////////////////////////////
 
 FCbObject::FCbObject()
-	: FCbField(CompactBinaryPrivate::GEmptyObjectPayload, ECbFieldType::Object)
+	: FCbField(CompactBinaryPrivate::GEmptyObjectPayload)
 {
 }
 
-uint64 FCbObject::GetSize() const
-{
-	return sizeof(ECbFieldType) + GetPayloadSize();
-}
-
 FCbFieldIterator FCbObject::CreateIterator() const
 {
 	const uint8* PayloadBytes = static_cast<const uint8*>(GetPayload());
@@ -560,17 +584,11 @@
 		const void* const PayloadEnd = PayloadBytes + PayloadSize;
 		const ECbFieldType UniformType = FCbFieldType::GetType(GetType()) == ECbFieldType::UniformObject ?
 			ECbFieldType(*PayloadBytes++) : ECbFieldType::HasFieldType;
-		return FCbFieldIterator(FCbField(PayloadBytes, UniformType), PayloadEnd);
+		return FCbFieldIterator::MakeRange(MakeMemoryView(PayloadBytes, PayloadEnd), UniformType);
 	}
 	return FCbFieldIterator();
 }
 
-bool FCbObject::Equals(const FCbObject& Other) const
-{
-	return FCbFieldType::GetType(GetType()) == FCbFieldType::GetType(Other.GetType()) &&
-		GetPayloadView().EqualBytes(Other.GetPayloadView());
-}
-
 FCbField FCbObject::Find(const FAnsiStringView Name) const
 {
 	for (const FCbField Field : *this)
@@ -595,6 +613,11 @@
 	return FCbField();
 }
 
+uint64 FCbObject::GetSize() const
+{
+	return sizeof(ECbFieldType) + GetPayloadSize();
+}
+
 FBlake3Hash FCbObject::GetHash() const
 {
 	FBlake3 Hash;
@@ -604,26 +627,190 @@
 
 void FCbObject::GetHash(FBlake3& Hash) const
 {
-	const uint8 LocalType = uint8(FCbFieldType::GetType(GetType()));
-	Hash.Update(&LocalType, sizeof(LocalType));
+	const ECbFieldType SerializedType = FCbFieldType::GetType(GetType());
+	Hash.Update(&SerializedType, sizeof(SerializedType));
 	Hash.Update(GetPayloadView());
 }
 
-void FCbObject::CopyTo(const FMutableMemoryView Buffer) const
+bool FCbObject::Equals(const FCbObject& Other) const
 {
-	uint8* BufferBytes = static_cast<uint8*>(Buffer.GetData());
-	*BufferBytes++ = uint8(FCbFieldType::GetType(GetType()));
-	const FConstMemoryView Source = GetPayloadView();
-	FMemory::Memcpy(BufferBytes, Source.GetData(), FMath::Min(Buffer.GetSize() - 1, Source.GetSize()));
+	return FCbFieldType::GetType(GetType()) == FCbFieldType::GetType(Other.GetType()) &&
+		GetPayloadView().EqualBytes(Other.GetPayloadView());
 }
 
+void FCbObject::CopyTo(FMutableMemoryView Buffer) const
+{
+	const FConstMemoryView Source = GetPayloadView();
+	checkf(Buffer.GetSize() == sizeof(ECbFieldType) + Source.GetSize(),
+		TEXT("Buffer is %" UINT64_FMT " bytes but %" UINT64_FMT " is required."),
+		Buffer.GetSize(), sizeof(ECbFieldType) + Source.GetSize());
+	*static_cast<ECbFieldType*>(Buffer.GetData()) = FCbFieldType::GetType(GetType());
+	Buffer.RightChopInline(sizeof(ECbFieldType));
+	FMemory::Memcpy(Buffer.GetData(), Source.GetData(), Source.GetSize());
+}
+
 void FCbObject::CopyTo(FArchive& Ar) const
 {
 	check(Ar.IsSaving());
-	uint8 TypeByte = uint8(FCbFieldType::GetType(GetType()));
 	const FConstMemoryView Source = GetPayloadView();
-	Ar.Serialize(&TypeByte, sizeof(TypeByte));
+	ECbFieldType SerializedType = FCbFieldType::GetType(GetType());
+	Ar.Serialize(&SerializedType, sizeof(SerializedType));
 	Ar.Serialize(const_cast<void*>(Source.GetData()), static_cast<int64>(Source.GetSize()));
 }
 
 ///////////////////////////////////////////////////////////////////////////////////////////////////////////////////////
+
+template <typename FieldType>
+uint64 TCbFieldIterator<FieldType>::GetRangeSize() const
+{
+	FConstMemoryView View;
+	if (TryGetSerializedRangeView(View))
+	{
+		return View.GetSize();
+	}
+	else
+	{
+		uint64 Size = 0;
+		for (TCbFieldIterator It(*this); It; ++It)
+		{
+			Size += It.GetSize();
+		}
+		return Size;
+	}
+}
+
+template <typename FieldType>
+FBlake3Hash TCbFieldIterator<FieldType>::GetRangeHash() const
+{
+	FBlake3 Hash;
+	GetRangeHash(Hash);
+	return Hash.Finalize();
+}
+
+template <typename FieldType>
+void TCbFieldIterator<FieldType>::GetRangeHash(FBlake3& Hash) const
+{
+	FConstMemoryView View;
+	if (TryGetSerializedRangeView(View))
+	{
+		Hash.Update(View);
+	}
+	else
+	{
+		for (TCbFieldIterator It(*this); It; ++It)
+		{
+			It.GetHash(Hash);
+		}
+	}
+}
+
+template <typename FieldType>
+void TCbFieldIterator<FieldType>::CopyRangeTo(FMutableMemoryView Buffer) const
+{
+	FConstMemoryView Source;
+	if (TryGetSerializedRangeView(Source))
+	{
+		checkf(Buffer.GetSize() == Source.GetSize(),
+			TEXT("Buffer is %" UINT64_FMT " bytes but %" UINT64_FMT " is required."),
+			Buffer.GetSize(), Source.GetSize());
+		FMemory::Memcpy(Buffer.GetData(), Source.GetData(), Source.GetSize());
+	}
+	else
+	{
+		for (TCbFieldIterator It(*this); It; ++It)
+		{
+			const uint64 Size = It.GetSize();
+			It.CopyTo(Buffer.Left(Size));
+			Buffer += Size;
+		}
+	}
+}
+
+template <typename FieldType>
+void TCbFieldIterator<FieldType>::CopyRangeTo(FArchive& Ar) const
+{
+	check(Ar.IsSaving());
+	FConstMemoryView Source;
+	if (TryGetSerializedRangeView(Source))
+	{
+		Ar.Serialize(const_cast<void*>(Source.GetData()), static_cast<int64>(Source.GetSize()));
+	}
+	else
+	{
+		for (TCbFieldIterator It(*this); It; ++It)
+		{
+			It.CopyTo(Ar);
+		}
+	}
+}
+
+template <typename FieldType>
+void TCbFieldIterator<FieldType>::IterateRangeReferences(FCbFieldVisitor Visitor) const
+{
+	// Always iterate over non-uniform ranges because we do not know if they contain a reference.
+	if (FCbFieldType::HasFieldType(FieldType::GetType()))
+	{
+		for (TCbFieldIterator It(*this); It; ++It)
+		{
+			if (FCbFieldType::MayContainReferences(It.GetType()))
+			{
+				It.IterateReferences(Visitor);
+			}
+		}
+	}
+	// Only iterate over uniform ranges if the uniform type may contain a reference.
+	else
+	{
+		if (FCbFieldType::MayContainReferences(FieldType::GetType()))
+		{
+			for (TCbFieldIterator It(*this); It; ++It)
+			{
+				It.IterateReferences(Visitor);
+			}
+		}
+	}
+}
+
+template class TCbFieldIterator<FCbField>;
+template class TCbFieldIterator<FCbFieldRef>;
+
+///////////////////////////////////////////////////////////////////////////////////////////////////////////////////////
+
+FCbFieldRefIterator FCbFieldRefIterator::CloneRange(const FCbFieldIterator& It)
+{
+	FSharedBufferPtr Buffer;
+	FConstMemoryView View;
+	if (It.TryGetSerializedRangeView(View))
+	{
+		Buffer = FSharedBuffer::Clone(View);
+	}
+	else
+	{
+		Buffer = FSharedBuffer::Alloc(It.GetRangeSize());
+		It.CopyRangeTo(*Buffer);
+	}
+	return MakeRange(MoveTemp(Buffer));
+}
+
+FSharedBufferConstPtr FCbFieldRefIterator::GetRangeBuffer() const
+{
+	const FConstMemoryView RangeView = GetRangeView();
+	const FSharedBufferConstPtr& OuterBuffer = GetOuterBuffer();
+	if (OuterBuffer)
+	{
+		if (OuterBuffer->GetView() == RangeView)
+		{
+			return OuterBuffer;
+		}
+		else
+		{
+			return FSharedBuffer::MakeView(RangeView, *OuterBuffer);
+		}
+	}
+	else
+	{
+		return FSharedBuffer::MakeView(RangeView);
+	}
+}
+
+///////////////////////////////////////////////////////////////////////////////////////////////////////////////////////
diff --git a/Engine/Source/Runtime/Core/Private/Serialization/CompactBinarySerialization.cpp b/Engine/Source/Runtime/Core/Private/Serialization/CompactBinarySerialization.cpp
--- a/Engine/Source/Runtime/Core/Private/Serialization/CompactBinarySerialization.cpp
+++ b/Engine/Source/Runtime/Core/Private/Serialization/CompactBinarySerialization.cpp
@@ -183,28 +183,23 @@ FCbFieldRef LoadCompactBinary(FArchive& Ar, FCbBufferAllocator Allocator)
 
 ///////////////////////////////////////////////////////////////////////////////////////////////////////////////////////
 
-template <typename T>
-static void SaveCompactBinaryToArchive(FArchive& Ar, const T& Value)
+void SaveCompactBinary(FArchive& Ar, const FCbField& Field)
 {
 	check(Ar.IsSaving());
-	struct FCopy : public T
-	{
-		using T::GetCopyType;
-		using T::CopyTo;
-	};
-	const FCopy& ValueCopy = static_cast<const FCopy&>(Value);
-	const ECbFieldType CopyType = ValueCopy.GetCopyType();
-	if (!EnumHasAnyFlags(CopyType, ECbFieldType::HasFieldType))
-	{
-		uint8 TypeByte = uint8(CopyType);
-		Ar.Serialize(&TypeByte, sizeof(TypeByte));
-	}
-	ValueCopy.CopyTo(Ar);
+	Field.CopyTo(Ar);
+}
+
+void SaveCompactBinary(FArchive& Ar, const FCbArray& Array)
+{
+	check(Ar.IsSaving());
+	Array.CopyTo(Ar);
 }
 
-void SaveCompactBinary(FArchive& Ar, const FCbField& Field) { SaveCompactBinaryToArchive(Ar, Field); }
-void SaveCompactBinary(FArchive& Ar, const FCbArray& Array) { SaveCompactBinaryToArchive(Ar, Array); }
-void SaveCompactBinary(FArchive& Ar, const FCbObject& Object) { SaveCompactBinaryToArchive(Ar, Object); }
+void SaveCompactBinary(FArchive& Ar, const FCbObject& Object)
+{
+	check(Ar.IsSaving());
+	Object.CopyTo(Ar);
+}
 
 ///////////////////////////////////////////////////////////////////////////////////////////////////////////////////////
 
@@ -221,7 +216,7 @@ static FArchive& SerializeCompactBinary(FArchive& Ar, T& Value, ConvertType&& Co
 	}
 	else if (Ar.IsSaving())
 	{
-		SaveCompactBinary(Ar, Value);
+		Value.CopyTo(Ar);
 	}
 	else
 	{
diff --git a/Engine/Source/Runtime/Core/Private/Serialization/CompactBinaryWriter.cpp b/Engine/Source/Runtime/Core/Private/Serialization/CompactBinaryWriter.cpp
--- a/Engine/Source/Runtime/Core/Private/Serialization/CompactBinaryWriter.cpp
+++ b/Engine/Source/Runtime/Core/Private/Serialization/CompactBinaryWriter.cpp
@@ -48,19 +48,18 @@ static constexpr bool IsUniformType(const ECbFieldType Type)
 }
 
 /** Append the payload from the compact binary value to the array and return its type. */
-template <typename T>
-static inline ECbFieldType AppendCompactBinary(const T& Value, TArray64<uint8>& OutData)
+static inline ECbFieldType AppendCompactBinary(const FCbField& Value, TArray64<uint8>& OutData)
 {
-	struct FCopy : public T
+	struct FCopy : public FCbField
 	{
-		using T::GetCopyType;
-		using T::GetPayloadView;
+		using FCbField::GetType;
+		using FCbField::GetPayloadView;
 	};
 	const FCopy& ValueCopy = static_cast<const FCopy&>(Value);
 	const FConstMemoryView SourceView = ValueCopy.GetPayloadView();
 	const int64 TargetOffset = OutData.AddUninitialized(SourceView.GetSize());
 	FMemory::Memcpy(OutData.GetData() + TargetOffset, SourceView.GetData(), SourceView.GetSize());
-	return FCbFieldType::GetType(ValueCopy.GetCopyType());
+	return FCbFieldType::GetType(ValueCopy.GetType());
 }
 
 FCbWriter::FCbWriter()
@@ -90,7 +89,7 @@ FCbFieldRefIterator FCbWriter::Save() const
 	const uint64 Size = GetSaveSize();
 	FSharedBufferPtr Buffer = FSharedBuffer::Alloc(Size);
 	const FCbFieldIterator Output = SaveToMemory(Buffer->GetView());
-	return FCbFieldRefIterator(Output, FSharedBuffer::MakeReadOnly(MoveTemp(Buffer)));
+	return FCbFieldRefIterator::MakeRangeView(Output, FSharedBuffer::MakeReadOnly(MoveTemp(Buffer)));
 }
 
 FCbFieldIterator FCbWriter::SaveToMemory(const FMutableMemoryView Buffer) const
@@ -101,7 +100,7 @@ FCbFieldIterator FCbWriter::SaveToMemory(const FMutableMemoryView Buffer) const
 	checkf(Buffer.GetSize() == Data.Num(),
 		TEXT("Buffer is %" UINT64_FMT " bytes but %" INT64_FMT " is required."), Buffer.GetSize(), Data.Num());
 	FMemory::Memcpy(Buffer.GetData(), Data.GetData(), Data.Num());
-	return FCbFieldIterator(FCbField(Buffer.GetData()), Buffer.GetDataEnd());
+	return FCbFieldIterator::MakeRange(Buffer);
 }
 
 void FCbWriter::SaveToArchive(FArchive& Ar) const
@@ -289,7 +288,7 @@ void FCbWriter::EndObject()
 void FCbWriter::Object(const FCbObject& Value)
 {
 	BeginField();
-	EndField(AppendCompactBinary(Value, Data));
+	EndField(AppendCompactBinary(Value.AsField(), Data));
 }
 
 void FCbWriter::Object(const FCbObjectRef& Value)
@@ -344,7 +343,7 @@ void FCbWriter::EndArray()
 void FCbWriter::Array(const FCbArray& Value)
 {
 	BeginField();
-	EndField(AppendCompactBinary(Value, Data));
+	EndField(AppendCompactBinary(Value.AsField(), Data));
 }
 
 void FCbWriter::Array(const FCbArrayRef& Value)
diff --git a/Engine/Source/Runtime/Core/Private/Tests/Serialization/CompactBinarySerializationTest.cpp b/Engine/Source/Runtime/Core/Private/Tests/Serialization/CompactBinarySerializationTest.cpp
--- a/Engine/Source/Runtime/Core/Private/Tests/Serialization/CompactBinarySerializationTest.cpp
+++ b/Engine/Source/Runtime/Core/Private/Tests/Serialization/CompactBinarySerializationTest.cpp
@@ -140,13 +140,9 @@ bool FCbSaveTest::RunTest(const FString& Parameters)
 			if (TestEqual(FString::Printf(TEXT("SaveCompactBinary(%s)->Size"), Test), uint64(WriteAr.Num()), ExpectedData.GetSize()) &&
 				TestTrue(FString::Printf(TEXT("SaveCompactBinary(%s)->EqualBytes"), Test), ExpectedData.EqualBytes(MakeMemoryView(WriteAr))))
 			{
-				struct FView : public FCbFieldRef
-				{
-					using FCbFieldRef::GetFieldView;
-				};
 				FMemoryReader ReadAr(WriteAr);
 				FCbFieldRef Field = LoadCompactBinary(ReadAr, [](ECbFieldType Type, uint64 Size) { return FSharedBuffer::Alloc(Size); });
-				TestTrue(FString::Printf(TEXT("LoadCompactBinary(%s)->EqualBytes"), Test), ExpectedData.EqualBytes(static_cast<FView&>(Field).GetFieldView()));
+				TestTrue(FString::Printf(TEXT("LoadCompactBinary(%s)->EqualBytes"), Test), ExpectedData.EqualBytes(Field.GetView()));
 			}
 		}
 		{
@@ -155,13 +151,9 @@ bool FCbSaveTest::RunTest(const FString& Parameters)
 			if (TestEqual(FString::Printf(TEXT("Ar << CompactBinary Save(%s)->Size"), Test), uint64(WriteAr.Num()), ExpectedData.GetSize()) &&
 				TestTrue(FString::Printf(TEXT("Ar << CompactBinary Save(%s)->EqualBytes"), Test), ExpectedData.EqualBytes(MakeMemoryView(WriteAr))))
 			{
-				struct FView : public decltype(Value)
-				{
-					using decltype(Value)::GetFieldView;
-				};
 				FMemoryReader ReadAr(WriteAr);
 				ReadAr << Value;
-				TestTrue(FString::Printf(TEXT("Ar << CompactBinary Load(%s)->EqualBytes"), Test), ExpectedData.EqualBytes(static_cast<FView&>(Value).GetFieldView()));
+				TestTrue(FString::Printf(TEXT("Ar << CompactBinary Load(%s)->EqualBytes"), Test), ExpectedData.EqualBytes(Value.GetView()));
 			}
 		}
 	};
diff --git a/Engine/Source/Runtime/Core/Private/Tests/Serialization/CompactBinaryTest.cpp b/Engine/Source/Runtime/Core/Private/Tests/Serialization/CompactBinaryTest.cpp
--- a/Engine/Source/Runtime/Core/Private/Tests/Serialization/CompactBinaryTest.cpp
+++ b/Engine/Source/Runtime/Core/Private/Tests/Serialization/CompactBinaryTest.cpp
@@ -127,7 +127,7 @@ protected:
 	void TestField(const TCHAR* What, TArrayView<const uint8> Payload, T ExpectedValue = T(), T DefaultValue = T())
 	{
 		FCbField Field(Payload.GetData(), FieldType);
-		TestEqual(FString::Printf(TEXT("FCbField::GetSize(%s)"), What), Field.GetSize(), uint64(Payload.Num()));
+		TestEqual(FString::Printf(TEXT("FCbField::GetSize(%s)"), What), Field.GetSize(), uint64(Payload.Num() + !FCbFieldType::HasFieldType(FieldType)));
 		TestTrue(FString::Printf(TEXT("FCbField::HasValue(%s)"), What), Field.HasValue());
 		TestFalse(FString::Printf(TEXT("FCbField::HasError(%s) == false"), What), Field.HasError());
 		TestEqual(FString::Printf(TEXT("FCbField::GetError(%s) == None"), What), Field.GetError(), ECbFieldError::None);
@@ -185,25 +185,31 @@ bool FCbFieldNoneTest::RunTest(const FString& Parameters)
 		static_assert(!DefaultField.HasValue(), "Error in HasValue()");
 		static_assert(!DefaultField.HasError(), "Error in HasError()");
 		static_assert(DefaultField.GetError() == ECbFieldError::None, "Error in GetError()");
-		TestEqual(TEXT("FCbField()::GetSize() == 0"), DefaultField.GetSize(), uint64(0));
+		TestEqual(TEXT("FCbField()::GetSize() == 1"), DefaultField.GetSize(), uint64(1));
 		TestEqual(TEXT("FCbField()::GetName().Len() == 0"), DefaultField.GetName().Len(), 0);
 		TestFalse(TEXT("!FCbField()::HasName()"), DefaultField.HasName());
 		TestFalse(TEXT("!FCbField()::HasValue()"), DefaultField.HasValue());
 		TestFalse(TEXT("!FCbField()::HasError()"), DefaultField.HasError());
 		TestEqual(TEXT("FCbField()::GetError() == None"), DefaultField.GetError(), ECbFieldError::None);
-		TestEqual(TEXT("FCbField()::GetHash()"), DefaultField.GetHash(), FBlake3::HashBuffer(FConstMemoryView()));
+		TestEqual(TEXT("FCbField()::GetHash()"), DefaultField.GetHash(), FBlake3::HashBuffer(MakeMemoryView<uint8>({uint8(ECbFieldType::None)})));
+		TestEqual(TEXT("FCbField()::GetView()"), DefaultField.GetView(), FConstMemoryView());
+		FConstMemoryView SerializedView;
+		TestFalse(TEXT("FCbField()::TryGetSerializedView()"), DefaultField.TryGetSerializedView(SerializedView));
 	}
 
 	// Test FCbField(None)
 	{
 		FCbField NoneField(nullptr, ECbFieldType::None);
-		TestEqual(TEXT("FCbField(None)::GetSize() == 0"), NoneField.GetSize(), uint64(0));
+		TestEqual(TEXT("FCbField(None)::GetSize() == 1"), NoneField.GetSize(), uint64(1));
 		TestEqual(TEXT("FCbField(None)::GetName().Len() == 0"), NoneField.GetName().Len(), 0);
 		TestFalse(TEXT("!FCbField(None)::HasName()"), NoneField.HasName());
 		TestFalse(TEXT("!FCbField(None)::HasValue()"), NoneField.HasValue());
 		TestFalse(TEXT("!FCbField(None)::HasError()"), NoneField.HasError());
 		TestEqual(TEXT("FCbField(None)::GetError() == None"), NoneField.GetError(), ECbFieldError::None);
 		TestEqual(TEXT("FCbField(None)::GetHash()"), NoneField.GetHash(), FCbField().GetHash());
+		TestEqual(TEXT("FCbField(None)::GetView()"), NoneField.GetView(), FConstMemoryView());
+		FConstMemoryView SerializedView;
+		TestFalse(TEXT("FCbField(None)::TryGetSerializedView()"), NoneField.TryGetSerializedView(SerializedView));
 	}
 
 	// Test FCbField(None|Type|Name)
@@ -215,7 +221,14 @@ bool FCbFieldNoneTest::RunTest(const FString& Parameters)
 		TestEqual(TEXT("FCbField(None|Type|Name)::GetName()"), NoneField.GetName(), "Name"_ASV);
 		TestTrue(TEXT("FCbField(None|Type|Name)::HasName()"), NoneField.HasName());
 		TestFalse(TEXT("!FCbField(None|Type|Name)::HasValue()"), NoneField.HasValue());
-		TestNotEqual(TEXT("FCbField(None|Type|Name)::GetHash()"), NoneField.GetHash(), FCbField().GetHash());
+		TestEqual(TEXT("FCbField(None|Type|Name)::GetHash()"), NoneField.GetHash(), FBlake3::HashBuffer(MakeMemoryView(NoneBytes)));
+		TestEqual(TEXT("FCbField(None|Type|Name)::GetView()"), NoneField.GetView(), MakeMemoryView(NoneBytes));
+		FConstMemoryView SerializedView;
+		TestTrue(TEXT("FCbField(None|Type|Name)::TryGetSerializedView()"), NoneField.TryGetSerializedView(SerializedView) && SerializedView == MakeMemoryView(NoneBytes));
+
+		uint8 CopyBytes[sizeof(NoneBytes)];
+		NoneField.CopyTo(MakeMemoryView(CopyBytes));
+		TestTrue(TEXT("FCbField(None|Type|Name)::CopyTo()"), MakeMemoryView(NoneBytes).EqualBytes(MakeMemoryView(CopyBytes)));
 	}
 
 	// Test FCbField(None|Type)
@@ -227,19 +240,44 @@ bool FCbFieldNoneTest::RunTest(const FString& Parameters)
 		TestEqual(TEXT("FCbField(None|Type)::GetName()"), NoneField.GetName().Len(), 0);
 		TestFalse(TEXT("FCbField(None|Type)::HasName()"), NoneField.HasName());
 		TestFalse(TEXT("!FCbField(None|Type)::HasValue()"), NoneField.HasValue());
-		TestNotEqual(TEXT("FCbField(None|Type)::GetHash()"), NoneField.GetHash(), FCbField().GetHash());
+		TestEqual(TEXT("FCbField(None|Type)::GetHash()"), NoneField.GetHash(), FCbField().GetHash());
+		TestEqual(TEXT("FCbField(None|Type)::GetView()"), NoneField.GetView(), MakeMemoryView(NoneBytes));
+		FConstMemoryView SerializedView;
+		TestTrue(TEXT("FCbField(None|Type)::TryGetSerializedView()"), NoneField.TryGetSerializedView(SerializedView) && SerializedView == MakeMemoryView(NoneBytes));
 	}
 
 	// Test FCbField(None|Name)
 	{
 		constexpr ECbFieldType FieldType = ECbFieldType::None | ECbFieldType::HasFieldName;
-		constexpr const ANSICHAR NoneBytes[] = { 4, 'N', 'a', 'm', 'e' };
-		FCbField NoneField(NoneBytes, FieldType);
+		constexpr const ANSICHAR NoneBytes[] = { ANSICHAR(FieldType), 4, 'N', 'a', 'm', 'e' };
+		FCbField NoneField(NoneBytes + 1, FieldType);
 		TestEqual(TEXT("FCbField(None|Name)::GetSize()"), NoneField.GetSize(), uint64(sizeof(NoneBytes)));
 		TestEqual(TEXT("FCbField(None|Name)::GetName()"), NoneField.GetName(), "Name"_ASV);
 		TestTrue(TEXT("FCbField(None|Name)::HasName()"), NoneField.HasName());
 		TestFalse(TEXT("!FCbField(None|Name)::HasValue()"), NoneField.HasValue());
-		TestNotEqual(TEXT("FCbField(None|Name)::GetHash()"), NoneField.GetHash(), FCbField().GetHash());
+		TestEqual(TEXT("FCbField(None|Name)::GetHash()"), NoneField.GetHash(), FBlake3::HashBuffer(MakeMemoryView(NoneBytes)));
+		TestEqual(TEXT("FCbField(None|Name)::GetView()"), NoneField.GetView(), MakeMemoryView(NoneBytes) + 1);
+		FConstMemoryView SerializedView;
+		TestFalse(TEXT("FCbField(None|Name)::TryGetSerializedView()"), NoneField.TryGetSerializedView(SerializedView));
+
+		uint8 CopyBytes[sizeof(NoneBytes)];
+		NoneField.CopyTo(MakeMemoryView(CopyBytes));
+		TestTrue(TEXT("FCbField(None|Name)::CopyTo()"), MakeMemoryView(NoneBytes).EqualBytes(MakeMemoryView(CopyBytes)));
+	}
+
+	// Test FCbField(None|EmptyName)
+	{
+		constexpr ECbFieldType FieldType = ECbFieldType::None | ECbFieldType::HasFieldName;
+		constexpr const uint8 NoneBytes[] = { uint8(FieldType), 0 };
+		FCbField NoneField(NoneBytes + 1, FieldType);
+		TestEqual(TEXT("FCbField(None|EmptyName)::GetSize()"), NoneField.GetSize(), uint64(sizeof(NoneBytes)));
+		TestEqual(TEXT("FCbField(None|EmptyName)::GetName()"), NoneField.GetName(), ""_ASV);
+		TestTrue(TEXT("FCbField(None|EmptyName)::HasName()"), NoneField.HasName());
+		TestFalse(TEXT("!FCbField(None|EmptyName)::HasValue()"), NoneField.HasValue());
+		TestEqual(TEXT("FCbField(None|EmptyName)::GetHash()"), NoneField.GetHash(), FBlake3::HashBuffer(MakeMemoryView(NoneBytes)));
+		TestEqual(TEXT("FCbField(None|EmptyName)::GetView()"), NoneField.GetView(), MakeMemoryView(NoneBytes) + 1);
+		FConstMemoryView SerializedView;
+		TestFalse(TEXT("FCbField(None|EmptyName)::TryGetSerializedView()"), NoneField.TryGetSerializedView(SerializedView));
 	}
 
 	return true;
@@ -251,12 +289,12 @@ bool FCbFieldNullTest::RunTest(const FString& Parameters)
 	// Test FCbField(Null)
 	{
 		FCbField NullField(nullptr, ECbFieldType::Null);
-		TestEqual(TEXT("FCbField(Null)::GetSize() == 0"), NullField.GetSize(), uint64(0));
+		TestEqual(TEXT("FCbField(Null)::GetSize() == 1"), NullField.GetSize(), uint64(1));
 		TestTrue(TEXT("FCbField(Null)::IsNull()"), NullField.IsNull());
 		TestTrue(TEXT("FCbField(Null)::HasValue()"), NullField.HasValue());
 		TestFalse(TEXT("!FCbField(Null)::HasError()"), NullField.HasError());
 		TestEqual(TEXT("FCbField(Null)::GetError() == None"), NullField.GetError(), ECbFieldError::None);
-		TestEqual(TEXT("FCbField(Null)::GetHash()"), NullField.GetHash(), FBlake3::HashBuffer(FConstMemoryView()));
+		TestEqual(TEXT("FCbField(Null)::GetHash()"), NullField.GetHash(), FBlake3::HashBuffer(MakeMemoryView<uint8>({uint8(ECbFieldType::Null)})));
 	}
 
 	// Test FCbField(None) as Null
@@ -341,6 +379,7 @@ bool FCbFieldObjectTest::RunTest(const FString& Parameters)
 		TestEqual(TEXT("FCbObject::FindIgnoreCase()"), Object.FindIgnoreCase("b"_ASV).AsInt32(), 2);
 		TestEqual(TEXT("FCbObject::operator[]"), Object["B"_ASV].AsInt32(), 2);
 		TestEqual(TEXT("FCbObject::operator[]"), Object["b"_ASV].AsInt32(4), 4);
+		TestEqual(TEXT("FCbObject::GetView()"), Field.AsObject().GetView(), Field.GetView());
 	}
 
 	// Test FCbField(UniformObject, NotEmpty)
@@ -352,27 +391,41 @@ bool FCbFieldObjectTest::RunTest(const FString& Parameters)
 		FCbObjectRef Object = FCbObjectRef::Clone(Field.AsObject());
 		TestIntObject(Object, 3, sizeof(Payload));
 		TestIntObject(Field.AsObject(), 3, sizeof(Payload));
-		TestTrue(TEXT("FCbObject{Uniform}::Equals()"), Object.Equals(Field.AsObject()));
-		TestEqual(TEXT("FCbObject{Uniform}::Find()"), Object.Find("B"_ASV).AsInt32(), 2);
-		TestEqual(TEXT("FCbObject{Uniform}::Find()"), Object.FindRef("B"_ASV).AsInt32(), 2);
-		TestEqual(TEXT("FCbObject{Uniform}::Find()"), Object.Find("b"_ASV).AsInt32(4), 4);
-		TestEqual(TEXT("FCbObject{Uniform}::Find()"), Object.FindRef("b"_ASV).AsInt32(4), 4);
-		TestEqual(TEXT("FCbObject{Uniform}::FindIgnoreCase()"), Object.FindIgnoreCase("B"_ASV).AsInt32(), 2);
-		TestEqual(TEXT("FCbObject{Uniform}::FindIgnoreCase()"), Object.FindRefIgnoreCase("B"_ASV).AsInt32(), 2);
-		TestEqual(TEXT("FCbObject{Uniform}::FindIgnoreCase()"), Object.FindIgnoreCase("b"_ASV).AsInt32(), 2);
-		TestEqual(TEXT("FCbObject{Uniform}::FindIgnoreCase()"), Object.FindRefIgnoreCase("b"_ASV).AsInt32(), 2);
-		TestEqual(TEXT("FCbObject{Uniform}::operator[]"), Object["B"_ASV].AsInt32(), 2);
-		TestEqual(TEXT("FCbObject{Uniform}::operator[]"), Object["b"_ASV].AsInt32(4), 4);
+		TestTrue(TEXT("FCbObject(Uniform)::Equals()"), Object.Equals(Field.AsObject()));
+		TestEqual(TEXT("FCbObject(Uniform)::Find()"), Object.Find("B"_ASV).AsInt32(), 2);
+		TestEqual(TEXT("FCbObject(Uniform)::Find()"), Object.FindRef("B"_ASV).AsInt32(), 2);
+		TestEqual(TEXT("FCbObject(Uniform)::Find()"), Object.Find("b"_ASV).AsInt32(4), 4);
+		TestEqual(TEXT("FCbObject(Uniform)::Find()"), Object.FindRef("b"_ASV).AsInt32(4), 4);
+		TestEqual(TEXT("FCbObject(Uniform)::FindIgnoreCase()"), Object.FindIgnoreCase("B"_ASV).AsInt32(), 2);
+		TestEqual(TEXT("FCbObject(Uniform)::FindIgnoreCase()"), Object.FindRefIgnoreCase("B"_ASV).AsInt32(), 2);
+		TestEqual(TEXT("FCbObject(Uniform)::FindIgnoreCase()"), Object.FindIgnoreCase("b"_ASV).AsInt32(), 2);
+		TestEqual(TEXT("FCbObject(Uniform)::FindIgnoreCase()"), Object.FindRefIgnoreCase("b"_ASV).AsInt32(), 2);
+		TestEqual(TEXT("FCbObject(Uniform)::operator[]"), Object["B"_ASV].AsInt32(), 2);
+		TestEqual(TEXT("FCbObject(Uniform)::operator[]"), Object["b"_ASV].AsInt32(4), 4);
+		TestEqual(TEXT("FCbObject(Uniform)::GetView()"), Field.AsObject().GetView(), Field.GetView());
 
 		TestTrue(TEXT("FCbObjectRef::AsFieldRef()"), Object.GetBuffer() == Object.AsFieldRef().AsObjectRef().GetBuffer());
 
-		Object = FCbFieldRef::MakeView(Field).AsObjectRef();
-
+		// Equals
 		const uint8 NamedPayload[] = { 1, 'O', 10, IntType, 1, 'A', 1, 1, 'B', 2, 1, 'C', 3 };
 		FCbField NamedField(NamedPayload, ECbFieldType::UniformObject | ECbFieldType::HasFieldName);
-		TestTrue(TEXT("FCbObject::Equals()"), Object.Equals(NamedField.AsObject()));
+		TestTrue(TEXT("FCbObject::Equals()"), Field.AsObject().Equals(NamedField.AsObject()));
 		TestTrue(TEXT("FCbObject::AsField().Equals()"), Field.Equals(Field.AsObject().AsField()));
 		TestTrue(TEXT("FCbObject::AsField().Equals()"), NamedField.Equals(NamedField.AsObject().AsField()));
+		TestEqual(TEXT("FCbObject(Name)::GetView()"), NamedField.AsObject().GetView(), NamedField.GetView());
+
+		// CopyTo
+		uint8 CopyBytes[sizeof(Payload) + 1];
+		Field.AsObject().CopyTo(MakeMemoryView(CopyBytes));
+		TestTrue(TEXT("FCbObject(NoType)::CopyTo()"), MakeMemoryView(Payload).EqualBytes(MakeMemoryView(CopyBytes) + 1));
+		NamedField.AsObject().CopyTo(MakeMemoryView(CopyBytes));
+		TestTrue(TEXT("FCbObject(NoType, Name)::CopyTo()"), MakeMemoryView(Payload).EqualBytes(MakeMemoryView(CopyBytes) + 1));
+
+		// GetSerializedView
+		FConstMemoryView SerializedView;
+		TestTrue(TEXT("FCbObject(Clone)::TryGetSerializedView()"), Object.TryGetSerializedView(SerializedView) && SerializedView == Object.AsField().GetView());
+		TestFalse(TEXT("FCbObject(NoType)::TryGetSerializedView()"), Field.AsObject().TryGetSerializedView(SerializedView));
+		TestFalse(TEXT("FCbObject(Name)::TryGetSerializedView()"), NamedField.AsObject().TryGetSerializedView(SerializedView));
 	}
 
 	// Test FCbField(None) as Object
@@ -398,17 +451,25 @@ bool FCbFieldObjectTest::RunTest(const FString& Parameters)
 			TestEqual(TEXT("FCbObjectRef::CreateRefIterator().GetName()"), Field.GetName(), "F"_ASV);
 			TestEqual(TEXT("FCbObjectRef::CreateRefIterator().AsInt32()"), Field.AsInt32(), 8);
 			TestTrue(TEXT("FCbObjectRef::CreateRefIterator().IsOwned()"), Field.IsOwned());
+			TestEqual(TEXT("FCbObjectRef::CreateRefIterator().GetBuffer()"), Field.GetBuffer()->GetView(), Field.GetView());
 		}
 		for (FCbFieldRefIterator It = ObjectClone.CreateRefIterator(), End; It != End; ++It)
 		{
 		}
+
+		// CopyTo
+		uint8 CopyBytes[6];
+		Object.CopyTo(MakeMemoryView(CopyBytes));
+		TestTrue(TEXT("FCbObject(Name)::CopyTo()"), ObjectClone.GetView().EqualBytes(MakeMemoryView(CopyBytes)));
+		ObjectClone.CopyTo(MakeMemoryView(CopyBytes));
+		TestTrue(TEXT("FCbObject()::CopyTo()"), ObjectClone.GetView().EqualBytes(MakeMemoryView(CopyBytes)));
 	}
 
 	// Test FCbObject as FCbFieldIterator
 	{
 		uint32 Count = 0;
 		FCbObject Object;
-		for (FCbField Field : FCbFieldIterator(Object.AsField()))
+		for (FCbField Field : FCbFieldIterator::MakeSingle(Object.AsField()))
 		{
 			TestTrue(TEXT("FCbObject::AsField() as Iterator"), Field.IsObject());
 			++Count;
@@ -421,7 +482,7 @@ bool FCbFieldObjectTest::RunTest(const FString& Parameters)
 		uint32 Count = 0;
 		FCbObjectRef Object;
 		Object.MakeOwned();
-		for (FCbFieldRef Field : FCbFieldRefIterator(Object.AsFieldRef()))
+		for (FCbFieldRef Field : FCbFieldRefIterator::MakeSingle(Object.AsFieldRef()))
 		{
 			TestTrue(TEXT("FCbObjectRef::AsField() as Iterator"), Field.IsObject());
 			++Count;
@@ -493,6 +554,7 @@ bool FCbFieldArrayTest::RunTest(const FString& Parameters)
 		TestIntArray(Array, 3, sizeof(Payload));
 		TestIntArray(Field.AsArray(), 3, sizeof(Payload));
 		TestTrue(TEXT("FCbArray::Equals()"), Array.Equals(Field.AsArray()));
+		TestEqual(TEXT("FCbArray::GetView()"), Field.AsArray().GetView(), Field.GetView());
 	}
 
 	// Test FCbField(UniformArray)
@@ -504,17 +566,31 @@ bool FCbFieldArrayTest::RunTest(const FString& Parameters)
 		FCbArrayRef Array = FCbArrayRef::Clone(Field.AsArray());
 		TestIntArray(Array, 3, sizeof(Payload));
 		TestIntArray(Field.AsArray(), 3, sizeof(Payload));
-		TestTrue(TEXT("FCbArray::Equals()"), Array.Equals(Field.AsArray()));
+		TestTrue(TEXT("FCbArray(Uniform)::Equals()"), Array.Equals(Field.AsArray()));
+		TestEqual(TEXT("FCbArray(Uniform)::GetView()"), Field.AsArray().GetView(), Field.GetView());
 
 		TestTrue(TEXT("FCbArrayRef::AsFieldRef()"), Array.GetBuffer() == Array.AsFieldRef().AsArrayRef().GetBuffer());
 
-		Array = FCbFieldRef::MakeView(Field).AsArrayRef();
-
+		// Equals
 		const uint8 NamedPayload[] = { 1, 'A', 5, 3, IntType, 1, 2, 3 };
 		FCbField NamedField(NamedPayload, ECbFieldType::UniformArray | ECbFieldType::HasFieldName);
-		TestTrue(TEXT("FCbArray::Equals()"), Array.Equals(NamedField.AsArray()));
+		TestTrue(TEXT("FCbArray::Equals()"), Field.AsArray().Equals(NamedField.AsArray()));
 		TestTrue(TEXT("FCbArray::AsField().Equals()"), Field.Equals(Field.AsArray().AsField()));
 		TestTrue(TEXT("FCbArray::AsField().Equals()"), NamedField.Equals(NamedField.AsArray().AsField()));
+		TestEqual(TEXT("FCbArray(Name)::GetView()"), NamedField.AsArray().GetView(), NamedField.GetView());
+
+		// CopyTo
+		uint8 CopyBytes[sizeof(Payload) + 1];
+		Field.AsArray().CopyTo(MakeMemoryView(CopyBytes));
+		TestTrue(TEXT("FCbArray(NoType)::CopyTo()"), MakeMemoryView(Payload).EqualBytes(MakeMemoryView(CopyBytes) + 1));
+		NamedField.AsArray().CopyTo(MakeMemoryView(CopyBytes));
+		TestTrue(TEXT("FCbArray(NoType, Name)::CopyTo()"), MakeMemoryView(Payload).EqualBytes(MakeMemoryView(CopyBytes) + 1));
+
+		// GetSerializedView
+		FConstMemoryView SerializedView;
+		TestTrue(TEXT("FCbArray(Clone)::TryGetSerializedView()"), Array.TryGetSerializedView(SerializedView) && SerializedView == Array.AsField().GetView());
+		TestFalse(TEXT("FCbArray(NoType)::TryGetSerializedView()"), Field.AsArray().TryGetSerializedView(SerializedView));
+		TestFalse(TEXT("FCbArray(Name)::TryGetSerializedView()"), NamedField.AsArray().TryGetSerializedView(SerializedView));
 	}
 
 	// Test FCbField(None) as Array
@@ -539,17 +615,25 @@ bool FCbFieldArrayTest::RunTest(const FString& Parameters)
 			FCbFieldRef Field = *It;
 			TestEqual(TEXT("FCbArrayRef::CreateRefIterator().AsInt32()"), Field.AsInt32(), 8);
 			TestTrue(TEXT("FCbArrayRef::CreateRefIterator().IsOwned()"), Field.IsOwned());
+			TestEqual(TEXT("FCbArrayRef::CreateRefIterator().GetBuffer()"), Field.GetBuffer()->GetView(), Field.GetView());
 		}
 		for (FCbFieldRefIterator It = ArrayClone.CreateRefIterator(), End; It != End; ++It)
 		{
 		}
+
+		// CopyTo
+		uint8 CopyBytes[5];
+		Array.CopyTo(MakeMemoryView(CopyBytes));
+		TestTrue(TEXT("FCbArray(Name)::CopyTo()"), ArrayClone.GetView().EqualBytes(MakeMemoryView(CopyBytes)));
+		ArrayClone.CopyTo(MakeMemoryView(CopyBytes));
+		TestTrue(TEXT("FCbArray()::CopyTo()"), ArrayClone.GetView().EqualBytes(MakeMemoryView(CopyBytes)));
 	}
 
 	// Test FCbArray as FCbFieldIterator
 	{
 		uint32 Count = 0;
 		FCbArray Array;
-		for (FCbField Field : FCbFieldIterator(Array.AsField()))
+		for (FCbField Field : FCbFieldIterator::MakeSingle(Array.AsField()))
 		{
 			TestTrue(TEXT("FCbArray::AsField() as Iterator"), Field.IsArray());
 			++Count;
@@ -562,7 +646,7 @@ bool FCbFieldArrayTest::RunTest(const FString& Parameters)
 		uint32 Count = 0;
 		FCbArrayRef Array;
 		Array.MakeOwned();
-		for (FCbFieldRef Field : FCbFieldRefIterator(Array.AsFieldRef()))
+		for (FCbFieldRef Field : FCbFieldRefIterator::MakeSingle(Array.AsFieldRef()))
 		{
 			TestTrue(TEXT("FCbArrayRef::AsField() as Iterator"), Field.IsArray());
 			++Count;
@@ -1285,7 +1369,6 @@ bool FCbFieldRefTest::RunTest(const FString& Parameters)
 
 		FCbField Field(Payload);
 
-		FCbFieldRef VoidTake = FCbFieldRef::TakeOwnership(ImplicitConv<const void*>(PayloadCopy), FMemory::Free);
 		FCbFieldRef VoidView = FCbFieldRef::MakeView(Payload);
 		FCbFieldRef VoidClone = FCbFieldRef::Clone(Payload);
 		FCbFieldRef FieldView = FCbFieldRef::MakeView(Field);
@@ -1294,15 +1377,12 @@ bool FCbFieldRefTest::RunTest(const FString& Parameters)
 
 		Payload[UE_ARRAY_COUNT(Payload) - 1] = 4;
 
-		TestEqualBytes(TEXT("FCbFieldRef::TakeOwnership(Void)"), VoidTake.AsBinary(), {4, 5, 6});
 		TestEqualBytes(TEXT("FCbFieldRef::MakeView(Void)"), VoidView.AsBinary(), {4, 5, 4});
 		TestEqualBytes(TEXT("FCbFieldRef::Clone(Void)"), VoidClone.AsBinary(), {4, 5, 6});
 		TestEqualBytes(TEXT("FCbFieldRef::MakeView(Field)"), FieldView.AsBinary(), {4, 5, 4});
 		TestEqualBytes(TEXT("FCbFieldRef::Clone(Field)"), FieldClone.AsBinary(), {4, 5, 6});
 		TestEqualBytes(TEXT("FCbFieldRef::Clone(FieldRef)"), FieldRefClone.AsBinary(), {4, 5, 6});
 
-		TestTrue(TEXT("FCbFieldRef::TakeOwnership(Void).IsOwned()"), VoidTake.IsOwned());
-		TestTrue(TEXT("FCbFieldRef::TakeOwnership(Void).IsReadOnly()"), VoidTake.IsReadOnly());
 		TestFalse(TEXT("FCbFieldRef::MakeView(Void).IsOwned()"), VoidView.IsOwned());
 		TestFalse(TEXT("FCbFieldRef::MakeView(Void).IsReadOnly()"), VoidView.IsReadOnly());
 		TestTrue(TEXT("FCbFieldRef::Clone(Void).IsOwned()"), VoidClone.IsOwned());
@@ -1412,29 +1492,11 @@ bool FCbObjectRefTest::RunTest(const FString& Parameters)
 IMPLEMENT_SIMPLE_AUTOMATION_TEST(FCbFieldRefIteratorTest, "System.Core.Serialization.CbFieldRefIterator", CompactBinaryTestFlags)
 bool FCbFieldRefIteratorTest::RunTest(const FString& Parameters)
 {
-	static_assert(std::is_constructible<FCbFieldIterator, const FCbField&>::value, "Missing constructor for FCbFieldIterator");
-	static_assert(std::is_constructible<FCbFieldIterator, const FCbFieldRef&>::value, "Missing constructor for FCbFieldIterator");
 	static_assert(std::is_constructible<FCbFieldIterator, const FCbFieldRefIterator&>::value, "Missing constructor for FCbFieldIterator");
-	static_assert(std::is_constructible<FCbFieldIterator, FCbField&&>::value, "Missing constructor for FCbFieldIterator");
-	static_assert(std::is_constructible<FCbFieldIterator, FCbFieldRef&&>::value, "Missing constructor for FCbFieldIterator");
 	static_assert(std::is_constructible<FCbFieldIterator, FCbFieldRefIterator&&>::value, "Missing constructor for FCbFieldIterator");
-	static_assert(std::is_constructible<FCbFieldIterator, FCbField&&, const void*>::value, "Missing constructor for FCbFieldIterator");
 
-	static_assert(std::is_constructible<FCbFieldRefIterator, const FCbFieldRef&>::value, "Missing constructor for FCbFieldRefIterator");
 	static_assert(std::is_constructible<FCbFieldRefIterator, const FCbFieldRefIterator&>::value, "Missing constructor for FCbFieldRefIterator");
-	static_assert(std::is_constructible<FCbFieldRefIterator, FCbFieldRef&&>::value, "Missing constructor for FCbFieldRefIterator");
 	static_assert(std::is_constructible<FCbFieldRefIterator, FCbFieldRefIterator&&>::value, "Missing constructor for FCbFieldRefIterator");
-	static_assert(std::is_constructible<FCbFieldRefIterator, FCbFieldRef&&, const void*>::value, "Missing constructor for FCbFieldRefIterator");
-
-	static_assert(std::is_constructible<FCbFieldRefIterator, const FCbFieldIterator&, const FSharedBufferRef&>::value, "Missing constructor for FCbFieldRefIterator");
-	static_assert(std::is_constructible<FCbFieldRefIterator, const FCbFieldIterator&, const FSharedBufferPtr&>::value, "Missing constructor for FCbFieldRefIterator");
-	static_assert(std::is_constructible<FCbFieldRefIterator, const FCbFieldIterator&, const FSharedBufferConstRef&>::value, "Missing constructor for FCbFieldRefIterator");
-	static_assert(std::is_constructible<FCbFieldRefIterator, const FCbFieldIterator&, const FSharedBufferConstPtr&>::value, "Missing constructor for FCbFieldRefIterator");
-
-	static_assert(std::is_constructible<FCbFieldRefIterator, const FCbFieldIterator&, FSharedBufferRef&&>::value, "Missing constructor for FCbFieldRefIterator");
-	static_assert(std::is_constructible<FCbFieldRefIterator, const FCbFieldIterator&, FSharedBufferPtr&&>::value, "Missing constructor for FCbFieldRefIterator");
-	static_assert(std::is_constructible<FCbFieldRefIterator, const FCbFieldIterator&, FSharedBufferConstRef&&>::value, "Missing constructor for FCbFieldRefIterator");
-	static_assert(std::is_constructible<FCbFieldRefIterator, const FCbFieldIterator&, FSharedBufferConstPtr&&>::value, "Missing constructor for FCbFieldRefIterator");
 
 	const auto GetCount = [](auto It) -> uint32
 	{
@@ -1456,68 +1518,107 @@ bool FCbFieldRefIteratorTest::RunTest(const FString& Parameters)
 	{
 		constexpr uint8 T = uint8(ECbFieldType::IntegerPositive);
 		const uint8 Payload[] = { T, 0, T, 1, T, 2, T, 3 };
-		const uint8* const PayloadEnd = Payload + UE_ARRAY_COUNT(Payload);
 
 		const FSharedBufferConstPtr View = FSharedBuffer::MakeView(MakeMemoryView(Payload));
-		const void* const ViewEnd = View->GetView().GetDataEnd();
 		const FSharedBufferConstPtr Clone = FSharedBuffer::Clone(MakeMemoryView(Payload));
-		const void* const CloneEnd = Clone->GetView().GetDataEnd();
 
 		const FConstMemoryView EmptyView;
 		const FSharedBufferConstPtr NullBuffer;
 
-		const FCbFieldIterator FieldIt(FCbField(View->GetData()), ViewEnd);
-		const FCbFieldRefIterator FieldRefIt(FCbFieldRef(View), ViewEnd);
+		const FCbFieldIterator FieldIt = FCbFieldIterator::MakeRange(*View);
+		const FCbFieldRefIterator FieldRefIt = FCbFieldRefIterator::MakeRange(View);
 
 		TestEqual(TEXT("FCbFieldIterator::GetRangeHash()"), FieldIt.GetRangeHash(), FBlake3::HashBuffer(*View));
 		TestEqual(TEXT("FCbFieldRefIterator::GetRangeHash()"), FieldRefIt.GetRangeHash(), FBlake3::HashBuffer(*View));
 
-		TestEqual(TEXT("FCbFieldIterator(EmptyView)"), GetCount(FCbFieldIterator(EmptyView)), 0);
-		TestEqual(TEXT("FCbFieldRefIterator(BufferNullL)"), GetCount(FCbFieldRefIterator(NullBuffer)), 0);
-		TestEqual(TEXT("FCbFieldRefIterator(BufferNullR)"), GetCount(FCbFieldRefIterator(FSharedBufferConstPtr(NullBuffer))), 0);
-
-		TestEqual(TEXT("FCbFieldIterator(View)"), GetCount(FCbFieldIterator(MakeMemoryView(Payload))), 4);
-		TestEqual(TEXT("FCbFieldRefIterator(BufferCloneL)"), GetCount(FCbFieldRefIterator(Clone)), 4);
-		TestEqual(TEXT("FCbFieldRefIterator(BufferCloneR)"), GetCount(FCbFieldRefIterator(FSharedBufferConstPtr(Clone))), 4);
-
-		TestEqual(TEXT("FCbFieldIterator(Field, End)"), GetCount(FCbFieldIterator(FCbField(Payload), PayloadEnd)), 4);
-		TestEqual(TEXT("FCbFieldRefIterator(FieldRef, End)"), GetCount(FCbFieldRefIterator(FCbFieldRef(View), PayloadEnd)), 4);
-
-		TestEqual(TEXT("FCbFieldRefIterator(FieldIt, BufferNullL)"), GetCount(FCbFieldRefIterator(FCbFieldIterator(FCbField(View->GetData()), ViewEnd), NullBuffer)), 4);
-		TestEqual(TEXT("FCbFieldRefIterator(FieldIt, BufferNullR)"), GetCount(FCbFieldRefIterator(FCbFieldIterator(FCbField(View->GetData()), ViewEnd), FSharedBufferConstPtr(NullBuffer))), 4);
-		TestEqual(TEXT("FCbFieldRefIterator(FieldIt, BufferViewL)"), GetCount(FCbFieldRefIterator(FCbFieldIterator(FCbField(View->GetData()), ViewEnd), View)), 4);
-		TestEqual(TEXT("FCbFieldRefIterator(FieldIt, BufferViewR)"), GetCount(FCbFieldRefIterator(FCbFieldIterator(FCbField(View->GetData()), ViewEnd), FSharedBufferConstPtr(View))), 4);
-		TestEqual(TEXT("FCbFieldRefIterator(FieldIt, BufferCloneL)"), GetCount(FCbFieldRefIterator(FCbFieldIterator(FCbField(Clone->GetData()), CloneEnd), Clone)), 4);
-		TestEqual(TEXT("FCbFieldRefIterator(FieldIt, BufferCloneR)"), GetCount(FCbFieldRefIterator(FCbFieldIterator(FCbField(Clone->GetData()), CloneEnd), FSharedBufferConstPtr(Clone))), 4);
+		TestEqual(TEXT("FCbFieldIterator::GetRangeView()"), FieldIt.GetRangeView(), MakeMemoryView(Payload));
+		TestEqual(TEXT("FCbFieldRefIterator::GetRangeView()"), FieldRefIt.GetRangeView(), MakeMemoryView(Payload));
+		FConstMemoryView SerializedView;
+		TestTrue(TEXT("FCbFieldIterator::TryGetSerializedRangeView()"), FieldIt.TryGetSerializedRangeView(SerializedView) && SerializedView == MakeMemoryView(Payload));
+		TestTrue(TEXT("FCbFieldRefIterator::TryGetSerializedRangeView()"), FieldRefIt.TryGetSerializedRangeView(SerializedView) && SerializedView == MakeMemoryView(Payload));
+
+		TestEqual(TEXT("FCbFieldRefIterator::CloneRange(EmptyIt)"), GetCount(FCbFieldRefIterator::CloneRange(FCbFieldIterator())), 0);
+		TestEqual(TEXT("FCbFieldRefIterator::CloneRange(EmptyRefIt)"), GetCount(FCbFieldRefIterator::CloneRange(FCbFieldRefIterator())), 0);
+		const FCbFieldRefIterator FieldItClone = FCbFieldRefIterator::CloneRange(FieldIt);
+		const FCbFieldRefIterator FieldRefItClone = FCbFieldRefIterator::CloneRange(FieldRefIt);
+		TestEqual(TEXT("FCbFieldRefIterator::CloneRange(FieldIt)"), GetCount(FieldItClone), 4);
+		TestEqual(TEXT("FCbFieldRefIterator::CloneRange(FieldRefIt)"), GetCount(FieldRefItClone), 4);
+		TestNotEqual(TEXT("FCbFieldRefIterator::CloneRange(FieldIt).Equals()"), FieldItClone, FieldRefIt);
+		TestNotEqual(TEXT("FCbFieldRefIterator::CloneRange(FieldRefIt).Equals()"), FieldRefItClone, FieldRefIt);
+
+		TestEqual(TEXT("FCbFieldIterator::MakeRange(EmptyView)"), GetCount(FCbFieldIterator::MakeRange(EmptyView)), 0);
+		TestEqual(TEXT("FCbFieldRefIterator::MakeRange(BufferNullL)"), GetCount(FCbFieldRefIterator::MakeRange(NullBuffer)), 0);
+		TestEqual(TEXT("FCbFieldRefIterator::MakeRange(BufferNullR)"), GetCount(FCbFieldRefIterator::MakeRange(FSharedBufferConstPtr(NullBuffer))), 0);
+
+		TestEqual(TEXT("FCbFieldIterator::MakeRange(BufferView)"), GetCount(FCbFieldIterator::MakeRange(MakeMemoryView(Payload))), 4);
+		TestEqual(TEXT("FCbFieldRefIterator::MakeRange(BufferCloneL)"), GetCount(FCbFieldRefIterator::MakeRange(Clone)), 4);
+		TestEqual(TEXT("FCbFieldRefIterator::MakeRange(BufferCloneR)"), GetCount(FCbFieldRefIterator::MakeRange(FSharedBufferConstPtr(Clone))), 4);
+
+		TestEqual(TEXT("FCbFieldRefIterator::MakeRangeView(FieldIt, BufferNullL)"), GetCount(FCbFieldRefIterator::MakeRangeView(FCbFieldIterator::MakeRange(*View), NullBuffer)), 4);
+		TestEqual(TEXT("FCbFieldRefIterator::MakeRangeView(FieldIt, BufferNullR)"), GetCount(FCbFieldRefIterator::MakeRangeView(FCbFieldIterator::MakeRange(*View), FSharedBufferConstPtr(NullBuffer))), 4);
+		TestEqual(TEXT("FCbFieldRefIterator::MakeRangeView(FieldIt, BufferViewL)"), GetCount(FCbFieldRefIterator::MakeRangeView(FCbFieldIterator::MakeRange(*View), View)), 4);
+		TestEqual(TEXT("FCbFieldRefIterator::MakeRangeView(FieldIt, BufferViewR)"), GetCount(FCbFieldRefIterator::MakeRangeView(FCbFieldIterator::MakeRange(*View), FSharedBufferConstPtr(View))), 4);
+		TestEqual(TEXT("FCbFieldRefIterator::MakeRangeView(FieldIt, BufferCloneL)"), GetCount(FCbFieldRefIterator::MakeRangeView(FCbFieldIterator::MakeRange(*Clone), Clone)), 4);
+		TestEqual(TEXT("FCbFieldRefIterator::MakeRangeView(FieldIt, BufferCloneR)"), GetCount(FCbFieldRefIterator::MakeRangeView(FCbFieldIterator::MakeRange(*Clone), FSharedBufferConstPtr(Clone))), 4);
 
 		TestEqual(TEXT("FCbFieldIterator(FieldRefItL)"), GetCount(FCbFieldIterator(FieldRefIt)), 4);
 		TestEqual(TEXT("FCbFieldIterator(FieldRefItR)"), GetCount(FCbFieldIterator(FCbFieldRefIterator(FieldRefIt))), 4);
 
+		// Uniform
+		const uint8 UniformPayload[] = { 0, 1, 2, 3 };
+		const FCbFieldIterator UniformFieldIt = FCbFieldIterator::MakeRange(MakeMemoryView(UniformPayload), ECbFieldType::IntegerPositive);
+		TestEqual(TEXT("FCbFieldIterator::MakeRange(Uniform).GetRangeHash()"), UniformFieldIt.GetRangeHash(), FieldIt.GetRangeHash());
+		TestEqual(TEXT("FCbFieldIterator::MakeRange(Uniform).GetRangeView()"), UniformFieldIt.GetRangeView(), MakeMemoryView(UniformPayload));
+		TestFalse(TEXT("FCbFieldIterator::MakeRange(Uniform).TryGetSerializedRangeView()"), UniformFieldIt.TryGetSerializedRangeView(SerializedView));
+		const FSharedBufferConstPtr UniformView = FSharedBuffer::MakeView(MakeMemoryView(UniformPayload));
+		const FCbFieldRefIterator UniformFieldRefIt = FCbFieldRefIterator::MakeRange(UniformView, ECbFieldType::IntegerPositive);
+		TestEqual(TEXT("FCbFieldRefIterator::MakeRange(Uniform).GetRangeHash()"), UniformFieldRefIt.GetRangeHash(), FieldIt.GetRangeHash());
+		TestEqual(TEXT("FCbFieldRefIterator::MakeRange(Uniform).GetRangeView()"), UniformFieldRefIt.GetRangeView(), MakeMemoryView(UniformPayload));
+		TestFalse(TEXT("FCbFieldRefIterator::MakeRange(Uniform).TryGetSerializedRangeView()"), UniformFieldRefIt.TryGetSerializedRangeView(SerializedView));
+
+		// CopyRangeTo
+		uint8 CopyBytes[sizeof(Payload)];
+		FieldIt.CopyRangeTo(MakeMemoryView(CopyBytes));
+		TestTrue(TEXT("FCbFieldIterator::MakeRange().CopyRangeTo()"), MakeMemoryView(CopyBytes).EqualBytes(MakeMemoryView(Payload)));
+		FieldRefIt.CopyRangeTo(MakeMemoryView(CopyBytes));
+		TestTrue(TEXT("FCbFieldRefIterator::MakeRange().CopyRangeTo()"), MakeMemoryView(CopyBytes).EqualBytes(MakeMemoryView(Payload)));
+		UniformFieldIt.CopyRangeTo(MakeMemoryView(CopyBytes));
+		TestTrue(TEXT("FCbFieldIterator::MakeRange(Uniform).CopyRangeTo()"), MakeMemoryView(CopyBytes).EqualBytes(MakeMemoryView(Payload)));
+
+		// MakeRangeOwned
+		FCbFieldRefIterator OwnedFromView = UniformFieldRefIt;
+		OwnedFromView.MakeRangeOwned();
+		TestTrue(TEXT("FCbFieldRefIterator::MakeRangeOwned(View)"), OwnedFromView.GetRangeView().EqualBytes(MakeMemoryView(Payload)));
+		FCbFieldRefIterator OwnedFromOwned = OwnedFromView;
+		OwnedFromOwned.MakeRangeOwned();
+		TestEqual(TEXT("FCbFieldRefIterator::MakeRangeOwned(Owned)"), OwnedFromOwned, OwnedFromView);
+
+		// GetRangeBuffer
+		TestEqual(TEXT("FCbFieldRefIterator::GetRangeBuffer()"), UniformFieldRefIt.GetRangeBuffer(), UniformView);
+		const FCbFieldRefIterator UniformShortFieldRefIt = FCbFieldRefIterator::MakeRangeView(FCbFieldIterator::MakeRange(MakeMemoryView(UniformPayload).LeftChop(1), ECbFieldType::IntegerPositive), UniformView);
+		TestNotEqual(TEXT("FCbFieldRefIterator::GetRangeBuffer(Short)"), UniformShortFieldRefIt.GetRangeBuffer(), UniformView);
+
 		// These lines are expected to assert when uncommented.
 		//const FSharedBufferConstPtr ShortView = FSharedBuffer::MakeView(MakeMemoryView(Payload).LeftChop(2));
-		//TestEqual(TEXT("FCbFieldRefIterator(FieldIt, InvalidBufferL)"), GetCount(FCbFieldRefIterator(FCbFieldIterator(FCbField(View->GetData()), ViewEnd), ShortView)), 4);
-		//TestEqual(TEXT("FCbFieldRefIterator(FieldIt, InvalidBufferR)"), GetCount(FCbFieldRefIterator(FCbFieldIterator(FCbField(View->GetData()), ViewEnd), FSharedBufferConstPtr(ShortView))), 4);
+		//TestEqual(TEXT("FCbFieldRefIterator::MakeRangeView(FieldIt, InvalidBufferL)"), GetCount(FCbFieldRefIterator::MakeRangeView(FCbFieldIterator::MakeRange(*View), ShortView)), 4);
+		//TestEqual(TEXT("FCbFieldRefIterator::MakeRangeView(FieldIt, InvalidBufferR)"), GetCount(FCbFieldRefIterator::MakeRangeView(FCbFieldIterator::MakeRange(*View), FSharedBufferConstPtr(ShortView))), 4);
 	}
 
 	// Test FCbField[Ref]Iterator(Scalar)
 	{
 		constexpr uint8 T = uint8(ECbFieldType::IntegerPositive);
 		const uint8 Payload[] = { T, 0 };
-		const uint8* const PayloadEnd = Payload + UE_ARRAY_COUNT(Payload);
 
 		const FSharedBufferConstPtr View = FSharedBuffer::MakeView(MakeMemoryView(Payload));
-		const void* const ViewEnd = View->GetView().GetDataEnd();
 		const FSharedBufferConstPtr Clone = FSharedBuffer::Clone(MakeMemoryView(Payload));
-		const void* const CloneEnd = Clone->GetView().GetDataEnd();
-		const FSharedBufferConstPtr NullBuffer;
 
 		const FCbField Field(Payload);
 		const FCbFieldRef FieldRef(View);
 
-		TestEqual(TEXT("FCbFieldIterator(FieldL)"), GetCount(FCbFieldIterator(Field)), 1);
-		TestEqual(TEXT("FCbFieldIterator(FieldR)"), GetCount(FCbFieldIterator(FCbField(Field))), 1);
-		TestEqual(TEXT("FCbFieldRefIterator(FieldRefL)"), GetCount(FCbFieldRefIterator(FieldRef)), 1);
-		TestEqual(TEXT("FCbFieldRefIterator(FieldRefR)"), GetCount(FCbFieldRefIterator(FCbFieldRef(FieldRef))), 1);
+		TestEqual(TEXT("FCbFieldIterator::MakeSingle(FieldL)"), GetCount(FCbFieldIterator::MakeSingle(Field)), 1);
+		TestEqual(TEXT("FCbFieldIterator::MakeSingle(FieldR)"), GetCount(FCbFieldIterator::MakeSingle(FCbField(Field))), 1);
+		TestEqual(TEXT("FCbFieldRefIterator::MakeSingle(FieldRefL)"), GetCount(FCbFieldRefIterator::MakeSingle(FieldRef)), 1);
+		TestEqual(TEXT("FCbFieldRefIterator::MakeSingle(FieldRefR)"), GetCount(FCbFieldRefIterator::MakeSingle(FCbFieldRef(FieldRef))), 1);
 	}
 
 	return true;
diff --git a/Engine/Source/Runtime/Core/Public/Serialization/CompactBinary.h b/Engine/Source/Runtime/Core/Public/Serialization/CompactBinary.h
--- a/Engine/Source/Runtime/Core/Public/Serialization/CompactBinary.h
+++ b/Engine/Source/Runtime/Core/Public/Serialization/CompactBinary.h
@@ -11,8 +11,6 @@
 #include "Templates/IsTriviallyDestructible.h"
 
 /**
- * \file
- *
  * This file declares a compact binary data format that is compatible with JSON and only slightly
  * more expressive. The format is designed to achieve fairly small encoded sizes while also being
  * efficient to read both sequentially and through random access. An atom of data in this compact
@@ -25,23 +23,19 @@
  * be cast to an object or array. This attribute means that a blob containing compact binary data
  * is always safe to interpret as a field, which allows for easy validation as described later.
  *
- * A field can be constructed as a view of the underlying memory with \ref FCbField or can either
- * view (wrap) or own (clone or assume ownership) the underlying memory with \ref FCbFieldRef. An
- * array provides the same functionality through \ref FCbArray or \ref FCbArrayRef, and an object
- * uses \ref FCbObject or \ref FCbObjectRef.
+ * A field can be constructed as a view of the underlying memory with FCbField or can either view
+ * or own (clone or take ownership) the underlying memory with FCbFieldRef. An array provides the
+ * same behavior with FCbArray or FCbArrayRef, and an object uses FCbObject or FCbObjectRef.
  *
  * It is optimal use the view types when possible, and reference types only when they are needed,
  * to avoid the overhead of the atomic reference counting of the shared buffer.
  *
- * A range of data validation functionality is provided through \ref ValidateCompactBinary. Using
- * the Default mode of validation is sufficient to be able to read the data without crashing. The
- * additional validation modes are required for compatibility with other formats (Names), correct
- * byte-wise comparisons of the encoded data (Format), or storage without arbitrary padding bytes
- * leading to non-deterministic values being stored (Padding).
+ * A suite of validation functionality is provided by ValidateCompactBinary and its siblings. The
+ * Default mode provides a guarantee that the data can be consumed without a crash. Documentation
+ * of the other modes is available on ECbValidateMode.
  *
  * Example:
  *
- * \code
  * void BeginBuild(FCbObjectRef Params)
  * {
  *     if (FSharedBufferConstPtr Data = Storage().Load(Params["Data"].AsBinaryReference()))
@@ -64,7 +58,6 @@
  *         BeginCompress(FName(Format.AsString()));
  *     }
  * }
- * \endcode
  */
 
 class FArchive;
@@ -79,7 +72,7 @@ template <typename FuncType> class TFunctionRef;
 ///////////////////////////////////////////////////////////////////////////////////////////////////////////////////////
 
 /**
- * Field types and flags for \ref FCbField.
+ * Field types and flags for FCbField.
  *
  * This is a private type and is only declared here to enable inline use below.
  *
@@ -198,7 +191,7 @@ ENUM_CLASS_FLAGS(ECbFieldType);
 
 ///////////////////////////////////////////////////////////////////////////////////////////////////////////////////////
 
-/** Functions that operate on \ref ECbFieldType. */
+/** Functions that operate on ECbFieldType. */
 class FCbFieldType
 {
 	static constexpr ECbFieldType SerializedTypeMask    = ECbFieldType(0b1011'1111);
@@ -273,52 +266,20 @@ using FCbFieldVisitor = TFunctionRef<void (FCbField)>;
 ///////////////////////////////////////////////////////////////////////////////////////////////////////////////////////
 
 /**
- * Iterator for \ref FCbField[Ref] that can operate on any contiguous range of fields.
+ * Iterator for FCbField[Ref] that can operate on any contiguous range of fields.
  *
  * The iterator *is* the current field that the iterator points to and exposes the full interface
- * of \ref FCbField[Ref]. An iterator that is at the end is equivalent to a field with no value.
+ * of FCbField[Ref]. An iterator that is at the end is equivalent to a field with no value.
+ *
+ * The iterator represents a range of fields from the current field to the last field.
  */
 template <typename FieldType>
 class TCbFieldIterator : public FieldType
 {
 public:
-	/** Construct a field iterator that has no fields to iterate. */
+	/** Construct an empty field range. */
 	constexpr TCbFieldIterator() = default;
 
-	/**
-	 * Construct a field iterator.
-	 *
-	 * @param InField The first field, or the default field if there are no fields.
-	 * @param InFieldsEnd A pointer to the end of the payload of the last field, or null.
-	 */
-	constexpr inline TCbFieldIterator(FieldType&& InField, const void* InFieldsEnd)
-		: FieldType(MoveTemp(InField))
-		, FieldsEnd(InFieldsEnd)
-	{
-	}
-
-	/**
-	 * Construct a field iterator for one field.
-	 *
-	 * @param InField The single field to be iterated.
-	 */
-	constexpr inline explicit TCbFieldIterator(const FieldType& InField)
-		: FieldType(InField)
-		, FieldsEnd(FieldType::GetPayloadEnd())
-	{
-	}
-
-	/**
-	 * Construct a field iterator for one field.
-	 *
-	 * @param InField The single field to be iterated.
-	 */
-	constexpr inline explicit TCbFieldIterator(FieldType&& InField)
-		: FieldType(MoveTemp(InField))
-		, FieldsEnd(FieldType::GetPayloadEnd())
-	{
-	}
-
 	inline TCbFieldIterator& operator++()
 	{
 		const void* const PayloadEnd = FieldType::GetPayloadEnd();
@@ -334,6 +295,14 @@ public:
 	constexpr inline FieldType& operator*() { return *this; }
 	constexpr inline FieldType* operator->() { return this; }
 
+	/** Returns the size of the fields in the range in bytes. */
+	CORE_API uint64 GetRangeSize() const;
+
+	/** Calculate the hash of every field in the range. */
+	CORE_API FBlake3Hash GetRangeHash() const;
+	/** Calculate the hash of every field in the range. */
+	CORE_API void GetRangeHash(FBlake3& Hash) const;
+
 	using FieldType::Equals;
 
 	template <typename OtherFieldType>
@@ -354,51 +323,62 @@ public:
 		return !Equals(Other);
 	}
 
-	/** Calculate the hash of the field range by hashing their underlying memory. */
-	inline FBlake3Hash GetRangeHash() const { return FBlake3::HashBuffer(GetFieldRangeView(*this)); }
-	/** Calculate the hash of the field range by hashing their underlying memory. */
-	inline void GetRangeHash(FBlake3& Hash) const { Hash.Update(GetFieldRangeView(*this)); }
+	/** Copy the field range into a buffer of exactly GetRangeSize() bytes. */
+	CORE_API void CopyRangeTo(FMutableMemoryView Buffer) const;
+
+	/** Copy the field range into an archive, as if calling CopyTo on every field. */
+	CORE_API void CopyRangeTo(FArchive& Ar) const;
 
 	/** Invoke the visitor for every reference or binary reference in the field range. */
-	inline void IterateRangeReferences(FCbFieldVisitor Visitor) const
+	CORE_API void IterateRangeReferences(FCbFieldVisitor Visitor) const;
+
+	/** Create a view of every field in the range. */
+	inline FConstMemoryView GetRangeView() const
+	{
+		return MakeMemoryView(FieldType::GetView().GetData(), FieldsEnd);
+	}
+
+	/**
+	 * Try to get a view of every field in the range as they would be serialized.
+	 *
+	 * A serialized view is not available if the underlying fields have an externally-provided type.
+	 * Access the serialized form of such ranges using FCbFieldRefIterator::CloneRange.
+	 */
+	inline bool TryGetSerializedRangeView(FConstMemoryView& OutView) const
 	{
-		// Always iterate over non-uniform ranges because we do not know if they contain a reference.
 		if (FCbFieldType::HasFieldType(FieldType::GetType()))
 		{
-			for (TCbFieldIterator It(*this); It; ++It)
-			{
-				if (FCbFieldType::MayContainReferences(It.GetType()))
-				{
-					It.IterateReferences(Visitor);
-				}
-			}
-		}
-		// Only iterate over uniform ranges if the uniform type may contain a reference.
-		else
-		{
-			if (FCbFieldType::MayContainReferences(FieldType::GetType()))
-			{
-				for (TCbFieldIterator It(*this); It; ++It)
-				{
-					It.IterateReferences(Visitor);
-				}
-			}
+			OutView = GetRangeView();
+			return true;
 		}
+		return false;
 	}
 
 protected:
-	/** Returns the end of the last field, or null for an iterator at the end. */
-	template <typename OtherFieldType>
-	static inline const void* GetFieldsEnd(const TCbFieldIterator<OtherFieldType>& It)
+	/** Construct a field range that contains exactly one field. */
+	constexpr inline explicit TCbFieldIterator(FieldType InField)
+		: FieldType(MoveTemp(InField))
+		, FieldsEnd(FieldType::GetPayloadEnd())
+	{
+	}
+
+	/**
+	 * Construct a field range from the first field and a pointer to the end of the last field.
+	 *
+	 * @param InField The first field, or the default field if there are no fields.
+	 * @param InFieldsEnd A pointer to the end of the payload of the last field, or null.
+	 */
+	constexpr inline TCbFieldIterator(FieldType&& InField, const void* InFieldsEnd)
+		: FieldType(MoveTemp(InField))
+		, FieldsEnd(InFieldsEnd)
 	{
-		return It.FieldsEnd;
 	}
 
-	/** Create a view of the field iterator. */
+	/** Returns the end of the last field, or null for an iterator at the end. */
 	template <typename OtherFieldType>
-	static inline FConstMemoryView GetFieldRangeView(const TCbFieldIterator<OtherFieldType>& It)
+	static inline const void* GetFieldsEnd(const TCbFieldIterator<OtherFieldType>& It)
 	{
-		return MakeMemoryView(It.OtherFieldType::GetFieldView().GetData(), It.FieldsEnd);
+		return It.FieldsEnd;
 	}
 
 private:
@@ -446,7 +426,7 @@ enum class ECbFieldError : uint8
  * uniform array of fields with no payload, where the answer is to encode as a non-uniform array.
  *
  * This type only provides a view into memory and does not perform any memory management itself.
- * Use \ref FCbFieldRef to hold a reference to the underlying memory when necessary.
+ * Use FCbFieldRef to hold a reference to the underlying memory when necessary.
  */
 class FCbField
 {
@@ -462,34 +442,12 @@ public:
 	 */
 	CORE_API explicit FCbField(const void* Data, ECbFieldType Type = ECbFieldType::HasFieldType);
 
-	/** Name of the field if it has one, otherwise empty. */
+	/** Returns the name of the field if it has a name, otherwise an empty view. */
 	constexpr inline FAnsiStringView GetName() const
 	{
 		return FAnsiStringView(static_cast<const ANSICHAR*>(Payload) - NameLen, NameLen);
 	}
 
-	/** Size of the field in bytes. */
-	CORE_API uint64 GetSize() const;
-
-	/**
-	 * Whether this field is identical to the other field.
-	 *
-	 * Performs a deep comparison of any contained arrays or objects and their fields. Comparison
-	 * assumes that both fields are valid and are written in the canonical format. Fields must be
-	 * written in the same order in arrays and objects, and name comparison is case sensitive. If
-	 * these assumptions do not hold, this may return false for equivalent inputs. Validation can
-	 * be performed with \ref ValidateCompactBinary, except for field order and field name case.
-	 */
-	CORE_API bool Equals(const FCbField& Other) const;
-
-	/** Calculate the hash of the field by hashing its underlying memory. */
-	inline FBlake3Hash GetHash() const { return FBlake3::HashBuffer(GetFieldView()); }
-	/** Calculate the hash of the field by hashing its underlying memory. */
-	inline void GetHash(FBlake3& Hash) const { Hash.Update(GetFieldView()); }
-
-	/** Invoke the visitor for every reference or binary reference in the field. */
-	CORE_API void IterateReferences(FCbFieldVisitor Visitor) const;
-
 	/** Access the field as an object. Defaults to an empty object on error. */
 	CORE_API FCbObject AsObject();
 
@@ -583,7 +541,7 @@ public:
 	constexpr inline bool IsDateTime() const        { return FCbFieldType::IsDateTime(Type); }
 	constexpr inline bool IsTimeSpan() const        { return FCbFieldType::IsTimeSpan(Type); }
 
-	/** Whether the field has a value. \see \ref FCbField::HasValue() */
+	/** Whether the field has a value. */
 	constexpr inline explicit operator bool() const { return HasValue(); }
 
 	/**
@@ -596,18 +554,66 @@ public:
 
 	/** Whether the last field access encountered an error. */
 	constexpr inline bool HasError() const          { return Error != ECbFieldError::None; }
+
 	/** The type of error that occurred on the last field access, or None. */
 	constexpr inline ECbFieldError GetError() const { return Error; }
 
+	/** Returns the size of the field in bytes, including the type and name. */
+	CORE_API uint64 GetSize() const;
+
+	/** Calculate the hash of the field, including the type and name. */
+	CORE_API FBlake3Hash GetHash() const;
+	/** Calculate the hash of the field, including the type and name. */
+	CORE_API void GetHash(FBlake3& Hash) const;
+
+	/**
+	 * Whether this field is identical to the other field.
+	 *
+	 * Performs a deep comparison of any contained arrays or objects and their fields. Comparison
+	 * assumes that both fields are valid and are written in the canonical format. Fields must be
+	 * written in the same order in arrays and objects, and name comparison is case sensitive. If
+	 * these assumptions do not hold, this may return false for equivalent inputs. Validation can
+	 * be performed with ValidateCompactBinary, except for field order and field name case.
+	 */
+	CORE_API bool Equals(const FCbField& Other) const;
+
+	/** Copy the field into a buffer of exactly GetSize() bytes, including the type and name. */
+	CORE_API void CopyTo(FMutableMemoryView Buffer) const;
+
+	/** Copy the field into an archive, including its type and name. */
+	CORE_API void CopyTo(FArchive& Ar) const;
+
+	/** Invoke the visitor for every reference or binary reference in the field. */
+	CORE_API void IterateReferences(FCbFieldVisitor Visitor) const;
+
+	/** Returns a view of the field, including the type and name when present. */
+	CORE_API FConstMemoryView GetView() const;
+
+	/**
+	 * Try to get a view of the field as it would be serialized, such as by CopyTo.
+	 *
+	 * A serialized view is not available if the field has an externally-provided type.
+	 * Access the serialized form of such fields using CopyTo or FCbFieldRef::Clone.
+	 */
+	inline bool TryGetSerializedView(FConstMemoryView& OutView) const
+	{
+		if (FCbFieldType::HasFieldType(Type))
+		{
+			OutView = GetView();
+			return true;
+		}
+		return false;
+	}
+
 protected:
-	/** Returns the type of the field including any relevant flags. */
-	constexpr inline ECbFieldType GetType() const { return Type; }
+	/** Returns a view of the name and value payload, which excludes the type. */
+	CORE_API FConstMemoryView GetViewNoType() const;
 
-	/** Returns the type to use when constructing a new field after using CopyTo. */
-	constexpr inline ECbFieldType GetCopyType() const { return Type; }
+	/** Returns a view of the value payload, which excludes the type and name. */
+	inline FConstMemoryView GetPayloadView() const { return MakeMemoryView(Payload, GetPayloadSize()); }
 
-	/** Size of the field payload in bytes. That is the field excluding its type and name. */
-	CORE_API uint64 GetPayloadSize() const;
+	/** Returns the type of the field including flags. */
+	constexpr inline ECbFieldType GetType() const { return Type; }
 
 	/** Returns the start of the value payload. */
 	constexpr inline const void* GetPayload() const { return Payload; }
@@ -615,8 +621,8 @@ protected:
 	/** Returns the end of the value payload. */
 	inline const void* GetPayloadEnd() const { return static_cast<const uint8*>(Payload) + GetPayloadSize(); }
 
-	/** Returns a view of the value payload, which excludes the type and name. */
-	inline FConstMemoryView GetPayloadView() const { return MakeMemoryView(Payload, GetPayloadSize()); }
+	/** Returns the size of the value payload in bytes, which is the field excluding the type and name. */
+	CORE_API uint64 GetPayloadSize() const;
 
 	/** Assign a field from a pointer to its data and an optional externally-provided type. */
 	inline void Assign(const void* InData, const ECbFieldType InType)
@@ -626,23 +632,6 @@ protected:
 		new(this) FCbField(InData, InType);
 	}
 
-	/**
-	 * Copy the field into a buffer of at least GetSize() bytes.
-	 *
-	 * The field type will only be copied if HasFieldType is set.
-	 */
-	CORE_API void CopyTo(FMutableMemoryView Buffer) const;
-
-	/**
-	 * Copy the field into an archive. This will write GetSize() bytes.
-	 *
-	 * The field type will only be copied if HasFieldType is set.
-	 */
-	CORE_API void CopyTo(FArchive& Ar) const;
-
-	/** Create a view of the field, including the type and name when present. */
-	CORE_API FConstMemoryView GetFieldView() const;
-
 private:
 	/** Parameters for converting to an integer. */
 	struct FIntegerParams
@@ -687,13 +676,33 @@ private:
 	const void* Payload = nullptr;
 };
 
-/** Iterator for \ref FCbField. \see \ref TCbFieldIterator */
+/**
+ * Iterator for FCbField.
+ *
+ * @see TCbFieldIterator
+ */
 class FCbFieldIterator : public TCbFieldIterator<FCbField>
 {
 public:
-	using TCbFieldIterator::TCbFieldIterator;
+	/** Construct a field range that contains exactly one field. */
+	static inline FCbFieldIterator MakeSingle(const FCbField& Field)
+	{
+		return FCbFieldIterator(Field);
+	}
+
+	/**
+	 * Construct a field range from a buffer containing zero or more valid fields.
+	 *
+	 * @param View A buffer containing zero or more valid fields.
+	 * @param Type HasFieldType means that View contains the type. Otherwise, use the given type.
+	 */
+	static inline FCbFieldIterator MakeRange(FConstMemoryView View, ECbFieldType Type = ECbFieldType::HasFieldType)
+	{
+		return !View.IsEmpty() ? FCbFieldIterator(FCbField(View.GetData(), Type), View.GetDataEnd()) : FCbFieldIterator();
+	}
 
-	FCbFieldIterator() = default;
+	/** Construct an empty field range. */
+	constexpr FCbFieldIterator() = default;
 
 	/** Construct an iterator from another iterator. */
 	template <typename OtherFieldType>
@@ -702,57 +711,50 @@ public:
 	{
 	}
 
-	/** Construct an iterator for the fields in the view. */
-	inline explicit FCbFieldIterator(FConstMemoryView View)
-		: FCbFieldIterator(MakeFromView(View))
-	{
-	}
-
 private:
-	/** Construct an iterator for the fields in the view. */
-	static inline FCbFieldIterator MakeFromView(FConstMemoryView View)
-	{
-		return !View.IsEmpty() ? FCbFieldIterator(FCbField(View.GetData()), View.GetDataEnd()) : FCbFieldIterator();
-	}
+	using TCbFieldIterator::TCbFieldIterator;
 };
 
 ///////////////////////////////////////////////////////////////////////////////////////////////////////////////////////
 
 /**
- * Array of \ref FCbField that have no names.
+ * Array of FCbField that have no names.
  *
  * Accessing a field of the array requires iteration. Access by index is not provided because the
  * cost of accessing an item by index scales linearly with the index.
  *
  * This type only provides a view into memory and does not perform any memory management itself.
- * Use \ref FCbArrayRef to hold a reference to the underlying memory when necessary.
+ * Use FCbArrayRef to hold a reference to the underlying memory when necessary.
  */
 class FCbArray : protected FCbField
 {
 public:
+	/** @see FCbField::FCbField */
+	using FCbField::FCbField;
+
 	/** Construct an array with no fields. */
 	CORE_API FCbArray();
 
-	/**
-	 * Construct an array from a pointer to its data and an optional externally-provided type.
-	 *
-	 * @param Data Pointer to the start of the array data.
-	 * @param Type HasFieldType means that Data contains the type. Otherwise, use the given type.
-	 */
-	inline explicit FCbArray(const void* Data, ECbFieldType Type = ECbFieldType::HasFieldType)
-		: FCbField(Data, Type)
-	{
-	}
-
-	/** Size of the array in bytes if serialized by itself with no name. */
-	CORE_API uint64 GetSize() const;
-
-	/** Number of items in the array. */
+	/** Returns the number of items in the array. */
 	CORE_API uint64 Num() const;
 
 	/** Create an iterator for the fields of this array. */
 	CORE_API FCbFieldIterator CreateIterator() const;
 
+	/** Access the array as an array field. */
+	inline FCbField AsField() const { return static_cast<const FCbField&>(*this); }
+
+	/** Construct an array from an array field. No type check is performed! */
+	static inline FCbArray FromField(const FCbField& Field) { return FCbArray(Field); }
+
+	/** Returns the size of the array in bytes if serialized by itself with no name. */
+	CORE_API uint64 GetSize() const;
+
+	/** Calculate the hash of the array if serialized by itself with no name. */
+	CORE_API FBlake3Hash GetHash() const;
+	/** Calculate the hash of the array if serialized by itself with no name. */
+	CORE_API void GetHash(FBlake3& Hash) const;
+
 	/**
 	 * Whether this array is identical to the other array.
 	 *
@@ -760,91 +762,66 @@ public:
 	 * assumes that both fields are valid and are written in the canonical format. Fields must be
 	 * written in the same order in arrays and objects, and name comparison is case sensitive. If
 	 * these assumptions do not hold, this may return false for equivalent inputs. Validation can
-	 * be done with the `All` mode to check these assumptions about the format of the inputs.
+	 * be done with the All mode to check these assumptions about the format of the inputs.
 	 */
 	CORE_API bool Equals(const FCbArray& Other) const;
 
-	/** Access the array as an array field. */
-	inline FCbField AsField() const { return static_cast<const FCbField&>(*this); }
-
-	/** Construct an array from an array field. No type check is performed! */
-	static FCbArray FromField(const FCbField& Field) { return FCbArray(Field); }
+	/** Copy the array into a buffer of exactly GetSize() bytes, with no name. */
+	CORE_API void CopyTo(FMutableMemoryView Buffer) const;
 
-	/** Calculate the hash of the array if serialized by itself with no name. */
-	CORE_API FBlake3Hash GetHash() const;
-	/** Calculate the hash of the array if serialized by itself with no name. */
-	CORE_API void GetHash(FBlake3& Hash) const;
+	/** Copy the array into an archive. This will write GetSize() bytes, with no name. */
+	CORE_API void CopyTo(FArchive& Ar) const;
 
 	/** Invoke the visitor for every reference or binary reference in the array. */
 	inline void IterateReferences(FCbFieldVisitor Visitor) const { CreateIterator().IterateRangeReferences(Visitor); }
 
-protected:
-	/** Returns the type to use when constructing a new array after using CopyTo. */
-	constexpr inline ECbFieldType GetCopyType() const
+	/** Returns a view of the array, including the type and name when present. */
+	using FCbField::GetView;
+
+	/**
+	 * Try to get a view of the array as it would be serialized, such as by CopyTo.
+	 *
+	 * A serialized view is not available if the underlying field has an externally-provided type or
+	 * a name. Access the serialized form of such arrays using CopyTo or FCbArrayRef::Clone.
+	 */
+	inline bool TryGetSerializedView(FConstMemoryView& OutView) const
 	{
-		return FCbFieldType::GetType(GetType()) | ECbFieldType::HasFieldType;
+		return !FCbField::HasName() && FCbField::TryGetSerializedView(OutView);
 	}
 
-	/** Copy the array into a buffer of at least GetSize() bytes. */
-	CORE_API void CopyTo(FMutableMemoryView Buffer) const;
-
-	/** Copy the array into an archive. This will write GetSize() bytes. */
-	CORE_API void CopyTo(FArchive& Ar) const;
-
 private:
 	friend inline FCbFieldIterator begin(const FCbArray& Array) { return Array.CreateIterator(); }
 	friend inline FCbFieldIterator end(const FCbArray&) { return FCbFieldIterator(); }
 
+	/** Construct an array from an array field. No type check is performed! Use via FromField. */
 	inline explicit FCbArray(const FCbField& Field) : FCbField(Field) {}
 };
 
 ///////////////////////////////////////////////////////////////////////////////////////////////////////////////////////
 
 /**
- * Array of \ref FCbField that have unique names.
+ * Array of FCbField that have unique names.
  *
  * Accessing the fields of an object is always a safe operation, even if the requested field does
  * not exist. Fields may be accessed by name or through iteration. When a field is requested that
  * is not found in the object, the field that it returns has no value (evaluates to false) though
- * attempting to access the empty field is also safe, as described by \ref FCbField.
+ * attempting to access the empty field is also safe, as described by FCbField.
  *
  * This type only provides a view into memory and does not perform any memory management itself.
- * Use \ref FCbObjectRef to hold a reference to the underlying memory when necessary.
+ * Use FCbObjectRef to hold a reference to the underlying memory when necessary.
  */
 class FCbObject : protected FCbField
 {
 public:
+	/** @see FCbField::FCbField */
+	using FCbField::FCbField;
+
 	/** Construct an object with no fields. */
 	CORE_API FCbObject();
 
-	/**
-	 * Construct an object from a pointer to its data and an optional externally-provided type.
-	 *
-	 * @param Data Pointer to the start of the object data.
-	 * @param Type HasFieldType means that Data contains the type. Otherwise, use the given type.
-	 */
-	inline explicit FCbObject(const void* Data, ECbFieldType Type = ECbFieldType::HasFieldType)
-		: FCbField(Data, Type)
-	{
-	}
-
-	/** Size of the object in bytes if serialized by itself with no name. */
-	CORE_API uint64 GetSize() const;
-
 	/** Create an iterator for the fields of this object. */
 	CORE_API FCbFieldIterator CreateIterator() const;
 
-	/**
-	 * Whether this object is identical to the other object.
-	 *
-	 * Performs a deep comparison of any contained arrays or objects and their fields. Comparison
-	 * assumes that both fields are valid and are written in the canonical format. Fields must be
-	 * written in the same order in arrays and objects, and name comparison is case sensitive. If
-	 * these assumptions do not hold, this may return false for equivalent inputs. Validation can
-	 * be done with the `All` mode to check these assumptions about the format of the inputs.
-	 */
-	CORE_API bool Equals(const FCbObject& Other) const;
-
 	/**
 	 * Find a field by case-sensitive name comparison.
 	 *
@@ -856,43 +833,65 @@ public:
 	 */
 	CORE_API FCbField Find(FAnsiStringView Name) const;
 
-	/** Find a field by case-insensitive name comparison. \see \ref FCbObject::Find */
+	/** Find a field by case-insensitive name comparison. */
 	CORE_API FCbField FindIgnoreCase(FAnsiStringView Name) const;
 
-	/** Find a field by case-sensitive name comparison. \see \ref FCbObject::Find */
+	/** Find a field by case-sensitive name comparison. */
 	inline FCbField operator[](FAnsiStringView Name) const { return Find(Name); }
 
 	/** Access the object as an object field. */
 	inline FCbField AsField() const { return static_cast<const FCbField&>(*this); }
 
 	/** Construct an object from an object field. No type check is performed! */
-	static FCbObject FromField(const FCbField& Field) { return FCbObject(Field); }
+	static inline FCbObject FromField(const FCbField& Field) { return FCbObject(Field); }
+
+	/** Returns the size of the object in bytes if serialized by itself with no name. */
+	CORE_API uint64 GetSize() const;
 
 	/** Calculate the hash of the object if serialized by itself with no name. */
 	CORE_API FBlake3Hash GetHash() const;
 	/** Calculate the hash of the object if serialized by itself with no name. */
 	CORE_API void GetHash(FBlake3& Hash) const;
 
+	/**
+	 * Whether this object is identical to the other object.
+	 *
+	 * Performs a deep comparison of any contained arrays or objects and their fields. Comparison
+	 * assumes that both fields are valid and are written in the canonical format. Fields must be
+	 * written in the same order in arrays and objects, and name comparison is case sensitive. If
+	 * these assumptions do not hold, this may return false for equivalent inputs. Validation can
+	 * be done with the All mode to check these assumptions about the format of the inputs.
+	 */
+	CORE_API bool Equals(const FCbObject& Other) const;
+
+	/** Copy the object into a buffer of exactly GetSize() bytes, with no name. */
+	CORE_API void CopyTo(FMutableMemoryView Buffer) const;
+
+	/** Copy the object into an archive. This will write GetSize() bytes, with no name. */
+	CORE_API void CopyTo(FArchive& Ar) const;
+
 	/** Invoke the visitor for every reference or binary reference in the object. */
 	inline void IterateReferences(FCbFieldVisitor Visitor) const { CreateIterator().IterateRangeReferences(Visitor); }
 
-protected:
-	/** Returns the type to use when constructing a new object after using CopyTo. */
-	constexpr inline ECbFieldType GetCopyType() const
+	/** Returns a view of the object, including the type and name when present. */
+	using FCbField::GetView;
+
+	/**
+	 * Try to get a view of the object as it would be serialized, such as by CopyTo.
+	 *
+	 * A serialized view is not available if the underlying field has an externally-provided type or
+	 * a name. Access the serialized form of such objects using CopyTo or FCbObjectRef::Clone.
+	 */
+	inline bool TryGetSerializedView(FConstMemoryView& OutView) const
 	{
-		return FCbFieldType::GetType(GetType()) | ECbFieldType::HasFieldType;
+		return !FCbField::HasName() && FCbField::TryGetSerializedView(OutView);
 	}
 
-	/** Copy the object into a buffer of at least GetSize() bytes. */
-	CORE_API void CopyTo(FMutableMemoryView Buffer) const;
-
-	/** Copy the object into an archive. This will write GetSize() bytes. */
-	CORE_API void CopyTo(FArchive& Ar) const;
-
 private:
 	friend inline FCbFieldIterator begin(const FCbObject& Object) { return Object.CreateIterator(); }
 	friend inline FCbFieldIterator end(const FCbObject&) { return FCbFieldIterator(); }
 
+	/** Construct an object from an object field. No type check is performed! Use via FromField. */
 	inline explicit FCbObject(const FCbField& Field) : FCbField(Field) {}
 };
 
@@ -903,13 +902,7 @@ using FCbBufferAllocator = TFunctionRef<FSharedBufferPtr (ECbFieldType Type, uin
 
 ///////////////////////////////////////////////////////////////////////////////////////////////////////////////////////
 
-/**
- * A wrapper that can hold a reference to the memory that contains its value.
- *
- * When constructed with an owned buffer, or by assuming ownership, or by cloning, this reference
- * type will prevent the memory containing its value from being released. To minimize overhead, a
- * buffer pointer is only stored for owned buffers.
- */
+/** A wrapper that holds a reference to the buffer that contains its compact binary value. */
 template <typename BaseType>
 class TCbBufferRef : public BaseType
 {
@@ -917,86 +910,36 @@ public:
 	/** Construct a default value. */
 	TCbBufferRef() = default;
 
-	/** Construct a value and hold a reference to the buffer that contains it. */
-	inline explicit TCbBufferRef(const FSharedBufferConstRef& ValueBuffer)
-		: BaseType(ValueBuffer->GetData())
-	{
-		if (ValueBuffer->IsOwned())
-		{
-			Buffer = ValueBuffer;
-		}
-	}
-
-	/** Construct a value and hold a reference to the buffer that contains it. */
-	inline explicit TCbBufferRef(const FSharedBufferConstPtr& ValueBuffer)
-	{
-		if (ValueBuffer)
-		{
-			static_cast<BaseType&>(*this) = BaseType(ValueBuffer->GetData());
-			if (ValueBuffer->IsOwned())
-			{
-				Buffer = ValueBuffer;
-			}
-		}
-	}
-
-	/** Construct a value and hold a reference to the buffer that contains it. */
-	inline explicit TCbBufferRef(FSharedBufferConstPtr&& ValueBuffer)
-	{
-		if (ValueBuffer)
-		{
-			static_cast<BaseType&>(*this) = BaseType(ValueBuffer->GetData());
-			if (ValueBuffer->IsOwned())
-			{
-				Buffer = MoveTemp(ValueBuffer);
-			}
-		}
-	}
-
-	/** Construct a value and hold a reference to the buffer that contains it. */
-	inline explicit TCbBufferRef(const FSharedBufferRef& ValueBuffer)
-		: TCbBufferRef(FSharedBufferConstRef(ValueBuffer))
-	{
-	}
-
-	/** Construct a value that holds a reference to the buffer that contains it. */
-	inline TCbBufferRef(const BaseType& Value, const FSharedBufferConstPtr& ValueBuffer)
-		: BaseType(Value)
+	/**
+	 * Construct a value from a pointer to its data and an optional externally-provided type.
+	 *
+	 * @param ValueBuffer A buffer that exactly contains the value.
+	 * @param Type HasFieldType means that ValueBuffer contains the type. Otherwise, use the given type.
+	 */
+	inline explicit TCbBufferRef(FSharedBufferConstPtr ValueBuffer, ECbFieldType Type = ECbFieldType::HasFieldType)
 	{
 		if (ValueBuffer)
 		{
-			check(ValueBuffer->GetView().Contains(BaseType::GetFieldView()));
-			if (ValueBuffer->IsOwned())
-			{
-				Buffer = ValueBuffer;
-			}
+			static_cast<BaseType&>(*this) = BaseType(ValueBuffer->GetData(), Type);
+			check(ValueBuffer->GetView().Contains(BaseType::GetView()));
+			Buffer = MoveTemp(ValueBuffer);
 		}
 	}
 
 	/** Construct a value that holds a reference to the buffer that contains it. */
-	inline TCbBufferRef(const BaseType& Value, FSharedBufferConstPtr&& ValueBuffer)
+	inline TCbBufferRef(const BaseType& Value, FSharedBufferConstPtr OuterBuffer)
 		: BaseType(Value)
 	{
-		if (ValueBuffer)
+		if (OuterBuffer)
 		{
-			check(ValueBuffer->GetView().Contains(BaseType::GetFieldView()));
-			if (ValueBuffer->IsOwned())
-			{
-				Buffer = MoveTemp(ValueBuffer);
-			}
+			check(OuterBuffer->GetView().Contains(BaseType::GetView()));
+			Buffer = MoveTemp(OuterBuffer);
 		}
 	}
 
 	/** Construct a value that holds a reference to the buffer of the outer that contains it. */
 	template <typename OtherBaseType>
-	inline TCbBufferRef(const BaseType& Value, const TCbBufferRef<OtherBaseType>& OuterRef)
-		: TCbBufferRef(Value, OuterRef.Buffer)
-	{
-	}
-
-	/** Construct a value that holds a reference to the buffer of the outer that contains it. */
-	template <typename OtherBaseType>
-	inline TCbBufferRef(const BaseType& Value, TCbBufferRef<OtherBaseType>&& OuterRef)
+	inline TCbBufferRef(const BaseType& Value, TCbBufferRef<OtherBaseType> OuterRef)
 		: TCbBufferRef(Value, MoveTemp(OuterRef.Buffer))
 	{
 	}
@@ -1009,7 +952,10 @@ public:
 	{
 		if (!IsOwned())
 		{
-			MakeReadOnly();
+			FSharedBufferPtr MutableBuffer = FSharedBuffer::Alloc(BaseType::GetSize());
+			BaseType::CopyTo(*MutableBuffer);
+			static_cast<BaseType&>(*this) = BaseType(MutableBuffer->GetData());
+			Buffer = FSharedBuffer::MakeReadOnly(MoveTemp(MutableBuffer));
 		}
 	}
 
@@ -1023,13 +969,35 @@ public:
 		{
 			FSharedBufferPtr MutableBuffer = FSharedBuffer::Alloc(BaseType::GetSize());
 			BaseType::CopyTo(*MutableBuffer);
-			static_cast<BaseType&>(*this) = BaseType(MutableBuffer->GetData(), BaseType::GetCopyType());
+			static_cast<BaseType&>(*this) = BaseType(MutableBuffer->GetData());
 			Buffer = FSharedBuffer::MakeReadOnly(MoveTemp(MutableBuffer));
 		}
 	}
 
-	/** The buffer (if any) that contains this value. */
-	inline const FSharedBufferConstPtr& GetBuffer() const { return Buffer; }
+	/** Returns a buffer that exactly contains this value. */
+	inline FSharedBufferConstPtr GetBuffer() const
+	{
+		const FConstMemoryView View = BaseType::GetView();
+		const FSharedBufferConstPtr& OuterBuffer = GetOuterBuffer();
+		if (OuterBuffer)
+		{
+			if (OuterBuffer->GetView() == View)
+			{
+				return OuterBuffer;
+			}
+			else
+			{
+				return FSharedBuffer::MakeView(View, *OuterBuffer);
+			}
+		}
+		else
+		{
+			return FSharedBuffer::MakeView(View);
+		}
+	}
+
+	/** Returns the outer buffer (if any) that contains this value. */
+	inline const FSharedBufferConstPtr& GetOuterBuffer() const { return Buffer; }
 
 private:
 	template <typename OtherType>
@@ -1064,26 +1032,16 @@ public:
 		return Ref;
 	}
 
-	/** Construct a value from a read-only view of its memory. */
-	static inline RefType MakeView(const void* const Data)
-	{
-		return MakeView(BaseType(Data));
-	}
-
-	/** Construct a value from a read-only view of its memory. */
-	static inline RefType MakeView(const BaseType& Value)
+	/** Construct a value from a read-only view of its memory and its optional outer buffer. */
+	static inline RefType MakeView(const void* const Data, FSharedBufferConstPtr OuterBuffer = FSharedBufferConstPtr())
 	{
-		return RefType(Value, FSharedBufferConstPtr());
+		return MakeView(BaseType(Data), MoveTemp(OuterBuffer));
 	}
 
-	/** Construct a value and take ownership its memory. */
-	template <typename DeleteFunctionOrBufferOwnerType>
-	static inline RefType TakeOwnership(const void* const Data, DeleteFunctionOrBufferOwnerType&& DeleteFunctionOrBufferOwner)
+	/** Construct a value from a read-only view of its memory and its optional outer buffer. */
+	static inline RefType MakeView(const BaseType& Value, FSharedBufferConstPtr OuterBuffer = FSharedBufferConstPtr())
 	{
-		const BaseType Value(Data);
-		FSharedBufferConstPtr Buffer = FSharedBuffer::TakeOwnership(Data, Value.GetSize(),
-			Forward<DeleteFunctionOrBufferOwnerType>(DeleteFunctionOrBufferOwner));
-		return RefType(Value, FSharedBuffer::MakeReadOnly(MoveTemp(Buffer)));
+		return RefType(Value, MoveTemp(OuterBuffer));
 	}
 };
 
@@ -1094,7 +1052,11 @@ class FCbObjectRef;
 
 ///////////////////////////////////////////////////////////////////////////////////////////////////////////////////////
 
-/** A field that can hold a reference to the memory that contains it. \see \ref TCbBufferRef */
+/**
+ * A field that can hold a reference to the memory that contains it.
+ *
+ * @see TCbBufferRef
+ */
 class FCbFieldRef : public TCbBufferRef<FCbField>, public TCbBufferRefFactory<FCbFieldRef, FCbField>
 {
 public:
@@ -1113,61 +1075,77 @@ public:
 	inline FCbArrayRef AsArrayRef() &&;
 };
 
-/** Iterator for \ref FCbFieldRef. \see \ref TCbFieldIterator */
+/**
+ * Iterator for FCbFieldRef.
+ *
+ * @see TCbFieldIterator
+ */
 class FCbFieldRefIterator : public TCbFieldIterator<FCbFieldRef>
 {
 public:
-	using TCbFieldIterator::TCbFieldIterator;
+	/** Construct a field range from an owned clone of a range. */
+	CORE_API static FCbFieldRefIterator CloneRange(const FCbFieldIterator& It);
 
-	/** Construct a field iterator that has no fields to iterate. */
-	constexpr FCbFieldRefIterator() = default;
-
-	/** Construct an iterator that holds a reference to the buffer that contains its fields. */
-	inline FCbFieldRefIterator(const FCbFieldIterator& It, const FSharedBufferConstPtr& ValueBuffer)
-		: TCbFieldIterator(MakeFieldRef(It, ValueBuffer), GetFieldsEnd(It))
+	/** Construct a field range from an owned clone of a range. */
+	static inline FCbFieldRefIterator CloneRange(const FCbFieldRefIterator& It)
 	{
+		return CloneRange(FCbFieldIterator(It));
 	}
 
-	/** Construct an iterator that holds a reference to the buffer that contains its fields. */
-	inline FCbFieldRefIterator(const FCbFieldIterator& It, FSharedBufferConstPtr&& ValueBuffer)
-		: TCbFieldIterator(MakeFieldRef(It, MoveTemp(ValueBuffer)), GetFieldsEnd(It))
+	/** Construct a field range that contains exactly one field. */
+	static inline FCbFieldRefIterator MakeSingle(FCbFieldRef Field)
 	{
+		return FCbFieldRefIterator(MoveTemp(Field));
 	}
 
-	/** Construct an iterator and hold a reference to the buffer that contains its fields. */
-	inline explicit FCbFieldRefIterator(const FSharedBufferConstPtr& ValueBuffer)
-		: FCbFieldRefIterator(MakeFieldIterator(ValueBuffer), ValueBuffer)
+	/**
+	 * Construct a field range from a buffer containing zero or more valid fields.
+	 *
+	 * @param Buffer A buffer containing zero or more valid fields.
+	 * @param Type HasFieldType means that Buffer contains the type. Otherwise, use the given type.
+	 */
+	static inline FCbFieldRefIterator MakeRange(FSharedBufferConstPtr Buffer, ECbFieldType Type = ECbFieldType::HasFieldType)
 	{
+		if (Buffer && Buffer->GetSize())
+		{
+			const void* const DataEnd = Buffer->GetView().GetDataEnd();
+			return FCbFieldRefIterator(FCbFieldRef(MoveTemp(Buffer), Type), DataEnd);
+		}
+		return FCbFieldRefIterator();
 	}
 
-	/** Construct an iterator and hold a reference to the buffer that contains its fields. */
-	inline explicit FCbFieldRefIterator(FSharedBufferConstPtr&& ValueBuffer)
-		: FCbFieldRefIterator(MakeFieldIterator(ValueBuffer), MoveTemp(ValueBuffer))
+	/** Construct a field range from an iterator and its optional outer buffer. */
+	static inline FCbFieldRefIterator MakeRangeView(const FCbFieldIterator& It, FSharedBufferConstPtr OuterBuffer = FSharedBufferConstPtr())
 	{
+		return FCbFieldRefIterator(FCbFieldRef(It, MoveTemp(OuterBuffer)), GetFieldsEnd(It));
 	}
 
-private:
-	/** Make a field ref from an iterator and the buffer that contains its fields. */
-	template <typename BufferType>
-	static inline FCbFieldRef MakeFieldRef(const FCbFieldIterator& It, BufferType&& ValueBuffer)
+	/** Construct an empty field range. */
+	constexpr FCbFieldRefIterator() = default;
+
+	/** Clone the range, if necessary, to a buffer that this reference has ownership of. */
+	inline void MakeRangeOwned()
 	{
-		if (ValueBuffer)
+		if (!IsOwned())
 		{
-			check(ValueBuffer->GetView().Contains(GetFieldRangeView(It)));
+			*this = CloneRange(*this);
 		}
-		return FCbFieldRef(It, Forward<BufferType>(ValueBuffer));
 	}
 
-	/** Make a field iterator from the possibly-null buffer that contains its fields. */
-	static inline FCbFieldIterator MakeFieldIterator(const FSharedBufferConstPtr& ValueBuffer)
-	{
-		return ValueBuffer ? FCbFieldIterator(ValueBuffer->GetView()) : FCbFieldIterator();
-	}
+	/** Returns a buffer that exactly contains the field range. */
+	CORE_API FSharedBufferConstPtr GetRangeBuffer() const;
+
+private:
+	using TCbFieldIterator::TCbFieldIterator;
 };
 
 ///////////////////////////////////////////////////////////////////////////////////////////////////////////////////////
 
-/** An array that can hold a reference to the memory that contains it. \see \ref TCbBufferRef */
+/**
+ * An array that can hold a reference to the memory that contains it.
+ *
+ * @see TCbBufferRef
+ */
 class FCbArrayRef : public TCbBufferRef<FCbArray>, public TCbBufferRefFactory<FCbArrayRef, FCbArray>
 {
 public:
@@ -1176,7 +1154,7 @@ public:
 	/** Create an iterator for the fields of this array. */
 	inline FCbFieldRefIterator CreateRefIterator() const
 	{
-		return FCbFieldRefIterator(CreateIterator(), GetBuffer());
+		return FCbFieldRefIterator::MakeRangeView(CreateIterator(), GetOuterBuffer());
 	}
 
 	/** Access the array as an array field. */
@@ -1194,7 +1172,11 @@ public:
 
 ///////////////////////////////////////////////////////////////////////////////////////////////////////////////////////
 
-/** An object that can hold a reference to the memory that contains it. \see \ref TCbBufferRef */
+/**
+ * An object that can hold a reference to the memory that contains it.
+ *
+ * @see TCbBufferRef
+ */
 class FCbObjectRef : public TCbBufferRef<FCbObject>, public TCbBufferRefFactory<FCbObjectRef, FCbObject>
 {
 public:
@@ -1203,7 +1185,7 @@ public:
 	/** Create an iterator for the fields of this object. */
 	inline FCbFieldRefIterator CreateRefIterator() const
 	{
-		return FCbFieldRefIterator(CreateIterator(), GetBuffer());
+		return FCbFieldRefIterator::MakeRangeView(CreateIterator(), GetOuterBuffer());
 	}
 
 	/** Find a field by case-sensitive name comparison. */
