CompactBinary: Rename a parameter to avoid a shadowing warning

#rnx

diff --git a/Engine/Source/Runtime/Core/Private/Serialization/CompactBinary.cpp b/Engine/Source/Runtime/Core/Private/Serialization/CompactBinary.cpp
--- a/Engine/Source/Runtime/Core/Private/Serialization/CompactBinary.cpp
+++ b/Engine/Source/Runtime/Core/Private/Serialization/CompactBinary.cpp
@@ -705,23 +705,23 @@ void TCbFieldIterator<FieldType>::GetRangeHash(FBlake3& Hash) const
 }
 
 template <typename FieldType>
-void TCbFieldIterator<FieldType>::CopyRangeTo(FMutableMemoryView Buffer) const
+void TCbFieldIterator<FieldType>::CopyRangeTo(FMutableMemoryView InBuffer) const
 {
 	FConstMemoryView Source;
 	if (TryGetSerializedRangeView(Source))
 	{
-		checkf(Buffer.GetSize() == Source.GetSize(),
+		checkf(InBuffer.GetSize() == Source.GetSize(),
 			TEXT("Buffer is %" UINT64_FMT " bytes but %" UINT64_FMT " is required."),
-			Buffer.GetSize(), Source.GetSize());
-		FMemory::Memcpy(Buffer.GetData(), Source.GetData(), Source.GetSize());
+			InBuffer.GetSize(), Source.GetSize());
+		FMemory::Memcpy(InBuffer.GetData(), Source.GetData(), Source.GetSize());
 	}
 	else
 	{
 		for (TCbFieldIterator It(*this); It; ++It)
 		{
 			const uint64 Size = It.GetSize();
-			It.CopyTo(Buffer.Left(Size));
-			Buffer += Size;
+			It.CopyTo(InBuffer.Left(Size));
+			InBuffer += Size;
 		}
 	}
 }
