MemoryView: Renamed FConstMemoryView to FMemoryView

#rnx

diff --git a/Engine/Source/Runtime/Core/Private/Serialization/CompactBinary.cpp b/Engine/Source/Runtime/Core/Private/Serialization/CompactBinary.cpp
--- a/Engine/Source/Runtime/Core/Private/Serialization/CompactBinary.cpp
+++ b/Engine/Source/Runtime/Core/Private/Serialization/CompactBinary.cpp
@@ -113,7 +113,7 @@ FCbArray FCbField::AsArray()
 	}
 }
 
-FConstMemoryView FCbField::AsBinary(const FConstMemoryView Default)
+FMemoryView FCbField::AsBinary(const FMemoryView Default)
 {
 	if (FCbFieldType::IsBinary(Type))
 	{
@@ -437,7 +437,7 @@ bool FCbField::Equals(const FCbField& Other) const
 
 void FCbField::CopyTo(FMutableMemoryView Buffer) const
 {
-	const FConstMemoryView Source = GetViewNoType();
+	const FMemoryView Source = GetViewNoType();
 	checkf(Buffer.GetSize() == sizeof(ECbFieldType) + Source.GetSize(),
 		TEXT("A buffer of %" UINT64_FMT " bytes was provided when %" UINT64_FMT " bytes are required"),
 		Buffer.GetSize(), sizeof(ECbFieldType) + Source.GetSize());
@@ -449,7 +449,7 @@ void FCbField::CopyTo(FMutableMemoryView Buffer) const
 void FCbField::CopyTo(FArchive& Ar) const
 {
 	check(Ar.IsSaving());
-	const FConstMemoryView Source = GetViewNoType();
+	const FMemoryView Source = GetViewNoType();
 	ECbFieldType SerializedType = FCbFieldType::GetSerializedType(Type);
 	Ar.Serialize(&SerializedType, sizeof(SerializedType));
 	Ar.Serialize(const_cast<void*>(Source.GetData()), static_cast<int64>(Source.GetSize()));
@@ -473,7 +473,7 @@ void FCbField::IterateReferences(FCbFieldVisitor Visitor) const
 	}
 }
 
-FConstMemoryView FCbField::GetView() const
+FMemoryView FCbField::GetView() const
 {
 	const uint32 TypeSize = FCbFieldType::HasFieldType(Type) ? sizeof(ECbFieldType) : 0;
 	const uint32 NameSize = FCbFieldType::HasFieldName(Type) ? NameLen + MeasureVarUInt(NameLen) : 0;
@@ -481,7 +481,7 @@ FConstMemoryView FCbField::GetView() const
 	return MakeMemoryView(static_cast<const uint8*>(Payload) - TypeSize - NameSize, TypeSize + NameSize + PayloadSize);
 }
 
-FConstMemoryView FCbField::GetViewNoType() const
+FMemoryView FCbField::GetViewNoType() const
 {
 	const uint32 NameSize = FCbFieldType::HasFieldName(Type) ? NameLen + MeasureVarUInt(NameLen) : 0;
 	const uint64 PayloadSize = GetPayloadSize();
@@ -548,7 +548,7 @@ bool FCbArray::Equals(const FCbArray& Other) const
 
 void FCbArray::CopyTo(FMutableMemoryView Buffer) const
 {
-	const FConstMemoryView Source = GetPayloadView();
+	const FMemoryView Source = GetPayloadView();
 	checkf(Buffer.GetSize() == sizeof(ECbFieldType) + Source.GetSize(),
 		TEXT("Buffer is %" UINT64_FMT " bytes but %" UINT64_FMT " is required."),
 		Buffer.GetSize(), sizeof(ECbFieldType) + Source.GetSize());
@@ -560,7 +560,7 @@ void FCbArray::CopyTo(FMutableMemoryView Buffer) const
 void FCbArray::CopyTo(FArchive& Ar) const
 {
 	check(Ar.IsSaving());
-	const FConstMemoryView Source = GetPayloadView();
+	const FMemoryView Source = GetPayloadView();
 	ECbFieldType SerializedType = FCbFieldType::GetType(GetType());
 	Ar.Serialize(&SerializedType, sizeof(SerializedType));
 	Ar.Serialize(const_cast<void*>(Source.GetData()), static_cast<int64>(Source.GetSize()));
@@ -640,7 +640,7 @@ bool FCbObject::Equals(const FCbObject& Other) const
 
 void FCbObject::CopyTo(FMutableMemoryView Buffer) const
 {
-	const FConstMemoryView Source = GetPayloadView();
+	const FMemoryView Source = GetPayloadView();
 	checkf(Buffer.GetSize() == sizeof(ECbFieldType) + Source.GetSize(),
 		TEXT("Buffer is %" UINT64_FMT " bytes but %" UINT64_FMT " is required."),
 		Buffer.GetSize(), sizeof(ECbFieldType) + Source.GetSize());
@@ -652,7 +652,7 @@ void FCbObject::CopyTo(FMutableMemoryView Buffer) const
 void FCbObject::CopyTo(FArchive& Ar) const
 {
 	check(Ar.IsSaving());
-	const FConstMemoryView Source = GetPayloadView();
+	const FMemoryView Source = GetPayloadView();
 	ECbFieldType SerializedType = FCbFieldType::GetType(GetType());
 	Ar.Serialize(&SerializedType, sizeof(SerializedType));
 	Ar.Serialize(const_cast<void*>(Source.GetData()), static_cast<int64>(Source.GetSize()));
@@ -663,7 +663,7 @@ void FCbObject::CopyTo(FArchive& Ar) const
 template <typename FieldType>
 uint64 TCbFieldIterator<FieldType>::GetRangeSize() const
 {
-	FConstMemoryView View;
+	FMemoryView View;
 	if (TryGetSerializedRangeView(View))
 	{
 		return View.GetSize();
@@ -690,7 +690,7 @@ FBlake3Hash TCbFieldIterator<FieldType>::GetRangeHash() const
 template <typename FieldType>
 void TCbFieldIterator<FieldType>::GetRangeHash(FBlake3& Hash) const
 {
-	FConstMemoryView View;
+	FMemoryView View;
 	if (TryGetSerializedRangeView(View))
 	{
 		Hash.Update(View);
@@ -707,7 +707,7 @@ void TCbFieldIterator<FieldType>::GetRangeHash(FBlake3& Hash) const
 template <typename FieldType>
 void TCbFieldIterator<FieldType>::CopyRangeTo(FMutableMemoryView InBuffer) const
 {
-	FConstMemoryView Source;
+	FMemoryView Source;
 	if (TryGetSerializedRangeView(Source))
 	{
 		checkf(InBuffer.GetSize() == Source.GetSize(),
@@ -730,7 +730,7 @@ template <typename FieldType>
 void TCbFieldIterator<FieldType>::CopyRangeTo(FArchive& Ar) const
 {
 	check(Ar.IsSaving());
-	FConstMemoryView Source;
+	FMemoryView Source;
 	if (TryGetSerializedRangeView(Source))
 	{
 		Ar.Serialize(const_cast<void*>(Source.GetData()), static_cast<int64>(Source.GetSize()));
@@ -779,7 +779,7 @@ template class TCbFieldIterator<FCbFieldRef>;
 FCbFieldRefIterator FCbFieldRefIterator::CloneRange(const FCbFieldIterator& It)
 {
 	FSharedBufferPtr Buffer;
-	FConstMemoryView View;
+	FMemoryView View;
 	if (It.TryGetSerializedRangeView(View))
 	{
 		Buffer = FSharedBuffer::Clone(View);
@@ -794,7 +794,7 @@ FCbFieldRefIterator FCbFieldRefIterator::CloneRange(const FCbFieldIterator& It)
 
 FSharedBufferConstPtr FCbFieldRefIterator::GetRangeBuffer() const
 {
-	const FConstMemoryView RangeView = GetRangeView();
+	const FMemoryView RangeView = GetRangeView();
 	const FSharedBufferConstPtr& OuterBuffer = GetOuterBuffer();
 	if (OuterBuffer)
 	{
diff --git a/Engine/Source/Runtime/Core/Private/Serialization/CompactBinarySerialization.cpp b/Engine/Source/Runtime/Core/Private/Serialization/CompactBinarySerialization.cpp
--- a/Engine/Source/Runtime/Core/Private/Serialization/CompactBinarySerialization.cpp
+++ b/Engine/Source/Runtime/Core/Private/Serialization/CompactBinarySerialization.cpp
@@ -8,7 +8,7 @@
 
 ///////////////////////////////////////////////////////////////////////////////////////////////////////////////////////
 
-uint64 MeasureCompactBinary(FConstMemoryView View, ECbFieldType Type)
+uint64 MeasureCompactBinary(FMemoryView View, ECbFieldType Type)
 {
 	uint64 Size;
 	return TryMeasureCompactBinary(View, Type, Size, Type) ? Size : 0;
@@ -16,7 +16,7 @@ uint64 MeasureCompactBinary(FConstMemoryView View, ECbFieldType Type)
 
 ///////////////////////////////////////////////////////////////////////////////////////////////////////////////////////
 
-bool TryMeasureCompactBinary(FConstMemoryView View, ECbFieldType& OutType, uint64& OutSize, ECbFieldType Type)
+bool TryMeasureCompactBinary(FMemoryView View, ECbFieldType& OutType, uint64& OutSize, ECbFieldType Type)
 {
 	uint64 Size = 0;
 
diff --git a/Engine/Source/Runtime/Core/Private/Serialization/CompactBinaryValidation.cpp b/Engine/Source/Runtime/Core/Private/Serialization/CompactBinaryValidation.cpp
--- a/Engine/Source/Runtime/Core/Private/Serialization/CompactBinaryValidation.cpp
+++ b/Engine/Source/Runtime/Core/Private/Serialization/CompactBinaryValidation.cpp
@@ -42,7 +42,7 @@ FORCENOINLINE static void AddError(ECbValidateError& OutError, const ECbValidate
  *
  * A type argument with the HasFieldType flag indicates that the type will not be read from the view.
  */
-static ECbFieldType ValidateCbFieldType(FConstMemoryView& View, ECbValidateMode Mode, ECbValidateError& Error, ECbFieldType Type = ECbFieldType::HasFieldType)
+static ECbFieldType ValidateCbFieldType(FMemoryView& View, ECbValidateMode Mode, ECbValidateError& Error, ECbFieldType Type = ECbFieldType::HasFieldType)
 {
 	if (FCbFieldType::HasFieldType(Type))
 	{
@@ -77,7 +77,7 @@ static ECbFieldType ValidateCbFieldType(FConstMemoryView& View, ECbValidateMode
  *
  * Modifies the view to start at the end of the value, and adds error flags if applicable.
  */
-static uint64 ValidateCbUInt(FConstMemoryView& View, ECbValidateMode Mode, ECbValidateError& Error)
+static uint64 ValidateCbUInt(FMemoryView& View, ECbValidateMode Mode, ECbValidateError& Error)
 {
 	if (View.GetSize() > 0 && View.GetSize() >= MeasureVarUInt(View.GetData()))
 	{
@@ -103,7 +103,7 @@ static uint64 ValidateCbUInt(FConstMemoryView& View, ECbValidateMode Mode, ECbVa
  *
  * Modifies the view to start at the end of the value, and adds error flags if applicable.
  */
-static void ValidateCbFloat64(FConstMemoryView& View, ECbValidateMode Mode, ECbValidateError& Error)
+static void ValidateCbFloat64(FMemoryView& View, ECbValidateMode Mode, ECbValidateError& Error)
 {
 	if (View.GetSize() >= sizeof(double))
 	{
@@ -130,7 +130,7 @@ static void ValidateCbFloat64(FConstMemoryView& View, ECbValidateMode Mode, ECbV
  *
  * Modifies the view to start at the end of the string, and adds error flags if applicable.
  */
-static FAnsiStringView ValidateCbString(FConstMemoryView& View, ECbValidateMode Mode, ECbValidateError& Error)
+static FAnsiStringView ValidateCbString(FMemoryView& View, ECbValidateMode Mode, ECbValidateError& Error)
 {
 	const uint64 NameSize = ValidateCbUInt(View, Mode, Error);
 	if (View.GetSize() >= NameSize)
@@ -147,7 +147,7 @@ static FAnsiStringView ValidateCbString(FConstMemoryView& View, ECbValidateMode
 	}
 }
 
-static FCbField ValidateCbField(FConstMemoryView& View, ECbValidateMode Mode, ECbValidateError& Error, ECbFieldType ExternalType);
+static FCbField ValidateCbField(FMemoryView& View, ECbValidateMode Mode, ECbValidateError& Error, ECbFieldType ExternalType);
 
 /** A type that checks whether all validated fields are of the same type. */
 class FCbUniformFieldsValidator
@@ -158,7 +158,7 @@ public:
 	{
 	}
 
-	inline FCbField ValidateField(FConstMemoryView& View, ECbValidateMode Mode, ECbValidateError& Error)
+	inline FCbField ValidateField(FMemoryView& View, ECbValidateMode Mode, ECbValidateError& Error)
 	{
 		const void* const FieldData = View.GetData();
 		if (FCbField Field = ValidateCbField(View, Mode, Error, ExternalType))
@@ -193,10 +193,10 @@ private:
 	ECbFieldType ExternalType;
 };
 
-static void ValidateCbObject(FConstMemoryView& View, ECbValidateMode Mode, ECbValidateError& Error, ECbFieldType ObjectType)
+static void ValidateCbObject(FMemoryView& View, ECbValidateMode Mode, ECbValidateError& Error, ECbFieldType ObjectType)
 {
 	const uint64 Size = ValidateCbUInt(View, Mode, Error);
-	FConstMemoryView ObjectView = View.Left(Size);
+	FMemoryView ObjectView = View.Left(Size);
 	View += Size;
 
 	if (Size > 0)
@@ -245,10 +245,10 @@ static void ValidateCbObject(FConstMemoryView& View, ECbValidateMode Mode, ECbVa
 	}
 }
 
-static void ValidateCbArray(FConstMemoryView& View, ECbValidateMode Mode, ECbValidateError& Error, ECbFieldType ArrayType)
+static void ValidateCbArray(FMemoryView& View, ECbValidateMode Mode, ECbValidateError& Error, ECbFieldType ArrayType)
 {
 	const uint64 Size = ValidateCbUInt(View, Mode, Error);
-	FConstMemoryView ArrayView = View.Left(Size);
+	FMemoryView ArrayView = View.Left(Size);
 	View += Size;
 
 	const uint64 Count = ValidateCbUInt(ArrayView, Mode, Error);
@@ -274,9 +274,9 @@ static void ValidateCbArray(FConstMemoryView& View, ECbValidateMode Mode, ECbVal
 	}
 }
 
-static FCbField ValidateCbField(FConstMemoryView& View, ECbValidateMode Mode, ECbValidateError& Error, const ECbFieldType ExternalType = ECbFieldType::HasFieldType)
+static FCbField ValidateCbField(FMemoryView& View, ECbValidateMode Mode, ECbValidateError& Error, const ECbFieldType ExternalType = ECbFieldType::HasFieldType)
 {
-	const FConstMemoryView FieldView = View;
+	const FMemoryView FieldView = View;
 	const ECbFieldType Type = ValidateCbFieldType(View, Mode, Error, ExternalType);
 	const FAnsiStringView Name = FCbFieldType::HasFieldName(Type) ? ValidateCbString(View, Mode, Error) : FAnsiStringView();
 
@@ -374,7 +374,7 @@ static FCbField ValidateCbField(FConstMemoryView& View, ECbValidateMode Mode, EC
 	return FCbField(FieldView.GetData(), ExternalType);
 }
 
-ECbValidateError ValidateCompactBinary(FConstMemoryView View, ECbValidateMode Mode, ECbFieldType Type)
+ECbValidateError ValidateCompactBinary(FMemoryView View, ECbValidateMode Mode, ECbFieldType Type)
 {
 	ECbValidateError Error = ECbValidateError::None;
 	ValidateCbField(View, Mode, Error, Type);
@@ -385,7 +385,7 @@ ECbValidateError ValidateCompactBinary(FConstMemoryView View, ECbValidateMode Mo
 	return Error;
 }
 
-ECbValidateError ValidateCompactBinaryRange(FConstMemoryView View, ECbValidateMode Mode)
+ECbValidateError ValidateCompactBinaryRange(FMemoryView View, ECbValidateMode Mode)
 {
 	ECbValidateError Error = ECbValidateError::None;
 	while (!View.IsEmpty())
diff --git a/Engine/Source/Runtime/Core/Private/Serialization/CompactBinaryWriter.cpp b/Engine/Source/Runtime/Core/Private/Serialization/CompactBinaryWriter.cpp
--- a/Engine/Source/Runtime/Core/Private/Serialization/CompactBinaryWriter.cpp
+++ b/Engine/Source/Runtime/Core/Private/Serialization/CompactBinaryWriter.cpp
@@ -56,7 +56,7 @@ static inline ECbFieldType AppendCompactBinary(const FCbField& Value, TArray64<u
 		using FCbField::GetPayloadView;
 	};
 	const FCopy& ValueCopy = static_cast<const FCopy&>(Value);
-	const FConstMemoryView SourceView = ValueCopy.GetPayloadView();
+	const FMemoryView SourceView = ValueCopy.GetPayloadView();
 	const int64 TargetOffset = OutData.AddUninitialized(SourceView.GetSize());
 	FMemory::Memcpy(OutData.GetData() + TargetOffset, SourceView.GetData(), SourceView.GetSize());
 	return FCbFieldType::GetType(ValueCopy.GetType());
diff --git a/Engine/Source/Runtime/Core/Private/Tests/Memory/MemoryViewTest.cpp b/Engine/Source/Runtime/Core/Private/Tests/Memory/MemoryViewTest.cpp
--- a/Engine/Source/Runtime/Core/Private/Tests/Memory/MemoryViewTest.cpp
+++ b/Engine/Source/Runtime/Core/Private/Tests/Memory/MemoryViewTest.cpp
@@ -6,6 +6,13 @@
 
 #include <type_traits>
 
+static_assert(std::is_trivially_copyable<FMemoryView>::value, "FMemoryView must be trivially copyable");
+static_assert(std::is_trivially_copy_constructible<FMemoryView>::value, "FMemoryView must be trivially copy constructible");
+static_assert(std::is_trivially_move_constructible<FMemoryView>::value, "FMemoryView must be trivially move constructible");
+static_assert(std::is_trivially_copy_assignable<FMemoryView>::value, "FMemoryView must be trivially copy assignable");
+static_assert(std::is_trivially_move_assignable<FMemoryView>::value, "FMemoryView must be trivially move assignable");
+static_assert(std::is_trivially_destructible<FMemoryView>::value, "FMemoryView must be trivially destructible");
+
 static_assert(std::is_trivially_copyable<FMutableMemoryView>::value, "FMutableMemoryView must be trivially copyable");
 static_assert(std::is_trivially_copy_constructible<FMutableMemoryView>::value, "FMutableMemoryView must be trivially copy constructible");
 static_assert(std::is_trivially_move_constructible<FMutableMemoryView>::value, "FMutableMemoryView must be trivially move constructible");
@@ -13,37 +20,30 @@ static_assert(std::is_trivially_copy_assignable<FMutableMemoryView>::value, "FMu
 static_assert(std::is_trivially_move_assignable<FMutableMemoryView>::value, "FMutableMemoryView must be trivially move assignable");
 static_assert(std::is_trivially_destructible<FMutableMemoryView>::value, "FMutableMemoryView must be trivially destructible");
 
-static_assert(std::is_trivially_copyable<FConstMemoryView>::value, "FConstMemoryView must be trivially copyable");
-static_assert(std::is_trivially_copy_constructible<FConstMemoryView>::value, "FConstMemoryView must be trivially copy constructible");
-static_assert(std::is_trivially_move_constructible<FConstMemoryView>::value, "FConstMemoryView must be trivially move constructible");
-static_assert(std::is_trivially_copy_assignable<FConstMemoryView>::value, "FConstMemoryView must be trivially copy assignable");
-static_assert(std::is_trivially_move_assignable<FConstMemoryView>::value, "FConstMemoryView must be trivially move assignable");
-static_assert(std::is_trivially_destructible<FConstMemoryView>::value, "FConstMemoryView must be trivially destructible");
-
-static_assert(std::is_constructible<FConstMemoryView, const FMutableMemoryView&>::value, "Missing constructor");
-static_assert(!std::is_constructible<FMutableMemoryView, const FConstMemoryView&>::value, "Invalid constructor");
+static_assert(std::is_constructible<FMemoryView, const FMutableMemoryView&>::value, "Missing constructor");
+static_assert(!std::is_constructible<FMutableMemoryView, const FMemoryView&>::value, "Invalid constructor");
 static_assert(!std::is_constructible<FMutableMemoryView, const void*, uint64>::value, "Invalid constructor");
 
-static_assert(std::is_assignable<FConstMemoryView, const FMutableMemoryView&>::value, "Missing assignment");
-static_assert(!std::is_assignable<FMutableMemoryView, const FConstMemoryView&>::value, "Invalid assignment");
+static_assert(std::is_assignable<FMemoryView, const FMutableMemoryView&>::value, "Missing assignment");
+static_assert(!std::is_assignable<FMutableMemoryView, const FMemoryView&>::value, "Invalid assignment");
 
 #if WITH_DEV_AUTOMATION_TESTS
 
 IMPLEMENT_SIMPLE_AUTOMATION_TEST(FMemoryViewTest, "System.Core.Memory.MemoryView", EAutomationTestFlags::ApplicationContextMask | EAutomationTestFlags::SmokeFilter)
 bool FMemoryViewTest::RunTest(const FString& Parameters)
 {
-	auto TestMutableMemoryView = [this](const FMutableMemoryView& View, void* Data, uint64 Size)
+	auto TestMemoryView = [this](const FMemoryView& View, const void* Data, uint64 Size)
 	{
 		TestEqual(TEXT("MemoryView.GetData()"), View.GetData(), Data);
-		TestEqual(TEXT("MemoryView.GetDataEnd()"), View.GetDataEnd(), static_cast<void*>(static_cast<uint8*>(Data) + Size));
+		TestEqual(TEXT("MemoryView.GetDataEnd()"), View.GetDataEnd(), static_cast<const void*>(static_cast<const uint8*>(Data) + Size));
 		TestEqual(TEXT("MemoryView.GetSize()"), View.GetSize(), Size);
 		TestEqual(TEXT("MemoryView.IsEmpty()"), View.IsEmpty(), Size == 0);
 	};
 
-	auto TestConstMemoryView = [this](const FConstMemoryView& View, const void* Data, uint64 Size)
+	auto TestMutableMemoryView = [this](const FMutableMemoryView& View, void* Data, uint64 Size)
 	{
 		TestEqual(TEXT("MemoryView.GetData()"), View.GetData(), Data);
-		TestEqual(TEXT("MemoryView.GetDataEnd()"), View.GetDataEnd(), static_cast<const void*>(static_cast<const uint8*>(Data) + Size));
+		TestEqual(TEXT("MemoryView.GetDataEnd()"), View.GetDataEnd(), static_cast<void*>(static_cast<uint8*>(Data) + Size));
 		TestEqual(TEXT("MemoryView.GetSize()"), View.GetSize(), Size);
 		TestEqual(TEXT("MemoryView.IsEmpty()"), View.IsEmpty(), Size == 0);
 	};
@@ -59,24 +59,24 @@ bool FMemoryViewTest::RunTest(const FString& Parameters)
 	uint32 IntArray[12]{};
 
 	// Test Empty Views
+	TestMemoryView(FMemoryView(), nullptr, 0);
+	TestMemoryView(FMutableMemoryView(), nullptr, 0);
 	TestMutableMemoryView(FMutableMemoryView(), nullptr, 0);
-	TestConstMemoryView(FMutableMemoryView(), nullptr, 0);
-	TestConstMemoryView(FConstMemoryView(), nullptr, 0);
 
 	// Test Construction from Type[], TArrayView, (Type*, uint64), (Type*, Type*)
+	TestMemoryView(MakeMemoryView(AsConst(IntArray)), IntArray, sizeof(IntArray));
+	TestMemoryView(MakeMemoryView(MakeArrayView(AsConst(IntArray))), IntArray, sizeof(IntArray));
+	TestMemoryView(MakeMemoryView(AsConst(IntArray), sizeof(IntArray)), IntArray, sizeof(IntArray));
+	TestMemoryView(MakeMemoryView(AsConst(IntArray), AsConst(IntArray) + 6), IntArray, sizeof(*IntArray) * 6);
 	TestMutableMemoryView(MakeMemoryView(IntArray), IntArray, sizeof(IntArray));
 	TestMutableMemoryView(MakeMemoryView(MakeArrayView(IntArray)), IntArray, sizeof(IntArray));
 	TestMutableMemoryView(MakeMemoryView(IntArray, sizeof(IntArray)), IntArray, sizeof(IntArray));
 	TestMutableMemoryView(MakeMemoryView(IntArray, IntArray + 6), IntArray, sizeof(*IntArray) * 6);
-	TestConstMemoryView(MakeMemoryView(AsConst(IntArray)), IntArray, sizeof(IntArray));
-	TestConstMemoryView(MakeMemoryView(MakeArrayView(AsConst(IntArray))), IntArray, sizeof(IntArray));
-	TestConstMemoryView(MakeMemoryView(AsConst(IntArray), sizeof(IntArray)), IntArray, sizeof(IntArray));
-	TestConstMemoryView(MakeMemoryView(AsConst(IntArray), AsConst(IntArray) + 6), IntArray, sizeof(*IntArray) * 6);
 
 	// Test Construction from std::initializer_list
 	//MakeMemoryView({1, 2, 3}); // fail because the type must be deduced
 	std::initializer_list<uint8> InitializerList{1, 2, 3};
-	TestConstMemoryView(MakeMemoryView(InitializerList), GetData(InitializerList), GetNum(InitializerList) * sizeof(uint8));
+	TestMemoryView(MakeMemoryView(InitializerList), GetData(InitializerList), GetNum(InitializerList) * sizeof(uint8));
 
 	// Test Reset
 	{
@@ -128,8 +128,8 @@ bool FMemoryViewTest::RunTest(const FString& Parameters)
 	TestEqual(TEXT("MemoryView.Mid(MaxSize,MaxSize)"), MakeMemoryView(IntArray).Mid(MAX_uint64, MAX_uint64), FMutableMemoryView());
 
 	// Test Contains
-	TestTrue(TEXT("MemoryView.Contains(Empty)"), FMutableMemoryView().Contains(FConstMemoryView()));
-	TestTrue(TEXT("MemoryView.Contains(Empty)"), FConstMemoryView().Contains(FMutableMemoryView()));
+	TestTrue(TEXT("MemoryView.Contains(Empty)"), FMemoryView().Contains(FMutableMemoryView()));
+	TestTrue(TEXT("MemoryView.Contains(Empty)"), FMutableMemoryView().Contains(FMemoryView()));
 	TestTrue(TEXT("MemoryView.Contains(Equal)"), MakeMemoryView(ByteArray).Contains(MakeMemoryView(ByteArray)));
 	TestTrue(TEXT("MemoryView.Contains(SmallerBy1Left)"), MakeMemoryView(ByteArray).Contains(MakeMemoryView(ByteArray + 1, 15)));
 	TestTrue(TEXT("MemoryView.Contains(SmallerBy1Right)"), MakeMemoryView(ByteArray).Contains(MakeMemoryView(ByteArray, 15)));
@@ -158,8 +158,8 @@ bool FMemoryViewTest::RunTest(const FString& Parameters)
 	TestTrue(TEXT("MemoryView.Intersects(LargerBy1Right)"), MakeMemoryView(ByteArray).Intersects(MakeMemoryView(ByteArray, 17)));
 	TestTrue(TEXT("MemoryView.Intersects(LargerBy2Both)"), MakeMemoryView(ByteArray).Intersects(MakeMemoryView(ByteArrayContainer.BeforeByteArray + 3, 18)));
 	TestTrue(TEXT("MemoryView.Intersects(EmptyMiddle)"), MakeMemoryView(ByteArray).Intersects(MakeMemoryView(ByteArray + 8, 0)));
-	TestFalse(TEXT("MemoryView.Intersects(Empty)"), FMutableMemoryView().Intersects(FConstMemoryView()));
-	TestFalse(TEXT("MemoryView.Intersects(Empty)"), FConstMemoryView().Intersects(FMutableMemoryView()));
+	TestFalse(TEXT("MemoryView.Intersects(Empty)"), FMemoryView().Intersects(FMutableMemoryView()));
+	TestFalse(TEXT("MemoryView.Intersects(Empty)"), FMutableMemoryView().Intersects(FMemoryView()));
 	TestFalse(TEXT("MemoryView.Intersects(EmptyLeft)"), MakeMemoryView(ByteArray).Intersects(MakeMemoryView(ByteArray, 0)));
 	TestFalse(TEXT("MemoryView.Intersects(EmptyRight)"), MakeMemoryView(ByteArray).Intersects(MakeMemoryView(ByteArray + 16, 0)));
 	TestFalse(TEXT("MemoryView.Intersects(EmptyOutside)"), MakeMemoryView(ByteArray).Intersects(MakeMemoryView(ByteArrayContainer.BeforeByteArray + 3, 0)));
@@ -169,12 +169,12 @@ bool FMemoryViewTest::RunTest(const FString& Parameters)
 
 	// Test CompareBytes
 	const uint8 CompareBytes[8] = { 5, 4, 6, 2, 4, 7, 1, 3 };
-	TestEqual(TEXT("MemoryView.CompareBytes(Empty)"), FMutableMemoryView().CompareBytes(FConstMemoryView()), 0);
-	TestEqual(TEXT("MemoryView.CompareBytes(Empty)"), FConstMemoryView().CompareBytes(FMutableMemoryView()), 0);
+	TestEqual(TEXT("MemoryView.CompareBytes(Empty)"), FMemoryView().CompareBytes(FMutableMemoryView()), 0);
+	TestEqual(TEXT("MemoryView.CompareBytes(Empty)"), FMutableMemoryView().CompareBytes(FMemoryView()), 0);
 	TestEqual(TEXT("MemoryView.CompareBytes(EqualView)"), MakeMemoryView(ByteArray).CompareBytes(MakeMemoryView(ByteArray)), 0);
 	TestEqual(TEXT("MemoryView.CompareBytes(EqualBytes)"), MakeMemoryView(ByteArray, 8).CompareBytes(MakeMemoryView(ByteArray + 8, 8)), 0);
-	TestTrue(TEXT("MemoryView.CompareBytes(EmptyLeft)"), FConstMemoryView().CompareBytes(MakeMemoryView(ByteArray)) < 0);
-	TestTrue(TEXT("MemoryView.CompareBytes(EmptyRight)"), MakeMemoryView(ByteArray).CompareBytes(FConstMemoryView()) > 0);
+	TestTrue(TEXT("MemoryView.CompareBytes(EmptyLeft)"), FMemoryView().CompareBytes(MakeMemoryView(ByteArray)) < 0);
+	TestTrue(TEXT("MemoryView.CompareBytes(EmptyRight)"), MakeMemoryView(ByteArray).CompareBytes(FMemoryView()) > 0);
 	TestTrue(TEXT("MemoryView.CompareBytes(ShorterLeft)"), MakeMemoryView(ByteArray, 8).CompareBytes(MakeMemoryView(ByteArray)) < 0);
 	TestTrue(TEXT("MemoryView.CompareBytes(ShorterRight)"), MakeMemoryView(ByteArray).CompareBytes(MakeMemoryView(ByteArray, 8)) > 0);
 	TestTrue(TEXT("MemoryView.CompareBytes(ShorterLeft)"), MakeMemoryView(IntArray, 8).CompareBytes(MakeMemoryView(ByteArray)) < 0);
@@ -185,8 +185,8 @@ bool FMemoryViewTest::RunTest(const FString& Parameters)
 	TestTrue(TEXT("MemoryView.CompareBytes(SameSizeLeftGreater)"), MakeMemoryView(CompareBytes, 3).CompareBytes(MakeMemoryView(CompareBytes + 3, 3)) > 0);
 
 	// Test EqualBytes
-	TestTrue(TEXT("MemoryView.EqualBytes(Empty)"), FMutableMemoryView().EqualBytes(FConstMemoryView()));
-	TestTrue(TEXT("MemoryView.EqualBytes(Empty)"), FConstMemoryView().EqualBytes(FMutableMemoryView()));
+	TestTrue(TEXT("MemoryView.EqualBytes(Empty)"), FMemoryView().EqualBytes(FMutableMemoryView()));
+	TestTrue(TEXT("MemoryView.EqualBytes(Empty)"), FMutableMemoryView().EqualBytes(FMemoryView()));
 	TestTrue(TEXT("MemoryView.EqualBytes(EqualView)"), MakeMemoryView(ByteArray).EqualBytes(MakeMemoryView(ByteArray)));
 	TestTrue(TEXT("MemoryView.EqualBytes(EqualBytes)"), MakeMemoryView(ByteArray, 8).EqualBytes(MakeMemoryView(ByteArray + 8, 8)));
 	TestFalse(TEXT("MemoryView.EqualBytes(DifferentSize)"), MakeMemoryView(ByteArray, 8).EqualBytes(MakeMemoryView(ByteArray, 4)));
@@ -194,10 +194,10 @@ bool FMemoryViewTest::RunTest(const FString& Parameters)
 	TestFalse(TEXT("MemoryView.EqualBytes(DifferentBytes)"), MakeMemoryView(CompareBytes, 4).EqualBytes(MakeMemoryView(CompareBytes + 4, 4)));
 
 	// Test Equals
+	TestTrue(TEXT("MemoryView.Equals(Empty)"), FMemoryView().Equals(FMemoryView()));
+	TestTrue(TEXT("MemoryView.Equals(Empty)"), FMemoryView().Equals(FMutableMemoryView()));
+	TestTrue(TEXT("MemoryView.Equals(Empty)"), FMutableMemoryView().Equals(FMemoryView()));
 	TestTrue(TEXT("MemoryView.Equals(Empty)"), FMutableMemoryView().Equals(FMutableMemoryView()));
-	TestTrue(TEXT("MemoryView.Equals(Empty)"), FMutableMemoryView().Equals(FConstMemoryView()));
-	TestTrue(TEXT("MemoryView.Equals(Empty)"), FConstMemoryView().Equals(FMutableMemoryView()));
-	TestTrue(TEXT("MemoryView.Equals(Empty)"), FConstMemoryView().Equals(FConstMemoryView()));
 	TestTrue(TEXT("MemoryView.Equals(Equal)"), MakeMemoryView(IntArray).Equals(MakeMemoryView(AsConst(IntArray))));
 	TestFalse(TEXT("MemoryView.Equals(DataDiff)"), MakeMemoryView(IntArray).Equals(MakeMemoryView(IntArray + 1, sizeof(IntArray) - sizeof(*IntArray))));
 	TestFalse(TEXT("MemoryView.Equals(SizeDiff)"), MakeMemoryView(IntArray).Equals(MakeMemoryView(IntArray, sizeof(*IntArray))));
diff --git a/Engine/Source/Runtime/Core/Private/Tests/Memory/SharedBufferTest.cpp b/Engine/Source/Runtime/Core/Private/Tests/Memory/SharedBufferTest.cpp
--- a/Engine/Source/Runtime/Core/Private/Tests/Memory/SharedBufferTest.cpp
+++ b/Engine/Source/Runtime/Core/Private/Tests/Memory/SharedBufferTest.cpp
@@ -124,13 +124,13 @@ static_assert(std::is_same<FSharedBufferPtr, decltype(FSharedBuffer::Clone(DeclV
 static_assert(std::is_same<FSharedBufferPtr, decltype(FSharedBuffer::Clone(*DeclVal<FSharedBufferRef>()))>::value, "Invalid constructor");
 static_assert(std::is_same<FSharedBufferPtr, decltype(FSharedBuffer::Clone(*DeclVal<FSharedBufferConstRef>()))>::value, "Invalid constructor");
 static_assert(std::is_same<FSharedBufferPtr, decltype(FSharedBuffer::Clone(DeclVal<FMutableMemoryView>()))>::value, "Invalid constructor");
-static_assert(std::is_same<FSharedBufferPtr, decltype(FSharedBuffer::Clone(DeclVal<FConstMemoryView>()))>::value, "Invalid constructor");
+static_assert(std::is_same<FSharedBufferPtr, decltype(FSharedBuffer::Clone(DeclVal<FMemoryView>()))>::value, "Invalid constructor");
 
 static_assert(std::is_same<FSharedBufferPtr, decltype(FSharedBuffer::MakeView(DeclVal<FMutableMemoryView>()))>::value, "Invalid constructor");
-static_assert(std::is_same<FSharedBufferConstPtr, decltype(FSharedBuffer::MakeView(DeclVal<FConstMemoryView>()))>::value, "Invalid constructor");
+static_assert(std::is_same<FSharedBufferConstPtr, decltype(FSharedBuffer::MakeView(DeclVal<FMemoryView>()))>::value, "Invalid constructor");
 static_assert(std::is_same<FSharedBufferPtr, decltype(FSharedBuffer::MakeView(DeclVal<void*>(), 0))>::value, "Invalid constructor");
 static_assert(std::is_same<FSharedBufferConstPtr, decltype(FSharedBuffer::MakeView(DeclVal<const void*>(), 0))>::value, "Invalid constructor");
-static_assert(std::is_same<FSharedBufferConstPtr, decltype(FSharedBuffer::MakeView(DeclVal<FConstMemoryView>(), DeclVal<const FSharedBuffer&>()))>::value, "Invalid constructor");
+static_assert(std::is_same<FSharedBufferConstPtr, decltype(FSharedBuffer::MakeView(DeclVal<FMemoryView>(), DeclVal<const FSharedBuffer&>()))>::value, "Invalid constructor");
 static_assert(std::is_same<FSharedBufferConstPtr, decltype(FSharedBuffer::MakeView(DeclVal<const void*>(), 0, DeclVal<const FSharedBuffer&>()))>::value, "Invalid constructor");
 
 static_assert(std::is_same<FSharedBufferPtr, decltype(FSharedBuffer::TakeOwnership(DeclVal<void*>(), 0, FMemory::Free))>::value, "Invalid constructor");
@@ -141,12 +141,12 @@ static_assert(std::is_same<const void*, decltype(DeclVal<const FSharedBuffer>().
 static_assert(std::is_same<uint64, decltype(DeclVal<FSharedBuffer>().GetSize())>::value, "Invalid accessor");
 
 static_assert(std::is_same<FMutableMemoryView, decltype(DeclVal<FSharedBuffer>().GetView())>::value, "Invalid accessor");
-static_assert(std::is_same<FConstMemoryView, decltype(DeclVal<const FSharedBuffer>().GetView())>::value, "Invalid accessor");
+static_assert(std::is_same<FMemoryView, decltype(DeclVal<const FSharedBuffer>().GetView())>::value, "Invalid accessor");
 
 static_assert(std::is_convertible<FSharedBuffer, FMutableMemoryView>::value, "Missing conversion");
-static_assert(std::is_convertible<FSharedBuffer, FConstMemoryView>::value, "Missing conversion");
+static_assert(std::is_convertible<FSharedBuffer, FMemoryView>::value, "Missing conversion");
 static_assert(!std::is_convertible<const FSharedBuffer, FMutableMemoryView>::value, "Invalid conversion");
-static_assert(std::is_convertible<const FSharedBuffer, FConstMemoryView>::value, "Missing conversion");
+static_assert(std::is_convertible<const FSharedBuffer, FMemoryView>::value, "Missing conversion");
 
 ///////////////////////////////////////////////////////////////////////////////////////////////////////////////////////
 
diff --git a/Engine/Source/Runtime/Core/Private/Tests/Serialization/CompactBinarySerializationTest.cpp b/Engine/Source/Runtime/Core/Private/Tests/Serialization/CompactBinarySerializationTest.cpp
--- a/Engine/Source/Runtime/Core/Private/Tests/Serialization/CompactBinarySerializationTest.cpp
+++ b/Engine/Source/Runtime/Core/Private/Tests/Serialization/CompactBinarySerializationTest.cpp
@@ -132,7 +132,7 @@ bool FCbMeasureTest::RunTest(const FString& Parameters)
 IMPLEMENT_SIMPLE_AUTOMATION_TEST(FCbSaveTest, "System.Core.Serialization.SaveCompactBinary", CompactBinarySerializationTestFlags)
 bool FCbSaveTest::RunTest(const FString& Parameters)
 {
-	const auto TestSave = [this](const TCHAR* Test, auto Value, FConstMemoryView ExpectedData)
+	const auto TestSave = [this](const TCHAR* Test, auto Value, FMemoryView ExpectedData)
 	{
 		{
 			FBufferArchive WriteAr;
diff --git a/Engine/Source/Runtime/Core/Private/Tests/Serialization/CompactBinaryTest.cpp b/Engine/Source/Runtime/Core/Private/Tests/Serialization/CompactBinaryTest.cpp
--- a/Engine/Source/Runtime/Core/Private/Tests/Serialization/CompactBinaryTest.cpp
+++ b/Engine/Source/Runtime/Core/Private/Tests/Serialization/CompactBinaryTest.cpp
@@ -52,7 +52,7 @@ constexpr auto TCbFieldAsTypeFn = TCbFieldTypeAccessors<FCbFieldType::GetType(Fi
 
 UE_CBFIELD_TYPE_ACCESSOR(Object, IsObject, AsObject, FCbObject);
 UE_CBFIELD_TYPE_ACCESSOR(Array, IsArray, AsArray, FCbArray);
-UE_CBFIELD_TYPE_ACCESSOR(Binary, IsBinary, AsBinary, FConstMemoryView);
+UE_CBFIELD_TYPE_ACCESSOR(Binary, IsBinary, AsBinary, FMemoryView);
 UE_CBFIELD_TYPE_ACCESSOR(String, IsString, AsString, FAnsiStringView);
 UE_CBFIELD_TYPE_ACCESSOR(IntegerPositive, IsInteger, AsUInt64, uint64);
 UE_CBFIELD_TYPE_ACCESSOR(IntegerNegative, IsInteger, AsInt64, int64);
@@ -75,7 +75,7 @@ protected:
 	using FAutomationTestBase::FAutomationTestBase;
 	using FAutomationTestBase::TestEqual;
 
-	void TestEqualBytes(const TCHAR* What, FConstMemoryView Actual, TArrayView<const uint8> Expected)
+	void TestEqualBytes(const TCHAR* What, FMemoryView Actual, TArrayView<const uint8> Expected)
 	{
 		TestTrue(What, Actual.EqualBytes(MakeMemoryView(Expected)));
 	}
@@ -192,8 +192,8 @@ bool FCbFieldNoneTest::RunTest(const FString& Parameters)
 		TestFalse(TEXT("!FCbField()::HasError()"), DefaultField.HasError());
 		TestEqual(TEXT("FCbField()::GetError() == None"), DefaultField.GetError(), ECbFieldError::None);
 		TestEqual(TEXT("FCbField()::GetHash()"), DefaultField.GetHash(), FBlake3::HashBuffer(MakeMemoryView<uint8>({uint8(ECbFieldType::None)})));
-		TestEqual(TEXT("FCbField()::GetView()"), DefaultField.GetView(), FConstMemoryView());
-		FConstMemoryView SerializedView;
+		TestEqual(TEXT("FCbField()::GetView()"), DefaultField.GetView(), FMemoryView());
+		FMemoryView SerializedView;
 		TestFalse(TEXT("FCbField()::TryGetSerializedView()"), DefaultField.TryGetSerializedView(SerializedView));
 	}
 
@@ -207,8 +207,8 @@ bool FCbFieldNoneTest::RunTest(const FString& Parameters)
 		TestFalse(TEXT("!FCbField(None)::HasError()"), NoneField.HasError());
 		TestEqual(TEXT("FCbField(None)::GetError() == None"), NoneField.GetError(), ECbFieldError::None);
 		TestEqual(TEXT("FCbField(None)::GetHash()"), NoneField.GetHash(), FCbField().GetHash());
-		TestEqual(TEXT("FCbField(None)::GetView()"), NoneField.GetView(), FConstMemoryView());
-		FConstMemoryView SerializedView;
+		TestEqual(TEXT("FCbField(None)::GetView()"), NoneField.GetView(), FMemoryView());
+		FMemoryView SerializedView;
 		TestFalse(TEXT("FCbField(None)::TryGetSerializedView()"), NoneField.TryGetSerializedView(SerializedView));
 	}
 
@@ -223,7 +223,7 @@ bool FCbFieldNoneTest::RunTest(const FString& Parameters)
 		TestFalse(TEXT("!FCbField(None|Type|Name)::HasValue()"), NoneField.HasValue());
 		TestEqual(TEXT("FCbField(None|Type|Name)::GetHash()"), NoneField.GetHash(), FBlake3::HashBuffer(MakeMemoryView(NoneBytes)));
 		TestEqual(TEXT("FCbField(None|Type|Name)::GetView()"), NoneField.GetView(), MakeMemoryView(NoneBytes));
-		FConstMemoryView SerializedView;
+		FMemoryView SerializedView;
 		TestTrue(TEXT("FCbField(None|Type|Name)::TryGetSerializedView()"), NoneField.TryGetSerializedView(SerializedView) && SerializedView == MakeMemoryView(NoneBytes));
 
 		uint8 CopyBytes[sizeof(NoneBytes)];
@@ -242,7 +242,7 @@ bool FCbFieldNoneTest::RunTest(const FString& Parameters)
 		TestFalse(TEXT("!FCbField(None|Type)::HasValue()"), NoneField.HasValue());
 		TestEqual(TEXT("FCbField(None|Type)::GetHash()"), NoneField.GetHash(), FCbField().GetHash());
 		TestEqual(TEXT("FCbField(None|Type)::GetView()"), NoneField.GetView(), MakeMemoryView(NoneBytes));
-		FConstMemoryView SerializedView;
+		FMemoryView SerializedView;
 		TestTrue(TEXT("FCbField(None|Type)::TryGetSerializedView()"), NoneField.TryGetSerializedView(SerializedView) && SerializedView == MakeMemoryView(NoneBytes));
 	}
 
@@ -257,7 +257,7 @@ bool FCbFieldNoneTest::RunTest(const FString& Parameters)
 		TestFalse(TEXT("!FCbField(None|Name)::HasValue()"), NoneField.HasValue());
 		TestEqual(TEXT("FCbField(None|Name)::GetHash()"), NoneField.GetHash(), FBlake3::HashBuffer(MakeMemoryView(NoneBytes)));
 		TestEqual(TEXT("FCbField(None|Name)::GetView()"), NoneField.GetView(), MakeMemoryView(NoneBytes) + 1);
-		FConstMemoryView SerializedView;
+		FMemoryView SerializedView;
 		TestFalse(TEXT("FCbField(None|Name)::TryGetSerializedView()"), NoneField.TryGetSerializedView(SerializedView));
 
 		uint8 CopyBytes[sizeof(NoneBytes)];
@@ -276,7 +276,7 @@ bool FCbFieldNoneTest::RunTest(const FString& Parameters)
 		TestFalse(TEXT("!FCbField(None|EmptyName)::HasValue()"), NoneField.HasValue());
 		TestEqual(TEXT("FCbField(None|EmptyName)::GetHash()"), NoneField.GetHash(), FBlake3::HashBuffer(MakeMemoryView(NoneBytes)));
 		TestEqual(TEXT("FCbField(None|EmptyName)::GetView()"), NoneField.GetView(), MakeMemoryView(NoneBytes) + 1);
-		FConstMemoryView SerializedView;
+		FMemoryView SerializedView;
 		TestFalse(TEXT("FCbField(None|EmptyName)::TryGetSerializedView()"), NoneField.TryGetSerializedView(SerializedView));
 	}
 
@@ -422,7 +422,7 @@ bool FCbFieldObjectTest::RunTest(const FString& Parameters)
 		TestTrue(TEXT("FCbObject(NoType, Name)::CopyTo()"), MakeMemoryView(Payload).EqualBytes(MakeMemoryView(CopyBytes) + 1));
 
 		// GetSerializedView
-		FConstMemoryView SerializedView;
+		FMemoryView SerializedView;
 		TestTrue(TEXT("FCbObject(Clone)::TryGetSerializedView()"), Object.TryGetSerializedView(SerializedView) && SerializedView == Object.AsField().GetView());
 		TestFalse(TEXT("FCbObject(NoType)::TryGetSerializedView()"), Field.AsObject().TryGetSerializedView(SerializedView));
 		TestFalse(TEXT("FCbObject(Name)::TryGetSerializedView()"), NamedField.AsObject().TryGetSerializedView(SerializedView));
@@ -587,7 +587,7 @@ bool FCbFieldArrayTest::RunTest(const FString& Parameters)
 		TestTrue(TEXT("FCbArray(NoType, Name)::CopyTo()"), MakeMemoryView(Payload).EqualBytes(MakeMemoryView(CopyBytes) + 1));
 
 		// GetSerializedView
-		FConstMemoryView SerializedView;
+		FMemoryView SerializedView;
 		TestTrue(TEXT("FCbArray(Clone)::TryGetSerializedView()"), Array.TryGetSerializedView(SerializedView) && SerializedView == Array.AsField().GetView());
 		TestFalse(TEXT("FCbArray(NoType)::TryGetSerializedView()"), Field.AsArray().TryGetSerializedView(SerializedView));
 		TestFalse(TEXT("FCbArray(Name)::TryGetSerializedView()"), NamedField.AsArray().TryGetSerializedView(SerializedView));
@@ -1522,7 +1522,7 @@ bool FCbFieldRefIteratorTest::RunTest(const FString& Parameters)
 		const FSharedBufferConstPtr View = FSharedBuffer::MakeView(MakeMemoryView(Payload));
 		const FSharedBufferConstPtr Clone = FSharedBuffer::Clone(MakeMemoryView(Payload));
 
-		const FConstMemoryView EmptyView;
+		const FMemoryView EmptyView;
 		const FSharedBufferConstPtr NullBuffer;
 
 		const FCbFieldIterator FieldIt = FCbFieldIterator::MakeRange(*View);
@@ -1533,7 +1533,7 @@ bool FCbFieldRefIteratorTest::RunTest(const FString& Parameters)
 
 		TestEqual(TEXT("FCbFieldIterator::GetRangeView()"), FieldIt.GetRangeView(), MakeMemoryView(Payload));
 		TestEqual(TEXT("FCbFieldRefIterator::GetRangeView()"), FieldRefIt.GetRangeView(), MakeMemoryView(Payload));
-		FConstMemoryView SerializedView;
+		FMemoryView SerializedView;
 		TestTrue(TEXT("FCbFieldIterator::TryGetSerializedRangeView()"), FieldIt.TryGetSerializedRangeView(SerializedView) && SerializedView == MakeMemoryView(Payload));
 		TestTrue(TEXT("FCbFieldRefIterator::TryGetSerializedRangeView()"), FieldRefIt.TryGetSerializedRangeView(SerializedView) && SerializedView == MakeMemoryView(Payload));
 
diff --git a/Engine/Source/Runtime/Core/Private/Tests/Serialization/CompactBinaryWriterTest.cpp b/Engine/Source/Runtime/Core/Private/Tests/Serialization/CompactBinaryWriterTest.cpp
--- a/Engine/Source/Runtime/Core/Private/Tests/Serialization/CompactBinaryWriterTest.cpp
+++ b/Engine/Source/Runtime/Core/Private/Tests/Serialization/CompactBinaryWriterTest.cpp
@@ -763,7 +763,7 @@ bool FCbWriterComplexTest::RunTest(const FString& Parameters)
 		Writer.Name("Binary"_ASV);
 		Writer.BeginObject();
 		{
-			Writer.Name("Empty"_ASV).Binary(FConstMemoryView());
+			Writer.Name("Empty"_ASV).Binary(FMemoryView());
 			Writer.Name("Value"_ASV).Binary(MakeMemoryView("BinaryValue"));
 			Writer.Name("LargeValue"_ASV).Binary(MakeMemoryView(FString::ChrN(256, TEXT('.'))));
 			Writer.Name("LargeRefValue"_ASV).Binary(FSharedBuffer::Clone(MakeMemoryView(FString::ChrN(256, TEXT('!')))));
diff --git a/Engine/Source/Runtime/Core/Public/Memory/MemoryFwd.h b/Engine/Source/Runtime/Core/Public/Memory/MemoryFwd.h
--- a/Engine/Source/Runtime/Core/Public/Memory/MemoryFwd.h
+++ b/Engine/Source/Runtime/Core/Public/Memory/MemoryFwd.h
@@ -7,12 +7,12 @@
 template <typename DataType>
 class TMemoryView;
 
+/** A non-owning view of a contiguous region of memory. */
+using FMemoryView = TMemoryView<const void>;
+
 /** A non-owning mutable view of a contiguous region of memory. */
 using FMutableMemoryView = TMemoryView<void>;
 
-/** A non-owning const view of a contiguous region of memory. */
-using FConstMemoryView = TMemoryView<const void>;
-
 // SharedBuffer
 
 class FSharedBufferRef;
diff --git a/Engine/Source/Runtime/Core/Public/Memory/MemoryView.h b/Engine/Source/Runtime/Core/Public/Memory/MemoryView.h
--- a/Engine/Source/Runtime/Core/Public/Memory/MemoryView.h
+++ b/Engine/Source/Runtime/Core/Public/Memory/MemoryView.h
@@ -21,7 +21,7 @@
 /**
  * A non-owning view of a contiguous region of memory.
  *
- * Prefer to use the aliases FConstMemoryView or FMutableMemoryView over this type.
+ * Prefer to use the aliases FMemoryView or FMutableMemoryView over this type.
  *
  * Functions that modify a view clamp sizes and offsets to always return a sub-view of the input.
  */
diff --git a/Engine/Source/Runtime/Core/Public/Memory/SharedBuffer.h b/Engine/Source/Runtime/Core/Public/Memory/SharedBuffer.h
--- a/Engine/Source/Runtime/Core/Public/Memory/SharedBuffer.h
+++ b/Engine/Source/Runtime/Core/Public/Memory/SharedBuffer.h
@@ -495,7 +495,7 @@ public:
 	}
 
 	/** Make an owned writable clone of the memory view. */
-	static inline FSharedBufferPtr Clone(const FConstMemoryView View)
+	static inline FSharedBufferPtr Clone(const FMemoryView View)
 	{
 		return Clone(View.GetData(), View.GetSize());
 	}
@@ -515,7 +515,7 @@ public:
 	}
 
 	/** Make a non-owned read-only view of the memory view. */
-	static inline FSharedBufferConstPtr MakeView(FConstMemoryView View)
+	static inline FSharedBufferConstPtr MakeView(FMemoryView View)
 	{
 		return MakeView(View.GetData(), View.GetSize());
 	}
@@ -534,7 +534,7 @@ public:
 	}
 
 	/** Make a non-owned read-only view of the outer buffer and hold a reference to it. */
-	static inline FSharedBufferConstPtr MakeView(FConstMemoryView View, const FSharedBuffer& OuterBuffer)
+	static inline FSharedBufferConstPtr MakeView(FMemoryView View, const FSharedBuffer& OuterBuffer)
 	{
 		return MakeView(View.GetData(), View.GetSize(), OuterBuffer);
 	}
@@ -724,10 +724,10 @@ public:
 	/** A writable view of buffer. */
 	inline FMutableMemoryView GetView() { return FMutableMemoryView(GetData(), GetSize()); }
 	/** A read-only view of the buffer. */
-	inline FConstMemoryView GetView() const { return FConstMemoryView(GetData(), GetSize()); }
+	inline FMemoryView GetView() const { return FMemoryView(GetData(), GetSize()); }
 
 	inline operator FMutableMemoryView() { return GetView(); }
-	inline operator FConstMemoryView() const { return GetView(); }
+	inline operator FMemoryView() const { return GetView(); }
 
 private:
 	/**
diff --git a/Engine/Source/Runtime/Core/Public/Misc/Blake3.h b/Engine/Source/Runtime/Core/Public/Misc/Blake3.h
--- a/Engine/Source/Runtime/Core/Public/Misc/Blake3.h
+++ b/Engine/Source/Runtime/Core/Public/Misc/Blake3.h
@@ -50,7 +50,7 @@ public:
 	CORE_API void Update(const void* Data, uint64 Size);
 
 	/** Add the view as input to the hash. May be called any number of times. */
-	inline void Update(FConstMemoryView View)
+	inline void Update(FMemoryView View)
 	{
 		Update(View.GetData(), View.GetSize());
 	}
@@ -62,7 +62,7 @@ public:
 	CORE_API static FBlake3Hash HashBuffer(const void* Data, uint64 Size);
 
 	/** Calculate the hash of the input view. */
-	static inline FBlake3Hash HashBuffer(FConstMemoryView View)
+	static inline FBlake3Hash HashBuffer(FMemoryView View)
 	{
 		return HashBuffer(View.GetData(), View.GetSize());
 	}
diff --git a/Engine/Source/Runtime/Core/Public/Serialization/CompactBinary.h b/Engine/Source/Runtime/Core/Public/Serialization/CompactBinary.h
--- a/Engine/Source/Runtime/Core/Public/Serialization/CompactBinary.h
+++ b/Engine/Source/Runtime/Core/Public/Serialization/CompactBinary.h
@@ -333,7 +333,7 @@ public:
 	CORE_API void IterateRangeReferences(FCbFieldVisitor Visitor) const;
 
 	/** Create a view of every field in the range. */
-	inline FConstMemoryView GetRangeView() const
+	inline FMemoryView GetRangeView() const
 	{
 		return MakeMemoryView(FieldType::GetView().GetData(), FieldsEnd);
 	}
@@ -344,7 +344,7 @@ public:
 	 * A serialized view is not available if the underlying fields have an externally-provided type.
 	 * Access the serialized form of such ranges using FCbFieldRefIterator::CloneRange.
 	 */
-	inline bool TryGetSerializedRangeView(FConstMemoryView& OutView) const
+	inline bool TryGetSerializedRangeView(FMemoryView& OutView) const
 	{
 		if (FCbFieldType::HasFieldType(FieldType::GetType()))
 		{
@@ -455,7 +455,7 @@ public:
 	CORE_API FCbArray AsArray();
 
 	/** Access the field as binary. Returns the provided default on error. */
-	CORE_API FConstMemoryView AsBinary(FConstMemoryView Default = FConstMemoryView());
+	CORE_API FMemoryView AsBinary(FMemoryView Default = FMemoryView());
 
 	/** Access the field as a UTF-8 string. Returns the provided default on error. */
 	CORE_API FAnsiStringView AsString(FAnsiStringView Default = FAnsiStringView());
@@ -587,7 +587,7 @@ public:
 	CORE_API void IterateReferences(FCbFieldVisitor Visitor) const;
 
 	/** Returns a view of the field, including the type and name when present. */
-	CORE_API FConstMemoryView GetView() const;
+	CORE_API FMemoryView GetView() const;
 
 	/**
 	 * Try to get a view of the field as it would be serialized, such as by CopyTo.
@@ -595,7 +595,7 @@ public:
 	 * A serialized view is not available if the field has an externally-provided type.
 	 * Access the serialized form of such fields using CopyTo or FCbFieldRef::Clone.
 	 */
-	inline bool TryGetSerializedView(FConstMemoryView& OutView) const
+	inline bool TryGetSerializedView(FMemoryView& OutView) const
 	{
 		if (FCbFieldType::HasFieldType(Type))
 		{
@@ -607,10 +607,10 @@ public:
 
 protected:
 	/** Returns a view of the name and value payload, which excludes the type. */
-	CORE_API FConstMemoryView GetViewNoType() const;
+	CORE_API FMemoryView GetViewNoType() const;
 
 	/** Returns a view of the value payload, which excludes the type and name. */
-	inline FConstMemoryView GetPayloadView() const { return MakeMemoryView(Payload, GetPayloadSize()); }
+	inline FMemoryView GetPayloadView() const { return MakeMemoryView(Payload, GetPayloadSize()); }
 
 	/** Returns the type of the field including flags. */
 	constexpr inline ECbFieldType GetType() const { return Type; }
@@ -696,7 +696,7 @@ public:
 	 * @param View A buffer containing zero or more valid fields.
 	 * @param Type HasFieldType means that View contains the type. Otherwise, use the given type.
 	 */
-	static inline FCbFieldIterator MakeRange(FConstMemoryView View, ECbFieldType Type = ECbFieldType::HasFieldType)
+	static inline FCbFieldIterator MakeRange(FMemoryView View, ECbFieldType Type = ECbFieldType::HasFieldType)
 	{
 		return !View.IsEmpty() ? FCbFieldIterator(FCbField(View.GetData(), Type), View.GetDataEnd()) : FCbFieldIterator();
 	}
@@ -784,7 +784,7 @@ public:
 	 * A serialized view is not available if the underlying field has an externally-provided type or
 	 * a name. Access the serialized form of such arrays using CopyTo or FCbArrayRef::Clone.
 	 */
-	inline bool TryGetSerializedView(FConstMemoryView& OutView) const
+	inline bool TryGetSerializedView(FMemoryView& OutView) const
 	{
 		return !FCbField::HasName() && FCbField::TryGetSerializedView(OutView);
 	}
@@ -882,7 +882,7 @@ public:
 	 * A serialized view is not available if the underlying field has an externally-provided type or
 	 * a name. Access the serialized form of such objects using CopyTo or FCbObjectRef::Clone.
 	 */
-	inline bool TryGetSerializedView(FConstMemoryView& OutView) const
+	inline bool TryGetSerializedView(FMemoryView& OutView) const
 	{
 		return !FCbField::HasName() && FCbField::TryGetSerializedView(OutView);
 	}
@@ -977,7 +977,7 @@ public:
 	/** Returns a buffer that exactly contains this value. */
 	inline FSharedBufferConstPtr GetBuffer() const
 	{
-		const FConstMemoryView View = BaseType::GetView();
+		const FMemoryView View = BaseType::GetView();
 		const FSharedBufferConstPtr& OuterBuffer = GetOuterBuffer();
 		if (OuterBuffer)
 		{
diff --git a/Engine/Source/Runtime/Core/Public/Serialization/CompactBinarySerialization.h b/Engine/Source/Runtime/Core/Public/Serialization/CompactBinarySerialization.h
--- a/Engine/Source/Runtime/Core/Public/Serialization/CompactBinarySerialization.h
+++ b/Engine/Source/Runtime/Core/Public/Serialization/CompactBinarySerialization.h
@@ -24,7 +24,7 @@ class FArchive;
  * @param View A memory view that may contain the start of a field.
  * @param Type HasFieldType means that View contains the type. Otherwise, use the given type.
  */
-CORE_API uint64 MeasureCompactBinary(FConstMemoryView View, ECbFieldType Type = ECbFieldType::HasFieldType);
+CORE_API uint64 MeasureCompactBinary(FMemoryView View, ECbFieldType Type = ECbFieldType::HasFieldType);
 
 /**
  * Try to determine the type and size of the compact binary field at the start of the view.
@@ -43,7 +43,7 @@ CORE_API uint64 MeasureCompactBinary(FConstMemoryView View, ECbFieldType Type =
  * @return true if the size of the field was determined, otherwise false.
  */
 CORE_API bool TryMeasureCompactBinary(
-	FConstMemoryView InView,
+	FMemoryView InView,
 	ECbFieldType& OutType,
 	uint64& OutSize,
 	ECbFieldType InType = ECbFieldType::HasFieldType);
diff --git a/Engine/Source/Runtime/Core/Public/Serialization/CompactBinaryValidation.h b/Engine/Source/Runtime/Core/Public/Serialization/CompactBinaryValidation.h
--- a/Engine/Source/Runtime/Core/Public/Serialization/CompactBinaryValidation.h
+++ b/Engine/Source/Runtime/Core/Public/Serialization/CompactBinaryValidation.h
@@ -115,7 +115,7 @@ ENUM_CLASS_FLAGS(ECbValidateError);
  * @param Type HasFieldType means that View contains the type. Otherwise, use the given type.
  * @return None on success, otherwise the flags for the types of errors that were detected.
  */
-CORE_API ECbValidateError ValidateCompactBinary(FConstMemoryView View, ECbValidateMode Mode, ECbFieldType Type = ECbFieldType::HasFieldType);
+CORE_API ECbValidateError ValidateCompactBinary(FMemoryView View, ECbValidateMode Mode, ECbFieldType Type = ECbFieldType::HasFieldType);
 
 /**
  * Validate the compact binary data for every field in the view as specified by the mode flags.
@@ -126,6 +126,6 @@ CORE_API ECbValidateError ValidateCompactBinary(FConstMemoryView View, ECbValida
  *
  * \see \ref ValidateCompactBinary for more detail.
  */
-CORE_API ECbValidateError ValidateCompactBinaryRange(FConstMemoryView View, ECbValidateMode Mode);
+CORE_API ECbValidateError ValidateCompactBinaryRange(FMemoryView View, ECbValidateMode Mode);
 
 ///////////////////////////////////////////////////////////////////////////////////////////////////////////////////////
diff --git a/Engine/Source/Runtime/Core/Public/Serialization/CompactBinaryWriter.h b/Engine/Source/Runtime/Core/Public/Serialization/CompactBinaryWriter.h
--- a/Engine/Source/Runtime/Core/Public/Serialization/CompactBinaryWriter.h
+++ b/Engine/Source/Runtime/Core/Public/Serialization/CompactBinaryWriter.h
@@ -139,7 +139,7 @@ public:
 	/** Write a binary field by copying Size bytes from Value. */
 	CORE_API void Binary(const void* Value, uint64 Size);
 	/** Write a binary field by copying the view. */
-	inline void Binary(FConstMemoryView Value) { Binary(Value.GetData(), Value.GetSize()); }
+	inline void Binary(FMemoryView Value) { Binary(Value.GetData(), Value.GetSize()); }
 	/** Write a binary field by copying the buffer. Holds a reference if owned. */
 	CORE_API void Binary(const FSharedBufferConstPtr& Buffer);
 
