CompactBinary: Simplify base class assignment based on review feedback

#rnx

diff --git a/Engine/Source/Runtime/Core/Public/Serialization/CompactBinary.h b/Engine/Source/Runtime/Core/Public/Serialization/CompactBinary.h
--- a/Engine/Source/Runtime/Core/Public/Serialization/CompactBinary.h
+++ b/Engine/Source/Runtime/Core/Public/Serialization/CompactBinary.h
@@ -920,7 +920,7 @@ public:
 	{
 		if (ValueBuffer)
 		{
-			static_cast<BaseType&>(*this) = BaseType(ValueBuffer->GetData(), Type);
+			BaseType::operator=(BaseType(ValueBuffer->GetData(), Type));
 			check(ValueBuffer->GetView().Contains(BaseType::GetView()));
 			Buffer = MoveTemp(ValueBuffer);
 		}
@@ -954,7 +954,7 @@ public:
 		{
 			FSharedBufferPtr MutableBuffer = FSharedBuffer::Alloc(BaseType::GetSize());
 			BaseType::CopyTo(*MutableBuffer);
-			static_cast<BaseType&>(*this) = BaseType(MutableBuffer->GetData());
+			BaseType::operator=(BaseType(MutableBuffer->GetData()));
 			Buffer = FSharedBuffer::MakeReadOnly(MoveTemp(MutableBuffer));
 		}
 	}
@@ -969,7 +969,7 @@ public:
 		{
 			FSharedBufferPtr MutableBuffer = FSharedBuffer::Alloc(BaseType::GetSize());
 			BaseType::CopyTo(*MutableBuffer);
-			static_cast<BaseType&>(*this) = BaseType(MutableBuffer->GetData());
+			BaseType::operator=(BaseType(MutableBuffer->GetData()));
 			Buffer = FSharedBuffer::MakeReadOnly(MoveTemp(MutableBuffer));
 		}
 	}
