CompactBinary: Fixed Equals on iterators to compare end pointers and compare across types

#rnx

diff --git a/Engine/Source/Runtime/Core/Private/Tests/Serialization/CompactBinaryTest.cpp b/Engine/Source/Runtime/Core/Private/Tests/Serialization/CompactBinaryTest.cpp
--- a/Engine/Source/Runtime/Core/Private/Tests/Serialization/CompactBinaryTest.cpp
+++ b/Engine/Source/Runtime/Core/Private/Tests/Serialization/CompactBinaryTest.cpp
@@ -1576,6 +1576,24 @@ bool FCbFieldRefIteratorTest::RunTest(const FString& Parameters)
 		TestEqual(TEXT("FCbFieldRefIterator::MakeRange(Uniform).GetRangeView()"), UniformFieldRefIt.GetRangeView(), MakeMemoryView(UniformPayload));
 		TestFalse(TEXT("FCbFieldRefIterator::MakeRange(Uniform).TryGetSerializedRangeView()"), UniformFieldRefIt.TryGetSerializedRangeView(SerializedView));
 
+		// Equals
+		TestTrue(TEXT("FCbFieldIterator::Equals(Self)"), FieldIt.Equals(FieldIt));
+		TestTrue(TEXT("FCbFieldIterator::Equals(OtherType)"), FieldIt.Equals(FieldRefIt));
+		TestTrue(TEXT("FCbFieldRefIterator::Equals(Self)"), FieldRefIt.Equals(FieldRefIt));
+		TestTrue(TEXT("FCbFieldRefIterator::Equals(OtherType)"), FieldRefIt.Equals(FieldIt));
+		TestFalse(TEXT("FCbFieldIterator::Equals(OtherRange)"), FieldIt.Equals(FieldItClone));
+		TestFalse(TEXT("FCbFieldRefIterator::Equals(OtherRange)"), FieldRefIt.Equals(FieldRefItClone));
+		TestTrue(TEXT("FCbFieldIterator::Equals(Uniform, Self)"), UniformFieldIt.Equals(UniformFieldIt));
+		TestTrue(TEXT("FCbFieldIterator::Equals(Uniform, OtherType)"), UniformFieldIt.Equals(UniformFieldRefIt));
+		TestTrue(TEXT("FCbFieldRefIterator::Equals(Uniform, Self)"), UniformFieldRefIt.Equals(UniformFieldRefIt));
+		TestTrue(TEXT("FCbFieldRefIterator::Equals(Uniform, OtherType)"), UniformFieldRefIt.Equals(UniformFieldIt));
+		TestFalse(TEXT("FCbFieldIterator::Equals(SamePayload, DifferentEnd)"),
+			FCbFieldIterator::MakeRange(MakeMemoryView(UniformPayload), ECbFieldType::IntegerPositive)
+				.Equals(FCbFieldIterator::MakeRange(MakeMemoryView(UniformPayload).LeftChop(1), ECbFieldType::IntegerPositive)));
+		TestFalse(TEXT("FCbFieldIterator::Equals(DifferentPayload, SameEnd)"),
+			FCbFieldIterator::MakeRange(MakeMemoryView(UniformPayload), ECbFieldType::IntegerPositive)
+				.Equals(FCbFieldIterator::MakeRange(MakeMemoryView(UniformPayload).RightChop(1), ECbFieldType::IntegerPositive)));
+
 		// CopyRangeTo
 		uint8 CopyBytes[sizeof(Payload)];
 		FieldIt.CopyRangeTo(MakeMemoryView(CopyBytes));
diff --git a/Engine/Source/Runtime/Core/Public/Serialization/CompactBinary.h b/Engine/Source/Runtime/Core/Public/Serialization/CompactBinary.h
--- a/Engine/Source/Runtime/Core/Public/Serialization/CompactBinary.h
+++ b/Engine/Source/Runtime/Core/Public/Serialization/CompactBinary.h
@@ -308,7 +308,7 @@ public:
 	template <typename OtherFieldType>
 	constexpr inline bool Equals(const TCbFieldIterator<OtherFieldType>& Other) const
 	{
-		return FieldType::GetPayload() == Other.FieldType::GetPayload();
+		return FieldType::GetPayload() == Other.OtherFieldType::GetPayload() && FieldsEnd == Other.FieldsEnd;
 	}
 
 	template <typename OtherFieldType>
