Use the LowLevelError() in AssertMacros.h for UE_LOG(Fatal) in >= Test variants. Produces better code on the hot path because error handling is no longer inlined, and at reduces redundancy.

#rnx

diff --git a/Engine/Source/Runtime/Core/Public/Logging/LogMacros.h b/Engine/Source/Runtime/Core/Public/Logging/LogMacros.h
--- a/Engine/Source/Runtime/Core/Public/Logging/LogMacros.h
+++ b/Engine/Source/Runtime/Core/Public/Logging/LogMacros.h
@@ -103,12 +103,9 @@ private:
 	// This will only log Fatal errors
 	#define UE_LOG(CategoryName, Verbosity, Format, ...) \
 	{ \
-		static_assert(TIsArrayOrRefOfType<decltype(Format), TCHAR>::Value, "Formatting string must be a TCHAR array."); \
 		if (ELogVerbosity::Verbosity == ELogVerbosity::Fatal) \
 		{ \
-			LowLevelFatalErrorHandler(UE_LOG_SOURCE_FILE(__FILE__), __LINE__, Format, ##__VA_ARGS__); \
-			_DebugBreakAndPromptForRemote(); \
-			FDebug::ProcessFatalError(); \
+			LowLevelFatalError(Format, ##__VA_ARGS__); \
 			UE_LOG_EXPAND_IS_FATAL(Verbosity, CA_ASSUME(false);, PREPROCESSOR_NOTHING) \
 		} \
 	}
@@ -118,14 +115,11 @@ private:
 	// Conditional logging (fatal errors only).
 	#define UE_CLOG(Condition, CategoryName, Verbosity, Format, ...) \
 	{ \
-		static_assert(TIsArrayOrRefOfType<decltype(Format), TCHAR>::Value, "Formatting string must be a TCHAR array."); \
 		if (ELogVerbosity::Verbosity == ELogVerbosity::Fatal) \
 		{ \
 			if (Condition) \
 			{ \
-				LowLevelFatalErrorHandler(UE_LOG_SOURCE_FILE(__FILE__), __LINE__, Format, ##__VA_ARGS__); \
-				_DebugBreakAndPromptForRemote(); \
-				FDebug::ProcessFatalError(); \
+				LowLevelFatalError(Format, ##__VA_ARGS__); \
 				UE_LOG_EXPAND_IS_FATAL(Verbosity, CA_ASSUME(false);, PREPROCESSOR_NOTHING) \
 			} \
 		} \
