Added macros for getting return address and a pointer onto the stack where a return address is located.

#rnx

diff --git a/Engine/Source/Runtime/Core/Public/Clang/ClangPlatform.h b/Engine/Source/Runtime/Core/Public/Clang/ClangPlatform.h
--- a/Engine/Source/Runtime/Core/Public/Clang/ClangPlatform.h
+++ b/Engine/Source/Runtime/Core/Public/Clang/ClangPlatform.h
@@ -17,3 +17,6 @@
 #else
 	#define PLATFORM_COMPILER_HAS_FOLD_EXPRESSIONS 0
 #endif
+
+#define PLATFORM_RETURN_ADDRESS()			__builtin_return_address(0)
+#define PLATFORM_RETURN_ADDRESS_POINTER()	__builtin_frame_address(0)
diff --git a/Engine/Source/Runtime/Core/Public/MSVC/MSVCPlatform.h b/Engine/Source/Runtime/Core/Public/MSVC/MSVCPlatform.h
--- a/Engine/Source/Runtime/Core/Public/MSVC/MSVCPlatform.h
+++ b/Engine/Source/Runtime/Core/Public/MSVC/MSVCPlatform.h
@@ -17,3 +17,6 @@
 #else
 	#define PLATFORM_COMPILER_HAS_FOLD_EXPRESSIONS 0
 #endif
+
+#define PLATFORM_RETURN_ADDRESS()	        _ReturnAddress()
+#define PLATFORM_RETURN_ADDRESS_POINTER()	_AddressOfReturnAddress()
