Fixed FUnixSystemWideCriticalSection to always honor the timeout parameter passed at construction (because flock() could block indefinitely).

diff --git a/Engine/Source/Runtime/Core/Private/Unix/UnixCriticalSection.cpp b/Engine/Source/Runtime/Core/Private/Unix/UnixCriticalSection.cpp
--- a/Engine/Source/Runtime/Core/Private/Unix/UnixCriticalSection.cpp
+++ b/Engine/Source/Runtime/Core/Private/Unix/UnixCriticalSection.cpp
@@ -20,25 +20,37 @@ FUnixSystemWideCriticalSection::FUnixSystemWideCriticalSection(const FString& In
 	const FString LockPath = FString(FPlatformProcess::ApplicationSettingsDir()) / InName;
 	FString NormalizedFilepath(LockPath);
 	NormalizedFilepath.ReplaceInline(TEXT("\\"), TEXT("/"));
+	FileHandle = -1;
 
-	// Attempt to open a file and then lock with flock (NOTE: not an atomic operation, but best we can do)
-	FileHandle = open(TCHAR_TO_UTF8(*NormalizedFilepath), O_CREAT | O_WRONLY | O_NONBLOCK, S_IRUSR | S_IWUSR | S_IRGRP | S_IWGRP | S_IROTH | S_IWOTH);
-
-	if (FileHandle == -1 && InTimeout != FTimespan::Zero())
+	FDateTime ExpireTime = FDateTime::UtcNow() + InTimeout;
+	while (true)
 	{
-		FDateTime ExpireTime = FDateTime::UtcNow() + InTimeout;
-		const float RetrySeconds = FMath::Min((float)InTimeout.GetTotalSeconds(), 0.25f);
-
-		do
+		if (FileHandle == -1)
 		{
-			// retry until timeout
-			FPlatformProcess::Sleep(RetrySeconds);
+			// Try to open the file.
 			FileHandle = open(TCHAR_TO_UTF8(*NormalizedFilepath), O_CREAT | O_WRONLY | O_NONBLOCK, S_IRUSR | S_IWUSR | S_IRGRP | S_IWGRP | S_IROTH | S_IWOTH);
-		} while (FileHandle == -1 && FDateTime::UtcNow() < ExpireTime);
-	}
-	if (FileHandle != -1)
-	{
-		flock(FileHandle, LOCK_EX);
+		}
+
+		// If the file is open, try to lock it, but don't block. If the file is already locked by another process or another thread, blocking here may not honor InTimeout.
+		if (FileHandle != -1 && flock(FileHandle, LOCK_EX | LOCK_NB) == 0)
+		{
+			return; // Lock was successfully taken.
+		}
+
+		// If the lock isn't acquired and no time is left to retry, clean up and set the state as 'invalid'
+		if (InTimeout != FTimespan::Zero() || FDateTime::UtcNow() < ExpireTime)
+		{
+			if (FileHandle != -1)
+			{
+				close(FileHandle);
+				FileHandle = -1;
+			}
+			return; // Lock wasn't acquired within the allowed time.
+		}
+
+		// Either the file did not open or the lock wasn't acquired, retry.
+		const float RetrySeconds = FMath::Min((float)InTimeout.GetTotalSeconds(), 0.25f);
+		FPlatformProcess::Sleep(RetrySeconds);
 	}
 }
 
