Wrapped location and expression information about checks, ensures, and fatal errors in a struct.

#rnx

diff --git a/Engine/Source/Runtime/Core/Private/Logging/LogMacros.cpp b/Engine/Source/Runtime/Core/Private/Logging/LogMacros.cpp
--- a/Engine/Source/Runtime/Core/Private/Logging/LogMacros.cpp
+++ b/Engine/Source/Runtime/Core/Private/Logging/LogMacros.cpp
@@ -9,7 +9,7 @@
 #include "Stats/Stats.h"
 #include "ProfilingDebugging/CsvProfiler.h"
 
-void StaticFailDebug( const TCHAR* Error, const ANSICHAR* File, int32 Line, const TCHAR* Description, bool bIsEnsure, int32 NumStackFramesToIgnore );
+void StaticFailDebug( const TCHAR* Error, const FDebug::FFailureInfo& Info, const TCHAR* Description, bool bIsEnsure, int32 NumStackFramesToIgnore );
 
 /** Statics to prevent FMsg::Logf from allocating too much stack memory. */
 static FCriticalSection					MsgLogfStaticBufferGuard;
@@ -62,7 +62,8 @@ void FMsg::LogfImpl(const ANSICHAR* File, int32 Line, const FLogCategoryName& Ca
 		}
 
 		const int32 NumStackFramesToIgnore = 1;
-		StaticFailDebug(TEXT("Fatal error:"), File, Line, Message, false, NumStackFramesToIgnore);
+		FDebug::FFailureInfo Info = { "FMsg::Logf(Fatal)", File, Line };
+		StaticFailDebug(TEXT("Fatal error:"), Info, Message, false, NumStackFramesToIgnore);
 		FDebug::AssertFailed("", File, Line, Message);
 	}
 #endif
@@ -109,7 +110,8 @@ void FMsg::Logf_InternalImpl(const ANSICHAR* File, int32 Line, const FLogCategor
 		}
 
 		const int32 NumStackFramesToIgnore = 1;
-		StaticFailDebug(TEXT("Fatal error:"), File, Line, Message, false, NumStackFramesToIgnore);
+		FDebug::FFailureInfo Info = { "FMsg::Logf(Fatal)", File, Line };
+		StaticFailDebug(TEXT("Fatal error:"), Info, Message, false, NumStackFramesToIgnore);
 	}
 #endif
 }
diff --git a/Engine/Source/Runtime/Core/Private/Misc/AssertionMacros.cpp b/Engine/Source/Runtime/Core/Private/Misc/AssertionMacros.cpp
--- a/Engine/Source/Runtime/Core/Private/Misc/AssertionMacros.cpp
+++ b/Engine/Source/Runtime/Core/Private/Misc/AssertionMacros.cpp
@@ -113,7 +113,7 @@ void PrintScriptCallstack()
 	}
 }
 
-static void AssertFailedImplV(const ANSICHAR* Expr, const ANSICHAR* File, int32 Line, const TCHAR* Format, va_list Args)
+static void AssertFailedImplV(const FDebug::FFailureInfo& Info, const TCHAR* Format, va_list Args)
 {
 	FTempCommandLineScope TempCommandLine;
 
@@ -127,8 +127,8 @@ static void AssertFailedImplV(const ANSICHAR* Expr, const ANSICHAR* File, int32
 	FCString::GetVarArgs(DescriptionString, UE_ARRAY_COUNT(DescriptionString), Format, Args);
 
 	TCHAR ErrorString[MAX_SPRINTF];
-	FCString::Sprintf(ErrorString, TEXT("%s"), ANSI_TO_TCHAR(Expr));
-	GError->Logf(TEXT("Assertion failed: %s") FILE_LINE_DESC TEXT("\n%s\n"), ErrorString, ANSI_TO_TCHAR(File), Line, DescriptionString);
+	FCString::Sprintf(ErrorString, TEXT("%s"), ANSI_TO_TCHAR(Info.Expr));
+	GError->Logf(TEXT("Assertion failed: %s") FILE_LINE_DESC TEXT("\n%s\n"), ErrorString, ANSI_TO_TCHAR(Info.File), Info.Line, DescriptionString);
 }
 
 /**
@@ -136,8 +136,11 @@ static void AssertFailedImplV(const ANSICHAR* Expr, const ANSICHAR* File, int32
  *	prompts for the remote debugging if there is not debugger, breaks into the debugger 
  *	and copies the error into the global error message.
  */
-FORCENOINLINE void StaticFailDebug( const TCHAR* Error, const ANSICHAR* File, int32 Line, const TCHAR* Description, bool bIsEnsure, int NumStackFramesToIgnore )
+FORCENOINLINE void StaticFailDebug( const TCHAR* Error, const FDebug::FFailureInfo& Info, const TCHAR* Description, bool bIsEnsure, int NumStackFramesToIgnore )
 {
+	const ANSICHAR* File = Info.File;
+	int32 Line = Info.Line;
+
 	// Print out the blueprint callstack
 	PrintScriptCallstack();
 
@@ -268,15 +271,15 @@ void FDebug::LogFormattedMessageWithCallstack(const FName& InLogName, const ANSI
 //warning: May be called at library startup time.
 //
 
-FORCENOINLINE void FDebug::LogAssertFailedMessageImpl(const ANSICHAR* Expr, const ANSICHAR* File, int32 Line, const TCHAR* Fmt, ...)
+FORCENOINLINE void FDebug::LogAssertFailedMessageImpl(const FFailureInfo& Info, const TCHAR* Fmt, ...)
 {
 	va_list Args;
 	va_start(Args, Fmt);
-	LogAssertFailedMessageImplV(Expr, File, Line, Fmt, Args);
+	LogAssertFailedMessageImplV(Info, Fmt, Args);
 	va_end(Args);
 }
 
-void FDebug::LogAssertFailedMessageImplV(const ANSICHAR* Expr, const ANSICHAR* File, int32 Line, const TCHAR* Fmt, va_list Args)
+void FDebug::LogAssertFailedMessageImplV(const FFailureInfo& Info, const TCHAR* Fmt, va_list Args)
 {
 	// Ignore this assert if we're already forcibly shutting down because of a critical error.
 	if( !GIsCriticalError )
@@ -285,10 +288,10 @@ void FDebug::LogAssertFailedMessageImplV(const ANSICHAR* Expr, const ANSICHAR* F
 		FCString::GetVarArgs( DescriptionString, UE_ARRAY_COUNT(DescriptionString), Fmt, Args );
 
 		TCHAR ErrorString[MAX_SPRINTF];
-		FCString::Sprintf( ErrorString, TEXT( "Assertion failed: %s" ), ANSI_TO_TCHAR( Expr ) );
+		FCString::Sprintf( ErrorString, TEXT( "Assertion failed: %s" ), ANSI_TO_TCHAR( Info.Expr ) );
 
 		const int32 NumStackFramesToIgnore = 1;
-		StaticFailDebug( ErrorString, File, Line, DescriptionString, false, NumStackFramesToIgnore );
+		StaticFailDebug( ErrorString, Info, DescriptionString, false, NumStackFramesToIgnore );
 	}
 }
 
@@ -301,8 +304,12 @@ void FDebug::LogAssertFailedMessageImplV(const ANSICHAR* Expr, const ANSICHAR* F
  * @param	Msg		Informative error message text
  * @param	NumStackFramesToIgnore	Number of stack frames to ignore in the callstack
  */
-FORCENOINLINE void FDebug::EnsureFailed(const ANSICHAR* Expr, const ANSICHAR* File, int32 Line, const TCHAR* Msg, int NumStackFramesToIgnore)
+FORCENOINLINE void FDebug::EnsureFailed(const FFailureInfo& Info, const TCHAR* Msg, int NumStackFramesToIgnore)
 {
+	const ANSICHAR* Expr = Info.Expr;
+	const ANSICHAR* File = Info.File;
+	int32 Line = Info.Line;
+
 	FTempCommandLineScope TempCommandLine;
 
 	// if time isn't ready yet, we better not continue
@@ -323,7 +330,7 @@ FORCENOINLINE void FDebug::EnsureFailed(const ANSICHAR* Expr, const ANSICHAR* Fi
 	if( bShouldCrash )
 	{
 		// Just trigger a regular assertion which will crash via GError->Logf()
-		FDebug::LogAssertFailedMessageImpl( Expr, File, Line, TEXT("%s"), Msg );
+		FDebug::LogAssertFailedMessageImpl( Info, TEXT("%s"), Msg );
 		return;
 	}
 
@@ -332,9 +339,9 @@ FORCENOINLINE void FDebug::EnsureFailed(const ANSICHAR* Expr, const ANSICHAR* Fi
 
 	// Print initial debug message for this error
 	TCHAR ErrorString[MAX_SPRINTF];
-	FCString::Sprintf(ErrorString,TEXT("Ensure condition failed: %s"),ANSI_TO_TCHAR(Expr));
+	FCString::Sprintf(ErrorString,TEXT("Ensure condition failed: %s"),ANSI_TO_TCHAR(Info.Expr));
 
-	StaticFailDebug( ErrorString, File, Line, Msg, true, NumStackFramesToIgnore + 1 );
+	StaticFailDebug( ErrorString, Info, Msg, true, NumStackFramesToIgnore + 1 );
 
 	// Is there a debugger attached?  If not we'll submit an error report.
 	if (FPlatformMisc::IsDebuggerPresent() && !GAlwaysReportCrash)
@@ -466,16 +473,14 @@ FORCENOINLINE void FDebug::EnsureFailed(const ANSICHAR* Expr, const ANSICHAR* Fi
 }
 
 void FORCENOINLINE FDebug::CheckVerifyFailedImpl(
-	const ANSICHAR* Expr,
-	const ANSICHAR* File,
-	const int Line,
+	const FFailureInfo& Info,
 	const TCHAR* Format,
 	...)
 {
 	va_list Args;
 
 	va_start(Args, Format);
-	FDebug::LogAssertFailedMessageImplV(Expr, File, Line, Format, Args);
+	FDebug::LogAssertFailedMessageImplV(Info, Format, Args);
 	va_end(Args);
 
 	if (!FPlatformMisc::IsDebuggerPresent())
@@ -483,7 +488,7 @@ void FORCENOINLINE FDebug::CheckVerifyFailedImpl(
 		FPlatformMisc::PromptForRemoteDebugging(false);
 
 		va_start(Args, Format);
-		AssertFailedImplV(Expr, File, Line, Format, Args);
+		AssertFailedImplV(Info, Format, Args);
 		va_end(Args);
 	}
 }
@@ -494,7 +499,8 @@ void VARARGS FDebug::AssertFailed(const ANSICHAR* Expr, const ANSICHAR* File, in
 {
 	va_list Args;
 	va_start(Args, Format);
-	AssertFailedImplV(Expr, File, Line, Format, Args);
+	FDebug::FFailureInfo Info = { Expr, File, Line };
+	AssertFailedImplV(Info, Format, Args);
 	va_end(Args);
 }
 
@@ -510,7 +516,7 @@ void FDebug::ProcessFatalError()
 }
 
 #if DO_CHECK || DO_GUARD_SLOW || DO_ENSURE
-FORCENOINLINE bool VARARGS FDebug::OptionallyLogFormattedEnsureMessageReturningFalseImpl( bool bLog, const ANSICHAR* Expr, const ANSICHAR* File, int32 Line, const TCHAR* FormattedMsg, ... )
+FORCENOINLINE bool VARARGS FDebug::OptionallyLogFormattedEnsureMessageReturningFalseImpl( bool bLog, const FFailureInfo& Info, const TCHAR* FormattedMsg, ... )
 {
 	if (bLog)
 	{
@@ -519,20 +525,20 @@ FORCENOINLINE bool VARARGS FDebug::OptionallyLogFormattedEnsureMessageReturningF
 		GET_VARARGS( TempStr, TempStrSize, TempStrSize - 1, FormattedMsg, FormattedMsg );
 
 		const int32 NumStackFramesToIgnore = 1; // Just ignore this frame
-		EnsureFailed( Expr, File, Line, TempStr, NumStackFramesToIgnore );
+		EnsureFailed( Info, TempStr, NumStackFramesToIgnore );
 	}
 	
 	return false;
 }
 #endif
 
-FORCENOINLINE void VARARGS LowLevelFatalErrorHandler(const ANSICHAR* File, int32 Line, const TCHAR* Format, ...)
+FORCENOINLINE void VARARGS LowLevelFatalErrorHandler(const FDebug::FFailureInfo& Info, const TCHAR* Format, ...)
 {
 	TCHAR DescriptionString[4096];
 	GET_VARARGS( DescriptionString, UE_ARRAY_COUNT(DescriptionString), UE_ARRAY_COUNT(DescriptionString)-1, Format, Format );
 
 	const int32 NumStackFramesToIgnore = 1; // Just ignore this frame
-	StaticFailDebug(TEXT("LowLevelFatalError"), File, Line, DescriptionString, false, NumStackFramesToIgnore);
+	StaticFailDebug(TEXT("LowLevelFatalError"), Info, DescriptionString, false, NumStackFramesToIgnore);
 }
 
 void FDebug::DumpStackTraceToLog(const ELogVerbosity::Type LogVerbosity)
diff --git a/Engine/Source/Runtime/Core/Public/Misc/AssertionMacros.h b/Engine/Source/Runtime/Core/Public/Misc/AssertionMacros.h
--- a/Engine/Source/Runtime/Core/Public/Misc/AssertionMacros.h
+++ b/Engine/Source/Runtime/Core/Public/Misc/AssertionMacros.h
@@ -41,6 +41,16 @@ extern "C" CORE_API void PrintScriptCallstack();
  **/
 struct CORE_API FDebug
 {
+	/**
+	 * Contextual information for a point of failure.
+	 */
+	struct FFailureInfo
+	{
+		const ANSICHAR* Expr;
+		const ANSICHAR* File;
+		int32 Line;
+	};
+
 	/** Logs final assert message and exits the program. */
 	static void VARARGS AssertFailed(const ANSICHAR* Expr, const ANSICHAR* File, int32 Line, const TCHAR* Format = TEXT(""), ...);
 
@@ -64,16 +74,16 @@ struct CORE_API FDebug
 
 #if DO_CHECK || DO_GUARD_SLOW || DO_ENSURE
 private:
-	static void VARARGS CheckVerifyFailedImpl(const ANSICHAR* Expr, const char* File, int32 Line, const TCHAR* Format, ...);
-	static void VARARGS LogAssertFailedMessageImpl(const ANSICHAR* Expr, const ANSICHAR* File, int32 Line, const TCHAR* Fmt, ...);
-	static void LogAssertFailedMessageImplV(const ANSICHAR* Expr, const ANSICHAR* File, int32 Line, const TCHAR* Fmt, va_list Args);
+	static void VARARGS CheckVerifyFailedImpl(const FFailureInfo& Info, const TCHAR* Format, ...);
+	static void VARARGS LogAssertFailedMessageImpl(const FFailureInfo& Info, const TCHAR* Fmt, ...);
+	static void LogAssertFailedMessageImplV(const FFailureInfo& Info, const TCHAR* Fmt, va_list Args);
 
 public:
 	/**
 	 * Called when a 'check/verify' assertion fails.
 	 */
 	template <typename FmtType, typename... Types>
-	static void UE_DEBUG_SECTION CheckVerifyFailed(const ANSICHAR* Expr, const char* File, int32 Line, const FmtType& Format, Types... Args);
+	static void UE_DEBUG_SECTION CheckVerifyFailed(const FFailureInfo& Info, const FmtType& Format, Types... Args);
 	
 	/**
 	 * Called when an 'ensure' assertion fails; gathers stack data and generates and error report.
@@ -86,10 +96,10 @@ public:
 	 * 
 	 * Don't change the name of this function, it's used to detect ensures by the crash reporter.
 	 */
-	static void EnsureFailed( const ANSICHAR* Expr, const ANSICHAR* File, int32 Line, const TCHAR* Msg, int NumStackFramesToIgnore );
+	static void EnsureFailed( const FFailureInfo& Info, const TCHAR* Msg, int NumStackFramesToIgnore );
 
 private:
-	static bool VARARGS OptionallyLogFormattedEnsureMessageReturningFalseImpl(bool bLog, const ANSICHAR* Expr, const ANSICHAR* File, int32 Line, const TCHAR* FormattedMsg, ...);
+	static bool VARARGS OptionallyLogFormattedEnsureMessageReturningFalseImpl(bool bLog, const FFailureInfo& Info, const TCHAR* FormattedMsg, ...);
 
 public:
 	/**
@@ -108,12 +118,12 @@ public:
 
 	/** Failed assertion handler.  Warning: May be called at library startup time. */
 	template <typename FmtType, typename... Types>
-	static FORCEINLINE typename TEnableIf<TIsArrayOrRefOfType<FmtType, TCHAR>::Value, bool>::Type OptionallyLogFormattedEnsureMessageReturningFalse(bool bLog, const ANSICHAR* Expr, const ANSICHAR* File, int32 Line, const FmtType& FormattedMsg, Types... Args)
+	static FORCEINLINE typename TEnableIf<TIsArrayOrRefOfType<FmtType, TCHAR>::Value, bool>::Type OptionallyLogFormattedEnsureMessageReturningFalse(bool bLog, const FFailureInfo& Info, const FmtType& FormattedMsg, Types... Args)
 	{
 		static_assert(TIsArrayOrRefOfType<FmtType, TCHAR>::Value, "Formatting string must be a TCHAR array.");
 		static_assert(TAnd<TIsValidVariadicFunctionArg<Types>...>::Value, "Invalid argument(s) passed to ensureMsgf");
 
-		return OptionallyLogFormattedEnsureMessageReturningFalseImpl(bLog, Expr, File, Line, FormattedMsg, Args...);
+		return OptionallyLogFormattedEnsureMessageReturningFalseImpl(bLog, Info, FormattedMsg, Args...);
 	}
 
 #endif // DO_CHECK || DO_GUARD_SLOW
@@ -144,27 +154,25 @@ public:
 #if DO_CHECK || DO_GUARD_SLOW || DO_ENSURE
 	template <typename FmtType, typename... Types>
 	void FORCENOINLINE UE_DEBUG_SECTION FDebug::CheckVerifyFailed(
-		const ANSICHAR* Expr,
-		const ANSICHAR* File,
-		const int Line,
+		const FFailureInfo& Info,
 		const FmtType& Format,
 		Types... Args)
 	{
 		static_assert(TIsArrayOrRefOfType<FmtType, TCHAR>::Value, "Formatting string must be a TCHAR array.");
 		static_assert(TAnd<TIsValidVariadicFunctionArg<Types>...>::Value, "Invalid argument(s) passed to CheckVerifyFailed()");
-		return CheckVerifyFailedImpl(Expr, File, Line, Format, Args...);
-	}
-
-	// MSVC (v19.00.24215.1 at time of writing) ignores no-inline attributes on
-	// lambdas. This can be worked around by calling the lambda from inside this
-	// templated (and correctly non-inlined) function.
-	template <typename RetType=void, class InnerType>
-	RetType FORCENOINLINE UE_DEBUG_SECTION DispatchCheckVerify(InnerType&& Inner)
-	{
-		return Inner();
+		return CheckVerifyFailedImpl(Info, Format, Args...);
 	}
 #endif
 
+// MSVC (v19.00.24215.1 at time of writing) ignores no-inline attributes on
+// lambdas. This can be worked around by calling the lambda from inside this
+// templated (and correctly non-inlined) function.
+template <typename RetType=void, class InnerType>
+RetType FORCENOINLINE UE_DEBUG_SECTION DispatchCheckVerify(InnerType&& Inner)
+{
+	return Inner();
+}
+
 #if !UE_BUILD_SHIPPING
 #define _DebugBreakAndPromptForRemote() \
 	if (!FPlatformMisc::IsDebuggerPresent()) { FPlatformMisc::PromptForRemoteDebugging(false); } UE_DEBUG_BREAK();
@@ -195,7 +203,8 @@ public:
 				{ \
 					static void FORCENOINLINE UE_DEBUG_SECTION ExecCheckImplInternal() \
 					{ \
-						FDebug::CheckVerifyFailed(#expr, __FILE__, __LINE__, TEXT("")); \
+						FDebug::FFailureInfo Info = { #expr, __FILE__, __LINE__ }; \
+						FDebug::CheckVerifyFailed(Info, TEXT("")); \
 					} \
 				}; \
 				Impl::ExecCheckImplInternal(); \
@@ -221,7 +230,8 @@ public:
 			{ \
 				DispatchCheckVerify([&] () FORCENOINLINE UE_DEBUG_SECTION \
 				{ \
-					FDebug::CheckVerifyFailed(#expr, __FILE__, __LINE__, format, ##__VA_ARGS__); \
+					FDebug::FFailureInfo Info = { #expr, __FILE__, __LINE__ }; \
+					FDebug::CheckVerifyFailed(Info, format, ##__VA_ARGS__); \
 				}); \
 				PLATFORM_BREAK(); \
 				CA_ASSUME(false); \
@@ -326,7 +336,8 @@ public:
 			if ((!bExecuted || Always) && FPlatformMisc::IsEnsureAllowed()) \
 			{ \
 				bExecuted = true; \
-				FDebug::OptionallyLogFormattedEnsureMessageReturningFalse(true, #InExpression, __FILE__, __LINE__, ##__VA_ARGS__); \
+				FDebug::FFailureInfo Info = { #InExpression, __FILE__, __LINE__ }; \
+				FDebug::OptionallyLogFormattedEnsureMessageReturningFalse(true, Info, ##__VA_ARGS__); \
 				if (!FPlatformMisc::IsDebuggerPresent()) \
 				{ \
 					FPlatformMisc::PromptForRemoteDebugging(true); \
@@ -385,13 +396,17 @@ namespace UE4Asserts_Private
 ----------------------------------------------------------------------------*/
 
 /** low level fatal error handler. */
-CORE_API void VARARGS LowLevelFatalErrorHandler(const ANSICHAR* File, int32 Line, const TCHAR* Format=TEXT(""), ... );
+CORE_API void VARARGS LowLevelFatalErrorHandler(const FDebug::FFailureInfo& Info, const TCHAR* Format=TEXT(""), ... );
 
 #define LowLevelFatalError(Format, ...) \
 	{ \
 		static_assert(TIsArrayOrRefOfType<decltype(Format), TCHAR>::Value, "Formatting string must be a TCHAR array."); \
-		LowLevelFatalErrorHandler(__FILE__, __LINE__, Format, ##__VA_ARGS__); \
-		_DebugBreakAndPromptForRemote(); \
-		FDebug::ProcessFatalError(); \
+		DispatchCheckVerify([&] () FORCENOINLINE UE_DEBUG_SECTION \
+		{ \
+			FDebug::FFailureInfo Info = { "LowLevelFatalError", __FILE__, __LINE__ }; \
+			LowLevelFatalErrorHandler(Info, Format, ##__VA_ARGS__); \
+			_DebugBreakAndPromptForRemote(); \
+			FDebug::ProcessFatalError(); \
+		}); \
 	}
 
