The name "Info" was a little to generic in hindsight and inevitably collided. Renamed to mitigate future collisions.

#rnx

diff --git a/Engine/Source/Runtime/Core/Public/Misc/AssertionMacros.h b/Engine/Source/Runtime/Core/Public/Misc/AssertionMacros.h
--- a/Engine/Source/Runtime/Core/Public/Misc/AssertionMacros.h
+++ b/Engine/Source/Runtime/Core/Public/Misc/AssertionMacros.h
@@ -203,8 +203,8 @@ RetType FORCENOINLINE UE_DEBUG_SECTION DispatchCheckVerify(InnerType&& Inner)
 				{ \
 					static void FORCENOINLINE UE_DEBUG_SECTION ExecCheckImplInternal() \
 					{ \
-						FDebug::FFailureInfo Info = { #expr, __FILE__, __LINE__ }; \
-						FDebug::CheckVerifyFailed(Info, TEXT("")); \
+						FDebug::FFailureInfo CheckFailureInfo = { #expr, __FILE__, __LINE__ }; \
+						FDebug::CheckVerifyFailed(CheckFailureInfo, TEXT("")); \
 					} \
 				}; \
 				Impl::ExecCheckImplInternal(); \
@@ -230,8 +230,8 @@ RetType FORCENOINLINE UE_DEBUG_SECTION DispatchCheckVerify(InnerType&& Inner)
 			{ \
 				DispatchCheckVerify([&] () FORCENOINLINE UE_DEBUG_SECTION \
 				{ \
-					FDebug::FFailureInfo Info = { #expr, __FILE__, __LINE__ }; \
-					FDebug::CheckVerifyFailed(Info, format, ##__VA_ARGS__); \
+					FDebug::FFailureInfo CheckFailureInfo = { #expr, __FILE__, __LINE__ }; \
+					FDebug::CheckVerifyFailed(CheckFailureInfo, format, ##__VA_ARGS__); \
 				}); \
 				PLATFORM_BREAK(); \
 				CA_ASSUME(false); \
@@ -336,8 +336,8 @@ RetType FORCENOINLINE UE_DEBUG_SECTION DispatchCheckVerify(InnerType&& Inner)
 			if ((!bExecuted || Always) && FPlatformMisc::IsEnsureAllowed()) \
 			{ \
 				bExecuted = true; \
-				FDebug::FFailureInfo Info = { #InExpression, __FILE__, __LINE__ }; \
-				FDebug::OptionallyLogFormattedEnsureMessageReturningFalse(true, Info, ##__VA_ARGS__); \
+				FDebug::FFailureInfo EnsureFailureInfo = { #InExpression, __FILE__, __LINE__ }; \
+				FDebug::OptionallyLogFormattedEnsureMessageReturningFalse(true, EnsureFailureInfo, ##__VA_ARGS__); \
 				if (!FPlatformMisc::IsDebuggerPresent()) \
 				{ \
 					FPlatformMisc::PromptForRemoteDebugging(true); \
@@ -403,8 +403,8 @@ CORE_API void VARARGS LowLevelFatalErrorHandler(const FDebug::FFailureInfo& Info
 		static_assert(TIsArrayOrRefOfType<decltype(Format), TCHAR>::Value, "Formatting string must be a TCHAR array."); \
 		DispatchCheckVerify([&] () FORCENOINLINE UE_DEBUG_SECTION \
 		{ \
-			FDebug::FFailureInfo Info = { "LowLevelFatalError", __FILE__, __LINE__ }; \
-			LowLevelFatalErrorHandler(Info, Format, ##__VA_ARGS__); \
+			FDebug::FFailureInfo LlfeInfo = { "LowLevelFatalError", __FILE__, __LINE__ }; \
+			LowLevelFatalErrorHandler(LlfeInfo, Format, ##__VA_ARGS__); \
 			_DebugBreakAndPromptForRemote(); \
 			FDebug::ProcessFatalError(); \
 		}); \
