Workaround for a precompiled lib that link is linking to symbols in Core whose signatures have changed.

#rnx

diff --git a/Engine/Source/Runtime/Core/Public/Misc/AssertionMacros.h b/Engine/Source/Runtime/Core/Public/Misc/AssertionMacros.h
--- a/Engine/Source/Runtime/Core/Public/Misc/AssertionMacros.h
+++ b/Engine/Source/Runtime/Core/Public/Misc/AssertionMacros.h
@@ -140,6 +140,11 @@ public:
 	*
 	*/
 	static void LogFormattedMessageWithCallstack(const FName& LogName, const ANSICHAR* File, int32 Line, const TCHAR* Heading, const TCHAR* Message, ELogVerbosity::Type Verbosity);
+
+private:
+	/* Temporary workaround for link errors in the precompiled cadkernel.lib */
+	static void VARARGS CheckVerifyFailedImpl(const ANSICHAR*, const ANSICHAR*, int, const TCHAR*, ...) {}
+	static bool VARARGS OptionallyLogFormattedEnsureMessageReturningFalseImpl(bool, const ANSICHAR*, const ANSICHAR*, int, const TCHAR*, ...) { return false; }
 };
 
 /*----------------------------------------------------------------------------
