Added "Get" verb prefix to a few methods following feedback.

#rnx

diff --git a/Engine/Source/Runtime/Core/Private/Memory/MemoryArena.cpp b/Engine/Source/Runtime/Core/Private/Memory/MemoryArena.cpp
--- a/Engine/Source/Runtime/Core/Private/Memory/MemoryArena.cpp
+++ b/Engine/Source/Runtime/Core/Private/Memory/MemoryArena.cpp
@@ -134,7 +134,7 @@ static void FreeArenaId(uint16 ArenaId, FMemoryArena* Arena)
 
 FMemoryArena& FArenaPointer::Arena() const
 {
-	const uint16 Index = ArenaIndex();
+	const uint16 Index = GetArenaIndex();
 
 	return *GKnownArenas[Index].Arena;
 }
@@ -167,9 +167,9 @@ UE_NOALIAS void FMemoryArena::Free(const void* MemoryBlock)
 	return InternalFree(MemoryBlock, 0);
 }
 
-SIZE_T FMemoryArena::BlockSize(const void* MemoryBlock) const
+SIZE_T FMemoryArena::GetBlockSize(const void* MemoryBlock) const
 {
-	return InternalBlockSize(MemoryBlock);
+	return InternalGetBlockSize(MemoryBlock);
 }
 
 const TCHAR* FMemoryArena::GetDebugName() const
@@ -216,7 +216,7 @@ FArenaPointer ArenaRealloc(FMemoryArena* Arena, void* InPtr, SIZE_T OldSize, SIZ
 
 FArenaPointer ArenaRealloc(FArenaPointer InPtr, SIZE_T OldSize, SIZE_T NewSize, SIZE_T Alignment)
 {
-	return ArenaRealloc(&InPtr.Arena(), InPtr.Pointer(), OldSize, NewSize, Alignment);
+	return ArenaRealloc(&InPtr.Arena(), InPtr.GetPointer(), OldSize, NewSize, Alignment);
 }
 
 //////////////////////////////////////////////////////////////////////////
@@ -243,7 +243,7 @@ void FHeapArena::InternalFree(const void* MemoryBlock, SIZE_T MemoryBlockSize)
 	return mspace_free(HeapHandle, (void*) MemoryBlock);
 }
 
-SIZE_T FHeapArena::InternalBlockSize(const void* MemoryBlock) const
+SIZE_T FHeapArena::InternalGetBlockSize(const void* MemoryBlock) const
 {
 	return mspace_usable_size(MemoryBlock);
 }
@@ -275,7 +275,7 @@ void FMallocArena::InternalFree(const void* MemoryBlock, SIZE_T MemoryBlockSize)
 	FMemory::Free(const_cast<void*>(MemoryBlock));
 }
 
-SIZE_T FMallocArena::InternalBlockSize(const void* MemoryBlock) const
+SIZE_T FMallocArena::InternalGetBlockSize(const void* MemoryBlock) const
 {
 	return FMemory::GetAllocSize(const_cast<void*>(MemoryBlock));
 }
@@ -302,7 +302,7 @@ void FAnsiArena::InternalFree(const void* MemoryBlock, SIZE_T MemoryBlockSize)
 	return GAnsiMalloc.Free(const_cast<void*>(MemoryBlock));
 }
 
-SIZE_T FAnsiArena::InternalBlockSize(const void* MemoryBlock) const
+SIZE_T FAnsiArena::InternalGetBlockSize(const void* MemoryBlock) const
 {
 	SIZE_T Size = 0;
 	GAnsiMalloc.GetAllocationSize(const_cast<void*>(MemoryBlock), /* out */ Size);
diff --git a/Engine/Source/Runtime/Core/Public/Memory/MemoryArena.h b/Engine/Source/Runtime/Core/Public/Memory/MemoryArena.h
--- a/Engine/Source/Runtime/Core/Public/Memory/MemoryArena.h
+++ b/Engine/Source/Runtime/Core/Public/Memory/MemoryArena.h
@@ -51,9 +51,9 @@ private:
 	uint16	ArenaTag = 0;
 
 public:
-	FORCEINLINE const uint16	ArenaIndex() const									{ return ArenaTag; }
-	FORCEINLINE void*			Pointer()											{ return Ptr; }
-	FORCEINLINE const void*		Pointer() const										{ return Ptr; }
+	FORCEINLINE const uint16	GetArenaIndex() const									{ return ArenaTag; }
+	FORCEINLINE void*			GetPointer()											{ return Ptr; }
+	FORCEINLINE const void*		GetPointer() const										{ return Ptr; }
 	FORCEINLINE void			SetPointerAndArena(void* InPtr, uint16 InArenaTag)	{ Ptr = InPtr; ArenaTag = InArenaTag; }
 
 	inline FArenaPointer(void* InPtr, uint16 InArenaTag) : Ptr(InPtr), ArenaTag(InArenaTag) {}
@@ -68,15 +68,15 @@ private:
 	void*	TaggedPointer = nullptr;
 
 public:
-	FORCEINLINE const uint16	ArenaIndex() const	{ return uint16(UPTRINT(TaggedPointer) >> ArenaShift); }
-	FORCEINLINE void*			Pointer()			{ return reinterpret_cast<void*>(UPTRINT(TaggedPointer) & PointerMask); }
-	FORCEINLINE const void*		Pointer() const		{ return reinterpret_cast<const void*>(UPTRINT(TaggedPointer) & PointerMask); }
+	FORCEINLINE const uint16	GetArenaIndex() const	{ return uint16(UPTRINT(TaggedPointer) >> ArenaShift); }
+	FORCEINLINE void*			GetPointer()			{ return reinterpret_cast<void*>(UPTRINT(TaggedPointer) & PointerMask); }
+	FORCEINLINE const void*		GetPointer() const		{ return reinterpret_cast<const void*>(UPTRINT(TaggedPointer) & PointerMask); }
 	FORCEINLINE void			SetPointerAndArena(void* InPtr, uint16 InArenaTag) { TaggedPointer = reinterpret_cast<void*>(UPTRINT(InPtr) | (UPTRINT(InArenaTag) << ArenaShift)); }
 
 	inline FArenaPointer(void* Ptr, uint16 ArenaIndex) : TaggedPointer(reinterpret_cast<void*>(UPTRINT(Ptr) | (UPTRINT(ArenaIndex) << ArenaShift))) {};
 #endif
 
-	inline operator bool() const	{ return !!Pointer(); }
+	inline operator bool() const	{ return !!GetPointer(); }
 
 	void Free() const;
 
@@ -92,10 +92,10 @@ public:
 
 	inline TArenaPointer& operator=(T* Rhs) { SetPointerAndArena(Rhs, NoTag); return *this; }
 	
-	inline operator T* ()					{ return reinterpret_cast<T*>(Pointer()); }
-	inline operator const T* () const		{ return reinterpret_cast<const T*>(Pointer()); }
-	inline T* operator -> ()				{ return reinterpret_cast<T*>(Pointer()); }
-	inline const T* operator -> () const	{ return reinterpret_cast<const T*>(Pointer()); }
+	inline operator T* ()					{ return reinterpret_cast<T*>(GetPointer()); }
+	inline operator const T* () const		{ return reinterpret_cast<const T*>(GetPointer()); }
+	inline T* operator -> ()				{ return reinterpret_cast<T*>(GetPointer()); }
+	inline const T* operator -> () const	{ return reinterpret_cast<const T*>(GetPointer()); }
 };
 
 /** Memory arena interface
@@ -112,13 +112,13 @@ public:
 	CORE_API UE_RESTRICT UE_NOALIAS void*	Alloc(SIZE_T Size, SIZE_T Alignment);
 	CORE_API UE_NOALIAS void				Free(const void* MemoryBlock);
 
-	CORE_API SIZE_T				BlockSize(const void* MemoryBlock) const;
+	CORE_API SIZE_T				GetBlockSize(const void* MemoryBlock) const;
 	CORE_API const TCHAR*		GetDebugName() const;
 
 private:
 	CORE_API virtual void*		InternalAlloc(SIZE_T Size, SIZE_T Alignment) = 0;
 	CORE_API virtual void		InternalFree(const void* MemoryBlock, SIZE_T MemoryBlockSize);
-	CORE_API virtual SIZE_T		InternalBlockSize(const void* MemoryBlock) const = 0;
+	CORE_API virtual SIZE_T		InternalGetBlockSize(const void* MemoryBlock) const = 0;
 
 	CORE_API virtual const TCHAR* InternalGetDebugName() const;
 
@@ -132,7 +132,7 @@ public:
 
 inline void FArenaPointer::Free() const 
 { 
-	return Arena().Free(Pointer());
+	return Arena().Free(GetPointer());
 }
 
 // These are somewhat temporary, to support (experimental) arena-based container allocators
@@ -156,7 +156,7 @@ public:
 private:
 	CORE_API virtual void*			InternalAlloc(SIZE_T Size, SIZE_T Alignment) override;
 	CORE_API virtual void			InternalFree(const void* MemoryBlock, SIZE_T MemoryBlockSize) override;
-	CORE_API virtual SIZE_T			InternalBlockSize(const void* MemoryBlock) const override;
+	CORE_API virtual SIZE_T			InternalGetBlockSize(const void* MemoryBlock) const override;
 	CORE_API virtual const TCHAR*	InternalGetDebugName() const override;
 
 	void* HeapHandle = nullptr;
@@ -165,7 +165,7 @@ private:
 
 /** Default heap allocator
 
-	All allocations are passed through to UE4 main heap allocation functions
+	All allocations are passed through to UE's main heap allocation functions
 
   */
 class FMallocArena final : public FMemoryArena
@@ -177,7 +177,7 @@ public:
 private:
 	CORE_API virtual void*			InternalAlloc(SIZE_T Size, SIZE_T Alignment) override;
 	CORE_API virtual void			InternalFree(const void* MemoryBlock, SIZE_T MemoryBlockSize) override;
-	CORE_API virtual SIZE_T			InternalBlockSize(const void* MemoryBlock) const override;
+	CORE_API virtual SIZE_T			InternalGetBlockSize(const void* MemoryBlock) const override;
 	CORE_API virtual const TCHAR*	InternalGetDebugName() const override;
 };
 
@@ -195,7 +195,7 @@ public:
 private:
 	CORE_API virtual void*			InternalAlloc(SIZE_T Size, SIZE_T Alignment) override;
 	CORE_API virtual void			InternalFree(const void* MemoryBlock, SIZE_T MemoryBlockSize) override;
-	CORE_API virtual SIZE_T			InternalBlockSize(const void* MemoryBlock) const override;
+	CORE_API virtual SIZE_T			InternalGetBlockSize(const void* MemoryBlock) const override;
 	CORE_API virtual const TCHAR*	InternalGetDebugName() const override;
 };
 
