UE5 - Identify Canceled Tasks with a seperate state. This can make some Frontends easier to implement. Also cleaned up all the differnt Taskstates that are mostly useless for the User and consolidated them for ease of use. Also used the opprotunity to unify the Busy Waiting using the modified TaskDelegate.

diff --git a/Engine/Source/Runtime/Core/Private/Async/Fundamental/Scheduler.cpp b/Engine/Source/Runtime/Core/Private/Async/Fundamental/Scheduler.cpp
--- a/Engine/Source/Runtime/Core/Private/Async/Fundamental/Scheduler.cpp
+++ b/Engine/Source/Runtime/Core/Private/Async/Fundamental/Scheduler.cpp
@@ -7,7 +7,7 @@
 
 namespace LowLevelTasks
 {
-	DEFINE_LOG_CATEGORY(LogTasks2);
+	DEFINE_LOG_CATEGORY(LowLevelTasks);
 
 	thread_local FScheduler::FLocalQueueType* FScheduler::LocalQueue = nullptr;
 	thread_local FTask* FScheduler::ActiveTask = nullptr;
@@ -96,11 +96,8 @@ namespace LowLevelTasks
 		}
 	}
 
-	void FScheduler::Launch(FTask& Task, EQueuePreference QueuePreference)
+	void FScheduler::LaunchInternal(FTask& Task, EQueuePreference QueuePreference)
 	{
-		Task.PrepareLaunch();
-		checkSlow(Task.IsScheduled() || Task.IsCanceled());
-
 		if (ActiveWorkers.load(std::memory_order_relaxed))
 		{
 			if (LocalQueue && QueuePreference != EQueuePreference::GlobalQueuePreference)
@@ -176,7 +173,7 @@ namespace LowLevelTasks
 
 			if (WaitCount < WorkerSpinCycles)
 			{
-				TRACE_CPUPROFILER_EVENT_SCOPE(TEXT("Spinning Worker"));
+				TRACE_CPUPROFILER_EVENT_SCOPE("Spinning Worker");
 				for (uint32 i = 0; i < WaitCycles; i++)
 				{
 					FPlatformProcess::Yield();
@@ -198,9 +195,9 @@ namespace LowLevelTasks
 		FMemory::ClearAndDisableTLSCachesOnCurrentThread();
 	}
 
-	bool FScheduler::BusyWaitInternal(const FTask& AwaitedTask)
+	void FScheduler::BusyWaitInternal(const FConditional& Conditional)
 	{
-		TRACE_CPUPROFILER_EVENT_SCOPE(BusyWait);
+		TRACE_CPUPROFILER_EVENT_SCOPE(FScheduler::BusyWaitInternal);
 
 		checkSlow(LocalQueue != nullptr);
 		check(ActiveWorkers.load(std::memory_order_relaxed));
@@ -212,9 +209,9 @@ namespace LowLevelTasks
 			while(TryExecuteTaskFrom<&FLocalQueueType::DequeueLocal>(WorkerLocalQueue)
 			   || TryExecuteTaskFrom<&FLocalQueueType::DequeueGlobal>(WorkerLocalQueue))
 			{
-				if (AwaitedTask.IsCompleted())
+				if (Conditional())
 				{
-					return true;
+					return;
 				}
 				WaitCount = 0;
 			}
@@ -222,16 +219,16 @@ namespace LowLevelTasks
 			while(TryExecuteTaskFrom<&FLocalQueueType::DequeueLocal>(WorkerLocalQueue)
 			   || TryExecuteTaskFrom<&FLocalQueueType::DequeueSteal>(WorkerLocalQueue))
 			{
-				if (AwaitedTask.IsCompleted())
+				if (Conditional())
 				{
-					return true;
+					return;
 				}
 				WaitCount = 0;
 			}
 
-			if (AwaitedTask.IsCompleted())
+			if (Conditional())
 			{
-				return false;
+				return;
 			}
 
 			if (WaitCount < WorkerSpinCycles)
@@ -242,7 +239,7 @@ namespace LowLevelTasks
 			}
 			else
 			{
-				TRACE_CPUPROFILER_EVENT_SCOPE(TEXT("Sleeping Busy Wait"));
+				TRACE_CPUPROFILER_EVENT_SCOPE(BusyWaitInternal::SleepNoStats);
 				FPlatformProcess::SleepNoStats(0);
 				WaitCount = 0;
 			}
diff --git a/Engine/Source/Runtime/Core/Private/Async/TaskGraph.cpp b/Engine/Source/Runtime/Core/Private/Async/TaskGraph.cpp
--- a/Engine/Source/Runtime/Core/Private/Async/TaskGraph.cpp
+++ b/Engine/Source/Runtime/Core/Private/Async/TaskGraph.cpp
@@ -1874,13 +1874,14 @@ private:
 				}
 
 				const int MaxBackgroundTasks = FMath::Max(1, NumWorkerThreads - 1 - ConcurrentNormalTasks.load(std::memory_order_relaxed));
+				uint32 NumSpawned = 0;
 				while (ConcurrentBackgroundTasks.load(std::memory_order_relaxed) < MaxBackgroundTasks)
 				{
 					FBaseGraphTask* QueuedTask = QueuedBackgroundTasks.dequeue();
 					if (QueuedTask)
 					{
 						ConcurrentBackgroundTasks++;
-						LowLevelTasks::LaunchTask(QueuedTask->TaskHandle);
+						verifySlow(LowLevelTasks::TryLaunch(QueuedTask->TaskHandle, (NumSpawned++) ? LowLevelTasks::EQueuePreference::GlobalQueuePreference : LowLevelTasks::EQueuePreference::LocalQueuePreference));
 						continue;
 					}
 					break;
@@ -1893,7 +1894,7 @@ private:
 				if(ConcurrentBackgroundTasks.load(std::memory_order_relaxed) < MaxBackgroundTasks)
 				{
 					ConcurrentBackgroundTasks++;
-					LowLevelTasks::LaunchTask(Task->TaskHandle);			
+					verifySlow(LowLevelTasks::TryLaunch(Task->TaskHandle, LowLevelTasks::EQueuePreference::GlobalQueuePreference));			
 				}
 				else
 				{
@@ -1903,7 +1904,7 @@ private:
 			else
 			{
 				ConcurrentNormalTasks++;
-				LowLevelTasks::LaunchTask(Task->TaskHandle);
+				verifySlow(LowLevelTasks::TryLaunch(Task->TaskHandle, LowLevelTasks::EQueuePreference::GlobalQueuePreference));
 			}
 			return;
 		}
diff --git a/Engine/Source/Runtime/Core/Public/Async/Fundamental/Scheduler.h b/Engine/Source/Runtime/Core/Public/Async/Fundamental/Scheduler.h
--- a/Engine/Source/Runtime/Core/Public/Async/Fundamental/Scheduler.h
+++ b/Engine/Source/Runtime/Core/Public/Async/Fundamental/Scheduler.h
@@ -31,6 +31,9 @@ namespace LowLevelTasks
 		static thread_local FTask* ActiveTask;
 		static CORE_API FScheduler Singleton;
 
+		// using 16 bytes here because it fits the vtable and one additional pointer
+		using FConditional = TTaskDelegate<16, bool>;
+
 	private: 
 		//the FLocalQueueInstaller installs a LocalQueue into the current thread
 		struct FLocalQueueInstaller
@@ -47,15 +50,21 @@ namespace LowLevelTasks
 		CORE_API void StartWorkers(uint32 NumWorkers = 0, EThreadPriority Priority = EThreadPriority::TPri_Normal, bool bIsForkable = false);
 		CORE_API void StopWorkers();
 
-		//Launching the Task and queuing it.
-		CORE_API void Launch(FTask& Task, EQueuePreference QueuePreference);
+		//try to launch the task, the return value will specify if the task was in the ready state and has been launced
+		inline bool TryLaunch(FTask& Task, EQueuePreference QueuePreference = EQueuePreference::DefaultPreference);
+		
+		//try to cancel the task and launching it if the task was in ready state
+		//you still need to wait for task completion before recycling the handle
+		//you can alternatively use FTask::TryCancel if you want to launch the task manually
+		inline bool TryCancelAndLaunch(FTask& Task, EQueuePreference QueuePreference = EQueuePreference::DefaultPreference);		
 
 		//tries to do some work until the Task is completed
-		inline void BusyWait(const FTask& Task);
+		template<typename TaskType>
+		inline void BusyWait(const TaskType& Task);
 
 		//tries to do some work until the Conditional return true
 		template<typename Conditional>
-		inline void BusyWaitUntil(const Conditional& Cond);
+		inline void BusyWaitUntil(Conditional&& Cond);
 
 		//tries to do some work until all the Tasks are completed
 		//the template parameter can be any Type that has a const conversion operator to FTask
@@ -91,8 +100,9 @@ namespace LowLevelTasks
 
 	private: 
 		TUniquePtr<FThread> CreateWorker(FLocalQueueType* ExternalWorkerLocalQueue = nullptr, EThreadPriority Priority = EThreadPriority::TPri_Normal, bool bIsForkable = false);
-		void WorkerMain(struct FSleepEvent* WorkerEvent, FLocalQueueType* ExternalWorkerLocalQueue, uint32 WaitCycles);	
-		CORE_API bool BusyWaitInternal(const FTask& Task);
+		void WorkerMain(struct FSleepEvent* WorkerEvent, FLocalQueueType* ExternalWorkerLocalQueue, uint32 WaitCycles);
+		CORE_API void LaunchInternal(FTask& Task, EQueuePreference QueuePreference);
+		CORE_API void BusyWaitInternal(const FConditional& Conditional);
 		FORCENOINLINE void TrySleeping(FSleepEvent* WorkerEvent, uint32& WaitCount);
 		inline bool WakeUpWorker();
 		template<FTask* (FLocalQueueType::*DequeueFunction)()>
@@ -108,10 +118,14 @@ namespace LowLevelTasks
 		std::atomic_uint NextWorkerId { 0 };
 	};
 
-	template<EQueuePreference QueuePreference = EQueuePreference::DefaultPreference>
-	FORCEINLINE_DEBUGGABLE void LaunchTask(FTask& Task)
+	FORCEINLINE_DEBUGGABLE bool TryLaunch(FTask& Task, EQueuePreference QueuePreference = EQueuePreference::DefaultPreference)
+	{
+		return FScheduler::Get().TryLaunch(Task, QueuePreference);
+	}
+
+	FORCEINLINE_DEBUGGABLE bool TryCancelAndLaunch(FTask& Task, EQueuePreference QueuePreference = EQueuePreference::DefaultPreference)
 	{
-		FScheduler::Get().Launch(Task, QueuePreference);
+		return FScheduler::Get().TryCancelAndLaunch(Task, QueuePreference);
 	}
 
 	FORCEINLINE_DEBUGGABLE void BusyWaitForTask(const FTask& Task)
@@ -120,9 +134,9 @@ namespace LowLevelTasks
 	}
 
 	template<typename Conditional>
-	FORCEINLINE_DEBUGGABLE void BusyWaitUntil(const Conditional& Cond)
+	FORCEINLINE_DEBUGGABLE void BusyWaitUntil(Conditional&& Cond)
 	{
-		FScheduler::Get().BusyWaitUntil<Conditional>(Cond);
+		FScheduler::Get().BusyWaitUntil<Conditional>(Forward<Conditional>(Cond));
 	}
 
 	template<typename TaskType>
@@ -134,71 +148,75 @@ namespace LowLevelTasks
    /******************
 	* IMPLEMENTATION *
 	******************/
+	inline bool FScheduler::TryLaunch(FTask& Task, EQueuePreference QueuePreference)
+	{
+		if(Task.TryPrepareLaunch())
+		{
+			FScheduler::Get().LaunchInternal(Task, QueuePreference);
+			return true;
+		}
+		return false;
+	}
+
+	inline bool FScheduler::TryCancelAndLaunch(FTask& Task, EQueuePreference QueuePreference)
+	{
+		bool WasCanceled = Task.TryCancel();
+		if(WasCanceled && Task.TryPrepareLaunch())
+		{
+			Task.ExecuteTask();
+			return true;
+		}
+		return WasCanceled;
+	}
 
 	inline uint32 FScheduler::GetNumWorkers() const
 	{
 		return ActiveWorkers.load(std::memory_order_relaxed);
 	}
 
-	inline void FScheduler::BusyWait(const FTask& Task)
+	template<typename TaskType>
+	inline void FScheduler::BusyWait(const TaskType& Task)
 	{
 		if(!Task.IsCompleted())
 		{
 			FLocalQueueInstaller Installer(*this);
-			FScheduler::BusyWaitInternal(Task);
+			FScheduler::BusyWaitInternal([&Task](){ return Task.IsCompleted(); });
 		}
 	}
 
 	template<typename Conditional>
-	inline void FScheduler::BusyWaitUntil(const Conditional& Cond)
+	inline void FScheduler::BusyWaitUntil(Conditional&& Cond)
 	{
 		static_assert(TIsInvocable<Conditional>::Value, "Conditional is not invocable");
 		static_assert(TIsSame<decltype(Cond()), bool>::Value, "Conditional must return a boolean");
-
-		TOptional<FLocalQueueInstaller> Installer;
-		uint32 WaitCount = 0;
-		FTask Dummy;
-		while(!Cond())
+		
+		if(!Cond())
 		{
-			if (!Installer.IsSet())
-			{
-				Installer.Emplace(*this);
-			}
-			if(FScheduler::BusyWaitInternal(Dummy))
-			{
-				WaitCount = 0;
-			}
-			if (!Cond())
-			{
-				if (WaitCount < WorkerSpinCycles)
-				{
-					WaitCount++;
-					FPlatformProcess::Yield();
-				}
-				else
-				{
-					TRACE_CPUPROFILER_EVENT_SCOPE(TEXT("Sleeping Conditional Busy Wait"));
-					WaitCount = 0;
-					FPlatformProcess::SleepNoStats(0);		
-				}
-			}
+			FLocalQueueInstaller Installer(*this);
+			FScheduler::BusyWaitInternal(Forward<Conditional>(Cond));
 		}
 	}
 
 	template<typename TaskType>
 	inline void FScheduler::BusyWait(const TArrayView<const TaskType>& Tasks)
 	{
-		TOptional<FLocalQueueInstaller> Installer;
-		for(const FTask& Task : Tasks)
+		auto AllTasksCompleted = [Index(0), &Tasks]() mutable
 		{
-			if(!Task.IsCompleted())
+			while (Index < Tasks.Num())
 			{
-				if (!Installer.IsSet())
+				if (!Tasks[Index].IsCompleted())
 				{
-					Installer.Emplace(*this);
+					return false;
 				}
-				FScheduler::BusyWaitInternal(Task);
+				Index++;
 			}
+			return true;
+		};
+
+		if (!AllTasksCompleted())
+		{
+			FLocalQueueInstaller Installer(*this);
+			FScheduler::BusyWaitInternal([&AllTasksCompleted](){ return AllTasksCompleted(); });
 		}
 	}
 
@@ -208,13 +226,13 @@ namespace LowLevelTasks
 		ESleepState RunningState  = ESleepState::Running;
 		if(WorkerEvent->SleepState.compare_exchange_strong(DrowsingState, ESleepState::Sleeping, std::memory_order_relaxed))
 		{
-			TRACE_CPUPROFILER_EVENT_SCOPE(TEXT("Sleeping Worker"));
+			TRACE_CPUPROFILER_EVENT_SCOPE("Sleeping Worker");
 			WaitCount = 0;
 			WorkerEvent->SleepEvent->Wait(); // State two: ((Running -> Drowsing) -> Sleeping)
 		}
 		else if(WorkerEvent->SleepState.compare_exchange_strong(RunningState, ESleepState::Drowsing, std::memory_order_relaxed))
 		{
-			TRACE_CPUPROFILER_EVENT_SCOPE(TEXT("Drowsing Worker"));
+			TRACE_CPUPROFILER_EVENT_SCOPE("Drowsing Worker");
 			SleepEventQueue.enqueue(WorkerEvent); // State one: (Running -> Drowsing)
 			FPlatformProcess::SleepNoStats(0);	
 		}
diff --git a/Engine/Source/Runtime/Core/Public/Async/Fundamental/Task.h b/Engine/Source/Runtime/Core/Public/Async/Fundamental/Task.h
--- a/Engine/Source/Runtime/Core/Public/Async/Fundamental/Task.h
+++ b/Engine/Source/Runtime/Core/Public/Async/Fundamental/Task.h
@@ -8,7 +8,7 @@
 
 namespace LowLevelTasks
 {
-	DECLARE_LOG_CATEGORY_EXTERN(LogTasks2, Log, All);
+	DECLARE_LOG_CATEGORY_EXTERN(LowLevelTasks, Log, All);
 
 	enum class ETaskPriority : int32
 	{
@@ -21,12 +21,13 @@ namespace LowLevelTasks
 
 	enum class ETaskState : int32
 	{
-		Completed,		//task is finshed or the default when we create a handle
-		Ready,			//task is ready for scheduling
-		Canceled,		//task was canceled
-		Scheduled,		//task is scheduled and can run concurrently any time from now
-		Running,		//task is executing
-		Continuation,   //task is executing the continuation	
+		CanceledAndCompleted,	//means the task is completed with execution of it's continuation but the runnable was cancelled
+		Completed,				//means the task is completed with execution or the default when we create a handle
+		Ready,					//means the Task is ready to be launched
+		Canceled,				//means the task was canceled and launched and therefore queued for execution by a worker (which already might be executing it's continuation)
+		CanceledAndReady,		//means the task was canceled and is ready to be launched (it still is required to be launched)
+		Scheduled,				//means the task is launched and therefore queued for execution by a worker
+		Running,				//means the task is executing it's runnable and continuation by a worker 
 		Count
 	};
 	
@@ -121,7 +122,7 @@ namespace LowLevelTasks
 		};
 
 	private:
-		using FTaskDelegate = TTaskDelegate<PLATFORM_CACHE_LINE_SIZE - sizeof(FPackedData)>;
+		using FTaskDelegate = TTaskDelegate<PLATFORM_CACHE_LINE_SIZE - sizeof(FPackedData), bool>;
 		FTaskDelegate Runnable;
 		std::atomic<FPackedData> PackedData { FPackedData() };
 
@@ -130,6 +131,10 @@ namespace LowLevelTasks
 		{
 			static_assert(sizeof(FTaskBase) == PLATFORM_CACHE_LINE_SIZE, "Require FTaskBase to be cacheline size");
 		}
+
+	private:
+		inline bool IsCanceled() const { return PackedData.load(std::memory_order_relaxed).GetState() == ETaskState::Canceled; }
+		inline bool IsScheduled() const { return PackedData.load(std::memory_order_relaxed).GetState() == ETaskState::Scheduled; }
 	};
 	}
 
@@ -139,16 +144,32 @@ namespace LowLevelTasks
 		friend class FScheduler;
 		UE_NONCOPYABLE(FTask); //means non movable
 
-	public: //Public Interface 
-		inline bool IsCompleted() const { return PackedData.load(std::memory_order_seq_cst).GetState() == ETaskState::Completed; }
-		inline bool IsReady() const { return PackedData.load(std::memory_order_relaxed).GetState() == ETaskState::Ready; }
-		inline bool IsCanceled() const { return PackedData.load(std::memory_order_relaxed).GetState() == ETaskState::Canceled; }
-		inline bool IsScheduled() const { return PackedData.load(std::memory_order_relaxed).GetState() == ETaskState::Scheduled; }
-		inline bool IsRunning() const { return PackedData.load(std::memory_order_relaxed).GetState() == ETaskState::Running; }
-		inline bool IsRunningContinuation() const { return PackedData.load(std::memory_order_relaxed).GetState() == ETaskState::Continuation; }
+	public: //Public Interface
+		//means the task is completed and this taskhandle can be recycled
+		inline bool IsCompleted() const
+		{
+			ETaskState State = PackedData.load(std::memory_order_seq_cst).GetState();
+			return State == ETaskState::CanceledAndCompleted || State == ETaskState::Completed;
+		}
+		
+		//means the task was canceled but might still need to be launched 
+		inline bool WasCanceled() 		
+		{
+			ETaskState State = PackedData.load(std::memory_order_relaxed).GetState();
+			return State == ETaskState::CanceledAndReady || State == ETaskState::Canceled || State == ETaskState::CanceledAndCompleted;
+		}
 
-		//Note even though is canceled we do not pull it out of the scheduler and it's continuation will still run regardless
-		//so we still have to wait for completion before we can recycle/reassign the handle.
+		//means the task is ready to be launched but might already been canceled 
+		inline bool IsReady() const 
+		{
+			ETaskState State = PackedData.load(std::memory_order_relaxed).GetState();
+			return State == ETaskState::Ready || State == ETaskState::CanceledAndReady; 
+		}
+
+		
+		//try to cancel the task without launching it. Even if the task is canceled it sill needs to be launched manually
+		//to launch the task manually you use FScheduler::TryLaunch
+		//you can alternatively use FScheduler::TryCancelAndLaunch if you want to launch the task in this case automatically
 		inline bool TryCancel();
 
 		template<typename TRunnable, typename TContinuation>
@@ -174,9 +195,9 @@ namespace LowLevelTasks
 		inline ~FTask();
 
 	private: //Interface of the Scheduler
-		inline void PrepareLaunch();
+		inline bool TryPrepareLaunch();
 		//after calling this function the task can be considered dead
-		inline void ExecuteTask();			
+		inline void ExecuteTask();
 	};
 
    /******************
@@ -188,7 +209,7 @@ namespace LowLevelTasks
 	{
 		checkf(IsCompleted(), TEXT("State: %d"), PackedData.load(std::memory_order_relaxed).GetState());
 		checkSlow(!Runnable.IsSet());
-		Runnable = [this, LocalRunnable = Forward<TRunnable>(InRunnable), LocalContinuation = Forward<TContinuation>(InContinuation)]()
+		Runnable = [this, LocalRunnable = Forward<TRunnable>(InRunnable), LocalContinuation = Forward<TContinuation>(InContinuation)]() -> bool
 		{
 			checkSlow(IsScheduled() || IsCanceled());
 			FPackedData LocalPackedData = PackedData.load(std::memory_order_relaxed);
@@ -196,14 +217,15 @@ namespace LowLevelTasks
 			if (PackedData.compare_exchange_strong(ScheduledState, FPackedData(LocalPackedData, ETaskState::Running), std::memory_order_relaxed))
 			{
 				LocalRunnable();
+				LocalContinuation();
+				return false;
 			}
 			else
 			{
 				checkSlow(IsCanceled());
+				LocalContinuation();
+				return true;
 			}
-
-			PackedData.store(FPackedData(LocalPackedData, ETaskState::Continuation), std::memory_order_relaxed);
-			LocalContinuation();
 		};
 		PackedData.store(FPackedData(InDebugName, InPriority, ETaskState::Ready), std::memory_order_release);
 	}
@@ -213,7 +235,7 @@ namespace LowLevelTasks
 	{
 		checkf(IsCompleted(), TEXT("State: %d"), PackedData.load(std::memory_order_relaxed).GetState());
 		checkSlow(!Runnable.IsSet());
-		Runnable = [this, LocalRunnable = Forward<TRunnable>(InRunnable)]()
+		Runnable = [this, LocalRunnable = Forward<TRunnable>(InRunnable)]() -> bool
 		{
 			checkSlow(IsScheduled() || IsCanceled());
 			FPackedData LocalPackedData = PackedData.load(std::memory_order_relaxed);
@@ -221,10 +243,12 @@ namespace LowLevelTasks
 			if (PackedData.compare_exchange_strong(ScheduledState, FPackedData(LocalPackedData, ETaskState::Running), std::memory_order_relaxed))
 			{
 				LocalRunnable();
+				return false;
 			}
 			else
 			{
 				checkSlow(IsCanceled());
+				return true;
 			}
 		};
 		PackedData.store(FPackedData(InDebugName, InPriority, ETaskState::Ready), std::memory_order_release);
@@ -247,24 +271,24 @@ namespace LowLevelTasks
 		checkf(IsCompleted(), TEXT("State: %d"), PackedData.load(std::memory_order_relaxed).GetState());
 	}
 
-	inline bool FTask::TryCancel()
+	inline bool FTask::TryPrepareLaunch()
 	{
 		FPackedData LocalPackedData = PackedData.load(std::memory_order_relaxed);
-		FPackedData ReadyState(LocalPackedData, ETaskState::Ready);
-		FPackedData ScheduledState(LocalPackedData, ETaskState::Scheduled);
-		return PackedData.compare_exchange_strong(ReadyState,	  FPackedData(LocalPackedData, ETaskState::Canceled), std::memory_order_relaxed) 
-			|| PackedData.compare_exchange_strong(ScheduledState, FPackedData(LocalPackedData, ETaskState::Canceled), std::memory_order_relaxed)
-			|| IsCanceled();
+		FPackedData ReadyState = FPackedData(LocalPackedData, ETaskState::Ready);
+		FPackedData CanceledAndReadyState = FPackedData(LocalPackedData, ETaskState::CanceledAndReady);
+		return PackedData.compare_exchange_strong(ReadyState, FPackedData(LocalPackedData, ETaskState::Scheduled), std::memory_order_acquire, std::memory_order_relaxed)
+			|| PackedData.compare_exchange_strong(CanceledAndReadyState, FPackedData(LocalPackedData, ETaskState::Canceled), std::memory_order_acquire, std::memory_order_relaxed);
 	}
 
-	inline void FTask::PrepareLaunch()
+	inline bool FTask::TryCancel()
 	{
 		FPackedData LocalPackedData = PackedData.load(std::memory_order_relaxed);
 		FPackedData ReadyState(LocalPackedData, ETaskState::Ready);
-		if (!PackedData.compare_exchange_strong(ReadyState, FPackedData(LocalPackedData, ETaskState::Scheduled), std::memory_order_acquire))
-		{
-			checkf(IsCanceled(), TEXT("State: %d"), PackedData.load(std::memory_order_relaxed).GetState());
-		}
+		FPackedData ScheduledState(LocalPackedData, ETaskState::Scheduled);
+		//we can use memory_order_relaxed in this case because, when a task is canceled it does not launch the task
+		//to launch a canceled  task it has to go though TryPrepareLaunch which is doing the memory_order_acquire
+		return PackedData.compare_exchange_strong(ReadyState, FPackedData(LocalPackedData, ETaskState::CanceledAndReady), std::memory_order_relaxed)
+			|| PackedData.compare_exchange_strong(ScheduledState, FPackedData(LocalPackedData, ETaskState::Canceled), std::memory_order_relaxed);
 	}
 
 	inline void FTask::ExecuteTask()
@@ -272,12 +296,12 @@ namespace LowLevelTasks
 		checkSlow(Runnable.IsSet());
 		alignas(PLATFORM_CACHE_LINE_SIZE) 
 		FTaskDelegate LocalRunnable;
-		Runnable.CallAndMove(LocalRunnable);
+		bool Canceled = Runnable.CallAndMove(LocalRunnable);
 		checkSlow(!Runnable.IsSet());
 		//do not access the task again after this call
 		//as by defitition the task can be considered dead
 		FPackedData LocalPackedData = PackedData.load(std::memory_order_relaxed);
-		PackedData.store(FPackedData(LocalPackedData, ETaskState::Completed), std::memory_order_seq_cst);
+		PackedData.store(FPackedData(LocalPackedData, Canceled ? ETaskState::CanceledAndCompleted : ETaskState::Completed), std::memory_order_seq_cst);
 	}
 
 	inline ETaskPriority FTask::GetPriority() const 
diff --git a/Engine/Source/Runtime/Core/Public/Async/Fundamental/TaskDelegate.h b/Engine/Source/Runtime/Core/Public/Async/Fundamental/TaskDelegate.h
--- a/Engine/Source/Runtime/Core/Public/Async/Fundamental/TaskDelegate.h
+++ b/Engine/Source/Runtime/Core/Public/Async/Fundamental/TaskDelegate.h
@@ -2,24 +2,39 @@
 
 #pragma once
 #include "Templates/IsInvocable.h"
+#include "Misc/ScopeExit.h"
 #include <type_traits>
 
 namespace LowLevelTasks
 {
-	//version of TUniqueFunction<void()> that is less wasteful with it's memory
-	//this class might be removed when TUniqueFunction<void()> is fixed
-	template<uint32 TotalSize = PLATFORM_CACHE_LINE_SIZE>
+	namespace TTaskDelegate_Impl
+	{
+		template<typename ReturnType>
+		inline ReturnType MakeDummyValue()
+		{
+			return *(reinterpret_cast<ReturnType*>(uintptr_t(1)));
+		}
+
+		template<>
+		inline void MakeDummyValue<void>()
+		{
+			return;
+		}
+	}
+	//version of TUniqueFunction<ReturnType()> that is less wasteful with it's memory
+	//this class might be removed when TUniqueFunction<ReturnType()> is fixed
+	template<uint32 TotalSize = PLATFORM_CACHE_LINE_SIZE, typename ReturnType = void>
 	class alignas(8) TTaskDelegate
 	{
-		template<uint32>
+		template<uint32, typename>
 		friend class TTaskDelegate;
 
-		using ThisClass = TTaskDelegate<TotalSize>;
+		using ThisClass = TTaskDelegate<TotalSize, ReturnType>;
 		struct TTaskDelegateBase
 		{
 			virtual void Move(TTaskDelegateBase&, void*, void*, uint32) {check(false); };
-			virtual void Call(const void*) const { check(false); };
-			virtual void CallAndMove(ThisClass&, void*, uint32) {check(false); };
+			virtual ReturnType Call(void*) const { check(false); return TTaskDelegate_Impl::MakeDummyValue<ReturnType>(); };
+			virtual ReturnType CallAndMove(ThisClass&, void*, uint32) {check(false); return TTaskDelegate_Impl::MakeDummyValue<ReturnType>(); };
 			virtual void Destroy(void*) { check(false); };
 			virtual bool IsHeapAllocated() const { check(false); return false; }
 			virtual bool IsSet() const { check(false); return false; }
@@ -29,57 +44,62 @@ namespace LowLevelTasks
 		struct TTaskDelegateDummy final : TTaskDelegateBase
 		{
 			void Move(TTaskDelegateBase&, void*, void*, uint32) override {};
-			void Call(const void*) const override {};
-			void CallAndMove(ThisClass&, void*, uint32) override {};
+			ReturnType Call(void*) const override { check(false); return TTaskDelegate_Impl::MakeDummyValue<ReturnType>(); };
+			ReturnType CallAndMove(ThisClass&, void*, uint32) override { check(false); return TTaskDelegate_Impl::MakeDummyValue<ReturnType>(); };
 			void Destroy(void*) override {};
 			bool IsHeapAllocated() const override { return false; }
 			bool IsSet() const override { return false; }
 			uint32 DelegateSize() const override { return 0; }
 		};
 
-		template<typename TCallable, bool HeapAllocated>
+		template<typename TCallableType, bool HeapAllocated>
 		struct TTaskDelegateImpl;
 
-		template<typename TCallable>
-		struct TTaskDelegateImpl<TCallable, false> final : TTaskDelegateBase
+		template<typename TCallableType>
+		struct TTaskDelegateImpl<TCallableType, false> final : TTaskDelegateBase
 		{
+			template<typename TCallable>
 			inline TTaskDelegateImpl(TCallable&& Callable, void* InlineData)
 			{
-				static_assert(TIsInvocable<TCallable>::Value, "TCallable is not invocable");
-				static_assert(sizeof(TTaskDelegateImpl<TCallable, false>) == sizeof(TTaskDelegateBase), "Size must match the Baseclass");
-				new (InlineData) TCallable(Forward<TCallable>(Callable));
+				static_assert(TIsInvocable<TCallableType>::Value, "TCallableType is not invocable");
+				static_assert(TIsSame<ReturnType, decltype(Callable())>::Value, "TCallableType return type does not match");
+				static_assert(sizeof(TTaskDelegateImpl<TCallableType, false>) == sizeof(TTaskDelegateBase), "Size must match the Baseclass");
+				new (InlineData) TCallableType(Forward<TCallable>(Callable));
 			}
 
 			inline void Move(TTaskDelegateBase& DstWrapper, void* DstData, void* SrcData, uint32 DestInlineSize) override
 			{
-				TCallable* SrcPtr = reinterpret_cast<TCallable*>(SrcData);
-				if ((sizeof(TCallable) <= DestInlineSize) && (uintptr_t(DstData) % alignof(TCallable)) == 0)
+				TCallableType* SrcPtr = reinterpret_cast<TCallableType*>(SrcData);
+				if ((sizeof(TCallableType) <= DestInlineSize) && (uintptr_t(DstData) % alignof(TCallableType)) == 0)
 				{
-					new (&DstWrapper) TTaskDelegateImpl<TCallable, false>(MoveTemp(*SrcPtr), DstData);
+					new (&DstWrapper) TTaskDelegateImpl<TCallableType, false>(MoveTemp(*SrcPtr), DstData);
 				}
 				else
 				{
-					new (&DstWrapper) TTaskDelegateImpl<TCallable, true>(MoveTemp(*SrcPtr), DstData);
+					new (&DstWrapper) TTaskDelegateImpl<TCallableType, true>(MoveTemp(*SrcPtr), DstData);
 				}
 				new (this) TTaskDelegateDummy();
 			}
 
-			inline void Call(const void* InlineData) const override
+			inline ReturnType Call(void* InlineData) const override
 			{
-				const TCallable* LocalPtr = reinterpret_cast<const TCallable*>(InlineData);
-				(*LocalPtr)();
+				TCallableType* LocalPtr = reinterpret_cast<TCallableType*>(InlineData);
+				return (*LocalPtr)();
 			}
 
-			void CallAndMove(ThisClass& Destination, void* InlineData, uint32 DestInlineSize) override
+			ReturnType CallAndMove(ThisClass& Destination, void* InlineData, uint32 DestInlineSize) override
 			{
-				Call(InlineData);
-				Move(Destination.CallableWrapper, Destination.InlineStorage, InlineData, DestInlineSize);
+				ON_SCOPE_EXIT
+				{
+					Move(Destination.CallableWrapper, Destination.InlineStorage, InlineData, DestInlineSize);
+				};
+				return Call(InlineData);		
 			}
 
 			void Destroy(void* InlineData) override
 			{
-				TCallable* LocalPtr = reinterpret_cast<TCallable*>(InlineData);
-				LocalPtr->~TCallable();
+				TCallableType* LocalPtr = reinterpret_cast<TCallableType*>(InlineData);
+				LocalPtr->~TCallableType();
 			}
 
 			bool IsHeapAllocated() const override
@@ -94,49 +114,54 @@ namespace LowLevelTasks
 
 			uint32 DelegateSize() const override 
 			{ 
-				return sizeof(TCallable); 
+				return sizeof(TCallableType); 
 			}
 		};
 
-		template<typename TCallable>
-		struct TTaskDelegateImpl<TCallable, true> final : TTaskDelegateBase
+		template<typename TCallableType>
+		struct TTaskDelegateImpl<TCallableType, true> final : TTaskDelegateBase
 		{
 		private:
 			inline TTaskDelegateImpl(void* DstData, void* SrcData)
 			{
-				memcpy(DstData, SrcData, sizeof(TCallable*));
+				memcpy(DstData, SrcData, sizeof(TCallableType*));
 			}
 
 		public:
+			template<typename TCallable>
 			inline TTaskDelegateImpl(TCallable&& Callable, void* InlineData)
 			{
-				static_assert(TIsInvocable<TCallable>::Value, "TCallable is not invocable");
-				static_assert(sizeof(TTaskDelegateImpl<TCallable, true>) == sizeof(TTaskDelegateBase), "Size must match the Baseclass");
-				TCallable** HeapPtr = reinterpret_cast<TCallable**>(InlineData);
-				*HeapPtr = new TCallable(Forward<TCallable>(Callable));
+				static_assert(TIsInvocable<TCallableType>::Value, "TCallableType is not invocable");
+				static_assert(TIsSame<ReturnType, decltype(Callable())>::Value, "TCallableType return type does not match");
+				static_assert(sizeof(TTaskDelegateImpl<TCallableType, true>) == sizeof(TTaskDelegateBase), "Size must match the Baseclass");
+				TCallableType** HeapPtr = reinterpret_cast<TCallableType**>(InlineData);
+				*HeapPtr = new TCallableType(Forward<TCallable>(Callable));
 			}
 
 			inline void Move(TTaskDelegateBase& DstWrapper, void* DstData, void* SrcData, uint32 DestInlineSize) override
 			{
-				new (&DstWrapper) TTaskDelegateImpl<TCallable, true>(DstData, SrcData);
+				new (&DstWrapper) TTaskDelegateImpl<TCallableType, true>(DstData, SrcData);
 				new (this) TTaskDelegateDummy();
 			}
 
-			inline void Call(const void* InlineData) const override
+			inline ReturnType Call(void* InlineData) const override
 			{
-				const TCallable* HeapPtr = reinterpret_cast<const TCallable*>(*reinterpret_cast<const void* const*>(InlineData));
-				(*HeapPtr)();
+				TCallableType* HeapPtr = reinterpret_cast<TCallableType*>(*reinterpret_cast<void* const*>(InlineData));
+				return (*HeapPtr)();
 			}
 
-			void CallAndMove(ThisClass& Destination, void* InlineData, uint32 DestInlineSize) override
+			ReturnType CallAndMove(ThisClass& Destination, void* InlineData, uint32 DestInlineSize) override
 			{
-				Call(InlineData);
-				Move(Destination.CallableWrapper, Destination.InlineStorage, InlineData, DestInlineSize);
+				ON_SCOPE_EXIT
+				{
+					Move(Destination.CallableWrapper, Destination.InlineStorage, InlineData, DestInlineSize);
+				};
+				return Call(InlineData);			
 			}
 
 			void Destroy(void* InlineData) override
 			{
-				TCallable* HeapPtr = reinterpret_cast<TCallable*>(*reinterpret_cast<void**>(InlineData));
+				TCallableType* HeapPtr = reinterpret_cast<TCallableType*>(*reinterpret_cast<void**>(InlineData));
 				delete HeapPtr;
 			}
 
@@ -152,7 +177,7 @@ namespace LowLevelTasks
 
 			uint32 DelegateSize() const override 
 			{ 
-				return sizeof(TCallable); 
+				return sizeof(TCallableType); 
 			}
 		};
 
@@ -160,15 +185,15 @@ namespace LowLevelTasks
 		TTaskDelegate() 
 		{
 			static_assert(TotalSize % 8 == 0,  "Totalsize must be dividable by 8");
-			static_assert(TotalSize >= 16,  "Totalsize must be larger or equal than 16");
+			static_assert(TotalSize >= (sizeof(TTaskDelegateBase) + sizeof(void*)),  "Totalsize must be large enough to fit a vtable and pointer");
 			new (&CallableWrapper) TTaskDelegateDummy();
 		}
 
 		template<uint32 SourceTotalSize>
-		TTaskDelegate(const TTaskDelegate<SourceTotalSize>&) = delete;
+		TTaskDelegate(const TTaskDelegate<SourceTotalSize, ReturnType>&) = delete;
 
 		template<uint32 SourceTotalSize>
-		TTaskDelegate(TTaskDelegate<SourceTotalSize>&& Other)
+		TTaskDelegate(TTaskDelegate<SourceTotalSize, ReturnType>&& Other)
 		{
 			static_cast<TTaskDelegateBase*>(&Other.CallableWrapper)->Move(CallableWrapper, InlineStorage, Other.InlineStorage, InlineStorageSize);
 		}
@@ -176,13 +201,14 @@ namespace LowLevelTasks
 		template<typename TCallable>
 		TTaskDelegate(TCallable&& Callable)
 		{
-			if ((sizeof(TCallable) <= InlineStorageSize) && ((uintptr_t(InlineStorageSize) % alignof(TCallable)) == 0))
+			using TCallableType = typename TRemoveConst<typename TRemoveReference<TCallable>::Type>::Type;
+			if ((sizeof(TCallableType) <= InlineStorageSize) && ((uintptr_t(InlineStorageSize) % alignof(TCallableType)) == 0))
 			{
-				new (&CallableWrapper) TTaskDelegateImpl<TCallable, false>(Forward<TCallable>(Callable), InlineStorage);
+				new (&CallableWrapper) TTaskDelegateImpl<TCallableType, false>(Forward<TCallable>(Callable), InlineStorage);
 			}
 			else
 			{
-				new (&CallableWrapper) TTaskDelegateImpl<TCallable, true>(Forward<TCallable>(Callable), InlineStorage);
+				new (&CallableWrapper) TTaskDelegateImpl<TCallableType, true>(Forward<TCallable>(Callable), InlineStorage);
 			}
 		}
 
@@ -191,23 +217,23 @@ namespace LowLevelTasks
 			static_cast<TTaskDelegateBase*>(&CallableWrapper)->Destroy(InlineStorage);
 		}
 
-		void operator()() const
+		ReturnType operator()() const
 		{
-			static_cast<const TTaskDelegateBase*>(&CallableWrapper)->Call(InlineStorage);
+			return static_cast<const TTaskDelegateBase*>(&CallableWrapper)->Call(InlineStorage);
 		}
 
 		template<uint32 DestTotalSize>
-		void CallAndMove(TTaskDelegate<DestTotalSize>& Destination)
+		ReturnType CallAndMove(TTaskDelegate<DestTotalSize, ReturnType>& Destination)
 		{
 			checkSlow(!Destination.IsSet());
-			static_cast<TTaskDelegateBase*>(&CallableWrapper)->CallAndMove(Destination, InlineStorage, TTaskDelegate<DestTotalSize>::InlineStorageSize);
+			return static_cast<TTaskDelegateBase*>(&CallableWrapper)->CallAndMove(Destination, InlineStorage, TTaskDelegate<DestTotalSize, ReturnType>::InlineStorageSize);
 		}
 
 		template<uint32 SourceTotalSize>
-		ThisClass& operator= (const TTaskDelegate<SourceTotalSize>&) = delete;
+		ThisClass& operator= (const TTaskDelegate<SourceTotalSize, ReturnType>&) = delete;
 
 		template<uint32 SourceTotalSize>
-		ThisClass& operator= (TTaskDelegate<SourceTotalSize>&& Other)
+		ThisClass& operator= (TTaskDelegate<SourceTotalSize, ReturnType>&& Other)
 		{
 			static_cast<TTaskDelegateBase*>(&CallableWrapper)->Destroy(InlineStorage);
 			static_cast<TTaskDelegateBase*>(&Other.CallableWrapper)->Move(CallableWrapper, InlineStorage, Other.InlineStorage, InlineStorageSize);
@@ -217,14 +243,15 @@ namespace LowLevelTasks
 		template<typename TCallable>
 		ThisClass& operator= (TCallable&& Callable)
 		{
+			using TCallableType = typename TRemoveConst<typename TRemoveReference<TCallable>::Type>::Type;
 			static_cast<TTaskDelegateBase*>(&CallableWrapper)->Destroy(InlineStorage);
-			if ((sizeof(TCallable) <= InlineStorageSize) && ((uintptr_t(InlineStorageSize) % alignof(TCallable)) == 0))
+			if ((sizeof(TCallableType) <= InlineStorageSize) && ((uintptr_t(InlineStorageSize) % alignof(TCallableType)) == 0))
 			{
-				new (&CallableWrapper) TTaskDelegateImpl<TCallable, false>(Forward<TCallable>(Callable), InlineStorage);
+				new (&CallableWrapper) TTaskDelegateImpl<TCallableType, false>(Forward<TCallable>(Callable), InlineStorage);
 			}
 			else
 			{
-				new (&CallableWrapper) TTaskDelegateImpl<TCallable, true>(Forward<TCallable>(Callable), InlineStorage);
+				new (&CallableWrapper) TTaskDelegateImpl<TCallableType, true>(Forward<TCallable>(Callable), InlineStorage);
 			}
 			return *this;
 		}
@@ -252,7 +279,7 @@ namespace LowLevelTasks
 
 	private:
 		static constexpr uint32 InlineStorageSize = TotalSize - sizeof(TTaskDelegateBase);
-		char InlineStorage[InlineStorageSize];
+		mutable char InlineStorage[InlineStorageSize];
 		TTaskDelegateBase CallableWrapper;
 	};
 }
\ No newline at end of file
diff --git a/Engine/Source/Runtime/Core/Public/Misc/QueuedThreadPoolWrapper.h b/Engine/Source/Runtime/Core/Public/Misc/QueuedThreadPoolWrapper.h
--- a/Engine/Source/Runtime/Core/Public/Misc/QueuedThreadPoolWrapper.h
+++ b/Engine/Source/Runtime/Core/Public/Misc/QueuedThreadPoolWrapper.h
@@ -221,7 +221,7 @@ class CORE_API FQueuedLowLevelThreadPool : public FQueuedThreadPool
 
 		virtual bool Retract()
 		{
-			return Task.TryCancel();
+			return Task.WasCanceled() || Task.TryCancel();
 		}
 	};
 public:
@@ -269,7 +269,7 @@ public:
 				break;
 			}
 			TaskCount++;
-			LowLevelTasks::LaunchTask<LowLevelTasks::EQueuePreference::GlobalQueuePreference>(QueuedWork->Task);
+			verifySlow(LowLevelTasks::TryLaunch(QueuedWork->Task, LowLevelTasks::EQueuePreference::GlobalQueuePreference));
 		}
 
 		if (InNumQueuedWork == -1)
@@ -298,13 +298,14 @@ private:
 		[this, InternalData = InQueuedWork->InternalData]()
 		{
 			--TaskCount;
+			uint32 NumSpawned = 0;
 			while (TaskCount < MaxConcurrency)
 			{
 				FQueuedWorkInternalData* QueuedWork = Dequeue();
 				if (QueuedWork)
 				{
 					TaskCount++;
-					LowLevelTasks::LaunchTask<LowLevelTasks::EQueuePreference::GlobalQueuePreference>(QueuedWork->Task);
+					verifySlow(LowLevelTasks::TryLaunch(QueuedWork->Task, (NumSpawned++) ? LowLevelTasks::EQueuePreference::GlobalQueuePreference : LowLevelTasks::EQueuePreference::LocalQueuePreference));
 				}
 				else
 				{
@@ -316,7 +317,7 @@ private:
 		if(!bIsPaused && TaskCount <= MaxConcurrency)
 		{
 			TaskCount++;
-			LowLevelTasks::LaunchTask<LowLevelTasks::EQueuePreference::GlobalQueuePreference>(QueuedWorkInternalData->Task);
+			verifySlow(LowLevelTasks::TryLaunch(QueuedWorkInternalData->Task, LowLevelTasks::EQueuePreference::GlobalQueuePreference));
 		}
 		else
 		{
@@ -361,7 +362,7 @@ protected:
 
 			verify(QueuedWork->Retract());
 			TaskCount++;
-			LowLevelTasks::LaunchTask<LowLevelTasks::EQueuePreference::GlobalQueuePreference>(QueuedWork->Task);
+			verifySlow(LowLevelTasks::TryLaunch(QueuedWork->Task, LowLevelTasks::EQueuePreference::GlobalQueuePreference));
 		}
 
 		while (TaskCount != 0)
