UE5 - fix WinRT compile error

diff --git a/Engine/Source/Runtime/Core/Public/Experimental/Containers/HazardPointer.h b/Engine/Source/Runtime/Core/Public/Experimental/Containers/HazardPointer.h
--- a/Engine/Source/Runtime/Core/Public/Experimental/Containers/HazardPointer.h
+++ b/Engine/Source/Runtime/Core/Public/Experimental/Containers/HazardPointer.h
@@ -164,7 +164,11 @@ public:
 			}
 		};
 
+#if PLATFORM_HOLOLENS //dynamic initialization of thread local data not allowed in WinRT code
+		uint32 StartIndex = 0;
+#else
 		static thread_local uint32 StartIndex = FPseudo::GetThreadId();
+#endif
 	
 		FHazardRecordChunk* p = &Head;
 		if (Cached)
