CompactBinary: Added AsAnyReference, renamed writer save functions, and changed buffer allocator parameters

Includes a minor optimization to TCbBufferRef::GetOuterBuffer() and adds an assert that the buffer allocator returned a valid buffer.

#rnx

diff --git a/Engine/Source/Runtime/Core/Private/Serialization/CompactBinary.cpp b/Engine/Source/Runtime/Core/Private/Serialization/CompactBinary.cpp
--- a/Engine/Source/Runtime/Core/Private/Serialization/CompactBinary.cpp
+++ b/Engine/Source/Runtime/Core/Private/Serialization/CompactBinary.cpp
@@ -281,6 +281,20 @@ FBlake3Hash FCbField::AsBinaryReference(const FBlake3Hash& Default)
 	}
 }
 
+FBlake3Hash FCbField::AsAnyReference(const FBlake3Hash& Default)
+{
+	if (FCbFieldType::IsAnyReference(Type))
+	{
+		Error = ECbFieldError::None;
+		return FBlake3Hash(*static_cast<const FBlake3Hash::ByteArray*>(Payload));
+	}
+	else
+	{
+		Error = ECbFieldError::TypeError;
+		return Default;
+	}
+}
+
 FBlake3Hash FCbField::AsHash(const FBlake3Hash& Default)
 {
 	if (FCbFieldType::IsHash(Type))
diff --git a/Engine/Source/Runtime/Core/Private/Serialization/CompactBinarySerialization.cpp b/Engine/Source/Runtime/Core/Private/Serialization/CompactBinarySerialization.cpp
--- a/Engine/Source/Runtime/Core/Private/Serialization/CompactBinarySerialization.cpp
+++ b/Engine/Source/Runtime/Core/Private/Serialization/CompactBinarySerialization.cpp
@@ -170,7 +170,10 @@ FCbFieldRef LoadCompactBinary(FArchive& Ar, FCbBufferAllocator Allocator)
 	}
 
 	// Allocate the buffer, copy the header, and read the remainder of the field.
-	FSharedBufferPtr Buffer = Allocator(FieldType, FieldSize);
+	FSharedBufferPtr Buffer = Allocator(FieldSize);
+	checkf(Buffer && Buffer->GetSize() == FieldSize,
+		TEXT("Allocator returned a buffer of size %" UINT64_FMT " bytes when %" UINT64_FMT " bytes were requested."),
+		Buffer ? Buffer->GetSize() : 0, FieldSize);
 	FMutableMemoryView View = Buffer->GetView();
 	FMemory::Memcpy(View.GetData(), HeaderBytes.GetData(), HeaderBytes.Num());
 	View += HeaderBytes.Num();
@@ -209,10 +212,7 @@ static FArchive& SerializeCompactBinary(FArchive& Ar, T& Value, ConvertType&& Co
 	if (Ar.IsLoading())
 	{
 		Value = Invoke(Forward<ConvertType>(Convert),
-			LoadCompactBinary(Ar, [](ECbFieldType Type, uint64 Size)
-			{
-				return FSharedBuffer::Alloc(Size);
-			}));
+			LoadCompactBinary(Ar, [](uint64 Size) { return FSharedBuffer::Alloc(Size); }));
 	}
 	else if (Ar.IsSaving())
 	{
diff --git a/Engine/Source/Runtime/Core/Private/Serialization/CompactBinaryWriter.cpp b/Engine/Source/Runtime/Core/Private/Serialization/CompactBinaryWriter.cpp
--- a/Engine/Source/Runtime/Core/Private/Serialization/CompactBinaryWriter.cpp
+++ b/Engine/Source/Runtime/Core/Private/Serialization/CompactBinaryWriter.cpp
@@ -88,11 +88,11 @@ FCbFieldRefIterator FCbWriter::Save() const
 {
 	const uint64 Size = GetSaveSize();
 	FSharedBufferPtr Buffer = FSharedBuffer::Alloc(Size);
-	const FCbFieldIterator Output = SaveToMemory(Buffer->GetView());
+	const FCbFieldIterator Output = Save(Buffer->GetView());
 	return FCbFieldRefIterator::MakeRangeView(Output, FSharedBuffer::MakeReadOnly(MoveTemp(Buffer)));
 }
 
-FCbFieldIterator FCbWriter::SaveToMemory(const FMutableMemoryView Buffer) const
+FCbFieldIterator FCbWriter::Save(const FMutableMemoryView Buffer) const
 {
 	checkf(States.Num() == 1 && States.Last().Flags == EStateFlags::None,
 		TEXT("It is invalid to save while there are incomplete write operations."));
@@ -103,7 +103,7 @@ FCbFieldIterator FCbWriter::SaveToMemory(const FMutableMemoryView Buffer) const
 	return FCbFieldIterator::MakeRange(Buffer);
 }
 
-void FCbWriter::SaveToArchive(FArchive& Ar) const
+void FCbWriter::Save(FArchive& Ar) const
 {
 	check(Ar.IsSaving());
 	checkf(States.Num() == 1 && States.Last().Flags == EStateFlags::None,
diff --git a/Engine/Source/Runtime/Core/Private/Tests/Serialization/CompactBinarySerializationTest.cpp b/Engine/Source/Runtime/Core/Private/Tests/Serialization/CompactBinarySerializationTest.cpp
--- a/Engine/Source/Runtime/Core/Private/Tests/Serialization/CompactBinarySerializationTest.cpp
+++ b/Engine/Source/Runtime/Core/Private/Tests/Serialization/CompactBinarySerializationTest.cpp
@@ -141,7 +141,7 @@ bool FCbSaveTest::RunTest(const FString& Parameters)
 				TestTrue(FString::Printf(TEXT("SaveCompactBinary(%s)->EqualBytes"), Test), ExpectedData.EqualBytes(MakeMemoryView(WriteAr))))
 			{
 				FMemoryReader ReadAr(WriteAr);
-				FCbFieldRef Field = LoadCompactBinary(ReadAr, [](ECbFieldType Type, uint64 Size) { return FSharedBuffer::Alloc(Size); });
+				FCbFieldRef Field = LoadCompactBinary(ReadAr, [](uint64 Size) { return FSharedBuffer::Alloc(Size); });
 				TestTrue(FString::Printf(TEXT("LoadCompactBinary(%s)->EqualBytes"), Test), ExpectedData.EqualBytes(Field.GetView()));
 			}
 		}
diff --git a/Engine/Source/Runtime/Core/Private/Tests/Serialization/CompactBinaryTest.cpp b/Engine/Source/Runtime/Core/Private/Tests/Serialization/CompactBinaryTest.cpp
--- a/Engine/Source/Runtime/Core/Private/Tests/Serialization/CompactBinaryTest.cpp
+++ b/Engine/Source/Runtime/Core/Private/Tests/Serialization/CompactBinaryTest.cpp
@@ -953,6 +953,12 @@ bool FCbFieldReferenceTest::RunTest(const FString& Parameters)
 	// Test FCbField(Reference, NonZero)
 	TestField<ECbFieldType::Reference>(TEXT("Reference, NonZero"), SequentialBytes, FBlake3Hash(SequentialBytes));
 
+	// Test FCbField(Reference, NonZero) AsAnyReference
+	{
+		FCbField Field(SequentialBytes, ECbFieldType::Reference);
+		TestFieldAsType(TEXT("Reference, NonZero, AsAnyReference"), Field, &FCbField::AsAnyReference, FBlake3Hash(SequentialBytes));
+	}
+
 	// Test FCbField(None) as Reference
 	{
 		FCbField DefaultField;
@@ -974,6 +980,12 @@ bool FCbFieldBinaryReferenceTest::RunTest(const FString& Parameters)
 	// Test FCbField(BinaryReference, NonZero)
 	TestField<ECbFieldType::BinaryReference>(TEXT("BinaryReference, NonZero"), SequentialBytes, FBlake3Hash(SequentialBytes));
 
+	// Test FCbField(BinaryReference, NonZero) AsAnyReference
+	{
+		FCbField Field(SequentialBytes, ECbFieldType::BinaryReference);
+		TestFieldAsType(TEXT("BinaryReference, NonZero, AsAnyReference"), Field, &FCbField::AsAnyReference, FBlake3Hash(SequentialBytes));
+	}
+
 	// Test FCbField(None) as BinaryReference
 	{
 		FCbField DefaultField;
diff --git a/Engine/Source/Runtime/Core/Private/Tests/Serialization/CompactBinaryWriterTest.cpp b/Engine/Source/Runtime/Core/Private/Tests/Serialization/CompactBinaryWriterTest.cpp
--- a/Engine/Source/Runtime/Core/Private/Tests/Serialization/CompactBinaryWriterTest.cpp
+++ b/Engine/Source/Runtime/Core/Private/Tests/Serialization/CompactBinaryWriterTest.cpp
@@ -216,7 +216,7 @@ bool FCbWriterNullTest::RunTest(const FString& Parameters)
 		TestEqual(TEXT("FCbWriter(Null, Uniform) Validate"), ValidateCompactBinaryRange(*Fields.GetBuffer(), ECbValidateMode::All), ECbValidateError::None);
 	}
 
-	// Test Null with SaveToMemory
+	// Test Null with Save(Buffer)
 	{
 		Writer.Reset();
 		constexpr int NullCount = 3;
@@ -225,7 +225,7 @@ bool FCbWriterNullTest::RunTest(const FString& Parameters)
 			Writer.Null();
 		}
 		uint8 Buffer[NullCount]{};
-		FCbFieldIterator Fields = Writer.SaveToMemory(MakeMemoryView(Buffer));
+		FCbFieldIterator Fields = Writer.Save(MakeMemoryView(Buffer));
 		if (TestEqual(TEXT("FCbWriter(Null, Memory) Validate"), ValidateCompactBinaryRange(MakeMemoryView(Buffer), ECbValidateMode::All), ECbValidateError::None))
 		{
 			for (int Index = 0; Index < NullCount; ++Index)
@@ -912,14 +912,14 @@ bool FCbWriterComplexTest::RunTest(const FString& Parameters)
 		Writer.EndObject();
 		Object = Writer.Save().AsObjectRef();
 
-		Writer.SaveToArchive(Archive);
-		TestEqual(TEXT("FCbWriter(Complex).SaveToArchive()->Num()"), uint64(Archive.Num()), Writer.GetSaveSize());
+		Writer.Save(Archive);
+		TestEqual(TEXT("FCbWriter(Complex).Save(Ar)->Num()"), uint64(Archive.Num()), Writer.GetSaveSize());
 	}
 
 	TestEqual(TEXT("FCbWriter(Complex).Save()->Validate"),
 		ValidateCompactBinary(*Object.GetBuffer(), ECbValidateMode::All), ECbValidateError::None);
 
-	TestEqual(TEXT("FCbWriter(Complex).SaveToArchive()->Validate"),
+	TestEqual(TEXT("FCbWriter(Complex).Save(Ar)->Validate"),
 		ValidateCompactBinary(MakeMemoryView(Archive), ECbValidateMode::All), ECbValidateError::None);
 
 	return true;
@@ -1010,19 +1010,19 @@ bool FCbWriterStateTest::RunTest(const FString& Parameters)
 	// Assert on saving an empty writer.
 	//uint8 EmptyField[1];
 	//Writer.Reset();
-	//Writer.SaveToMemory(MakeMemoryView(EmptyField));
+	//Writer.Save(MakeMemoryView(EmptyField));
 
 	// Assert on under-sized save buffer.
 	//uint8 ZeroFieldSmall[1];
 	//Writer.Reset();
 	//Writer.Integer(0);
-	//Writer.SaveToMemory(MakeMemoryView(ZeroFieldSmall));
+	//Writer.Save(MakeMemoryView(ZeroFieldSmall));
 
 	// Assert on over-sized save buffer.
 	//uint8 ZeroFieldLarge[3];
 	//Writer.Reset();
 	//Writer.Integer(0);
-	//Writer.SaveToMemory(MakeMemoryView(ZeroFieldLarge));
+	//Writer.Save(MakeMemoryView(ZeroFieldLarge));
 
 	// Assert on empty name.
 	//Writer.Name(""_ASV);
@@ -1044,14 +1044,14 @@ bool FCbWriterStateTest::RunTest(const FString& Parameters)
 	//uint8 InvalidObject[1];
 	//Writer.Reset();
 	//Writer.BeginObject();
-	//Writer.SaveToMemory(MakeMemoryView(InvalidObject));
+	//Writer.Save(MakeMemoryView(InvalidObject));
 	//Writer.EndObject();
 
 	// Assert on save in array.
 	//uint8 InvalidArray[1];
 	//Writer.Reset();
 	//Writer.BeginArray();
-	//Writer.SaveToMemory(MakeMemoryView(InvalidArray));
+	//Writer.Save(MakeMemoryView(InvalidArray));
 	//Writer.EndArray();
 
 	// Assert on object end with no begin.
diff --git a/Engine/Source/Runtime/Core/Public/Serialization/CompactBinary.h b/Engine/Source/Runtime/Core/Public/Serialization/CompactBinary.h
--- a/Engine/Source/Runtime/Core/Public/Serialization/CompactBinary.h
+++ b/Engine/Source/Runtime/Core/Public/Serialization/CompactBinary.h
@@ -489,6 +489,8 @@ public:
 	CORE_API FBlake3Hash AsReference(const FBlake3Hash& Default = FBlake3Hash());
 	/** Access the field as a hash referencing a blob. Returns the provided default on error. */
 	CORE_API FBlake3Hash AsBinaryReference(const FBlake3Hash& Default = FBlake3Hash());
+	/** Access the field as a hash referencing a blob or compact binary. Returns the provided default on error. */
+	CORE_API FBlake3Hash AsAnyReference(const FBlake3Hash& Default = FBlake3Hash());
 
 	/** Access the field as a hash. Returns the provided default on error. */
 	CORE_API FBlake3Hash AsHash(const FBlake3Hash& Default = FBlake3Hash());
@@ -898,7 +900,7 @@ private:
 ///////////////////////////////////////////////////////////////////////////////////////////////////////////////////////
 
 /** A reference to a function that is used to allocate buffers for compact binary data. */
-using FCbBufferAllocator = TFunctionRef<FSharedBufferPtr (ECbFieldType Type, uint64 Size)>;
+using FCbBufferAllocator = TFunctionRef<FSharedBufferPtr (uint64 Size)>;
 
 ///////////////////////////////////////////////////////////////////////////////////////////////////////////////////////
 
@@ -997,7 +999,8 @@ public:
 	}
 
 	/** Returns the outer buffer (if any) that contains this value. */
-	inline const FSharedBufferConstPtr& GetOuterBuffer() const { return Buffer; }
+	inline const FSharedBufferConstPtr& GetOuterBuffer() const & { return Buffer; }
+	inline FSharedBufferConstPtr GetOuterBuffer() && { return MoveTemp(Buffer); }
 
 private:
 	template <typename OtherType>
diff --git a/Engine/Source/Runtime/Core/Public/Serialization/CompactBinaryWriter.h b/Engine/Source/Runtime/Core/Public/Serialization/CompactBinaryWriter.h
--- a/Engine/Source/Runtime/Core/Public/Serialization/CompactBinaryWriter.h
+++ b/Engine/Source/Runtime/Core/Public/Serialization/CompactBinaryWriter.h
@@ -20,7 +20,7 @@ struct FTimespan;
  *
  * The writer produces a sequence of fields that can be saved to a provided memory buffer or into
  * a new owned buffer. The typical use case is to write a single object, which can be accessed by
- * calling Save().AsObjectRef() or SaveToMemory(Buffer).AsObject().
+ * calling Save().AsObjectRef() or Save(Buffer).AsObject().
  *
  * The writer will assert on most incorrect usage and will always produce valid compact binary if
  * provided with valid input. The writer does not check for invalid UTF-8 string encoding, object
@@ -81,7 +81,7 @@ public:
 	 * @param Buffer A mutable memory view to write to. Must be exactly GetSaveSize() bytes.
 	 * @return An iterator for the field(s) written to the buffer.
 	 */
-	CORE_API FCbFieldIterator SaveToMemory(FMutableMemoryView Buffer) const;
+	CORE_API FCbFieldIterator Save(FMutableMemoryView Buffer) const;
 
 	/**
 	 * Serialize the field(s) to an archive.
@@ -91,7 +91,7 @@ public:
 	 *
 	 * @param Ar An archive to write to. Exactly GetSaveSize() bytes will be written.
 	 */
-	CORE_API void SaveToArchive(FArchive& Ar) const;
+	CORE_API void Save(FArchive& Ar) const;
 
 	/**
 	 * The size of buffer (in bytes) required to serialize the fields that have been written.
