Additional crash context providers

Add functionality for users to provide additional context in the case of a crash. A per thread stack with callbacks that are issued in case of a crash.

diff --git a/Engine/Source/Runtime/Core/Private/GenericPlatform/GenericPlatformCrashContext.cpp b/Engine/Source/Runtime/Core/Private/GenericPlatform/GenericPlatformCrashContext.cpp
--- a/Engine/Source/Runtime/Core/Private/GenericPlatform/GenericPlatformCrashContext.cpp
+++ b/Engine/Source/Runtime/Core/Private/GenericPlatform/GenericPlatformCrashContext.cpp
@@ -1080,6 +1080,8 @@ void FGenericCrashContext::DumpLog(const FString& CrashFolderAbsolute)
 		static_cast<void>(IFileManager::Get().Copy(*LogDstAbsolute, *LogSrcAbsolute, bReplace, bEvenIfReadOnly, bAttributes, CopyProgress, FILEREAD_AllowWrite, FILEWRITE_AllowRead));	// best effort, so don't care about result: couldn't copy -> tough, no log
 	}
 #endif // !NO_LOGGING
+
+	
 }
 
 FORCENOINLINE void FGenericCrashContext::CapturePortableCallStack(int32 NumStackFramesToIgnore, void* Context)
@@ -1153,6 +1155,8 @@ void FGenericCrashContext::AddPortableThreadCallStack(uint32 ThreadId, const TCH
 	// Not implemented for generic class
 }
 
+
+
 void FGenericCrashContext::CopyPlatformSpecificFiles(const TCHAR* OutputDirectory, void* Context)
 {
 	// If present, include the crash report config file to pass config values to the CRC
@@ -1163,8 +1167,120 @@ void FGenericCrashContext::CopyPlatformSpecificFiles(const TCHAR* OutputDirector
 		const FString CrashConfigDstAbsolute = FPaths::Combine(OutputDirectory, *CrashConfigFilename);
 		IFileManager::Get().Copy(*CrashConfigDstAbsolute, CrashConfigSrcPath);	// best effort, so don't care about result: couldn't copy -> tough, no config
 	}
+
+	
+}
+
+#if WITH_ADDITIONAL_CRASH_CONTEXTS
+
+thread_local FAdditionalCrashContextStack FAdditionalCrashContextStack::ThreadContextProvider;
+static FCriticalSection GAdditionalProviderLock;
+static FAdditionalCrashContextStack* GProviderHead;
+
+FAdditionalCrashContextStack::FAdditionalCrashContextStack()
+	: Next(nullptr)
+{
+	// Register by appending self to the the linked list. 
+	FScopeLock _(&GAdditionalProviderLock);
+	FAdditionalCrashContextStack** Current = &GProviderHead;
+	while(*Current != nullptr)
+	{
+		Current = &((*Current)->Next);
+	}
+	*Current = this;
+}
+
+FAdditionalCrashContextStack::~FAdditionalCrashContextStack()
+{
+	// Unregister by iterating the list, replacing self with next
+	// on the list.
+	FScopeLock _(&GAdditionalProviderLock);
+	FAdditionalCrashContextStack** Current = &GProviderHead;
+	while (*Current != this)
+	{
+		Current = &((*Current)->Next);
+	}
+	*Current = this->Next;
+}
+
+void FAdditionalCrashContextStack::ExecuteProviders(FCrashContextExtendedWriter& Writer)
+{
+	// Attempt to lock. If a thread crashed while holding the lock
+	// we could potentially deadlock here otherwise.
+	if (GAdditionalProviderLock.TryLock())
+	{
+		FAdditionalCrashContextStack* Provider = GProviderHead;
+		while (Provider)
+		{
+			for (uint32 i = 0; i < Provider->StackIndex; i++)
+			{
+				const FScopedAdditionalCrashContextProvider* Callback = Provider->Stack[i];
+				Callback->Execute(Writer);
+			}
+			Provider = Provider->Next;
+		}
+		GAdditionalProviderLock.Unlock();
+	}
+}
+
+struct FCrashContextExtendedWriterImpl : public FCrashContextExtendedWriter
+{
+	FCrashContextExtendedWriterImpl(const TCHAR* InOutputDirectory)
+		: OutputDirectory(InOutputDirectory)
+	{
+	}
+	
+	void OutputBuffer(const TCHAR* Identifier, const uint8* Data, uint32 DataSize, const TCHAR* Extension)
+	{
+		TCHAR Filename[1024] = { 0 };
+		FCString::Snprintf(Filename, 1024, TEXT("%s/%s.%s"), OutputDirectory, Identifier, Extension);
+		TUniquePtr<IFileHandle> File(IPlatformFile::GetPlatformPhysical().OpenWrite(Filename));
+		if (File)
+		{
+			File->Write(Data, DataSize);
+			File->Flush();
+		}
+	}
+
+	virtual void AddBuffer(const TCHAR* Identifier, const uint8* Data, uint32 DataSize) override
+	{
+		if (Identifier == nullptr || Data == nullptr || DataSize == 0)
+		{
+			return;
+		}
+
+		OutputBuffer(Identifier, Data, DataSize, TEXT(".bin"));
+	}
+
+	virtual void AddString(const TCHAR* Identifier, const TCHAR* DataStr) override
+	{
+		if (Identifier == nullptr || DataStr == nullptr)
+		{
+			return;
+		}
+
+		FPlatformMisc::LowLevelOutputDebugStringf(TEXT("Additional Crash Context (Key=\"%s\", Value=\"%s\")"), Identifier, DataStr);
+
+		FTCHARToUTF8 Converter(DataStr);
+		OutputBuffer(Identifier, (uint8*)Converter.Get(), Converter.Length(), TEXT(".txt"));
+	}
+
+private:
+	const TCHAR* OutputDirectory;
+};
+
+#endif //WITH_ADDITIONAL_CRASH_CONTEXTS 
+
+
+void FGenericCrashContext::DumpAdditionalContext(const TCHAR* CrashFolderAbsolute)
+{
+#if WITH_ADDITIONAL_CRASH_CONTEXTS 
+	FCrashContextExtendedWriterImpl Writer(CrashFolderAbsolute);
+	FAdditionalCrashContextStack::ExecuteProviders(Writer);
+#endif
 }
 
+
 /**
  * Attempts to create the output report directory.
  */
@@ -1234,3 +1350,4 @@ bool RecoveryService::TokenizeSessionName(const FString& SessionName, FString* O
 
 	return true; // Successfully parsed.
 }
+
diff --git a/Engine/Source/Runtime/Core/Private/Windows/WindowsPlatformCrashContext.cpp b/Engine/Source/Runtime/Core/Private/Windows/WindowsPlatformCrashContext.cpp
--- a/Engine/Source/Runtime/Core/Private/Windows/WindowsPlatformCrashContext.cpp
+++ b/Engine/Source/Runtime/Core/Private/Windows/WindowsPlatformCrashContext.cpp
@@ -590,6 +590,25 @@ enum class EErrorReportUI
 
 /** This lock is to prevent an ensure and a crash to concurrently report to CrashReportClient (CRC) when CRC is running in background and waiting for crash/ensure (monitor mode). */
 static FCriticalSection GMonitorLock;
+/**
+ * Guard against additional context callbacks crashing
+ */
+void GuardedDumpAdditionalContext(const TCHAR* CrashDirectoryAbsolute)
+{
+#if !PLATFORM_SEH_EXCEPTIONS_DISABLED
+	__try
+#endif
+	{
+		FGenericCrashContext::DumpAdditionalContext(CrashDirectoryAbsolute);
+	}
+#if !PLATFORM_SEH_EXCEPTIONS_DISABLED
+	__except (TRUE)
+	{
+		// do nothing in case of an error
+		FPlatformMisc::LowLevelOutputDebugString(TEXT("An error occurred while executing addtional crash contexts"));
+	}
+#endif
+}
 
 /**
  * Write required information about the crash to the shared context, and then signal the crash reporter client 
@@ -645,6 +664,7 @@ int32 ReportCrashForMonitor(
 	if (FCommandLine::IsInitialized())
 	{
 		bNoDialog |= FApp::IsUnattended();
+		bNoDialog |= IsRunningDedicatedServer();
 	}
 
 	if (GConfig)
@@ -751,6 +771,9 @@ int32 ReportCrashForMonitor(
 		FCString::Strcpy(SharedContext->CrashFilesDirectory, *CrashDirectoryAbsolute);
 		// Copy the log file to output
 		FGenericCrashContext::DumpLog(CrashDirectoryAbsolute);
+
+		// Dump additional context
+		GuardedDumpAdditionalContext(*CrashDirectoryAbsolute);
 	}
 
 	// Allow the monitor process to take window focus
@@ -1683,6 +1706,7 @@ FORCENOINLINE void ReportEnsure(const TCHAR* ErrorMessage, int NumStackFramesToI
 	ReportEnsureInner(ErrorMessage, NumStackFramesToIgnore + 1);
 }
 
+
 #if !IS_PROGRAM && 0
 namespace {
 	/** Utility class to test crashes during static initialization. */
diff --git a/Engine/Source/Runtime/Core/Public/GenericPlatform/GenericPlatformCrashContext.h b/Engine/Source/Runtime/Core/Public/GenericPlatform/GenericPlatformCrashContext.h
--- a/Engine/Source/Runtime/Core/Public/GenericPlatform/GenericPlatformCrashContext.h
+++ b/Engine/Source/Runtime/Core/Public/GenericPlatform/GenericPlatformCrashContext.h
@@ -9,6 +9,10 @@
 #include "Misc/Optional.h"
 #include "Containers/UnrealString.h"
 
+#define WITH_ADDITIONAL_CRASH_CONTEXTS 1
+
+struct FProgramCounterSymbolInfo;
+
 /** Defines special exit codes used to diagnose abnormal terminations. The code values are arbitrary, but easily recongnizable in decimal. They are meant to be
     used with the out-of-process monitoring/analytics in order to figure out unexpected cases. */
 enum ECrashExitCodes : int32
@@ -394,6 +398,9 @@ public:
 	/** Flushes the logs. In the case of in memory logs is used on this configuration, dumps them to file. */
 	static void DumpLog(const FString& CrashFolderAbsolute);
 
+	/** Collects additional crash context providers. See FAdditionalCrashContextStack. */
+	static void DumpAdditionalContext(const TCHAR* CrashFolderAbsolute);
+
 	/** Initializes a shared crash context from current state. Will not set all fields in Dst. */
 	static void CopySharedCrashContext(FSharedCrashContext& Dst);
 
@@ -547,3 +554,89 @@ namespace RecoveryService
 	/** Tokenize the session name into its components. */
 	CORE_API bool TokenizeSessionName(const FString& SessionName, FString* OutServerName, int32* SeqNum, FString* ProjName, FDateTime* DateTime);
 }
+
+#if WITH_ADDITIONAL_CRASH_CONTEXTS
+
+/**
+ * Interface for callbacks to add context to the crash report.
+ */
+struct FCrashContextExtendedWriter
+{
+	/** Adds a named buffer to the report. Intended for larger payloads. */
+	CORE_API virtual void AddBuffer(const TCHAR* Identifier, const uint8* Data, uint32 DataSize) = 0;
+		
+	/** Add a named buffer containing a string to the report. */
+	CORE_API virtual void AddString(const TCHAR* Identifier, const TCHAR* DataStr) = 0;
+};
+
+/**
+ * A thread local stack of callbacks that can be issued at time of the crash.
+ */
+struct FAdditionalCrashContextStack
+{
+	CORE_API static void PushProvider(struct FScopedAdditionalCrashContextProvider* Provider)
+	{
+		ThreadContextProvider.PushProviderInternal(Provider);
+	}
+
+	CORE_API static void PopProvider()
+	{
+		ThreadContextProvider.PopProviderInternal();
+	}
+
+	static void ExecuteProviders(FCrashContextExtendedWriter& Writer);
+
+private:
+	enum { MaxStackDepth = 16 };
+	static thread_local FAdditionalCrashContextStack ThreadContextProvider;
+	FAdditionalCrashContextStack* Next;
+	const FScopedAdditionalCrashContextProvider* Stack[MaxStackDepth];
+	uint32 StackIndex = 0;
+
+	FAdditionalCrashContextStack();
+	~FAdditionalCrashContextStack();
+	
+	inline void PushProviderInternal(const FScopedAdditionalCrashContextProvider* Provider) 
+	{
+		check(StackIndex < MaxStackDepth);
+		Stack[StackIndex++] = Provider;
+	}
+
+	inline void PopProviderInternal()
+	{
+		check(StackIndex > 0);
+		Stack[--StackIndex] = nullptr;
+	}
+};
+
+struct FScopedAdditionalCrashContextProvider
+{
+public:
+	FScopedAdditionalCrashContextProvider(TUniqueFunction<void(FCrashContextExtendedWriter&)> InFunc)
+		: Func(MoveTemp(InFunc))
+	{
+		FAdditionalCrashContextStack::PushProvider(this);
+	}
+
+	~FScopedAdditionalCrashContextProvider()
+	{
+		FAdditionalCrashContextStack::PopProvider();
+	}
+
+	void Execute(FCrashContextExtendedWriter& Writer) const
+	{
+		Func(Writer);
+	}
+
+private:
+	TUniqueFunction<void(FCrashContextExtendedWriter&)> Func;
+};
+
+#define UE_ADD_CRASH_CONTEXT_SCOPE(FuncExpr) FScopedAdditionalCrashContextProvider ANONYMOUS_VARIABLE(AddCrashCtx)(FuncExpr)
+
+#else
+
+#define UE_ADD_CRASH_CONTEXT_SCOPE(FuncExpr) 
+
+#endif // WITH_ADDITIONAL_CRASH_CONTEXTS
+
diff --git a/Engine/Source/Runtime/Engine/Private/UnrealEngine.cpp b/Engine/Source/Runtime/Engine/Private/UnrealEngine.cpp
--- a/Engine/Source/Runtime/Engine/Private/UnrealEngine.cpp
+++ b/Engine/Source/Runtime/Engine/Private/UnrealEngine.cpp
@@ -9013,6 +9013,19 @@ bool UEngine::PerformError(const TCHAR* Cmd, FOutputDevice& Ar)
 		}, TStatId());
 		return true;
 	}
+	else if (FParse::Command(&Cmd, TEXT("CRASHCONTEXTWRITER")))
+	{
+		TArray<uint8> TestData = { 0xd, 0xe, 0xa, 0xd };
+		UE_ADD_CRASH_CONTEXT_SCOPE([&](FCrashContextExtendedWriter& Writer) { 
+				TCHAR TextBuffer[1024] = { 0 };
+				FCString::Sprintf(TextBuffer, TEXT("Last error was: %d"), FPlatformMisc::GetLastError());
+				Writer.AddString(TEXT("TestText"), TextBuffer); 
+				Writer.AddBuffer(TEXT("TestData"), TestData.GetData(), TestData.Num());
+				
+		});
+		FGenericCrashContext::SetCrashTrigger(ECrashTrigger::Debug);
+		UE_LOG(LogEngine, Fatal, TEXT("Crashing the worker thread at your request"));
+	}
 #endif // !UE_BUILD_SHIPPING
 	return false;
 }
