Use PLATFORM_ macros for the address of the return address when determining owners of allocations.

#rnx

diff --git a/Engine/Source/Runtime/Core/Private/ProfilingDebugging/Windows/MemoryTrace.cpp b/Engine/Source/Runtime/Core/Private/ProfilingDebugging/Windows/MemoryTrace.cpp
--- a/Engine/Source/Runtime/Core/Private/ProfilingDebugging/Windows/MemoryTrace.cpp
+++ b/Engine/Source/Runtime/Core/Private/ProfilingDebugging/Windows/MemoryTrace.cpp
@@ -37,15 +37,9 @@ void*	Backtracer_GetBacktraceId(void*);
 void 	MemoryTrace_InitTags(FMalloc*);
 
 ////////////////////////////////////////////////////////////////////////////////
-#if defined(_MSC_VER)
-#	define UE_RETURN_ADDRESS_ADDRESS()	(void*)_AddressOfReturnAddress()
-#elif defined(__clang__) || defined(__GNUC__)
-#	define UE_RETURN_ADDRESS_ADDRESS()	__builtin_frame_address(0)
-#endif
-
 FORCEINLINE void* GetOwner(bool bLight)
 {
-	void* RetAddrAddr = UE_RETURN_ADDRESS_ADDRESS();
+	void* RetAddrAddr = PLATFORM_RETURN_ADDRESS_POINTER();
 	if (bLight)
 	{
 		return *(void**)RetAddrAddr;
