Fix non-unity

diff --git a/Engine/Source/Runtime/Core/Private/GenericPlatform/GenericPlatformCrashContext.cpp b/Engine/Source/Runtime/Core/Private/GenericPlatform/GenericPlatformCrashContext.cpp
--- a/Engine/Source/Runtime/Core/Private/GenericPlatform/GenericPlatformCrashContext.cpp
+++ b/Engine/Source/Runtime/Core/Private/GenericPlatform/GenericPlatformCrashContext.cpp
@@ -19,6 +19,7 @@
 #include "Misc/EngineVersion.h"
 #include "Misc/EngineBuildSettings.h"
 #include "Misc/OutputDeviceRedirector.h"
+#include "Misc/ScopeLock.h"
 #include "Stats/Stats.h"
 #include "Internationalization/Regex.h"
 #include "Internationalization/TextLocalizationManager.h"
