CompactBinary: Added a compact binary package that represents a Merkle tree of an object and its transitive references

#rnx

diff --git a/Engine/Source/Runtime/Core/Private/Serialization/CompactBinaryPackage.cpp b/Engine/Source/Runtime/Core/Private/Serialization/CompactBinaryPackage.cpp
new file mode 100644
--- /dev/null
+++ b/Engine/Source/Runtime/Core/Private/Serialization/CompactBinaryPackage.cpp
@@ -0,0 +1,390 @@
+// Copyright Epic Games, Inc. All Rights Reserved.
+
+#include "Serialization/CompactBinaryPackage.h"
+
+#include "Algo/BinarySearch.h"
+#include "Serialization/CompactBinarySerialization.h"
+#include "Serialization/CompactBinaryWriter.h"
+
+///////////////////////////////////////////////////////////////////////////////////////////////////////////////////////
+
+FCbAttachment::FCbAttachment(FCbFieldRefIterator InValue, const FBlake3Hash* const InHash)
+{
+	if (InValue)
+	{
+		if (!InValue.IsOwned())
+		{
+			InValue = FCbFieldRefIterator::CloneRange(InValue);
+		}
+
+		CompactBinary = FCbFieldIterator(InValue);
+		Buffer = MoveTemp(InValue).GetOuterBuffer();
+	}
+
+	if (InHash)
+	{
+		Hash = *InHash;
+		if (CompactBinary)
+		{
+			checkSlow(Hash == CompactBinary.GetRangeHash());
+		}
+		else
+		{
+			checkfSlow(Hash == FBlake3Hash(), TEXT("A null or empty field range must use a hash of zero."));
+		}
+	}
+	else if (CompactBinary)
+	{
+		Hash = CompactBinary.GetRangeHash();
+	}
+}
+
+FCbAttachment::FCbAttachment(FSharedBufferConstPtr InBuffer, const FBlake3Hash* const InHash)
+	: Buffer(FSharedBuffer::MakeReadOnly(MoveTemp(InBuffer)))
+{
+	if (InHash)
+	{
+		Hash = *InHash;
+		if (Buffer && Buffer->GetSize())
+		{
+			checkSlow(Hash == FBlake3::HashBuffer(*Buffer));
+		}
+		else
+		{
+			checkfSlow(Hash == FBlake3Hash(), TEXT("A null or empty buffer must use a hash of zero."));
+		}
+	}
+	else if (Buffer && Buffer->GetSize())
+	{
+		Hash = FBlake3::HashBuffer(*Buffer);
+	}
+	else
+	{
+		Buffer.Reset();
+	}
+}
+
+FSharedBufferConstPtr FCbAttachment::AsBinary() const
+{
+	if (!CompactBinary)
+	{
+		return Buffer;
+	}
+
+	FMemoryView SerializedView;
+	if (CompactBinary.TryGetSerializedRangeView(SerializedView))
+	{
+		return SerializedView == Buffer->GetView() ? Buffer : FSharedBuffer::MakeView(SerializedView, *Buffer);
+	}
+
+	return FCbFieldRefIterator::CloneRange(CompactBinary).GetRangeBuffer();
+}
+
+FCbFieldRefIterator FCbAttachment::AsCompactBinary() const
+{
+	return CompactBinary ? FCbFieldRefIterator::MakeRangeView(CompactBinary, Buffer) : FCbFieldRefIterator();
+}
+
+void FCbAttachment::Load(FCbFieldRefIterator& Fields)
+{
+	checkf(Fields.IsBinary(), TEXT("Attachments must start with a binary field."));
+	const FMemoryView View = Fields.AsBinary();
+	if (View.GetSize() > 0)
+	{
+		if (FSharedBufferConstPtr OuterBuffer = Fields.GetOuterBuffer())
+		{
+			Buffer = FSharedBuffer::MakeReadOnly(FSharedBuffer::MakeView(View, *OuterBuffer));
+		}
+		else
+		{
+			Buffer = FSharedBuffer::MakeReadOnly(FSharedBuffer::Clone(View));
+		}
+		++Fields;
+		Hash = Fields.AsAnyReference();
+		checkf(!Fields.HasError(), TEXT("Attachments must be a non-empty binary value with a content hash."));
+		if (Fields.IsReference())
+		{
+			CompactBinary = FCbFieldIterator::MakeRange(*Buffer);
+		}
+		++Fields;
+	}
+	else
+	{
+		++Fields;
+		Buffer = FSharedBufferConstPtr();
+		CompactBinary = FCbFieldIterator();
+		Hash = FBlake3Hash();
+	}
+}
+
+void FCbAttachment::Load(FArchive& Ar, FCbBufferAllocator Allocator)
+{
+	FCbFieldRef BufferField = LoadCompactBinary(Ar, Allocator);
+	checkf(BufferField.IsBinary(), TEXT("Attachments must start with a binary field."));
+	const FMemoryView View = BufferField.AsBinary();
+	if (View.GetSize() > 0)
+	{
+		Buffer = FSharedBuffer::MakeReadOnly(FSharedBuffer::MakeView(View, *BufferField.GetOuterBuffer()));
+		CompactBinary = FCbFieldIterator();
+
+		TArray<uint8, TInlineAllocator<64>> HashBuffer;
+		FCbFieldRef HashField = LoadCompactBinary(Ar, [&HashBuffer](uint64 Size)
+			{
+				HashBuffer.SetNumUninitialized(int32(Size));
+				return FSharedBuffer::MakeView(HashBuffer.GetData(), Size);
+			});
+		Hash = HashField.AsAnyReference();
+		checkf(!HashField.HasError(), TEXT("Attachments must be a non-empty binary value with a content hash."));
+		if (HashField.IsReference())
+		{
+			CompactBinary = FCbFieldIterator::MakeRange(*Buffer);
+		}
+	}
+	else
+	{
+		Buffer = FSharedBufferConstPtr();
+		CompactBinary = FCbFieldIterator();
+		Hash = FBlake3Hash();
+	}
+}
+
+void FCbAttachment::Save(FCbWriter& Writer) const
+{
+	if (CompactBinary)
+	{
+		FMemoryView SerializedView;
+		if (CompactBinary.TryGetSerializedRangeView(SerializedView))
+		{
+			Writer.Binary(SerializedView);
+		}
+		else
+		{
+			Writer.Binary(AsBinary());
+		}
+		Writer.Reference(Hash);
+	}
+	else if (Buffer && Buffer->GetSize())
+	{
+		Writer.Binary(Buffer);
+		Writer.BinaryReference(Hash);
+	}
+	else // Null
+	{
+		Writer.Binary(FMemoryView());
+	}
+}
+
+void FCbAttachment::Save(FArchive& Ar) const
+{
+	FCbWriter Writer;
+	Save(Writer);
+	Writer.Save(Ar);
+}
+
+///////////////////////////////////////////////////////////////////////////////////////////////////////////////////////
+
+void FCbPackage::SetObject(FCbObjectRef InObject, const FBlake3Hash* InObjectHash, FAttachmentResolver* InResolver)
+{
+	if (InObject.CreateIterator())
+	{
+		Object = InObject.IsOwned() ? MoveTemp(InObject) : FCbObjectRef::Clone(InObject);
+		if (InObjectHash)
+		{
+			ObjectHash = *InObjectHash;
+			checkSlow(ObjectHash == Object.GetHash());
+		}
+		else
+		{
+			ObjectHash = Object.GetHash();
+		}
+		if (InResolver)
+		{
+			GatherAttachments(Object.CreateIterator(), *InResolver);
+		}
+	}
+	else
+	{
+		Object = FCbObjectRef();
+		ObjectHash = FBlake3Hash();
+	}
+}
+
+void FCbPackage::AddAttachment(const FCbAttachment& Attachment, FAttachmentResolver* Resolver)
+{
+	if (!Attachment.IsNull())
+	{
+		const int32 Index = Algo::LowerBound(Attachments, Attachment);
+		if (Attachments.IsValidIndex(Index) && Attachments[Index] == Attachment)
+		{
+			FCbAttachment& Existing = Attachments[Index];
+			if (Attachment.IsCompactBinary() && !Existing.IsCompactBinary())
+			{
+				Existing = FCbAttachment(FCbFieldRefIterator::MakeRange(Existing.AsBinary()));
+			}
+		}
+		else
+		{
+			Attachments.Insert(Attachment, Index);
+		}
+		if (Attachment.IsCompactBinary() && Resolver)
+		{
+			GatherAttachments(Attachment.AsCompactBinary(), *Resolver);
+		}
+	}
+}
+
+int32 FCbPackage::RemoveAttachment(const FBlake3Hash& Hash)
+{
+	const int32 Index = Algo::BinarySearchBy(Attachments, Hash,
+		[](const FCbAttachment& Attachment) -> const FBlake3Hash& { return Attachment.GetHash(); });
+	if (Attachments.IsValidIndex(Index))
+	{
+		Attachments.RemoveAt(Index);
+		return 1;
+	}
+	return 0;
+}
+
+bool FCbPackage::Equals(const FCbPackage& Package) const
+{
+	return ObjectHash == Package.ObjectHash && Attachments == Package.Attachments;
+}
+
+const FCbAttachment* FCbPackage::FindAttachment(const FBlake3Hash& Hash) const
+{
+	const int32 Index = Algo::BinarySearchBy(Attachments, Hash,
+		[](const FCbAttachment& Attachment) -> const FBlake3Hash& { return Attachment.GetHash(); });
+	return Attachments.IsValidIndex(Index) ? &Attachments[Index] : nullptr;
+}
+
+void FCbPackage::GatherAttachments(const FCbFieldIterator& Fields, FAttachmentResolver Resolver)
+{
+	Fields.IterateRangeReferences([this, &Resolver](FCbField Field)
+		{
+			const FBlake3Hash& Hash = Field.AsAnyReference();
+			if (FSharedBufferConstPtr Buffer = Resolver(Hash))
+			{
+				if (Field.IsReference())
+				{
+					AddAttachment(FCbAttachment(FCbFieldRefIterator::MakeRange(MoveTemp(Buffer)), Hash), &Resolver);
+				}
+				else
+				{
+					AddAttachment(FCbAttachment(MoveTemp(Buffer), Hash));
+				}
+			}
+		});
+}
+
+void FCbPackage::Load(FCbFieldRefIterator& Fields)
+{
+	*this = FCbPackage();
+	while (Fields)
+	{
+		if (Fields.IsNull())
+		{
+			++Fields;
+			break;
+		}
+		else if (Fields.IsBinary())
+		{
+			FCbAttachment Attachment;
+			Attachment.Load(Fields);
+			AddAttachment(Attachment);
+		}
+		else
+		{
+			checkf(Fields.IsObject(), TEXT("Expected Object, Binary, or Null field when loading a package."));
+			Object = Fields.AsObjectRef();
+			Object.MakeOwned();
+			++Fields;
+			if (Object.CreateIterator())
+			{
+				ObjectHash = Fields.AsReference();
+				checkf(!Fields.HasError(), TEXT("Object must be followed by a Reference with the object hash."));
+				++Fields;
+			}
+			else
+			{
+				Object = FCbObjectRef();
+			}
+		}
+	}
+}
+
+void FCbPackage::Load(FArchive& Ar, FCbBufferAllocator Allocator)
+{
+	uint8 StackBuffer[64];
+	const auto StackAllocator = [&Allocator, &StackBuffer](uint64 Size)
+	{
+		return Size <= sizeof(StackBuffer) ? FSharedBuffer::MakeView(StackBuffer, Size) : Allocator(Size);
+	};
+
+	*this = FCbPackage();
+	for (;;)
+	{
+		FCbFieldRef ValueField = LoadCompactBinary(Ar, StackAllocator);
+		if (ValueField.IsNull())
+		{
+			break;
+		}
+		else if (ValueField.IsBinary())
+		{
+			const FMemoryView View = ValueField.AsBinary();
+			if (View.GetSize() > 0)
+			{
+				FSharedBufferConstPtr Buffer = FSharedBuffer::MakeReadOnly(FSharedBuffer::MakeView(View, *ValueField.GetOuterBuffer()));
+				FCbFieldRef HashField = LoadCompactBinary(Ar, StackAllocator);
+				const FBlake3Hash& Hash = HashField.AsAnyReference();
+				checkf(!HashField.HasError(), TEXT("Attachments must be a non-empty binary value with a content hash."));
+				if (HashField.IsReference())
+				{
+					AddAttachment(FCbAttachment(FCbFieldRefIterator::MakeRange(MoveTemp(Buffer)), Hash));
+				}
+				else
+				{
+					AddAttachment(FCbAttachment(MoveTemp(Buffer), Hash));
+				}
+			}
+		}
+		else
+		{
+			checkf(ValueField.IsObject(), TEXT("Expected Object, Binary, or Null field when loading a package."));
+			Object = ValueField.AsObjectRef();
+			Object.MakeOwned();
+			if (Object.CreateIterator())
+			{
+				FCbFieldRef HashField = LoadCompactBinary(Ar, StackAllocator);
+				ObjectHash = HashField.AsReference();
+				checkf(!HashField.HasError(), TEXT("Object must be followed by a Reference with the object hash."));
+			}
+			else
+			{
+				Object = FCbObjectRef();
+			}
+		}
+	}
+}
+
+void FCbPackage::Save(FCbWriter& Writer) const
+{
+	if (Object.CreateIterator())
+	{
+		Writer.Object(Object);
+		Writer.Reference(ObjectHash);
+	}
+	for (const FCbAttachment& Attachment : Attachments)
+	{
+		Attachment.Save(Writer);
+	}
+	Writer.Null();
+}
+
+void FCbPackage::Save(FArchive& Ar) const
+{
+	FCbWriter Writer;
+	Save(Writer);
+	Writer.Save(Ar);
+}
+
+///////////////////////////////////////////////////////////////////////////////////////////////////////////////////////
diff --git a/Engine/Source/Runtime/Core/Private/Serialization/CompactBinaryValidation.cpp b/Engine/Source/Runtime/Core/Private/Serialization/CompactBinaryValidation.cpp
--- a/Engine/Source/Runtime/Core/Private/Serialization/CompactBinaryValidation.cpp
+++ b/Engine/Source/Runtime/Core/Private/Serialization/CompactBinaryValidation.cpp
@@ -374,13 +374,102 @@ static FCbField ValidateCbField(FMemoryView& View, ECbValidateMode Mode, ECbVali
 	return FCbField(FieldView.GetData(), ExternalType);
 }
 
+static FCbField ValidateCbPackageField(FMemoryView& View, ECbValidateMode Mode, ECbValidateError& Error)
+{
+	if (View.IsEmpty())
+	{
+		if (EnumHasAnyFlags(Mode, ECbValidateMode::Package))
+		{
+			AddError(Error, ECbValidateError::InvalidPackageFormat);
+		}
+		return FCbField();
+	}
+	if (FCbField Field = ValidateCbField(View, Mode, Error))
+	{
+		if (Field.HasName() && EnumHasAnyFlags(Mode, ECbValidateMode::Package))
+		{
+			AddError(Error, ECbValidateError::InvalidPackageFormat);
+		}
+		return Field;
+	}
+	return FCbField();
+}
+
+static FBlake3Hash ValidateCbPackageAttachment(FCbField& Value, FMemoryView& View, ECbValidateMode Mode, ECbValidateError& Error)
+{
+	const FMemoryView ValueView = Value.AsBinary();
+	if (Value.HasError() && EnumHasAnyFlags(Mode, ECbValidateMode::Package))
+	{
+		if (EnumHasAnyFlags(Mode, ECbValidateMode::Package))
+		{
+			AddError(Error, ECbValidateError::InvalidPackageFormat);
+		}
+	}
+	else if (ValueView.GetSize())
+	{
+		if (FCbField HashField = ValidateCbPackageField(View, Mode, Error))
+		{
+			const FBlake3Hash Hash = HashField.AsAnyReference();
+			if (EnumHasAnyFlags(Mode, ECbValidateMode::Package))
+			{
+				if (HashField.HasError())
+				{
+					AddError(Error, ECbValidateError::InvalidPackageFormat);
+				}
+				else if (Hash != FBlake3::HashBuffer(ValueView))
+				{
+					AddError(Error, ECbValidateError::InvalidPackageHash);
+				}
+			}
+			return Hash;
+		}
+	}
+	return FBlake3Hash();
+}
+
+static FBlake3Hash ValidateCbPackageObject(FCbField& Value, FMemoryView& View, ECbValidateMode Mode, ECbValidateError& Error)
+{
+	FCbObject Object = Value.AsObject();
+	if (Value.HasError())
+	{
+		if (EnumHasAnyFlags(Mode, ECbValidateMode::Package))
+		{
+			AddError(Error, ECbValidateError::InvalidPackageFormat);
+		}
+	}
+	else if (FCbField HashField = ValidateCbPackageField(View, Mode, Error))
+	{
+		const FBlake3Hash Hash = HashField.AsAnyReference();
+		if (EnumHasAnyFlags(Mode, ECbValidateMode::Package))
+		{
+			if (!Object.CreateIterator())
+			{
+				AddError(Error, ECbValidateError::NullPackageObject);
+			}
+			if (HashField.HasError())
+			{
+				AddError(Error, ECbValidateError::InvalidPackageFormat);
+			}
+			else if (Hash != Value.GetHash())
+			{
+				AddError(Error, ECbValidateError::InvalidPackageHash);
+			}
+		}
+		return Hash;
+	}
+	return FBlake3Hash();
+}
+
 ECbValidateError ValidateCompactBinary(FMemoryView View, ECbValidateMode Mode, ECbFieldType Type)
 {
 	ECbValidateError Error = ECbValidateError::None;
-	ValidateCbField(View, Mode, Error, Type);
-	if (!View.IsEmpty() && EnumHasAnyFlags(Mode, ECbValidateMode::Padding))
+	if (EnumHasAnyFlags(Mode, ECbValidateMode::All))
 	{
-		AddError(Error, ECbValidateError::Padding);
+		ValidateCbField(View, Mode, Error, Type);
+		if (!View.IsEmpty() && EnumHasAnyFlags(Mode, ECbValidateMode::Padding))
+		{
+			AddError(Error, ECbValidateError::Padding);
+		}
 	}
 	return Error;
 }
@@ -388,9 +477,94 @@ ECbValidateError ValidateCompactBinary(FMemoryView View, ECbValidateMode Mode, E
 ECbValidateError ValidateCompactBinaryRange(FMemoryView View, ECbValidateMode Mode)
 {
 	ECbValidateError Error = ECbValidateError::None;
-	while (!View.IsEmpty())
+	if (EnumHasAnyFlags(Mode, ECbValidateMode::All))
+	{
+		while (!View.IsEmpty())
+		{
+			ValidateCbField(View, Mode, Error);
+		}
+	}
+	return Error;
+}
+
+ECbValidateError ValidateCompactBinaryAttachment(FMemoryView View, ECbValidateMode Mode)
+{
+	ECbValidateError Error = ECbValidateError::None;
+	if (EnumHasAnyFlags(Mode, ECbValidateMode::All))
 	{
-		ValidateCbField(View, Mode, Error);
+		if (FCbField Value = ValidateCbPackageField(View, Mode, Error))
+		{
+			ValidateCbPackageAttachment(Value, View, Mode, Error);
+		}
+		if (!View.IsEmpty() && EnumHasAnyFlags(Mode, ECbValidateMode::Padding))
+		{
+			AddError(Error, ECbValidateError::Padding);
+		}
+	}
+	return Error;
+}
+
+ECbValidateError ValidateCompactBinaryPackage(FMemoryView View, ECbValidateMode Mode)
+{
+	TArray<FBlake3Hash, TInlineAllocator<16>> Attachments;
+	ECbValidateError Error = ECbValidateError::None;
+	if (EnumHasAnyFlags(Mode, ECbValidateMode::All))
+	{
+		uint32 ObjectCount = 0;
+		while (FCbField Value = ValidateCbPackageField(View, Mode, Error))
+		{
+			if (Value.IsBinary())
+			{
+				const FBlake3Hash Hash = ValidateCbPackageAttachment(Value, View, Mode, Error);
+				if (EnumHasAnyFlags(Mode, ECbValidateMode::Package))
+				{
+					Attachments.Add(Hash);
+					if (Value.AsBinary().IsEmpty())
+					{
+						AddError(Error, ECbValidateError::NullPackageAttachment);
+					}
+				}
+			}
+			else if (Value.IsObject())
+			{
+				ValidateCbPackageObject(Value, View, Mode, Error);
+				if (++ObjectCount > 1 && EnumHasAnyFlags(Mode, ECbValidateMode::Package))
+				{
+					AddError(Error, ECbValidateError::MultiplePackageObjects);
+				}
+			}
+			else if (Value.IsNull())
+			{
+				break;
+			}
+			else if (EnumHasAnyFlags(Mode, ECbValidateMode::Package))
+			{
+				AddError(Error, ECbValidateError::InvalidPackageFormat);
+			}
+
+			if (EnumHasAnyFlags(Error, ECbValidateError::OutOfBounds))
+			{
+				break;
+			}
+		}
+
+		if (!View.IsEmpty() && EnumHasAnyFlags(Mode, ECbValidateMode::Padding))
+		{
+			AddError(Error, ECbValidateError::Padding);
+		}
+
+		if (Attachments.Num() && EnumHasAnyFlags(Mode, ECbValidateMode::Package))
+		{
+			Algo::Sort(Attachments);
+			for (const FBlake3Hash* It = Attachments.GetData(), *End = It + Attachments.Num() - 1; It != End; ++It)
+			{
+				if (It[0] == It[1])
+				{
+					AddError(Error, ECbValidateError::DuplicateAttachments);
+					break;
+				}
+			}
+		}
 	}
 	return Error;
 }
diff --git a/Engine/Source/Runtime/Core/Private/Tests/Serialization/CompactBinaryPackageTest.cpp b/Engine/Source/Runtime/Core/Private/Tests/Serialization/CompactBinaryPackageTest.cpp
new file mode 100644
--- /dev/null
+++ b/Engine/Source/Runtime/Core/Private/Tests/Serialization/CompactBinaryPackageTest.cpp
@@ -0,0 +1,559 @@
+// Copyright Epic Games, Inc. All Rights Reserved.
+
+#include "Serialization/CompactBinaryPackage.h"
+
+#include "Algo/IsSorted.h"
+#include "Misc/AutomationTest.h"
+#include "Serialization/BufferArchive.h"
+#include "Serialization/CompactBinaryValidation.h"
+#include "Serialization/CompactBinaryWriter.h"
+#include "Serialization/MemoryReader.h"
+
+#if WITH_DEV_AUTOMATION_TESTS
+
+///////////////////////////////////////////////////////////////////////////////////////////////////////////////////////
+
+static constexpr EAutomationTestFlags::Type CompactBinaryPackageTestFlags =
+	EAutomationTestFlags::Type(EAutomationTestFlags::ApplicationContextMask | EAutomationTestFlags::SmokeFilter);
+
+///////////////////////////////////////////////////////////////////////////////////////////////////////////////////////
+
+IMPLEMENT_SIMPLE_AUTOMATION_TEST(FCbAttachmentTest, "System.Core.Serialization.CbAttachment", CompactBinaryPackageTestFlags)
+bool FCbAttachmentTest::RunTest(const FString& Parameters)
+{
+	const auto TestSaveLoadValidate = [this](const TCHAR* Test, const FCbAttachment& Attachment)
+	{
+		TCbWriter<256> Writer;
+		FBufferArchive WriteAr;
+		Attachment.Save(Writer);
+		Attachment.Save(WriteAr);
+		FCbFieldRefIterator Fields = Writer.Save();
+
+		TestTrue(FString::Printf(TEXT("FCbAttachment(%s).Save()->Equals"), Test),
+			MakeMemoryView(WriteAr).EqualBytes(Fields.GetRangeBuffer()->GetView()));
+		TestEqual(FString::Printf(TEXT("FCbAttachment(%s).Save()->ValidateRange"), Test),
+			ValidateCompactBinaryRange(MakeMemoryView(WriteAr), ECbValidateMode::All), ECbValidateError::None);
+		TestEqual(FString::Printf(TEXT("FCbAttachment(%s).Save()->ValidateAttachment"), Test),
+			ValidateCompactBinaryAttachment(MakeMemoryView(WriteAr), ECbValidateMode::All), ECbValidateError::None);
+
+		FCbAttachment FromFields;
+		FromFields.Load(Fields);
+		TestTrue(FString::Printf(TEXT("FCbAttachment(%s).Load(Iterator)->AtEnd"), Test), !bool(Fields));
+		TestEqual(FString::Printf(TEXT("FCbAttachment(%s).Load(Iterator)->Equals"), Test), FromFields, Attachment);
+
+		FCbAttachment FromArchive;
+		FMemoryReader ReadAr(WriteAr);
+		FromArchive.Load(ReadAr, [](uint64 Size) { return FSharedBuffer::Alloc(Size); });
+		TestTrue(FString::Printf(TEXT("FCbAttachment(%s).Load(Archive)->AtEnd"), Test), ReadAr.AtEnd());
+		TestEqual(FString::Printf(TEXT("FCbAttachment(%s).Load(Archive)->Equals"), Test), FromArchive, Attachment);
+	};
+
+	// Empty Attachment
+	{
+		FCbAttachment Attachment;
+		TestTrue(TEXT("FCbAttachment(Null).IsNull()"), Attachment.IsNull());
+		TestFalse(TEXT("FCbAttachment(Null) as bool"), bool(Attachment));
+		TestFalse(TEXT("FCbAttachment(Null).AsBinary()"), Attachment.AsBinary().IsValid());
+		TestFalse(TEXT("FCbAttachment(Null).AsCompactBinary()"), Attachment.AsCompactBinary().HasValue());
+		TestFalse(TEXT("FCbAttachment(Null).IsBinary()"), Attachment.IsBinary());
+		TestFalse(TEXT("FCbAttachment(Null).IsCompactBinary()"), Attachment.IsCompactBinary());
+		TestEqual(TEXT("FCbAttachment(Null).GetHash()"), Attachment.GetHash(), FBlake3Hash());
+		TestSaveLoadValidate(TEXT("Null"), Attachment);
+	}
+
+	// Binary Attachment
+	{
+		FSharedBufferConstPtr Buffer = FSharedBuffer::MakeReadOnly(FSharedBuffer::Clone(MakeMemoryView<uint8>({0, 1, 2, 3})));
+		FCbAttachment Attachment(Buffer);
+		TestFalse(TEXT("FCbAttachment(Binary).IsNull()"), Attachment.IsNull());
+		TestTrue(TEXT("FCbAttachment(Binary) as bool"), bool(Attachment));
+		TestEqual(TEXT("FCbAttachment(Binary).AsBinary()"), Attachment.AsBinary(), Buffer);
+		TestFalse(TEXT("FCbAttachment(Binary).AsCompactBinary()"), Attachment.AsCompactBinary().HasValue());
+		TestTrue(TEXT("FCbAttachment(Binary).IsBinary()"), Attachment.IsBinary());
+		TestFalse(TEXT("FCbAttachment(Binary).IsCompactBinary()"), Attachment.IsCompactBinary());
+		TestEqual(TEXT("FCbAttachment(Binary).GetHash()"), Attachment.GetHash(), FBlake3::HashBuffer(Buffer->GetView()));
+		TestSaveLoadValidate(TEXT("Binary"), Attachment);
+	}
+
+	// Compact Binary Attachment
+	{
+		FCbWriter Writer;
+		Writer << "Name"_ASV << 42;
+		FCbFieldRefIterator Fields = Writer.Save();
+		FCbAttachment Attachment(Fields);
+		TestFalse(TEXT("FCbAttachment(CompactBinary).IsNull()"), Attachment.IsNull());
+		TestTrue(TEXT("FCbAttachment(CompactBinary) as bool"), bool(Attachment));
+		TestEqual(TEXT("FCbAttachment(CompactBinary).AsBinary()"), Attachment.AsBinary(), Fields.GetRangeBuffer());
+		TestEqual(TEXT("FCbAttachment(CompactBinary).AsCompactBinary()"), Attachment.AsCompactBinary(), Fields);
+		TestTrue(TEXT("FCbAttachment(CompactBinary).IsBinary()"), Attachment.IsBinary());
+		TestTrue(TEXT("FCbAttachment(CompactBinary).IsCompactBinary()"), Attachment.IsCompactBinary());
+		TestEqual(TEXT("FCbAttachment(CompactBinary).GetHash()"), Attachment.GetHash(), Fields.GetRangeHash());
+		TestSaveLoadValidate(TEXT("CompactBinary"), Attachment);
+	}
+
+	// Binary View
+	{
+		const uint8 Value[]{0, 1, 2, 3};
+		FSharedBufferConstPtr Buffer = FSharedBuffer::MakeView(MakeMemoryView(Value));
+		FCbAttachment Attachment(Buffer);
+		TestFalse(TEXT("FCbAttachment(BinaryView).IsNull()"), Attachment.IsNull());
+		TestTrue(TEXT("FCbAttachment(BinaryView) as bool"), bool(Attachment));
+		TestNotEqual(TEXT("FCbAttachment(BinaryView).AsBinary()"), Attachment.AsBinary(), Buffer);
+		TestTrue(TEXT("FCbAttachment(BinaryView).AsBinary()"),
+			Attachment.AsBinary()->GetView().EqualBytes(Buffer->GetView()));
+		TestFalse(TEXT("FCbAttachment(BinaryView).AsCompactBinary()"), Attachment.AsCompactBinary().HasValue());
+		TestTrue(TEXT("FCbAttachment(BinaryView).IsBinary()"), Attachment.IsBinary());
+		TestFalse(TEXT("FCbAttachment(BinaryView).IsCompactBinary()"), Attachment.IsCompactBinary());
+		TestEqual(TEXT("FCbAttachment(BinaryView).GetHash()"),
+			Attachment.GetHash(), FBlake3::HashBuffer(Buffer->GetView()));
+	}
+
+	// Compact Binary View
+	{
+		FCbWriter Writer;
+		Writer << "Name"_ASV << 42;
+		FCbFieldRefIterator Fields = Writer.Save();
+		FCbFieldRefIterator FieldsView = FCbFieldRefIterator::MakeRangeView(FCbFieldIterator(Fields));
+		FCbAttachment Attachment(FieldsView);
+		TestFalse(TEXT("FCbAttachment(CompactBinaryView).IsNull()"), Attachment.IsNull());
+		TestTrue(TEXT("FCbAttachment(CompactBinaryView) as bool"), bool(Attachment));
+		TestNotEqual(TEXT("FCbAttachment(CompactBinaryView).AsBinary()"),
+			Attachment.AsBinary(), FieldsView.GetRangeBuffer());
+		TestTrue(TEXT("FCbAttachment(CompactBinaryView).AsCompactBinary()"),
+			Attachment.AsCompactBinary().GetRangeView().EqualBytes(Fields.GetRangeView()));
+		TestTrue(TEXT("FCbAttachment(CompactBinaryView).IsBinary()"), Attachment.IsBinary());
+		TestTrue(TEXT("FCbAttachment(CompactBinaryView).IsCompactBinary()"), Attachment.IsCompactBinary());
+		TestEqual(TEXT("FCbAttachment(CompactBinaryView).GetHash()"), Attachment.GetHash(), Fields.GetRangeHash());
+	}
+
+	// Binary Load from View
+	{
+		const uint8 Value[]{0, 1, 2, 3};
+		const FSharedBufferConstPtr Buffer = FSharedBuffer::MakeView(MakeMemoryView(Value));
+		FCbAttachment Attachment(Buffer);
+
+		FCbWriter Writer;
+		Attachment.Save(Writer);
+		FCbFieldRefIterator Fields = Writer.Save();
+		FCbFieldRefIterator FieldsView = FCbFieldRefIterator::MakeRangeView(FCbFieldIterator(Fields));
+
+		Attachment.Load(FieldsView);
+		TestFalse(TEXT("FCbAttachment(LoadBinaryView).IsNull()"), Attachment.IsNull());
+		TestTrue(TEXT("FCbAttachment(LoadBinaryView) as bool"), bool(Attachment));
+		TestFalse(TEXT("FCbAttachment(LoadBinaryView).AsBinary()->!InView"),
+			FieldsView.GetRangeBuffer()->GetView().Contains(Attachment.AsBinary()->GetView()));
+		TestTrue(TEXT("FCbAttachment(LoadBinaryView).AsBinary()->EqualBytes"),
+			Attachment.AsBinary()->GetView().EqualBytes(Buffer->GetView()));
+		TestFalse(TEXT("FCbAttachment(LoadBinaryView).AsCompactBinary()"), Attachment.AsCompactBinary().HasValue());
+		TestTrue(TEXT("FCbAttachment(LoadBinaryView).IsBinary()"), Attachment.IsBinary());
+		TestFalse(TEXT("FCbAttachment(LoadBinaryView).IsCompactBinary()"), Attachment.IsCompactBinary());
+		TestEqual(TEXT("FCbAttachment(LoadBinaryView).GetHash()"),
+			Attachment.GetHash(), FBlake3::HashBuffer(MakeMemoryView(Value)));
+	}
+
+	// Compact Binary Load from View
+	{
+		FCbWriter ValueWriter;
+		ValueWriter << "Name"_ASV << 42;
+		const FCbFieldRefIterator Value = ValueWriter.Save();
+		TestEqual(TEXT("FCbAttachment(LoadCompactBinaryView).Validate"),
+			ValidateCompactBinaryRange(Value.GetRangeView(), ECbValidateMode::All), ECbValidateError::None);
+		FCbAttachment Attachment(Value);
+
+		FCbWriter Writer;
+		Attachment.Save(Writer);
+		FCbFieldRefIterator Fields = Writer.Save();
+		FCbFieldRefIterator FieldsView = FCbFieldRefIterator::MakeRangeView(FCbFieldIterator(Fields));
+
+		Attachment.Load(FieldsView);
+		TestFalse(TEXT("FCbAttachment(LoadCompactBinaryView).IsNull()"), Attachment.IsNull());
+		TestTrue(TEXT("FCbAttachment(LoadCompactBinaryView) as bool"), bool(Attachment));
+		TestTrue(TEXT("FCbAttachment(LoadCompactBinaryView).AsBinary()->EqualBytes"),
+			Attachment.AsBinary()->GetView().EqualBytes(Value.GetRangeView()));
+		TestFalse(TEXT("FCbAttachment(LoadCompactBinaryView).AsCompactBinary()->!InView"),
+			FieldsView.GetRangeBuffer()->GetView().Contains(Attachment.AsCompactBinary().GetRangeBuffer()->GetView()));
+		TestTrue(TEXT("FCbAttachment(LoadCompactBinaryView).IsBinary()"), Attachment.IsBinary());
+		TestTrue(TEXT("FCbAttachment(LoadCompactBinaryView).IsCompactBinary()"), Attachment.IsCompactBinary());
+		TestEqual(TEXT("FCbAttachment(LoadCompactBinaryView).GetHash()"), Attachment.GetHash(), Value.GetRangeHash());
+	}
+
+	// Compact Binary Uniform Sub-View
+	{
+		const FSharedBufferConstPtr Buffer = FSharedBuffer::Clone(MakeMemoryView<uint8>({0, 1, 2, 3}));
+		const FCbFieldIterator Fields = FCbFieldIterator::MakeRange(Buffer->GetView().RightChop(2), ECbFieldType::IntegerPositive);
+		const FCbFieldRefIterator SavedFieldRefs = FCbFieldRefIterator::CloneRange(Fields);
+		FCbFieldRefIterator FieldRefs = FCbFieldRefIterator::MakeRangeView(Fields, Buffer);
+		FCbAttachment Attachment(FieldRefs);
+		const FSharedBufferConstPtr Binary = Attachment.AsBinary();
+		TestEqual(TEXT("FCbAttachment(CompactBinaryUniformSubView).AsCompactBinary()->Equals()"),
+			Attachment.AsCompactBinary(), FieldRefs);
+		TestEqual(TEXT("FCbAttachment(CompactBinaryUniformSubView).AsBinary()->GetSize()"),
+			Binary->GetSize(), SavedFieldRefs.GetRangeSize());
+		TestTrue(TEXT("FCbAttachment(CompactBinaryUniformSubView).AsBinary()->EqualBytes()"),
+			Binary->GetView().EqualBytes(SavedFieldRefs.GetRangeView()));
+		TestEqual(TEXT("FCbAttachment(CompactBinaryUniformSubView).GetHash()"),
+			Attachment.GetHash(), SavedFieldRefs.GetRangeHash());
+		TestSaveLoadValidate(TEXT("CompactBinaryUniformSubView"), Attachment);
+	}
+
+	// Binary Null
+	{
+		const FCbAttachment Attachment(FSharedBufferConstPtr{});
+		TestTrue(TEXT("FCbAttachment(BinaryNull).IsNull()"), Attachment.IsNull());
+		TestFalse(TEXT("FCbAttachment(BinaryNull).IsBinary()"), Attachment.IsBinary());
+		TestFalse(TEXT("FCbAttachment(BinaryNull).IsCompactBinary()"), Attachment.IsCompactBinary());
+		TestEqual(TEXT("FCbAttachment(BinaryNull).GetHash()"), Attachment.GetHash(), FBlake3Hash());
+	}
+
+	// Binary Empty
+	{
+		const FCbAttachment Attachment(FSharedBuffer::Alloc(0));
+		TestTrue(TEXT("FCbAttachment(BinaryEmpty).IsNull()"), Attachment.IsNull());
+		TestFalse(TEXT("FCbAttachment(BinaryEmpty).IsBinary()"), Attachment.IsBinary());
+		TestFalse(TEXT("FCbAttachment(BinaryEmpty).IsCompactBinary()"), Attachment.IsCompactBinary());
+		TestEqual(TEXT("FCbAttachment(BinaryEmpty).GetHash()"), Attachment.GetHash(), FBlake3Hash());
+	}
+
+	// Compact Binary Empty
+	{
+		const FCbAttachment Attachment(FCbFieldRefIterator{});
+		TestTrue(TEXT("FCbAttachment(CompactBinaryEmpty).IsNull()"), Attachment.IsNull());
+		TestFalse(TEXT("FCbAttachment(CompactBinaryEmpty).IsBinary()"), Attachment.IsBinary());
+		TestFalse(TEXT("FCbAttachment(CompactBinaryEmpty).IsCompactBinary()"), Attachment.IsCompactBinary());
+		TestEqual(TEXT("FCbAttachment(CompactBinaryEmpty).GetHash()"), Attachment.GetHash(), FBlake3Hash());
+	}
+
+	return true;
+}
+
+IMPLEMENT_SIMPLE_AUTOMATION_TEST(FCbPackageTest, "System.Core.Serialization.CbPackage", CompactBinaryPackageTestFlags)
+bool FCbPackageTest::RunTest(const FString& Parameters)
+{
+	const auto TestSaveLoadValidate = [this](const TCHAR* Test, const FCbPackage& Package)
+	{
+		TCbWriter<256> Writer;
+		FBufferArchive WriteAr;
+		Package.Save(Writer);
+		Package.Save(WriteAr);
+		FCbFieldRefIterator Fields = Writer.Save();
+
+		TestTrue(FString::Printf(TEXT("FCbPackage(%s).Save()->Equals"), Test),
+			MakeMemoryView(WriteAr).EqualBytes(Fields.GetRangeBuffer()->GetView()));
+		TestEqual(FString::Printf(TEXT("FCbPackage(%s).Save()->ValidateRange"), Test),
+			ValidateCompactBinaryRange(MakeMemoryView(WriteAr), ECbValidateMode::All), ECbValidateError::None);
+		TestEqual(FString::Printf(TEXT("FCbPackage(%s).Save()->ValidatePackage"), Test),
+			ValidateCompactBinaryPackage(MakeMemoryView(WriteAr), ECbValidateMode::All), ECbValidateError::None);
+
+		FCbPackage FromFields;
+		FromFields.Load(Fields);
+		TestFalse(FString::Printf(TEXT("FCbPackage(%s).Load(Iterator)->AtEnd"), Test), bool(Fields));
+		TestEqual(FString::Printf(TEXT("FCbPackage(%s).Load(Iterator)->Equals"), Test), FromFields, Package);
+
+		FCbPackage FromArchive;
+		FMemoryReader ReadAr(WriteAr);
+		FromArchive.Load(ReadAr, [](uint64 Size) { return FSharedBuffer::Alloc(Size); });
+		TestTrue(FString::Printf(TEXT("FCbPackage(%s).Load(Archive)->AtEnd"), Test), ReadAr.AtEnd());
+		TestEqual(FString::Printf(TEXT("FCbPackage(%s).Load(Archive)->Equals"), Test), FromArchive, Package);
+	};
+
+	// Empty
+	{
+		FCbPackage Package;
+		TestTrue(TEXT("FCbPackage(Empty).IsNull()"), Package.IsNull());
+		TestFalse(TEXT("FCbPackage(Empty) as bool"), bool(Package));
+		TestEqual(TEXT("FCbPackage(Empty).GetAttachments()"), Package.GetAttachments().Num(), 0);
+		TestSaveLoadValidate(TEXT("Empty"), Package);
+	}
+
+	// Object Only
+	{
+		TCbWriter<256> Writer;
+		Writer.BeginObject();
+		Writer << "Field" << 42;
+		Writer.EndObject();
+
+		const FCbObjectRef Object = Writer.Save().AsObjectRef();
+		FCbPackage Package(Object);
+		TestFalse(TEXT("FCbPackage(Object).IsNull()"), Package.IsNull());
+		TestTrue(TEXT("FCbPackage(Object) as bool"), bool(Package));
+		TestEqual(TEXT("FCbPackage(Object).GetAttachments()"), Package.GetAttachments().Num(), 0);
+		TestEqual(TEXT("FCbPackage(Object).GetObject()->IsClone"), Package.GetObject().GetOuterBuffer(), Object.GetOuterBuffer());
+		TestEqual(TEXT("FCbPackage(Object).GetObject()"), Package.GetObject()["Field"].AsInt32(), 42);
+		TestEqual(TEXT("FCbPackage(Object).GetObjectHash()"), Package.GetObjectHash(), Package.GetObject().GetHash());
+		TestSaveLoadValidate(TEXT("Object"), Package);
+	}
+
+	// Object View Only
+	{
+		TCbWriter<256> Writer;
+		Writer.BeginObject();
+		Writer << "Field" << 42;
+		Writer.EndObject();
+
+		const FCbObjectRef Object = Writer.Save().AsObjectRef();
+		FCbPackage Package(FCbObjectRef::MakeView(Object));
+		TestFalse(TEXT("FCbPackage(Object).IsNull()"), Package.IsNull());
+		TestTrue(TEXT("FCbPackage(Object) as bool"), bool(Package));
+		TestEqual(TEXT("FCbPackage(Object).GetAttachments()"), Package.GetAttachments().Num(), 0);
+		TestNotEqual(TEXT("FCbPackage(Object).GetObject()->IsClone"), Package.GetObject().GetOuterBuffer(), Object.GetOuterBuffer());
+		TestEqual(TEXT("FCbPackage(Object).GetObject()"), Package.GetObject()["Field"].AsInt32(), 42);
+		TestEqual(TEXT("FCbPackage(Object).GetObjectHash()"), Package.GetObjectHash(), Package.GetObject().GetHash());
+		TestSaveLoadValidate(TEXT("Object"), Package);
+	}
+
+	// Attachment Only
+	{
+		FCbObjectRef Object;
+		{
+			TCbWriter<256> Writer;
+			Writer.BeginObject();
+			Writer << "Field" << 42;
+			Writer.EndObject();
+			Object = Writer.Save().AsObjectRef();
+		}
+		FCbFieldRef Field = FCbFieldRef::Clone(Object["Field"]);
+
+		FCbPackage Package;
+		Package.AddAttachment(FCbAttachment(FCbFieldRefIterator::MakeSingle(Object.AsFieldRef())));
+		Package.AddAttachment(FCbAttachment(Field.GetBuffer()));
+
+		TestFalse(TEXT("FCbPackage(Attachments).IsNull()"), Package.IsNull());
+		TestTrue(TEXT("FCbPackage(Attachments) as bool"), bool(Package));
+		TestEqual(TEXT("FCbPackage(Attachments).GetAttachments()"), Package.GetAttachments().Num(), 2);
+		TestTrue(TEXT("FCbPackage(Attachments).GetObject()"), Package.GetObject().Equals(FCbObjectRef()));
+		TestEqual(TEXT("FCbPackage(Attachments).GetObjectHash()"), Package.GetObjectHash(), FBlake3Hash());
+		TestSaveLoadValidate(TEXT("Attachments"), Package);
+
+		const FCbAttachment* const ObjectAttachment = Package.FindAttachment(Object.GetHash());
+		const FCbAttachment* const FieldAttachment = Package.FindAttachment(Field.GetHash());
+
+		TestTrue(TEXT("FCbPackage(Attachments).FindAttachment(Object)"),
+			ObjectAttachment && ObjectAttachment->AsCompactBinary().AsObjectRef().Equals(Object));
+		TestTrue(TEXT("FCbPackage(Attachments).FindAttachment(Field)"),
+			FieldAttachment && FieldAttachment->AsBinary() == Field.GetBuffer());
+
+		Package.AddAttachment(FCbAttachment(FSharedBuffer::Clone(Object.GetView())));
+		Package.AddAttachment(FCbAttachment(FCbFieldRefIterator::CloneRange(FCbFieldIterator::MakeSingle(Field))));
+
+		TestEqual(TEXT("FCbPackage(Attachments).GetAttachments()"), Package.GetAttachments().Num(), 2);
+		TestEqual(TEXT("FCbPackage(Attachments).FindAttachment(Object, Re-Add)"),
+			Package.FindAttachment(Object.GetHash()), ObjectAttachment);
+		TestEqual(TEXT("FCbPackage(Attachments).FindAttachment(Field, Re-Add)"),
+			Package.FindAttachment(Field.GetHash()), FieldAttachment);
+
+		TestTrue(TEXT("FCbPackage(Attachments).FindAttachment(ObjectAsObject)"),
+			ObjectAttachment && ObjectAttachment->AsCompactBinary().AsObjectRef().Equals(Object));
+		TestTrue(TEXT("FCbPackage(Attachments).FindAttachment(ObjectAsBinary)"),
+			ObjectAttachment && ObjectAttachment->AsBinary() == Object.GetBuffer());
+		TestTrue(TEXT("FCbPackage(Attachments).FindAttachment(FieldAsField)"),
+			FieldAttachment && FieldAttachment->AsCompactBinary().Equals(Field));
+		TestTrue(TEXT("FCbPackage(Attachments).FindAttachment(FieldAsBinary)"),
+			FieldAttachment && FieldAttachment->AsBinary() == Field.GetBuffer());
+
+		TestTrue(TEXT("FCbPackage(Attachments).GetAttachments()->Sorted"),
+			Algo::IsSorted(Package.GetAttachments()));
+	}
+
+	// Shared Values
+	const uint8 Level4Values[]{0, 1, 2, 3};
+	FSharedBufferConstPtr Level4 = FSharedBuffer::MakeView(MakeMemoryView(Level4Values));
+	const FBlake3Hash Level4Hash = FBlake3::HashBuffer(*Level4);
+
+	FCbFieldRef Level3;
+	{
+		TCbWriter<256> Writer;
+		Writer.Name("Level4").BinaryReference(Level4Hash);
+		Level3 = Writer.Save();
+	}
+	const FBlake3Hash Level3Hash = Level3.GetHash();
+
+	FCbArrayRef Level2;
+	{
+		TCbWriter<256> Writer;
+		Writer.Name("Level3");
+		Writer.BeginArray();
+		Writer.Reference(Level3Hash);
+		Writer.EndArray();
+		Level2 = Writer.Save().AsArrayRef();
+	}
+	const FBlake3Hash Level2Hash = Level2.AsField().GetHash();
+
+	FCbObjectRef Level1;
+	{
+		TCbWriter<256> Writer;
+		Writer.BeginObject();
+		Writer.Name("Level2").Reference(Level2Hash);
+		Writer.EndObject();
+		Level1 = Writer.Save().AsObjectRef();
+	}
+	const FBlake3Hash Level1Hash = Level1.AsField().GetHash();
+
+	const auto Resolver = [&Level2, &Level2Hash, &Level3, &Level3Hash, &Level4, &Level4Hash]
+		(const FBlake3Hash& Hash) -> FSharedBufferConstPtr
+		{
+			return
+				Hash == Level2Hash ? Level2.GetBuffer() :
+				Hash == Level3Hash ? Level3.GetBuffer() :
+				Hash == Level4Hash ? Level4 :
+				FSharedBufferConstPtr();
+		};
+
+	// Object + Attachments
+	{
+		FCbPackage Package;
+		Package.SetObject(Level1, Level1Hash, Resolver);
+
+		TestFalse(TEXT("FCbPackage(Object+Attachments).IsNull()"), Package.IsNull());
+		TestTrue(TEXT("FCbPackage(Object+Attachments) as bool"), bool(Package));
+		TestEqual(TEXT("FCbPackage(Object+Attachments).GetAttachments()"), Package.GetAttachments().Num(), 3);
+		TestTrue(TEXT("FCbPackage(Object+Attachments).GetObject()"),
+			Package.GetObject().GetBuffer() == Level1.GetBuffer());
+		TestEqual(TEXT("FCbPackage(Object+Attachments).GetObjectHash()"), Package.GetObjectHash(), Level1Hash);
+		TestSaveLoadValidate(TEXT("Object+Attachments"), Package);
+
+		const FCbAttachment* const Level2Attachment = Package.FindAttachment(Level2Hash);
+		const FCbAttachment* const Level3Attachment = Package.FindAttachment(Level3Hash);
+		const FCbAttachment* const Level4Attachment = Package.FindAttachment(Level4Hash);
+		TestTrue(TEXT("FCbPackage(Object+Attachments).FindAttachment(Level2)"),
+			Level2Attachment && Level2Attachment->AsCompactBinary().AsArrayRef().Equals(Level2));
+		TestTrue(TEXT("FCbPackage(Object+Attachments).FindAttachment(Level3)"),
+			Level3Attachment && Level3Attachment->AsCompactBinary().Equals(Level3));
+		TestTrue(TEXT("FCbPackage(Object+Attachments).FindAttachment(Level4)"),
+			Level4Attachment &&
+			Level4Attachment->AsBinary() != Level4 &&
+			Level4Attachment->AsBinary()->GetView().EqualBytes(Level4->GetView()));
+
+		TestTrue(TEXT("FCbPackage(Object+Attachments).GetAttachments()->Sorted"),
+			Algo::IsSorted(Package.GetAttachments()));
+
+		const FCbPackage PackageCopy = Package;
+		TestEqual(TEXT("FCbPackage(Object+Attachments).Equals(EqualCopied)"), PackageCopy, Package);
+
+		TestEqual(TEXT("FCbPackage(Object+Attachments).RemoveAttachment(Level1)"),
+			Package.RemoveAttachment(Level1Hash), 0);
+		TestEqual(TEXT("FCbPackage(Object+Attachments).RemoveAttachment(Level2)"),
+			Package.RemoveAttachment(Level2Hash), 1);
+		TestEqual(TEXT("FCbPackage(Object+Attachments).RemoveAttachment(Level3)"),
+			Package.RemoveAttachment(Level3Hash), 1);
+		TestEqual(TEXT("FCbPackage(Object+Attachments).RemoveAttachment(Level4)"),
+			Package.RemoveAttachment(Level4Hash), 1);
+		TestEqual(TEXT("FCbPackage(Object+Attachments).RemoveAttachment(Level4, Again)"),
+			Package.RemoveAttachment(Level4Hash), 0);
+		TestEqual(TEXT("FCbPackage(Object+Attachments).GetAttachments(Removed)"), Package.GetAttachments().Num(), 0);
+
+		TestNotEqual(TEXT("FCbPackage(Object+Attachments).Equals(AttachmentsNotEqual)"), PackageCopy, Package);
+		Package = PackageCopy;
+		TestEqual(TEXT("FCbPackage(Object+Attachments).Equals(EqualAssigned)"), PackageCopy, Package);
+		Package.SetObject(FCbObjectRef());
+		TestNotEqual(TEXT("FCbPackage(Object+Attachments).Equals(ObjectNotEqual)"), PackageCopy, Package);
+		TestEqual(TEXT("FCbPackage(Object+Attachments).GetObjectHash(Null)"), Package.GetObjectHash(), FBlake3Hash());
+	}
+
+	// Out of Order
+	{
+		TCbWriter<384> Writer;
+		Writer.Binary(Level2.GetBuffer());
+		Writer.Reference(Level2Hash);
+		Writer.Binary(Level4);
+		Writer.BinaryReference(Level4Hash);
+		Writer.Object(Level1);
+		Writer.Reference(Level1Hash);
+		Writer.Binary(Level3.GetBuffer());
+		Writer.Reference(Level3Hash);
+		Writer.Null();
+
+		FCbFieldRefIterator Fields = Writer.Save();
+		FCbPackage FromFields;
+		FromFields.Load(Fields);
+
+		const FCbAttachment* const Level2Attachment = FromFields.FindAttachment(Level2Hash);
+		const FCbAttachment* const Level3Attachment = FromFields.FindAttachment(Level3Hash);
+		const FCbAttachment* const Level4Attachment = FromFields.FindAttachment(Level4Hash);
+
+		TestTrue(TEXT("FCbPackage(OutOfOrder).Load()->Level1"), FromFields.GetObject().Equals(Level1));
+		TestEqual(TEXT("FCbPackage(OutOfOrder).Load()->Level1Buffer"),
+			FromFields.GetObject().GetOuterBuffer(), Fields.GetOuterBuffer());
+		TestEqual(TEXT("FCbPackage(OutOfOrder).Load()->Level1Hash"), FromFields.GetObjectHash(), Level1Hash);
+
+		TestTrue(TEXT("FCbPackage(OutOfOrder).Load()->Level2"),
+			Level2Attachment && Level2Attachment->AsCompactBinary().AsArrayRef().Equals(Level2));
+		TestTrue(TEXT("FCbPackage(OutOfOrder).Load()->Level2Buffer"),
+			Level2Attachment && Fields.GetOuterBuffer()->GetView().Contains(Level2Attachment->AsBinary()->GetView()));
+		TestTrue(TEXT("FCbPackage(OutOfOrder).Load()->Level2Hash"),
+			Level2Attachment && Level2Attachment->GetHash() == Level2Hash);
+
+		TestTrue(TEXT("FCbPackage(OutOfOrder).Load()->Level3"),
+			Level3Attachment && Level3Attachment->AsCompactBinary().Equals(Level3));
+		TestTrue(TEXT("FCbPackage(OutOfOrder).Load()->Level3Buffer"),
+			Level3Attachment && Fields.GetOuterBuffer()->GetView().Contains(Level3Attachment->AsBinary()->GetView()));
+		TestTrue(TEXT("FCbPackage(OutOfOrder).Load()->Level3Hash"),
+			Level3Attachment && Level3Attachment->GetHash() == Level3Hash);
+
+		TestTrue(TEXT("FCbPackage(OutOfOrder).Load()->Level4"),
+			Level4Attachment && Level4Attachment->AsBinary()->GetView().EqualBytes(Level4->GetView()));
+		TestTrue(TEXT("FCbPackage(OutOfOrder).Load()->Level4Buffer"),
+			Level4Attachment && Fields.GetOuterBuffer()->GetView().Contains(Level4Attachment->AsBinary()->GetView()));
+		TestTrue(TEXT("FCbPackage(OutOfOrder).Load()->Level4Hash"),
+			Level4Attachment && Level4Attachment->GetHash() == Level4Hash);
+
+		FBufferArchive WriteAr;
+		Writer.Save(WriteAr);
+		FCbPackage FromArchive;
+		FMemoryReader ReadAr(WriteAr);
+		FromArchive.Load(ReadAr, [](uint64 Size) { return FSharedBuffer::Alloc(Size); });
+
+		Writer.Reset();
+		FromArchive.Save(Writer);
+		FCbFieldRefIterator Saved = Writer.Save();
+		TestTrue(TEXT("FCbPackage(OutOfOrder).Save()->Level1"), Saved.AsObjectRef().Equals(Level1));
+		++Saved;
+		TestEqual(TEXT("FCbPackage(OutOfOrder).Save()->Level1Hash"), Saved.AsReference(), Level1Hash);
+		++Saved;
+		TestTrue(TEXT("FCbPackage(OutOfOrder).Save()->Level3"), Saved.AsBinary().EqualBytes(Level3.GetView()));
+		++Saved;
+		TestEqual(TEXT("FCbPackage(OutOfOrder).Save()->Level3Hash"), Saved.AsReference(), Level3Hash);
+		++Saved;
+		TestTrue(TEXT("FCbPackage(OutOfOrder).Save()->Level2"), Saved.AsBinary().EqualBytes(Level2.GetView()));
+		++Saved;
+		TestEqual(TEXT("FCbPackage(OutOfOrder).Save()->Level2Hash"), Saved.AsReference(), Level2Hash);
+		++Saved;
+		TestTrue(TEXT("FCbPackage(OutOfOrder).Save()->Level4"), Saved.AsBinary().EqualBytes(Level4->GetView()));
+		++Saved;
+		TestEqual(TEXT("FCbPackage(OutOfOrder).Save()->Level4Hash"), Saved.AsBinaryReference(), Level4Hash);
+		++Saved;
+		TestTrue(TEXT("FCbPackage(OutOfOrder).Save()->Null"), Saved.IsNull());
+		++Saved;
+		TestTrue(TEXT("FCbPackage(OutOfOrder).Save()->AtEnd"), !Saved);
+	}
+
+	// Null Attachment
+	{
+		const FCbAttachment NullAttachment;
+		FCbPackage Package;
+		Package.AddAttachment(NullAttachment);
+		TestTrue(TEXT("FCbPackage(NullAttachment).IsNull()"), Package.IsNull());
+		TestFalse(TEXT("FCbPackage(NullAttachment) as bool"), bool(Package));
+		TestEqual(TEXT("FCbPackage(NullAttachment).GetAttachments()"), Package.GetAttachments().Num(), 0);
+		TestTrue(TEXT("FCbPackage(NullAttachment).FindAttachment()"), !Package.FindAttachment(NullAttachment));
+	}
+
+	// Resolve After Merge
+	{
+		bool bResolved = false;
+		FCbPackage Package;
+		Package.AddAttachment(FCbAttachment(Level3.GetBuffer()));
+		Package.AddAttachment(FCbAttachment(FCbFieldRefIterator::MakeSingle(Level3)),
+			[&bResolved](const FBlake3Hash& Hash) -> FSharedBufferConstPtr
+			{
+				bResolved = true;
+				return FSharedBufferConstPtr();
+			});
+		TestTrue(TEXT("FCbPackage(ResolveAfterMerge)->Resolved"), bResolved);
+	}
+
+	return true;
+}
+
+///////////////////////////////////////////////////////////////////////////////////////////////////////////////////////
+
+#endif // WITH_DEV_AUTOMATION_TESTS
diff --git a/Engine/Source/Runtime/Core/Private/Tests/Serialization/CompactBinaryValidationTest.cpp b/Engine/Source/Runtime/Core/Private/Tests/Serialization/CompactBinaryValidationTest.cpp
--- a/Engine/Source/Runtime/Core/Private/Tests/Serialization/CompactBinaryValidationTest.cpp
+++ b/Engine/Source/Runtime/Core/Private/Tests/Serialization/CompactBinaryValidationTest.cpp
@@ -3,6 +3,10 @@
 #include "Serialization/CompactBinaryValidation.h"
 
 #include "Misc/AutomationTest.h"
+#include "Serialization/BufferArchive.h"
+#include "Serialization/CompactBinary.h"
+#include "Serialization/CompactBinaryPackage.h"
+#include "Serialization/CompactBinaryWriter.h"
 
 #if WITH_DEV_AUTOMATION_TESTS
 
@@ -34,6 +38,7 @@ bool FCbValidateTest::RunTest(const FString& Parameters)
 
 	// Test OutOfBounds
 	TestEqual(TEXT("ValidateCompactBinary(OutOfBounds, Empty)"), Validate({}), ECbValidateError::OutOfBounds);
+	TestEqual(TEXT("ValidateCompactBinary(OutOfBounds, Empty, Mode)"), ValidateMode({}, ECbValidateMode::None), ECbValidateError::None);
 
 	TestEqual(TEXT("ValidateCompactBinary(Valid, Null)"), Validate({NullNoName}), ECbValidateError::None);
 	TestEqual(TEXT("ValidateCompactBinary(Valid, Null, Name)"), Validate({NullWithName, 1, 'N'}), ECbValidateError::None);
@@ -253,9 +258,9 @@ bool FCbValidateTest::RunTest(const FString& Parameters)
 IMPLEMENT_SIMPLE_AUTOMATION_TEST(FCbValidateRangeTest, "System.Core.Serialization.ValidateCompactBinaryRange", CompactBinaryValidationTestFlags)
 bool FCbValidateRangeTest::RunTest(const FString& Parameters)
 {
-	auto Validate = [this](std::initializer_list<uint8> Data) -> ECbValidateError
+	auto Validate = [this](std::initializer_list<uint8> Data, ECbValidateMode Mode = ECbValidateMode::All) -> ECbValidateError
 	{
-		return ValidateCompactBinaryRange(MakeMemoryView(Data), ECbValidateMode::All);
+		return ValidateCompactBinaryRange(MakeMemoryView(Data), Mode);
 	};
 
 	// Test Empty
@@ -267,6 +272,276 @@ bool FCbValidateRangeTest::RunTest(const FString& Parameters)
 	// Test Padding
 	TestEqual(TEXT("ValidateCompactBinaryRange(Padding InvalidType)"), Validate({uint8(ECbFieldType::Null), 0}), ECbValidateError::InvalidType);
 	TestEqual(TEXT("ValidateCompactBinaryRange(Padding OutOfBounds)"), Validate({uint8(ECbFieldType::Null), uint8(ECbFieldType::Binary)}), ECbValidateError::OutOfBounds);
+	TestEqual(TEXT("ValidateCompactBinaryRange(Padding OutOfBounds, Mode)"), Validate({uint8(ECbFieldType::Null), uint8(ECbFieldType::Binary)}, ECbValidateMode::None), ECbValidateError::None);
+
+	return true;
+}
+
+///////////////////////////////////////////////////////////////////////////////////////////////////////////////////////
+
+IMPLEMENT_SIMPLE_AUTOMATION_TEST(FCbValidateAttachmentTest, "System.Core.Serialization.ValidateCompactBinaryAttachment", CompactBinaryValidationTestFlags)
+bool FCbValidateAttachmentTest::RunTest(const FString& Parameters)
+{
+	auto Validate = [this](std::initializer_list<uint8> Data, ECbValidateMode Mode = ECbValidateMode::All) -> ECbValidateError
+	{
+		return ValidateCompactBinaryAttachment(MakeMemoryView(Data), Mode);
+	};
+
+	// Test Null
+	{
+		TestEqual(TEXT("ValidateCompactBinaryAttachment(Null)"), Validate({uint8(ECbFieldType::Binary), 0}), ECbValidateError::None);
+		TestEqual(TEXT("ValidateCompactBinaryAttachment(Null, Padding)"), Validate({uint8(ECbFieldType::Binary), 0, 0}), ECbValidateError::Padding);
+		TestEqual(TEXT("ValidateCompactBinaryAttachment(Null, Padding, Mode)"), Validate({uint8(ECbFieldType::Binary), 0, 0}, ECbValidateMode::All & ~ECbValidateMode::Padding), ECbValidateError::None);
+	}
+
+	// Test Binary
+	{
+		FBufferArchive Buffer;
+		FCbAttachment(FSharedBuffer::MakeView(MakeMemoryView<uint8>({0}))).Save(Buffer);
+		TestEqual(TEXT("ValidateCompactBinaryAttachment(Binary)"), ValidateCompactBinaryAttachment(MakeMemoryView(Buffer), ECbValidateMode::All), ECbValidateError::None);
+		TestEqual(TEXT("ValidateCompactBinaryAttachment(Binary, MissingHash)"), ValidateCompactBinaryAttachment(MakeMemoryView(Buffer).Left(3), ECbValidateMode::All), ECbValidateError::InvalidPackageFormat);
+		TestEqual(TEXT("ValidateCompactBinaryAttachment(Binary, MissingHash, Mode)"), ValidateCompactBinaryAttachment(MakeMemoryView(Buffer).Left(3), ECbValidateMode::All & ~ECbValidateMode::Package), ECbValidateError::None);
+		++Buffer.Last();
+		TestEqual(TEXT("ValidateCompactBinaryAttachment(Binary, InvalidHash)"), ValidateCompactBinaryAttachment(MakeMemoryView(Buffer), ECbValidateMode::All), ECbValidateError::InvalidPackageHash);
+		TestEqual(TEXT("ValidateCompactBinaryAttachment(Binary, InvalidHash, Mode)"), ValidateCompactBinaryAttachment(MakeMemoryView(Buffer), ECbValidateMode::All & ~ECbValidateMode::Package), ECbValidateError::None);
+		--Buffer.Last();
+		Buffer.AddZeroed(1);
+		TestEqual(TEXT("ValidateCompactBinaryAttachment(Binary, Padding)"), ValidateCompactBinaryAttachment(MakeMemoryView(Buffer), ECbValidateMode::All), ECbValidateError::Padding);
+		TestEqual(TEXT("ValidateCompactBinaryAttachment(Binary, Padding, Mode)"), ValidateCompactBinaryAttachment(MakeMemoryView(Buffer), ECbValidateMode::All & ~ECbValidateMode::Padding), ECbValidateError::None);
+		Buffer.Pop(/*bAllowShrinking*/ false);
+		TestEqual(TEXT("ValidateCompactBinaryAttachment(Binary, OutOfBoundsValue)"), ValidateCompactBinaryAttachment(MakeMemoryView(Buffer).Left(1), ECbValidateMode::All), ECbValidateError::OutOfBounds);
+		TestEqual(TEXT("ValidateCompactBinaryAttachment(Binary, OutOfBoundsHash)"), ValidateCompactBinaryAttachment(MakeMemoryView(Buffer).Left(4), ECbValidateMode::All), ECbValidateError::OutOfBounds);
+		TestEqual(TEXT("ValidateCompactBinaryAttachment(Binary, OutOfBoundsHash, Mode)"), ValidateCompactBinaryAttachment(MakeMemoryView(Buffer).Left(4), ECbValidateMode::None), ECbValidateError::None);
+	}
+
+	// Test CompactBinary
+	{
+		FBufferArchive Buffer;
+		FCbAttachment(FCbFieldRefIterator::MakeSingle(FCbFieldRef::MakeView(FCbField(nullptr, ECbFieldType::Null)))).Save(Buffer);
+		TestEqual(TEXT("ValidateCompactBinaryAttachment(CompactBinary)"), ValidateCompactBinaryAttachment(MakeMemoryView(Buffer), ECbValidateMode::All), ECbValidateError::None);
+		TestEqual(TEXT("ValidateCompactBinaryAttachment(CompactBinary, MissingHash)"), ValidateCompactBinaryAttachment(MakeMemoryView(Buffer).Left(3), ECbValidateMode::All), ECbValidateError::InvalidPackageFormat);
+		TestEqual(TEXT("ValidateCompactBinaryAttachment(CompactBinary, MissingHash, Mode)"), ValidateCompactBinaryAttachment(MakeMemoryView(Buffer).Left(3), ECbValidateMode::All & ~ECbValidateMode::Package), ECbValidateError::None);
+		++Buffer.Last();
+		TestEqual(TEXT("ValidateCompactBinaryAttachment(CompactBinary, InvalidHash)"), ValidateCompactBinaryAttachment(MakeMemoryView(Buffer), ECbValidateMode::All), ECbValidateError::InvalidPackageHash);
+		TestEqual(TEXT("ValidateCompactBinaryAttachment(CompactBinary, InvalidHash, Mode)"), ValidateCompactBinaryAttachment(MakeMemoryView(Buffer), ECbValidateMode::All & ~ECbValidateMode::Package), ECbValidateError::None);
+		--Buffer.Last();
+		Buffer.AddZeroed(1);
+		TestEqual(TEXT("ValidateCompactBinaryAttachment(CompactBinary, Padding)"), ValidateCompactBinaryAttachment(MakeMemoryView(Buffer), ECbValidateMode::All), ECbValidateError::Padding);
+		TestEqual(TEXT("ValidateCompactBinaryAttachment(CompactBinary, Padding, Mode)"), ValidateCompactBinaryAttachment(MakeMemoryView(Buffer), ECbValidateMode::All & ~ECbValidateMode::Padding), ECbValidateError::None);
+		Buffer.Pop(/*bAllowShrinking*/ false);
+		TestEqual(TEXT("ValidateCompactBinaryAttachment(CompactBinary, OutOfBoundsValue)"), ValidateCompactBinaryAttachment(MakeMemoryView(Buffer).Left(1), ECbValidateMode::All), ECbValidateError::OutOfBounds);
+		TestEqual(TEXT("ValidateCompactBinaryAttachment(CompactBinary, OutOfBoundsHash)"), ValidateCompactBinaryAttachment(MakeMemoryView(Buffer).Left(4), ECbValidateMode::All), ECbValidateError::OutOfBounds);
+		TestEqual(TEXT("ValidateCompactBinaryAttachment(CompactBinary, OutOfBoundsHash, Mode)"), ValidateCompactBinaryAttachment(MakeMemoryView(Buffer).Left(4), ECbValidateMode::None), ECbValidateError::None);
+	}
+
+	// Test InvalidPackageFormat
+	{
+		TestEqual(TEXT("ValidateCompactBinaryAttachment(Empty)"), Validate({}), ECbValidateError::InvalidPackageFormat);
+		TestEqual(TEXT("ValidateCompactBinaryAttachment(Empty, Mode)"), Validate({}, ECbValidateMode::All & ~ECbValidateMode::Package), ECbValidateError::None);
+	}
+	{
+		FBufferArchive Buffer;
+		FCbFieldRefIterator::MakeSingle(FCbObjectRef().AsFieldRef()).CopyRangeTo(Buffer);
+		TestEqual(TEXT("ValidateCompactBinaryAttachment(Object)"), ValidateCompactBinaryAttachment(MakeMemoryView(Buffer), ECbValidateMode::All), ECbValidateError::InvalidPackageFormat);
+		TestEqual(TEXT("ValidateCompactBinaryAttachment(Object, Mode)"), ValidateCompactBinaryAttachment(MakeMemoryView(Buffer), ECbValidateMode::All & ~ECbValidateMode::Package), ECbValidateError::None);
+	}
+	{
+		FBufferArchive Buffer;
+		FCbWriter Writer;
+		Writer.Reference(FBlake3::HashBuffer(MakeMemoryView<uint8>({0, 1, 2, 3})));
+		Writer.Binary(MakeMemoryView<uint8>({0, 1, 2, 3}));
+		Writer.Save(Buffer);
+		TestEqual(TEXT("ValidateCompactBinaryAttachment(HashBeforeValue)"), ValidateCompactBinaryAttachment(MakeMemoryView(Buffer), ECbValidateMode::All), ECbValidateError::InvalidPackageFormat | ECbValidateError::Padding);
+		TestEqual(TEXT("ValidateCompactBinaryAttachment(HashBeforeValue, Mode)"), ValidateCompactBinaryAttachment(MakeMemoryView(Buffer), ECbValidateMode::All & ~ECbValidateMode::Package), ECbValidateError::Padding);
+	}
+	{
+		FBufferArchive Buffer;
+		FCbWriter Writer;
+		Writer.Name("Name"_ASV).Binary(MakeMemoryView<uint8>({0, 1, 2, 3}));
+		Writer.Reference(FBlake3::HashBuffer(MakeMemoryView<uint8>({0, 1, 2, 3})));
+		Writer.Save(Buffer);
+		TestEqual(TEXT("ValidateCompactBinaryAttachment(NameOnValue)"), ValidateCompactBinaryAttachment(MakeMemoryView(Buffer), ECbValidateMode::All), ECbValidateError::InvalidPackageFormat);
+		TestEqual(TEXT("ValidateCompactBinaryAttachment(NameOnValue, Mode)"), ValidateCompactBinaryAttachment(MakeMemoryView(Buffer), ECbValidateMode::All & ~ECbValidateMode::Package), ECbValidateError::None);
+	}
+	{
+		FBufferArchive Buffer;
+		FCbWriter Writer;
+		Writer.Binary(MakeMemoryView<uint8>({0, 1, 2, 3}));
+		Writer.Name("Name"_ASV).Reference(FBlake3::HashBuffer(MakeMemoryView<uint8>({0, 1, 2, 3})));
+		Writer.Save(Buffer);
+		TestEqual(TEXT("ValidateCompactBinaryAttachment(NameOnHash)"), ValidateCompactBinaryAttachment(MakeMemoryView(Buffer), ECbValidateMode::All), ECbValidateError::InvalidPackageFormat);
+		TestEqual(TEXT("ValidateCompactBinaryAttachment(NameOnHash, Mode)"), ValidateCompactBinaryAttachment(MakeMemoryView(Buffer), ECbValidateMode::All & ~ECbValidateMode::Package), ECbValidateError::None);
+	}
+
+	return true;
+}
+
+///////////////////////////////////////////////////////////////////////////////////////////////////////////////////////
+
+IMPLEMENT_SIMPLE_AUTOMATION_TEST(FCbValidatePackageTest, "System.Core.Serialization.ValidateCompactBinaryPackage", CompactBinaryValidationTestFlags)
+bool FCbValidatePackageTest::RunTest(const FString& Parameters)
+{
+	// Test Null
+	{
+		TestEqual(TEXT("ValidateCompactBinaryPackage(Null)"), ValidateCompactBinaryPackage(MakeMemoryView<uint8>({uint8(ECbFieldType::Null)}), ECbValidateMode::All), ECbValidateError::None);
+		TestEqual(TEXT("ValidateCompactBinaryPackage(Null, Padding)"), ValidateCompactBinaryPackage(MakeMemoryView<uint8>({uint8(ECbFieldType::Null), 0}), ECbValidateMode::All), ECbValidateError::Padding);
+		TestEqual(TEXT("ValidateCompactBinaryPackage(Null, Padding, Mode)"), ValidateCompactBinaryPackage(MakeMemoryView<uint8>({uint8(ECbFieldType::Null), 0}), ECbValidateMode::All & ~ECbValidateMode::Padding), ECbValidateError::None);
+	}
+
+	// Test Object
+	{
+		FBufferArchive Buffer;
+		TCbWriter<256> Writer;
+		Writer.BeginObject();
+		Writer << "Field" << 42;
+		Writer.EndObject();
+		Writer.Reference(Writer.Save().GetHash());
+		Writer.Null();
+		Writer.Save(Buffer);
+		TestEqual(TEXT("ValidateCompactBinaryPackage(Object)"), ValidateCompactBinaryPackage(MakeMemoryView(Buffer), ECbValidateMode::All), ECbValidateError::None);
+		TestEqual(TEXT("ValidateCompactBinaryPackage(Object, MissingHash)"), ValidateCompactBinaryPackage(MakeMemoryView(Buffer).LeftChop(34), ECbValidateMode::All), ECbValidateError::InvalidPackageFormat);
+		TestEqual(TEXT("ValidateCompactBinaryPackage(Object, MissingHash, Mode)"), ValidateCompactBinaryPackage(MakeMemoryView(Buffer).LeftChop(34), ECbValidateMode::All & ~ECbValidateMode::Package), ECbValidateError::None);
+		TestEqual(TEXT("ValidateCompactBinaryPackage(Object, MissingNull)"), ValidateCompactBinaryPackage(MakeMemoryView(Buffer).LeftChop(1), ECbValidateMode::All), ECbValidateError::InvalidPackageFormat);
+		TestEqual(TEXT("ValidateCompactBinaryPackage(Object, MissingNull, Mode)"), ValidateCompactBinaryPackage(MakeMemoryView(Buffer).LeftChop(1), ECbValidateMode::All & ~ECbValidateMode::Package), ECbValidateError::None);
+		++Buffer.Last(1);
+		TestEqual(TEXT("ValidateCompactBinaryPackage(Object, InvalidHash)"), ValidateCompactBinaryPackage(MakeMemoryView(Buffer), ECbValidateMode::All), ECbValidateError::InvalidPackageHash);
+		TestEqual(TEXT("ValidateCompactBinaryPackage(Object, InvalidHash, Mode)"), ValidateCompactBinaryPackage(MakeMemoryView(Buffer), ECbValidateMode::All & ~ECbValidateMode::Package), ECbValidateError::None);
+		--Buffer.Last(1);
+		Buffer.AddZeroed(1);
+		TestEqual(TEXT("ValidateCompactBinaryPackage(Object, Padding)"), ValidateCompactBinaryPackage(MakeMemoryView(Buffer), ECbValidateMode::All), ECbValidateError::Padding);
+		TestEqual(TEXT("ValidateCompactBinaryPackage(Object, Padding, Mode)"), ValidateCompactBinaryPackage(MakeMemoryView(Buffer), ECbValidateMode::All & ~ECbValidateMode::Padding), ECbValidateError::None);
+		Buffer.Pop(/*bAllowShrinking*/ false);
+		TestEqual(TEXT("ValidateCompactBinaryPackage(Object, OutOfBoundsValue)"), ValidateCompactBinaryPackage(MakeMemoryView(Buffer).Left(1), ECbValidateMode::All), ECbValidateError::OutOfBounds);
+		TestEqual(TEXT("ValidateCompactBinaryPackage(Object, OutOfBoundsHash)"), ValidateCompactBinaryPackage(MakeMemoryView(Buffer).Left(11), ECbValidateMode::All), ECbValidateError::OutOfBounds);
+		TestEqual(TEXT("ValidateCompactBinaryPackage(Object, OutOfBoundsHash, Mode)"), ValidateCompactBinaryPackage(MakeMemoryView(Buffer).Left(11), ECbValidateMode::None), ECbValidateError::None);
+	}
+
+	// Test Object + Attachment
+	{
+		FBufferArchive Buffer;
+		TCbWriter<256> Writer;
+		Writer.BeginObject();
+		Writer << "Field" << 42;
+		Writer.EndObject();
+		Writer.Reference(Writer.Save().GetHash());
+		const uint64 AttachmentOffset = Writer.GetSaveSize();
+		Writer.Binary(MakeMemoryView<uint8>({0, 1, 2, 3}));
+		Writer.Reference(FBlake3::HashBuffer(MakeMemoryView<uint8>({0, 1, 2, 3})));
+		Writer.Null();
+		Writer.Save(Buffer);
+		TestEqual(TEXT("ValidateCompactBinaryPackage(Object+Attachment)"), ValidateCompactBinaryPackage(MakeMemoryView(Buffer), ECbValidateMode::All), ECbValidateError::None);
+		TestEqual(TEXT("ValidateCompactBinaryPackage(Object+Attachment, MissingHash)"), ValidateCompactBinaryPackage(MakeMemoryView(Buffer).LeftChop(34), ECbValidateMode::All), ECbValidateError::InvalidPackageFormat);
+		TestEqual(TEXT("ValidateCompactBinaryPackage(Object+Attachment, MissingHash, Mode)"), ValidateCompactBinaryPackage(MakeMemoryView(Buffer).LeftChop(34), ECbValidateMode::All & ~ECbValidateMode::Package), ECbValidateError::None);
+		TestEqual(TEXT("ValidateCompactBinaryPackage(Object+Attachment, MissingNull)"), ValidateCompactBinaryPackage(MakeMemoryView(Buffer).LeftChop(1), ECbValidateMode::All), ECbValidateError::InvalidPackageFormat);
+		TestEqual(TEXT("ValidateCompactBinaryPackage(Object+Attachment, MissingNull, Mode)"), ValidateCompactBinaryPackage(MakeMemoryView(Buffer).LeftChop(1), ECbValidateMode::All & ~ECbValidateMode::Package), ECbValidateError::None);
+		++Buffer.Last(1);
+		TestEqual(TEXT("ValidateCompactBinaryPackage(Object+Attachment, InvalidHash)"), ValidateCompactBinaryPackage(MakeMemoryView(Buffer), ECbValidateMode::All), ECbValidateError::InvalidPackageHash);
+		TestEqual(TEXT("ValidateCompactBinaryPackage(Object+Attachment, InvalidHash, Mode)"), ValidateCompactBinaryPackage(MakeMemoryView(Buffer), ECbValidateMode::All & ~ECbValidateMode::Package), ECbValidateError::None);
+		--Buffer.Last(1);
+		Buffer.AddZeroed(1);
+		TestEqual(TEXT("ValidateCompactBinaryPackage(Object+Attachment, Padding)"), ValidateCompactBinaryPackage(MakeMemoryView(Buffer), ECbValidateMode::All), ECbValidateError::Padding);
+		TestEqual(TEXT("ValidateCompactBinaryPackage(Object+Attachment, Padding, Mode)"), ValidateCompactBinaryPackage(MakeMemoryView(Buffer), ECbValidateMode::All & ~ECbValidateMode::Padding), ECbValidateError::None);
+		Buffer.Pop(/*bAllowShrinking*/ false);
+		TestEqual(TEXT("ValidateCompactBinaryPackage(Object+Attachment, OutOfBoundsValue)"), ValidateCompactBinaryPackage(MakeMemoryView(Buffer).Left(AttachmentOffset + 1), ECbValidateMode::All), ECbValidateError::OutOfBounds);
+		TestEqual(TEXT("ValidateCompactBinaryPackage(Object+Attachment, OutOfBoundsHash)"), ValidateCompactBinaryPackage(MakeMemoryView(Buffer).Left(AttachmentOffset + 7), ECbValidateMode::All), ECbValidateError::OutOfBounds);
+		TestEqual(TEXT("ValidateCompactBinaryPackage(Object+Attachment, OutOfBoundsHash, Mode)"), ValidateCompactBinaryPackage(MakeMemoryView(Buffer).Left(AttachmentOffset + 7), ECbValidateMode::None), ECbValidateError::None);
+	}
+
+	// Test InvalidPackageFormat
+	{
+		TestEqual(TEXT("ValidateCompactBinaryPackage(Empty)"), ValidateCompactBinaryPackage(FMemoryView(), ECbValidateMode::All), ECbValidateError::InvalidPackageFormat);
+		TestEqual(TEXT("ValidateCompactBinaryPackage(Empty, Mode)"), ValidateCompactBinaryPackage(FMemoryView(), ECbValidateMode::All & ~ECbValidateMode::Package), ECbValidateError::None);
+	}
+	{
+		FBufferArchive Buffer;
+		TCbWriter<256> Writer;
+		Writer.Name("Name"_ASV);
+		Writer.BeginObject();
+		Writer << "Field" << 42;
+		Writer.EndObject();
+		Writer.Reference(Writer.Save().GetHash());
+		Writer.Null();
+		Writer.Save(Buffer);
+		TestEqual(TEXT("ValidateCompactBinaryPackage(NameOnObject)"), ValidateCompactBinaryPackage(MakeMemoryView(Buffer), ECbValidateMode::All), ECbValidateError::InvalidPackageFormat);
+		TestEqual(TEXT("ValidateCompactBinaryPackage(NameOnObject, Mode)"), ValidateCompactBinaryPackage(MakeMemoryView(Buffer), ECbValidateMode::All & ~ECbValidateMode::Package), ECbValidateError::None);
+	}
+	{
+		FBufferArchive Buffer;
+		TCbWriter<256> Writer;
+		Writer.BeginObject();
+		Writer << "Field" << 42;
+		Writer.EndObject();
+		const FBlake3Hash ObjectHash = Writer.Save().GetHash();
+		Writer.Name("Name"_ASV);
+		Writer.Reference(ObjectHash);
+		Writer.Null();
+		Writer.Save(Buffer);
+		TestEqual(TEXT("ValidateCompactBinaryPackage(NameOnHash)"), ValidateCompactBinaryPackage(MakeMemoryView(Buffer), ECbValidateMode::All), ECbValidateError::InvalidPackageFormat);
+		TestEqual(TEXT("ValidateCompactBinaryPackage(NameOnHash, Mode)"), ValidateCompactBinaryPackage(MakeMemoryView(Buffer), ECbValidateMode::All & ~ECbValidateMode::Package), ECbValidateError::None);
+	}
+	{
+		FBufferArchive Buffer;
+		TCbWriter<256> Writer;
+		Writer.BeginObject();
+		Writer << "Field" << 42;
+		Writer.EndObject();
+		Writer.Reference(Writer.Save().GetHash());
+		Writer.Name("Name"_ASV).Binary(MakeMemoryView<uint8>({0, 1, 2, 3}));
+		Writer.Reference(FBlake3::HashBuffer(MakeMemoryView<uint8>({0, 1, 2, 3})));
+		Writer.Null();
+		Writer.Save(Buffer);
+		TestEqual(TEXT("ValidateCompactBinaryPackage(NameOnAttachment)"), ValidateCompactBinaryPackage(MakeMemoryView(Buffer), ECbValidateMode::All), ECbValidateError::InvalidPackageFormat);
+		TestEqual(TEXT("ValidateCompactBinaryPackage(NameOnAttachment, Mode)"), ValidateCompactBinaryPackage(MakeMemoryView(Buffer), ECbValidateMode::All & ~ECbValidateMode::Package), ECbValidateError::None);
+	}
+	{
+		FBufferArchive Buffer;
+		TCbWriter<256> Writer;
+		Writer.BeginObject();
+		Writer << "Field" << 42;
+		Writer.EndObject();
+		const FBlake3Hash ObjectHash = Writer.Save().GetHash();
+		Writer.Reference(ObjectHash);
+		Writer.BeginObject();
+		Writer << "Field" << 42;
+		Writer.EndObject();
+		Writer.Reference(ObjectHash);
+		Writer.Null();
+		Writer.Save(Buffer);
+		TestEqual(TEXT("ValidateCompactBinaryPackage(MultipleObjects)"), ValidateCompactBinaryPackage(MakeMemoryView(Buffer), ECbValidateMode::All), ECbValidateError::MultiplePackageObjects);
+		TestEqual(TEXT("ValidateCompactBinaryPackage(MultipleObjects, Mode)"), ValidateCompactBinaryPackage(MakeMemoryView(Buffer), ECbValidateMode::All & ~ECbValidateMode::Package), ECbValidateError::None);
+	}
+	{
+		FBufferArchive Buffer;
+		TCbWriter<256> Writer;
+		Writer.Binary(MakeMemoryView<uint8>({0, 1, 2, 3}));
+		Writer.Reference(FBlake3::HashBuffer(MakeMemoryView<uint8>({0, 1, 2, 3})));
+		Writer.Binary(MakeMemoryView<uint8>({0, 1, 2, 3}));
+		Writer.Reference(FBlake3::HashBuffer(MakeMemoryView<uint8>({0, 1, 2, 3})));
+		Writer.Null();
+		Writer.Save(Buffer);
+		TestEqual(TEXT("ValidateCompactBinaryPackage(DuplicateAttachments)"), ValidateCompactBinaryPackage(MakeMemoryView(Buffer), ECbValidateMode::All), ECbValidateError::DuplicateAttachments);
+		TestEqual(TEXT("ValidateCompactBinaryPackage(DuplicateAttachments, Mode)"), ValidateCompactBinaryPackage(MakeMemoryView(Buffer), ECbValidateMode::All & ~ECbValidateMode::Package), ECbValidateError::None);
+	}
+	{
+		FBufferArchive Buffer;
+		TCbWriter<256> Writer;
+		Writer.BeginObject();
+		Writer.EndObject();
+		Writer.Reference(FCbObject().GetHash());
+		Writer.Null();
+		Writer.Save(Buffer);
+		TestEqual(TEXT("ValidateCompactBinaryPackage(NullObject)"), ValidateCompactBinaryPackage(MakeMemoryView(Buffer), ECbValidateMode::All), ECbValidateError::NullPackageObject);
+		TestEqual(TEXT("ValidateCompactBinaryPackage(NullObject, Mode)"), ValidateCompactBinaryPackage(MakeMemoryView(Buffer), ECbValidateMode::All & ~ECbValidateMode::Package), ECbValidateError::None);
+	}
+	{
+		FBufferArchive Buffer;
+		TCbWriter<256> Writer;
+		Writer.Binary(FMemoryView());
+		Writer.Null();
+		Writer.Save(Buffer);
+		TestEqual(TEXT("ValidateCompactBinaryPackage(NullAttachment)"), ValidateCompactBinaryPackage(MakeMemoryView(Buffer), ECbValidateMode::All), ECbValidateError::NullPackageAttachment);
+		TestEqual(TEXT("ValidateCompactBinaryPackage(NullAttachment, Mode)"), ValidateCompactBinaryPackage(MakeMemoryView(Buffer), ECbValidateMode::All & ~ECbValidateMode::Package), ECbValidateError::None);
+	}
 
 	return true;
 }
diff --git a/Engine/Source/Runtime/Core/Public/Serialization/CompactBinaryPackage.h b/Engine/Source/Runtime/Core/Public/Serialization/CompactBinaryPackage.h
new file mode 100644
--- /dev/null
+++ b/Engine/Source/Runtime/Core/Public/Serialization/CompactBinaryPackage.h
@@ -0,0 +1,339 @@
+// Copyright Epic Games, Inc. All Rights Reserved.
+
+#pragma once
+
+#include "CoreTypes.h"
+#include "Containers/Array.h"
+#include "Memory/SharedBuffer.h"
+#include "Misc/Blake3.h"
+#include "Serialization/CompactBinary.h"
+
+class FArchive;
+class FCbWriter;
+template <typename FuncType> class TFunctionRef;
+
+///////////////////////////////////////////////////////////////////////////////////////////////////////////////////////
+
+/**
+ * An attachment is either binary or compact binary and is identified by its hash.
+ *
+ * A compact binary attachment is also a valid binary attachment and may be accessed as binary.
+ *
+ * Attachments are serialized as one or two compact binary fields with no name. A Binary field is
+ * written first with its content. The content hash non-zero as a Reference or BinaryReference if
+ * the content size is non-zero, and is otherwise omitted.
+ */
+class FCbAttachment
+{
+public:
+	/** Construct a null attachment. */
+	FCbAttachment() = default;
+
+	/** Construct a compact binary attachment. Value is cloned unless read-only. */
+	inline explicit FCbAttachment(FCbFieldRefIterator Value)
+		: FCbAttachment(MoveTemp(Value), nullptr)
+	{
+	}
+
+	/** Construct a compact binary attachment. Value is cloned unless read-only. Hash must match Value. */
+	inline explicit FCbAttachment(FCbFieldRefIterator Value, const FBlake3Hash& Hash)
+		: FCbAttachment(MoveTemp(Value), &Hash)
+	{
+	}
+
+	/** Construct a binary attachment. Value is cloned unless read-only. */
+	inline explicit FCbAttachment(FSharedBufferConstPtr Value)
+		: FCbAttachment(MoveTemp(Value), nullptr)
+	{
+	}
+
+	/** Construct a binary attachment. Value is cloned unless read-only. Hash must match Value. */
+	inline explicit FCbAttachment(FSharedBufferConstPtr Value, const FBlake3Hash& Hash)
+		: FCbAttachment(MoveTemp(Value), &Hash)
+	{
+	}
+
+	/** Whether the attachment has a value. */
+	inline explicit operator bool() const { return !IsNull(); }
+
+	/** Whether the attachment has a value. */
+	inline bool IsNull() const { return !Buffer; }
+
+	/** Access the attachment as binary. Defaults to a null buffer on error. */
+	CORE_API FSharedBufferConstPtr AsBinary() const;
+
+	/** Access the attachment as compact binary. Defaults to a field iterator with no value on error. */
+	CORE_API FCbFieldRefIterator AsCompactBinary() const;
+
+	/** Returns whether the attachment is binary or compact binary. */
+	inline bool IsBinary() const { return Buffer.IsValid(); }
+
+	/** Returns whether the attachment is compact binary. */
+	inline bool IsCompactBinary() const { return CompactBinary.HasValue(); }
+
+	/** Returns the hash of the attachment value. */
+	inline const FBlake3Hash& GetHash() const { return Hash; }
+
+	/** Compares attachments by their hash. Any discrepancy in type must be handled externally. */
+	inline bool operator==(const FCbAttachment& Attachment) const { return Hash == Attachment.Hash; }
+	inline bool operator!=(const FCbAttachment& Attachment) const { return Hash != Attachment.Hash; }
+	inline bool operator<(const FCbAttachment& Attachment) const { return Hash < Attachment.Hash; }
+
+	/**
+	 * Load the attachment from compact binary as written by Save.
+	 *
+	 * The attachment will hold a reference to the input iterator, if it is read-only and owned, and
+	 * otherwise clones the attachment to make a read-only copy.
+	 *
+	 * The iterator is advanced as attachment fields are consumed from it.
+	 */
+	CORE_API void Load(FCbFieldRefIterator& Fields);
+
+	/**
+	 * Load the attachment from compact binary as written by Save.
+	 *
+	 * The attachment will be read into a read-only, owned buffer.
+	 */
+	CORE_API void Load(FArchive& Ar, FCbBufferAllocator Allocator);
+
+	/** Save the attachment into the writer as a stream of compact binary fields. */
+	CORE_API void Save(FCbWriter& Writer) const;
+
+	/** Save the attachment into the archive as a stream of compact binary fields. */
+	CORE_API void Save(FArchive& Ar) const;
+
+private:
+	CORE_API FCbAttachment(FCbFieldRefIterator Value, const FBlake3Hash* Hash);
+	CORE_API FCbAttachment(FSharedBufferConstPtr Value, const FBlake3Hash* Hash);
+
+	/** A read-only owned buffer containing the binary or compact binary data. */
+	FSharedBufferConstPtr Buffer;
+	/** A field iterator that is valid only for compact binary attachments. */
+	FCbFieldIterator CompactBinary;
+	/** A hash of the attachment value. */
+	FBlake3Hash Hash;
+};
+
+/** Hashes attachments by their hash. Any discrepancy in type must be handled externally. */
+inline uint32 GetTypeHash(const FCbAttachment& Attachment)
+{
+	return GetTypeHash(Attachment.GetHash());
+}
+
+///////////////////////////////////////////////////////////////////////////////////////////////////////////////////////
+
+/**
+ * A package is a compact binary object with attachments for its external references.
+ *
+ * A package is basically a Merkle tree with compact binary as its root and other non-leaf nodes,
+ * and either binary or compact binary as its leaf nodes. A node references its child nodes using
+ * BinaryHash or FieldHash fields in its compact binary representation.
+ *
+ * It is invalid for a package to include attachments that are not referenced by its object or by
+ * one of its referenced compact binary attachments. When attachments are added explicitly, it is
+ * the responsibility of the package creator to follow this requirement. Attachments that are not
+ * referenced may not survive a round-trip through certain storage systems.
+ *
+ * It is valid for a package to exclude referenced attachments, but then it is the responsibility
+ * of the package consumer to have a mechanism for resolving those references when necessary.
+ *
+ * A package is serialized as a sequence of compact binary fields with no name. The object may be
+ * both preceded and followed by attachments. The object itself is written as an Object field and
+ * then its hash in a Reference field. The end of the package is marked by a Null field. Standard
+ * ordering of components is the object, then the attachments ordered by hash, though this is not
+ * a requirement, and it is valid for a package to contain components in an any order.
+ */
+class FCbPackage
+{
+public:
+	/**
+	 * A function that resolves a hash to a buffer containing the data matching that hash.
+	 *
+	 * The resolver may return a null buffer to skip resolving an attachment for the hash.
+	 */
+	using FAttachmentResolver = TFunctionRef<FSharedBufferConstPtr (const FBlake3Hash& Hash)>;
+
+	/** Construct a null package. */
+	FCbPackage() = default;
+
+	/**
+	 * Construct a package from a root object without gathering attachments.
+	 *
+	 * @param InObject The root object, which will be cloned unless it is read-only.
+	 */
+	inline explicit FCbPackage(FCbObjectRef InObject)
+	{
+		SetObject(MoveTemp(InObject));
+	}
+
+	/**
+	 * Construct a package from a root object and gather attachments using the resolver.
+	 *
+	 * @param InObject The root object, which will be cloned unless it is read-only.
+	 * @param InResolver A function that is invoked for every reference and binary reference field.
+	 */
+	inline explicit FCbPackage(FCbObjectRef InObject, FAttachmentResolver InResolver)
+	{
+		SetObject(MoveTemp(InObject), InResolver);
+	}
+
+	/**
+	 * Construct a package from a root object without gathering attachments.
+	 *
+	 * @param InObject The root object, which will be cloned unless it is read-only.
+	 * @param InObjectHash The hash of the object, which must match to avoid validation errors.
+	 */
+	inline explicit FCbPackage(FCbObjectRef InObject, const FBlake3Hash& InObjectHash)
+	{
+		SetObject(MoveTemp(InObject), InObjectHash);
+	}
+
+	/**
+	 * Construct a package from a root object and gather attachments using the resolver.
+	 *
+	 * @param InObject The root object, which will be cloned unless it is read-only.
+	 * @param InObjectHash The hash of the object, which must match to avoid validation errors.
+	 * @param InResolver A function that is invoked for every reference and binary reference field.
+	 */
+	inline explicit FCbPackage(FCbObjectRef InObject, const FBlake3Hash& InObjectHash, FAttachmentResolver InResolver)
+	{
+		SetObject(MoveTemp(InObject), InObjectHash, InResolver);
+	}
+
+	/** Whether the package has a non-empty object or attachments. */
+	inline explicit operator bool() const { return !IsNull(); }
+
+	/** Whether the package has an empty object and no attachments. */
+	inline bool IsNull() const
+	{
+		return !Object.CreateIterator() && Attachments.Num() == 0;
+	}
+
+	/** Returns the compact binary object for the package. */
+	inline const FCbObjectRef& GetObject() const { return Object; }
+
+	/** Returns the has of the compact binary object for the package. */
+	inline const FBlake3Hash& GetObjectHash() const { return ObjectHash; }
+
+	/**
+	 * Set the root object without gathering attachments.
+	 *
+	 * @param InObject The root object, which will be cloned unless it is read-only.
+	 */
+	inline void SetObject(FCbObjectRef InObject)
+	{
+		SetObject(MoveTemp(InObject), nullptr, nullptr);
+	}
+
+	/**
+	 * Set the root object and gather attachments using the resolver.
+	 *
+	 * @param InObject The root object, which will be cloned unless it is read-only.
+	 * @param InResolver A function that is invoked for every reference and binary reference field.
+	 */
+	inline void SetObject(FCbObjectRef InObject, FAttachmentResolver InResolver)
+	{
+		SetObject(MoveTemp(InObject), nullptr, &InResolver);
+	}
+
+	/**
+	 * Set the root object without gathering attachments.
+	 *
+	 * @param InObject The root object, which will be cloned unless it is read-only.
+	 * @param InObjectHash The hash of the object, which must match to avoid validation errors.
+	 */
+	inline void SetObject(FCbObjectRef InObject, const FBlake3Hash& InObjectHash)
+	{
+		SetObject(MoveTemp(InObject), &InObjectHash, nullptr);
+	}
+
+	/**
+	 * Set the root object and gather attachments using the resolver.
+	 *
+	 * @param InObject The root object, which will be cloned unless it is read-only.
+	 * @param InObjectHash The hash of the object, which must match to avoid validation errors.
+	 * @param InResolver A function that is invoked for every reference and binary reference field.
+	 */
+	inline void SetObject(FCbObjectRef InObject, const FBlake3Hash& InObjectHash, FAttachmentResolver InResolver)
+	{
+		SetObject(MoveTemp(InObject), &InObjectHash, &InResolver);
+	}
+
+	/** Returns the attachments in this package. */
+	inline TConstArrayView<FCbAttachment> GetAttachments() const { return Attachments; }
+
+	/**
+	 * Find an attachment by its hash.
+	 *
+	 * @return The attachment, or null if the attachment is not found.
+	 * @note The returned pointer is only valid until the attachments on this package are modified.
+	 */
+	CORE_API const FCbAttachment* FindAttachment(const FBlake3Hash& Hash) const;
+
+	/** Find an attachment if it exists in the package. */
+	inline const FCbAttachment* FindAttachment(const FCbAttachment& Attachment) const
+	{
+		return FindAttachment(Attachment.GetHash());
+	}
+
+	/** Add the attachment to this package. */
+	inline void AddAttachment(const FCbAttachment& Attachment)
+	{
+		AddAttachment(Attachment, nullptr);
+	}
+
+	/** Add the attachment to this package, along with any references that can be resolved. */
+	inline void AddAttachment(const FCbAttachment& Attachment, FAttachmentResolver Resolver)
+	{
+		AddAttachment(Attachment, &Resolver);
+	}
+
+	/**
+	 * Remove an attachment by hash.
+	 *
+	 * @return Number of attachments removed, which will be either 0 or 1.
+	 */
+	CORE_API int32 RemoveAttachment(const FBlake3Hash& Hash);
+	inline int32 RemoveAttachment(const FCbAttachment& Attachment) { return RemoveAttachment(Attachment.GetHash()); }
+
+	/** Compares packages by their object and attachment hashes. */
+	CORE_API bool Equals(const FCbPackage& Package) const;
+	inline bool operator==(const FCbPackage& Package) const { return Equals(Package); }
+	inline bool operator!=(const FCbPackage& Package) const { return !Equals(Package); }
+
+	/**
+	 * Load the object and attachments from compact binary as written by Save.
+	 *
+	 * The object and attachments will hold a reference to the input package, if it is read-only and
+	 * owned, and otherwise individually clones the object and attachments to make read-only copies.
+	 *
+	 * The iterator is advanced as package and attachment fields are consumed from it.
+	 */
+	CORE_API void Load(FCbFieldRefIterator& Fields);
+
+	/**
+	 * Load the object and attachments from compact binary as written by Save.
+	 *
+	 * The object and attachments will be individually read into read-only, owned buffers.
+	 */
+	CORE_API void Load(FArchive& Ar, FCbBufferAllocator Allocator);
+
+	/** Save the object and attachments into the writer as a stream of compact binary fields. */
+	CORE_API void Save(FCbWriter& Writer) const;
+
+	/** Save the object and attachments into the archive as a stream of compact binary fields. */
+	CORE_API void Save(FArchive& Ar) const;
+
+private:
+	CORE_API void SetObject(FCbObjectRef Object, const FBlake3Hash* Hash, FAttachmentResolver* Resolver);
+	CORE_API void AddAttachment(const FCbAttachment& Attachment, FAttachmentResolver* Resolver);
+
+	void GatherAttachments(const FCbFieldIterator& Fields, FAttachmentResolver Resolver);
+
+	/** Attachments ordered by their hash. */
+	TArray<FCbAttachment> Attachments;
+	FCbObjectRef Object;
+	FBlake3Hash ObjectHash;
+};
+
+///////////////////////////////////////////////////////////////////////////////////////////////////////////////////////
diff --git a/Engine/Source/Runtime/Core/Public/Serialization/CompactBinaryValidation.h b/Engine/Source/Runtime/Core/Public/Serialization/CompactBinaryValidation.h
--- a/Engine/Source/Runtime/Core/Public/Serialization/CompactBinaryValidation.h
+++ b/Engine/Source/Runtime/Core/Public/Serialization/CompactBinaryValidation.h
@@ -12,13 +12,16 @@
 /** Flags for validating compact binary data. */
 enum class ECbValidateMode : uint32
 {
+	/** Skip validation if no other validation modes are enabled. */
+	None                    = 0,
+
 	/**
 	 * Validate that the value can be read and stays inside the bounds of the memory view.
 	 *
 	 * This is the minimum level of validation required to be able to safely read a field, array,
 	 * or object without the risk of crashing or reading out of bounds.
 	 */
-	Default                 = 0,
+	Default                 = 1 << 0,
 
 	/**
 	 * Validate that object fields have unique non-empty names and array fields have no names.
@@ -27,7 +30,7 @@ enum class ECbValidateMode : uint32
 	 * cannot be looked up by name other than the first, and converting to other data formats can
 	 * fail in the presence of naming issues.
 	 */
-	Names                   = 1 << 0,
+	Names                   = 1 << 1,
 
 	/**
 	 * Validate that fields are serialized in the canonical format.
@@ -38,7 +41,7 @@ enum class ECbValidateMode : uint32
 	 * were not encoded uniformly, variable-length integers that could be encoded in fewer bytes,
 	 * or 64-bit floats that could be encoded in 32 bits without loss of precision.
 	 */
-	Format                  = 1 << 1,
+	Format                  = 1 << 2,
 
 	/**
 	 * Validate that there is no padding after the value before the end of the memory view.
@@ -46,10 +49,15 @@ enum class ECbValidateMode : uint32
 	 * Padding validation failures have no impact on the ability to read the input, but are using
 	 * more memory than necessary.
 	 */
-	Padding                 = 1 << 2,
+	Padding                 = 1 << 3,
+
+	/**
+	 * Validate that a package or attachment has the expected fields and matches its saved hashes.
+	 */
+	Package                 = 1 << 4,
 
 	/** Perform all validation described above. */
-	All                     = Default | Names | Format | Padding,
+	All                     = Default | Names | Format | Padding | Package,
 };
 
 ENUM_CLASS_FLAGS(ECbValidateMode);
@@ -95,6 +103,21 @@ enum class ECbValidateError : uint32
 
 	/** A value did not use the entire memory view given for validation. */
 	Padding                 = 1 << 10,
+
+	// Mode: Package
+
+	/** The package or attachment had missing fields or fields out of order. */
+	InvalidPackageFormat    = 1 << 11,
+	/** The object or an attachment did not match the hash stored for it. */
+	InvalidPackageHash      = 1 << 12,
+	/** The package contained more than one copy of the same attachment. */
+	DuplicateAttachments    = 1 << 13,
+	/** The package contained more than one object. */
+	MultiplePackageObjects  = 1 << 14,
+	/** The package contained an object with no fields. */
+	NullPackageObject       = 1 << 15,
+	/** The package contained a null attachment. */
+	NullPackageAttachment   = 1 << 16,
 };
 
 ENUM_CLASS_FLAGS(ECbValidateError);
@@ -124,8 +147,36 @@ CORE_API ECbValidateError ValidateCompactBinary(FMemoryView View, ECbValidateMod
  * does not contain a valid field will produce an OutOfBounds or InvalidType error instead of the
  * Padding error that would be produced by the single field validation function.
  *
- * \see \ref ValidateCompactBinary for more detail.
+ * @see ValidateCompactBinary
  */
 CORE_API ECbValidateError ValidateCompactBinaryRange(FMemoryView View, ECbValidateMode Mode);
 
+/**
+ * Validate the compact binary attachment pointed to by the view as specified by the mode flags.
+ *
+ * The attachment is validated with ValidateCompactBinary by using the validation mode specified.
+ * Include ECbValidateMode::Package to validate the attachment format and hash.
+ *
+ * @see ValidateCompactBinary
+ *
+ * @param View A memory view containing a package.
+ * @param Mode A combination of the flags for the types of validation to perform.
+ * @return None on success, otherwise the flags for the types of errors that were detected.
+ */
+CORE_API ECbValidateError ValidateCompactBinaryAttachment(FMemoryView View, ECbValidateMode Mode);
+
+/**
+ * Validate the compact binary package pointed to by the view as specified by the mode flags.
+ *
+ * The package, and attachments, are validated with ValidateCompactBinary by using the validation
+ * mode specified. Include ECbValidateMode::Package to validate the package format and hashes.
+ *
+ * @see ValidateCompactBinary
+ *
+ * @param View A memory view containing a package.
+ * @param Mode A combination of the flags for the types of validation to perform.
+ * @return None on success, otherwise the flags for the types of errors that were detected.
+ */
+CORE_API ECbValidateError ValidateCompactBinaryPackage(FMemoryView View, ECbValidateMode Mode);
+
 ///////////////////////////////////////////////////////////////////////////////////////////////////////////////////////
