Added an overload for dumping backtraces to strings that uses an instruction address to condition the callstack as an alternative to a frame-strip count.

Unix-related files are in the changelist because FUnixPlatformStackWalk overrides StackWalkAndDump[Ex]. The compiler won't pick up the overload in FGenericPlatformStackWalk, hence the need to forward for Unix. Coincidentally, the old versions use the context to condition backtraces by program counter. That we can skip out on for these new overloads.

#rnx

diff --git a/Engine/Source/Runtime/Core/Private/GenericPlatform/GenericPlatformStackWalk.cpp b/Engine/Source/Runtime/Core/Private/GenericPlatform/GenericPlatformStackWalk.cpp
--- a/Engine/Source/Runtime/Core/Private/GenericPlatform/GenericPlatformStackWalk.cpp
+++ b/Engine/Source/Runtime/Core/Private/GenericPlatform/GenericPlatformStackWalk.cpp
@@ -200,12 +200,44 @@ void FGenericPlatformStackWalk::StackWalkAndDump( ANSICHAR* HumanReadableString,
 	}
 }
 
+void FGenericPlatformStackWalk::StackWalkAndDump( ANSICHAR* HumanReadableString, SIZE_T HumanReadableStringSize, void* ProgramCounter, void* Context )
+{
+	uint64 StackTrace[100];
+	int32 Depth = FPlatformStackWalk::CaptureStackBackTrace(StackTrace, UE_ARRAY_COUNT(StackTrace), Context);
+
+	int32 CurrentDepth = 0;
+	if (ProgramCounter != nullptr)
+	{
+		for (int32 i = 0; i < Depth; ++i)
+		{
+			if (StackTrace[i] != uint64(ProgramCounter))
+			{
+				continue;
+			}
+
+			CurrentDepth = i;
+			break;
+		}
+	}
+
+	for (; CurrentDepth < Depth; CurrentDepth++)
+	{
+		FPlatformStackWalk::ProgramCounterToHumanReadableString( CurrentDepth, StackTrace[CurrentDepth], HumanReadableString, HumanReadableStringSize, reinterpret_cast< FGenericCrashContext* >( Context ) );
+		FCStringAnsi::Strncat(HumanReadableString, LINE_TERMINATOR_ANSI, (int32)HumanReadableStringSize);
+	}
+}
+
 void FGenericPlatformStackWalk::StackWalkAndDumpEx(ANSICHAR* HumanReadableString, SIZE_T HumanReadableStringSize, int32 IgnoreCount, uint32 Flags, void* Context)
 {
 	// generic implementation ignores extra flags
 	return FPlatformStackWalk::StackWalkAndDump(HumanReadableString, HumanReadableStringSize, IgnoreCount, Context);
 }
 
+void FGenericPlatformStackWalk::StackWalkAndDumpEx(ANSICHAR* HumanReadableString, SIZE_T HumanReadableStringSize, void* ProgramCounter, uint32 Flags, void* Context)
+{
+	return FPlatformStackWalk::StackWalkAndDump(HumanReadableString, HumanReadableStringSize, ProgramCounter, Context);
+}
+
 TArray<FProgramCounterSymbolInfo> FGenericPlatformStackWalk::GetStack(int32 IgnoreCount, int32 MaxDepth, void* Context)
 {
 	TArray<FProgramCounterSymbolInfo> Stack;
diff --git a/Engine/Source/Runtime/Core/Private/Unix/UnixPlatformStackWalk.cpp b/Engine/Source/Runtime/Core/Private/Unix/UnixPlatformStackWalk.cpp
--- a/Engine/Source/Runtime/Core/Private/Unix/UnixPlatformStackWalk.cpp
+++ b/Engine/Source/Runtime/Core/Private/Unix/UnixPlatformStackWalk.cpp
@@ -609,6 +609,33 @@ void FUnixPlatformStackWalk::StackWalkAndDump( ANSICHAR* HumanReadableString, SI
 	}
 }
 
+void FUnixPlatformStackWalk::StackWalkAndDump(ANSICHAR* HumanReadableString, SIZE_T HumanReadableStringSize, void* ProgramCounter, void* Context)
+{
+	FGenericPlatformStackWalk::StackWalkAndDump(HumanReadableString, HumanReadableStringSize, ProgramCounter, Context);
+}
+
+namespace
+{
+	/** Helper sets the ensure value in the context and guarantees it gets reset
+	 * afterwards (even if an exception is thrown) */
+	struct FLocalGuardHelper
+	{
+		FLocalGuardHelper(FUnixCrashContext* InContext, ECrashContextType NewType)
+			: Context(InContext), OldType(Context->GetType())
+		{
+			Context->SetType(NewType);
+		}
+		~FLocalGuardHelper()
+		{
+			Context->SetType(OldType);
+		}
+
+	private:
+		FUnixCrashContext* Context;
+		ECrashContextType OldType;
+	};
+} // namespace
+
 void FUnixPlatformStackWalk::StackWalkAndDumpEx(ANSICHAR* HumanReadableString, SIZE_T HumanReadableStringSize, int32 IgnoreCount, uint32 Flags, void* Context)
 {
 	const bool bHandlingEnsure = (Flags & EStackWalkFlags::FlagsUsedWhenHandlingEnsure) == EStackWalkFlags::FlagsUsedWhenHandlingEnsure;
@@ -624,24 +651,6 @@ void FUnixPlatformStackWalk::StackWalkAndDumpEx(ANSICHAR* HumanReadableString, S
 	}
 	else
 	{
-		/** Helper sets the ensure value in the context and guarantees it gets reset afterwards (even if an exception is thrown) */
-		struct FLocalGuardHelper
-		{
-			FLocalGuardHelper(FUnixCrashContext* InContext, ECrashContextType NewType)
-				: Context(InContext), OldType(Context->GetType())
-			{
-				Context->SetType(NewType);
-			}
-			~FLocalGuardHelper()
-			{
-				Context->SetType(OldType);
-			}
-
-		private:
-			FUnixCrashContext* Context;
-			ECrashContextType OldType;
-		};
-
 		FLocalGuardHelper Guard(reinterpret_cast<FUnixCrashContext*>(Context), HandlingType);
 		FPlatformStackWalk::StackWalkAndDump(HumanReadableString, HumanReadableStringSize, IgnoreCount, Context);
 	}
@@ -649,6 +658,18 @@ void FUnixPlatformStackWalk::StackWalkAndDumpEx(ANSICHAR* HumanReadableString, S
 	GHandlingEnsure = false;
 }
 
+void FUnixPlatformStackWalk::StackWalkAndDumpEx(ANSICHAR* HumanReadableString, SIZE_T HumanReadableStringSize, void* ProgramCounter, uint32 Flags, void* Context)
+{
+	const bool bHandlingEnsure = (Flags & EStackWalkFlags::FlagsUsedWhenHandlingEnsure) == EStackWalkFlags::FlagsUsedWhenHandlingEnsure;
+	GHandlingEnsure = bHandlingEnsure;
+	ECrashContextType HandlingType = bHandlingEnsure? ECrashContextType::Ensure : ECrashContextType::Crash;
+
+	FLocalGuardHelper Guard(reinterpret_cast<FUnixCrashContext*>(Context), HandlingType);
+	FPlatformStackWalk::StackWalkAndDump(HumanReadableString, HumanReadableStringSize, ProgramCounter, Context);
+
+	GHandlingEnsure = false;
+}
+
 namespace
 {
 	uint32 OverwriteBacktraceWithRealCallstack(uint64* BackTrace, uint32 Size, uint64* FirstCrashHandlerFrame)
diff --git a/Engine/Source/Runtime/Core/Private/Windows/WindowsPlatformStackWalk.cpp b/Engine/Source/Runtime/Core/Private/Windows/WindowsPlatformStackWalk.cpp
--- a/Engine/Source/Runtime/Core/Private/Windows/WindowsPlatformStackWalk.cpp
+++ b/Engine/Source/Runtime/Core/Private/Windows/WindowsPlatformStackWalk.cpp
@@ -266,6 +266,12 @@ void FWindowsPlatformStackWalk::StackWalkAndDump( ANSICHAR* HumanReadableString,
 	FGenericPlatformStackWalk::StackWalkAndDump(HumanReadableString, HumanReadableStringSize, IgnoreCount, Context);
 }
 
+void FWindowsPlatformStackWalk::StackWalkAndDump( ANSICHAR* HumanReadableString, SIZE_T HumanReadableStringSize, void* ProgramCounter, void* Context )
+{
+	InitStackWalking();
+	FGenericPlatformStackWalk::StackWalkAndDump(HumanReadableString, HumanReadableStringSize, ProgramCounter, Context);
+}
+
 FORCENOINLINE TArray<FProgramCounterSymbolInfo> FWindowsPlatformStackWalk::GetStack(int32 IgnoreCount, int32 MaxDepth, void* Context)
 {
 	InitStackWalking();
diff --git a/Engine/Source/Runtime/Core/Public/GenericPlatform/GenericPlatformStackWalk.h b/Engine/Source/Runtime/Core/Public/GenericPlatform/GenericPlatformStackWalk.h
--- a/Engine/Source/Runtime/Core/Public/GenericPlatform/GenericPlatformStackWalk.h
+++ b/Engine/Source/Runtime/Core/Public/GenericPlatform/GenericPlatformStackWalk.h
@@ -233,7 +233,18 @@ struct CORE_API FGenericPlatformStackWalk
 	 * @param	Context				Optional thread context information
 	 */ 
 	static void StackWalkAndDump( ANSICHAR* HumanReadableString, SIZE_T HumanReadableStringSize, int32 IgnoreCount, void* Context = nullptr );
-	
+
+	/**
+	 * Walks the stack and appends the human readable string to the passed in one.
+	 * @warning: The code assumes that HumanReadableString is large enough to contain the information.
+	 *
+	 * @param	HumanReadableString	String to concatenate information with
+	 * @param	HumanReadableStringSize size of string in characters
+	 * @param	ProgramCounter		Instruction address that should be at the top of the stack
+	 * @param	Context				Optional thread context information
+	 */ 
+	static void StackWalkAndDump( ANSICHAR* HumanReadableString, SIZE_T HumanReadableStringSize, void* ProgramCounter, void* Context = nullptr );
+
 	/**
 	 * Walks the stack and updates the Stack array with the symbol information for each line in the stack.
 	 *
@@ -269,6 +280,17 @@ struct CORE_API FGenericPlatformStackWalk
 	 * @param	Context				Optional thread context information
 	 */ 
 	static void StackWalkAndDumpEx( ANSICHAR* HumanReadableString, SIZE_T HumanReadableStringSize, int32 IgnoreCount, uint32 Flags, void* Context = nullptr );
+	/**
+	 * Walks the stack and appends the human readable string to the passed in one.
+	 * @warning: The code assumes that HumanReadableString is large enough to contain the information.
+	 *
+	 * @param	HumanReadableString	String to concatenate information with
+	 * @param	HumanReadableStringSize size of string in characters
+	 * @param	ProgramCounter		Instruction address that should be at the top of the stack
+	 * @param   Flags				Used to pass additional information (see StackWalkFlags)
+	 * @param	Context				Optional thread context information
+	 */ 
+	static void StackWalkAndDumpEx( ANSICHAR* HumanReadableString, SIZE_T HumanReadableStringSize, void* ProgramCounter, uint32 Flags, void* Context = nullptr );
 
 	/**
 	 * Returns the number of modules loaded by the currently running process.
diff --git a/Engine/Source/Runtime/Core/Public/Unix/UnixPlatformStackWalk.h b/Engine/Source/Runtime/Core/Public/Unix/UnixPlatformStackWalk.h
--- a/Engine/Source/Runtime/Core/Public/Unix/UnixPlatformStackWalk.h
+++ b/Engine/Source/Runtime/Core/Public/Unix/UnixPlatformStackWalk.h
@@ -17,7 +17,9 @@ struct CORE_API FUnixPlatformStackWalk : public FGenericPlatformStackWalk
 	static bool ProgramCounterToHumanReadableString( int32 CurrentCallDepth, uint64 ProgramCounter, ANSICHAR* HumanReadableString, SIZE_T HumanReadableStringSize, FGenericCrashContext* Context = nullptr );
 	static uint32 CaptureStackBackTrace( uint64* BackTrace, uint32 MaxDepth, void* Context = nullptr );
 	static void StackWalkAndDump( ANSICHAR* HumanReadableString, SIZE_T HumanReadableStringSize, int32 IgnoreCount, void* Context = nullptr );
+	static void StackWalkAndDump( ANSICHAR* HumanReadableString, SIZE_T HumanReadableStringSize, void* ProgramCounter, void* Context = nullptr );
 	static void StackWalkAndDumpEx(ANSICHAR* HumanReadableString, SIZE_T HumanReadableStringSize, int32 IgnoreCount, uint32 Flags, void* Context = nullptr);
+	static void StackWalkAndDumpEx(ANSICHAR* HumanReadableString, SIZE_T HumanReadableStringSize, void* ProgramCounter, uint32 Flags, void* Context = nullptr);
 	static uint32 CaptureThreadStackBackTrace(uint64 ThreadId, uint64* BackTrace, uint32 MaxDepth);
 	static void ThreadStackWalkAndDump(ANSICHAR* HumanReadableString, SIZE_T HumanReadableStringSize, int32 IgnoreCount, uint32 ThreadId);
 	static int32 GetProcessModuleCount();
diff --git a/Engine/Source/Runtime/Core/Public/Windows/WindowsPlatformStackWalk.h b/Engine/Source/Runtime/Core/Public/Windows/WindowsPlatformStackWalk.h
--- a/Engine/Source/Runtime/Core/Public/Windows/WindowsPlatformStackWalk.h
+++ b/Engine/Source/Runtime/Core/Public/Windows/WindowsPlatformStackWalk.h
@@ -21,6 +21,7 @@ struct CORE_API FWindowsPlatformStackWalk
 	static uint32 CaptureStackBackTrace( uint64* BackTrace, uint32 MaxDepth, void* Context = nullptr );
 	static uint32 CaptureThreadStackBackTrace( uint64 ThreadId, uint64* BackTrace, uint32 MaxDepth );
 	static void StackWalkAndDump( ANSICHAR* HumanReadableString, SIZE_T HumanReadableStringSize, int32 IgnoreCount, void* Context = nullptr );
+	static void StackWalkAndDump( ANSICHAR* HumanReadableString, SIZE_T HumanReadableStringSize, void* ProgramCounter, void* Context = nullptr );
 	static void ThreadStackWalkAndDump(ANSICHAR* HumanReadableString, SIZE_T HumanReadableStringSize, int32 IgnoreCount, uint32 ThreadId);
 
 	static int32 GetProcessModuleCount();
