Disabled additional crash context due to crashes.

diff --git a/Engine/Source/Runtime/Core/Public/GenericPlatform/GenericPlatformCrashContext.h b/Engine/Source/Runtime/Core/Public/GenericPlatform/GenericPlatformCrashContext.h
--- a/Engine/Source/Runtime/Core/Public/GenericPlatform/GenericPlatformCrashContext.h
+++ b/Engine/Source/Runtime/Core/Public/GenericPlatform/GenericPlatformCrashContext.h
@@ -9,7 +9,7 @@
 #include "Misc/Optional.h"
 #include "Containers/UnrealString.h"
 
-#define WITH_ADDITIONAL_CRASH_CONTEXTS 1
+#define WITH_ADDITIONAL_CRASH_CONTEXTS 0
 
 struct FProgramCounterSymbolInfo;
 
