Converted pass-by-value FNames to references. TagTrace.h is included so early it really cannot depend on any other headers otherwise it quickly ends in a circular loop. We need to forward declare as much as possible.

#rnx

diff --git a/Engine/Source/Runtime/Core/Private/ProfilingDebugging/TagTrace.cpp b/Engine/Source/Runtime/Core/Private/ProfilingDebugging/TagTrace.cpp
--- a/Engine/Source/Runtime/Core/Private/ProfilingDebugging/TagTrace.cpp
+++ b/Engine/Source/Runtime/Core/Private/ProfilingDebugging/TagTrace.cpp
@@ -7,6 +7,7 @@
 #include "Misc/ScopeLock.h"
 #include "HAL/LowLevelMemTracker.h"
 #include "Trace/Trace.inl"
+#include "UObject/NameTypes.h"
 
 #if USE_MEMORY_TRACE_TAGS
 
@@ -53,7 +54,7 @@ FMemScope::FMemScope(ELLMTag InTag)
 }
 
 ////////////////////////////////////////////////////////////////////////////////
-FMemScope::FMemScope(FName InName)
+FMemScope::FMemScope(const FName& InName)
 {
 	if (UE_TRACE_CHANNELEXPR_IS_ENABLED(MemAllocChannel))
 	{
@@ -129,7 +130,7 @@ public:
 	void 			AnnounceTagDeclarations();
 	static void 	OnAnnounceTagDeclaration(FLLMTagDeclaration& TagDeclaration);
 	int32			AnnounceCustomTag(int32 Tag, int32 ParentTag, const ANSICHAR* Display);
-	int32 			AnnounceFNameTag(FName TagName);
+	int32 			AnnounceFNameTag(const FName& TagName);
 
 private:
 	/**
@@ -229,7 +230,7 @@ void FTagTrace::OnAnnounceTagDeclaration(FLLMTagDeclaration& TagDeclaration)
 }
 
 ////////////////////////////////////////////////////////////////////////////////
-int32 FTagTrace::AnnounceFNameTag(FName Name)
+int32 FTagTrace::AnnounceFNameTag(const FName& Name)
 {
 	FScopeLock _(&Cs);
 	const int32 NameIndex = Name.GetDisplayIndex().ToUnstableInt();
@@ -281,7 +282,7 @@ int32 MemoryTrace_AnnounceCustomTag(int32 Tag, int32 ParentTag, const TCHAR* Dis
 }
 
 ////////////////////////////////////////////////////////////////////////////////
-int32 MemoryTrace_AnnounceFNameTag(FName TagName)
+int32 MemoryTrace_AnnounceFNameTag(const FName& TagName)
 {
 #if USE_MEMORY_TRACE_TAGS
 	if (GTagTrace)
diff --git a/Engine/Source/Runtime/Core/Public/ProfilingDebugging/TagTrace.h b/Engine/Source/Runtime/Core/Public/ProfilingDebugging/TagTrace.h
--- a/Engine/Source/Runtime/Core/Public/ProfilingDebugging/TagTrace.h
+++ b/Engine/Source/Runtime/Core/Public/ProfilingDebugging/TagTrace.h
@@ -1,7 +1,6 @@
 // Copyright Epic Games, Inc. All Rights Reserved.
 #pragma once
 #include "CoreTypes.h"
-#include "UObject/NameTypes.h"
 #include "Trace/Trace.h"
 #include "Trace/Detail/LogScope.h"
 
@@ -17,7 +16,7 @@ namespace UE {
 
 ////////////////////////////////////////////////////////////////////////////////
 CORE_API int32	MemoryTrace_AnnounceCustomTag(int32 Tag, int32 ParentTag, const TCHAR* Display);
-CORE_API int32	MemoryTrace_AnnounceFNameTag(FName TagName);
+CORE_API int32	MemoryTrace_AnnounceFNameTag(const class FName& TagName);
 CORE_API int32	MemoryTrace_GetActiveTag();
 
 ////////////////////////////////////////////////////////////////////////////////
@@ -51,7 +50,7 @@ class FMemScope
 public:
 	CORE_API FMemScope(int32 InTag);
 	CORE_API FMemScope(ELLMTag InTag);
-	CORE_API FMemScope(FName InName);
+	CORE_API FMemScope(const class FName& InName);
 	CORE_API FMemScope(const UE::LLMPrivate::FTagData* TagData);
 	CORE_API ~FMemScope();
 private:
