Fix for static init issue with additional crash contexts

diff --git a/Engine/Source/Runtime/Core/Private/GenericPlatform/GenericPlatformCrashContext.cpp b/Engine/Source/Runtime/Core/Private/GenericPlatform/GenericPlatformCrashContext.cpp
--- a/Engine/Source/Runtime/Core/Private/GenericPlatform/GenericPlatformCrashContext.cpp
+++ b/Engine/Source/Runtime/Core/Private/GenericPlatform/GenericPlatformCrashContext.cpp
@@ -1175,14 +1175,19 @@ void FGenericCrashContext::CopyPlatformSpecificFiles(const TCHAR* OutputDirector
 #if WITH_ADDITIONAL_CRASH_CONTEXTS
 
 thread_local FAdditionalCrashContextStack FAdditionalCrashContextStack::ThreadContextProvider;
-static FCriticalSection GAdditionalProviderLock;
 static FAdditionalCrashContextStack* GProviderHead;
 
+FCriticalSection* GetAdditionalProviderLock()
+{
+	static FCriticalSection GAdditionalProviderLock;
+	return &GAdditionalProviderLock;
+}
+
 FAdditionalCrashContextStack::FAdditionalCrashContextStack()
 	: Next(nullptr)
 {
 	// Register by appending self to the the linked list. 
-	FScopeLock _(&GAdditionalProviderLock);
+	FScopeLock _(GetAdditionalProviderLock());
 	FAdditionalCrashContextStack** Current = &GProviderHead;
 	while(*Current != nullptr)
 	{
@@ -1195,7 +1200,7 @@ FAdditionalCrashContextStack::~FAdditionalCrashContextStack()
 {
 	// Unregister by iterating the list, replacing self with next
 	// on the list.
-	FScopeLock _(&GAdditionalProviderLock);
+	FScopeLock _(GetAdditionalProviderLock());
 	FAdditionalCrashContextStack** Current = &GProviderHead;
 	while (*Current != this)
 	{
@@ -1208,7 +1213,8 @@ void FAdditionalCrashContextStack::ExecuteProviders(FCrashContextExtendedWriter&
 {
 	// Attempt to lock. If a thread crashed while holding the lock
 	// we could potentially deadlock here otherwise.
-	if (GAdditionalProviderLock.TryLock())
+	FCriticalSection* AdditionalProviderLock = GetAdditionalProviderLock();
+	if (AdditionalProviderLock->TryLock())
 	{
 		FAdditionalCrashContextStack* Provider = GProviderHead;
 		while (Provider)
@@ -1220,7 +1226,7 @@ void FAdditionalCrashContextStack::ExecuteProviders(FCrashContextExtendedWriter&
 			}
 			Provider = Provider->Next;
 		}
-		GAdditionalProviderLock.Unlock();
+		AdditionalProviderLock->Unlock();
 	}
 }
 
@@ -1250,7 +1256,7 @@ struct FCrashContextExtendedWriterImpl : public FCrashContextExtendedWriter
 			return;
 		}
 
-		OutputBuffer(Identifier, Data, DataSize, TEXT(".bin"));
+		OutputBuffer(Identifier, Data, DataSize, TEXT("bin"));
 	}
 
 	virtual void AddString(const TCHAR* Identifier, const TCHAR* DataStr) override
@@ -1263,7 +1269,7 @@ struct FCrashContextExtendedWriterImpl : public FCrashContextExtendedWriter
 		FPlatformMisc::LowLevelOutputDebugStringf(TEXT("Additional Crash Context (Key=\"%s\", Value=\"%s\")"), Identifier, DataStr);
 
 		FTCHARToUTF8 Converter(DataStr);
-		OutputBuffer(Identifier, (uint8*)Converter.Get(), Converter.Length(), TEXT(".txt"));
+		OutputBuffer(Identifier, (uint8*)Converter.Get(), Converter.Length(), TEXT("txt"));
 	}
 
 private:
diff --git a/Engine/Source/Runtime/Core/Private/Windows/WindowsPlatformCrashContext.cpp b/Engine/Source/Runtime/Core/Private/Windows/WindowsPlatformCrashContext.cpp
--- a/Engine/Source/Runtime/Core/Private/Windows/WindowsPlatformCrashContext.cpp
+++ b/Engine/Source/Runtime/Core/Private/Windows/WindowsPlatformCrashContext.cpp
@@ -894,6 +894,9 @@ int32 ReportCrashUsingCrashReportClient(FWindowsPlatformCrashContext& InContext,
 			// Copy platform specific files (e.g. minidump) to output directory
 			InContext.CopyPlatformSpecificFiles(*CrashFolderAbsolute, (void*) ExceptionInfo);
 
+			// Dump additional context
+			GuardedDumpAdditionalContext(*CrashFolderAbsolute);
+
 			// Copy the log file to output
 			FGenericCrashContext::DumpLog(CrashFolderAbsolute);
 
diff --git a/Engine/Source/Runtime/Core/Public/GenericPlatform/GenericPlatformCrashContext.h b/Engine/Source/Runtime/Core/Public/GenericPlatform/GenericPlatformCrashContext.h
--- a/Engine/Source/Runtime/Core/Public/GenericPlatform/GenericPlatformCrashContext.h
+++ b/Engine/Source/Runtime/Core/Public/GenericPlatform/GenericPlatformCrashContext.h
@@ -9,7 +9,7 @@
 #include "Misc/Optional.h"
 #include "Containers/UnrealString.h"
 
-#define WITH_ADDITIONAL_CRASH_CONTEXTS 0
+#define WITH_ADDITIONAL_CRASH_CONTEXTS 1
 
 struct FProgramCounterSymbolInfo;
 
