UE5 - This implements throttle stealing (similar to Tokio). The idea is 2 fold: 1) if a lot of workers are out of work we want to put half of them to sleep early, so that they don't create contention for the other half. The 2) Idea is to ramp up workers more slowly, the old algo was waking a worker everytime work was added which quickly wakes all workers and some of them might never find work. The new algo will only wake workers if there are no other workers that are running and out of work.

diff --git a/Engine/Source/Runtime/Core/Private/Async/Fundamental/Scheduler.cpp b/Engine/Source/Runtime/Core/Private/Async/Fundamental/Scheduler.cpp
--- a/Engine/Source/Runtime/Core/Private/Async/Fundamental/Scheduler.cpp
+++ b/Engine/Source/Runtime/Core/Private/Async/Fundamental/Scheduler.cpp
@@ -102,13 +102,18 @@ namespace LowLevelTasks
 		{
 			if (LocalQueue && QueuePreference != EQueuePreference::GlobalQueuePreference)
 			{
-				LocalQueue->Enqueue(&Task, uint32(Task.GetPriority()));
+				if (LocalQueue->Enqueue(&Task, uint32(Task.GetPriority())))
+				{
+					WakeUpWorker();
+				}
 			}
 			else
 			{
-				QueueRegistry.Enqueue(&Task, uint32(Task.GetPriority()));
+				if (QueueRegistry.Enqueue(&Task, uint32(Task.GetPriority())))
+				{
+					WakeUpWorker();
+				}
 			}
-			WakeUpWorker();
 		}
 		else
 		{
@@ -122,11 +127,12 @@ namespace LowLevelTasks
 	}
 
 	template<FTask* (FScheduler::FLocalQueueType::*DequeueFunction)()>
-	FORCEINLINE_DEBUGGABLE bool FScheduler::TryExecuteTaskFrom(FLocalQueueType* Queue)
+	FORCEINLINE_DEBUGGABLE bool FScheduler::TryExecuteTaskFrom(FLocalQueueType* Queue, FQueueRegistry::FOutOfWork& OutOfWork)
 	{
 		FTask* Task = (Queue->*DequeueFunction)();
 		if (Task)
 		{	
+			OutOfWork.Stop();
 			FTask* OldTask = ActiveTask;
 			ActiveTask = Task;
 			Task->ExecuteTask();
@@ -151,18 +157,22 @@ namespace LowLevelTasks
 		}
 		FLocalQueueType* WorkerLocalQueue = LocalQueue;
 
+		bool Drowsing = false;
 		uint32 WaitCount = 0;
+		FQueueRegistry::FOutOfWork OutOfWork = QueueRegistry.GetOutOfWorkScope();
 		while (true)
 		{
-			while(TryExecuteTaskFrom<&FLocalQueueType::DequeueLocal>(WorkerLocalQueue)
-			   || TryExecuteTaskFrom<&FLocalQueueType::DequeueGlobal>(WorkerLocalQueue))
-			{
+			while(TryExecuteTaskFrom<&FLocalQueueType::DequeueLocal>(WorkerLocalQueue, OutOfWork)
+			   || TryExecuteTaskFrom<&FLocalQueueType::DequeueGlobal>(WorkerLocalQueue, OutOfWork))
+			{			
+				Drowsing = false;
 				WaitCount = 0;
 			}
 
-			while(TryExecuteTaskFrom<&FLocalQueueType::DequeueLocal>(WorkerLocalQueue)
-			   || TryExecuteTaskFrom<&FLocalQueueType::DequeueSteal>(WorkerLocalQueue))
+			while(TryExecuteTaskFrom<&FLocalQueueType::DequeueLocal>(WorkerLocalQueue, OutOfWork)
+			   || TryExecuteTaskFrom<&FLocalQueueType::DequeueSteal>(WorkerLocalQueue, OutOfWork))
 			{
+				Drowsing = false;
 				WaitCount = 0;
 			}
 
@@ -173,7 +183,7 @@ namespace LowLevelTasks
 
 			if (WaitCount < WorkerSpinCycles)
 			{
-				TRACE_CPUPROFILER_EVENT_SCOPE("Spinning Worker");
+				OutOfWork.Start();			
 				for (uint32 i = 0; i < WaitCycles; i++)
 				{
 					FPlatformProcess::Yield();
@@ -182,7 +192,7 @@ namespace LowLevelTasks
 				continue;
 			}
 
-			TrySleeping(WorkerEvent, WaitCount);
+			TrySleeping(WorkerEvent, OutOfWork, WaitCount, Drowsing);
 		}
 
 		while (WakeUpWorker())
@@ -204,10 +214,11 @@ namespace LowLevelTasks
 		FLocalQueueType* WorkerLocalQueue = LocalQueue;
 
 		uint32 WaitCount = 0;
+		FQueueRegistry::FOutOfWork OutOfWork = QueueRegistry.GetOutOfWorkScope();
 		while (true)
 		{
-			while(TryExecuteTaskFrom<&FLocalQueueType::DequeueLocal>(WorkerLocalQueue)
-			   || TryExecuteTaskFrom<&FLocalQueueType::DequeueGlobal>(WorkerLocalQueue))
+			while(TryExecuteTaskFrom<&FLocalQueueType::DequeueLocal>(WorkerLocalQueue, OutOfWork)
+			   || TryExecuteTaskFrom<&FLocalQueueType::DequeueGlobal>(WorkerLocalQueue, OutOfWork))
 			{
 				if (Conditional())
 				{
@@ -216,8 +227,8 @@ namespace LowLevelTasks
 				WaitCount = 0;
 			}
 
-			while(TryExecuteTaskFrom<&FLocalQueueType::DequeueLocal>(WorkerLocalQueue)
-			   || TryExecuteTaskFrom<&FLocalQueueType::DequeueSteal>(WorkerLocalQueue))
+			while(TryExecuteTaskFrom<&FLocalQueueType::DequeueLocal>(WorkerLocalQueue, OutOfWork)
+			   || TryExecuteTaskFrom<&FLocalQueueType::DequeueSteal>(WorkerLocalQueue, OutOfWork))
 			{
 				if (Conditional())
 				{
diff --git a/Engine/Source/Runtime/Core/Public/Async/Fundamental/LocalQueue.h b/Engine/Source/Runtime/Core/Public/Async/Fundamental/LocalQueue.h
--- a/Engine/Source/Runtime/Core/Public/Async/Fundamental/LocalQueue.h
+++ b/Engine/Source/Runtime/Core/Public/Async/Fundamental/LocalQueue.h
@@ -1,8 +1,10 @@
 // Copyright Epic Games, Inc. All Rights Reserved.
 
 #pragma once
+#include "CoreTypes.h"
 #include "Math/RandomStream.h"
 #include "Experimental/Containers/FAAArrayQueue.h"
+#include "ProfilingDebugging/CpuProfilerTrace.h"
 
 #include <atomic>
 
@@ -94,7 +96,7 @@ protected:
 			{
 				return false;
 			}
-			else if (Slot == uintptr_t(ESlotState::Free) && ItemSlots[Idx].Value.compare_exchange_strong(Slot, uintptr_t(ESlotState::Taken), std::memory_order_relaxed))
+			else if (Slot == uintptr_t(ESlotState::Free) && ItemSlots[Idx].Value.compare_exchange_weak(Slot, uintptr_t(ESlotState::Taken), std::memory_order_relaxed))
 			{
 				if(IdxVer == Tail.load(std::memory_order_relaxed))
 				{
@@ -157,6 +159,64 @@ class TLocalQueueRegistry
 public:
 	class TLocalQueue;
 
+	// FOutOfWork is used to track the time while a worker is waiting for work
+	// this happens after a worker was unable to aquire any task from the queues and until it finds work again or it goes into drowsing state.
+	class FOutOfWork
+	{
+	private:
+		template<typename, uint32, uint32>
+		friend class TLocalQueueRegistry;
+
+		std::atomic_int& NumWorkersLookingForWork;
+		bool ActivelyLookingForWork = false;
+
+		static uint32 WorkerLookingForWorkTraceId;
+
+		inline FOutOfWork(std::atomic_int& InNumWorkersLookingForWork) : NumWorkersLookingForWork(InNumWorkersLookingForWork)
+		{
+#if CPUPROFILERTRACE_ENABLED
+			if (WorkerLookingForWorkTraceId == 0) 
+			{
+				WorkerLookingForWorkTraceId = FCpuProfilerTrace::OutputEventType("TaskWorkerIsLookingForWork");
+			}
+#endif
+		}
+
+	public:
+		inline ~FOutOfWork()
+		{
+			Stop();
+		}
+
+		inline bool Start()
+		{
+			if (!ActivelyLookingForWork)
+			{
+#if CPUPROFILERTRACE_ENABLED
+				FCpuProfilerTrace::OutputBeginEvent(WorkerLookingForWorkTraceId);
+#endif
+				NumWorkersLookingForWork.fetch_add(1, std::memory_order_relaxed);
+				ActivelyLookingForWork = true;
+				return true;
+			}
+			return false;
+		}
+
+		inline bool Stop()
+		{
+			if (ActivelyLookingForWork)
+			{
+#if CPUPROFILERTRACE_ENABLED
+				FCpuProfilerTrace::OutputEndEvent();
+#endif
+				NumWorkersLookingForWork.fetch_sub(1, std::memory_order_release);
+				ActivelyLookingForWork = false;
+				return true;
+			}
+			return false;
+		}
+	};
+
 private:
 	using FLocalQueueType	 = LocalQueue_Impl::TWorkStealingQueue2<ItemType, NumLocalItems>;
 	using FOverflowQueueType = FAAArrayQueue<ItemType>;
@@ -235,7 +295,8 @@ public:
 		}
 
 		//add an item to the local queue and overflow into the global queue if full
-		inline void Enqueue(ItemType* Item, uint32 PriorityIndex)
+		// returns true if we should wake a worker
+		inline bool Enqueue(ItemType* Item, uint32 PriorityIndex)
 		{
 			checkSlow(Registry);
 			checkSlow(PriorityIndex < NumPriorities);
@@ -244,11 +305,14 @@ public:
 			if (!LocalQueues[PriorityIndex].Put(Item))
 			{
 				Registry->OverflowQueues[PriorityIndex].enqueue(Item);
+				return Registry->AnyWorkerLookingForWork();
 			}
+			return (LocalTasksSinceLastDequeue++ != 0) && Registry->AnyWorkerLookingForWork();
 		}
 
 		inline ItemType* DequeueLocal()
 		{
+			LocalTasksSinceLastDequeue = 0;
 			for (int32 PriorityIndex = 0; PriorityIndex < NumPriorities; PriorityIndex++)
 			{
 				ItemType* Item;
@@ -262,12 +326,15 @@ public:
 
 		inline ItemType* DequeueGlobal()
 		{
-			for (int32 PriorityIndex = 0; PriorityIndex < NumPriorities; PriorityIndex++)
+			if (Registry->NumActiveWorkers.load(std::memory_order_relaxed) >= (2 * Registry->NumWorkersLookingForWork.load(std::memory_order_relaxed) - 1))
 			{
-				ItemType* Item = Registry->OverflowQueues[PriorityIndex].dequeue(DequeueHazards[PriorityIndex]);
-				if (Item)
+				for (int32 PriorityIndex = 0; PriorityIndex < NumPriorities; PriorityIndex++)
 				{
-					return Item;
+					ItemType* Item = Registry->OverflowQueues[PriorityIndex].dequeue(DequeueHazards[PriorityIndex]);
+					if (Item)
+					{
+						return Item;
+					}
 				}
 			}
 			return nullptr;
@@ -275,15 +342,18 @@ public:
 
 		inline ItemType* DequeueSteal()
 		{
-			if (CachedRandomIndex == InvalidIndex)
+			if (Registry->NumActiveWorkers.load(std::memory_order_relaxed) >= (2 * Registry->NumWorkersLookingForWork.load(std::memory_order_relaxed) - 1))
 			{
-				CachedRandomIndex = Random.GetUnsignedInt();
-			}
+				if (CachedRandomIndex == InvalidIndex)
+				{
+					CachedRandomIndex = Random.GetUnsignedInt();
+				}
 
-			ItemType* Result = Registry->StealItem(StealHazard, CachedRandomIndex, CachedPriorityIndex);
-			if (Result)
-			{
-				return Result;
+				ItemType* Result = Registry->StealItem(StealHazard, CachedRandomIndex, CachedPriorityIndex);
+				if (Result)
+				{
+					return Result;
+				}
 			}
 			return nullptr;
 		}
@@ -297,6 +367,7 @@ public:
 		FRandomStream			Random;
 		uint32					CachedRandomIndex = InvalidIndex;
 		uint32					CachedPriorityIndex = 0;
+		uint32					LocalTasksSinceLastDequeue = 0;
 	};
 
 	TLocalQueueRegistry()
@@ -308,6 +379,7 @@ private:
 	// add a queue to the Registry
 	void AddLocalQueue(FStealHazard& Hazard, TLocalQueue* QueueToAdd)
 	{
+		NumActiveWorkers.fetch_add(1, std::memory_order_relaxed);
 		while(true)
 		{
 			FLocalQueueCollection* Previous = Hazard.Get();
@@ -328,6 +400,7 @@ private:
 	// remove a queue from the Registry
 	void DeleteLocalQueue(FStealHazard& Hazard, TLocalQueue* QueueToRemove, bool WorkerOwned)
 	{
+		NumActiveWorkers.fetch_sub(1, std::memory_order_relaxed);
 		while(true)
 		{
 			FLocalQueueCollection* Previous = Hazard.Get();
@@ -381,12 +454,15 @@ private:
 
 public:
 	// enqueue an Item directy into the Global OverflowQueue
-	void Enqueue(ItemType* Item, uint32 PriorityIndex)
+	// returns true if we should wake a worker for stealing
+	bool Enqueue(ItemType* Item, uint32 PriorityIndex)
 	{
 		check(PriorityIndex < NumPriorities);
 		check(Item != nullptr);
 
 		OverflowQueues[PriorityIndex].enqueue(Item);
+
+		return AnyWorkerLookingForWork();
 	}
 
 	// grab an Item directy from the Global OverflowQueue
@@ -403,8 +479,23 @@ public:
 		return nullptr;
 	}
 
+	inline FOutOfWork GetOutOfWorkScope()
+	{
+		return FOutOfWork(NumWorkersLookingForWork);
+	}
+
 private:
+	inline bool AnyWorkerLookingForWork() const
+	{
+		return NumWorkersLookingForWork.load(std::memory_order_acquire) == 0;
+	}
+
 	FOverflowQueueType	  OverflowQueues[NumPriorities];
 	FHazardPointerCollection		  HazardsCollection;
-	std::atomic<FLocalQueueCollection*>	QueueCollection;	
+	std::atomic<FLocalQueueCollection*>	QueueCollection;
+	std::atomic_int NumWorkersLookingForWork { 0 };
+	std::atomic_int NumActiveWorkers { 0 };
 };
+
+template<typename ItemType, uint32 NumPriorities, uint32 NumLocalItems>
+uint32 TLocalQueueRegistry<ItemType, NumPriorities, NumLocalItems>::FOutOfWork::WorkerLookingForWorkTraceId = 0;
\ No newline at end of file
diff --git a/Engine/Source/Runtime/Core/Public/Async/Fundamental/Scheduler.h b/Engine/Source/Runtime/Core/Public/Async/Fundamental/Scheduler.h
--- a/Engine/Source/Runtime/Core/Public/Async/Fundamental/Scheduler.h
+++ b/Engine/Source/Runtime/Core/Public/Async/Fundamental/Scheduler.h
@@ -87,10 +87,8 @@ namespace LowLevelTasks
 
 		struct FSleepEvent
 		{
-			FEvent* SleepEvent;
+			FEventRef SleepEvent;
 			std::atomic<ESleepState> SleepState { ESleepState::Running };
-			inline FSleepEvent();
-			inline ~FSleepEvent();
 		};
 		using FEventQueueType = FAAArrayQueue<FSleepEvent>;
 
@@ -103,19 +101,19 @@ namespace LowLevelTasks
 		void WorkerMain(struct FSleepEvent* WorkerEvent, FLocalQueueType* ExternalWorkerLocalQueue, uint32 WaitCycles);
 		CORE_API void LaunchInternal(FTask& Task, EQueuePreference QueuePreference);
 		CORE_API void BusyWaitInternal(const FConditional& Conditional);
-		FORCENOINLINE void TrySleeping(FSleepEvent* WorkerEvent, uint32& WaitCount);
+		FORCENOINLINE void TrySleeping(FSleepEvent* WorkerEvent, FQueueRegistry::FOutOfWork& OutOfWork, uint32& WaitCount, bool& Drowsing);
 		inline bool WakeUpWorker();
 		template<FTask* (FLocalQueueType::*DequeueFunction)()>
-		FORCEINLINE_DEBUGGABLE static bool TryExecuteTaskFrom(FLocalQueueType* Queue);
+		FORCEINLINE_DEBUGGABLE static bool TryExecuteTaskFrom(FLocalQueueType* Queue, FQueueRegistry::FOutOfWork& OutOfWork);
 
 	private:
-		FEventQueueType SleepEventQueue;
-		FQueueRegistry QueueRegistry;
-		FCriticalSection WorkerThreadsCS;
+		FEventQueueType 			SleepEventQueue;
+		FQueueRegistry 				QueueRegistry;
+		FCriticalSection 			WorkerThreadsCS;
 		TArray<TUniquePtr<FThread>> WorkerThreads;
-		TArray<FLocalQueueType> WorkerLocalQueues;
-		std::atomic_uint ActiveWorkers { 0 };
-		std::atomic_uint NextWorkerId { 0 };
+		TArray<FLocalQueueType>		WorkerLocalQueues;
+		std::atomic_uint			ActiveWorkers { 0 };
+		std::atomic_uint			NextWorkerId { 0 };
 	};
 
 	FORCEINLINE_DEBUGGABLE bool TryLaunch(FTask& Task, EQueuePreference QueuePreference = EQueuePreference::DefaultPreference)
@@ -220,21 +218,30 @@ namespace LowLevelTasks
 		}
 	}
 
-	inline void FScheduler::TrySleeping(FSleepEvent* WorkerEvent, uint32& WaitCount)
+	inline void FScheduler::TrySleeping(FSleepEvent* WorkerEvent, FQueueRegistry::FOutOfWork& OutOfWork, uint32& WaitCount, bool& Drowsing)
 	{
-		ESleepState DrowsingState = ESleepState::Drowsing;
+		ESleepState DrowsingState1 = ESleepState::Drowsing;
+		ESleepState DrowsingState2 = ESleepState::Drowsing;
 		ESleepState RunningState  = ESleepState::Running;
-		if(WorkerEvent->SleepState.compare_exchange_strong(DrowsingState, ESleepState::Sleeping, std::memory_order_relaxed))
+		if(!Drowsing && WorkerEvent->SleepState.compare_exchange_strong(DrowsingState1, ESleepState::Drowsing, std::memory_order_relaxed)) //continue drowsing
 		{
-			TRACE_CPUPROFILER_EVENT_SCOPE("Sleeping Worker");
+			verifySlow(OutOfWork.Stop());
+			Drowsing = true;
+			FPlatformProcess::SleepNoStats(0);	// Alternative State one: ((Running -> Drowsing) -> Drowsing)
+		}
+		else if(WorkerEvent->SleepState.compare_exchange_strong(DrowsingState2, ESleepState::Sleeping, std::memory_order_relaxed))
+		{
+			verifySlow(!OutOfWork.Stop());
 			WaitCount = 0;
+			Drowsing = false;
 			WorkerEvent->SleepEvent->Wait(); // State two: ((Running -> Drowsing) -> Sleeping)
 		}
 		else if(WorkerEvent->SleepState.compare_exchange_strong(RunningState, ESleepState::Drowsing, std::memory_order_relaxed))
 		{
-			TRACE_CPUPROFILER_EVENT_SCOPE("Drowsing Worker");
-			SleepEventQueue.enqueue(WorkerEvent); // State one: (Running -> Drowsing)
-			FPlatformProcess::SleepNoStats(0);	
+			verifySlow(OutOfWork.Stop() == !Drowsing);
+			Drowsing = true;
+			SleepEventQueue.enqueue(WorkerEvent);
+			FPlatformProcess::SleepNoStats(0); // State one: (Running -> Drowsing)
 		}
 		else
 		{
@@ -247,22 +254,15 @@ namespace LowLevelTasks
 		FSleepEvent* WorkerEvent = SleepEventQueue.dequeue();
 		if (WorkerEvent)
 		{
-			ESleepState DrowsingState = ESleepState::Drowsing;
-			ESleepState SleepingState = ESleepState::Sleeping;
-			if (WorkerEvent->SleepState.compare_exchange_strong(DrowsingState, ESleepState::Running, std::memory_order_relaxed))
-			{
-				return true; // Solving State one: (Running -> Drowsing) -> Running
-			}
-			else if (WorkerEvent->SleepState.compare_exchange_strong(SleepingState, ESleepState::Running, std::memory_order_relaxed))
+			// Solving State : (Running -> Drowsing) -> Running  OR ((Running -> Drowsing) -> Drowsing) -> Running OR (((Running -> Drowsing) -> Sleeping) -> Running)
+			ESleepState SleepState = WorkerEvent->SleepState.exchange(ESleepState::Running, std::memory_order_relaxed);
+			if (SleepState == ESleepState::Sleeping)
 			{
-				WorkerEvent->SleepEvent->Trigger();
-				return true; // Solving State two: (((Running -> Drowsing) -> Sleeping) -> Running)
-			}
-			else
-			{
-				checkf(false, TEXT("Worker was not sleeping or drowsing: %d"), WorkerEvent->SleepState.load(std::memory_order_relaxed));
+				WorkerEvent->SleepEvent->Trigger(); // Solving State two: (((Running -> Drowsing) -> Sleeping) -> Running)
 				return true;
-			}		
+			}
+			checkf(SleepState == ESleepState::Drowsing, TEXT("Worker was not drowsing: %d"), WorkerEvent->SleepState.load(std::memory_order_relaxed));
+			return true;
 		}
 		return false;
 	}
@@ -276,14 +276,5 @@ namespace LowLevelTasks
 	{
 		StopWorkers();
 	}
-
-	inline FScheduler::FSleepEvent::FSleepEvent() : SleepEvent(FPlatformProcess::GetSynchEventFromPool())
-	{
-	}
-
-	inline FScheduler::FSleepEvent::~FSleepEvent()
-	{
-		FPlatformProcess::ReturnSynchEventToPool(SleepEvent);
-	}
 }
 
