Fix static member visibility to resolve linktime errors on macOS.

#rnx

diff --git a/Engine/Source/Runtime/Core/Private/GenericPlatform/GenericPlatformCrashContext.cpp b/Engine/Source/Runtime/Core/Private/GenericPlatform/GenericPlatformCrashContext.cpp
--- a/Engine/Source/Runtime/Core/Private/GenericPlatform/GenericPlatformCrashContext.cpp
+++ b/Engine/Source/Runtime/Core/Private/GenericPlatform/GenericPlatformCrashContext.cpp
@@ -1174,8 +1174,23 @@ void FGenericCrashContext::CopyPlatformSpecificFiles(const TCHAR* OutputDirector
 
 #if WITH_ADDITIONAL_CRASH_CONTEXTS
 
-thread_local FAdditionalCrashContextStack FAdditionalCrashContextStack::ThreadContextProvider;
-static FAdditionalCrashContextStack* GProviderHead;
+static FAdditionalCrashContextStack* GProviderHead = nullptr;
+
+FAdditionalCrashContextStack& FAdditionalCrashContextStack::GetThreadContextProvider()
+{
+	static thread_local FAdditionalCrashContextStack ThreadContextProvider;
+	return ThreadContextProvider;
+}
+
+void FAdditionalCrashContextStack::PushProvider(struct FScopedAdditionalCrashContextProvider* Provider)
+{
+	GetThreadContextProvider().PushProviderInternal(Provider);
+}
+
+void FAdditionalCrashContextStack::PopProvider()
+{
+	GetThreadContextProvider().PopProviderInternal();
+}
 
 FCriticalSection* GetAdditionalProviderLock()
 {
diff --git a/Engine/Source/Runtime/Core/Public/GenericPlatform/GenericPlatformCrashContext.h b/Engine/Source/Runtime/Core/Public/GenericPlatform/GenericPlatformCrashContext.h
--- a/Engine/Source/Runtime/Core/Public/GenericPlatform/GenericPlatformCrashContext.h
+++ b/Engine/Source/Runtime/Core/Public/GenericPlatform/GenericPlatformCrashContext.h
@@ -574,21 +574,14 @@ struct FCrashContextExtendedWriter
  */
 struct FAdditionalCrashContextStack
 {
-	CORE_API static void PushProvider(struct FScopedAdditionalCrashContextProvider* Provider)
-	{
-		ThreadContextProvider.PushProviderInternal(Provider);
-	}
-
-	CORE_API static void PopProvider()
-	{
-		ThreadContextProvider.PopProviderInternal();
-	}
+	CORE_API static FAdditionalCrashContextStack& GetThreadContextProvider();
+	CORE_API static void PushProvider(struct FScopedAdditionalCrashContextProvider* Provider);
+	CORE_API static void PopProvider();
 
 	static void ExecuteProviders(FCrashContextExtendedWriter& Writer);
 
 private:
 	enum { MaxStackDepth = 16 };
-	static thread_local FAdditionalCrashContextStack ThreadContextProvider;
 	FAdditionalCrashContextStack* Next;
 	const FScopedAdditionalCrashContextProvider* Stack[MaxStackDepth];
 	uint32 StackIndex = 0;
