Call lambdas through a non-inlined function. MSVC doesn't pay attention to FORCENOINLINE on macro so UE_LOG's lambdas don't achieve their purpose.
 See DispatchCheckVerify for further information.

#rnx

diff --git a/Engine/Source/Runtime/Core/Public/Logging/LogMacros.h b/Engine/Source/Runtime/Core/Public/Logging/LogMacros.h
--- a/Engine/Source/Runtime/Core/Public/Logging/LogMacros.h
+++ b/Engine/Source/Runtime/Core/Public/Logging/LogMacros.h
@@ -191,21 +191,20 @@ private:
 		{ \
 			UE_LOG_EXPAND_IS_FATAL(Verbosity, PREPROCESSOR_NOTHING, if (!CategoryName.IsSuppressed(ELogVerbosity::Verbosity))) \
 			{ \
-				auto UE_LOG_noinline_lambda = [](const auto& LCategoryName, const auto& LFormat, const auto&... UE_LOG_Args) FORCENOINLINE \
+				DispatchCheckVerify([&] () FORCENOINLINE \
 				{ \
-					TRACE_LOG_MESSAGE(LCategoryName, Verbosity, LFormat, UE_LOG_Args...) \
+					TRACE_LOG_MESSAGE(CategoryName, Verbosity, Format, ##__VA_ARGS__) \
 					UE_LOG_EXPAND_IS_FATAL(Verbosity, \
 						{ \
-							FMsg::Logf_Internal(UE_LOG_SOURCE_FILE(__FILE__), __LINE__, LCategoryName.GetCategoryName(), ELogVerbosity::Verbosity, LFormat, UE_LOG_Args...); \
+							FMsg::Logf_Internal(UE_LOG_SOURCE_FILE(__FILE__), __LINE__, CategoryName.GetCategoryName(), ELogVerbosity::Verbosity, Format, ##__VA_ARGS__); \
 							_DebugBreakAndPromptForRemote(); \
 							FDebug::ProcessFatalError(); \
 						}, \
 						{ \
-							FMsg::Logf_Internal(nullptr, 0, LCategoryName.GetCategoryName(), ELogVerbosity::Verbosity, LFormat, UE_LOG_Args...); \
+							FMsg::Logf_Internal(nullptr, 0, CategoryName.GetCategoryName(), ELogVerbosity::Verbosity, Format, ##__VA_ARGS__); \
 						} \
 					) \
-				}; \
-				UE_LOG_noinline_lambda(CategoryName, Format, ##__VA_ARGS__); \
+				}); \
 				UE_LOG_EXPAND_IS_FATAL(Verbosity, CA_ASSUME(false);, PREPROCESSOR_NOTHING) \
 			} \
 		} \
@@ -270,22 +269,21 @@ private:
 			{ \
 				if (Condition) \
 				{ \
-					auto UE_LOG_noinline_lambda = [](const auto& LCategoryName, const auto& LFormat, const auto&... UE_LOG_Args) FORCENOINLINE \
+					DispatchCheckVerify([&] () FORCENOINLINE \
 					{ \
-						TRACE_LOG_MESSAGE(LCategoryName, Verbosity, LFormat, UE_LOG_Args...) \
+						TRACE_LOG_MESSAGE(CategoryName, Verbosity, Format, ##__VA_ARGS__) \
 						UE_LOG_EXPAND_IS_FATAL(Verbosity, \
 							{ \
-								FMsg::Logf_Internal(UE_LOG_SOURCE_FILE(__FILE__), __LINE__, LCategoryName.GetCategoryName(), ELogVerbosity::Verbosity, LFormat, UE_LOG_Args...); \
+								FMsg::Logf_Internal(UE_LOG_SOURCE_FILE(__FILE__), __LINE__, CategoryName.GetCategoryName(), ELogVerbosity::Verbosity, Format, ##__VA_ARGS__); \
 								_DebugBreakAndPromptForRemote(); \
 								FDebug::ProcessFatalError(); \
 							}, \
 							{ \
-								FMsg::Logf_Internal(nullptr, 0, LCategoryName.GetCategoryName(), ELogVerbosity::Verbosity, LFormat, UE_LOG_Args...); \
+								FMsg::Logf_Internal(nullptr, 0, CategoryName.GetCategoryName(), ELogVerbosity::Verbosity, Format, ##__VA_ARGS__); \
 							} \
 						) \
 						CA_ASSUME(true); \
-					}; \
-					UE_LOG_noinline_lambda(CategoryName, Format, ##__VA_ARGS__); \
+					}); \
 					UE_LOG_EXPAND_IS_FATAL(Verbosity, CA_ASSUME(false);, PREPROCESSOR_NOTHING) \
 				} \
 			} \
