Pass through a FFailureInfo object to FDebug::ProcessFatalError().

#rnx

diff --git a/Engine/Source/Runtime/Core/Private/Misc/AssertionMacros.cpp b/Engine/Source/Runtime/Core/Private/Misc/AssertionMacros.cpp
--- a/Engine/Source/Runtime/Core/Private/Misc/AssertionMacros.cpp
+++ b/Engine/Source/Runtime/Core/Private/Misc/AssertionMacros.cpp
@@ -504,7 +504,7 @@ void VARARGS FDebug::AssertFailed(const ANSICHAR* Expr, const ANSICHAR* File, in
 	va_end(Args);
 }
 
-void FDebug::ProcessFatalError()
+void FDebug::ProcessFatalError(const FFailureInfo& Info)
 {
 	// This is not perfect because another thread might crash and be handled before this assert
 	// but this static variable will report the crash as an assert. Given complexity of a thread
diff --git a/Engine/Source/Runtime/Core/Public/Logging/LogMacros.h b/Engine/Source/Runtime/Core/Public/Logging/LogMacros.h
--- a/Engine/Source/Runtime/Core/Public/Logging/LogMacros.h
+++ b/Engine/Source/Runtime/Core/Public/Logging/LogMacros.h
@@ -198,7 +198,7 @@ private:
 						{ \
 							FMsg::Logf_Internal(UE_LOG_SOURCE_FILE(__FILE__), __LINE__, CategoryName.GetCategoryName(), ELogVerbosity::Verbosity, Format, ##__VA_ARGS__); \
 							_DebugBreakAndPromptForRemote(); \
-							FDebug::ProcessFatalError(); \
+							FDebug::ProcessFatalError({"UE_LOG(Fatal)", __FILE__, __LINE__}); \
 						}, \
 						{ \
 							FMsg::Logf_Internal(nullptr, 0, CategoryName.GetCategoryName(), ELogVerbosity::Verbosity, Format, ##__VA_ARGS__); \
@@ -229,7 +229,7 @@ private:
 					{ \
 						FMsg::Logf_Internal(UE_LOG_SOURCE_FILE(__FILE__), __LINE__, CategoryName.GetCategoryName(), ELogVerbosity::Verbosity, Format,  ##__VA_ARGS__); \
 						_DebugBreakAndPromptForRemote(); \
-						FDebug::ProcessFatalError(); \
+						FDebug::ProcessFatalError({"UE_LOG(Fatal)", __FILE__, __LINE__}); \
 						CA_ASSUME(false); \
 					}, \
 					{ \
@@ -276,7 +276,7 @@ private:
 							{ \
 								FMsg::Logf_Internal(UE_LOG_SOURCE_FILE(__FILE__), __LINE__, CategoryName.GetCategoryName(), ELogVerbosity::Verbosity, Format, ##__VA_ARGS__); \
 								_DebugBreakAndPromptForRemote(); \
-								FDebug::ProcessFatalError(); \
+								FDebug::ProcessFatalError({"UE_LOG(Fatal)", __FILE__, __LINE__}); \
 							}, \
 							{ \
 								FMsg::Logf_Internal(nullptr, 0, CategoryName.GetCategoryName(), ELogVerbosity::Verbosity, Format, ##__VA_ARGS__); \
diff --git a/Engine/Source/Runtime/Core/Public/Misc/AssertionMacros.h b/Engine/Source/Runtime/Core/Public/Misc/AssertionMacros.h
--- a/Engine/Source/Runtime/Core/Public/Misc/AssertionMacros.h
+++ b/Engine/Source/Runtime/Core/Public/Misc/AssertionMacros.h
@@ -55,7 +55,7 @@ struct CORE_API FDebug
 	static void VARARGS AssertFailed(const ANSICHAR* Expr, const ANSICHAR* File, int32 Line, const TCHAR* Format = TEXT(""), ...);
 
 	/** Triggers a fatal error, using the error formatted to GErrorHist via a previous call to FMsg*/
-	static void ProcessFatalError();
+	static void ProcessFatalError(const FFailureInfo& Info);
 
 	// returns true if an assert has occurred
 	static bool HasAsserted();
@@ -411,7 +411,7 @@ CORE_API void VARARGS LowLevelFatalErrorHandler(const FDebug::FFailureInfo& Info
 			FDebug::FFailureInfo LlfeInfo = { "LowLevelFatalError", __FILE__, __LINE__ }; \
 			LowLevelFatalErrorHandler(LlfeInfo, Format, ##__VA_ARGS__); \
 			_DebugBreakAndPromptForRemote(); \
-			FDebug::ProcessFatalError(); \
+			FDebug::ProcessFatalError(LlfeInfo); \
 		}); \
 	}
 
