Include the instruction address where a failure (i.e. UE_LOG(fatal) or check) occurred in the FDebug::FFailureInfo.

#rnx

diff --git a/Engine/Source/Runtime/Core/Private/Logging/LogMacros.cpp b/Engine/Source/Runtime/Core/Private/Logging/LogMacros.cpp
--- a/Engine/Source/Runtime/Core/Private/Logging/LogMacros.cpp
+++ b/Engine/Source/Runtime/Core/Private/Logging/LogMacros.cpp
@@ -2,6 +2,7 @@
 
 #include "Logging/LogMacros.h"
 #include "CoreGlobals.h"
+#include "HAL/Platform.h"
 #include "Misc/ScopeLock.h"
 #include "Misc/OutputDeviceRedirector.h"
 #include "Misc/FeedbackContext.h"
@@ -62,7 +63,7 @@ void FMsg::LogfImpl(const ANSICHAR* File, int32 Line, const FLogCategoryName& Ca
 		}
 
 		const int32 NumStackFramesToIgnore = 1;
-		FDebug::FFailureInfo Info = { "FMsg::Logf(Fatal)", File, Line };
+		FDebug::FFailureInfo Info = { "FMsg::Logf(Fatal)", File, Line, PLATFORM_RETURN_ADDRESS() };
 		StaticFailDebug(TEXT("Fatal error:"), Info, Message, false, NumStackFramesToIgnore);
 		FDebug::AssertFailed("", File, Line, Message);
 	}
@@ -110,7 +111,7 @@ void FMsg::Logf_InternalImpl(const ANSICHAR* File, int32 Line, const FLogCategor
 		}
 
 		const int32 NumStackFramesToIgnore = 1;
-		FDebug::FFailureInfo Info = { "FMsg::Logf(Fatal)", File, Line };
+		FDebug::FFailureInfo Info = { "FMsg::Logf(Fatal)", File, Line, PLATFORM_RETURN_ADDRESS() };
 		StaticFailDebug(TEXT("Fatal error:"), Info, Message, false, NumStackFramesToIgnore);
 	}
 #endif
diff --git a/Engine/Source/Runtime/Core/Public/Logging/LogMacros.h b/Engine/Source/Runtime/Core/Public/Logging/LogMacros.h
--- a/Engine/Source/Runtime/Core/Public/Logging/LogMacros.h
+++ b/Engine/Source/Runtime/Core/Public/Logging/LogMacros.h
@@ -198,7 +198,7 @@ private:
 						{ \
 							FMsg::Logf_Internal(UE_LOG_SOURCE_FILE(__FILE__), __LINE__, CategoryName.GetCategoryName(), ELogVerbosity::Verbosity, Format, ##__VA_ARGS__); \
 							_DebugBreakAndPromptForRemote(); \
-							FDebug::ProcessFatalError({"UE_LOG(Fatal)", __FILE__, __LINE__}); \
+							FDebug::ProcessFatalError({"UE_LOG(Fatal)", __FILE__, __LINE__, PLATFORM_RETURN_ADDRESS()}); \
 						}, \
 						{ \
 							FMsg::Logf_Internal(nullptr, 0, CategoryName.GetCategoryName(), ELogVerbosity::Verbosity, Format, ##__VA_ARGS__); \
@@ -229,7 +229,7 @@ private:
 					{ \
 						FMsg::Logf_Internal(UE_LOG_SOURCE_FILE(__FILE__), __LINE__, CategoryName.GetCategoryName(), ELogVerbosity::Verbosity, Format,  ##__VA_ARGS__); \
 						_DebugBreakAndPromptForRemote(); \
-						FDebug::ProcessFatalError({"UE_LOG(Fatal)", __FILE__, __LINE__}); \
+						FDebug::ProcessFatalError({"UE_LOG(Fatal)", __FILE__, __LINE__, PLATFORM_RETURN_ADDRESS()}); \
 						CA_ASSUME(false); \
 					}, \
 					{ \
@@ -276,7 +276,7 @@ private:
 							{ \
 								FMsg::Logf_Internal(UE_LOG_SOURCE_FILE(__FILE__), __LINE__, CategoryName.GetCategoryName(), ELogVerbosity::Verbosity, Format, ##__VA_ARGS__); \
 								_DebugBreakAndPromptForRemote(); \
-								FDebug::ProcessFatalError({"UE_LOG(Fatal)", __FILE__, __LINE__}); \
+								FDebug::ProcessFatalError({"UE_LOG(Fatal)", __FILE__, __LINE__, PLATFORM_RETURN_ADDRESS()}); \
 							}, \
 							{ \
 								FMsg::Logf_Internal(nullptr, 0, CategoryName.GetCategoryName(), ELogVerbosity::Verbosity, Format, ##__VA_ARGS__); \
diff --git a/Engine/Source/Runtime/Core/Public/Misc/AssertionMacros.h b/Engine/Source/Runtime/Core/Public/Misc/AssertionMacros.h
--- a/Engine/Source/Runtime/Core/Public/Misc/AssertionMacros.h
+++ b/Engine/Source/Runtime/Core/Public/Misc/AssertionMacros.h
@@ -3,6 +3,7 @@
 #pragma once
 
 #include "CoreTypes.h"
+#include "HAL/Platform.h"
 #include "HAL/PlatformMisc.h"
 #include "Templates/AndOrNot.h"
 #include "Templates/EnableIf.h"
@@ -49,6 +50,7 @@ struct CORE_API FDebug
 		const ANSICHAR* Expr;
 		const ANSICHAR* File;
 		int32 Line;
+		void* ProgramCounter;
 	};
 
 	/** Logs final assert message and exits the program. */
@@ -208,7 +210,7 @@ RetType FORCENOINLINE UE_DEBUG_SECTION DispatchCheckVerify(InnerType&& Inner)
 				{ \
 					static void FORCENOINLINE UE_DEBUG_SECTION ExecCheckImplInternal() \
 					{ \
-						FDebug::FFailureInfo CheckFailureInfo = { #expr, __FILE__, __LINE__ }; \
+						FDebug::FFailureInfo CheckFailureInfo = { #expr, __FILE__, __LINE__, PLATFORM_RETURN_ADDRESS() }; \
 						FDebug::CheckVerifyFailed(CheckFailureInfo, TEXT("")); \
 					} \
 				}; \
@@ -235,7 +237,7 @@ RetType FORCENOINLINE UE_DEBUG_SECTION DispatchCheckVerify(InnerType&& Inner)
 			{ \
 				DispatchCheckVerify([&] () FORCENOINLINE UE_DEBUG_SECTION \
 				{ \
-					FDebug::FFailureInfo CheckFailureInfo = { #expr, __FILE__, __LINE__ }; \
+					FDebug::FFailureInfo CheckFailureInfo = { #expr, __FILE__, __LINE__, PLATFORM_RETURN_ADDRESS() }; \
 					FDebug::CheckVerifyFailed(CheckFailureInfo, format, ##__VA_ARGS__); \
 				}); \
 				PLATFORM_BREAK(); \
@@ -341,7 +343,7 @@ RetType FORCENOINLINE UE_DEBUG_SECTION DispatchCheckVerify(InnerType&& Inner)
 			if ((!bExecuted || Always) && FPlatformMisc::IsEnsureAllowed()) \
 			{ \
 				bExecuted = true; \
-				FDebug::FFailureInfo EnsureFailureInfo = { #InExpression, __FILE__, __LINE__ }; \
+				FDebug::FFailureInfo EnsureFailureInfo = { #InExpression, __FILE__, __LINE__, PLATFORM_RETURN_ADDRESS() }; \
 				FDebug::OptionallyLogFormattedEnsureMessageReturningFalse(true, EnsureFailureInfo, ##__VA_ARGS__); \
 				if (!FPlatformMisc::IsDebuggerPresent()) \
 				{ \
@@ -408,7 +410,7 @@ CORE_API void VARARGS LowLevelFatalErrorHandler(const FDebug::FFailureInfo& Info
 		static_assert(TIsArrayOrRefOfType<decltype(Format), TCHAR>::Value, "Formatting string must be a TCHAR array."); \
 		DispatchCheckVerify([&] () FORCENOINLINE UE_DEBUG_SECTION \
 		{ \
-			FDebug::FFailureInfo LlfeInfo = { "LowLevelFatalError", __FILE__, __LINE__ }; \
+			FDebug::FFailureInfo LlfeInfo = { "LowLevelFatalError", __FILE__, __LINE__, PLATFORM_RETURN_ADDRESS() }; \
 			LowLevelFatalErrorHandler(LlfeInfo, Format, ##__VA_ARGS__); \
 			_DebugBreakAndPromptForRemote(); \
 			FDebug::ProcessFatalError(LlfeInfo); \
