Add UE_LOG_REF so we can call logging with a const FLogCategoryBase& as argument.

Reasonings
  1. UE_LOG keeps it's token pasting behavior to allow compile time constant evaluation and avoid any perf regression.
  2. UE_LOG_REF doesn't assume that the category can be token pasted inside the macro at the cost of an additional runtime condition evaluation
  3. UE_LOG_REF maintains the benefit from lazy argument evaluation only for enabled verbosity.
  4. UE_INTERNAL_LOG_IMPL name was used to avoid polluting UE_LOG auto-completion.

diff --git a/Engine/Source/Runtime/Core/Public/Logging/LogCategory.h b/Engine/Source/Runtime/Core/Public/Logging/LogCategory.h
--- a/Engine/Source/Runtime/Core/Public/Logging/LogCategory.h
+++ b/Engine/Source/Runtime/Core/Public/Logging/LogCategory.h
@@ -52,6 +52,11 @@ struct CORE_API FLogCategoryBase
 	/** Sets up the working verbosity and clamps to the compile time verbosity. **/
 	void SetVerbosity(ELogVerbosity::Type Verbosity);
 
+	/** Gets the compile time verbosity **/
+	inline ELogVerbosity::Type GetCompileTimeVerbosity() const
+	{
+		return CompileTimeVerbosity;
+	}
 private:
 	friend class FLogSuppressionImplementation;
 	friend class FLogScopedVerbosityOverride;
diff --git a/Engine/Source/Runtime/Core/Public/Logging/LogMacros.h b/Engine/Source/Runtime/Core/Public/Logging/LogMacros.h
--- a/Engine/Source/Runtime/Core/Public/Logging/LogMacros.h
+++ b/Engine/Source/Runtime/Core/Public/Logging/LogMacros.h
@@ -109,6 +109,7 @@ private:
 			UE_LOG_EXPAND_IS_FATAL(Verbosity, CA_ASSUME(false);, PREPROCESSOR_NOTHING) \
 		} \
 	}
+	#define UE_LOG_REF(ConstLogCategoryRef, Verbosity, Format, ...) UE_LOG(ConstLogCategoryRef, Verbosity, Format, __VA_ARGS__ )
 
 	#define UE_LOG_CLINKAGE(CategoryName, Verbosity, Format, ...) UE_LOG(CategoryName, Verbosity, Format, __VA_ARGS__ )
 
@@ -178,18 +179,12 @@ private:
 		CategoryName.SetVerbosity(ELogVerbosity::Verbosity);
 
 	/** 
-	 * A  macro that outputs a formatted message to log if a given logging category is active at a given verbosity level
-	 * @param CategoryName name of the logging category
-	 * @param Verbosity, verbosity level to test against
-	 * @param Format, format text
-	 ***/
-	#define UE_LOG(CategoryName, Verbosity, Format, ...) \
-	{ \
+	 * INTERNAL IMPLEMENTATION. DO NOT CALL DIRECTLY!
+	**/
+	#define UE_INTERNAL_LOG_IMPL(CategoryName, Verbosity, Format, ...) \
 		static_assert(TIsArrayOrRefOfType<decltype(Format), TCHAR>::Value, "Formatting string must be a TCHAR array."); \
 		static_assert((ELogVerbosity::Verbosity & ELogVerbosity::VerbosityMask) < ELogVerbosity::NumVerbosity && ELogVerbosity::Verbosity > 0, "Verbosity must be constant and in range."); \
-		CA_CONSTANT_IF((ELogVerbosity::Verbosity & ELogVerbosity::VerbosityMask) <= ELogVerbosity::COMPILED_IN_MINIMUM_VERBOSITY && (ELogVerbosity::Warning & ELogVerbosity::VerbosityMask) <= FLogCategory##CategoryName::CompileTimeVerbosity) \
-		{ \
-			UE_LOG_EXPAND_IS_FATAL(Verbosity, PREPROCESSOR_NOTHING, if (!CategoryName.IsSuppressed(ELogVerbosity::Verbosity))) \
+		UE_LOG_EXPAND_IS_FATAL(Verbosity, PREPROCESSOR_NOTHING, if (!CategoryName.IsSuppressed(ELogVerbosity::Verbosity))) \
 			{ \
 				DispatchCheckVerify([&] () FORCENOINLINE \
 				{ \
@@ -207,6 +202,32 @@ private:
 				}); \
 				UE_LOG_EXPAND_IS_FATAL(Verbosity, CA_ASSUME(false);, PREPROCESSOR_NOTHING) \
 			} \
+
+	/**
+	 * A  macro that outputs a formatted message to log if a given logging category is active at a given verbosity level
+	 * @param Category  A reference to a FLogCategoryBase instance
+	 * @param Verbosity Verbosity level to test against
+	 * @param Format    Format text
+	 ***/
+	#define UE_LOG_REF(ConstLogCategoryRef, Verbosity, Format, ...) \
+	{ \
+		CA_CONSTANT_IF((ELogVerbosity::Verbosity & ELogVerbosity::VerbosityMask) <= ELogVerbosity::COMPILED_IN_MINIMUM_VERBOSITY && (ELogVerbosity::Warning & ELogVerbosity::VerbosityMask) <= ConstLogCategoryRef.GetCompileTimeVerbosity()) \
+		{ \
+			UE_INTERNAL_LOG_IMPL(ConstLogCategoryRef, Verbosity, Format, ##__VA_ARGS__); \
+		} \
+	}
+
+	/** 
+	 * A  macro that outputs a formatted message to log if a given logging category is active at a given verbosity level
+	 * @param CategoryName name of the logging category
+	 * @param Verbosity, verbosity level to test against
+	 * @param Format, format text
+	 ***/
+	#define UE_LOG(CategoryName, Verbosity, Format, ...) \
+	{ \
+		CA_CONSTANT_IF((ELogVerbosity::Verbosity & ELogVerbosity::VerbosityMask) <= ELogVerbosity::COMPILED_IN_MINIMUM_VERBOSITY && (ELogVerbosity::Warning & ELogVerbosity::VerbosityMask) <= FLogCategory##CategoryName::CompileTimeVerbosity) \
+		{ \
+			UE_INTERNAL_LOG_IMPL(CategoryName, Verbosity, Format, ##__VA_ARGS__); \
 		} \
 	}
 
