Added somewhere to store an error's program counter. Concrete error devices on desktop platforms are responsible for kicking off the reporting of asserts. This happens in Serialize() who's API can't be changed, but at the same time we want to be able to pass through where the error occurred to aid better/more-correct reporting.

#rnx

diff --git a/Engine/Source/Runtime/Core/Private/Misc/AssertionMacros.cpp b/Engine/Source/Runtime/Core/Private/Misc/AssertionMacros.cpp
--- a/Engine/Source/Runtime/Core/Private/Misc/AssertionMacros.cpp
+++ b/Engine/Source/Runtime/Core/Private/Misc/AssertionMacros.cpp
@@ -123,6 +123,8 @@ static void AssertFailedImplV(const FDebug::FFailureInfo& Info, const TCHAR* For
 	// look into fixing this.
 	bHasAsserted = true;
 
+	GError->SetErrorProgramCounter(Info.ProgramCounter);
+
 	TCHAR DescriptionString[4096];
 	FCString::GetVarArgs(DescriptionString, UE_ARRAY_COUNT(DescriptionString), Format, Args);
 
@@ -512,6 +514,7 @@ void FDebug::ProcessFatalError(const FFailureInfo& Info)
 	// look into fixing this.
 	bHasAsserted = true;
 
+	GError->SetErrorProgramCounter(Info.ProgramCounter);
 	GError->Logf(TEXT("%s"), GErrorHist);
 }
 
diff --git a/Engine/Source/Runtime/Core/Public/Misc/OutputDeviceError.h b/Engine/Source/Runtime/Core/Public/Misc/OutputDeviceError.h
--- a/Engine/Source/Runtime/Core/Public/Misc/OutputDeviceError.h
+++ b/Engine/Source/Runtime/Core/Public/Misc/OutputDeviceError.h
@@ -10,5 +10,14 @@ class CORE_API FOutputDeviceError : public FOutputDevice
 {
 public:
 	virtual void HandleError()=0;
+
+	/** Sets the location of the instruction that raise the next error */
+	void SetErrorProgramCounter(void* InProgramCounter) { ProgramCounter = InProgramCounter; }
+
+	/* Returns the instruction location of where an error occurred */
+	void* GetErrorProgramCounter() const { return ProgramCounter; }
+
+private:
+	void* ProgramCounter = nullptr;
 };
 
