- Removed unnecessary conditional. StackTrace will never be null.
- Undef FILE_LINE_DESC as we leave this translation unit.

diff --git a/Engine/Source/Runtime/Core/Private/Misc/AssertionMacros.cpp b/Engine/Source/Runtime/Core/Private/Misc/AssertionMacros.cpp
--- a/Engine/Source/Runtime/Core/Private/Misc/AssertionMacros.cpp
+++ b/Engine/Source/Runtime/Core/Private/Misc/AssertionMacros.cpp
@@ -155,14 +155,11 @@ FORCENOINLINE void StaticFailDebug( const TCHAR* Error, const FDebug::FFailureIn
 	if (FPlatformProperties::AllowsCallStackDumpDuringAssert() && bIsEnsure == false)
 	{
 		ANSICHAR StackTrace[4096];
-		if (StackTrace != NULL)
-		{
-			StackTrace[0] = 0;
-			FPlatformStackWalk::StackWalkAndDump(StackTrace, UE_ARRAY_COUNT(StackTrace), NumStackFramesToIgnore + 1);
+		StackTrace[0] = 0;
+		FPlatformStackWalk::StackWalkAndDump(StackTrace, UE_ARRAY_COUNT(StackTrace), NumStackFramesToIgnore + 1);
 
-			FCString::Strncat(DescriptionAndTrace, TEXT("\n"), UE_ARRAY_COUNT(DescriptionAndTrace) - 1);
-			FCString::Strncat(DescriptionAndTrace, ANSI_TO_TCHAR(StackTrace), UE_ARRAY_COUNT(DescriptionAndTrace) - 1);
-		}
+		FCString::Strncat(DescriptionAndTrace, TEXT("\n"), UE_ARRAY_COUNT(DescriptionAndTrace) - 1);
+		FCString::Strncat(DescriptionAndTrace, ANSI_TO_TCHAR(StackTrace), UE_ARRAY_COUNT(DescriptionAndTrace) - 1);
 	}
 
 	FScopeLock Lock( &GetFailDebugCriticalSection());
@@ -574,3 +571,5 @@ FORCENOINLINE void FDebug::DumpStackTraceToLog(const TCHAR* Heading, const ELogV
 	FMemory::SystemFree(StackTrace);
 #endif
 }
+
+#undef FILE_LINE_DESC
