Fix IOSChunckInstaller ObjC & C++ Lambda issue
ChunkID doesn't need to be copied to a __block as it's captured by value and not changed by the completion callback. The __block confuses C++ lambdas used inside UELOG

diff --git a/Engine/Source/Runtime/Core/Private/IOS/IOSChunkInstaller.cpp b/Engine/Source/Runtime/Core/Private/IOS/IOSChunkInstaller.cpp
--- a/Engine/Source/Runtime/Core/Private/IOS/IOSChunkInstaller.cpp
+++ b/Engine/Source/Runtime/Core/Private/IOS/IOSChunkInstaller.cpp
@@ -136,7 +136,6 @@ bool FIOSChunkInstall::PrioritizeChunk( uint32 ChunkID, EChunkPriority::Type Pri
     ChunkStatus.Add(ChunkID, Status);
     
     // start the download
-    __block uint32 BlockID = ChunkID;
     [ChunkRequest beginAccessingResourcesWithCompletionHandler:^(NSError* error)
      {
         if (error)
@@ -147,23 +146,23 @@ bool FIOSChunkInstall::PrioritizeChunk( uint32 ChunkID, EChunkPriority::Type Pri
         else
         {
             // mount the pak file in the NSBundle
-            if (!MountedChunks.Find(BlockID))
+            if (!MountedChunks.Find(ChunkID))
             {
-                FIOSChunkStatus* FoundStatus = ChunkStatus.Find(BlockID);
+                FIOSChunkStatus* FoundStatus = ChunkStatus.Find(ChunkID);
                 NSBundle* Bundle = FoundStatus->Request.bundle;
-                FString PakFileName = FString::Printf(TEXT("pakchunk%i-ios.pak"), BlockID);
+                FString PakFileName = FString::Printf(TEXT("pakchunk%i-ios.pak"), ChunkID);
                 NSString* ResourcePath = [Bundle pathForResource: PakFileName.GetNSString() ofType:nil];
                 NSLog(@"ResourcePath: %@", ResourcePath);
                 IPakFile* PakFile = FCoreDelegates::MountPak.Execute(FString(ResourcePath), 0);
                 if (PakFile)
                 {
-                    MountedChunks.Add(BlockID);
-					InstallDelegate.Broadcast(BlockID, true);
+                    MountedChunks.Add(ChunkID);
+					InstallDelegate.Broadcast(ChunkID, true);
                 }
                 else
                 {
-                    UE_LOG(LogChunkInstaller, Warning, TEXT("NSBundle Chunk %i couldn't be mounted."), BlockID);
-					InstallDelegate.Broadcast(BlockID, false);
+                    UE_LOG(LogChunkInstaller, Warning, TEXT("NSBundle Chunk %i couldn't be mounted."), ChunkID);
+					InstallDelegate.Broadcast(ChunkID, false);
                 }
             }
         }
