SharedBuffer: Rewrote SharedBuffer to make it easier to use, and support lazy materialization of the underlying buffer

- FUniqueBuffer has been added as a single-ownership mutable buffer.
- FSharedBuffer remains a shared-ownership buffer and is now always immutable. In addition to the factory functions that were available previously, it can be also be created by moving FUniqueBuffer.
- FSharedBuffer is now used directly rather than requiring several pointer and reference types.
- FWeakSharedBuffer has been added as a weak reference to a shared buffer, in place of the previous weak pointer types.
- FBufferOwner has been promoted out of FSharedBuffer and can be used for advanced buffer management, including lazy buffer materialization where allocating and/or populating the buffer may be deferred until the buffer is accessed.

#rnx

diff --git a/Engine/Source/Runtime/Core/Private/Memory/SharedBuffer.cpp b/Engine/Source/Runtime/Core/Private/Memory/SharedBuffer.cpp
--- a/Engine/Source/Runtime/Core/Private/Memory/SharedBuffer.cpp
+++ b/Engine/Source/Runtime/Core/Private/Memory/SharedBuffer.cpp
@@ -2,89 +2,239 @@
 
 #include "Memory/SharedBuffer.h"
 
-FSharedBufferRef FSharedBuffer::NewBuffer(
-	void* const Data,
-	const uint64 Size,
-	const ESharedBufferFlags Flags,
-	const uint32 OwnerSize)
-{
-	checkf((OwnerSize > 0) == EnumHasAnyFlags(Flags, ESharedBufferFlags::HasBufferOwner),
-		TEXT("Mismatch between OwnerSize %u and Flags."), OwnerSize);
+#include "HAL/UnrealMemory.h"
 
-	void* const BufferMemory = FMemory::Malloc(sizeof(FSharedBuffer) + OwnerSize, alignof(FSharedBuffer));
+///////////////////////////////////////////////////////////////////////////////////////////////////////////////////////
 
-	FSharedBuffer* const Buffer = new(BufferMemory) FSharedBuffer;
-	Buffer->Data = Data;
-	Buffer->Size = Size;
-	Buffer->ReferenceCountAndFlags = SetFlags(Flags);
+namespace BufferOwnerPrivate
+{
 
-	return FSharedBufferRef(SharedBufferPrivate::TSharedBufferPtr</*bAllowNull*/ false, /*bIsConst*/ false, /*bIsWeak*/ false>(Buffer));
+template <typename FOps>
+inline TBufferOwnerPtr<FOps>::TBufferOwnerPtr(FBufferOwner* const InOwner)
+	: Owner(InOwner)
+{
+	if (InOwner)
+	{
+		checkf(!FOps::HasRef(*InOwner), TEXT("FBufferOwner is referenced by another TBufferOwnerPtr. ")
+			TEXT("Construct this from an existing pointer instead of a raw pointer."));
+		FOps::AddRef(*InOwner);
+	}
 }
 
-void FSharedBuffer::DeleteBuffer(FSharedBuffer* const Buffer)
+template <typename FOps>
+inline void TBufferOwnerPtr<FOps>::Reset()
 {
-	checkSlow(!Buffer->Data && !Buffer->Size);
-	checkSlow(Buffer->ReferenceCountAndFlags == 0);
-	Buffer->~FSharedBuffer();
-	FMemory::Free(Buffer);
+	FOps::Release(Owner);
+	Owner = nullptr;
 }
 
-void FSharedBuffer::ReleaseData()
+class FBufferOwnerHeap final : public FBufferOwner
 {
-	const ESharedBufferFlags Flags = GetFlags(ReferenceCountAndFlags);
-	if (EnumHasAnyFlags(Flags, ESharedBufferFlags::HasBufferOwner))
+public:
+	inline explicit FBufferOwnerHeap(uint64 Size)
+		: FBufferOwner(FMemory::Malloc(Size), Size)
 	{
-		FBufferOwner& BufferOwner = *reinterpret_cast<FBufferOwner*>(this + 1);
-		BufferOwner.Free(Data, Size);
-		BufferOwner.~FBufferOwner();
+		SetIsMaterialized();
+		SetIsOwned();
 	}
-	else if (EnumHasAnyFlags(Flags, ESharedBufferFlags::Owned))
+
+protected:
+	virtual void FreeBuffer() final
 	{
-		FMemory::Free(Data);
+		FMemory::Free(GetData());
 	}
-	Data = nullptr;
-	Size = 0;
-	ReferenceCountAndFlags.fetch_and(~SetFlags(~ESharedBufferFlags::None));
-}
+};
 
-bool FSharedBuffer::TryMakeReadOnly() const
+class FBufferOwnerView final : public FBufferOwner
 {
-	for (uint64 Value = ReferenceCountAndFlags.load(std::memory_order_relaxed);;)
+public:
+	inline FBufferOwnerView(void* Data, uint64 Size)
+		: FBufferOwner(Data, Size)
 	{
-		if (EnumHasAnyFlags(GetFlags(Value), ESharedBufferFlags::ReadOnly))
-		{
-			return true;
-		}
-		if (GetSharedRefCount(Value) != 1 || GetWeakRefCount(Value) != 1 ||
-			!EnumHasAnyFlags(GetFlags(Value), ESharedBufferFlags::Owned))
-		{
-			return false;
-		}
-		if (ReferenceCountAndFlags.compare_exchange_weak(Value, Value | SetFlags(ESharedBufferFlags::ReadOnly),
-			std::memory_order_relaxed, std::memory_order_relaxed))
+		SetIsMaterialized();
+	}
+
+protected:
+	virtual void FreeBuffer() final
+	{
+	}
+};
+
+class FBufferOwnerOuterView final : public FBufferOwner
+{
+public:
+	inline FBufferOwnerOuterView(void* Data, uint64 Size, FSharedBuffer InOuterBuffer)
+		: FBufferOwner(Data, Size)
+		, OuterBuffer(MoveTemp(InOuterBuffer))
+	{
+		check(OuterBuffer.GetView().Contains(MakeMemoryView(Data, Size)));
+		SetIsMaterialized();
+		if (OuterBuffer.IsOwned())
 		{
-			return true;
+			SetIsOwned();
 		}
 	}
+
+protected:
+	virtual void FreeBuffer() final
+	{
+		OuterBuffer.Reset();
+	}
+
+private:
+	FSharedBuffer OuterBuffer;
+};
+
+} // BufferOwnerPrivate
+
+///////////////////////////////////////////////////////////////////////////////////////////////////////////////////////
+
+FUniqueBuffer FUniqueBuffer::Alloc(uint64 InSize)
+{
+	return FUniqueBuffer(new BufferOwnerPrivate::FBufferOwnerHeap(InSize));
+}
+
+FUniqueBuffer FUniqueBuffer::Clone(FMemoryView View)
+{
+	return Clone(View.GetData(), View.GetSize());
+}
+
+FUniqueBuffer FUniqueBuffer::Clone(const void* Data, uint64 Size)
+{
+	FUniqueBuffer Buffer = Alloc(Size);
+	FMemory::Memcpy(Buffer.GetData(), Data, Size);
+	return Buffer;
+}
+
+FUniqueBuffer FUniqueBuffer::MakeView(FMutableMemoryView View)
+{
+	return MakeView(View.GetData(), View.GetSize());
+}
+
+FUniqueBuffer FUniqueBuffer::MakeView(void* Data, uint64 Size)
+{
+	return FUniqueBuffer(new BufferOwnerPrivate::FBufferOwnerView(Data, Size));
+}
+
+FUniqueBuffer::FUniqueBuffer(FBufferOwner* InOwner)
+	: Owner(InOwner)
+{
+}
+
+void FUniqueBuffer::Reset()
+{
+	Owner.Reset();
+}
+
+void FUniqueBuffer::MakeOwned()
+{
+	if (!IsOwned())
+	{
+		*this = Clone(GetView());
+	}
 }
 
-bool FSharedBuffer::TryMakeWritable() const
+void FUniqueBuffer::Materialize() const
 {
-	for (uint64 Value = ReferenceCountAndFlags.load(std::memory_order_relaxed);;)
+	if (Owner)
 	{
-		if (GetSharedRefCount(Value) != 1 || GetWeakRefCount(Value) != 1 ||
-			!EnumHasAnyFlags(GetFlags(Value), ESharedBufferFlags::Owned))
-		{
-			return false;
-		}
-		if (!EnumHasAnyFlags(GetFlags(Value), ESharedBufferFlags::ReadOnly))
-		{
-			return true;
-		}
-		if (ReferenceCountAndFlags.compare_exchange_weak(Value, Value & ~SetFlags(ESharedBufferFlags::ReadOnly),
-			std::memory_order_relaxed, std::memory_order_relaxed))
-		{
-			return true;
-		}
+		Owner->Materialize();
 	}
 }
+
+///////////////////////////////////////////////////////////////////////////////////////////////////////////////////////
+
+FSharedBuffer FSharedBuffer::Clone(FMemoryView View)
+{
+	return FSharedBuffer(FUniqueBuffer::Clone(View));
+}
+
+FSharedBuffer FSharedBuffer::Clone(const void* Data, uint64 Size)
+{
+	return FSharedBuffer(FUniqueBuffer::Clone(Data, Size));
+}
+
+FSharedBuffer FSharedBuffer::MakeView(FMemoryView View)
+{
+	return MakeView(View.GetData(), View.GetSize());
+}
+
+FSharedBuffer FSharedBuffer::MakeView(FMemoryView View, FSharedBuffer OuterBuffer)
+{
+	if (OuterBuffer.IsNull())
+	{
+		return MakeView(View);
+	}
+	if (View == OuterBuffer.GetView())
+	{
+		return MoveTemp(OuterBuffer);
+	}
+	return FSharedBuffer(new BufferOwnerPrivate::FBufferOwnerOuterView(
+		const_cast<void*>(View.GetData()), View.GetSize(), MoveTemp(OuterBuffer)));
+}
+
+FSharedBuffer FSharedBuffer::MakeView(const void* Data, uint64 Size)
+{
+	return FSharedBuffer(new BufferOwnerPrivate::FBufferOwnerView(const_cast<void*>(Data), Size));
+}
+
+FSharedBuffer FSharedBuffer::MakeView(const void* Data, uint64 Size, FSharedBuffer OuterBuffer)
+{
+	return MakeView(MakeMemoryView(Data, Size), MoveTemp(OuterBuffer));
+}
+
+FSharedBuffer::FSharedBuffer(FBufferOwner* InOwner)
+	: Owner(InOwner)
+{
+}
+
+FSharedBuffer::FSharedBuffer(const BufferOwnerPrivate::TBufferOwnerPtr<BufferOwnerPrivate::FWeakOps>& WeakOwner)
+	: Owner(WeakOwner)
+{
+}
+
+void FSharedBuffer::Reset()
+{
+	Owner.Reset();
+}
+
+void FSharedBuffer::MakeOwned()
+{
+	if (!IsOwned())
+	{
+		*this = FSharedBuffer::Clone(GetView());
+	}
+}
+
+void FSharedBuffer::Materialize() const
+{
+	if (Owner)
+	{
+		Owner->Materialize();
+	}
+}
+
+///////////////////////////////////////////////////////////////////////////////////////////////////////////////////////
+
+FWeakSharedBuffer::FWeakSharedBuffer(const FSharedBuffer& Buffer)
+	: Owner(ToPrivateOwnerPtr(Buffer))
+{
+}
+
+FWeakSharedBuffer& FWeakSharedBuffer::operator=(const FSharedBuffer& Buffer)
+{
+	Owner = ToPrivateOwnerPtr(Buffer);
+	return *this;
+}
+
+void FWeakSharedBuffer::Reset()
+{
+	Owner.Reset();
+}
+
+FSharedBuffer FWeakSharedBuffer::Pin() const
+{
+	return FSharedBuffer(Owner);
+}
+
+///////////////////////////////////////////////////////////////////////////////////////////////////////////////////////
diff --git a/Engine/Source/Runtime/Core/Private/Serialization/CompactBinary.cpp b/Engine/Source/Runtime/Core/Private/Serialization/CompactBinary.cpp
--- a/Engine/Source/Runtime/Core/Private/Serialization/CompactBinary.cpp
+++ b/Engine/Source/Runtime/Core/Private/Serialization/CompactBinary.cpp
@@ -792,39 +792,24 @@ template class TCbFieldIterator<FCbFieldRef>;
 
 FCbFieldRefIterator FCbFieldRefIterator::CloneRange(const FCbFieldIterator& It)
 {
-	FSharedBufferPtr Buffer;
 	FMemoryView View;
 	if (It.TryGetSerializedRangeView(View))
 	{
-		Buffer = FSharedBuffer::Clone(View);
+		return MakeRange(FSharedBuffer::Clone(View));
 	}
 	else
 	{
-		Buffer = FSharedBuffer::Alloc(It.GetRangeSize());
-		It.CopyRangeTo(*Buffer);
+		FUniqueBuffer Buffer = FUniqueBuffer::Alloc(It.GetRangeSize());
+		It.CopyRangeTo(Buffer);
+		return MakeRange(FSharedBuffer(MoveTemp(Buffer)));
 	}
-	return MakeRange(MoveTemp(Buffer));
 }
 
-FSharedBufferConstPtr FCbFieldRefIterator::GetRangeBuffer() const
+FSharedBuffer FCbFieldRefIterator::GetRangeBuffer() const
 {
 	const FMemoryView RangeView = GetRangeView();
-	const FSharedBufferConstPtr& OuterBuffer = GetOuterBuffer();
-	if (OuterBuffer)
-	{
-		if (OuterBuffer->GetView() == RangeView)
-		{
-			return OuterBuffer;
-		}
-		else
-		{
-			return FSharedBuffer::MakeView(RangeView, *OuterBuffer);
-		}
-	}
-	else
-	{
-		return FSharedBuffer::MakeView(RangeView);
-	}
+	const FSharedBuffer& OuterBuffer = GetOuterBuffer();
+	return OuterBuffer.GetView() == RangeView ? OuterBuffer : FSharedBuffer::MakeView(RangeView, OuterBuffer);
 }
 
 ///////////////////////////////////////////////////////////////////////////////////////////////////////////////////////
diff --git a/Engine/Source/Runtime/Core/Private/Serialization/CompactBinaryPackage.cpp b/Engine/Source/Runtime/Core/Private/Serialization/CompactBinaryPackage.cpp
--- a/Engine/Source/Runtime/Core/Private/Serialization/CompactBinaryPackage.cpp
+++ b/Engine/Source/Runtime/Core/Private/Serialization/CompactBinaryPackage.cpp
@@ -39,24 +39,25 @@ FCbAttachment::FCbAttachment(FCbFieldRefIterator InValue, const FBlake3Hash* con
 	}
 }
 
-FCbAttachment::FCbAttachment(FSharedBufferConstPtr InBuffer, const FBlake3Hash* const InHash)
-	: Buffer(FSharedBuffer::MakeReadOnly(MoveTemp(InBuffer)))
+FCbAttachment::FCbAttachment(FSharedBuffer InBuffer, const FBlake3Hash* const InHash)
+	: Buffer(MoveTemp(InBuffer))
 {
+	Buffer.MakeOwned();
 	if (InHash)
 	{
 		Hash = *InHash;
-		if (Buffer && Buffer->GetSize())
+		if (Buffer.GetSize())
 		{
-			checkSlow(Hash == FBlake3::HashBuffer(*Buffer));
+			checkSlow(Hash == FBlake3::HashBuffer(Buffer));
 		}
 		else
 		{
 			checkfSlow(Hash == FBlake3Hash(), TEXT("A null or empty buffer must use a hash of zero."));
 		}
 	}
-	else if (Buffer && Buffer->GetSize())
+	else if (Buffer.GetSize())
 	{
-		Hash = FBlake3::HashBuffer(*Buffer);
+		Hash = FBlake3::HashBuffer(Buffer);
 	}
 	else
 	{
@@ -64,7 +65,7 @@ FCbAttachment::FCbAttachment(FSharedBufferConstPtr InBuffer, const FBlake3Hash*
 	}
 }
 
-FSharedBufferConstPtr FCbAttachment::AsBinary() const
+FSharedBuffer FCbAttachment::AsBinary() const
 {
 	if (!CompactBinary)
 	{
@@ -74,7 +75,7 @@ FSharedBufferConstPtr FCbAttachment::AsBinary() const
 	FMemoryView SerializedView;
 	if (CompactBinary.TryGetSerializedRangeView(SerializedView))
 	{
-		return SerializedView == Buffer->GetView() ? Buffer : FSharedBuffer::MakeView(SerializedView, *Buffer);
+		return SerializedView == Buffer.GetView() ? Buffer : FSharedBuffer::MakeView(SerializedView, Buffer);
 	}
 
 	return FCbFieldRefIterator::CloneRange(CompactBinary).GetRangeBuffer();
@@ -91,27 +92,21 @@ void FCbAttachment::Load(FCbFieldRefIterator& Fields)
 	const FMemoryView View = Fields.AsBinary();
 	if (View.GetSize() > 0)
 	{
-		if (FSharedBufferConstPtr OuterBuffer = Fields.GetOuterBuffer())
-		{
-			Buffer = FSharedBuffer::MakeReadOnly(FSharedBuffer::MakeView(View, *OuterBuffer));
-		}
-		else
-		{
-			Buffer = FSharedBuffer::MakeReadOnly(FSharedBuffer::Clone(View));
-		}
+		Buffer = FSharedBuffer::MakeView(View, Fields.GetOuterBuffer());
+		Buffer.MakeOwned();
 		++Fields;
 		Hash = Fields.AsAnyReference();
 		checkf(!Fields.HasError(), TEXT("Attachments must be a non-empty binary value with a content hash."));
 		if (Fields.IsReference())
 		{
-			CompactBinary = FCbFieldIterator::MakeRange(*Buffer);
+			CompactBinary = FCbFieldIterator::MakeRange(Buffer);
 		}
 		++Fields;
 	}
 	else
 	{
 		++Fields;
-		Buffer = FSharedBufferConstPtr();
+		Buffer = FSharedBuffer();
 		CompactBinary = FCbFieldIterator();
 		Hash = FBlake3Hash();
 	}
@@ -124,25 +119,27 @@ void FCbAttachment::Load(FArchive& Ar, FCbBufferAllocator Allocator)
 	const FMemoryView View = BufferField.AsBinary();
 	if (View.GetSize() > 0)
 	{
-		Buffer = FSharedBuffer::MakeReadOnly(FSharedBuffer::MakeView(View, *BufferField.GetOuterBuffer()));
+		Buffer = FSharedBuffer::MakeView(View, BufferField.GetOuterBuffer());
+		Buffer.MakeOwned();
 		CompactBinary = FCbFieldIterator();
 
 		TArray<uint8, TInlineAllocator<64>> HashBuffer;
-		FCbFieldRef HashField = LoadCompactBinary(Ar, [&HashBuffer](uint64 Size)
+		FCbFieldRef HashField = LoadCompactBinary(Ar,
+			[&HashBuffer](uint64 Size) -> FUniqueBuffer
 			{
 				HashBuffer.SetNumUninitialized(int32(Size));
-				return FSharedBuffer::MakeView(HashBuffer.GetData(), Size);
+				return FUniqueBuffer::MakeView(HashBuffer.GetData(), Size);
 			});
 		Hash = HashField.AsAnyReference();
 		checkf(!HashField.HasError(), TEXT("Attachments must be a non-empty binary value with a content hash."));
 		if (HashField.IsReference())
 		{
-			CompactBinary = FCbFieldIterator::MakeRange(*Buffer);
+			CompactBinary = FCbFieldIterator::MakeRange(Buffer);
 		}
 	}
 	else
 	{
-		Buffer = FSharedBufferConstPtr();
+		Buffer = FSharedBuffer();
 		CompactBinary = FCbFieldIterator();
 		Hash = FBlake3Hash();
 	}
@@ -163,7 +160,7 @@ void FCbAttachment::Save(FCbWriter& Writer) const
 		}
 		Writer.Reference(Hash);
 	}
-	else if (Buffer && Buffer->GetSize())
+	else if (Buffer.GetSize())
 	{
 		Writer.Binary(Buffer);
 		Writer.BinaryReference(Hash);
@@ -262,7 +259,7 @@ void FCbPackage::GatherAttachments(const FCbFieldIterator& Fields, FAttachmentRe
 	Fields.IterateRangeReferences([this, &Resolver](FCbField Field)
 		{
 			const FBlake3Hash& Hash = Field.AsAnyReference();
-			if (FSharedBufferConstPtr Buffer = Resolver(Hash))
+			if (FSharedBuffer Buffer = Resolver(Hash))
 			{
 				if (Field.IsReference())
 				{
@@ -315,9 +312,9 @@ void FCbPackage::Load(FCbFieldRefIterator& Fields)
 void FCbPackage::Load(FArchive& Ar, FCbBufferAllocator Allocator)
 {
 	uint8 StackBuffer[64];
-	const auto StackAllocator = [&Allocator, &StackBuffer](uint64 Size)
+	const auto StackAllocator = [&Allocator, &StackBuffer](uint64 Size) -> FUniqueBuffer
 	{
-		return Size <= sizeof(StackBuffer) ? FSharedBuffer::MakeView(StackBuffer, Size) : Allocator(Size);
+		return Size <= sizeof(StackBuffer) ? FUniqueBuffer::MakeView(StackBuffer, Size) : Allocator(Size);
 	};
 
 	*this = FCbPackage();
@@ -333,7 +330,8 @@ void FCbPackage::Load(FArchive& Ar, FCbBufferAllocator Allocator)
 			const FMemoryView View = ValueField.AsBinary();
 			if (View.GetSize() > 0)
 			{
-				FSharedBufferConstPtr Buffer = FSharedBuffer::MakeReadOnly(FSharedBuffer::MakeView(View, *ValueField.GetOuterBuffer()));
+				FSharedBuffer Buffer = FSharedBuffer::MakeView(View, ValueField.GetOuterBuffer());
+				Buffer.MakeOwned();
 				FCbFieldRef HashField = LoadCompactBinary(Ar, StackAllocator);
 				const FBlake3Hash& Hash = HashField.AsAnyReference();
 				checkf(!HashField.HasError(), TEXT("Attachments must be a non-empty binary value with a content hash."));
diff --git a/Engine/Source/Runtime/Core/Private/Serialization/CompactBinarySerialization.cpp b/Engine/Source/Runtime/Core/Private/Serialization/CompactBinarySerialization.cpp
--- a/Engine/Source/Runtime/Core/Private/Serialization/CompactBinarySerialization.cpp
+++ b/Engine/Source/Runtime/Core/Private/Serialization/CompactBinarySerialization.cpp
@@ -170,18 +170,18 @@ FCbFieldRef LoadCompactBinary(FArchive& Ar, FCbBufferAllocator Allocator)
 	}
 
 	// Allocate the buffer, copy the header, and read the remainder of the field.
-	FSharedBufferPtr Buffer = Allocator(FieldSize);
-	checkf(Buffer && Buffer->GetSize() == FieldSize,
+	FUniqueBuffer Buffer = Allocator(FieldSize);
+	checkf(Buffer.GetSize() == FieldSize,
 		TEXT("Allocator returned a buffer of size %" UINT64_FMT " bytes when %" UINT64_FMT " bytes were requested."),
-		Buffer ? Buffer->GetSize() : 0, FieldSize);
-	FMutableMemoryView View = Buffer->GetView();
+		Buffer.GetSize(), FieldSize);
+	FMutableMemoryView View = Buffer.GetView();
 	FMemory::Memcpy(View.GetData(), HeaderBytes.GetData(), HeaderBytes.Num());
 	View += HeaderBytes.Num();
 	if (!View.IsEmpty())
 	{
 		Ar.Serialize(View.GetData(), static_cast<int64>(View.GetSize()));
 	}
-	return FCbFieldRef(MoveTemp(Buffer));
+	return FCbFieldRef(FSharedBuffer(MoveTemp(Buffer)));
 }
 
 ///////////////////////////////////////////////////////////////////////////////////////////////////////////////////////
@@ -211,8 +211,7 @@ static FArchive& SerializeCompactBinary(FArchive& Ar, T& Value, ConvertType&& Co
 {
 	if (Ar.IsLoading())
 	{
-		Value = Invoke(Forward<ConvertType>(Convert),
-			LoadCompactBinary(Ar, [](uint64 Size) { return FSharedBuffer::Alloc(Size); }));
+		Value = Invoke(Forward<ConvertType>(Convert), LoadCompactBinary(Ar));
 	}
 	else if (Ar.IsSaving())
 	{
diff --git a/Engine/Source/Runtime/Core/Private/Serialization/CompactBinaryWriter.cpp b/Engine/Source/Runtime/Core/Private/Serialization/CompactBinaryWriter.cpp
--- a/Engine/Source/Runtime/Core/Private/Serialization/CompactBinaryWriter.cpp
+++ b/Engine/Source/Runtime/Core/Private/Serialization/CompactBinaryWriter.cpp
@@ -87,9 +87,9 @@ void FCbWriter::Reset()
 FCbFieldRefIterator FCbWriter::Save() const
 {
 	const uint64 Size = GetSaveSize();
-	FSharedBufferPtr Buffer = FSharedBuffer::Alloc(Size);
-	const FCbFieldIterator Output = Save(Buffer->GetView());
-	return FCbFieldRefIterator::MakeRangeView(Output, FSharedBuffer::MakeReadOnly(MoveTemp(Buffer)));
+	FUniqueBuffer Buffer = FUniqueBuffer::Alloc(Size);
+	const FCbFieldIterator Output = Save(Buffer);
+	return FCbFieldRefIterator::MakeRangeView(Output, FSharedBuffer(MoveTemp(Buffer)));
 }
 
 FCbFieldIterator FCbWriter::Save(const FMutableMemoryView Buffer) const
@@ -367,12 +367,9 @@ void FCbWriter::Binary(const void* const Value, const uint64 Size)
 	EndField(ECbFieldType::Binary);
 }
 
-void FCbWriter::Binary(const FSharedBufferConstPtr& Buffer)
+void FCbWriter::Binary(const FSharedBuffer& Buffer)
 {
-	if (Buffer)
-	{
-		Binary(Buffer->GetData(), Buffer->GetSize());
-	}
+	Binary(Buffer.GetData(), Buffer.GetSize());
 }
 
 void FCbWriter::String(const FAnsiStringView Value)
diff --git a/Engine/Source/Runtime/Core/Private/Tests/Memory/SharedBufferTest.cpp b/Engine/Source/Runtime/Core/Private/Tests/Memory/SharedBufferTest.cpp
--- a/Engine/Source/Runtime/Core/Private/Tests/Memory/SharedBufferTest.cpp
+++ b/Engine/Source/Runtime/Core/Private/Tests/Memory/SharedBufferTest.cpp
@@ -8,531 +8,387 @@
 
 ///////////////////////////////////////////////////////////////////////////////////////////////////////////////////////
 
-static_assert(!std::is_constructible<FSharedBuffer>::value, "Invalid constructor");
-
-static_assert(!std::is_constructible<FSharedBufferRef>::value, "Invalid constructor");
-static_assert( std::is_constructible<FSharedBufferRef, const FSharedBufferRef&>::value, "Invalid constructor");
-static_assert(!std::is_constructible<FSharedBufferRef, FSharedBufferPtr>::value, "Invalid constructor");
-static_assert(!std::is_constructible<FSharedBufferRef, FSharedBufferWeakPtr>::value, "Invalid constructor");
-static_assert(!std::is_constructible<FSharedBufferRef, FSharedBufferConstRef>::value, "Invalid constructor");
-static_assert(!std::is_constructible<FSharedBufferRef, FSharedBufferConstPtr>::value, "Invalid constructor");
-static_assert(!std::is_constructible<FSharedBufferRef, FSharedBufferConstWeakPtr>::value, "Invalid constructor");
-
-static_assert(!std::is_constructible<FSharedBufferConstRef>::value, "Invalid constructor");
-static_assert( std::is_constructible<FSharedBufferConstRef, const FSharedBufferRef&>::value, "Missing constructor");
-static_assert(!std::is_constructible<FSharedBufferConstRef, FSharedBufferPtr>::value, "Invalid constructor");
-static_assert(!std::is_constructible<FSharedBufferConstRef, FSharedBufferWeakPtr>::value, "Invalid constructor");
-static_assert( std::is_constructible<FSharedBufferConstRef, const FSharedBufferConstRef&>::value, "Invalid constructor");
-static_assert(!std::is_constructible<FSharedBufferConstRef, FSharedBufferConstPtr>::value, "Invalid constructor");
-static_assert(!std::is_constructible<FSharedBufferConstRef, FSharedBufferConstWeakPtr>::value, "Invalid constructor");
-
-static_assert( std::is_constructible<FSharedBufferPtr>::value, "Missing constructor");
-static_assert( std::is_constructible<FSharedBufferPtr, const FSharedBufferRef&>::value, "Missing constructor");
-static_assert( std::is_constructible<FSharedBufferPtr, const FSharedBufferPtr&>::value, "Missing constructor");
-static_assert( std::is_constructible<FSharedBufferPtr, FSharedBufferPtr&&>::value, "Missing constructor");
-static_assert(!std::is_constructible<FSharedBufferPtr, FSharedBufferWeakPtr>::value, "Invalid constructor");
-static_assert(!std::is_constructible<FSharedBufferPtr, FSharedBufferConstRef>::value, "Invalid constructor");
-static_assert(!std::is_constructible<FSharedBufferPtr, FSharedBufferConstPtr>::value, "Invalid constructor");
-static_assert(!std::is_constructible<FSharedBufferPtr, FSharedBufferConstWeakPtr>::value, "Invalid constructor");
-
-static_assert( std::is_constructible<FSharedBufferConstPtr>::value, "Missing constructor");
-static_assert( std::is_constructible<FSharedBufferConstPtr, const FSharedBufferRef&>::value, "Missing constructor");
-static_assert( std::is_constructible<FSharedBufferConstPtr, const FSharedBufferPtr&>::value, "Missing constructor");
-static_assert( std::is_constructible<FSharedBufferConstPtr, FSharedBufferPtr&&>::value, "Missing constructor");
-static_assert(!std::is_constructible<FSharedBufferConstPtr, FSharedBufferWeakPtr>::value, "Invalid constructor");
-static_assert( std::is_constructible<FSharedBufferConstPtr, const FSharedBufferConstRef&>::value, "Missing constructor");
-static_assert( std::is_constructible<FSharedBufferConstPtr, const FSharedBufferConstPtr&>::value, "Missing constructor");
-static_assert( std::is_constructible<FSharedBufferConstPtr, FSharedBufferConstPtr&&>::value, "Missing constructor");
-static_assert(!std::is_constructible<FSharedBufferConstPtr, FSharedBufferConstWeakPtr>::value, "Invalid constructor");
-
-static_assert( std::is_constructible<FSharedBufferWeakPtr>::value, "Missing constructor");
-static_assert( std::is_constructible<FSharedBufferWeakPtr, FSharedBufferRef>::value, "Missing constructor");
-static_assert( std::is_constructible<FSharedBufferWeakPtr, FSharedBufferPtr>::value, "Missing constructor");
-static_assert( std::is_constructible<FSharedBufferWeakPtr, const FSharedBufferWeakPtr&>::value, "Missing constructor");
-static_assert( std::is_constructible<FSharedBufferWeakPtr, FSharedBufferWeakPtr&&>::value, "Missing constructor");
-static_assert(!std::is_constructible<FSharedBufferWeakPtr, FSharedBufferConstRef>::value, "Invalid constructor");
-static_assert(!std::is_constructible<FSharedBufferWeakPtr, FSharedBufferConstPtr>::value, "Invalid constructor");
-static_assert(!std::is_constructible<FSharedBufferWeakPtr, FSharedBufferConstWeakPtr>::value, "Invalid constructor");
-
-static_assert(std::is_constructible<FSharedBufferConstWeakPtr>::value, "Missing constructor");
-static_assert(std::is_constructible<FSharedBufferConstWeakPtr, FSharedBufferRef>::value, "Missing constructor");
-static_assert(std::is_constructible<FSharedBufferConstWeakPtr, FSharedBufferPtr>::value, "Missing constructor");
-static_assert(std::is_constructible<FSharedBufferConstWeakPtr, const FSharedBufferWeakPtr&>::value, "Missing constructor");
-static_assert(std::is_constructible<FSharedBufferConstWeakPtr, FSharedBufferWeakPtr&&>::value, "Missing constructor");
-static_assert(std::is_constructible<FSharedBufferConstWeakPtr, FSharedBufferConstRef>::value, "Missing constructor");
-static_assert(std::is_constructible<FSharedBufferConstWeakPtr, FSharedBufferConstPtr>::value, "Missing constructor");
-static_assert(std::is_constructible<FSharedBufferConstWeakPtr, const FSharedBufferConstWeakPtr&>::value, "Missing constructor");
-static_assert(std::is_constructible<FSharedBufferConstWeakPtr, FSharedBufferConstWeakPtr&&>::value, "Missing constructor");
+static_assert( std::is_constructible<FUniqueBuffer>::value, "Missing constructor");
+static_assert( std::is_constructible<FUniqueBuffer, FUniqueBuffer&&>::value, "Missing constructor");
+static_assert(!std::is_constructible<FUniqueBuffer, FSharedBuffer&&>::value, "Invalid constructor");
+static_assert(!std::is_constructible<FUniqueBuffer, FWeakSharedBuffer&&>::value, "Invalid constructor");
+static_assert(!std::is_constructible<FUniqueBuffer, const FUniqueBuffer&>::value, "Invalid constructor");
+static_assert(!std::is_constructible<FUniqueBuffer, const FSharedBuffer&>::value, "Invalid constructor");
+static_assert(!std::is_constructible<FUniqueBuffer, const FWeakSharedBuffer&>::value, "Invalid constructor");
+
+static_assert( std::is_constructible<FSharedBuffer>::value, "Missing constructor");
+static_assert( std::is_constructible<FSharedBuffer, FUniqueBuffer&&>::value, "Missing constructor");
+static_assert( std::is_constructible<FSharedBuffer, FSharedBuffer&&>::value, "Missing constructor");
+static_assert(!std::is_constructible<FSharedBuffer, FWeakSharedBuffer&&>::value, "Invalid constructor");
+static_assert(!std::is_constructible<FSharedBuffer, const FUniqueBuffer&>::value, "Invalid constructor");
+static_assert( std::is_constructible<FSharedBuffer, const FSharedBuffer&>::value, "Missing constructor");
+static_assert(!std::is_constructible<FSharedBuffer, const FWeakSharedBuffer&>::value, "Invalid constructor");
+
+static_assert( std::is_constructible<FWeakSharedBuffer>::value, "Missing constructor");
+static_assert(!std::is_constructible<FWeakSharedBuffer, FUniqueBuffer&&>::value, "Invalid constructor");
+static_assert( std::is_constructible<FWeakSharedBuffer, FSharedBuffer&&>::value, "Missing constructor");
+static_assert( std::is_constructible<FWeakSharedBuffer, FWeakSharedBuffer&&>::value, "Missing constructor");
+static_assert(!std::is_constructible<FWeakSharedBuffer, const FUniqueBuffer&>::value, "Invalid constructor");
+static_assert( std::is_constructible<FWeakSharedBuffer, const FSharedBuffer&>::value, "Missing constructor");
+static_assert( std::is_constructible<FWeakSharedBuffer, const FWeakSharedBuffer&>::value, "Missing constructor");
 
 ///////////////////////////////////////////////////////////////////////////////////////////////////////////////////////
 
-static_assert( std::is_assignable<FSharedBufferRef, const FSharedBufferRef&>::value, "Missing assignment");
-static_assert(!std::is_assignable<FSharedBufferRef, FSharedBufferPtr>::value, "Invalid assignment");
-static_assert(!std::is_assignable<FSharedBufferRef, FSharedBufferWeakPtr>::value, "Invalid assignment");
-static_assert(!std::is_assignable<FSharedBufferRef, FSharedBufferConstRef>::value, "Invalid assignment");
-static_assert(!std::is_assignable<FSharedBufferRef, FSharedBufferConstPtr>::value, "Invalid assignment");
-static_assert(!std::is_assignable<FSharedBufferRef, FSharedBufferConstWeakPtr>::value, "Invalid assignment");
-
-static_assert( std::is_assignable<FSharedBufferConstRef, const FSharedBufferRef&>::value, "Missing assignment");
-static_assert(!std::is_assignable<FSharedBufferConstRef, FSharedBufferPtr>::value, "Invalid assignment");
-static_assert(!std::is_assignable<FSharedBufferConstRef, FSharedBufferWeakPtr>::value, "Invalid assignment");
-static_assert( std::is_assignable<FSharedBufferConstRef, const FSharedBufferConstRef&>::value, "Missing assignment");
-static_assert(!std::is_assignable<FSharedBufferConstRef, FSharedBufferConstPtr>::value, "Invalid assignment");
-static_assert(!std::is_assignable<FSharedBufferConstRef, FSharedBufferConstWeakPtr>::value, "Invalid assignment");
-
-static_assert( std::is_assignable<FSharedBufferPtr, const FSharedBufferRef&>::value, "Missing assignment");
-static_assert( std::is_assignable<FSharedBufferPtr, const FSharedBufferPtr&>::value, "Missing assignment");
-static_assert( std::is_assignable<FSharedBufferPtr, FSharedBufferPtr&&>::value, "Missing assignment");
-static_assert(!std::is_assignable<FSharedBufferPtr, FSharedBufferWeakPtr>::value, "Invalid assignment");
-static_assert(!std::is_assignable<FSharedBufferPtr, FSharedBufferConstRef>::value, "Invalid assignment");
-static_assert(!std::is_assignable<FSharedBufferPtr, FSharedBufferConstPtr>::value, "Invalid assignment");
-static_assert(!std::is_assignable<FSharedBufferPtr, FSharedBufferConstWeakPtr>::value, "Invalid assignment");
-
-static_assert( std::is_assignable<FSharedBufferConstPtr, const FSharedBufferRef&>::value, "Missing assignment");
-static_assert( std::is_assignable<FSharedBufferConstPtr, const FSharedBufferPtr&>::value, "Missing assignment");
-static_assert( std::is_assignable<FSharedBufferConstPtr, FSharedBufferPtr&&>::value, "Missing assignment");
-static_assert(!std::is_assignable<FSharedBufferConstPtr, FSharedBufferWeakPtr>::value, "Invalid assignment");
-static_assert( std::is_assignable<FSharedBufferConstPtr, const FSharedBufferConstRef&>::value, "Missing assignment");
-static_assert( std::is_assignable<FSharedBufferConstPtr, const FSharedBufferConstPtr&>::value, "Missing assignment");
-static_assert( std::is_assignable<FSharedBufferConstPtr, FSharedBufferConstPtr&&>::value, "Missing assignment");
-static_assert(!std::is_assignable<FSharedBufferConstPtr, FSharedBufferConstWeakPtr>::value, "Invalid assignment");
-
-static_assert( std::is_assignable<FSharedBufferWeakPtr, const FSharedBufferRef&>::value, "Missing assignment");
-static_assert( std::is_assignable<FSharedBufferWeakPtr, const FSharedBufferPtr&>::value, "Missing assignment");
-static_assert( std::is_assignable<FSharedBufferWeakPtr, const FSharedBufferWeakPtr&>::value, "Missing assignment");
-static_assert( std::is_assignable<FSharedBufferWeakPtr, FSharedBufferWeakPtr&&>::value, "Missing assignment");
-static_assert(!std::is_assignable<FSharedBufferWeakPtr, FSharedBufferConstRef>::value, "Invalid assignment");
-static_assert(!std::is_assignable<FSharedBufferWeakPtr, FSharedBufferConstPtr>::value, "Invalid assignment");
-static_assert(!std::is_assignable<FSharedBufferWeakPtr, FSharedBufferConstWeakPtr>::value, "Invalid assignment");
-
-static_assert(std::is_assignable<FSharedBufferConstWeakPtr, const FSharedBufferRef&>::value, "Missing assignment");
-static_assert(std::is_assignable<FSharedBufferConstWeakPtr, const FSharedBufferPtr&>::value, "Missing assignment");
-static_assert(std::is_assignable<FSharedBufferConstWeakPtr, const FSharedBufferWeakPtr&>::value, "Missing assignment");
-static_assert(std::is_assignable<FSharedBufferConstWeakPtr, FSharedBufferWeakPtr&&>::value, "Missing assignment");
-static_assert(std::is_assignable<FSharedBufferConstWeakPtr, const FSharedBufferConstRef&>::value, "Missing assignment");
-static_assert(std::is_assignable<FSharedBufferConstWeakPtr, const FSharedBufferConstPtr&>::value, "Missing assignment");
-static_assert(std::is_assignable<FSharedBufferConstWeakPtr, FSharedBufferConstPtr&&>::value, "Missing assignment");
-static_assert(std::is_assignable<FSharedBufferConstWeakPtr, const FSharedBufferConstWeakPtr&>::value, "Missing assignment");
-static_assert(std::is_assignable<FSharedBufferConstWeakPtr, FSharedBufferConstWeakPtr&&>::value, "Missing assignment");
+static_assert( std::is_assignable<FUniqueBuffer, FUniqueBuffer&&>::value, "Missing assignment");
+static_assert(!std::is_assignable<FUniqueBuffer, FSharedBuffer&&>::value, "Invalid assignment");
+static_assert(!std::is_assignable<FUniqueBuffer, FWeakSharedBuffer&&>::value, "Invalid assignment");
+static_assert(!std::is_assignable<FUniqueBuffer, const FUniqueBuffer&>::value, "Invalid assignment");
+static_assert(!std::is_assignable<FUniqueBuffer, const FSharedBuffer&>::value, "Invalid assignment");
+static_assert(!std::is_assignable<FUniqueBuffer, const FWeakSharedBuffer&>::value, "Invalid assignment");
+
+static_assert( std::is_assignable<FSharedBuffer, FUniqueBuffer&&>::value, "Missing assignment");
+static_assert( std::is_assignable<FSharedBuffer, FSharedBuffer&&>::value, "Missing assignment");
+static_assert(!std::is_assignable<FSharedBuffer, FWeakSharedBuffer&&>::value, "Invalid assignment");
+static_assert(!std::is_assignable<FSharedBuffer, const FUniqueBuffer&>::value, "Invalid assignment");
+static_assert( std::is_assignable<FSharedBuffer, const FSharedBuffer&>::value, "Missing assignment");
+static_assert(!std::is_assignable<FSharedBuffer, const FWeakSharedBuffer&>::value, "Invalid assignment");
+
+static_assert(!std::is_assignable<FWeakSharedBuffer, FUniqueBuffer&&>::value, "Invalid assignment");
+static_assert( std::is_assignable<FWeakSharedBuffer, FSharedBuffer&&>::value, "Missing assignment");
+static_assert( std::is_assignable<FWeakSharedBuffer, FWeakSharedBuffer&&>::value, "Missing assignment");
+static_assert(!std::is_assignable<FWeakSharedBuffer, const FUniqueBuffer&>::value, "Invalid assignment");
+static_assert( std::is_assignable<FWeakSharedBuffer, const FSharedBuffer&>::value, "Missing assignment");
+static_assert( std::is_assignable<FWeakSharedBuffer, const FWeakSharedBuffer&>::value, "Missing assignment");
 
 ///////////////////////////////////////////////////////////////////////////////////////////////////////////////////////
 
-static_assert(std::is_same<FSharedBufferPtr, decltype(FSharedBuffer::Alloc(DeclVal<uint64>()))>::value, "Invalid constructor");
-
-static_assert(std::is_same<FSharedBufferPtr, decltype(FSharedBuffer::Clone(DeclVal<void*>(), 0))>::value, "Invalid constructor");
-static_assert(std::is_same<FSharedBufferPtr, decltype(FSharedBuffer::Clone(DeclVal<const void*>(), 0))>::value, "Invalid constructor");
-static_assert(std::is_same<FSharedBufferPtr, decltype(FSharedBuffer::Clone(*DeclVal<FSharedBufferRef>()))>::value, "Invalid constructor");
-static_assert(std::is_same<FSharedBufferPtr, decltype(FSharedBuffer::Clone(*DeclVal<FSharedBufferConstRef>()))>::value, "Invalid constructor");
-static_assert(std::is_same<FSharedBufferPtr, decltype(FSharedBuffer::Clone(DeclVal<FMutableMemoryView>()))>::value, "Invalid constructor");
-static_assert(std::is_same<FSharedBufferPtr, decltype(FSharedBuffer::Clone(DeclVal<FMemoryView>()))>::value, "Invalid constructor");
-
-static_assert(std::is_same<FSharedBufferPtr, decltype(FSharedBuffer::MakeView(DeclVal<FMutableMemoryView>()))>::value, "Invalid constructor");
-static_assert(std::is_same<FSharedBufferConstPtr, decltype(FSharedBuffer::MakeView(DeclVal<FMemoryView>()))>::value, "Invalid constructor");
-static_assert(std::is_same<FSharedBufferPtr, decltype(FSharedBuffer::MakeView(DeclVal<void*>(), 0))>::value, "Invalid constructor");
-static_assert(std::is_same<FSharedBufferConstPtr, decltype(FSharedBuffer::MakeView(DeclVal<const void*>(), 0))>::value, "Invalid constructor");
-static_assert(std::is_same<FSharedBufferConstPtr, decltype(FSharedBuffer::MakeView(DeclVal<FMemoryView>(), DeclVal<const FSharedBuffer&>()))>::value, "Invalid constructor");
-static_assert(std::is_same<FSharedBufferConstPtr, decltype(FSharedBuffer::MakeView(DeclVal<const void*>(), 0, DeclVal<const FSharedBuffer&>()))>::value, "Invalid constructor");
-
-static_assert(std::is_same<FSharedBufferPtr, decltype(FSharedBuffer::TakeOwnership(DeclVal<void*>(), 0, FMemory::Free))>::value, "Invalid constructor");
-static_assert(std::is_same<FSharedBufferConstPtr, decltype(FSharedBuffer::TakeOwnership(DeclVal<const void*>(), 0, FMemory::Free))>::value, "Invalid constructor");
-
-static_assert(std::is_same<void*, decltype(DeclVal<FSharedBuffer>().GetData())>::value, "Invalid accessor");
+static_assert(std::is_same<      void*, decltype(DeclVal<      FUniqueBuffer>().GetData())>::value, "Invalid accessor");
+static_assert(std::is_same<const void*, decltype(DeclVal<const FUniqueBuffer>().GetData())>::value, "Invalid accessor");
+static_assert(std::is_same<const void*, decltype(DeclVal<      FSharedBuffer>().GetData())>::value, "Invalid accessor");
 static_assert(std::is_same<const void*, decltype(DeclVal<const FSharedBuffer>().GetData())>::value, "Invalid accessor");
+
+static_assert(std::is_same<uint64, decltype(DeclVal<FUniqueBuffer>().GetSize())>::value, "Invalid accessor");
 static_assert(std::is_same<uint64, decltype(DeclVal<FSharedBuffer>().GetSize())>::value, "Invalid accessor");
 
-static_assert(std::is_same<FMutableMemoryView, decltype(DeclVal<FSharedBuffer>().GetView())>::value, "Invalid accessor");
-static_assert(std::is_same<FMemoryView, decltype(DeclVal<const FSharedBuffer>().GetView())>::value, "Invalid accessor");
+///////////////////////////////////////////////////////////////////////////////////////////////////////////////////////
 
-static_assert(std::is_convertible<FSharedBuffer, FMutableMemoryView>::value, "Missing conversion");
-static_assert(std::is_convertible<FSharedBuffer, FMemoryView>::value, "Missing conversion");
-static_assert(!std::is_convertible<const FSharedBuffer, FMutableMemoryView>::value, "Invalid conversion");
-static_assert(std::is_convertible<const FSharedBuffer, FMemoryView>::value, "Missing conversion");
+static_assert(std::is_same<FMutableMemoryView, decltype(DeclVal<      FUniqueBuffer>().GetView())>::value, "Invalid accessor");
+static_assert(std::is_same<       FMemoryView, decltype(DeclVal<const FUniqueBuffer>().GetView())>::value, "Invalid accessor");
+static_assert(std::is_same<       FMemoryView, decltype(DeclVal<      FSharedBuffer>().GetView())>::value, "Invalid accessor");
+static_assert(std::is_same<       FMemoryView, decltype(DeclVal<const FSharedBuffer>().GetView())>::value, "Invalid accessor");
 
 ///////////////////////////////////////////////////////////////////////////////////////////////////////////////////////
 
-static_assert(std::is_same<FSharedBufferRef, decltype(FSharedBuffer::MakeOwned(DeclVal<const FSharedBufferRef&>()))>::value, "Invalid conversion");
-static_assert(std::is_same<FSharedBufferRef, decltype(FSharedBuffer::MakeOwned(DeclVal<FSharedBufferRef&&>()))>::value, "Invalid conversion");
-static_assert(std::is_same<FSharedBufferPtr, decltype(FSharedBuffer::MakeOwned(DeclVal<const FSharedBufferPtr&>()))>::value, "Invalid conversion");
-static_assert(std::is_same<FSharedBufferPtr, decltype(FSharedBuffer::MakeOwned(DeclVal<FSharedBufferPtr&&>()))>::value, "Invalid conversion");
+static_assert( std::is_convertible<      FUniqueBuffer, FMutableMemoryView>::value, "Missing conversion");
+static_assert(!std::is_convertible<const FUniqueBuffer, FMutableMemoryView>::value, "Invalid conversion");
+static_assert( std::is_convertible<      FUniqueBuffer, FMemoryView>::value, "Missing conversion");
+static_assert( std::is_convertible<const FUniqueBuffer, FMemoryView>::value, "Missing conversion");
+
+static_assert(!std::is_convertible<      FSharedBuffer, FMutableMemoryView>::value, "Invalid conversion");
+static_assert(!std::is_convertible<const FSharedBuffer, FMutableMemoryView>::value, "Invalid conversion");
+static_assert( std::is_convertible<      FSharedBuffer, FMemoryView>::value, "Missing conversion");
+static_assert( std::is_convertible<const FSharedBuffer, FMemoryView>::value, "Missing conversion");
 
-static_assert(std::is_same<FSharedBufferConstRef, decltype(FSharedBuffer::MakeOwned(DeclVal<const FSharedBufferConstRef&>()))>::value, "Invalid conversion");
-static_assert(std::is_same<FSharedBufferConstRef, decltype(FSharedBuffer::MakeOwned(DeclVal<FSharedBufferConstRef&&>()))>::value, "Invalid conversion");
-static_assert(std::is_same<FSharedBufferConstPtr, decltype(FSharedBuffer::MakeOwned(DeclVal<const FSharedBufferConstPtr&>()))>::value, "Invalid conversion");
-static_assert(std::is_same<FSharedBufferConstPtr, decltype(FSharedBuffer::MakeOwned(DeclVal<FSharedBufferConstPtr&&>()))>::value, "Invalid conversion");
+static_assert(!std::is_convertible<      FWeakSharedBuffer, FMutableMemoryView>::value, "Invalid conversion");
+static_assert(!std::is_convertible<const FWeakSharedBuffer, FMutableMemoryView>::value, "Invalid conversion");
+static_assert(!std::is_convertible<      FWeakSharedBuffer, FMemoryView>::value, "Invalid conversion");
+static_assert(!std::is_convertible<const FWeakSharedBuffer, FMemoryView>::value, "Invalid conversion");
 
 ///////////////////////////////////////////////////////////////////////////////////////////////////////////////////////
 
-static_assert(std::is_same<bool, decltype(DeclVal<const FSharedBufferRef&>() == DeclVal<const FSharedBufferRef&>())>::value, "Invalid equality");
-static_assert(std::is_same<bool, decltype(DeclVal<const FSharedBufferRef&>() == DeclVal<const FSharedBufferPtr&>())>::value, "Invalid equality");
-static_assert(std::is_same<bool, decltype(DeclVal<const FSharedBufferRef&>() == DeclVal<const FSharedBufferWeakPtr&>())>::value, "Invalid equality");
-static_assert(std::is_same<bool, decltype(DeclVal<const FSharedBufferRef&>() == DeclVal<const FSharedBufferConstRef&>())>::value, "Invalid equality");
-static_assert(std::is_same<bool, decltype(DeclVal<const FSharedBufferRef&>() == DeclVal<const FSharedBufferConstPtr&>())>::value, "Invalid equality");
-static_assert(std::is_same<bool, decltype(DeclVal<const FSharedBufferRef&>() == DeclVal<const FSharedBufferConstWeakPtr&>())>::value, "Invalid equality");
-
-static_assert(std::is_same<bool, decltype(DeclVal<const FSharedBufferPtr&>() == DeclVal<const FSharedBufferRef&>())>::value, "Invalid equality");
-static_assert(std::is_same<bool, decltype(DeclVal<const FSharedBufferPtr&>() == DeclVal<const FSharedBufferPtr&>())>::value, "Invalid equality");
-static_assert(std::is_same<bool, decltype(DeclVal<const FSharedBufferPtr&>() == DeclVal<const FSharedBufferWeakPtr&>())>::value, "Invalid equality");
-static_assert(std::is_same<bool, decltype(DeclVal<const FSharedBufferPtr&>() == DeclVal<const FSharedBufferConstRef&>())>::value, "Invalid equality");
-static_assert(std::is_same<bool, decltype(DeclVal<const FSharedBufferPtr&>() == DeclVal<const FSharedBufferConstPtr&>())>::value, "Invalid equality");
-static_assert(std::is_same<bool, decltype(DeclVal<const FSharedBufferPtr&>() == DeclVal<const FSharedBufferConstWeakPtr&>())>::value, "Invalid equality");
-
-static_assert(std::is_same<bool, decltype(DeclVal<const FSharedBufferWeakPtr&>() == DeclVal<const FSharedBufferRef&>())>::value, "Invalid equality");
-static_assert(std::is_same<bool, decltype(DeclVal<const FSharedBufferWeakPtr&>() == DeclVal<const FSharedBufferPtr&>())>::value, "Invalid equality");
-static_assert(std::is_same<bool, decltype(DeclVal<const FSharedBufferWeakPtr&>() == DeclVal<const FSharedBufferWeakPtr&>())>::value, "Invalid equality");
-static_assert(std::is_same<bool, decltype(DeclVal<const FSharedBufferWeakPtr&>() == DeclVal<const FSharedBufferConstRef&>())>::value, "Invalid equality");
-static_assert(std::is_same<bool, decltype(DeclVal<const FSharedBufferWeakPtr&>() == DeclVal<const FSharedBufferConstPtr&>())>::value, "Invalid equality");
-static_assert(std::is_same<bool, decltype(DeclVal<const FSharedBufferWeakPtr&>() == DeclVal<const FSharedBufferConstWeakPtr&>())>::value, "Invalid equality");
-
-static_assert(std::is_same<bool, decltype(DeclVal<const FSharedBufferConstRef&>() == DeclVal<const FSharedBufferRef&>())>::value, "Invalid equality");
-static_assert(std::is_same<bool, decltype(DeclVal<const FSharedBufferConstRef&>() == DeclVal<const FSharedBufferPtr&>())>::value, "Invalid equality");
-static_assert(std::is_same<bool, decltype(DeclVal<const FSharedBufferConstRef&>() == DeclVal<const FSharedBufferWeakPtr&>())>::value, "Invalid equality");
-static_assert(std::is_same<bool, decltype(DeclVal<const FSharedBufferConstRef&>() == DeclVal<const FSharedBufferConstRef&>())>::value, "Invalid equality");
-static_assert(std::is_same<bool, decltype(DeclVal<const FSharedBufferConstRef&>() == DeclVal<const FSharedBufferConstPtr&>())>::value, "Invalid equality");
-static_assert(std::is_same<bool, decltype(DeclVal<const FSharedBufferConstRef&>() == DeclVal<const FSharedBufferConstWeakPtr&>())>::value, "Invalid equality");
-
-static_assert(std::is_same<bool, decltype(DeclVal<const FSharedBufferConstPtr&>() == DeclVal<const FSharedBufferRef&>())>::value, "Invalid equality");
-static_assert(std::is_same<bool, decltype(DeclVal<const FSharedBufferConstPtr&>() == DeclVal<const FSharedBufferPtr&>())>::value, "Invalid equality");
-static_assert(std::is_same<bool, decltype(DeclVal<const FSharedBufferConstPtr&>() == DeclVal<const FSharedBufferWeakPtr&>())>::value, "Invalid equality");
-static_assert(std::is_same<bool, decltype(DeclVal<const FSharedBufferConstPtr&>() == DeclVal<const FSharedBufferConstRef&>())>::value, "Invalid equality");
-static_assert(std::is_same<bool, decltype(DeclVal<const FSharedBufferConstPtr&>() == DeclVal<const FSharedBufferConstPtr&>())>::value, "Invalid equality");
-static_assert(std::is_same<bool, decltype(DeclVal<const FSharedBufferConstPtr&>() == DeclVal<const FSharedBufferConstWeakPtr&>())>::value, "Invalid equality");
-
-static_assert(std::is_same<bool, decltype(DeclVal<const FSharedBufferConstWeakPtr&>() == DeclVal<const FSharedBufferRef&>())>::value, "Invalid equality");
-static_assert(std::is_same<bool, decltype(DeclVal<const FSharedBufferConstWeakPtr&>() == DeclVal<const FSharedBufferPtr&>())>::value, "Invalid equality");
-static_assert(std::is_same<bool, decltype(DeclVal<const FSharedBufferConstWeakPtr&>() == DeclVal<const FSharedBufferWeakPtr&>())>::value, "Invalid equality");
-static_assert(std::is_same<bool, decltype(DeclVal<const FSharedBufferConstWeakPtr&>() == DeclVal<const FSharedBufferConstRef&>())>::value, "Invalid equality");
-static_assert(std::is_same<bool, decltype(DeclVal<const FSharedBufferConstWeakPtr&>() == DeclVal<const FSharedBufferConstPtr&>())>::value, "Invalid equality");
-static_assert(std::is_same<bool, decltype(DeclVal<const FSharedBufferConstWeakPtr&>() == DeclVal<const FSharedBufferConstWeakPtr&>())>::value, "Invalid equality");
+static_assert(std::is_same<bool, decltype(DeclVal<const FUniqueBuffer&>() == DeclVal<const FUniqueBuffer&>())>::value, "Invalid equality");
+static_assert(std::is_same<bool, decltype(DeclVal<const FUniqueBuffer&>() == DeclVal<const FSharedBuffer&>())>::value, "Invalid equality");
+static_assert(std::is_same<bool, decltype(DeclVal<const FUniqueBuffer&>() == DeclVal<const FWeakSharedBuffer&>())>::value, "Invalid equality");
+static_assert(std::is_same<bool, decltype(DeclVal<const FSharedBuffer&>() == DeclVal<const FUniqueBuffer&>())>::value, "Invalid equality");
+static_assert(std::is_same<bool, decltype(DeclVal<const FSharedBuffer&>() == DeclVal<const FSharedBuffer&>())>::value, "Invalid equality");
+static_assert(std::is_same<bool, decltype(DeclVal<const FSharedBuffer&>() == DeclVal<const FWeakSharedBuffer&>())>::value, "Invalid equality");
+static_assert(std::is_same<bool, decltype(DeclVal<const FWeakSharedBuffer&>() == DeclVal<const FUniqueBuffer&>())>::value, "Invalid equality");
+static_assert(std::is_same<bool, decltype(DeclVal<const FWeakSharedBuffer&>() == DeclVal<const FSharedBuffer&>())>::value, "Invalid equality");
+static_assert(std::is_same<bool, decltype(DeclVal<const FWeakSharedBuffer&>() == DeclVal<const FWeakSharedBuffer&>())>::value, "Invalid equality");
 
 ///////////////////////////////////////////////////////////////////////////////////////////////////////////////////////
 
-static_assert(std::is_same<bool, decltype(DeclVal<const FSharedBufferRef&>() != DeclVal<const FSharedBufferRef&>())>::value, "Invalid inequality");
-static_assert(std::is_same<bool, decltype(DeclVal<const FSharedBufferRef&>() != DeclVal<const FSharedBufferPtr&>())>::value, "Invalid inequality");
-static_assert(std::is_same<bool, decltype(DeclVal<const FSharedBufferRef&>() != DeclVal<const FSharedBufferWeakPtr&>())>::value, "Invalid inequality");
-static_assert(std::is_same<bool, decltype(DeclVal<const FSharedBufferRef&>() != DeclVal<const FSharedBufferConstRef&>())>::value, "Invalid inequality");
-static_assert(std::is_same<bool, decltype(DeclVal<const FSharedBufferRef&>() != DeclVal<const FSharedBufferConstPtr&>())>::value, "Invalid inequality");
-static_assert(std::is_same<bool, decltype(DeclVal<const FSharedBufferRef&>() != DeclVal<const FSharedBufferConstWeakPtr&>())>::value, "Invalid inequality");
-
-static_assert(std::is_same<bool, decltype(DeclVal<const FSharedBufferPtr&>() != DeclVal<const FSharedBufferRef&>())>::value, "Invalid inequality");
-static_assert(std::is_same<bool, decltype(DeclVal<const FSharedBufferPtr&>() != DeclVal<const FSharedBufferPtr&>())>::value, "Invalid inequality");
-static_assert(std::is_same<bool, decltype(DeclVal<const FSharedBufferPtr&>() != DeclVal<const FSharedBufferWeakPtr&>())>::value, "Invalid inequality");
-static_assert(std::is_same<bool, decltype(DeclVal<const FSharedBufferPtr&>() != DeclVal<const FSharedBufferConstRef&>())>::value, "Invalid inequality");
-static_assert(std::is_same<bool, decltype(DeclVal<const FSharedBufferPtr&>() != DeclVal<const FSharedBufferConstPtr&>())>::value, "Invalid inequality");
-static_assert(std::is_same<bool, decltype(DeclVal<const FSharedBufferPtr&>() != DeclVal<const FSharedBufferConstWeakPtr&>())>::value, "Invalid inequality");
-
-static_assert(std::is_same<bool, decltype(DeclVal<const FSharedBufferWeakPtr&>() != DeclVal<const FSharedBufferRef&>())>::value, "Invalid inequality");
-static_assert(std::is_same<bool, decltype(DeclVal<const FSharedBufferWeakPtr&>() != DeclVal<const FSharedBufferPtr&>())>::value, "Invalid inequality");
-static_assert(std::is_same<bool, decltype(DeclVal<const FSharedBufferWeakPtr&>() != DeclVal<const FSharedBufferWeakPtr&>())>::value, "Invalid inequality");
-static_assert(std::is_same<bool, decltype(DeclVal<const FSharedBufferWeakPtr&>() != DeclVal<const FSharedBufferConstRef&>())>::value, "Invalid inequality");
-static_assert(std::is_same<bool, decltype(DeclVal<const FSharedBufferWeakPtr&>() != DeclVal<const FSharedBufferConstPtr&>())>::value, "Invalid inequality");
-static_assert(std::is_same<bool, decltype(DeclVal<const FSharedBufferWeakPtr&>() != DeclVal<const FSharedBufferConstWeakPtr&>())>::value, "Invalid inequality");
-
-static_assert(std::is_same<bool, decltype(DeclVal<const FSharedBufferConstRef&>() != DeclVal<const FSharedBufferRef&>())>::value, "Invalid inequality");
-static_assert(std::is_same<bool, decltype(DeclVal<const FSharedBufferConstRef&>() != DeclVal<const FSharedBufferPtr&>())>::value, "Invalid inequality");
-static_assert(std::is_same<bool, decltype(DeclVal<const FSharedBufferConstRef&>() != DeclVal<const FSharedBufferWeakPtr&>())>::value, "Invalid inequality");
-static_assert(std::is_same<bool, decltype(DeclVal<const FSharedBufferConstRef&>() != DeclVal<const FSharedBufferConstRef&>())>::value, "Invalid inequality");
-static_assert(std::is_same<bool, decltype(DeclVal<const FSharedBufferConstRef&>() != DeclVal<const FSharedBufferConstPtr&>())>::value, "Invalid inequality");
-static_assert(std::is_same<bool, decltype(DeclVal<const FSharedBufferConstRef&>() != DeclVal<const FSharedBufferConstWeakPtr&>())>::value, "Invalid inequality");
-
-static_assert(std::is_same<bool, decltype(DeclVal<const FSharedBufferConstPtr&>() != DeclVal<const FSharedBufferRef&>())>::value, "Invalid inequality");
-static_assert(std::is_same<bool, decltype(DeclVal<const FSharedBufferConstPtr&>() != DeclVal<const FSharedBufferPtr&>())>::value, "Invalid inequality");
-static_assert(std::is_same<bool, decltype(DeclVal<const FSharedBufferConstPtr&>() != DeclVal<const FSharedBufferWeakPtr&>())>::value, "Invalid inequality");
-static_assert(std::is_same<bool, decltype(DeclVal<const FSharedBufferConstPtr&>() != DeclVal<const FSharedBufferConstRef&>())>::value, "Invalid inequality");
-static_assert(std::is_same<bool, decltype(DeclVal<const FSharedBufferConstPtr&>() != DeclVal<const FSharedBufferConstPtr&>())>::value, "Invalid inequality");
-static_assert(std::is_same<bool, decltype(DeclVal<const FSharedBufferConstPtr&>() != DeclVal<const FSharedBufferConstWeakPtr&>())>::value, "Invalid inequality");
-
-static_assert(std::is_same<bool, decltype(DeclVal<const FSharedBufferConstWeakPtr&>() != DeclVal<const FSharedBufferRef&>())>::value, "Invalid inequality");
-static_assert(std::is_same<bool, decltype(DeclVal<const FSharedBufferConstWeakPtr&>() != DeclVal<const FSharedBufferPtr&>())>::value, "Invalid inequality");
-static_assert(std::is_same<bool, decltype(DeclVal<const FSharedBufferConstWeakPtr&>() != DeclVal<const FSharedBufferWeakPtr&>())>::value, "Invalid inequality");
-static_assert(std::is_same<bool, decltype(DeclVal<const FSharedBufferConstWeakPtr&>() != DeclVal<const FSharedBufferConstRef&>())>::value, "Invalid inequality");
-static_assert(std::is_same<bool, decltype(DeclVal<const FSharedBufferConstWeakPtr&>() != DeclVal<const FSharedBufferConstPtr&>())>::value, "Invalid inequality");
-static_assert(std::is_same<bool, decltype(DeclVal<const FSharedBufferConstWeakPtr&>() != DeclVal<const FSharedBufferConstWeakPtr&>())>::value, "Invalid inequality");
+static_assert(std::is_same<bool, decltype(DeclVal<const FUniqueBuffer&>() != DeclVal<const FUniqueBuffer&>())>::value, "Invalid inequality");
+static_assert(std::is_same<bool, decltype(DeclVal<const FUniqueBuffer&>() != DeclVal<const FSharedBuffer&>())>::value, "Invalid inequality");
+static_assert(std::is_same<bool, decltype(DeclVal<const FUniqueBuffer&>() != DeclVal<const FWeakSharedBuffer&>())>::value, "Invalid inequality");
+static_assert(std::is_same<bool, decltype(DeclVal<const FSharedBuffer&>() != DeclVal<const FUniqueBuffer&>())>::value, "Invalid inequality");
+static_assert(std::is_same<bool, decltype(DeclVal<const FSharedBuffer&>() != DeclVal<const FSharedBuffer&>())>::value, "Invalid inequality");
+static_assert(std::is_same<bool, decltype(DeclVal<const FSharedBuffer&>() != DeclVal<const FWeakSharedBuffer&>())>::value, "Invalid inequality");
+static_assert(std::is_same<bool, decltype(DeclVal<const FWeakSharedBuffer&>() != DeclVal<const FUniqueBuffer&>())>::value, "Invalid inequality");
+static_assert(std::is_same<bool, decltype(DeclVal<const FWeakSharedBuffer&>() != DeclVal<const FSharedBuffer&>())>::value, "Invalid inequality");
+static_assert(std::is_same<bool, decltype(DeclVal<const FWeakSharedBuffer&>() != DeclVal<const FWeakSharedBuffer&>())>::value, "Invalid inequality");
 
 ///////////////////////////////////////////////////////////////////////////////////////////////////////////////////////
 
-static_assert(std::is_same<uint32, decltype(GetTypeHash(DeclVal<const FSharedBufferRef>()))>::value, "Missing or invalid hash function");
-static_assert(std::is_same<uint32, decltype(GetTypeHash(DeclVal<const FSharedBufferConstRef>()))>::value, "Missing or invalid hash function");
-static_assert(std::is_same<uint32, decltype(GetTypeHash(DeclVal<const FSharedBufferPtr>()))>::value, "Missing or invalid hash function");
-static_assert(std::is_same<uint32, decltype(GetTypeHash(DeclVal<const FSharedBufferConstPtr>()))>::value, "Missing or invalid hash function");
-static_assert(std::is_same<uint32, decltype(GetTypeHash(DeclVal<const FSharedBufferWeakPtr>()))>::value, "Missing or invalid hash function");
-static_assert(std::is_same<uint32, decltype(GetTypeHash(DeclVal<const FSharedBufferConstWeakPtr>()))>::value, "Missing or invalid hash function");
+static_assert(std::is_same<uint32, decltype(GetTypeHash(DeclVal<const FUniqueBuffer>()))>::value, "Missing or invalid hash function");
+static_assert(std::is_same<uint32, decltype(GetTypeHash(DeclVal<const FSharedBuffer>()))>::value, "Missing or invalid hash function");
+static_assert(std::is_same<uint32, decltype(GetTypeHash(DeclVal<const FWeakSharedBuffer>()))>::value, "Missing or invalid hash function");
 
 ///////////////////////////////////////////////////////////////////////////////////////////////////////////////////////
 
 #if WITH_DEV_AUTOMATION_TESTS
 
-IMPLEMENT_SIMPLE_AUTOMATION_TEST(FSharedBufferTest, "System.Core.Memory.SharedBuffer", EAutomationTestFlags::ApplicationContextMask | EAutomationTestFlags::SmokeFilter)
-bool FSharedBufferTest::RunTest(const FString& Parameters)
+IMPLEMENT_SIMPLE_AUTOMATION_TEST(FUniqueBufferTest, "System.Core.Memory.UniqueBuffer", EAutomationTestFlags::ApplicationContextMask | EAutomationTestFlags::SmokeFilter)
+bool FUniqueBufferTest::RunTest(const FString& Parameters)
 {
-	// Test Size
+	// Test Alloc
 	{
 		constexpr uint64 Size = 64;
-		FSharedBufferRef Ref = FSharedBuffer::Alloc(64).ToSharedRef();
-		TestTrue(TEXT("FSharedBuffer(Size).IsOwned()"), Ref->IsOwned());
-		TestEqual(TEXT("FSharedBuffer(Size).GetSize()"), Ref->GetSize(), Size);
+		FUniqueBuffer Buffer = FUniqueBuffer::Alloc(Size);
+		TestTrue(TEXT("FUniqueBuffer::Alloc().IsOwned()"), Buffer.IsOwned());
+		TestEqual(TEXT("FUniqueBuffer::Alloc().GetSize()"), Buffer.GetSize(), Size);
+		TestEqual(TEXT("FUniqueBuffer::Alloc().GetView().GetSize()"), Buffer.GetView().GetSize(), Size);
+		TestEqual(TEXT("FUniqueBuffer::Alloc().GetView().GetData()"), Buffer.GetView().GetData(), Buffer.GetData());
 	}
 
 	// Test Clone
 	{
 		constexpr uint64 Size = 64;
 		const uint8 Data[Size]{};
-		FSharedBufferRef Ref = FSharedBuffer::Clone(Data, Size).ToSharedRef();
-		TestTrue(TEXT("FSharedBuffer::Clone().IsOwned()"), Ref->IsOwned());
-		TestEqual(TEXT("FSharedBuffer::Clone().GetSize()"), Ref->GetSize(), Size);
-		TestNotEqual(TEXT("FSharedBuffer::Clone().GetData()"), static_cast<const void*>(Ref->GetData()), static_cast<const void*>(Data));
+		FUniqueBuffer Buffer = FUniqueBuffer::Clone(Data, Size);
+		TestTrue(TEXT("FUniqueBuffer::Clone().IsOwned()"), Buffer.IsOwned());
+		TestEqual(TEXT("FUniqueBuffer::Clone().GetSize()"), Buffer.GetSize(), Size);
+		TestNotEqual(TEXT("FUniqueBuffer::Clone().GetData()"), static_cast<const void*>(Buffer.GetData()), static_cast<const void*>(Data));
 	}
 
 	// Test MakeView
 	{
 		constexpr uint64 Size = 64;
 		uint8 Data[Size]{};
-		FSharedBufferRef Ref = FSharedBuffer::MakeView(Data, Size).ToSharedRef();
-		TestFalse(TEXT("FSharedBuffer::MakeView().IsOwned()"), Ref->IsOwned());
-		TestEqual(TEXT("FSharedBuffer::MakeView().GetSize()"), Ref->GetSize(), Size);
-		TestEqual(TEXT("FSharedBuffer::MakeView().GetData()"), Ref->GetData(), static_cast<void*>(Data));
+		FUniqueBuffer Buffer = FUniqueBuffer::MakeView(Data, Size);
+		TestFalse(TEXT("FUniqueBuffer::MakeView().IsOwned()"), Buffer.IsOwned());
+		TestEqual(TEXT("FUniqueBuffer::MakeView().GetSize()"), Buffer.GetSize(), Size);
+		TestEqual(TEXT("FUniqueBuffer::MakeView().GetData()"), Buffer.GetData(), static_cast<void*>(Data));
 	}
 
-	// Test MakeView Outer
+	// Test TakeOwnership
 	{
+		bool bDeleted = false;
 		constexpr uint64 Size = 64;
-		FSharedBufferWeakPtr WeakPtr;
+		void* Data = new uint8[Size];
+		FUniqueBuffer Buffer = FUniqueBuffer::TakeOwnership(Data, Size,
+			[&bDeleted](void* InData, uint64 InSize)
 		{
-			FSharedBufferConstPtr Ptr;
-			{
-				FSharedBufferRef OuterRef = FSharedBuffer::Alloc(Size).ToSharedRef();
-				Ptr = FSharedBuffer::MakeView(OuterRef->GetData(), OuterRef->GetSize() / 2, *OuterRef);
-				WeakPtr = OuterRef;
-			}
-			TestFalse(TEXT("FSharedBuffer::MakeView(Outer).IsOwned()"), Ptr->IsOwned());
-			TestFalse(TEXT("FSharedBuffer::MakeView(Outer).IsReadOnly()"), Ptr->IsReadOnly());
-			TestTrue(TEXT("FSharedBuffer::MakeView(Outer) Outer Valid"), WeakPtr.Pin().IsValid());
-		}
-		TestFalse(TEXT("FSharedBuffer::MakeView(Outer) Outer Not Valid"), WeakPtr.Pin().IsValid());
+			bDeleted = true;
+			delete[] static_cast<uint8*>(InData);
+		});
+		TestTrue(TEXT("FUniqueBuffer::TakeOwnership().IsOwned()"), Buffer.IsOwned());
+		TestEqual(TEXT("FUniqueBuffer::TakeOwnership().GetSize()"), Buffer.GetSize(), Size);
+		TestEqual(TEXT("FUniqueBuffer::TakeOwnership().GetData()"), Buffer.GetData(), Data);
+		Buffer.Reset();
+		TestTrue(TEXT("FUniqueBuffer::TakeOwnership() Deleted"), bDeleted);
 	}
 
-	// Test MakeView Outer Read-Only
+	// Test MakeOwned
 	{
 		constexpr uint64 Size = 64;
-		FSharedBufferConstPtr Ptr;
-		{
-			FSharedBufferConstRef OuterRef = FSharedBuffer::MakeReadOnly(FSharedBuffer::Alloc(Size)).ToSharedRef();
-			Ptr = FSharedBuffer::MakeView(OuterRef->GetData(), OuterRef->GetSize() / 2, *OuterRef);
-		}
-		TestFalse(TEXT("FSharedBuffer::MakeView(Outer ReadOnly).IsOwned()"), Ptr->IsOwned());
-		TestTrue(TEXT("FSharedBuffer::MakeView(Outer ReadOnly).IsReadOnly()"), Ptr->IsReadOnly());
+		uint8 Data[Size]{};
+		FUniqueBuffer Buffer = FUniqueBuffer::MakeView(Data, Size);
+		Buffer.MakeOwned();
+		TestTrue(TEXT("FUniqueBuffer::MakeOwned(View).IsOwned()"), Buffer.IsOwned());
+		TestEqual(TEXT("FUniqueBuffer::MakeOwned(View).GetSize()"), Buffer.GetSize(), Size);
+		TestFalse(TEXT("FUniqueBuffer::MakeOwned(View).GetData()"), Buffer.GetData() == Data);
+		void* const OwnedData = Buffer.GetData();
+		Buffer.MakeOwned();
+		TestTrue(TEXT("FUniqueBuffer::MakeOwned(Owned).IsOwned()"), Buffer.IsOwned());
+		TestEqual(TEXT("FUniqueBuffer::MakeOwned(Owned).GetSize()"), Buffer.GetSize(), Size);
+		TestEqual(TEXT("FUniqueBuffer::MakeOwned(Owned).GetData()"), Buffer.GetData(), OwnedData);
 	}
 
-	// Test TakeOwnership FMemory::Free
+	return true;
+}
+
+IMPLEMENT_SIMPLE_AUTOMATION_TEST(FSharedBuffer2Test, "System.Core.Memory.SharedBuffer2", EAutomationTestFlags::ApplicationContextMask | EAutomationTestFlags::SmokeFilter)
+bool FSharedBuffer2Test::RunTest(const FString& Parameters)
+{
+	// Test Clone
 	{
 		constexpr uint64 Size = 64;
-		void* const Data = FMemory::Malloc(Size);
-		FSharedBufferRef Ref = FSharedBuffer::TakeOwnership(Data, Size, FMemory::Free).ToSharedRef();
-		TestTrue(TEXT("FSharedBuffer::TakeOwnership(FMemory::Free).IsOwned()"), Ref->IsOwned());
-		TestEqual(TEXT("FSharedBuffer::TakeOwnership(FMemory::Free).GetSize()"), Ref->GetSize(), Size);
-		TestEqual(TEXT("FSharedBuffer::TakeOwnership(FMemory::Free).GetData()"), Ref->GetData(), Data);
+		const uint8 Data[Size]{};
+		FSharedBuffer Buffer = FSharedBuffer::Clone(Data, Size);
+		TestTrue(TEXT("FSharedBuffer::Clone().IsOwned()"), Buffer.IsOwned());
+		TestEqual(TEXT("FSharedBuffer::Clone().GetSize()"), Buffer.GetSize(), Size);
+		TestNotEqual(TEXT("FSharedBuffer::Clone().GetData()"), Buffer.GetData(), static_cast<const void*>(Data));
 	}
 
-	// Test TakeOwnership Lambda
+	// Test MakeView
 	{
-		bool bDeleted = false;
-		auto Deleter = [&bDeleted](void* Data)
-		{
-			bDeleted = true;
-			delete[] static_cast<uint8*>(Data);
-		};
 		constexpr uint64 Size = 64;
-		FSharedBuffer::TakeOwnership(new uint8[Size], Size, Deleter);
-		TestTrue(TEXT("FSharedBuffer::TakeOwnership(Lambda) Deleted"), bDeleted);
+		uint8 Data[Size]{};
+		FSharedBuffer Buffer = FSharedBuffer::MakeView(Data, Size);
+		TestFalse(TEXT("FSharedBuffer::MakeView().IsOwned()"), Buffer.IsOwned());
+		TestEqual(TEXT("FSharedBuffer::MakeView().GetSize()"), Buffer.GetSize(), Size);
+		TestEqual(TEXT("FSharedBuffer::MakeView().GetData()"), Buffer.GetData(), static_cast<const void*>(Data));
 	}
 
-	// Test TakeOwnership BufferOwner
+	// Test MakeView Outer View
 	{
-		bool bDeleted = false;
-		class FTestBufferOwner final : public FSharedBuffer::FBufferOwner
+		constexpr uint64 Size = 64;
+		const uint8 Data[Size]{};
+		FWeakSharedBuffer WeakBuffer;
 		{
-		public:
-			FTestBufferOwner(bool& bInDeletedRef)
-				: bDeletedRef(bInDeletedRef)
-			{
-			}
-
-		private:
-			virtual void Free(void* Data, uint64 Size) final
+			FSharedBuffer Buffer;
 			{
-				bDeletedRef = true;
+				FSharedBuffer OuterBuffer = FSharedBuffer::MakeView(Data, Size);
+				Buffer = FSharedBuffer::MakeView(OuterBuffer.GetData(), OuterBuffer.GetSize() / 2, OuterBuffer);
+				WeakBuffer = OuterBuffer;
 			}
-
-			bool& bDeletedRef;
-		};
-		constexpr uint64 Size = 64;
-		FSharedBuffer::TakeOwnership(nullptr, 0, FTestBufferOwner(bDeleted));
-		TestTrue(TEXT("FSharedBuffer::TakeOwnership(BufferOwner) Deleted"), bDeleted);
+			TestFalse(TEXT("FSharedBuffer::MakeView(OuterView).IsOwned()"), Buffer.IsOwned());
+			TestEqual(TEXT("FSharedBuffer::MakeView(OuterView).GetSize()"), Buffer.GetSize(), Size / 2);
+			TestEqual(TEXT("FSharedBuffer::MakeView(OuterView).GetData()"), Buffer.GetData(), static_cast<const void*>(Data));
+			TestFalse(TEXT("FSharedBuffer::MakeView(OuterView) Outer Not Null"), WeakBuffer.Pin().IsNull());
+		}
+		TestTrue(TEXT("FSharedBuffer::MakeView(OuterView) Outer Null"), WeakBuffer.Pin().IsNull());
 	}
 
-	// Test MakeOwned
+	// Test MakeView Outer Owned
 	{
 		constexpr uint64 Size = 64;
-		uint8 Data[Size]{};
-		FSharedBufferRef Ref = FSharedBuffer::MakeOwned(FSharedBuffer::MakeView(Data, Size).ToSharedRef());
-		TestTrue(TEXT("FSharedBuffer::MakeOwned(Wrap).IsOwned()"), Ref->IsOwned());
-		TestEqual(TEXT("FSharedBuffer::MakeOwned(Wrap).GetSize()"), Ref->GetSize(), Size);
-		TestNotEqual(TEXT("FSharedBuffer::MakeOwned(Wrap).GetData()"), Ref->GetData(), static_cast<void*>(Data));
-		FSharedBufferRef RefCopy = FSharedBuffer::MakeOwned(Ref);
-		TestTrue(TEXT("FSharedBuffer::MakeOwned(Owned).IsOwned()"), RefCopy->IsOwned());
-		TestEqual(TEXT("FSharedBuffer::MakeOwned(Owned).GetSize()"), RefCopy->GetSize(), Size);
-		TestEqual(TEXT("FSharedBuffer::MakeOwned(Owned).GetData()"), RefCopy->GetData(), Ref->GetData());
-	}
-	{
-		FSharedBufferPtr Ptr;
-		TestFalse(TEXT("FSharedBuffer::MakeOwned(Null)"), FSharedBuffer::MakeOwned(Ptr).IsValid());
-		TestFalse(TEXT("FSharedBuffer::MakeOwned(Null)"), FSharedBuffer::MakeOwned(AsConst(Ptr)).IsValid());
-		TestFalse(TEXT("FSharedBuffer::MakeOwned(Null)"), FSharedBuffer::MakeOwned(FSharedBufferPtr()).IsValid());
-		TestFalse(TEXT("FSharedBuffer::MakeOwned(Null)"), FSharedBuffer::MakeOwned(FSharedBufferConstPtr()).IsValid());
+		FWeakSharedBuffer WeakBuffer;
+		{
+			FSharedBuffer Buffer;
+			{
+				FSharedBuffer OuterBuffer = FSharedBuffer(FUniqueBuffer::Alloc(Size));
+				Buffer = FSharedBuffer::MakeView(OuterBuffer.GetData(), OuterBuffer.GetSize(), OuterBuffer);
+				TestEqual(TEXT("FSharedBuffer::MakeView(OuterOwned, SameView)"), Buffer, OuterBuffer);
+				Buffer = FSharedBuffer::MakeView(OuterBuffer.GetData(), OuterBuffer.GetSize() / 2, OuterBuffer);
+				WeakBuffer = OuterBuffer;
+			}
+			TestTrue(TEXT("FSharedBuffer::MakeView(OuterOwned).IsOwned()"), Buffer.IsOwned());
+			TestEqual(TEXT("FSharedBuffer::MakeView(OuterOwned).GetSize()"), Buffer.GetSize(), Size / 2);
+			TestEqual(TEXT("FSharedBuffer::MakeView(OuterOwned).GetData()"), Buffer.GetData(), WeakBuffer.Pin().GetData());
+			TestFalse(TEXT("FSharedBuffer::MakeView(OuterOwned) Outer Not Null"), WeakBuffer.Pin().IsNull());
+		}
+		TestTrue(TEXT("FSharedBuffer::MakeView(OuterOwned) Outer  Null"), WeakBuffer.Pin().IsNull());
 	}
 
-	// Test MakeReadOnly
-	{
-		// MakeReadOnly from a new reference.
-		constexpr uint64 Size = 64;
-		FSharedBufferConstRef Ref = FSharedBuffer::MakeReadOnly(FSharedBuffer::Alloc(Size)).ToSharedRef();
-		TestTrue(TEXT("FSharedBuffer::MakeReadOnly(AllocRef).IsOwned()"), Ref->IsOwned());
-		TestTrue(TEXT("FSharedBuffer::MakeReadOnly(AllocRef).IsReadOnly()"), Ref->IsReadOnly());
-		TestEqual(TEXT("FSharedBuffer::MakeReadOnly(AllocRef).GetSize()"), Ref->GetSize(), Size);
-		FSharedBufferConstRef OtherRef = Ref;
-		FSharedBufferConstRef ReadOnlyRef = FSharedBuffer::MakeReadOnly(MoveTemp(Ref));
-		TestEqual(TEXT("FSharedBuffer::MakeReadOnly(AllocRef).GetData()"), Ref->GetData(), ReadOnlyRef->GetData());
-	}
-	{
-		// MakeReadOnly from a new pointer.
-		constexpr uint64 Size = 64;
-		FSharedBufferPtr MutablePtr = FSharedBuffer::Alloc(Size);
-		TestFalse(TEXT("FSharedBuffer::MakeReadOnly(MovePtr).IsReadOnly()"), MutablePtr->IsReadOnly());
-		TestTrue(TEXT("FSharedBuffer::MakeReadOnly(MovePtr).IsOwned()"), MutablePtr->IsOwned());
-		const void* const Data = MutablePtr->GetData();
-		FSharedBufferConstPtr Ptr = FSharedBuffer::MakeReadOnly(MoveTemp(MutablePtr));
-		TestTrue(TEXT("FSharedBuffer::MakeReadOnly(MovePtr).IsOwned()"), Ptr->IsOwned());
-		TestTrue(TEXT("FSharedBuffer::MakeReadOnly(MovePtr).IsReadOnly()"), Ptr->IsReadOnly());
-		TestEqual(TEXT("FSharedBuffer::MakeReadOnly(MovePtr).GetSize()"), Ptr->GetSize(), Size);
-		TestEqual(TEXT("FSharedBuffer::MakeReadOnly(MovePtr).GetData()"), Ptr->GetData(), Data);
-	}
-	{
-		// MakeReadOnly from a view.
-		constexpr uint64 Size = 64;
-		uint8 Data[Size]{};
-		FSharedBufferPtr MutablePtr = FSharedBuffer::MakeView(Data, Size);
-		TestFalse(TEXT("FSharedBuffer::MakeReadOnly(View).IsOwned()"), MutablePtr->IsOwned());
-		TestFalse(TEXT("FSharedBuffer::MakeReadOnly(View).IsReadOnly()"), MutablePtr->IsReadOnly());
-		FSharedBufferConstPtr Ptr = FSharedBuffer::MakeReadOnly(MoveTemp(MutablePtr));
-		TestTrue(TEXT("FSharedBuffer::MakeReadOnly(View).IsOwned()"), Ptr->IsOwned());
-		TestTrue(TEXT("FSharedBuffer::MakeReadOnly(View).IsReadOnly()"), Ptr->IsReadOnly());
-		TestEqual(TEXT("FSharedBuffer::MakeReadOnly(View).GetSize()"), Ptr->GetSize(), Size);
-		TestNotEqual(TEXT("FSharedBuffer::MakeReadOnly(View).GetData()"), Ptr->GetData(), static_cast<const void*>(Data));
-	}
-	{
-		// MakeReadOnly with another shared reference.
-		constexpr uint64 Size = 64;
-		FSharedBufferPtr MutablePtr = FSharedBuffer::Alloc(Size);
-		FSharedBufferConstPtr SharedPtr = MutablePtr;
-		TestFalse(TEXT("FSharedBuffer::MakeReadOnly(SharedRef).IsReadOnly()"), MutablePtr->IsReadOnly());
-		const void* const Data = MutablePtr->GetData();
-		FSharedBufferConstPtr Ptr = FSharedBuffer::MakeReadOnly(MoveTemp(MutablePtr));
-		TestTrue(TEXT("FSharedBuffer::MakeReadOnly(SharedRef).IsOwned()"), Ptr->IsOwned());
-		TestTrue(TEXT("FSharedBuffer::MakeReadOnly(SharedRef).IsReadOnly()"), Ptr->IsReadOnly());
-		TestFalse(TEXT("FSharedBuffer::MakeReadOnly(SharedRef).IsReadOnly(SharedPtr)"), SharedPtr->IsReadOnly());
-		TestNotEqual(TEXT("FSharedBuffer::MakeReadOnly(SharedRef).GetData()"), Ptr->GetData(), Data);
-	}
+	// Test TakeOwnership
 	{
-		// MakeReadOnly with another weak reference.
+		bool bDeleted = false;
 		constexpr uint64 Size = 64;
-		FSharedBufferPtr MutablePtr = FSharedBuffer::Alloc(Size);
-		FSharedBufferConstWeakPtr WeakPtr = MutablePtr;
-		TestFalse(TEXT("FSharedBuffer::MakeReadOnly(WeakRef).IsReadOnly()"), MutablePtr->IsReadOnly());
-		const void* const Data = MutablePtr->GetData();
-		FSharedBufferConstPtr Ptr = FSharedBuffer::MakeReadOnly(MoveTemp(MutablePtr));
-		TestTrue(TEXT("FSharedBuffer::MakeReadOnly(WeakRef).IsOwned()"), Ptr->IsOwned());
-		TestTrue(TEXT("FSharedBuffer::MakeReadOnly(WeakRef).IsReadOnly()"), Ptr->IsReadOnly());
-		TestNotEqual(TEXT("FSharedBuffer::MakeReadOnly(WeakRef).GetData()"), Ptr->GetData(), Data);
-		TestFalse(TEXT("FSharedBuffer::MakeReadOnly(WeakRef).Pin().IsValid()"), WeakPtr.Pin().IsValid());
+		const void* Data = new uint8[Size];
+		FSharedBuffer Buffer = FSharedBuffer::TakeOwnership(Data, Size,
+			[&bDeleted](void* InData, uint64 InSize)
+		{
+			bDeleted = true;
+			delete[] static_cast<uint8*>(InData);
+		});
+		TestTrue(TEXT("FSharedBuffer::TakeOwnership().IsOwned()"), Buffer.IsOwned());
+		TestEqual(TEXT("FSharedBuffer::TakeOwnership().GetSize()"), Buffer.GetSize(), Size);
+		TestEqual(TEXT("FSharedBuffer::TakeOwnership().GetData()"), Buffer.GetData(), Data);
+		Buffer.Reset();
+		TestTrue(TEXT("FSharedBuffer::TakeOwnership() Deleted"), bDeleted);
 	}
 
-	// Test MakeWritable
-	{
-		// MakeWritable from a new reference.
-		constexpr uint64 Size = 64;
-		FSharedBufferConstRef ConstRef = FSharedBuffer::MakeReadOnly(FSharedBuffer::Alloc(Size)).ToSharedRef();
-		FSharedBufferRef Ref = FSharedBuffer::MakeWritable(MoveTemp(ConstRef));
-		TestTrue(TEXT("FSharedBuffer::MakeWritable(AllocRef).IsOwned()"), Ref->IsOwned());
-		TestFalse(TEXT("FSharedBuffer::MakeWritable(AllocRef).IsReadOnly()"), Ref->IsReadOnly());
-		TestEqual(TEXT("FSharedBuffer::MakeWritable(AllocRef).GetSize()"), Ref->GetSize(), Size);
-		TestEqual(TEXT("FSharedBuffer::MakeWritable(AllocRef).GetData()"), const_cast<const void*>(Ref->GetData()), ConstRef->GetData());
-		FSharedBufferRef MutableRef = FSharedBuffer::MakeWritable(MoveTemp(Ref));
-		TestEqual(TEXT("FSharedBuffer::MakeWritable(AllocRef).GetData()"), Ref->GetData(), MutableRef->GetData());
-	}
-	{
-		// MakeWritable from a new pointer.
-		constexpr uint64 Size = 64;
-		FSharedBufferConstPtr ConstPtr = FSharedBuffer::MakeReadOnly(FSharedBuffer::Alloc(Size));
-		const void* const Data = ConstPtr->GetData();
-		FSharedBufferPtr Ptr = FSharedBuffer::MakeWritable(MoveTemp(ConstPtr));
-		TestFalse(TEXT("FSharedBuffer::MakeReadOnly(MovePtr).IsReadOnly()"), Ptr->IsReadOnly());
-		TestTrue(TEXT("FSharedBuffer::MakeReadOnly(MovePtr).IsOwned()"), Ptr->IsOwned());
-		TestEqual(TEXT("FSharedBuffer::MakeReadOnly(MovePtr).GetSize()"), Ptr->GetSize(), Size);
-		TestEqual(TEXT("FSharedBuffer::MakeReadOnly(MovePtr).GetData()"), const_cast<const void*>(Ptr->GetData()), Data);
-	}
+	// Test MakeOwned
 	{
-		// MakeWritable from a mutable view.
 		constexpr uint64 Size = 64;
 		uint8 Data[Size]{};
-		FSharedBufferPtr Ptr = FSharedBuffer::MakeWritable(FSharedBuffer::MakeView(Data, Size));
-		TestFalse(TEXT("FSharedBuffer::MakeWritable(View).IsOwned()"), Ptr->IsOwned());
-		TestFalse(TEXT("FSharedBuffer::MakeWritable(View).IsReadOnly()"), Ptr->IsReadOnly());
-		TestEqual(TEXT("FSharedBuffer::MakeWritable(View).GetSize()"), Ptr->GetSize(), Size);
-		TestEqual(TEXT("FSharedBuffer::MakeWritable(View).GetData()"), const_cast<const void*>(Ptr->GetData()), static_cast<const void*>(Data));
-	}
-	{
-		// MakeWritable from a const view.
-		constexpr uint64 Size = 64;
-		const uint8 Data[Size]{};
-		FSharedBufferPtr Ptr = FSharedBuffer::MakeWritable(FSharedBuffer::MakeView(Data, Size));
-		TestTrue(TEXT("FSharedBuffer::MakeWritable(View).IsOwned()"), Ptr->IsOwned());
-		TestFalse(TEXT("FSharedBuffer::MakeWritable(View).IsReadOnly()"), Ptr->IsReadOnly());
-		TestEqual(TEXT("FSharedBuffer::MakeWritable(View).GetSize()"), Ptr->GetSize(), Size);
-		TestNotEqual(TEXT("FSharedBuffer::MakeWritable(View).GetData()"), const_cast<const void*>(Ptr->GetData()), static_cast<const void*>(Data));
-	}
-	{
-		// MakeWritable with another shared reference.
-		constexpr uint64 Size = 64;
-		FSharedBufferConstPtr ConstPtr = FSharedBuffer::MakeReadOnly(FSharedBuffer::Alloc(Size));
-		FSharedBufferConstPtr SharedPtr = ConstPtr;
-		const void* const Data = ConstPtr->GetData();
-		FSharedBufferPtr Ptr = FSharedBuffer::MakeWritable(MoveTemp(ConstPtr));
-		TestTrue(TEXT("FSharedBuffer::MakeWritable(SharedRef).IsOwned()"), Ptr->IsOwned());
-		TestFalse(TEXT("FSharedBuffer::MakeWritable(SharedRef).IsReadOnly()"), Ptr->IsReadOnly());
-		TestNotEqual(TEXT("FSharedBuffer::MakeWritable(SharedRef).GetData()"), const_cast<const void*>(Ptr->GetData()), Data);
-	}
-	{
-		// MakeWritable with another weak reference.
-		constexpr uint64 Size = 64;
-		FSharedBufferConstPtr ConstPtr = FSharedBuffer::MakeReadOnly(FSharedBuffer::Alloc(Size));
-		FSharedBufferConstWeakPtr WeakPtr = ConstPtr;
-		const void* const Data = ConstPtr->GetData();
-		FSharedBufferPtr Ptr = FSharedBuffer::MakeWritable(MoveTemp(ConstPtr));
-		TestTrue(TEXT("FSharedBuffer::MakeWritable(WeakRef).IsOwned()"), Ptr->IsOwned());
-		TestFalse(TEXT("FSharedBuffer::MakeWritable(WeakRef).IsReadOnly()"), Ptr->IsReadOnly());
-		TestNotEqual(TEXT("FSharedBuffer::MakeWritable(WeakRef).GetData()"), const_cast<const void*>(Ptr->GetData()), Data);
-		TestTrue(TEXT("FSharedBuffer::MakeWritable(WeakRef).Pin().IsValid()"), WeakPtr.Pin().IsValid());
+		FSharedBuffer Buffer = FSharedBuffer::MakeView(Data, Size);
+		Buffer.MakeOwned();
+		TestTrue(TEXT("FSharedBuffer::MakeOwned(View).IsOwned()"), Buffer.IsOwned());
+		TestEqual(TEXT("FSharedBuffer::MakeOwned(View).GetSize()"), Buffer.GetSize(), Size);
+		TestFalse(TEXT("FSharedBuffer::MakeOwned(View).GetData()"), Buffer.GetData() == Data);
+		FSharedBuffer BufferCopy = Buffer;
+		BufferCopy.MakeOwned();
+		TestTrue(TEXT("FSharedBuffer::MakeOwned(Owned).IsOwned()"), BufferCopy.IsOwned());
+		TestEqual(TEXT("FSharedBuffer::MakeOwned(Owned).GetSize()"), BufferCopy.GetSize(), Size);
+		TestEqual(TEXT("FSharedBuffer::MakeOwned(Owned).GetData()"), BufferCopy.GetData(), Buffer.GetData());
 	}
 
-	// Test WeakPtr
+	return true;
+}
+
+IMPLEMENT_SIMPLE_AUTOMATION_TEST(FBufferOwnerTest, "System.Core.Memory.BufferOwner", EAutomationTestFlags::ApplicationContextMask | EAutomationTestFlags::SmokeFilter)
+bool FBufferOwnerTest::RunTest(const FString& Parameters)
+{
+	class FTestBufferOwner final : public FBufferOwner
 	{
-		FSharedBufferWeakPtr WeakPtr;
+	public:
+		FTestBufferOwner(bool& bInMaterializedRef, bool& bInFreedRef, bool& bInDeletedRef)
+			: bMaterializedRef(bInMaterializedRef)
+			, bFreedRef(bInFreedRef)
+			, bDeletedRef(bInDeletedRef)
 		{
-			FSharedBufferRef Ref = FSharedBuffer::Alloc(0).ToSharedRef();
-			WeakPtr = Ref;
-			TestTrue(TEXT("FSharedBufferWeakPtr(Ref).Pin().IsValid()"), WeakPtr.Pin().IsValid());
+			SetIsOwned();
 		}
-		TestFalse(TEXT("FSharedBufferWeakPtr(Ref).Pin().IsValid()"), WeakPtr.Pin().IsValid());
-	}
+
+	private:
+		virtual void MaterializeBuffer()
+		{
+			SetIsMaterialized();
+			bMaterializedRef = true;
+		}
+
+		virtual void FreeBuffer() final
+		{
+			bFreedRef = true;
+		}
+
+		virtual ~FTestBufferOwner()
+		{
+			bDeletedRef = true;
+		}
+
+		bool& bMaterializedRef;
+		bool& bFreedRef;
+		bool& bDeletedRef;
+	};
+
+	bool bMaterialized = false;
+	bool bFreed = false;
+	bool bDeleted = false;
+
+	FSharedBuffer Buffer(new FTestBufferOwner(bMaterialized, bFreed, bDeleted));
+	TestFalse(TEXT("FSharedBuffer(BufferOwner)->!Materialized"), bMaterialized);
+	TestFalse(TEXT("FSharedBuffer(BufferOwner)->!Freed"), bFreed);
+	TestFalse(TEXT("FSharedBuffer(BufferOwner)->!Deleted"), bDeleted);
+
+	Buffer.GetData();
+	TestTrue(TEXT("FSharedBuffer(BufferOwner)->GetData()->Materialized"), bMaterialized);
+	bMaterialized = false;
+	Buffer.GetData();
+	TestFalse(TEXT("FSharedBuffer(BufferOwner)->GetData()->!Materialized"), bMaterialized);
+	bMaterialized = false;
+
+	Buffer = FUniqueBuffer(new FTestBufferOwner(bMaterialized, bFreed, bDeleted));
+	TestFalse(TEXT("FSharedBuffer(BufferOwner)->Assign(FUniqueBuffer)->!Materialized"), bMaterialized);
+	TestTrue(TEXT("FSharedBuffer(BufferOwner)->Assign(FUniqueBuffer)->Freed"), bFreed);
+	TestTrue(TEXT("FSharedBuffer(BufferOwner)->Assign(FUniqueBuffer)->Deleted"), bDeleted);
+	bFreed = false;
+	bDeleted = false;
+
+	Buffer.GetSize();
+	TestTrue(TEXT("FSharedBuffer(BufferOwner)->GetSize()->Materialized"), bMaterialized);
+	bMaterialized = false;
+	Buffer.GetSize();
+	TestFalse(TEXT("FSharedBuffer(BufferOwner)->GetSize()->!Materialized"), bMaterialized);
+	bMaterialized = false;
+
+	FWeakSharedBuffer WeakBuffer = Buffer;
+	Buffer.Reset();
+	TestTrue(TEXT("FSharedBuffer(BufferOwner)->Reset(Shared)->Freed"), bFreed);
+	TestFalse(TEXT("FSharedBuffer(BufferOwner)->Reset(Shared)->!Deleted"), bDeleted);
+	WeakBuffer.Reset();
+	TestTrue(TEXT("FSharedBuffer(BufferOwner)->Reset(Weak)->Deleted"), bDeleted);
+
+	bMaterialized = false;
+	bFreed = false;
+	bDeleted = false;
+	Buffer = FUniqueBuffer(new FTestBufferOwner(bMaterialized, bFreed, bDeleted));
+
+	TestFalse(TEXT("FSharedBuffer(BufferOwner)->IsMaterialized()"), Buffer.IsMaterialized());
+	Buffer.Materialize();
+	TestTrue(TEXT("FSharedBuffer(BufferOwner)->Materialize()->IsMaterialized()"), Buffer.IsMaterialized());
+	TestTrue(TEXT("FSharedBuffer(BufferOwner)->Materialize()->Materialized"), bMaterialized);
+	bMaterialized = false;
+	Buffer.Materialize();
+	TestFalse(TEXT("FSharedBuffer(BufferOwner)->Materialize()->!Materialized"), bMaterialized);
+	bMaterialized = false;
+	Buffer.Reset();
+
+	bFreed = false;
+	bDeleted = false;
+	Buffer = FUniqueBuffer(new FTestBufferOwner(bMaterialized, bFreed, bDeleted));
+	Buffer.Reset();
+	TestFalse(TEXT("FSharedBuffer(BufferOwner)->Reset(!Materialized)->!Materialized"), bMaterialized);
+	TestTrue(TEXT("FSharedBuffer(BufferOwner)->Reset(!Materialized)->Freed"), bFreed);
+	TestTrue(TEXT("FSharedBuffer(BufferOwner)->Reset(!Materialized)->Deleted"), bDeleted);
 
 	return true;
 }
diff --git a/Engine/Source/Runtime/Core/Private/Tests/Serialization/CompactBinaryPackageTest.cpp b/Engine/Source/Runtime/Core/Private/Tests/Serialization/CompactBinaryPackageTest.cpp
--- a/Engine/Source/Runtime/Core/Private/Tests/Serialization/CompactBinaryPackageTest.cpp
+++ b/Engine/Source/Runtime/Core/Private/Tests/Serialization/CompactBinaryPackageTest.cpp
@@ -30,7 +30,7 @@ bool FCbAttachmentTest::RunTest(const FString& Parameters)
 		FCbFieldRefIterator Fields = Writer.Save();
 
 		TestTrue(FString::Printf(TEXT("FCbAttachment(%s).Save()->Equals"), Test),
-			MakeMemoryView(WriteAr).EqualBytes(Fields.GetRangeBuffer()->GetView()));
+			MakeMemoryView(WriteAr).EqualBytes(Fields.GetRangeBuffer().GetView()));
 		TestEqual(FString::Printf(TEXT("FCbAttachment(%s).Save()->ValidateRange"), Test),
 			ValidateCompactBinaryRange(MakeMemoryView(WriteAr), ECbValidateMode::All), ECbValidateError::None);
 		TestEqual(FString::Printf(TEXT("FCbAttachment(%s).Save()->ValidateAttachment"), Test),
@@ -43,7 +43,7 @@ bool FCbAttachmentTest::RunTest(const FString& Parameters)
 
 		FCbAttachment FromArchive;
 		FMemoryReader ReadAr(WriteAr);
-		FromArchive.Load(ReadAr, [](uint64 Size) { return FSharedBuffer::Alloc(Size); });
+		FromArchive.Load(ReadAr);
 		TestTrue(FString::Printf(TEXT("FCbAttachment(%s).Load(Archive)->AtEnd"), Test), ReadAr.AtEnd());
 		TestEqual(FString::Printf(TEXT("FCbAttachment(%s).Load(Archive)->Equals"), Test), FromArchive, Attachment);
 	};
@@ -53,7 +53,7 @@ bool FCbAttachmentTest::RunTest(const FString& Parameters)
 		FCbAttachment Attachment;
 		TestTrue(TEXT("FCbAttachment(Null).IsNull()"), Attachment.IsNull());
 		TestFalse(TEXT("FCbAttachment(Null) as bool"), bool(Attachment));
-		TestFalse(TEXT("FCbAttachment(Null).AsBinary()"), Attachment.AsBinary().IsValid());
+		TestFalse(TEXT("FCbAttachment(Null).AsBinary()"), bool(Attachment.AsBinary()));
 		TestFalse(TEXT("FCbAttachment(Null).AsCompactBinary()"), Attachment.AsCompactBinary().HasValue());
 		TestFalse(TEXT("FCbAttachment(Null).IsBinary()"), Attachment.IsBinary());
 		TestFalse(TEXT("FCbAttachment(Null).IsCompactBinary()"), Attachment.IsCompactBinary());
@@ -63,7 +63,7 @@ bool FCbAttachmentTest::RunTest(const FString& Parameters)
 
 	// Binary Attachment
 	{
-		FSharedBufferConstPtr Buffer = FSharedBuffer::MakeReadOnly(FSharedBuffer::Clone(MakeMemoryView<uint8>({0, 1, 2, 3})));
+		FSharedBuffer Buffer = FSharedBuffer::Clone(MakeMemoryView<uint8>({0, 1, 2, 3}));
 		FCbAttachment Attachment(Buffer);
 		TestFalse(TEXT("FCbAttachment(Binary).IsNull()"), Attachment.IsNull());
 		TestTrue(TEXT("FCbAttachment(Binary) as bool"), bool(Attachment));
@@ -71,7 +71,7 @@ bool FCbAttachmentTest::RunTest(const FString& Parameters)
 		TestFalse(TEXT("FCbAttachment(Binary).AsCompactBinary()"), Attachment.AsCompactBinary().HasValue());
 		TestTrue(TEXT("FCbAttachment(Binary).IsBinary()"), Attachment.IsBinary());
 		TestFalse(TEXT("FCbAttachment(Binary).IsCompactBinary()"), Attachment.IsCompactBinary());
-		TestEqual(TEXT("FCbAttachment(Binary).GetHash()"), Attachment.GetHash(), FBlake3::HashBuffer(Buffer->GetView()));
+		TestEqual(TEXT("FCbAttachment(Binary).GetHash()"), Attachment.GetHash(), FBlake3::HashBuffer(Buffer));
 		TestSaveLoadValidate(TEXT("Binary"), Attachment);
 	}
 
@@ -94,18 +94,18 @@ bool FCbAttachmentTest::RunTest(const FString& Parameters)
 	// Binary View
 	{
 		const uint8 Value[]{0, 1, 2, 3};
-		FSharedBufferConstPtr Buffer = FSharedBuffer::MakeView(MakeMemoryView(Value));
+		FSharedBuffer Buffer = FSharedBuffer::MakeView(MakeMemoryView(Value));
 		FCbAttachment Attachment(Buffer);
 		TestFalse(TEXT("FCbAttachment(BinaryView).IsNull()"), Attachment.IsNull());
 		TestTrue(TEXT("FCbAttachment(BinaryView) as bool"), bool(Attachment));
 		TestNotEqual(TEXT("FCbAttachment(BinaryView).AsBinary()"), Attachment.AsBinary(), Buffer);
 		TestTrue(TEXT("FCbAttachment(BinaryView).AsBinary()"),
-			Attachment.AsBinary()->GetView().EqualBytes(Buffer->GetView()));
+			Attachment.AsBinary().GetView().EqualBytes(Buffer.GetView()));
 		TestFalse(TEXT("FCbAttachment(BinaryView).AsCompactBinary()"), Attachment.AsCompactBinary().HasValue());
 		TestTrue(TEXT("FCbAttachment(BinaryView).IsBinary()"), Attachment.IsBinary());
 		TestFalse(TEXT("FCbAttachment(BinaryView).IsCompactBinary()"), Attachment.IsCompactBinary());
 		TestEqual(TEXT("FCbAttachment(BinaryView).GetHash()"),
-			Attachment.GetHash(), FBlake3::HashBuffer(Buffer->GetView()));
+			Attachment.GetHash(), FBlake3::HashBuffer(Buffer));
 	}
 
 	// Compact Binary View
@@ -129,7 +129,7 @@ bool FCbAttachmentTest::RunTest(const FString& Parameters)
 	// Binary Load from View
 	{
 		const uint8 Value[]{0, 1, 2, 3};
-		const FSharedBufferConstPtr Buffer = FSharedBuffer::MakeView(MakeMemoryView(Value));
+		const FSharedBuffer Buffer = FSharedBuffer::MakeView(MakeMemoryView(Value));
 		FCbAttachment Attachment(Buffer);
 
 		FCbWriter Writer;
@@ -141,9 +141,9 @@ bool FCbAttachmentTest::RunTest(const FString& Parameters)
 		TestFalse(TEXT("FCbAttachment(LoadBinaryView).IsNull()"), Attachment.IsNull());
 		TestTrue(TEXT("FCbAttachment(LoadBinaryView) as bool"), bool(Attachment));
 		TestFalse(TEXT("FCbAttachment(LoadBinaryView).AsBinary()->!InView"),
-			FieldsView.GetRangeBuffer()->GetView().Contains(Attachment.AsBinary()->GetView()));
+			FieldsView.GetRangeBuffer().GetView().Contains(Attachment.AsBinary().GetView()));
 		TestTrue(TEXT("FCbAttachment(LoadBinaryView).AsBinary()->EqualBytes"),
-			Attachment.AsBinary()->GetView().EqualBytes(Buffer->GetView()));
+			Attachment.AsBinary().GetView().EqualBytes(Buffer.GetView()));
 		TestFalse(TEXT("FCbAttachment(LoadBinaryView).AsCompactBinary()"), Attachment.AsCompactBinary().HasValue());
 		TestTrue(TEXT("FCbAttachment(LoadBinaryView).IsBinary()"), Attachment.IsBinary());
 		TestFalse(TEXT("FCbAttachment(LoadBinaryView).IsCompactBinary()"), Attachment.IsCompactBinary());
@@ -169,9 +169,9 @@ bool FCbAttachmentTest::RunTest(const FString& Parameters)
 		TestFalse(TEXT("FCbAttachment(LoadCompactBinaryView).IsNull()"), Attachment.IsNull());
 		TestTrue(TEXT("FCbAttachment(LoadCompactBinaryView) as bool"), bool(Attachment));
 		TestTrue(TEXT("FCbAttachment(LoadCompactBinaryView).AsBinary()->EqualBytes"),
-			Attachment.AsBinary()->GetView().EqualBytes(Value.GetRangeView()));
+			Attachment.AsBinary().GetView().EqualBytes(Value.GetRangeView()));
 		TestFalse(TEXT("FCbAttachment(LoadCompactBinaryView).AsCompactBinary()->!InView"),
-			FieldsView.GetRangeBuffer()->GetView().Contains(Attachment.AsCompactBinary().GetRangeBuffer()->GetView()));
+			FieldsView.GetRangeBuffer().GetView().Contains(Attachment.AsCompactBinary().GetRangeBuffer().GetView()));
 		TestTrue(TEXT("FCbAttachment(LoadCompactBinaryView).IsBinary()"), Attachment.IsBinary());
 		TestTrue(TEXT("FCbAttachment(LoadCompactBinaryView).IsCompactBinary()"), Attachment.IsCompactBinary());
 		TestEqual(TEXT("FCbAttachment(LoadCompactBinaryView).GetHash()"), Attachment.GetHash(), Value.GetRangeHash());
@@ -179,18 +179,18 @@ bool FCbAttachmentTest::RunTest(const FString& Parameters)
 
 	// Compact Binary Uniform Sub-View
 	{
-		const FSharedBufferConstPtr Buffer = FSharedBuffer::Clone(MakeMemoryView<uint8>({0, 1, 2, 3}));
-		const FCbFieldIterator Fields = FCbFieldIterator::MakeRange(Buffer->GetView().RightChop(2), ECbFieldType::IntegerPositive);
+		const FSharedBuffer Buffer = FSharedBuffer::Clone(MakeMemoryView<uint8>({0, 1, 2, 3}));
+		const FCbFieldIterator Fields = FCbFieldIterator::MakeRange(Buffer.GetView().RightChop(2), ECbFieldType::IntegerPositive);
 		const FCbFieldRefIterator SavedFieldRefs = FCbFieldRefIterator::CloneRange(Fields);
 		FCbFieldRefIterator FieldRefs = FCbFieldRefIterator::MakeRangeView(Fields, Buffer);
 		FCbAttachment Attachment(FieldRefs);
-		const FSharedBufferConstPtr Binary = Attachment.AsBinary();
+		const FSharedBuffer Binary = Attachment.AsBinary();
 		TestEqual(TEXT("FCbAttachment(CompactBinaryUniformSubView).AsCompactBinary()->Equals()"),
 			Attachment.AsCompactBinary(), FieldRefs);
 		TestEqual(TEXT("FCbAttachment(CompactBinaryUniformSubView).AsBinary()->GetSize()"),
-			Binary->GetSize(), SavedFieldRefs.GetRangeSize());
+			Binary.GetSize(), SavedFieldRefs.GetRangeSize());
 		TestTrue(TEXT("FCbAttachment(CompactBinaryUniformSubView).AsBinary()->EqualBytes()"),
-			Binary->GetView().EqualBytes(SavedFieldRefs.GetRangeView()));
+			Binary.GetView().EqualBytes(SavedFieldRefs.GetRangeView()));
 		TestEqual(TEXT("FCbAttachment(CompactBinaryUniformSubView).GetHash()"),
 			Attachment.GetHash(), SavedFieldRefs.GetRangeHash());
 		TestSaveLoadValidate(TEXT("CompactBinaryUniformSubView"), Attachment);
@@ -198,7 +198,7 @@ bool FCbAttachmentTest::RunTest(const FString& Parameters)
 
 	// Binary Null
 	{
-		const FCbAttachment Attachment(FSharedBufferConstPtr{});
+		const FCbAttachment Attachment(FSharedBuffer{});
 		TestTrue(TEXT("FCbAttachment(BinaryNull).IsNull()"), Attachment.IsNull());
 		TestFalse(TEXT("FCbAttachment(BinaryNull).IsBinary()"), Attachment.IsBinary());
 		TestFalse(TEXT("FCbAttachment(BinaryNull).IsCompactBinary()"), Attachment.IsCompactBinary());
@@ -207,7 +207,7 @@ bool FCbAttachmentTest::RunTest(const FString& Parameters)
 
 	// Binary Empty
 	{
-		const FCbAttachment Attachment(FSharedBuffer::Alloc(0));
+		const FCbAttachment Attachment(FSharedBuffer(FUniqueBuffer::Alloc(0)));
 		TestTrue(TEXT("FCbAttachment(BinaryEmpty).IsNull()"), Attachment.IsNull());
 		TestFalse(TEXT("FCbAttachment(BinaryEmpty).IsBinary()"), Attachment.IsBinary());
 		TestFalse(TEXT("FCbAttachment(BinaryEmpty).IsCompactBinary()"), Attachment.IsCompactBinary());
@@ -238,7 +238,7 @@ bool FCbPackageTest::RunTest(const FString& Parameters)
 		FCbFieldRefIterator Fields = Writer.Save();
 
 		TestTrue(FString::Printf(TEXT("FCbPackage(%s).Save()->Equals"), Test),
-			MakeMemoryView(WriteAr).EqualBytes(Fields.GetRangeBuffer()->GetView()));
+			MakeMemoryView(WriteAr).EqualBytes(Fields.GetRangeBuffer().GetView()));
 		TestEqual(FString::Printf(TEXT("FCbPackage(%s).Save()->ValidateRange"), Test),
 			ValidateCompactBinaryRange(MakeMemoryView(WriteAr), ECbValidateMode::All), ECbValidateError::None);
 		TestEqual(FString::Printf(TEXT("FCbPackage(%s).Save()->ValidatePackage"), Test),
@@ -251,7 +251,7 @@ bool FCbPackageTest::RunTest(const FString& Parameters)
 
 		FCbPackage FromArchive;
 		FMemoryReader ReadAr(WriteAr);
-		FromArchive.Load(ReadAr, [](uint64 Size) { return FSharedBuffer::Alloc(Size); });
+		FromArchive.Load(ReadAr);
 		TestTrue(FString::Printf(TEXT("FCbPackage(%s).Load(Archive)->AtEnd"), Test), ReadAr.AtEnd());
 		TestEqual(FString::Printf(TEXT("FCbPackage(%s).Load(Archive)->Equals"), Test), FromArchive, Package);
 	};
@@ -356,8 +356,8 @@ bool FCbPackageTest::RunTest(const FString& Parameters)
 
 	// Shared Values
 	const uint8 Level4Values[]{0, 1, 2, 3};
-	FSharedBufferConstPtr Level4 = FSharedBuffer::MakeView(MakeMemoryView(Level4Values));
-	const FBlake3Hash Level4Hash = FBlake3::HashBuffer(*Level4);
+	FSharedBuffer Level4 = FSharedBuffer::MakeView(MakeMemoryView(Level4Values));
+	const FBlake3Hash Level4Hash = FBlake3::HashBuffer(Level4);
 
 	FCbFieldRef Level3;
 	{
@@ -389,13 +389,13 @@ bool FCbPackageTest::RunTest(const FString& Parameters)
 	const FBlake3Hash Level1Hash = Level1.AsField().GetHash();
 
 	const auto Resolver = [&Level2, &Level2Hash, &Level3, &Level3Hash, &Level4, &Level4Hash]
-		(const FBlake3Hash& Hash) -> FSharedBufferConstPtr
+		(const FBlake3Hash& Hash) -> FSharedBuffer
 		{
 			return
 				Hash == Level2Hash ? Level2.GetBuffer() :
 				Hash == Level3Hash ? Level3.GetBuffer() :
 				Hash == Level4Hash ? Level4 :
-				FSharedBufferConstPtr();
+				FSharedBuffer();
 		};
 
 	// Object + Attachments
@@ -421,7 +421,7 @@ bool FCbPackageTest::RunTest(const FString& Parameters)
 		TestTrue(TEXT("FCbPackage(Object+Attachments).FindAttachment(Level4)"),
 			Level4Attachment &&
 			Level4Attachment->AsBinary() != Level4 &&
-			Level4Attachment->AsBinary()->GetView().EqualBytes(Level4->GetView()));
+			Level4Attachment->AsBinary().GetView().EqualBytes(Level4.GetView()));
 
 		TestTrue(TEXT("FCbPackage(Object+Attachments).GetAttachments()->Sorted"),
 			Algo::IsSorted(Package.GetAttachments()));
@@ -478,21 +478,21 @@ bool FCbPackageTest::RunTest(const FString& Parameters)
 		TestTrue(TEXT("FCbPackage(OutOfOrder).Load()->Level2"),
 			Level2Attachment && Level2Attachment->AsCompactBinary().AsArrayRef().Equals(Level2));
 		TestTrue(TEXT("FCbPackage(OutOfOrder).Load()->Level2Buffer"),
-			Level2Attachment && Fields.GetOuterBuffer()->GetView().Contains(Level2Attachment->AsBinary()->GetView()));
+			Level2Attachment && Fields.GetOuterBuffer().GetView().Contains(Level2Attachment->AsBinary().GetView()));
 		TestTrue(TEXT("FCbPackage(OutOfOrder).Load()->Level2Hash"),
 			Level2Attachment && Level2Attachment->GetHash() == Level2Hash);
 
 		TestTrue(TEXT("FCbPackage(OutOfOrder).Load()->Level3"),
 			Level3Attachment && Level3Attachment->AsCompactBinary().Equals(Level3));
 		TestTrue(TEXT("FCbPackage(OutOfOrder).Load()->Level3Buffer"),
-			Level3Attachment && Fields.GetOuterBuffer()->GetView().Contains(Level3Attachment->AsBinary()->GetView()));
+			Level3Attachment && Fields.GetOuterBuffer().GetView().Contains(Level3Attachment->AsBinary().GetView()));
 		TestTrue(TEXT("FCbPackage(OutOfOrder).Load()->Level3Hash"),
 			Level3Attachment && Level3Attachment->GetHash() == Level3Hash);
 
 		TestTrue(TEXT("FCbPackage(OutOfOrder).Load()->Level4"),
-			Level4Attachment && Level4Attachment->AsBinary()->GetView().EqualBytes(Level4->GetView()));
+			Level4Attachment && Level4Attachment->AsBinary().GetView().EqualBytes(Level4.GetView()));
 		TestTrue(TEXT("FCbPackage(OutOfOrder).Load()->Level4Buffer"),
-			Level4Attachment && Fields.GetOuterBuffer()->GetView().Contains(Level4Attachment->AsBinary()->GetView()));
+			Level4Attachment && Fields.GetOuterBuffer().GetView().Contains(Level4Attachment->AsBinary().GetView()));
 		TestTrue(TEXT("FCbPackage(OutOfOrder).Load()->Level4Hash"),
 			Level4Attachment && Level4Attachment->GetHash() == Level4Hash);
 
@@ -500,7 +500,7 @@ bool FCbPackageTest::RunTest(const FString& Parameters)
 		Writer.Save(WriteAr);
 		FCbPackage FromArchive;
 		FMemoryReader ReadAr(WriteAr);
-		FromArchive.Load(ReadAr, [](uint64 Size) { return FSharedBuffer::Alloc(Size); });
+		FromArchive.Load(ReadAr);
 
 		Writer.Reset();
 		FromArchive.Save(Writer);
@@ -517,7 +517,7 @@ bool FCbPackageTest::RunTest(const FString& Parameters)
 		++Saved;
 		TestEqual(TEXT("FCbPackage(OutOfOrder).Save()->Level2Hash"), Saved.AsReference(), Level2Hash);
 		++Saved;
-		TestTrue(TEXT("FCbPackage(OutOfOrder).Save()->Level4"), Saved.AsBinary().EqualBytes(Level4->GetView()));
+		TestTrue(TEXT("FCbPackage(OutOfOrder).Save()->Level4"), Saved.AsBinary().EqualBytes(Level4.GetView()));
 		++Saved;
 		TestEqual(TEXT("FCbPackage(OutOfOrder).Save()->Level4Hash"), Saved.AsBinaryReference(), Level4Hash);
 		++Saved;
@@ -543,10 +543,10 @@ bool FCbPackageTest::RunTest(const FString& Parameters)
 		FCbPackage Package;
 		Package.AddAttachment(FCbAttachment(Level3.GetBuffer()));
 		Package.AddAttachment(FCbAttachment(FCbFieldRefIterator::MakeSingle(Level3)),
-			[&bResolved](const FBlake3Hash& Hash) -> FSharedBufferConstPtr
+			[&bResolved](const FBlake3Hash& Hash) -> FSharedBuffer
 			{
 				bResolved = true;
-				return FSharedBufferConstPtr();
+				return FSharedBuffer();
 			});
 		TestTrue(TEXT("FCbPackage(ResolveAfterMerge)->Resolved"), bResolved);
 	}
diff --git a/Engine/Source/Runtime/Core/Private/Tests/Serialization/CompactBinarySerializationTest.cpp b/Engine/Source/Runtime/Core/Private/Tests/Serialization/CompactBinarySerializationTest.cpp
--- a/Engine/Source/Runtime/Core/Private/Tests/Serialization/CompactBinarySerializationTest.cpp
+++ b/Engine/Source/Runtime/Core/Private/Tests/Serialization/CompactBinarySerializationTest.cpp
@@ -141,7 +141,7 @@ bool FCbSaveTest::RunTest(const FString& Parameters)
 				TestTrue(FString::Printf(TEXT("SaveCompactBinary(%s)->EqualBytes"), Test), ExpectedData.EqualBytes(MakeMemoryView(WriteAr))))
 			{
 				FMemoryReader ReadAr(WriteAr);
-				FCbFieldRef Field = LoadCompactBinary(ReadAr, [](uint64 Size) { return FSharedBuffer::Alloc(Size); });
+				FCbFieldRef Field = LoadCompactBinary(ReadAr);
 				TestTrue(FString::Printf(TEXT("LoadCompactBinary(%s)->EqualBytes"), Test), ExpectedData.EqualBytes(Field.GetView()));
 			}
 		}
@@ -161,20 +161,20 @@ bool FCbSaveTest::RunTest(const FString& Parameters)
 	// Field
 	{
 		const uint8 Payload[] = { uint8(ECbFieldType::IntegerPositive), 42 };
-		TestSave(TEXT("Field"), FCbFieldRef(FCbField(Payload), FSharedBufferConstPtr()), MakeMemoryView(Payload));
+		TestSave(TEXT("Field"), FCbFieldRef(FCbField(Payload), FSharedBuffer()), MakeMemoryView(Payload));
 	}
 	{
 		const uint8 Payload[] = { uint8(ECbFieldType::IntegerPositive | ECbFieldType::HasFieldName), 1, 'A', 42 };
-		TestSave(TEXT("Field, Name"), FCbFieldRef(FCbField(Payload), FSharedBufferConstPtr()), MakeMemoryView(Payload));
+		TestSave(TEXT("Field, Name"), FCbFieldRef(FCbField(Payload), FSharedBuffer()), MakeMemoryView(Payload));
 	}
 	{
 		const uint8 Payload[] = { 42 };
-		TestSave(TEXT("Field, NoType"), FCbFieldRef(FCbField(Payload, ECbFieldType::IntegerPositive), FSharedBufferConstPtr()),
+		TestSave(TEXT("Field, NoType"), FCbFieldRef(FCbField(Payload, ECbFieldType::IntegerPositive), FSharedBuffer()),
 			MakeMemoryView<uint8>({ uint8(ECbFieldType::IntegerPositive), 42 }));
 	}
 	{
 		const uint8 Payload[] = { 1, 'I', 42 };
-		TestSave(TEXT("Field, NoType, Name"), FCbFieldRef(FCbField(Payload, ECbFieldType::IntegerPositive | ECbFieldType::HasFieldName), FSharedBufferConstPtr()),
+		TestSave(TEXT("Field, NoType, Name"), FCbFieldRef(FCbField(Payload, ECbFieldType::IntegerPositive | ECbFieldType::HasFieldName), FSharedBuffer()),
 			MakeMemoryView<uint8>({ uint8(ECbFieldType::IntegerPositive | ECbFieldType::HasFieldName), 1, 'I', 42 }));
 	}
 
@@ -182,25 +182,25 @@ bool FCbSaveTest::RunTest(const FString& Parameters)
 	{
 		constexpr uint8 IntType = uint8(ECbFieldType::IntegerPositive);
 		const uint8 Payload[] = { uint8(ECbFieldType::UniformArray), 5, 3, IntType, 1, 2, 3 };
-		TestSave(TEXT("Array"), FCbArrayRef(FCbArray(Payload), FSharedBufferConstPtr()),
+		TestSave(TEXT("Array"), FCbArrayRef(FCbArray(Payload), FSharedBuffer()),
 			MakeMemoryView(Payload));
 	}
 	{
 		constexpr uint8 IntType = uint8(ECbFieldType::IntegerPositive);
 		const uint8 Payload[] = { uint8(ECbFieldType::UniformArray | ECbFieldType::HasFieldName), 1, 'A', 5, 3, IntType, 1, 2, 3 };
-		TestSave(TEXT("Array, Name"), FCbArrayRef(FCbArray(Payload), FSharedBufferConstPtr()),
+		TestSave(TEXT("Array, Name"), FCbArrayRef(FCbArray(Payload), FSharedBuffer()),
 			MakeMemoryView<uint8>({ uint8(ECbFieldType::UniformArray), 5, 3, IntType, 1, 2, 3 }));
 	}
 	{
 		constexpr uint8 IntType = uint8(ECbFieldType::IntegerPositive);
 		const uint8 Payload[] = { 5, 3, IntType, 1, 2, 3 };
-		TestSave(TEXT("Array, NoType"), FCbArrayRef(FCbArray(Payload, ECbFieldType::UniformArray), FSharedBufferConstPtr()),
+		TestSave(TEXT("Array, NoType"), FCbArrayRef(FCbArray(Payload, ECbFieldType::UniformArray), FSharedBuffer()),
 			MakeMemoryView<uint8>({ uint8(ECbFieldType::UniformArray), 5, 3, IntType, 1, 2, 3 }));
 	}
 	{
 		constexpr uint8 IntType = uint8(ECbFieldType::IntegerPositive);
 		const uint8 Payload[] = { 1, 'A', 5, 3, IntType, 1, 2, 3 };
-		TestSave(TEXT("Array, NoType, Name"), FCbArrayRef(FCbArray(Payload, ECbFieldType::UniformArray | ECbFieldType::HasFieldName), FSharedBufferConstPtr()),
+		TestSave(TEXT("Array, NoType, Name"), FCbArrayRef(FCbArray(Payload, ECbFieldType::UniformArray | ECbFieldType::HasFieldName), FSharedBuffer()),
 			MakeMemoryView<uint8>({ uint8(ECbFieldType::UniformArray), 5, 3, IntType, 1, 2, 3 }));
 	}
 
@@ -208,25 +208,25 @@ bool FCbSaveTest::RunTest(const FString& Parameters)
 	{
 		constexpr uint8 IntType = uint8(ECbFieldType::HasFieldName | ECbFieldType::IntegerPositive);
 		const uint8 Payload[] = { uint8(ECbFieldType::UniformObject), 10, IntType, 1, 'A', 1, 1, 'B', 2, 1, 'C', 3 };
-		TestSave(TEXT("Object"), FCbObjectRef(FCbObject(Payload), FSharedBufferConstPtr()),
+		TestSave(TEXT("Object"), FCbObjectRef(FCbObject(Payload), FSharedBuffer()),
 			MakeMemoryView(Payload));
 	}
 	{
 		constexpr uint8 IntType = uint8(ECbFieldType::HasFieldName | ECbFieldType::IntegerPositive);
 		const uint8 Payload[] = { uint8(ECbFieldType::UniformObject | ECbFieldType::HasFieldName), 1, 'O', 10, IntType, 1, 'A', 1, 1, 'B', 2, 1, 'C', 3 };
-		TestSave(TEXT("Object, Name"), FCbObjectRef(FCbObject(Payload), FSharedBufferConstPtr()),
+		TestSave(TEXT("Object, Name"), FCbObjectRef(FCbObject(Payload), FSharedBuffer()),
 			MakeMemoryView<uint8>({ uint8(ECbFieldType::UniformObject), 10, IntType, 1, 'A', 1, 1, 'B', 2, 1, 'C', 3 }));
 	}
 	{
 		constexpr uint8 IntType = uint8(ECbFieldType::HasFieldName | ECbFieldType::IntegerPositive);
 		const uint8 Payload[] = { 10, IntType, 1, 'A', 1, 1, 'B', 2, 1, 'C', 3 };
-		TestSave(TEXT("Object, NoType"), FCbObjectRef(FCbObject(Payload, ECbFieldType::UniformObject), FSharedBufferConstPtr()),
+		TestSave(TEXT("Object, NoType"), FCbObjectRef(FCbObject(Payload, ECbFieldType::UniformObject), FSharedBuffer()),
 			MakeMemoryView<uint8>({ uint8(ECbFieldType::UniformObject), 10, IntType, 1, 'A', 1, 1, 'B', 2, 1, 'C', 3 }));
 	}
 	{
 		constexpr uint8 IntType = uint8(ECbFieldType::HasFieldName | ECbFieldType::IntegerPositive);
 		const uint8 Payload[] = { 1, 'O', 10, IntType, 1, 'A', 1, 1, 'B', 2, 1, 'C', 3 };
-		TestSave(TEXT("Object, NoType, Name"), FCbObjectRef(FCbObject(Payload, ECbFieldType::UniformObject | ECbFieldType::HasFieldName), FSharedBufferConstPtr()),
+		TestSave(TEXT("Object, NoType, Name"), FCbObjectRef(FCbObject(Payload, ECbFieldType::UniformObject | ECbFieldType::HasFieldName), FSharedBuffer()),
 			MakeMemoryView<uint8>({ uint8(ECbFieldType::UniformObject), 10, IntType, 1, 'A', 1, 1, 'B', 2, 1, 'C', 3 }));
 	}
 
diff --git a/Engine/Source/Runtime/Core/Private/Tests/Serialization/CompactBinaryTest.cpp b/Engine/Source/Runtime/Core/Private/Tests/Serialization/CompactBinaryTest.cpp
--- a/Engine/Source/Runtime/Core/Private/Tests/Serialization/CompactBinaryTest.cpp
+++ b/Engine/Source/Runtime/Core/Private/Tests/Serialization/CompactBinaryTest.cpp
@@ -451,7 +451,7 @@ bool FCbFieldObjectTest::RunTest(const FString& Parameters)
 			TestEqual(TEXT("FCbObjectRef::CreateRefIterator().GetName()"), Field.GetName(), "F"_ASV);
 			TestEqual(TEXT("FCbObjectRef::CreateRefIterator().AsInt32()"), Field.AsInt32(), 8);
 			TestTrue(TEXT("FCbObjectRef::CreateRefIterator().IsOwned()"), Field.IsOwned());
-			TestEqual(TEXT("FCbObjectRef::CreateRefIterator().GetBuffer()"), Field.GetBuffer()->GetView(), Field.GetView());
+			TestEqual(TEXT("FCbObjectRef::CreateRefIterator().GetBuffer()"), Field.GetBuffer().GetView(), Field.GetView());
 		}
 		for (FCbFieldRefIterator It = ObjectClone.CreateRefIterator(), End; It != End; ++It)
 		{
@@ -615,7 +615,7 @@ bool FCbFieldArrayTest::RunTest(const FString& Parameters)
 			FCbFieldRef Field = *It;
 			TestEqual(TEXT("FCbArrayRef::CreateRefIterator().AsInt32()"), Field.AsInt32(), 8);
 			TestTrue(TEXT("FCbArrayRef::CreateRefIterator().IsOwned()"), Field.IsOwned());
-			TestEqual(TEXT("FCbArrayRef::CreateRefIterator().GetBuffer()"), Field.GetBuffer()->GetView(), Field.GetView());
+			TestEqual(TEXT("FCbArrayRef::CreateRefIterator().GetBuffer()"), Field.GetBuffer().GetView(), Field.GetView());
 		}
 		for (FCbFieldRefIterator It = ArrayClone.CreateRefIterator(), End; It != End; ++It)
 		{
@@ -1252,7 +1252,7 @@ bool FCbFieldIterateReferencesTest::RunTest(const FString& Parameters)
 		Fields = Writer.Save();
 	}
 
-	TestEqual(TEXT("FCbField::IterateReferences Validate"), ValidateCompactBinaryRange(*Fields.GetBuffer(), ECbValidateMode::All), ECbValidateError::None);
+	TestEqual(TEXT("FCbField::IterateReferences Validate"), ValidateCompactBinaryRange(Fields.GetBuffer(), ECbValidateMode::All), ECbValidateError::None);
 
 	uint32 ReferenceIndex = 0;
 	Fields.IterateRangeReferences([this, &ReferenceIndex, &MakeTestHash](FCbField Field)
@@ -1271,29 +1271,16 @@ bool FCbFieldIterateReferencesTest::RunTest(const FString& Parameters)
 IMPLEMENT_CUSTOM_SIMPLE_AUTOMATION_TEST(FCbFieldRefTest, FCbFieldTestBase, "System.Core.Serialization.CbFieldRef", CompactBinaryTestFlags)
 bool FCbFieldRefTest::RunTest(const FString& Parameters)
 {
-	static_assert(std::is_constructible<FCbFieldRef, const FSharedBufferRef&>::value, "Missing constructor for FCbFieldRef");
-	static_assert(std::is_constructible<FCbFieldRef, const FSharedBufferPtr&>::value, "Missing constructor for FCbFieldRef");
-	static_assert(std::is_constructible<FCbFieldRef, const FSharedBufferConstRef&>::value, "Missing constructor for FCbFieldRef");
-	static_assert(std::is_constructible<FCbFieldRef, const FSharedBufferConstPtr&>::value, "Missing constructor for FCbFieldRef");
-
-	static_assert(std::is_constructible<FCbFieldRef, FSharedBufferRef&&>::value, "Missing constructor for FCbFieldRef");
-	static_assert(std::is_constructible<FCbFieldRef, FSharedBufferPtr&&>::value, "Missing constructor for FCbFieldRef");
-	static_assert(std::is_constructible<FCbFieldRef, FSharedBufferConstRef&&>::value, "Missing constructor for FCbFieldRef");
-	static_assert(std::is_constructible<FCbFieldRef, FSharedBufferConstPtr&&>::value, "Missing constructor for FCbFieldRef");
-
-	static_assert(std::is_constructible<FCbFieldRef, const FCbField&, const FSharedBufferRef&>::value, "Missing constructor for FCbFieldRef");
-	static_assert(std::is_constructible<FCbFieldRef, const FCbField&, const FSharedBufferPtr&>::value, "Missing constructor for FCbFieldRef");
-	static_assert(std::is_constructible<FCbFieldRef, const FCbField&, const FSharedBufferConstRef&>::value, "Missing constructor for FCbFieldRef");
-	static_assert(std::is_constructible<FCbFieldRef, const FCbField&, const FSharedBufferConstPtr&>::value, "Missing constructor for FCbFieldRef");
+	static_assert(std::is_constructible<FCbFieldRef, const FSharedBuffer&>::value, "Missing constructor for FCbFieldRef");
+	static_assert(std::is_constructible<FCbFieldRef, FSharedBuffer&&>::value, "Missing constructor for FCbFieldRef");
+
+	static_assert(std::is_constructible<FCbFieldRef, const FCbField&, const FSharedBuffer&>::value, "Missing constructor for FCbFieldRef");
 	static_assert(std::is_constructible<FCbFieldRef, const FCbField&, const FCbFieldRefIterator&>::value, "Missing constructor for FCbFieldRef");
 	static_assert(std::is_constructible<FCbFieldRef, const FCbField&, const FCbFieldRef&>::value, "Missing constructor for FCbFieldRef");
 	static_assert(std::is_constructible<FCbFieldRef, const FCbField&, const FCbArrayRef&>::value, "Missing constructor for FCbFieldRef");
 	static_assert(std::is_constructible<FCbFieldRef, const FCbField&, const FCbObjectRef&>::value, "Missing constructor for FCbFieldRef");
 
-	static_assert(std::is_constructible<FCbFieldRef, const FCbField&, FSharedBufferRef&&>::value, "Missing constructor for FCbFieldRef");
-	static_assert(std::is_constructible<FCbFieldRef, const FCbField&, FSharedBufferPtr&&>::value, "Missing constructor for FCbFieldRef");
-	static_assert(std::is_constructible<FCbFieldRef, const FCbField&, FSharedBufferConstRef&&>::value, "Missing constructor for FCbFieldRef");
-	static_assert(std::is_constructible<FCbFieldRef, const FCbField&, FSharedBufferConstPtr&&>::value, "Missing constructor for FCbFieldRef");
+	static_assert(std::is_constructible<FCbFieldRef, const FCbField&, FSharedBuffer&&>::value, "Missing constructor for FCbFieldRef");
 	static_assert(std::is_constructible<FCbFieldRef, const FCbField&, FCbFieldRefIterator&&>::value, "Missing constructor for FCbFieldRef");
 	static_assert(std::is_constructible<FCbFieldRef, const FCbField&, FCbFieldRef&&>::value, "Missing constructor for FCbFieldRef");
 	static_assert(std::is_constructible<FCbFieldRef, const FCbField&, FCbArrayRef&&>::value, "Missing constructor for FCbFieldRef");
@@ -1304,82 +1291,59 @@ bool FCbFieldRefTest::RunTest(const FString& Parameters)
 		FCbFieldRef DefaultField;
 		TestFalse(TEXT("FCbFieldRef().HasValue()"), DefaultField.HasValue());
 		TestFalse(TEXT("FCbFieldRef().IsOwned()"), DefaultField.IsOwned());
-		TestFalse(TEXT("FCbFieldRef().IsReadOnly()"), DefaultField.IsReadOnly());
 		DefaultField.MakeOwned();
 		TestTrue(TEXT("FCbFieldRef().MakeOwned().IsOwned()"), DefaultField.IsOwned());
-		TestTrue(TEXT("FCbFieldRef().MakeOwned().IsReadOnly()"), DefaultField.IsReadOnly());
 	}
 
 	// Test Field w/ Type from Shared Buffer
 	{
-		uint8 Payload[] = { uint8(ECbFieldType::Binary), 3, 4, 5, 6 }; // Size: 3, Data: 4/5/6
-		FSharedBufferPtr ViewBuffer = FSharedBuffer::MakeView(MakeMemoryView(Payload));
-		FSharedBufferPtr OwnedBuffer = FSharedBuffer::MakeOwned(ViewBuffer);
-		FSharedBufferConstPtr ViewBufferPtr = ViewBuffer;
-		FSharedBufferConstPtr OwnedBufferPtr = OwnedBuffer;
-
-		FCbFieldRef ViewRef(ViewBuffer);
-		FCbFieldRef ViewPtr(ViewBufferPtr);
-		FCbFieldRef ViewPtrMove{FSharedBufferConstPtr(ViewBufferPtr)};
-		FCbFieldRef ViewOuterFieldRef(ImplicitConv<FCbField>(ViewPtr), ViewPtrMove);
-		FCbFieldRef ViewOuterBufferRef(ImplicitConv<FCbField>(ViewRef), ViewPtr);
-		FCbFieldRef OwnedRef(OwnedBuffer);
-		FCbFieldRef OwnedPtr(OwnedBufferPtr);
-		FCbFieldRef OwnedPtrMove{FSharedBufferConstPtr(OwnedBufferPtr)};
-		FCbFieldRef OwnedOuterFieldRef(ImplicitConv<FCbField>(OwnedPtr), OwnedPtrMove);
-		FCbFieldRef OwnedOuterBufferRef(ImplicitConv<FCbField>(OwnedRef), OwnedPtr);
+		const uint8 Payload[] = { uint8(ECbFieldType::Binary), 3, 4, 5, 6 }; // Size: 3, Data: 4/5/6
+		FSharedBuffer ViewBuffer = FSharedBuffer::MakeView(MakeMemoryView(Payload));
+		FSharedBuffer OwnedBuffer = FSharedBuffer::Clone(ViewBuffer);
+
+		FCbFieldRef View(ViewBuffer);
+		FCbFieldRef ViewMove{FSharedBuffer(ViewBuffer)};
+		FCbFieldRef ViewOuterField(ImplicitConv<FCbField>(View), ViewBuffer);
+		FCbFieldRef ViewOuterBuffer(ImplicitConv<FCbField>(View), View);
+		FCbFieldRef Owned(OwnedBuffer);
+		FCbFieldRef OwnedMove{FSharedBuffer(OwnedBuffer)};
+		FCbFieldRef OwnedOuterField(ImplicitConv<FCbField>(Owned), OwnedBuffer);
+		FCbFieldRef OwnedOuterBuffer(ImplicitConv<FCbField>(Owned), Owned);
 
 		// These lines are expected to assert when uncommented.
-		//FCbFieldRef InvalidOuterBuffer(ImplicitConv<FCbField>(OwnedRef), ViewBufferPtr);
-		//FCbFieldRef InvalidOuterBufferMove(ImplicitConv<FCbField>(OwnedRef), FSharedBufferConstPtr(ViewBufferPtr));
-
-		Payload[UE_ARRAY_COUNT(Payload) - 1] = 4;
-
-		TestEqualBytes(TEXT("FCbFieldRef(ViewBufferRef)"), ViewRef.AsBinary(), {4, 5, 4});
-		TestEqualBytes(TEXT("FCbFieldRef(ViewBufferPtr)"), ViewPtr.AsBinary(), {4, 5, 4});
-		TestEqualBytes(TEXT("FCbFieldRef(ViewBufferPtr&&)"), ViewPtrMove.AsBinary(), {4, 5, 4});
-		TestEqualBytes(TEXT("FCbFieldRef(ViewOuterFieldRef)"), ViewOuterFieldRef.AsBinary(), {4, 5, 4});
-		TestEqualBytes(TEXT("FCbFieldRef(ViewOuterBufferRef)"), ViewOuterBufferRef.AsBinary(), {4, 5, 4});
-		TestEqualBytes(TEXT("FCbFieldRef(OwnedBufferRef)"), OwnedRef.AsBinary(), {4, 5, 6});
-		TestEqualBytes(TEXT("FCbFieldRef(OwnedBufferPtr)"), OwnedPtr.AsBinary(), {4, 5, 6});
-		TestEqualBytes(TEXT("FCbFieldRef(OwnedBufferPtr&&)"), OwnedPtrMove.AsBinary(), {4, 5, 6});
-		TestEqualBytes(TEXT("FCbFieldRef(OwnedOuterFieldRef)"), OwnedOuterFieldRef.AsBinary(), {4, 5, 6});
-		TestEqualBytes(TEXT("FCbFieldRef(OwnedOuterBufferRef)"), OwnedOuterBufferRef.AsBinary(), {4, 5, 6});
-
-		TestFalse(TEXT("FCbFieldRef(ViewBufferRef).IsOwned()"), ViewRef.IsOwned());
-		TestFalse(TEXT("FCbFieldRef(ViewBufferPtr).IsOwned()"), ViewPtr.IsOwned());
-		TestFalse(TEXT("FCbFieldRef(ViewBufferPtr&&).IsOwned()"), ViewPtrMove.IsOwned());
-		TestFalse(TEXT("FCbFieldRef(ViewOuterFieldRef).IsOwned()"), ViewOuterFieldRef.IsOwned());
-		TestFalse(TEXT("FCbFieldRef(ViewOuterBufferRef).IsOwned()"), ViewOuterBufferRef.IsOwned());
-		TestTrue(TEXT("FCbFieldRef(OwnedBufferRef).IsOwned()"), OwnedRef.IsOwned());
-		TestTrue(TEXT("FCbFieldRef(OwnedBufferPtr).IsOwned()"), OwnedPtr.IsOwned());
-		TestTrue(TEXT("FCbFieldRef(OwnedBufferPtr&&).IsOwned()"), OwnedPtrMove.IsOwned());
-		TestTrue(TEXT("FCbFieldRef(OwnedOuterFieldRef).IsOwned()"), OwnedOuterFieldRef.IsOwned());
-		TestTrue(TEXT("FCbFieldRef(OwnedOuterBufferRef).IsOwned()"), OwnedOuterBufferRef.IsOwned());
-
-		ViewRef.MakeOwned();
-		OwnedRef.MakeOwned();
-		static_cast<uint8*>(OwnedBuffer->GetData())[UE_ARRAY_COUNT(Payload) - 1] = 5;
-		TestEqualBytes(TEXT("FCbFieldRef(View).MakeOwned()"), ViewRef.AsBinary(), {4, 5, 4});
-		TestTrue(TEXT("FCbFieldRef(View).MakeOwned().IsOwned()"), ViewRef.IsOwned());
-		TestTrue(TEXT("FCbFieldRef(View).MakeOwned().IsReadOnly()"), ViewRef.IsReadOnly());
-		TestEqualBytes(TEXT("FCbFieldRef(Owned).MakeOwned()"), OwnedRef.AsBinary(), {4, 5, 5});
-		TestTrue(TEXT("FCbFieldRef(Owned).MakeOwned().IsOwned()"), OwnedRef.IsOwned());
-		TestFalse(TEXT("FCbFieldRef(Owned).MakeOwned().IsReadOnly()"), OwnedRef.IsReadOnly());
-
-		OwnedRef.MakeReadOnly();
-		TestEqualBytes(TEXT("FCbFieldRef(Owned).MakeReadOnly()"), OwnedRef.AsBinary(), {4, 5, 5});
-		TestTrue(TEXT("FCbFieldRef(Owned).MakeReadOnly().IsOwned()"), OwnedRef.IsOwned());
-		TestTrue(TEXT("FCbFieldRef(Owned).MakeReadOnly().IsReadOnly()"), OwnedRef.IsReadOnly());
+		//FCbFieldRef InvalidOuterBuffer(ImplicitConv<FCbField>(Owned), ViewBuffer);
+		//FCbFieldRef InvalidOuterBufferMove(ImplicitConv<FCbField>(Owned), FSharedBuffer(ViewBuffer));
+
+		TestEqual(TEXT("FCbFieldRef(ViewBuffer)"), View.AsBinary(), ViewBuffer.GetView().Right(3));
+		TestEqual(TEXT("FCbFieldRef(ViewBuffer&&)"), ViewMove.AsBinary(), View.AsBinary());
+		TestEqual(TEXT("FCbFieldRef(ViewOuterField)"), ViewOuterField.AsBinary(), View.AsBinary());
+		TestEqual(TEXT("FCbFieldRef(ViewOuterBuffer)"), ViewOuterBuffer.AsBinary(), View.AsBinary());
+		TestEqual(TEXT("FCbFieldRef(OwnedBuffer)"), Owned.AsBinary(), OwnedBuffer.GetView().Right(3));
+		TestEqual(TEXT("FCbFieldRef(OwnedBuffer&&)"), OwnedMove.AsBinary(), Owned.AsBinary());
+		TestEqual(TEXT("FCbFieldRef(OwnedOuterField)"), OwnedOuterField.AsBinary(), Owned.AsBinary());
+		TestEqual(TEXT("FCbFieldRef(OwnedOuterBuffer)"), OwnedOuterBuffer.AsBinary(), Owned.AsBinary());
+
+		TestFalse(TEXT("FCbFieldRef(ViewBuffer).IsOwned()"), View.IsOwned());
+		TestFalse(TEXT("FCbFieldRef(ViewBuffer&&).IsOwned()"), ViewMove.IsOwned());
+		TestFalse(TEXT("FCbFieldRef(ViewOuterField).IsOwned()"), ViewOuterField.IsOwned());
+		TestFalse(TEXT("FCbFieldRef(ViewOuterBuffer).IsOwned()"), ViewOuterBuffer.IsOwned());
+		TestTrue(TEXT("FCbFieldRef(OwnedBuffer).IsOwned()"), Owned.IsOwned());
+		TestTrue(TEXT("FCbFieldRef(OwnedBuffer&&).IsOwned()"), OwnedMove.IsOwned());
+		TestTrue(TEXT("FCbFieldRef(OwnedOuterField).IsOwned()"), OwnedOuterField.IsOwned());
+		TestTrue(TEXT("FCbFieldRef(OwnedOuterBuffer).IsOwned()"), OwnedOuterBuffer.IsOwned());
+
+		View.MakeOwned();
+		Owned.MakeOwned();
+		TestNotEqual(TEXT("FCbFieldRef(View).MakeOwned()"), View.AsBinary(), ViewBuffer.GetView().Right(3));
+		TestTrue(TEXT("FCbFieldRef(View).MakeOwned().IsOwned()"), View.IsOwned());
+		TestEqual(TEXT("FCbFieldRef(Owned).MakeOwned()"), Owned.AsBinary(), OwnedBuffer.GetView().Right(3));
+		TestTrue(TEXT("FCbFieldRef(Owned).MakeOwned().IsOwned()"), Owned.IsOwned());
 	}
 
 	// Test Field w/ Type
 	{
-		uint8 Payload[] = { uint8(ECbFieldType::Binary), 3, 4, 5, 6 }; // Size: 3, Data: 4/5/6
-		uint8* PayloadCopy = static_cast<uint8*>(FMemory::Malloc(sizeof(Payload)));
-		FMemory::Memcpy(PayloadCopy, Payload, sizeof(Payload));
-
-		FCbField Field(Payload);
+		const uint8 Payload[] = { uint8(ECbFieldType::Binary), 3, 4, 5, 6 }; // Size: 3, Data: 4/5/6
+		const FCbField Field(Payload);
 
 		FCbFieldRef VoidView = FCbFieldRef::MakeView(Payload);
 		FCbFieldRef VoidClone = FCbFieldRef::Clone(Payload);
@@ -1387,53 +1351,42 @@ bool FCbFieldRefTest::RunTest(const FString& Parameters)
 		FCbFieldRef FieldClone = FCbFieldRef::Clone(Field);
 		FCbFieldRef FieldRefClone = FCbFieldRef::Clone(FieldView);
 
-		Payload[UE_ARRAY_COUNT(Payload) - 1] = 4;
-
-		TestEqualBytes(TEXT("FCbFieldRef::MakeView(Void)"), VoidView.AsBinary(), {4, 5, 4});
-		TestEqualBytes(TEXT("FCbFieldRef::Clone(Void)"), VoidClone.AsBinary(), {4, 5, 6});
-		TestEqualBytes(TEXT("FCbFieldRef::MakeView(Field)"), FieldView.AsBinary(), {4, 5, 4});
-		TestEqualBytes(TEXT("FCbFieldRef::Clone(Field)"), FieldClone.AsBinary(), {4, 5, 6});
-		TestEqualBytes(TEXT("FCbFieldRef::Clone(FieldRef)"), FieldRefClone.AsBinary(), {4, 5, 6});
+		TestEqual(TEXT("FCbFieldRef::MakeView(Void)"), VoidView.AsBinary(), MakeMemoryView(Payload).Right(3));
+		TestNotEqual(TEXT("FCbFieldRef::Clone(Void)"), VoidClone.AsBinary(), MakeMemoryView(Payload).Right(3));
+		TestTrue(TEXT("FCbFieldRef::Clone(Void)->EqualBytes"), VoidClone.AsBinary().EqualBytes(VoidView.AsBinary()));
+		TestEqual(TEXT("FCbFieldRef::MakeView(Field)"), FieldView.AsBinary(), MakeMemoryView(Payload).Right(3));
+		TestNotEqual(TEXT("FCbFieldRef::Clone(Field)"), FieldClone.AsBinary(), MakeMemoryView(Payload).Right(3));
+		TestTrue(TEXT("FCbFieldRef::Clone(Field)->EqualBytes"), FieldClone.AsBinary().EqualBytes(VoidView.AsBinary()));
+		TestNotEqual(TEXT("FCbFieldRef::Clone(FieldRef)"), FieldRefClone.AsBinary(), FieldView.AsBinary());
+		TestTrue(TEXT("FCbFieldRef::Clone(FieldRef)->EqualBytes"), FieldRefClone.AsBinary().EqualBytes(VoidView.AsBinary()));
 
 		TestFalse(TEXT("FCbFieldRef::MakeView(Void).IsOwned()"), VoidView.IsOwned());
-		TestFalse(TEXT("FCbFieldRef::MakeView(Void).IsReadOnly()"), VoidView.IsReadOnly());
 		TestTrue(TEXT("FCbFieldRef::Clone(Void).IsOwned()"), VoidClone.IsOwned());
-		TestTrue(TEXT("FCbFieldRef::Clone(Void).IsReadOnly()"), VoidClone.IsReadOnly());
 		TestFalse(TEXT("FCbFieldRef::MakeView(Field).IsOwned()"), FieldView.IsOwned());
-		TestFalse(TEXT("FCbFieldRef::MakeView(Field).IsReadOnly()"), FieldView.IsReadOnly());
 		TestTrue(TEXT("FCbFieldRef::Clone(Field).IsOwned()"), FieldClone.IsOwned());
-		TestTrue(TEXT("FCbFieldRef::Clone(Field).IsReadOnly()"), FieldClone.IsReadOnly());
 		TestTrue(TEXT("FCbFieldRef::Clone(FieldRef).IsOwned()"), FieldRefClone.IsOwned());
-		TestTrue(TEXT("FCbFieldRef::Clone(FieldRef).IsReadOnly()"), FieldRefClone.IsReadOnly());
 	}
 
 	// Test Field w/o Type
 	{
-		uint8 Payload[] = { 3, 4, 5, 6 }; // Size: 3, Data: 4/5/6
-		FCbField Field(Payload, ECbFieldType::Binary);
+		const uint8 Payload[] = { 3, 4, 5, 6 }; // Size: 3, Data: 4/5/6
+		const FCbField Field(Payload, ECbFieldType::Binary);
 
 		FCbFieldRef FieldView = FCbFieldRef::MakeView(Field);
 		FCbFieldRef FieldClone = FCbFieldRef::Clone(Field);
 		FCbFieldRef FieldRefClone = FCbFieldRef::Clone(FieldView);
 
-		Payload[UE_ARRAY_COUNT(Payload) - 1] = 4;
-
-		TestEqualBytes(TEXT("FCbFieldRef::MakeView(Field, NoType)"), FieldView.AsBinary(), {4, 5, 4});
-		TestEqualBytes(TEXT("FCbFieldRef::Clone(Field, NoType)"), FieldClone.AsBinary(), {4, 5, 6});
-		TestEqualBytes(TEXT("FCbFieldRef::Clone(FieldRef, NoType)"), FieldRefClone.AsBinary(), {4, 5, 6});
+		TestEqual(TEXT("FCbFieldRef::MakeView(Field, NoType)"), FieldView.AsBinary(), MakeMemoryView(Payload).Right(3));
+		TestTrue(TEXT("FCbFieldRef::Clone(Field, NoType)"), FieldClone.AsBinary().EqualBytes(FieldView.AsBinary()));
+		TestTrue(TEXT("FCbFieldRef::Clone(FieldRef, NoType)"), FieldRefClone.AsBinary().EqualBytes(FieldView.AsBinary()));
 
 		TestFalse(TEXT("FCbFieldRef::MakeView(Field, NoType).IsOwned()"), FieldView.IsOwned());
 		TestTrue(TEXT("FCbFieldRef::Clone(Field, NoType).IsOwned()"), FieldClone.IsOwned());
 		TestTrue(TEXT("FCbFieldRef::Clone(FieldRef, NoType).IsOwned()"), FieldRefClone.IsOwned());
 
-		TestFalse(TEXT("FCbFieldRef::MakeView(Field, NoType).IsReadOnly()"), FieldView.IsReadOnly());
-		TestTrue(TEXT("FCbFieldRef::Clone(Field, NoType).IsReadOnly()"), FieldClone.IsReadOnly());
-		TestTrue(TEXT("FCbFieldRef::Clone(FieldRef, NoType).IsReadOnly()"), FieldRefClone.IsReadOnly());
-
 		FieldView.MakeOwned();
-		TestEqualBytes(TEXT("FCbFieldRef::MakeView(NoType).MakeOwned()"), FieldView.AsBinary(), {4, 5, 4});
+		TestTrue(TEXT("FCbFieldRef::MakeView(NoType).MakeOwned()"), FieldView.AsBinary().EqualBytes(MakeMemoryView(Payload).Right(3)));
 		TestTrue(TEXT("FCbFieldRef::MakeView(NoType).MakeOwned().IsOwned()"), FieldView.IsOwned());
-		TestTrue(TEXT("FCbFieldRef::MakeView(NoType).MakeOwned().IsReadOnly()"), FieldView.IsReadOnly());
 	}
 
 	return true;
@@ -1444,13 +1397,13 @@ bool FCbFieldRefTest::RunTest(const FString& Parameters)
 IMPLEMENT_SIMPLE_AUTOMATION_TEST(FCbArrayRefTest, "System.Core.Serialization.CbArrayRef", CompactBinaryTestFlags)
 bool FCbArrayRefTest::RunTest(const FString& Parameters)
 {
-	static_assert(std::is_constructible<FCbArrayRef, const FCbArray&, const FSharedBufferConstPtr&>::value, "Missing constructor for FCbArrayRef");
+	static_assert(std::is_constructible<FCbArrayRef, const FCbArray&, const FSharedBuffer&>::value, "Missing constructor for FCbArrayRef");
 	static_assert(std::is_constructible<FCbArrayRef, const FCbArray&, const FCbFieldRefIterator&>::value, "Missing constructor for FCbArrayRef");
 	static_assert(std::is_constructible<FCbArrayRef, const FCbArray&, const FCbFieldRef&>::value, "Missing constructor for FCbArrayRef");
 	static_assert(std::is_constructible<FCbArrayRef, const FCbArray&, const FCbArrayRef&>::value, "Missing constructor for FCbArrayRef");
 	static_assert(std::is_constructible<FCbArrayRef, const FCbArray&, const FCbObjectRef&>::value, "Missing constructor for FCbArrayRef");
 
-	static_assert(std::is_constructible<FCbArrayRef, const FCbArray&, FSharedBufferConstPtr&&>::value, "Missing constructor for FCbArrayRef");
+	static_assert(std::is_constructible<FCbArrayRef, const FCbArray&, FSharedBuffer&&>::value, "Missing constructor for FCbArrayRef");
 	static_assert(std::is_constructible<FCbArrayRef, const FCbArray&, FCbFieldRefIterator&&>::value, "Missing constructor for FCbArrayRef");
 	static_assert(std::is_constructible<FCbArrayRef, const FCbArray&, FCbFieldRef&&>::value, "Missing constructor for FCbArrayRef");
 	static_assert(std::is_constructible<FCbArrayRef, const FCbArray&, FCbArrayRef&&>::value, "Missing constructor for FCbArrayRef");
@@ -1460,10 +1413,8 @@ bool FCbArrayRefTest::RunTest(const FString& Parameters)
 	{
 		FCbArrayRef DefaultArray;
 		TestFalse(TEXT("FCbArrayRef().IsOwned()"), DefaultArray.IsOwned());
-		TestFalse(TEXT("FCbArrayRef().IsReadOnly()"), DefaultArray.IsReadOnly());
 		DefaultArray.MakeOwned();
 		TestTrue(TEXT("FCbArrayRef().MakeOwned().IsOwned()"), DefaultArray.IsOwned());
-		TestTrue(TEXT("FCbArrayRef().MakeOwned().IsReadOnly()"), DefaultArray.IsReadOnly());
 	}
 
 	return true;
@@ -1474,13 +1425,13 @@ bool FCbArrayRefTest::RunTest(const FString& Parameters)
 IMPLEMENT_SIMPLE_AUTOMATION_TEST(FCbObjectRefTest, "System.Core.Serialization.CbObjectRef", CompactBinaryTestFlags)
 bool FCbObjectRefTest::RunTest(const FString& Parameters)
 {
-	static_assert(std::is_constructible<FCbObjectRef, const FCbObject&, const FSharedBufferConstPtr&>::value, "Missing constructor for FCbObjectRef");
+	static_assert(std::is_constructible<FCbObjectRef, const FCbObject&, const FSharedBuffer&>::value, "Missing constructor for FCbObjectRef");
 	static_assert(std::is_constructible<FCbObjectRef, const FCbObject&, const FCbFieldRefIterator&>::value, "Missing constructor for FCbObjectRef");
 	static_assert(std::is_constructible<FCbObjectRef, const FCbObject&, const FCbFieldRef&>::value, "Missing constructor for FCbObjectRef");
 	static_assert(std::is_constructible<FCbObjectRef, const FCbObject&, const FCbArrayRef&>::value, "Missing constructor for FCbObjectRef");
 	static_assert(std::is_constructible<FCbObjectRef, const FCbObject&, const FCbObjectRef&>::value, "Missing constructor for FCbObjectRef");
 
-	static_assert(std::is_constructible<FCbObjectRef, const FCbObject&, FSharedBufferConstPtr&&>::value, "Missing constructor for FCbObjectRef");
+	static_assert(std::is_constructible<FCbObjectRef, const FCbObject&, FSharedBuffer&&>::value, "Missing constructor for FCbObjectRef");
 	static_assert(std::is_constructible<FCbObjectRef, const FCbObject&, FCbFieldRefIterator&&>::value, "Missing constructor for FCbObjectRef");
 	static_assert(std::is_constructible<FCbObjectRef, const FCbObject&, FCbFieldRef&&>::value, "Missing constructor for FCbObjectRef");
 	static_assert(std::is_constructible<FCbObjectRef, const FCbObject&, FCbArrayRef&&>::value, "Missing constructor for FCbObjectRef");
@@ -1490,10 +1441,8 @@ bool FCbObjectRefTest::RunTest(const FString& Parameters)
 	{
 		FCbObjectRef DefaultObject;
 		TestFalse(TEXT("FCbObjectRef().IsOwned()"), DefaultObject.IsOwned());
-		TestFalse(TEXT("FCbObjectRef().IsReadOnly()"), DefaultObject.IsReadOnly());
 		DefaultObject.MakeOwned();
 		TestTrue(TEXT("FCbObjectRef().MakeOwned().IsOwned()"), DefaultObject.IsOwned());
-		TestTrue(TEXT("FCbObjectRef().MakeOwned().IsReadOnly()"), DefaultObject.IsReadOnly());
 	}
 
 	return true;
@@ -1531,17 +1480,17 @@ bool FCbFieldRefIteratorTest::RunTest(const FString& Parameters)
 		constexpr uint8 T = uint8(ECbFieldType::IntegerPositive);
 		const uint8 Payload[] = { T, 0, T, 1, T, 2, T, 3 };
 
-		const FSharedBufferConstPtr View = FSharedBuffer::MakeView(MakeMemoryView(Payload));
-		const FSharedBufferConstPtr Clone = FSharedBuffer::Clone(MakeMemoryView(Payload));
+		const FSharedBuffer View = FSharedBuffer::MakeView(MakeMemoryView(Payload));
+		const FSharedBuffer Clone = FSharedBuffer::Clone(View);
 
 		const FMemoryView EmptyView;
-		const FSharedBufferConstPtr NullBuffer;
+		const FSharedBuffer NullBuffer;
 
-		const FCbFieldIterator FieldIt = FCbFieldIterator::MakeRange(*View);
+		const FCbFieldIterator FieldIt = FCbFieldIterator::MakeRange(View);
 		const FCbFieldRefIterator FieldRefIt = FCbFieldRefIterator::MakeRange(View);
 
-		TestEqual(TEXT("FCbFieldIterator::GetRangeHash()"), FieldIt.GetRangeHash(), FBlake3::HashBuffer(*View));
-		TestEqual(TEXT("FCbFieldRefIterator::GetRangeHash()"), FieldRefIt.GetRangeHash(), FBlake3::HashBuffer(*View));
+		TestEqual(TEXT("FCbFieldIterator::GetRangeHash()"), FieldIt.GetRangeHash(), FBlake3::HashBuffer(View));
+		TestEqual(TEXT("FCbFieldRefIterator::GetRangeHash()"), FieldRefIt.GetRangeHash(), FBlake3::HashBuffer(View));
 
 		TestEqual(TEXT("FCbFieldIterator::GetRangeView()"), FieldIt.GetRangeView(), MakeMemoryView(Payload));
 		TestEqual(TEXT("FCbFieldRefIterator::GetRangeView()"), FieldRefIt.GetRangeView(), MakeMemoryView(Payload));
@@ -1560,18 +1509,18 @@ bool FCbFieldRefIteratorTest::RunTest(const FString& Parameters)
 
 		TestEqual(TEXT("FCbFieldIterator::MakeRange(EmptyView)"), GetCount(FCbFieldIterator::MakeRange(EmptyView)), 0);
 		TestEqual(TEXT("FCbFieldRefIterator::MakeRange(BufferNullL)"), GetCount(FCbFieldRefIterator::MakeRange(NullBuffer)), 0);
-		TestEqual(TEXT("FCbFieldRefIterator::MakeRange(BufferNullR)"), GetCount(FCbFieldRefIterator::MakeRange(FSharedBufferConstPtr(NullBuffer))), 0);
+		TestEqual(TEXT("FCbFieldRefIterator::MakeRange(BufferNullR)"), GetCount(FCbFieldRefIterator::MakeRange(FSharedBuffer(NullBuffer))), 0);
 
 		TestEqual(TEXT("FCbFieldIterator::MakeRange(BufferView)"), GetCount(FCbFieldIterator::MakeRange(MakeMemoryView(Payload))), 4);
 		TestEqual(TEXT("FCbFieldRefIterator::MakeRange(BufferCloneL)"), GetCount(FCbFieldRefIterator::MakeRange(Clone)), 4);
-		TestEqual(TEXT("FCbFieldRefIterator::MakeRange(BufferCloneR)"), GetCount(FCbFieldRefIterator::MakeRange(FSharedBufferConstPtr(Clone))), 4);
+		TestEqual(TEXT("FCbFieldRefIterator::MakeRange(BufferCloneR)"), GetCount(FCbFieldRefIterator::MakeRange(FSharedBuffer(Clone))), 4);
 
-		TestEqual(TEXT("FCbFieldRefIterator::MakeRangeView(FieldIt, BufferNullL)"), GetCount(FCbFieldRefIterator::MakeRangeView(FCbFieldIterator::MakeRange(*View), NullBuffer)), 4);
-		TestEqual(TEXT("FCbFieldRefIterator::MakeRangeView(FieldIt, BufferNullR)"), GetCount(FCbFieldRefIterator::MakeRangeView(FCbFieldIterator::MakeRange(*View), FSharedBufferConstPtr(NullBuffer))), 4);
-		TestEqual(TEXT("FCbFieldRefIterator::MakeRangeView(FieldIt, BufferViewL)"), GetCount(FCbFieldRefIterator::MakeRangeView(FCbFieldIterator::MakeRange(*View), View)), 4);
-		TestEqual(TEXT("FCbFieldRefIterator::MakeRangeView(FieldIt, BufferViewR)"), GetCount(FCbFieldRefIterator::MakeRangeView(FCbFieldIterator::MakeRange(*View), FSharedBufferConstPtr(View))), 4);
-		TestEqual(TEXT("FCbFieldRefIterator::MakeRangeView(FieldIt, BufferCloneL)"), GetCount(FCbFieldRefIterator::MakeRangeView(FCbFieldIterator::MakeRange(*Clone), Clone)), 4);
-		TestEqual(TEXT("FCbFieldRefIterator::MakeRangeView(FieldIt, BufferCloneR)"), GetCount(FCbFieldRefIterator::MakeRangeView(FCbFieldIterator::MakeRange(*Clone), FSharedBufferConstPtr(Clone))), 4);
+		TestEqual(TEXT("FCbFieldRefIterator::MakeRangeView(FieldIt, BufferNullL)"), GetCount(FCbFieldRefIterator::MakeRangeView(FCbFieldIterator::MakeRange(View), NullBuffer)), 4);
+		TestEqual(TEXT("FCbFieldRefIterator::MakeRangeView(FieldIt, BufferNullR)"), GetCount(FCbFieldRefIterator::MakeRangeView(FCbFieldIterator::MakeRange(View), FSharedBuffer(NullBuffer))), 4);
+		TestEqual(TEXT("FCbFieldRefIterator::MakeRangeView(FieldIt, BufferViewL)"), GetCount(FCbFieldRefIterator::MakeRangeView(FCbFieldIterator::MakeRange(View), View)), 4);
+		TestEqual(TEXT("FCbFieldRefIterator::MakeRangeView(FieldIt, BufferViewR)"), GetCount(FCbFieldRefIterator::MakeRangeView(FCbFieldIterator::MakeRange(View), FSharedBuffer(View))), 4);
+		TestEqual(TEXT("FCbFieldRefIterator::MakeRangeView(FieldIt, BufferCloneL)"), GetCount(FCbFieldRefIterator::MakeRangeView(FCbFieldIterator::MakeRange(Clone), Clone)), 4);
+		TestEqual(TEXT("FCbFieldRefIterator::MakeRangeView(FieldIt, BufferCloneR)"), GetCount(FCbFieldRefIterator::MakeRangeView(FCbFieldIterator::MakeRange(Clone), FSharedBuffer(Clone))), 4);
 
 		TestEqual(TEXT("FCbFieldIterator(FieldRefItL)"), GetCount(FCbFieldIterator(FieldRefIt)), 4);
 		TestEqual(TEXT("FCbFieldIterator(FieldRefItR)"), GetCount(FCbFieldIterator(FCbFieldRefIterator(FieldRefIt))), 4);
@@ -1582,7 +1531,7 @@ bool FCbFieldRefIteratorTest::RunTest(const FString& Parameters)
 		TestEqual(TEXT("FCbFieldIterator::MakeRange(Uniform).GetRangeHash()"), UniformFieldIt.GetRangeHash(), FieldIt.GetRangeHash());
 		TestEqual(TEXT("FCbFieldIterator::MakeRange(Uniform).GetRangeView()"), UniformFieldIt.GetRangeView(), MakeMemoryView(UniformPayload));
 		TestFalse(TEXT("FCbFieldIterator::MakeRange(Uniform).TryGetSerializedRangeView()"), UniformFieldIt.TryGetSerializedRangeView(SerializedView));
-		const FSharedBufferConstPtr UniformView = FSharedBuffer::MakeView(MakeMemoryView(UniformPayload));
+		const FSharedBuffer UniformView = FSharedBuffer::MakeView(MakeMemoryView(UniformPayload));
 		const FCbFieldRefIterator UniformFieldRefIt = FCbFieldRefIterator::MakeRange(UniformView, ECbFieldType::IntegerPositive);
 		TestEqual(TEXT("FCbFieldRefIterator::MakeRange(Uniform).GetRangeHash()"), UniformFieldRefIt.GetRangeHash(), FieldIt.GetRangeHash());
 		TestEqual(TEXT("FCbFieldRefIterator::MakeRange(Uniform).GetRangeView()"), UniformFieldRefIt.GetRangeView(), MakeMemoryView(UniformPayload));
@@ -1629,9 +1578,9 @@ bool FCbFieldRefIteratorTest::RunTest(const FString& Parameters)
 		TestNotEqual(TEXT("FCbFieldRefIterator::GetRangeBuffer(Short)"), UniformShortFieldRefIt.GetRangeBuffer(), UniformView);
 
 		// These lines are expected to assert when uncommented.
-		//const FSharedBufferConstPtr ShortView = FSharedBuffer::MakeView(MakeMemoryView(Payload).LeftChop(2));
+		//const FSharedBuffer ShortView = FSharedBuffer::MakeView(MakeMemoryView(Payload).LeftChop(2));
 		//TestEqual(TEXT("FCbFieldRefIterator::MakeRangeView(FieldIt, InvalidBufferL)"), GetCount(FCbFieldRefIterator::MakeRangeView(FCbFieldIterator::MakeRange(*View), ShortView)), 4);
-		//TestEqual(TEXT("FCbFieldRefIterator::MakeRangeView(FieldIt, InvalidBufferR)"), GetCount(FCbFieldRefIterator::MakeRangeView(FCbFieldIterator::MakeRange(*View), FSharedBufferConstPtr(ShortView))), 4);
+		//TestEqual(TEXT("FCbFieldRefIterator::MakeRangeView(FieldIt, InvalidBufferR)"), GetCount(FCbFieldRefIterator::MakeRangeView(FCbFieldIterator::MakeRange(*View), FSharedBuffer(ShortView))), 4);
 	}
 
 	// Test FCbField[Ref]Iterator(Scalar)
@@ -1639,8 +1588,8 @@ bool FCbFieldRefIteratorTest::RunTest(const FString& Parameters)
 		constexpr uint8 T = uint8(ECbFieldType::IntegerPositive);
 		const uint8 Payload[] = { T, 0 };
 
-		const FSharedBufferConstPtr View = FSharedBuffer::MakeView(MakeMemoryView(Payload));
-		const FSharedBufferConstPtr Clone = FSharedBuffer::Clone(MakeMemoryView(Payload));
+		const FSharedBuffer View = FSharedBuffer::MakeView(MakeMemoryView(Payload));
+		const FSharedBuffer Clone = FSharedBuffer::Clone(View);
 
 		const FCbField Field(Payload);
 		const FCbFieldRef FieldRef(View);
diff --git a/Engine/Source/Runtime/Core/Private/Tests/Serialization/CompactBinaryValidationTest.cpp b/Engine/Source/Runtime/Core/Private/Tests/Serialization/CompactBinaryValidationTest.cpp
--- a/Engine/Source/Runtime/Core/Private/Tests/Serialization/CompactBinaryValidationTest.cpp
+++ b/Engine/Source/Runtime/Core/Private/Tests/Serialization/CompactBinaryValidationTest.cpp
@@ -287,6 +287,9 @@ bool FCbValidateAttachmentTest::RunTest(const FString& Parameters)
 		return ValidateCompactBinaryAttachment(MakeMemoryView(Data), Mode);
 	};
 
+	const uint8 BinaryValue[] = { 0, 1, 2, 3 };
+	const FMemoryView BinaryView = MakeMemoryView(BinaryValue);
+
 	// Test Null
 	{
 		TestEqual(TEXT("ValidateCompactBinaryAttachment(Null)"), Validate({uint8(ECbFieldType::Binary), 0}), ECbValidateError::None);
@@ -348,8 +351,8 @@ bool FCbValidateAttachmentTest::RunTest(const FString& Parameters)
 	{
 		FBufferArchive Buffer;
 		FCbWriter Writer;
-		Writer.Reference(FBlake3::HashBuffer(MakeMemoryView<uint8>({0, 1, 2, 3})));
-		Writer.Binary(MakeMemoryView<uint8>({0, 1, 2, 3}));
+		Writer.Reference(FBlake3::HashBuffer(BinaryView));
+		Writer.Binary(BinaryView);
 		Writer.Save(Buffer);
 		TestEqual(TEXT("ValidateCompactBinaryAttachment(HashBeforeValue)"), ValidateCompactBinaryAttachment(MakeMemoryView(Buffer), ECbValidateMode::All), ECbValidateError::InvalidPackageFormat | ECbValidateError::Padding);
 		TestEqual(TEXT("ValidateCompactBinaryAttachment(HashBeforeValue, Mode)"), ValidateCompactBinaryAttachment(MakeMemoryView(Buffer), ECbValidateMode::All & ~ECbValidateMode::Package), ECbValidateError::Padding);
@@ -357,8 +360,8 @@ bool FCbValidateAttachmentTest::RunTest(const FString& Parameters)
 	{
 		FBufferArchive Buffer;
 		FCbWriter Writer;
-		Writer.Name("Name"_ASV).Binary(MakeMemoryView<uint8>({0, 1, 2, 3}));
-		Writer.Reference(FBlake3::HashBuffer(MakeMemoryView<uint8>({0, 1, 2, 3})));
+		Writer.Name("Name"_ASV).Binary(BinaryView);
+		Writer.Reference(FBlake3::HashBuffer(BinaryView));
 		Writer.Save(Buffer);
 		TestEqual(TEXT("ValidateCompactBinaryAttachment(NameOnValue)"), ValidateCompactBinaryAttachment(MakeMemoryView(Buffer), ECbValidateMode::All), ECbValidateError::InvalidPackageFormat);
 		TestEqual(TEXT("ValidateCompactBinaryAttachment(NameOnValue, Mode)"), ValidateCompactBinaryAttachment(MakeMemoryView(Buffer), ECbValidateMode::All & ~ECbValidateMode::Package), ECbValidateError::None);
@@ -366,8 +369,8 @@ bool FCbValidateAttachmentTest::RunTest(const FString& Parameters)
 	{
 		FBufferArchive Buffer;
 		FCbWriter Writer;
-		Writer.Binary(MakeMemoryView<uint8>({0, 1, 2, 3}));
-		Writer.Name("Name"_ASV).Reference(FBlake3::HashBuffer(MakeMemoryView<uint8>({0, 1, 2, 3})));
+		Writer.Binary(BinaryView);
+		Writer.Name("Name"_ASV).Reference(FBlake3::HashBuffer(BinaryView));
 		Writer.Save(Buffer);
 		TestEqual(TEXT("ValidateCompactBinaryAttachment(NameOnHash)"), ValidateCompactBinaryAttachment(MakeMemoryView(Buffer), ECbValidateMode::All), ECbValidateError::InvalidPackageFormat);
 		TestEqual(TEXT("ValidateCompactBinaryAttachment(NameOnHash, Mode)"), ValidateCompactBinaryAttachment(MakeMemoryView(Buffer), ECbValidateMode::All & ~ECbValidateMode::Package), ECbValidateError::None);
@@ -381,6 +384,9 @@ bool FCbValidateAttachmentTest::RunTest(const FString& Parameters)
 IMPLEMENT_SIMPLE_AUTOMATION_TEST(FCbValidatePackageTest, "System.Core.Serialization.ValidateCompactBinaryPackage", CompactBinaryValidationTestFlags)
 bool FCbValidatePackageTest::RunTest(const FString& Parameters)
 {
+	const uint8 BinaryValue[] = { 0, 1, 2, 3 };
+	const FMemoryView BinaryView = MakeMemoryView(BinaryValue);
+
 	// Test Null
 	{
 		TestEqual(TEXT("ValidateCompactBinaryPackage(Null)"), ValidateCompactBinaryPackage(MakeMemoryView<uint8>({uint8(ECbFieldType::Null)}), ECbValidateMode::All), ECbValidateError::None);
@@ -425,8 +431,8 @@ bool FCbValidatePackageTest::RunTest(const FString& Parameters)
 		Writer.EndObject();
 		Writer.Reference(Writer.Save().GetHash());
 		const uint64 AttachmentOffset = Writer.GetSaveSize();
-		Writer.Binary(MakeMemoryView<uint8>({0, 1, 2, 3}));
-		Writer.Reference(FBlake3::HashBuffer(MakeMemoryView<uint8>({0, 1, 2, 3})));
+		Writer.Binary(BinaryView);
+		Writer.Reference(FBlake3::HashBuffer(BinaryView));
 		Writer.Null();
 		Writer.Save(Buffer);
 		TestEqual(TEXT("ValidateCompactBinaryPackage(Object+Attachment)"), ValidateCompactBinaryPackage(MakeMemoryView(Buffer), ECbValidateMode::All), ECbValidateError::None);
@@ -486,8 +492,8 @@ bool FCbValidatePackageTest::RunTest(const FString& Parameters)
 		Writer << "Field" << 42;
 		Writer.EndObject();
 		Writer.Reference(Writer.Save().GetHash());
-		Writer.Name("Name"_ASV).Binary(MakeMemoryView<uint8>({0, 1, 2, 3}));
-		Writer.Reference(FBlake3::HashBuffer(MakeMemoryView<uint8>({0, 1, 2, 3})));
+		Writer.Name("Name"_ASV).Binary(BinaryView);
+		Writer.Reference(FBlake3::HashBuffer(BinaryView));
 		Writer.Null();
 		Writer.Save(Buffer);
 		TestEqual(TEXT("ValidateCompactBinaryPackage(NameOnAttachment)"), ValidateCompactBinaryPackage(MakeMemoryView(Buffer), ECbValidateMode::All), ECbValidateError::InvalidPackageFormat);
@@ -513,10 +519,10 @@ bool FCbValidatePackageTest::RunTest(const FString& Parameters)
 	{
 		FBufferArchive Buffer;
 		TCbWriter<256> Writer;
-		Writer.Binary(MakeMemoryView<uint8>({0, 1, 2, 3}));
-		Writer.Reference(FBlake3::HashBuffer(MakeMemoryView<uint8>({0, 1, 2, 3})));
-		Writer.Binary(MakeMemoryView<uint8>({0, 1, 2, 3}));
-		Writer.Reference(FBlake3::HashBuffer(MakeMemoryView<uint8>({0, 1, 2, 3})));
+		Writer.Binary(BinaryView);
+		Writer.Reference(FBlake3::HashBuffer(BinaryView));
+		Writer.Binary(BinaryView);
+		Writer.Reference(FBlake3::HashBuffer(BinaryView));
 		Writer.Null();
 		Writer.Save(Buffer);
 		TestEqual(TEXT("ValidateCompactBinaryPackage(DuplicateAttachments)"), ValidateCompactBinaryPackage(MakeMemoryView(Buffer), ECbValidateMode::All), ECbValidateError::DuplicateAttachments);
diff --git a/Engine/Source/Runtime/Core/Private/Tests/Serialization/CompactBinaryWriterTest.cpp b/Engine/Source/Runtime/Core/Private/Tests/Serialization/CompactBinaryWriterTest.cpp
--- a/Engine/Source/Runtime/Core/Private/Tests/Serialization/CompactBinaryWriterTest.cpp
+++ b/Engine/Source/Runtime/Core/Private/Tests/Serialization/CompactBinaryWriterTest.cpp
@@ -28,7 +28,7 @@ bool FCbWriterObjectTest::RunTest(const FString& Parameters)
 		Writer.BeginObject();
 		Writer.EndObject();
 		FCbFieldRef Field = Writer.Save();
-		if (TestEqual(TEXT("FCbWriter(Object, Empty) Validate"), ValidateCompactBinary(*Field.GetBuffer(), ECbValidateMode::All), ECbValidateError::None))
+		if (TestEqual(TEXT("FCbWriter(Object, Empty) Validate"), ValidateCompactBinary(Field.GetBuffer(), ECbValidateMode::All), ECbValidateError::None))
 		{
 			TestTrue(TEXT("FCbWriter(Object, Empty).IsObject()"), Field.IsObject());
 			TestFalse(TEXT("FCbWriter(Object, Empty).AsObject()"), Field.AsObject().CreateIterator().HasValue());
@@ -42,7 +42,7 @@ bool FCbWriterObjectTest::RunTest(const FString& Parameters)
 		Writer.BeginObject();
 		Writer.EndObject();
 		FCbFieldRef Field = Writer.Save();
-		if (TestEqual(TEXT("FCbWriter(Object, Empty, Name) Validate"), ValidateCompactBinary(*Field.GetBuffer(), ECbValidateMode::All), ECbValidateError::None))
+		if (TestEqual(TEXT("FCbWriter(Object, Empty, Name) Validate"), ValidateCompactBinary(Field.GetBuffer(), ECbValidateMode::All), ECbValidateError::None))
 		{
 			TestTrue(TEXT("FCbWriter(Object, Empty, Name).IsObject()"), Field.IsObject());
 			TestFalse(TEXT("FCbWriter(Object, Empty, Name).AsObject()"), Field.AsObject().CreateIterator().HasValue());
@@ -57,7 +57,7 @@ bool FCbWriterObjectTest::RunTest(const FString& Parameters)
 		Writer.Name("Float"_ASV).Float(0.0f);
 		Writer.EndObject();
 		FCbFieldRef Field = Writer.Save();
-		if (TestEqual(TEXT("FCbWriter(Object, Basic) Validate"), ValidateCompactBinary(*Field.GetBuffer(), ECbValidateMode::All), ECbValidateError::None))
+		if (TestEqual(TEXT("FCbWriter(Object, Basic) Validate"), ValidateCompactBinary(Field.GetBuffer(), ECbValidateMode::All), ECbValidateError::None))
 		{
 			TestTrue(TEXT("FCbWriter(Object, Basic).IsObject()"), Field.IsObject());
 			FCbObject Object = Field.AsObject();
@@ -74,7 +74,7 @@ bool FCbWriterObjectTest::RunTest(const FString& Parameters)
 		Writer.Name("Field2"_ASV).Integer(1);
 		Writer.EndObject();
 		FCbFieldRef Field = Writer.Save();
-		if (TestEqual(TEXT("FCbWriter(Object, Uniform) Validate"), ValidateCompactBinary(*Field.GetBuffer(), ECbValidateMode::All), ECbValidateError::None))
+		if (TestEqual(TEXT("FCbWriter(Object, Uniform) Validate"), ValidateCompactBinary(Field.GetBuffer(), ECbValidateMode::All), ECbValidateError::None))
 		{
 			TestTrue(TEXT("FCbWriter(Object, Uniform).IsObject()"), Field.IsObject());
 			FCbObject Object = Field.AsObject();
@@ -99,7 +99,7 @@ bool FCbWriterArrayTest::RunTest(const FString& Parameters)
 		Writer.BeginArray();
 		Writer.EndArray();
 		FCbFieldRef Field = Writer.Save();
-		if (TestEqual(TEXT("FCbWriter(Array, Empty) Validate"), ValidateCompactBinary(*Field.GetBuffer(), ECbValidateMode::All), ECbValidateError::None))
+		if (TestEqual(TEXT("FCbWriter(Array, Empty) Validate"), ValidateCompactBinary(Field.GetBuffer(), ECbValidateMode::All), ECbValidateError::None))
 		{
 			TestTrue(TEXT("FCbWriter(Array, Empty).IsArray()"), Field.IsArray());
 			TestEqual(TEXT("FCbWriter(Array, Empty).AsArray()"), Field.AsArray().Num(), uint64(0));
@@ -113,7 +113,7 @@ bool FCbWriterArrayTest::RunTest(const FString& Parameters)
 		Writer.BeginArray();
 		Writer.EndArray();
 		FCbFieldRef Field = Writer.Save();
-		if (TestEqual(TEXT("FCbWriter(Array, Empty, Name) Validate"), ValidateCompactBinary(*Field.GetBuffer(), ECbValidateMode::All), ECbValidateError::None))
+		if (TestEqual(TEXT("FCbWriter(Array, Empty, Name) Validate"), ValidateCompactBinary(Field.GetBuffer(), ECbValidateMode::All), ECbValidateError::None))
 		{
 			TestTrue(TEXT("FCbWriter(Array, Empty, Name).IsArray()"), Field.IsArray());
 			TestEqual(TEXT("FCbWriter(Array, Empty, Name).AsArray()"), Field.AsArray().Num(), uint64(0));
@@ -128,7 +128,7 @@ bool FCbWriterArrayTest::RunTest(const FString& Parameters)
 		Writer.Float(0.0f);
 		Writer.EndArray();
 		FCbFieldRef Field = Writer.Save();
-		if (TestEqual(TEXT("FCbWriter(Array, Basic) Validate"), ValidateCompactBinary(*Field.GetBuffer(), ECbValidateMode::All), ECbValidateError::None))
+		if (TestEqual(TEXT("FCbWriter(Array, Basic) Validate"), ValidateCompactBinary(Field.GetBuffer(), ECbValidateMode::All), ECbValidateError::None))
 		{
 			TestTrue(TEXT("FCbWriter(Array, Basic).IsArray()"), Field.IsArray());
 			FCbFieldIterator Iterator = Field.AsArray().CreateIterator();
@@ -148,7 +148,7 @@ bool FCbWriterArrayTest::RunTest(const FString& Parameters)
 		Writer.Integer(1);
 		Writer.EndArray();
 		FCbFieldRef Field = Writer.Save();
-		if (TestEqual(TEXT("FCbWriter(Array, Uniform) Validate"), ValidateCompactBinary(*Field.GetBuffer(), ECbValidateMode::All), ECbValidateError::None))
+		if (TestEqual(TEXT("FCbWriter(Array, Uniform) Validate"), ValidateCompactBinary(Field.GetBuffer(), ECbValidateMode::All), ECbValidateError::None))
 		{
 			TestTrue(TEXT("FCbWriter(Array, Uniform).IsArray()"), Field.IsArray());
 			FCbFieldIterator Iterator = Field.AsArray().CreateIterator();
@@ -175,7 +175,7 @@ bool FCbWriterNullTest::RunTest(const FString& Parameters)
 		Writer.Reset();
 		Writer.Null();
 		FCbFieldRef Field = Writer.Save();
-		if (TestEqual(TEXT("FCbWriter(Null) Validate"), ValidateCompactBinary(*Field.GetBuffer(), ECbValidateMode::All), ECbValidateError::None))
+		if (TestEqual(TEXT("FCbWriter(Null) Validate"), ValidateCompactBinary(Field.GetBuffer(), ECbValidateMode::All), ECbValidateError::None))
 		{
 			TestFalse(TEXT("FCbWriter(Null).HasName()"), Field.HasName());
 			TestTrue(TEXT("FCbWriter(Null).IsNull()"), Field.IsNull());
@@ -188,7 +188,7 @@ bool FCbWriterNullTest::RunTest(const FString& Parameters)
 		Writer.Name("Null"_ASV);
 		Writer.Null();
 		FCbFieldRef Field = Writer.Save();
-		if (TestEqual(TEXT("FCbWriter(Null, Name) Validate"), ValidateCompactBinary(*Field.GetBuffer(), ECbValidateMode::All), ECbValidateError::None))
+		if (TestEqual(TEXT("FCbWriter(Null, Name) Validate"), ValidateCompactBinary(Field.GetBuffer(), ECbValidateMode::All), ECbValidateError::None))
 		{
 			TestEqual(TEXT("FCbWriter(Null, Name).GetName()"), Field.GetName(), "Null"_ASV);
 			TestTrue(TEXT("FCbWriter(Null, Name).HasName()"), Field.HasName());
@@ -213,7 +213,7 @@ bool FCbWriterNullTest::RunTest(const FString& Parameters)
 		Writer.EndObject();
 
 		FCbFieldRefIterator Fields = Writer.Save();
-		TestEqual(TEXT("FCbWriter(Null, Uniform) Validate"), ValidateCompactBinaryRange(*Fields.GetBuffer(), ECbValidateMode::All), ECbValidateError::None);
+		TestEqual(TEXT("FCbWriter(Null, Uniform) Validate"), ValidateCompactBinaryRange(Fields.GetBuffer(), ECbValidateMode::All), ECbValidateError::None);
 	}
 
 	// Test Null with Save(Buffer)
@@ -265,7 +265,7 @@ bool FCbWriterBinaryTest::RunTest(const FString& Parameters)
 		Writer.Reset();
 		Writer.Binary(nullptr, 0);
 		FCbFieldRef Field = Writer.Save();
-		if (TestEqual(TEXT("FCbWriter(Binary, Empty) Validate"), ValidateCompactBinary(*Field.GetBuffer(), ECbValidateMode::All), ECbValidateError::None))
+		if (TestEqual(TEXT("FCbWriter(Binary, Empty) Validate"), ValidateCompactBinary(Field.GetBuffer(), ECbValidateMode::All), ECbValidateError::None))
 		{
 			TestFalse(TEXT("FCbWriter(Binary, Empty).HasName()"), Field.HasName());
 			TestTrue(TEXT("FCbWriter(Binary, Empty).IsBinary()"), Field.IsBinary());
@@ -280,7 +280,7 @@ bool FCbWriterBinaryTest::RunTest(const FString& Parameters)
 		Writer.Name("Binary"_ASV);
 		Writer.Binary(BinaryValue, sizeof(BinaryValue));
 		FCbFieldRef Field = Writer.Save();
-		if (TestEqual(TEXT("FCbWriter(Binary, Array) Validate"), ValidateCompactBinary(*Field.GetBuffer(), ECbValidateMode::All), ECbValidateError::None))
+		if (TestEqual(TEXT("FCbWriter(Binary, Array) Validate"), ValidateCompactBinary(Field.GetBuffer(), ECbValidateMode::All), ECbValidateError::None))
 		{
 			TestEqual(TEXT("FCbWriter(Binary, Array).GetName()"), Field.GetName(), "Binary"_ASV);
 			TestTrue(TEXT("FCbWriter(Binary, Array).HasName()"), Field.HasName());
@@ -305,7 +305,7 @@ bool FCbWriterStringTest::RunTest(const FString& Parameters)
 		Writer.String(FAnsiStringView());
 		Writer.String(FWideStringView());
 		FCbFieldRefIterator Fields = Writer.Save();
-		if (TestEqual(TEXT("FCbWriter(String, Empty) Validate"), ValidateCompactBinaryRange(*Fields.GetBuffer(), ECbValidateMode::All), ECbValidateError::None))
+		if (TestEqual(TEXT("FCbWriter(String, Empty) Validate"), ValidateCompactBinaryRange(Fields.GetBuffer(), ECbValidateMode::All), ECbValidateError::None))
 		{
 			for (FCbField Field : Fields)
 			{
@@ -322,7 +322,7 @@ bool FCbWriterStringTest::RunTest(const FString& Parameters)
 		Writer.Name("String"_ASV).String("Value"_ASV);
 		Writer.Name("String"_ASV).String(TEXT("Value"_SV));
 		FCbFieldRefIterator Fields = Writer.Save();
-		if (TestEqual(TEXT("FCbWriter(String, Basic) Validate"), ValidateCompactBinaryRange(*Fields.GetBuffer(), ECbValidateMode::All), ECbValidateError::None))
+		if (TestEqual(TEXT("FCbWriter(String, Basic) Validate"), ValidateCompactBinaryRange(Fields.GetBuffer(), ECbValidateMode::All), ECbValidateError::None))
 		{
 			for (FCbField Field : Fields)
 			{
@@ -346,7 +346,7 @@ bool FCbWriterStringTest::RunTest(const FString& Parameters)
 		Writer.String(Dots);
 		Writer.String(FString::ChrN(DotCount, TEXT('.')));
 		FCbFieldRefIterator Fields = Writer.Save();
-		if (TestEqual(TEXT("FCbWriter(String, Long) Validate"), ValidateCompactBinaryRange(*Fields.GetBuffer(), ECbValidateMode::All), ECbValidateError::None))
+		if (TestEqual(TEXT("FCbWriter(String, Long) Validate"), ValidateCompactBinaryRange(Fields.GetBuffer(), ECbValidateMode::All), ECbValidateError::None))
 		{
 			for (FCbField Field : Fields)
 			{
@@ -362,7 +362,7 @@ bool FCbWriterStringTest::RunTest(const FString& Parameters)
 		Writer.String("\xf0\x9f\x98\x80"_ASV);
 		Writer.String(FWideStringView(Value, UE_ARRAY_COUNT(Value)));
 		FCbFieldRefIterator Fields = Writer.Save();
-		if (TestEqual(TEXT("FCbWriter(String, Unicode) Validate"), ValidateCompactBinaryRange(*Fields.GetBuffer(), ECbValidateMode::All), ECbValidateError::None))
+		if (TestEqual(TEXT("FCbWriter(String, Unicode) Validate"), ValidateCompactBinaryRange(Fields.GetBuffer(), ECbValidateMode::All), ECbValidateError::None))
 		{
 			for (FCbField Field : Fields)
 			{
@@ -386,7 +386,7 @@ bool FCbWriterIntegerTest::RunTest(const FString& Parameters)
 		Writer.Reset();
 		Writer.Integer(Value);
 		FCbFieldRef Field = Writer.Save();
-		if (TestEqual(TEXT("FCbWriter(Integer, Int32) Validate"), ValidateCompactBinary(*Field.GetBuffer(), ECbValidateMode::All), ECbValidateError::None))
+		if (TestEqual(TEXT("FCbWriter(Integer, Int32) Validate"), ValidateCompactBinary(Field.GetBuffer(), ECbValidateMode::All), ECbValidateError::None))
 		{
 			TestEqual(TEXT("FCbWriter(Integer, Int32) Value"), Field.AsInt32(), Value);
 			TestFalse(TEXT("FCbWriter(Integer, Int32) Error"), Field.HasError());
@@ -398,7 +398,7 @@ bool FCbWriterIntegerTest::RunTest(const FString& Parameters)
 		Writer.Reset();
 		Writer.Integer(Value);
 		FCbFieldRef Field = Writer.Save();
-		if (TestEqual(TEXT("FCbWriter(Integer, UInt32) Validate"), ValidateCompactBinary(*Field.GetBuffer(), ECbValidateMode::All), ECbValidateError::None))
+		if (TestEqual(TEXT("FCbWriter(Integer, UInt32) Validate"), ValidateCompactBinary(Field.GetBuffer(), ECbValidateMode::All), ECbValidateError::None))
 		{
 			TestEqual(TEXT("FCbWriter(Integer, UInt32) Value"), Field.AsUInt32(), Value);
 			TestFalse(TEXT("FCbWriter(Integer, UInt32) Error"), Field.HasError());
@@ -410,7 +410,7 @@ bool FCbWriterIntegerTest::RunTest(const FString& Parameters)
 		Writer.Reset();
 		Writer.Integer(Value);
 		FCbFieldRef Field = Writer.Save();
-		if (TestEqual(TEXT("FCbWriter(Integer, Int64) Validate"), ValidateCompactBinary(*Field.GetBuffer(), ECbValidateMode::All), ECbValidateError::None))
+		if (TestEqual(TEXT("FCbWriter(Integer, Int64) Validate"), ValidateCompactBinary(Field.GetBuffer(), ECbValidateMode::All), ECbValidateError::None))
 		{
 			TestEqual(TEXT("FCbWriter(Integer, Int64) Value"), Field.AsInt64(), Value);
 			TestFalse(TEXT("FCbWriter(Integer, Int64) Error"), Field.HasError());
@@ -422,7 +422,7 @@ bool FCbWriterIntegerTest::RunTest(const FString& Parameters)
 		Writer.Reset();
 		Writer.Integer(Value);
 		FCbFieldRef Field = Writer.Save();
-		if (TestEqual(TEXT("FCbWriter(Integer, UInt64) Validate"), ValidateCompactBinary(*Field.GetBuffer(), ECbValidateMode::All), ECbValidateError::None))
+		if (TestEqual(TEXT("FCbWriter(Integer, UInt64) Validate"), ValidateCompactBinary(Field.GetBuffer(), ECbValidateMode::All), ECbValidateError::None))
 		{
 			TestEqual(TEXT("FCbWriter(Integer, UInt64) Value"), Field.AsUInt64(), Value);
 			TestFalse(TEXT("FCbWriter(Integer, UInt64) Error"), Field.HasError());
@@ -481,7 +481,7 @@ bool FCbWriterFloatTest::RunTest(const FString& Parameters)
 			Writer.Float(Value);
 		}
 		FCbFieldRefIterator Fields = Writer.Save();
-		if (TestEqual(TEXT("FCbWriter(Float, Single) Validate"), ValidateCompactBinaryRange(*Fields.GetBuffer(), ECbValidateMode::All), ECbValidateError::None))
+		if (TestEqual(TEXT("FCbWriter(Float, Single) Validate"), ValidateCompactBinaryRange(Fields.GetBuffer(), ECbValidateMode::All), ECbValidateError::None))
 		{
 			const float* CheckValue = Values;
 			for (FCbField Field : Fields)
@@ -501,7 +501,7 @@ bool FCbWriterFloatTest::RunTest(const FString& Parameters)
 			Writer.Float(Value);
 		}
 		FCbFieldRefIterator Fields = Writer.Save();
-		if (TestEqual(TEXT("FCbWriter(Float, Double) Validate"), ValidateCompactBinaryRange(*Fields.GetBuffer(), ECbValidateMode::All), ECbValidateError::None))
+		if (TestEqual(TEXT("FCbWriter(Float, Double) Validate"), ValidateCompactBinaryRange(Fields.GetBuffer(), ECbValidateMode::All), ECbValidateError::None))
 		{
 			const double* CheckValue = Values;
 			for (FCbField Field : Fields)
@@ -528,7 +528,7 @@ bool FCbWriterBoolTest::RunTest(const FString& Parameters)
 		Writer.Bool(false);
 
 		FCbFieldRefIterator Fields = Writer.Save();
-		if (TestEqual(TEXT("FCbWriter(Bool) Validate"), ValidateCompactBinaryRange(*Fields.GetBuffer(), ECbValidateMode::All), ECbValidateError::None))
+		if (TestEqual(TEXT("FCbWriter(Bool) Validate"), ValidateCompactBinaryRange(Fields.GetBuffer(), ECbValidateMode::All), ECbValidateError::None))
 		{
 			TestTrue(TEXT("FCbWriter(Bool).AsBool()"), Fields.AsBool());
 			TestFalse(TEXT("FCbWriter(Bool) Error"), Fields.HasError());
@@ -555,7 +555,7 @@ bool FCbWriterBoolTest::RunTest(const FString& Parameters)
 		Writer.EndObject();
 
 		FCbFieldRefIterator Fields = Writer.Save();
-		TestEqual(TEXT("FCbWriter(Bool, Uniform) Validate"), ValidateCompactBinaryRange(*Fields.GetBuffer(), ECbValidateMode::All), ECbValidateError::None);
+		TestEqual(TEXT("FCbWriter(Bool, Uniform) Validate"), ValidateCompactBinaryRange(Fields.GetBuffer(), ECbValidateMode::All), ECbValidateError::None);
 	}
 
 	return true;
@@ -578,7 +578,7 @@ bool FCbWriterReferenceTest::RunTest(const FString& Parameters)
 	}
 
 	FCbFieldRefIterator Fields = Writer.Save();
-	if (TestEqual(TEXT("FCbWriter(Reference) Validate"), ValidateCompactBinaryRange(*Fields.GetBuffer(), ECbValidateMode::All), ECbValidateError::None))
+	if (TestEqual(TEXT("FCbWriter(Reference) Validate"), ValidateCompactBinaryRange(Fields.GetBuffer(), ECbValidateMode::All), ECbValidateError::None))
 	{
 		const FBlake3Hash* CheckValue = Values;
 		for (FCbField Field : Fields)
@@ -608,7 +608,7 @@ bool FCbWriterBinaryReferenceTest::RunTest(const FString& Parameters)
 	}
 
 	FCbFieldRefIterator Fields = Writer.Save();
-	if (TestEqual(TEXT("FCbWriter(BinaryReference) Validate"), ValidateCompactBinaryRange(*Fields.GetBuffer(), ECbValidateMode::All), ECbValidateError::None))
+	if (TestEqual(TEXT("FCbWriter(BinaryReference) Validate"), ValidateCompactBinaryRange(Fields.GetBuffer(), ECbValidateMode::All), ECbValidateError::None))
 	{
 		const FBlake3Hash* CheckValue = Values;
 		for (FCbField Field : Fields)
@@ -638,7 +638,7 @@ bool FCbWriterHashTest::RunTest(const FString& Parameters)
 	}
 
 	FCbFieldRefIterator Fields = Writer.Save();
-	if (TestEqual(TEXT("FCbWriter(Hash) Validate"), ValidateCompactBinaryRange(*Fields.GetBuffer(), ECbValidateMode::All), ECbValidateError::None))
+	if (TestEqual(TEXT("FCbWriter(Hash) Validate"), ValidateCompactBinaryRange(Fields.GetBuffer(), ECbValidateMode::All), ECbValidateError::None))
 	{
 		const FBlake3Hash* CheckValue = Values;
 		for (FCbField Field : Fields)
@@ -665,7 +665,7 @@ bool FCbWriterUuidTest::RunTest(const FString& Parameters)
 	}
 
 	FCbFieldRefIterator Fields = Writer.Save();
-	if (TestEqual(TEXT("FCbWriter(Uuid) Validate"), ValidateCompactBinaryRange(*Fields.GetBuffer(), ECbValidateMode::All), ECbValidateError::None))
+	if (TestEqual(TEXT("FCbWriter(Uuid) Validate"), ValidateCompactBinaryRange(Fields.GetBuffer(), ECbValidateMode::All), ECbValidateError::None))
 	{
 		const FGuid* CheckValue = Values;
 		for (FCbField Field : Fields)
@@ -692,7 +692,7 @@ bool FCbWriterDateTimeTest::RunTest(const FString& Parameters)
 	}
 
 	FCbFieldRefIterator Fields = Writer.Save();
-	if (TestEqual(TEXT("FCbWriter(DateTime) Validate"), ValidateCompactBinaryRange(*Fields.GetBuffer(), ECbValidateMode::All), ECbValidateError::None))
+	if (TestEqual(TEXT("FCbWriter(DateTime) Validate"), ValidateCompactBinaryRange(Fields.GetBuffer(), ECbValidateMode::All), ECbValidateError::None))
 	{
 		const FDateTime* CheckValue = Values;
 		for (FCbField Field : Fields)
@@ -719,7 +719,7 @@ bool FCbWriterTimeSpanTest::RunTest(const FString& Parameters)
 	}
 
 	FCbFieldRefIterator Fields = Writer.Save();
-	if (TestEqual(TEXT("FCbWriter(TimeSpan) Validate"), ValidateCompactBinaryRange(*Fields.GetBuffer(), ECbValidateMode::All), ECbValidateError::None))
+	if (TestEqual(TEXT("FCbWriter(TimeSpan) Validate"), ValidateCompactBinaryRange(Fields.GetBuffer(), ECbValidateMode::All), ECbValidateError::None))
 	{
 		const FTimespan* CheckValue = Values;
 		for (FCbField Field : Fields)
@@ -917,7 +917,7 @@ bool FCbWriterComplexTest::RunTest(const FString& Parameters)
 	}
 
 	TestEqual(TEXT("FCbWriter(Complex).Save()->Validate"),
-		ValidateCompactBinary(*Object.GetBuffer(), ECbValidateMode::All), ECbValidateError::None);
+		ValidateCompactBinary(Object.GetBuffer(), ECbValidateMode::All), ECbValidateError::None);
 
 	TestEqual(TEXT("FCbWriter(Complex).Save(Ar)->Validate"),
 		ValidateCompactBinary(MakeMemoryView(Archive), ECbValidateMode::All), ECbValidateError::None);
@@ -935,7 +935,6 @@ bool FCbWriterOwnedReadOnlyTest::RunTest(const FString& Parameters)
 	Writer.EndObject();
 	FCbObjectRef Object = Writer.Save().AsObjectRef();
 	TestTrue(TEXT("FCbWriter().Save().IsOwned()"), Object.IsOwned());
-	TestTrue(TEXT("FCbWriter().Save().IsReadOnly()"), Object.IsReadOnly());
 
 	return true;
 }
@@ -995,7 +994,7 @@ bool FCbWriterStreamTest::RunTest(const FString& Parameters)
 		Object = Writer.Save().AsObjectRef();
 	}
 
-	TestEqual(TEXT("FCbWriter(Stream) Validate"), ValidateCompactBinary(*Object.GetBuffer(), ECbValidateMode::All), ECbValidateError::None);
+	TestEqual(TEXT("FCbWriter(Stream) Validate"), ValidateCompactBinary(Object.GetBuffer(), ECbValidateMode::All), ECbValidateError::None);
 
 	return true;
 }
diff --git a/Engine/Source/Runtime/Core/Public/Memory/MemoryFwd.h b/Engine/Source/Runtime/Core/Public/Memory/MemoryFwd.h
--- a/Engine/Source/Runtime/Core/Public/Memory/MemoryFwd.h
+++ b/Engine/Source/Runtime/Core/Public/Memory/MemoryFwd.h
@@ -15,9 +15,6 @@ using FMutableMemoryView = TMemoryView<void>;
 
 // SharedBuffer
 
-class FSharedBufferRef;
-class FSharedBufferConstRef;
-class FSharedBufferPtr;
-class FSharedBufferConstPtr;
-class FSharedBufferWeakPtr;
-class FSharedBufferConstWeakPtr;
+class FUniqueBuffer;
+class FSharedBuffer;
+class FWeakSharedBuffer;
diff --git a/Engine/Source/Runtime/Core/Public/Memory/SharedBuffer.h b/Engine/Source/Runtime/Core/Public/Memory/SharedBuffer.h
--- a/Engine/Source/Runtime/Core/Public/Memory/SharedBuffer.h
+++ b/Engine/Source/Runtime/Core/Public/Memory/SharedBuffer.h
@@ -3,7 +3,6 @@
 #pragma once
 
 #include "CoreTypes.h"
-#include "HAL/UnrealMemory.h"
 #include "Memory/MemoryView.h"
 #include "Misc/EnumClassFlags.h"
 #include "Templates/Invoke.h"
@@ -15,950 +14,764 @@
 
 ///////////////////////////////////////////////////////////////////////////////////////////////////////////////////////
 
-class FSharedBuffer;
+namespace BufferOwnerPrivate { struct FSharedOps; }
+namespace BufferOwnerPrivate { struct FWeakOps; }
 
-namespace SharedBufferPrivate
-{
+///////////////////////////////////////////////////////////////////////////////////////////////////////////////////////
 
-enum ESharedBufferFlags : uint16
+/**
+ * A reference-counted buffer.
+ *
+ * A buffer owner may own its memory or provide a view into memory owned externally. When used as
+ * a non-owning view, the viewed memory must be guaranteed to outlive the buffer owner.
+ *
+ * A buffer owner must be referenced and accessed through one of its three reference types:
+ * FUniqueBuffer, FSharedBuffer, or FWeakSharedBuffer.
+ *
+ * FUniqueBuffer and FSharedBuffer offer static functions to create a buffer using private buffer
+ * owner types that are ideal for most use cases. The TakeOwnership function allows creation of a
+ * buffer with a custom delete function. Advanced use cases require deriving from FBufferOwner to
+ * enable storage of arbitrary members alongside the data and size, and to enable materialization
+ * of the buffer to be deferred.
+ *
+ * A derived type must call SetIsOwned from its constructor if they own (or will own) the buffer.
+ * A derived type must call SetIsMaterialized from its constructor, unless it implements deferred
+ * materialization by overriding MaterializeBuffer.
+ */
+class FBufferOwner
 {
-	None = 0,
-	Owned = 1 << 0,
-	ReadOnly = 1 << 1,
-	HasBufferOwner = 1 << 2,
+private:
+	enum class EBufferOwnerFlags : uint8;
+
+protected:
+	FBufferOwner() = default;
+
+	FBufferOwner(const FBufferOwner&) = delete;
+	FBufferOwner& operator=(const FBufferOwner&) = delete;
+
+	inline FBufferOwner(void* InData, uint64 InSize);
+	virtual ~FBufferOwner();
+
+	/**
+	 * Materialize the buffer by making it ready to be accessed.
+	 *
+	 * This will be called before any access to the data or size, unless SetIsMaterialized is called
+	 * by the constructor. Accesses from multiple threads will cause multiple calls to this function
+	 * until at least one has finished.
+	 */
+	virtual void MaterializeBuffer();
+
+	/**
+	 * Free the buffer and any associated resources.
+	 *
+	 * This is called when the last shared reference is released. The destructor will be called when
+	 * the last weak reference is released. A buffer owner will always call this function before the
+	 * calling the destructor, unless an exception was thrown by the constructor.
+	 */
+	virtual void FreeBuffer() = 0;
+
+	inline void* GetData();
+	inline uint64 GetSize();
+	inline void SetBuffer(void* InData, uint64 InSize);
+
+	inline bool IsOwned() const;
+	inline void SetIsOwned();
+
+	inline void Materialize();
+	inline bool IsMaterialized() const;
+	inline void SetIsMaterialized();
+
+	inline bool IsReferenced() const;
+
+private:
+	static uint32 GetSharedRefCount(uint64 RefCountsAndFlags) { return uint32(RefCountsAndFlags >> 0) & 0x7fffffff; }
+	static uint64 SetSharedRefCount(uint32 RefCount) { return uint64(RefCount) << 0; }
+	static uint32 GetWeakRefCount(uint64 RefCountsAndFlags) { return uint32(RefCountsAndFlags >> 31) & 0x7fffffff; }
+	static uint64 SetWeakRefCount(uint32 RefCount) { return uint64(RefCount) << 31; }
+	static EBufferOwnerFlags GetFlags(uint64 RefCountsAndFlags) { return EBufferOwnerFlags(RefCountsAndFlags >> 62); }
+	static uint64 SetFlags(EBufferOwnerFlags Flags) { return uint64(Flags) << 62; }
+
+	inline void AddSharedReference();
+	inline void ReleaseSharedReference();
+	inline bool TryAddSharedReference();
+	inline void AddWeakReference();
+	inline void ReleaseWeakReference();
+
+	friend class FUniqueBuffer;
+	friend class FSharedBuffer;
+	friend class FWeakSharedBuffer;
+	friend struct BufferOwnerPrivate::FSharedOps;
+	friend struct BufferOwnerPrivate::FWeakOps;
+
+private:
+	void* Data = nullptr;
+	uint64 Size = 0;
+	std::atomic<uint64> ReferenceCountsAndFlags{0};
 };
 
-ENUM_CLASS_FLAGS(ESharedBufferFlags);
+///////////////////////////////////////////////////////////////////////////////////////////////////////////////////////
 
-struct FSharedRefOps final
+namespace BufferOwnerPrivate
 {
-	static bool TryAddRef(const FSharedBuffer* Buffer);
-	static void AddRef(const FSharedBuffer* Buffer);
-	static void Release(const FSharedBuffer* Buffer);
-};
 
-struct FSharedPtrOps final
+struct FSharedOps final
 {
-	static bool TryAddRef(const FSharedBuffer* Buffer);
-	static void AddRef(const FSharedBuffer* Buffer);
-	static void Release(const FSharedBuffer* Buffer);
+	static inline bool HasRef(FBufferOwner& Owner) { return Owner.IsReferenced(); }
+	static inline bool TryAddRef(FBufferOwner& Owner) { return Owner.TryAddSharedReference(); }
+	static inline void AddRef(FBufferOwner& Owner) { Owner.AddSharedReference(); }
+	static inline void Release(FBufferOwner* Owner) { if (Owner) { Owner->ReleaseSharedReference(); } }
 };
 
-struct FWeakPtrOps final
+struct FWeakOps final
 {
-	static bool TryAddRef(const FSharedBuffer* Buffer);
-	static void AddRef(const FSharedBuffer* Buffer);
-	static void Release(const FSharedBuffer* Buffer);
+	static inline bool HasRef(FBufferOwner& Owner) { return Owner.IsReferenced(); }
+	static inline bool TryAddRef(FBufferOwner& Owner) { AddRef(Owner); return true; }
+	static inline void AddRef(FBufferOwner& Owner) { Owner.AddWeakReference(); }
+	static inline void Release(FBufferOwner* Owner) { if (Owner) { Owner->ReleaseWeakReference(); } }
 };
 
-template <bool bInAllowNull, bool bInIsConst, bool bInIsWeak>
-class TSharedBufferPtr final
+template <typename FOps>
+class TBufferOwnerPtr final
 {
-	static constexpr bool bAllowNull = bInAllowNull;
-	static constexpr bool bIsConst = bInIsConst;
-	static constexpr bool bIsWeak = bInIsWeak;
-	using BufferType = std::conditional_t<bIsConst, const FSharedBuffer, FSharedBuffer>;
-	using Ops = std::conditional_t<bIsWeak, FWeakPtrOps, std::conditional_t<bAllowNull, FSharedPtrOps, FSharedRefOps>>;
+	static constexpr bool bIsWeak = std::is_same<FOps, FWeakOps>::value;
 
-	template <bool bOtherAllowNull, bool bOtherIsConst, bool bOtherIsWeak>
-	friend class TSharedBufferPtr;
+	template <typename FOtherOps>
+	friend class TBufferOwnerPtr;
 
-	template <bool bOtherAllowNull, bool bOtherIsConst, bool bOtherIsWeak>
-	static inline BufferType* CopyFrom(const TSharedBufferPtr<bOtherAllowNull, bOtherIsConst, bOtherIsWeak>& Ptr)
-	{
-		BufferType* NewBuffer = Ptr.Buffer;
-		if (bOtherIsWeak && !bIsWeak)
-		{
-			if (!Ops::TryAddRef(NewBuffer))
-			{
-				NewBuffer = nullptr;
-			}
-		}
-		else
-		{
-			Ops::AddRef(NewBuffer);
-		}
-		return NewBuffer;
-	}
+	template <typename FOtherOps>
+	static inline FBufferOwner* CopyFrom(const TBufferOwnerPtr<FOtherOps>& Ptr);
 
-	template <bool bOtherAllowNull, bool bOtherIsConst, bool bOtherIsWeak>
-	static inline BufferType* MoveFrom(TSharedBufferPtr<bOtherAllowNull, bOtherIsConst, bOtherIsWeak>&& Ptr)
-	{
-		BufferType* NewBuffer = Ptr.Buffer;
-		if (bOtherIsWeak)
-		{
-			if (bIsWeak)
-			{
-				Ptr.Buffer = nullptr;
-			}
-			else
-			{
-				if (!Ops::TryAddRef(NewBuffer))
-				{
-					NewBuffer = nullptr;
-				}
-			}
-		}
-		else
-		{
-			if (bIsWeak)
-			{
-				Ops::AddRef(NewBuffer);
-			}
-			else
-			{
-				if (bOtherAllowNull)
-				{
-					Ptr.Buffer = nullptr;
-				}
-				else
-				{
-					Ops::AddRef(NewBuffer);
-				}
-			}
-		}
-		return NewBuffer;
-	}
+	template <typename FOtherOps>
+	static inline FBufferOwner* MoveFrom(TBufferOwnerPtr<FOtherOps>&& Ptr);
 
 public:
-	inline explicit TSharedBufferPtr(BufferType* const InBuffer)
-		: Buffer(InBuffer)
-	{
-		Ops::AddRef(InBuffer);
-	}
+	inline TBufferOwnerPtr() = default;
+	inline explicit TBufferOwnerPtr(FBufferOwner* const InOwner);
 
-	inline TSharedBufferPtr(const TSharedBufferPtr& Ptr)
-		: Buffer(CopyFrom(Ptr))
-	{
-	}
-
-	inline TSharedBufferPtr(TSharedBufferPtr&& Ptr)
-		: Buffer(MoveFrom(MoveTemp(Ptr)))
-	{
-	}
-
-	template <bool bOtherAllowNull, bool bOtherIsConst, bool bOtherIsWeak>
-	inline explicit TSharedBufferPtr(const TSharedBufferPtr<bOtherAllowNull, bOtherIsConst, bOtherIsWeak>& Ptr)
-		: Buffer(CopyFrom(Ptr))
-	{
-	}
-
-	template <bool bOtherAllowNull, bool bOtherIsConst, bool bOtherIsWeak>
-	inline explicit TSharedBufferPtr(TSharedBufferPtr<bOtherAllowNull, bOtherIsConst, bOtherIsWeak>&& Ptr)
-		: Buffer(MoveFrom(MoveTemp(Ptr)))
-	{
-	}
+	inline TBufferOwnerPtr(const TBufferOwnerPtr& Ptr);
+	inline TBufferOwnerPtr(TBufferOwnerPtr&& Ptr);
 
-	inline ~TSharedBufferPtr()
-	{
-		Ops::Release(Buffer);
-	}
+	template <typename FOtherOps>
+	inline explicit TBufferOwnerPtr(const TBufferOwnerPtr<FOtherOps>& Ptr);
+	template <typename FOtherOps>
+	inline explicit TBufferOwnerPtr(TBufferOwnerPtr<FOtherOps>&& Ptr);
 
-	inline TSharedBufferPtr& operator=(const TSharedBufferPtr& Ptr)
-	{
-		BufferType* const OldBuffer = Buffer;
-		Buffer = CopyFrom(Ptr);
-		Ops::Release(OldBuffer);
-		return *this;
-	}
+	inline ~TBufferOwnerPtr();
 
-	inline TSharedBufferPtr& operator=(TSharedBufferPtr&& Ptr)
-	{
-		BufferType* const OldBuffer = Buffer;
-		Buffer = MoveFrom(MoveTemp(Ptr));
-		Ops::Release(OldBuffer);
-		return *this;
-	}
-
-	template <bool bOtherAllowNull, bool bOtherIsConst, bool bOtherIsWeak>
-	inline TSharedBufferPtr& operator=(const TSharedBufferPtr<bOtherAllowNull, bOtherIsConst, bOtherIsWeak>& Ptr)
-	{
-		BufferType* const OldBuffer = Buffer;
-		Buffer = CopyFrom(Ptr);
-		Ops::Release(OldBuffer);
-		return *this;
-	}
+	inline TBufferOwnerPtr& operator=(const TBufferOwnerPtr& Ptr);
+	inline TBufferOwnerPtr& operator=(TBufferOwnerPtr&& Ptr);
 
-	template <bool bOtherAllowNull, bool bOtherIsConst, bool bOtherIsWeak>
-	inline TSharedBufferPtr& operator=(TSharedBufferPtr<bOtherAllowNull, bOtherIsConst, bOtherIsWeak>&& Ptr)
-	{
-		BufferType* const OldBuffer = Buffer;
-		Buffer = MoveFrom(MoveTemp(Ptr));
-		Ops::Release(OldBuffer);
-		return *this;
-	}
+	template <typename FOtherOps>
+	inline TBufferOwnerPtr& operator=(const TBufferOwnerPtr<FOtherOps>& Ptr);
+	template <typename FOtherOps>
+	inline TBufferOwnerPtr& operator=(TBufferOwnerPtr<FOtherOps>&& Ptr);
 
-	BufferType* Get() const { return Buffer; }
+	template <typename FOtherOps>
+	inline bool operator==(const TBufferOwnerPtr<FOtherOps>& Ptr) const;
+	template <typename FOtherOps>
+	inline bool operator!=(const TBufferOwnerPtr<FOtherOps>& Ptr) const;
 
-	template <bool bOtherAllowNull, bool bOtherIsConst, bool bOtherIsWeak>
-	inline bool operator==(const TSharedBufferPtr<bOtherAllowNull, bOtherIsConst, bOtherIsWeak>& Ptr) const
-	{
-		return Buffer == Ptr.Buffer;
-	}
+	inline FBufferOwner* Get() const { return Owner; }
+	inline FBufferOwner* operator->() const { return Get(); }
+	inline explicit operator bool() const { return !IsNull(); }
+	inline bool IsNull() const { return Owner == nullptr; }
 
-	template <bool bOtherAllowNull, bool bOtherIsConst, bool bOtherIsWeak>
-	inline bool operator!=(const TSharedBufferPtr<bOtherAllowNull, bOtherIsConst, bOtherIsWeak>& Ptr) const
-	{
-		return Buffer != Ptr.Buffer;
-	}
+	inline void Reset();
 
 private:
-	BufferType* Buffer = nullptr;
+	FBufferOwner* Owner = nullptr;
 };
 
-}
+} // BufferOwnerPrivate
 
 ///////////////////////////////////////////////////////////////////////////////////////////////////////////////////////
 
-/** A non-nullable, thread-safe, reference to a writable shared buffer. */
-class FSharedBufferRef final
+/**
+ * A single-ownership mutable buffer.
+ *
+ * Ownership can be transferred by moving to FUniqueBuffer or it can be converted to an immutable
+ * shared buffer by moving to FSharedBuffer.
+ *
+ * @see FBufferOwner
+ */
+class FUniqueBuffer
 {
 public:
-	using PtrType = SharedBufferPrivate::TSharedBufferPtr</*bAllowNull*/ false, /*bIsConst*/ false, /*bIsWeak*/ false>;
+	/** Make an owned buffer of Size bytes. */
+	CORE_API static FUniqueBuffer Alloc(uint64 Size);
 
-	inline explicit FSharedBufferRef(PtrType&& InPtr) : Ptr(MoveTemp(InPtr)) {}
+	/** Make an owned clone of the input. */
+	CORE_API static FUniqueBuffer Clone(FMemoryView View);
+	CORE_API static FUniqueBuffer Clone(const void* Data, uint64 Size);
 
-	inline FSharedBuffer& Get() const { return *Ptr.Get(); }
+	/** Make a non-owned view of the input. */
+	CORE_API static FUniqueBuffer MakeView(FMutableMemoryView View);
+	CORE_API static FUniqueBuffer MakeView(void* Data, uint64 Size);
 
-	inline FSharedBuffer& operator*() const { return *Ptr.Get(); }
-	inline FSharedBuffer* operator->() const { return Ptr.Get(); }
+	/**
+	 * Make an owned buffer by taking ownership of the input.
+	 *
+	 * @param DeleteFunction Called with Data to free memory when the last shared reference is released.
+	 */
+	template <typename DeleteFunctionType,
+		decltype(Invoke(std::declval<DeleteFunctionType>(), std::declval<void*>()))* = nullptr>
+	static inline FUniqueBuffer TakeOwnership(void* Data, uint64 Size, DeleteFunctionType&& DeleteFunction);
 
-	inline friend uint32 GetTypeHash(const FSharedBufferRef& Ref) { return PointerHash(Ref.Ptr.Get()); }
+	/**
+	 * Make an owned buffer by taking ownership of the input.
+	 *
+	 * @param DeleteFunction Called with (Data, Size) to free memory when the last shared reference is released.
+	 */
+	template <typename DeleteFunctionType,
+		decltype(Invoke(std::declval<DeleteFunctionType>(), std::declval<void*>(), std::declval<uint64>()))* = nullptr>
+	static inline FUniqueBuffer TakeOwnership(void* Data, uint64 Size, DeleteFunctionType&& DeleteFunction);
 
-private:
-	constexpr inline friend bool IsValid(const FSharedBufferRef&) { return true; }
-	inline friend FSharedBufferRef ToSharedRef(const FSharedBufferRef& Ref) { return Ref; }
-	inline friend const PtrType& ToPrivateSharedBufferPtr(const FSharedBufferRef& Ref) { return Ref.Ptr; }
+	/** Construct a null unique buffer. */
+	FUniqueBuffer() = default;
 
-	friend class FSharedBufferConstRef;
-	friend class FSharedBufferPtr;
-	friend class FSharedBufferConstPtr;
-	friend class FSharedBufferWeakPtr;
-	friend class FSharedBufferConstWeakPtr;
+	/** Construct a unique buffer from a new unreferenced buffer owner. */
+	CORE_API explicit FUniqueBuffer(FBufferOwner* Owner);
 
-	PtrType Ptr;
-};
+	FUniqueBuffer(FUniqueBuffer&&) = default;
+	FUniqueBuffer& operator=(FUniqueBuffer&&) = default;
 
-/** A non-nullable, thread-safe, const reference to a shared buffer. */
-class FSharedBufferConstRef final
-{
-public:
-	using PtrType = SharedBufferPrivate::TSharedBufferPtr</*bAllowNull*/ false, /*bIsConst*/ true, /*bIsWeak*/ false>;
+	FUniqueBuffer(const FUniqueBuffer&) = delete;
+	FUniqueBuffer& operator=(const FUniqueBuffer&) = delete;
+
+	/** Reset this to null. */
+	CORE_API void Reset();
+
+	/** Returns a pointer to the start of the buffer. */
+	inline void* GetData() { return Owner ? Owner->GetData() : nullptr; }
+	inline const void* GetData() const { return Owner ? Owner->GetData() : nullptr; }
+
+	/** Returns the size of the buffer in bytes. */
+	inline uint64 GetSize() const { return Owner ? Owner->GetSize() : 0; }
+
+	/** Returns a view of the buffer. */
+	inline FMutableMemoryView GetView() { return FMutableMemoryView(GetData(), GetSize()); }
+	inline FMemoryView GetView() const { return FMemoryView(GetData(), GetSize()); }
+	inline operator FMutableMemoryView() { return GetView(); }
+	inline operator FMemoryView() const { return GetView(); }
+
+	/** Returns true if this points to a buffer. */
+	inline explicit operator bool() const { return !IsNull(); }
 
-	inline explicit FSharedBufferConstRef(PtrType&& InPtr) : Ptr(MoveTemp(InPtr)) {}
+	/** Returns true if this does not point to a buffer. */
+	inline bool IsNull() const { return Owner.IsNull(); }
 
-	inline FSharedBufferConstRef(const FSharedBufferRef& Ref) : Ptr(Ref.Ptr) {}
+	/** Returns true if this owns the memory for its buffer. */
+	inline bool IsOwned() const { return Owner && Owner->IsOwned(); }
 
-	inline const FSharedBuffer& Get() const { return *Ptr.Get(); }
+	/** Clone into a new buffer if this does not own the memory for its buffer. */
+	CORE_API void MakeOwned();
 
-	inline const FSharedBuffer& operator*() const { return *Ptr.Get(); }
-	inline const FSharedBuffer* operator->() const { return Ptr.Get(); }
+	/** Returns true if this points to a buffer that has been materialized. */
+	inline bool IsMaterialized() const { return Owner && Owner->IsMaterialized(); }
 
-	inline friend uint32 GetTypeHash(const FSharedBufferConstRef& Ref) { return PointerHash(Ref.Ptr.Get()); }
+	/** Make the data and size of the buffer available. GetData and GetSize do this automatically. */
+	CORE_API void Materialize() const;
 
 private:
-	constexpr inline friend bool IsValid(const FSharedBufferConstRef&) { return true; }
-	inline friend FSharedBufferConstRef ToSharedRef(const FSharedBufferConstRef& Ref) { return Ref; }
-	inline friend const PtrType& ToPrivateSharedBufferPtr(const FSharedBufferConstRef& Ref) { return Ref.Ptr; }
+	using OwnerPtrType = BufferOwnerPrivate::TBufferOwnerPtr<BufferOwnerPrivate::FSharedOps>;
 
-	friend class FSharedBufferConstPtr;
-	friend class FSharedBufferConstWeakPtr;
+	inline friend const OwnerPtrType& ToPrivateOwnerPtr(const FUniqueBuffer& Buffer) { return Buffer.Owner; }
+	inline friend OwnerPtrType ToPrivateOwnerPtr(FUniqueBuffer&& Buffer) { return MoveTemp(Buffer.Owner); }
 
-	PtrType Ptr;
+	OwnerPtrType Owner;
 };
 
 ///////////////////////////////////////////////////////////////////////////////////////////////////////////////////////
 
-/** A non-nullable, thread-safe, pointer to a writable shared buffer. */
-class FSharedBufferPtr final
+/**
+ * A shared-ownership immutable buffer.
+ *
+ * @see FBufferOwner
+ */
+class FSharedBuffer
 {
 public:
-	using PtrType = SharedBufferPrivate::TSharedBufferPtr</*bAllowNull*/ true, /*bIsConst*/ false, /*bIsWeak*/ false>;
+	/** Make an owned clone of the input. */
+	CORE_API static FSharedBuffer Clone(FMemoryView View);
+	CORE_API static FSharedBuffer Clone(const void* Data, uint64 Size);
+
+	/** Make a non-owned view of the input. */
+	CORE_API static FSharedBuffer MakeView(FMemoryView View);
+	CORE_API static FSharedBuffer MakeView(const void* Data, uint64 Size);
 
-	inline explicit FSharedBufferPtr(PtrType&& InPtr) : Ptr(MoveTemp(InPtr)) {}
+	/** Make a view of the input within its outer buffer. Ownership matches OuterBuffer. */
+	CORE_API static FSharedBuffer MakeView(FMemoryView View, FSharedBuffer OuterBuffer);
+	CORE_API static FSharedBuffer MakeView(const void* Data, uint64 Size, FSharedBuffer OuterBuffer);
 
-	inline FSharedBufferPtr() : Ptr(nullptr) {}
-	inline FSharedBufferPtr(const FSharedBufferRef& Ref) : Ptr(Ref.Ptr) {}
+	/**
+	 * Make an owned buffer by taking ownership of the input.
+	 *
+	 * @param DeleteFunction Called with Data to free memory when the last shared reference is released.
+	 */
+	template <typename DeleteFunctionType,
+		decltype(Invoke(std::declval<DeleteFunctionType>(), std::declval<void*>()))* = nullptr>
+	static inline FSharedBuffer TakeOwnership(const void* Data, uint64 Size, DeleteFunctionType&& DeleteFunction);
 
-	inline FSharedBufferRef ToSharedRef() const& { return FSharedBufferRef(FSharedBufferRef::PtrType(Ptr)); }
-	inline FSharedBufferRef ToSharedRef() && { return FSharedBufferRef(FSharedBufferRef::PtrType(MoveTemp(Ptr))); }
+	/**
+	 * Make an owned buffer by taking ownership of the input.
+	 *
+	 * @param DeleteFunction Called with (Data, Size) to free memory when the last shared reference is released.
+	 */
+	template <typename DeleteFunctionType,
+		decltype(Invoke(std::declval<DeleteFunctionType>(), std::declval<void*>(), std::declval<uint64>()))* = nullptr>
+	static inline FSharedBuffer TakeOwnership(const void* Data, uint64 Size, DeleteFunctionType&& DeleteFunction);
 
-	inline bool IsValid() const { return Ptr.Get() != nullptr; }
-	inline FSharedBuffer* Get() const { return Ptr.Get(); }
+	/** Construct a null shared buffer. */
+	FSharedBuffer() = default;
 
-	inline explicit operator bool() const { return Ptr.Get() != nullptr; }
-	inline FSharedBuffer& operator*() const { return *Ptr.Get(); }
-	inline FSharedBuffer* operator->() const { return Ptr.Get(); }
+	/** Construct a shared buffer from a new unreferenced buffer owner. */
+	CORE_API explicit FSharedBuffer(FBufferOwner* Owner);
 
-	inline void Reset() { *this = FSharedBufferPtr(); }
+	/** Construct a shared buffer by taking ownership of a unique buffer. */
+	inline explicit FSharedBuffer(FUniqueBuffer&& Buffer)
+		: Owner(ToPrivateOwnerPtr(MoveTemp(Buffer)))
+	{
+	}
 
-	inline friend uint32 GetTypeHash(const FSharedBufferPtr& InPtr) { return PointerHash(InPtr.Ptr.Get()); }
+	/** Assign a shared buffer by taking ownership of a unique buffer. */
+	inline FSharedBuffer& operator=(FUniqueBuffer&& Buffer)
+	{
+		Owner = ToPrivateOwnerPtr(MoveTemp(Buffer));
+		return *this;
+	}
 
-private:
-	inline friend bool IsValid(const FSharedBufferPtr& InPtr) { return InPtr.IsValid(); }
-	inline friend FSharedBufferRef ToSharedRef(const FSharedBufferPtr& InPtr) { return InPtr.ToSharedRef(); }
-	inline friend FSharedBufferRef ToSharedRef(FSharedBufferPtr&& InPtr) { return MoveTemp(InPtr).ToSharedRef(); }
-	inline friend const PtrType& ToPrivateSharedBufferPtr(const FSharedBufferPtr& InPtr) { return InPtr.Ptr; }
+	/** Reset this to null. */
+	CORE_API void Reset();
 
-	friend class FSharedBufferConstPtr;
-	friend class FSharedBufferWeakPtr;
-	friend class FSharedBufferConstWeakPtr;
+	/** Returns a pointer to the start of the buffer. */
+	inline const void* GetData() const { return Owner ? Owner->GetData() : nullptr; }
 
-	PtrType Ptr;
-};
+	/** Returns the size of the buffer in bytes. */
+	inline uint64 GetSize() const { return Owner ? Owner->GetSize() : 0; }
 
-/** A non-nullable, thread-safe, const pointer to a shared buffer. */
-class FSharedBufferConstPtr final
-{
-public:
-	using PtrType = SharedBufferPrivate::TSharedBufferPtr</*bAllowNull*/ true, /*bIsConst*/ true, /*bIsWeak*/ false>;
+	/** Returns a view of the buffer. */
+	inline FMemoryView GetView() const { return FMemoryView(GetData(), GetSize()); }
+	inline operator FMemoryView() const { return GetView(); }
 
-	inline explicit FSharedBufferConstPtr(PtrType&& InPtr) : Ptr(MoveTemp(InPtr)) {}
+	/** Returns true if this points to a buffer. */
+	inline explicit operator bool() const { return !IsNull(); }
 
-	inline FSharedBufferConstPtr() : Ptr(nullptr) {}
-	inline FSharedBufferConstPtr(const FSharedBufferRef& Ref) : Ptr(Ref.Ptr) {}
-	inline FSharedBufferConstPtr(const FSharedBufferConstRef& Ref) : Ptr(Ref.Ptr) {}
-	inline FSharedBufferConstPtr(const FSharedBufferPtr& InPtr) : Ptr(InPtr.Ptr) {}
-	inline FSharedBufferConstPtr(FSharedBufferPtr&& InPtr) : Ptr(MoveTemp(InPtr.Ptr)) {}
+	/** Returns true if this does not point to a buffer. */
+	inline bool IsNull() const { return Owner.IsNull(); }
 
-	inline FSharedBufferConstRef ToSharedRef() const& { return FSharedBufferConstRef(FSharedBufferConstRef::PtrType(Ptr)); }
-	inline FSharedBufferConstRef ToSharedRef() && { return FSharedBufferConstRef(FSharedBufferConstRef::PtrType(MoveTemp(Ptr))); }
+	/** Returns true if this owns the memory for its buffer. */
+	inline bool IsOwned() const { return Owner && Owner->IsOwned(); }
 
-	inline bool IsValid() const { return Ptr.Get() != nullptr; }
-	inline const FSharedBuffer* Get() const { return Ptr.Get(); }
+	/** Clone into a new buffer if this does not own the memory for its buffer. */
+	CORE_API void MakeOwned();
 
-	inline explicit operator bool() const { return Ptr.Get() != nullptr; }
-	inline const FSharedBuffer& operator*() const { return *Ptr.Get(); }
-	inline const FSharedBuffer* operator->() const { return Ptr.Get(); }
+	/** Returns true if this points to a buffer that has been materialized. */
+	inline bool IsMaterialized() const { return Owner && Owner->IsMaterialized(); }
 
-	inline void Reset() { *this = FSharedBufferConstPtr(); }
+	/** Make the data and size of the buffer available. GetData and GetSize do this automatically. */
+	CORE_API void Materialize() const;
 
-	inline friend uint32 GetTypeHash(const FSharedBufferConstPtr& InPtr) { return PointerHash(InPtr.Ptr.Get()); }
+	friend class FWeakSharedBuffer;
 
 private:
-	inline friend bool IsValid(const FSharedBufferConstPtr& InPtr) { return InPtr.IsValid(); }
-	inline friend FSharedBufferConstRef ToSharedRef(const FSharedBufferConstPtr& InPtr) { return InPtr.ToSharedRef(); }
-	inline friend FSharedBufferConstRef ToSharedRef(FSharedBufferConstPtr&& InPtr) { return MoveTemp(InPtr).ToSharedRef(); }
-	inline friend const PtrType& ToPrivateSharedBufferPtr(const FSharedBufferConstPtr& InPtr) { return InPtr.Ptr; }
+	FSharedBuffer(const BufferOwnerPrivate::TBufferOwnerPtr<BufferOwnerPrivate::FWeakOps>& WeakOwner);
+
+	using OwnerPtrType = BufferOwnerPrivate::TBufferOwnerPtr<BufferOwnerPrivate::FSharedOps>;
 
-	friend class FSharedBufferConstWeakPtr;
+	inline friend const OwnerPtrType& ToPrivateOwnerPtr(const FSharedBuffer& Buffer) { return Buffer.Owner; }
 
-	PtrType Ptr;
+	OwnerPtrType Owner;
 };
 
 ///////////////////////////////////////////////////////////////////////////////////////////////////////////////////////
 
-/** A non-nullable, thread-safe, weak, pointer to a writable shared buffer. */
-class FSharedBufferWeakPtr final
+/**
+ * A weak reference to a shared-ownership immutable buffer.
+ *
+ * @see FSharedBuffer
+ */
+class FWeakSharedBuffer
 {
 public:
-	using PtrType = SharedBufferPrivate::TSharedBufferPtr</*bAllowNull*/ true, /*bIsConst*/ false, /*bIsWeak*/ true>;
+	/** Construct a null weak shared buffer. */
+	FWeakSharedBuffer() = default;
 
-	inline explicit FSharedBufferWeakPtr(PtrType&& InPtr) : Ptr(MoveTemp(InPtr)) {}
+	/** Construct a weak shared buffer from a shared buffer. */
+	CORE_API FWeakSharedBuffer(const FSharedBuffer& Buffer);
 
-	inline FSharedBufferWeakPtr() : Ptr(nullptr) {}
-	inline FSharedBufferWeakPtr(const FSharedBufferRef& Ref) : Ptr(Ref.Ptr) {}
-	inline FSharedBufferWeakPtr(const FSharedBufferPtr& InPtr) : Ptr(InPtr.Ptr) {}
+	/** Assign a weak shared buffer from a shared buffer. */
+	CORE_API FWeakSharedBuffer& operator=(const FSharedBuffer& Buffer);
 
-	inline FSharedBufferPtr Pin() const { return FSharedBufferPtr(FSharedBufferPtr::PtrType(Ptr)); }
+	/** Reset this to null. */
+	CORE_API void Reset();
 
-	inline void Reset() { *this = FSharedBufferWeakPtr(); }
-
-	inline friend uint32 GetTypeHash(const FSharedBufferWeakPtr& InPtr) { return PointerHash(InPtr.Ptr.Get()); }
+	/** Convert this to a shared buffer. */
+	CORE_API FSharedBuffer Pin() const;
 
 private:
-	friend class FSharedBufferConstWeakPtr;
-	inline friend const PtrType& ToPrivateSharedBufferPtr(const FSharedBufferWeakPtr& InPtr) { return InPtr.Ptr; }
-
-	PtrType Ptr;
-};
-
-/** A non-nullable, thread-safe, weak, const pointer to a shared buffer. */
-class FSharedBufferConstWeakPtr final
-{
-public:
-	using PtrType = SharedBufferPrivate::TSharedBufferPtr</*bAllowNull*/ true, /*bIsConst*/ true, /*bIsWeak*/ true>;
-
-	inline explicit FSharedBufferConstWeakPtr(PtrType&& InPtr) : Ptr(MoveTemp(InPtr)) {}
-
-	inline FSharedBufferConstWeakPtr() : Ptr(nullptr) {}
-	inline FSharedBufferConstWeakPtr(const FSharedBufferRef& Ref) : Ptr(Ref.Ptr) {}
-	inline FSharedBufferConstWeakPtr(const FSharedBufferConstRef& Ref) : Ptr(Ref.Ptr) {}
-	inline FSharedBufferConstWeakPtr(const FSharedBufferPtr& InPtr) : Ptr(InPtr.Ptr) {}
-	inline FSharedBufferConstWeakPtr(const FSharedBufferConstPtr& InPtr) : Ptr(InPtr.Ptr) {}
-	inline FSharedBufferConstWeakPtr(const FSharedBufferWeakPtr& InPtr) : Ptr(InPtr.Ptr) {}
-	inline FSharedBufferConstWeakPtr(FSharedBufferWeakPtr&& InPtr) : Ptr(MoveTemp(InPtr.Ptr)) {}
+	using OwnerPtrType = BufferOwnerPrivate::TBufferOwnerPtr<BufferOwnerPrivate::FWeakOps>;
 
-	inline FSharedBufferConstPtr Pin() const { return FSharedBufferConstPtr(FSharedBufferConstPtr::PtrType(Ptr)); }
+	inline friend const OwnerPtrType& ToPrivateOwnerPtr(const FWeakSharedBuffer& Buffer) { return Buffer.Owner; }
 
-	inline void Reset() { *this = FSharedBufferConstWeakPtr(); }
-
-	inline friend uint32 GetTypeHash(const FSharedBufferConstWeakPtr& InPtr) { return PointerHash(InPtr.Ptr.Get()); }
-
-private:
-	inline friend const PtrType& ToPrivateSharedBufferPtr(const FSharedBufferConstWeakPtr& InPtr) { return InPtr.Ptr; }
-
-	PtrType Ptr;
+	OwnerPtrType Owner;
 };
 
 ///////////////////////////////////////////////////////////////////////////////////////////////////////////////////////
 
 template <typename TypeA, typename TypeB>
 inline auto operator==(const TypeA& BufferA, const TypeB& BufferB)
-	-> decltype(ToPrivateSharedBufferPtr(BufferA) == ToPrivateSharedBufferPtr(BufferB))
+	-> decltype(ToPrivateOwnerPtr(BufferA) == ToPrivateOwnerPtr(BufferB))
 {
-	return ToPrivateSharedBufferPtr(BufferA) == ToPrivateSharedBufferPtr(BufferB);
+	return ToPrivateOwnerPtr(BufferA) == ToPrivateOwnerPtr(BufferB);
 }
 
 template <typename TypeA, typename TypeB>
 inline auto operator!=(const TypeA& BufferA, const TypeB& BufferB)
-	-> decltype(ToPrivateSharedBufferPtr(BufferA) != ToPrivateSharedBufferPtr(BufferB))
+	-> decltype(ToPrivateOwnerPtr(BufferA) != ToPrivateOwnerPtr(BufferB))
 {
-	return ToPrivateSharedBufferPtr(BufferA) != ToPrivateSharedBufferPtr(BufferB);
+	return ToPrivateOwnerPtr(BufferA) != ToPrivateOwnerPtr(BufferB);
 }
 
+inline uint32 GetTypeHash(const FUniqueBuffer& Buffer) { return PointerHash(ToPrivateOwnerPtr(Buffer).Get()); }
+inline uint32 GetTypeHash(const FSharedBuffer& Buffer) { return PointerHash(ToPrivateOwnerPtr(Buffer).Get()); }
+inline uint32 GetTypeHash(const FWeakSharedBuffer& Buffer) { return PointerHash(ToPrivateOwnerPtr(Buffer).Get()); }
+
 ///////////////////////////////////////////////////////////////////////////////////////////////////////////////////////
 
-template<> struct TIsZeroConstructType<FSharedBufferRef> { enum { Value = true }; };
-template<> struct TIsZeroConstructType<FSharedBufferConstRef> { enum { Value = true }; };
-template<> struct TIsZeroConstructType<FSharedBufferPtr> { enum { Value = true }; };
-template<> struct TIsZeroConstructType<FSharedBufferConstPtr> { enum { Value = true }; };
-template<> struct TIsZeroConstructType<FSharedBufferWeakPtr> { enum { Value = true }; };
-template<> struct TIsZeroConstructType<FSharedBufferConstWeakPtr> { enum { Value = true }; };
+template<> struct TIsZeroConstructType<FUniqueBuffer> { enum { Value = true }; };
+template<> struct TIsZeroConstructType<FSharedBuffer> { enum { Value = true }; };
+template<> struct TIsZeroConstructType<FWeakSharedBuffer> { enum { Value = true }; };
 
-template<> struct TIsWeakPointerType<FSharedBufferWeakPtr> { enum { Value = true }; };
-template<> struct TIsWeakPointerType<FSharedBufferConstWeakPtr> { enum { Value = true }; };
+template<> struct TIsWeakPointerType<FWeakSharedBuffer> { enum { Value = true }; };
 
 ///////////////////////////////////////////////////////////////////////////////////////////////////////////////////////
 
-/**
- * A reference-counted shared buffer type.
- *
- * A shared buffer may either own its memory or provide a view into memory owned externally. When
- * used as to wrap a externally-owned memory, the user must guarantee the lifetime of the wrapped
- * memory will exceed the lifetime of the shared buffer.
- *
- * A shared buffer must be stored and accessed through one of its six reference types:
- * FSharedBuffer[Const]Ref or FSharedBuffer[Const][Weak]Ptr.
- *
- * Make an owned sized buffer with optional alignment.
- *
- *     FSharedBuffer::Alloc(128)
- *     FSharedBuffer::Alloc(128, 16)
- *
- * Make an owned buffer by cloning the buffer pointed to by the arguments.
- *
- *     FSharedBuffer::Clone(Data, Size)
- *     FSharedBuffer::Clone(MakeMemoryView(...))
- *     FSharedBuffer::Clone(*SharedBufferRefOrPtr)
- *
- * Make a non-owned buffer by creating a view of the buffer pointed to by the arguments.
- *
- *     FSharedBuffer::MakeView(Data, Size)
- *     FSharedBuffer::MakeView(MakeMemoryView(...))
- *     FSharedBuffer::MakeView(Data, Size, OuterBuffer)
- *     FSharedBuffer::MakeView(MakeMemoryView(...), OuterBuffer)
- *
- * Make an owned buffer by taking ownership of the buffer pointed to by the arguments.
- *
- *     FSharedBuffer::TakeOwnership(Data, Size, FMemory::Free)
- *     FSharedBuffer::TakeOwnership(Data, Size, [](void* Data) { delete[] static_cast<uint8*>(Data); })
- *     FSharedBuffer::TakeOwnership(Data, Size, FCustomBufferOwner())
- */
-class FSharedBuffer final
+namespace BufferOwnerPrivate
 {
-	using ESharedBufferFlags = SharedBufferPrivate::ESharedBufferFlags;
 
+template <typename DeleteFunctionType>
+class TBufferOwnerDeleteFunction final : public FBufferOwner
+{
 public:
-	/**
-	 * Base for customizing lifetime management for the buffer.
-	 *
-	 * Types derived from this must be movable. Prefer an arbitrary callable (including a lambda) to
-	 * handle deletion in place of the extra code required to derive from this type.
-	 */
-	class FBufferOwner
-	{
-	protected:
-		FBufferOwner() = default;
-
-	public:
-		virtual ~FBufferOwner() = default;
-
-		/** Free the memory for the buffer. */
-		virtual void Free(void* Data, uint64 Size) = 0;
-	};
-
-	/** Make an owned writable buffer of Size bytes. */
-	static inline FSharedBufferPtr Alloc(const uint64 Size, const uint32 Alignment = DEFAULT_ALIGNMENT)
-	{
-		return NewBuffer(FMemory::Malloc(Size, Alignment), Size, ESharedBufferFlags::Owned);
-	}
-
-	/** Make an owned writable clone of the memory view. */
-	static inline FSharedBufferPtr Clone(const FMemoryView View)
-	{
-		return Clone(View.GetData(), View.GetSize());
-	}
-
-	/** Make an owned writable clone of the buffer. */
-	static inline FSharedBufferPtr Clone(const void* const Data, const uint64 Size)
+	explicit TBufferOwnerDeleteFunction(void* Data, uint64 Size, DeleteFunctionType&& InDeleteFunction)
+		: FBufferOwner(Data, Size)
+		, DeleteFunction(Forward<DeleteFunctionType>(InDeleteFunction))
 	{
-		const FSharedBufferPtr Buffer = Alloc(Size);
-		FMemory::Memcpy(Buffer->GetData(), Data, Size);
-		return Buffer;
+		SetIsMaterialized();
+		SetIsOwned();
 	}
 
-	/** Make a non-owned writable view of the memory view. */
-	static inline FSharedBufferPtr MakeView(FMutableMemoryView View)
-	{
-		return MakeView(View.GetData(), View.GetSize());
-	}
-
-	/** Make a non-owned read-only view of the memory view. */
-	static inline FSharedBufferConstPtr MakeView(FMemoryView View)
+private:
+	virtual void FreeBuffer() final
 	{
-		return MakeView(View.GetData(), View.GetSize());
+		Invoke(DeleteFunction, GetData(), GetSize());
 	}
 
-	/** Make a non-owned writable view of the buffer. */
-	static inline FSharedBufferPtr MakeView(void* Data, uint64 Size)
-	{
-		return NewBuffer(Data, Size, ESharedBufferFlags::None);
-	}
+	std::decay_t<DeleteFunctionType> DeleteFunction;
+};
 
-	/** Make a non-owned read-only view of the buffer. */
-	static inline FSharedBufferConstPtr MakeView(const void* Data, uint64 Size)
-	{
-		// const_cast is safe here because the return value only allows accessing data through a const void*
-		return MakeView(const_cast<void*>(Data), Size);
-	}
+} // BufferOwnerPrivate
 
-	/** Make a non-owned read-only view of the outer buffer and hold a reference to it. */
-	static inline FSharedBufferConstPtr MakeView(FMemoryView View, const FSharedBuffer& OuterBuffer)
+template <typename DeleteFunctionType,
+	decltype(Invoke(std::declval<DeleteFunctionType>(), std::declval<void*>()))*>
+inline FUniqueBuffer FUniqueBuffer::TakeOwnership(void* Data, uint64 Size, DeleteFunctionType&& DeleteFunction)
+{
+	return TakeOwnership(Data, Size, [Delete=Forward<DeleteFunctionType>(DeleteFunction)](void* InData, uint64 InSize)
 	{
-		return MakeView(View.GetData(), View.GetSize(), OuterBuffer);
-	}
+		Delete(InData);
+	});
+}
 
-	/** Make a non-owned read-only view of the outer buffer and hold a reference to it. */
-	static inline FSharedBufferConstPtr MakeView(const void* Data, uint64 Size, const FSharedBuffer& OuterBuffer)
-	{
-		check(OuterBuffer.GetView().Contains(MakeMemoryView(Data, Size)));
-		// Create a reference now to prevent the read-only state from changing before we add a reference later.
-		FSharedBufferConstPtr OuterBufferPtr{FSharedBufferConstPtr::PtrType(&OuterBuffer)};
-		// This buffer is not owned, but is read-only if its outer is read-only.
-		const ESharedBufferFlags Flags = GetFlags(OuterBuffer.ReferenceCountAndFlags) & ESharedBufferFlags::ReadOnly;
-		// const_cast is safe here because the return value only allows accessing data through a const void*
-		return NewBufferWithOwner<FOuterBufferOwner>(const_cast<void*>(Data), Size, Flags, MoveTemp(OuterBufferPtr));
-	}
+template <typename DeleteFunctionType,
+	decltype(Invoke(std::declval<DeleteFunctionType>(), std::declval<void*>(), std::declval<uint64>()))*>
+inline FUniqueBuffer FUniqueBuffer::TakeOwnership(void* Data, uint64 Size, DeleteFunctionType&& DeleteFunction)
+{
+	using OwnerType = BufferOwnerPrivate::TBufferOwnerDeleteFunction<DeleteFunctionType>;
+	return FUniqueBuffer(new OwnerType(Data, Size, Forward<DeleteFunctionType>(DeleteFunction)));
+}
 
-	/**
-	 * Make an owned writable buffer by taking ownership of the provided memory.
-	 *
-	 * @param DeleteFunction Called with the data pointer to free memory when there are no shared references.
-	 */
-	template <typename DeleteFunctionType,
-		decltype(Invoke(std::declval<DeleteFunctionType>(), std::declval<void*>()))* = nullptr>
-	static inline FSharedBufferPtr TakeOwnership(void* Data, uint64 Size, DeleteFunctionType&& DeleteFunction)
+template <typename DeleteFunctionType,
+	decltype(Invoke(std::declval<DeleteFunctionType>(), std::declval<void*>()))*>
+inline FSharedBuffer FSharedBuffer::TakeOwnership(const void* Data, uint64 Size, DeleteFunctionType&& DeleteFunction)
+{
+	return TakeOwnership(Data, Size, [Delete=Forward<DeleteFunctionType>(DeleteFunction)](void* InData, uint64 InSize)
 	{
-		return NewBufferWithOwner<TDeleteFunctionOwner<DeleteFunctionType>>(Data, Size, ESharedBufferFlags::Owned,
-			Forward<DeleteFunctionType>(DeleteFunction));
-	}
+		Delete(InData);
+	});
+}
 
-	/**
-	 * Make an owned read-only buffer by taking ownership of the provided memory.
-	 *
-	 * @param DeleteFunction Called with the data pointer to free memory when there are no shared references.
-	 */
-	template <typename DeleteFunctionType,
-		decltype(Invoke(std::declval<DeleteFunctionType>(), std::declval<void*>()))* = nullptr>
-	static inline FSharedBufferConstPtr TakeOwnership(const void* Data, uint64 Size, DeleteFunctionType&& DeleteFunction)
-	{
-		// const_cast is safe here because the return value only allows accessing data through a const void*
-		return TakeOwnership(const_cast<void*>(Data), Size, Forward<DeleteFunctionType>(DeleteFunction));
-	}
+template <typename DeleteFunctionType,
+	decltype(Invoke(std::declval<DeleteFunctionType>(), std::declval<void*>(), std::declval<uint64>()))*>
+inline FSharedBuffer FSharedBuffer::TakeOwnership(const void* Data, uint64 Size, DeleteFunctionType&& DeleteFunction)
+{
+	using OwnerType = BufferOwnerPrivate::TBufferOwnerDeleteFunction<DeleteFunctionType>;
+	return FSharedBuffer(new OwnerType(const_cast<void*>(Data), Size, Forward<DeleteFunctionType>(DeleteFunction)));
+}
 
-	/**
-	 * Make an owned read-only buffer by taking ownership of the provided memory.
-	 *
-	 * @param BufferOwner An owner for the buffer that is derived from FBufferOwner.
-	 */
-	template <typename BufferOwnerType,
-		decltype(ImplicitConv<FBufferOwner*>(std::declval<std::decay_t<BufferOwnerType>*>()))* = nullptr>
-	static inline FSharedBufferPtr TakeOwnership(void* Data, uint64 Size, BufferOwnerType&& BufferOwner)
-	{
-		return NewBufferWithOwner<BufferOwnerType>(Data, Size, ESharedBufferFlags::Owned, Forward<BufferOwnerType>(BufferOwner));
-	}
+///////////////////////////////////////////////////////////////////////////////////////////////////////////////////////
 
-	/**
-	 * Make an owned read-only buffer by taking ownership of the provided memory.
-	 *
-	 * @param BufferOwner An owner for the buffer that is derived from FBufferOwner.
-	 */
-	template <typename BufferOwnerType,
-		decltype(&ImplicitConv<FBufferOwner*>(std::declval<std::decay_t<BufferOwnerType>*>()))* = nullptr>
-	static inline FSharedBufferConstPtr TakeOwnership(const void* Data, uint64 Size, BufferOwnerType&& BufferOwner)
-	{
-		// const_cast is safe here because the return value only allows accessing data through a const void*
-		return TakeOwnership(const_cast<void*>(Data), Size, Forward<BufferOwnerType>(BufferOwner));
-	}
+enum class FBufferOwner::EBufferOwnerFlags : uint8
+{
+	None         = 0,
+	Owned        = 1 << 0,
+	Materialized = 1 << 1,
+};
 
-	/**
-	 * Return the buffer if it is owned, or a clone otherwise.
-	 *
-	 * @param Buffer A FSharedBuffer[Const]Ref or FSharedBuffer[Const]Ptr to make owned.
-	 * @return An owned copy of the input buffer in the same reference type as the input.
-	 */
-	template <typename T,
-		std::enable_if_t<std::is_same<FSharedBuffer, std::decay_t<decltype(*std::declval<T>())>>::value>* = nullptr
-		>
-	static inline auto MakeOwned(T&& Buffer) -> decltype(std::decay_t<T>(ToSharedRef(Clone(Buffer->GetData(), Buffer->GetSize()))))
-	{
-		return !IsValid(Buffer) || Buffer->IsOwned() ? Forward<T>(Buffer) : ToSharedRef(Clone(Buffer->GetData(), Buffer->GetSize()));
-	}
+inline FBufferOwner::FBufferOwner(void* InData, uint64 InSize)
+	: Data(InData)
+	, Size(InSize)
+{
+}
 
-	/** Return the buffer if it is owned and either read-only or this is the only reference, or a clone otherwise. */
-	static inline FSharedBufferConstRef MakeReadOnly(FSharedBufferConstRef&& Buffer)
-	{
-		return Buffer->TryMakeReadOnly() ? MoveTemp(Buffer) : ToSharedRef(MakeReadOnly(Clone(*Buffer)));
-	}
+inline FBufferOwner::~FBufferOwner()
+{
+	checkSlow(!IsReferenced());
+}
 
-	/** Return the buffer if it is owned and either read-only or this is the only reference, or a clone otherwise. */
-	static inline FSharedBufferConstPtr MakeReadOnly(FSharedBufferPtr&& Buffer)
-	{
-		return MakeReadOnly(FSharedBufferConstPtr(MoveTemp(Buffer)));
-	}
+inline void FBufferOwner::MaterializeBuffer()
+{
+	SetIsMaterialized();
+}
 
-	/** Return the buffer if it is owned and either read-only or this is the only reference, or a clone otherwise. */
-	static inline FSharedBufferConstPtr MakeReadOnly(FSharedBufferConstPtr&& Buffer)
-	{
-		return Buffer ? Buffer->TryMakeReadOnly() ? MoveTemp(Buffer) : MakeReadOnly(Clone(*Buffer)) : FSharedBufferConstPtr();
-	}
+inline void* FBufferOwner::GetData()
+{
+	Materialize();
+	return Data;
+}
 
-	// Disable overloads that would allow the argument to be copied.
-	static void MakeReadOnly(FSharedBufferRef&&) = delete;
-	static void MakeReadOnly(const FSharedBufferRef&) = delete;
-	static void MakeReadOnly(const FSharedBufferConstRef&) = delete;
-	static void MakeReadOnly(const FSharedBufferPtr&) = delete;
-	static void MakeReadOnly(const FSharedBufferConstPtr&) = delete;
+inline uint64 FBufferOwner::GetSize()
+{
+	Materialize();
+	return Size;
+}
 
-	/** Return the buffer because it is already writable. */
-	static inline FSharedBufferRef MakeWritable(FSharedBufferRef&& Buffer) { return MoveTemp(Buffer); }
+inline void FBufferOwner::SetBuffer(void* InData, uint64 InSize)
+{
+	Data = InData;
+	Size = InSize;
+}
 
-	/** Return the buffer if it is owned and either writable or this is the only reference, or a clone otherwise. */
-	static inline FSharedBufferRef MakeWritable(FSharedBufferConstRef&& Buffer)
-	{
-		if (Buffer->TryMakeWritable())
-		{
-			return FSharedBufferRef(FSharedBufferRef::PtrType(const_cast<FSharedBuffer*>(&Buffer.Get())));
-		}
-		return Clone(*Buffer).ToSharedRef();
-	}
+inline bool FBufferOwner::IsOwned() const
+{
+	return (uint8(GetFlags(ReferenceCountsAndFlags)) & uint8(EBufferOwnerFlags::Owned)) != 0;
+}
 
-	/** Return the buffer because it is already writable. */
-	static inline FSharedBufferPtr MakeWritable(FSharedBufferPtr&& Buffer) { return MoveTemp(Buffer); }
+inline void FBufferOwner::SetIsOwned()
+{
+	ReferenceCountsAndFlags.fetch_or(SetFlags(EBufferOwnerFlags::Owned));
+}
 
-	/** Return the buffer if it is owned and either writable or this is the only reference, or a clone otherwise. */
-	static inline FSharedBufferPtr MakeWritable(FSharedBufferConstPtr&& Buffer)
+inline void FBufferOwner::Materialize()
+{
+	if (!IsMaterialized())
 	{
-		if (Buffer->TryMakeWritable())
-		{
-			return FSharedBufferPtr(FSharedBufferPtr::PtrType(const_cast<FSharedBuffer*>(Buffer.Get())));
-		}
-		return Clone(*Buffer);
+		MaterializeBuffer();
+		checkSlow(IsMaterialized());
 	}
+}
 
-	// Disable overloads that would allow the argument to be copied.
-	static void MakeWritable(const FSharedBufferRef&) = delete;
-	static void MakeWritable(const FSharedBufferConstRef&) = delete;
-	static void MakeWritable(const FSharedBufferPtr&) = delete;
-	static void MakeWritable(const FSharedBufferConstPtr&) = delete;
-
-	/** A pointer to the start of the buffer. */
-	inline void* GetData() { return Data; }
-	/** A pointer to the start of the buffer. */
-	inline const void* GetData() const { return Data; }
+inline bool FBufferOwner::IsMaterialized() const
+{
+	return (uint8(GetFlags(ReferenceCountsAndFlags)) & uint8(EBufferOwnerFlags::Materialized)) != 0;
+}
 
-	/** The size of the buffer in bytes. */
-	inline uint64 GetSize() const { return Size; }
+inline void FBufferOwner::SetIsMaterialized()
+{
+	ReferenceCountsAndFlags.fetch_or(SetFlags(EBufferOwnerFlags::Materialized));
+}
 
-	/** Whether this shared buffer owns the memory that it provides a view of. */
-	inline bool IsOwned() const
-	{
-		return EnumHasAnyFlags(GetFlags(ReferenceCountAndFlags), ESharedBufferFlags::Owned);
-	}
+inline bool FBufferOwner::IsReferenced() const
+{
+	const uint64 Value = ReferenceCountsAndFlags.load(std::memory_order_relaxed);
+	return GetSharedRefCount(Value) || GetWeakRefCount(Value);
+}
 
-	/**
-	 * Whether the shared buffer is read-only.
-	 *
-	 * A read-only shared buffer (or its outer) is owned and every reference to it is const.
-	 */
-	inline bool IsReadOnly() const
+inline void FBufferOwner::AddSharedReference()
+{
+	const uint64 PreviousValue = ReferenceCountsAndFlags.fetch_add(SetSharedRefCount(1), std::memory_order_relaxed);
+	checkSlow(GetSharedRefCount(PreviousValue) < 0x7fffffff);
+	if (GetSharedRefCount(PreviousValue) == 0)
 	{
-		return EnumHasAnyFlags(GetFlags(ReferenceCountAndFlags), ESharedBufferFlags::ReadOnly);
+		AddWeakReference();
 	}
+}
 
-	/**
-	 * Try to make the shared buffer read-only.
-	 *
-	 * A shared buffer can be made read-only if it is owned and there is only one reference to it.
-	 *
-	 * @return true if the buffer was read-only or was made read-only, otherwise false.
-	 */
-	CORE_API bool TryMakeReadOnly() const;
-
-	// Disable this overload because it would allow a read-only buffer to be written.
-	bool TryMakeReadOnly() = delete;
-
-	/**
-	 * Try to make the shared buffer writable.
-	 *
-	 * A shared buffer can be made writable if it is owned and there is only one reference to it.
-	 *
-	 * @return true if the buffer was writable or was made writable, otherwise false.
-	 */
-	CORE_API bool TryMakeWritable() const;
-
-	/** The shared buffer is already writable. */
-	inline bool TryMakeWritable() { return true; }
-
-	/** A writable view of buffer. */
-	inline FMutableMemoryView GetView() { return FMutableMemoryView(GetData(), GetSize()); }
-	/** A read-only view of the buffer. */
-	inline FMemoryView GetView() const { return FMemoryView(GetData(), GetSize()); }
-
-	inline operator FMutableMemoryView() { return GetView(); }
-	inline operator FMemoryView() const { return GetView(); }
-
-private:
-	/**
-	 * Construct a new shared buffer with a buffer owner.
-	 *
-	 * @param Data A pointer to the start of the data buffer.
-	 * @param Size The size of the data buffer.
-	 * @param Flags The flags associated with the data buffer. Internally adds HasBufferOwner.
-	 * @param OwnerArgs Arguments to the constructor of the buffer owner.
-	 */
-	template <typename OwnerType, typename... OwnerArgTypes>
-	static inline FSharedBufferRef NewBufferWithOwner(
-		void* const Data,
-		const uint64 Size,
-		ESharedBufferFlags Flags,
-		OwnerArgTypes&&... OwnerArgs)
+inline void FBufferOwner::ReleaseSharedReference()
+{
+	const uint64 PreviousValue = ReferenceCountsAndFlags.fetch_sub(SetSharedRefCount(1), std::memory_order_acq_rel);
+	checkSlow(GetSharedRefCount(PreviousValue) > 0);
+	if (GetSharedRefCount(PreviousValue) == 1)
 	{
-		Flags |= ESharedBufferFlags::HasBufferOwner;
-		static_assert(alignof(OwnerType) <= alignof(FSharedBuffer), "Required alignment of OwnerType is too high.");
-		FSharedBufferRef Buffer = NewBuffer(Data, Size, Flags, sizeof(OwnerType));
-		new(&Buffer.Get() + 1) OwnerType(Forward<OwnerArgTypes>(OwnerArgs)...);
-		return Buffer;
+		FreeBuffer();
+		Data = nullptr;
+		Size = 0;
+		ReleaseWeakReference();
 	}
+}
 
-	/**
-	 * Construct a new shared buffer.
-	 *
-	 * @param Data A pointer to the start of the data buffer.
-	 * @param Size The size of the data buffer.
-	 * @param Flags The flags associated with the data buffer.
-	 * @param OwnerSize The size of the optional buffer owner to allocate after the shared buffer.
-	 */
-	CORE_API static FSharedBufferRef NewBuffer(void* Data, uint64 Size, ESharedBufferFlags Flags, uint32 OwnerSize = 0);
-
-	/**
-	 * Delete a shared buffer. Requires ReleaseData to have been called already.
-	 */
-	CORE_API static void DeleteBuffer(FSharedBuffer* Buffer);
-
-	/**
-	 * Free the data if it is owned and release the pointer to it.
-	 *
-	 * This does not delete the shared buffer object, which survives until the last shared or weak
-	 * reference to it has been released.
-	 */
-	CORE_API void ReleaseData();
-
-	FSharedBuffer() = default;
-	FSharedBuffer(const FSharedBuffer&) = delete;
-	FSharedBuffer& operator=(const FSharedBuffer&) = delete;
-	~FSharedBuffer() = default;
-
-	/** A buffer owner that holds a reference to an outer shared buffer. */
-	class FOuterBufferOwner final : public FBufferOwner
+inline bool FBufferOwner::TryAddSharedReference()
+{
+	for (uint64 Value = ReferenceCountsAndFlags.load(std::memory_order_relaxed);;)
 	{
-	public:
-		explicit FOuterBufferOwner(FSharedBufferConstPtr&& InOuterBuffer)
-			: OuterBuffer(MoveTemp(InOuterBuffer))
+		if (GetSharedRefCount(Value) == 0)
 		{
+			return false;
 		}
-
-	private:
-		virtual void Free(void* InData, uint64 InSize) final
+		if (ReferenceCountsAndFlags.compare_exchange_weak(Value, Value + SetSharedRefCount(1),
+			std::memory_order_relaxed, std::memory_order_relaxed))
 		{
-			// Destruction of the OuterBuffer reference will release the buffer.
+			return true;
 		}
+	}
+}
 
-		FSharedBufferConstPtr OuterBuffer;
-	};
+inline void FBufferOwner::AddWeakReference()
+{
+	const uint64 PreviousValue = ReferenceCountsAndFlags.fetch_add(SetWeakRefCount(1), std::memory_order_relaxed);
+	checkSlow(GetWeakRefCount(PreviousValue) < 0x7fffffff);
+}
 
-	/** A buffer owner that wraps an arbitrary callable delete function. */
-	template <typename DeleteFunctionType>
-	class TDeleteFunctionOwner final : public FBufferOwner
+inline void FBufferOwner::ReleaseWeakReference()
+{
+	const uint64 PreviousValue = ReferenceCountsAndFlags.fetch_sub(SetWeakRefCount(1), std::memory_order_acq_rel);
+	checkSlow(GetWeakRefCount(PreviousValue) > 0);
+	if (GetWeakRefCount(PreviousValue) == 1)
 	{
-	public:
-		explicit TDeleteFunctionOwner(DeleteFunctionType&& InDeleteFunction)
-			: InDeleteFunction(Forward<DeleteFunctionType>(InDeleteFunction))
-		{
-		}
-
-	private:
-		virtual void Free(void* InData, uint64 InSize) final
-		{
-			Invoke(InDeleteFunction, InData);
-		}
-
-		std::decay_t<DeleteFunctionType> InDeleteFunction;
-	};
+		delete this;
+	}
+}
 
-private:
-	static uint32 GetSharedRefCount(uint64 RefCountsAndFlags) { return uint32(RefCountsAndFlags >> 0) & 0xffffff; }
-	static uint64 SetSharedRefCount(uint32 RefCount) { return uint64(RefCount) << 0; }
-	static uint32 GetWeakRefCount(uint64 RefCountsAndFlags) { return uint32(RefCountsAndFlags >> 24) & 0xffffff; }
-	static uint64 SetWeakRefCount(uint32 RefCount) { return uint64(RefCount) << 24; }
-	static ESharedBufferFlags GetFlags(uint64 RefCountsAndFlags) { return ESharedBufferFlags(RefCountsAndFlags >> 48); }
-	static uint64 SetFlags(ESharedBufferFlags Flags) { return uint64(Flags) << 48; }
+///////////////////////////////////////////////////////////////////////////////////////////////////////////////////////
 
-	inline void AddSharedReference() const
-	{
-		const uint64 PreviousValue = ReferenceCountAndFlags.fetch_add(SetSharedRefCount(1), std::memory_order_relaxed);
-		checkSlow(GetSharedRefCount(PreviousValue) < 0xffffff);
-		if (GetSharedRefCount(PreviousValue) == 0)
-		{
-			AddWeakReference();
-		}
-	}
+namespace BufferOwnerPrivate
+{
 
-	inline void ReleaseSharedReference() const
+template <typename FOps>
+template <typename FOtherOps>
+inline FBufferOwner* TBufferOwnerPtr<FOps>::CopyFrom(const TBufferOwnerPtr<FOtherOps>& Ptr)
+{
+	FBufferOwner* NewOwner = Ptr.Owner;
+	if (NewOwner)
 	{
-		const uint64 PreviousValue = ReferenceCountAndFlags.fetch_sub(SetSharedRefCount(1), std::memory_order_acq_rel);
-		checkSlow(GetSharedRefCount(PreviousValue) > 0);
-		if (GetSharedRefCount(PreviousValue) == 1)
+		if (bIsWeak || !Ptr.bIsWeak)
 		{
-			const_cast<FSharedBuffer&>(*this).ReleaseData();
-			ReleaseWeakReference();
+			FOps::AddRef(*NewOwner);
 		}
-	}
-
-	inline bool TryAddSharedReference() const
-	{
-		for (uint64 Value = ReferenceCountAndFlags.load(std::memory_order_relaxed);;)
+		else if (!FOps::TryAddRef(*NewOwner))
 		{
-			if (GetSharedRefCount(Value) == 0)
-			{
-				return false;
-			}
-			if (ReferenceCountAndFlags.compare_exchange_weak(Value, Value + SetSharedRefCount(1),
-				std::memory_order_relaxed, std::memory_order_relaxed))
-			{
-				return true;
-			}
+			NewOwner = nullptr;
 		}
 	}
+	return NewOwner;
+}
 
-	inline void AddWeakReference() const
+template <typename FOps>
+template <typename FOtherOps>
+inline FBufferOwner* TBufferOwnerPtr<FOps>::MoveFrom(TBufferOwnerPtr<FOtherOps>&& Ptr)
+{
+	FBufferOwner* NewOwner = Ptr.Owner;
+	if (bIsWeak == Ptr.bIsWeak)
 	{
-		const uint64 PreviousValue = ReferenceCountAndFlags.fetch_add(SetWeakRefCount(1), std::memory_order_relaxed);
-		checkSlow(GetWeakRefCount(PreviousValue) < 0xffffff);
+		Ptr.Owner = nullptr;
 	}
-
-	inline void ReleaseWeakReference() const
+	else if (NewOwner)
 	{
-		const uint64 PreviousValue = ReferenceCountAndFlags.fetch_sub(SetWeakRefCount(1), std::memory_order_acq_rel);
-		checkSlow(GetWeakRefCount(PreviousValue) > 0);
-		if (GetWeakRefCount(PreviousValue) == 1)
+		if (bIsWeak)
 		{
-			DeleteBuffer(const_cast<FSharedBuffer*>(this));
+			FOps::AddRef(*NewOwner);
+		}
+		else if (!FOps::TryAddRef(*NewOwner))
+		{
+			NewOwner = nullptr;
 		}
 	}
+	return NewOwner;
+}
 
-	friend struct SharedBufferPrivate::FSharedRefOps;
-	friend struct SharedBufferPrivate::FSharedPtrOps;
-	friend struct SharedBufferPrivate::FWeakPtrOps;
-
-private:
-	void* Data;
-	uint64 Size;
-	mutable std::atomic<uint64> ReferenceCountAndFlags;
-};
-
-///////////////////////////////////////////////////////////////////////////////////////////////////////////////////////
-
-namespace SharedBufferPrivate
+template <typename FOps>
+inline TBufferOwnerPtr<FOps>::TBufferOwnerPtr(const TBufferOwnerPtr& Ptr)
+	: Owner(CopyFrom(Ptr))
 {
+}
 
-inline bool FSharedRefOps::TryAddRef(const FSharedBuffer* const Buffer)
+template <typename FOps>
+inline TBufferOwnerPtr<FOps>::TBufferOwnerPtr(TBufferOwnerPtr&& Ptr)
+	: Owner(MoveFrom(MoveTemp(Ptr)))
 {
-	check(Buffer);
-	verify(Buffer->TryAddSharedReference());
-	return true;
 }
 
-inline void FSharedRefOps::AddRef(const FSharedBuffer* const Buffer)
+template <typename FOps>
+template <typename FOtherOps>
+inline TBufferOwnerPtr<FOps>::TBufferOwnerPtr(const TBufferOwnerPtr<FOtherOps>& Ptr)
+	: Owner(CopyFrom(Ptr))
 {
-	check(Buffer);
-	Buffer->AddSharedReference();
 }
 
-inline void FSharedRefOps::Release(const FSharedBuffer* const Buffer)
+template <typename FOps>
+template <typename FOtherOps>
+inline TBufferOwnerPtr<FOps>::TBufferOwnerPtr(TBufferOwnerPtr<FOtherOps>&& Ptr)
+	: Owner(MoveFrom(MoveTemp(Ptr)))
 {
-	check(Buffer);
-	Buffer->ReleaseSharedReference();
 }
 
-inline bool FSharedPtrOps::TryAddRef(const FSharedBuffer* const Buffer)
+template <typename FOps>
+inline TBufferOwnerPtr<FOps>::~TBufferOwnerPtr()
 {
-	return Buffer && Buffer->TryAddSharedReference();
+	FOps::Release(Owner);
 }
 
-inline void FSharedPtrOps::AddRef(const FSharedBuffer* const Buffer)
+template <typename FOps>
+inline TBufferOwnerPtr<FOps>& TBufferOwnerPtr<FOps>::operator=(const TBufferOwnerPtr& Ptr)
 {
-	if (Buffer)
-	{
-		Buffer->AddSharedReference();
-	}
+	FBufferOwner* const OldOwner = Owner;
+	Owner = CopyFrom(Ptr);
+	FOps::Release(OldOwner);
+	return *this;
 }
 
-inline void FSharedPtrOps::Release(const FSharedBuffer* const Buffer)
+template <typename FOps>
+inline TBufferOwnerPtr<FOps>& TBufferOwnerPtr<FOps>::operator=(TBufferOwnerPtr&& Ptr)
 {
-	if (Buffer)
-	{
-		Buffer->ReleaseSharedReference();
-	}
+	FBufferOwner* const OldOwner = Owner;
+	Owner = MoveFrom(MoveTemp(Ptr));
+	FOps::Release(OldOwner);
+	return *this;
 }
 
-inline bool FWeakPtrOps::TryAddRef(const FSharedBuffer* const Buffer)
+template <typename FOps>
+template <typename FOtherOps>
+inline TBufferOwnerPtr<FOps>& TBufferOwnerPtr<FOps>::operator=(const TBufferOwnerPtr<FOtherOps>& Ptr)
 {
-	if (Buffer)
-	{
-		Buffer->AddWeakReference();
-		return true;
-	}
-	return false;
+	FBufferOwner* const OldOwner = Owner;
+	Owner = CopyFrom(Ptr);
+	FOps::Release(OldOwner);
+	return *this;
 }
 
-inline void FWeakPtrOps::AddRef(const FSharedBuffer* const Buffer)
+template <typename FOps>
+template <typename FOtherOps>
+inline TBufferOwnerPtr<FOps>& TBufferOwnerPtr<FOps>::operator=(TBufferOwnerPtr<FOtherOps>&& Ptr)
 {
-	if (Buffer)
-	{
-		Buffer->AddWeakReference();
-	}
+	FBufferOwner* const OldOwner = Owner;
+	Owner = MoveFrom(MoveTemp(Ptr));
+	FOps::Release(OldOwner);
+	return *this;
 }
 
-inline void FWeakPtrOps::Release(const FSharedBuffer* const Buffer)
+template <typename FOps>
+template <typename FOtherOps>
+inline bool TBufferOwnerPtr<FOps>::operator==(const TBufferOwnerPtr<FOtherOps>& Ptr) const
 {
-	if (Buffer)
-	{
-		Buffer->ReleaseWeakReference();
-	}
+	return Owner == Ptr.Owner;
 }
 
+template <typename FOps>
+template <typename FOtherOps>
+inline bool TBufferOwnerPtr<FOps>::operator!=(const TBufferOwnerPtr<FOtherOps>& Ptr) const
+{
+	return Owner != Ptr.Owner;
 }
 
+} // BufferOwnerPrivate
+
 ///////////////////////////////////////////////////////////////////////////////////////////////////////////////////////
diff --git a/Engine/Source/Runtime/Core/Public/Serialization/CompactBinary.h b/Engine/Source/Runtime/Core/Public/Serialization/CompactBinary.h
--- a/Engine/Source/Runtime/Core/Public/Serialization/CompactBinary.h
+++ b/Engine/Source/Runtime/Core/Public/Serialization/CompactBinary.h
@@ -38,7 +38,7 @@
  *
  * void BeginBuild(FCbObjectRef Params)
  * {
- *     if (FSharedBufferConstPtr Data = Storage().Load(Params["Data"].AsBinaryReference()))
+ *     if (FSharedBuffer Data = Storage().Load(Params["Data"].AsBinaryReference()))
  *     {
  *         SetData(Data);
  *     }
@@ -900,7 +900,7 @@ private:
 ///////////////////////////////////////////////////////////////////////////////////////////////////////////////////////
 
 /** A reference to a function that is used to allocate buffers for compact binary data. */
-using FCbBufferAllocator = TFunctionRef<FSharedBufferPtr (uint64 Size)>;
+using FCbBufferAllocator = TFunctionRef<FUniqueBuffer (uint64 Size)>;
 
 ///////////////////////////////////////////////////////////////////////////////////////////////////////////////////////
 
@@ -918,23 +918,23 @@ public:
 	 * @param ValueBuffer A buffer that exactly contains the value.
 	 * @param Type HasFieldType means that ValueBuffer contains the type. Otherwise, use the given type.
 	 */
-	inline explicit TCbBufferRef(FSharedBufferConstPtr ValueBuffer, ECbFieldType Type = ECbFieldType::HasFieldType)
+	inline explicit TCbBufferRef(FSharedBuffer ValueBuffer, ECbFieldType Type = ECbFieldType::HasFieldType)
 	{
 		if (ValueBuffer)
 		{
-			BaseType::operator=(BaseType(ValueBuffer->GetData(), Type));
-			check(ValueBuffer->GetView().Contains(BaseType::GetView()));
+			BaseType::operator=(BaseType(ValueBuffer.GetData(), Type));
+			check(ValueBuffer.GetView().Contains(BaseType::GetView()));
 			Buffer = MoveTemp(ValueBuffer);
 		}
 	}
 
 	/** Construct a value that holds a reference to the buffer that contains it. */
-	inline TCbBufferRef(const BaseType& Value, FSharedBufferConstPtr OuterBuffer)
+	inline TCbBufferRef(const BaseType& Value, FSharedBuffer OuterBuffer)
 		: BaseType(Value)
 	{
 		if (OuterBuffer)
 		{
-			check(OuterBuffer->GetView().Contains(BaseType::GetView()));
+			check(OuterBuffer.GetView().Contains(BaseType::GetView()));
 			Buffer = MoveTemp(OuterBuffer);
 		}
 	}
@@ -947,66 +947,37 @@ public:
 	}
 
 	/** Whether this reference has ownership of the memory in its buffer. */
-	inline bool IsOwned() const { return Buffer && Buffer->IsOwned(); }
+	inline bool IsOwned() const { return Buffer.IsOwned(); }
 
 	/** Clone the value, if necessary, to a buffer that this reference has ownership of. */
 	inline void MakeOwned()
 	{
 		if (!IsOwned())
 		{
-			FSharedBufferPtr MutableBuffer = FSharedBuffer::Alloc(BaseType::GetSize());
-			BaseType::CopyTo(*MutableBuffer);
-			BaseType::operator=(BaseType(MutableBuffer->GetData()));
-			Buffer = FSharedBuffer::MakeReadOnly(MoveTemp(MutableBuffer));
-		}
-	}
-
-	/** Whether this reference is backed by a read-only buffer. */
-	inline bool IsReadOnly() const { return Buffer && Buffer->IsReadOnly(); }
-
-	/** Clone the value, if necessary, to a read-only buffer. */
-	inline void MakeReadOnly()
-	{
-		if (!IsReadOnly())
-		{
-			FSharedBufferPtr MutableBuffer = FSharedBuffer::Alloc(BaseType::GetSize());
-			BaseType::CopyTo(*MutableBuffer);
-			BaseType::operator=(BaseType(MutableBuffer->GetData()));
-			Buffer = FSharedBuffer::MakeReadOnly(MoveTemp(MutableBuffer));
+			FUniqueBuffer MutableBuffer = FUniqueBuffer::Alloc(BaseType::GetSize());
+			BaseType::CopyTo(MutableBuffer);
+			BaseType::operator=(BaseType(MutableBuffer.GetData()));
+			Buffer = MoveTemp(MutableBuffer);
 		}
 	}
 
 	/** Returns a buffer that exactly contains this value. */
-	inline FSharedBufferConstPtr GetBuffer() const
+	inline FSharedBuffer GetBuffer() const
 	{
 		const FMemoryView View = BaseType::GetView();
-		const FSharedBufferConstPtr& OuterBuffer = GetOuterBuffer();
-		if (OuterBuffer)
-		{
-			if (OuterBuffer->GetView() == View)
-			{
-				return OuterBuffer;
-			}
-			else
-			{
-				return FSharedBuffer::MakeView(View, *OuterBuffer);
-			}
-		}
-		else
-		{
-			return FSharedBuffer::MakeView(View);
-		}
+		const FSharedBuffer& OuterBuffer = GetOuterBuffer();
+		return View == OuterBuffer.GetView() ? OuterBuffer : FSharedBuffer::MakeView(View, OuterBuffer);
 	}
 
 	/** Returns the outer buffer (if any) that contains this value. */
-	inline const FSharedBufferConstPtr& GetOuterBuffer() const & { return Buffer; }
-	inline FSharedBufferConstPtr GetOuterBuffer() && { return MoveTemp(Buffer); }
+	inline const FSharedBuffer& GetOuterBuffer() const & { return Buffer; }
+	inline FSharedBuffer GetOuterBuffer() && { return MoveTemp(Buffer); }
 
 private:
 	template <typename OtherType>
 	friend class TCbBufferRef;
 
-	FSharedBufferConstPtr Buffer;
+	FSharedBuffer Buffer;
 };
 
 ///////////////////////////////////////////////////////////////////////////////////////////////////////////////////////
@@ -1036,13 +1007,13 @@ public:
 	}
 
 	/** Construct a value from a read-only view of its memory and its optional outer buffer. */
-	static inline RefType MakeView(const void* const Data, FSharedBufferConstPtr OuterBuffer = FSharedBufferConstPtr())
+	static inline RefType MakeView(const void* const Data, FSharedBuffer OuterBuffer = FSharedBuffer())
 	{
 		return MakeView(BaseType(Data), MoveTemp(OuterBuffer));
 	}
 
 	/** Construct a value from a read-only view of its memory and its optional outer buffer. */
-	static inline RefType MakeView(const BaseType& Value, FSharedBufferConstPtr OuterBuffer = FSharedBufferConstPtr())
+	static inline RefType MakeView(const BaseType& Value, FSharedBuffer OuterBuffer = FSharedBuffer())
 	{
 		return RefType(Value, MoveTemp(OuterBuffer));
 	}
@@ -1107,18 +1078,18 @@ public:
 	 * @param Buffer A buffer containing zero or more valid fields.
 	 * @param Type HasFieldType means that Buffer contains the type. Otherwise, use the given type.
 	 */
-	static inline FCbFieldRefIterator MakeRange(FSharedBufferConstPtr Buffer, ECbFieldType Type = ECbFieldType::HasFieldType)
+	static inline FCbFieldRefIterator MakeRange(FSharedBuffer Buffer, ECbFieldType Type = ECbFieldType::HasFieldType)
 	{
-		if (Buffer && Buffer->GetSize())
+		if (Buffer.GetSize())
 		{
-			const void* const DataEnd = Buffer->GetView().GetDataEnd();
+			const void* const DataEnd = Buffer.GetView().GetDataEnd();
 			return FCbFieldRefIterator(FCbFieldRef(MoveTemp(Buffer), Type), DataEnd);
 		}
 		return FCbFieldRefIterator();
 	}
 
 	/** Construct a field range from an iterator and its optional outer buffer. */
-	static inline FCbFieldRefIterator MakeRangeView(const FCbFieldIterator& It, FSharedBufferConstPtr OuterBuffer = FSharedBufferConstPtr())
+	static inline FCbFieldRefIterator MakeRangeView(const FCbFieldIterator& It, FSharedBuffer OuterBuffer = FSharedBuffer())
 	{
 		return FCbFieldRefIterator(FCbFieldRef(It, MoveTemp(OuterBuffer)), GetFieldsEnd(It));
 	}
@@ -1136,7 +1107,7 @@ public:
 	}
 
 	/** Returns a buffer that exactly contains the field range. */
-	CORE_API FSharedBufferConstPtr GetRangeBuffer() const;
+	CORE_API FSharedBuffer GetRangeBuffer() const;
 
 private:
 	using TCbFieldIterator::TCbFieldIterator;
diff --git a/Engine/Source/Runtime/Core/Public/Serialization/CompactBinaryPackage.h b/Engine/Source/Runtime/Core/Public/Serialization/CompactBinaryPackage.h
--- a/Engine/Source/Runtime/Core/Public/Serialization/CompactBinaryPackage.h
+++ b/Engine/Source/Runtime/Core/Public/Serialization/CompactBinaryPackage.h
@@ -29,26 +29,26 @@ public:
 	/** Construct a null attachment. */
 	FCbAttachment() = default;
 
-	/** Construct a compact binary attachment. Value is cloned unless read-only. */
+	/** Construct a compact binary attachment. Value is cloned if not owned. */
 	inline explicit FCbAttachment(FCbFieldRefIterator Value)
 		: FCbAttachment(MoveTemp(Value), nullptr)
 	{
 	}
 
-	/** Construct a compact binary attachment. Value is cloned unless read-only. Hash must match Value. */
+	/** Construct a compact binary attachment. Value is cloned if not owned. Hash must match Value. */
 	inline explicit FCbAttachment(FCbFieldRefIterator Value, const FBlake3Hash& Hash)
 		: FCbAttachment(MoveTemp(Value), &Hash)
 	{
 	}
 
-	/** Construct a binary attachment. Value is cloned unless read-only. */
-	inline explicit FCbAttachment(FSharedBufferConstPtr Value)
+	/** Construct a binary attachment. Value is cloned if not owned. */
+	inline explicit FCbAttachment(FSharedBuffer Value)
 		: FCbAttachment(MoveTemp(Value), nullptr)
 	{
 	}
 
-	/** Construct a binary attachment. Value is cloned unless read-only. Hash must match Value. */
-	inline explicit FCbAttachment(FSharedBufferConstPtr Value, const FBlake3Hash& Hash)
+	/** Construct a binary attachment. Value is cloned if not owned. Hash must match Value. */
+	inline explicit FCbAttachment(FSharedBuffer Value, const FBlake3Hash& Hash)
 		: FCbAttachment(MoveTemp(Value), &Hash)
 	{
 	}
@@ -60,13 +60,13 @@ public:
 	inline bool IsNull() const { return !Buffer; }
 
 	/** Access the attachment as binary. Defaults to a null buffer on error. */
-	CORE_API FSharedBufferConstPtr AsBinary() const;
+	CORE_API FSharedBuffer AsBinary() const;
 
 	/** Access the attachment as compact binary. Defaults to a field iterator with no value on error. */
 	CORE_API FCbFieldRefIterator AsCompactBinary() const;
 
 	/** Returns whether the attachment is binary or compact binary. */
-	inline bool IsBinary() const { return Buffer.IsValid(); }
+	inline bool IsBinary() const { return !Buffer.IsNull(); }
 
 	/** Returns whether the attachment is compact binary. */
 	inline bool IsCompactBinary() const { return CompactBinary.HasValue(); }
@@ -82,8 +82,7 @@ public:
 	/**
 	 * Load the attachment from compact binary as written by Save.
 	 *
-	 * The attachment will hold a reference to the input iterator, if it is read-only and owned, and
-	 * otherwise clones the attachment to make a read-only copy.
+	 * The attachment references the input iterator if it is owned, and otherwise clones the value.
 	 *
 	 * The iterator is advanced as attachment fields are consumed from it.
 	 */
@@ -92,9 +91,12 @@ public:
 	/**
 	 * Load the attachment from compact binary as written by Save.
 	 *
-	 * The attachment will be read into a read-only, owned buffer.
+	 * The attachments value will be loaded into an owned buffer.
+	 *
+	 * @param Allocator Allocator for the attachment value buffer.
+	 * @note Allocated buffers will be cloned if they are not owned.
 	 */
-	CORE_API void Load(FArchive& Ar, FCbBufferAllocator Allocator);
+	CORE_API void Load(FArchive& Ar, FCbBufferAllocator Allocator = FUniqueBuffer::Alloc);
 
 	/** Save the attachment into the writer as a stream of compact binary fields. */
 	CORE_API void Save(FCbWriter& Writer) const;
@@ -104,10 +106,10 @@ public:
 
 private:
 	CORE_API FCbAttachment(FCbFieldRefIterator Value, const FBlake3Hash* Hash);
-	CORE_API FCbAttachment(FSharedBufferConstPtr Value, const FBlake3Hash* Hash);
+	CORE_API FCbAttachment(FSharedBuffer Value, const FBlake3Hash* Hash);
 
-	/** A read-only owned buffer containing the binary or compact binary data. */
-	FSharedBufferConstPtr Buffer;
+	/** An owned buffer containing the binary or compact binary data. */
+	FSharedBuffer Buffer;
 	/** A field iterator that is valid only for compact binary attachments. */
 	FCbFieldIterator CompactBinary;
 	/** A hash of the attachment value. */
@@ -151,7 +153,7 @@ public:
 	 *
 	 * The resolver may return a null buffer to skip resolving an attachment for the hash.
 	 */
-	using FAttachmentResolver = TFunctionRef<FSharedBufferConstPtr (const FBlake3Hash& Hash)>;
+	using FAttachmentResolver = TFunctionRef<FSharedBuffer (const FBlake3Hash& Hash)>;
 
 	/** Construct a null package. */
 	FCbPackage() = default;
@@ -159,7 +161,7 @@ public:
 	/**
 	 * Construct a package from a root object without gathering attachments.
 	 *
-	 * @param InObject The root object, which will be cloned unless it is read-only.
+	 * @param InObject The root object, which will be cloned unless it is owned.
 	 */
 	inline explicit FCbPackage(FCbObjectRef InObject)
 	{
@@ -169,7 +171,7 @@ public:
 	/**
 	 * Construct a package from a root object and gather attachments using the resolver.
 	 *
-	 * @param InObject The root object, which will be cloned unless it is read-only.
+	 * @param InObject The root object, which will be cloned unless it is owned.
 	 * @param InResolver A function that is invoked for every reference and binary reference field.
 	 */
 	inline explicit FCbPackage(FCbObjectRef InObject, FAttachmentResolver InResolver)
@@ -180,7 +182,7 @@ public:
 	/**
 	 * Construct a package from a root object without gathering attachments.
 	 *
-	 * @param InObject The root object, which will be cloned unless it is read-only.
+	 * @param InObject The root object, which will be cloned unless it is owned.
 	 * @param InObjectHash The hash of the object, which must match to avoid validation errors.
 	 */
 	inline explicit FCbPackage(FCbObjectRef InObject, const FBlake3Hash& InObjectHash)
@@ -191,7 +193,7 @@ public:
 	/**
 	 * Construct a package from a root object and gather attachments using the resolver.
 	 *
-	 * @param InObject The root object, which will be cloned unless it is read-only.
+	 * @param InObject The root object, which will be cloned unless it is owned.
 	 * @param InObjectHash The hash of the object, which must match to avoid validation errors.
 	 * @param InResolver A function that is invoked for every reference and binary reference field.
 	 */
@@ -218,7 +220,7 @@ public:
 	/**
 	 * Set the root object without gathering attachments.
 	 *
-	 * @param InObject The root object, which will be cloned unless it is read-only.
+	 * @param InObject The root object, which will be cloned unless it is owned.
 	 */
 	inline void SetObject(FCbObjectRef InObject)
 	{
@@ -228,7 +230,7 @@ public:
 	/**
 	 * Set the root object and gather attachments using the resolver.
 	 *
-	 * @param InObject The root object, which will be cloned unless it is read-only.
+	 * @param InObject The root object, which will be cloned unless it is owned.
 	 * @param InResolver A function that is invoked for every reference and binary reference field.
 	 */
 	inline void SetObject(FCbObjectRef InObject, FAttachmentResolver InResolver)
@@ -239,7 +241,7 @@ public:
 	/**
 	 * Set the root object without gathering attachments.
 	 *
-	 * @param InObject The root object, which will be cloned unless it is read-only.
+	 * @param InObject The root object, which will be cloned unless it is owned.
 	 * @param InObjectHash The hash of the object, which must match to avoid validation errors.
 	 */
 	inline void SetObject(FCbObjectRef InObject, const FBlake3Hash& InObjectHash)
@@ -250,7 +252,7 @@ public:
 	/**
 	 * Set the root object and gather attachments using the resolver.
 	 *
-	 * @param InObject The root object, which will be cloned unless it is read-only.
+	 * @param InObject The root object, which will be cloned unless it is owned.
 	 * @param InObjectHash The hash of the object, which must match to avoid validation errors.
 	 * @param InResolver A function that is invoked for every reference and binary reference field.
 	 */
@@ -304,19 +306,22 @@ public:
 	/**
 	 * Load the object and attachments from compact binary as written by Save.
 	 *
-	 * The object and attachments will hold a reference to the input package, if it is read-only and
-	 * owned, and otherwise individually clones the object and attachments to make read-only copies.
+	 * The object and attachments reference the input iterator, if it is owned, and otherwise clones
+	 * the object and attachments individually to make owned copies.
 	 *
-	 * The iterator is advanced as package and attachment fields are consumed from it.
+	 * The iterator is advanced as object and attachment fields are consumed from it.
 	 */
 	CORE_API void Load(FCbFieldRefIterator& Fields);
 
 	/**
 	 * Load the object and attachments from compact binary as written by Save.
 	 *
-	 * The object and attachments will be individually read into read-only, owned buffers.
+	 * The object and attachments will be individually loaded into owned buffers.
+	 *
+	 * @param Allocator Allocator for object and attachment buffers.
+	 * @note Allocated buffers will be cloned if they are not owned.
 	 */
-	CORE_API void Load(FArchive& Ar, FCbBufferAllocator Allocator);
+	CORE_API void Load(FArchive& Ar, FCbBufferAllocator Allocator = FUniqueBuffer::Alloc);
 
 	/** Save the object and attachments into the writer as a stream of compact binary fields. */
 	CORE_API void Save(FCbWriter& Writer) const;
diff --git a/Engine/Source/Runtime/Core/Public/Serialization/CompactBinarySerialization.h b/Engine/Source/Runtime/Core/Public/Serialization/CompactBinarySerialization.h
--- a/Engine/Source/Runtime/Core/Public/Serialization/CompactBinarySerialization.h
+++ b/Engine/Source/Runtime/Core/Public/Serialization/CompactBinarySerialization.h
@@ -4,6 +4,7 @@
 
 #include "CoreTypes.h"
 #include "Memory/MemoryView.h"
+#include "Memory/SharedBuffer.h"
 #include "Serialization/CompactBinary.h"
 
 class FArchive;
@@ -61,7 +62,7 @@ CORE_API bool TryMeasureCompactBinary(
  * @param Allocator Allocator for the buffer that the field is loaded into.
  * @return A field with a reference to the provided buffer if it is owned.
  */
-CORE_API FCbFieldRef LoadCompactBinary(FArchive& Ar, FCbBufferAllocator Allocator);
+CORE_API FCbFieldRef LoadCompactBinary(FArchive& Ar, FCbBufferAllocator Allocator = FUniqueBuffer::Alloc);
 
 /** Save a compact binary value to an archive. */
 CORE_API void SaveCompactBinary(FArchive& Ar, const FCbField& Field);
diff --git a/Engine/Source/Runtime/Core/Public/Serialization/CompactBinaryWriter.h b/Engine/Source/Runtime/Core/Public/Serialization/CompactBinaryWriter.h
--- a/Engine/Source/Runtime/Core/Public/Serialization/CompactBinaryWriter.h
+++ b/Engine/Source/Runtime/Core/Public/Serialization/CompactBinaryWriter.h
@@ -4,12 +4,12 @@
 
 #include "CoreTypes.h"
 #include "Containers/StringView.h"
+#include "Memory/MemoryFwd.h"
 #include "Serialization/CompactBinary.h"
 
 ///////////////////////////////////////////////////////////////////////////////////////////////////////////////////////
 
 class FArchive;
-class FSharedBufferConstPtr;
 struct FBlake3Hash;
 struct FDateTime;
 struct FGuid;
@@ -141,7 +141,7 @@ public:
 	/** Write a binary field by copying the view. */
 	inline void Binary(FMemoryView Value) { Binary(Value.GetData(), Value.GetSize()); }
 	/** Write a binary field by copying the buffer. Holds a reference if owned. */
-	CORE_API void Binary(const FSharedBufferConstPtr& Buffer);
+	CORE_API void Binary(const FSharedBuffer& Buffer);
 
 	/** Write a string field by copying the UTF-8 value. */
 	CORE_API void String(FAnsiStringView Value);
diff --git a/Engine/Source/Runtime/CoreUObject/Private/Tests/Virtualization/VirtualizedBulkDataTests.cpp b/Engine/Source/Runtime/CoreUObject/Private/Tests/Virtualization/VirtualizedBulkDataTests.cpp
--- a/Engine/Source/Runtime/CoreUObject/Private/Tests/Virtualization/VirtualizedBulkDataTests.cpp
+++ b/Engine/Source/Runtime/CoreUObject/Private/Tests/Virtualization/VirtualizedBulkDataTests.cpp
@@ -31,20 +31,20 @@ bool FVirtualizationWrapperTestBasic::RunTest(const FString& Parameters)
 
 	auto ValidateBulkData = [this, &SourceBuffer, BufferSize](const FVirtualizedUntypedBulkData& BulkDataToValidate, const TCHAR* Label)
 		{
-			FSharedBufferConstPtr RetrievedBuffer = BulkDataToValidate.GetData().Get();
-			TestEqual(FString::Printf(TEXT("%s buffer length"),Label), (int64)RetrievedBuffer->GetSize(), BufferSize);
-			TestTrue(FString::Printf(TEXT("SourceBuffer values == %s values"),Label), FMemory::Memcmp(SourceBuffer.Get(), RetrievedBuffer->GetData(), BufferSize) == 0);
+			FSharedBuffer RetrievedBuffer = BulkDataToValidate.GetData().Get();
+			TestEqual(FString::Printf(TEXT("%s buffer length"),Label), (int64)RetrievedBuffer.GetSize(), BufferSize);
+			TestTrue(FString::Printf(TEXT("SourceBuffer values == %s values"),Label), FMemory::Memcmp(SourceBuffer.Get(), RetrievedBuffer.GetData(), BufferSize) == 0);
 		};
 
 	// Create a basic bulkdata (but retain ownership of the buffer!)
 	FByteVirtualizedBulkData BulkData;
-	BulkData.UpdatePayload(FSharedBuffer::MakeView(SourceBuffer.Get(), BufferSize).ToSharedRef());
+	BulkData.UpdatePayload(FSharedBuffer::MakeView(SourceBuffer.Get(), BufferSize));
 
 	// Test access via callback
-	BulkData.GetData([this,BufferSize,&SourceBuffer](FSharedBufferConstPtr Payload)
+	BulkData.GetData([this,BufferSize,&SourceBuffer](FSharedBuffer Payload)
 		{
-			TestEqual(TEXT("Callback buffer length"), (int64)Payload->GetSize(), BufferSize);
-			TestTrue(TEXT("SourceBuffer values == Callback values"), FMemory::Memcmp(SourceBuffer.Get(), Payload->GetData(), BufferSize) == 0);
+			TestEqual(TEXT("Callback buffer length"), (int64)Payload.GetSize(), BufferSize);
+			TestTrue(TEXT("SourceBuffer values == Callback values"), FMemory::Memcmp(SourceBuffer.Get(), Payload.GetData(), BufferSize) == 0);
 		});
 
 	// Test accessing the data from the bulkdata object
@@ -55,7 +55,7 @@ bool FVirtualizationWrapperTestBasic::RunTest(const FString& Parameters)
 
 	// Create a new bulkdata object and copy by assignment (note we assign some junk data that will get overwritten)
 	FByteVirtualizedBulkData BulkDataAssignment;
-	BulkDataAssignment.UpdatePayload(FSharedBuffer::Alloc(128).ToSharedRef());
+	BulkDataAssignment.UpdatePayload(FSharedBuffer(FUniqueBuffer::Alloc(128)));
 	BulkDataAssignment = BulkData;
 
 	// Test both bulkdata objects
@@ -87,12 +87,12 @@ bool FVirtualizationWrapperTestEmpty::RunTest(const FString& Parameters)
 	TestFalse(TEXT("Return value of ::IsDataLoaded()"), BulkData.IsDataLoaded());
 
 	// Validate the payload accessors
-	FSharedBufferConstPtr Payload = BulkData.GetData().Get();
-	TestFalse(TEXT("The payload from the GetData TFuture is invalid"), Payload.IsValid());
+	FSharedBuffer Payload = BulkData.GetData().Get();
+	TestTrue(TEXT("The payload from the GetData TFuture is null"), Payload.IsNull());
 
-	BulkData.GetData([this](FSharedBufferConstPtr Payload) 
+	BulkData.GetData([this](FSharedBuffer Payload) 
 		{
-			TestFalse(TEXT("The payload via the GetData callback is invalid"), Payload.IsValid());
+			TestTrue(TEXT("The payload via the GetData callback is null"), Payload.IsNull());
 		});
 
 	return true;
@@ -111,32 +111,35 @@ bool FVirtualizationWrapperTestUpdatePayload::RunTest(const FString& Parameters)
 
 	// Pass the buffer to to bulkdata but retain ownership
 	FByteVirtualizedBulkData BulkData;
-	BulkData.UpdatePayload(FSharedBuffer::MakeView(OriginalData.Get(), BufferSize).ToSharedRef());
+	BulkData.UpdatePayload(FSharedBuffer::MakeView(OriginalData.Get(), BufferSize));
 
 	// Access the payload, edit it and push it back into the bulkdata object
 	{
 		// The payload should be the same size and same contents as the original buffer but a different
 		// memory address since we retained ownership in the TUniquePtr, so the bulkdata object should 
 		// have created it's own copy.
-		FSharedBufferConstPtr Payload = BulkData.GetData().Get();
-		TestEqual(TEXT("Payload length"), (int64)Payload->GetSize(), BufferSize);
-		TestNotEqual(TEXT("OriginalData and the payload should have different memory addresses"), (uint8*)OriginalData.Get(), (uint8*)Payload->GetData());
-		TestTrue(TEXT("Orginal buffer == Payload data"), FMemory::Memcmp(OriginalData.Get(), Payload->GetData(), Payload->GetSize()) == 0);
+		FSharedBuffer Payload = BulkData.GetData().Get();
+		TestEqual(TEXT("Payload length"), (int64)Payload.GetSize(), BufferSize);
+		TestNotEqual(TEXT("OriginalData and the payload should have different memory addresses"), (uint8*)OriginalData.Get(), (uint8*)Payload.GetData());
+		TestTrue(TEXT("Orginal buffer == Payload data"), FMemory::Memcmp(OriginalData.Get(), Payload.GetData(), Payload.GetSize()) == 0);
 
 		// Make a copy of the payload that we can edit
-		FSharedBufferRef EditablePayload = FSharedBuffer::Clone(*Payload).ToSharedRef();
-
 		const uint8 NewValue = 255;
-		FMemory::Memset(EditablePayload->GetData(), NewValue, EditablePayload->GetSize());
+		FSharedBuffer EditedPayload;
+		{
+			FUniqueBuffer EditablePayload = FUniqueBuffer::Clone(Payload);
+			FMemory::Memset(EditablePayload.GetData(), NewValue, EditablePayload.GetSize());
+			EditedPayload = MoveTemp(EditablePayload);
+		}
 
 		// Update the bulkdata object with the new edited payload
-		BulkData.UpdatePayload(EditablePayload);
+		BulkData.UpdatePayload(EditedPayload);
 
 		Payload = BulkData.GetData().Get();
-		TestEqual(TEXT("Updated payload length"), (int64)Payload->GetSize(), BufferSize);
-		TestEqual(TEXT("Payload and EditablePayload should have the same memory addresses"), (uint8*)Payload->GetData(), (uint8*)EditablePayload->GetData());
+		TestEqual(TEXT("Updated payload length"), (int64)Payload.GetSize(), BufferSize);
+		TestEqual(TEXT("Payload and EditablePayload should have the same memory addresses"), (uint8*)Payload.GetData(), (uint8*)EditedPayload.GetData());
 
-		const bool bAllElementsCorrect = Algo::AllOf(TArrayView<uint8>((uint8*)Payload->GetData(), Payload->GetSize()), [NewValue](uint8 Val)
+		const bool bAllElementsCorrect = Algo::AllOf(TArrayView<uint8>((uint8*)Payload.GetData(), Payload.GetSize()), [NewValue](uint8 Val)
 			{ 
 				return Val == NewValue; 
 			});
@@ -151,14 +154,14 @@ bool FVirtualizationWrapperTestUpdatePayload::RunTest(const FString& Parameters)
 
 		// Update the bulkdata object with the original data but this time we give ownership of the buffer to
 		// the bulkdata object.
-		BulkData.UpdatePayload(FSharedBuffer::TakeOwnership(OriginalData.Release(), BufferSize, FMemory::Free).ToSharedRef());
+		BulkData.UpdatePayload(FSharedBuffer::TakeOwnership(OriginalData.Release(), BufferSize, FMemory::Free));
 
-		FSharedBufferConstPtr Payload = BulkData.GetData().Get();
-		TestEqual(TEXT("Updated payload length"), (int64)Payload->GetSize(), BufferSize);
-		TestEqual(TEXT("Payload and OriginalDataPtr should have the same memory addresses"), (uint8*)Payload->GetData(), OriginalDataPtr);
+		FSharedBuffer Payload = BulkData.GetData().Get();
+		TestEqual(TEXT("Updated payload length"), (int64)Payload.GetSize(), BufferSize);
+		TestEqual(TEXT("Payload and OriginalDataPtr should have the same memory addresses"), (uint8*)Payload.GetData(), OriginalDataPtr);
 
 		// The original data was all zeros, so we can test for that to make sure that the contents are correct.
-		const bool bAllElementsCorrect = Algo::AllOf(TArrayView<uint8>((uint8*)Payload->GetData(), Payload->GetSize()), [](uint8 Val)
+		const bool bAllElementsCorrect = Algo::AllOf(TArrayView<uint8>((uint8*)Payload.GetData(), Payload.GetSize()), [](uint8 Val)
 			{
 				return Val == 0;
 			});
diff --git a/Engine/Source/Runtime/CoreUObject/Private/Virtualization/VirtualizationManager.cpp b/Engine/Source/Runtime/CoreUObject/Private/Virtualization/VirtualizationManager.cpp
--- a/Engine/Source/Runtime/CoreUObject/Private/Virtualization/VirtualizationManager.cpp
+++ b/Engine/Source/Runtime/CoreUObject/Private/Virtualization/VirtualizationManager.cpp
@@ -22,7 +22,7 @@ FVirtualizationManager::FVirtualizationManager()
 	LoadSettingsFromConfigFiles();
 }
 
-bool FVirtualizationManager::PushData(const FSharedBufferConstPtr& Payload, const FGuid& Guid)
+bool FVirtualizationManager::PushData(const FSharedBuffer& Payload, const FGuid& Guid)
 {
 	checkf(Guid.IsValid(), TEXT("VirtualizationManager::PushData called with an invalid guid (%s)."), *Guid.ToString());
 
@@ -33,13 +33,13 @@ bool FVirtualizationManager::PushData(const FSharedBufferConstPtr& Payload, cons
 	}
 
 	// Early out if we have no payload
-	if (!Payload.IsValid() || Payload->GetSize() == 0)
+	if (Payload.GetSize() == 0)
 	{
 		return false;
 	}
 
 	// Early out if the payload length is below our minimum required length
-	if ((int64)Payload->GetSize() < MinPayloadLength)
+	if ((int64)Payload.GetSize() < MinPayloadLength)
 	{
 		return false;
 	}
@@ -52,9 +52,9 @@ bool FVirtualizationManager::PushData(const FSharedBufferConstPtr& Payload, cons
 		TUniquePtr<FArchive> FileAr(IFileManager::Get().CreateFileWriter(*FilePath));
 		check(FileAr);
 
-		uint64 PayloadLength = Payload->GetSize();
+		uint64 PayloadLength = Payload.GetSize();
 		// Cast away the const because FArchive requires a non-const pointer even if we are only reading from it (saving)
-		void* PayloadPtr = const_cast<void*>(Payload->GetData());
+		void* PayloadPtr = const_cast<void*>(Payload.GetData());
 
 		*FileAr << PayloadLength;
 		FileAr->Serialize(PayloadPtr, PayloadLength);
@@ -63,7 +63,7 @@ bool FVirtualizationManager::PushData(const FSharedBufferConstPtr& Payload, cons
 	return true;
 }
 
-FSharedBufferConstPtr FVirtualizationManager::PullData(const FGuid& Guid)
+FSharedBuffer FVirtualizationManager::PullData(const FGuid& Guid)
 {
 	checkf(Guid.IsValid(), TEXT("VirtualizationManager::PullData called with an invalid guid (%s)."), *Guid.ToString());
 
@@ -78,10 +78,9 @@ FSharedBufferConstPtr FVirtualizationManager::PullData(const FGuid& Guid)
 	int64 DataLength = 0;
 	*FileAr << DataLength;
 
-	uint8* DataPtr = (uint8*)FMemory::Malloc(DataLength, DEFAULT_ALIGNMENT);
-	FileAr->Serialize(DataPtr, DataLength);
-
-	return FSharedBuffer::TakeOwnership(DataPtr, DataLength, FMemory::Free);
+	FUniqueBuffer Buffer = FUniqueBuffer::Alloc(DataLength);
+	FileAr->Serialize(Buffer.GetData(), DataLength);
+	return FSharedBuffer(MoveTemp(Buffer));
 }
 
 void FVirtualizationManager::LoadSettingsFromConfigFiles()
diff --git a/Engine/Source/Runtime/CoreUObject/Private/Virtualization/VirtualizedBulkData.cpp b/Engine/Source/Runtime/CoreUObject/Private/Virtualization/VirtualizedBulkData.cpp
--- a/Engine/Source/Runtime/CoreUObject/Private/Virtualization/VirtualizedBulkData.cpp
+++ b/Engine/Source/Runtime/CoreUObject/Private/Virtualization/VirtualizedBulkData.cpp
@@ -93,7 +93,7 @@ void FVirtualizedUntypedBulkData::Serialize(FArchive& Ar, UObject* Owner)
 			if (!bIsVirtualized)
 			{
 				// Need to load the payload so that we can write it out
-				FSharedBufferConstPtr PayloadToSerialize = GetDataInternal();
+				FSharedBuffer PayloadToSerialize = GetDataInternal();
 
 				// Need to update the flag before it is written out
 				bIsDataStoredAsCompressed = !bIsVirtualized && bShouldCompressData;
@@ -183,10 +183,10 @@ void FVirtualizedUntypedBulkData::Serialize(FArchive& Ar, UObject* Owner)
 
 void FVirtualizedUntypedBulkData::CalculateKey()
 {
-	if (Payload.IsValid())
+	if (Payload)
 	{
 		uint32 Hash[5] = {};
-		FSHA1::HashBuffer(Payload->GetData(), GetBulkDataSize(), (uint8*)Hash);
+		FSHA1::HashBuffer(Payload.GetData(), GetBulkDataSize(), (uint8*)Hash);
 
 		Key = FGuid(Hash[0] ^ Hash[4], Hash[1], Hash[2], Hash[3]);
 	}
@@ -196,11 +196,11 @@ void FVirtualizedUntypedBulkData::CalculateKey()
 	}
 }
 
-FSharedBufferConstPtr FVirtualizedUntypedBulkData::LoadFromDisk() const
+FSharedBuffer FVirtualizedUntypedBulkData::LoadFromDisk() const
 {
 	TRACE_CPUPROFILER_EVENT_SCOPE(FVirtualizedUntypedBulkData::LoadFromDisk);
 
-	FSharedBufferConstPtr PayloadFromDisk;
+	FSharedBuffer PayloadFromDisk;
 
 	// Open a reader to the file
 	FOpenPackageResult Result = IPackageResourceManager::Get().OpenReadPackage(PackagePath, PackageSegment);
@@ -221,20 +221,25 @@ FSharedBufferConstPtr FVirtualizedUntypedBulkData::LoadFromDisk() const
 	return PayloadFromDisk;
 }
 
-bool FVirtualizedUntypedBulkData::SerializeData(FArchive& Ar, FSharedBufferConstPtr& InPayload) const
+bool FVirtualizedUntypedBulkData::SerializeData(FArchive& Ar, FSharedBuffer& InPayload) const
 {
 	TRACE_CPUPROFILER_EVENT_SCOPE(FVirtualizedUntypedBulkData::SerializeData);
 
-	const int64 PayloadSize = Ar.IsSaving() ? InPayload->GetSize() : GetBulkDataSize();
+	const int64 PayloadSize = Ar.IsSaving() ? InPayload.GetSize() : GetBulkDataSize();
 
 	if (PayloadSize > 0)
 	{
+		FUniqueBuffer LoadPayload;
+
+		if (Ar.IsLoading())
+		{
+			LoadPayload = FUniqueBuffer::Alloc(PayloadSize);
+		}
+
 		// Our serialization methods require a non-const pointer even when we are only reading from it
 		// as the methods themselves are bi-directional. This means if we are saving we need to cast 
 		// away the const qualifier.
-		void* DataPtr = Ar.IsLoading() ?
-						FMemory::Malloc(PayloadSize, DEFAULT_ALIGNMENT) :
-						const_cast<void*>(InPayload->GetData());
+		void* DataPtr = Ar.IsLoading() ? LoadPayload.GetData() : const_cast<void*>(InPayload.GetData());
 
 		checkf(DataPtr != nullptr, TEXT("Attempting to serialize to/from an invalid shared buffer!"));
 
@@ -249,7 +254,7 @@ bool FVirtualizedUntypedBulkData::SerializeData(FArchive& Ar, FSharedBufferConst
 
 		if (Ar.IsLoading())
 		{
-			InPayload = FSharedBuffer::TakeOwnership(DataPtr, PayloadSize, FMemory::Free);
+			InPayload = FSharedBuffer(MoveTemp(LoadPayload));
 		}
 
 		return true;
@@ -268,7 +273,7 @@ bool FVirtualizedUntypedBulkData::SerializeData(FArchive& Ar, FSharedBufferConst
 
 void FVirtualizedUntypedBulkData::PushData()
 {
-	checkf(bIsVirtualized == false || Payload.IsValid() == false, TEXT("Cannot have a valid payload in memory if the payload is virtualized!")); // Sanity check
+	checkf(bIsVirtualized == false || Payload.IsNull(), TEXT("Cannot have a valid payload in memory if the payload is virtualized!")); // Sanity check
 
 	// We only need to push if the payload is not currently virtualized (either we have an updated
 	// payload in memory or the payload is currently non-virtualized and stored on disk)
@@ -280,7 +285,7 @@ void FVirtualizedUntypedBulkData::PushData()
 		// a non-virtualized payload to a virtualized one. If the bulkdata is merely being
 		// edited then we should have the payload in memory already and are just accessing a
 		// reference to it.
-		FSharedBufferConstPtr PayloadToPush = GetDataInternal();
+		FSharedBuffer PayloadToPush = GetDataInternal();
 
 		if (FVirtualizationManager::Get().PushData(PayloadToPush, Key))
 		{
@@ -289,18 +294,18 @@ void FVirtualizedUntypedBulkData::PushData()
 	}	
 }
 
-FSharedBufferConstPtr FVirtualizedUntypedBulkData::PullData() const
+FSharedBuffer FVirtualizedUntypedBulkData::PullData() const
 {
 	TRACE_CPUPROFILER_EVENT_SCOPE(FVirtualizedUntypedBulkData::PullData);
 
-	FSharedBufferConstPtr PulledPayload = FVirtualizationManager::Get().PullData(Key);
+	FSharedBuffer PulledPayload = FVirtualizationManager::Get().PullData(Key);
 
-	if (PulledPayload.IsValid())
+	if (PulledPayload)
 	{
-		checkf(	PayloadLength == PulledPayload->GetSize(),	
-				TEXT("Mismatch between serialized length (%d) and virtualized data length (%d)"), 
+		checkf(	PayloadLength == PulledPayload.GetSize(),
+				TEXT("Mismatch between serialized length (%" INT64_FMT ") and virtualized data length (%" UINT64_FMT ")"),
 				PayloadLength,
-				PulledPayload->GetSize());
+				PulledPayload.GetSize());
 	}
 	else
 	{
@@ -352,44 +357,44 @@ void FVirtualizedUntypedBulkData::UnloadData()
 	Payload.Reset();
 }
 
-FSharedBufferConstPtr FVirtualizedUntypedBulkData::GetDataInternal() const
+FSharedBuffer FVirtualizedUntypedBulkData::GetDataInternal() const
 {
 	TRACE_CPUPROFILER_EVENT_SCOPE(FVirtualizedUntypedBulkData::GetDataInternal);
 
 	// Early out there isn't any data to actually load
 	if (GetBulkDataSize() == 0)
 	{
-		return FSharedBufferConstPtr();
+		return FSharedBuffer();
 	}
 
 	// Check if we already have the data in memory
-	if (Payload.IsValid())
+	if (Payload)
 	{
 		return Payload;
 	}
 
 	if (bIsVirtualized)
 	{
-		FSharedBufferConstPtr Buffer = PullData();
-		check(!Payload.IsValid()); //Make sure that we did not assign the buffer internally
+		FSharedBuffer Buffer = PullData();
+		check(Payload.IsNull()); //Make sure that we did not assign the buffer internally
 		return Buffer;
 	}
 	else
 	{
-		FSharedBufferConstPtr Buffer = LoadFromDisk();
-		check(!Payload.IsValid()); //Make sure that we did not assign the buffer internally
+		FSharedBuffer Buffer = LoadFromDisk();
+		check(Payload.IsNull()); //Make sure that we did not assign the buffer internally
 		return Buffer;
 	}
 }
 
-TFuture<FSharedBufferConstPtr> FVirtualizedUntypedBulkData::GetData() const
+TFuture<FSharedBuffer> FVirtualizedUntypedBulkData::GetData() const
 {
-	TPromise<FSharedBufferConstPtr> Promise;
+	TPromise<FSharedBuffer> Promise;
 
-	FSharedBufferConstPtr SharedBuffer = GetDataInternal();
+	FSharedBuffer SharedBuffer = GetDataInternal();
 
 	// TODO: Not actually async yet!
-	Promise.SetValue(SharedBuffer);
+	Promise.SetValue(MoveTemp(SharedBuffer));
 
 	return Promise.GetFuture();
 }
@@ -399,15 +404,16 @@ void FVirtualizedUntypedBulkData::GetData(OnDataReadyCallback&& Callback) const
 	Callback(GetDataInternal());
 }
 
-void FVirtualizedUntypedBulkData::UpdatePayload(const FSharedBufferConstRef& InPayload)
+void FVirtualizedUntypedBulkData::UpdatePayload(const FSharedBuffer& InPayload)
 {
 	TRACE_CPUPROFILER_EVENT_SCOPE(FVirtualizedUntypedBulkData::UpdatePayload);
 
 	UnloadData();
 
 	// Make sure that we own the memory in the shared buffer
-	Payload = FSharedBuffer::MakeOwned(InPayload);
-	PayloadLength = (int64)InPayload->GetSize();
+	Payload = InPayload;
+	Payload.MakeOwned();
+	PayloadLength = (int64)Payload.GetSize();
 
 	bIsVirtualized = false;
 	bIsDataStoredAsCompressed = false;
diff --git a/Engine/Source/Runtime/CoreUObject/Public/Virtualization/VirtualizationManager.h b/Engine/Source/Runtime/CoreUObject/Public/Virtualization/VirtualizationManager.h
--- a/Engine/Source/Runtime/CoreUObject/Public/Virtualization/VirtualizationManager.h
+++ b/Engine/Source/Runtime/CoreUObject/Public/Virtualization/VirtualizationManager.h
@@ -35,10 +35,10 @@ public:
 	~FVirtualizationManager() = default;
 
 	/** Push a payload to the backends*/
-	bool PushData(const FSharedBufferConstPtr& Payload, const FGuid& Guid);
+	bool PushData(const FSharedBuffer& Payload, const FGuid& Guid);
 
 	/** Pull a payload from the backends*/
-	FSharedBufferConstPtr PullData(const FGuid& Guid);
+	FSharedBuffer PullData(const FGuid& Guid);
 
 private:
 
diff --git a/Engine/Source/Runtime/CoreUObject/Public/Virtualization/VirtualizationUtilities.h b/Engine/Source/Runtime/CoreUObject/Public/Virtualization/VirtualizationUtilities.h
--- a/Engine/Source/Runtime/CoreUObject/Public/Virtualization/VirtualizationUtilities.h
+++ b/Engine/Source/Runtime/CoreUObject/Public/Virtualization/VirtualizationUtilities.h
@@ -17,16 +17,16 @@ public:
 		SetIsSaving(true);
 		SetIsPersistent(bIsPersistent);
 
-		FSharedBufferConstPtr Payload = InBulkData.GetData().Get();
+		FSharedBuffer Payload = InBulkData.GetData().Get();
 
-		if (Payload.IsValid())
+		if (Payload)
 		{
-			const int64 CurrentDataLength = Payload->GetSize();
+			const int64 CurrentDataLength = Payload.GetSize();
 
 			// Clone the payload so that we have a local copy that we can
 			// append additional data to.
 			Buffer = FMemory::Malloc(CurrentDataLength, DEFAULT_ALIGNMENT);
-			FMemory::Memcpy(Buffer, Payload->GetData(), CurrentDataLength);
+			FMemory::Memcpy(Buffer, Payload.GetData(), CurrentDataLength);
 
 			BufferLength = CurrentDataLength;
 
@@ -48,7 +48,7 @@ public:
 	{
 		// Remove the slack from the allocated bulk data
 		Buffer = FMemory::Realloc(Buffer, DataLength, DEFAULT_ALIGNMENT);
-		BulkData.UpdatePayload(FSharedBuffer::TakeOwnership(Buffer, DataLength, FMemory::Free).ToSharedRef());
+		BulkData.UpdatePayload(FSharedBuffer::TakeOwnership(Buffer, DataLength, FMemory::Free));
 	}
 
 	virtual void Serialize(void* Data, int64 Num)
@@ -136,26 +136,19 @@ namespace UE4VirtualizedBulkData_Private
 
 		void* GetData() const
 		{
-			if(Payload.IsValid())
-			{
-				// It's okay to remove the const qualifier here as it will only be passed
-				// on to FBufferReaderBase, which will not change the data at all, but takes a non-const
-				// pointer so that it can free the memory if requested, which we don't.
-				return const_cast<void*>(Payload->GetData());
-			}
-			else
-			{
-				return nullptr;
-			}
+			// It's okay to remove the const qualifier here as it will only be passed
+			// on to FBufferReaderBase, which will not change the data at all, but takes a non-const
+			// pointer so that it can free the memory if requested, which we don't.
+			return const_cast<void*>(Payload.GetData());
 		}
 
 		int64 GetDataLength() const
 		{
-			return Payload.IsValid() ? Payload->GetSize() : 0;
+			return Payload.GetSize();
 		}
 
 	private:
-		FSharedBufferConstPtr Payload;
+		FSharedBuffer Payload;
 	};
 }
 
diff --git a/Engine/Source/Runtime/CoreUObject/Public/Virtualization/VirtualizedBulkData.h b/Engine/Source/Runtime/CoreUObject/Public/Virtualization/VirtualizedBulkData.h
--- a/Engine/Source/Runtime/CoreUObject/Public/Virtualization/VirtualizedBulkData.h
+++ b/Engine/Source/Runtime/CoreUObject/Public/Virtualization/VirtualizedBulkData.h
@@ -46,7 +46,7 @@ class UObject;
 class COREUOBJECT_API FVirtualizedUntypedBulkData
 {
 public:
-	using OnDataReadyCallback = TUniqueFunction<void(FSharedBufferConstPtr)>;
+	using OnDataReadyCallback = TUniqueFunction<void(FSharedBuffer)>;
 
 	FVirtualizedUntypedBulkData() = default;
 	~FVirtualizedUntypedBulkData();
@@ -80,7 +80,7 @@ public:
 	// needs to be loaded in order to be able to run on a background thread. Since FVirtualizedUntypedBulkData 
 	// will aim to be thread safe we could change TextureDerivedData and remove this.
 	/** Returns if the data is being held in memory (true) or will be loaded from disk (false) */
-	bool IsDataLoaded() const { return Payload.IsValid(); }
+	bool IsDataLoaded() const { return !Payload.IsNull(); }
 
 	/** 
 	 * Enable/Disable if the payload should be compressed if it ends up being serialized to disk on the local
@@ -95,7 +95,7 @@ public:
 	 * To convert to a mutable reference call'FSharedBufferRef Mutable = MakeSharedBuffer(FSharedBuffer::Clone, Immutable->GetView());'
 	 * which will return a newly allocated clone of the payload which the caller will have full ownership over.
 	 */
-	TFuture<FSharedBufferConstPtr> GetData() const;
+	TFuture<FSharedBuffer> GetData() const;
 
 	/** Returns an immutable FSharedBuffer reference to the payload data via the OnDataReadyCallback */
 	void GetData(OnDataReadyCallback&& Callback) const;
@@ -110,19 +110,19 @@ public:
 	 * by using 'FSharedBuffer::Wrap') then a clone of the data will be created internally and assigned to
 	 * the bulkdata object.
 	 */
-	void UpdatePayload(const FSharedBufferConstRef& InPayload);
+	void UpdatePayload(const FSharedBuffer& InPayload);
 
 private:
 
-	FSharedBufferConstPtr GetDataInternal() const;
+	FSharedBuffer GetDataInternal() const;
 
 	void CalculateKey();
 
-	FSharedBufferConstPtr LoadFromDisk() const;
-	bool SerializeData(FArchive& Ar, FSharedBufferConstPtr& Payload) const;
+	FSharedBuffer LoadFromDisk() const;
+	bool SerializeData(FArchive& Ar, FSharedBuffer& Payload) const;
 
 	void PushData();
-	FSharedBufferConstPtr PullData() const;
+	FSharedBuffer PullData() const;
 
 	FPackagePath GetPackagePathFromOwner(UObject* Owner, EPackageSegment& OutPackageSegment) const;
 
@@ -132,7 +132,7 @@ private:
 	FGuid Key;
 
 	/** Pointer to the payload if it is held in memory (it has been updated but not yet saved to disk for example) */
-	FSharedBufferConstPtr Payload;
+	FSharedBuffer Payload;
 
 	/** Length of the payload in bytes */
 	int64 PayloadLength = 0;
