Adding Mantle Serialization Version.

diff --git a/Engine/Source/Runtime/Core/Public/UObject/UE5MainStreamObjectVersion.h b/Engine/Source/Runtime/Core/Public/UObject/UE5MainStreamObjectVersion.h
--- a/Engine/Source/Runtime/Core/Public/UObject/UE5MainStreamObjectVersion.h
+++ b/Engine/Source/Runtime/Core/Public/UObject/UE5MainStreamObjectVersion.h
@@ -61,6 +61,9 @@ struct CORE_API FUE5MainStreamObjectVersion
 		// Added support for per-object collision constraint flag. [Chaos]
 		AddCollisionConstraintFlag,
 
+		// Initial Mantle Serialize Version
+		MantleDbSerialize,
+
 		// -----<new versions can be added above this line>-------------------------------------------------
 		VersionPlusOne,
 		LatestVersion = VersionPlusOne - 1
