Implement IORegistry discovery of Apple Silicon gfx.

#rnx

diff --git a/Engine/Source/Runtime/Apple/MetalRHI/Private/MetalRHI.cpp b/Engine/Source/Runtime/Apple/MetalRHI/Private/MetalRHI.cpp
--- a/Engine/Source/Runtime/Apple/MetalRHI/Private/MetalRHI.cpp
+++ b/Engine/Source/Runtime/Apple/MetalRHI/Private/MetalRHI.cpp
@@ -347,6 +347,15 @@ FMetalDynamicRHI::FMetalDynamicRHI(ERHIFeatureLevel::Type RequestedFeatureLevel)
 		bSupportsDistanceFields = (FPlatformMisc::MacOSXVersionCompare(10,12,2) >= 0);
 		bIsIntelHaswell = (GRHIAdapterName == TEXT("Intel HD Graphics 5000") || GRHIAdapterName == TEXT("Intel Iris Graphics") || GRHIAdapterName == TEXT("Intel Iris Pro Graphics"));
 	}
+	else if(GRHIAdapterName.Contains("Apple"))
+	{
+		bSupportsPointLights = true;
+		GRHIVendorId = 0x106B;
+		bSupportsTiledReflections = true;
+		bSupportsDistanceFields = true;
+		bSupportsRHIThread = true;
+		GSupportsTimestampRenderQueries = true;
+	}
 
 	bool const bRequestedSM5 = (RequestedFeatureLevel == ERHIFeatureLevel::SM5 || (!bRequestedFeatureLevel && (FParse::Param(FCommandLine::Get(),TEXT("metalsm5")) || FParse::Param(FCommandLine::Get(),TEXT("metalmrt")))));
 	if(bSupportsSM5 && bRequestedSM5)
@@ -491,7 +500,7 @@ FMetalDynamicRHI::FMetalDynamicRHI(ERHIFeatureLevel::Type RequestedFeatureLevel)
 #endif
 		GRHISupportsParallelRHIExecute = GRHISupportsRHIThread && ((!IsRHIDeviceIntel() && !IsRHIDeviceNVIDIA()) || FParse::Param(FCommandLine::Get(),TEXT("metalparallel")));
 #endif
-		GSupportsEfficientAsyncCompute = GRHISupportsParallelRHIExecute && (IsRHIDeviceAMD() || PLATFORM_IOS || FParse::Param(FCommandLine::Get(),TEXT("metalasynccompute"))); // Only AMD currently support async. compute and it requires parallel execution to be useful.
+		GSupportsEfficientAsyncCompute = GRHISupportsParallelRHIExecute && (IsRHIDeviceAMD() || /*TODO: IsRHIDeviceApple()*/ (GRHIVendorId == 0x106B) || PLATFORM_IOS || FParse::Param(FCommandLine::Get(),TEXT("metalasynccompute"))); // Only AMD and Apple currently support async. compute and it requires parallel execution to be useful.
 		GSupportsParallelOcclusionQueries = GRHISupportsRHIThread;
 	}
 	else
diff --git a/Engine/Source/Runtime/Core/Private/Mac/MacPlatformMisc.cpp b/Engine/Source/Runtime/Core/Private/Mac/MacPlatformMisc.cpp
--- a/Engine/Source/Runtime/Core/Private/Mac/MacPlatformMisc.cpp
+++ b/Engine/Source/Runtime/Core/Private/Mac/MacPlatformMisc.cpp
@@ -76,6 +76,83 @@ static TAutoConsoleVariable<int32> CVarMacPlatformDumpAllThreadsOnHang(
 	1,
 	TEXT("If > 0, then when reporting a hang generate a backtrace for all threads."));
 
+/*------------------------------------------------------------------------------
+ Platform property discovery.
+ ------------------------------------------------------------------------------*/
+
+#define PLATFORM_MAC_IOSERVICE_MATCHING_NAME_ARM64	"AppleARMIODevice"
+#define PLATFORM_MAC_IOSERVICE_MATCHING_NAME_X86	"IOPCIDevice"
+#define PLATFORM_MAC_CLASS_CODE_NAME_ARM64			"device_type"
+#define PLATFORM_MAC_CLASS_CODE_NAME_X86			"class-code"
+
+#define PLATFORM_MAC_MAKE_FOURCC(ch0, ch1, ch2, ch3)							\
+   ((uint32_t)(uint8_t)(ch0)		| ((uint32_t)(uint8_t)(ch1) << 8	)	|	\
+   ((uint32_t)(uint8_t)(ch2) << 16) | ((uint32_t)(uint8_t)(ch3) << 24	))
+
+#if PLATFORM_MAC_ARM64
+
+inline bool IsRunningOnAppleSilicon()
+{
+	return true;
+}
+
+static const char* GetIOServiceMatchingName()
+{
+	return PLATFORM_MAC_IOSERVICE_MATCHING_NAME_ARM64;
+}
+
+static const char* GetClassCode()
+{
+	return PLATFORM_MAC_CLASS_CODE_NAME_ARM64;
+}
+
+#elif PLATFORM_MAC_X86
+
+#define MAC_PROCESS_TYPE_NATIVE		 0
+#define MAC_PROCESS_TYPE_TRANSLATED	 1
+#define MAC_PROCESS_TYPE_UNKNOWN	-1
+
+static int32 GetProcessTranslationType()
+{
+	static int32 MacProcessType = MAC_PROCESS_TYPE_UNKNOWN;
+
+	if (MacProcessType == MAC_PROCESS_TYPE_UNKNOWN)
+	{
+		int32 Value = 0;
+		size_t ValueSize = sizeof(Value);
+		if (sysctlbyname("sysctl.proc_translated", &Value, &ValueSize, NULL, 0) == -1)
+		{
+			if (errno == ENOENT)
+			{
+				Value = MAC_PROCESS_TYPE_NATIVE;
+			}
+		}
+		MacProcessType = Value;
+	}
+
+	return MacProcessType;
+}
+
+inline bool IsRunningOnAppleSilicon()
+{
+	return MAC_PROCESS_TYPE_TRANSLATED == GetProcessTranslationType();
+}
+
+static const char* GetIOServiceMatchingName()
+{
+	return IsRunningOnAppleSilicon() ? PLATFORM_MAC_IOSERVICE_MATCHING_NAME_ARM64 : PLATFORM_MAC_IOSERVICE_MATCHING_NAME_X86;
+}
+
+static const char* GetClassCode()
+{
+	return IsRunningOnAppleSilicon() ? PLATFORM_MAC_CLASS_CODE_NAME_ARM64 : PLATFORM_MAC_CLASS_CODE_NAME_X86;
+}
+
+#else
+	#error "Undefined Mac platform"
+
+#endif // PLATFORM_MAC_XXXXX
+
 /*------------------------------------------------------------------------------
  FMacApplicationInfo - class to contain all state for crash reporting that is unsafe to acquire in a signal.
  ------------------------------------------------------------------------------*/
@@ -821,6 +898,9 @@ FMacPlatformMisc::FGPUDescriptorARM64::~FGPUDescriptorARM64()
 
 void FMacPlatformMisc::FGPUDescriptorARM64::CopyFromImpl(FGPUDescriptorCommon<FMacPlatformMisc::FGPUDescriptorARM64> const& Other)
 {
+	FMacPlatformMisc::FGPUDescriptorARM64 const& ConcreteOther = static_cast<FMacPlatformMisc::FGPUDescriptorARM64 const&>(Other);
+
+	RegistryID = ConcreteOther.RegistryID;
 }
 
 TMap<FString, float> FMacPlatformMisc::FGPUDescriptorARM64::GetPerformanceStatisticsImpl() const
@@ -839,7 +919,101 @@ class FMacPlatformGPUManager
 	TArray<FMacPlatformMisc::FGPUDescriptor> CurrentGPUs;
 	TArray<FMacPlatformMisc::FGPUDescriptor> UpdatedGPUs;
 	TAtomic<bool> bRequiresUpdate;
-	
+
+	void InitializeDescriptorFromDeviceEntryM(FMacPlatformMisc::FGPUDescriptor& Desc, io_registry_entry_t ServiceEntry, CFMutableDictionaryRef ServiceInfo)
+	{
+		static CFStringRef IOMatchCategoryRef	= CFSTR("IOMatchCategory");
+		static CFStringRef IOAcceleratorRef		= CFSTR("IOAccelerator");
+		static CFStringRef CFBundleIdentifier	= CFSTR("CFBundleIdentifier");
+		static CFStringRef VendorIDRef			= CFSTR("vendor-id");
+		static CFStringRef MetalPluginNameRef	= CFSTR("MetalPluginName");
+		static CFStringRef GLBundleNameRef		= CFSTR("IOGLBundleName");
+		static CFStringRef ModelRef				= CFSTR("model");
+
+		io_iterator_t ChildIterator;
+		if(IORegistryEntryGetChildIterator(ServiceEntry, kIOServicePlane, &ChildIterator) == kIOReturnSuccess)
+		{
+			io_registry_entry_t ChildEntry;
+			while ((Desc.RegistryID == 0) && (ChildEntry = IOIteratorNext(ChildIterator)))
+			{
+				CFStringRef IOMatchCategory = (CFStringRef)IORegistryEntrySearchCFProperty(ChildEntry, kIOServicePlane, IOMatchCategoryRef, kCFAllocatorDefault, 0);
+				if (IOMatchCategory && (CFGetTypeID(IOMatchCategory) == CFStringGetTypeID()) && (CFStringCompare(IOMatchCategory, IOAcceleratorRef, 0) == kCFCompareEqualTo))
+				{
+					CFMutableDictionaryRef Properties = nullptr;
+					if (kIOReturnSuccess == IORegistryEntryCreateCFProperties(ChildEntry, &Properties, kCFAllocatorDefault, kIORegistryIterateRecursively))
+					{
+						kern_return_t Result = IORegistryEntryGetRegistryEntryID(ChildEntry, &Desc.RegistryID);
+						check(Result == kIOReturnSuccess);
+
+						CFStringRef BundleID = (CFStringRef)CFDictionaryGetValue(Properties, CFBundleIdentifier);
+						if (BundleID && (CFGetTypeID(BundleID) == CFStringGetTypeID()))
+						{
+							Desc.GPUBundleID = [[NSString alloc] initWithString:(__bridge NSString*)BundleID];
+						}
+
+						{
+							char Buffer[0x40] = { 0 };
+							size_t BufferSize = sizeof(Buffer) / sizeof(Buffer[0]);
+							if (0 == sysctlbyname("hw.targettype", &Buffer, &BufferSize, NULL, 0))
+							{
+								uint32_t Value = PLATFORM_MAC_MAKE_FOURCC(Buffer[0], Buffer[1], Buffer[2], Buffer[3]);
+								Desc.GPUDeviceId = Value;
+							}
+						}
+
+						CFDataRef VendorID = (CFDataRef)CFDictionaryGetValue(Properties, VendorIDRef);
+						if (VendorID && (CFGetTypeID(VendorID) == CFDataGetTypeID()))
+						{
+							const uint32_t* Value = reinterpret_cast<const uint32_t*>(CFDataGetBytePtr(VendorID));
+							Desc.GPUVendorId = *Value;
+						}
+
+						CFStringRef MetalPluginName = (CFStringRef)CFDictionaryGetValue(Properties, MetalPluginNameRef);
+						if (MetalPluginName && (CFGetTypeID(MetalPluginName) == CFStringGetTypeID()))
+						{
+							Desc.GPUMetalBundle = [[NSString alloc] initWithString:(__bridge NSString*)MetalPluginName];
+						}
+
+						CFStringRef GLBundleName = (CFStringRef)CFDictionaryGetValue(Properties, GLBundleNameRef);
+						if (GLBundleName && (CFGetTypeID(GLBundleName) == CFStringGetTypeID()))
+						{
+							Desc.GPUOpenGLBundle = [[NSString alloc] initWithString:(__bridge NSString*)GLBundleName];
+						}
+
+						CFStringRef ModelName = (CFStringRef)CFDictionaryGetValue(Properties, ModelRef);
+						if (ModelName && (CFGetTypeID(ModelName) == CFStringGetTypeID()))
+						{
+							Desc.GPUName = [[NSString alloc] initWithString:(__bridge NSString*)ModelName];
+						}
+
+						{
+							uint64_t Value = 0;
+							size_t ValueSize = sizeof(Value);
+							if (0 == sysctlbyname("hw.memsize", &Value, &ValueSize, NULL, 0))
+							{
+								Desc.GPUMemoryMB = uint64(float(Value) * 0.75f) / 1024 / 1024;
+							}
+						}
+					}
+
+					if (Properties)
+					{
+						CFRelease(Properties);
+					}
+				}
+
+				if (IOMatchCategory)
+				{
+					CFRelease(IOMatchCategory);
+				}
+
+				IOObjectRelease(ChildEntry);
+			}
+
+			IOObjectRelease(ChildIterator);
+		}
+	}
+
 #if PLATFORM_MAC_X86
 	void InitializeDescriptorFromDeviceEntry(FMacPlatformMisc::FGPUDescriptor& Desc, io_registry_entry_t ServiceEntry, CFMutableDictionaryRef ServiceInfo)
 	{
@@ -962,7 +1136,7 @@ class FMacPlatformGPUManager
 			}
 		}
 	}
-#endif
+#endif // PLATFORM_MAC_X86
     
 public:
 	static FMacPlatformGPUManager& Get()
@@ -970,14 +1144,17 @@ public:
 		static FMacPlatformGPUManager sSelf;
 		return sSelf;
 	}
-	
+
 	FMacPlatformGPUManager()
 	{
 		FScopeLock Lock(&Mutex);
-#if PLATFORM_MAC_X86
+
+		CFStringRef ClassCodeRef = CFStringCreateWithCString(kCFAllocatorDefault, GetClassCode(), kCFStringEncodingUTF8);
+		check(ClassCodeRef);
+
 		// Enumerate the GPUs via IOKit to avoid dragging in OpenGL
 		io_iterator_t Iterator;
-		CFMutableDictionaryRef MatchDictionary = IOServiceMatching("IOPCIDevice");
+		CFMutableDictionaryRef MatchDictionary = IOServiceMatching(GetIOServiceMatchingName());
 		if(IOServiceGetMatchingServices(kIOMasterPortDefault, MatchDictionary, &Iterator) == kIOReturnSuccess)
 		{
 			uint32 Index = 0;
@@ -987,25 +1164,49 @@ public:
 				CFMutableDictionaryRef ServiceInfo;
 				if(IORegistryEntryCreateCFProperties(ServiceEntry, &ServiceInfo, kCFAllocatorDefault, kNilOptions) == kIOReturnSuccess)
 				{
-					// GPUs are class-code 0x30000 || 0x38000
-					static CFStringRef ClassCodeRef = CFSTR("class-code");
 					const CFDataRef ClassCode = (const CFDataRef)CFDictionaryGetValue(ServiceInfo, ClassCodeRef);
 					if(ClassCode && CFGetTypeID(ClassCode) == CFDataGetTypeID())
 					{
-						const uint32* ClassCodeValue = reinterpret_cast<const uint32*>(CFDataGetBytePtr(ClassCode));
-						if(ClassCodeValue && (*ClassCodeValue == 0x30000 || *ClassCodeValue == 0x38000))
+						if (IsRunningOnAppleSilicon())
 						{
-							FMacPlatformMisc::FGPUDescriptor Desc;
-							
-                            InitializeDescriptorFromDeviceEntry(Desc, ServiceEntry, ServiceInfo);
-							
-							if (Desc.GPUMetalBundle)
+							const char* ClassCodeValue = reinterpret_cast<const char*>(CFDataGetBytePtr(ClassCode));
+							if (!strncasecmp(ClassCodeValue, "sgx", 3))
 							{
-								Desc.GPUIndex = Index++;
-								
-								CurrentGPUs.Add(Desc);
+								FMacPlatformMisc::FGPUDescriptor Desc;
+#if PLATFORM_MAC_X86
+								// Default initialized to 0 on ARM64, but initialized here for X86 due to
+								// 4.26.0 compatibility requirements.
+								Desc.RegistryID = 0;
+#endif // PLATFORM_MAC_X86
+								InitializeDescriptorFromDeviceEntryM(Desc, ServiceEntry, ServiceInfo);
+								if (Desc.GPUMetalBundle)
+								{
+									Desc.GPUIndex = Index++;
+									CurrentGPUs.Add(Desc);
+								}
 							}
 						}
+#if PLATFORM_MAC_X86
+						else
+						{
+							const uint32* ClassCodeValue = reinterpret_cast<const uint32*>(CFDataGetBytePtr(ClassCode));
+
+							// GPUs are class-code 0x30000 || 0x38000
+							if (ClassCodeValue && (*ClassCodeValue == 0x30000 || *ClassCodeValue == 0x38000))
+							{
+								FMacPlatformMisc::FGPUDescriptor Desc;
+
+								InitializeDescriptorFromDeviceEntry(Desc, ServiceEntry, ServiceInfo);
+
+								if (Desc.GPUMetalBundle)
+								{
+									Desc.GPUIndex = Index++;
+
+									CurrentGPUs.Add(Desc);
+								}
+							}
+						}
+#endif // PLATFORM_MAC_X86
 					}
 					CFRelease(ServiceInfo);
 				}
@@ -1013,15 +1214,9 @@ public:
 			}
 			IOObjectRelease(Iterator);
 		}
-#else
-		FMacPlatformMisc::FGPUDescriptor Desc;
-		Desc.GPUName = @"Apple";
-		Desc.GPUVendorId = 1;
-		Desc.GPUDeviceId = 1;
-		Desc.GPUMemoryMB = 8192;
-		Desc.GPUIndex = 0;
-		CurrentGPUs.Add(Desc);
-#endif // PLATFORM_MAC_X86
+
+		CFRelease(ClassCodeRef);
+
 		UpdatedGPUs = CurrentGPUs;
 	}
 	
@@ -1230,10 +1425,6 @@ FGPUDriverInfo FMacPlatformMisc::GetGPUDriverInfo(const FString& DeviceDescripti
 			bool bGotUserVersionInfo = false;
 			bool bGotDate = false;
 
-#if PLATFORM_MAC_ARM64
-            // MAC_ARM_TODO
-
-#else
 			for(uint32 Index = 0; Index < _dyld_image_count(); Index++)
 			{
 				char const* IndexName = _dyld_get_image_name(Index);
@@ -1302,14 +1493,9 @@ FGPUDriverInfo FMacPlatformMisc::GetGPUDriverInfo(const FString& DeviceDescripti
 					}
 				}
 			}
-#endif // PLATFORM_MAC_ARM
-            
-#if PLATFORM_MAC_X86
-            bool bCanPullDriverInfo = !GMacAppInfo.bIsSandboxed;
-#else
-            bool bCanPullDriverInfo = false;
-#endif
-			
+
+			bool bCanPullDriverInfo = !GMacAppInfo.bIsSandboxed;
+
 			if(bCanPullDriverInfo)
 			{
 				if(!bGotDate || !bGotInternalVersionInfo || !bGotUserVersionInfo)
@@ -1447,32 +1633,40 @@ FString FMacPlatformMisc::GetCPUVendor()
 
 FString FMacPlatformMisc::GetCPUBrand()
 {
-#if PLATFORM_MAC_ARM64
-    return TEXT("Apple Silicon");
-#else
 	static FString Result = FGenericPlatformMisc::GetCPUBrand();
 	static bool bHaveResult = false;
 
 	if (!bHaveResult)
 	{
-		// @see for more information http://msdn.microsoft.com/en-us/library/vstudio/hskdteyh(v=vs.100).aspx
+		// @see (x86_64) for more information http://msdn.microsoft.com/en-us/library/vstudio/hskdteyh(v=vs.100).aspx
 		ANSICHAR BrandString[0x40] = { 0 };
-		int32 CPUInfo[4] = { -1 };
-		const SIZE_T CPUInfoSize = sizeof(CPUInfo);
 
-		asm( "cpuid" : "=a" (CPUInfo[0]), "=b" (CPUInfo[1]), "=c" (CPUInfo[2]), "=d" (CPUInfo[3]) : "a" (0x80000000));
-		const uint32 MaxExtIDs = CPUInfo[0];
-
-		if (MaxExtIDs >= 0x80000004)
+		if (IsRunningOnAppleSilicon())
+		{
+			size_t BrandStringLength = sizeof(BrandString) / sizeof(BrandString[0]);
+			sysctlbyname("machdep.cpu.brand_string", &BrandString, &BrandStringLength, NULL, 0);
+		}
+#if PLATFORM_MAC_X86
+		else
 		{
-			const uint32 FirstBrandString = 0x80000002;
-			const uint32 NumBrandStrings = 3;
-			for (uint32 Index = 0; Index < NumBrandStrings; ++Index)
+			int32 CPUInfo[4] = { -1 };
+			const SIZE_T CPUInfoSize = sizeof(CPUInfo);
+
+			asm( "cpuid" : "=a" (CPUInfo[0]), "=b" (CPUInfo[1]), "=c" (CPUInfo[2]), "=d" (CPUInfo[3]) : "a" (0x80000000));
+			const uint32 MaxExtIDs = CPUInfo[0];
+
+			if (MaxExtIDs >= 0x80000004)
 			{
-				asm( "cpuid" : "=a" (CPUInfo[0]), "=b" (CPUInfo[1]), "=c" (CPUInfo[2]), "=d" (CPUInfo[3]) : "a" (FirstBrandString + Index));
-				FPlatformMemory::Memcpy(BrandString + CPUInfoSize * Index, CPUInfo, CPUInfoSize);
+				const uint32 FirstBrandString = 0x80000002;
+				const uint32 NumBrandStrings = 3;
+				for (uint32 Index = 0; Index < NumBrandStrings; ++Index)
+				{
+					asm( "cpuid" : "=a" (CPUInfo[0]), "=b" (CPUInfo[1]), "=c" (CPUInfo[2]), "=d" (CPUInfo[3]) : "a" (FirstBrandString + Index));
+					FPlatformMemory::Memcpy(BrandString + CPUInfoSize * Index, CPUInfo, CPUInfoSize);
+				}
 			}
 		}
+#endif // PLATFORM_MAC_X86
 
 		Result = BrandString;
 
@@ -1480,7 +1674,6 @@ FString FMacPlatformMisc::GetCPUBrand()
 	}
 
 	return FString(Result);
-#endif
 }
 
 uint32 FMacPlatformMisc::GetCPUInfo()
diff --git a/Engine/Source/Runtime/Core/Public/Mac/MacPlatformMisc.h b/Engine/Source/Runtime/Core/Public/Mac/MacPlatformMisc.h
--- a/Engine/Source/Runtime/Core/Public/Mac/MacPlatformMisc.h
+++ b/Engine/Source/Runtime/Core/Public/Mac/MacPlatformMisc.h
@@ -191,6 +191,9 @@ struct CORE_API FMacPlatformMisc : public FApplePlatformMisc
 			void CopyFromImpl(FGPUDescriptorCommon const& Other);
 			TMap<FString, float> GetPerformanceStatisticsImpl() const;
 
+		public:
+			uint64 RegistryID = 0;
+
 	} FGPUDescriptor;
 
 #else
