Fixed up Windows thread priorities so we play better with others. Previously as of CL3747560 all normal priority threads were boosted, including ShaderCompilerWorker which meant it could cause system stalls by starving other applications of CPU.

The only practical difference is that TPri_Normal now actually maps to THREAD_PRIORITY_NORMAL instead of THREAD_PRIORITY_ABOVE_NORMAL and TPri_SlightlyBelowNormal maps to THREAD_PRIORITY_BELOW_NORMAL (which is the same as TPri_BelowNormal)

Also moved implementation code into cpp to make maintenance easier and cleaned up some comments

diff --git a/Engine/Source/Runtime/Core/Private/Windows/WindowsRunnableThread.cpp b/Engine/Source/Runtime/Core/Private/Windows/WindowsRunnableThread.cpp
--- a/Engine/Source/Runtime/Core/Private/Windows/WindowsRunnableThread.cpp
+++ b/Engine/Source/Runtime/Core/Private/Windows/WindowsRunnableThread.cpp
@@ -1,14 +1,191 @@
 // Copyright Epic Games, Inc. All Rights Reserved.
 
 #include "WindowsRunnableThread.h"
+#include "Containers/StringConv.h"
+#include "CoreGlobals.h"
+#include "GenericPlatform/GenericPlatformCrashContext.h"
+#include "HAL/Event.h"
+#include "HAL/ExceptionHandling.h"
+#include "HAL/LowLevelMemTracker.h"
+#include "HAL/ThreadManager.h"
+#include "Misc/FeedbackContext.h"
 #include "Misc/OutputDeviceError.h"
 #include "Stats/Stats.h"
-#include "Misc/FeedbackContext.h"
-#include "HAL/ExceptionHandling.h"
-#include "GenericPlatform/GenericPlatformCrashContext.h"
 
 DEFINE_LOG_CATEGORY_STATIC(LogThreadingWindows, Log, All);
 
+FRunnableThreadWin::~FRunnableThreadWin()
+{
+	if (Thread)
+	{
+		Kill(true);
+	}
+}
+
+::DWORD STDCALL FRunnableThreadWin::_ThreadProc(LPVOID pThis)
+{
+	check(pThis);
+	auto* ThisThread = (FRunnableThreadWin*)pThis;
+	FThreadManager::Get().AddThread(ThisThread->GetThreadID(), ThisThread);
+	return ThisThread->GuardedRun();
+}
+
+static int TranslateThreadPriority(EThreadPriority Priority)
+{
+	// If this triggers, 
+	static_assert(TPri_Num == 7, "Need to add a case for new TPri_xxx enum value");
+
+	switch (Priority)
+	{
+		case TPri_AboveNormal:			return THREAD_PRIORITY_ABOVE_NORMAL;
+		case TPri_Normal:				return THREAD_PRIORITY_NORMAL;
+		case TPri_BelowNormal:			return THREAD_PRIORITY_BELOW_NORMAL;
+		case TPri_Highest:				return THREAD_PRIORITY_HIGHEST;
+		case TPri_TimeCritical:			return THREAD_PRIORITY_HIGHEST;
+		case TPri_Lowest:				return THREAD_PRIORITY_LOWEST;
+		case TPri_SlightlyBelowNormal:	return THREAD_PRIORITY_BELOW_NORMAL;
+
+	// Note: previously, the behaviour was:
+	//
+	//case TPri_AboveNormal:			return THREAD_PRIORITY_HIGHEST;
+	//case TPri_Normal:					return THREAD_PRIORITY_HIGHEST - 1;
+	//case TPri_BelowNormal:			return THREAD_PRIORITY_HIGHEST - 3;
+	//case TPri_Highest:				return THREAD_PRIORITY_HIGHEST;
+	//case TPri_TimeCritical:			return THREAD_PRIORITY_HIGHEST;
+	//case TPri_Lowest:					return THREAD_PRIORITY_HIGHEST - 4;
+	//case TPri_SlightlyBelowNormal:	return THREAD_PRIORITY_HIGHEST - 2;
+	//
+	// But the change (CL3747560) was not well documented (it didn't describe
+	// the symptoms it was supposed to address) and introduces undesirable 
+	// system behaviour on Windows since it starves out other processes in
+	// the system when UE compiles shaders or otherwise goes wide due to the
+	// inflation in priority (Normal mapped to THREAD_PRIORITY_ABOVE_NORMAL)
+	// I kept the TPri_TimeCritical mapping to THREAD_PRIORITY_HIGHEST however
+	// to avoid introducing poor behaviour since time critical priority is
+	// similarly detrimental to overall system behaviour.
+	//
+	// If we discover thread scheduling issues it would maybe be better to 
+	// adjust actual thread priorities at the source instead of this mapping?
+
+	default: UE_LOG(LogHAL, Fatal, TEXT("Unknown Priority passed to TranslateThreadPriority()")); return THREAD_PRIORITY_NORMAL;
+	}
+}
+
+void FRunnableThreadWin::SetThreadPriority(EThreadPriority NewPriority)
+{
+	ThreadPriority = NewPriority;
+
+	::SetThreadPriority(Thread, TranslateThreadPriority(ThreadPriority));
+}
+
+void FRunnableThreadWin::Suspend(bool bShouldPause)
+{
+	check(Thread);
+	if (bShouldPause == true)
+	{
+		SuspendThread(Thread);
+	}
+	else
+	{
+		ResumeThread(Thread);
+	}
+}
+
+bool FRunnableThreadWin::Kill(bool bShouldWait)
+{
+	check(Thread && "Did you forget to call Create()?");
+	bool bDidExitOK = true;
+
+	// Let the runnable have a chance to stop without brute force killing
+	if (Runnable)
+	{
+		Runnable->Stop();
+	}
+
+	if (bShouldWait == true)
+	{
+		// Wait indefinitely for the thread to finish.  IMPORTANT:  It's not safe to just go and
+		// kill the thread with TerminateThread() as it could have a mutex lock that's shared
+		// with a thread that's continuing to run, which would cause that other thread to
+		// dead-lock.  
+		//
+		// This can manifest itself in code as simple as the synchronization
+		// object that is used by our logging output classes
+
+		WaitForSingleObject(Thread, INFINITE);
+	}
+
+	CloseHandle(Thread);
+	Thread = NULL;
+
+	return bDidExitOK;
+}
+
+void FRunnableThreadWin::WaitForCompletion()
+{
+	WaitForSingleObject(Thread, INFINITE);
+}
+
+bool FRunnableThreadWin::CreateInternal(
+	FRunnable* InRunnable, 
+	const TCHAR* InThreadName,
+	uint32 InStackSize,
+	EThreadPriority InThreadPri, 
+	uint64 InThreadAffinityMask,
+	EThreadCreateFlags InCreateFlags)
+{
+	static bool bOnce = false;
+	if (!bOnce)
+	{
+		bOnce = true;
+		::SetThreadPriority(::GetCurrentThread(), TranslateThreadPriority(TPri_Normal));
+	}
+
+	check(InRunnable);
+	Runnable = InRunnable;
+	ThreadAffinityMask = InThreadAffinityMask;
+
+	// Create a sync event to guarantee the Init() function is called first
+	ThreadInitSyncEvent = FPlatformProcess::GetSynchEventFromPool(true);
+
+	ThreadName = InThreadName ? InThreadName : TEXT("Unnamed UE");
+	ThreadPriority = InThreadPri;
+
+	// Create the new thread
+	{
+		LLM_SCOPE(ELLMTag::ThreadStack);
+		LLM_PLATFORM_SCOPE(ELLMTag::ThreadStackPlatform);
+		// add in the thread size, since it's allocated in a black box we can't track
+		// note: I don't see any accounting for this when threads are destroyed
+		LLM(FLowLevelMemTracker::Get().OnLowLevelAlloc(ELLMTracker::Default, nullptr, InStackSize));
+		LLM(FLowLevelMemTracker::Get().OnLowLevelAlloc(ELLMTracker::Platform, nullptr, InStackSize));
+
+		// Create the thread as suspended, so we can ensure ThreadId is initialized and the thread manager knows about the thread before it runs.
+		Thread = CreateThread(NULL, InStackSize, _ThreadProc, this, STACK_SIZE_PARAM_IS_A_RESERVATION | CREATE_SUSPENDED, (::DWORD*)&ThreadID);
+	}
+
+	// If it fails, clear all the vars
+	if (Thread == NULL)
+	{
+		Runnable = nullptr;
+	}
+	else
+	{
+		ResumeThread(Thread);
+
+		// Let the thread start up
+		ThreadInitSyncEvent->Wait(INFINITE);
+
+		ThreadPriority = TPri_Normal; // Set back to default in case any SetThreadPrio() impls compare against current value to reduce syscalls
+		SetThreadPriority(InThreadPri);
+	}
+
+	// Cleanup the sync event
+	FPlatformProcess::ReturnSynchEventToPool(ThreadInitSyncEvent);
+	ThreadInitSyncEvent = nullptr;
+	return Thread != NULL;
+}
+
 uint32 FRunnableThreadWin::GuardedRun()
 {
 	uint32 ExitCode = 0;
@@ -54,7 +231,10 @@ uint32 FRunnableThreadWin::GuardedRun()
 			}
 			__except(EXCEPTION_EXECUTE_HANDLER)
 			{
-				// The crash handler crashed itself, exit with a code that the out-of-process monitor will be able to pick up and report into analytics.
+				// The crash handler crashed itself, exit with a code which the 
+				// out-of-process monitor will be able to pick up and report into 
+				// analytics.
+
 				::exit(ECrashExitCodes::CrashHandlerCrashed);
 			}
 		}
@@ -67,21 +247,18 @@ uint32 FRunnableThreadWin::GuardedRun()
 
 uint32 FRunnableThreadWin::Run()
 {
-	// Assume we'll fail init
 	uint32 ExitCode = 1;
 	check(Runnable);
 
-	// Initialize the runnable object
 	if (Runnable->Init() == true)
 	{
-		// Initialization has completed, release the sync event
 		ThreadInitSyncEvent->Trigger();
 
 		// Setup TLS for this thread, used by FTlsAutoCleanup objects.
 		SetTls();
 
-		// Now run the task that needs to be done
 		ExitCode = Runnable->Run();
+
 		// Allow any allocated resources to be cleaned up
 		Runnable->Exit();
 
diff --git a/Engine/Source/Runtime/Core/Private/Windows/WindowsRunnableThread.h b/Engine/Source/Runtime/Core/Private/Windows/WindowsRunnableThread.h
--- a/Engine/Source/Runtime/Core/Private/Windows/WindowsRunnableThread.h
+++ b/Engine/Source/Runtime/Core/Private/Windows/WindowsRunnableThread.h
@@ -3,15 +3,8 @@
 #pragma once
 
 #include "CoreTypes.h"
-#include "HAL/PlatformProcess.h"
-#include "Containers/StringConv.h"
-#include "HAL/Event.h"
-#include "HAL/Runnable.h"
 #include "HAL/RunnableThread.h"
-#include "HAL/ThreadManager.h"
-#include "CoreGlobals.h"
 #include "Windows/WindowsHWrapper.h"
-#include "HAL/LowLevelMemTracker.h"
 
 class FRunnable;
 
@@ -23,19 +16,13 @@ class FRunnableThreadWin
 	: public FRunnableThread
 {
 	/** The thread handle for the thread. */
-	HANDLE Thread;
+	HANDLE Thread = 0;
 
 	/**
 	 * The thread entry point. Simply forwards the call on to the right
 	 * thread main function
 	 */
-	static ::DWORD STDCALL _ThreadProc( LPVOID pThis )
-	{
-		check(pThis);
-		auto* ThisThread = (FRunnableThreadWin*)pThis;
-		FThreadManager::Get().AddThread(ThisThread->GetThreadID(), ThisThread);
-		return ThisThread->GuardedRun();
-	}
+	static ::DWORD STDCALL _ThreadProc(LPVOID pThis);
 
 	/** Guarding works only if debugger is not attached or GAlwaysReportCrash is true. */
 	uint32 GuardedRun();
@@ -47,145 +34,17 @@ class FRunnableThreadWin
 	uint32 Run();
 
 public:
-
-	FRunnableThreadWin( )
-		: Thread(NULL)
-	{}
-
-	~FRunnableThreadWin( )
-	{
-		// Clean up our thread if it is still active
-		if (Thread != NULL)
-		{
-			Kill(true);
-		}
-	}
+	~FRunnableThreadWin();
 	
-	static int TranslateThreadPriority(EThreadPriority Priority)
-	{
-		switch (Priority)
-		{
-		case TPri_AboveNormal: return THREAD_PRIORITY_HIGHEST;
-		case TPri_Normal: return THREAD_PRIORITY_HIGHEST - 1;
-		case TPri_BelowNormal: return THREAD_PRIORITY_HIGHEST - 3;
-		case TPri_Highest: return THREAD_PRIORITY_HIGHEST;
-		case TPri_TimeCritical: return THREAD_PRIORITY_HIGHEST;
-		case TPri_Lowest: return THREAD_PRIORITY_HIGHEST - 4;
-		case TPri_SlightlyBelowNormal: return THREAD_PRIORITY_HIGHEST - 2;
-		default: UE_LOG(LogHAL, Fatal, TEXT("Unknown Priority passed to TranslateThreadPriority()")); return TPri_Normal;
-		}
-	}
-
-	virtual void SetThreadPriority( EThreadPriority NewPriority ) override
-	{
-		// Don't bother calling the OS if there is no need
-			ThreadPriority = NewPriority;
-			// Change the priority on the thread
-			::SetThreadPriority(Thread, TranslateThreadPriority(ThreadPriority));
-		}
-
-	virtual void Suspend( bool bShouldPause = true ) override
-	{
-		check(Thread);
-		if (bShouldPause == true)
-		{
-			SuspendThread(Thread);
-		}
-		else
-		{
-			ResumeThread(Thread);
-		}
-	}
-
-	virtual bool Kill( bool bShouldWait = false ) override
-	{
-		check(Thread && "Did you forget to call Create()?");
-		bool bDidExitOK = true;
-		// Let the runnable have a chance to stop without brute force killing
-		if (Runnable)
-		{
-			Runnable->Stop();
-		}
-		// If waiting was specified, wait the amount of time. If that fails,
-		// brute force kill that thread. Very bad as that might leak.
-		if (bShouldWait == true)
-		{
-			// Wait indefinitely for the thread to finish.  IMPORTANT:  It's not safe to just go and
-			// kill the thread with TerminateThread() as it could have a mutex lock that's shared
-			// with a thread that's continuing to run, which would cause that other thread to
-			// dead-lock.  (This can manifest itself in code as simple as the synchronization
-			// object that is used by our logging output classes.  Trust us, we've seen it!)
-			WaitForSingleObject(Thread,INFINITE);
-		}
-		// Now clean up the thread handle so we don't leak
-		CloseHandle(Thread);
-		Thread = NULL;
-
-		return bDidExitOK;
-	}
-
-	virtual void WaitForCompletion( ) override
-	{
-		// Block until this thread exits
-		WaitForSingleObject(Thread,INFINITE);
-	}
+	virtual void SetThreadPriority(EThreadPriority NewPriority) override;
+	virtual void Suspend(bool bShouldPause = true) override;
+	virtual bool Kill(bool bShouldWait = false) override;
+	virtual void WaitForCompletion() override;
 
 protected:
 
-	virtual bool CreateInternal( FRunnable* InRunnable, const TCHAR* InThreadName,
+	virtual bool CreateInternal(FRunnable* InRunnable, const TCHAR* InThreadName,
 		uint32 InStackSize = 0,
 		EThreadPriority InThreadPri = TPri_Normal, uint64 InThreadAffinityMask = 0,
-		EThreadCreateFlags InCreateFlags = EThreadCreateFlags::None) override
-	{
-		static bool bOnce = false;
-		if (!bOnce)
-		{
-			bOnce = true;
-			::SetThreadPriority(::GetCurrentThread(), TranslateThreadPriority(TPri_Normal)); // set the main thread to be normal, since this is no longer the windows default.
-		}
-
-
-		check(InRunnable);
-		Runnable = InRunnable;
-		ThreadAffinityMask = InThreadAffinityMask;
-
-		// Create a sync event to guarantee the Init() function is called first
-		ThreadInitSyncEvent	= FPlatformProcess::GetSynchEventFromPool(true);
-
-		ThreadName = InThreadName ? InThreadName : TEXT("Unnamed UE4");
-		ThreadPriority = InThreadPri;
-
-		// Create the new thread
-		{
-			LLM_SCOPE(ELLMTag::ThreadStack);
-			LLM_PLATFORM_SCOPE(ELLMTag::ThreadStackPlatform);
-			// add in the thread size, since it's allocated in a black box we can't track
-			LLM(FLowLevelMemTracker::Get().OnLowLevelAlloc(ELLMTracker::Default, nullptr, InStackSize));
-			LLM(FLowLevelMemTracker::Get().OnLowLevelAlloc(ELLMTracker::Platform, nullptr, InStackSize));
-
-			// Create the thread as suspended, so we can ensure ThreadId is initialized and the thread manager knows about the thread before it runs.
-			Thread = CreateThread(NULL, InStackSize, _ThreadProc, this, STACK_SIZE_PARAM_IS_A_RESERVATION | CREATE_SUSPENDED, (::DWORD *)&ThreadID);
-		}
-
-		// If it fails, clear all the vars
-		if (Thread == NULL)
-		{
-			Runnable = nullptr;
-		}
-		else
-		{
-			ResumeThread(Thread);
-
-			// Let the thread start up
-			ThreadInitSyncEvent->Wait(INFINITE);
-
-			ThreadPriority = TPri_Normal; // Set back to default in case any SetThreadPrio() impls compare against current value to reduce syscalls
-			SetThreadPriority(InThreadPri);
-		}
-
-		// Cleanup the sync event
-		FPlatformProcess::ReturnSynchEventToPool(ThreadInitSyncEvent);
-		ThreadInitSyncEvent = nullptr;
-		return Thread != NULL;
-	}
+		EThreadCreateFlags InCreateFlags = EThreadCreateFlags::None) override;
 };
