Ensure that preprocessor macros are expanded before calling UE_LOG()'s internal lambda. There are numerous patterns like this; UE_LOG(..., __FUNCTION__) and a recent change to UE_LOG() was causing __FUNCTION__ to be incorrectly expanded within the lambda's scope, leading to misleading/obfuscated log output.

#rnx

diff --git a/Engine/Source/Runtime/Core/Public/Logging/LogMacros.h b/Engine/Source/Runtime/Core/Public/Logging/LogMacros.h
--- a/Engine/Source/Runtime/Core/Public/Logging/LogMacros.h
+++ b/Engine/Source/Runtime/Core/Public/Logging/LogMacros.h
@@ -186,20 +186,20 @@ private:
 		static_assert((ELogVerbosity::Verbosity & ELogVerbosity::VerbosityMask) < ELogVerbosity::NumVerbosity && ELogVerbosity::Verbosity > 0, "Verbosity must be constant and in range."); \
 		UE_LOG_EXPAND_IS_FATAL(Verbosity, PREPROCESSOR_NOTHING, if (!CategoryName.IsSuppressed(ELogVerbosity::Verbosity))) \
 			{ \
-				DispatchCheckVerify([&] () FORCENOINLINE \
+				DispatchCheckVerify([] (const auto& LCategoryName, const auto& LFormat, const auto&... UE_LOG_Args) FORCENOINLINE \
 				{ \
-					TRACE_LOG_MESSAGE(CategoryName, Verbosity, Format, ##__VA_ARGS__) \
+					TRACE_LOG_MESSAGE(LCategoryName, Verbosity, LFormat, UE_LOG_Args...) \
 					UE_LOG_EXPAND_IS_FATAL(Verbosity, \
 						{ \
-							FMsg::Logf_Internal(UE_LOG_SOURCE_FILE(__FILE__), __LINE__, CategoryName.GetCategoryName(), ELogVerbosity::Verbosity, Format, ##__VA_ARGS__); \
+							FMsg::Logf_Internal(UE_LOG_SOURCE_FILE(__FILE__), __LINE__, LCategoryName.GetCategoryName(), ELogVerbosity::Verbosity, LFormat, UE_LOG_Args...); \
 							_DebugBreakAndPromptForRemote(); \
 							FDebug::ProcessFatalError({"UE_LOG(Fatal)", __FILE__, __LINE__, PLATFORM_RETURN_ADDRESS()}); \
 						}, \
 						{ \
-							FMsg::Logf_Internal(nullptr, 0, CategoryName.GetCategoryName(), ELogVerbosity::Verbosity, Format, ##__VA_ARGS__); \
+							FMsg::Logf_Internal(nullptr, 0, LCategoryName.GetCategoryName(), ELogVerbosity::Verbosity, LFormat, UE_LOG_Args...); \
 						} \
 					) \
-				}); \
+				}, CategoryName, Format, ##__VA_ARGS__); \
 				UE_LOG_EXPAND_IS_FATAL(Verbosity, CA_ASSUME(false);, PREPROCESSOR_NOTHING) \
 			} \
 
@@ -290,21 +290,21 @@ private:
 			{ \
 				if (Condition) \
 				{ \
-					DispatchCheckVerify([&] () FORCENOINLINE \
+					DispatchCheckVerify([] (const auto& LCategoryName, const auto& LFormat, const auto&... UE_LOG_Args) FORCENOINLINE \
 					{ \
-						TRACE_LOG_MESSAGE(CategoryName, Verbosity, Format, ##__VA_ARGS__) \
+						TRACE_LOG_MESSAGE(LCategoryName, Verbosity, LFormat, UE_LOG_Args...) \
 						UE_LOG_EXPAND_IS_FATAL(Verbosity, \
 							{ \
-								FMsg::Logf_Internal(UE_LOG_SOURCE_FILE(__FILE__), __LINE__, CategoryName.GetCategoryName(), ELogVerbosity::Verbosity, Format, ##__VA_ARGS__); \
+								FMsg::Logf_Internal(UE_LOG_SOURCE_FILE(__FILE__), __LINE__, LCategoryName.GetCategoryName(), ELogVerbosity::Verbosity, LFormat, UE_LOG_Args...); \
 								_DebugBreakAndPromptForRemote(); \
 								FDebug::ProcessFatalError({"UE_LOG(Fatal)", __FILE__, __LINE__, PLATFORM_RETURN_ADDRESS()}); \
 							}, \
 							{ \
-								FMsg::Logf_Internal(nullptr, 0, CategoryName.GetCategoryName(), ELogVerbosity::Verbosity, Format, ##__VA_ARGS__); \
+								FMsg::Logf_Internal(nullptr, 0, LCategoryName.GetCategoryName(), ELogVerbosity::Verbosity, LFormat, UE_LOG_Args...); \
 							} \
 						) \
 						CA_ASSUME(true); \
-					}); \
+					}, CategoryName, Format, ##__VA_ARGS__); \
 					UE_LOG_EXPAND_IS_FATAL(Verbosity, CA_ASSUME(false);, PREPROCESSOR_NOTHING) \
 				} \
 			} \
diff --git a/Engine/Source/Runtime/Core/Public/Misc/AssertionMacros.h b/Engine/Source/Runtime/Core/Public/Misc/AssertionMacros.h
--- a/Engine/Source/Runtime/Core/Public/Misc/AssertionMacros.h
+++ b/Engine/Source/Runtime/Core/Public/Misc/AssertionMacros.h
@@ -174,10 +174,10 @@ private:
 // MSVC (v19.00.24215.1 at time of writing) ignores no-inline attributes on
 // lambdas. This can be worked around by calling the lambda from inside this
 // templated (and correctly non-inlined) function.
-template <typename RetType=void, class InnerType>
-RetType FORCENOINLINE UE_DEBUG_SECTION DispatchCheckVerify(InnerType&& Inner)
+template <typename RetType=void, class InnerType, typename... ArgTypes>
+RetType FORCENOINLINE UE_DEBUG_SECTION DispatchCheckVerify(InnerType&& Inner, ArgTypes const&... Args)
 {
-	return Inner();
+	return Inner(Args...);
 }
 
 #if !UE_BUILD_SHIPPING
