Changed default sleep interval for FMonitoredProcess from 0.0 to 0.01 to stop it from spinning continuously hogging an entire core if the caller does not update the interval using SetSleepInterval(), which only a single call site did, which resulted in multi-second 100% utilisation of cores during startup and at other points in the editor

diff --git a/Engine/Source/Runtime/Core/Private/Misc/MonitoredProcess.cpp b/Engine/Source/Runtime/Core/Private/Misc/MonitoredProcess.cpp
--- a/Engine/Source/Runtime/Core/Private/Misc/MonitoredProcess.cpp
+++ b/Engine/Source/Runtime/Core/Private/Misc/MonitoredProcess.cpp
@@ -3,31 +3,24 @@
 #include "Misc/MonitoredProcess.h"
 #include "HAL/RunnableThread.h"
 #include "Misc/Paths.h"
+#include <atomic>
 
 /* FMonitoredProcess structors
  *****************************************************************************/
 
 FMonitoredProcess::FMonitoredProcess( const FString& InURL, const FString& InParams, bool InHidden, bool InCreatePipes )
 	: FMonitoredProcess(InURL, InParams, FPaths::RootDir(), InHidden, InCreatePipes)
-{ }
+{ 
+}
 
 FMonitoredProcess::FMonitoredProcess( const FString& InURL, const FString& InParams, const FString& InWorkingDir, bool InHidden, bool InCreatePipes )
-	: Canceling(false)
-	, EndTime(0)
-	, Hidden(InHidden)
-	, KillTree(false)
+	: Hidden(InHidden)
 	, Params(InParams)
-	, ReadPipe(nullptr)
-	, ReturnCode(0)
-	, StartTime(0)
-	, Thread(nullptr)
-	, bIsRunning(false)
 	, URL(InURL)
 	, WorkingDir(InWorkingDir)
-	, WritePipe(nullptr)
 	, bCreatePipes(InCreatePipes)
-	, SleepInterval(0.0f)
-{ }
+{ 
+}
 
  
 FMonitoredProcess::~FMonitoredProcess()
@@ -80,9 +73,8 @@ bool FMonitoredProcess::Launch()
 		return false;
 	}
 
-	static int32 MonitoredProcessIndex = 0;
-	const FString MonitoredProcessName = FString::Printf( TEXT( "FMonitoredProcess %d" ), MonitoredProcessIndex );
-	MonitoredProcessIndex++;
+	static std::atomic_uint32_t MonitoredProcessIndex { 0 };
+	const FString MonitoredProcessName = FString::Printf( TEXT( "FMonitoredProcess %d" ), MonitoredProcessIndex.fetch_add(1) );
 
 	bIsRunning = true;
 	Thread = FRunnableThread::Create(this, *MonitoredProcessName, 128 * 1024, TPri_AboveNormal);
diff --git a/Engine/Source/Runtime/Core/Public/Misc/MonitoredProcess.h b/Engine/Source/Runtime/Core/Public/Misc/MonitoredProcess.h
--- a/Engine/Source/Runtime/Core/Public/Misc/MonitoredProcess.h
+++ b/Engine/Source/Runtime/Core/Public/Misc/MonitoredProcess.h
@@ -192,16 +192,16 @@ private:
 
 
 	// Whether the process is being canceled. */
-	bool Canceling;
+	bool Canceling = false;
 
 	// Holds the time at which the process ended. */
 	FDateTime EndTime;
 
 	// Whether the window of the process should be hidden. */
-	bool Hidden;
+	bool Hidden = false;
 
 	// Whether to kill the entire process tree when cancelling this process. */
-	bool KillTree;
+	bool KillTree = false;
 
 	// Holds the command line parameters. */
 	FString Params;
@@ -210,19 +210,19 @@ private:
 	FProcHandle ProcessHandle;
 
 	// Holds the read pipe. */
-	void* ReadPipe;
+	void* ReadPipe = nullptr;
 
 	// Holds the return code. */
-	int ReturnCode;
+	int ReturnCode = 0;
 
 	// Holds the time at which the process started. */
-	FDateTime StartTime;
+	FDateTime StartTime { 0 };
 
 	// Holds the monitoring thread object. */
-	FRunnableThread* Thread;
+	FRunnableThread* Thread = nullptr;
 
 	// Is the thread running? 
-	TSAN_ATOMIC(bool) bIsRunning;
+	TSAN_ATOMIC(bool) bIsRunning = false;
 
 	// Holds the URL of the executable to launch. */
 	FString URL;
@@ -231,13 +231,13 @@ private:
 	FString WorkingDir;
 
 	// Holds the write pipe. */
-	void* WritePipe;
+	void* WritePipe = nullptr;
 
 	// Holds if we should create pipes
-	bool bCreatePipes;
+	bool bCreatePipes = false;
 
 	// Sleep interval to use
-	float SleepInterval;
+	float SleepInterval = 0.01f;
 
 	// Buffered output text which does not contain a newline
 	FString OutputBuffer;
