Fix NO_LOGGING compilation
#rnx

diff --git a/Engine/Source/Runtime/Core/Public/ProfilingDebugging/ScopedTimers.h b/Engine/Source/Runtime/Core/Public/ProfilingDebugging/ScopedTimers.h
--- a/Engine/Source/Runtime/Core/Public/ProfilingDebugging/ScopedTimers.h
+++ b/Engine/Source/Runtime/Core/Public/ProfilingDebugging/ScopedTimers.h
@@ -263,5 +263,9 @@ private:
 	FDurationTimer Timer;
 };
 
+#if NO_LOGGING
+#define UE_SCOPED_TIMER(Title, Category)
+#else
 #define UE_SCOPED_TIMER(Title, Category) \
 	FScopedDurationTimeCustomLogger BODY_MACRO_COMBINE(Scoped,Timer,_,__LINE__)(Title, Category)
+#endif
