Optional is the atomic_uint32_t alias is in C++11 and seemingly not available on all platforms.

#rnx

diff --git a/Engine/Source/Runtime/Core/Private/Misc/MonitoredProcess.cpp b/Engine/Source/Runtime/Core/Private/Misc/MonitoredProcess.cpp
--- a/Engine/Source/Runtime/Core/Private/Misc/MonitoredProcess.cpp
+++ b/Engine/Source/Runtime/Core/Private/Misc/MonitoredProcess.cpp
@@ -73,7 +73,7 @@ bool FMonitoredProcess::Launch()
 		return false;
 	}
 
-	static std::atomic_uint32_t MonitoredProcessIndex { 0 };
+	static std::atomic<uint32> MonitoredProcessIndex { 0 };
 	const FString MonitoredProcessName = FString::Printf( TEXT( "FMonitoredProcess %d" ), MonitoredProcessIndex.fetch_add(1) );
 
 	bIsRunning = true;
