Linker "missing symbol" error fixes. Some of FRunnableThreadWin's methods are inlined because the class is shared with other platforms.

#rnx

diff --git a/Engine/Source/Runtime/Core/Private/Windows/WindowsRunnableThread.cpp b/Engine/Source/Runtime/Core/Private/Windows/WindowsRunnableThread.cpp
--- a/Engine/Source/Runtime/Core/Private/Windows/WindowsRunnableThread.cpp
+++ b/Engine/Source/Runtime/Core/Private/Windows/WindowsRunnableThread.cpp
@@ -14,23 +14,7 @@
 
 DEFINE_LOG_CATEGORY_STATIC(LogThreadingWindows, Log, All);
 
-FRunnableThreadWin::~FRunnableThreadWin()
-{
-	if (Thread)
-	{
-		Kill(true);
-	}
-}
-
-::DWORD STDCALL FRunnableThreadWin::_ThreadProc(LPVOID pThis)
-{
-	check(pThis);
-	auto* ThisThread = (FRunnableThreadWin*)pThis;
-	FThreadManager::Get().AddThread(ThisThread->GetThreadID(), ThisThread);
-	return ThisThread->GuardedRun();
-}
-
-static int TranslateThreadPriority(EThreadPriority Priority)
+int FRunnableThreadWin::TranslateThreadPriority(EThreadPriority Priority)
 {
 	// If this triggers, 
 	static_assert(TPri_Num == 7, "Need to add a case for new TPri_xxx enum value");
@@ -71,121 +55,6 @@ static int TranslateThreadPriority(EThreadPriority Priority)
 	}
 }
 
-void FRunnableThreadWin::SetThreadPriority(EThreadPriority NewPriority)
-{
-	ThreadPriority = NewPriority;
-
-	::SetThreadPriority(Thread, TranslateThreadPriority(ThreadPriority));
-}
-
-void FRunnableThreadWin::Suspend(bool bShouldPause)
-{
-	check(Thread);
-	if (bShouldPause == true)
-	{
-		SuspendThread(Thread);
-	}
-	else
-	{
-		ResumeThread(Thread);
-	}
-}
-
-bool FRunnableThreadWin::Kill(bool bShouldWait)
-{
-	check(Thread && "Did you forget to call Create()?");
-	bool bDidExitOK = true;
-
-	// Let the runnable have a chance to stop without brute force killing
-	if (Runnable)
-	{
-		Runnable->Stop();
-	}
-
-	if (bShouldWait == true)
-	{
-		// Wait indefinitely for the thread to finish.  IMPORTANT:  It's not safe to just go and
-		// kill the thread with TerminateThread() as it could have a mutex lock that's shared
-		// with a thread that's continuing to run, which would cause that other thread to
-		// dead-lock.  
-		//
-		// This can manifest itself in code as simple as the synchronization
-		// object that is used by our logging output classes
-
-		WaitForSingleObject(Thread, INFINITE);
-	}
-
-	CloseHandle(Thread);
-	Thread = NULL;
-
-	return bDidExitOK;
-}
-
-void FRunnableThreadWin::WaitForCompletion()
-{
-	WaitForSingleObject(Thread, INFINITE);
-}
-
-bool FRunnableThreadWin::CreateInternal(
-	FRunnable* InRunnable, 
-	const TCHAR* InThreadName,
-	uint32 InStackSize,
-	EThreadPriority InThreadPri, 
-	uint64 InThreadAffinityMask,
-	EThreadCreateFlags InCreateFlags)
-{
-	static bool bOnce = false;
-	if (!bOnce)
-	{
-		bOnce = true;
-		::SetThreadPriority(::GetCurrentThread(), TranslateThreadPriority(TPri_Normal));
-	}
-
-	check(InRunnable);
-	Runnable = InRunnable;
-	ThreadAffinityMask = InThreadAffinityMask;
-
-	// Create a sync event to guarantee the Init() function is called first
-	ThreadInitSyncEvent = FPlatformProcess::GetSynchEventFromPool(true);
-
-	ThreadName = InThreadName ? InThreadName : TEXT("Unnamed UE");
-	ThreadPriority = InThreadPri;
-
-	// Create the new thread
-	{
-		LLM_SCOPE(ELLMTag::ThreadStack);
-		LLM_PLATFORM_SCOPE(ELLMTag::ThreadStackPlatform);
-		// add in the thread size, since it's allocated in a black box we can't track
-		// note: I don't see any accounting for this when threads are destroyed
-		LLM(FLowLevelMemTracker::Get().OnLowLevelAlloc(ELLMTracker::Default, nullptr, InStackSize));
-		LLM(FLowLevelMemTracker::Get().OnLowLevelAlloc(ELLMTracker::Platform, nullptr, InStackSize));
-
-		// Create the thread as suspended, so we can ensure ThreadId is initialized and the thread manager knows about the thread before it runs.
-		Thread = CreateThread(NULL, InStackSize, _ThreadProc, this, STACK_SIZE_PARAM_IS_A_RESERVATION | CREATE_SUSPENDED, (::DWORD*)&ThreadID);
-	}
-
-	// If it fails, clear all the vars
-	if (Thread == NULL)
-	{
-		Runnable = nullptr;
-	}
-	else
-	{
-		ResumeThread(Thread);
-
-		// Let the thread start up
-		ThreadInitSyncEvent->Wait(INFINITE);
-
-		ThreadPriority = TPri_Normal; // Set back to default in case any SetThreadPrio() impls compare against current value to reduce syscalls
-		SetThreadPriority(InThreadPri);
-	}
-
-	// Cleanup the sync event
-	FPlatformProcess::ReturnSynchEventToPool(ThreadInitSyncEvent);
-	ThreadInitSyncEvent = nullptr;
-	return Thread != NULL;
-}
-
 uint32 FRunnableThreadWin::GuardedRun()
 {
 	uint32 ExitCode = 0;
diff --git a/Engine/Source/Runtime/Core/Private/Windows/WindowsRunnableThread.h b/Engine/Source/Runtime/Core/Private/Windows/WindowsRunnableThread.h
--- a/Engine/Source/Runtime/Core/Private/Windows/WindowsRunnableThread.h
+++ b/Engine/Source/Runtime/Core/Private/Windows/WindowsRunnableThread.h
@@ -4,6 +4,7 @@
 
 #include "CoreTypes.h"
 #include "HAL/RunnableThread.h"
+#include "HAL/ThreadManager.h"
 #include "Windows/WindowsHWrapper.h"
 
 class FRunnable;
@@ -22,7 +23,13 @@ class FRunnableThreadWin
 	 * The thread entry point. Simply forwards the call on to the right
 	 * thread main function
 	 */
-	static ::DWORD STDCALL _ThreadProc(LPVOID pThis);
+	static ::DWORD STDCALL _ThreadProc(LPVOID pThis)
+	{
+		check(pThis);
+		auto* ThisThread = (FRunnableThreadWin*)pThis;
+		FThreadManager::Get().AddThread(ThisThread->GetThreadID(), ThisThread);
+		return ThisThread->GuardedRun();
+	}
 
 	/** Guarding works only if debugger is not attached or GAlwaysReportCrash is true. */
 	uint32 GuardedRun();
@@ -34,17 +41,126 @@ class FRunnableThreadWin
 	uint32 Run();
 
 public:
-	~FRunnableThreadWin();
-	
-	virtual void SetThreadPriority(EThreadPriority NewPriority) override;
-	virtual void Suspend(bool bShouldPause = true) override;
-	virtual bool Kill(bool bShouldWait = false) override;
-	virtual void WaitForCompletion() override;
+	~FRunnableThreadWin()
+	{
+		if (Thread)
+		{
+			Kill(true);
+		}
+	}
+
+	virtual void SetThreadPriority(EThreadPriority NewPriority) override
+	{
+		ThreadPriority = NewPriority;
+
+		::SetThreadPriority(Thread, TranslateThreadPriority(ThreadPriority));
+	}
+
+	virtual void Suspend(bool bShouldPause = true) override
+	{
+		check(Thread);
+		if (bShouldPause == true)
+		{
+			SuspendThread(Thread);
+		}
+		else
+		{
+			ResumeThread(Thread);
+		}
+	}
+
+	virtual bool Kill(bool bShouldWait = false) override
+	{
+		check(Thread && "Did you forget to call Create()?");
+		bool bDidExitOK = true;
+
+		// Let the runnable have a chance to stop without brute force killing
+		if (Runnable)
+		{
+			Runnable->Stop();
+		}
+
+		if (bShouldWait == true)
+		{
+			// Wait indefinitely for the thread to finish.  IMPORTANT:  It's not safe to just go and
+			// kill the thread with TerminateThread() as it could have a mutex lock that's shared
+			// with a thread that's continuing to run, which would cause that other thread to
+			// dead-lock.  
+			//
+			// This can manifest itself in code as simple as the synchronization
+			// object that is used by our logging output classes
+
+			WaitForSingleObject(Thread, INFINITE);
+		}
+
+		CloseHandle(Thread);
+		Thread = NULL;
+
+		return bDidExitOK;
+	}
+
+	virtual void WaitForCompletion() override
+	{
+		WaitForSingleObject(Thread, INFINITE);
+	}
 
 protected:
+ 	static int TranslateThreadPriority(EThreadPriority Priority);
 
 	virtual bool CreateInternal(FRunnable* InRunnable, const TCHAR* InThreadName,
 		uint32 InStackSize = 0,
 		EThreadPriority InThreadPri = TPri_Normal, uint64 InThreadAffinityMask = 0,
-		EThreadCreateFlags InCreateFlags = EThreadCreateFlags::None) override;
+		EThreadCreateFlags InCreateFlags = EThreadCreateFlags::None) override
+	{
+		static bool bOnce = false;
+		if (!bOnce)
+		{
+			bOnce = true;
+			::SetThreadPriority(::GetCurrentThread(), TranslateThreadPriority(TPri_Normal));
+		}
+
+		check(InRunnable);
+		Runnable = InRunnable;
+		ThreadAffinityMask = InThreadAffinityMask;
+
+		// Create a sync event to guarantee the Init() function is called first
+		ThreadInitSyncEvent = FPlatformProcess::GetSynchEventFromPool(true);
+
+		ThreadName = InThreadName ? InThreadName : TEXT("Unnamed UE");
+		ThreadPriority = InThreadPri;
+
+		// Create the new thread
+		{
+			LLM_SCOPE(ELLMTag::ThreadStack);
+			LLM_PLATFORM_SCOPE(ELLMTag::ThreadStackPlatform);
+			// add in the thread size, since it's allocated in a black box we can't track
+			// note: I don't see any accounting for this when threads are destroyed
+			LLM(FLowLevelMemTracker::Get().OnLowLevelAlloc(ELLMTracker::Default, nullptr, InStackSize));
+			LLM(FLowLevelMemTracker::Get().OnLowLevelAlloc(ELLMTracker::Platform, nullptr, InStackSize));
+
+			// Create the thread as suspended, so we can ensure ThreadId is initialized and the thread manager knows about the thread before it runs.
+			Thread = CreateThread(NULL, InStackSize, _ThreadProc, this, STACK_SIZE_PARAM_IS_A_RESERVATION | CREATE_SUSPENDED, (::DWORD*)&ThreadID);
+		}
+
+		// If it fails, clear all the vars
+		if (Thread == NULL)
+		{
+			Runnable = nullptr;
+		}
+		else
+		{
+			ResumeThread(Thread);
+
+			// Let the thread start up
+			ThreadInitSyncEvent->Wait(INFINITE);
+
+			ThreadPriority = TPri_Normal; // Set back to default in case any SetThreadPrio() impls compare against current value to reduce syscalls
+			SetThreadPriority(InThreadPri);
+		}
+
+		// Cleanup the sync event
+		FPlatformProcess::ReturnSynchEventToPool(ThreadInitSyncEvent);
+		ThreadInitSyncEvent = nullptr;
+		return Thread != NULL;
+	}
 };
