Left out some include statements for non-unity builds.

#rnx

diff --git a/Engine/Source/Runtime/Core/Private/Windows/WindowsRunnableThread.cpp b/Engine/Source/Runtime/Core/Private/Windows/WindowsRunnableThread.cpp
--- a/Engine/Source/Runtime/Core/Private/Windows/WindowsRunnableThread.cpp
+++ b/Engine/Source/Runtime/Core/Private/Windows/WindowsRunnableThread.cpp
@@ -1,13 +1,8 @@
 // Copyright Epic Games, Inc. All Rights Reserved.
 
 #include "WindowsRunnableThread.h"
-#include "Containers/StringConv.h"
-#include "CoreGlobals.h"
 #include "GenericPlatform/GenericPlatformCrashContext.h"
-#include "HAL/Event.h"
 #include "HAL/ExceptionHandling.h"
-#include "HAL/LowLevelMemTracker.h"
-#include "HAL/ThreadManager.h"
 #include "Misc/FeedbackContext.h"
 #include "Misc/OutputDeviceError.h"
 #include "Stats/Stats.h"
diff --git a/Engine/Source/Runtime/Core/Private/Windows/WindowsRunnableThread.h b/Engine/Source/Runtime/Core/Private/Windows/WindowsRunnableThread.h
--- a/Engine/Source/Runtime/Core/Private/Windows/WindowsRunnableThread.h
+++ b/Engine/Source/Runtime/Core/Private/Windows/WindowsRunnableThread.h
@@ -2,7 +2,13 @@
 
 #pragma once
 
+#include "Containers/StringConv.h"
+#include "CoreGlobals.h"
 #include "CoreTypes.h"
+#include "HAL/Event.h"
+#include "HAL/LowLevelMemTracker.h"
+#include "HAL/PlatformProcess.h"
+#include "HAL/Runnable.h"
 #include "HAL/RunnableThread.h"
 #include "HAL/ThreadManager.h"
 #include "Windows/WindowsHWrapper.h"
