Marked Logging.LogMessage as unsynchronised. These trace events already have an embedded timestamp field.

#rnx

diff --git a/Engine/Source/Runtime/Core/Private/Logging/LogTrace.cpp b/Engine/Source/Runtime/Core/Private/Logging/LogTrace.cpp
--- a/Engine/Source/Runtime/Core/Private/Logging/LogTrace.cpp
+++ b/Engine/Source/Runtime/Core/Private/Logging/LogTrace.cpp
@@ -25,7 +25,7 @@ UE_TRACE_EVENT_BEGIN(Logging, LogMessageSpec, NoSync|Important)
 	UE_TRACE_EVENT_FIELD(UE::Trace::WideString, FormatString)
 UE_TRACE_EVENT_END()
 
-UE_TRACE_EVENT_BEGIN(Logging, LogMessage)
+UE_TRACE_EVENT_BEGIN(Logging, LogMessage, NoSync)
 	UE_TRACE_EVENT_FIELD(const void*, LogPoint)
 	UE_TRACE_EVENT_FIELD(uint64, Cycle)
 UE_TRACE_EVENT_END()
