Cheap solution for Windows deadlocks when tracing memory from modular targets.

#rnx

diff --git a/Engine/Source/Runtime/Core/Private/ProfilingDebugging/Windows/CallstackTrace.cpp b/Engine/Source/Runtime/Core/Private/ProfilingDebugging/Windows/CallstackTrace.cpp
--- a/Engine/Source/Runtime/Core/Private/ProfilingDebugging/Windows/CallstackTrace.cpp
+++ b/Engine/Source/Runtime/Core/Private/ProfilingDebugging/Windows/CallstackTrace.cpp
@@ -460,6 +460,17 @@ void FBacktracer::AddModule(UPTRINT ModuleBase, const TCHAR* Name)
 ////////////////////////////////////////////////////////////////////////////////
 void FBacktracer::RemoveModule(UPTRINT ModuleBase)
 {
+	// When Windows' RequestExit() is called it hard-terminates all threads except
+	// the main thread and then proceeds to unload the process' DLLs. This hard 
+	// thread termination can result is dangling locked locks. Not an issue as
+	// the rule is "do not do anything multithreaded in DLL load/unload". And here
+	// we are, taking write locks during DLL unload which is, quite unsurprisingly,
+	// deadlocking. In reality tracking Windows' DLL unloads doesn't tell us
+	// anything due to how DLLs and processes' address spaces work. So we will...
+#if defined PLATFORM_WINDOWS
+	return;
+#endif
+
 	uint32 ModuleId = AddressToId(ModuleBase);
 	TArrayView<FModule> ModulesView(Modules.Data, Modules.Num);
 	int32 Index = Algo::LowerBound(ModulesView, ModuleId, FIdPredicate());
