Removed truncation of callstack hashes to 47 bits. The original intention was to limit the bits to match those that are addressable and later save space in trace events. 47 bits has quite a high rish of collision though.

#rnx

diff --git a/Engine/Source/Runtime/Core/Private/ProfilingDebugging/Windows/CallstackTrace.cpp b/Engine/Source/Runtime/Core/Private/ProfilingDebugging/Windows/CallstackTrace.cpp
--- a/Engine/Source/Runtime/Core/Private/ProfilingDebugging/Windows/CallstackTrace.cpp
+++ b/Engine/Source/Runtime/Core/Private/ProfilingDebugging/Windows/CallstackTrace.cpp
@@ -630,7 +630,7 @@ void* FBacktracer::GetBacktraceId(void* AddressOfReturnAddress) const
 	// Add to queue to be processed. This might block until there is room in the
 	// queue (i.e. the processing thread has caught up processing).
 	ProcessingThreadRunnable->AddWork(BacktraceEntry);
-	return (void*)(BacktraceId & ((1ull << 47) - 1));
+	return (void*)BacktraceId;
 }
 
 ////////////////////////////////////////////////////////////////////////////////
