Double math extensions

diff --git a/Engine/Source/Runtime/Core/Private/GenericPlatform/GenericPlatformMath.cpp b/Engine/Source/Runtime/Core/Private/GenericPlatform/GenericPlatformMath.cpp
--- a/Engine/Source/Runtime/Core/Private/GenericPlatform/GenericPlatformMath.cpp
+++ b/Engine/Source/Runtime/Core/Private/GenericPlatform/GenericPlatformMath.cpp
@@ -75,6 +75,49 @@ float FGenericPlatformMath::Atan2(float Y, float X)
 	return t3;
 }
 
+double FGenericPlatformMath::Atan2(double Y, double X)
+{
+	//return atan2(Y,X);
+	// Potential for NaN generation from valid inputs. See Atan2(float, float) comment above.
+
+	const double absX = FMath::Abs(X);
+	const double absY = FMath::Abs(Y);
+	const bool yAbsBigger = (absY > absX);
+	double t0 = yAbsBigger ? absY : absX; // Max(absY, absX)
+	double t1 = yAbsBigger ? absX : absY; // Min(absX, absY)
+
+	if (t0 == 0.0)
+		return 0.0;
+
+	double t3 = t1 / t0;
+	double t4 = t3 * t3;
+
+	static const double c[7] = {
+		+7.2128853633444123e-03,
+		-3.5059680836411644e-02,
+		+8.1675882859940430e-02,
+		-1.3374657325451267e-01,
+		+1.9856563505717162e-01,
+		-3.3324998579202170e-01,
+		+1.0
+	};
+
+	t0 = c[0];
+	t0 = t0 * t4 + c[1];
+	t0 = t0 * t4 + c[2];
+	t0 = t0 * t4 + c[3];
+	t0 = t0 * t4 + c[4];
+	t0 = t0 * t4 + c[5];
+	t0 = t0 * t4 + c[6];
+	t3 = t0 * t3;
+
+	t3 = yAbsBigger ? (0.5 * PI) - t3 : t3;
+	t3 = (X < 0.0) ? PI - t3 : t3;
+	t3 = (Y < 0.0) ? -t3 : t3;
+
+	return t3;
+}
+
 /*FORCENOINLINE*/ float FGenericPlatformMath::Fmod(float X, float Y)
 {
 	const float AbsY = fabsf(Y);
@@ -101,6 +144,32 @@ float FGenericPlatformMath::Atan2(float Y, float X)
 	return FMath::Clamp(Result, -AbsY, AbsY);
 }
 
+/*FORCENOINLINE*/ double FGenericPlatformMath::Fmod(double X, double Y)
+{
+	const double AbsY = fabs(Y);
+	if (AbsY <= 1.e-8)
+	{
+		FmodReportError(X, Y);
+		return 0.0;
+	}
+	const double Div = (X / Y);
+	// All doubles where abs(f) >= 2^52 (4503599627370496.0) are whole numbers so do not need truncation, and avoid overflow in TruncToDouble as they get even larger.
+	const double Quotient = fabs(Div) < DOUBLE_NON_FRACTIONAL ? TruncToDouble(Div) : Div;
+	double IntPortion = Y * Quotient;
+
+	// Rounding and imprecision could cause IntPortion to exceed X and cause the result to be outside the expected range.
+	// For example Fmod(55.8, 9.3) would result in a very small negative value!
+	if (fabs(IntPortion) > fabs(X))
+	{
+		IntPortion = X;
+	}
+
+	const double Result = X - IntPortion;
+	// Clamp to [-AbsY, AbsY] because of possible failures for very large numbers (>1e10) due to precision loss.
+	// We could instead fall back to stock fmodf() for large values, however this would diverge from the SIMD VectorMod() which has no similar fallback with reasonable performance.
+	return FMath::Clamp(Result, -AbsY, AbsY);
+}
+
 void FGenericPlatformMath::FmodReportError(float X, float Y)
 {
 	if (Y == 0)
@@ -109,6 +178,14 @@ void FGenericPlatformMath::FmodReportError(float X, float Y)
 	}
 }
 
+void FGenericPlatformMath::FmodReportError(double X, double Y)
+{
+	if (Y == 0)
+	{
+		ensureMsgf(Y != 0, TEXT("FMath::FMod(X=%f, Y=%f) : Y is zero, this is invalid and would result in NaN!"), X, Y);
+	}
+}
+
 #if WITH_DEV_AUTOMATION_TESTS
 
 // Various constants used by the test code below. These are declared
diff --git a/Engine/Source/Runtime/Core/Public/GenericPlatform/GenericPlatformMath.h b/Engine/Source/Runtime/Core/Public/GenericPlatform/GenericPlatformMath.h
--- a/Engine/Source/Runtime/Core/Public/GenericPlatform/GenericPlatformMath.h
+++ b/Engine/Source/Runtime/Core/Public/GenericPlatform/GenericPlatformMath.h
@@ -10,6 +10,10 @@
 #include "CoreTypes.h"
 #include "Containers/ContainersFwd.h"
 #include "HAL/PlatformCrt.h"
+#include "Templates/AndOrNot.h"
+#include "Templates/Decay.h"
+#include "Templates/IsFloatingPoint.h"
+#include "Templates/UnrealTypeTraits.h"
 
 /**
  * Generic implementation for most platforms
@@ -132,7 +136,7 @@ struct FGenericPlatformMath
 		VectorStoreHalf(Dst + 4, Src + 4);
 	}
 
-	/** 
+	/**
 	 * Performs a bit cast of the given float to an unsigned int of the same bit width.
 	 * @param F The float to bit cast to an unsigned integer.
 	*  @return A bitwise copy of the float in a 32-bit unsigned integer value.
@@ -349,12 +353,18 @@ struct FGenericPlatformMath
 
 	// Returns e^Value
 	static FORCEINLINE float Exp( float Value ) { return expf(Value); }
+	static FORCEINLINE double Exp(double Value) { return exp(Value); }
 	// Returns 2^Value
 	static FORCEINLINE float Exp2( float Value ) { return powf(2.f, Value); /*exp2f(Value);*/ }
+	static FORCEINLINE double Exp2(double Value) { return pow(2.0, Value); /*exp2(Value);*/ }
 	static FORCEINLINE float Loge( float Value ) {	return logf(Value); }
+	static FORCEINLINE double Loge(double Value) { return log(Value); }
 	static FORCEINLINE float LogX( float Base, float Value ) { return Loge(Value) / Loge(Base); }
+	static FORCEINLINE double LogX(double Base, double Value) { return Loge(Value) / Loge(Base); }
 	// 1.0 / Loge(2) = 1.4426950f
 	static FORCEINLINE float Log2( float Value ) { return Loge(Value) * 1.4426950f; }	
+	// 1.0 / Loge(2) = 1.442695040888963387
+	static FORCEINLINE double Log2(double Value) { return Loge(Value) * 1.442695040888963387; }
 
 	/** 
 	* Returns the floating-point remainder of X / Y
@@ -366,29 +376,36 @@ struct FGenericPlatformMath
 	* which is not consistent with the intent nor with the vectorized version.
 	*/
 	static CORE_API FORCENOINLINE float Fmod(float X, float Y);
+	static CORE_API FORCENOINLINE double Fmod(double X, double Y);
 
 	static FORCEINLINE float Sin( float Value ) { return sinf(Value); }
+	static FORCEINLINE double Sin( double Value ) { return sin(Value); }
 	static FORCEINLINE float Asin( float Value ) { return asinf( (Value<-1.f) ? -1.f : ((Value<1.f) ? Value : 1.f) ); }
+	static FORCEINLINE double Asin( double Value ) { return asin( (Value<-1.0) ? -1.0 : ((Value<1.0) ? Value : 1.0) ); }
 	static FORCEINLINE float Sinh(float Value) { return sinhf(Value); }
+	static FORCEINLINE double Sinh(double Value) { return sinh(Value); }
 	static FORCEINLINE float Cos( float Value ) { return cosf(Value); }
+	static FORCEINLINE double Cos( double Value ) { return cos(Value); }
 	static FORCEINLINE float Acos( float Value ) { return acosf( (Value<-1.f) ? -1.f : ((Value<1.f) ? Value : 1.f) ); }
+	static FORCEINLINE double Acos( double Value ) { return acos( (Value<-1.0) ? -1.0 : ((Value<1.0) ? Value : 1.0) ); }
 	static FORCEINLINE float Tan( float Value ) { return tanf(Value); }
+	static FORCEINLINE double Tan( double Value ) { return tan(Value); }
 	static FORCEINLINE float Atan( float Value ) { return atanf(Value); }
+	static FORCEINLINE double Atan( double Value ) { return atan(Value); }
 	static CORE_API float Atan2( float Y, float X );
+	static CORE_API double Atan2( double Y, double X );
 	static FORCEINLINE float Sqrt( float Value ) { return sqrtf(Value); }
+	static FORCEINLINE double Sqrt( double Value ) { return sqrt(Value); }
 	static FORCEINLINE float Pow( float A, float B ) { return powf(A,B); }
+	static FORCEINLINE double Pow( double A, double B ) { return pow(A,B); }
 
 	/** Computes a fully accurate inverse square root */
-	static FORCEINLINE float InvSqrt( float F )
-	{
-		return 1.0f / sqrtf( F );
-	}
+	static FORCEINLINE float InvSqrt( float F ) { return 1.0f / sqrtf( F ); }
+	static FORCEINLINE double InvSqrt( double F ) { return 1.0 / sqrt( F ); }
 
 	/** Computes a faster but less accurate inverse square root */
-	static FORCEINLINE float InvSqrtEst( float F )
-	{
-		return InvSqrt( F );
-	}
+	static FORCEINLINE float InvSqrtEst( float F ) { return InvSqrt( F ); }
+	static FORCEINLINE double InvSqrtEst( double F ) { return InvSqrt( F ); }
 
 	/** Return true if value is NaN (not a number). */
 	static FORCEINLINE bool IsNaN( float A ) 
@@ -865,10 +882,60 @@ struct FGenericPlatformMath
 	static void AutoTest();
 #endif
 
+
+	// Temporary support for ambiguities to simplify UE4 -> UE5 upgrades - falls back to float variant.
+	// A call is considered ambiguous if it passes no float types, any long double, or multiple mismatched float/double types.
+
+	template<typename T1, typename T2>
+	struct TIsSameOrNotFP { static const bool Value = TIsSame<T1, T2>::Value || !TIsFloatingPoint<T2>::Value; };
+	template<typename... Ts>
+	static constexpr bool TIsAmbiguous =	!TOr<TIsFloatingPoint<typename TDecay<Ts>::Type>...>::Value ||
+											TOr<TIsSame<long double, typename TDecay<Ts>::Type>...>::Value ||
+											(TOr<TIsSame<float, typename TDecay<Ts>::Type>...>::Value && !TAnd<TIsSameOrNotFP<float, typename TDecay<Ts>::Type>...>::Value);
+
+
+	template<typename T> UE_DEPRECATED(5.0, "Arguments cause function resolution ambiguity.")
+	static FORCEINLINE float Exp(T&& Value) { return Exp((float)Value); }
+	template<typename T> UE_DEPRECATED(5.0, "Arguments cause function resolution ambiguity.")
+	static FORCEINLINE float Exp2(T&& Value) { return Exp2((float)Value); }
+	template<typename T> UE_DEPRECATED(5.0, "Arguments cause function resolution ambiguity.")
+	static FORCEINLINE float Loge(T&& Value) { return Loge((float)Value); }
+	template<typename T1, typename T2, TEMPLATE_REQUIRES(TIsAmbiguous<T1, T2>)> UE_DEPRECATED(5.0, "Arguments cause function resolution ambiguity.")
+	static FORCEINLINE float LogX(T1&& Base, T2&& Value) { return LogX((float)Base, (float)Value); }
+	template<typename T> UE_DEPRECATED(5.0, "Arguments cause function resolution ambiguity.")
+	static FORCEINLINE float Log2(T&& Value) { return Log2((float)Value); }
+	template<typename T1, typename T2, TEMPLATE_REQUIRES(TIsAmbiguous<T1, T2>)> UE_DEPRECATED(5.0, "Arguments cause function resolution ambiguity.")
+	static FORCEINLINE float Fmod(T1&& X, T2&& Y) { return Fmod((float)X, (float)Y); }
+	template<typename T> UE_DEPRECATED(5.0, "Arguments cause function resolution ambiguity.")
+	static FORCEINLINE float Sin( T&& Value ) { return Sin((float)Value); }
+	template<typename T> UE_DEPRECATED(5.0, "Arguments cause function resolution ambiguity.")
+	static FORCEINLINE float Asin( T&& Value ) { return Asin((float)Value); }
+	template<typename T> UE_DEPRECATED(5.0, "Arguments cause function resolution ambiguity.")
+	static FORCEINLINE float Sinh(T&& Value) { return Sinh((float)Value); }
+	template<typename T> UE_DEPRECATED(5.0, "Arguments cause function resolution ambiguity.")
+	static FORCEINLINE float Cos( T&& Value ) { return Cos((float)Value); }
+	template<typename T> UE_DEPRECATED(5.0, "Arguments cause function resolution ambiguity.")
+	static FORCEINLINE float Acos( T&& Value ) { return Acos((float)Value); }
+	template<typename T> UE_DEPRECATED(5.0, "Arguments cause function resolution ambiguity.")
+	static FORCEINLINE float Tan( T&& Value ) { return Tan((float)Value); }
+	template<typename T> UE_DEPRECATED(5.0, "Arguments cause function resolution ambiguity.")
+	static FORCEINLINE float Atan( T&& Value ) { return Atan((float)Value); }
+	template<typename T1, typename T2, TEMPLATE_REQUIRES(TIsAmbiguous<T1, T2>)> UE_DEPRECATED(5.0, "Arguments cause function resolution ambiguity.")
+	static FORCEINLINE float Atan2( T1&& Y, T2&& X ) { return Atan2((float)Y, (float)X); }
+	template<typename T> UE_DEPRECATED(5.0, "Arguments cause function resolution ambiguity.")
+	static FORCEINLINE float Sqrt( T&& Value ) { return Sqrt((float)Value); }
+	template<typename T1, typename T2, TEMPLATE_REQUIRES(TIsAmbiguous<T1, T2>)> UE_DEPRECATED(5.0, "Arguments cause function resolution ambiguity.")
+	static FORCEINLINE float Pow( T1&& A, T2&& B ) { return Pow((float)A, (float)B); }
+	template<typename T> UE_DEPRECATED(5.0, "Arguments cause function resolution ambiguity.")
+	static FORCEINLINE float InvSqrt( T&& F ) { return InvSqrt((float)F); }
+	template<typename T> UE_DEPRECATED(5.0, "Arguments cause function resolution ambiguity.")
+	static FORCEINLINE float InvSqrtEst(T&& F) { return InvSqrt((float)F); }
+
 private:
 
 	/** Error reporting for Fmod. Not inlined to avoid compilation issues and avoid all the checks and error reporting at all callsites. */
 	static CORE_API void FmodReportError(float X, float Y);
+	static CORE_API void FmodReportError(double X, double Y);
 };
 
 /** Float specialization */
@@ -877,4 +944,8 @@ FORCEINLINE float FGenericPlatformMath::Abs( const float A )
 {
 	return fabsf( A );
 }
-
+template<>
+FORCEINLINE double FGenericPlatformMath::Abs(const double A)
+{
+	return fabs(A);
+}
\ No newline at end of file
diff --git a/Engine/Source/Runtime/Core/Public/Math/UnrealMathUtility.h b/Engine/Source/Runtime/Core/Public/Math/UnrealMathUtility.h
--- a/Engine/Source/Runtime/Core/Public/Math/UnrealMathUtility.h
+++ b/Engine/Source/Runtime/Core/Public/Math/UnrealMathUtility.h
@@ -52,6 +52,7 @@ class TRange;
 #define EULERS_NUMBER       (2.71828182845904523536f)
 #define UE_GOLDEN_RATIO		(1.6180339887498948482045868343656381f)	/* Also known as divine proportion, golden mean, or golden section - related to the Fibonacci Sequence = (1 + sqrt(5)) / 2 */
 #define FLOAT_NON_FRACTIONAL (8388608.f) /* All single-precision floating point numbers greater than or equal to this have no fractional value. */
+#define DOUBLE_NON_FRACTIONAL (4503599627370496.0) /* All double-precision floating point numbers greater than or equal to this have no fractional value. 2^52 */
 
 // Copied from float.h
 #define MAX_FLT 3.402823466e+38F
@@ -469,6 +470,21 @@ public:
 		return Loge(Value) * LogToLog2;
 	}
 
+	/**
+	 * Computes the base 2 logarithm of the specified value
+	 *
+	 * @param Value the value to perform the log on
+	 *
+	 * @return the base 2 log of the value
+	 */
+	static FORCEINLINE double Log2(double Value)
+	{
+		// Cached value for fast conversions
+		static const double LogToLog2 = 1.0 / Loge(2.0);
+		// Do the platform specific log and convert using the cached value
+		return Loge(Value) * LogToLog2;
+	}
+
 	/**
 	* Computes the sine and cosine of a scalar value.
 	*
@@ -1641,7 +1657,32 @@ public:
 		const float InterpFraction = (X - A) / (B - A);
 		return InterpFraction * InterpFraction * (3.0f - 2.0f * InterpFraction);
 	}
-	
+
+	/** 
+	 * Returns a smooth Hermite interpolation between 0 and 1 for the value X (where X ranges between A and B)
+	 * Clamped to 0 for X <= A and 1 for X >= B.
+	 *
+	 * @param A Minimum value of X
+	 * @param B Maximum value of X
+	 * @param X Parameter
+	 *
+	 * @return Smoothed value between 0 and 1
+	 */
+	static double SmoothStep(double A, double B, double X)
+	{
+		if (X < A)
+		{
+			return 0.0;
+		}
+		else if (X >= B)
+		{
+			return 1.0;
+		}
+		const double InterpFraction = (X - A) / (B - A);
+		return InterpFraction * InterpFraction * (3.0 - 2.0 * InterpFraction);
+	}
+
+
 	/**
 	 * Get a bit in memory created from bitflags (uint32 Value:1), used for EngineShowFlags,
 	 * TestBitFieldFunctions() tests the implementation
@@ -1792,4 +1833,58 @@ public:
 		}
 		return WeightedMovingAverage(CurrentSample, PreviousSample, Weight);
 	}
+
+	/**
+	 * Calculates the new value in a weighted moving average series using the previous value and the weight
+	 *
+	 * @param CurrentSample - The value to blend with the previous sample to get a new weighted value
+	 * @param PreviousSample - The last value from the series
+	 * @param Weight - The weight to blend with
+	 *
+	 * @return the next value in the series
+	 */
+	static CORE_API inline double WeightedMovingAverage(double CurrentSample, double PreviousSample, double Weight)
+	{
+		Weight = Clamp<double>(Weight, 0.0, 1.0);
+		double WAvg = (CurrentSample * Weight) + (PreviousSample * (1.0 - Weight));
+		return WAvg;
+	}
+
+	/**
+	 * Calculates the new value in a weighted moving average series using the previous value and a weight range.
+	 * The weight range is used to dynamically adjust based upon distance between the samples
+	 * This allows you to smooth a value more aggressively for small noise and let large movements be smoothed less (or vice versa)
+	 *
+	 * @param CurrentSample - The value to blend with the previous sample to get a new weighted value
+	 * @param PreviousSample - The last value from the series
+	 * @param MaxDistance - Distance to use as the blend between min weight or max weight
+	 * @param MinWeight - The weight use when the distance is small
+	 * @param MaxWeight - The weight use when the distance is large
+	 *
+	 * @return the next value in the series
+	 */
+	static CORE_API inline double DynamicWeightedMovingAverage(double CurrentSample, double PreviousSample, double MaxDistance, double MinWeight, double MaxWeight)
+	{
+		// We need the distance between samples to determine how much of each weight to use
+		const double Distance = Abs<double>(CurrentSample - PreviousSample);
+		double Weight = MinWeight;
+		if (MaxDistance > 0)
+		{
+			// Figure out the lerp value to use between the min/max weights
+			const double LerpAlpha = Clamp<double>(Distance / MaxDistance, 0.0, 1.0);
+			Weight = Lerp<double>(MinWeight, MaxWeight, LerpAlpha);
+		}
+		return WeightedMovingAverage(CurrentSample, PreviousSample, Weight);
+	}
+
+	// Temporary support for ambiguities to simplify UE4 -> UE5 upgrades - falls back to float variant.
+	// A call is considered ambiguous if it passes no float types, any long double, or multiple mismatched float/double types.
+	template<typename T> UE_DEPRECATED(5.0, "Arguments cause function resolution ambiguity.")
+	static FORCEINLINE float Log2(T&& Value) { return Log2((float)Value); }
+	template<typename T1, typename T2, typename T3, TEMPLATE_REQUIRES(TIsAmbiguous<T1, T2, T3>)> UE_DEPRECATED(5.0, "Arguments cause function resolution ambiguity.")
+	static FORCEINLINE float SmoothStep(T1&& A, T2&& B, T3&& C) { return SmoothStep((float)A, (float)B, (float)C); }
+	template<typename T1, typename T2, typename T3, TEMPLATE_REQUIRES(TIsAmbiguous<T1, T2, T3>)> UE_DEPRECATED(5.0, "Arguments cause function resolution ambiguity.")
+	static FORCEINLINE float WeightedMovingAverage(T1&& CurrentSample, T2&& PreviousSample, T3&& Weight) { return WeightedMovingAverage((float)CurrentSample, (float)PreviousSample, (float)Weight); }
+	template<typename T1, typename T2, typename T3, typename T4, typename T5, TEMPLATE_REQUIRES(TIsAmbiguous<T1, T2, T3, T4, T5>)> UE_DEPRECATED(5.0, "Arguments cause function resolution ambiguity.")
+	static FORCEINLINE float DynamicWeightedMovingAverage(T1&& CurrentSample, T2&& PreviousSample, T3&& MaxDistance, T4&& MinWeight, T5&& MaxWeight) { return DynamicWeightedMovingAverage((float)CurrentSample, (float)PreviousSample, (float)MaxDistance, (float)MinWeight, (float)MaxWeight); }
 };
diff --git a/Engine/Source/Runtime/Core/Public/Templates/UnrealTypeTraits.h b/Engine/Source/Runtime/Core/Public/Templates/UnrealTypeTraits.h
--- a/Engine/Source/Runtime/Core/Public/Templates/UnrealTypeTraits.h
+++ b/Engine/Source/Runtime/Core/Public/Templates/UnrealTypeTraits.h
@@ -13,6 +13,7 @@
 #include "Misc/AssertionMacros.h"
 #include "Templates/AndOrNot.h"
 #include "Templates/AreTypesEqual.h"
+#include "Templates/EnableIf.h"
 #include "Templates/IsArithmetic.h"
 #include "Templates/IsEnum.h"
 #include "Templates/RemoveCV.h"
@@ -22,6 +23,13 @@
 #include "Templates/IsPODType.h"
 #include "Templates/IsTriviallyCopyConstructible.h"
 
+/*-----------------------------------------------------------------------------
+	Readability macro for enable_if in template definitions. Usage:
+
+	template<typename T, TEMPLATE_REQUIRES(TIsFloatingPoint<T>)>
+	void FloatingPointOnlyPlease(T In) {}
+ -----------------------------------------------------------------------------*/
+#define TEMPLATE_REQUIRES(...) typename TEnableIf<__VA_ARGS__, int>::type = 0
 
 /*-----------------------------------------------------------------------------
  * Macros to abstract the presence of certain compiler intrinsic type traits 
