Change Android and iOS logging tags from [UE4] to [UE]
Change default logfile names from UE4 to Unreal

diff --git a/Engine/Source/Developer/Android/AndroidTargetPlatform/Private/AndroidTargetDeviceOutput.inl b/Engine/Source/Developer/Android/AndroidTargetPlatform/Private/AndroidTargetDeviceOutput.inl
--- a/Engine/Source/Developer/Android/AndroidTargetPlatform/Private/AndroidTargetDeviceOutput.inl
+++ b/Engine/Source/Developer/Android/AndroidTargetPlatform/Private/AndroidTargetDeviceOutput.inl
@@ -19,7 +19,7 @@ inline FAndroidDeviceOutputReaderRunnable::FAndroidDeviceOutputReaderRunnable(co
 
 inline bool FAndroidDeviceOutputReaderRunnable::StartLogcatProcess(void)
 {
-	FString Params = FString::Printf(TEXT(" -s %s logcat UE4:V DEBUG:V *:S -v tag"), *DeviceSerialNumber);
+	FString Params = FString::Printf(TEXT(" -s %s logcat UE:V DEBUG:V *:S -v tag"), *DeviceSerialNumber);
 	LogcatProcHandle = FPlatformProcess::CreateProc(*AdbFilename, *Params, true, false, false, NULL, 0, NULL, LogcatWritePipe);
 	return LogcatProcHandle.IsValid();
 }
diff --git a/Engine/Source/Developer/IOS/IOSTargetPlatform/Private/IOSTargetDeviceOutput.inl b/Engine/Source/Developer/IOS/IOSTargetPlatform/Private/IOSTargetDeviceOutput.inl
--- a/Engine/Source/Developer/IOS/IOSTargetPlatform/Private/IOSTargetDeviceOutput.inl
+++ b/Engine/Source/Developer/IOS/IOSTargetPlatform/Private/IOSTargetDeviceOutput.inl
@@ -81,7 +81,7 @@ inline uint32 FIOSDeviceOutputReaderRunnable::Run(void)
 			FString Text;
 			if (OutputQueue.Dequeue(Text))
 			{
-				if (Text.Contains(TEXT("[UE4]"), ESearchCase::CaseSensitive))
+				if (Text.Contains(TEXT("[UE]"), ESearchCase::CaseSensitive))
 				{
 					Output->Serialize(*Text, ELogVerbosity::Log, NAME_None);
 				}
diff --git a/Engine/Source/Runtime/Core/Private/Android/AndroidPlatformMisc.cpp b/Engine/Source/Runtime/Core/Private/Android/AndroidPlatformMisc.cpp
--- a/Engine/Source/Runtime/Core/Private/Android/AndroidPlatformMisc.cpp
+++ b/Engine/Source/Runtime/Core/Private/Android/AndroidPlatformMisc.cpp
@@ -283,7 +283,7 @@ void FAndroidMisc::LocalPrint(const TCHAR *Message)
 			}
 		}
 		*WritePtr = '\0';
-		__android_log_print(ANDROID_LOG_DEBUG, "UE4", "%ls", MessageBuffer);
+		__android_log_print(ANDROID_LOG_DEBUG, "UE", "%ls", MessageBuffer);
 	}
 #endif
 }
@@ -2061,9 +2061,9 @@ static EDeviceVulkanSupportStatus AttemptVulkanInit(void* VulkanLib)
 	VkApplicationInfo App;
 	FMemory::Memzero(App);
 	App.sType = VK_STRUCTURE_TYPE_APPLICATION_INFO;
-	App.pApplicationName = "UE4";
+	App.pApplicationName = "UE";
 	App.applicationVersion = 0;
-	App.pEngineName = "UE4";
+	App.pEngineName = "UE";
 	App.engineVersion = 0;
 	App.apiVersion = UE_VK_API_VERSION;
 
diff --git a/Engine/Source/Runtime/Core/Private/Android/AndroidPlatformStackWalk.cpp b/Engine/Source/Runtime/Core/Private/Android/AndroidPlatformStackWalk.cpp
--- a/Engine/Source/Runtime/Core/Private/Android/AndroidPlatformStackWalk.cpp
+++ b/Engine/Source/Runtime/Core/Private/Android/AndroidPlatformStackWalk.cpp
@@ -153,7 +153,7 @@ uint32 FAndroidPlatformStackWalk::CaptureStackBackTrace(uint64* BackTrace, uint3
 		static int32 OnceOnly = 0;
 		if (Context == nullptr && OnceOnly == 0)
 		{
-			__android_log_print(ANDROID_LOG_DEBUG, "UE4", "FAndroidPlatformStackWalk::CaptureStackBackTrace disabled on Android 10 with TargetSDK >= 29 due to XOM.");
+			__android_log_print(ANDROID_LOG_DEBUG, "UE", "FAndroidPlatformStackWalk::CaptureStackBackTrace disabled on Android 10 with TargetSDK >= 29 due to XOM.");
 			OnceOnly = 1;
 		}
 		return 0;
diff --git a/Engine/Source/Runtime/Core/Private/Apple/ApplePlatformMisc.cpp b/Engine/Source/Runtime/Core/Private/Apple/ApplePlatformMisc.cpp
--- a/Engine/Source/Runtime/Core/Private/Apple/ApplePlatformMisc.cpp
+++ b/Engine/Source/Runtime/Core/Private/Apple/ApplePlatformMisc.cpp
@@ -51,7 +51,7 @@ FString FApplePlatformMisc::GetEnvironmentVariable(const TCHAR* VariableName)
 void FApplePlatformMisc::LocalPrint(const TCHAR* Message)
 {
 	//NsLog will out to all iOS output consoles, instead of just the Xcode console.
-	NSLog(@"[UE4] %s", TCHAR_TO_UTF8(Message));
+	NSLog(@"[UE] %s", TCHAR_TO_UTF8(Message));
 }
 
 const TCHAR* FApplePlatformMisc::GetSystemErrorMessage(TCHAR* OutBuffer, int32 BufferCount, int32 Error)
diff --git a/Engine/Source/Runtime/Core/Private/GenericPlatform/GenericPlatformOutputDevices.cpp b/Engine/Source/Runtime/Core/Private/GenericPlatform/GenericPlatformOutputDevices.cpp
--- a/Engine/Source/Runtime/Core/Private/GenericPlatform/GenericPlatformOutputDevices.cpp
+++ b/Engine/Source/Runtime/Core/Private/GenericPlatform/GenericPlatformOutputDevices.cpp
@@ -85,7 +85,7 @@ void FGenericPlatformOutputDevices::InitDefaultOutputDeviceFile()
 		}
 		else
 		{
-			LogFilename = TEXT("UE4");
+			LogFilename = TEXT("Unreal");
 		}
 
 		LogFilename += TEXT(".log");
diff --git a/Engine/Source/Runtime/Core/Private/IOS/IOSPlatformCrashContext.cpp b/Engine/Source/Runtime/Core/Private/IOS/IOSPlatformCrashContext.cpp
--- a/Engine/Source/Runtime/Core/Private/IOS/IOSPlatformCrashContext.cpp
+++ b/Engine/Source/Runtime/Core/Private/IOS/IOSPlatformCrashContext.cpp
@@ -184,7 +184,7 @@ void FIOSCrashContext::GenerateInfoInFolder(char const* const InfoFolder, bool b
         // copy log
         FCStringAnsi::Strncpy(FilePath, CrashInfoFolder, PATH_MAX);
         FCStringAnsi::Strcat(FilePath, PATH_MAX, "/");
-        FCStringAnsi::Strcat(FilePath, PATH_MAX, (!GIOSAppInfo.AppName.IsEmpty() ? GIOSAppInfo.AppNameUTF8 : "UE4"));
+        FCStringAnsi::Strcat(FilePath, PATH_MAX, (!GIOSAppInfo.AppName.IsEmpty() ? GIOSAppInfo.AppNameUTF8 : "Unreal"));
         FCStringAnsi::Strcat(FilePath, PATH_MAX, ".log");
         
         int LogSrc = open(GIOSAppInfo.AppLogPath, O_RDONLY);
diff --git a/Engine/Source/Runtime/Core/Private/Mac/MacPlatformMisc.cpp b/Engine/Source/Runtime/Core/Private/Mac/MacPlatformMisc.cpp
--- a/Engine/Source/Runtime/Core/Private/Mac/MacPlatformMisc.cpp
+++ b/Engine/Source/Runtime/Core/Private/Mac/MacPlatformMisc.cpp
@@ -2073,7 +2073,7 @@ void FMacCrashContext::GenerateInfoInFolder(char const* const InfoFolder) const
 		// copy log
 		FCStringAnsi::Strncpy(FilePath, CrashInfoFolder, PATH_MAX);
 		FCStringAnsi::Strcat(FilePath, PATH_MAX, "/");
-		FCStringAnsi::Strcat(FilePath, PATH_MAX, (!GMacAppInfo.AppName.IsEmpty() ? GMacAppInfo.AppNameUTF8 : "UE4"));
+		FCStringAnsi::Strcat(FilePath, PATH_MAX, (!GMacAppInfo.AppName.IsEmpty() ? GMacAppInfo.AppNameUTF8 : "Unreal"));
 		FCStringAnsi::Strcat(FilePath, PATH_MAX, ".log");
 		int LogSrc = open(GMacAppInfo.AppLogPath, O_RDONLY);
 		int LogDst = open(FilePath, O_CREAT|O_WRONLY, 0766);
