CompactBinary: Added a post-increment operator to the field iterators

#rnx

diff --git a/Engine/Source/Runtime/Core/Public/Serialization/CompactBinary.h b/Engine/Source/Runtime/Core/Public/Serialization/CompactBinary.h
--- a/Engine/Source/Runtime/Core/Public/Serialization/CompactBinary.h
+++ b/Engine/Source/Runtime/Core/Public/Serialization/CompactBinary.h
@@ -292,6 +292,13 @@ public:
 		return *this;
 	}
 
+	inline TCbFieldIterator operator++(int)
+	{
+		TCbFieldIterator It(*this);
+		++*this;
+		return It;
+	}
+
 	constexpr inline FieldType& operator*() { return *this; }
 	constexpr inline FieldType* operator->() { return this; }
 
