Improved Windows Superluminal integration so custom installation paths are now supported.

diff --git a/Engine/Source/Runtime/Core/Core.Build.cs b/Engine/Source/Runtime/Core/Core.Build.cs
--- a/Engine/Source/Runtime/Core/Core.Build.cs
+++ b/Engine/Source/Runtime/Core/Core.Build.cs
@@ -197,14 +197,20 @@ public class Core : ModuleRules
 		// Superluminal instrumentation support, if one has it installed
 		if (Target.Platform.IsInGroup(UnrealPlatformGroup.Windows))
 		{
-			string SuperluminalRootDir = Path.Combine(Environment.GetFolderPath(Environment.SpecialFolder.ProgramFiles), "Superluminal/Performance/API/");
+			string SuperluminalInstallDir = Microsoft.Win32.Registry.GetValue(@"HKEY_LOCAL_MACHINE\SOFTWARE\Superluminal\Performance", "InstallDir", null) as string;
+			if (String.IsNullOrEmpty(SuperluminalInstallDir))
+			{
+				SuperluminalInstallDir = Path.Combine(Environment.GetFolderPath(Environment.SpecialFolder.ProgramFiles), "Superluminal/Performance");
+			}
+
+			string SuperluminalApiDir = Path.Combine(SuperluminalInstallDir, "API/");
 			string SubFolderName = (Target.Platform == UnrealTargetPlatform.Win32) ? "lib/x86/" : "lib/x64/";
-			string SuperluminalLibDir = Path.Combine(SuperluminalRootDir, SubFolderName);
+			string SuperluminalLibDir = Path.Combine(SuperluminalApiDir, SubFolderName);
 
-			if (File.Exists(Path.Combine(SuperluminalRootDir, "include/Superluminal/PerformanceAPI_capi.h"))
-				&& Target.Configuration != UnrealTargetConfiguration.Shipping)
+			if (Target.Configuration != UnrealTargetConfiguration.Shipping &&
+				File.Exists(Path.Combine(SuperluminalApiDir, "include/Superluminal/PerformanceAPI_capi.h")))
 			{
-				PublicSystemIncludePaths.Add(Path.Combine(SuperluminalRootDir, "include/"));
+				PublicSystemIncludePaths.Add(Path.Combine(SuperluminalApiDir, "include/"));
 				PublicAdditionalLibraries.Add(Path.Combine(SuperluminalLibDir, "PerformanceAPI_MD.lib"));
 				PublicDefinitions.Add("WITH_SUPERLUMINAL_PROFILER=1");
 			}
