TRingBuffer parity with TArray.
1) Support declaration of TRingBuffer with forward-declare classes by
1A) Removing unused bConstructElements that had declaration-time use of TIsPODType<T>
1B) Change TIsTriviallyDestructible<T> from being required at declaration-time (for the definition of an enum) to being required only at function-call-time.
2) Add GetAllocatedSize function.
#rnx

diff --git a/Engine/Source/Runtime/Core/Public/Containers/RingBuffer.h b/Engine/Source/Runtime/Core/Public/Containers/RingBuffer.h
--- a/Engine/Source/Runtime/Core/Public/Containers/RingBuffer.h
+++ b/Engine/Source/Runtime/Core/Public/Containers/RingBuffer.h
@@ -576,7 +576,7 @@ public:
 		{
 			if (!Predicate(Data[ReadIndex & IndexMask]))
 			{
-				if (bDestructElements)
+				if (NeedsDestructElements())
 				{
 					DestructItem(&Data[WriteIndex & IndexMask]);
 				}
@@ -588,7 +588,7 @@ public:
 		AfterBack = WriteIndex;
 
 		SizeType NumDeleted = static_cast<SizeType>(ReadIndex - WriteIndex);
-		if (bDestructElements)
+		if (NeedsDestructElements())
 		{
 			while (WriteIndex != ReadIndex)
 			{
@@ -640,12 +640,23 @@ public:
 		return TArrayView<T>(GetStorage() + MaskedFront, Num());
 	}
 
+	/**
+	 * Helper function to return the amount of memory allocated by this
+	 * container.
+	 * Only returns the size of allocations made directly by the container, not the elements themselves.
+	 *
+	 * @returns Number of bytes allocated by this container.
+	 */
+	SIZE_T GetAllocatedSize(void) const
+	{
+		return Max() * sizeof(ElementType);
+	}
+
 private:
-	enum : uint32
+	static constexpr bool NeedsDestructElements()
 	{
-		bConstructElements = (TIsPODType<T>::Value ? 0U : 1U),
-		bDestructElements = (TIsTriviallyDestructible<T>::Value ? 0U : 1U),
-	};
+		return !TIsTriviallyDestructible<T>::Value;
+	}
 
 	/** Set the capacity to the given value and move or copy all elements from the old storage into a new storage with the given capacity.  Assumes the capacity has already been normalized and is greater than or equal to the number of elements in the RingBuffer. */
 	void Reallocate(SizeType NewCapacity)
@@ -726,7 +737,7 @@ private:
 			return;
 		}
 		const StorageModuloType DestructCount = RangeEnd - RangeStart;
-		if (bDestructElements && DestructCount > 0)
+		if (NeedsDestructElements() && DestructCount > 0)
 		{
 			ElementType* Data = GetStorage();
 			const StorageModuloType MaskedRangeStart = RangeStart & IndexMask;
@@ -779,7 +790,7 @@ private:
 
 		ElementType* Data = GetStorage();
 		ElementType Copy(MoveTemp(Data[RangeLast & IndexMask]));
-		if (bDestructElements)
+		if (NeedsDestructElements())
 		{
 			for (StorageModuloType Index = RangeLast; Index != RangeFirst; Index += RangeDirection)
 			{
