Fixing PVS 7.7 Issues: warning V598: The 'Memcpy' function is used to copy the fields of 'XYZ' class. Virtual table pointer will be damaged by this.

diff --git a/Engine/Source/Runtime/Core/Public/Containers/ChunkedArray.h b/Engine/Source/Runtime/Core/Public/Containers/ChunkedArray.h
--- a/Engine/Source/Runtime/Core/Public/Containers/ChunkedArray.h
+++ b/Engine/Source/Runtime/Core/Public/Containers/ChunkedArray.h
@@ -236,7 +236,7 @@ public:
 			{
 				const int32 NumElementsInCurrentChunk = FMath::Min<int32>(NumElements - ChunkIndex * NumElementsPerChunk, NumElementsPerChunk);
 				check(NumElementsInCurrentChunk > 0);
-				FMemory::Memcpy(CopyDestPtr, &Chunks[ChunkIndex].Elements[0], NumElementsInCurrentChunk * sizeof(ElementType));
+				FMemory::Memcpy(CopyDestPtr, &Chunks[ChunkIndex].Elements[0], NumElementsInCurrentChunk * sizeof(ElementType)); //-V598
 				CopyDestPtr += NumElementsInCurrentChunk;
 			}
 		}
diff --git a/Engine/Source/Runtime/Core/Public/Experimental/Containers/HazardPointer.h b/Engine/Source/Runtime/Core/Public/Experimental/Containers/HazardPointer.h
--- a/Engine/Source/Runtime/Core/Public/Experimental/Containers/HazardPointer.h
+++ b/Engine/Source/Runtime/Core/Public/Experimental/Containers/HazardPointer.h
@@ -29,13 +29,13 @@ public:
 	FHazardDeleter(const FHazardDeleter& Other)
 	{
 		//deliberately stomp over the vtable pointer so that this FHazardDeleter becomes a THazardDeleter
-		FPlatformMemory::Memcpy(this, &Other, sizeof(FHazardDeleter));
+		FPlatformMemory::Memcpy(this, &Other, sizeof(FHazardDeleter)); //-V598
 	}
 
 	FHazardDeleter& operator= (const FHazardDeleter& Other)
 	{
 		//deliberately stomp over the vtable pointer so that this FHazardDeleter becomes a THazardDeleter
-		FPlatformMemory::Memcpy(this, &Other, sizeof(FHazardDeleter));
+		FPlatformMemory::Memcpy(this, &Other, sizeof(FHazardDeleter)); //-V598
 		return *this;
 	}
 
