Fixing PVS 7.7 Issues:

warning V614: Uninitialized variable 'AllocType' used. Consider checking the seventh actual argument of the 'Init' function.

warning V678: An object is used as an argument to its own method. Consider checking the first actual argument of the 'RemoveFromList' function.

diff --git a/Engine/Source/Runtime/Core/Private/HAL/LowLevelMemTracker.cpp b/Engine/Source/Runtime/Core/Private/HAL/LowLevelMemTracker.cpp
--- a/Engine/Source/Runtime/Core/Private/HAL/LowLevelMemTracker.cpp
+++ b/Engine/Source/Runtime/Core/Private/HAL/LowLevelMemTracker.cpp
@@ -1933,7 +1933,7 @@ FLLMPauseScope::FLLMPauseScope(FName TagName, bool bIsStatTag, uint64 Amount, EL
 		bEnabled = false;
 		return;
 	}
-	Init(TagName, ELLMTag::Untagged, false, bIsStatTag, Amount, TrackerToPause, AllocType);
+	Init(TagName, ELLMTag::Untagged, false, bIsStatTag, Amount, TrackerToPause, InAllocType);
 }
 
 FLLMPauseScope::FLLMPauseScope(ELLMTag TagEnum, bool bIsStatTag, uint64 Amount, ELLMTracker TrackerToPause, ELLMAllocType InAllocType)
@@ -2498,7 +2498,7 @@ namespace LLMPrivate
 			void* Result = FreePages->Allocate();
 			if (FreePages->IsFull())
 			{
-				FreePages->RemoveFromList(FreePages);
+				FreePages->RemoveFromList(FreePages); //-V678
 			}
 			return Result;
 		}
