PathViews: add functions to manipulate relative paths.

diff --git a/Engine/Source/Runtime/Core/Private/Misc/PathViews.cpp b/Engine/Source/Runtime/Core/Private/Misc/PathViews.cpp
--- a/Engine/Source/Runtime/Core/Private/Misc/PathViews.cpp
+++ b/Engine/Source/Runtime/Core/Private/Misc/PathViews.cpp
@@ -2,9 +2,11 @@
 
 #include "Misc/PathViews.h"
 
+#include "Algo/Find.h"
 #include "Algo/FindLast.h"
 #include "Containers/UnrealString.h"
 #include "Containers/StringView.h"
+#include "Misc/Char.h"
 #include "Misc/StringBuilder.h"
 #include "String/ParseTokens.h"
 #include "Templates/Function.h"
@@ -15,6 +17,119 @@ namespace UE4PathViews_Private
 	static bool IsNotSlashOrBackslash(TCHAR C) { return C != TEXT('/') && C != TEXT('\\'); }
 
 	static bool IsSlashOrBackslashOrPeriod(TCHAR C) { return C == TEXT('/') || C == TEXT('\\') || C == TEXT('.'); }
+
+	static bool HasRedundantTerminatingSeparator(FStringView A)
+	{
+		if (A.Len() <= 2)
+		{
+			if (A.Len() <= 1)
+			{
+				// "", "c", or "/", All of these are either not slash terminating or not redundant
+				return false;
+			}
+			else if (!IsSlashOrBackslash(A[1]))
+			{
+				// "/c" or "cd"
+				return false;
+			}
+			else if (IsSlashOrBackslash(A[0]))
+			{
+				// "//"
+				return false;
+			}
+			else if (A[0] == ':')
+			{
+				// ":/", which is an invalid path, and we arbitrarily decide its terminating slash is not redundant
+				return false;
+			}
+			else
+			{
+				// "c/"
+				return true;
+			}
+		}
+		else if (!IsSlashOrBackslash(A[A.Len() - 1]))
+		{
+			// "/Some/Path"
+			return false;
+		}
+		else if (A[A.Len() - 2] == ':')
+		{
+			// "Volume:/",  "Volume:/Some/Path:/"
+			// The first case is the root dir of the volume and is not redundant
+			// The second case is an invalid path (at most one colon), and we arbitrarily decide its terminating slash is not redundant
+			return false;
+		}
+		else
+		{
+			// /Some/Path/
+			return true;
+		}
+	}
+
+	static bool StringEqualsIgnoreCaseIgnoreSeparator(FStringView A, FStringView B)
+	{
+		if (A.Len() != B.Len())
+		{
+			return false;
+		}
+		const TCHAR* AData = A.GetData();
+		const TCHAR* AEnd = AData + A.Len();
+		const TCHAR* BData = B.GetData();
+		for (; AData < AEnd; ++AData, ++BData)
+		{
+			TCHAR AChar = FChar::ToUpper(*AData);
+			TCHAR BChar = FChar::ToUpper(*BData);
+			if (IsSlashOrBackslash(AChar))
+			{
+				if (!IsSlashOrBackslash(BChar))
+				{
+					return false;
+				}
+			}
+			else
+			{
+				if (AChar != BChar)
+				{
+					return false;
+				}
+			}
+		}
+		return true;
+	}
+
+	static bool StringLessIgnoreCaseIgnoreSeparator(FStringView A, FStringView B)
+	{
+		const TCHAR* AData = A.GetData();
+		const TCHAR* AEnd = AData + A.Len();
+		const TCHAR* BData = B.GetData();
+		const TCHAR* BEnd = BData + B.Len();
+		for (; AData < AEnd && BData < BEnd; ++AData, ++BData)
+		{
+			TCHAR AChar = FChar::ToUpper(*AData);
+			TCHAR BChar = FChar::ToUpper(*BData);
+			if (IsSlashOrBackslash(AChar))
+			{
+				if (!IsSlashOrBackslash(BChar))
+				{
+					return AChar < BChar;
+				}
+			}
+			else
+			{
+				if (AChar != BChar)
+				{
+					return AChar < BChar;
+				}
+			}
+		}
+		if (BData < BEnd)
+		{
+			// B is longer than A, so A is a prefix of B. A string is greater than any of its prefixes
+			return true;
+		}
+		return false;
+	}
 }
 
 FStringView FPathViews::GetCleanFilename(const FStringView& InPath)
@@ -71,6 +186,26 @@ FStringView FPathViews::GetPathLeaf(const FStringView& InPath)
 	return FStringView();
 }
 
+bool FPathViews::IsPathLeaf(FStringView InPath)
+{
+	using namespace UE4PathViews_Private;
+
+	const TCHAR* FirstSlash = Algo::FindByPredicate(InPath, IsSlashOrBackslash);
+	if (FirstSlash == nullptr)
+	{
+		return true;
+	}
+	if (Algo::FindByPredicate(InPath.RightChop((FStringView::SizeType)(FirstSlash - InPath.GetData())), IsNotSlashOrBackslash) == nullptr)
+	{
+		// The first slash is after the last non-slash
+		// This means it is either required for e.g. // or D:/ or /, or it is a redundant terminating slash D:/Foo/
+		// In all of those cases the token is still considered a (possibly unnormalized) leaf
+		return true;
+	}
+	return false;
+
+}
+
 void FPathViews::IterateComponents(FStringView InPath, TFunctionRef<void(FStringView)> ComponentVisitor)
 {
 	UE::String::ParseTokensMultiple(InPath, { TEXT('/'), TEXT('\\') }, ComponentVisitor);
@@ -126,3 +261,239 @@ FString FPathViews::ChangeExtension(const FStringView& InPath, const FStringView
 		return FString(InPath);
 	}
 }
+
+bool FPathViews::IsSeparator(TCHAR c)
+{
+	return UE4PathViews_Private::IsSlashOrBackslash(c);
+}
+
+bool FPathViews::Equals(FStringView A, FStringView B)
+{
+	using namespace UE4PathViews_Private;
+
+	while (HasRedundantTerminatingSeparator(A))
+	{
+		A.LeftChopInline(1);
+	}
+	while (HasRedundantTerminatingSeparator(B))
+	{
+		B.LeftChopInline(1);
+	}
+	return StringEqualsIgnoreCaseIgnoreSeparator(A, B);
+}
+
+bool FPathViews::Less(FStringView A, FStringView B)
+{
+	using namespace UE4PathViews_Private;
+
+	while (HasRedundantTerminatingSeparator(A))
+	{
+		A.LeftChopInline(1);
+	}
+	while (HasRedundantTerminatingSeparator(B))
+	{
+		B.LeftChopInline(1);
+	}
+	return StringLessIgnoreCaseIgnoreSeparator(A, B);
+}
+
+bool FPathViews::TryMakeChildPathRelativeTo(FStringView Child, FStringView Parent, FStringView& OutRelPath)
+{
+	using namespace UE4PathViews_Private;
+
+	while (HasRedundantTerminatingSeparator(Parent))
+	{
+		Parent.LeftChopInline(1);
+		// Note that Parent can not be the empty string, because HasRedundantTerminatingSeparator("/") is false
+	}
+	if (Parent.Len() == 0)
+	{
+		// We arbitrarily define an empty string as never the parent of any directory
+		OutRelPath.Reset();
+		return false;
+	}
+	if (Child.Len() < Parent.Len())
+	{
+		// A shorter path can not be a child path
+		OutRelPath.Reset();
+		return false;
+	}
+
+	check(Child.Len() > 0 && Parent.Len() > 0);
+
+	if (!StringEqualsIgnoreCaseIgnoreSeparator(Parent, Child.SubStr(0, Parent.Len())))
+	{
+		// Child paths match their parent paths exactly up to the parent's length
+		OutRelPath.Reset();
+		return false;
+	}
+	if (Child.Len() == Parent.Len())
+	{
+		// The child is equal to the parent, which we define as being a child path
+		OutRelPath.Reset();
+		return true;
+	}
+
+	check(Child.Len() >= Parent.Len() + 1);
+
+	if (IsSlashOrBackslash(Parent[Parent.Len() - 1]))
+	{
+		// Parent is the special case "//", "Volume:/", or "/" (these are the only cases that end in / after HasRedundantTerminatingSeparator)
+		if (Parent.Len() == 1 && IsSlashOrBackslash(Child[1]))
+		{
+			// Parent is "/", Child is "//SomethingOrNothing". // is not a child of /, so this is not a child
+			OutRelPath.Reset();
+			return false;
+		}
+		// Since child starts with parent, including the terminating /, child is a child path of parent
+		OutRelPath = Child.RightChop(Parent.Len()); // Chop all of the parent including the terminating slash
+	}
+	else
+	{
+		if (!IsSlashOrBackslash(Child[Parent.Len()]))
+		{
+			// Child is in a different directory that happens to have the parent's final directory component as a prefix
+			// e.g. "Volume:/Some/Path", "Volume:/Some/PathOther"
+			OutRelPath.Reset();
+			return false;
+		}
+		// Parent   = Volume:/Some/Path
+		// Child    = Volume:/Some/Path/Child1/Child2/etc
+		// or Child = Volume:/Some/Path/
+		// or Child = Volume:/Some/Path/////
+		OutRelPath = Child.RightChop(Parent.Len() + 1); // Chop the Parent and the slash in the child that follows the parent
+	}
+	// If there were invalid duplicate slashes after the parent path, remove them since starting with a slash
+	// is not a valid relpath
+	while (OutRelPath.Len() > 0 && IsSlashOrBackslash(OutRelPath[0]))
+	{
+		OutRelPath.RightChopInline(1);
+	}
+	return true;
+}
+
+bool FPathViews::IsParentPathOf(FStringView Parent, FStringView Child)
+{
+	FStringView RelPath;
+	return TryMakeChildPathRelativeTo(Child, Parent, RelPath);
+}
+
+bool FPathViews::IsRelativePath(FStringView InPath)
+{
+	using namespace UE4PathViews_Private;
+
+	if (const TCHAR* EndPos = Algo::FindByPredicate(InPath, IsSlashOrBackslash))
+	{
+		FStringView::SizeType FirstLen = static_cast<FStringView::SizeType>(EndPos - InPath.GetData());
+		if (FirstLen == 0)
+		{
+			// Path starts with /; it may be either /Foo or //Foo
+			return false;
+		}
+		else if (InPath[FirstLen - 1] == ':')
+		{
+			// InPath == Volume:/SomethingOrNothing
+			return false;
+		}
+		else
+		{
+			// InPath == RelativeComponent/SomethingOrNothing
+			return true;
+		}
+	}
+	else
+	{
+		// InPath == SomethingOrNothing
+		return true;
+	}
+}
+
+void FPathViews::SplitFirstComponent(FStringView InPath, FStringView& OutFirstComponent, FStringView& OutRemainder)
+{
+	using namespace UE4PathViews_Private;
+
+	if (const TCHAR* EndPos = Algo::FindByPredicate(InPath, IsSlashOrBackslash))
+	{
+		FStringView::SizeType FirstLen = static_cast<FStringView::SizeType>(EndPos - InPath.GetData());
+		if (FirstLen == 0)
+		{
+			// Path starts with /; it may be either /Foo or //Foo
+			if (InPath.Len() == 1)
+			{
+				// InPath == /
+				// FirstComponent = /
+				// Remainder = <Empty>
+				OutFirstComponent = InPath;
+				OutRemainder.Reset();
+			}
+			else if (IsSlashOrBackslash(InPath[1]))
+			{
+				// InPath == //SomethingOrNothing
+				// FirstComponent = //
+				// Remainder = SomethingOrNothing
+				OutFirstComponent = InPath.Left(2);
+				OutRemainder = InPath.RightChop(2);
+			}
+			else
+			{
+				// InPath == /SomethingOrNothing
+				// FirstComponent = /
+				// Remainder = SomethingOrNothing
+				OutFirstComponent = InPath.Left(1);
+				OutRemainder = InPath.RightChop(1);
+			}
+		}
+		else if (InPath[FirstLen - 1] == ':')
+		{
+			// InPath == Volume:/SomethingOrNothing
+			// FirstComponent = Volume:/
+			// Remainder = SomethingOrNothing
+			OutFirstComponent = InPath.Left(FirstLen + 1);
+			OutRemainder = InPath.RightChop(FirstLen + 1);
+		}
+		else
+		{
+			// InPath == RelativeComponent/SomethingOrNothing
+			// FirstComponent = RelativeComponent
+			// Remainder = SomethingOrNothing
+			OutFirstComponent = InPath.Left(FirstLen);
+			OutRemainder = InPath.RightChop(FirstLen + 1);
+		}
+	}
+	else
+	{
+		// InPath == SomethingOrNothing
+		// FirstComponent = SomethingOrNothing
+		// Remainder = <Empty>
+		OutFirstComponent = InPath;
+		OutRemainder.Reset();
+	}
+	// If there were invalid duplicate slashes after the first component, remove them since starting with a slash
+	// is not a valid relpath
+	while (OutRemainder.Len() > 0 && IsSlashOrBackslash(OutRemainder[0]))
+	{
+		OutRemainder.RightChopInline(1);
+	}
+}
+
+void FPathViews::AppendPath(FStringBuilderBase& InOutPath, FStringView AppendPath)
+{
+	using namespace UE4PathViews_Private;
+
+	if (AppendPath.Len() == 0)
+	{
+	}
+	else if (IsRelativePath(AppendPath))
+	{
+		if (InOutPath.Len() > 0 && !IsSlashOrBackslash(InOutPath.LastChar()))
+		{
+			InOutPath << TEXT('/');
+		}
+		InOutPath << AppendPath;
+	}
+	else
+	{
+		InOutPath.Reset();
+		InOutPath << AppendPath;
+	}
+}
\ No newline at end of file
diff --git a/Engine/Source/Runtime/Core/Private/Tests/Misc/PathViewsTest.cpp b/Engine/Source/Runtime/Core/Private/Tests/Misc/PathViewsTest.cpp
--- a/Engine/Source/Runtime/Core/Private/Tests/Misc/PathViewsTest.cpp
+++ b/Engine/Source/Runtime/Core/Private/Tests/Misc/PathViewsTest.cpp
@@ -4,6 +4,7 @@
 
 #include "Containers/StringView.h"
 #include "Misc/AutomationTest.h"
+#include "Misc/StringBuilder.h"
 
 #if WITH_DEV_AUTOMATION_TESTS
 
@@ -21,6 +22,11 @@ protected:
 				InPath.Len(), InPath.GetData(), Actual.Len(), Actual.GetData(), InExpected));
 		}
 	}
+
+	static const TCHAR* BoolToString(bool bValue)
+	{
+		return bValue ? TEXT("true") : TEXT("false");
+	}
 };
 
 IMPLEMENT_CUSTOM_SIMPLE_AUTOMATION_TEST(FPathViewsGetCleanFilenameTest, FPathViewsTest, "System.Core.Misc.PathViews.GetCleanFilename", EAutomationTestFlags::ApplicationContextMask | EAutomationTestFlags::SmokeFilter)
@@ -190,6 +196,46 @@ bool FPathViewsGetPathLeafTest::RunTest(const FString& InParameters)
 	return true;
 }
 
+IMPLEMENT_CUSTOM_SIMPLE_AUTOMATION_TEST(FPathViewsIsPathLeafTest, FPathViewsTest, "System.Core.Misc.PathViews.IsPathLeaf", EAutomationTestFlags::ApplicationContextMask | EAutomationTestFlags::SmokeFilter)
+bool FPathViewsIsPathLeafTest::RunTest(const FString& InParameters)
+{
+	auto RunIsPathLeafTest = [this](const TCHAR* InPath, bool bExpected)
+	{
+		bool bResult = FPathViews::IsPathLeaf(InPath);
+		if (bResult != bExpected)
+		{
+			AddError(FString::Printf(TEXT("GetPathLeaf(\"%s\") == %s, expected %s."), InPath, BoolToString(bResult), BoolToString(bExpected)));
+		}
+	};
+
+	RunIsPathLeafTest(TEXT(""), true);
+	RunIsPathLeafTest(TEXT(".txt"), true);
+	RunIsPathLeafTest(TEXT(".tar.gz"), true);
+	RunIsPathLeafTest(TEXT(".tar.gz/"), true);
+	RunIsPathLeafTest(TEXT(".tar.gz\\"), true);
+	RunIsPathLeafTest(TEXT("File"), true);
+	RunIsPathLeafTest(TEXT("File.txt"), true);
+	RunIsPathLeafTest(TEXT("File.tar.gz"), true);
+	RunIsPathLeafTest(TEXT("File.tar.gz/"), true);
+	RunIsPathLeafTest(TEXT("File.tar.gz\\"), true);
+	RunIsPathLeafTest(TEXT("//"), true);
+	RunIsPathLeafTest(TEXT("\\\\"), true);
+	RunIsPathLeafTest(TEXT("/"), true);
+	RunIsPathLeafTest(TEXT("\\"), true);
+	RunIsPathLeafTest(TEXT("C:/"), true);
+	RunIsPathLeafTest(TEXT("C:\\"), true);
+	RunIsPathLeafTest(TEXT("//Folder"), false);
+	RunIsPathLeafTest(TEXT("\\\\Folder"), false);
+	RunIsPathLeafTest(TEXT("/Folder"), false);
+	RunIsPathLeafTest(TEXT("\\Folder"), false);
+	RunIsPathLeafTest(TEXT("C:/Folder"), false);
+	RunIsPathLeafTest(TEXT("C:\\Folder"), false);
+	RunIsPathLeafTest(TEXT("C:/Folder"), false);
+	RunIsPathLeafTest(TEXT("C:\\Folder"), false);
+
+	return true;
+}
+
 IMPLEMENT_CUSTOM_SIMPLE_AUTOMATION_TEST(FPathViewsSplitTest, FPathViewsTest, "System.Core.Misc.PathViews.Split", EAutomationTestFlags::ApplicationContextMask | EAutomationTestFlags::SmokeFilter)
 bool FPathViewsSplitTest::RunTest(const FString& InParameters)
 {
@@ -262,4 +308,346 @@ bool FPathViewsChangeExtensionTest::RunTest(const FString& InParameters)
 	return true;
 }
 
+IMPLEMENT_CUSTOM_SIMPLE_AUTOMATION_TEST(FPathViewsEqualsAndLessTest, FPathViewsTest, "System.Core.Misc.PathViews.EqualsAndLess", EAutomationTestFlags::ApplicationContextMask | EAutomationTestFlags::SmokeFilter)
+bool FPathViewsEqualsAndLessTest::RunTest(const FString& InParameters)
+{
+	auto RunEqualsLessTest = [this](const TCHAR* A, const TCHAR* B, int Expected)
+	{
+		bool bEqual = FPathViews::Equals(A, B);
+		bool bALessThanB = FPathViews::Less(A, B);
+		bool bBLessThanA = FPathViews::Less(B, A);
+		if (bEqual != (Expected == 0))
+		{
+			AddError(FString::Printf(TEXT("Equals(%s,%s) == %s, expected %s"), A, B, BoolToString(bEqual), BoolToString(Expected == 0)));
+		}
+		if (bALessThanB != Expected < 0)
+		{
+			AddError(FString::Printf(TEXT("Less(%s,%s) == %s, expected %s"), A, B, BoolToString(bALessThanB), BoolToString(Expected < 0)));
+		}
+		if (bBLessThanA != Expected > 0)
+		{
+			AddError(FString::Printf(TEXT("Less(%s,%s) == %s, expected %s"), B, A, BoolToString(bBLessThanA), BoolToString(Expected > 0)));
+		}
+	};
+
+	RunEqualsLessTest(TEXT("A"), TEXT("B"), -1);
+	RunEqualsLessTest(TEXT("A"), TEXT("b"), -1);
+	RunEqualsLessTest(TEXT("a"), TEXT("B"), -1);
+	RunEqualsLessTest(TEXT("A"), TEXT("A"), 0);
+	RunEqualsLessTest(TEXT("A"), TEXT("a/"), 0);
+	RunEqualsLessTest(TEXT("A"), TEXT("a\\"), 0);
+	RunEqualsLessTest(TEXT("A"), TEXT("abby"), -1);
+	RunEqualsLessTest(TEXT("a"), TEXT("Abby"), -1);
+	RunEqualsLessTest(TEXT("\\A/B"), TEXT("/A\\B/"), 0);
+	//RunEqualsLessTest(TEXT("../../../Engine"), TEXT("C:/Engine"), 0); // Detecting relpath == abspath is not yet implemented
+	//RunEqualsLessTest(TEXT("C:/A/B"), TEXT("C:/A/../A/./B"), 0); // Collapsing .. and . is not yet implemented
+	RunEqualsLessTest(TEXT("/"), TEXT("/"), 0);
+	RunEqualsLessTest(TEXT("/"), TEXT("//"), -1);
+	RunEqualsLessTest(TEXT("/"), TEXT("C:/"), -1); // int('/') == 47 < int('C') == 67
+	RunEqualsLessTest(TEXT("/"), TEXT("A"), -1);
+	RunEqualsLessTest(TEXT("//"), TEXT("//"), 0);
+	RunEqualsLessTest(TEXT("//"), TEXT("C:/"), -1);
+	RunEqualsLessTest(TEXT("//"), TEXT("A"), -1);
+	RunEqualsLessTest(TEXT("C:/"), TEXT("C:/"), 0);
+	RunEqualsLessTest(TEXT("C:/"), TEXT("C"), 1);
+
+	return true;
+}
+
+IMPLEMENT_CUSTOM_SIMPLE_AUTOMATION_TEST(FPathViewsTryMakeChildPathRelativeToTest, FPathViewsTest, "System.Core.Misc.PathViews.TryMakeChildPathRelativeTo", EAutomationTestFlags::ApplicationContextMask | EAutomationTestFlags::SmokeFilter)
+bool FPathViewsTryMakeChildPathRelativeToTest::RunTest(const FString& InParameters)
+{
+	auto RunRelChildTest = [this](const TCHAR* Child, const TCHAR* Parent, bool bExpectedIsChild, const TCHAR* ExpectedRelPath)
+	{
+		FStringView ActualRelPath;
+		bool bActualIsChild = FPathViews::TryMakeChildPathRelativeTo(Child, Parent, ActualRelPath);
+		bool bActualIsParent = FPathViews::IsParentPathOf(Parent, Child);
+		if (bExpectedIsChild != bActualIsChild || !ActualRelPath.Equals(ExpectedRelPath, ESearchCase::IgnoreCase))
+		{
+			AddError(FString::Printf(TEXT("TryMakeChildPathRelativeTo(\"%s\", \"%s\") returned (%s, \"%.*s\"), expected (%s, \"%s\")."), Child, Parent,
+				BoolToString(bActualIsChild), ActualRelPath.Len(), ActualRelPath.GetData(), BoolToString(bExpectedIsChild), ExpectedRelPath));
+		}
+		if (bActualIsParent != bExpectedIsChild)
+		{
+			AddError(FString::Printf(TEXT("IsParentPathOf(\"%s\", \"%s\") returned %s, expected %s."), Child, Parent,
+				BoolToString(bActualIsParent), BoolToString(bExpectedIsChild)));
+		}
+	};
+
+	RunRelChildTest(TEXT("A/B"), TEXT("A"), true, TEXT("B"));
+	RunRelChildTest(TEXT("A\\B"), TEXT("A"), true, TEXT("B"));
+	RunRelChildTest(TEXT("A/B"), TEXT("A/"), true, TEXT("B"));
+	RunRelChildTest(TEXT("A\\B"), TEXT("A/"), true, TEXT("B"));
+	RunRelChildTest(TEXT("A/B"), TEXT("A\\"), true, TEXT("B"));
+	RunRelChildTest(TEXT("A\\B"), TEXT("A\\"), true, TEXT("B"));
+	RunRelChildTest(TEXT("A"), TEXT("A"), true, TEXT(""));
+	RunRelChildTest(TEXT("A/"), TEXT("A"), true, TEXT(""));
+	RunRelChildTest(TEXT("A\\"), TEXT("A"), true, TEXT(""));
+	RunRelChildTest(TEXT("A"), TEXT("A/"), true, TEXT(""));
+	RunRelChildTest(TEXT("A"), TEXT("A\\"), true, TEXT(""));
+	RunRelChildTest(TEXT("../A"), TEXT(".."), true, TEXT("A"));
+	RunRelChildTest(TEXT("/A/B"), TEXT("/A"), true, TEXT("B"));
+	RunRelChildTest(TEXT("../A/B"), TEXT("../A"), true, TEXT("B"));
+	RunRelChildTest(TEXT("../"), TEXT(".."), true, TEXT(""));
+	RunRelChildTest(TEXT("C:/"), TEXT("C:/"), true, TEXT(""));
+	RunRelChildTest(TEXT("C:/A"), TEXT("C:/"), true, TEXT("A"));
+	RunRelChildTest(TEXT("//A"), TEXT("//A"), true, TEXT(""));
+	RunRelChildTest(TEXT("//A/"), TEXT("//A"), true, TEXT(""));
+	RunRelChildTest(TEXT("//A"), TEXT("//A/"), true, TEXT(""));
+	RunRelChildTest(TEXT("//A/"), TEXT("//A/"), true, TEXT(""));
+	RunRelChildTest(TEXT("//A/B"), TEXT("//A"), true, TEXT("B"));
+	RunRelChildTest(TEXT("//A/B/"), TEXT("//A"), true, TEXT("B/"));
+	RunRelChildTest(TEXT("//A/B"), TEXT("//A/"), true, TEXT("B"));
+	RunRelChildTest(TEXT("//A/B/"), TEXT("//A/"), true, TEXT("B/"));
+
+	RunRelChildTest(TEXT("//A/BFoo"), TEXT("//A/B"), false, TEXT(""));
+	RunRelChildTest(TEXT("//A/C"), TEXT("//A/B"), false, TEXT(""));
+	RunRelChildTest(TEXT("//A/C"), TEXT("C:/A"), false, TEXT(""));
+	RunRelChildTest(TEXT("//A/C"), TEXT("/A"), false, TEXT(""));
+	RunRelChildTest(TEXT("C:/A/BFoo"), TEXT("C:/A/B"), false, TEXT(""));
+	RunRelChildTest(TEXT("C:/A/C"), TEXT("C:/A/B"), false, TEXT(""));
+	RunRelChildTest(TEXT("C:/A/C"), TEXT("//A"), false, TEXT(""));
+	RunRelChildTest(TEXT("C:/A/C"), TEXT("/A"), false, TEXT(""));
+	RunRelChildTest(TEXT("/A/BFoo"), TEXT("/A/B"), false, TEXT(""));
+	RunRelChildTest(TEXT("/A/C"), TEXT("/A/B"), false, TEXT(""));
+	RunRelChildTest(TEXT("/A/C"), TEXT("C:/A"), false, TEXT(""));
+	RunRelChildTest(TEXT("/A/C"), TEXT("//A"), false, TEXT(""));
+
+	RunRelChildTest(TEXT("/"), TEXT("/"), true, TEXT(""));
+	RunRelChildTest(TEXT("/"), TEXT("//"), false, TEXT(""));
+	RunRelChildTest(TEXT("/"), TEXT("C"), false, TEXT(""));
+	RunRelChildTest(TEXT("/"), TEXT("C:/"), false, TEXT(""));
+	RunRelChildTest(TEXT("//"), TEXT("/"), false, TEXT(""));
+	RunRelChildTest(TEXT("//"), TEXT("//"), true, TEXT(""));
+	RunRelChildTest(TEXT("//"), TEXT("C"), false, TEXT(""));
+	RunRelChildTest(TEXT("//"), TEXT("C:/"), false, TEXT(""));
+	RunRelChildTest(TEXT("C"), TEXT("/"), false, TEXT(""));
+	RunRelChildTest(TEXT("C"), TEXT("//"), false, TEXT(""));
+	RunRelChildTest(TEXT("C"), TEXT("C"), true, TEXT(""));
+	RunRelChildTest(TEXT("C"), TEXT("C:/"), false, TEXT(""));
+	RunRelChildTest(TEXT("C:/"), TEXT("/"), false, TEXT(""));
+	RunRelChildTest(TEXT("C:/"), TEXT("//"), false, TEXT(""));
+	RunRelChildTest(TEXT("C:/"), TEXT("C"), false, TEXT(""));
+	RunRelChildTest(TEXT("C:/"), TEXT("C:/"), true, TEXT(""));
+
+	RunRelChildTest(TEXT("C:/A/B/C"), TEXT("C:/A"), true, TEXT("B/C"));
+	RunRelChildTest(TEXT("C:/A"), TEXT("C:/A/B/C"), false, TEXT(""));
+
+	// Correctly handle paths with invalid duplicate slashes
+	RunRelChildTest(TEXT("///A"), TEXT("///"), true, TEXT("A"));
+	RunRelChildTest(TEXT("////////////A"), TEXT("///"), true, TEXT("A"));
+	RunRelChildTest(TEXT("///A"), TEXT("//"), true, TEXT("A"));
+	RunRelChildTest(TEXT("////////////A"), TEXT("/"), false, TEXT(""));
+	RunRelChildTest(TEXT("C://A"), TEXT("C://"), true, TEXT("A"));
+	RunRelChildTest(TEXT("C://////////A"), TEXT("C://"), true, TEXT("A"));
+	RunRelChildTest(TEXT("C://A"), TEXT("C:/"), true, TEXT("A"));
+	RunRelChildTest(TEXT("C://////////A"), TEXT("C:/"), true, TEXT("A"));
+	RunRelChildTest(TEXT("A//B"), TEXT("A//"), true, TEXT("B"));
+	RunRelChildTest(TEXT("A///////////B"), TEXT("A//"), true, TEXT("B"));
+	RunRelChildTest(TEXT("A//B"), TEXT("A/"), true, TEXT("B"));
+	RunRelChildTest(TEXT("A///////////B"), TEXT("A/"), true, TEXT("B"));
+	RunRelChildTest(TEXT("A//B"), TEXT("A"), true, TEXT("B"));
+	RunRelChildTest(TEXT("A///////////B"), TEXT("A"), true, TEXT("B"));
+
+	return true;
+}
+
+IMPLEMENT_CUSTOM_SIMPLE_AUTOMATION_TEST(FPathViewsIsRelativePathTest, FPathViewsTest, "System.Core.Misc.PathViews.IsRelativePath", EAutomationTestFlags::ApplicationContextMask | EAutomationTestFlags::SmokeFilter)
+bool FPathViewsIsRelativePathTest::RunTest(const FString& InParameters)
+{
+	auto RunRelTest = [this](const TCHAR* A, bool bExpected)
+	{
+		bool bActual = FPathViews::IsRelativePath(A);
+		if (bActual != bExpected)
+		{
+			AddError(FString::Printf(TEXT("IsRelativePath(\"%s\") == %s, expected %s."), A, BoolToString(bActual), BoolToString(bExpected)));
+		}
+	};
+
+	RunRelTest(TEXT("A"), true);
+	RunRelTest(TEXT("A/"), true);
+	RunRelTest(TEXT("A/B"), true);
+	RunRelTest(TEXT("A\\B"), true);
+	RunRelTest(TEXT("/A"), false);
+	RunRelTest(TEXT("\\A"), false);
+	RunRelTest(TEXT("/A/B"), false);
+	RunRelTest(TEXT("//A"), false);
+	RunRelTest(TEXT("\\\\A"), false);
+	RunRelTest(TEXT("//A/B"), false);
+	RunRelTest(TEXT("C:/A"), false);
+	RunRelTest(TEXT("C:\\A"), false);
+	RunRelTest(TEXT("C:/A/B"), false);
+
+	return true;
+}
+
+IMPLEMENT_CUSTOM_SIMPLE_AUTOMATION_TEST(FPathViewsSplitFirstComponentTest, FPathViewsTest, "System.Core.Misc.PathViews.SplitFirstComponent", EAutomationTestFlags::ApplicationContextMask | EAutomationTestFlags::SmokeFilter)
+bool FPathViewsSplitFirstComponentTest::RunTest(const FString& InParameters)
+{
+	auto RunSplitFirstTest = [this](const TCHAR* FullPath, const TCHAR* ExpectedFirst, const TCHAR* ExpectedRemaining)
+	{
+		FStringView ActualFirst;
+		FStringView ActualRemaining;
+		FPathViews::SplitFirstComponent(FullPath, ActualFirst, ActualRemaining);
+		if (!ActualFirst.Equals(ExpectedFirst, ESearchCase::IgnoreCase) || !ActualRemaining.Equals(ExpectedRemaining, ESearchCase::IgnoreCase))
+		{
+			AddError(FString::Printf(TEXT("SplitFirstComponent(\"%s\") == (\"%.*s\", \"%.*s\"), expected (\"%s\", \"%s\")."),
+				FullPath, ActualFirst.Len(), ActualFirst.GetData(), ActualRemaining.Len(), ActualRemaining.GetData(),
+				ExpectedFirst, ExpectedRemaining));
+		}
+	};
+
+	RunSplitFirstTest(TEXT(""), TEXT(""), TEXT(""));
+
+	RunSplitFirstTest(TEXT("A"), TEXT("A"), TEXT(""));
+	RunSplitFirstTest(TEXT("A/"), TEXT("A"), TEXT(""));
+	RunSplitFirstTest(TEXT("A\\"), TEXT("A"), TEXT(""));
+	RunSplitFirstTest(TEXT("A/B"), TEXT("A"), TEXT("B"));
+	RunSplitFirstTest(TEXT("A\\B"), TEXT("A"), TEXT("B"));
+	RunSplitFirstTest(TEXT("A/B/"), TEXT("A"), TEXT("B/"));
+	RunSplitFirstTest(TEXT("A\\B\\"), TEXT("A"), TEXT("B\\"));
+	RunSplitFirstTest(TEXT("A/B/C"), TEXT("A"), TEXT("B/C"));
+	RunSplitFirstTest(TEXT("A\\B\\C"), TEXT("A"), TEXT("B\\C"));
+
+	RunSplitFirstTest(TEXT("/A"), TEXT("/"), TEXT("A"));
+	RunSplitFirstTest(TEXT("\\A"), TEXT("\\"), TEXT("A"));
+	RunSplitFirstTest(TEXT("/A/"), TEXT("/"), TEXT("A/"));
+	RunSplitFirstTest(TEXT("\\A\\"), TEXT("\\"), TEXT("A\\"));
+	RunSplitFirstTest(TEXT("/A/B"), TEXT("/"), TEXT("A/B"));
+	RunSplitFirstTest(TEXT("\\A\\B"), TEXT("\\"), TEXT("A\\B"));
+
+	RunSplitFirstTest(TEXT("//A"), TEXT("//"), TEXT("A"));
+	RunSplitFirstTest(TEXT("\\\\A"), TEXT("\\\\"), TEXT("A"));
+	RunSplitFirstTest(TEXT("//A/"), TEXT("//"), TEXT("A/"));
+	RunSplitFirstTest(TEXT("\\\\A\\"), TEXT("\\\\"), TEXT("A\\"));
+	RunSplitFirstTest(TEXT("//A/B"), TEXT("//"), TEXT("A/B"));
+	RunSplitFirstTest(TEXT("\\\\A\\B"), TEXT("\\\\"), TEXT("A\\B"));
+
+	RunSplitFirstTest(TEXT("C:/A"), TEXT("C:/"), TEXT("A"));
+	RunSplitFirstTest(TEXT("C:\\A"), TEXT("C:\\"), TEXT("A"));
+	RunSplitFirstTest(TEXT("C:/A/"), TEXT("C:/"), TEXT("A/"));
+	RunSplitFirstTest(TEXT("C:\\A\\"), TEXT("C:\\"), TEXT("A\\"));
+	RunSplitFirstTest(TEXT("C:/A/B"), TEXT("C:/"), TEXT("A/B"));
+	RunSplitFirstTest(TEXT("C:\\A\\B"), TEXT("C:\\"), TEXT("A\\B"));
+
+	// Correctly handle paths with invalid duplicate slashes
+	RunSplitFirstTest(TEXT("///A"), TEXT("//"), TEXT("A"));
+	RunSplitFirstTest(TEXT("////////////A"), TEXT("//"), TEXT("A"));
+	RunSplitFirstTest(TEXT("///A"), TEXT("//"), TEXT("A"));
+	RunSplitFirstTest(TEXT("C://A"), TEXT("C:/"), TEXT("A"));
+	RunSplitFirstTest(TEXT("C://////////A"), TEXT("C:/"), TEXT("A"));
+	RunSplitFirstTest(TEXT("A//B"), TEXT("A"), TEXT("B"));
+	RunSplitFirstTest(TEXT("A///////////B"), TEXT("A"), TEXT("B"));
+
+	return true;
+}
+
+IMPLEMENT_CUSTOM_SIMPLE_AUTOMATION_TEST(FPathViewsAppendPathTest, FPathViewsTest, "System.Core.Misc.PathViews.AppendPath", EAutomationTestFlags::ApplicationContextMask | EAutomationTestFlags::SmokeFilter)
+bool FPathViewsAppendPathTest::RunTest(const FString& InParameters)
+{
+	auto RunAppendTest = [this](const TCHAR* Base, const TCHAR* Append, const TCHAR* ExpectedNewBase)
+	{
+		TStringBuilder<128> BaseBuilder;
+		BaseBuilder << Base;
+		FPathViews::AppendPath(BaseBuilder, Append);
+		if (!FStringView(BaseBuilder).Equals(ExpectedNewBase, ESearchCase::IgnoreCase))
+		{
+			AddError(FString::Printf(TEXT("AppendPath(\"%s\", \"%s\") == \"%.*s\", expected \"%s\"."),
+				Base, Append, BaseBuilder.Len(), BaseBuilder.GetData(), ExpectedNewBase));
+		}
+	};
+
+	RunAppendTest(TEXT(""), TEXT("A"), TEXT("A"));
+	RunAppendTest(TEXT(""), TEXT("A/B"), TEXT("A/B"));
+	RunAppendTest(TEXT(""), TEXT("A\\B"), TEXT("A\\B"));
+
+	RunAppendTest(TEXT("Root"), TEXT("A"), TEXT("Root/A"));
+	RunAppendTest(TEXT("Root"), TEXT("A/B"), TEXT("Root/A/B"));
+	RunAppendTest(TEXT("Root"), TEXT("A\\B"), TEXT("Root/A\\B"));
+
+	RunAppendTest(TEXT("/"), TEXT("A"), TEXT("/A"));
+	RunAppendTest(TEXT("/"), TEXT("A/B"), TEXT("/A/B"));
+	RunAppendTest(TEXT("/"), TEXT("A\\B"), TEXT("/A\\B"));
+	RunAppendTest(TEXT("\\"), TEXT("A"), TEXT("\\A"));
+	RunAppendTest(TEXT("\\"), TEXT("A/B"), TEXT("\\A/B"));
+	RunAppendTest(TEXT("\\"), TEXT("A\\B"), TEXT("\\A\\B"));
+
+	RunAppendTest(TEXT("/Root"), TEXT("A"), TEXT("/Root/A"));
+	RunAppendTest(TEXT("/Root"), TEXT("A/B"), TEXT("/Root/A/B"));
+	RunAppendTest(TEXT("/Root"), TEXT("A\\B"), TEXT("/Root/A\\B"));
+	RunAppendTest(TEXT("\\Root"), TEXT("A"), TEXT("\\Root/A"));
+	RunAppendTest(TEXT("\\Root"), TEXT("A/B"), TEXT("\\Root/A/B"));
+	RunAppendTest(TEXT("\\Root"), TEXT("A\\B"), TEXT("\\Root/A\\B"));
+	RunAppendTest(TEXT("/Root/"), TEXT("A"), TEXT("/Root/A"));
+	RunAppendTest(TEXT("/Root/"), TEXT("A/B"), TEXT("/Root/A/B"));
+	RunAppendTest(TEXT("/Root/"), TEXT("A\\B"), TEXT("/Root/A\\B"));
+	RunAppendTest(TEXT("\\Root\\"), TEXT("A"), TEXT("\\Root\\A"));
+	RunAppendTest(TEXT("\\Root\\"), TEXT("A/B"), TEXT("\\Root\\A/B"));
+	RunAppendTest(TEXT("\\Root\\"), TEXT("A\\B"), TEXT("\\Root\\A\\B"));
+
+	RunAppendTest(TEXT("//"), TEXT("A"), TEXT("//A"));
+	RunAppendTest(TEXT("//"), TEXT("A/B"), TEXT("//A/B"));
+	RunAppendTest(TEXT("//"), TEXT("A\\B"), TEXT("//A\\B"));
+	RunAppendTest(TEXT("\\\\"), TEXT("A"), TEXT("\\\\A"));
+	RunAppendTest(TEXT("\\\\"), TEXT("A/B"), TEXT("\\\\A/B"));
+	RunAppendTest(TEXT("\\\\"), TEXT("A\\B"), TEXT("\\\\A\\B"));
+
+	RunAppendTest(TEXT("//Root"), TEXT("A"), TEXT("//Root/A"));
+	RunAppendTest(TEXT("//Root"), TEXT("A/B"), TEXT("//Root/A/B"));
+	RunAppendTest(TEXT("//Root"), TEXT("A\\B"), TEXT("//Root/A\\B"));
+	RunAppendTest(TEXT("\\\\Root"), TEXT("A"), TEXT("\\\\Root/A"));
+	RunAppendTest(TEXT("\\\\Root"), TEXT("A/B"), TEXT("\\\\Root/A/B"));
+	RunAppendTest(TEXT("\\\\Root"), TEXT("A\\B"), TEXT("\\\\Root/A\\B"));
+	RunAppendTest(TEXT("//Root/"), TEXT("A"), TEXT("//Root/A"));
+	RunAppendTest(TEXT("//Root/"), TEXT("A/B"), TEXT("//Root/A/B"));
+	RunAppendTest(TEXT("//Root/"), TEXT("A\\B"), TEXT("//Root/A\\B"));
+	RunAppendTest(TEXT("\\\\Root\\"), TEXT("A"), TEXT("\\\\Root\\A"));
+	RunAppendTest(TEXT("\\\\Root\\"), TEXT("A/B"), TEXT("\\\\Root\\A/B"));
+	RunAppendTest(TEXT("\\\\Root\\"), TEXT("A\\B"), TEXT("\\\\Root\\A\\B"));
+
+	RunAppendTest(TEXT("C:/"), TEXT("A"), TEXT("C:/A"));
+	RunAppendTest(TEXT("C:/"), TEXT("A/B"), TEXT("C:/A/B"));
+	RunAppendTest(TEXT("C:/"), TEXT("A\\B"), TEXT("C:/A\\B"));
+	RunAppendTest(TEXT("C:\\"), TEXT("A"), TEXT("C:\\A"));
+	RunAppendTest(TEXT("C:\\"), TEXT("A/B"), TEXT("C:\\A/B"));
+	RunAppendTest(TEXT("C:\\"), TEXT("A\\B"), TEXT("C:\\A\\B"));
+
+	RunAppendTest(TEXT("C:/Root"), TEXT("A"), TEXT("C:/Root/A"));
+	RunAppendTest(TEXT("C:/Root"), TEXT("A/B"), TEXT("C:/Root/A/B"));
+	RunAppendTest(TEXT("C:/Root"), TEXT("A\\B"), TEXT("C:/Root/A\\B"));
+	RunAppendTest(TEXT("C:\\Root"), TEXT("A"), TEXT("C:\\Root/A"));
+	RunAppendTest(TEXT("C:\\Root"), TEXT("A/B"), TEXT("C:\\Root/A/B"));
+	RunAppendTest(TEXT("C:\\Root"), TEXT("A\\B"), TEXT("C:\\Root/A\\B"));
+	RunAppendTest(TEXT("C:/Root/"), TEXT("A"), TEXT("C:/Root/A"));
+	RunAppendTest(TEXT("C:/Root/"), TEXT("A/B"), TEXT("C:/Root/A/B"));
+	RunAppendTest(TEXT("C:/Root/"), TEXT("A\\B"), TEXT("C:/Root/A\\B"));
+	RunAppendTest(TEXT("C:\\Root\\"), TEXT("A"), TEXT("C:\\Root\\A"));
+	RunAppendTest(TEXT("C:\\Root\\"), TEXT("A/B"), TEXT("C:\\Root\\A/B"));
+	RunAppendTest(TEXT("C:\\Root\\"), TEXT("A\\B"), TEXT("C:\\Root\\A\\B"));
+
+	// No matter the prefix, appending a rooted path should result in only the rooted path
+	for (const TCHAR* Prefix : { TEXT(""),
+		TEXT("/"), TEXT("\\"), TEXT("/Root"), TEXT("\\Root"), TEXT("/Root/"), TEXT("\\Root\\"),
+		TEXT("//"), TEXT("\\\\"), TEXT("//Root"), TEXT("\\\\Root"), TEXT("//Root/"), TEXT("\\\\Root\\"),
+		TEXT("C:/"), TEXT("C:\\"), TEXT("C:/Root"), TEXT("C:\\Root"), TEXT("C:/Root/"), TEXT("C:\\Root\\")})
+	{
+		RunAppendTest(Prefix, TEXT("/A"), TEXT("/A"));
+		RunAppendTest(Prefix, TEXT("\\A"), TEXT("\\A"));
+		RunAppendTest(Prefix, TEXT("/A/B"), TEXT("/A/B"));
+		RunAppendTest(Prefix, TEXT("\\A\\B"), TEXT("\\A\\B"));
+
+		RunAppendTest(Prefix, TEXT("//A"), TEXT("//A"));
+		RunAppendTest(Prefix, TEXT("\\\\A"), TEXT("\\\\A"));
+		RunAppendTest(Prefix, TEXT("//A/B"), TEXT("//A/B"));
+		RunAppendTest(Prefix, TEXT("\\\\A\\B"), TEXT("\\\\A\\B"));
+
+		RunAppendTest(Prefix, TEXT("C:/A"), TEXT("C:/A"));
+		RunAppendTest(Prefix, TEXT("C:\\A"), TEXT("C:\\A"));
+		RunAppendTest(Prefix, TEXT("C:/A/B"), TEXT("C:/A/B"));
+		RunAppendTest(Prefix, TEXT("C:\\A\\B"), TEXT("C:\\A\\B"));
+	}
+
+	return true;
+}
+
 #endif
diff --git a/Engine/Source/Runtime/Core/Public/Misc/PathViews.h b/Engine/Source/Runtime/Core/Public/Misc/PathViews.h
--- a/Engine/Source/Runtime/Core/Public/Misc/PathViews.h
+++ b/Engine/Source/Runtime/Core/Public/Misc/PathViews.h
@@ -100,6 +100,21 @@ public:
 	 */
 	static FStringView GetPathLeaf(const FStringView& InPath);
 
+	/**
+	 * Return whether the given relative or absolute path is a leaf path - has no separators.
+	 * 
+	 * Examples: (Using '/' but '\' functions the same way)
+	 * A		-> true
+	 * A/		-> true
+	 * D:/		-> true
+	 * /		-> true
+	 * //		-> true
+	 * A/B		-> false
+	 * D:/A		-> false
+	 * //A		-> false
+	 */
+	static bool IsPathLeaf(FStringView InPath);
+
 	/**
 	 * Splits InPath into individual directory components, and calls ComponentVisitor on each.
 	 *
@@ -147,7 +162,8 @@ public:
 	 */
 	static void Append(FStringBuilderBase& Builder, const FStringView& Suffix);
 
-	/** Replaces the pre-existing file extension of a filename. 
+	/**
+	 * Replaces the pre-existing file extension of a filename. 
 	 *
 	 * @param InPath A valid file path with a pre-existing extension.
 	 * @param InNewExtension The new extension to use (prefixing with a '.' is optional)
@@ -156,4 +172,131 @@ public:
 	 * case a copy of InPath will be returned instead.
 	 */
 	static FString ChangeExtension(const FStringView& InPath, const FStringView& InNewExtension);
+
+	/** Return whether the given character is a path-separator character (/ or \) */
+	static bool IsSeparator(TCHAR c);
+
+	/**
+	 * Return true if the given paths are the same path (with exceptions noted below).
+	 * Case-insensitive
+	 * / is treated as equal to \
+	 * Presence or absence of terminating separator (/) is ignored in the comparison.
+	 * Directory elements of . and .. are currently not interpreted and are treated as literal characters.
+	 *    Callers should not rely on this behavior as it may be corrected in the future.
+	 *    callers should instead conform the paths before calling.
+	 * Relative paths and absolute paths are not resolved, and relative paths will never equal absolute paths.
+	 *    Callers should not rely on this behavior as it may be corrected in the future;
+	 *    callers should instead conform the paths before calling.
+	 * Examples:
+	 * ("../A/B.C", "../A/B.C") -> true
+	 * ("../A/B", "../A/B.C")	-> false
+	 * ("../A/", "../A/")		-> true
+	 * ("../A/", "../A")		-> true
+	 * ("d:/root/Engine/", "d:\root\Engine") -> true
+	 * (../../../Engine/Content", "d:/root/Engine/Content") -> false
+	 * (d:/root/Engine/..", "d:/root") -> false
+	 * (d:/root/Engine/./Content", "d:/root/Engine/Content") -> false
+	 */
+	static bool Equals(FStringView A, FStringView B);
+
+	/**
+	 * Return true if the the first path is lexicographically less than the second path (with caveats noted below).
+	 * Case-insensitive
+	 * / is treated as equal to \
+	 * Presence or absence of terminating separator (/) is ignored in the comparison.
+	 * Directory elements of . and .. are currently not interpreted and are treated as literal characters.
+	 *    Callers should not rely on this behavior as it may be corrected in the future.
+	 *    callers should instead conform the paths before calling.
+	 * Relative paths and absolute paths are not resolved, and relative paths will never equal absolute paths.
+	 *    Callers should not rely on this behavior as it may be corrected in the future;
+	 *    callers should instead conform the paths before calling.
+	 * Examples:
+	 * ("../A/B.C", "../A/B.C") -> false (they are equal)
+	 * ("../A/B", "../A/B.C")	-> true (a string is greater than any prefix of itself)
+	 * ("../A/", "../A/")		-> false (they are equal)
+	 * ("../A/", "../A")		-> false (they are equal)
+	 * ("../A", "../A/")		-> false (they are equal)
+	 * ("d:/root/Engine/", "d:\root\Engine") -> false (they are equal)
+	 * (../../../Engine/Content", "d:/root/Engine/Content") -> true ('.' is less than 'd')
+	 * (d:/root/Engine/..", "d:/root") -> false (A string is greater than any prefix of itself)
+	 * (d:/root/Engine/./Content", "d:/root/Engine/Content") -> false
+	 */
+	static bool Less(FStringView A, FStringView B);
+
+	/**
+	 * Check whether Parent is a parent path of Child and report the relative path if so.
+	 * Case-insensitive
+	 * / is treated as equal to \
+	 * Presence or absence of terminating separator (/) is ignored in the comparison.
+	 * Directory elements of . and .. are currently not interpreted and are treated as literal characters.
+	 *    Callers should not rely on this behavior as it may be corrected in the future.
+	 *    callers should instead conform the paths before calling.
+	 * Relative paths and absolute paths are not resolved, and relative paths will never equal absolute paths.
+	 *    Callers should not rely on this behavior as it may be corrected in the future;
+	 *    callers should instead conform the paths before calling.
+	 * Examples:
+	 * ("../A/B", "../A")	-> (true, "B")
+	 * ("../A\B", "../A/")	-> (true, "B")
+	 * ("../A/", "../A")	-> (true, "")
+	 * (".././A/", "../A")	-> (false, "")
+	 * ("../../../Engine", "d:/root/Engine") -> (false, "")
+	 *
+	 * @param Child An absolute path that may be a child path of Parent.
+	 * @param Parent An absolute path that may be a parent path of Child.
+	 * @param OutRelPath Receives the relative path from Parent to Child, or empty if Parent is not a parent of Child.
+	 *
+	 * @return True if and only if Child is a child path of Parent (or is equal to it).
+	 */
+	static bool TryMakeChildPathRelativeTo(FStringView Child, FStringView Parent, FStringView& OutRelPath);
+
+	/**
+	 * Return whether Parent is a parent path of (or is equal to) Child.
+	 * Case-insensitive
+	 * / is treated as equal to \
+	 * Presence or absence of terminating separator (/) is ignored in the comparison.
+	 * Directory elements of . and .. are currently not interpreted and are treated as literal characters.
+	 *    Callers should not rely on this behavior as it may be corrected in the future.
+	 *    callers should instead conform the paths before calling.
+	 * Relative paths and absolute paths are not resolved, and relative paths will never equal absolute paths.
+	 *    Callers should not rely on this behavior as it may be corrected in the future;
+	 *    callers should instead conform the paths before calling.
+	 * Examples:
+	 * ("../A", "../A/B")	-> true
+	 * ("../A/", "../A\B")	-> true
+	 * ("../A", "../A/")	-> true
+	 * ("../A", ".././A/")	-> false
+	 * ("d:/root/Engine", "../../../Engine") -> false
+	 *
+	 * @param Parent An absolute path that may be a parent path of Child.
+	 * @param Child An absolute path that may be a child path of Parent.
+	 *
+	 * @return True if and only if Child is a child path of Parent (or is equal to it).
+	 */
+	static bool IsParentPathOf(FStringView Parent, FStringView Child);
+
+	/**
+	 * Return whether the given path is a relativepath - does not start with a separator or volume:.
+	 * Returns true for empty paths.
+	 */
+	static bool IsRelativePath(FStringView InPath);
+
+	/**
+	 * Split the given absolute or relative path into its topmost directory and the relative path from that directory.
+	 * Directory elements of . and .. are currently not interpreted and are treated as literal characters.
+	 *    Callers should not rely on this behavior as it may be corrected in the future.
+	 *    callers should instead conform the paths before calling.
+	 *
+	 * @param InPath The path to split.
+	 * @param OutFirstComponent Receives the first directory element in the path, or InPath if it is a leaf path.
+	 * @param OutRemainder Receives the relative path from OutFirstComponent to InPath, or empty if InPath is a leaf path.
+	 */
+	static void SplitFirstComponent(FStringView InPath, FStringView& OutFirstComponent, FStringView& OutRemainder);
+
+	/**
+	 * If AppendPath is a relative path, append it as a relative path onto InOutPath.
+	 * If AppendPath is absolute, reset InOutPath and replace it with RelPath.
+	 * Handles presence or absence of terminating separator in BasePath.
+	 * Does not interpret . or ..; each occurrence of these in either path will remain in the combined InOutPath.
+	 */
+	static void AppendPath(FStringBuilderBase& InOutPath, FStringView AppendPath);
 };
