StringBuilder::GetAllocatedSize
#rnx

diff --git a/Engine/Source/Runtime/Core/Public/Misc/StringBuilder.h b/Engine/Source/Runtime/Core/Public/Misc/StringBuilder.h
--- a/Engine/Source/Runtime/Core/Public/Misc/StringBuilder.h
+++ b/Engine/Source/Runtime/Core/Public/Misc/StringBuilder.h
@@ -85,6 +85,17 @@ public:
 
 	inline const CharType	LastChar() const	{ return *(CurPos - 1); }
 
+	/**
+	 * Helper function to return the amount of memory allocated by this container.
+	 * Does not include the sizeof of the inline buffer, only includes the size of the overflow buffer.
+	 *
+	 * @returns Number of bytes allocated by this container.
+	 */
+	SIZE_T GetAllocatedSize() const
+	{
+		return bIsDynamic ? (End - Base) * sizeof(CharType) : 0;
+	}
+
 	/**
 	 * Empties the string builder, but doesn't change memory allocation.
 	 */
