Fix Config Cache ini FindConfigFileWithBaseName not properly looking in known configs

diff --git a/Engine/Source/Runtime/Core/Private/Misc/ConfigCacheIni.cpp b/Engine/Source/Runtime/Core/Private/Misc/ConfigCacheIni.cpp
--- a/Engine/Source/Runtime/Core/Private/Misc/ConfigCacheIni.cpp
+++ b/Engine/Source/Runtime/Core/Private/Misc/ConfigCacheIni.cpp
@@ -2156,6 +2156,11 @@ FConfigFile* FConfigCacheIni::Find(const FString& Filename, bool CreateIfNotFoun
 
 FConfigFile* FConfigCacheIni::FindConfigFileWithBaseName(FName BaseName)
 {
+	if (FConfigFile* Result = KnownFiles.GetMutableFile(BaseName))
+	{
+		return Result;
+	}
+
 	for (TPair<FString,FConfigFile>& CurrentFilePair : *this)
 	{
 		if (CurrentFilePair.Value.Name == BaseName)
