Remove constexpr to fix mac build

diff --git a/Engine/Source/Runtime/Core/Public/Misc/QueuedThreadPool.h b/Engine/Source/Runtime/Core/Public/Misc/QueuedThreadPool.h
--- a/Engine/Source/Runtime/Core/Public/Misc/QueuedThreadPool.h
+++ b/Engine/Source/Runtime/Core/Public/Misc/QueuedThreadPool.h
@@ -18,7 +18,7 @@ enum class EQueuedWorkPriority : uint8
 	Count
 };
 
-constexpr inline const TCHAR* LexToString(EQueuedWorkPriority Priority)
+inline const TCHAR* LexToString(EQueuedWorkPriority Priority)
 {
 	switch (Priority)
 	{
