Support opening an asset in a true default editor.

diff --git a/Engine/Source/Runtime/Core/Private/Windows/WindowsPlatformProcess.cpp b/Engine/Source/Runtime/Core/Private/Windows/WindowsPlatformProcess.cpp
--- a/Engine/Source/Runtime/Core/Private/Windows/WindowsPlatformProcess.cpp
+++ b/Engine/Source/Runtime/Core/Private/Windows/WindowsPlatformProcess.cpp
@@ -1221,6 +1221,12 @@ void FWindowsPlatformProcess::LaunchFileInDefaultExternalApplication( const TCHA
 	
 	UE_LOG(LogWindows, Log,  TEXT("Launch application code for %s %s: %d"), FileName, Parms ? Parms : TEXT(""), (PTRINT)Code );
 
+	// Fallback to a true windows-defined default for the asset type
+	if ((PTRINT)Code == SE_ERR_NOASSOC || (PTRINT)Code == SE_ERR_ASSOCINCOMPLETE)
+	{
+		Code = ::ShellExecuteW(NULL, NULL, FileName, NULL, NULL, SW_SHOW);
+	}
+
 	// If opening the file in the default application failed, check to see if it's because the file's extension does not have
 	// a default application associated with it. If so, prompt the user with the Windows "Open With..." dialog to allow them to specify
 	// an application to use.
