BLAKE3: Added FBlake3Hash::IsZero for convenience

#rnx

diff --git a/Engine/Source/Runtime/Core/Public/Misc/Blake3.h b/Engine/Source/Runtime/Core/Public/Misc/Blake3.h
--- a/Engine/Source/Runtime/Core/Public/Misc/Blake3.h
+++ b/Engine/Source/Runtime/Core/Public/Misc/Blake3.h
@@ -28,6 +28,8 @@ public:
 
 	inline const ByteArray& Bytes() const { return Hash; }
 
+	inline bool IsZero() const;
+
 private:
 	alignas(uint32) ByteArray Hash{};
 };
@@ -78,6 +80,20 @@ inline FBlake3Hash::FBlake3Hash(const ByteArray& InHash)
 	FMemory::Memcpy(Hash, InHash, sizeof(ByteArray));
 }
 
+inline bool FBlake3Hash::IsZero() const
+{
+	using UInt32Array = uint32[8];
+	static_assert(sizeof(UInt32Array) == sizeof(ByteArray), "Invalid size for UInt32Array");
+	for (uint32 Value : reinterpret_cast<const UInt32Array&>(Hash))
+	{
+		if (Value != 0)
+		{
+			return false;
+		}
+	}
+	return true;
+}
+
 inline bool operator==(const FBlake3Hash& A, const FBlake3Hash& B)
 {
 	return FMemory::Memcmp(A.Bytes(), B.Bytes(), sizeof(decltype(A.Bytes()))) == 0;
