Fixing PVS 7.7 Issues:

warning V575: The null pointer is passed into 'Strnlen' function. Inspect the first argument.
warning V575: The 'Strnlen' function processes '0' elements. Inspect the second argument.

diff --git a/Engine/Source/Runtime/Core/Private/Tests/HAL/PlatformStringTest.cpp b/Engine/Source/Runtime/Core/Private/Tests/HAL/PlatformStringTest.cpp
--- a/Engine/Source/Runtime/Core/Private/Tests/HAL/PlatformStringTest.cpp
+++ b/Engine/Source/Runtime/Core/Private/Tests/HAL/PlatformStringTest.cpp
@@ -37,9 +37,9 @@ bool FPlatformStringTestStrnlen::RunTest(const FString& Parameters)
 	ANSICHAR AnsiBuffer[128] = "123456789";
 	TestEqual(TEXT("Strnlen(PaddedBuffer)"), FPlatformString::Strnlen(AnsiBuffer, UE_ARRAY_COUNT(AnsiBuffer)), 9);
 
-	TestEqual(TEXT("Strnlen(nullptr, 0)"), FPlatformString::Strnlen((const TCHAR*)nullptr, 0), 0);
-	TestEqual(TEXT("Strnlen(\"\", 0)"), FPlatformString::Strnlen(TEXT(""), 0), 0);
-	TestEqual(TEXT("Strnlen(\"1\", 0)"), FPlatformString::Strnlen(TEXT("1"), 0), 0);
+	TestEqual(TEXT("Strnlen(nullptr, 0)"), FPlatformString::Strnlen((const TCHAR*)nullptr, 0), 0); //-V575
+	TestEqual(TEXT("Strnlen(\"\", 0)"), FPlatformString::Strnlen(TEXT(""), 0), 0); //-V575
+	TestEqual(TEXT("Strnlen(\"1\", 0)"), FPlatformString::Strnlen(TEXT("1"), 0), 0); //-V575
 	TestEqual(TEXT("Strnlen(\"1\", 1)"), FPlatformString::Strnlen(TEXT("1"), 1), 1);
 	TestEqual(TEXT("Strnlen(\"1\", 2)"), FPlatformString::Strnlen(TEXT("1"), 2), 1);
 	TestEqual(TEXT("Strnlen(\"123\", 2)"), FPlatformString::Strnlen(TEXT("123"), 2), 2);
