Fixing PVS 7.7 Issues:

warning V501: There are identical sub-expressions 'MakeMemoryView(ByteArrayContainer.ByteArray)' to the left and to the right of the '==' operator.

diff --git a/Engine/Source/Runtime/Core/Private/Tests/Memory/MemoryViewTest.cpp b/Engine/Source/Runtime/Core/Private/Tests/Memory/MemoryViewTest.cpp
--- a/Engine/Source/Runtime/Core/Private/Tests/Memory/MemoryViewTest.cpp
+++ b/Engine/Source/Runtime/Core/Private/Tests/Memory/MemoryViewTest.cpp
@@ -204,10 +204,10 @@ bool FMemoryViewTest::RunTest(const FString& Parameters)
 	TestFalse(TEXT("MemoryView.Equals(BothDiff)"), MakeMemoryView(IntArray).Equals(FMutableMemoryView()));
 
 	// Test operator==
-	static_assert(MakeMemoryView(ByteArrayContainer.ByteArray) == MakeMemoryView(ByteArrayContainer.ByteArray), "Error in MemoryView == MemoryView");
+	static_assert(MakeMemoryView(ByteArrayContainer.ByteArray) == MakeMemoryView(ByteArrayContainer.ByteArray), "Error in MemoryView == MemoryView"); //-V501
 	static_assert(MakeMemoryView(ByteArrayContainer.ByteArray) == MakeMemoryView(AsConst(ByteArrayContainer.ByteArray)), "Error in MemoryView == MemoryView");
 	static_assert(MakeMemoryView(AsConst(ByteArrayContainer.ByteArray)) == MakeMemoryView(ByteArrayContainer.ByteArray), "Error in MemoryView == MemoryView");
-	static_assert(MakeMemoryView(AsConst(ByteArrayContainer.ByteArray)) == MakeMemoryView(AsConst(ByteArrayContainer.ByteArray)), "Error in MemoryView == MemoryView");
+	static_assert(MakeMemoryView(AsConst(ByteArrayContainer.ByteArray)) == MakeMemoryView(AsConst(ByteArrayContainer.ByteArray)), "Error in MemoryView == MemoryView"); //-V501
 
 	// Test operator!=
 	static_assert(MakeMemoryView(ByteArrayContainer.ByteArray) != MakeMemoryView(IntArray), "Error in MemoryView != MemoryView");
