Fix inconsistent total virtual memory being reported for Windows

Before:
LogMemory: Memory total: Physical=255.9GB (256GB approx)
LogMemory: Platform Memory Stats for WindowsEditor
LogMemory: Process Physical Memory: 266.32 MB used, 266.34 MB peak
LogMemory: Process Virtual Memory: 251.57 MB used, 251.57 MB peak
LogMemory: Physical Memory: 30622.70 MB used,  231393.88 MB free, 262016.56 MB total
LogMemory: Virtual Memory: 133801944.00 MB used,  415784.91 MB free, 134217728.00 MB total

After:
LogMemory: Memory total: Physical=255.9GB (256GB approx)
LogMemory: Platform Memory Stats for WindowsEditor
LogMemory: Process Physical Memory: 265.60 MB used, 265.60 MB peak
LogMemory: Process Virtual Memory: 247.32 MB used, 247.32 MB peak
LogMemory: Physical Memory: 30700.99 MB used,  231315.58 MB free, 262016.56 MB total
LogMemory: Virtual Memory: 38424.43 MB used,  415592.16 MB free, 454016.56 MB total

diff --git a/Engine/Source/Runtime/Core/Private/Windows/WindowsPlatformMemory.cpp b/Engine/Source/Runtime/Core/Private/Windows/WindowsPlatformMemory.cpp
--- a/Engine/Source/Runtime/Core/Private/Windows/WindowsPlatformMemory.cpp
+++ b/Engine/Source/Runtime/Core/Private/Windows/WindowsPlatformMemory.cpp
@@ -229,6 +229,10 @@ FPlatformMemoryStats FWindowsPlatformMemory::GetStats()
 	MemoryStats.AvailablePhysical = MemoryStatusEx.ullAvailPhys;
 	MemoryStats.AvailableVirtual = MemoryStatusEx.ullAvailVirtual;
 
+	// On Windows, ullTotalVirtual is artificial and represent the SKU limitation (128TB on Win10Pro) instead of the commit limit of the system we're after.
+	// ullTotalPageFile represents PhysicalMemory + Disk Swap Space, which is the value we care about
+	MemoryStats.TotalVirtual = MemoryStatusEx.ullTotalPageFile;
+
 	// On Windows, Virtual Memory is limited per process to the address space (e.g. 47 bits (128Tb)), but is additionally limited by the sum of used virtual memory across all processes
 	// must be less than PhysicalMemory plus the Virtual Memory Page Size. The remaining virtual memory space given this system-wide limit is stored in ullAvailPageSize
 	MemoryStats.AvailableVirtual = FMath::Min(MemoryStats.AvailableVirtual, MemoryStatusEx.ullAvailPageFile);
@@ -270,7 +274,9 @@ const FPlatformMemoryConstants& FWindowsPlatformMemory::GetConstants()
 		::GetSystemInfo(&SystemInfo);
 
 		MemoryConstants.TotalPhysical = MemoryStatusEx.ullTotalPhys;
-		MemoryConstants.TotalVirtual = MemoryStatusEx.ullTotalVirtual;
+		// On Windows, ullTotalVirtual is artificial and represent the SKU limitation (128TB on Win10Pro) instead of the commit limit of the system we're after.
+		// ullTotalPageFile represents PhysicalMemory + Disk Swap Space, which is the value we care about
+		MemoryConstants.TotalVirtual = MemoryStatusEx.ullTotalPageFile;
 		MemoryConstants.BinnedPageSize = SystemInfo.dwAllocationGranularity;	// Use this so we get larger 64KiB pages, instead of 4KiB
 		MemoryConstants.BinnedAllocationGranularity = SystemInfo.dwPageSize; // Use 4KiB pages for more efficient use of memory - 64KiB pages don't really exist on this CPU
 		MemoryConstants.OsAllocationGranularity = SystemInfo.dwAllocationGranularity;	// VirtualAlloc cannot allocate memory less than that
