Add comments to ContainerElementTypeCompatibility.h explaining limitations of use for the compatibility types.  Split out from review here:

https://p4-swarm.epicgames.net/reviews/15065157

diff --git a/Engine/Source/Runtime/Core/Public/Containers/ContainerElementTypeCompatibility.h b/Engine/Source/Runtime/Core/Public/Containers/ContainerElementTypeCompatibility.h
--- a/Engine/Source/Runtime/Core/Public/Containers/ContainerElementTypeCompatibility.h
+++ b/Engine/Source/Runtime/Core/Public/Containers/ContainerElementTypeCompatibility.h
@@ -4,6 +4,11 @@
 
 #include "CoreTypes.h"
 
+/**
+ * Temporary compatibility mechanism to be used solely for the purpose of raw pointers to wrapped pointers.
+ * Specialization and use of this type is not supported beyond the specific instances in ObjectPtr.h.
+ * These types and the container methods that leverage them will be removed in a future release.
+ */
 template <typename InElementType>
 struct TContainerElementTypeCompatibility
 {
@@ -19,12 +24,20 @@ struct TContainerElementTypeCompatibility
 	static constexpr void CopyingFromOtherType() {}
 };
 
+/**
+ * Temporary compatibility mechanism to be used solely for the purpose of raw pointers to wrapped pointers.
+ * These types and the container methods that leverage them will be removed in a future release.
+ */
 template <typename ElementType>
 struct TIsContainerElementTypeReinterpretable
 {
 	enum { Value = !TIsSame<typename TContainerElementTypeCompatibility<ElementType>::ReinterpretType, ElementType>::Value };
 };
 
+/**
+ * Temporary compatibility mechanism to be used solely for the purpose of raw pointers to wrapped pointers.
+ * These types and the container methods that leverage them will be removed in a future release.
+ */
 template <typename ElementType>
 struct TIsContainerElementTypeCopyable
 {
