Replace use of TEnableIf with std::enable_if_t in TArray/TSet.  Split out from review here:

https://p4-swarm.epicgames.net/reviews/15065157

diff --git a/Engine/Source/Runtime/Core/Public/Containers/Array.h b/Engine/Source/Runtime/Core/Public/Containers/Array.h
--- a/Engine/Source/Runtime/Core/Public/Containers/Array.h
+++ b/Engine/Source/Runtime/Core/Public/Containers/Array.h
@@ -428,7 +428,7 @@ private:
 	template <
 		typename FromArrayType,
 		typename ToArrayType,
-		typename TEnableIf<TCanMoveBetweenAllocators<typename FromArrayType::Allocator, typename ToArrayType::Allocator>::Value>::Type* = nullptr
+		std::enable_if_t<TCanMoveBetweenAllocators<typename FromArrayType::Allocator, typename ToArrayType::Allocator>::Value>* = nullptr
 	>
 	static FORCEINLINE void MoveAllocatorToEmpty(FromArrayType& FromArray, ToArrayType& ToArray)
 	{
@@ -438,7 +438,7 @@ private:
 	template <
 		typename FromArrayType,
 		typename ToArrayType,
-		typename TEnableIf<!TCanMoveBetweenAllocators<typename FromArrayType::Allocator, typename ToArrayType::Allocator>::Value>::Type* = nullptr
+		std::enable_if_t<!TCanMoveBetweenAllocators<typename FromArrayType::Allocator, typename ToArrayType::Allocator>::Value>* = nullptr
 	>
 	static FORCEINLINE void MoveAllocatorToEmpty(FromArrayType& FromArray, ToArrayType& ToArray)
 	{
@@ -455,7 +455,7 @@ private:
 	 * @param FromArray Array to move from.
 	 */
 	template <typename FromArrayType, typename ToArrayType>
-	static FORCEINLINE typename TEnableIf<UE4Array_Private::TCanMoveTArrayPointersBetweenArrayTypes<FromArrayType, ToArrayType>::Value>::Type MoveOrCopy(ToArrayType& ToArray, FromArrayType& FromArray, SizeType PrevMax)
+	static FORCEINLINE std::enable_if_t<UE4Array_Private::TCanMoveTArrayPointersBetweenArrayTypes<FromArrayType, ToArrayType>::Value> MoveOrCopy(ToArrayType& ToArray, FromArrayType& FromArray, SizeType PrevMax)
 	{
 		using FromAllocatorType = typename FromArrayType::Allocator;
 		using ToAllocatorType   = typename ToArrayType::Allocator;
@@ -494,7 +494,7 @@ private:
 	 *                   at the end of the array in the number of elements.
 	 */
 	template <typename FromArrayType, typename ToArrayType>
-	static FORCEINLINE typename TEnableIf<!UE4Array_Private::TCanMoveTArrayPointersBetweenArrayTypes<FromArrayType, ToArrayType>::Value>::Type MoveOrCopy(ToArrayType& ToArray, FromArrayType& FromArray, SizeType PrevMax)
+	static FORCEINLINE std::enable_if_t<!UE4Array_Private::TCanMoveTArrayPointersBetweenArrayTypes<FromArrayType, ToArrayType>::Value> MoveOrCopy(ToArrayType& ToArray, FromArrayType& FromArray, SizeType PrevMax)
 	{
 		ToArray.CopyToEmpty(FromArray.GetData(), FromArray.Num(), PrevMax, 0);
 	}
@@ -510,7 +510,7 @@ private:
 	 *                   at the end of the array in the number of elements.
 	 */
 	template <typename FromArrayType, typename ToArrayType>
-	static FORCEINLINE typename TEnableIf<UE4Array_Private::TCanMoveTArrayPointersBetweenArrayTypes<FromArrayType, ToArrayType>::Value>::Type MoveOrCopyWithSlack(ToArrayType& ToArray, FromArrayType& FromArray, SizeType PrevMax, SizeType ExtraSlack)
+	static FORCEINLINE std::enable_if_t<UE4Array_Private::TCanMoveTArrayPointersBetweenArrayTypes<FromArrayType, ToArrayType>::Value> MoveOrCopyWithSlack(ToArrayType& ToArray, FromArrayType& FromArray, SizeType PrevMax, SizeType ExtraSlack)
 	{
 		MoveOrCopy(ToArray, FromArray, PrevMax);
 
@@ -528,7 +528,7 @@ private:
 	 *                   at the end of the array in the number of elements.
 	 */
 	template <typename FromArrayType, typename ToArrayType>
-	static FORCEINLINE typename TEnableIf<!UE4Array_Private::TCanMoveTArrayPointersBetweenArrayTypes<FromArrayType, ToArrayType>::Value>::Type MoveOrCopyWithSlack(ToArrayType& ToArray, FromArrayType& FromArray, SizeType PrevMax, SizeType ExtraSlack)
+	static FORCEINLINE std::enable_if_t<!UE4Array_Private::TCanMoveTArrayPointersBetweenArrayTypes<FromArrayType, ToArrayType>::Value> MoveOrCopyWithSlack(ToArrayType& ToArray, FromArrayType& FromArray, SizeType PrevMax, SizeType ExtraSlack)
 	{
 		ToArray.CopyToEmpty(FromArray.GetData(), FromArray.Num(), PrevMax, ExtraSlack);
 	}
@@ -2113,7 +2113,7 @@ public:
 	/** Implicit conversion operator to container of compatible element type. */
 	template <
 		typename AliasElementType = ElementType,
-		typename TEnableIf<TIsContainerElementTypeReinterpretable<AliasElementType>::Value>::Type* = nullptr
+		std::enable_if_t<TIsContainerElementTypeReinterpretable<AliasElementType>::Value>* = nullptr
 		>
 	operator TArray<typename TContainerElementTypeCompatibility<AliasElementType>::ReinterpretType, Allocator>& ()
 	{
diff --git a/Engine/Source/Runtime/Core/Public/Containers/Set.h b/Engine/Source/Runtime/Core/Public/Containers/Set.h
--- a/Engine/Source/Runtime/Core/Public/Containers/Set.h
+++ b/Engine/Source/Runtime/Core/Public/Containers/Set.h
@@ -176,7 +176,7 @@ public:
 	FORCEINLINE TSetElementBase() {}
 
 	/** Initialization constructor. */
-	template <typename InitType, typename = typename TEnableIf<!TAreTypesEqual<TSetElementBase, typename TDecay<InitType>::Type>::Value>::Type> explicit FORCEINLINE TSetElementBase(InitType&& InValue) : Value(Forward<InitType>(InValue)) {}
+	template <typename InitType, typename = std::enable_if_t<!TAreTypesEqual<TSetElementBase, typename TDecay<InitType>::Type>::Value>> explicit FORCEINLINE TSetElementBase(InitType&& InValue) : Value(Forward<InitType>(InValue)) {}
 
 	TSetElementBase(TSetElementBase&&) = default;
 	TSetElementBase(const TSetElementBase&) = default;
@@ -203,7 +203,7 @@ public:
 	FORCEINLINE TSetElementBase() {}
 
 	/** Initialization constructor. */
-	template <typename InitType, typename = typename TEnableIf<!TAreTypesEqual<TSetElementBase, typename TDecay<InitType>::Type>::Value>::Type> explicit FORCEINLINE TSetElementBase(InitType&& InValue) : Value(Forward<InitType>(InValue)) {}
+	template <typename InitType, typename = std::enable_if_t<!TAreTypesEqual<TSetElementBase, typename TDecay<InitType>::Type>::Value>> explicit FORCEINLINE TSetElementBase(InitType&& InValue) : Value(Forward<InitType>(InValue)) {}
 
 	TSetElementBase(TSetElementBase&&) = default;
 	TSetElementBase(const TSetElementBase&) = default;
@@ -231,7 +231,7 @@ public:
 	{}
 
 	/** Initialization constructor. */
-	template <typename InitType, typename = typename TEnableIf<!TAreTypesEqual<TSetElement, typename TDecay<InitType>::Type>::Value>::Type> explicit FORCEINLINE TSetElement(InitType&& InValue) : Super(Forward<InitType>(InValue)) {}
+	template <typename InitType, typename = std::enable_if_t<!TAreTypesEqual<TSetElement, typename TDecay<InitType>::Type>::Value>> explicit FORCEINLINE TSetElement(InitType&& InValue) : Super(Forward<InitType>(InValue)) {}
 
 	TSetElement(TSetElement&&) = default;
 	TSetElement(const TSetElement&) = default;
@@ -349,7 +349,7 @@ public:
 
 private:
 	template <typename SetType>
-	static FORCEINLINE typename TEnableIf<TContainerTraits<SetType>::MoveWillEmptyContainer>::Type MoveOrCopy(SetType& ToSet, SetType& FromSet)
+	static FORCEINLINE std::enable_if_t<TContainerTraits<SetType>::MoveWillEmptyContainer> MoveOrCopy(SetType& ToSet, SetType& FromSet)
 	{
 		ToSet.Elements = (ElementArrayType&&)FromSet.Elements;
 
@@ -360,7 +360,7 @@ private:
 	}
 
 	template <typename SetType>
-	static FORCEINLINE typename TEnableIf<!TContainerTraits<SetType>::MoveWillEmptyContainer>::Type MoveOrCopy(SetType& ToSet, SetType& FromSet)
+	static FORCEINLINE std::enable_if_t<!TContainerTraits<SetType>::MoveWillEmptyContainer> MoveOrCopy(SetType& ToSet, SetType& FromSet)
 	{
 		ToSet = FromSet;
 	}
