Introduce TIsTSet and TIsTMap which are TSet/TMap equivalents of TIsTArray.  Split out from review here:

https://p4-swarm.epicgames.net/reviews/15065157

diff --git a/Engine/Source/Runtime/Core/Public/Containers/Map.h b/Engine/Source/Runtime/Core/Public/Containers/Map.h
--- a/Engine/Source/Runtime/Core/Public/Containers/Map.h
+++ b/Engine/Source/Runtime/Core/Public/Containers/Map.h
@@ -1879,6 +1879,16 @@ struct TIsZeroConstructType<TScriptMap<AllocatorType>>
 	enum { Value = true };
 };
 
+/**
+ * Traits class which determines whether or not a type is a TMap.
+ */
+template <typename T> struct TIsTMap { enum { Value = false }; };
+
+template <typename KeyType, typename ValueType, typename SetAllocator, typename KeyFuncs> struct TIsTMap<               TMap<KeyType, ValueType, SetAllocator, KeyFuncs>> { enum { Value = true }; };
+template <typename KeyType, typename ValueType, typename SetAllocator, typename KeyFuncs> struct TIsTMap<const          TMap<KeyType, ValueType, SetAllocator, KeyFuncs>> { enum { Value = true }; };
+template <typename KeyType, typename ValueType, typename SetAllocator, typename KeyFuncs> struct TIsTMap<      volatile TMap<KeyType, ValueType, SetAllocator, KeyFuncs>> { enum { Value = true }; };
+template <typename KeyType, typename ValueType, typename SetAllocator, typename KeyFuncs> struct TIsTMap<const volatile TMap<KeyType, ValueType, SetAllocator, KeyFuncs>> { enum { Value = true }; };
+
 class FScriptMap : public TScriptMap<FDefaultSetAllocator, FScriptMap>
 {
 	using Super = TScriptMap<FDefaultSetAllocator, FScriptMap>;
diff --git a/Engine/Source/Runtime/Core/Public/Containers/Set.h b/Engine/Source/Runtime/Core/Public/Containers/Set.h
--- a/Engine/Source/Runtime/Core/Public/Containers/Set.h
+++ b/Engine/Source/Runtime/Core/Public/Containers/Set.h
@@ -2112,6 +2112,16 @@ struct TIsZeroConstructType<TScriptSet<AllocatorType, InDerivedType>>
 	enum { Value = true };
 };
 
+/**
+ * Traits class which determines whether or not a type is a TSet.
+ */
+template <typename T> struct TIsTSet { enum { Value = false }; };
+
+template <typename ElementType, typename KeyFuncs, typename Allocator> struct TIsTSet<               TSet<ElementType, KeyFuncs, Allocator>> { enum { Value = true }; };
+template <typename ElementType, typename KeyFuncs, typename Allocator> struct TIsTSet<const          TSet<ElementType, KeyFuncs, Allocator>> { enum { Value = true }; };
+template <typename ElementType, typename KeyFuncs, typename Allocator> struct TIsTSet<      volatile TSet<ElementType, KeyFuncs, Allocator>> { enum { Value = true }; };
+template <typename ElementType, typename KeyFuncs, typename Allocator> struct TIsTSet<const volatile TSet<ElementType, KeyFuncs, Allocator>> { enum { Value = true }; };
+
 class FScriptSet : public TScriptSet<FDefaultSetAllocator, FScriptSet>
 {
 	using Super = TScriptSet<FDefaultSetAllocator, FScriptSet>;
