Expose template parameter type names for TSet/TMap and adjust exposed name for Allocator to AllocatorType in TArray.  Split out from review here:

https://p4-swarm.epicgames.net/reviews/15065157

diff --git a/Engine/Source/Runtime/Core/Public/Containers/Array.h b/Engine/Source/Runtime/Core/Public/Containers/Array.h
--- a/Engine/Source/Runtime/Core/Public/Containers/Array.h
+++ b/Engine/Source/Runtime/Core/Public/Containers/Array.h
@@ -284,16 +284,19 @@ namespace UE4Array_Private
  * Caution: as noted below some methods are not safe for element types that require constructors.
  *
  **/
-template<typename InElementType, typename InAllocator>
+template<typename InElementType, typename InAllocatorType>
 class TArray
 {
 	template <typename OtherInElementType, typename OtherAllocator>
 	friend class TArray;
 
 public:
-	typedef typename InAllocator::SizeType SizeType;
+	typedef typename InAllocatorType::SizeType SizeType;
 	typedef InElementType ElementType;
-	typedef InAllocator   Allocator;
+	typedef InAllocatorType AllocatorType;
+
+	// This member type name is deprecated, new code should prefer the "AllocatorType" name instead
+	typedef InAllocatorType Allocator;
 
 	typedef typename TChooseClass<
 		Allocator::NeedsElementType,
diff --git a/Engine/Source/Runtime/Core/Public/Containers/Map.h b/Engine/Source/Runtime/Core/Public/Containers/Map.h
--- a/Engine/Source/Runtime/Core/Public/Containers/Map.h
+++ b/Engine/Source/Runtime/Core/Public/Containers/Map.h
@@ -1103,8 +1103,8 @@ template <typename AllocatorType, typename InDerivedType = void>
 class TScriptMap;
 
 /** A TMapBase specialization that only allows a single value associated with each key.*/
-template<typename KeyType, typename ValueType, typename SetAllocator /*= FDefaultSetAllocator*/, typename KeyFuncs /*= TDefaultMapHashableKeyFuncs<KeyType,ValueType,false>*/>
-class TMap : public TSortableMapBase<KeyType, ValueType, SetAllocator, KeyFuncs>
+template<typename InKeyType, typename InValueType, typename SetAllocator /*= FDefaultSetAllocator*/, typename KeyFuncs /*= TDefaultMapHashableKeyFuncs<KeyType,ValueType,false>*/>
+class TMap : public TSortableMapBase<InKeyType, InValueType, SetAllocator, KeyFuncs>
 {
 	friend struct TContainerTraits<TMap>;
 
@@ -1114,6 +1114,11 @@ class TMap : public TSortableMapBase<KeyType, ValueType, SetAllocator, KeyFuncs>
 	static_assert(!KeyFuncs::bAllowDuplicateKeys, "TMap cannot be instantiated with a KeyFuncs which allows duplicate keys");
 
 public:
+	typedef InKeyType      KeyType;
+	typedef InValueType    ValueType;
+	typedef SetAllocator   SetAllocatorType;
+	typedef KeyFuncs       KeyFuncsType;
+
 	typedef TSortableMapBase<KeyType, ValueType, SetAllocator, KeyFuncs> Super;
 	typedef typename Super::KeyInitType KeyInitType;
 	typedef typename Super::KeyConstPointerType KeyConstPointerType;
diff --git a/Engine/Source/Runtime/Core/Public/Containers/Set.h b/Engine/Source/Runtime/Core/Public/Containers/Set.h
--- a/Engine/Source/Runtime/Core/Public/Containers/Set.h
+++ b/Engine/Source/Runtime/Core/Public/Containers/Set.h
@@ -286,6 +286,10 @@ class TSet
 public:
 	static const bool SupportsFreezeMemoryImage = TAllocatorTraits<Allocator>::SupportsFreezeMemoryImage && THasTypeLayout<InElementType>::Value;
 
+	typedef InElementType ElementType;
+	typedef KeyFuncs    KeyFuncsType;
+	typedef Allocator   AllocatorType;
+
 private:
 	friend struct TContainerTraits<TSet>;
 
@@ -298,8 +302,6 @@ private:
 	typedef TSetElement<InElementType> SetElementType;
 
 public:
-	typedef InElementType ElementType;
-
 	/** Initialization constructor. */
 	FORCEINLINE TSet()
 	:	HashSize(0)
