Include <type_traits> in Array.h, Set.h, and Map.h to adhere to "include what you use" since the preceding CL to use std::enable_if_t instead of TEnableIf.

diff --git a/Engine/Source/Runtime/Core/Public/Containers/Array.h b/Engine/Source/Runtime/Core/Public/Containers/Array.h
--- a/Engine/Source/Runtime/Core/Public/Containers/Array.h
+++ b/Engine/Source/Runtime/Core/Public/Containers/Array.h
@@ -24,6 +24,7 @@
 #include "Templates/ChooseClass.h"
 #include "Templates/Sorting.h"
 #include "Templates/AlignmentTemplates.h"
+#include <type_traits>
 
 
 #if UE_BUILD_SHIPPING || UE_BUILD_TEST
diff --git a/Engine/Source/Runtime/Core/Public/Containers/Map.h b/Engine/Source/Runtime/Core/Public/Containers/Map.h
--- a/Engine/Source/Runtime/Core/Public/Containers/Map.h
+++ b/Engine/Source/Runtime/Core/Public/Containers/Map.h
@@ -15,6 +15,7 @@
 #include "Templates/Tuple.h"
 #include "Templates/UnrealTemplate.h"
 #include "Templates/UnrealTypeTraits.h"
+#include <type_traits>
 
 #define ExchangeB(A,B) {bool T=A; A=B; B=T;}
 
diff --git a/Engine/Source/Runtime/Core/Public/Containers/Set.h b/Engine/Source/Runtime/Core/Public/Containers/Set.h
--- a/Engine/Source/Runtime/Core/Public/Containers/Set.h
+++ b/Engine/Source/Runtime/Core/Public/Containers/Set.h
@@ -21,6 +21,7 @@
 #include "Serialization/StructuredArchive.h"
 #include "Serialization/MemoryImageWriter.h"
 #include "ContainersFwd.h"
+#include <type_traits>
 
 /**
  * The base KeyFuncs type with some useful definitions for all KeyFuncs; meant to be derived from instead of used directly.
