U5 - cleanup TaskDelegates, rename Task function for clarity and enable the new backend to use more than 2 threads.

diff --git a/Engine/Source/Runtime/Core/Private/GenericPlatform/GenericPlatformMisc.cpp b/Engine/Source/Runtime/Core/Private/GenericPlatform/GenericPlatformMisc.cpp
--- a/Engine/Source/Runtime/Core/Private/GenericPlatform/GenericPlatformMisc.cpp
+++ b/Engine/Source/Runtime/Core/Private/GenericPlatform/GenericPlatformMisc.cpp
@@ -1181,7 +1181,9 @@ int32 FGenericPlatformMisc::NumberOfCoresIncludingHyperthreads()
 int32 FGenericPlatformMisc::NumberOfWorkerThreadsToSpawn()
 {
 	static int32 MaxGameThreads = 4;
-	static int32 MaxThreads = 16;
+
+	extern CORE_API int32 GUseNewTaskBackend;
+	int32 MaxThreads = GUseNewTaskBackend ? INT32_MAX : 16;
 
 	int32 NumberOfCores = FPlatformMisc::NumberOfCores();
 	int32 MaxWorkerThreadsWanted = (IsRunningGame() || IsRunningDedicatedServer() || IsRunningClientOnly()) ? MaxGameThreads : MaxThreads;
diff --git a/Engine/Source/Runtime/Core/Private/Windows/WindowsPlatformMisc.cpp b/Engine/Source/Runtime/Core/Private/Windows/WindowsPlatformMisc.cpp
--- a/Engine/Source/Runtime/Core/Private/Windows/WindowsPlatformMisc.cpp
+++ b/Engine/Source/Runtime/Core/Private/Windows/WindowsPlatformMisc.cpp
@@ -1799,7 +1799,9 @@ const TCHAR* FWindowsPlatformMisc::GetPlatformFeaturesModuleName()
 int32 FWindowsPlatformMisc::NumberOfWorkerThreadsToSpawn()
 {
 	static int32 MaxServerWorkerThreads = 4;
-	static int32 MaxWorkerThreads = 26;
+
+	extern CORE_API int32 GUseNewTaskBackend;
+	int32 MaxWorkerThreads = GUseNewTaskBackend ? INT32_MAX : 26;
 
 	int32 NumberOfCores = FWindowsPlatformMisc::NumberOfCores();
 	int32 NumberOfCoresIncludingHyperthreads = FWindowsPlatformMisc::NumberOfCoresIncludingHyperthreads();
diff --git a/Engine/Source/Runtime/Core/Public/Async/Fundamental/Scheduler.h b/Engine/Source/Runtime/Core/Public/Async/Fundamental/Scheduler.h
--- a/Engine/Source/Runtime/Core/Public/Async/Fundamental/Scheduler.h
+++ b/Engine/Source/Runtime/Core/Public/Async/Fundamental/Scheduler.h
@@ -56,7 +56,7 @@ namespace LowLevelTasks
 		//try to cancel the task and launching it if the task was in ready state
 		//you still need to wait for task completion before recycling the handle
 		//you can alternatively use FTask::TryCancel if you want to launch the task manually
-		inline bool TryCancelAndLaunch(FTask& Task, EQueuePreference QueuePreference = EQueuePreference::DefaultPreference);		
+		inline bool TryCancelAndLaunchContinuation(FTask& Task, EQueuePreference QueuePreference = EQueuePreference::DefaultPreference);		
 
 		//tries to do some work until the Task is completed
 		template<typename TaskType>
@@ -121,9 +121,9 @@ namespace LowLevelTasks
 		return FScheduler::Get().TryLaunch(Task, QueuePreference);
 	}
 
-	FORCEINLINE_DEBUGGABLE bool TryCancelAndLaunch(FTask& Task, EQueuePreference QueuePreference = EQueuePreference::DefaultPreference)
+	FORCEINLINE_DEBUGGABLE bool TryCancelAndLaunchContinuation(FTask& Task, EQueuePreference QueuePreference = EQueuePreference::DefaultPreference)
 	{
-		return FScheduler::Get().TryCancelAndLaunch(Task, QueuePreference);
+		return FScheduler::Get().TryCancelAndLaunchContinuation(Task, QueuePreference);
 	}
 
 	FORCEINLINE_DEBUGGABLE void BusyWaitForTask(const FTask& Task)
@@ -156,7 +156,7 @@ namespace LowLevelTasks
 		return false;
 	}
 
-	inline bool FScheduler::TryCancelAndLaunch(FTask& Task, EQueuePreference QueuePreference)
+	inline bool FScheduler::TryCancelAndLaunchContinuation(FTask& Task, EQueuePreference QueuePreference)
 	{
 		bool WasCanceled = Task.TryCancel();
 		if(WasCanceled && Task.TryPrepareLaunch())
diff --git a/Engine/Source/Runtime/Core/Public/Async/Fundamental/TaskDelegate.h b/Engine/Source/Runtime/Core/Public/Async/Fundamental/TaskDelegate.h
--- a/Engine/Source/Runtime/Core/Public/Async/Fundamental/TaskDelegate.h
+++ b/Engine/Source/Runtime/Core/Public/Async/Fundamental/TaskDelegate.h
@@ -3,6 +3,7 @@
 #pragma once
 #include "Templates/IsInvocable.h"
 #include "Misc/ScopeExit.h"
+#include "Misc/Launder.h"
 #include <type_traits>
 
 namespace LowLevelTasks
@@ -32,24 +33,83 @@ namespace LowLevelTasks
 		using ThisClass = TTaskDelegate<TotalSize, ReturnType>;
 		struct TTaskDelegateBase
 		{
-			virtual void Move(TTaskDelegateBase&, void*, void*, uint32) {check(false); };
-			virtual ReturnType Call(void*) const { check(false); return TTaskDelegate_Impl::MakeDummyValue<ReturnType>(); };
-			virtual ReturnType CallAndMove(ThisClass&, void*, uint32) {check(false); return TTaskDelegate_Impl::MakeDummyValue<ReturnType>(); };
-			virtual void Destroy(void*) { check(false); };
-			virtual bool IsHeapAllocated() const { check(false); return false; }
-			virtual bool IsSet() const { check(false); return false; }
-			virtual uint32 DelegateSize() const { check(false); return 0; }
+			virtual void Move(TTaskDelegateBase&, void*, void*, uint32) 
+			{
+				checkNoEntry(); //if we get to this place than the compiler is optimizing our vtable lookup and ignores our undefined/shallow copy of the memory
+			};
+
+			virtual ReturnType Call(void*) const 
+			{ 
+				checkNoEntry(); //if we get to this place than the compiler is optimizing our vtable lookup and ignores our undefined/shallow copy of the memory 
+				return TTaskDelegate_Impl::MakeDummyValue<ReturnType>(); 
+			};
+
+			virtual ReturnType CallAndMove(ThisClass&, void*, uint32) 
+			{
+				checkNoEntry(); //if we get to this place than the compiler is optimizing our vtable lookup and ignores our undefined/shallow copy of the memory 
+				return TTaskDelegate_Impl::MakeDummyValue<ReturnType>(); 
+			};
+
+			virtual void Destroy(void*) 
+			{ 
+				checkNoEntry(); //if we get to this place than the compiler is optimizing our vtable lookup and ignores our undefined/shallow copy of the memory 
+			};
+
+			virtual bool IsHeapAllocated() const 
+			{ 
+				checkNoEntry(); //if we get to this place than the compiler is optimizing our vtable lookup and ignores our undefined/shallow copy of the memory
+				return false; 
+			}
+
+			virtual bool IsSet() const 
+			{ 
+				checkNoEntry(); //if we get to this place than the compiler is optimizing our vtable lookup and ignores our undefined/shallow copy of the memory 
+				return false; 
+			}
+
+			virtual uint32 DelegateSize() const 
+			{ 
+				checkNoEntry(); //if we get to this place than the compiler is optimizing our vtable lookup and ignores our undefined/shallow copy of the memory 
+				return 0; 
+			}
 		};
 
 		struct TTaskDelegateDummy final : TTaskDelegateBase
 		{
-			void Move(TTaskDelegateBase&, void*, void*, uint32) override {};
-			ReturnType Call(void*) const override { check(false); return TTaskDelegate_Impl::MakeDummyValue<ReturnType>(); };
-			ReturnType CallAndMove(ThisClass&, void*, uint32) override { check(false); return TTaskDelegate_Impl::MakeDummyValue<ReturnType>(); };
-			void Destroy(void*) override {};
-			bool IsHeapAllocated() const override { return false; }
-			bool IsSet() const override { return false; }
-			uint32 DelegateSize() const override { return 0; }
+			void Move(TTaskDelegateBase&, void*, void*, uint32) override 
+			{
+			}
+
+			ReturnType Call(void*) const override 
+			{ 
+				checkf(false, TEXT("trying to Call a dummy TaskDelegate"));
+				return TTaskDelegate_Impl::MakeDummyValue<ReturnType>(); 
+			}
+
+			ReturnType CallAndMove(ThisClass&, void*, uint32) override 
+			{ 
+				checkf(false, TEXT("trying to Call a dummy TaskDelegate"));
+				return TTaskDelegate_Impl::MakeDummyValue<ReturnType>(); 
+			}
+
+			void Destroy(void*) override 
+			{
+			}
+
+			bool IsHeapAllocated() const override 
+			{ 
+				return false; 
+			}
+
+			bool IsSet() const override 
+			{ 
+				return false; 
+			}
+
+			uint32 DelegateSize() const override 
+			{ 
+				return 0; 
+			}
 		};
 
 		template<typename TCallableType, bool HeapAllocated>
@@ -58,13 +118,13 @@ namespace LowLevelTasks
 		template<typename TCallableType>
 		struct TTaskDelegateImpl<TCallableType, false> final : TTaskDelegateBase
 		{
-			template<typename TCallable>
-			inline TTaskDelegateImpl(TCallable&& Callable, void* InlineData)
+			template<typename CallableT>
+			inline TTaskDelegateImpl(CallableT&& Callable, void* InlineData)
 			{
 				static_assert(TIsInvocable<TCallableType>::Value, "TCallableType is not invocable");
 				static_assert(TIsSame<ReturnType, decltype(Callable())>::Value, "TCallableType return type does not match");
 				static_assert(sizeof(TTaskDelegateImpl<TCallableType, false>) == sizeof(TTaskDelegateBase), "Size must match the Baseclass");
-				new (InlineData) TCallableType(Forward<TCallable>(Callable));
+				new (InlineData) TCallableType(Forward<CallableT>(Callable));
 			}
 
 			inline void Move(TTaskDelegateBase& DstWrapper, void* DstData, void* SrcData, uint32 DestInlineSize) override
@@ -84,7 +144,7 @@ namespace LowLevelTasks
 			inline ReturnType Call(void* InlineData) const override
 			{
 				TCallableType* LocalPtr = reinterpret_cast<TCallableType*>(InlineData);
-				return (*LocalPtr)();
+				return Invoke(*LocalPtr);
 			}
 
 			ReturnType CallAndMove(ThisClass& Destination, void* InlineData, uint32 DestInlineSize) override
@@ -128,14 +188,14 @@ namespace LowLevelTasks
 			}
 
 		public:
-			template<typename TCallable>
-			inline TTaskDelegateImpl(TCallable&& Callable, void* InlineData)
+			template<typename CallableT>
+			inline TTaskDelegateImpl(CallableT&& Callable, void* InlineData)
 			{
 				static_assert(TIsInvocable<TCallableType>::Value, "TCallableType is not invocable");
 				static_assert(TIsSame<ReturnType, decltype(Callable())>::Value, "TCallableType return type does not match");
 				static_assert(sizeof(TTaskDelegateImpl<TCallableType, true>) == sizeof(TTaskDelegateBase), "Size must match the Baseclass");
 				TCallableType** HeapPtr = reinterpret_cast<TCallableType**>(InlineData);
-				*HeapPtr = new TCallableType(Forward<TCallable>(Callable));
+				*HeapPtr = new TCallableType(Forward<CallableT>(Callable));
 			}
 
 			inline void Move(TTaskDelegateBase& DstWrapper, void* DstData, void* SrcData, uint32 DestInlineSize) override
@@ -147,7 +207,7 @@ namespace LowLevelTasks
 			inline ReturnType Call(void* InlineData) const override
 			{
 				TCallableType* HeapPtr = reinterpret_cast<TCallableType*>(*reinterpret_cast<void* const*>(InlineData));
-				return (*HeapPtr)();
+				return Invoke(*HeapPtr);
 			}
 
 			ReturnType CallAndMove(ThisClass& Destination, void* InlineData, uint32 DestInlineSize) override
@@ -195,38 +255,38 @@ namespace LowLevelTasks
 		template<uint32 SourceTotalSize>
 		TTaskDelegate(TTaskDelegate<SourceTotalSize, ReturnType>&& Other)
 		{
-			static_cast<TTaskDelegateBase*>(&Other.CallableWrapper)->Move(CallableWrapper, InlineStorage, Other.InlineStorage, InlineStorageSize);
+			Other.GetWrapper()->Move(CallableWrapper, InlineStorage, Other.InlineStorage, InlineStorageSize);
 		}
 
-		template<typename TCallable>
-		TTaskDelegate(TCallable&& Callable)
+		template<typename CallableT>
+		TTaskDelegate(CallableT&& Callable)
 		{
-			using TCallableType = typename TRemoveConst<typename TRemoveReference<TCallable>::Type>::Type;
+			using TCallableType = std::decay_t<CallableT>;
 			if ((sizeof(TCallableType) <= InlineStorageSize) && ((uintptr_t(InlineStorageSize) % alignof(TCallableType)) == 0))
 			{
-				new (&CallableWrapper) TTaskDelegateImpl<TCallableType, false>(Forward<TCallable>(Callable), InlineStorage);
+				new (&CallableWrapper) TTaskDelegateImpl<TCallableType, false>(Forward<CallableT>(Callable), InlineStorage);
 			}
 			else
 			{
-				new (&CallableWrapper) TTaskDelegateImpl<TCallableType, true>(Forward<TCallable>(Callable), InlineStorage);
+				new (&CallableWrapper) TTaskDelegateImpl<TCallableType, true>(Forward<CallableT>(Callable), InlineStorage);
 			}
 		}
 
 		~TTaskDelegate()
 		{
-			static_cast<TTaskDelegateBase*>(&CallableWrapper)->Destroy(InlineStorage);
+			GetWrapper()->Destroy(InlineStorage);
 		}
 
 		ReturnType operator()() const
 		{
-			return static_cast<const TTaskDelegateBase*>(&CallableWrapper)->Call(InlineStorage);
+			return GetWrapper()->Call(InlineStorage);
 		}
 
 		template<uint32 DestTotalSize>
 		ReturnType CallAndMove(TTaskDelegate<DestTotalSize, ReturnType>& Destination)
 		{
 			checkSlow(!Destination.IsSet());
-			return static_cast<TTaskDelegateBase*>(&CallableWrapper)->CallAndMove(Destination, InlineStorage, TTaskDelegate<DestTotalSize, ReturnType>::InlineStorageSize);
+			return GetWrapper()->CallAndMove(Destination, InlineStorage, TTaskDelegate<DestTotalSize, ReturnType>::InlineStorageSize);
 		}
 
 		template<uint32 SourceTotalSize>
@@ -235,51 +295,61 @@ namespace LowLevelTasks
 		template<uint32 SourceTotalSize>
 		ThisClass& operator= (TTaskDelegate<SourceTotalSize, ReturnType>&& Other)
 		{
-			static_cast<TTaskDelegateBase*>(&CallableWrapper)->Destroy(InlineStorage);
-			static_cast<TTaskDelegateBase*>(&Other.CallableWrapper)->Move(CallableWrapper, InlineStorage, Other.InlineStorage, InlineStorageSize);
+			GetWrapper()->Destroy(InlineStorage);
+			Other.GetWrapper()->Move(CallableWrapper, InlineStorage, Other.InlineStorage, InlineStorageSize);
 			return *this;
 		}
 
-		template<typename TCallable>
-		ThisClass& operator= (TCallable&& Callable)
+		template<typename CallableT>
+		ThisClass& operator= (CallableT&& Callable)
 		{
-			using TCallableType = typename TRemoveConst<typename TRemoveReference<TCallable>::Type>::Type;
-			static_cast<TTaskDelegateBase*>(&CallableWrapper)->Destroy(InlineStorage);
+			using TCallableType = std::decay_t<CallableT>;
+			GetWrapper()->Destroy(InlineStorage);
 			if ((sizeof(TCallableType) <= InlineStorageSize) && ((uintptr_t(InlineStorageSize) % alignof(TCallableType)) == 0))
 			{
-				new (&CallableWrapper) TTaskDelegateImpl<TCallableType, false>(Forward<TCallable>(Callable), InlineStorage);
+				new (&CallableWrapper) TTaskDelegateImpl<TCallableType, false>(Forward<CallableT>(Callable), InlineStorage);
 			}
 			else
 			{
-				new (&CallableWrapper) TTaskDelegateImpl<TCallableType, true>(Forward<TCallable>(Callable), InlineStorage);
+				new (&CallableWrapper) TTaskDelegateImpl<TCallableType, true>(Forward<CallableT>(Callable), InlineStorage);
 			}
 			return *this;
 		}
 
 		void Destroy()
 		{
-			static_cast<TTaskDelegateBase*>(&CallableWrapper)->Destroy(InlineStorage);
+			GetWrapper()->Destroy(InlineStorage);
 			new (&CallableWrapper) TTaskDelegateDummy();
 		}
 
 		bool IsHeapAllocated() const
 		{
-			return static_cast<const TTaskDelegateBase*>(&CallableWrapper)->IsHeapAllocated();
+			return GetWrapper()->IsHeapAllocated();
 		}
 
 		bool IsSet() const
 		{
-			return static_cast<const TTaskDelegateBase*>(&CallableWrapper)->IsSet();
+			return GetWrapper()->IsSet();
 		}
 
 		uint32 DelegateSize() const  
 		{ 
-			return static_cast<const TTaskDelegateBase*>(&CallableWrapper)->DelegateSize(); 
+			return GetWrapper()->DelegateSize(); 
 		}
 
 	private:
 		static constexpr uint32 InlineStorageSize = TotalSize - sizeof(TTaskDelegateBase);
 		mutable char InlineStorage[InlineStorageSize];
 		TTaskDelegateBase CallableWrapper;
+
+		TTaskDelegateBase* GetWrapper()
+		{
+			return UE_LAUNDER(static_cast<TTaskDelegateBase*>(&CallableWrapper));
+		}
+
+		const TTaskDelegateBase* GetWrapper() const
+		{
+			return UE_LAUNDER(static_cast<const TTaskDelegateBase*>(&CallableWrapper));
+		}
 	};
 }
\ No newline at end of file
diff --git a/Engine/Source/Runtime/Core/Public/Misc/Launder.h b/Engine/Source/Runtime/Core/Public/Misc/Launder.h
new file mode 100644
--- /dev/null
+++ b/Engine/Source/Runtime/Core/Public/Misc/Launder.h
@@ -0,0 +1,9 @@
+// Copyright Epic Games, Inc. All Rights Reserved.
+#pragma once
+
+#if __cplusplus >= 201703L
+#include <new>
+#define UE_LAUNDER(x) std::launder(x)
+#else
+#define UE_LAUNDER(x) (x)
+#endif
