Mark TArray::Allocator to be deprecated with UE_DEPRECATED (instead of just by comment).  Performed a blanket replacement of the term Allocator (whole word) in Array.h with AllocatorType.  This affected TArray as well as other templates in Array.h

diff --git a/Engine/Source/Runtime/Core/Public/Containers/Array.h b/Engine/Source/Runtime/Core/Public/Containers/Array.h
--- a/Engine/Source/Runtime/Core/Public/Containers/Array.h
+++ b/Engine/Source/Runtime/Core/Public/Containers/Array.h
@@ -255,10 +255,10 @@ namespace UE4Array_Private
 	template <typename FromArrayType, typename ToArrayType>
 	struct TCanMoveTArrayPointersBetweenArrayTypes
 	{
-		typedef typename FromArrayType::Allocator   FromAllocatorType;
-		typedef typename ToArrayType  ::Allocator   ToAllocatorType;
-		typedef typename FromArrayType::ElementType FromElementType;
-		typedef typename ToArrayType  ::ElementType ToElementType;
+		typedef typename FromArrayType::AllocatorType FromAllocatorType;
+		typedef typename ToArrayType  ::AllocatorType ToAllocatorType;
+		typedef typename FromArrayType::ElementType   FromElementType;
+		typedef typename ToArrayType  ::ElementType   ToElementType;
 
 		enum
 		{
@@ -296,13 +296,13 @@ public:
 	typedef InElementType ElementType;
 	typedef InAllocatorType AllocatorType;
 
-	// This member type name is deprecated, new code should prefer the "AllocatorType" name instead
+	UE_DEPRECATED(5.0, "TArray::Allocator type is deprecated, please use TArray::AllocatorType instead.")
 	typedef InAllocatorType Allocator;
 
 	typedef typename TChooseClass<
-		Allocator::NeedsElementType,
-		typename Allocator::template ForElementType<ElementType>,
-		typename Allocator::ForAnyElementType
+		AllocatorType::NeedsElementType,
+		typename AllocatorType::template ForElementType<ElementType>,
+		typename AllocatorType::ForAnyElementType
 	>::Result ElementAllocatorType;
 
 	static_assert(TIsSigned<SizeType>::Value, "TArray only supports signed index types");
@@ -396,12 +396,12 @@ public:
 	 * Assignment operator. First deletes all currently contained elements
 	 * and then copies from other array.
 	 *
-	 * Allocator changing version.
+	 * AllocatorType changing version.
 	 *
 	 * @param Other The source array to assign from.
 	 */
-	template<typename OtherAllocator>
-	TArray& operator=(const TArray<ElementType, OtherAllocator>& Other)
+	template<typename OtherAllocatorType>
+	TArray& operator=(const TArray<ElementType, OtherAllocatorType>& Other)
 	{
 		DestructItems(GetData(), ArrayNum);
 		CopyToEmpty(Other.GetData(), Other.Num(), ArrayMax, 0);
@@ -432,17 +432,17 @@ private:
 	template <
 		typename FromArrayType,
 		typename ToArrayType,
-		std::enable_if_t<TCanMoveBetweenAllocators<typename FromArrayType::Allocator, typename ToArrayType::Allocator>::Value>* = nullptr
+		std::enable_if_t<TCanMoveBetweenAllocators<typename FromArrayType::AllocatorType, typename ToArrayType::AllocatorType>::Value>* = nullptr
 	>
 	static FORCEINLINE void MoveAllocatorToEmpty(FromArrayType& FromArray, ToArrayType& ToArray)
 	{
-		ToArray.AllocatorInstance.template MoveToEmptyFromOtherAllocator<typename FromArrayType::Allocator>(FromArray.AllocatorInstance);
+		ToArray.AllocatorInstance.template MoveToEmptyFromOtherAllocator<typename FromArrayType::AllocatorType>(FromArray.AllocatorInstance);
 	}
 
 	template <
 		typename FromArrayType,
 		typename ToArrayType,
-		std::enable_if_t<!TCanMoveBetweenAllocators<typename FromArrayType::Allocator, typename ToArrayType::Allocator>::Value>* = nullptr
+		std::enable_if_t<!TCanMoveBetweenAllocators<typename FromArrayType::AllocatorType, typename ToArrayType::AllocatorType>::Value>* = nullptr
 	>
 	static FORCEINLINE void MoveAllocatorToEmpty(FromArrayType& FromArray, ToArrayType& ToArray)
 	{
@@ -461,8 +461,8 @@ private:
 	template <typename FromArrayType, typename ToArrayType>
 	static FORCEINLINE std::enable_if_t<UE4Array_Private::TCanMoveTArrayPointersBetweenArrayTypes<FromArrayType, ToArrayType>::Value> MoveOrCopy(ToArrayType& ToArray, FromArrayType& FromArray, SizeType PrevMax)
 	{
-		using FromAllocatorType = typename FromArrayType::Allocator;
-		using ToAllocatorType   = typename ToArrayType::Allocator;
+		using FromAllocatorType = typename FromArrayType::AllocatorType;
+		using ToAllocatorType   = typename ToArrayType::AllocatorType;
 
 #if PLATFORM_COMPILER_HAS_IF_CONSTEXPR
 		if constexpr (TCanMoveBetweenAllocators<FromAllocatorType, ToAllocatorType>::Value)
@@ -567,7 +567,7 @@ public:
 	 *                   at the end of the array in the number of elements.
 	 */
 	template <typename OtherElementType>
-	TArray(TArray<OtherElementType, Allocator>&& Other, SizeType ExtraSlack)
+	TArray(TArray<OtherElementType, AllocatorType>&& Other, SizeType ExtraSlack)
 	{
 		// We don't implement move semantics for general OtherAllocators, as there's no way
 		// to tell if they're compatible with the current one.  Probably going to be a pretty
@@ -674,7 +674,7 @@ public:
 		CheckInvariants();
 
 		// Template property, branch will be optimized out
-		if (Allocator::RequireRangeCheck)
+		if (AllocatorType::RequireRangeCheck)
 		{
 			checkf((Index >= 0) & (Index < ArrayNum),TEXT("Array index out of bounds: %i from an array of size %i"),Index,ArrayNum); // & for one branch
 		}
@@ -1814,13 +1814,13 @@ public:
 	/**
 	 * Appends the specified array to this array.
 	 *
-	 * Allocator changing version.
+	 * AllocatorType changing version.
 	 *
 	 * @param Source The array to append.
 	 * @see Add, Insert
 	 */
-	template <typename OtherElementType, typename OtherAllocator>
-	void Append(const TArray<OtherElementType, OtherAllocator>& Source)
+	template <typename OtherElementType, typename OtherAllocatorType>
+	void Append(const TArray<OtherElementType, OtherAllocatorType>& Source)
 	{
 		check((void*)this != (void*)&Source);
 
@@ -2119,7 +2119,7 @@ public:
 		typename AliasElementType = ElementType,
 		std::enable_if_t<TIsContainerElementTypeReinterpretable<AliasElementType>::Value>* = nullptr
 		>
-	operator TArray<typename TContainerElementTypeCompatibility<AliasElementType>::ReinterpretType, Allocator>& ()
+	operator TArray<typename TContainerElementTypeCompatibility<AliasElementType>::ReinterpretType, AllocatorType>& ()
 	{
 		using ElementCompat = TContainerElementTypeCompatibility<ElementType>;
 		ElementCompat::ReinterpretRange(begin(), end());
@@ -2131,7 +2131,7 @@ public:
 		typename AliasElementType = ElementType,
 		typename std::enable_if_t<TIsContainerElementTypeReinterpretable<AliasElementType>::Value>* = nullptr
 		>
-	operator const TArray<typename TContainerElementTypeCompatibility<AliasElementType>::ReinterpretType, Allocator>& () const
+	operator const TArray<typename TContainerElementTypeCompatibility<AliasElementType>::ReinterpretType, AllocatorType>& () const
 	{
 		using ElementCompat = TContainerElementTypeCompatibility<ElementType>;
 		ElementCompat::ReinterpretRange(begin(), end());
@@ -2148,7 +2148,7 @@ public:
 		typename AliasElementType = ElementType,
 		typename std::enable_if_t<TIsContainerElementTypeCopyable<AliasElementType>::Value>* = nullptr
 	>
-	TArray& operator=(TArray<typename TContainerElementTypeCompatibility<ElementType>::CopyFromOtherType, Allocator>&& Other)
+	TArray& operator=(TArray<typename TContainerElementTypeCompatibility<ElementType>::CopyFromOtherType, AllocatorType>&& Other)
 	{
 		TContainerElementTypeCompatibility<ElementType>::CopyingFromOtherType();
 		DestructItems(GetData(), ArrayNum);
@@ -2780,7 +2780,7 @@ private:
 
 protected:
 
-	template<typename ElementType, typename Allocator>
+	template<typename ElementType, typename AllocatorType>
 	friend class TIndirectArray;
 
 	ElementAllocatorType AllocatorInstance;
@@ -2832,25 +2832,25 @@ private:
 public:
 	void WriteMemoryImage(FMemoryImageWriter& Writer) const
 	{
-		static const bool bSupportsFreezeMemoryImage = TAllocatorTraits<Allocator>::SupportsFreezeMemoryImage && THasTypeLayout<ElementType>::Value;
+		static const bool bSupportsFreezeMemoryImage = TAllocatorTraits<AllocatorType>::SupportsFreezeMemoryImage && THasTypeLayout<ElementType>::Value;
 		TSupportsFreezeMemoryImageHelper<bSupportsFreezeMemoryImage>::WriteMemoryImage(Writer, *this);
 	}
 
 	void CopyUnfrozen(const FMemoryUnfreezeContent& Context, void* Dst) const
 	{
-		static const bool bSupportsFreezeMemoryImage = TAllocatorTraits<Allocator>::SupportsFreezeMemoryImage && THasTypeLayout<ElementType>::Value;
+		static const bool bSupportsFreezeMemoryImage = TAllocatorTraits<AllocatorType>::SupportsFreezeMemoryImage && THasTypeLayout<ElementType>::Value;
 		TSupportsFreezeMemoryImageHelper<bSupportsFreezeMemoryImage>::CopyUnfrozen(Context, *this, Dst);
 	}
 
 	static void AppendHash(const FPlatformTypeLayoutParameters& LayoutParams, FSHA1& Hasher)
 	{
-		static const bool bSupportsFreezeMemoryImage = TAllocatorTraits<Allocator>::SupportsFreezeMemoryImage && THasTypeLayout<ElementType>::Value;
+		static const bool bSupportsFreezeMemoryImage = TAllocatorTraits<AllocatorType>::SupportsFreezeMemoryImage && THasTypeLayout<ElementType>::Value;
 		TSupportsFreezeMemoryImageHelper<bSupportsFreezeMemoryImage>::AppendHash(LayoutParams, Hasher);
 	}
 
 	void ToString(const FPlatformTypeLayoutParameters& LayoutParams, FMemoryToStringContext& OutContext) const
 	{
-		static const bool bSupportsFreezeMemoryImage = TAllocatorTraits<Allocator>::SupportsFreezeMemoryImage && THasTypeLayout<ElementType>::Value;
+		static const bool bSupportsFreezeMemoryImage = TAllocatorTraits<AllocatorType>::SupportsFreezeMemoryImage && THasTypeLayout<ElementType>::Value;
 		TSupportsFreezeMemoryImageHelper<bSupportsFreezeMemoryImage>::ToString(LayoutParams, OutContext, *this);
 	}
 
@@ -3174,21 +3174,21 @@ namespace Freeze
 
 DECLARE_TEMPLATE_INTRINSIC_TYPE_LAYOUT((template <typename T, typename AllocatorType>), (TArray<T, AllocatorType>));
 
-template <typename InElementType, typename Allocator>
-struct TIsZeroConstructType<TArray<InElementType, Allocator>>
+template <typename InElementType, typename AllocatorType>
+struct TIsZeroConstructType<TArray<InElementType, AllocatorType>>
 {
-	enum { Value = TAllocatorTraits<Allocator>::IsZeroConstruct };
+	enum { Value = TAllocatorTraits<AllocatorType>::IsZeroConstruct };
 };
 
-template <typename InElementType, typename Allocator>
-struct TContainerTraits<TArray<InElementType, Allocator> > : public TContainerTraitsBase<TArray<InElementType, Allocator> >
+template <typename InElementType, typename AllocatorType>
+struct TContainerTraits<TArray<InElementType, AllocatorType> > : public TContainerTraitsBase<TArray<InElementType, AllocatorType> >
 {
-	static_assert(TAllocatorTraits<Allocator>::SupportsMove, "TArray no longer supports move-unaware allocators");
-	enum { MoveWillEmptyContainer = TAllocatorTraits<Allocator>::SupportsMove };
+	static_assert(TAllocatorTraits<AllocatorType>::SupportsMove, "TArray no longer supports move-unaware allocators");
+	enum { MoveWillEmptyContainer = TAllocatorTraits<AllocatorType>::SupportsMove };
 };
 
-template <typename T, typename Allocator>
-struct TIsContiguousContainer<TArray<T, Allocator>>
+template <typename T, typename AllocatorType>
+struct TIsContiguousContainer<TArray<T, AllocatorType>>
 {
 	enum { Value = true };
 };
@@ -3198,22 +3198,22 @@ struct TIsContiguousContainer<TArray<T, Allocator>>
  */
 template <typename T> struct TIsTArray { enum { Value = false }; };
 
-template <typename InElementType, typename Allocator> struct TIsTArray<               TArray<InElementType, Allocator>> { enum { Value = true }; };
-template <typename InElementType, typename Allocator> struct TIsTArray<const          TArray<InElementType, Allocator>> { enum { Value = true }; };
-template <typename InElementType, typename Allocator> struct TIsTArray<      volatile TArray<InElementType, Allocator>> { enum { Value = true }; };
-template <typename InElementType, typename Allocator> struct TIsTArray<const volatile TArray<InElementType, Allocator>> { enum { Value = true }; };
+template <typename InElementType, typename InAllocatorType> struct TIsTArray<               TArray<InElementType, InAllocatorType>> { enum { Value = true }; };
+template <typename InElementType, typename InAllocatorType> struct TIsTArray<const          TArray<InElementType, InAllocatorType>> { enum { Value = true }; };
+template <typename InElementType, typename InAllocatorType> struct TIsTArray<      volatile TArray<InElementType, InAllocatorType>> { enum { Value = true }; };
+template <typename InElementType, typename InAllocatorType> struct TIsTArray<const volatile TArray<InElementType, InAllocatorType>> { enum { Value = true }; };
 
 
 //
 // Array operator news.
 //
-template <typename T,typename Allocator> void* operator new( size_t Size, TArray<T,Allocator>& Array )
+template <typename T,typename AllocatorType> void* operator new( size_t Size, TArray<T,AllocatorType>& Array )
 {
 	check(Size == sizeof(T));
 	const auto Index = Array.AddUninitialized(1);
 	return &Array[Index];
 }
-template <typename T,typename Allocator> void* operator new( size_t Size, TArray<T,Allocator>& Array, typename TArray<T, Allocator>::SizeType Index )
+template <typename T,typename AllocatorType> void* operator new( size_t Size, TArray<T,AllocatorType>& Array, typename TArray<T, AllocatorType>::SizeType Index )
 {
 	check(Size == sizeof(T));
 	Array.InsertUninitialized(Index);
diff --git a/Engine/Source/Runtime/CoreUObject/Public/UObject/ObjectPtr.h b/Engine/Source/Runtime/CoreUObject/Public/UObject/ObjectPtr.h
--- a/Engine/Source/Runtime/CoreUObject/Public/UObject/ObjectPtr.h
+++ b/Engine/Source/Runtime/CoreUObject/Public/UObject/ObjectPtr.h
@@ -491,7 +491,7 @@ template <
 decltype(auto) ToRawPtrTArrayUnsafe(ArrayType&& Array)
 {
 	using ArrayElementType         = typename ArrayTypeNoRef::ElementType;
-	using ArrayAllocatorType       = typename ArrayTypeNoRef::Allocator;
+	using ArrayAllocatorType       = typename ArrayTypeNoRef::AllocatorType;
 	using RawPointerType           = typename ObjectPtr_Private::TRawPointerType<ArrayElementType>::Type;
 	using QualifiedRawPointerType  = typename TCopyQualifiersFromTo<ArrayElementType, RawPointerType>::Type;
 	using NewArrayType             = TArray<QualifiedRawPointerType, ArrayAllocatorType>;
