PR #7201: Fixed computer name function on Windows (Contributed by [REDACTED_PERSON])

diff --git a/Engine/Source/Runtime/Core/Private/Windows/WindowsPlatformProcess.cpp b/Engine/Source/Runtime/Core/Private/Windows/WindowsPlatformProcess.cpp
--- a/Engine/Source/Runtime/Core/Private/Windows/WindowsPlatformProcess.cpp
+++ b/Engine/Source/Runtime/Core/Private/Windows/WindowsPlatformProcess.cpp
@@ -1046,11 +1046,11 @@ const TCHAR* FWindowsPlatformProcess::ApplicationSettingsDir()
 
 const TCHAR* FWindowsPlatformProcess::ComputerName()
 {
-	static TCHAR Result[256]=TEXT("");
+	static TCHAR Result[MAX_COMPUTERNAME_LENGTH]=TEXT("");
 	if( !Result[0] )
 	{
 		uint32 Size=UE_ARRAY_COUNT(Result);
-		GetComputerName( Result, (::DWORD*)&Size );
+		GetComputerNameEx( ComputerNameDnsHostname, Result, (::DWORD*)&Size );
 	}
 	return Result;
 }
