Added BytesToHexLower and cleaned up BytesToHex and HexToBytes

diff --git a/Engine/Source/Runtime/Core/Private/String/BytesToHex.cpp b/Engine/Source/Runtime/Core/Private/String/BytesToHex.cpp
--- a/Engine/Source/Runtime/Core/Private/String/BytesToHex.cpp
+++ b/Engine/Source/Runtime/Core/Private/String/BytesToHex.cpp
@@ -9,10 +9,10 @@ namespace UE
 namespace String
 {
 
-template <typename CharType>
-void BytesToHexImpl(TArrayView<const uint8> Bytes, CharType* OutHex)
+template <typename CharType, CharType LetterA>
+static inline void BytesToHexImpl(TConstArrayView<uint8> Bytes, CharType* OutHex)
 {
-	auto NibbleToHex = [](uint8 Value) -> CharType { return Value + (Value > 9 ? 'A' - 10 : '0'); };
+	const auto NibbleToHex = [](uint8 Value) -> CharType { return Value + (Value > 9 ? LetterA - 10 : '0'); };
 	const uint8* Data = Bytes.GetData();
 	
 	for (const uint8* DataEnd = Data + Bytes.Num(); Data != DataEnd; ++Data)
@@ -22,26 +22,48 @@ void BytesToHexImpl(TArrayView<const uint8> Bytes, CharType* OutHex)
 	}
 }
 
-void BytesToHex(TArrayView<const uint8> Bytes, WIDECHAR* OutHex)
+void BytesToHex(TConstArrayView<uint8> Bytes, ANSICHAR* OutHex)
 {
-	BytesToHexImpl(Bytes, OutHex);
+	BytesToHexImpl<ANSICHAR, 'A'>(Bytes, OutHex);
 }
 
-void BytesToHex(TArrayView<const uint8> Bytes, ANSICHAR* OutHex)
+void BytesToHex(TConstArrayView<uint8> Bytes, WIDECHAR* OutHex)
 {
-	BytesToHexImpl(Bytes, OutHex);
+	BytesToHexImpl<WIDECHAR, 'A'>(Bytes, OutHex);
 }
 
-void BytesToHex(TArrayView<const uint8> Bytes, FStringBuilderBase& Builder)
+void BytesToHexLower(TConstArrayView<uint8> Bytes, ANSICHAR* OutHex)
+{
+	BytesToHexImpl<ANSICHAR, 'a'>(Bytes, OutHex);
+}
+
+void BytesToHexLower(TConstArrayView<uint8> Bytes, WIDECHAR* OutHex)
+{
+	BytesToHexImpl<WIDECHAR, 'a'>(Bytes, OutHex);
+}
+
+void BytesToHex(TConstArrayView<uint8> Bytes, FAnsiStringBuilderBase& Builder)
+{
+	const int32 Offset = Builder.AddUninitialized(Bytes.Num() * 2);
+	BytesToHexImpl<ANSICHAR, 'A'>(Bytes, GetData(Builder) + Offset);
+}
+
+void BytesToHex(TConstArrayView<uint8> Bytes, FWideStringBuilderBase& Builder)
+{
+	const int32 Offset = Builder.AddUninitialized(Bytes.Num() * 2);
+	BytesToHexImpl<WIDECHAR, 'A'>(Bytes, GetData(Builder) + Offset);
+}
+
+void BytesToHexLower(TConstArrayView<uint8> Bytes, FAnsiStringBuilderBase& Builder)
 {
 	const int32 Offset = Builder.AddUninitialized(Bytes.Num() * 2);
-	BytesToHexImpl(Bytes, GetData(Builder) + Offset);
+	BytesToHexImpl<ANSICHAR, 'a'>(Bytes, GetData(Builder) + Offset);
 }
 
-void BytesToHex(TArrayView<const uint8> Bytes, FAnsiStringBuilderBase& Builder)
+void BytesToHexLower(TConstArrayView<uint8> Bytes, FWideStringBuilderBase& Builder)
 {
 	const int32 Offset = Builder.AddUninitialized(Bytes.Num() * 2);
-	BytesToHexImpl(Bytes, GetData(Builder) + Offset);
+	BytesToHexImpl<WIDECHAR, 'a'>(Bytes, GetData(Builder) + Offset);
 }
 
 }
diff --git a/Engine/Source/Runtime/Core/Private/String/HexToBytes.cpp b/Engine/Source/Runtime/Core/Private/String/HexToBytes.cpp
--- a/Engine/Source/Runtime/Core/Private/String/HexToBytes.cpp
+++ b/Engine/Source/Runtime/Core/Private/String/HexToBytes.cpp
@@ -11,7 +11,7 @@ namespace String
 {
 
 template <typename CharType>
-int32 HexToBytesImpl(const TStringView<CharType>& Hex, uint8* const OutBytes)
+static inline int32 HexToBytesImpl(TStringView<CharType> Hex, uint8* const OutBytes)
 {
 	const int32 HexCount = Hex.Len();
 	const CharType* HexPos = Hex.GetData();
@@ -29,14 +29,14 @@ int32 HexToBytesImpl(const TStringView<CharType>& Hex, uint8* const OutBytes)
 	return static_cast<int32>(OutPos - OutBytes);
 }
 
-int32 HexToBytes(const FStringView& Hex, uint8* OutBytes)
+int32 HexToBytes(FAnsiStringView Hex, uint8* OutBytes)
 {
-	return HexToBytesImpl<TCHAR>(Hex, OutBytes);
+	return HexToBytesImpl<ANSICHAR>(Hex, OutBytes);
 }
 
-int32 HexToBytes(const FAnsiStringView& Hex, uint8* OutBytes)
+int32 HexToBytes(FWideStringView Hex, uint8* OutBytes)
 {
-	return HexToBytesImpl<ANSICHAR>(Hex, OutBytes);
+	return HexToBytesImpl<WIDECHAR>(Hex, OutBytes);
 }
 
 }
diff --git a/Engine/Source/Runtime/Core/Public/String/BytesToHex.h b/Engine/Source/Runtime/Core/Public/String/BytesToHex.h
--- a/Engine/Source/Runtime/Core/Public/String/BytesToHex.h
+++ b/Engine/Source/Runtime/Core/Public/String/BytesToHex.h
@@ -17,8 +17,10 @@ namespace String
  * @param Bytes Array of bytes to convert.
  * @param OutHex [out] Array of at least 2*Bytes.Len() output characters.
  */
-CORE_API void BytesToHex(TArrayView<const uint8> Bytes, WIDECHAR* OutHex);
-CORE_API void BytesToHex(TArrayView<const uint8> Bytes, ANSICHAR* OutHex);
+CORE_API void BytesToHex(TConstArrayView<uint8> Bytes, ANSICHAR* OutHex);
+CORE_API void BytesToHex(TConstArrayView<uint8> Bytes, WIDECHAR* OutHex);
+CORE_API void BytesToHexLower(TConstArrayView<uint8> Bytes, ANSICHAR* OutHex);
+CORE_API void BytesToHexLower(TConstArrayView<uint8> Bytes, WIDECHAR* OutHex);
 
 /**
  * Append an array of bytes to a string builder as hex digits.
@@ -26,8 +28,10 @@ CORE_API void BytesToHex(TArrayView<const uint8> Bytes, ANSICHAR* OutHex);
  * @param Bytes Array of bytes to convert.
  * @param Builder Builder to append the converted string to.
  */
-CORE_API void BytesToHex(TArrayView<const uint8> Bytes, FStringBuilderBase& Builder);
-CORE_API void BytesToHex(TArrayView<const uint8> Bytes, FAnsiStringBuilderBase& Builder);
+CORE_API void BytesToHex(TConstArrayView<uint8> Bytes, FAnsiStringBuilderBase& Builder);
+CORE_API void BytesToHex(TConstArrayView<uint8> Bytes, FWideStringBuilderBase& Builder);
+CORE_API void BytesToHexLower(TConstArrayView<uint8> Bytes, FAnsiStringBuilderBase& Builder);
+CORE_API void BytesToHexLower(TConstArrayView<uint8> Bytes, FWideStringBuilderBase& Builder);
 
 }
 }
diff --git a/Engine/Source/Runtime/Core/Public/String/HexToBytes.h b/Engine/Source/Runtime/Core/Public/String/HexToBytes.h
--- a/Engine/Source/Runtime/Core/Public/String/HexToBytes.h
+++ b/Engine/Source/Runtime/Core/Public/String/HexToBytes.h
@@ -18,8 +18,8 @@ namespace String
  *
  * @return Number of bytes written to the output.
  */
-CORE_API int32 HexToBytes(const FStringView& Hex, uint8* OutBytes);
-CORE_API int32 HexToBytes(const FAnsiStringView& Hex, uint8* OutBytes);
+CORE_API int32 HexToBytes(FAnsiStringView Hex, uint8* OutBytes);
+CORE_API int32 HexToBytes(FWideStringView Hex, uint8* OutBytes);
 
 }
 }
