U5 - Enable the new Scheduler Backend

diff --git a/Engine/Source/Runtime/Core/Private/Async/TaskGraph.cpp b/Engine/Source/Runtime/Core/Private/Async/TaskGraph.cpp
--- a/Engine/Source/Runtime/Core/Private/Async/TaskGraph.cpp
+++ b/Engine/Source/Runtime/Core/Private/Async/TaskGraph.cpp
@@ -97,7 +97,7 @@ static FAutoConsoleVariableRef CVarForkedProcessMaxWorkerThreads(
 	TEXT("Configures the number of worker threads a forked process should spawn if it allows multithreading.")
 );
 
-CORE_API int32 GUseNewTaskBackend = 0;
+CORE_API int32 GUseNewTaskBackend = 1;
 
 #if CREATE_HIPRI_TASK_THREADS || CREATE_BACKGROUND_TASK_THREADS
 	static void ThreadSwitchForABTest(const TArray<FString>& Args)
