U5 - temporarely disable the new Scheduler again after reinventing the way how it deals with background tasks. To get the changes tested again. The TaskGraph Backwardscompat layer will now also support Busy waiting.

diff --git a/Engine/Source/Runtime/Core/Private/Async/Fundamental/Scheduler.cpp b/Engine/Source/Runtime/Core/Private/Async/Fundamental/Scheduler.cpp
--- a/Engine/Source/Runtime/Core/Private/Async/Fundamental/Scheduler.cpp
+++ b/Engine/Source/Runtime/Core/Private/Async/Fundamental/Scheduler.cpp
@@ -4,6 +4,7 @@
 #include "Async/Fundamental/Task.h"
 #include "Logging/LogMacros.h"
 #include "Misc/ScopeLock.h" 
+#include "CoreGlobals.h"
 
 namespace LowLevelTasks
 {
@@ -126,17 +127,30 @@ namespace LowLevelTasks
 		return ActiveTask;
 	}
 
-	template<FTask* (FScheduler::FLocalQueueType::*DequeueFunction)()>
-	FORCEINLINE_DEBUGGABLE bool FScheduler::TryExecuteTaskFrom(FLocalQueueType* Queue, FQueueRegistry::FOutOfWork& OutOfWork)
+	template<FTask* (FScheduler::FLocalQueueType::*DequeueFunction)(bool)>
+	FORCEINLINE_DEBUGGABLE bool FScheduler::TryExecuteTaskFrom(FLocalQueueType* Queue, FQueueRegistry::FOutOfWork& OutOfWork, bool GetBackgroundTask)
 	{
-		FTask* Task = (Queue->*DequeueFunction)();
+		int32 NumActiveWorkers = int(ActiveWorkers.load(std::memory_order_relaxed));
+		GetBackgroundTask &= int(ActiveBackgroundTasks.load(std::memory_order_relaxed)) < FMath::Max(1, NumActiveWorkers - 2);
+		FTask* Task = (Queue->*DequeueFunction)(GetBackgroundTask);
 		if (Task)
 		{	
 			OutOfWork.Stop();
+			const bool IsBackgroundTask = Task->GetPriority() == ETaskPriority::Background;
+			if (IsBackgroundTask)
+			{
+				ActiveBackgroundTasks.fetch_add(1, std::memory_order_relaxed);
+			}
+			
 			FTask* OldTask = ActiveTask;
 			ActiveTask = Task;
 			Task->ExecuteTask();
 			ActiveTask = OldTask;
+
+			if (IsBackgroundTask)
+			{
+				ActiveBackgroundTasks.fetch_sub(1, std::memory_order_relaxed);
+			}
 			return true;
 		}
 		return false;
@@ -162,15 +176,15 @@ namespace LowLevelTasks
 		FQueueRegistry::FOutOfWork OutOfWork = QueueRegistry.GetOutOfWorkScope();
 		while (true)
 		{
-			while(TryExecuteTaskFrom<&FLocalQueueType::DequeueLocal>(WorkerLocalQueue, OutOfWork)
-			   || TryExecuteTaskFrom<&FLocalQueueType::DequeueGlobal>(WorkerLocalQueue, OutOfWork))
-			{			
+			while(TryExecuteTaskFrom<&FLocalQueueType::DequeueLocal>(WorkerLocalQueue, OutOfWork, true)
+			   || TryExecuteTaskFrom<&FLocalQueueType::DequeueGlobal>(WorkerLocalQueue, OutOfWork, true))
+			{		
 				Drowsing = false;
 				WaitCount = 0;
 			}
 
-			while(TryExecuteTaskFrom<&FLocalQueueType::DequeueLocal>(WorkerLocalQueue, OutOfWork)
-			   || TryExecuteTaskFrom<&FLocalQueueType::DequeueSteal>(WorkerLocalQueue, OutOfWork))
+			while(TryExecuteTaskFrom<&FLocalQueueType::DequeueLocal>(WorkerLocalQueue, OutOfWork, true)
+			   || TryExecuteTaskFrom<&FLocalQueueType::DequeueSteal>(WorkerLocalQueue, OutOfWork, true))
 			{
 				Drowsing = false;
 				WaitCount = 0;
@@ -208,30 +222,44 @@ namespace LowLevelTasks
 	void FScheduler::BusyWaitInternal(const FConditional& Conditional)
 	{
 		TRACE_CPUPROFILER_EVENT_SCOPE(FScheduler::BusyWaitInternal);
+		FTaskTagScope NoneScope(ETaskTag::EBusyWait);
 
 		checkSlow(LocalQueue != nullptr);
 		check(ActiveWorkers.load(std::memory_order_relaxed));
 		FLocalQueueType* WorkerLocalQueue = LocalQueue;
 
 		uint32 WaitCount = 0;
+		bool GetBackgroundtasks = ActiveTask->GetPriority() == ETaskPriority::Background;
+		if (GetBackgroundtasks)
+		{
+			ActiveBackgroundTasks.fetch_sub(1, std::memory_order_relaxed);
+		}
 		FQueueRegistry::FOutOfWork OutOfWork = QueueRegistry.GetOutOfWorkScope();
 		while (true)
 		{
-			while(TryExecuteTaskFrom<&FLocalQueueType::DequeueLocal>(WorkerLocalQueue, OutOfWork)
-			   || TryExecuteTaskFrom<&FLocalQueueType::DequeueGlobal>(WorkerLocalQueue, OutOfWork))
+			while(TryExecuteTaskFrom<&FLocalQueueType::DequeueLocal>(WorkerLocalQueue, OutOfWork, GetBackgroundtasks)
+			   || TryExecuteTaskFrom<&FLocalQueueType::DequeueGlobal>(WorkerLocalQueue, OutOfWork, GetBackgroundtasks))
 			{
 				if (Conditional())
 				{
+					if (GetBackgroundtasks)
+					{
+						ActiveBackgroundTasks.fetch_add(1, std::memory_order_relaxed);
+					}
 					return;
 				}
 				WaitCount = 0;
 			}
 
-			while(TryExecuteTaskFrom<&FLocalQueueType::DequeueLocal>(WorkerLocalQueue, OutOfWork)
-			   || TryExecuteTaskFrom<&FLocalQueueType::DequeueSteal>(WorkerLocalQueue, OutOfWork))
+			while(TryExecuteTaskFrom<&FLocalQueueType::DequeueLocal>(WorkerLocalQueue, OutOfWork, GetBackgroundtasks)
+			   || TryExecuteTaskFrom<&FLocalQueueType::DequeueSteal>(WorkerLocalQueue, OutOfWork, GetBackgroundtasks))
 			{
 				if (Conditional())
 				{
+					if (GetBackgroundtasks)
+					{
+						ActiveBackgroundTasks.fetch_add(1, std::memory_order_relaxed);
+					}
 					return;
 				}
 				WaitCount = 0;
@@ -239,11 +267,16 @@ namespace LowLevelTasks
 
 			if (Conditional())
 			{
+				if (GetBackgroundtasks)
+				{
+					ActiveBackgroundTasks.fetch_add(1, std::memory_order_relaxed);
+				}
 				return;
 			}
 
 			if (WaitCount < WorkerSpinCycles)
 			{
+				OutOfWork.Start();
 				FPlatformProcess::Yield();
 				FPlatformProcess::Yield();
 				WaitCount++;
diff --git a/Engine/Source/Runtime/Core/Private/Async/TaskGraph.cpp b/Engine/Source/Runtime/Core/Private/Async/TaskGraph.cpp
--- a/Engine/Source/Runtime/Core/Private/Async/TaskGraph.cpp
+++ b/Engine/Source/Runtime/Core/Private/Async/TaskGraph.cpp
@@ -97,7 +97,7 @@ static FAutoConsoleVariableRef CVarForkedProcessMaxWorkerThreads(
 	TEXT("Configures the number of worker threads a forked process should spawn if it allows multithreading.")
 );
 
-CORE_API int32 GUseNewTaskBackend = 1;
+CORE_API int32 GUseNewTaskBackend = 0;
 
 #if CREATE_HIPRI_TASK_THREADS || CREATE_BACKGROUND_TASK_THREADS
 	static void ThreadSwitchForABTest(const TArray<FString>& Args)
@@ -1765,9 +1765,6 @@ class FTaskGraphCompatibilityImplementation final : public FTaskGraphInterface
 	TArray<FWorkerThread> NamedThreads;
 
 	FThreadSafeCounter	ReentrancyCheck;
-
-	std::atomic_int					ConcurrentNormalTasks { 0 };
-	std::atomic_int					ConcurrentBackgroundTasks { 0 };
 	FAAArrayQueue<FBaseGraphTask>	QueuedBackgroundTasks;
 
 public:
@@ -1854,58 +1851,15 @@ private:
 		{
 			uint32 ThreadPriority = GetThreadPriorityIndex(InThreadToExecuteOn);
 			check(ThreadPriority < uint32(LowLevelTasks::ETaskPriority::Count));
-			LowLevelTasks::ETaskPriority Conversion[int(ENamedThreads::NumThreadPriorities)] = { LowLevelTasks::ETaskPriority::Normal, LowLevelTasks::ETaskPriority::High, LowLevelTasks::ETaskPriority::Low };
-			
-			const bool IsBackgroundTask = !!(InThreadToExecuteOn & ENamedThreads::BackgroundThreadPriority);
+			LowLevelTasks::ETaskPriority Conversion[int(ENamedThreads::NumThreadPriorities)] = { LowLevelTasks::ETaskPriority::Normal, LowLevelTasks::ETaskPriority::High, LowLevelTasks::ETaskPriority::Background };
+			LowLevelTasks::ETaskPriority Priority = Conversion[ThreadPriority];
 
-			Task->TaskHandle.Init(TEXT("TaskGraphTask"), Conversion[ThreadPriority], [Task, InThreadToExecuteOn]()
+			Task->TaskHandle.Init(TEXT("TaskGraphTask"), Priority, [Task, InThreadToExecuteOn, Deleter(LowLevelTasks::TDeleter<FBaseGraphTask, &FBaseGraphTask::DeleteTask>(Task))]()
 			{
 				Task->Execute(NewTasks, InThreadToExecuteOn, false);
-			}, 
-			[this, Deleter(LowLevelTasks::TDeleter<FBaseGraphTask, &FBaseGraphTask::DeleteTask>(Task)), IsBackgroundTask]()
-			{
-				if (IsBackgroundTask)
-				{
-					ConcurrentBackgroundTasks--;
-				}
-				else
-				{
-					ConcurrentNormalTasks--;
-				}
-
-				const int MaxBackgroundTasks = FMath::Max(1, NumWorkerThreads - 1 - ConcurrentNormalTasks.load(std::memory_order_relaxed));
-				uint32 NumSpawned = 0;
-				while (ConcurrentBackgroundTasks.load(std::memory_order_relaxed) < MaxBackgroundTasks)
-				{
-					FBaseGraphTask* QueuedTask = QueuedBackgroundTasks.dequeue();
-					if (QueuedTask)
-					{
-						ConcurrentBackgroundTasks++;
-						verifySlow(LowLevelTasks::TryLaunch(QueuedTask->TaskHandle, (NumSpawned++) ? LowLevelTasks::EQueuePreference::GlobalQueuePreference : LowLevelTasks::EQueuePreference::LocalQueuePreference));
-						continue;
-					}
-					break;
-				}
 			});
-
-			if (IsBackgroundTask)
-			{
-				const int MaxBackgroundTasks = FMath::Max(1, NumWorkerThreads - 1 - ConcurrentNormalTasks.load(std::memory_order_relaxed));
-				if(ConcurrentBackgroundTasks.load(std::memory_order_relaxed) < MaxBackgroundTasks)
-				{
-					ConcurrentBackgroundTasks++;
-					verifySlow(LowLevelTasks::TryLaunch(Task->TaskHandle, LowLevelTasks::EQueuePreference::GlobalQueuePreference));			
-				}
-				else
-				{
-					QueuedBackgroundTasks.enqueue(Task);
-				}
-			}
-			else
-			{
-				ConcurrentNormalTasks++;
-				verifySlow(LowLevelTasks::TryLaunch(Task->TaskHandle, LowLevelTasks::EQueuePreference::GlobalQueuePreference));
-			}
+	
+			verifySlow(LowLevelTasks::TryLaunch(Task->TaskHandle, LowLevelTasks::EQueuePreference::GlobalQueuePreference));
 			return;
 		}
 
@@ -2043,27 +1997,18 @@ private:
 		}
 		else
 		{
-			if (!FTaskGraphInterface::IsMultithread())
+			LowLevelTasks::BusyWaitUntil([Index(0), &Tasks]() mutable
 			{
-				bool bAnyPending = false;
-				for (int32 Index = 0; Index < Tasks.Num(); Index++)
+				while (Index < Tasks.Num())
 				{
-					FGraphEvent* Task = Tasks[Index].GetReference();
-					if (Task && !Task->IsComplete())
+					if (!Tasks[Index]->IsComplete())
 					{
-						bAnyPending = true;
-						break;
+						return false;
 					}
+					Index++;
 				}
-				if (!bAnyPending)
-				{
-					return;
-				}
-				UE_LOG(LogTaskGraph, Fatal, TEXT("Recursive waits are not allowed in single threaded mode."));
-			}
-			// We will just stall this thread on an event while we wait
-			FScopedEvent Event;
-			TriggerEventWhenTasksComplete(Event.Get(), Tasks, CurrentThreadIfKnown);
+				return true;
+			});
 		}
 	}
 
diff --git a/Engine/Source/Runtime/Core/Private/HAL/ThreadingBase.cpp b/Engine/Source/Runtime/Core/Private/HAL/ThreadingBase.cpp
--- a/Engine/Source/Runtime/Core/Private/HAL/ThreadingBase.cpp
+++ b/Engine/Source/Runtime/Core/Private/HAL/ThreadingBase.cpp
@@ -90,7 +90,7 @@ FTaskTagScope::FTaskTagScope(bool InTagOnlyIfNone, ETaskTag InTag) : Tag(InTag),
 	}
 	
 	ParentTag = ActiveTaskTag;
-	if (!TagOnlyIfNone || ActiveTaskTag == ETaskTag::ENone)
+	if (!TagOnlyIfNone || ActiveTaskTag == ETaskTag::ENone || ActiveTaskTag == ETaskTag::EBusyWait)
 	{
 		ActiveTaskTag = Tag;
 	}
@@ -111,7 +111,7 @@ FTaskTagScope::FTaskTagScope(bool InTagOnlyIfNone, ETaskTag InTag) : Tag(InTag),
 FTaskTagScope::~FTaskTagScope()
 {
 	checkf(TagOnlyIfNone || ActiveTaskTag == Tag, TEXT("ActiveTaskTag(%x) corrupted needs to be Tag(%x)"), FTaskTagScope::GetCurrentTag(), Tag);
-	if (!TagOnlyIfNone || ParentTag == ETaskTag::ENone)
+	if (!TagOnlyIfNone || ParentTag == ETaskTag::ENone || ParentTag == ETaskTag::EBusyWait)
 	{
 		ActiveTaskTag = ParentTag;
 	}
diff --git a/Engine/Source/Runtime/Core/Public/Async/Fundamental/LocalQueue.h b/Engine/Source/Runtime/Core/Public/Async/Fundamental/LocalQueue.h
--- a/Engine/Source/Runtime/Core/Public/Async/Fundamental/LocalQueue.h
+++ b/Engine/Source/Runtime/Core/Public/Async/Fundamental/LocalQueue.h
@@ -5,9 +5,12 @@
 #include "Math/RandomStream.h"
 #include "Experimental/Containers/FAAArrayQueue.h"
 #include "ProfilingDebugging/CpuProfilerTrace.h"
+#include "Async/Fundamental/Task.h"
 
 #include <atomic>
 
+namespace LowLevelTasks
+{
 namespace LocalQueue_Impl
 {
 template<uint32 NumItems>
@@ -153,7 +156,7 @@ public:
  * A Dequeue Operation can only be done starting from a LocalQueue, than the GlobalQueue will be checked.                                   *
  * Finally Items might get Stolen from other LocalQueues that are registered with the LocalQueueRegistry.                                   *
  ********************************************************************************************************************************************/
-template<typename ItemType, uint32 NumPriorities, uint32 NumLocalItems = 1024>
+template<uint32 NumLocalItems = 1024>
 class TLocalQueueRegistry
 {
 public:
@@ -164,7 +167,7 @@ public:
 	class FOutOfWork
 	{
 	private:
-		template<typename, uint32, uint32>
+		template<uint32>
 		friend class TLocalQueueRegistry;
 
 		std::atomic_int& NumWorkersLookingForWork;
@@ -218,8 +221,8 @@ public:
 	};
 
 private:
-	using FLocalQueueType	 = LocalQueue_Impl::TWorkStealingQueue2<ItemType, NumLocalItems>;
-	using FOverflowQueueType = FAAArrayQueue<ItemType>;
+	using FLocalQueueType	 = LocalQueue_Impl::TWorkStealingQueue2<FTask, NumLocalItems>;
+	using FOverflowQueueType = FAAArrayQueue<FTask>;
 	using DequeueHazard		 = typename FOverflowQueueType::DequeueHazard;
 
 private:
@@ -250,7 +253,7 @@ private:
 public:
 	class TLocalQueue
 	{
-		template<typename, uint32, uint32>
+		template<uint32>
 		friend class TLocalQueueRegistry;
 
 	public:
@@ -259,7 +262,7 @@ public:
 			checkSlow(Registry);
 			StealHazard = FStealHazard(Registry->QueueCollection, Registry->HazardsCollection);
 			Registry->AddLocalQueue(StealHazard, this);
-			for (int PriorityIndex = 0; PriorityIndex < NumPriorities; PriorityIndex++)
+			for (int32 PriorityIndex = 0; PriorityIndex < int32(ETaskPriority::Count); PriorityIndex++)
 			{
 				DequeueHazards[PriorityIndex] = Registry->OverflowQueues[PriorityIndex].getHeadHazard();
 			}
@@ -279,11 +282,11 @@ public:
 			TLocalQueueRegistry* Registry = Queue->Registry;
 			Queue->Registry = nullptr;
 
-			for (int32 PriorityIndex = 0; PriorityIndex < NumPriorities; PriorityIndex++)
+			for (int32 PriorityIndex = 0; PriorityIndex < int32(ETaskPriority::Count); PriorityIndex++)
 			{
 				while (true)
 				{
-					ItemType* Item;
+					FTask* Item;
 					if (!Queue->LocalQueues[PriorityIndex].Get(Item))
 					{
 						break;
@@ -296,10 +299,10 @@ public:
 
 		//add an item to the local queue and overflow into the global queue if full
 		// returns true if we should wake a worker
-		inline bool Enqueue(ItemType* Item, uint32 PriorityIndex)
+		inline bool Enqueue(FTask* Item, uint32 PriorityIndex)
 		{
 			checkSlow(Registry);
-			checkSlow(PriorityIndex < NumPriorities);
+			checkSlow(PriorityIndex < int32(ETaskPriority::Count));
 			checkSlow(Item != nullptr);
 
 			if (!LocalQueues[PriorityIndex].Put(Item))
@@ -310,12 +313,13 @@ public:
 			return (LocalTasksSinceLastDequeue++ != 0) && Registry->AnyWorkerLookingForWork();
 		}
 
-		inline ItemType* DequeueLocal()
+		inline FTask* DequeueLocal(bool GetBackGroundTasks)
 		{
 			LocalTasksSinceLastDequeue = 0;
-			for (int32 PriorityIndex = 0; PriorityIndex < NumPriorities; PriorityIndex++)
+			int32 MaxPriority = GetBackGroundTasks ? int32(ETaskPriority::Count) : int32(ETaskPriority::Background);
+			for (int32 PriorityIndex = 0; PriorityIndex < MaxPriority; PriorityIndex++)
 			{
-				ItemType* Item;
+				FTask* Item;
 				if (LocalQueues[PriorityIndex].Get(Item))
 				{
 					return Item;
@@ -324,13 +328,14 @@ public:
 			return nullptr;
 		}
 
-		inline ItemType* DequeueGlobal()
+		inline FTask* DequeueGlobal(bool GetBackGroundTasks)
 		{
 			if (Registry->NumActiveWorkers.load(std::memory_order_relaxed) >= (2 * Registry->NumWorkersLookingForWork.load(std::memory_order_relaxed) - 1))
 			{
-				for (int32 PriorityIndex = 0; PriorityIndex < NumPriorities; PriorityIndex++)
+				int32 MaxPriority = GetBackGroundTasks ? int32(ETaskPriority::Count) : int32(ETaskPriority::Background);
+				for (int32 PriorityIndex = 0; PriorityIndex < MaxPriority; PriorityIndex++)
 				{
-					ItemType* Item = Registry->OverflowQueues[PriorityIndex].dequeue(DequeueHazards[PriorityIndex]);
+					FTask* Item = Registry->OverflowQueues[PriorityIndex].dequeue(DequeueHazards[PriorityIndex]);
 					if (Item)
 					{
 						return Item;
@@ -340,7 +345,7 @@ public:
 			return nullptr;
 		}
 
-		inline ItemType* DequeueSteal()
+		inline FTask* DequeueSteal(bool GetBackGroundTasks)
 		{
 			if (Registry->NumActiveWorkers.load(std::memory_order_relaxed) >= (2 * Registry->NumWorkersLookingForWork.load(std::memory_order_relaxed) - 1))
 			{
@@ -349,7 +354,7 @@ public:
 					CachedRandomIndex = Random.GetUnsignedInt();
 				}
 
-				ItemType* Result = Registry->StealItem(StealHazard, CachedRandomIndex, CachedPriorityIndex);
+				FTask* Result = Registry->StealItem(StealHazard, CachedRandomIndex, CachedPriorityIndex, GetBackGroundTasks);
 				if (Result)
 				{
 					return Result;
@@ -360,8 +365,8 @@ public:
 
 	private:
 		static constexpr uint32	InvalidIndex = ~0u;
-		FLocalQueueType			LocalQueues[NumPriorities];
-		DequeueHazard			DequeueHazards[NumPriorities];
+		FLocalQueueType			LocalQueues[uint32(ETaskPriority::Count)];
+		DequeueHazard			DequeueHazards[uint32(ETaskPriority::Count)];
 		FStealHazard			StealHazard;
 		TLocalQueueRegistry*	Registry;
 		FRandomStream			Random;
@@ -425,24 +430,25 @@ private:
 	}
 
 	// StealItem tries to steal an Item from a Registered LocalQueue
-	ItemType* StealItem(FStealHazard& Hazard, uint32& CachedRandomIndex, uint32& CachedPriorityIndex)
+	FTask* StealItem(FStealHazard& Hazard, uint32& CachedRandomIndex, uint32& CachedPriorityIndex, bool GetBackGroundTasks)
 	{
 		FLocalQueueCollection* Queues = Hazard.Get();
 		uint32 NumQueues = Queues->LocalQueues.Num();
+		uint32 MaxPriority = GetBackGroundTasks ? int32(ETaskPriority::Count) : int32(ETaskPriority::Background);
 		CachedRandomIndex = CachedRandomIndex % NumQueues;
 
 		for(uint32 i = 0; i < NumQueues; i++)
 		{
 			TLocalQueue* LocalQueue = Queues->LocalQueues[CachedRandomIndex];
-			for(uint32 j = 0; j < NumPriorities; j++)
+			for(uint32 PriorityIndex = 0; PriorityIndex < MaxPriority; PriorityIndex++)
 			{	
-				ItemType* Item;
+				FTask* Item;
 				if (LocalQueue->LocalQueues[CachedPriorityIndex].Steal(Item))
 				{
 					Hazard.Retire();
 					return Item;
 				}
-				CachedPriorityIndex = ++CachedPriorityIndex < NumPriorities ? CachedPriorityIndex : 0;
+				CachedPriorityIndex = ++CachedPriorityIndex < MaxPriority ? CachedPriorityIndex : 0;
 			}
 			CachedRandomIndex = ++CachedRandomIndex < NumQueues ? CachedRandomIndex : 0;
 		}
@@ -455,9 +461,9 @@ private:
 public:
 	// enqueue an Item directy into the Global OverflowQueue
 	// returns true if we should wake a worker for stealing
-	bool Enqueue(ItemType* Item, uint32 PriorityIndex)
+	bool Enqueue(FTask* Item, uint32 PriorityIndex)
 	{
-		check(PriorityIndex < NumPriorities);
+		check(PriorityIndex < int32(ETaskPriority::Count));
 		check(Item != nullptr);
 
 		OverflowQueues[PriorityIndex].enqueue(Item);
@@ -466,11 +472,11 @@ public:
 	}
 
 	// grab an Item directy from the Global OverflowQueue
-	ItemType* Dequeue()
+	FTask* Dequeue()
 	{
-		for (int32 PriorityIndex = 0; PriorityIndex < NumPriorities; PriorityIndex++)
+		for (int32 PriorityIndex = 0; PriorityIndex < int32(ETaskPriority::Count); PriorityIndex++)
 		{
-			ItemType* Result = OverflowQueues[PriorityIndex].dequeue();
+			FTask* Result = OverflowQueues[PriorityIndex].dequeue();
 			if (Result)
 			{
 				return Result;
@@ -484,18 +490,19 @@ public:
 		return FOutOfWork(NumWorkersLookingForWork);
 	}
 
-private:
 	inline bool AnyWorkerLookingForWork() const
 	{
 		return NumWorkersLookingForWork.load(std::memory_order_acquire) == 0;
 	}
 
-	FOverflowQueueType	  OverflowQueues[NumPriorities];
+private:
+	FOverflowQueueType	  OverflowQueues[uint32(ETaskPriority::Count)];
 	FHazardPointerCollection		  HazardsCollection;
 	std::atomic<FLocalQueueCollection*>	QueueCollection;
 	std::atomic_int NumWorkersLookingForWork { 0 };
 	std::atomic_int NumActiveWorkers { 0 };
 };
 
-template<typename ItemType, uint32 NumPriorities, uint32 NumLocalItems>
-uint32 TLocalQueueRegistry<ItemType, NumPriorities, NumLocalItems>::FOutOfWork::WorkerLookingForWorkTraceId = 0;
\ No newline at end of file
+template<uint32 NumLocalItems>
+uint32 TLocalQueueRegistry<NumLocalItems>::FOutOfWork::WorkerLookingForWorkTraceId = 0;
+}
\ No newline at end of file
diff --git a/Engine/Source/Runtime/Core/Public/Async/Fundamental/Scheduler.h b/Engine/Source/Runtime/Core/Public/Async/Fundamental/Scheduler.h
--- a/Engine/Source/Runtime/Core/Public/Async/Fundamental/Scheduler.h
+++ b/Engine/Source/Runtime/Core/Public/Async/Fundamental/Scheduler.h
@@ -24,7 +24,7 @@ namespace LowLevelTasks
 		UE_NONCOPYABLE(FScheduler);
 		static constexpr uint32 WorkerSpinCycles = 53;
 
-		using FQueueRegistry	= TLocalQueueRegistry<FTask, uint32(ETaskPriority::Count)>;
+		using FQueueRegistry	= TLocalQueueRegistry<>;
 		using FLocalQueueType	= FQueueRegistry::TLocalQueue;
 
 		static thread_local FLocalQueueType* LocalQueue;
@@ -103,8 +103,8 @@ namespace LowLevelTasks
 		CORE_API void BusyWaitInternal(const FConditional& Conditional);
 		FORCENOINLINE void TrySleeping(FSleepEvent* WorkerEvent, FQueueRegistry::FOutOfWork& OutOfWork, uint32& WaitCount, bool& Drowsing);
 		inline bool WakeUpWorker();
-		template<FTask* (FLocalQueueType::*DequeueFunction)()>
-		FORCEINLINE_DEBUGGABLE static bool TryExecuteTaskFrom(FLocalQueueType* Queue, FQueueRegistry::FOutOfWork& OutOfWork);
+		template<FTask* (FLocalQueueType::*DequeueFunction)(bool)>
+		FORCEINLINE_DEBUGGABLE bool TryExecuteTaskFrom(FLocalQueueType* Queue, FQueueRegistry::FOutOfWork& OutOfWork, bool GetBackgroundTask);
 
 	private:
 		FEventQueueType 			SleepEventQueue;
@@ -114,6 +114,7 @@ namespace LowLevelTasks
 		TArray<FLocalQueueType>		WorkerLocalQueues;
 		std::atomic_uint			ActiveWorkers { 0 };
 		std::atomic_uint			NextWorkerId { 0 };
+		std::atomic_uint			ActiveBackgroundTasks { 0 };
 	};
 
 	FORCEINLINE_DEBUGGABLE bool TryLaunch(FTask& Task, EQueuePreference QueuePreference = EQueuePreference::DefaultPreference)
diff --git a/Engine/Source/Runtime/Core/Public/Async/Fundamental/Task.h b/Engine/Source/Runtime/Core/Public/Async/Fundamental/Task.h
--- a/Engine/Source/Runtime/Core/Public/Async/Fundamental/Task.h
+++ b/Engine/Source/Runtime/Core/Public/Async/Fundamental/Task.h
@@ -15,7 +15,7 @@ namespace LowLevelTasks
 		High,
 		Normal,
 		Default = Normal,
-		Low,
+		Background,
 		Count
 	};
 
diff --git a/Engine/Source/Runtime/Core/Public/CoreGlobals.h b/Engine/Source/Runtime/Core/Public/CoreGlobals.h
--- a/Engine/Source/Runtime/Core/Public/CoreGlobals.h
+++ b/Engine/Source/Runtime/Core/Public/CoreGlobals.h
@@ -507,9 +507,10 @@ enum class ETaskTag : int32
 	ERenderingThread			= 1 << 4,
 	ERhiThread					= 1 << 5,
 	EAsyncLoadingThread			= 1 << 6,
+	EBusyWait					= 1 << 7,
 
 	ENamedThreadBits			= (EAsyncLoadingThread << 1) - 1,
-	EParallelThread				= 1 << 7,										//This can be used when multipe threads or jobs are involved (usually a parallel for) It will avoid the check for uniqieness of the named thread tag.
+	EParallelThread				= 1 << 8,										//This can be used when multipe threads or jobs are involved (usually a parallel for) It will avoid the check for uniqieness of the named thread tag.
 	EParallelRenderingThread	= ERenderingThread | EParallelThread,
 	EParallelGameThread			= EGameThread | EParallelThread,
 	EParallelRhiThread			= ERhiThread | EParallelThread,
diff --git a/Engine/Source/Runtime/Core/Public/Misc/QueuedThreadPoolWrapper.h b/Engine/Source/Runtime/Core/Public/Misc/QueuedThreadPoolWrapper.h
--- a/Engine/Source/Runtime/Core/Public/Misc/QueuedThreadPoolWrapper.h
+++ b/Engine/Source/Runtime/Core/Public/Misc/QueuedThreadPoolWrapper.h
@@ -287,10 +287,9 @@ private:
 		InQueuedWork->InternalData = QueuedWorkInternalData;
 		
 		checkSlow(int32(InPriority) < int32(EQueuedWorkPriority::Count));
-		const LowLevelTasks::ETaskPriority Mapping[int32(EQueuedWorkPriority::Count)] = { LowLevelTasks::ETaskPriority::Normal, LowLevelTasks::ETaskPriority::Normal, LowLevelTasks::ETaskPriority::Low, LowLevelTasks::ETaskPriority::Low, LowLevelTasks::ETaskPriority::Low };
+		EQueuedWorkPriority Priority = PriorityMapper(InPriority);
 
-		LowLevelTasks::ETaskPriority Priority = Mapping[int32(PriorityMapper(InPriority))];
-		QueuedWorkInternalData->Task.Init(TEXT("FQueuedLowLevelThreadPoolTask"), Priority, [InQueuedWork]
+		QueuedWorkInternalData->Task.Init(TEXT("FQueuedLowLevelThreadPoolTask"), LowLevelTasks::ETaskPriority::Background, [InQueuedWork]
 		{
 			FMemMark Mark(FMemStack::Get());
 			InQueuedWork->DoThreadedWork();
@@ -372,11 +371,11 @@ protected:
 	}
 
 private:
-	FAAArrayQueue<FQueuedWorkInternalData> PendingWork[int32(LowLevelTasks::ETaskPriority::Count)];
+	FAAArrayQueue<FQueuedWorkInternalData> PendingWork[int32(EQueuedWorkPriority::Count)];
 
 	inline FQueuedWorkInternalData* Dequeue()
 	{
-		for (int32 i = 0; i < int32(LowLevelTasks::ETaskPriority::Count); i++)
+		for (int32 i = 0; i < int32(EQueuedWorkPriority::Count); i++)
 		{
 			FQueuedWorkInternalData* QueuedWork = PendingWork[i].dequeue();
 			if (QueuedWork)
@@ -387,7 +386,7 @@ private:
 		return nullptr;
 	}
 
-	inline void Enqueue(LowLevelTasks::ETaskPriority Priority, FQueuedWorkInternalData* Item)
+	inline void Enqueue(EQueuedWorkPriority Priority, FQueuedWorkInternalData* Item)
 	{
 		PendingWork[int32(Priority)].enqueue(Item);
 	}
diff --git a/Engine/Source/Runtime/Launch/Private/LaunchEngineLoop.cpp b/Engine/Source/Runtime/Launch/Private/LaunchEngineLoop.cpp
--- a/Engine/Source/Runtime/Launch/Private/LaunchEngineLoop.cpp
+++ b/Engine/Source/Runtime/Launch/Private/LaunchEngineLoop.cpp
@@ -9,12 +9,14 @@
 #include "Misc/MessageDialog.h"
 #include "Misc/ScopedSlowTask.h"
 #include "Misc/QueuedThreadPool.h"
+#include "Misc/QueuedThreadPoolWrapper.h"
 #include "HAL/FileManager.h"
 #include "HAL/PlatformAffinity.h"
 #include "Misc/FileHelper.h"
 #include "Internationalization/TextLocalizationManagerGlobals.h"
 #include "Logging/LogSuppressionInterface.h"
 #include "Async/TaskGraphInterfaces.h"
+#include "Async/Fundamental/Scheduler.h"
 #include "MemPro/MemProProfiler.h"
 #include "Misc/TimeGuard.h"
 #include "Misc/Paths.h"
@@ -92,7 +94,6 @@
 	#include "Interfaces/IEditorStyleModule.h"
 	#include "PIEPreviewDeviceProfileSelectorModule.h"
 	#include "AssetCompilingManager.h"
-	#include "Misc/QueuedThreadPoolWrapper.h"
 	#include "ShaderCompiler.h"
 
 	#if PLATFORM_WINDOWS
@@ -2184,7 +2185,6 @@ int32 FEngineLoop::PreInitPreStartupScreen(const TCHAR* CmdLine)
 				NumThreadsInLargeThreadPool = FMath::Max(FPlatformMisc::NumberOfCoresIncludingHyperthreads() - 2, 2);
 				NumThreadsInThreadPool = FPlatformMisc::NumberOfWorkerThreadsToSpawn();
 				GLargeThreadPool = FQueuedThreadPool::Allocate();
-
 			}
 			else
 			{
@@ -2216,15 +2216,40 @@ int32 FEngineLoop::PreInitPreStartupScreen(const TCHAR* CmdLine)
 		}
 #else
 		{
-			GThreadPool = FQueuedThreadPool::Allocate();
-			int32 NumThreadsInThreadPool = FPlatformMisc::NumberOfWorkerThreadsToSpawn();
+			extern CORE_API int32 GUseNewTaskBackend;
+			if (!GUseNewTaskBackend)
+			{
+				GThreadPool = FQueuedThreadPool::Allocate();
+				int32 NumThreadsInThreadPool = FPlatformMisc::NumberOfWorkerThreadsToSpawn();
 
-			// we are only going to give dedicated servers one pool thread
-			if (FPlatformProperties::IsServerOnly())
+				// we are only going to give dedicated servers one pool thread
+				if (FPlatformProperties::IsServerOnly())
+				{
+					NumThreadsInThreadPool = 1;
+				}
+				verify(GThreadPool->Create(NumThreadsInThreadPool, StackSize * 1024, TPri_SlightlyBelowNormal, TEXT("ThreadPool")));
+			}
+			else
 			{
-				NumThreadsInThreadPool = 1;
+				int NumThreadsInWorkerPool = FMath::Max(LowLevelTasks::FScheduler::Get().GetNumWorkers() - 2, 2u);			
+				if(NumThreadsInWorkerPool < int(LowLevelTasks::FScheduler::Get().GetNumWorkers()))
+				{
+					NumThreadsInWorkerPool = FMath::Max(NumThreadsInWorkerPool - 1, 1);
+					GThreadPool = new FQueuedLowLevelThreadPool(NumThreadsInWorkerPool);
+				}
+				else
+				{
+					GThreadPool = FQueuedThreadPool::Allocate();
+					int32 NumThreadsInThreadPool = FPlatformMisc::NumberOfWorkerThreadsToSpawn();
+
+					// we are only going to give dedicated servers one pool thread
+					if (FPlatformProperties::IsServerOnly())
+					{
+						NumThreadsInThreadPool = 1;
+					}
+					verify(GThreadPool->Create(NumThreadsInThreadPool, StackSize * 1024, TPri_SlightlyBelowNormal, TEXT("ThreadPool")));
+				}
 			}
-			verify(GThreadPool->Create(NumThreadsInThreadPool, StackSize * 1024, TPri_SlightlyBelowNormal, TEXT("ThreadPool")));
 		}
 #endif
 		{
