U5 - Fix nullptr

diff --git a/Engine/Source/Runtime/Core/Private/Async/Fundamental/Scheduler.cpp b/Engine/Source/Runtime/Core/Private/Async/Fundamental/Scheduler.cpp
--- a/Engine/Source/Runtime/Core/Private/Async/Fundamental/Scheduler.cpp
+++ b/Engine/Source/Runtime/Core/Private/Async/Fundamental/Scheduler.cpp
@@ -229,7 +229,7 @@ namespace LowLevelTasks
 		FLocalQueueType* WorkerLocalQueue = LocalQueue;
 
 		uint32 WaitCount = 0;
-		bool GetBackgroundtasks = ActiveTask->GetPriority() == ETaskPriority::Background;
+		bool GetBackgroundtasks = ActiveTask && (ActiveTask->GetPriority() == ETaskPriority::Background);
 		if (GetBackgroundtasks)
 		{
 			ActiveBackgroundTasks.fetch_sub(1, std::memory_order_relaxed);
