U5 - Dedicate Workers as foreground or background works and avoid the complicated logic to determine the amount of background works. Background tasks start with below nomal thread priority. Adding another Thread priority for Background workers.

diff --git a/Engine/Source/Runtime/Core/Private/Async/Fundamental/Scheduler.cpp b/Engine/Source/Runtime/Core/Private/Async/Fundamental/Scheduler.cpp
--- a/Engine/Source/Runtime/Core/Private/Async/Fundamental/Scheduler.cpp
+++ b/Engine/Source/Runtime/Core/Private/Async/Fundamental/Scheduler.cpp
@@ -33,7 +33,7 @@ namespace LowLevelTasks
 		}
 	}
 
-	TUniquePtr<FThread> FScheduler::CreateWorker(FLocalQueueType* ExternalWorkerLocalQueue, EThreadPriority Priority, bool bIsForkable)
+	TUniquePtr<FThread> FScheduler::CreateWorker(FLocalQueueType* ExternalWorkerLocalQueue, EThreadPriority Priority, bool bPermitBackgroundWork, bool bIsForkable)
 	{
 		uint32 WorkerId = NextWorkerId++;
 		const uint32 WaitTimes[8] = { 23, 31, 41, 37, 47, 29, 19, 43 };
@@ -41,21 +41,25 @@ namespace LowLevelTasks
 		uint64 ThreadAffinityMask = FPlatformAffinity::GetTaskGraphThreadMask();
 		return MakeUnique<FThread>
 		(
-			*FString::Printf(TEXT("Task Worker #%d"), WorkerId),
-			[this, ExternalWorkerLocalQueue, WaitTime]
+			bPermitBackgroundWork ? *FString::Printf(TEXT("Background Worker #%d"), WorkerId) : *FString::Printf(TEXT("Task Worker #%d"), WorkerId),
+			[this, ExternalWorkerLocalQueue, WaitTime, bPermitBackgroundWork]
 			{ 
 				FSleepEvent Event;
-				WorkerMain(&Event, ExternalWorkerLocalQueue, WaitTime);
+				WorkerMain(&Event, ExternalWorkerLocalQueue, WaitTime, bPermitBackgroundWork);
 			}, 0, Priority, ThreadAffinityMask
 		);
 	}
 
-	void FScheduler::StartWorkers(uint32 NumWorkers, EThreadPriority Priority, bool bIsForkable)
+	void FScheduler::StartWorkers(uint32 NumWorkers, uint32 NumBackgroundWorkers, EThreadPriority WorkerPriority,  EThreadPriority BackgroundPriority, bool bIsForkable)
 	{
-		NumWorkers = NumWorkers == 0 ? FPlatformMisc::NumberOfWorkerThreadsToSpawn() : NumWorkers;
+		if (NumWorkers == 0)
+		{
+			NumWorkers = 2;
+			NumBackgroundWorkers = FMath::Max(1, FPlatformMisc::NumberOfWorkerThreadsToSpawn() - 2);
+		}
 
 		uint32 OldActiveWorkers = ActiveWorkers.load(std::memory_order_relaxed);
-		if(OldActiveWorkers == 0 && FPlatformProcess::SupportsMultithreading() && ActiveWorkers.compare_exchange_strong(OldActiveWorkers, NumWorkers, std::memory_order_relaxed))
+		if(OldActiveWorkers == 0 && FPlatformProcess::SupportsMultithreading() && ActiveWorkers.compare_exchange_strong(OldActiveWorkers, NumWorkers + NumBackgroundWorkers, std::memory_order_relaxed))
 		{
 			UE::Trace::ThreadGroupBegin(TEXT("Task Workers"));
 			FScopeLock Lock(&WorkerThreadsCS);
@@ -63,12 +67,17 @@ namespace LowLevelTasks
 			check(!WorkerLocalQueues.Num());
 			check(NextWorkerId == 0);
 
-			WorkerThreads.Reserve(NumWorkers);
-			WorkerLocalQueues.Reserve(NumWorkers);
+			WorkerThreads.Reserve(NumWorkers + NumBackgroundWorkers);
+			WorkerLocalQueues.Reserve(NumWorkers + NumBackgroundWorkers);
 			for (uint32 WorkerId = 0; WorkerId < NumWorkers; ++WorkerId)
 			{
 				WorkerLocalQueues.Emplace(QueueRegistry);
-				WorkerThreads.Add(CreateWorker(&WorkerLocalQueues.Last(), Priority, bIsForkable));
+				WorkerThreads.Add(CreateWorker(&WorkerLocalQueues.Last(), WorkerPriority, false, bIsForkable));
+			}
+			for (uint32 WorkerId = 0; WorkerId < NumBackgroundWorkers; ++WorkerId)
+			{
+				WorkerLocalQueues.Emplace(QueueRegistry);
+				WorkerThreads.Add(CreateWorker(&WorkerLocalQueues.Last(), BackgroundPriority, true, bIsForkable));
 			}
 			UE::Trace::ThreadGroupEnd();
 		}
@@ -128,35 +137,22 @@ namespace LowLevelTasks
 	}
 
 	template<FTask* (FScheduler::FLocalQueueType::*DequeueFunction)(bool)>
-	FORCEINLINE_DEBUGGABLE bool FScheduler::TryExecuteTaskFrom(FLocalQueueType* Queue, FQueueRegistry::FOutOfWork& OutOfWork, bool GetBackgroundTask)
+	bool FScheduler::TryExecuteTaskFrom(FLocalQueueType* Queue, FQueueRegistry::FOutOfWork& OutOfWork, bool bPermitBackgroundWork)
 	{
-		int32 NumActiveWorkers = int(ActiveWorkers.load(std::memory_order_relaxed));
-		GetBackgroundTask &= int(ActiveBackgroundTasks.load(std::memory_order_relaxed)) < FMath::Max(1, NumActiveWorkers - 2);
-		FTask* Task = (Queue->*DequeueFunction)(GetBackgroundTask);
+		FTask* Task = (Queue->*DequeueFunction)(bPermitBackgroundWork);
 		if (Task)
 		{	
-			OutOfWork.Stop();
-			const bool IsBackgroundTask = Task->GetPriority() == ETaskPriority::Background;
-			if (IsBackgroundTask)
-			{
-				ActiveBackgroundTasks.fetch_add(1, std::memory_order_relaxed);
-			}
-			
+			OutOfWork.Stop();		
 			FTask* OldTask = ActiveTask;
 			ActiveTask = Task;
 			Task->ExecuteTask();
 			ActiveTask = OldTask;
-
-			if (IsBackgroundTask)
-			{
-				ActiveBackgroundTasks.fetch_sub(1, std::memory_order_relaxed);
-			}
 			return true;
 		}
 		return false;
 	}
 
-	void FScheduler::WorkerMain(FSleepEvent* WorkerEvent, FLocalQueueType* ExternalWorkerLocalQueue, uint32 WaitCycles)
+	void FScheduler::WorkerMain(FSleepEvent* WorkerEvent, FLocalQueueType* ExternalWorkerLocalQueue, uint32 WaitCycles, bool bPermitBackgroundWork)
 	{
 		FMemory::SetupTLSCachesOnCurrentThread();
 
@@ -176,15 +172,15 @@ namespace LowLevelTasks
 		FQueueRegistry::FOutOfWork OutOfWork = QueueRegistry.GetOutOfWorkScope();
 		while (true)
 		{
-			while(TryExecuteTaskFrom<&FLocalQueueType::DequeueLocal>(WorkerLocalQueue, OutOfWork, true)
-			   || TryExecuteTaskFrom<&FLocalQueueType::DequeueGlobal>(WorkerLocalQueue, OutOfWork, true))
+			while(TryExecuteTaskFrom<&FLocalQueueType::DequeueLocal>(WorkerLocalQueue, OutOfWork, bPermitBackgroundWork)
+			   || TryExecuteTaskFrom<&FLocalQueueType::DequeueGlobal>(WorkerLocalQueue, OutOfWork, bPermitBackgroundWork))
 			{		
 				Drowsing = false;
 				WaitCount = 0;
 			}
 
-			while(TryExecuteTaskFrom<&FLocalQueueType::DequeueLocal>(WorkerLocalQueue, OutOfWork, true)
-			   || TryExecuteTaskFrom<&FLocalQueueType::DequeueSteal>(WorkerLocalQueue, OutOfWork, true))
+			while(TryExecuteTaskFrom<&FLocalQueueType::DequeueLocal>(WorkerLocalQueue, OutOfWork, bPermitBackgroundWork)
+			   || TryExecuteTaskFrom<&FLocalQueueType::DequeueSteal>(WorkerLocalQueue, OutOfWork, bPermitBackgroundWork))
 			{
 				Drowsing = false;
 				WaitCount = 0;
@@ -229,37 +225,25 @@ namespace LowLevelTasks
 		FLocalQueueType* WorkerLocalQueue = LocalQueue;
 
 		uint32 WaitCount = 0;
-		bool GetBackgroundtasks = ActiveTask && (ActiveTask->GetPriority() == ETaskPriority::Background);
-		if (GetBackgroundtasks)
-		{
-			ActiveBackgroundTasks.fetch_sub(1, std::memory_order_relaxed);
-		}
+		bool bPermitBackgroundWork = ActiveTask && (ActiveTask->GetPriority() == ETaskPriority::BackgroundNormal || ActiveTask->GetPriority() == ETaskPriority::BackgroundLow);
 		FQueueRegistry::FOutOfWork OutOfWork = QueueRegistry.GetOutOfWorkScope();
 		while (true)
 		{
-			while(TryExecuteTaskFrom<&FLocalQueueType::DequeueLocal>(WorkerLocalQueue, OutOfWork, GetBackgroundtasks)
-			   || TryExecuteTaskFrom<&FLocalQueueType::DequeueGlobal>(WorkerLocalQueue, OutOfWork, GetBackgroundtasks))
+			while(TryExecuteTaskFrom<&FLocalQueueType::DequeueLocal>(WorkerLocalQueue, OutOfWork, bPermitBackgroundWork)
+			   || TryExecuteTaskFrom<&FLocalQueueType::DequeueGlobal>(WorkerLocalQueue, OutOfWork, bPermitBackgroundWork))
 			{
 				if (Conditional())
 				{
-					if (GetBackgroundtasks)
-					{
-						ActiveBackgroundTasks.fetch_add(1, std::memory_order_relaxed);
-					}
 					return;
 				}
 				WaitCount = 0;
 			}
 
-			while(TryExecuteTaskFrom<&FLocalQueueType::DequeueLocal>(WorkerLocalQueue, OutOfWork, GetBackgroundtasks)
-			   || TryExecuteTaskFrom<&FLocalQueueType::DequeueSteal>(WorkerLocalQueue, OutOfWork, GetBackgroundtasks))
+			while(TryExecuteTaskFrom<&FLocalQueueType::DequeueLocal>(WorkerLocalQueue, OutOfWork, bPermitBackgroundWork)
+			   || TryExecuteTaskFrom<&FLocalQueueType::DequeueSteal>(WorkerLocalQueue, OutOfWork, bPermitBackgroundWork))
 			{
 				if (Conditional())
 				{
-					if (GetBackgroundtasks)
-					{
-						ActiveBackgroundTasks.fetch_add(1, std::memory_order_relaxed);
-					}
 					return;
 				}
 				WaitCount = 0;
@@ -267,10 +251,6 @@ namespace LowLevelTasks
 
 			if (Conditional())
 			{
-				if (GetBackgroundtasks)
-				{
-					ActiveBackgroundTasks.fetch_add(1, std::memory_order_relaxed);
-				}
 				return;
 			}
 
diff --git a/Engine/Source/Runtime/Core/Private/Async/TaskGraph.cpp b/Engine/Source/Runtime/Core/Private/Async/TaskGraph.cpp
--- a/Engine/Source/Runtime/Core/Private/Async/TaskGraph.cpp
+++ b/Engine/Source/Runtime/Core/Private/Async/TaskGraph.cpp
@@ -1773,7 +1773,10 @@ public:
 	{
 		if (FTaskGraphInterface::IsMultithread())
 		{
-			LowLevelTasks::FScheduler::Get().StartWorkers(NumWorkerThreads, EThreadPriority::TPri_Normal, FForkProcessHelper::IsForkedMultithreadInstance());
+			int32 NumBackgroundWorkers = FMath::Max(1, NumWorkerThreads - 2);
+			int32 NumWorkers =  FMath::Max(1, NumWorkerThreads - NumBackgroundWorkers);
+
+			LowLevelTasks::FScheduler::Get().StartWorkers(NumWorkers, NumBackgroundWorkers, EThreadPriority::TPri_Normal, EThreadPriority::TPri_BelowNormal, FForkProcessHelper::IsForkedMultithreadInstance());
 			NumNamedThreads = ENamedThreads::ActualRenderingThread + 1;
 			ENamedThreads::bHasBackgroundThreads = 1;
 			ENamedThreads::bHasHighPriorityThreads = 1;
@@ -1839,8 +1842,11 @@ public:
 	{
 		if (FTaskGraphInterface::IsMultithread())
 		{
+			int32 NumBackgroundWorkers = FMath::Max(1, NumWorkerThreads - 2);
+			int32 NumWorkers =  FMath::Max(1, NumWorkerThreads - NumBackgroundWorkers);
+
 			LowLevelTasks::FScheduler::Get().StopWorkers();
-			LowLevelTasks::FScheduler::Get().StartWorkers(NumWorkerThreads, Pri, FForkProcessHelper::IsForkedMultithreadInstance());
+			LowLevelTasks::FScheduler::Get().StartWorkers(NumWorkers, NumBackgroundWorkers, Pri, EThreadPriority::TPri_BelowNormal, FForkProcessHelper::IsForkedMultithreadInstance());
 		}
 	}
 
@@ -1851,7 +1857,7 @@ private:
 		{
 			uint32 ThreadPriority = GetThreadPriorityIndex(InThreadToExecuteOn);
 			check(ThreadPriority < uint32(LowLevelTasks::ETaskPriority::Count));
-			LowLevelTasks::ETaskPriority Conversion[int(ENamedThreads::NumThreadPriorities)] = { LowLevelTasks::ETaskPriority::Normal, LowLevelTasks::ETaskPriority::High, LowLevelTasks::ETaskPriority::Background };
+			LowLevelTasks::ETaskPriority Conversion[int(ENamedThreads::NumThreadPriorities)] = { LowLevelTasks::ETaskPriority::Normal, LowLevelTasks::ETaskPriority::High, LowLevelTasks::ETaskPriority::BackgroundNormal };
 			LowLevelTasks::ETaskPriority Priority = Conversion[ThreadPriority];
 
 			Task->TaskHandle.Init(TEXT("TaskGraphTask"), Priority, [Task, InThreadToExecuteOn, Deleter(LowLevelTasks::TDeleter<FBaseGraphTask, &FBaseGraphTask::DeleteTask>(Task))]()
@@ -2085,7 +2091,7 @@ private:
 			const LowLevelTasks::FTask* Task = LowLevelTasks::FScheduler::Get().GetActiveTask();
 			if (Task != nullptr)
 			{
-				ENamedThreads::Type Conversion[int(LowLevelTasks::ETaskPriority::Count)] = { ENamedThreads::HighThreadPriority, ENamedThreads::NormalThreadPriority, ENamedThreads::BackgroundThreadPriority };
+				ENamedThreads::Type Conversion[int(LowLevelTasks::ETaskPriority::Count)] = { ENamedThreads::HighThreadPriority, ENamedThreads::NormalThreadPriority, ENamedThreads::BackgroundThreadPriority, ENamedThreads::BackgroundThreadPriority };
 				CurrentThreadIfKnown = (ENamedThreads::Type)(ENamedThreads::AnyThread | Conversion[int(Task->GetPriority())]);
 			}
 		}
diff --git a/Engine/Source/Runtime/Core/Public/Async/Fundamental/LocalQueue.h b/Engine/Source/Runtime/Core/Public/Async/Fundamental/LocalQueue.h
--- a/Engine/Source/Runtime/Core/Public/Async/Fundamental/LocalQueue.h
+++ b/Engine/Source/Runtime/Core/Public/Async/Fundamental/LocalQueue.h
@@ -316,7 +316,7 @@ public:
 		inline FTask* DequeueLocal(bool GetBackGroundTasks)
 		{
 			LocalTasksSinceLastDequeue = 0;
-			int32 MaxPriority = GetBackGroundTasks ? int32(ETaskPriority::Count) : int32(ETaskPriority::Background);
+			int32 MaxPriority = GetBackGroundTasks ? int32(ETaskPriority::Count) : int32(ETaskPriority::BackgroundNormal);
 			for (int32 PriorityIndex = 0; PriorityIndex < MaxPriority; PriorityIndex++)
 			{
 				FTask* Item;
@@ -332,7 +332,7 @@ public:
 		{
 			if (Registry->NumActiveWorkers.load(std::memory_order_relaxed) >= (2 * Registry->NumWorkersLookingForWork.load(std::memory_order_relaxed) - 1))
 			{
-				int32 MaxPriority = GetBackGroundTasks ? int32(ETaskPriority::Count) : int32(ETaskPriority::Background);
+				int32 MaxPriority = GetBackGroundTasks ? int32(ETaskPriority::Count) : int32(ETaskPriority::BackgroundNormal);
 				for (int32 PriorityIndex = 0; PriorityIndex < MaxPriority; PriorityIndex++)
 				{
 					FTask* Item = Registry->OverflowQueues[PriorityIndex].dequeue(DequeueHazards[PriorityIndex]);
@@ -434,7 +434,7 @@ private:
 	{
 		FLocalQueueCollection* Queues = Hazard.Get();
 		uint32 NumQueues = Queues->LocalQueues.Num();
-		uint32 MaxPriority = GetBackGroundTasks ? int32(ETaskPriority::Count) : int32(ETaskPriority::Background);
+		uint32 MaxPriority = GetBackGroundTasks ? int32(ETaskPriority::Count) : int32(ETaskPriority::BackgroundNormal);
 		CachedRandomIndex = CachedRandomIndex % NumQueues;
 
 		for(uint32 i = 0; i < NumQueues; i++)
diff --git a/Engine/Source/Runtime/Core/Public/Async/Fundamental/Scheduler.h b/Engine/Source/Runtime/Core/Public/Async/Fundamental/Scheduler.h
--- a/Engine/Source/Runtime/Core/Public/Async/Fundamental/Scheduler.h
+++ b/Engine/Source/Runtime/Core/Public/Async/Fundamental/Scheduler.h
@@ -47,7 +47,7 @@ namespace LowLevelTasks
 		FORCEINLINE_DEBUGGABLE static FScheduler& Get();
 
 		//start number of workers where 0 is the system default
-		CORE_API void StartWorkers(uint32 NumWorkers = 0, EThreadPriority Priority = EThreadPriority::TPri_Normal, bool bIsForkable = false);
+		CORE_API void StartWorkers(uint32 NumWorkers = 0, uint32 NumBackgroundWorkers = 0, EThreadPriority WorkerPriority = EThreadPriority::TPri_Normal,  EThreadPriority BackgroundPriority = EThreadPriority::TPri_BelowNormal, bool bIsForkable = false);
 		CORE_API void StopWorkers();
 
 		//try to launch the task, the return value will specify if the task was in the ready state and has been launced
@@ -97,14 +97,15 @@ namespace LowLevelTasks
 		~FScheduler();
 
 	private: 
-		TUniquePtr<FThread> CreateWorker(FLocalQueueType* ExternalWorkerLocalQueue = nullptr, EThreadPriority Priority = EThreadPriority::TPri_Normal, bool bIsForkable = false);
-		void WorkerMain(struct FSleepEvent* WorkerEvent, FLocalQueueType* ExternalWorkerLocalQueue, uint32 WaitCycles);
+		TUniquePtr<FThread> CreateWorker(FLocalQueueType* ExternalWorkerLocalQueue = nullptr, EThreadPriority Priority = EThreadPriority::TPri_Normal, bool bPermitBackgroundWork = false, bool bIsForkable = false);
+		void WorkerMain(struct FSleepEvent* WorkerEvent, FLocalQueueType* ExternalWorkerLocalQueue, uint32 WaitCycles, bool bPermitBackgroundWork);
 		CORE_API void LaunchInternal(FTask& Task, EQueuePreference QueuePreference);
 		CORE_API void BusyWaitInternal(const FConditional& Conditional);
 		FORCENOINLINE void TrySleeping(FSleepEvent* WorkerEvent, FQueueRegistry::FOutOfWork& OutOfWork, uint32& WaitCount, bool& Drowsing);
 		inline bool WakeUpWorker();
+
 		template<FTask* (FLocalQueueType::*DequeueFunction)(bool)>
-		FORCEINLINE_DEBUGGABLE bool TryExecuteTaskFrom(FLocalQueueType* Queue, FQueueRegistry::FOutOfWork& OutOfWork, bool GetBackgroundTask);
+		static bool TryExecuteTaskFrom(FLocalQueueType* Queue, FQueueRegistry::FOutOfWork& OutOfWork, bool bPermitBackgroundWork);
 
 	private:
 		FEventQueueType 			SleepEventQueue;
@@ -114,7 +115,6 @@ namespace LowLevelTasks
 		TArray<FLocalQueueType>		WorkerLocalQueues;
 		std::atomic_uint			ActiveWorkers { 0 };
 		std::atomic_uint			NextWorkerId { 0 };
-		std::atomic_uint			ActiveBackgroundTasks { 0 };
 	};
 
 	FORCEINLINE_DEBUGGABLE bool TryLaunch(FTask& Task, EQueuePreference QueuePreference = EQueuePreference::DefaultPreference)
diff --git a/Engine/Source/Runtime/Core/Public/Async/Fundamental/Task.h b/Engine/Source/Runtime/Core/Public/Async/Fundamental/Task.h
--- a/Engine/Source/Runtime/Core/Public/Async/Fundamental/Task.h
+++ b/Engine/Source/Runtime/Core/Public/Async/Fundamental/Task.h
@@ -15,7 +15,8 @@ namespace LowLevelTasks
 		High,
 		Normal,
 		Default = Normal,
-		Background,
+		BackgroundNormal,
+		BackgroundLow,
 		Count
 	};
 
diff --git a/Engine/Source/Runtime/Core/Public/Misc/QueuedThreadPoolWrapper.h b/Engine/Source/Runtime/Core/Public/Misc/QueuedThreadPoolWrapper.h
--- a/Engine/Source/Runtime/Core/Public/Misc/QueuedThreadPoolWrapper.h
+++ b/Engine/Source/Runtime/Core/Public/Misc/QueuedThreadPoolWrapper.h
@@ -289,7 +289,7 @@ private:
 		checkSlow(int32(InPriority) < int32(EQueuedWorkPriority::Count));
 		EQueuedWorkPriority Priority = PriorityMapper(InPriority);
 
-		QueuedWorkInternalData->Task.Init(TEXT("FQueuedLowLevelThreadPoolTask"), LowLevelTasks::ETaskPriority::Background, [InQueuedWork]
+		QueuedWorkInternalData->Task.Init(TEXT("FQueuedLowLevelThreadPoolTask"), LowLevelTasks::ETaskPriority::BackgroundLow, [InQueuedWork]
 		{
 			FMemMark Mark(FMemStack::Get());
 			InQueuedWork->DoThreadedWork();
