U5 - Make Xbox compiler happy again

diff --git a/Engine/Source/Runtime/Core/Public/Misc/QueuedThreadPoolWrapper.h b/Engine/Source/Runtime/Core/Public/Misc/QueuedThreadPoolWrapper.h
--- a/Engine/Source/Runtime/Core/Public/Misc/QueuedThreadPoolWrapper.h
+++ b/Engine/Source/Runtime/Core/Public/Misc/QueuedThreadPoolWrapper.h
@@ -240,6 +240,16 @@ public:
 		Destroy();
 	}
 
+	void* operator new(size_t size)
+	{
+		return FMemory::Malloc(size, 128);
+	}
+
+	void operator delete(void* ptr)
+	{
+		FMemory::Free(ptr);
+	}
+
 	/**
 	*  Queued task are not scheduled against the wrapped thread-pool until resumed
 	*/
