U5 - Fine tune the priorities of the taskpool tasks

diff --git a/Engine/Source/Runtime/Core/Public/Misc/QueuedThreadPoolWrapper.h b/Engine/Source/Runtime/Core/Public/Misc/QueuedThreadPoolWrapper.h
--- a/Engine/Source/Runtime/Core/Public/Misc/QueuedThreadPoolWrapper.h
+++ b/Engine/Source/Runtime/Core/Public/Misc/QueuedThreadPoolWrapper.h
@@ -286,20 +286,40 @@ public:
 		{
 			bIsPaused = false;
 		}
+
+		ScheduleTasks();
 	}
 
 private:
+	void ScheduleTasks()
+	{
+		LowLevelTasks::EQueuePreference QueuePreference = LowLevelTasks::EQueuePreference::LocalQueuePreference;
+		while (!bIsPaused && TaskCount < MaxConcurrency)
+		{
+			FQueuedWorkInternalData* QueuedWork = Dequeue();
+			if (QueuedWork)
+			{
+				TaskCount++;
+				verifySlow(LowLevelTasks::TryLaunch(QueuedWork->Task, QueuePreference));
+				QueuePreference = LowLevelTasks::EQueuePreference::GlobalQueuePreference;
+			}
+			else
+			{
+				break;
+			}
+		}
+	}
+
 	void AddQueuedWork(IQueuedWork* InQueuedWork, EQueuedWorkPriority InPriority = EQueuedWorkPriority::Normal) override
 	{
 		check(bIsExiting == false);
 
 		FQueuedWorkInternalData* QueuedWorkInternalData = new FQueuedWorkInternalData();
 		InQueuedWork->InternalData = QueuedWorkInternalData;
-		
-		checkSlow(int32(InPriority) < int32(EQueuedWorkPriority::Count));
 		EQueuedWorkPriority Priority = PriorityMapper(InPriority);
+		LowLevelTasks::ETaskPriority TaskPriority = Priority > EQueuedWorkPriority::Normal ? LowLevelTasks::ETaskPriority::BackgroundLow : LowLevelTasks::ETaskPriority::BackgroundNormal;
 
-		QueuedWorkInternalData->Task.Init(TEXT("FQueuedLowLevelThreadPoolTask"), LowLevelTasks::ETaskPriority::BackgroundLow, [InQueuedWork]
+		QueuedWorkInternalData->Task.Init(TEXT("FQueuedLowLevelThreadPoolTask"), TaskPriority, [InQueuedWork]
 		{
 			FMemMark Mark(FMemStack::Get());
 			InQueuedWork->DoThreadedWork();
@@ -307,31 +327,12 @@ private:
 		[this, InternalData = InQueuedWork->InternalData]()
 		{
 			--TaskCount;
-			uint32 NumSpawned = 0;
-			while (TaskCount < MaxConcurrency)
-			{
-				FQueuedWorkInternalData* QueuedWork = Dequeue();
-				if (QueuedWork)
-				{
-					TaskCount++;
-					verifySlow(LowLevelTasks::TryLaunch(QueuedWork->Task, (NumSpawned++) ? LowLevelTasks::EQueuePreference::GlobalQueuePreference : LowLevelTasks::EQueuePreference::LocalQueuePreference));
-				}
-				else
-				{
-					break;
-				}
-			}
+			LowLevelTasks::EQueuePreference QueuePreference = LowLevelTasks::EQueuePreference::LocalQueuePreference;
+			ScheduleTasks();
 		});
 
-		if(!bIsPaused && TaskCount <= MaxConcurrency)
-		{
-			TaskCount++;
-			verifySlow(LowLevelTasks::TryLaunch(QueuedWorkInternalData->Task, LowLevelTasks::EQueuePreference::GlobalQueuePreference));
-		}
-		else
-		{
-			Enqueue(Priority, QueuedWorkInternalData);
-		}
+		Enqueue(Priority, QueuedWorkInternalData);
+		ScheduleTasks();	
 	}
 
 	bool RetractQueuedWork(IQueuedWork* InQueuedWork) override
