U5 - teach Taskgraph how to use the other BackgoundPriority using the HighPrio TaskFlag. Also enabling the possibility to request a number of foregroundworkers (that never will run background work) using the -foregroundworkers=X cmdLine switch

diff --git a/Engine/Source/Runtime/Core/Private/Async/TaskGraph.cpp b/Engine/Source/Runtime/Core/Private/Async/TaskGraph.cpp
--- a/Engine/Source/Runtime/Core/Private/Async/TaskGraph.cpp
+++ b/Engine/Source/Runtime/Core/Private/Async/TaskGraph.cpp
@@ -98,6 +98,7 @@ static FAutoConsoleVariableRef CVarForkedProcessMaxWorkerThreads(
 );
 
 CORE_API int32 GUseNewTaskBackend = 0;
+uint32 GNumForegroundWorkers = 2;
 
 #if CREATE_HIPRI_TASK_THREADS || CREATE_BACKGROUND_TASK_THREADS
 	static void ThreadSwitchForABTest(const TArray<FString>& Args)
@@ -1773,7 +1774,7 @@ public:
 	{
 		if (FTaskGraphInterface::IsMultithread())
 		{
-			int32 NumBackgroundWorkers = FMath::Max(1, NumWorkerThreads - 2);
+			int32 NumBackgroundWorkers = FMath::Max(1, NumWorkerThreads - FMath::Min<int>(GNumForegroundWorkers, NumWorkerThreads));
 			int32 NumWorkers =  FMath::Max(1, NumWorkerThreads - NumBackgroundWorkers);
 
 			LowLevelTasks::FScheduler::Get().StartWorkers(NumWorkers, NumBackgroundWorkers, EThreadPriority::TPri_Normal, EThreadPriority::TPri_BelowNormal, FForkProcessHelper::IsForkedMultithreadInstance());
@@ -1842,7 +1843,7 @@ public:
 	{
 		if (FTaskGraphInterface::IsMultithread())
 		{
-			int32 NumBackgroundWorkers = FMath::Max(1, NumWorkerThreads - 2);
+			int32 NumBackgroundWorkers = FMath::Max(1, NumWorkerThreads - FMath::Min<int>(GNumForegroundWorkers, NumWorkerThreads));
 			int32 NumWorkers =  FMath::Max(1, NumWorkerThreads - NumBackgroundWorkers);
 
 			LowLevelTasks::FScheduler::Get().StopWorkers();
@@ -1857,9 +1858,14 @@ private:
 		{
 			uint32 ThreadPriority = GetThreadPriorityIndex(InThreadToExecuteOn);
 			check(ThreadPriority < uint32(LowLevelTasks::ETaskPriority::Count));
-			LowLevelTasks::ETaskPriority Conversion[int(ENamedThreads::NumThreadPriorities)] = { LowLevelTasks::ETaskPriority::Normal, LowLevelTasks::ETaskPriority::High, LowLevelTasks::ETaskPriority::BackgroundNormal };
+			LowLevelTasks::ETaskPriority Conversion[int(ENamedThreads::NumThreadPriorities)] = { LowLevelTasks::ETaskPriority::Normal, LowLevelTasks::ETaskPriority::High, LowLevelTasks::ETaskPriority::BackgroundLow };
 			LowLevelTasks::ETaskPriority Priority = Conversion[ThreadPriority];
 
+			if (Priority == LowLevelTasks::ETaskPriority::BackgroundLow && GetTaskPriority(InThreadToExecuteOn))
+			{
+				Priority = LowLevelTasks::ETaskPriority::BackgroundNormal;
+			}
+
 			Task->TaskHandle.Init(TEXT("TaskGraphTask"), Priority, [Task, InThreadToExecuteOn, Deleter(LowLevelTasks::TDeleter<FBaseGraphTask, &FBaseGraphTask::DeleteTask>(Task))]()
 			{
 				Task->Execute(NewTasks, InThreadToExecuteOn, false);
@@ -2091,8 +2097,9 @@ private:
 			const LowLevelTasks::FTask* Task = LowLevelTasks::FScheduler::Get().GetActiveTask();
 			if (Task != nullptr)
 			{
-				ENamedThreads::Type Conversion[int(LowLevelTasks::ETaskPriority::Count)] = { ENamedThreads::HighThreadPriority, ENamedThreads::NormalThreadPriority, ENamedThreads::BackgroundThreadPriority, ENamedThreads::BackgroundThreadPriority };
-				CurrentThreadIfKnown = (ENamedThreads::Type)(ENamedThreads::AnyThread | Conversion[int(Task->GetPriority())]);
+				ENamedThreads::Type ThreadConversion[int(LowLevelTasks::ETaskPriority::Count)] = { ENamedThreads::HighThreadPriority, ENamedThreads::NormalThreadPriority, ENamedThreads::BackgroundThreadPriority, ENamedThreads::BackgroundThreadPriority };
+				ENamedThreads::Type TaskConversion[int(LowLevelTasks::ETaskPriority::Count)] = { ENamedThreads::NormalTaskPriority, ENamedThreads::NormalTaskPriority, ENamedThreads::HighTaskPriority, ENamedThreads::NormalTaskPriority };
+				CurrentThreadIfKnown = (ENamedThreads::Type)(ENamedThreads::AnyThread | ThreadConversion[int(Task->GetPriority())] | TaskConversion[int(Task->GetPriority())]);
 			}
 		}
 		return CurrentThreadIfKnown;
@@ -2135,6 +2142,7 @@ void FTaskGraphInterface::Startup(int32 NumThreads)
 
 	if (GUseNewTaskBackend)
 	{
+		FParse::Value(FCommandLine::Get(), TEXT("-foregroundworkers="), GNumForegroundWorkers);
 		TaskGraphImplementationSingleton = new FTaskGraphCompatibilityImplementation(NumThreads);
 	}
 	else
