U5 - Renaming and regrouping of workers

diff --git a/Engine/Source/Runtime/Core/Private/Async/Fundamental/Scheduler.cpp b/Engine/Source/Runtime/Core/Private/Async/Fundamental/Scheduler.cpp
--- a/Engine/Source/Runtime/Core/Private/Async/Fundamental/Scheduler.cpp
+++ b/Engine/Source/Runtime/Core/Private/Async/Fundamental/Scheduler.cpp
@@ -41,7 +41,7 @@ namespace LowLevelTasks
 		uint64 ThreadAffinityMask = FPlatformAffinity::GetTaskGraphThreadMask();
 		return MakeUnique<FThread>
 		(
-			bPermitBackgroundWork ? *FString::Printf(TEXT("Background Worker #%d"), WorkerId) : *FString::Printf(TEXT("Task Worker #%d"), WorkerId),
+			bPermitBackgroundWork ? *FString::Printf(TEXT("Background Worker #%d"), WorkerId) : *FString::Printf(TEXT("Foreground Worker #%d"), WorkerId),
 			[this, ExternalWorkerLocalQueue, WaitTime, bPermitBackgroundWork]
 			{ 
 				FSleepEvent Event;
@@ -61,7 +61,6 @@ namespace LowLevelTasks
 		uint32 OldActiveWorkers = ActiveWorkers.load(std::memory_order_relaxed);
 		if(OldActiveWorkers == 0 && FPlatformProcess::SupportsMultithreading() && ActiveWorkers.compare_exchange_strong(OldActiveWorkers, NumWorkers + NumBackgroundWorkers, std::memory_order_relaxed))
 		{
-			UE::Trace::ThreadGroupBegin(TEXT("Task Workers"));
 			FScopeLock Lock(&WorkerThreadsCS);
 			check(!WorkerThreads.Num());
 			check(!WorkerLocalQueues.Num());
@@ -69,11 +68,14 @@ namespace LowLevelTasks
 
 			WorkerThreads.Reserve(NumWorkers + NumBackgroundWorkers);
 			WorkerLocalQueues.Reserve(NumWorkers + NumBackgroundWorkers);
+			UE::Trace::ThreadGroupBegin(TEXT("Foreground Workers"));
 			for (uint32 WorkerId = 0; WorkerId < NumWorkers; ++WorkerId)
 			{
 				WorkerLocalQueues.Emplace(QueueRegistry);
 				WorkerThreads.Add(CreateWorker(&WorkerLocalQueues.Last(), WorkerPriority, false, bIsForkable));
 			}
+			UE::Trace::ThreadGroupEnd();
+			UE::Trace::ThreadGroupBegin(TEXT("Background Workers"));
 			for (uint32 WorkerId = 0; WorkerId < NumBackgroundWorkers; ++WorkerId)
 			{
 				WorkerLocalQueues.Emplace(QueueRegistry);
