U5 - prevent the backgoundtasks from waking foregroundworkers that they cannot run on

diff --git a/Engine/Source/Runtime/Core/Private/Async/Fundamental/Scheduler.cpp b/Engine/Source/Runtime/Core/Private/Async/Fundamental/Scheduler.cpp
--- a/Engine/Source/Runtime/Core/Private/Async/Fundamental/Scheduler.cpp
+++ b/Engine/Source/Runtime/Core/Private/Async/Fundamental/Scheduler.cpp
@@ -91,9 +91,9 @@ namespace LowLevelTasks
 		if(OldActiveWorkers != 0 && ActiveWorkers.compare_exchange_strong(OldActiveWorkers, 0, std::memory_order_relaxed))
 		{
 			FScopeLock Lock(&WorkerThreadsCS);
-			while (WakeUpWorker())
-			{
-			}
+			while (WakeUpWorker(true)) {}
+			while (WakeUpWorker(false)) {}
+
 			for (TUniquePtr<FThread>& Thread : WorkerThreads)
 			{
 				Thread->Join();
@@ -112,18 +112,25 @@ namespace LowLevelTasks
 	{
 		if (ActiveWorkers.load(std::memory_order_relaxed))
 		{
+			bool bIsBackgroundTask = Task.IsBackgroundTask();
 			if (LocalQueue && QueuePreference != EQueuePreference::GlobalQueuePreference)
 			{
 				if (LocalQueue->Enqueue(&Task, uint32(Task.GetPriority())))
 				{
-					WakeUpWorker();
+					if (!WakeUpWorker(bIsBackgroundTask) && !bIsBackgroundTask)
+					{
+						WakeUpWorker(true);
+					}
 				}
 			}
 			else
 			{
 				if (QueueRegistry.Enqueue(&Task, uint32(Task.GetPriority())))
 				{
-					WakeUpWorker();
+					if (!WakeUpWorker(bIsBackgroundTask) && !bIsBackgroundTask)
+					{
+						WakeUpWorker(true);
+					}
 				}
 			}
 		}
@@ -204,12 +211,10 @@ namespace LowLevelTasks
 				continue;
 			}
 
-			TrySleeping(WorkerEvent, OutOfWork, WaitCount, Drowsing);
+			TrySleeping(WorkerEvent, OutOfWork, WaitCount, Drowsing, bPermitBackgroundWork);
 		}
 
-		while (WakeUpWorker())
-		{
-		}
+		while (WakeUpWorker(bPermitBackgroundWork)) {}
 
 		FLocalQueueType::DeleteLocalQueue(WorkerLocalQueue, ExternalWorkerLocalQueue != nullptr);
 		LocalQueue = nullptr;
@@ -227,7 +232,7 @@ namespace LowLevelTasks
 		FLocalQueueType* WorkerLocalQueue = LocalQueue;
 
 		uint32 WaitCount = 0;
-		bool bPermitBackgroundWork = ActiveTask && (ActiveTask->GetPriority() == ETaskPriority::BackgroundNormal || ActiveTask->GetPriority() == ETaskPriority::BackgroundLow);
+		bool bPermitBackgroundWork = ActiveTask && ActiveTask->IsBackgroundTask();
 		FQueueRegistry::FOutOfWork OutOfWork = QueueRegistry.GetOutOfWorkScope();
 		while (true)
 		{
diff --git a/Engine/Source/Runtime/Core/Public/Async/Fundamental/Scheduler.h b/Engine/Source/Runtime/Core/Public/Async/Fundamental/Scheduler.h
--- a/Engine/Source/Runtime/Core/Public/Async/Fundamental/Scheduler.h
+++ b/Engine/Source/Runtime/Core/Public/Async/Fundamental/Scheduler.h
@@ -101,14 +101,14 @@ namespace LowLevelTasks
 		void WorkerMain(struct FSleepEvent* WorkerEvent, FLocalQueueType* ExternalWorkerLocalQueue, uint32 WaitCycles, bool bPermitBackgroundWork);
 		CORE_API void LaunchInternal(FTask& Task, EQueuePreference QueuePreference);
 		CORE_API void BusyWaitInternal(const FConditional& Conditional);
-		FORCENOINLINE void TrySleeping(FSleepEvent* WorkerEvent, FQueueRegistry::FOutOfWork& OutOfWork, uint32& WaitCount, bool& Drowsing);
-		inline bool WakeUpWorker();
+		FORCENOINLINE void TrySleeping(FSleepEvent* WorkerEvent, FQueueRegistry::FOutOfWork& OutOfWork, uint32& WaitCount, bool& Drowsing, bool bBackgroundWorker);
+		inline bool WakeUpWorker(bool bBackgroundWorker);
 
 		template<FTask* (FLocalQueueType::*DequeueFunction)(bool)>
 		static bool TryExecuteTaskFrom(FLocalQueueType* Queue, FQueueRegistry::FOutOfWork& OutOfWork, bool bPermitBackgroundWork);
 
 	private:
-		FEventQueueType 			SleepEventQueue;
+		FEventQueueType 			SleepEventQueue[2];
 		FQueueRegistry 				QueueRegistry;
 		FCriticalSection 			WorkerThreadsCS;
 		TArray<TUniquePtr<FThread>> WorkerThreads;
@@ -219,7 +219,7 @@ namespace LowLevelTasks
 		}
 	}
 
-	inline void FScheduler::TrySleeping(FSleepEvent* WorkerEvent, FQueueRegistry::FOutOfWork& OutOfWork, uint32& WaitCount, bool& Drowsing)
+	inline void FScheduler::TrySleeping(FSleepEvent* WorkerEvent, FQueueRegistry::FOutOfWork& OutOfWork, uint32& WaitCount, bool& Drowsing, bool bBackgroundWorker)
 	{
 		ESleepState DrowsingState1 = ESleepState::Drowsing;
 		ESleepState DrowsingState2 = ESleepState::Drowsing;
@@ -241,7 +241,7 @@ namespace LowLevelTasks
 		{
 			verifySlow(OutOfWork.Stop() == !Drowsing);
 			Drowsing = true;
-			SleepEventQueue.enqueue(WorkerEvent);
+			SleepEventQueue[bBackgroundWorker].enqueue(WorkerEvent);
 			FPlatformProcess::SleepNoStats(0); // State one: (Running -> Drowsing)
 		}
 		else
@@ -250,9 +250,9 @@ namespace LowLevelTasks
 		}
 	}
 
-	inline bool FScheduler::WakeUpWorker()
+	inline bool FScheduler::WakeUpWorker(bool bBackgroundWorker)
 	{
-		FSleepEvent* WorkerEvent = SleepEventQueue.dequeue();
+		FSleepEvent* WorkerEvent = SleepEventQueue[bBackgroundWorker].dequeue();
 		if (WorkerEvent)
 		{
 			// Solving State : (Running -> Drowsing) -> Running  OR ((Running -> Drowsing) -> Drowsing) -> Running OR (((Running -> Drowsing) -> Sleeping) -> Running)
diff --git a/Engine/Source/Runtime/Core/Public/Async/Fundamental/Task.h b/Engine/Source/Runtime/Core/Public/Async/Fundamental/Task.h
--- a/Engine/Source/Runtime/Core/Public/Async/Fundamental/Task.h
+++ b/Engine/Source/Runtime/Core/Public/Async/Fundamental/Task.h
@@ -187,6 +187,7 @@ namespace LowLevelTasks
 
 		inline const TCHAR* GetDebugName() const;
 		inline ETaskPriority GetPriority() const;
+		inline bool IsBackgroundTask() const;
 
 	public:
 		FTask()
@@ -314,4 +315,9 @@ namespace LowLevelTasks
 	{
 		return PackedData.load(std::memory_order_relaxed).GetDebugName(); 
 	}
+
+	inline bool FTask::IsBackgroundTask() const
+	{
+		return PackedData.load(std::memory_order_relaxed).GetPriority() >= ETaskPriority::BackgroundNormal;
+	}
 }
\ No newline at end of file
