U5 - Fix the throttle stealing logic to account for the two worker types

diff --git a/Engine/Source/Runtime/Core/Private/Async/Fundamental/Scheduler.cpp b/Engine/Source/Runtime/Core/Private/Async/Fundamental/Scheduler.cpp
--- a/Engine/Source/Runtime/Core/Private/Async/Fundamental/Scheduler.cpp
+++ b/Engine/Source/Runtime/Core/Private/Async/Fundamental/Scheduler.cpp
@@ -20,7 +20,8 @@ namespace LowLevelTasks
 		RegisteredLocalQueue = LocalQueue == nullptr;
 		if (RegisteredLocalQueue)
 		{
-			LocalQueue = FLocalQueueType::AllocateLocalQueue(Scheduler.QueueRegistry);
+			bool bPermitBackgroundWork = ActiveTask && ActiveTask->IsBackgroundTask();
+			LocalQueue = FLocalQueueType::AllocateLocalQueue(Scheduler.QueueRegistry, bPermitBackgroundWork);
 		}
 	}
 
@@ -28,7 +29,8 @@ namespace LowLevelTasks
 	{
 		if (RegisteredLocalQueue)
 		{
-			FLocalQueueType::DeleteLocalQueue(LocalQueue);
+			bool bPermitBackgroundWork = ActiveTask && ActiveTask->IsBackgroundTask();
+			FLocalQueueType::DeleteLocalQueue(LocalQueue, bPermitBackgroundWork);
 			LocalQueue = nullptr;
 		}
 	}
@@ -71,14 +73,14 @@ namespace LowLevelTasks
 			UE::Trace::ThreadGroupBegin(TEXT("Foreground Workers"));
 			for (uint32 WorkerId = 0; WorkerId < NumWorkers; ++WorkerId)
 			{
-				WorkerLocalQueues.Emplace(QueueRegistry);
+				WorkerLocalQueues.Emplace(QueueRegistry, false);
 				WorkerThreads.Add(CreateWorker(&WorkerLocalQueues.Last(), WorkerPriority, false, bIsForkable));
 			}
 			UE::Trace::ThreadGroupEnd();
 			UE::Trace::ThreadGroupBegin(TEXT("Background Workers"));
 			for (uint32 WorkerId = 0; WorkerId < NumBackgroundWorkers; ++WorkerId)
 			{
-				WorkerLocalQueues.Emplace(QueueRegistry);
+				WorkerLocalQueues.Emplace(QueueRegistry, true);
 				WorkerThreads.Add(CreateWorker(&WorkerLocalQueues.Last(), BackgroundPriority, true, bIsForkable));
 			}
 			UE::Trace::ThreadGroupEnd();
@@ -172,13 +174,13 @@ namespace LowLevelTasks
 		}
 		else
 		{
-			LocalQueue = FLocalQueueType::AllocateLocalQueue(QueueRegistry);
+			LocalQueue = FLocalQueueType::AllocateLocalQueue(QueueRegistry, bPermitBackgroundWork);
 		}
 		FLocalQueueType* WorkerLocalQueue = LocalQueue;
 
 		bool Drowsing = false;
 		uint32 WaitCount = 0;
-		FQueueRegistry::FOutOfWork OutOfWork = QueueRegistry.GetOutOfWorkScope();
+		FQueueRegistry::FOutOfWork OutOfWork = QueueRegistry.GetOutOfWorkScope(bPermitBackgroundWork);
 		while (true)
 		{
 			while(TryExecuteTaskFrom<&FLocalQueueType::DequeueLocal>(WorkerLocalQueue, OutOfWork, bPermitBackgroundWork)
@@ -216,7 +218,7 @@ namespace LowLevelTasks
 
 		while (WakeUpWorker(bPermitBackgroundWork)) {}
 
-		FLocalQueueType::DeleteLocalQueue(WorkerLocalQueue, ExternalWorkerLocalQueue != nullptr);
+		FLocalQueueType::DeleteLocalQueue(WorkerLocalQueue, bPermitBackgroundWork, ExternalWorkerLocalQueue != nullptr);
 		LocalQueue = nullptr;
 
 		FMemory::ClearAndDisableTLSCachesOnCurrentThread();
@@ -233,7 +235,7 @@ namespace LowLevelTasks
 
 		uint32 WaitCount = 0;
 		bool bPermitBackgroundWork = ActiveTask && ActiveTask->IsBackgroundTask();
-		FQueueRegistry::FOutOfWork OutOfWork = QueueRegistry.GetOutOfWorkScope();
+		FQueueRegistry::FOutOfWork OutOfWork = QueueRegistry.GetOutOfWorkScope(bPermitBackgroundWork);
 		while (true)
 		{
 			while(TryExecuteTaskFrom<&FLocalQueueType::DequeueLocal>(WorkerLocalQueue, OutOfWork, bPermitBackgroundWork)
diff --git a/Engine/Source/Runtime/Core/Public/Async/Fundamental/LocalQueue.h b/Engine/Source/Runtime/Core/Public/Async/Fundamental/LocalQueue.h
--- a/Engine/Source/Runtime/Core/Public/Async/Fundamental/LocalQueue.h
+++ b/Engine/Source/Runtime/Core/Public/Async/Fundamental/LocalQueue.h
@@ -257,27 +257,27 @@ public:
 		friend class TLocalQueueRegistry;
 
 	public:
-		TLocalQueue(TLocalQueueRegistry& InRegistry) : Registry(&InRegistry)
+		TLocalQueue(TLocalQueueRegistry& InRegistry, bool bBackgroundWorker) : Registry(&InRegistry)
 		{
 			checkSlow(Registry);
 			StealHazard = FStealHazard(Registry->QueueCollection, Registry->HazardsCollection);
-			Registry->AddLocalQueue(StealHazard, this);
+			Registry->AddLocalQueue(StealHazard, this, bBackgroundWorker);
 			for (int32 PriorityIndex = 0; PriorityIndex < int32(ETaskPriority::Count); PriorityIndex++)
 			{
 				DequeueHazards[PriorityIndex] = Registry->OverflowQueues[PriorityIndex].getHeadHazard();
 			}
 		}
 
-		static TLocalQueue* AllocateLocalQueue(TLocalQueueRegistry& InRegistry)
+		static TLocalQueue* AllocateLocalQueue(TLocalQueueRegistry& InRegistry, bool bBackgroundWorker)
 		{
 			void* Memory = FMemory::Malloc(sizeof(TLocalQueue), 128u);
-			return new (Memory) TLocalQueue(InRegistry);
+			return new (Memory) TLocalQueue(InRegistry, bBackgroundWorker);
 		}
 
 		//delete a queue
 		//WorkeOwned means that the queue will not be automatically deleted when removal succeded.
 		//It is a special case where the memory for the local queues is alloated linearly by the Scheduler for improved stealing performance.
-		static void DeleteLocalQueue(TLocalQueue* Queue, bool WorkerOwned = false)
+		static void DeleteLocalQueue(TLocalQueue* Queue, bool bBackgroundWorker, bool WorkerOwned = false)
 		{
 			TLocalQueueRegistry* Registry = Queue->Registry;
 			Queue->Registry = nullptr;
@@ -294,7 +294,7 @@ public:
 					Registry->OverflowQueues[PriorityIndex].enqueue(Item);
 				}
 			}
-			Registry->DeleteLocalQueue(Queue->StealHazard, Queue, WorkerOwned);
+			Registry->DeleteLocalQueue(Queue->StealHazard, Queue, bBackgroundWorker, WorkerOwned);
 		}
 
 		//add an item to the local queue and overflow into the global queue if full
@@ -305,12 +305,13 @@ public:
 			checkSlow(PriorityIndex < int32(ETaskPriority::Count));
 			checkSlow(Item != nullptr);
 
+			bool bBackgroundTask = Item->IsBackgroundTask();
 			if (!LocalQueues[PriorityIndex].Put(Item))
 			{
 				Registry->OverflowQueues[PriorityIndex].enqueue(Item);
-				return Registry->AnyWorkerLookingForWork();
+				return Registry->AnyWorkerLookingForWork(bBackgroundTask);
 			}
-			return (LocalTasksSinceLastDequeue++ != 0) && Registry->AnyWorkerLookingForWork();
+			return (LocalTasksSinceLastDequeue++ != 0) && Registry->AnyWorkerLookingForWork(bBackgroundTask);
 		}
 
 		inline FTask* DequeueLocal(bool GetBackGroundTasks)
@@ -330,7 +331,7 @@ public:
 
 		inline FTask* DequeueGlobal(bool GetBackGroundTasks)
 		{
-			if (Registry->NumActiveWorkers.load(std::memory_order_relaxed) >= (2 * Registry->NumWorkersLookingForWork.load(std::memory_order_relaxed) - 1))
+			if (Registry->NumActiveWorkers[GetBackGroundTasks].load(std::memory_order_relaxed) >= (2 * Registry->NumWorkersLookingForWork[GetBackGroundTasks].load(std::memory_order_relaxed) - 1))
 			{
 				int32 MaxPriority = GetBackGroundTasks ? int32(ETaskPriority::Count) : int32(ETaskPriority::BackgroundNormal);
 				for (int32 PriorityIndex = 0; PriorityIndex < MaxPriority; PriorityIndex++)
@@ -347,7 +348,7 @@ public:
 
 		inline FTask* DequeueSteal(bool GetBackGroundTasks)
 		{
-			if (Registry->NumActiveWorkers.load(std::memory_order_relaxed) >= (2 * Registry->NumWorkersLookingForWork.load(std::memory_order_relaxed) - 1))
+			if (Registry->NumActiveWorkers[GetBackGroundTasks].load(std::memory_order_relaxed) >= (2 * Registry->NumWorkersLookingForWork[GetBackGroundTasks].load(std::memory_order_relaxed) - 1))
 			{
 				if (CachedRandomIndex == InvalidIndex)
 				{
@@ -382,9 +383,9 @@ public:
 
 private:
 	// add a queue to the Registry
-	void AddLocalQueue(FStealHazard& Hazard, TLocalQueue* QueueToAdd)
+	void AddLocalQueue(FStealHazard& Hazard, TLocalQueue* QueueToAdd, bool bBackgroundWorker)
 	{
-		NumActiveWorkers.fetch_add(1, std::memory_order_relaxed);
+		NumActiveWorkers[bBackgroundWorker].fetch_add(1, std::memory_order_relaxed);
 		while(true)
 		{
 			FLocalQueueCollection* Previous = Hazard.Get();
@@ -403,9 +404,9 @@ private:
 	}
 
 	// remove a queue from the Registry
-	void DeleteLocalQueue(FStealHazard& Hazard, TLocalQueue* QueueToRemove, bool WorkerOwned)
+	void DeleteLocalQueue(FStealHazard& Hazard, TLocalQueue* QueueToRemove, bool bBackgroundWorker, bool WorkerOwned)
 	{
-		NumActiveWorkers.fetch_sub(1, std::memory_order_relaxed);
+		NumActiveWorkers[bBackgroundWorker].fetch_sub(1, std::memory_order_relaxed);
 		while(true)
 		{
 			FLocalQueueCollection* Previous = Hazard.Get();
@@ -466,9 +467,10 @@ public:
 		check(PriorityIndex < int32(ETaskPriority::Count));
 		check(Item != nullptr);
 
+		bool bBackgroundTask = Item->IsBackgroundTask();
 		OverflowQueues[PriorityIndex].enqueue(Item);
 
-		return AnyWorkerLookingForWork();
+		return AnyWorkerLookingForWork(bBackgroundTask);
 	}
 
 	// grab an Item directy from the Global OverflowQueue
@@ -485,22 +487,22 @@ public:
 		return nullptr;
 	}
 
-	inline FOutOfWork GetOutOfWorkScope()
+	inline FOutOfWork GetOutOfWorkScope(bool bBackgroundWorker)
 	{
-		return FOutOfWork(NumWorkersLookingForWork);
+		return FOutOfWork(NumWorkersLookingForWork[bBackgroundWorker]);
 	}
 
-	inline bool AnyWorkerLookingForWork() const
+private:
+	inline bool AnyWorkerLookingForWork(bool bBackgroundWorker) const
 	{
-		return NumWorkersLookingForWork.load(std::memory_order_acquire) == 0;
+		return NumWorkersLookingForWork[bBackgroundWorker].load(std::memory_order_acquire) == 0;
 	}
 
-private:
 	FOverflowQueueType	  OverflowQueues[uint32(ETaskPriority::Count)];
 	FHazardPointerCollection		  HazardsCollection;
 	std::atomic<FLocalQueueCollection*>	QueueCollection;
-	std::atomic_int NumWorkersLookingForWork { 0 };
-	std::atomic_int NumActiveWorkers { 0 };
+	std::atomic_int NumWorkersLookingForWork[2] { 0 };
+	std::atomic_int NumActiveWorkers[2] { 0 };
 };
 
 template<uint32 NumLocalItems>
