Fix compile errors due to incorrect initialization of TLocalQueueRegistry::NumWorkersLookingForWork and TLocalQueueRegistry::NumActiveWorkers.

diff --git a/Engine/Source/Runtime/Core/Public/Async/Fundamental/LocalQueue.h b/Engine/Source/Runtime/Core/Public/Async/Fundamental/LocalQueue.h
--- a/Engine/Source/Runtime/Core/Public/Async/Fundamental/LocalQueue.h
+++ b/Engine/Source/Runtime/Core/Public/Async/Fundamental/LocalQueue.h
@@ -501,8 +501,8 @@ private:
 	FOverflowQueueType	  OverflowQueues[uint32(ETaskPriority::Count)];
 	FHazardPointerCollection		  HazardsCollection;
 	std::atomic<FLocalQueueCollection*>	QueueCollection;
-	std::atomic_int NumWorkersLookingForWork[2] { 0 };
-	std::atomic_int NumActiveWorkers[2] { 0 };
+	std::atomic_int NumWorkersLookingForWork[2] = { {0}, {0} };
+	std::atomic_int NumActiveWorkers[2] = { {0}, {0} };
 };
 
 template<uint32 NumLocalItems>
