U5 - Tweaking Priorities and worker waking

diff --git a/Engine/Source/Runtime/Core/Private/Async/Fundamental/Scheduler.cpp b/Engine/Source/Runtime/Core/Private/Async/Fundamental/Scheduler.cpp
--- a/Engine/Source/Runtime/Core/Private/Async/Fundamental/Scheduler.cpp
+++ b/Engine/Source/Runtime/Core/Private/Async/Fundamental/Scheduler.cpp
@@ -113,8 +113,14 @@ namespace LowLevelTasks
 	void FScheduler::LaunchInternal(FTask& Task, EQueuePreference QueuePreference)
 	{
 		if (ActiveWorkers.load(std::memory_order_relaxed))
-		{
-			bool bIsBackgroundTask = Task.IsBackgroundTask();
+		{			
+			const bool bIsBackgroundTask = Task.IsBackgroundTask();
+			const bool bIsBackgroundWorker = ActiveTask && ActiveTask->IsBackgroundTask();
+			if (bIsBackgroundTask && !bIsBackgroundWorker)
+			{
+				QueuePreference = EQueuePreference::GlobalQueuePreference;
+			}
+
 			if (LocalQueue && QueuePreference != EQueuePreference::GlobalQueuePreference)
 			{
 				if (LocalQueue->Enqueue(&Task, uint32(Task.GetPriority())))
@@ -156,7 +162,10 @@ namespace LowLevelTasks
 			OutOfWork.Stop();		
 			FTask* OldTask = ActiveTask;
 			ActiveTask = Task;
-			Task->ExecuteTask();
+			{
+				TRACE_CPUPROFILER_EVENT_SCOPE(ExecuteTask);
+				Task->ExecuteTask();
+			}
 			ActiveTask = OldTask;
 			return true;
 		}
diff --git a/Engine/Source/Runtime/Core/Public/Async/Fundamental/LocalQueue.h b/Engine/Source/Runtime/Core/Public/Async/Fundamental/LocalQueue.h
--- a/Engine/Source/Runtime/Core/Public/Async/Fundamental/LocalQueue.h
+++ b/Engine/Source/Runtime/Core/Public/Async/Fundamental/LocalQueue.h
@@ -495,7 +495,7 @@ public:
 private:
 	inline bool AnyWorkerLookingForWork(bool bBackgroundWorker) const
 	{
-		return NumWorkersLookingForWork[bBackgroundWorker].load(std::memory_order_acquire) == 0;
+		return (NumWorkersLookingForWork[bBackgroundWorker].load(std::memory_order_acquire) == 0) && (bBackgroundWorker || (NumWorkersLookingForWork[true].load(std::memory_order_acquire) == 0));
 	}
 
 	FOverflowQueueType	  OverflowQueues[uint32(ETaskPriority::Count)];
diff --git a/Engine/Source/Runtime/Core/Public/Misc/QueuedThreadPoolWrapper.h b/Engine/Source/Runtime/Core/Public/Misc/QueuedThreadPoolWrapper.h
--- a/Engine/Source/Runtime/Core/Public/Misc/QueuedThreadPoolWrapper.h
+++ b/Engine/Source/Runtime/Core/Public/Misc/QueuedThreadPoolWrapper.h
@@ -293,15 +293,13 @@ public:
 private:
 	void ScheduleTasks()
 	{
-		LowLevelTasks::EQueuePreference QueuePreference = LowLevelTasks::EQueuePreference::LocalQueuePreference;
 		while (!bIsPaused && TaskCount < MaxConcurrency)
 		{
 			FQueuedWorkInternalData* QueuedWork = Dequeue();
 			if (QueuedWork)
 			{
 				TaskCount++;
-				verifySlow(LowLevelTasks::TryLaunch(QueuedWork->Task, QueuePreference));
-				QueuePreference = LowLevelTasks::EQueuePreference::GlobalQueuePreference;
+				verifySlow(LowLevelTasks::TryLaunch(QueuedWork->Task, LowLevelTasks::EQueuePreference::LocalQueuePreference));
 			}
 			else
 			{
