U5 - avoid Taskgraph from waking up workers while ppcessing the first subsequent as often the worker will not wake up fast enough to grab that task. Also shorten the time of spinwaiting if a worker was woken up that could not pick up any work.

diff --git a/Engine/Source/Runtime/Core/Private/Async/Fundamental/Scheduler.cpp b/Engine/Source/Runtime/Core/Private/Async/Fundamental/Scheduler.cpp
--- a/Engine/Source/Runtime/Core/Private/Async/Fundamental/Scheduler.cpp
+++ b/Engine/Source/Runtime/Core/Private/Async/Fundamental/Scheduler.cpp
@@ -110,7 +110,7 @@ namespace LowLevelTasks
 		}
 	}
 
-	void FScheduler::LaunchInternal(FTask& Task, EQueuePreference QueuePreference)
+	void FScheduler::LaunchInternal(FTask& Task, EQueuePreference QueuePreference, bool bWakeUpWorker)
 	{
 		if (ActiveWorkers.load(std::memory_order_relaxed))
 		{			
@@ -121,11 +121,13 @@ namespace LowLevelTasks
 				QueuePreference = EQueuePreference::GlobalQueuePreference;
 			}
 
+			bWakeUpWorker |= LocalQueue == nullptr;
+
 			if (LocalQueue && QueuePreference != EQueuePreference::GlobalQueuePreference)
 			{
 				if (LocalQueue->Enqueue(&Task, uint32(Task.GetPriority())))
 				{
-					if (!WakeUpWorker(bIsBackgroundTask) && !bIsBackgroundTask)
+					if(bWakeUpWorker && !WakeUpWorker(bIsBackgroundTask) && !bIsBackgroundTask)
 					{
 						WakeUpWorker(true);
 					}
@@ -135,7 +137,7 @@ namespace LowLevelTasks
 			{
 				if (QueueRegistry.Enqueue(&Task, uint32(Task.GetPriority())))
 				{
-					if (!WakeUpWorker(bIsBackgroundTask) && !bIsBackgroundTask)
+					if (bWakeUpWorker && !WakeUpWorker(bIsBackgroundTask) && !bIsBackgroundTask)
 					{
 						WakeUpWorker(true);
 					}
diff --git a/Engine/Source/Runtime/Core/Private/Async/TaskGraph.cpp b/Engine/Source/Runtime/Core/Private/Async/TaskGraph.cpp
--- a/Engine/Source/Runtime/Core/Private/Async/TaskGraph.cpp
+++ b/Engine/Source/Runtime/Core/Private/Async/TaskGraph.cpp
@@ -1360,7 +1360,7 @@ public:
 	 *	@param	ThreadToExecuteOn; Either a named thread for a threadlocked task or ENamedThreads::AnyThread for a task that is to run on a worker thread
 	 *	@param	CurrentThreadIfKnown; This should be the current thread if it is known, or otherwise use ENamedThreads::AnyThread and the current thread will be determined.
 	**/
-	virtual void QueueTask(FBaseGraphTask* Task, ENamedThreads::Type ThreadToExecuteOn, ENamedThreads::Type InCurrentThreadIfKnown = ENamedThreads::AnyThread) final override
+	virtual void QueueTask(FBaseGraphTask* Task, bool bWakeUpWorker, ENamedThreads::Type ThreadToExecuteOn, ENamedThreads::Type InCurrentThreadIfKnown = ENamedThreads::AnyThread) final override
 	{
 		TASKGRAPH_SCOPE_CYCLE_COUNTER(2, STAT_TaskGraph_QueueTask);
 
@@ -1852,7 +1852,7 @@ public:
 	}
 
 private:
-	void QueueTask(class FBaseGraphTask* Task, ENamedThreads::Type InThreadToExecuteOn, ENamedThreads::Type InCurrentThreadIfKnown) override
+	void QueueTask(class FBaseGraphTask* Task, bool bWakeUpWorker, ENamedThreads::Type InThreadToExecuteOn, ENamedThreads::Type InCurrentThreadIfKnown) override
 	{
 		if (ENamedThreads::GetThreadIndex(InThreadToExecuteOn) == ENamedThreads::AnyThread)
 		{
@@ -1871,7 +1871,7 @@ private:
 				Task->Execute(NewTasks, InThreadToExecuteOn, false);
 			});
 	
-			verifySlow(LowLevelTasks::TryLaunch(Task->TaskHandle, LowLevelTasks::EQueuePreference::GlobalQueuePreference));
+			verifySlow(LowLevelTasks::TryLaunch(Task->TaskHandle, bWakeUpWorker ? LowLevelTasks::EQueuePreference::GlobalQueuePreference : LowLevelTasks::EQueuePreference::LocalQueuePreference, bWakeUpWorker));
 			return;
 		}
 
@@ -2248,12 +2248,13 @@ void FGraphEvent::DispatchSubsequents(TArray<FBaseGraphTask*>& NewTasks, ENamedT
 		}
 	}
 
+	bool bWakeUpWorker = false;
 	SubsequentList.PopAllAndClose(NewTasks);
 	for (int32 Index = NewTasks.Num() - 1; Index >= 0 ; Index--) // reverse the order since PopAll is implicitly backwards
 	{
 		FBaseGraphTask* NewTask = NewTasks[Index];
 		checkThreadGraph(NewTask);
-		NewTask->ConditionalQueueTask(CurrentThreadIfKnown);
+		NewTask->ConditionalQueueTask(CurrentThreadIfKnown, bWakeUpWorker);
 	}
 	NewTasks.Reset();
 }
diff --git a/Engine/Source/Runtime/Core/Public/Async/Fundamental/LocalQueue.h b/Engine/Source/Runtime/Core/Public/Async/Fundamental/LocalQueue.h
--- a/Engine/Source/Runtime/Core/Public/Async/Fundamental/LocalQueue.h
+++ b/Engine/Source/Runtime/Core/Public/Async/Fundamental/LocalQueue.h
@@ -32,8 +32,9 @@ protected:
 		uint32 Idx = (Head + 1) % NumItems;
 		uintptr_t Slot = ItemSlots[Idx].Value.load(std::memory_order_relaxed);
 
-		if (Slot == uintptr_t(ESlotState::Free) && ItemSlots[Idx].Value.compare_exchange_strong(Slot, Item, std::memory_order_relaxed))
+		if (Slot == uintptr_t(ESlotState::Free))
 		{
+			ItemSlots[Idx].Value.store(Item, std::memory_order_relaxed);
 			Head++;
 			checkSlow(Head % NumItems == Idx);
 			return true;		
@@ -85,35 +86,6 @@ protected:
 		} while(true);
 	}
 
-	/* Untested
-	//add an item at the tail position (this can be done from any thread including the one that accesses the head)
-	inline bool PutTail(uintptr_t Item)
-	{
-		do
-		{
-			uint32 IdxVer = Tail.load(std::memory_order_relaxed);
-			uint32 Idx = (IdxVer - 1) % NumItems;
-			uintptr_t Slot = ItemSlots[Idx].Value.load(std::memory_order_relaxed);
-
-			if (Slot > uintptr_t(ESlotState::Taken))
-			{
-				return false;
-			}
-			else if (Slot == uintptr_t(ESlotState::Free) && ItemSlots[Idx].Value.compare_exchange_weak(Slot, uintptr_t(ESlotState::Taken), std::memory_order_relaxed))
-			{
-				if(IdxVer == Tail.load(std::memory_order_relaxed))
-				{
-					uint32 Prev = Tail.fetch_sub(1, std::memory_order_relaxed);
-					checkSlow((Prev - 1) % NumItems == Idx);
-					ItemSlots[Idx].Value.store(Item, std::memory_order_relaxed);
-					return true;
-				}
-				ItemSlots[Idx].Value.store(uintptr_t(ESlotState::Free), std::memory_order_relaxed);
-			}
-		} while(true);
-	}
-	*/
-
 private:
 	struct FAlignedElement
 	{
@@ -309,14 +281,12 @@ public:
 			if (!LocalQueues[PriorityIndex].Put(Item))
 			{
 				Registry->OverflowQueues[PriorityIndex].enqueue(Item);
-				return Registry->AnyWorkerLookingForWork(bBackgroundTask);
 			}
-			return (LocalTasksSinceLastDequeue++ != 0) && Registry->AnyWorkerLookingForWork(bBackgroundTask);
+			return Registry->AnyWorkerLookingForWork(bBackgroundTask);
 		}
 
 		inline FTask* DequeueLocal(bool GetBackGroundTasks)
 		{
-			LocalTasksSinceLastDequeue = 0;
 			int32 MaxPriority = GetBackGroundTasks ? int32(ETaskPriority::Count) : int32(ETaskPriority::BackgroundNormal);
 			for (int32 PriorityIndex = 0; PriorityIndex < MaxPriority; PriorityIndex++)
 			{
@@ -373,7 +343,6 @@ public:
 		FRandomStream			Random;
 		uint32					CachedRandomIndex = InvalidIndex;
 		uint32					CachedPriorityIndex = 0;
-		uint32					LocalTasksSinceLastDequeue = 0;
 	};
 
 	TLocalQueueRegistry()
diff --git a/Engine/Source/Runtime/Core/Public/Async/Fundamental/Scheduler.h b/Engine/Source/Runtime/Core/Public/Async/Fundamental/Scheduler.h
--- a/Engine/Source/Runtime/Core/Public/Async/Fundamental/Scheduler.h
+++ b/Engine/Source/Runtime/Core/Public/Async/Fundamental/Scheduler.h
@@ -51,7 +51,7 @@ namespace LowLevelTasks
 		CORE_API void StopWorkers();
 
 		//try to launch the task, the return value will specify if the task was in the ready state and has been launced
-		inline bool TryLaunch(FTask& Task, EQueuePreference QueuePreference = EQueuePreference::DefaultPreference);
+		inline bool TryLaunch(FTask& Task, EQueuePreference QueuePreference = EQueuePreference::DefaultPreference, bool bWakeUpWorker = true);
 		
 		//try to cancel the task and launching it if the task was in ready state
 		//you still need to wait for task completion before recycling the handle
@@ -99,7 +99,7 @@ namespace LowLevelTasks
 	private: 
 		TUniquePtr<FThread> CreateWorker(FLocalQueueType* ExternalWorkerLocalQueue = nullptr, EThreadPriority Priority = EThreadPriority::TPri_Normal, bool bPermitBackgroundWork = false, bool bIsForkable = false);
 		void WorkerMain(struct FSleepEvent* WorkerEvent, FLocalQueueType* ExternalWorkerLocalQueue, uint32 WaitCycles, bool bPermitBackgroundWork);
-		CORE_API void LaunchInternal(FTask& Task, EQueuePreference QueuePreference);
+		CORE_API void LaunchInternal(FTask& Task, EQueuePreference QueuePreference, bool bWakeUpWorker);
 		CORE_API void BusyWaitInternal(const FConditional& Conditional);
 		FORCENOINLINE void TrySleeping(FSleepEvent* WorkerEvent, FQueueRegistry::FOutOfWork& OutOfWork, uint32& WaitCount, bool& Drowsing, bool bBackgroundWorker);
 		inline bool WakeUpWorker(bool bBackgroundWorker);
@@ -117,9 +117,9 @@ namespace LowLevelTasks
 		std::atomic_uint			NextWorkerId { 0 };
 	};
 
-	FORCEINLINE_DEBUGGABLE bool TryLaunch(FTask& Task, EQueuePreference QueuePreference = EQueuePreference::DefaultPreference)
+	FORCEINLINE_DEBUGGABLE bool TryLaunch(FTask& Task, EQueuePreference QueuePreference = EQueuePreference::DefaultPreference, bool bWakeUpWorker = true)
 	{
-		return FScheduler::Get().TryLaunch(Task, QueuePreference);
+		return FScheduler::Get().TryLaunch(Task, QueuePreference, bWakeUpWorker);
 	}
 
 	FORCEINLINE_DEBUGGABLE bool TryCancelAndLaunchContinuation(FTask& Task, EQueuePreference QueuePreference = EQueuePreference::DefaultPreference)
@@ -147,11 +147,11 @@ namespace LowLevelTasks
    /******************
 	* IMPLEMENTATION *
 	******************/
-	inline bool FScheduler::TryLaunch(FTask& Task, EQueuePreference QueuePreference)
+	inline bool FScheduler::TryLaunch(FTask& Task, EQueuePreference QueuePreference, bool bWakeUpWorker)
 	{
 		if(Task.TryPrepareLaunch())
 		{
-			FScheduler::Get().LaunchInternal(Task, QueuePreference);
+			FScheduler::Get().LaunchInternal(Task, QueuePreference, bWakeUpWorker);
 			return true;
 		}
 		return false;
@@ -233,9 +233,9 @@ namespace LowLevelTasks
 		else if(WorkerEvent->SleepState.compare_exchange_strong(DrowsingState2, ESleepState::Sleeping, std::memory_order_relaxed))
 		{
 			verifySlow(!OutOfWork.Stop());
-			WaitCount = 0;
 			Drowsing = false;
 			WorkerEvent->SleepEvent->Wait(); // State two: ((Running -> Drowsing) -> Sleeping)
+			verifySlow(OutOfWork.Start());
 		}
 		else if(WorkerEvent->SleepState.compare_exchange_strong(RunningState, ESleepState::Drowsing, std::memory_order_relaxed))
 		{
diff --git a/Engine/Source/Runtime/Core/Public/Async/TaskGraphInterfaces.h b/Engine/Source/Runtime/Core/Public/Async/TaskGraphInterfaces.h
--- a/Engine/Source/Runtime/Core/Public/Async/TaskGraphInterfaces.h
+++ b/Engine/Source/Runtime/Core/Public/Async/TaskGraphInterfaces.h
@@ -288,7 +288,7 @@ class FTaskGraphInterface
 	 *	@param	ThreadToExecuteOn; Either a named thread for a threadlocked task or ENamedThreads::AnyThread for a task that is to run on a worker thread
 	 *	@param	CurrentThreadIfKnown; This should be the current thread if it is known, or otherwise use ENamedThreads::AnyThread and the current thread will be determined.
 	**/
-	virtual void QueueTask(class FBaseGraphTask* Task, ENamedThreads::Type ThreadToExecuteOn, ENamedThreads::Type CurrentThreadIfKnown = ENamedThreads::AnyThread) = 0;
+	virtual void QueueTask(class FBaseGraphTask* Task, bool bWakeUpWorker, ENamedThreads::Type ThreadToExecuteOn, ENamedThreads::Type CurrentThreadIfKnown = ENamedThreads::AnyThread) = 0;
 
 public:
 
@@ -478,7 +478,8 @@ protected:
 		int32 NumToSub = NumAlreadyFinishedPrequistes + (bUnlock ? 1 : 0); // the +1 is for the "lock" we set up in the constructor
 		if (NumberOfPrerequistitesOutstanding.Subtract(NumToSub) == NumToSub) 
 		{
-			QueueTask(CurrentThread);
+			bool bWakeUpWorker = true;
+			QueueTask(CurrentThread, bWakeUpWorker);	
 		}
 	}
 	/** destructor, just checks the life stage **/
@@ -499,11 +500,12 @@ protected:
 	 *	An indication that a prerequisite has been completed. Reduces the number of prerequisites by one and if no prerequisites are outstanding, it queues the task for execution.
 	 *	@param CurrentThread; provides the index of the thread we are running on. This is handy for submitting new taks. Can be ENamedThreads::AnyThread if the current thread is unknown.
 	 **/
-	void ConditionalQueueTask(ENamedThreads::Type CurrentThread)
+	void ConditionalQueueTask(ENamedThreads::Type CurrentThread, bool& bWakeUpWorker)
 	{
 		if (NumberOfPrerequistitesOutstanding.Decrement()==0)
 		{
-			QueueTask(CurrentThread);
+			QueueTask(CurrentThread, bWakeUpWorker);
+			bWakeUpWorker = true;
 		}
 	}
 
@@ -552,10 +554,10 @@ private:
 	 *	Queues the task for execution.
 	 *	@param CurrentThread; provides the index of the thread we are running on. This is handy for submitting new taks. Can be ENamedThreads::AnyThread if the current thread is unknown.
 	 **/
-	void QueueTask(ENamedThreads::Type CurrentThreadIfKnown)
+	void QueueTask(ENamedThreads::Type CurrentThreadIfKnown, bool bWakeUpWorker)
 	{
 		checkThreadGraph(LifeStage.Increment() == int32(LS_Queued));
-		FTaskGraphInterface::Get().QueueTask(this, ThreadToExecuteOn, CurrentThreadIfKnown);
+		FTaskGraphInterface::Get().QueueTask(this, bWakeUpWorker, ThreadToExecuteOn, CurrentThreadIfKnown);
 	}
 
 	/**	Thread to execute on, can be ENamedThreads::AnyThread to execute on any unnamed thread **/
@@ -869,7 +871,8 @@ public:
 
 	void Unlock(ENamedThreads::Type CurrentThreadIfKnown = ENamedThreads::AnyThread)
 	{
-		ConditionalQueueTask(CurrentThreadIfKnown);
+		bool bWakeUpWorker = true;
+		ConditionalQueueTask(CurrentThreadIfKnown, bWakeUpWorker);
 	}
 
 	FGraphEventRef GetCompletionEvent()
diff --git a/Engine/Source/Runtime/Core/Public/Misc/QueuedThreadPoolWrapper.h b/Engine/Source/Runtime/Core/Public/Misc/QueuedThreadPoolWrapper.h
--- a/Engine/Source/Runtime/Core/Public/Misc/QueuedThreadPoolWrapper.h
+++ b/Engine/Source/Runtime/Core/Public/Misc/QueuedThreadPoolWrapper.h
@@ -287,11 +287,12 @@ public:
 			bIsPaused = false;
 		}
 
-		ScheduleTasks();
+		bool bWakeUpWorker = true;
+		ScheduleTasks(bWakeUpWorker);
 	}
 
 private:
-	void ScheduleTasks()
+	void ScheduleTasks(bool &bWakeUpWorker)
 	{
 		while (!bIsPaused && TaskCount < MaxConcurrency)
 		{
@@ -299,7 +300,8 @@ private:
 			if (QueuedWork)
 			{
 				TaskCount++;
-				verifySlow(LowLevelTasks::TryLaunch(QueuedWork->Task, LowLevelTasks::EQueuePreference::LocalQueuePreference));
+				verifySlow(LowLevelTasks::TryLaunch(QueuedWork->Task, bWakeUpWorker ? LowLevelTasks::EQueuePreference::GlobalQueuePreference : LowLevelTasks::EQueuePreference::LocalQueuePreference, bWakeUpWorker));
+				bWakeUpWorker = true;
 			}
 			else
 			{
@@ -325,12 +327,13 @@ private:
 		[this, InternalData = InQueuedWork->InternalData]()
 		{
 			--TaskCount;
-			LowLevelTasks::EQueuePreference QueuePreference = LowLevelTasks::EQueuePreference::LocalQueuePreference;
-			ScheduleTasks();
+			bool bWakeUpWorker = false;
+			ScheduleTasks(bWakeUpWorker);
 		});
 
 		Enqueue(Priority, QueuedWorkInternalData);
-		ScheduleTasks();	
+		bool bWakeUpWorker = true;
+		ScheduleTasks(bWakeUpWorker);	
 	}
 
 	bool RetractQueuedWork(IQueuedWork* InQueuedWork) override
