U5 - remove now unused function parameter

diff --git a/Engine/Source/Runtime/Core/Private/Async/Fundamental/Scheduler.cpp b/Engine/Source/Runtime/Core/Private/Async/Fundamental/Scheduler.cpp
--- a/Engine/Source/Runtime/Core/Private/Async/Fundamental/Scheduler.cpp
+++ b/Engine/Source/Runtime/Core/Private/Async/Fundamental/Scheduler.cpp
@@ -224,7 +224,7 @@ namespace LowLevelTasks
 				continue;
 			}
 
-			TrySleeping(WorkerEvent, OutOfWork, WaitCount, Drowsing, bPermitBackgroundWork);
+			TrySleeping(WorkerEvent, OutOfWork, Drowsing, bPermitBackgroundWork);
 		}
 
 		while (WakeUpWorker(bPermitBackgroundWork)) {}
diff --git a/Engine/Source/Runtime/Core/Public/Async/Fundamental/Scheduler.h b/Engine/Source/Runtime/Core/Public/Async/Fundamental/Scheduler.h
--- a/Engine/Source/Runtime/Core/Public/Async/Fundamental/Scheduler.h
+++ b/Engine/Source/Runtime/Core/Public/Async/Fundamental/Scheduler.h
@@ -101,7 +101,7 @@ namespace LowLevelTasks
 		void WorkerMain(struct FSleepEvent* WorkerEvent, FLocalQueueType* ExternalWorkerLocalQueue, uint32 WaitCycles, bool bPermitBackgroundWork);
 		CORE_API void LaunchInternal(FTask& Task, EQueuePreference QueuePreference, bool bWakeUpWorker);
 		CORE_API void BusyWaitInternal(const FConditional& Conditional);
-		FORCENOINLINE void TrySleeping(FSleepEvent* WorkerEvent, FQueueRegistry::FOutOfWork& OutOfWork, uint32& WaitCount, bool& Drowsing, bool bBackgroundWorker);
+		FORCENOINLINE void TrySleeping(FSleepEvent* WorkerEvent, FQueueRegistry::FOutOfWork& OutOfWork, bool& Drowsing, bool bBackgroundWorker);
 		inline bool WakeUpWorker(bool bBackgroundWorker);
 
 		template<FTask* (FLocalQueueType::*DequeueFunction)(bool)>
@@ -219,7 +219,7 @@ namespace LowLevelTasks
 		}
 	}
 
-	inline void FScheduler::TrySleeping(FSleepEvent* WorkerEvent, FQueueRegistry::FOutOfWork& OutOfWork, uint32& WaitCount, bool& Drowsing, bool bBackgroundWorker)
+	inline void FScheduler::TrySleeping(FSleepEvent* WorkerEvent, FQueueRegistry::FOutOfWork& OutOfWork, bool& Drowsing, bool bBackgroundWorker)
 	{
 		ESleepState DrowsingState1 = ESleepState::Drowsing;
 		ESleepState DrowsingState2 = ESleepState::Drowsing;
