Fixed comment for TRACE_CPUPROFILER_EVENT_SCOPE macro.

diff --git a/Engine/Source/Runtime/Core/Public/ProfilingDebugging/CpuProfilerTrace.h b/Engine/Source/Runtime/Core/Public/ProfilingDebugging/CpuProfilerTrace.h
--- a/Engine/Source/Runtime/Core/Public/ProfilingDebugging/CpuProfilerTrace.h
+++ b/Engine/Source/Runtime/Core/Public/ProfilingDebugging/CpuProfilerTrace.h
@@ -148,7 +148,7 @@ struct FCpuProfilerTrace
 #define TRACE_CPUPROFILER_EVENT_SCOPE_STR(NameStr) \
 	TRACE_CPUPROFILER_EVENT_SCOPE_ON_CHANNEL_STR(NameStr, CpuChannel)
 
-// Trace a scoped cpu timing event providing a scope name (plain text) and a trace channel.
+// Trace a scoped cpu timing event providing a scope name (plain text).
 // It will use the Cpu trace channel.
 // Example: TRACE_CPUPROFILER_EVENT_SCOPE(MyScopedTimer::A)
 // Note: Do not use this macro with a static string because, in that case, additional quotes will
