Prototype Wrapped Object Pointers (ObjectHandle/ObjectPtr)

Treat wrapped object pointers identically to raw pointers in the base FArchive implementation (they are silently ignored for serialization and do not generate an error).  Some code paths (eg: Niagara) depend on this behavior for archives derived from FArchive (not FArchiveUObject).

diff --git a/Engine/Source/Runtime/Core/Private/Serialization/Archive.cpp b/Engine/Source/Runtime/Core/Private/Serialization/Archive.cpp
--- a/Engine/Source/Runtime/Core/Private/Serialization/Archive.cpp
+++ b/Engine/Source/Runtime/Core/Private/Serialization/Archive.cpp
@@ -22,6 +22,7 @@
 #include "Interfaces/ITargetPlatform.h"
 #include "Serialization/CompressedChunkInfo.h"
 #include "Serialization/ArchiveSerializedPropertyChain.h"
+#include "UObject/NameTypes.h"
 
 PRAGMA_DISABLE_UNSAFE_TYPECAST_WARNINGS
 
@@ -461,8 +462,6 @@ FArchive& FArchive::operator<<(struct FLazyObjectPtr& Value)
 
 FArchive& FArchive::operator<<(struct FObjectPtr& Value)
 {
-	// The base FArchive does not implement this method. Use FArchiveUObject instead.
-	UE_LOG(LogSerialization, Fatal, TEXT("FArchive does not support FObjectPtr serialization. Use FArchiveUObject instead."));
 	return *this;
 }
 
