Zero initialize FDateTime and FTimespan

diff --git a/Engine/Source/Runtime/Core/Public/Misc/DateTime.h b/Engine/Source/Runtime/Core/Public/Misc/DateTime.h
--- a/Engine/Source/Runtime/Core/Public/Misc/DateTime.h
+++ b/Engine/Source/Runtime/Core/Public/Misc/DateTime.h
@@ -74,8 +74,10 @@ struct FDateTime
 {
 public:
 
-	/** Default constructor (no initialization). */
-	FDateTime() { }
+	/** Default constructor (zero initialization). */
+	FDateTime()
+		: Ticks(0)
+	{ }
 
 	/**
 	 * Creates and initializes a new instance with the specified number of ticks.
diff --git a/Engine/Source/Runtime/Core/Public/Misc/Timespan.h b/Engine/Source/Runtime/Core/Public/Misc/Timespan.h
--- a/Engine/Source/Runtime/Core/Public/Misc/Timespan.h
+++ b/Engine/Source/Runtime/Core/Public/Misc/Timespan.h
@@ -75,8 +75,10 @@ struct FTimespan
 {
 public:
 
-	/** Default constructor (no initialization). */
-	FTimespan() { }
+	/** Default constructor (zero initialization). */
+	FTimespan()
+		: Ticks(0)
+	{ }
 
 	/**
 	 * Create and initialize a new time interval with the specified number of ticks.
