TArrayView::Slice: Allow InNum == 0 so that callers can robustly create possibly empty slices.
#rnx

diff --git a/Engine/Source/Runtime/Core/Public/Containers/ArrayView.h b/Engine/Source/Runtime/Core/Public/Containers/ArrayView.h
--- a/Engine/Source/Runtime/Core/Public/Containers/ArrayView.h
+++ b/Engine/Source/Runtime/Core/Public/Containers/ArrayView.h
@@ -252,6 +252,20 @@ public:
 		checkf((Index >= 0) & (Index < ArrayNum),TEXT("Array index out of bounds: %i from an array of size %i"),Index,ArrayNum); // & for one branch
 	}
 
+	/**
+	 * Checks if a slice range [Index, Index+InNum) is in array range.
+	 * Length is 0 is allowed on empty arrays; Index must be 0 in that case.
+	 *
+	 * @param Index Starting index of the slice.
+	 * @param InNum Length of the slice.
+	 */
+	FORCEINLINE void SliceRangeCheck(SizeType Index, SizeType InNum) const
+	{
+		checkf(Index >= 0, TEXT("Invalid index (%d)"), Index);
+		checkf(InNum >= 0, TEXT("Invalid count (%d)"), InNum);
+		checkf(Index + InNum <= ArrayNum, TEXT("Range (index: %d, count: %d) lies outside the view of %d elements"), Index, InNum, ArrayNum);
+	}
+
 	/**
 	 * Tests if index is valid, i.e. than or equal to zero, and less than the number of elements in the array.
 	 *
@@ -319,9 +333,7 @@ public:
 	 */
 	FORCEINLINE TArrayView Slice(SizeType Index, SizeType InNum) const
 	{
-		check(InNum > 0);
-		check(IsValidIndex(Index));
-		check(IsValidIndex(Index + InNum - 1));
+		SliceRangeCheck(Index, InNum);
 		return TArrayView(DataPtr + Index, InNum);
 	}
 
