IoDispatcher: Added support for external backends
#rnx

diff --git a/Engine/Source/Developer/IoStoreUtilities/Private/IoStoreUtilities.cpp b/Engine/Source/Developer/IoStoreUtilities/Private/IoStoreUtilities.cpp
--- a/Engine/Source/Developer/IoStoreUtilities/Private/IoStoreUtilities.cpp
+++ b/Engine/Source/Developer/IoStoreUtilities/Private/IoStoreUtilities.cpp
@@ -553,7 +553,6 @@ struct FContainerTargetSpec
 	FString OutputPath;
 	FIoStoreWriter* IoStoreWriter;
 	TArray<FContainerTargetFile> TargetFiles;
-	TUniquePtr<FIoStoreEnvironment> IoStoreEnv;
 	TArray<TUniquePtr<FIoStoreReader>> PatchSourceReaders;
 	FNameMapBuilder LocalNameMapBuilder;
 	FNameMapBuilder* NameMapBuilder = &LocalNameMapBuilder;
@@ -2989,8 +2988,6 @@ static void ProcessLocalizedPackages(
 
 TUniquePtr<FIoStoreReader> CreateIoStoreReader(const TCHAR* Path, const FKeyChain& KeyChain)
 {
-	FIoStoreEnvironment IoEnvironment;
-	IoEnvironment.InitializeFileEnvironment(FPaths::ChangeExtension(Path, TEXT("")));
 	TUniquePtr<FIoStoreReader> IoStoreReader(new FIoStoreReader());
 
 	TMap<FGuid, FAES::FAESKey> DecryptionKeys;
@@ -2998,7 +2995,7 @@ TUniquePtr<FIoStoreReader> CreateIoStoreReader(const TCHAR* Path, const FKeyChai
 	{
 		DecryptionKeys.Add(KV.Key, KV.Value.Key);
 	}
-	FIoStatus Status = IoStoreReader->Initialize(IoEnvironment, DecryptionKeys);
+	FIoStatus Status = IoStoreReader->Initialize(*FPaths::ChangeExtension(Path, TEXT("")), DecryptionKeys);
 	if (Status.IsOk())
 	{
 		return IoStoreReader;
@@ -3745,14 +3742,12 @@ int32 CreateTarget(const FIoStoreArguments& Arguments, const FIoStoreWriterSetti
 
 	TUniquePtr<FIoStoreWriterContext> IoStoreWriterContext(new FIoStoreWriterContext());
 	TArray<FIoStoreWriter*> IoStoreWriters;
-	FIoStoreEnvironment GlobalIoStoreEnv;
 	FIoStoreWriter* GlobalIoStoreWriter = nullptr;
 	{
 		IOSTORE_CPU_SCOPE(InitializeIoStoreWriters);
 		if (!Arguments.IsDLC())
 		{
-			GlobalIoStoreEnv.InitializeFileEnvironment(*Arguments.GlobalContainerPath);
-			GlobalIoStoreWriter = new FIoStoreWriter(GlobalIoStoreEnv);
+			GlobalIoStoreWriter = new FIoStoreWriter(*Arguments.GlobalContainerPath);
 			IoStoreWriters.Add(GlobalIoStoreWriter);
 		}
 		for (FContainerTargetSpec* ContainerTarget : ContainerTargets)
@@ -3760,9 +3755,7 @@ int32 CreateTarget(const FIoStoreArguments& Arguments, const FIoStoreWriterSetti
 			check(ContainerTarget->Header.ContainerId.IsValid());
 			if (!ContainerTarget->OutputPath.IsEmpty())
 			{
-				ContainerTarget->IoStoreEnv.Reset(new FIoStoreEnvironment());
-				ContainerTarget->IoStoreEnv->InitializeFileEnvironment(ContainerTarget->OutputPath);
-				ContainerTarget->IoStoreWriter = new FIoStoreWriter(*ContainerTarget->IoStoreEnv);
+				ContainerTarget->IoStoreWriter = new FIoStoreWriter(*ContainerTarget->OutputPath);
 				IoStoreWriters.Add(ContainerTarget->IoStoreWriter);
 			}
 		}
@@ -4101,9 +4094,7 @@ int32 CreateContentPatch(const FIoStoreArguments& Arguments, const FIoStoreWrite
 			return -1;
 		}
 
-		FIoStoreEnvironment IoStoreEnv;
-		FIoStoreWriter IoStoreWriter(IoStoreEnv);
-		IoStoreEnv.InitializeFileEnvironment(*Container.OutputPath);
+		FIoStoreWriter IoStoreWriter(*Container.OutputPath);
 		
 		EIoContainerFlags TargetContainerFlags = TargetReader->GetContainerFlags();
 
diff --git a/Engine/Source/Runtime/Core/Private/GenericPlatform/GenericPlatformIoDispatcher.cpp b/Engine/Source/Runtime/Core/Private/GenericPlatform/GenericPlatformIoDispatcher.cpp
--- a/Engine/Source/Runtime/Core/Private/GenericPlatform/GenericPlatformIoDispatcher.cpp
+++ b/Engine/Source/Runtime/Core/Private/GenericPlatform/GenericPlatformIoDispatcher.cpp
@@ -17,39 +17,27 @@ TRACE_DECLARE_INT_COUNTER(IoDispatcherForwardSeeks, TEXT("IoDispatcher/ForwardSe
 TRACE_DECLARE_INT_COUNTER(IoDispatcherBackwardSeeks, TEXT("IoDispatcher/BackwardSeeks"));
 TRACE_DECLARE_MEMORY_COUNTER(IoDispatcherTotalSeekDistance, TEXT("IoDispatcher/TotalSeekDistance"));
 
-FGenericIoDispatcherEventQueue::FGenericIoDispatcherEventQueue()
-	: DispatcherEvent(FPlatformProcess::GetSynchEventFromPool())
-	, ServiceEvent(FPlatformProcess::GetSynchEventFromPool())
+FGenericFileIoStoreEventQueue::FGenericFileIoStoreEventQueue()
+	: ServiceEvent(FPlatformProcess::GetSynchEventFromPool())
 {
 }
 
-FGenericIoDispatcherEventQueue::~FGenericIoDispatcherEventQueue()
+FGenericFileIoStoreEventQueue::~FGenericFileIoStoreEventQueue()
 {
 	FPlatformProcess::ReturnSynchEventToPool(ServiceEvent);
-	FPlatformProcess::ReturnSynchEventToPool(DispatcherEvent);
 }
 
-void FGenericIoDispatcherEventQueue::DispatcherNotify()
-{
-	DispatcherEvent->Trigger();
-}
-
-void FGenericIoDispatcherEventQueue::DispatcherWait()
-{
-	DispatcherEvent->Wait();
-}
-
-void FGenericIoDispatcherEventQueue::ServiceNotify()
+void FGenericFileIoStoreEventQueue::ServiceNotify()
 {
 	ServiceEvent->Trigger();
 }
 
-void FGenericIoDispatcherEventQueue::ServiceWait()
+void FGenericFileIoStoreEventQueue::ServiceWait()
 {
 	ServiceEvent->Wait();
 }
 
-FGenericFileIoStoreImpl::FGenericFileIoStoreImpl(FGenericIoDispatcherEventQueue& InEventQueue, FFileIoStoreBufferAllocator& InBufferAllocator, FFileIoStoreBlockCache& InBlockCache)
+FGenericFileIoStoreImpl::FGenericFileIoStoreImpl(FGenericFileIoStoreEventQueue& InEventQueue, FFileIoStoreBufferAllocator& InBufferAllocator, FFileIoStoreBlockCache& InBlockCache)
 	: EventQueue(InEventQueue)
 	, BufferAllocator(InBufferAllocator)
 	, BlockCache(InBlockCache)
@@ -92,7 +80,7 @@ bool FGenericFileIoStoreImpl::StartRequests(FFileIoStoreRequestQueue& RequestQue
 			FScopeLock _(&CompletedRequestsCritical);
 			CompletedRequests.Add(NextRequest);
 		}
-		EventQueue.DispatcherNotify();
+		WakeUpDispatcherThreadDelegate->Execute();
 		return true;
 	}
 
@@ -157,7 +145,7 @@ bool FGenericFileIoStoreImpl::StartRequests(FFileIoStoreRequestQueue& RequestQue
 		FScopeLock _(&CompletedRequestsCritical);
 		CompletedRequests.Add(NextRequest);
 	}
-	EventQueue.DispatcherNotify();
+	WakeUpDispatcherThreadDelegate->Execute();
 	return true;
 }
 
diff --git a/Engine/Source/Runtime/Core/Private/GenericPlatform/GenericPlatformIoDispatcher.h b/Engine/Source/Runtime/Core/Private/GenericPlatform/GenericPlatformIoDispatcher.h
--- a/Engine/Source/Runtime/Core/Private/GenericPlatform/GenericPlatformIoDispatcher.h
+++ b/Engine/Source/Runtime/Core/Private/GenericPlatform/GenericPlatformIoDispatcher.h
@@ -10,30 +10,27 @@
 #include "IO/IoDispatcher.h"
 #include "IO/IoDispatcherFileBackendTypes.h"
 
-class FGenericIoDispatcherEventQueue
+class FGenericFileIoStoreEventQueue
 {
 public:
-	FGenericIoDispatcherEventQueue();
-	~FGenericIoDispatcherEventQueue();
-	void DispatcherNotify();
-	void DispatcherWait();
-	void DispatcherWaitForIo()
-	{
-		DispatcherWait();
-	}
+	FGenericFileIoStoreEventQueue();
+	~FGenericFileIoStoreEventQueue();
 	void ServiceNotify();
 	void ServiceWait();
 
 private:
-	FEvent* DispatcherEvent = nullptr;
 	FEvent* ServiceEvent = nullptr;
 };
 
 class FGenericFileIoStoreImpl
 {
 public:
-	FGenericFileIoStoreImpl(FGenericIoDispatcherEventQueue& InEventQueue, FFileIoStoreBufferAllocator& InBufferAllocator, FFileIoStoreBlockCache& InBlockCache);
+	FGenericFileIoStoreImpl(FGenericFileIoStoreEventQueue& InEventQueue, FFileIoStoreBufferAllocator& InBufferAllocator, FFileIoStoreBlockCache& InBlockCache);
 	~FGenericFileIoStoreImpl();
+	void Initialize(const FWakeUpIoDispatcherThreadDelegate* InWakeUpDispatcherThreadDelegate)
+	{
+		WakeUpDispatcherThreadDelegate = InWakeUpDispatcherThreadDelegate;
+	}
 	bool OpenContainer(const TCHAR* ContainerFilePath, uint64& ContainerFileHandle, uint64& ContainerFileSize);
 	bool CreateCustomRequests(FFileIoStoreRequestAllocator& RequestAllocator, FFileIoStoreResolvedRequest& ResolvedRequest, FFileIoStoreReadRequestList& OutRequests)
 	{
@@ -43,7 +40,8 @@ public:
 	void GetCompletedRequests(FFileIoStoreReadRequestList& OutRequests);
 
 private:
-	FGenericIoDispatcherEventQueue& EventQueue;
+	const FWakeUpIoDispatcherThreadDelegate* WakeUpDispatcherThreadDelegate = nullptr;
+	FGenericFileIoStoreEventQueue& EventQueue;
 	FFileIoStoreBufferAllocator& BufferAllocator;
 	FFileIoStoreBlockCache& BlockCache;
 
diff --git a/Engine/Source/Runtime/Core/Private/IO/IoDispatcher.cpp b/Engine/Source/Runtime/Core/Private/IO/IoDispatcher.cpp
--- a/Engine/Source/Runtime/Core/Private/IO/IoDispatcher.cpp
+++ b/Engine/Source/Runtime/Core/Private/IO/IoDispatcher.cpp
@@ -3,7 +3,6 @@
 #include "IO/IoDispatcher.h"
 #include "IO/IoDispatcherPrivate.h"
 #include "IO/IoStore.h"
-#include "IO/IoDispatcherFileBackend.h"
 #include "Misc/ScopeRWLock.h"
 #include "Misc/CoreDelegates.h"
 #include "Math/RandomStream.h"
@@ -17,6 +16,7 @@
 #include "Async/MappedFileHandle.h"
 #include "ProfilingDebugging/CountersTrace.h"
 #include "ProfilingDebugging/CsvProfiler.h"
+#include "IO/IoDispatcherBackend.h"
 
 DEFINE_LOG_CATEGORY(LogIoDispatcher);
 
@@ -139,9 +139,11 @@ class FIoDispatcherImpl
 {
 public:
 	FIoDispatcherImpl(bool bInIsMultithreaded)
-		: bIsMultithreaded(bInIsMultithreaded)
-		, FileIoStore(EventQueue, SignatureErrorEvent, bIsMultithreaded)
+		: BackendContext(MakeShared<FIoDispatcherBackendContext>())
 	{
+		DispatcherEvent = FPlatformProcess::GetSynchEventFromPool(false);
+		BackendContext->WakeUpDispatcherThreadDelegate.BindRaw(this, &FIoDispatcherImpl::WakeUpDispatcherThread);
+		BackendContext->bIsMultiThreaded = bInIsMultithreaded;
 		FCoreDelegates::GetMemoryTrimDelegate().AddLambda([this]()
 		{
 			RequestAllocator.Trim();
@@ -152,6 +154,7 @@ public:
 	~FIoDispatcherImpl()
 	{
 		delete Thread;
+		FPlatformProcess::ReturnSynchEventToPool(DispatcherEvent);
 	}
 
 	FIoStatus Initialize()
@@ -161,7 +164,11 @@ public:
 
 	bool InitializePostSettings()
 	{
-		FileIoStore.Initialize();
+		check(!Thread);
+		for (const TSharedRef<IIoDispatcherBackend>& Backend : Backends)
+		{
+			Backend->Initialize(BackendContext);
+		}
 		Thread = FRunnableThread::Create(this, TEXT("IoDispatcher"), 0, TPri_AboveNormal, FPlatformAffinity::GetIoDispatcherThreadMask());
 		return true;
 	}
@@ -192,9 +199,9 @@ public:
 
 	void WakeUpDispatcherThread()
 	{
-		if (bIsMultithreaded)
+		if (BackendContext->bIsMultiThreaded)
 		{
-			EventQueue.DispatcherNotify();
+			DispatcherEvent->Trigger();
 		}
 		else
 		{
@@ -230,7 +237,15 @@ public:
 	{
 		if (ChunkId.IsValid())
 		{
-			return FileIoStore.OpenMapped(ChunkId, Options);
+			for (const TSharedRef<IIoDispatcherBackend>& Backend : Backends)
+			{
+				TIoStatusOr<FIoMappedRegion> Result = Backend->OpenMapped(ChunkId, Options);
+				if (Result.IsOk())
+				{
+					return Result;
+				}
+			}
+			return FIoStatus(EIoErrorCode::NotFound);
 		}
 		else
 		{
@@ -238,30 +253,25 @@ public:
 		}
 	}
 
-	FIoStatus Mount(const FIoStoreEnvironment& Environment, const FGuid& EncryptionKeyGuid, const FAES::FAESKey& EncryptionKey)
+	void Mount(TSharedRef<IIoDispatcherBackend> Backend)
 	{
-		TIoStatusOr<FIoContainerId> ContainerId = FileIoStore.Mount(Environment, EncryptionKeyGuid, EncryptionKey);
-
-		if (ContainerId.IsOk())
+		Backends.Add(Backend);
+		if (Thread)
 		{
-			FIoDispatcherMountedContainer MountedContainer;
-			MountedContainer.ContainerId = ContainerId.ValueOrDie();
-			MountedContainer.Environment = Environment;
-			if (ContainerMountedEvent.IsBound())
-			{
-				ContainerMountedEvent.Broadcast(MountedContainer);
-			}
-			FScopeLock Lock(&MountedContainersCritical);
-			MountedContainers.Add(MoveTemp(MountedContainer));
-			return FIoStatus::Ok;
+			Backend->Initialize(BackendContext);
 		}
-
-		return ContainerId.Status();
 	}
 
 	bool DoesChunkExist(const FIoChunkId& ChunkId) const
 	{
-		return FileIoStore.DoesChunkExist(ChunkId);
+		for (const TSharedRef<IIoDispatcherBackend>& Backend : Backends)
+		{
+			if (Backend->DoesChunkExist(ChunkId))
+			{
+				return true;
+			}
+		}
+		return false;
 	}
 
 	TIoStatusOr<uint64> GetSizeForChunk(const FIoChunkId& ChunkId) const
@@ -269,7 +279,15 @@ public:
 		// Only attempt to find the size if the FIoChunkId is valid
 		if (ChunkId.IsValid())
 		{
-			return FileIoStore.GetSizeForChunk(ChunkId);
+			for (const TSharedRef<IIoDispatcherBackend>& Backend : Backends)
+			{
+				TIoStatusOr<uint64> Result = Backend->GetSizeForChunk(ChunkId);
+				if (Result.IsOk())
+				{
+					return Result;
+				}
+			}
+			return FIoStatus(EIoErrorCode::NotFound);
 		}
 		else
 		{
@@ -277,20 +295,24 @@ public:
 		}	
 	}
 
-	TArray<FIoDispatcherMountedContainer> GetMountedContainers() const
+	TSet<FIoContainerId> GetMountedContainers() const
 	{
-		FScopeLock Lock(&MountedContainersCritical);
-		return MountedContainers;
+		TSet<FIoContainerId> Result;
+		for (const TSharedRef<IIoDispatcherBackend>& Backend : Backends)
+		{
+			Backend->AppendMountedContainers(Result);
+		}
+		return Result;
 	}
 
-	FIoDispatcher::FIoContainerMountedEvent& OnContainerMounted()
+	FIoContainerMountedDelegate& OnContainerMounted()
 	{
-		return ContainerMountedEvent;
+		return BackendContext->ContainerMountedDelegate;
 	}
 
-	FIoSignatureErrorEvent& GetSignatureErrorEvent()
+	FIoSignatureErrorDelegate& OnSignatureError()
 	{
-		return SignatureErrorEvent;
+		return BackendContext->SignatureErrorDelegate;
 	}
 
 	void IssueBatchInternal(FIoBatch& Batch, FIoBatchImpl* BatchImpl)
@@ -371,27 +393,30 @@ private:
 	{
 		//TRACE_CPUPROFILER_EVENT_SCOPE(ProcessCompletedRequests);
 
-		FIoRequestImpl* CompletedRequestsHead = FileIoStore.GetCompletedRequests();
-		while (CompletedRequestsHead)
+		for (const TSharedRef<IIoDispatcherBackend>& Backend : Backends)
 		{
-			FIoRequestImpl* NextRequest = CompletedRequestsHead->NextRequest;
-			if (CompletedRequestsHead->bCancelled)
-			{
-				CompleteRequest(CompletedRequestsHead, EIoErrorCode::Cancelled);
-			}
-			else if (CompletedRequestsHead->bFailed)
+			FIoRequestImpl* CompletedRequestsHead = Backend->GetCompletedRequests();
+			while (CompletedRequestsHead)
 			{
-				CompleteRequest(CompletedRequestsHead, EIoErrorCode::ReadError);
-			}
-			else
-			{
-				FPlatformAtomics::InterlockedAdd(&TotalLoaded, CompletedRequestsHead->IoBuffer.DataSize());
-				CompleteRequest(CompletedRequestsHead, EIoErrorCode::Ok);
+				FIoRequestImpl* NextRequest = CompletedRequestsHead->NextRequest;
+				if (CompletedRequestsHead->bCancelled)
+				{
+					CompleteRequest(CompletedRequestsHead, EIoErrorCode::Cancelled);
+				}
+				else if (CompletedRequestsHead->bFailed)
+				{
+					CompleteRequest(CompletedRequestsHead, EIoErrorCode::ReadError);
+				}
+				else
+				{
+					FPlatformAtomics::InterlockedAdd(&TotalLoaded, CompletedRequestsHead->IoBuffer.DataSize());
+					CompleteRequest(CompletedRequestsHead, EIoErrorCode::Ok);
+				}
+				CompletedRequestsHead->ReleaseRef();
+				CompletedRequestsHead = NextRequest;
+				--PendingIoRequestsCount;
+				TRACE_COUNTER_SET(PendingIoRequests, PendingIoRequestsCount);
 			}
-			CompletedRequestsHead->ReleaseRef();
-			CompletedRequestsHead = NextRequest;
-			--PendingIoRequestsCount;
-			TRACE_COUNTER_SET(PendingIoRequests, PendingIoRequestsCount);
 		}
 	}
 
@@ -484,18 +509,18 @@ private:
 				if (!RequestToCancel->bCancelled)
 				{
 					RequestToCancel->bCancelled = true;
-					if (RequestToCancel->bSubmitted)
+					if (RequestToCancel->Backend)
 					{
-						FileIoStore.CancelIoRequest(RequestToCancel);
+						RequestToCancel->Backend->CancelIoRequest(RequestToCancel);
 					}
 				}
 				RequestToCancel->ReleaseRef();
 			}
 			for (FIoRequestImpl* RequestToRePrioritize : LocalRequestsToReprioritize)
 			{
-				if (RequestToRePrioritize->bSubmitted)
+				if (RequestToRePrioritize->Backend)
 				{
-					FileIoStore.UpdatePriorityForIoRequest(RequestToRePrioritize);
+					RequestToRePrioritize->Backend->UpdatePriorityForIoRequest(RequestToRePrioritize);
 				}
 				RequestToRePrioritize->ReleaseRef();
 			}
@@ -507,6 +532,7 @@ private:
 			FIoRequestImpl* Request = RequestsToSubmitHead;
 			RequestsToSubmitHead = RequestsToSubmitHead->NextRequest;
 			Request->NextRequest = nullptr;
+
 			if (!RequestsToSubmitHead)
 			{
 				RequestsToSubmitTail = nullptr;
@@ -523,9 +549,17 @@ private:
 			if (Request->ChunkId.IsValid())
 			{
 				TRACE_CPUPROFILER_EVENT_SCOPE(ResolveRequest);
-				EIoStoreResolveResult Result = FileIoStore.Resolve(Request);
-				Request->bSubmitted = true;
-				if (Result != IoStoreResolveResult_OK)
+				bool bResolved = false;
+				for (const TSharedRef<IIoDispatcherBackend>& Backend : Backends)
+				{
+					if (Backend->Resolve(Request))
+					{
+						bResolved = true;
+						Request->Backend = &Backend.Get();
+						break;
+					}
+				}
+				if (!bResolved)
 				{
 					CompleteRequest(Request, EIoErrorCode::NotFound);
 					Request->ReleaseRef();
@@ -559,11 +593,11 @@ private:
 			if (PendingIoRequestsCount)
 			{
 				TRACE_CPUPROFILER_EVENT_SCOPE(IoDispatcherWaitForIo);
-				EventQueue.DispatcherWaitForIo();
+				DispatcherEvent->Wait();
 			}
 			else
 			{
-				EventQueue.DispatcherWait();
+				DispatcherEvent->Wait();
 			}
 			ProcessIncomingRequests();
 			ProcessCompletedRequests();
@@ -574,20 +608,18 @@ private:
 	virtual void Stop()
 	{
 		bStopRequested = true;
-		EventQueue.DispatcherNotify();
+		DispatcherEvent->Trigger();
 	}
 
 	using FRequestAllocator = TBlockAllocator<FIoRequestImpl, 4096>;
 	using FBatchAllocator = TBlockAllocator<FIoBatchImpl, 4096>;
 
-	bool bIsMultithreaded;
-	FIoDispatcherEventQueue EventQueue;
-
-	FIoSignatureErrorEvent SignatureErrorEvent;
-	FFileIoStore FileIoStore;
+	TSharedRef<FIoDispatcherBackendContext> BackendContext;
+	TArray<TSharedRef<IIoDispatcherBackend>> Backends;
 	FRequestAllocator RequestAllocator;
 	FBatchAllocator BatchAllocator;
 	FRunnableThread* Thread = nullptr;
+	FEvent* DispatcherEvent = nullptr;
 	FCriticalSection WaitingLock;
 	FIoRequestImpl* WaitingRequestsHead = nullptr;
 	FIoRequestImpl* WaitingRequestsTail = nullptr;
@@ -595,9 +627,6 @@ private:
 	TArray<FIoRequestImpl*> RequestsToCancel;
 	TArray<FIoRequestImpl*> RequestsToReprioritize;
 	TAtomic<bool> bStopRequested { false };
-	mutable FCriticalSection MountedContainersCritical;
-	TArray<FIoDispatcherMountedContainer> MountedContainers;
-	FIoDispatcher::FIoContainerMountedEvent ContainerMountedEvent;
 	uint64 PendingIoRequestsCount = 0;
 	int64 TotalLoaded = 0;
 };
@@ -612,10 +641,10 @@ FIoDispatcher::~FIoDispatcher()
 	delete Impl;
 }
 
-FIoStatus FIoDispatcher::Mount(const FIoStoreEnvironment& Environment, const FGuid& EncryptionKeyGuid, const FAES::FAESKey& EncryptionKey)
+void
+FIoDispatcher::Mount(TSharedRef<IIoDispatcherBackend> Backend)
 {
-	LLM_SCOPE(ELLMTag::FileSystem);
-	return Impl->Mount(Environment, EncryptionKeyGuid, EncryptionKey);
+	Impl->Mount(Backend);
 }
 
 FIoBatch
@@ -643,7 +672,7 @@ FIoDispatcher::GetSizeForChunk(const FIoChunkId& ChunkId) const
 	return Impl->GetSizeForChunk(ChunkId);
 }
 
-TArray<FIoDispatcherMountedContainer>
+TSet<FIoContainerId>
 FIoDispatcher::GetMountedContainers() const
 {
 	return Impl->GetMountedContainers();
@@ -655,16 +684,16 @@ FIoDispatcher::GetTotalLoaded() const
 	return Impl->GetTotalLoaded();
 }
 
-FIoDispatcher::FIoContainerMountedEvent&
+FIoContainerMountedDelegate&
 FIoDispatcher::OnContainerMounted()
 {
 	return Impl->OnContainerMounted();
 }
 
-FIoSignatureErrorEvent&
-FIoDispatcher::GetSignatureErrorEvent()
+FIoSignatureErrorDelegate&
+FIoDispatcher::OnSignatureError()
 {
-	return Impl->GetSignatureErrorEvent();
+	return Impl->OnSignatureError();
 }
 
 bool
@@ -673,12 +702,6 @@ FIoDispatcher::IsInitialized()
 	return GIoDispatcher.IsValid();
 }
 
-bool
-FIoDispatcher::IsValidEnvironment(const FIoStoreEnvironment& Environment)
-{
-	return FFileIoStore::IsValidEnvironment(Environment);
-}
-
 FIoStatus
 FIoDispatcher::Initialize()
 {
diff --git a/Engine/Source/Runtime/Core/Private/IO/IoDispatcherFileBackend.cpp b/Engine/Source/Runtime/Core/Private/IO/IoDispatcherFileBackend.cpp
--- a/Engine/Source/Runtime/Core/Private/IO/IoDispatcherFileBackend.cpp
+++ b/Engine/Source/Runtime/Core/Private/IO/IoDispatcherFileBackend.cpp
@@ -294,12 +294,12 @@ FFileIoStoreReader::FFileIoStoreReader(FFileIoStoreImpl& InPlatformImpl)
 {
 }
 
-FIoStatus FFileIoStoreReader::Initialize(const FIoStoreEnvironment& Environment)
+FIoStatus FFileIoStoreReader::Initialize(const TCHAR* InContainerPath, int32 InOrder)
 {
 	IPlatformFile& Ipf = FPlatformFileManager::Get().GetPlatformFile();
 
 	TStringBuilder<256> TocFilePath;
-	TocFilePath.Append(Environment.GetPath());
+	TocFilePath.Append(InContainerPath);
 	TocFilePath.Append(TEXT(".utoc"));
 	ContainerFile.FilePath = TocFilePath;
 
@@ -319,7 +319,7 @@ FIoStatus FFileIoStoreReader::Initialize(const FIoStoreEnvironment& Environment)
 	{
 		FFileIoStoreContainerFilePartition& Partition = ContainerFile.Partitions[PartitionIndex];
 		TStringBuilder<256> ContainerFilePath;
-		ContainerFilePath.Append(Environment.GetPath());
+		ContainerFilePath.Append(InContainerPath);
 		if (PartitionIndex > 0)
 		{
 			ContainerFilePath.Appendf(TEXT("_s%d"), PartitionIndex);
@@ -349,7 +349,7 @@ FIoStatus FFileIoStoreReader::Initialize(const FIoStoreEnvironment& Environment)
 	ContainerFile.BlockSignatureHashes	= MoveTemp(TocResource.ChunkBlockSignatures);
 
 	ContainerId = TocResource.Header.ContainerId;
-	Order = Environment.GetOrder();
+	Order = InOrder;
 	return FIoStatus::Ok;
 }
 
@@ -608,12 +608,9 @@ void FFileIoStoreRequestTracker::ReleaseIoRequestReferences(FFileIoStoreResolved
 	ResolvedRequest.ReadRequestsTail = nullptr;
 }
 
-FFileIoStore::FFileIoStore(FIoDispatcherEventQueue& InEventQueue, FIoSignatureErrorEvent& InSignatureErrorEvent, bool bInIsMultithreaded)
-	: EventQueue(InEventQueue)
-	, SignatureErrorEvent(InSignatureErrorEvent)
-	, RequestTracker(RequestAllocator, RequestQueue)
-	, PlatformImpl(InEventQueue, BufferAllocator, BlockCache)
-	, bIsMultithreaded(bInIsMultithreaded)
+FFileIoStore::FFileIoStore()
+	: RequestTracker(RequestAllocator, RequestQueue)
+	, PlatformImpl(EventQueue, BufferAllocator, BlockCache)
 {
 }
 
@@ -622,8 +619,12 @@ FFileIoStore::~FFileIoStore()
 	delete Thread;
 }
 
-void FFileIoStore::Initialize()
+void FFileIoStore::Initialize(TSharedRef<const FIoDispatcherBackendContext> InContext)
 {
+	check(!Thread);
+	BackendContext = InContext;
+	bIsMultithreaded = InContext->bIsMultiThreaded;
+
 	ReadBufferSize = (GIoDispatcherBufferSizeKB > 0 ? uint64(GIoDispatcherBufferSizeKB) << 10 : 256 << 10);
 
 	uint64 BufferMemorySize = uint64(GIoDispatcherBufferMemoryMB) << 20ull;
@@ -634,6 +635,8 @@ void FFileIoStore::Initialize()
 	uint64 CacheMemorySize = uint64(GIoDispatcherCacheSizeMB) << 20ull;
 	BlockCache.Initialize(CacheMemorySize, BufferSize);
 
+	PlatformImpl.Initialize(&BackendContext->WakeUpDispatcherThreadDelegate);
+
 	uint64 DecompressionContextCount = uint64(GIoDispatcherDecompressionWorkerCount > 0 ? GIoDispatcherDecompressionWorkerCount : 4);
 	for (uint64 ContextIndex = 0; ContextIndex < DecompressionContextCount; ++ContextIndex)
 	{
@@ -645,10 +648,10 @@ void FFileIoStore::Initialize()
 	Thread = FRunnableThread::Create(this, TEXT("IoService"), 0, TPri_AboveNormal);
 }
 
-TIoStatusOr<FIoContainerId> FFileIoStore::Mount(const FIoStoreEnvironment& Environment, const FGuid& EncryptionKeyGuid, const FAES::FAESKey& EncryptionKey)
+TIoStatusOr<FIoContainerId> FFileIoStore::Mount(const TCHAR* ContainerPath, int32 Order, const FGuid& EncryptionKeyGuid, const FAES::FAESKey& EncryptionKey)
 {
 	TUniquePtr<FFileIoStoreReader> Reader(new FFileIoStoreReader(PlatformImpl));
-	FIoStatus IoStatus = Reader->Initialize(Environment);
+	FIoStatus IoStatus = Reader->Initialize(ContainerPath, Order);
 	if (!IoStatus.IsOk())
 	{
 		return IoStatus;
@@ -663,7 +666,7 @@ TIoStatusOr<FIoContainerId> FFileIoStore::Mount(const FIoStoreEnvironment& Envir
 		else
 		{
 			return FIoStatus(EIoErrorCode::InvalidEncryptionKey, *FString::Printf(TEXT("Invalid encryption key '%s' (container '%s', encryption key '%s')"),
-				*EncryptionKeyGuid.ToString(), *FPaths::GetBaseFilename(Environment.GetPath()), *Reader->GetEncryptionKeyGuid().ToString()));
+				*EncryptionKeyGuid.ToString(), *FPaths::GetBaseFilename(ContainerPath), *Reader->GetEncryptionKeyGuid().ToString()));
 		}
 	}
 
@@ -683,12 +686,16 @@ TIoStatusOr<FIoContainerId> FFileIoStore::Mount(const FIoStoreEnvironment& Envir
 		FFileIoStoreReader* RawReader = Reader.Release();
 		UnorderedIoStoreReaders.Add(RawReader);		
 		OrderedIoStoreReaders.Insert(RawReader, InsertionIndex);
-		UE_LOG(LogIoDispatcher, Display, TEXT("Mounting container '%s' in location slot %d"), *FPaths::GetBaseFilename(Environment.GetPath()), InsertionIndex);
+		UE_LOG(LogIoDispatcher, Display, TEXT("Mounting container '%s' in location slot %d"), ContainerPath, InsertionIndex);
+	}
+	if (BackendContext && BackendContext->ContainerMountedDelegate.IsBound())
+	{
+		BackendContext->ContainerMountedDelegate.Broadcast(ContainerId);
 	}
 	return ContainerId;
 }
 
-EIoStoreResolveResult FFileIoStore::Resolve(FIoRequestImpl* Request)
+bool FFileIoStore::Resolve(FIoRequestImpl* Request)
 {
 	FReadScopeLock _(IoStoreReadersLock);
 	for (FFileIoStoreReader* Reader : OrderedIoStoreReaders)
@@ -742,11 +749,11 @@ EIoStoreResolveResult FFileIoStore::Resolve(FIoRequestImpl* Request)
 				CompleteDispatcherRequest(ResolvedRequest);
 			}
 
-			return IoStoreResolveResult_OK;
+			return true;
 		}
 	}
 
-	return IoStoreResolveResult_NotFound;
+	return false;
 }
 
 void FFileIoStore::CancelIoRequest(FIoRequestImpl* Request)
@@ -794,14 +801,6 @@ TIoStatusOr<uint64> FFileIoStore::GetSizeForChunk(const FIoChunkId& ChunkId) con
 	return FIoStatus(EIoErrorCode::NotFound);
 }
 
-bool FFileIoStore::IsValidEnvironment(const FIoStoreEnvironment& Environment)
-{
-	TStringBuilder<256> TocFilePath;
-	TocFilePath.Append(Environment.GetPath());
-	TocFilePath.Append(TEXT(".utoc"));
-	return FPlatformFileManager::Get().GetPlatformFile().FileExists(*TocFilePath);
-}
-
 FAutoConsoleTaskPriority CPrio_IoDispatcherTaskPriority(
 	TEXT("TaskGraph.TaskPriorities.IoDispatcherAsyncTasks"),
 	TEXT("Task and thread priority for IoDispatcher decompression."),
@@ -855,10 +854,10 @@ void FFileIoStore::ScatterBlock(FFileIoStoreCompressedBlock* CompressedBlock, bo
 
 			UE_LOG(LogIoDispatcher, Warning, TEXT("Signature error detected in container '%s' at block index '%d'"), *Error.ContainerName, Error.BlockIndex);
 
-			FScopeLock _(&SignatureErrorEvent.CriticalSection);
-			if (SignatureErrorEvent.SignatureErrorDelegate.IsBound())
+			check(BackendContext);
+			if (BackendContext->SignatureErrorDelegate.IsBound())
 			{
-				SignatureErrorEvent.SignatureErrorDelegate.Broadcast(Error);
+				BackendContext->SignatureErrorDelegate.Broadcast(Error);
 			}
 		}
 	}
@@ -903,7 +902,7 @@ void FFileIoStore::ScatterBlock(FFileIoStoreCompressedBlock* CompressedBlock, bo
 		CompressedBlock->Next = FirstDecompressedBlock;
 		FirstDecompressedBlock = CompressedBlock;
 
-		EventQueue.DispatcherNotify();
+		BackendContext->WakeUpDispatcherThreadDelegate.Execute();
 	}
 }
 
@@ -1135,6 +1134,15 @@ FIoRequestImpl* FFileIoStore::GetCompletedRequests()
 	return Result;
 }
 
+void FFileIoStore::AppendMountedContainers(TSet<FIoContainerId>& OutContainers)
+{
+	FReadScopeLock _(IoStoreReadersLock);
+	for (FFileIoStoreReader* Reader : UnorderedIoStoreReaders)
+	{
+		OutContainers.Add(Reader->GetContainerId());
+	}
+}
+
 TIoStatusOr<FIoMappedRegion> FFileIoStore::OpenMapped(const FIoChunkId& ChunkId, const FIoReadOptions& Options)
 {
 	if (!FPlatformProperties::SupportsMemoryMappedFiles())
@@ -1354,3 +1362,8 @@ uint32 FFileIoStore::Run()
 	}
 	return 0;
 }
+
+TSharedRef<IIoDispatcherFileBackend> CreateIoDispatcherFileBackend()
+{
+	return MakeShared<FFileIoStore>();
+}
\ No newline at end of file
diff --git a/Engine/Source/Runtime/Core/Private/IO/IoDispatcherFileBackend.h b/Engine/Source/Runtime/Core/Private/IO/IoDispatcherFileBackend.h
--- a/Engine/Source/Runtime/Core/Private/IO/IoDispatcherFileBackend.h
+++ b/Engine/Source/Runtime/Core/Private/IO/IoDispatcherFileBackend.h
@@ -14,6 +14,18 @@
 #include "Misc/AES.h"
 #include "GenericPlatform/GenericPlatformFile.h"
 
+#ifndef PLATFORM_IMPLEMENTS_IO
+#define PLATFORM_IMPLEMENTS_IO 0
+#endif
+
+#if PLATFORM_IMPLEMENTS_IO
+#include COMPILED_PLATFORM_HEADER(PlatformIoDispatcher.h)
+#else
+#include "GenericPlatform/GenericPlatformIoDispatcher.h"
+typedef FGenericFileIoStoreEventQueue FFileIoStoreEventQueue;
+typedef FGenericFileIoStoreImpl FFileIoStoreImpl;
+#endif
+
 class IMappedFileHandle;
 
 struct FFileIoStoreCompressionContext
@@ -27,7 +39,7 @@ class FFileIoStoreReader
 {
 public:
 	FFileIoStoreReader(FFileIoStoreImpl& InPlatformImpl);
-	FIoStatus Initialize(const FIoStoreEnvironment& Environment);
+	FIoStatus Initialize(const TCHAR* ContainerPath, int32 Order);
 	void SetIndex(uint32 InIndex)
 	{
 		Index = InIndex;
@@ -84,23 +96,23 @@ private:
 	TMap<FFileIoStoreBlockKey, FFileIoStoreReadRequest*> RawBlocksMap;
 };
 
-class FFileIoStore
+class FFileIoStore final
 	: public FRunnable
+	, public IIoDispatcherFileBackend
 {
 public:
-	FFileIoStore(FIoDispatcherEventQueue& InEventQueue, FIoSignatureErrorEvent& InSignatureErrorEvent, bool bInIsMultithreaded);
+	FFileIoStore();
 	~FFileIoStore();
-	void Initialize();
-	TIoStatusOr<FIoContainerId> Mount(const FIoStoreEnvironment& Environment, const FGuid& EncryptionKeyGuid, const FAES::FAESKey& EncryptionKey);
-	EIoStoreResolveResult Resolve(FIoRequestImpl* Request);
-	void CancelIoRequest(FIoRequestImpl* Request);
-	void UpdatePriorityForIoRequest(FIoRequestImpl* Request);
-	bool DoesChunkExist(const FIoChunkId& ChunkId) const;
+	void Initialize(TSharedRef<const FIoDispatcherBackendContext> Context) override;
+	TIoStatusOr<FIoContainerId> Mount(const TCHAR* ContainerPath, int32 Order, const FGuid& EncryptionKeyGuid, const FAES::FAESKey& EncryptionKey);
+	bool Resolve(FIoRequestImpl* Request) override;
+	void CancelIoRequest(FIoRequestImpl* Request) override;
+	void UpdatePriorityForIoRequest(FIoRequestImpl* Request) override;
+	bool DoesChunkExist(const FIoChunkId& ChunkId) const override;
 	TIoStatusOr<uint64> GetSizeForChunk(const FIoChunkId& ChunkId) const;
-	FIoRequestImpl* GetCompletedRequests();
-	TIoStatusOr<FIoMappedRegion> OpenMapped(const FIoChunkId& ChunkId, const FIoReadOptions& Options);
-	
-	static bool IsValidEnvironment(const FIoStoreEnvironment& Environment);
+	FIoRequestImpl* GetCompletedRequests() override;
+	TIoStatusOr<FIoMappedRegion> OpenMapped(const FIoChunkId& ChunkId, const FIoReadOptions& Options) override;
+	void AppendMountedContainers(TSet<FIoContainerId>& OutContainers) override;
 
 	virtual bool Init() override;
 	virtual uint32 Run() override;
@@ -150,8 +162,8 @@ private:
 	void UpdateAsyncIOMinimumPriority();
 
 	uint64 ReadBufferSize = 0;
-	FIoDispatcherEventQueue& EventQueue;
-	FIoSignatureErrorEvent& SignatureErrorEvent;
+	TSharedPtr<const FIoDispatcherBackendContext> BackendContext;
+	FFileIoStoreEventQueue EventQueue;
 	FFileIoStoreBlockCache BlockCache;
 	FFileIoStoreBufferAllocator BufferAllocator;
 	FFileIoStoreRequestAllocator RequestAllocator;
diff --git a/Engine/Source/Runtime/Core/Private/IO/IoDispatcherPrivate.h b/Engine/Source/Runtime/Core/Private/IO/IoDispatcherPrivate.h
--- a/Engine/Source/Runtime/Core/Private/IO/IoDispatcherPrivate.h
+++ b/Engine/Source/Runtime/Core/Private/IO/IoDispatcherPrivate.h
@@ -4,26 +4,8 @@
 
 #include "IO/IoDispatcher.h"
 
-#ifndef PLATFORM_IMPLEMENTS_IO
-#define PLATFORM_IMPLEMENTS_IO 0
-#endif
-
-#if PLATFORM_IMPLEMENTS_IO
-#include COMPILED_PLATFORM_HEADER(PlatformIoDispatcher.h)
-#else
-#include "GenericPlatform/GenericPlatformIoDispatcher.h"
-typedef FGenericIoDispatcherEventQueue FIoDispatcherEventQueue;
-typedef FGenericFileIoStoreImpl FFileIoStoreImpl;
-#endif
-
 class FIoRequestImpl;
 
-enum EIoStoreResolveResult
-{
-	IoStoreResolveResult_OK,
-	IoStoreResolveResult_NotFound,
-};
-
 class FIoBatchImpl
 {
 public:
@@ -32,3 +14,4 @@ public:
 	FGraphEventRef GraphEvent;
 	TAtomic<uint32> UnfinishedRequestsCount{ 0 };
 };
+
diff --git a/Engine/Source/Runtime/Core/Private/IO/IoStore.cpp b/Engine/Source/Runtime/Core/Private/IO/IoStore.cpp
--- a/Engine/Source/Runtime/Core/Private/IO/IoStore.cpp
+++ b/Engine/Source/Runtime/Core/Private/IO/IoStore.cpp
@@ -134,22 +134,6 @@ static FIoStatus ValidateContainerSignature(
 	return FIoStatus::Ok;
 }
 
-FIoStoreEnvironment::FIoStoreEnvironment()
-{
-}
-
-FIoStoreEnvironment::~FIoStoreEnvironment()
-{
-}
-
-void FIoStoreEnvironment::InitializeFileEnvironment(FStringView InPath, int32 InOrder)
-{
-	Path = InPath;
-	Order = InOrder;
-}
-
-//////////////////////////////////////////////////////////////////////////
-
 struct FChunkBlock
 {
 	uint64 Offset = 0;
@@ -550,8 +534,8 @@ private:
 class FIoStoreWriterImpl
 {
 public:
-	FIoStoreWriterImpl(FIoStoreEnvironment& InEnvironment)
-		: Environment(InEnvironment)
+	FIoStoreWriterImpl(const TCHAR* InContainerPath)
+		: ContainerPath(InContainerPath)
 	{
 	}
 
@@ -560,7 +544,7 @@ public:
 		WriterContext = &InContext;
 		ContainerSettings = InContainerSettings;
 
-		TocFilePath = Environment.GetPath() + TEXT(".utoc");
+		TocFilePath = ContainerPath + TEXT(".utoc");
 		
 		IPlatformFile& Ipf = IPlatformFile::GetPlatformPhysical();
 		Ipf.CreateDirectoryTree(*FPaths::GetPath(TocFilePath));
@@ -578,7 +562,7 @@ public:
 
 	FIoStatus EnableCsvOutput()
 	{
-		FString CsvFilePath = Environment.GetPath() + TEXT(".csv");
+		FString CsvFilePath = ContainerPath + TEXT(".csv");
 		CsvArchive.Reset(IFileManager::Get().CreateFileWriter(*CsvFilePath));
 		if (!CsvArchive)
 		{
@@ -1071,7 +1055,7 @@ private:
 	FIoStatus CreatePartitionContainerFile(FPartition& Partition)
 	{
 		check(!Partition.ContainerFileHandle);
-		FString ContainerFilePath = Environment.GetPath();
+		FString ContainerFilePath = ContainerPath;
 		if (Partition.Index > 0)
 		{
 			ContainerFilePath += FString::Printf(TEXT("_s%d"), Partition.Index);
@@ -1244,7 +1228,7 @@ private:
 		}
 	}
 
-	FIoStoreEnvironment&		Environment;
+	const FString				ContainerPath;
 	FIoStoreWriterContextImpl*	WriterContext = nullptr;
 	FIoContainerSettings		ContainerSettings;
 	FString						TocFilePath;
@@ -1263,8 +1247,8 @@ private:
 	bool						IsMetadataDirty = true;
 };
 
-FIoStoreWriter::FIoStoreWriter(FIoStoreEnvironment& InEnvironment)
-:	Impl(new FIoStoreWriterImpl(InEnvironment))
+FIoStoreWriter::FIoStoreWriter(const TCHAR* ContainerPath)
+:	Impl(new FIoStoreWriterImpl(ContainerPath))
 {
 }
 
@@ -1336,10 +1320,10 @@ public:
 
 	}
 
-	UE_NODISCARD FIoStatus Initialize(const FIoStoreEnvironment& InEnvironment, const TMap<FGuid, FAES::FAESKey>& InDecryptionKeys)
+	UE_NODISCARD FIoStatus Initialize(const TCHAR* InContainerPath, const TMap<FGuid, FAES::FAESKey>& InDecryptionKeys)
 	{
 		TStringBuilder<256> TocFilePath;
-		TocFilePath.Append(InEnvironment.GetPath());
+		TocFilePath.Append(InContainerPath);
 		TocFilePath.Append(TEXT(".utoc"));
 
 		FIoStoreTocResource& TocResource = Toc.GetTocResource();
@@ -1356,7 +1340,7 @@ public:
 		for (uint32 PartitionIndex = 0; PartitionIndex < TocResource.Header.PartitionCount; ++PartitionIndex)
 		{
 			TStringBuilder<256> ContainerFilePath;
-			ContainerFilePath.Append(InEnvironment.GetPath());
+			ContainerFilePath.Append(InContainerPath);
 			if (PartitionIndex > 0)
 			{
 				ContainerFilePath.Append(FString::Printf(TEXT("_s%d"), PartitionIndex));
@@ -1454,6 +1438,14 @@ public:
 			return FIoStatus(EIoErrorCode::NotFound, TEXT("Unknown chunk ID"));
 		}
 
+		uint64 RequestedOffset = Options.GetOffset();
+		uint64 ResolvedOffset = OffsetAndLength->GetOffset() + RequestedOffset;
+		uint64 ResolvedSize = 0;
+		if (RequestedOffset <= OffsetAndLength->GetLength())
+		{
+			ResolvedSize = FMath::Min(Options.GetSize(), OffsetAndLength->GetLength() - RequestedOffset);
+		}
+
 		if (!ThreadBuffers)
 		{
 			ThreadBuffers = new FThreadBuffers();
@@ -1464,13 +1456,13 @@ public:
 
 		const FIoStoreTocResource& TocResource = Toc.GetTocResource();
 		const uint64 CompressionBlockSize = TocResource.Header.CompressionBlockSize;
-		FIoBuffer IoBuffer = FIoBuffer(OffsetAndLength->GetLength());
-		int32 FirstBlockIndex = int32(OffsetAndLength->GetOffset() / CompressionBlockSize);
-		int32 LastBlockIndex = int32((Align(OffsetAndLength->GetOffset() + OffsetAndLength->GetLength(), CompressionBlockSize) - 1) / CompressionBlockSize);
-		uint64 OffsetInBlock = OffsetAndLength->GetOffset() % CompressionBlockSize;
+		FIoBuffer IoBuffer = FIoBuffer(ResolvedSize);
+		int32 FirstBlockIndex = int32(ResolvedOffset / CompressionBlockSize);
+		int32 LastBlockIndex = int32((Align(ResolvedOffset + ResolvedSize, CompressionBlockSize) - 1) / CompressionBlockSize);
+		uint64 OffsetInBlock = ResolvedOffset % CompressionBlockSize;
 		uint8* Dst = IoBuffer.Data();
 		uint8* Src = nullptr;
-		uint64 RemainingSize = OffsetAndLength->GetLength();
+		uint64 RemainingSize = ResolvedSize;
 		for (int32 BlockIndex = FirstBlockIndex; BlockIndex <= LastBlockIndex; ++BlockIndex)
 		{
 			const FIoStoreTocCompressedBlockEntry& CompressionBlock = TocResource.CompressionBlocks[BlockIndex];
@@ -1601,9 +1593,9 @@ FIoStoreReader::~FIoStoreReader()
 	delete Impl;
 }
 
-FIoStatus FIoStoreReader::Initialize(const FIoStoreEnvironment& InEnvironment, const TMap<FGuid, FAES::FAESKey>& InDecryptionKeys)
+FIoStatus FIoStoreReader::Initialize(const TCHAR* InContainerPath, const TMap<FGuid, FAES::FAESKey>& InDecryptionKeys)
 {
-	return Impl->Initialize(InEnvironment, InDecryptionKeys);
+	return Impl->Initialize(InContainerPath, InDecryptionKeys);
 }
 
 FIoContainerId FIoStoreReader::GetContainerId() const
diff --git a/Engine/Source/Runtime/Core/Public/IO/IoDispatcher.h b/Engine/Source/Runtime/Core/Public/IO/IoDispatcher.h
--- a/Engine/Source/Runtime/Core/Public/IO/IoDispatcher.h
+++ b/Engine/Source/Runtime/Core/Public/IO/IoDispatcher.h
@@ -32,6 +32,7 @@ class FIoStoreReaderImpl;
 class IMappedFileHandle;
 class IMappedFileRegion;
 class FIoDirectoryIndexReaderImpl;
+struct IIoDispatcherBackend;
 
 CORE_API DECLARE_LOG_CATEGORY_EXTERN(LogIoDispatcher, Log, All);
 
@@ -407,26 +408,6 @@ TIoStatusOr<T>& TIoStatusOr<T>::operator=(const TIoStatusOr<U>& Other)
 	return *this;
 }
 
-//////////////////////////////////////////////////////////////////////////
-
-/** Helper used to manage creation of I/O store file handles etc
-  */
-class FIoStoreEnvironment
-{
-public:
-	CORE_API FIoStoreEnvironment();
-	CORE_API ~FIoStoreEnvironment();
-
-	CORE_API void InitializeFileEnvironment(FStringView InPath, int32 InOrder = 0);
-
-	CORE_API const FString& GetPath() const { return Path; }
-	CORE_API int32 GetOrder() const { return Order; }
-
-private:
-	FString			Path;
-	int32			Order = 0;
-};
-
 /** Reference to buffer data used by I/O dispatcher APIs
   */
 class FIoBuffer
@@ -850,12 +831,6 @@ struct FIoMappedRegion
 	IMappedFileRegion* MappedFileRegion = nullptr;
 };
 
-struct FIoDispatcherMountedContainer
-{
-	FIoStoreEnvironment Environment;
-	FIoContainerId ContainerId;
-};
-
 struct FIoSignatureError
 {
 	FString ContainerName;
@@ -866,48 +841,36 @@ struct FIoSignatureError
 
 DECLARE_MULTICAST_DELEGATE_OneParam(FIoSignatureErrorDelegate, const FIoSignatureError&);
 
-struct FIoSignatureErrorEvent
-{
-	FCriticalSection CriticalSection;
-	FIoSignatureErrorDelegate SignatureErrorDelegate;
-};
+DECLARE_MULTICAST_DELEGATE_OneParam(FIoContainerMountedDelegate, const FIoContainerId&);
 
 /** I/O dispatcher
   */
-class FIoDispatcher
+class FIoDispatcher final
 {
 public:
-	DECLARE_EVENT_OneParam(FIoDispatcher, FIoContainerMountedEvent, const FIoDispatcherMountedContainer&);
-
 	CORE_API						FIoDispatcher();
-	CORE_API virtual				~FIoDispatcher();
+	CORE_API						~FIoDispatcher();
 
-	CORE_API FIoStatus				Mount(const FIoStoreEnvironment& Environment, const FGuid& EncryptionKeyGuid, const FAES::FAESKey& EncryptionKey);
+	CORE_API void					Mount(TSharedRef<IIoDispatcherBackend> Backend);
 
 	CORE_API FIoBatch				NewBatch();
 
-	UE_DEPRECATED(4.26, "Remove this call")
-	CORE_API void					FreeBatch(FIoBatch& Batch)
-	{
-	}
-
 	CORE_API TIoStatusOr<FIoMappedRegion> OpenMapped(const FIoChunkId& ChunkId, const FIoReadOptions& Options);
 
 	// Polling methods
 	CORE_API bool					DoesChunkExist(const FIoChunkId& ChunkId) const;
 	CORE_API TIoStatusOr<uint64>	GetSizeForChunk(const FIoChunkId& ChunkId) const;
-	CORE_API TArray<FIoDispatcherMountedContainer> GetMountedContainers() const;
+	CORE_API TSet<FIoContainerId>	GetMountedContainers() const;
 	CORE_API int64					GetTotalLoaded() const;
 
 
 	// Events
-	CORE_API FIoContainerMountedEvent& OnContainerMounted();
-	CORE_API FIoSignatureErrorEvent& GetSignatureErrorEvent();
+	CORE_API FIoContainerMountedDelegate& OnContainerMounted();
+	CORE_API FIoSignatureErrorDelegate& OnSignatureError();
 
 	FIoDispatcher(const FIoDispatcher&) = default;
 	FIoDispatcher& operator=(const FIoDispatcher&) = delete;
 
-	static CORE_API bool IsValidEnvironment(const FIoStoreEnvironment& Environment);
 	static CORE_API bool IsInitialized();
 	static CORE_API FIoStatus Initialize();
 	static CORE_API void InitializePostSettings();
@@ -1120,7 +1083,7 @@ public:
 class FIoStoreWriter
 {
 public:
-	CORE_API 			FIoStoreWriter(FIoStoreEnvironment& InEnvironment);
+	CORE_API 			FIoStoreWriter(const TCHAR* ContainerPath);
 	CORE_API virtual	~FIoStoreWriter();
 
 	FIoStoreWriter(const FIoStoreWriter&) = delete;
@@ -1154,7 +1117,7 @@ public:
 	CORE_API FIoStoreReader();
 	CORE_API ~FIoStoreReader();
 
-	UE_NODISCARD CORE_API FIoStatus Initialize(const FIoStoreEnvironment& InEnvironment, const TMap<FGuid, FAES::FAESKey>& InDecryptionKeys);
+	UE_NODISCARD CORE_API FIoStatus Initialize(const TCHAR* ContainerPath, const TMap<FGuid, FAES::FAESKey>& InDecryptionKeys);
 	CORE_API FIoContainerId GetContainerId() const;
 	CORE_API EIoContainerFlags GetContainerFlags() const;
 	CORE_API FGuid GetEncryptionKeyGuid() const;
diff --git a/Engine/Source/Runtime/Core/Public/IO/IoDispatcherBackend.h b/Engine/Source/Runtime/Core/Public/IO/IoDispatcherBackend.h
--- a/Engine/Source/Runtime/Core/Public/IO/IoDispatcherBackend.h
+++ b/Engine/Source/Runtime/Core/Public/IO/IoDispatcherBackend.h
@@ -2,6 +2,8 @@
 
 #pragma once
 
+#include "IoDispatcher.h"
+
 class FIoRequestImpl
 {
 public:
@@ -49,12 +51,43 @@ private:
 	void FreeRequest();
 
 	FIoDispatcherImpl& Dispatcher;
+	struct IIoDispatcherBackend* Backend = nullptr;
 	FIoBatchImpl* Batch = nullptr;
 	FIoReadCallback Callback;
 	TAtomic<uint32> RefCount{ 0 };
 	TAtomic<EIoErrorCode> ErrorCode{ EIoErrorCode::Unknown };
-	bool bSubmitted = false;
 	bool bCancelled = false;
 	bool bFailed = false;
 };
 
+DECLARE_DELEGATE(FWakeUpIoDispatcherThreadDelegate);
+
+struct FIoDispatcherBackendContext
+{
+	FWakeUpIoDispatcherThreadDelegate WakeUpDispatcherThreadDelegate;
+	FIoContainerMountedDelegate ContainerMountedDelegate;
+	FIoSignatureErrorDelegate SignatureErrorDelegate;
+	bool bIsMultiThreaded;
+};
+
+struct IIoDispatcherBackend
+{
+	virtual void Initialize(TSharedRef<const FIoDispatcherBackendContext> Context) = 0;
+	virtual bool Resolve(FIoRequestImpl* Request) = 0;
+	virtual void CancelIoRequest(FIoRequestImpl* Request) = 0;
+	virtual void UpdatePriorityForIoRequest(FIoRequestImpl* Request) = 0;
+	virtual bool DoesChunkExist(const FIoChunkId& ChunkId) const = 0;
+	virtual TIoStatusOr<uint64> GetSizeForChunk(const FIoChunkId& ChunkId) const = 0;
+	virtual FIoRequestImpl* GetCompletedRequests() = 0;
+	virtual TIoStatusOr<FIoMappedRegion> OpenMapped(const FIoChunkId& ChunkId, const FIoReadOptions& Options) = 0;
+	virtual void AppendMountedContainers(TSet<FIoContainerId>& OutContainers) = 0;
+};
+
+struct IIoDispatcherFileBackend
+	: public IIoDispatcherBackend
+{
+	virtual TIoStatusOr<FIoContainerId> Mount(const TCHAR* ContainerPath, int32 Order, const FGuid& EncryptionKeyGuid, const FAES::FAESKey& EncryptionKey) = 0;
+};
+
+CORE_API TSharedRef<IIoDispatcherFileBackend> CreateIoDispatcherFileBackend();
+
diff --git a/Engine/Source/Runtime/CoreUObject/Private/Serialization/AsyncLoading2.cpp b/Engine/Source/Runtime/CoreUObject/Private/Serialization/AsyncLoading2.cpp
--- a/Engine/Source/Runtime/CoreUObject/Private/Serialization/AsyncLoading2.cpp
+++ b/Engine/Source/Runtime/CoreUObject/Private/Serialization/AsyncLoading2.cpp
@@ -844,7 +844,6 @@ public:
 		TUniquePtr<FNameMap> ContainerNameMap;
 		TArray<uint8> StoreEntries; //FPackageStoreEntry[PackageCount];
 		uint32 PackageCount;
-		int32 Order = 0;
 		bool bValid = false;
 	};
 
@@ -916,15 +915,15 @@ public:
 		}
 	}
 
-	void LoadContainers(TArrayView<const FIoDispatcherMountedContainer> Containers)
+	void LoadContainers(TArrayView<const FIoContainerId> Containers)
 	{
 		TRACE_CPUPROFILER_EVENT_SCOPE(LoadContainers);
 
 		int32 ContainersToLoad = 0;
 
-		for (const FIoDispatcherMountedContainer& Container : Containers)
+		for (const FIoContainerId& ContainerId : Containers)
 		{
-			if (Container.ContainerId.IsValid())
+			if (ContainerId.IsValid())
 			{
 				++ContainersToLoad;
 			}
@@ -940,9 +939,8 @@ public:
 		FEvent* Event = FPlatformProcess::GetSynchEventFromPool();
 		FIoBatch IoBatch = IoDispatcher.NewBatch();
 
-		for (const FIoDispatcherMountedContainer& Container : Containers)
+		for (const FIoContainerId& ContainerId : Containers)
 		{
-			const FIoContainerId& ContainerId = Container.ContainerId;
 			if (!ContainerId.IsValid())
 			{
 				continue;
@@ -954,19 +952,9 @@ public:
 				LoadedContainerPtr.Reset(new FLoadedContainer);
 			}
 			FLoadedContainer& LoadedContainer = *LoadedContainerPtr;
-			if (LoadedContainer.bValid && LoadedContainer.Order >= Container.Environment.GetOrder())
-			{
-				UE_LOG(LogStreaming, Log, TEXT("Skipping loading mounted container ID '0x%llX', already loaded with higher order"), ContainerId.Value());
-				if (--Remaining == 0)
-				{
-					Event->Trigger();
-				}
-				continue;
-			}
 
 			UE_LOG(LogStreaming, Log, TEXT("Loading mounted container ID '0x%llX'"), ContainerId.Value());
 			LoadedContainer.bValid = true;
-			LoadedContainer.Order = Container.Environment.GetOrder();
 
 			FIoChunkId HeaderChunkId = CreateIoChunkId(ContainerId.Value(), 0, EIoChunkType::ContainerHeader);
 			IoBatch.ReadWithCallback(HeaderChunkId, FIoReadOptions(), IoDispatcherPriority_High, [this, &Remaining, Event, &LoadedContainer, ContainerId](TIoStatusOr<FIoBuffer> Result)
@@ -1082,10 +1070,10 @@ public:
 		ApplyRedirects(RedirectsPackageMap);
 	}
 
-	void OnContainerMounted(const FIoDispatcherMountedContainer& Container)
+	void OnContainerMounted(const FIoContainerId& ContainerId)
 	{
 		LLM_SCOPE(ELLMTag::AsyncLoading);
-		LoadContainers(MakeArrayView(&Container, 1));
+		LoadContainers(MakeArrayView(&ContainerId, 1));
 	}
 
 	void ApplyRedirects(const TMap<FPackageId, FPackageId>& Redirects)
@@ -4820,7 +4808,7 @@ void FAsyncLoadingThread2::LazyInitializeFromLoadPackage()
 		GlobalPackageStore.SetupInitialLoadData();
 	}
 	GlobalPackageStore.SetupCulture();
-	GlobalPackageStore.LoadContainers(IoDispatcher.GetMountedContainers());
+	GlobalPackageStore.LoadContainers(IoDispatcher.GetMountedContainers().Array());
 	IoDispatcher.OnContainerMounted().AddRaw(&GlobalPackageStore, &FPackageStore::OnContainerMounted);
 }
 
diff --git a/Engine/Source/Runtime/PakFile/Private/IPlatformFilePak.cpp b/Engine/Source/Runtime/PakFile/Private/IPlatformFilePak.cpp
--- a/Engine/Source/Runtime/PakFile/Private/IPlatformFilePak.cpp
+++ b/Engine/Source/Runtime/PakFile/Private/IPlatformFilePak.cpp
@@ -34,7 +34,7 @@
 #include "Misc/Fnv.h"
 #include "ProfilingDebugging/ScopedTimers.h"
 #include "Async/MappedFileHandle.h"
-#include "IO/IoDispatcher.h"
+#include "IO/IoDispatcherBackend.h"
 
 #include "ProfilingDebugging/LoadTimeTracker.h"
 
@@ -268,17 +268,18 @@ void FPakPlatformFile::GetFilenamesFromIostoreByBlockIndex(const FString& InCont
 {
 	if (FIoDispatcher::IsInitialized())
 	{
-		FIoDispatcher& IoDispatcher = FIoDispatcher::Get();
-
-		for (const FIoDispatcherMountedContainer& Container : IoDispatcher.GetMountedContainers())
+		FPakPlatformFile* PakPlatformFile = static_cast<FPakPlatformFile*>(FPlatformFileManager::Get().FindPlatformFile(FPakPlatformFile::GetTypeName()));
+		if (!PakPlatformFile)
+		{
+			return;
+		}
+		FScopeLock ScopedLock(&PakPlatformFile->PakListCritical);
+		for (const FPakListEntry& PakListEntry : PakPlatformFile->PakFiles)
 		{
-			if (FPaths::GetBaseFilename(Container.Environment.GetPath()) == InContainerName)
+			if (FPaths::GetBaseFilename(PakListEntry.PakFile->PakFilename) == InContainerName)
 			{
-				FIoStoreEnvironment IoEnvironment;
-				IoEnvironment.InitializeFileEnvironment(FPaths::ChangeExtension(Container.Environment.GetPath(), TEXT("")));
 				TUniquePtr<FIoStoreReader> IoStoreReader(new FIoStoreReader());
-
-				FIoStatus Status = IoStoreReader->Initialize(IoEnvironment, GetRegisteredEncryptionKeys().GetKeys());
+				FIoStatus Status = IoStoreReader->Initialize(*FPaths::ChangeExtension(PakListEntry.PakFile->PakFilename, TEXT("")), GetRegisteredEncryptionKeys().GetKeys());
 				if (Status.IsOk())
 				{
 					IoStoreReader->GetFilenamesByBlockIndex(InBlockIndex, OutFileList);
@@ -309,17 +310,18 @@ void FPakPlatformFile::GetFilenamesFromIostoreContainer(const FString& InContain
 {
 	if (FIoDispatcher::IsInitialized())
 	{
-		FIoDispatcher& IoDispatcher = FIoDispatcher::Get();
-
-		for (const FIoDispatcherMountedContainer& Container : IoDispatcher.GetMountedContainers())
+		FPakPlatformFile* PakPlatformFile = static_cast<FPakPlatformFile*>(FPlatformFileManager::Get().FindPlatformFile(FPakPlatformFile::GetTypeName()));
+		if (!PakPlatformFile)
 		{
-			if (FPaths::GetBaseFilename(Container.Environment.GetPath()) == InContainerName)
+			return;
+		}
+		FScopeLock ScopedLock(&PakPlatformFile->PakListCritical);
+		for (const FPakListEntry& PakListEntry : PakPlatformFile->PakFiles)
+		{
+			if (FPaths::GetBaseFilename(PakListEntry.PakFile->PakFilename) == InContainerName)
 			{
-				FIoStoreEnvironment IoEnvironment;
-				IoEnvironment.InitializeFileEnvironment(FPaths::ChangeExtension(Container.Environment.GetPath(), TEXT("")));
 				TUniquePtr<FIoStoreReader> IoStoreReader(new FIoStoreReader());
-
-				FIoStatus Status = IoStoreReader->Initialize(IoEnvironment, GetRegisteredEncryptionKeys().GetKeys());
+				FIoStatus Status = IoStoreReader->Initialize(*FPaths::ChangeExtension(PakListEntry.PakFile->PakFilename, TEXT("")), GetRegisteredEncryptionKeys().GetKeys());
 				if (Status.IsOk())
 				{
 					IoStoreReader->GetFilenames(OutFileList);
@@ -6957,22 +6959,20 @@ bool FPakPlatformFile::Initialize(IPlatformFile* Inner, const TCHAR* CmdLine)
 	FParse::Value(FCommandLine::Get(), TEXT("StartupPaksWildcard="), StartupPaksWildcard);
 #endif
 
-	FIoStoreEnvironment IoStoreGlobalEnvironment;
-	IoStoreGlobalEnvironment.InitializeFileEnvironment(FString::Printf(TEXT("%sPaks/global"), *FPaths::ProjectContentDir()));
-	if (FIoDispatcher::IsValidEnvironment(IoStoreGlobalEnvironment))
+	FString GlobalUTocPath = FString::Printf(TEXT("%sPaks/global.utoc"), *FPaths::ProjectContentDir());
+	if (FPlatformFileManager::Get().GetPlatformFile().FileExists(*GlobalUTocPath))
 	{
 		FIoStatus IoDispatcherInitStatus = FIoDispatcher::Initialize();
 		if (IoDispatcherInitStatus.IsOk())
 		{
 			FIoDispatcher& IoDispatcher = FIoDispatcher::Get();
-			FIoStatus IoDispatcherMountStatus = IoDispatcher.Mount(IoStoreGlobalEnvironment, FGuid(), FAES::FAESKey());
+			IoDispatcherFileBackend = CreateIoDispatcherFileBackend();
+			IoDispatcher.Mount(IoDispatcherFileBackend.ToSharedRef());
+			TIoStatusOr<FIoContainerId> IoDispatcherMountStatus = IoDispatcherFileBackend->Mount(*FPaths::ChangeExtension(GlobalUTocPath, TEXT("")), 0, FGuid(), FAES::FAESKey());
 			if (IoDispatcherMountStatus.IsOk())
 			{
 				UE_LOG(LogPakFile, Display, TEXT("Initialized I/O dispatcher"));
-
-				FIoSignatureErrorEvent& SignatureErrorEvent = IoDispatcher.GetSignatureErrorEvent();
-				FScopeLock _(&SignatureErrorEvent.CriticalSection);
-				SignatureErrorEvent.SignatureErrorDelegate.AddLambda([](const FIoSignatureError& Error)
+				IoDispatcher.OnSignatureError().AddLambda([](const FIoSignatureError& Error)
 				{
 					FPakChunkSignatureCheckFailedData FailedData(Error.ContainerName, TPakChunkHash(), TPakChunkHash(), Error.BlockIndex);
 #if !PAKHASH_USE_CRC
@@ -6984,7 +6984,7 @@ bool FPakPlatformFile::Initialize(IPlatformFile* Inner, const TCHAR* CmdLine)
 			}
 			else
 			{
-				UE_LOG(LogPakFile, Error, TEXT("Failed to mount I/O dispatcher global environment: '%s'"), *IoDispatcherMountStatus.ToString());
+				UE_LOG(LogPakFile, Error, TEXT("Failed to mount I/O dispatcher container: '%s'"), *IoDispatcherMountStatus.Status().ToString());
 			}
 		}
 		else
@@ -7215,11 +7215,8 @@ bool FPakPlatformFile::Mount(const TCHAR* InPakFilename, uint32 PakOrder, const
 			UE_LOG(LogPakFile, Warning, TEXT("Failed to mount pak \"%s\", pak is invalid."), InPakFilename);
 		}
 
-		if (FIoDispatcher::IsInitialized())
+		if (IoDispatcherFileBackend.IsValid())
 		{
-			FIoStoreEnvironment IoStoreEnvironment;
-			IoStoreEnvironment.InitializeFileEnvironment(FPaths::ChangeExtension(InPakFilename, FString()), PakOrder);
-
 			FGuid EncryptionKeyGuid = Pak->GetInfo().EncryptionKeyGuid;
 			FAES::FAESKey EncryptionKey;
 
@@ -7231,15 +7228,16 @@ bool FPakPlatformFile::Mount(const TCHAR* InPakFilename, uint32 PakOrder, const
 				}
 			}
 
-			FIoStatus IoStatus = FIoDispatcher::Get().Mount(IoStoreEnvironment, EncryptionKeyGuid, EncryptionKey);
+			FString ContainerPath = FPaths::ChangeExtension(InPakFilename, FString());
+			TIoStatusOr<FIoContainerId> IoStatus = IoDispatcherFileBackend->Mount(*ContainerPath, PakOrder, EncryptionKeyGuid, EncryptionKey);
 			if (IoStatus.IsOk())
 			{
-				UE_LOG(LogPakFile, Display, TEXT("Mounted IoStore environment \"%s\""), *IoStoreEnvironment.GetPath());
+				UE_LOG(LogPakFile, Display, TEXT("Mounted IoStore container \"%s\""), *ContainerPath);
 			}
 			else
 			{
 				bIoStoreSuccess = false;
-				UE_LOG(LogPakFile, Warning, TEXT("Failed to mount IoStore environment \"%s\" [%s]"), *IoStoreEnvironment.GetPath(), *IoStatus.ToString());
+				UE_LOG(LogPakFile, Warning, TEXT("Failed to mount IoStore container \"%s\" [%s]"), *ContainerPath, *IoStatus.Status().ToString());
 			}
 		}
 
diff --git a/Engine/Source/Runtime/PakFile/Public/IPlatformFilePak.h b/Engine/Source/Runtime/PakFile/Public/IPlatformFilePak.h
--- a/Engine/Source/Runtime/PakFile/Public/IPlatformFilePak.h
+++ b/Engine/Source/Runtime/PakFile/Public/IPlatformFilePak.h
@@ -19,6 +19,7 @@
 
 class FChunkCacheWorker;
 class IAsyncReadFileHandle;
+struct IIoDispatcherFileBackend;
 
 PAKFILE_API DECLARE_LOG_CATEGORY_EXTERN(LogPakFile, Log, All);
 DECLARE_FLOAT_ACCUMULATOR_STAT_EXTERN(TEXT("Total pak file read time"), STAT_PakFile_Read, STATGROUP_PakFile, PAKFILE_API);
@@ -1985,6 +1986,7 @@ class PAKFILE_API FPakPlatformFile : public IPlatformFile
 	FString IniFileExtension;
 	/** The filename for the gameusersettings ini file, used for excluding ini files, but not gameusersettings */
 	FString GameUserSettingsIniFilename;
+	TSharedPtr<IIoDispatcherFileBackend> IoDispatcherFileBackend;
 
 	/**
 	 * Gets mounted pak files
