- Undoing a recent pull request that is crashing the editor, due to empty windows machine name

diff --git a/Engine/Source/Runtime/Core/Private/Windows/WindowsPlatformProcess.cpp b/Engine/Source/Runtime/Core/Private/Windows/WindowsPlatformProcess.cpp
--- a/Engine/Source/Runtime/Core/Private/Windows/WindowsPlatformProcess.cpp
+++ b/Engine/Source/Runtime/Core/Private/Windows/WindowsPlatformProcess.cpp
@@ -1046,11 +1046,11 @@ const TCHAR* FWindowsPlatformProcess::ApplicationSettingsDir()
 
 const TCHAR* FWindowsPlatformProcess::ComputerName()
 {
-	static TCHAR Result[MAX_COMPUTERNAME_LENGTH]=TEXT("");
+	static TCHAR Result[256]=TEXT("");
 	if( !Result[0] )
 	{
 		uint32 Size=UE_ARRAY_COUNT(Result);
-		GetComputerNameEx( ComputerNameDnsHostname, Result, (::DWORD*)&Size );
+		GetComputerName( Result, (::DWORD*)&Size );
 	}
 	return Result;
 }
