U5 - fix static analyisis warning

diff --git a/Engine/Source/Runtime/Core/Private/Async/TaskGraph.cpp b/Engine/Source/Runtime/Core/Private/Async/TaskGraph.cpp
--- a/Engine/Source/Runtime/Core/Private/Async/TaskGraph.cpp
+++ b/Engine/Source/Runtime/Core/Private/Async/TaskGraph.cpp
@@ -1857,7 +1857,7 @@ private:
 		if (ENamedThreads::GetThreadIndex(InThreadToExecuteOn) == ENamedThreads::AnyThread)
 		{
 			uint32 ThreadPriority = GetThreadPriorityIndex(InThreadToExecuteOn);
-			check(ThreadPriority < uint32(LowLevelTasks::ETaskPriority::Count));
+			check(ThreadPriority < uint32(ENamedThreads::NumThreadPriorities));
 			LowLevelTasks::ETaskPriority Conversion[int(ENamedThreads::NumThreadPriorities)] = { LowLevelTasks::ETaskPriority::Normal, LowLevelTasks::ETaskPriority::High, LowLevelTasks::ETaskPriority::BackgroundLow };
 			LowLevelTasks::ETaskPriority Priority = Conversion[ThreadPriority];
 
