Removed enable_if that was preventing a clean static_assert

#rnx

diff --git a/Engine/Source/Runtime/Core/Public/Containers/UnrealString.h b/Engine/Source/Runtime/Core/Public/Containers/UnrealString.h
--- a/Engine/Source/Runtime/Core/Public/Containers/UnrealString.h
+++ b/Engine/Source/Runtime/Core/Public/Containers/UnrealString.h
@@ -1637,7 +1637,7 @@ public:
 	 * @returns FString object that was constructed using format and additional parameters.
 	 */
 	template <typename FmtType, typename... Types>
-	static typename TEnableIf<TIsArrayOrRefOfType<FmtType, TCHAR>::Value, FString>::Type Printf(const FmtType& Fmt, Types... Args)
+	static FString Printf(const FmtType& Fmt, Types... Args)
 	{
 		static_assert(TIsArrayOrRefOfType<FmtType, TCHAR>::Value, "Formatting string must be a TCHAR array.");
 		static_assert(TAnd<TIsValidVariadicFunctionArg<Types>...>::Value, "Invalid argument(s) passed to FString::Printf");
@@ -1650,7 +1650,7 @@ public:
 	 * @return a reference to the modified string, so that it can be chained
 	 */
 	template <typename FmtType, typename... Types>
-	typename TEnableIf<TIsArrayOrRefOfType<FmtType, TCHAR>::Value, FString&>::Type Appendf(const FmtType& Fmt, Types... Args)
+	FString& Appendf(const FmtType& Fmt, Types... Args)
 	{
 		static_assert(TIsArrayOrRefOfType<FmtType, TCHAR>::Value, "Formatting string must be a TCHAR array.");
 		static_assert(TAnd<TIsValidVariadicFunctionArg<Types>...>::Value, "Invalid argument(s) passed to FString::Appendf");
