SharedBuffer: Added FUniqueBuffer::MakeUnique(FSharedBuffer) clone or transfer ownership from a shared buffer

#rnx

diff --git a/Engine/Source/Runtime/Core/Private/Memory/SharedBuffer.cpp b/Engine/Source/Runtime/Core/Private/Memory/SharedBuffer.cpp
--- a/Engine/Source/Runtime/Core/Private/Memory/SharedBuffer.cpp
+++ b/Engine/Source/Runtime/Core/Private/Memory/SharedBuffer.cpp
@@ -116,6 +116,19 @@ FUniqueBuffer FUniqueBuffer::MakeView(void* Data, uint64 Size)
 	return FUniqueBuffer(new BufferOwnerPrivate::FBufferOwnerView(Data, Size));
 }
 
+FUniqueBuffer FUniqueBuffer::MakeUnique(FSharedBuffer Buffer)
+{
+	OwnerPtrType ExistingOwner = ToPrivateOwnerPtr(MoveTemp(Buffer));
+	if (!ExistingOwner || (ExistingOwner->IsOwned() && ExistingOwner->GetTotalRefCount() == 1))
+	{
+		return FUniqueBuffer(MoveTemp(ExistingOwner));
+	}
+	else
+	{
+		return Clone(ExistingOwner->GetData(), ExistingOwner->GetSize());
+	}
+}
+
 FUniqueBuffer::FUniqueBuffer(FBufferOwner* InOwner)
 	: Owner(InOwner)
 {
diff --git a/Engine/Source/Runtime/Core/Private/Tests/Memory/SharedBufferTest.cpp b/Engine/Source/Runtime/Core/Private/Tests/Memory/SharedBufferTest.cpp
--- a/Engine/Source/Runtime/Core/Private/Tests/Memory/SharedBufferTest.cpp
+++ b/Engine/Source/Runtime/Core/Private/Tests/Memory/SharedBufferTest.cpp
@@ -126,6 +126,15 @@ static_assert(std::is_same<uint32, decltype(GetTypeHash(DeclVal<const FWeakShare
 IMPLEMENT_SIMPLE_AUTOMATION_TEST(FUniqueBufferTest, "System.Core.Memory.UniqueBuffer", EAutomationTestFlags::ApplicationContextMask | EAutomationTestFlags::SmokeFilter)
 bool FUniqueBufferTest::RunTest(const FString& Parameters)
 {
+	// Test Null
+	{
+		FUniqueBuffer Buffer;
+		TestFalse(TEXT("FUniqueBuffer().IsOwned()"), Buffer.IsOwned());
+		TestEqual(TEXT("FUniqueBuffer().GetSize()"), Buffer.GetSize(), uint64(0));
+		TestEqual(TEXT("FUniqueBuffer().GetView().GetSize()"), Buffer.GetView().GetSize(), uint64(0));
+		TestEqual(TEXT("FUniqueBuffer().GetView().GetData()"), Buffer.GetView().GetData(), Buffer.GetData());
+	}
+
 	// Test Alloc
 	{
 		constexpr uint64 Size = 64;
@@ -174,6 +183,45 @@ bool FUniqueBufferTest::RunTest(const FString& Parameters)
 		TestTrue(TEXT("FUniqueBuffer::TakeOwnership() Deleted"), bDeleted);
 	}
 
+	// Test MakeUnique
+	{
+		FUniqueBuffer Buffer = FUniqueBuffer::MakeUnique(FSharedBuffer());
+		TestFalse(TEXT("FUniqueBuffer::MakeUnique(Null).IsOwned()"), Buffer.IsOwned());
+		TestEqual(TEXT("FUniqueBuffer::MakeUnique(Null).GetSize()"), Buffer.GetSize(), uint64(0));
+	}
+	{
+		FSharedBuffer Shared = FSharedBuffer::Clone(MakeMemoryView<uint8>({1, 2, 3, 4}));
+		FUniqueBuffer Buffer = FUniqueBuffer::MakeUnique(Shared);
+		TestTrue(TEXT("FUniqueBuffer::MakeUnique(Shared).IsOwned()"), Buffer.IsOwned());
+		TestEqual(TEXT("FUniqueBuffer::MakeUnique(Shared).GetSize()"), Buffer.GetSize(), Shared.GetSize());
+		TestNotEqual(TEXT("FUniqueBuffer::MakeUnique(Shared).GetData()"), const_cast<const void*>(Buffer.GetData()), Shared.GetData());
+		TestTrue(TEXT("FUniqueBuffer::MakeUnique(Shared)->Equal"), Buffer.GetView().EqualBytes(Shared.GetView()));
+	}
+	{
+		bool bDeleted = false;
+		constexpr uint64 Size = 64;
+		void* Data = new uint8[Size];
+		FUniqueBuffer Buffer = FUniqueBuffer::MakeUnique(FSharedBuffer::TakeOwnership(Data, Size,
+			[&bDeleted](void* InData, uint64 InSize)
+		{
+			bDeleted = true;
+			delete[] static_cast<uint8*>(InData);
+		}));
+		TestTrue(TEXT("FUniqueBuffer::MakeUnique(Owned).IsOwned()"), Buffer.IsOwned());
+		TestEqual(TEXT("FUniqueBuffer::MakeUnique(Owned).GetSize()"), Buffer.GetSize(), Size);
+		TestEqual(TEXT("FUniqueBuffer::MakeUnique(Owned).GetData()"), Buffer.GetData(), Data);
+		Buffer.Reset();
+		TestTrue(TEXT("FUniqueBuffer::MakeUnique(Owned) Deleted"), bDeleted);
+	}
+	{
+		constexpr uint64 Size = 64;
+		uint8 Data[Size]{};
+		FUniqueBuffer Buffer = FUniqueBuffer::MakeUnique(FSharedBuffer::MakeView(MakeMemoryView(Data)));
+		TestTrue(TEXT("FUniqueBuffer::MakeUnique(View).IsOwned()"), Buffer.IsOwned());
+		TestEqual(TEXT("FUniqueBuffer::MakeUnique(View).GetSize()"), Buffer.GetSize(), Size);
+		TestNotEqual(TEXT("FUniqueBuffer::MakeUnique(View).GetData()"), Buffer.GetData(), static_cast<void*>(Data));
+	}
+
 	// Test MakeOwned
 	{
 		constexpr uint64 Size = 64;
@@ -193,9 +241,18 @@ bool FUniqueBufferTest::RunTest(const FString& Parameters)
 	return true;
 }
 
-IMPLEMENT_SIMPLE_AUTOMATION_TEST(FSharedBuffer2Test, "System.Core.Memory.SharedBuffer2", EAutomationTestFlags::ApplicationContextMask | EAutomationTestFlags::SmokeFilter)
-bool FSharedBuffer2Test::RunTest(const FString& Parameters)
+IMPLEMENT_SIMPLE_AUTOMATION_TEST(FSharedBufferTest, "System.Core.Memory.SharedBuffer", EAutomationTestFlags::ApplicationContextMask | EAutomationTestFlags::SmokeFilter)
+bool FSharedBufferTest::RunTest(const FString& Parameters)
 {
+	// Test Null
+	{
+		FSharedBuffer Buffer;
+		TestFalse(TEXT("FSharedBuffer().IsOwned()"), Buffer.IsOwned());
+		TestEqual(TEXT("FSharedBuffer().GetSize()"), Buffer.GetSize(), uint64(0));
+		TestEqual(TEXT("FSharedBuffer().GetView().GetSize()"), Buffer.GetView().GetSize(), uint64(0));
+		TestEqual(TEXT("FSharedBuffer().GetView().GetData()"), Buffer.GetView().GetData(), Buffer.GetData());
+	}
+
 	// Test Clone
 	{
 		constexpr uint64 Size = 64;
diff --git a/Engine/Source/Runtime/Core/Public/Memory/SharedBuffer.h b/Engine/Source/Runtime/Core/Public/Memory/SharedBuffer.h
--- a/Engine/Source/Runtime/Core/Public/Memory/SharedBuffer.h
+++ b/Engine/Source/Runtime/Core/Public/Memory/SharedBuffer.h
@@ -81,7 +81,7 @@ protected:
 	inline bool IsMaterialized() const;
 	inline void SetIsMaterialized();
 
-	inline bool IsReferenced() const;
+	inline uint32 GetTotalRefCount() const;
 
 private:
 	static uint32 GetSharedRefCount(uint64 RefCountsAndFlags) { return uint32(RefCountsAndFlags >> 0) & 0x7fffffff; }
@@ -116,7 +116,7 @@ namespace BufferOwnerPrivate
 
 struct FSharedOps final
 {
-	static inline bool HasRef(FBufferOwner& Owner) { return Owner.IsReferenced(); }
+	static inline bool HasRef(FBufferOwner& Owner) { return Owner.GetTotalRefCount() > 0; }
 	static inline bool TryAddRef(FBufferOwner& Owner) { return Owner.TryAddSharedReference(); }
 	static inline void AddRef(FBufferOwner& Owner) { Owner.AddSharedReference(); }
 	static inline void Release(FBufferOwner* Owner) { if (Owner) { Owner->ReleaseSharedReference(); } }
@@ -124,7 +124,7 @@ struct FSharedOps final
 
 struct FWeakOps final
 {
-	static inline bool HasRef(FBufferOwner& Owner) { return Owner.IsReferenced(); }
+	static inline bool HasRef(FBufferOwner& Owner) { return Owner.GetTotalRefCount() > 0; }
 	static inline bool TryAddRef(FBufferOwner& Owner) { AddRef(Owner); return true; }
 	static inline void AddRef(FBufferOwner& Owner) { Owner.AddWeakReference(); }
 	static inline void Release(FBufferOwner* Owner) { if (Owner) { Owner->ReleaseWeakReference(); } }
@@ -226,6 +226,14 @@ public:
 		decltype(Invoke(std::declval<DeleteFunctionType>(), std::declval<void*>(), std::declval<uint64>()))* = nullptr>
 	static inline FUniqueBuffer TakeOwnership(void* Data, uint64 Size, DeleteFunctionType&& DeleteFunction);
 
+	/**
+	 * Make a unique buffer from a shared buffer.
+	 *
+	 * Steals the buffer owner from the shared buffer if this is the last reference to it, otherwise
+	 * clones the shared buffer to guarantee unique ownership. An non-owned buffer is always cloned.
+	 */
+	CORE_API static FUniqueBuffer MakeUnique(FSharedBuffer Buffer);
+
 	/** Construct a null unique buffer. */
 	FUniqueBuffer() = default;
 
@@ -278,6 +286,8 @@ private:
 	inline friend const OwnerPtrType& ToPrivateOwnerPtr(const FUniqueBuffer& Buffer) { return Buffer.Owner; }
 	inline friend OwnerPtrType ToPrivateOwnerPtr(FUniqueBuffer&& Buffer) { return MoveTemp(Buffer.Owner); }
 
+	inline explicit FUniqueBuffer(OwnerPtrType&& InOwner) : Owner(MoveTemp(InOwner)) {}
+
 	OwnerPtrType Owner;
 };
 
@@ -379,6 +389,7 @@ private:
 	using OwnerPtrType = BufferOwnerPrivate::TBufferOwnerPtr<BufferOwnerPrivate::FSharedOps>;
 
 	inline friend const OwnerPtrType& ToPrivateOwnerPtr(const FSharedBuffer& Buffer) { return Buffer.Owner; }
+	inline friend OwnerPtrType ToPrivateOwnerPtr(FSharedBuffer&& Buffer) { return MoveTemp(Buffer.Owner); }
 
 	OwnerPtrType Owner;
 };
@@ -525,7 +536,7 @@ inline FBufferOwner::FBufferOwner(void* InData, uint64 InSize)
 
 inline FBufferOwner::~FBufferOwner()
 {
-	checkSlow(!IsReferenced());
+	checkSlow(GetTotalRefCount() == 0);
 }
 
 inline void FBufferOwner::MaterializeBuffer()
@@ -580,10 +591,13 @@ inline void FBufferOwner::SetIsMaterialized()
 	ReferenceCountsAndFlags.fetch_or(SetFlags(EBufferOwnerFlags::Materialized));
 }
 
-inline bool FBufferOwner::IsReferenced() const
+inline uint32 FBufferOwner::GetTotalRefCount() const
 {
 	const uint64 Value = ReferenceCountsAndFlags.load(std::memory_order_relaxed);
-	return GetSharedRefCount(Value) || GetWeakRefCount(Value);
+	const uint32 SharedRefCount = GetSharedRefCount(Value);
+	// A non-zero SharedRefCount adds 1 to WeakRefCount to keep the owner alive.
+	// Subtract that extra reference when it is present to return an accurate count.
+	return GetWeakRefCount(Value) + SharedRefCount - !!SharedRefCount;
 }
 
 inline void FBufferOwner::AddSharedReference()
