BLAKE3: Added FBlake3Hash::Reset()

#rnx

diff --git a/Engine/Source/Runtime/Core/Public/Misc/Blake3.h b/Engine/Source/Runtime/Core/Public/Misc/Blake3.h
--- a/Engine/Source/Runtime/Core/Public/Misc/Blake3.h
+++ b/Engine/Source/Runtime/Core/Public/Misc/Blake3.h
@@ -18,16 +18,23 @@ struct FBlake3Hash
 public:
 	using ByteArray = uint8[32];
 
+	/** Construct a zero hash. */
 	FBlake3Hash() = default;
 	explicit FBlake3Hash(const ByteArray& Hash);
 
+	/** Construct a hash from a 64-character hex string. */
 	CORE_API explicit FBlake3Hash(FAnsiStringView HexHash);
 	CORE_API explicit FBlake3Hash(FWideStringView HexHash);
 
 	CORE_API friend FArchive& operator<<(FArchive& Ar, FBlake3Hash& Hash);
 
+	/** Reset this to a zero hash. */
+	inline void Reset() { *this = FBlake3Hash(); }
+
+	/** Returns a reference to the raw byte array for the hash. */
 	inline const ByteArray& Bytes() const { return Hash; }
 
+	/** Returns whether this is a zero hash. */
 	inline bool IsZero() const;
 
 private:
