CompactBinary: Added Reset() functions to several compact binary types

#rnx

diff --git a/Engine/Source/Runtime/Core/Private/Serialization/CompactBinaryPackage.cpp b/Engine/Source/Runtime/Core/Private/Serialization/CompactBinaryPackage.cpp
--- a/Engine/Source/Runtime/Core/Private/Serialization/CompactBinaryPackage.cpp
+++ b/Engine/Source/Runtime/Core/Private/Serialization/CompactBinaryPackage.cpp
@@ -30,7 +30,7 @@ FCbAttachment::FCbAttachment(FCbFieldRefIterator InValue, const FBlake3Hash* con
 		}
 		else
 		{
-			checkfSlow(Hash == FBlake3Hash(), TEXT("A null or empty field range must use a hash of zero."));
+			checkfSlow(Hash.IsZero(), TEXT("A null or empty field range must use a hash of zero."));
 		}
 	}
 	else if (CompactBinary)
@@ -52,7 +52,7 @@ FCbAttachment::FCbAttachment(FSharedBuffer InBuffer, const FBlake3Hash* const In
 		}
 		else
 		{
-			checkfSlow(Hash == FBlake3Hash(), TEXT("A null or empty buffer must use a hash of zero."));
+			checkfSlow(Hash.IsZero(), TEXT("A null or empty buffer must use a hash of zero."));
 		}
 	}
 	else if (Buffer.GetSize())
@@ -106,9 +106,9 @@ void FCbAttachment::Load(FCbFieldRefIterator& Fields)
 	else
 	{
 		++Fields;
-		Buffer = FSharedBuffer();
-		CompactBinary = FCbFieldIterator();
-		Hash = FBlake3Hash();
+		Buffer.Reset();
+		CompactBinary.Reset();
+		Hash.Reset();
 	}
 }
 
@@ -139,9 +139,9 @@ void FCbAttachment::Load(FArchive& Ar, FCbBufferAllocator Allocator)
 	}
 	else
 	{
-		Buffer = FSharedBuffer();
-		CompactBinary = FCbFieldIterator();
-		Hash = FBlake3Hash();
+		Buffer.Reset();
+		CompactBinary.Reset();
+		Hash.Reset();
 	}
 }
 
@@ -201,8 +201,8 @@ void FCbPackage::SetObject(FCbObjectRef InObject, const FBlake3Hash* InObjectHas
 	}
 	else
 	{
-		Object = FCbObjectRef();
-		ObjectHash = FBlake3Hash();
+		Object.Reset();
+		ObjectHash.Reset();
 	}
 }
 
@@ -303,7 +303,7 @@ void FCbPackage::Load(FCbFieldRefIterator& Fields)
 			}
 			else
 			{
-				Object = FCbObjectRef();
+				Object.Reset();
 			}
 		}
 	}
@@ -358,7 +358,7 @@ void FCbPackage::Load(FArchive& Ar, FCbBufferAllocator Allocator)
 			}
 			else
 			{
-				Object = FCbObjectRef();
+				Object.Reset();
 			}
 		}
 	}
diff --git a/Engine/Source/Runtime/Core/Public/Serialization/CompactBinary.h b/Engine/Source/Runtime/Core/Public/Serialization/CompactBinary.h
--- a/Engine/Source/Runtime/Core/Public/Serialization/CompactBinary.h
+++ b/Engine/Source/Runtime/Core/Public/Serialization/CompactBinary.h
@@ -302,6 +302,9 @@ public:
 	constexpr inline FieldType& operator*() { return *this; }
 	constexpr inline FieldType* operator->() { return this; }
 
+	/** Reset this to an empty field range. */
+	inline void Reset() { *this = TCbFieldIterator(); }
+
 	/** Returns the size of the fields in the range in bytes. */
 	CORE_API uint64 GetRangeSize() const;
 
@@ -953,6 +956,9 @@ public:
 	{
 	}
 
+	/** Reset this to a default value and null buffer. */
+	inline void Reset() { *this = TCbBufferRef(); }
+
 	/** Whether this reference has ownership of the memory in its buffer. */
 	inline bool IsOwned() const { return Buffer.IsOwned(); }
 
diff --git a/Engine/Source/Runtime/Core/Public/Serialization/CompactBinaryPackage.h b/Engine/Source/Runtime/Core/Public/Serialization/CompactBinaryPackage.h
--- a/Engine/Source/Runtime/Core/Public/Serialization/CompactBinaryPackage.h
+++ b/Engine/Source/Runtime/Core/Public/Serialization/CompactBinaryPackage.h
@@ -53,6 +53,9 @@ public:
 	{
 	}
 
+	/** Reset this to a null attachment. */
+	inline void Reset() { *this = FCbAttachment(); }
+
 	/** Whether the attachment has a value. */
 	inline explicit operator bool() const { return !IsNull(); }
 
@@ -202,6 +205,9 @@ public:
 		SetObject(MoveTemp(InObject), InObjectHash, InResolver);
 	}
 
+	/** Reset this to a null package. */
+	inline void Reset() { *this = FCbPackage(); }
+
 	/** Whether the package has a non-empty object or attachments. */
 	inline explicit operator bool() const { return !IsNull(); }
 
