U5 - Fix a rare deadlock where all Backgroundworkers have been promoted to foregroundworkers (due to BusyWaiting) and one of them is waiting on a BackgroundTask. The solution is to demote a BusyWaiting Worker again when no work was found, but only it was a Backgroundworker to begin with.

diff --git a/Engine/Source/Runtime/Core/Private/Async/Fundamental/Scheduler.cpp b/Engine/Source/Runtime/Core/Private/Async/Fundamental/Scheduler.cpp
--- a/Engine/Source/Runtime/Core/Private/Async/Fundamental/Scheduler.cpp
+++ b/Engine/Source/Runtime/Core/Private/Async/Fundamental/Scheduler.cpp
@@ -12,6 +12,7 @@ namespace LowLevelTasks
 
 	thread_local FScheduler::FLocalQueueType* FScheduler::LocalQueue = nullptr;
 	thread_local FTask* FScheduler::ActiveTask = nullptr;
+	thread_local bool FScheduler::bIsBackgroundWorker = false;
 
 	FScheduler FScheduler::Singleton;
 
@@ -115,7 +116,6 @@ namespace LowLevelTasks
 		if (ActiveWorkers.load(std::memory_order_relaxed))
 		{			
 			const bool bIsBackgroundTask = Task.IsBackgroundTask();
-			const bool bIsBackgroundWorker = ActiveTask && ActiveTask->IsBackgroundTask();
 			if (bIsBackgroundTask && !bIsBackgroundWorker)
 			{
 				QueuePreference = EQueuePreference::GlobalQueuePreference;
@@ -177,6 +177,7 @@ namespace LowLevelTasks
 	void FScheduler::WorkerMain(FSleepEvent* WorkerEvent, FLocalQueueType* ExternalWorkerLocalQueue, uint32 WaitCycles, bool bPermitBackgroundWork)
 	{
 		FMemory::SetupTLSCachesOnCurrentThread();
+		bIsBackgroundWorker = bPermitBackgroundWork;
 
 		checkSlow(LocalQueue == nullptr);
 		if(ExternalWorkerLocalQueue)
@@ -232,6 +233,7 @@ namespace LowLevelTasks
 		FLocalQueueType::DeleteLocalQueue(WorkerLocalQueue, bPermitBackgroundWork, ExternalWorkerLocalQueue != nullptr);
 		LocalQueue = nullptr;
 
+		bIsBackgroundWorker = false;
 		FMemory::ClearAndDisableTLSCachesOnCurrentThread();
 	}
 
@@ -281,6 +283,10 @@ namespace LowLevelTasks
 				FPlatformProcess::Yield();
 				WaitCount++;
 			}
+			else if (!bPermitBackgroundWork && bIsBackgroundWorker)
+			{
+				bPermitBackgroundWork = true;
+			}
 			else
 			{
 				TRACE_CPUPROFILER_EVENT_SCOPE(BusyWaitInternal::SleepNoStats);
diff --git a/Engine/Source/Runtime/Core/Public/Async/Fundamental/Scheduler.h b/Engine/Source/Runtime/Core/Public/Async/Fundamental/Scheduler.h
--- a/Engine/Source/Runtime/Core/Public/Async/Fundamental/Scheduler.h
+++ b/Engine/Source/Runtime/Core/Public/Async/Fundamental/Scheduler.h
@@ -29,6 +29,7 @@ namespace LowLevelTasks
 
 		static thread_local FLocalQueueType* LocalQueue;
 		static thread_local FTask* ActiveTask;
+		static thread_local bool bIsBackgroundWorker;
 		static CORE_API FScheduler Singleton;
 
 		// using 16 bytes here because it fits the vtable and one additional pointer
