U5 - disable Busy Waiting for all NamedThreads, before it was possible that they do Busy wait when their Queues are empty. But there are some concerns with the TaskTagScopes that might cause some ensures when Tasks are pulled in on the Gamethread that have a parallel RenderingThread etc. scope.

diff --git a/Engine/Source/Runtime/Core/Private/Async/TaskGraph.cpp b/Engine/Source/Runtime/Core/Private/Async/TaskGraph.cpp
--- a/Engine/Source/Runtime/Core/Private/Async/TaskGraph.cpp
+++ b/Engine/Source/Runtime/Core/Private/Async/TaskGraph.cpp
@@ -2007,7 +2007,7 @@ private:
 			TGraphTask<FReturnGraphTask>::CreateTask(&Tasks, CurrentThread).ConstructAndDispatchWhenReady(CurrentThread);
 			ProcessThreadUntilRequestReturn(CurrentThread);
 		}
-		else
+		else if (LowLevelTasks::FScheduler::Get().GetActiveTask() != nullptr)
 		{
 			LowLevelTasks::BusyWaitUntil([Index(0), &Tasks]() mutable
 			{
@@ -2022,6 +2022,30 @@ private:
 				return true;
 			});
 		}
+		else
+		{
+			if (!FTaskGraphInterface::IsMultithread())
+			{
+				bool bAnyPending = false;
+				for (int32 Index = 0; Index < Tasks.Num(); Index++)
+				{
+					FGraphEvent* Task = Tasks[Index].GetReference();
+					if (Task && !Task->IsComplete())
+					{
+						bAnyPending = true;
+						break;
+					}
+				}
+				if (!bAnyPending)
+				{
+					return;
+				}
+				UE_LOG(LogTaskGraph, Fatal, TEXT("Recursive waits are not allowed in single threaded mode."));
+			}
+			// We will just stall this thread on an event while we wait
+			FScopedEvent Event;
+			TriggerEventWhenTasksComplete(Event.Get(), Tasks, CurrentThreadIfKnown);
+		}
 	}
 
 	void TriggerEventWhenTasksComplete(FEvent* InEvent, const FGraphEventArray& Tasks, ENamedThreads::Type CurrentThreadIfKnown = ENamedThreads::AnyThread, ENamedThreads::Type TriggerThread = ENamedThreads::AnyHiPriThreadHiPriTask) final override
