It is okay to call GAux.EnableChannels() directly here.

diff --git a/Engine/Source/Runtime/Core/Private/ProfilingDebugging/TraceAuxiliary.cpp b/Engine/Source/Runtime/Core/Private/ProfilingDebugging/TraceAuxiliary.cpp
--- a/Engine/Source/Runtime/Core/Private/ProfilingDebugging/TraceAuxiliary.cpp
+++ b/Engine/Source/Runtime/Core/Private/ProfilingDebugging/TraceAuxiliary.cpp
@@ -417,7 +417,7 @@ void FTraceAuxiliary::Initialize(const TCHAR* CommandLine)
 	{
 		if (Reason == EModuleChangeReason::ModuleLoaded)
 		{
-			EnableChannels();
+			GTraceAuxiliary.EnableChannels();
 		}
 	});
 
