Speculative fix for an issue where the deprecation warning for the use of the "TArray::Allocator" type is triggered intermittently on the farm.  I cannot repro the issue even when doing clean builds using the same MSVC version as the farm machines, and on the farm machines, it seems intermittent.

diff --git a/Engine/Source/Runtime/Core/Public/Containers/MRUArray.h b/Engine/Source/Runtime/Core/Public/Containers/MRUArray.h
--- a/Engine/Source/Runtime/Core/Public/Containers/MRUArray.h
+++ b/Engine/Source/Runtime/Core/Public/Containers/MRUArray.h
@@ -16,12 +16,12 @@
  * - Has an upper limit of the number of items it will store.
  * - Any item that is added to the array is moved to the top.
  */
-template<typename T, typename Allocator = FDefaultAllocator>
+template<typename T, typename InAllocator = FDefaultAllocator>
 class TMRUArray
-	: public TArray<T, Allocator>
+	: public TArray<T, InAllocator>
 {
 public:
-	typedef TArray<T, Allocator> Super;
+	typedef TArray<T, InAllocator> Super;
 
 	/** The maximum number of items we can store in this array. */
 	int32 MaxItems;
