Added UE_NODISCARD to String, StringView, MemoryView

diff --git a/Engine/Source/Runtime/Core/Public/Containers/StringView.h b/Engine/Source/Runtime/Core/Public/Containers/StringView.h
--- a/Engine/Source/Runtime/Core/Public/Containers/StringView.h
+++ b/Engine/Source/Runtime/Core/Public/Containers/StringView.h
@@ -158,15 +158,15 @@ public:
 	inline const CharType& operator[](SizeType Index) const;
 
 	/** Returns a pointer to the start of the view. This is NOT guaranteed to be null-terminated! */
-	constexpr inline const CharType* GetData() const { return DataPtr; }
+	UE_NODISCARD constexpr inline const CharType* GetData() const { return DataPtr; }
 
 	// Capacity
 
 	/** Returns the length of the string view. */
-	constexpr inline SizeType Len() const { return Size; }
+	UE_NODISCARD constexpr inline SizeType Len() const { return Size; }
 
 	/** Returns whether the string view is empty. */
-	constexpr inline bool IsEmpty() const { return Size == 0; }
+	UE_NODISCARD constexpr inline bool IsEmpty() const { return Size == 0; }
 
 	// Modifiers
 
@@ -191,24 +191,24 @@ public:
 	inline SizeType CopyString(CharType* Dest, SizeType CharCount, SizeType Position = 0) const;
 
 	/** Alias for Mid. */
-	inline ViewType SubStr(SizeType Position, SizeType CharCount) const { return Mid(Position, CharCount); }
+	UE_NODISCARD inline ViewType SubStr(SizeType Position, SizeType CharCount) const { return Mid(Position, CharCount); }
 
 	/** Returns the left-most part of the view by taking the given number of characters from the left. */
-	inline ViewType Left(SizeType CharCount) const;
+	UE_NODISCARD inline ViewType Left(SizeType CharCount) const;
 	/** Returns the left-most part of the view by chopping the given number of characters from the right. */
-	inline ViewType LeftChop(SizeType CharCount) const;
+	UE_NODISCARD inline ViewType LeftChop(SizeType CharCount) const;
 	/** Returns the right-most part of the view by taking the given number of characters from the right. */
-	inline ViewType Right(SizeType CharCount) const;
+	UE_NODISCARD inline ViewType Right(SizeType CharCount) const;
 	/** Returns the right-most part of the view by chopping the given number of characters from the left. */
-	inline ViewType RightChop(SizeType CharCount) const;
+	UE_NODISCARD inline ViewType RightChop(SizeType CharCount) const;
 	/** Returns the middle part of the view by taking up to the given number of characters from the given position. */
-	inline ViewType Mid(SizeType Position, SizeType CharCount = TNumericLimits<SizeType>::Max()) const;
+	UE_NODISCARD inline ViewType Mid(SizeType Position, SizeType CharCount = TNumericLimits<SizeType>::Max()) const;
 	/** Returns the middle part of the view between any whitespace at the start and end. */
-	inline ViewType TrimStartAndEnd() const;
+	UE_NODISCARD inline ViewType TrimStartAndEnd() const;
 	/** Returns the right part of the view after any whitespace at the start. */
-	CORE_API ViewType TrimStart() const;
+	UE_NODISCARD CORE_API ViewType TrimStart() const;
 	/** Returns the left part of the view before any whitespace at the end. */
-	CORE_API ViewType TrimEnd() const;
+	UE_NODISCARD CORE_API ViewType TrimEnd() const;
 
 	/** Modifies the view to be the given number of characters from the left. */
 	inline void LeftInline(SizeType CharCount) { *this = Left(CharCount); }
@@ -235,7 +235,7 @@ public:
 	 * @param SearchCase Whether the comparison should ignore case.
 	 */
 	template <typename OtherType, typename TEnableIf<StringViewPrivate::TIsConvertibleToStringView<OtherType>::Value>::Type* = nullptr>
-	inline bool Equals(OtherType&& Other, ESearchCase::Type SearchCase = ESearchCase::CaseSensitive) const;
+	UE_NODISCARD inline bool Equals(OtherType&& Other, ESearchCase::Type SearchCase = ESearchCase::CaseSensitive) const;
 
 	/**
 	 * Check whether this view is lexicographically equivalent to another view.
@@ -243,7 +243,7 @@ public:
 	 * @param SearchCase Whether the comparison should ignore case.
 	 */
 	template <typename OtherCharType>
-	inline bool Equals(const OtherCharType* Other, ESearchCase::Type SearchCase = ESearchCase::CaseSensitive) const;
+	UE_NODISCARD inline bool Equals(const OtherCharType* Other, ESearchCase::Type SearchCase = ESearchCase::CaseSensitive) const;
 
 	/**
 	 * Compare this view lexicographically with another view.
@@ -253,7 +253,7 @@ public:
 	 * @return 0 is equal, negative if this view is less, positive if this view is greater.
 	 */
 	template <typename OtherType, typename TEnableIf<StringViewPrivate::TIsConvertibleToStringView<OtherType>::Value>::Type* = nullptr>
-	inline int32 Compare(OtherType&& Other, ESearchCase::Type SearchCase = ESearchCase::CaseSensitive) const;
+	UE_NODISCARD inline int32 Compare(OtherType&& Other, ESearchCase::Type SearchCase = ESearchCase::CaseSensitive) const;
 
 	/**
 	 * Compare this view lexicographically with another view.
@@ -263,17 +263,17 @@ public:
 	 * @return 0 is equal, negative if this view is less, positive if this view is greater.
 	 */
 	template <typename OtherCharType>
-	inline int32 Compare(const OtherCharType* Other, ESearchCase::Type SearchCase = ESearchCase::CaseSensitive) const;
+	UE_NODISCARD inline int32 Compare(const OtherCharType* Other, ESearchCase::Type SearchCase = ESearchCase::CaseSensitive) const;
 
 	/** Returns whether this view starts with the prefix character compared case-sensitively. */
-	inline bool StartsWith(CharType Prefix) const { return Size >= 1 && DataPtr[0] == Prefix; }
+	UE_NODISCARD inline bool StartsWith(CharType Prefix) const { return Size >= 1 && DataPtr[0] == Prefix; }
 	/** Returns whether this view starts with the prefix with optional case sensitivity. */
-	inline bool StartsWith(ViewType Prefix, ESearchCase::Type SearchCase = ESearchCase::IgnoreCase) const;
+	UE_NODISCARD inline bool StartsWith(ViewType Prefix, ESearchCase::Type SearchCase = ESearchCase::IgnoreCase) const;
 
 	/** Returns whether this view ends with the suffix character compared case-sensitively. */
-	inline bool EndsWith(CharType Suffix) const { return Size >= 1 && DataPtr[Size-1] == Suffix; }
+	UE_NODISCARD inline bool EndsWith(CharType Suffix) const { return Size >= 1 && DataPtr[Size-1] == Suffix; }
 	/** Returns whether this view ends with the suffix with optional case sensitivity. */
-	inline bool EndsWith(ViewType Suffix, ESearchCase::Type SearchCase = ESearchCase::IgnoreCase) const;
+	UE_NODISCARD inline bool EndsWith(ViewType Suffix, ESearchCase::Type SearchCase = ESearchCase::IgnoreCase) const;
 
 	// Searching/Finding
 
diff --git a/Engine/Source/Runtime/Core/Public/Containers/UnrealString.h b/Engine/Source/Runtime/Core/Public/Containers/UnrealString.h
--- a/Engine/Source/Runtime/Core/Public/Containers/UnrealString.h
+++ b/Engine/Source/Runtime/Core/Public/Containers/UnrealString.h
@@ -408,7 +408,7 @@ public:
 	 *
 	 * @return true if this string is empty, otherwise return false.
 	 */
-	FORCEINLINE bool IsEmpty() const
+	UE_NODISCARD FORCEINLINE bool IsEmpty() const
 	{
 		return Data.Num() <= 1;
 	}
@@ -443,7 +443,7 @@ public:
 	 *
 	 * @returns True if index is valid. False otherwise.
 	 */
-	FORCEINLINE bool IsValidIndex(int32 Index) const
+	UE_NODISCARD FORCEINLINE bool IsValidIndex(int32 Index) const
 	{
 		return Index >= 0 && Index < Len();
 	}
@@ -453,7 +453,7 @@ public:
 	 *
 	 * @Return Pointer to Array of TCHAR if Num, otherwise the empty string
 	 */
-	FORCEINLINE const TCHAR* operator*() const
+	UE_NODISCARD FORCEINLINE const TCHAR* operator*() const
 	{
 		return Data.Num() ? Data.GetData() : TEXT("");
 	}
@@ -464,13 +464,13 @@ public:
 	 * @warning: Operations on the TArray<*CHAR> can be unsafe, such as adding
 	 *		non-terminating 0's or removing the terminating zero.
 	 */
-	FORCEINLINE DataType& GetCharArray()
+	UE_NODISCARD FORCEINLINE DataType& GetCharArray()
 	{
 		return Data;
 	}
 
 	/** Get string as const array of TCHARS */
-	FORCEINLINE const DataType& GetCharArray() const
+	UE_NODISCARD FORCEINLINE const DataType& GetCharArray() const
 	{
 		return Data;
 	}
@@ -648,7 +648,7 @@ public:
 	 * @return The concatenated string.
 	 */
 	template <typename CharType>
-	FORCEINLINE friend typename TEnableIf<TIsCharType<CharType>::Value, FString>::Type operator+(const FString& Lhs, CharType Rhs)
+	UE_NODISCARD FORCEINLINE friend typename TEnableIf<TIsCharType<CharType>::Value, FString>::Type operator+(const FString& Lhs, CharType Rhs)
 	{
 		Lhs.CheckInvariants();
 
@@ -667,7 +667,7 @@ public:
 	 * @return The concatenated string.
 	 */
 	template <typename CharType>
-	FORCEINLINE friend typename TEnableIf<TIsCharType<CharType>::Value, FString>::Type operator+(FString&& Lhs, CharType Rhs)
+	UE_NODISCARD FORCEINLINE friend typename TEnableIf<TIsCharType<CharType>::Value, FString>::Type operator+(FString&& Lhs, CharType Rhs)
 	{
 		Lhs.CheckInvariants();
 
@@ -679,7 +679,7 @@ public:
 
 private:
 	template <typename LhsType, typename RhsType>
-	FORCEINLINE static FString ConcatFStrings(typename TIdentity<LhsType>::Type Lhs, typename TIdentity<RhsType>::Type Rhs)
+	UE_NODISCARD FORCEINLINE static FString ConcatFStrings(typename TIdentity<LhsType>::Type Lhs, typename TIdentity<RhsType>::Type Rhs)
 	{
 		Lhs.CheckInvariants();
 		Rhs.CheckInvariants();
@@ -698,7 +698,7 @@ private:
 	}
 
 	template <typename RhsType>
-	FORCEINLINE static FString ConcatTCHARsToFString(const TCHAR* Lhs, typename TIdentity<RhsType>::Type Rhs)
+	UE_NODISCARD FORCEINLINE static FString ConcatTCHARsToFString(const TCHAR* Lhs, typename TIdentity<RhsType>::Type Rhs)
 	{
 		checkSlow(Lhs);
 		if (!Lhs)
@@ -709,7 +709,7 @@ private:
 	}
 
 	template <typename RhsType>
-	FORCEINLINE static FString ConcatTCHARRangeToFString(const TCHAR* Lhs, int32 LhsLen, typename TIdentity<RhsType>::Type Rhs)
+	UE_NODISCARD FORCEINLINE static FString ConcatTCHARRangeToFString(const TCHAR* Lhs, int32 LhsLen, typename TIdentity<RhsType>::Type Rhs)
 	{
 		checkSlow(LhsLen >= 0);
 		Rhs.CheckInvariants();
@@ -736,7 +736,7 @@ private:
 	}
 
 	template <typename LhsType>
-	FORCEINLINE static FString ConcatFStringToTCHARs(typename TIdentity<LhsType>::Type Lhs, const TCHAR* Rhs)
+	UE_NODISCARD FORCEINLINE static FString ConcatFStringToTCHARs(typename TIdentity<LhsType>::Type Lhs, const TCHAR* Rhs)
 	{
 		checkSlow(Rhs);
 		if (!Rhs)
@@ -747,7 +747,7 @@ private:
 	}
 
 	template <typename LhsType>
-	FORCEINLINE static FString ConcatFStringToTCHARRange(typename TIdentity<LhsType>::Type Lhs, const TCHAR* Rhs, int32 RhsLen)
+	UE_NODISCARD FORCEINLINE static FString ConcatFStringToTCHARRange(typename TIdentity<LhsType>::Type Lhs, const TCHAR* Rhs, int32 RhsLen)
 	{
 		Lhs.CheckInvariants();
 		checkSlow(RhsLen >= 0);
@@ -771,7 +771,7 @@ public:
 	 *
 	 * @return The concatenated string.
 	 */
-	FORCEINLINE friend FString operator+(const FString& Lhs, const FString& Rhs)
+	UE_NODISCARD FORCEINLINE friend FString operator+(const FString& Lhs, const FString& Rhs)
 	{
 		return ConcatFStrings<const FString&, const FString&>(Lhs, Rhs);
 	}
@@ -784,7 +784,7 @@ public:
 	 *
 	 * @return The concatenated string.
 	 */
-	FORCEINLINE friend FString operator+(FString&& Lhs, const FString& Rhs)
+	UE_NODISCARD FORCEINLINE friend FString operator+(FString&& Lhs, const FString& Rhs)
 	{
 		return ConcatFStrings<FString&&, const FString&>(MoveTemp(Lhs), Rhs);
 	}
@@ -797,7 +797,7 @@ public:
 	 *
 	 * @return The concatenated string.
 	 */
-	FORCEINLINE friend FString operator+(const FString& Lhs, FString&& Rhs)
+	UE_NODISCARD FORCEINLINE friend FString operator+(const FString& Lhs, FString&& Rhs)
 	{
 		return ConcatFStrings<const FString&, FString&&>(Lhs, MoveTemp(Rhs));
 	}
@@ -810,7 +810,7 @@ public:
 	 *
 	 * @return The concatenated string.
 	 */
-	FORCEINLINE friend FString operator+(FString&& Lhs, FString&& Rhs)
+	UE_NODISCARD FORCEINLINE friend FString operator+(FString&& Lhs, FString&& Rhs)
 	{
 		return ConcatFStrings<FString&&, FString&&>(MoveTemp(Lhs), MoveTemp(Rhs));
 	}
@@ -823,7 +823,7 @@ public:
 	 *
 	 * @return The concatenated string.
 	 */
-	FORCEINLINE friend FString operator+(const TCHAR* Lhs, const FString& Rhs)
+	UE_NODISCARD FORCEINLINE friend FString operator+(const TCHAR* Lhs, const FString& Rhs)
 	{
 		return ConcatTCHARsToFString<const FString&>(Lhs, Rhs);
 	}
@@ -837,7 +837,7 @@ public:
 	 * @return The concatenated string.
 	 */
 	template <typename CharRangeType, typename TEnableIf<TIsTCharRangeNotCArrayNotFString<CharRangeType>::Value>::Type* = nullptr>
-	FORCEINLINE friend FString operator+(CharRangeType&& Lhs, const FString& Rhs)
+	UE_NODISCARD FORCEINLINE friend FString operator+(CharRangeType&& Lhs, const FString& Rhs)
 	{
 		return ConcatTCHARRangeToFString<const FString&>(GetData(Lhs), GetNum(Lhs), Rhs);
 	}
@@ -850,7 +850,7 @@ public:
 	 *
 	 * @return The concatenated string.
 	 */
-	FORCEINLINE friend FString operator+(const TCHAR* Lhs, FString&& Rhs)
+	UE_NODISCARD FORCEINLINE friend FString operator+(const TCHAR* Lhs, FString&& Rhs)
 	{
 		return ConcatTCHARsToFString<FString&&>(Lhs, MoveTemp(Rhs));
 	}
@@ -864,7 +864,7 @@ public:
 	 * @return The concatenated string.
 	 */
 	template <typename CharRangeType, typename TEnableIf<TIsTCharRangeNotCArrayNotFString<CharRangeType>::Value>::Type* = nullptr>
-	FORCEINLINE friend FString operator+(CharRangeType&& Lhs, FString&& Rhs)
+	UE_NODISCARD FORCEINLINE friend FString operator+(CharRangeType&& Lhs, FString&& Rhs)
 	{
 		return ConcatTCHARRangeToFString<FString&&>(GetData(Lhs), GetNum(Lhs), MoveTemp(Rhs));
 	}
@@ -877,7 +877,7 @@ public:
 	 *
 	 * @return The concatenated string.
 	 */
-	FORCEINLINE friend FString operator+(const FString& Lhs, const TCHAR* Rhs)
+	UE_NODISCARD FORCEINLINE friend FString operator+(const FString& Lhs, const TCHAR* Rhs)
 	{
 		return ConcatFStringToTCHARs<const FString&>(Lhs, Rhs);
 	}
@@ -891,7 +891,7 @@ public:
 	 * @return The concatenated string.
 	 */
 	template <typename CharRangeType, typename TEnableIf<TIsTCharRangeNotCArrayNotFString<CharRangeType>::Value>::Type* = nullptr>
-	FORCEINLINE friend FString operator+(const FString& Lhs, CharRangeType&& Rhs)
+	UE_NODISCARD FORCEINLINE friend FString operator+(const FString& Lhs, CharRangeType&& Rhs)
 	{
 		return ConcatFStringToTCHARRange<const FString&>(Lhs, GetData(Rhs), GetNum(Rhs));
 	}
@@ -904,7 +904,7 @@ public:
 	 *
 	 * @return The concatenated string.
 	 */
-	FORCEINLINE friend FString operator+(FString&& Lhs, const TCHAR* Rhs)
+	UE_NODISCARD FORCEINLINE friend FString operator+(FString&& Lhs, const TCHAR* Rhs)
 	{
 		return ConcatFStringToTCHARs<FString&&>(MoveTemp(Lhs), Rhs);
 	}
@@ -918,7 +918,7 @@ public:
 	 * @return The concatenated string.
 	 */
 	template <typename CharRangeType, typename TEnableIf<TIsTCharRangeNotCArrayNotFString<CharRangeType>::Value>::Type* = nullptr>
-	FORCEINLINE friend FString operator+(FString&& Lhs, CharRangeType&& Rhs)
+	UE_NODISCARD FORCEINLINE friend FString operator+(FString&& Lhs, CharRangeType&& Rhs)
 	{
 		return ConcatFStringToTCHARRange<FString&&>(MoveTemp(Lhs), GetData(Rhs), GetNum(Rhs));
 	}
@@ -956,7 +956,7 @@ public:
 	 * @param Rhs Path to concatenate.
 	 * @return new FString of the path
 	 */
-	FORCEINLINE friend FString operator/(const FString& Lhs, const TCHAR* Rhs)
+	UE_NODISCARD FORCEINLINE friend FString operator/(const FString& Lhs, const TCHAR* Rhs)
 	{
 		checkSlow(Rhs);
 
@@ -974,7 +974,7 @@ public:
 	 * @param Rhs Path to concatenate.
 	 * @return new FString of the path
 	 */
-	FORCEINLINE friend FString operator/(FString&& Lhs, const TCHAR* Rhs)
+	UE_NODISCARD FORCEINLINE friend FString operator/(FString&& Lhs, const TCHAR* Rhs)
 	{
 		checkSlow(Rhs);
 
@@ -992,7 +992,7 @@ public:
 	 * @param Rhs Path to concatenate.
 	 * @return new FString of the path
 	 */
-	FORCEINLINE friend FString operator/(const FString& Lhs, const FString& Rhs)
+	UE_NODISCARD FORCEINLINE friend FString operator/(const FString& Lhs, const FString& Rhs)
 	{
 		int32 StrLength = Rhs.Len();
 
@@ -1008,7 +1008,7 @@ public:
 	 * @param Rhs Path to concatenate.
 	 * @return new FString of the path
 	 */
-	FORCEINLINE friend FString operator/(FString&& Lhs, const FString& Rhs)
+	UE_NODISCARD FORCEINLINE friend FString operator/(FString&& Lhs, const FString& Rhs)
 	{
 		int32 StrLength = Rhs.Len();
 
@@ -1024,7 +1024,7 @@ public:
 	 * @param Rhs Path to concatenate.
 	 * @return new FString of the path
 	 */
-	FORCEINLINE friend FString operator/(const TCHAR* Lhs, const FString& Rhs)
+	UE_NODISCARD FORCEINLINE friend FString operator/(const TCHAR* Lhs, const FString& Rhs)
 	{
 		int32 StrLength = Rhs.Len();
 
@@ -1041,7 +1041,7 @@ public:
 	 * @return true if the left string is lexicographically <= the right string, otherwise false
 	 * @note case insensitive
 	 */
-	FORCEINLINE friend bool operator<=(const FString& Lhs, const FString& Rhs)
+	UE_NODISCARD FORCEINLINE friend bool operator<=(const FString& Lhs, const FString& Rhs)
 	{
 		return FPlatformString::Stricmp(*Lhs, *Rhs) <= 0;
 	}
@@ -1055,7 +1055,7 @@ public:
 	 * @note case insensitive
 	 */
 	template <typename CharType>
-	FORCEINLINE friend bool operator<=(const FString& Lhs, const CharType* Rhs)
+	UE_NODISCARD FORCEINLINE friend bool operator<=(const FString& Lhs, const CharType* Rhs)
 	{
 		return FPlatformString::Stricmp(*Lhs, Rhs) <= 0;
 	}
@@ -1069,7 +1069,7 @@ public:
 	 * @note case insensitive
 	 */
 	template <typename CharType>
-	FORCEINLINE friend bool operator<=(const CharType* Lhs, const FString& Rhs)
+	UE_NODISCARD FORCEINLINE friend bool operator<=(const CharType* Lhs, const FString& Rhs)
 	{
 		return FPlatformString::Stricmp(Lhs, *Rhs) <= 0;
 	}
@@ -1082,7 +1082,7 @@ public:
 	 * @return true if the left string is lexicographically < the right string, otherwise false
 	 * @note case insensitive
 	 */
-	FORCEINLINE friend bool operator<(const FString& Lhs, const FString& Rhs)
+	UE_NODISCARD FORCEINLINE friend bool operator<(const FString& Lhs, const FString& Rhs)
 	{
 		return FPlatformString::Stricmp(*Lhs, *Rhs) < 0;
 	}
@@ -1096,7 +1096,7 @@ public:
 	 * @note case insensitive
 	 */
 	template <typename CharType>
-	FORCEINLINE friend bool operator<(const FString& Lhs, const CharType* Rhs)
+	UE_NODISCARD FORCEINLINE friend bool operator<(const FString& Lhs, const CharType* Rhs)
 	{
 		return FPlatformString::Stricmp(*Lhs, Rhs) < 0;
 	}
@@ -1110,7 +1110,7 @@ public:
 	 * @note case insensitive
 	 */
 	template <typename CharType>
-	FORCEINLINE friend bool operator<(const CharType* Lhs, const FString& Rhs)
+	UE_NODISCARD FORCEINLINE friend bool operator<(const CharType* Lhs, const FString& Rhs)
 	{
 		return FPlatformString::Stricmp(Lhs, *Rhs) < 0;
 	}
@@ -1123,7 +1123,7 @@ public:
 	 * @return true if the left string is lexicographically >= the right string, otherwise false
 	 * @note case insensitive
 	 */
-	FORCEINLINE friend bool operator>=(const FString& Lhs, const FString& Rhs)
+	UE_NODISCARD FORCEINLINE friend bool operator>=(const FString& Lhs, const FString& Rhs)
 	{
 		return FPlatformString::Stricmp(*Lhs, *Rhs) >= 0;
 	}
@@ -1137,7 +1137,7 @@ public:
 	 * @note case insensitive
 	 */
 	template <typename CharType>
-	FORCEINLINE friend bool operator>=(const FString& Lhs, const CharType* Rhs)
+	UE_NODISCARD FORCEINLINE friend bool operator>=(const FString& Lhs, const CharType* Rhs)
 	{
 		return FPlatformString::Stricmp(*Lhs, Rhs) >= 0;
 	}
@@ -1151,7 +1151,7 @@ public:
 	 * @note case insensitive
 	 */
 	template <typename CharType>
-	FORCEINLINE friend bool operator>=(const CharType* Lhs, const FString& Rhs)
+	UE_NODISCARD FORCEINLINE friend bool operator>=(const CharType* Lhs, const FString& Rhs)
 	{
 		return FPlatformString::Stricmp(Lhs, *Rhs) >= 0;
 	}
@@ -1164,7 +1164,7 @@ public:
 	 * @return true if the left string is lexicographically > the right string, otherwise false
 	 * @note case insensitive
 	 */
-	FORCEINLINE friend bool operator>(const FString& Lhs, const FString& Rhs)
+	UE_NODISCARD FORCEINLINE friend bool operator>(const FString& Lhs, const FString& Rhs)
 	{
 		return FPlatformString::Stricmp(*Lhs, *Rhs) > 0;
 	}
@@ -1178,7 +1178,7 @@ public:
 	 * @note case insensitive
 	 */
 	template <typename CharType>
-	FORCEINLINE friend bool operator>(const FString& Lhs, const CharType* Rhs)
+	UE_NODISCARD FORCEINLINE friend bool operator>(const FString& Lhs, const CharType* Rhs)
 	{
 		return FPlatformString::Stricmp(*Lhs, Rhs) > 0;
 	}
@@ -1192,7 +1192,7 @@ public:
 	 * @note case insensitive
 	 */
 	template <typename CharType>
-	FORCEINLINE friend bool operator>(const CharType* Lhs, const FString& Rhs)
+	UE_NODISCARD FORCEINLINE friend bool operator>(const CharType* Lhs, const FString& Rhs)
 	{
 		return FPlatformString::Stricmp(Lhs, *Rhs) > 0;
 	}
@@ -1205,7 +1205,7 @@ public:
 	 * @return true if the left string is lexicographically == the right string, otherwise false
 	 * @note case insensitive
 	 */
-	FORCEINLINE friend bool operator==(const FString& Lhs, const FString& Rhs)
+	UE_NODISCARD FORCEINLINE friend bool operator==(const FString& Lhs, const FString& Rhs)
 	{
 		return Lhs.Equals(Rhs, ESearchCase::IgnoreCase);
 	}
@@ -1219,7 +1219,7 @@ public:
 	 * @note case insensitive
 	 */
 	template <typename CharType>
-	FORCEINLINE friend bool operator==(const FString& Lhs, const CharType* Rhs)
+	UE_NODISCARD FORCEINLINE friend bool operator==(const FString& Lhs, const CharType* Rhs)
 	{
 		return FPlatformString::Stricmp(*Lhs, Rhs) == 0;
 	}
@@ -1233,7 +1233,7 @@ public:
 	 * @note case insensitive
 	 */
 	template <typename CharType>
-	FORCEINLINE friend bool operator==(const CharType* Lhs, const FString& Rhs)
+	UE_NODISCARD FORCEINLINE friend bool operator==(const CharType* Lhs, const FString& Rhs)
 	{
 		return FPlatformString::Stricmp(Lhs, *Rhs) == 0;
 	}
@@ -1246,7 +1246,7 @@ public:
 	 * @return true if the left string is lexicographically != the right string, otherwise false
 	 * @note case insensitive
 	 */
-	FORCEINLINE friend bool operator!=(const FString& Lhs, const FString& Rhs)
+	UE_NODISCARD FORCEINLINE friend bool operator!=(const FString& Lhs, const FString& Rhs)
 	{
 		return !(Lhs == Rhs);
 	}
@@ -1260,7 +1260,7 @@ public:
 	 * @note case insensitive
 	 */
 	template <typename CharType>
-	FORCEINLINE friend bool operator!=(const FString& Lhs, const CharType* Rhs)
+	UE_NODISCARD FORCEINLINE friend bool operator!=(const FString& Lhs, const CharType* Rhs)
 	{
 		return FPlatformString::Stricmp(*Lhs, Rhs) != 0;
 	}
@@ -1274,24 +1274,24 @@ public:
 	 * @note case insensitive
 	 */
 	template <typename CharType>
-	FORCEINLINE friend bool operator!=(const CharType* Lhs, const FString& Rhs)
+	UE_NODISCARD FORCEINLINE friend bool operator!=(const CharType* Lhs, const FString& Rhs)
 	{
 		return FPlatformString::Stricmp(Lhs, *Rhs) != 0;
 	}
 
 	/** Get the length of the string, excluding terminating character */
-	FORCEINLINE int32 Len() const
+	UE_NODISCARD FORCEINLINE int32 Len() const
 	{
 		return Data.Num() ? Data.Num() - 1 : 0;
 	}
 
 	/** Returns the left most given number of characters */
-	FORCEINLINE FString Left( int32 Count ) const &
+	UE_NODISCARD FORCEINLINE FString Left( int32 Count ) const &
 	{
 		return FString( FMath::Clamp(Count,0,Len()), **this );
 	}
 
-	FORCEINLINE FString Left(int32 Count) &&
+	UE_NODISCARD FORCEINLINE FString Left(int32 Count) &&
 	{
 		LeftInline(Count, false);
 		return MoveTemp(*this);
@@ -1306,13 +1306,13 @@ public:
 	}
 
 	/** Returns the left most characters from the string chopping the given number of characters from the end */
-	FORCEINLINE FString LeftChop( int32 Count ) const &
+	UE_NODISCARD FORCEINLINE FString LeftChop( int32 Count ) const &
 	{
 		const int32 Length = Len();
 		return FString( FMath::Clamp(Length-Count,0, Length), **this );
 	}
 
-	FORCEINLINE FString LeftChop(int32 Count)&&
+	UE_NODISCARD FORCEINLINE FString LeftChop(int32 Count)&&
 	{
 		LeftChopInline(Count, false);
 		return MoveTemp(*this);
@@ -1326,13 +1326,13 @@ public:
 	}
 
 	/** Returns the string to the right of the specified location, counting back from the right (end of the word). */
-	FORCEINLINE FString Right( int32 Count ) const &
+	UE_NODISCARD FORCEINLINE FString Right( int32 Count ) const &
 	{
 		const int32 Length = Len();
 		return FString( **this + Length-FMath::Clamp(Count,0,Length) );
 	}
 
-	FORCEINLINE FString Right(int32 Count) &&
+	UE_NODISCARD FORCEINLINE FString Right(int32 Count) &&
 	{
 		RightInline(Count, false);
 		return MoveTemp(*this);
@@ -1346,13 +1346,13 @@ public:
 	}
 
 	/** Returns the string to the right of the specified location, counting forward from the left (from the beginning of the word). */
-	FORCEINLINE FString RightChop( int32 Count ) const &
+	UE_NODISCARD FORCEINLINE FString RightChop( int32 Count ) const &
 	{
 		const int32 Length = Len();
 		return FString( **this + Length-FMath::Clamp(Length-Count,0, Length) );
 	}
 
-	FORCEINLINE FString RightChop(int32 Count) &&
+	UE_NODISCARD FORCEINLINE FString RightChop(int32 Count) &&
 	{
 		RightChopInline(Count, false);
 		return MoveTemp(*this);
@@ -1365,7 +1365,7 @@ public:
 	}
 
 	/** Returns the substring from Start position for Count characters. */
-	FORCEINLINE FString Mid( int32 Start, int32 Count=MAX_int32 ) const &
+	UE_NODISCARD FORCEINLINE FString Mid( int32 Start, int32 Count=MAX_int32 ) const &
 	{
 		FString Result;
 		if (Count >= 0)
@@ -1379,7 +1379,7 @@ public:
 		return Result;
 	}
 
-	FORCEINLINE FString Mid(int32 Start, int32 Count = MAX_int32) &&
+	UE_NODISCARD FORCEINLINE FString Mid(int32 Start, int32 Count = MAX_int32) &&
 	{
 		MidInline(Start, Count, false);
 		return MoveTemp(*this);
@@ -1401,7 +1401,7 @@ public:
 	 * @param SearchCase		Indicates whether the search is case sensitive or not
 	 * @param SearchDir			Indicates whether the search starts at the beginning or at the end.
 	 */
-	int32 Find( const TCHAR* SubStr, ESearchCase::Type SearchCase = ESearchCase::IgnoreCase, 
+	UE_NODISCARD int32 Find( const TCHAR* SubStr, ESearchCase::Type SearchCase = ESearchCase::IgnoreCase, 
 				ESearchDir::Type SearchDir = ESearchDir::FromStart, int32 StartPosition=INDEX_NONE ) const;
 
 	/**
@@ -1413,7 +1413,7 @@ public:
 	 * @param SearchCase		Indicates whether the search is case sensitive or not ( defaults to ESearchCase::IgnoreCase )
 	 * @param SearchDir			Indicates whether the search starts at the beginning or at the end ( defaults to ESearchDir::FromStart )
 	 */
-	FORCEINLINE int32 Find( const FString& SubStr, ESearchCase::Type SearchCase = ESearchCase::IgnoreCase, 
+	UE_NODISCARD FORCEINLINE int32 Find( const FString& SubStr, ESearchCase::Type SearchCase = ESearchCase::IgnoreCase, 
 							ESearchDir::Type SearchDir = ESearchDir::FromStart, int32 StartPosition=INDEX_NONE ) const
 	{
 		return Find( *SubStr, SearchCase, SearchDir, StartPosition );
@@ -1427,7 +1427,7 @@ public:
 	 * @param SearchDir			Indicates whether the search starts at the beginning or at the end ( defaults to ESearchDir::FromStart )
 	 * @return					Returns whether the string contains the substring
 	 **/
-	FORCEINLINE bool Contains(const TCHAR* SubStr, ESearchCase::Type SearchCase = ESearchCase::IgnoreCase, 
+	UE_NODISCARD FORCEINLINE bool Contains(const TCHAR* SubStr, ESearchCase::Type SearchCase = ESearchCase::IgnoreCase, 
 							  ESearchDir::Type SearchDir = ESearchDir::FromStart ) const
 	{
 		return Find(SubStr, SearchCase, SearchDir) != INDEX_NONE;
@@ -1441,7 +1441,7 @@ public:
 	 * @param SearchDir			Indicates whether the search starts at the beginning or at the end ( defaults to ESearchDir::FromStart )
 	 * @return					Returns whether the string contains the substring
 	 **/
-	FORCEINLINE bool Contains(const FString& SubStr, ESearchCase::Type SearchCase = ESearchCase::IgnoreCase, 
+	UE_NODISCARD FORCEINLINE bool Contains(const FString& SubStr, ESearchCase::Type SearchCase = ESearchCase::IgnoreCase, 
 							  ESearchDir::Type SearchDir = ESearchDir::FromStart ) const
 	{
 		return Find(*SubStr, SearchCase, SearchDir) != INDEX_NONE;
@@ -1507,7 +1507,7 @@ public:
 	 * @param SearchCase 	Whether or not the comparison should ignore case
 	 * @return true if this string is lexicographically equivalent to the other, otherwise false
 	 */
-	FORCEINLINE bool Equals(const FString& Other, ESearchCase::Type SearchCase = ESearchCase::CaseSensitive) const
+	UE_NODISCARD FORCEINLINE bool Equals(const FString& Other, ESearchCase::Type SearchCase = ESearchCase::CaseSensitive) const
 	{
 		int32 Num = Data.Num();
 		int32 OtherNum = Other.Data.Num();
@@ -1539,7 +1539,7 @@ public:
 	 * @param SearchCase 	Whether or not the comparison should ignore case
 	 * @return 0 if equal, negative if less than, positive if greater than
 	 */
-	FORCEINLINE int32 Compare( const FString& Other, ESearchCase::Type SearchCase = ESearchCase::CaseSensitive ) const
+	UE_NODISCARD FORCEINLINE int32 Compare( const FString& Other, ESearchCase::Type SearchCase = ESearchCase::CaseSensitive ) const
 	{
 		if( SearchCase == ESearchCase::CaseSensitive )
 		{
@@ -1593,37 +1593,37 @@ public:
 	}
 
 	/** Returns a new string with the characters of this converted to uppercase */
-	FString ToUpper() const &;
+	UE_NODISCARD FString ToUpper() const &;
 
 	/**
 	 * Converts all characters in this rvalue string to uppercase and moves it into the returned string.
 	 * @return a new string with the characters of this converted to uppercase
 	 */
-	FString ToUpper() &&;
+	UE_NODISCARD FString ToUpper() &&;
 
 	/** Converts all characters in this string to uppercase */
 	void ToUpperInline();
 
 	/** Returns a new string with the characters of this converted to lowercase */
-	FString ToLower() const &;
+	UE_NODISCARD FString ToLower() const &;
 
 	/**
 	 * Converts all characters in this rvalue string to lowercase and moves it into the returned string.
 	 * @return a new string with the characters of this converted to lowercase
 	 */
-	FString ToLower() &&;
+	UE_NODISCARD FString ToLower() &&;
 
 	/** Converts all characters in this string to lowercase */
 	void ToLowerInline();
 
 	/** Pad the left of this string for ChCount characters */
-	FString LeftPad( int32 ChCount ) const;
+	UE_NODISCARD FString LeftPad( int32 ChCount ) const;
 
 	/** Pad the right of this string for ChCount characters */
-	FString RightPad( int32 ChCount ) const;
+	UE_NODISCARD FString RightPad( int32 ChCount ) const;
 	
 	/** Returns true if the string only contains numeric characters */
-	bool IsNumeric() const;
+	UE_NODISCARD bool IsNumeric() const;
 	
 	/** Removes spaces from the string.  I.E. "Spaces Are Cool" --> "SpacesAreCool". */
 	void RemoveSpacesInline();
@@ -1637,7 +1637,7 @@ public:
 	 * @returns FString object that was constructed using format and additional parameters.
 	 */
 	template <typename FmtType, typename... Types>
-	static FString Printf(const FmtType& Fmt, Types... Args)
+	UE_NODISCARD static FString Printf(const FmtType& Fmt, Types... Args)
 	{
 		static_assert(TIsArrayOrRefOfType<FmtType, TCHAR>::Value, "Formatting string must be a TCHAR array.");
 		static_assert(TAnd<TIsValidVariadicFunctionArg<Types>...>::Value, "Invalid argument(s) passed to FString::Printf");
@@ -1670,7 +1670,7 @@ public:
 	 * @param InNamedArguments		A map of named arguments that match the tokens specified in InExpression
 	 * @return A string containing the formatted text
 	 */
-	static FString Format(const TCHAR* InFormatString, const FStringFormatNamedArguments& InNamedArguments);
+	UE_NODISCARD static FString Format(const TCHAR* InFormatString, const FStringFormatNamedArguments& InNamedArguments);
 
 	/**
 	 * Format the specified string using the specified arguments. Replaces instances of {0} with indices from the given array matching the index specified in the token
@@ -1678,10 +1678,10 @@ public:
 	 * @param InOrderedArguments	An array of ordered arguments that match the tokens specified in InExpression
 	 * @return A string containing the formatted text
 	 */
-	static FString Format(const TCHAR* InFormatString, const FStringFormatOrderedArguments& InOrderedArguments);
+	UE_NODISCARD static FString Format(const TCHAR* InFormatString, const FStringFormatOrderedArguments& InOrderedArguments);
 
 	/** Returns a string containing only the Ch character */
-	static FString Chr( TCHAR Ch );
+	UE_NODISCARD static FString Chr( TCHAR Ch );
 
 	/**
 	 * Returns a string that is full of a variable number of characters
@@ -1691,7 +1691,7 @@ public:
 	 * 
 	 * @return The string of NumCharacters characters.
 	 */
-	static FString ChrN( int32 NumCharacters, TCHAR Char );
+	UE_NODISCARD static FString ChrN( int32 NumCharacters, TCHAR Char );
 
 	/**
 	 * Serializes the string.
@@ -1710,7 +1710,7 @@ public:
 	 * @param SearchCase		Indicates whether the search is case sensitive or not ( defaults to ESearchCase::IgnoreCase )
 	 * @return true if this string begins with specified text, false otherwise
 	 */
-	bool StartsWith(const TCHAR* InSuffix, ESearchCase::Type SearchCase = ESearchCase::IgnoreCase) const;
+	UE_NODISCARD bool StartsWith(const TCHAR* InSuffix, ESearchCase::Type SearchCase = ESearchCase::IgnoreCase) const;
 
 	/**
 	 * Test whether this string starts with given string.
@@ -1718,7 +1718,7 @@ public:
 	 * @param SearchCase		Indicates whether the search is case sensitive or not ( defaults to ESearchCase::IgnoreCase )
 	 * @return true if this string begins with specified text, false otherwise
 	 */
-	bool StartsWith(const FString& InPrefix, ESearchCase::Type SearchCase = ESearchCase::IgnoreCase) const;
+	UE_NODISCARD bool StartsWith(const FString& InPrefix, ESearchCase::Type SearchCase = ESearchCase::IgnoreCase) const;
 
 	/**
 	 * Test whether this string ends with given string.
@@ -1726,7 +1726,7 @@ public:
 	 * @param SearchCase		Indicates whether the search is case sensitive or not ( defaults to ESearchCase::IgnoreCase )
 	 * @return true if this string ends with specified text, false otherwise
 	 */
-	bool EndsWith(const TCHAR* InSuffix, ESearchCase::Type SearchCase = ESearchCase::IgnoreCase) const;
+	UE_NODISCARD bool EndsWith(const TCHAR* InSuffix, ESearchCase::Type SearchCase = ESearchCase::IgnoreCase) const;
 
 	/**
 	 * Test whether this string ends with given string.
@@ -1734,7 +1734,7 @@ public:
 	 * @param SearchCase		Indicates whether the search is case sensitive or not ( defaults to ESearchCase::IgnoreCase )
 	 * @return true if this string ends with specified text, false otherwise
 	 */
-	bool EndsWith(const FString& InSuffix, ESearchCase::Type SearchCase = ESearchCase::IgnoreCase ) const;
+	UE_NODISCARD bool EndsWith(const FString& InSuffix, ESearchCase::Type SearchCase = ESearchCase::IgnoreCase ) const;
 
 	/**
 	 * Searches this string for a given wild card
@@ -1744,7 +1744,7 @@ public:
 	 * @return true if this string matches the *?-type wildcard given. 
 	 * @warning This is a simple, SLOW routine. Use with caution
 	 */
-	bool MatchesWildcard(const TCHAR* Wildcard, ESearchCase::Type SearchCase = ESearchCase::IgnoreCase) const;
+	UE_NODISCARD bool MatchesWildcard(const TCHAR* Wildcard, ESearchCase::Type SearchCase = ESearchCase::IgnoreCase) const;
 
 	/**
 	 * Searches this string for a given wild card
@@ -1754,7 +1754,7 @@ public:
 	 * @return true if this string matches the *?-type wildcard given.
 	 * @warning This is a simple, SLOW routine. Use with caution
 	 */
-	FORCEINLINE bool MatchesWildcard(const FString& Wildcard, ESearchCase::Type SearchCase = ESearchCase::IgnoreCase) const
+	UE_NODISCARD FORCEINLINE bool MatchesWildcard(const FString& Wildcard, ESearchCase::Type SearchCase = ESearchCase::IgnoreCase) const
 	{
 		return MatchesWildcard(*Wildcard, SearchCase);
 	}
@@ -1768,13 +1768,13 @@ public:
 	 * Removes whitespace characters from the start and end of this string.
 	 * @note Unlike Trim() this function returns a copy, and does not mutate the string.
 	 */
-	FString TrimStartAndEnd() const &;
+	UE_NODISCARD FString TrimStartAndEnd() const &;
 
 	/**
 	 * Removes whitespace characters from the start and end of this string.
 	 * @note Unlike Trim() this function returns a copy, and does not mutate the string.
 	 */
-	FString TrimStartAndEnd() &&;
+	UE_NODISCARD FString TrimStartAndEnd() &&;
 
 	/**
 	 * Removes whitespace characters from the start of this string. Modifies the string in-place.
@@ -1785,13 +1785,13 @@ public:
 	 * Removes whitespace characters from the start of this string.
 	 * @note Unlike Trim() this function returns a copy, and does not mutate the string.
 	 */
-	FString TrimStart() const &;
+	UE_NODISCARD FString TrimStart() const &;
 
 	/**
 	 * Removes whitespace characters from the start of this string.
 	 * @note Unlike Trim() this function returns a copy, and does not mutate the string.
 	 */
-	FString TrimStart() &&;
+	UE_NODISCARD FString TrimStart() &&;
 
 	/**
 	 * Removes whitespace characters from the end of this string. Modifies the string in-place.
@@ -1802,13 +1802,13 @@ public:
 	 * Removes whitespace characters from the end of this string.
 	 * @note Unlike TrimTrailing() this function returns a copy, and does not mutate the string.
 	 */
-	FString TrimEnd() const &;
+	UE_NODISCARD FString TrimEnd() const &;
 
 	/**
 	 * Removes whitespace characters from the end of this string.
 	 * @note Unlike TrimTrailing() this function returns a copy, and does not mutate the string.
 	 */
-	FString TrimEnd() &&;
+	UE_NODISCARD FString TrimEnd() &&;
 
 	/** 
 	 * Trims the inner array after the null terminator.
@@ -1824,12 +1824,12 @@ public:
 	/**
 	 * Returns a copy of this string with wrapping quotation marks removed.
 	 */
-	FString TrimQuotes( bool* bQuotesRemoved = nullptr ) const &;
+	UE_NODISCARD FString TrimQuotes( bool* bQuotesRemoved = nullptr ) const &;
 
 	/**
 	 * Returns this string with wrapping quotation marks removed.
 	 */
-	FString TrimQuotes(bool* bQuotesRemoved = nullptr) &&;
+	UE_NODISCARD FString TrimQuotes(bool* bQuotesRemoved = nullptr) &&;
 
 	/**
 	 * Breaks up a delimited string into elements of a string array.
@@ -1888,12 +1888,12 @@ public:
 	/**
 	 * Returns a copy of this string, with the characters in reverse order
 	 */
-	FString Reverse() const &;
+	UE_NODISCARD FString Reverse() const &;
 
 	/**
 	 * Returns this string, with the characters in reverse order
 	 */
-	FString Reverse() &&;
+	UE_NODISCARD FString Reverse() &&;
 
 	/**
 	 * Reverses the order of characters in this string
@@ -1908,7 +1908,7 @@ public:
 	 * @param SearchCase	Indicates whether the search is case sensitive or not ( defaults to ESearchCase::IgnoreCase )
 	 * @return a copy of this string with the replacement made
 	 */
-	FString Replace(const TCHAR* From, const TCHAR* To, ESearchCase::Type SearchCase = ESearchCase::IgnoreCase) const &;
+	UE_NODISCARD FString Replace(const TCHAR* From, const TCHAR* To, ESearchCase::Type SearchCase = ESearchCase::IgnoreCase) const &;
 
 	/**
 	 * Replace all occurrences of a substring in this string
@@ -1918,7 +1918,7 @@ public:
 	 * @param SearchCase	Indicates whether the search is case sensitive or not ( defaults to ESearchCase::IgnoreCase )
 	 * @return a copy of this string with the replacement made
 	 */
-	FString Replace(const TCHAR* From, const TCHAR* To, ESearchCase::Type SearchCase = ESearchCase::IgnoreCase) &&;
+	UE_NODISCARD FString Replace(const TCHAR* From, const TCHAR* To, ESearchCase::Type SearchCase = ESearchCase::IgnoreCase) &&;
 
 	/**
 	 * Replace all occurrences of SearchText with ReplacementText in this string.
@@ -1960,7 +1960,7 @@ public:
 	/**
 	 * Returns a copy of this string with all quote marks escaped (unless the quote is already escaped)
 	 */
-	FString ReplaceQuotesWithEscapedQuotes() const &
+	UE_NODISCARD FString ReplaceQuotesWithEscapedQuotes() const &
 	{
 		FString Result(*this);
 		return MoveTemp(Result).ReplaceQuotesWithEscapedQuotes();
@@ -1969,7 +1969,7 @@ public:
 	/**
 	 * Returns a copy of this string with all quote marks escaped (unless the quote is already escaped)
 	 */
-	FString ReplaceQuotesWithEscapedQuotes() &&;
+	UE_NODISCARD FString ReplaceQuotesWithEscapedQuotes() &&;
 
 	/**
 	 * Replaces certain characters with the "escaped" version of that character (i.e. replaces "\n" with "\\n").
@@ -1979,7 +1979,7 @@ public:
 	 *
 	 * @return	a string with all control characters replaced by the escaped version.
 	 */
-	FString ReplaceCharWithEscapedChar( const TArray<TCHAR>* Chars = nullptr ) const &
+	UE_NODISCARD FString ReplaceCharWithEscapedChar( const TArray<TCHAR>* Chars = nullptr ) const &
 	{
 		FString Result(*this);
 		return MoveTemp(Result).ReplaceCharWithEscapedChar(Chars);
@@ -1993,14 +1993,14 @@ public:
 	 *
 	 * @return	a string with all control characters replaced by the escaped version.
 	 */
-	FString ReplaceCharWithEscapedChar( const TArray<TCHAR>* Chars = nullptr ) &&;
+	UE_NODISCARD FString ReplaceCharWithEscapedChar( const TArray<TCHAR>* Chars = nullptr ) &&;
 
 	/**
 	 * Removes the escape backslash for all supported characters, replacing the escape and character with the non-escaped version.  (i.e.
 	 * replaces "\\n" with "\n".  Counterpart to ReplaceCharWithEscapedChar().
 	 * @return copy of this string with replacement made
 	 */
-	FString ReplaceEscapedCharWithChar( const TArray<TCHAR>* Chars = nullptr ) const &
+	UE_NODISCARD FString ReplaceEscapedCharWithChar( const TArray<TCHAR>* Chars = nullptr ) const &
 	{
 		FString Result(*this);
 		return MoveTemp(Result).ReplaceEscapedCharWithChar(Chars);
@@ -2011,7 +2011,7 @@ public:
 	 * replaces "\\n" with "\n".  Counterpart to ReplaceCharWithEscapedChar().
 	 * @return copy of this string with replacement made
 	 */
-	FString ReplaceEscapedCharWithChar( const TArray<TCHAR>* Chars = nullptr ) &&;
+	UE_NODISCARD FString ReplaceEscapedCharWithChar( const TArray<TCHAR>* Chars = nullptr ) &&;
 
 	/**
 	 * Replaces all instances of '\t' with TabWidth number of spaces
@@ -2024,7 +2024,7 @@ public:
 	 * @param InSpacesPerTab - Number of spaces that a tab represents
 	 * @return copy of this string with replacement made
 	 */
-	FString ConvertTabsToSpaces(const int32 InSpacesPerTab) const &
+	UE_NODISCARD FString ConvertTabsToSpaces(const int32 InSpacesPerTab) const &
 	{
 		FString FinalString(*this);
 		FinalString.ConvertTabsToSpacesInline(InSpacesPerTab);
@@ -2036,14 +2036,14 @@ public:
 	 * @param InSpacesPerTab - Number of spaces that a tab represents
 	 * @return copy of this string with replacement made
 	 */
-	FString ConvertTabsToSpaces(const int32 InSpacesPerTab) &&
+	UE_NODISCARD FString ConvertTabsToSpaces(const int32 InSpacesPerTab) &&
 	{
 		ConvertTabsToSpacesInline(InSpacesPerTab);
 		return MoveTemp(*this);
 	}
 
 	// Takes the number passed in and formats the string in comma format ( 12345 becomes "12,345")
-	static FString FormatAsNumber( int32 InNumber );
+	UE_NODISCARD static FString FormatAsNumber( int32 InNumber );
 
 	// To allow more efficient memory handling, automatically adds one for the string termination.
 	FORCEINLINE void Reserve(int32 CharacterCount)
@@ -2066,7 +2066,7 @@ public:
 
 
 	/** Converts an integer to a string. */
-	static FORCEINLINE FString FromInt( int32 Num )
+	UE_NODISCARD static FORCEINLINE FString FromInt( int32 Num )
 	{
 		FString Ret; 
 		Ret.AppendInt(Num); 
@@ -2083,7 +2083,7 @@ public:
 	 *
 	 * @return The boolean value
 	 */
-	bool ToBool() const;
+	UE_NODISCARD bool ToBool() const;
 
 	/**
 	 * Converts a buffer to a string
@@ -2093,7 +2093,7 @@ public:
 	 *
 	 * @return the blob in string form
 	 */
-	static FString FromBlob(const uint8* SrcBuffer,const uint32 SrcSize);
+	UE_NODISCARD static FString FromBlob(const uint8* SrcBuffer,const uint32 SrcSize);
 
 	/**
 	 * Converts a string into a buffer
@@ -2113,7 +2113,7 @@ public:
 	 *
 	 * @return the blob in string form
 	 */
-	static FString FromHexBlob(const uint8* SrcBuffer,const uint32 SrcSize);
+	UE_NODISCARD static FString FromHexBlob(const uint8* SrcBuffer,const uint32 SrcSize);
 
 	/**
 	 * Converts a string into a buffer
@@ -2134,7 +2134,7 @@ public:
 	 *
 	 * @return sanitized string version of float
 	 */
-	static FString SanitizeFloat( double InFloat, const int32 InMinFractionalDigits = 1 );
+	UE_NODISCARD static FString SanitizeFloat( double InFloat, const int32 InMinFractionalDigits = 1 );
 
 	/**
 	 * Joins a range of 'something that can be concatentated to strings with +=' together into a single string with separators.
@@ -2145,7 +2145,7 @@ public:
 	 * @return	The final, joined, separated string.
 	 */
 	template <typename RangeType>
-	static FString Join(const RangeType& Range, const TCHAR* Separator)
+	UE_NODISCARD static FString Join(const RangeType& Range, const TCHAR* Separator)
 	{
 		FString Result;
 		bool    First = true;
@@ -2176,7 +2176,7 @@ public:
 	 * @return	The final, joined, separated string.
 	 */
 	template <typename RangeType, typename ProjectionType>
-	static FString JoinBy(const RangeType& Range, const TCHAR* Separator, ProjectionType Proj)
+	UE_NODISCARD static FString JoinBy(const RangeType& Range, const TCHAR* Separator, ProjectionType Proj)
 	{
 		FString Result;
 		bool    First = true;
@@ -2240,7 +2240,7 @@ FORCEINLINE uint32 GetTypeHash(const FString& S)
  * @param Count number of bytes to convert
  * @return Valid string representing bytes.
  */
-inline FString BytesToString(const uint8* In, int32 Count)
+UE_NODISCARD inline FString BytesToString(const uint8* In, int32 Count)
 {
 	FString Result;
 	Result.Empty(Count);
@@ -2281,7 +2281,7 @@ inline int32 StringToBytes( const FString& String, uint8* OutBytes, int32 MaxBuf
 }
 
 /** Returns Char value of Nibble */
-inline TCHAR NibbleToTChar(uint8 Num)
+UE_NODISCARD inline TCHAR NibbleToTChar(uint8 Num)
 {
 	if (Num > 9)
 	{
@@ -2307,7 +2307,7 @@ inline void ByteToHex(uint8 In, FString& Result)
  * @param Count number of bytes to convert
  * @return Hex value in string.
  */
-inline FString BytesToHex(const uint8* In, int32 Count)
+UE_NODISCARD inline FString BytesToHex(const uint8* In, int32 Count)
 {
 	FString Result;
 	Result.Empty(Count * 2);
@@ -2325,7 +2325,7 @@ inline FString BytesToHex(const uint8* In, int32 Count)
  * @param Char		The character
  * @return	True if in 0-9 and A-F ranges
  */
-inline const bool CheckTCharIsHex( const TCHAR Char )
+UE_NODISCARD inline const bool CheckTCharIsHex( const TCHAR Char )
 {
 	return ( Char >= TEXT('0') && Char <= TEXT('9') ) || ( Char >= TEXT('A') && Char <= TEXT('F') ) || ( Char >= TEXT('a') && Char <= TEXT('f') );
 }
@@ -2335,7 +2335,7 @@ inline const bool CheckTCharIsHex( const TCHAR Char )
  * @param Hex		The character
  * @return	The uint8 value of a hex character
  */
-inline const uint8 TCharToNibble(const TCHAR Hex)
+UE_NODISCARD inline const uint8 TCharToNibble(const TCHAR Hex)
 {
 	if (Hex >= '0' && Hex <= '9')
 	{
@@ -2395,7 +2395,7 @@ inline void LexFromString(FString& OutValue, 	const TCHAR* Buffer)	{	OutValue =
 
  /** Convert numeric types to a string */
 template<typename T>
-typename TEnableIf<TIsArithmetic<T>::Value, FString>::Type
+UE_NODISCARD typename TEnableIf<TIsArithmetic<T>::Value, FString>::Type
 LexToString(const T& Value)
 {
 	// TRemoveCV to remove potential volatile decorations. Removing const is pointless, but harmless because it's specified in the param declaration.
@@ -2403,42 +2403,42 @@ LexToString(const T& Value)
 }
 
 template<typename CharType>
-typename TEnableIf<TIsCharType<CharType>::Value, FString>::Type
+UE_NODISCARD typename TEnableIf<TIsCharType<CharType>::Value, FString>::Type
 LexToString(const CharType* Ptr)
 {
 	return FString(Ptr);
 }
 
-inline FString LexToString(bool Value)
+UE_NODISCARD inline FString LexToString(bool Value)
 {
 	return Value ? TEXT("true") : TEXT("false");
 }
 
-FORCEINLINE FString LexToString(FString&& Str)
+UE_NODISCARD FORCEINLINE FString LexToString(FString&& Str)
 {
 	return MoveTemp(Str);
 }
 
-FORCEINLINE FString LexToString(const FString& Str)
+UE_NODISCARD FORCEINLINE FString LexToString(const FString& Str)
 {
 	return Str;
 }
 
 /** Helper template to convert to sanitized strings */
 template<typename T>
-FString LexToSanitizedString(const T& Value)
+UE_NODISCARD FString LexToSanitizedString(const T& Value)
 {
 	return LexToString(Value);
 }
 
 /** Overloaded for floats */
-inline FString LexToSanitizedString(float Value)
+UE_NODISCARD inline FString LexToSanitizedString(float Value)
 {
 	return FString::SanitizeFloat(Value);
 }
 
 /** Overloaded for doubles */
-inline FString LexToSanitizedString(double Value)
+UE_NODISCARD inline FString LexToSanitizedString(double Value)
 {
 	return FString::SanitizeFloat(Value);
 }
@@ -2493,8 +2493,8 @@ static bool LexTryParseString(bool& OutValue, const TCHAR* Buffer)
 template<typename T>
 struct TTypeToString
 {
-	static FString ToString(const T& Value)				{ return LexToString(Value); }
-	static FString ToSanitizedString(const T& Value)	{ return LexToSanitizedString(Value); }
+	UE_NODISCARD static FString ToString(const T& Value)				{ return LexToString(Value); }
+	UE_NODISCARD static FString ToSanitizedString(const T& Value)	{ return LexToSanitizedString(Value); }
 };
 template<typename T>
 struct TTypeFromString
diff --git a/Engine/Source/Runtime/Core/Public/Memory/MemoryView.h b/Engine/Source/Runtime/Core/Public/Memory/MemoryView.h
--- a/Engine/Source/Runtime/Core/Public/Memory/MemoryView.h
+++ b/Engine/Source/Runtime/Core/Public/Memory/MemoryView.h
@@ -62,22 +62,22 @@ public:
 	}
 
 	/** Returns a pointer to the start of the view. */
-	constexpr inline DataType* GetData() const { return Data; }
+	UE_NODISCARD constexpr inline DataType* GetData() const { return Data; }
 
 	/** Returns a pointer to the end of the view. */
-	inline DataType* GetDataEnd() const { return GetDataAtOffsetNoCheck(Size); }
+	UE_NODISCARD inline DataType* GetDataEnd() const { return GetDataAtOffsetNoCheck(Size); }
 
 	/** Returns the number of bytes in the view. */
-	constexpr inline uint64 GetSize() const { return Size; }
+	UE_NODISCARD constexpr inline uint64 GetSize() const { return Size; }
 
 	/** Returns whether the view has a size of 0 regardless of its data pointer. */
-	constexpr inline bool IsEmpty() const { return Size == 0; }
+	UE_NODISCARD constexpr inline bool IsEmpty() const { return Size == 0; }
 
 	/** Resets to an empty view. */
 	constexpr inline void Reset() { *this = TMemoryView(); }
 
 	/** Returns the left-most part of the view by taking the given number of bytes from the left. */
-	constexpr inline TMemoryView Left(uint64 InSize) const
+	UE_NODISCARD constexpr inline TMemoryView Left(uint64 InSize) const
 	{
 		TMemoryView View(*this);
 		View.LeftInline(InSize);
@@ -85,7 +85,7 @@ public:
 	}
 
 	/** Returns the left-most part of the view by chopping the given number of bytes from the right. */
-	constexpr inline TMemoryView LeftChop(uint64 InSize) const
+	UE_NODISCARD constexpr inline TMemoryView LeftChop(uint64 InSize) const
 	{
 		TMemoryView View(*this);
 		View.LeftChopInline(InSize);
@@ -93,7 +93,7 @@ public:
 	}
 
 	/** Returns the right-most part of the view by taking the given number of bytes from the right. */
-	inline TMemoryView Right(uint64 InSize) const
+	UE_NODISCARD inline TMemoryView Right(uint64 InSize) const
 	{
 		TMemoryView View(*this);
 		View.RightInline(InSize);
@@ -101,7 +101,7 @@ public:
 	}
 
 	/** Returns the right-most part of the view by chopping the given number of bytes from the left. */
-	inline TMemoryView RightChop(uint64 InSize) const
+	UE_NODISCARD inline TMemoryView RightChop(uint64 InSize) const
 	{
 		TMemoryView View(*this);
 		View.RightChopInline(InSize);
@@ -109,7 +109,7 @@ public:
 	}
 
 	/** Returns the middle part of the view by taking up to the given number of bytes from the given position. */
-	inline TMemoryView Mid(uint64 InOffset, uint64 InSize = TNumericLimits<uint64>::Max()) const
+	UE_NODISCARD inline TMemoryView Mid(uint64 InOffset, uint64 InSize = TNumericLimits<uint64>::Max()) const
 	{
 		TMemoryView View(*this);
 		View.MidInline(InOffset, InSize);
@@ -154,21 +154,21 @@ public:
 
 	/** Returns whether this view fully contains the other view. */
 	template <typename OtherDataType>
-	inline bool Contains(const TMemoryView<OtherDataType>& InView) const
+	UE_NODISCARD inline bool Contains(const TMemoryView<OtherDataType>& InView) const
 	{
 		return Data <= InView.Data && GetDataAtOffsetNoCheck(Size) >= InView.GetDataAtOffsetNoCheck(InView.Size);
 	}
 
 	/** Returns whether this view intersects the other view. */
 	template <typename OtherDataType>
-	inline bool Intersects(const TMemoryView<OtherDataType>& InView) const
+	UE_NODISCARD inline bool Intersects(const TMemoryView<OtherDataType>& InView) const
 	{
 		return Data < InView.GetDataAtOffsetNoCheck(InView.Size) && InView.Data < GetDataAtOffsetNoCheck(Size);
 	}
 
 	/** Returns whether the bytes of this view are equal or less/greater than the bytes of the other view. */
 	template <typename OtherDataType>
-	inline int32 CompareBytes(const TMemoryView<OtherDataType>& InView) const
+	UE_NODISCARD inline int32 CompareBytes(const TMemoryView<OtherDataType>& InView) const
 	{
 		const int32 Compare = Data == InView.Data ? 0 : FMemory::Memcmp(Data, InView.Data, FMath::Min(Size, InView.Size));
 		return Compare || Size == InView.Size ? Compare : Size < InView.Size ? -1 : 1;
@@ -176,28 +176,28 @@ public:
 
 	/** Returns whether the bytes of this views are equal to the bytes of the other view. */
 	template <typename OtherDataType>
-	inline bool EqualBytes(const TMemoryView<OtherDataType>& InView) const
+	UE_NODISCARD inline bool EqualBytes(const TMemoryView<OtherDataType>& InView) const
 	{
 		return Size == InView.Size && (Data == InView.Data || FMemory::Memcmp(Data, InView.Data, Size) == 0);
 	}
 
 	/** Returns whether the data pointers and sizes of this view and the other view are equal. */
 	template <typename OtherDataType>
-	constexpr inline bool Equals(const TMemoryView<OtherDataType>& InView) const
+	UE_NODISCARD constexpr inline bool Equals(const TMemoryView<OtherDataType>& InView) const
 	{
 		return Size == InView.Size && (Size == 0 || Data == InView.Data);
 	}
 
 	/** Returns whether the data pointers and sizes of this view and the other view are equal. */
 	template <typename OtherDataType>
-	constexpr inline bool operator==(const TMemoryView<OtherDataType>& InView) const
+	UE_NODISCARD constexpr inline bool operator==(const TMemoryView<OtherDataType>& InView) const
 	{
 		return Equals(InView);
 	}
 
 	/** Returns whether the data pointers and sizes of this view and the other view are not equal. */
 	template <typename OtherDataType>
-	constexpr inline bool operator!=(const TMemoryView<OtherDataType>& InView) const
+	UE_NODISCARD constexpr inline bool operator!=(const TMemoryView<OtherDataType>& InView) const
 	{
 		return !Equals(InView);
 	}
@@ -211,7 +211,7 @@ public:
 
 private:
 	/** Returns the data pointer advanced by an offset in bytes. */
-	inline DataType* GetDataAtOffsetNoCheck(uint64 InOffset) const
+	UE_NODISCARD inline DataType* GetDataAtOffsetNoCheck(uint64 InOffset) const
 	{
 		return reinterpret_cast<ByteType*>(Data) + InOffset;
 	}
@@ -226,26 +226,26 @@ private:
 
 /** Advances the start of the view by an offset, which is clamped to stay within the view. */
 template <typename DataType>
-constexpr inline TMemoryView<DataType> operator+(const TMemoryView<DataType>& View, uint64 Offset)
+UE_NODISCARD constexpr inline TMemoryView<DataType> operator+(const TMemoryView<DataType>& View, uint64 Offset)
 {
 	return TMemoryView<DataType>(View) += Offset;
 }
 
 /** Advances the start of the view by an offset, which is clamped to stay within the view. */
 template <typename DataType>
-constexpr inline TMemoryView<DataType> operator+(uint64 Offset, const TMemoryView<DataType>& View)
+UE_NODISCARD constexpr inline TMemoryView<DataType> operator+(uint64 Offset, const TMemoryView<DataType>& View)
 {
 	return TMemoryView<DataType>(View) += Offset;
 }
 
 /** Make a non-owning mutable view of Size bytes starting at Data. */
-constexpr inline TMemoryView<void> MakeMemoryView(void* Data, uint64 Size)
+UE_NODISCARD constexpr inline TMemoryView<void> MakeMemoryView(void* Data, uint64 Size)
 {
 	return TMemoryView<void>(Data, Size);
 }
 
 /** Make a non-owning const view of Size bytes starting at Data. */
-constexpr inline TMemoryView<const void> MakeMemoryView(const void* Data, uint64 Size)
+UE_NODISCARD constexpr inline TMemoryView<const void> MakeMemoryView(const void* Data, uint64 Size)
 {
 	return TMemoryView<const void>(Data, Size);
 }
@@ -253,7 +253,7 @@ constexpr inline TMemoryView<const void> MakeMemoryView(const void* Data, uint64
 /** Make a non-owning mutable view starting at Data and ending at DataEnd. */
 template <typename DataEndType,
 	decltype(ImplicitConv<void*>(DeclVal<DataEndType*>()))* = nullptr>
-inline TMemoryView<void> MakeMemoryView(void* Data, DataEndType* DataEnd)
+UE_NODISCARD inline TMemoryView<void> MakeMemoryView(void* Data, DataEndType* DataEnd)
 {
 	return TMemoryView<void>(Data, DataEnd);
 }
@@ -261,7 +261,7 @@ inline TMemoryView<void> MakeMemoryView(void* Data, DataEndType* DataEnd)
 /** Make a non-owning const view starting at Data and ending at DataEnd. */
 template <typename DataEndType,
 	decltype(ImplicitConv<const void*>(DeclVal<DataEndType*>()))* = nullptr>
-inline TMemoryView<const void> MakeMemoryView(const void* Data, DataEndType* DataEnd)
+UE_NODISCARD inline TMemoryView<const void> MakeMemoryView(const void* Data, DataEndType* DataEnd)
 {
 	return TMemoryView<const void>(Data, DataEnd);
 }
@@ -272,7 +272,7 @@ inline TMemoryView<const void> MakeMemoryView(const void* Data, DataEndType* Dat
  * This overload is only available when the element type does not need to be deduced.
  */
 template <typename T>
-constexpr inline TMemoryView<const void> MakeMemoryView(std::initializer_list<typename TIdentity<T>::Type> List)
+UE_NODISCARD constexpr inline TMemoryView<const void> MakeMemoryView(std::initializer_list<typename TIdentity<T>::Type> List)
 {
 	return TMemoryView<const void>(GetData(List), GetNum(List) * sizeof(T));
 }
@@ -280,7 +280,7 @@ constexpr inline TMemoryView<const void> MakeMemoryView(std::initializer_list<ty
 /** Make a non-owning view of the memory of the contiguous container. */
 template <typename ContainerType,
 	typename TEnableIf<TIsContiguousContainer<ContainerType>::Value>::Type* = nullptr>
-constexpr inline auto MakeMemoryView(ContainerType&& Container)
+UE_NODISCARD constexpr inline auto MakeMemoryView(ContainerType&& Container)
 {
 	using ElementType = typename TRemovePointer<decltype(GetData(DeclVal<ContainerType>()))>::Type;
 	constexpr bool bIsConst = TIsConst<ElementType>::Value;
