Inform the user if trace.start was not possible. There exist some plugins that create their own trace store and sequester the runtime trace. As tracing is process wide this ties up the trace system such that other use-cases like trace.start can't be satisfied. When this scenario is encountered will inform the user and invite them to run trace.stop first.

diff --git a/Engine/Source/Runtime/Core/Private/ProfilingDebugging/TraceAuxiliary.cpp b/Engine/Source/Runtime/Core/Private/ProfilingDebugging/TraceAuxiliary.cpp
--- a/Engine/Source/Runtime/Core/Private/ProfilingDebugging/TraceAuxiliary.cpp
+++ b/Engine/Source/Runtime/Core/Private/ProfilingDebugging/TraceAuxiliary.cpp
@@ -185,6 +185,7 @@ bool FTraceAuxiliaryImpl::Stop()
 
 	DisableChannels();
 	State = EState::Stopped;
+	TraceDest.Reset();
 	return true;
 }
 
@@ -338,6 +339,18 @@ static void TraceAuxiliaryStart(const TArray<FString>& Args)
 		return;
 	}
 
+	// It is possible that something outside of TraceAux's world view has called
+	// UE::Trace::SendTo/WriteTo(). A plugin that has created its own store for
+	// example. There's not really much that can be done about that here (tracing
+	// is singular within a process. We can at least detect the obvious case and
+	// inform the user.
+	const TCHAR* TraceDest = GTraceAuxiliary.GetDest();
+	if (TraceDest[0] == '\0')
+	{
+		UE_LOG(LogConsoleResponse, Warning, TEXT("Trace system already in use by a plugin or -trace*=... argument. Use 'Trace.Stop' first."));
+		return;
+	}
+
 	// Give the user some feedback that everything's underway.
 	FString Channels;
 	GTraceAuxiliary.ReadChannels([&Channels] (const TCHAR* Channel)
@@ -349,7 +362,7 @@ static void TraceAuxiliaryStart(const TArray<FString>& Args)
 
 		Channels += Channel;
 	});
-	UE_LOG(LogConsoleResponse, Log, TEXT("Tracing to; %s"), GTraceAuxiliary.GetDest());
+	UE_LOG(LogConsoleResponse, Log, TEXT("Tracing to; %s"), TraceDest);
 	UE_LOG(LogConsoleResponse, Log, TEXT("Trace channels; %s"), *Channels);
 }
 
