[Trace] Removed use of TMiniArray.

#rnx

diff --git a/Engine/Source/Runtime/Core/Private/ProfilingDebugging/Windows/CallstackTrace.cpp b/Engine/Source/Runtime/Core/Private/ProfilingDebugging/Windows/CallstackTrace.cpp
--- a/Engine/Source/Runtime/Core/Private/ProfilingDebugging/Windows/CallstackTrace.cpp
+++ b/Engine/Source/Runtime/Core/Private/ProfilingDebugging/Windows/CallstackTrace.cpp
@@ -6,7 +6,6 @@
 #include "CoreTypes.h"
 #include "HAL/CriticalSection.h"
 #include "HAL/MemoryBase.h"
-#include "MemoryTrace.inl"
 #include "Misc/ScopeRWLock.h"
 #include "Containers/CircularQueue.h"
 #include "HAL/RunnableThread.h"
@@ -187,7 +186,9 @@ private:
 	const FFunction*		LookupFunction(UPTRINT Address, FLookupState& State) const;
 	static FBacktracer*		Instance;
 	mutable FRWLock			Lock;
-	TMiniArray<FModule>		Modules;
+	FModule*				Modules;
+	int32					ModulesNum;
+	int32					ModulesCapacity;
 	FMalloc*				Malloc;
 	FCallstackProcWorker*	ProcessingThreadRunnable;
 #if BACKTRACE_DBGLVL >= 1
@@ -201,10 +202,12 @@ FBacktracer* FBacktracer::Instance = nullptr;
 
 ////////////////////////////////////////////////////////////////////////////////
 FBacktracer::FBacktracer(FMalloc* InMalloc)
-: Modules(InMalloc)
-, Malloc(InMalloc)
+	: Malloc(InMalloc)
 {
-	Modules.MakeRoom();
+	ModulesCapacity = 8;
+	ModulesNum = 0;
+	Modules = (FModule*)Malloc->Malloc(sizeof(FModule) * ModulesCapacity);
+
 	Instance = this;
 	ProcessingThreadRunnable = (FCallstackProcWorker*) Malloc->Malloc(sizeof(FCallstackProcWorker));
 	ProcessingThreadRunnable = new(ProcessingThreadRunnable) FCallstackProcWorker();
@@ -218,7 +221,8 @@ FBacktracer::FBacktracer(FMalloc* InMalloc)
 ////////////////////////////////////////////////////////////////////////////////
 FBacktracer::~FBacktracer()
 {
-	for (FModule& Module : Modules)
+	TArrayView<FModule> ModulesView(Modules, ModulesNum);
+	for (FModule& Module : ModulesView)
 	{
 		Malloc->Free(Module.Functions);
 	}
@@ -444,11 +448,19 @@ void FBacktracer::AddModule(UPTRINT ModuleBase, const TCHAR* Name)
 		OutTable,
 	};
 
-	TArrayView<FModule> ModulesView(Modules.Data, Modules.Num);
-	int32 Index = Algo::UpperBound(ModulesView, Module.Id, FIdPredicate());
 	{
 		FWriteScopeLock _(Lock);
-		Modules.Insert(Module, Index);
+
+		TArrayView<FModule> ModulesView(Modules, ModulesNum);
+		int32 Index = Algo::UpperBound(ModulesView, Module.Id, FIdPredicate());
+
+		if (ModulesNum + 1 > ModulesCapacity)
+		{
+			ModulesCapacity += 8;
+			Modules = (FModule*) Malloc->Realloc(Modules, sizeof(FModule)* ModulesCapacity);
+		}
+		Modules[ModulesNum++] = Module;
+		Algo::Sort(TArrayView<FModule>(Modules, ModulesNum), [](const FModule& A, const FModule& B) { return A.Id < B.Id; });
 	}
 
 #if BACKTRACE_DBGLVL >= 1
@@ -471,15 +483,17 @@ void FBacktracer::RemoveModule(UPTRINT ModuleBase)
 	return;
 #endif
 
+	FWriteScopeLock _(Lock);
+
 	uint32 ModuleId = AddressToId(ModuleBase);
-	TArrayView<FModule> ModulesView(Modules.Data, Modules.Num);
+	TArrayView<FModule> ModulesView(Modules, ModulesNum);
 	int32 Index = Algo::LowerBound(ModulesView, ModuleId, FIdPredicate());
-	if (Index >= Modules.Num)
+	if (Index >= ModulesNum)
 	{
 		return;
 	}
 
-	const FModule& Module = Modules.Data[Index];
+	const FModule& Module = Modules[Index];
 	if (Module.Id != ModuleId)
 	{
 		return;
@@ -494,10 +508,12 @@ void FBacktracer::RemoveModule(UPTRINT ModuleBase)
 	// table knowing know one is looking at it.
 	Malloc->Free(Module.Functions);
 
+	for (SIZE_T i = Index; i < ModulesNum; i++)
 	{
-		FWriteScopeLock _(Lock);
-		Modules.RemoveAt(Index);
+		Modules[i] = Modules[i + 1];
 	}
+
+	--ModulesNum;
 }
 
 ////////////////////////////////////////////////////////////////////////////////
@@ -512,14 +528,14 @@ const FBacktracer::FFunction* FBacktracer::LookupFunction(UPTRINT Address, FLook
 	uint32 AddressId = AddressToId(Address);
 	if ((AddressId - State.Module.Id) >= State.Module.IdSize)
 	{
-		TArrayView<FModule> ModulesView(Modules.Data, Modules.Num);
+		TArrayView<FModule> ModulesView(Modules, ModulesNum);
 		uint32 Index = Algo::UpperBound(ModulesView, AddressId, IdPredicate);
 		if (Index == 0)
 		{
 			return nullptr;
 		}
 
-		State.Module = Modules.Data[Index - 1];
+		State.Module = Modules[Index - 1];
 	}
 
 	// Check that the address is within the address space of the best-found module
@@ -723,15 +739,17 @@ void Backtracer_Create(FMalloc* Malloc)
 		return;
 	}
 
-	static FBacktracer Instance(Malloc);
+	// Allocate, construct and intentionally leak backtracer
+	void* Alloc = Malloc->Malloc(sizeof(FBacktracer));
+	new (Alloc) FBacktracer(Malloc);
 
 	Modules_Create(Malloc);
 	Modules_Subscribe(
 		[] (bool bLoad, void* Module, const TCHAR* Name)
 		{
 			bLoad
-				? Instance.AddModule(UPTRINT(Module), Name)
-				: Instance.RemoveModule(UPTRINT(Module));
+				? FBacktracer::Get()->AddModule(UPTRINT(Module), Name)
+				: FBacktracer::Get()->RemoveModule(UPTRINT(Module));
 		}
 	);
 }
diff --git a/Engine/Source/Runtime/Core/Private/ProfilingDebugging/Windows/MemoryTrace.inl b/Engine/Source/Runtime/Core/Private/ProfilingDebugging/Windows/MemoryTrace.inl
deleted file mode 100644
diff --git a/Engine/Source/Runtime/Core/Private/ProfilingDebugging/Windows/ModuleDiagnostics.cpp b/Engine/Source/Runtime/Core/Private/ProfilingDebugging/Windows/ModuleDiagnostics.cpp
--- a/Engine/Source/Runtime/Core/Private/ProfilingDebugging/Windows/ModuleDiagnostics.cpp
+++ b/Engine/Source/Runtime/Core/Private/ProfilingDebugging/Windows/ModuleDiagnostics.cpp
@@ -1,7 +1,6 @@
 // Copyright Epic Games, Inc. All Rights Reserved.
 
 #include "CoreTypes.h"
-#include "MemoryTrace.inl"
 #include "Trace/Trace.inl"
 
 #include "Windows/AllowWindowsPlatformTypes.h"
@@ -53,6 +52,7 @@ class FModuleTrace
 {
 public:
 	typedef void				(*SubscribeFunc)(bool, void*, const TCHAR*);
+	typedef TArray<SubscribeFunc, TFixedAllocator<64>> SubscriberArray;
 
 								FModuleTrace(FMalloc* InMalloc);
 								~FModuleTrace();
@@ -65,7 +65,7 @@ private:
 	void						OnDllUnloaded(UPTRINT Base);
 	void						OnDllNotification(unsigned int Reason, const void* DataPtr);
 	static FModuleTrace*		Instance;
-	TMiniArray<SubscribeFunc>	Subscribers;
+	SubscriberArray				Subscribers;
 	void*						CallbackCookie;
 };
 
@@ -74,7 +74,6 @@ FModuleTrace* FModuleTrace::Instance = nullptr;
 
 ////////////////////////////////////////////////////////////////////////////////
 FModuleTrace::FModuleTrace(FMalloc* InMalloc)
-: Subscribers(InMalloc)
 {
 	Instance = this;
 }
