PR #7696: Fix typo of Tuesday in FDateTime (Contributed by [REDACTED_PERSON])

diff --git a/Engine/Source/Runtime/Core/Private/Misc/DateTime.cpp b/Engine/Source/Runtime/Core/Private/Misc/DateTime.cpp
--- a/Engine/Source/Runtime/Core/Private/Misc/DateTime.cpp
+++ b/Engine/Source/Runtime/Core/Private/Misc/DateTime.cpp
@@ -454,7 +454,7 @@ bool FDateTime::ParseHttpDate(const FString& HttpDate, FDateTime& OutDateTime)
 			{
 				return 1;
 			}
-			else if (WeekDay.Equals(TEXT("Tueday")))
+			else if (WeekDay.Equals(TEXT("Tuesday")))
 			{
 				return 2;
 			}
