U5 - ported the UnitTests from the TaskStream to Main

diff --git a/Engine/Source/Runtime/Core/Private/Async/Fundamental/UnitTests.cpp b/Engine/Source/Runtime/Core/Private/Async/Fundamental/UnitTests.cpp
new file mode 100644
--- /dev/null
+++ b/Engine/Source/Runtime/Core/Private/Async/Fundamental/UnitTests.cpp
@@ -0,0 +1,572 @@
+// Copyright Epic Games, Inc. All Rights Reserved.
+
+#include "CoreTypes.h"
+#include "Misc/AutomationTest.h"
+
+#include "Async/Fundamental/Scheduler.h"
+#include "Experimental/Async/AwaitableTask.h"
+
+#include <atomic>
+
+#if WITH_DEV_AUTOMATION_TESTS
+
+namespace Tasks2Tests
+{
+	using namespace LowLevelTasks;
+
+	template<int NumTasks>
+	void EmptyPerfTest(bool last)
+	{
+		const double Alpha = 0.01;
+
+		if(last)
+		{
+			UE_LOG(LogTemp, Display, TEXT("--------------------------EMPTY Perf Test--------------------------------"));
+		}
+
+		{
+			double StartTime = FPlatformTime::Seconds();
+
+			TArray<FTask> Tasks;
+			Tasks.AddDefaulted(NumTasks);
+
+			for (int i = 0; i < NumTasks; ++i)
+			{
+				Tasks[i].Init(TEXT("Perf Test Basic"), [](){});
+			}
+
+			double SpawnedTime = FPlatformTime::Seconds();
+
+			for (int i = 0; i < NumTasks; ++i)
+			{
+				TryLaunch(Tasks[i]);
+			}
+
+			BusyWaitForTasks<FTask>(Tasks);
+
+			double EndTime = FPlatformTime::Seconds();
+			double Total = EndTime - StartTime;
+			static double TotalAvg = Total;
+			TotalAvg = (Alpha * Total) + (1.0 - Alpha) * TotalAvg;
+
+			double Spawned = SpawnedTime - StartTime;
+			static double SpawnedAvg = Spawned;
+			SpawnedAvg = (Alpha * Spawned) + (1.0 - Alpha) * SpawnedAvg;
+
+			double End = EndTime - SpawnedTime;
+			static double EndAvg = End;
+			EndAvg = (Alpha * End) + (1.0 - Alpha) * EndAvg;
+			if(last)
+			{
+				UE_LOG(LogTemp, Display, TEXT("EmptyBusyWaitGlobal Basic Performance:          Spawning %d tasks took %f (%f avg) secs total, %f (%f avg) secs spawning and %f (%f avg) secs dispatching"), NumTasks, Total, TotalAvg, Spawned, SpawnedAvg, End, EndAvg);
+			}
+		}
+
+		{
+			using namespace LowLevelTasks;
+
+			TArray<FTask> Tasks;
+			Tasks.AddDefaulted(NumTasks);
+
+			double StartTime = FPlatformTime::Seconds();
+
+			for (int i = 0; i < NumTasks; i++)
+			{
+				Tasks[i].Init(TEXT("Empty Perf Test"), [](){});
+			}
+
+			FTask SignalTask;
+			SignalTask.Init(TEXT("Empty Signal Task"), [](){}, [&Tasks]()
+			{
+				for (int i = 0; i < NumTasks; i++)
+				{
+					TryLaunch(Tasks[i]);
+				}
+			});
+
+			double SpawnedTime = FPlatformTime::Seconds();
+
+			TryLaunch(SignalTask);
+			BusyWaitForTasks<FTask>(Tasks);
+			BusyWaitForTask(SignalTask);
+
+			double EndTime = FPlatformTime::Seconds();
+			double Total = EndTime - StartTime;
+			static double TotalAvg = Total;
+			TotalAvg = (Alpha * Total) + (1.0 - Alpha) * TotalAvg;
+
+			double Spawned = SpawnedTime - StartTime;
+			static double SpawnedAvg = Spawned;
+			SpawnedAvg = (Alpha * Spawned) + (1.0 - Alpha) * SpawnedAvg;
+
+			double End = EndTime - SpawnedTime;
+			static double EndAvg = End;
+			EndAvg = (Alpha * End) + (1.0 - Alpha) * EndAvg;
+			if(last)
+			{
+				UE_LOG(LogTemp, Display, TEXT("EmptyBusyWaitLocal Basic Performance:           Spawning %d tasks took %f (%f avg) secs total, %f (%f avg) secs spawning and %f (%f avg) secs dispatching"), NumTasks, Total, TotalAvg, Spawned, SpawnedAvg, End, EndAvg);
+			}
+		}
+
+		{
+			TArray<TAwaitableTask<void>> Tasks;
+			Tasks.AddDefaulted(NumTasks);
+
+			double StartTime = FPlatformTime::Seconds();
+
+			TAwaitableTask<void> SignalTask;
+			SignalTask.InitAndLaunch(TEXT("Empty Awaitable Signal Task"), [&Tasks]()
+			{
+				for (int i = 0; i < NumTasks; i++)
+				{
+					Tasks[i].InitAndLaunch(TEXT("Empty Awaitable Perf Test"), [](){});
+				}
+			});
+
+			double SpawnedTime = FPlatformTime::Seconds();
+
+			SignalTask.Await();
+			for (int i = 0; i < NumTasks; i++)
+			{
+				Tasks[i].Await();
+			}
+
+			double EndTime = FPlatformTime::Seconds();
+			double Total = EndTime - StartTime;
+			static double TotalAvg = Total;
+			TotalAvg = (Alpha * Total) + (1.0 - Alpha) * TotalAvg;
+
+			double Spawned = SpawnedTime - StartTime;
+			static double SpawnedAvg = Spawned;
+			SpawnedAvg = (Alpha * Spawned) + (1.0 - Alpha) * SpawnedAvg;
+
+			double End = EndTime - SpawnedTime;
+			static double EndAvg = End;
+			EndAvg = (Alpha * End) + (1.0 - Alpha) * EndAvg;
+			if(last)
+			{
+				UE_LOG(LogTemp, Display, TEXT("EmptyBusyWaitLocal Awaitable Performance:       Spawning %d tasks took %f (%f avg) secs total, %f (%f avg) secs spawning and %f (%f avg) secs dispatching"), NumTasks, Total, TotalAvg, Spawned, SpawnedAvg, End, EndAvg);
+			}
+		}
+
+		if(last)
+		{
+			UE_LOG(LogTemp, Display, TEXT("--------------------------END EMPTY Perf Test-----------------------------"));
+		}
+	}
+
+	template<int NumTasks>
+	void BasicPerfTest(bool last)
+	{
+		const double Alpha = 0.01;
+		if(last)
+		{
+			UE_LOG(LogTemp, Display, TEXT("--------------------------SIMPLE Perf Test--------------------------------"));
+		}
+
+		struct Test
+		{
+			static int64 Fibonacci(int64 N)
+			{
+				if (N > 2)
+				{
+					int64 Res1 = Test::Fibonacci(N - 1);
+					int64 Res2 = Test::Fibonacci(N - 2);
+					return Res1 + Res2;
+				}
+				return 1;
+			}
+		};
+
+		static int64 SGarbage = 0;
+		static int64 STestVal = 15;
+
+		volatile int64& Garbage = SGarbage;
+		volatile int64& TestVal = STestVal;
+
+		{
+			using namespace LowLevelTasks;
+			TArray<FTask> Tasks;
+			Tasks.AddDefaulted(NumTasks);
+
+			double StartTime = FPlatformTime::Seconds();
+
+			for (int i = 0; i < NumTasks; i++)
+			{
+				Tasks[i].Init(TEXT("Perf Test"), [&Garbage, &TestVal]()
+				{
+					Garbage = Test::Fibonacci(TestVal);
+				});
+			}
+
+			FTask SignalTask;
+			SignalTask.Init(TEXT("Signal Task"), [](){}, [&Tasks]()
+			{
+				for (int i = 0; i < NumTasks; i++)
+				{
+					TryLaunch(Tasks[i]);
+				}
+			});
+
+			double SpawnedTime = FPlatformTime::Seconds();
+
+			TryLaunch(SignalTask);
+			BusyWaitForTasks<FTask>(Tasks);
+			BusyWaitForTask(SignalTask);
+
+			double EndTime = FPlatformTime::Seconds();
+			double Total = EndTime - StartTime;
+			static double TotalAvg = Total;
+			TotalAvg = (Alpha * Total) + (1.0 - Alpha) * TotalAvg;
+
+			double Spawned = SpawnedTime - StartTime;
+			static double SpawnedAvg = Spawned;
+			SpawnedAvg = (Alpha * Spawned) + (1.0 - Alpha) * SpawnedAvg;
+
+			double End = EndTime - SpawnedTime;
+			static double EndAvg = End;
+			EndAvg = (Alpha * End) + (1.0 - Alpha) * EndAvg;
+
+			if(last)
+			{
+				UE_LOG(LogTemp, Display, TEXT("FibworkBusyWait Basic Performance:              Spawning %d tasks took %f (%f avg) secs total, %f (%f avg) secs spawning and %f (%f avg) secs dispatching"), NumTasks, Total, TotalAvg, Spawned, SpawnedAvg, End, EndAvg);
+			}
+		}
+
+		{
+			TArray<TAwaitableTask<int64>> Tasks;
+			Tasks.AddDefaulted(NumTasks);
+
+			double StartTime = FPlatformTime::Seconds();
+
+			TAwaitableTask<void> SignalTask;
+			SignalTask.InitAndLaunch(TEXT("Simple Awaitable Signal Task"), [&Tasks, &TestVal]()
+			{
+				for (int i = 0; i < NumTasks; i++)
+				{
+					Tasks[i].InitAndLaunch(TEXT("Simple Awaitable Perf Test"), [&TestVal]()
+					{
+						return Test::Fibonacci(TestVal);
+					});
+				}
+			});
+
+			double SpawnedTime = FPlatformTime::Seconds();
+
+			SignalTask.Await();
+			for (int i = 0; i < NumTasks; i++)
+			{
+				Garbage = Tasks[i].Await();
+			}
+
+			double EndTime = FPlatformTime::Seconds();
+			double Total = EndTime - StartTime;
+			static double TotalAvg = Total;
+			TotalAvg = (Alpha * Total) + (1.0 - Alpha) * TotalAvg;
+
+			double Spawned = SpawnedTime - StartTime;
+			static double SpawnedAvg = Spawned;
+			SpawnedAvg = (Alpha * Spawned) + (1.0 - Alpha) * SpawnedAvg;
+
+			double End = EndTime - SpawnedTime;
+			static double EndAvg = End;
+			EndAvg = (Alpha * End) + (1.0 - Alpha) * EndAvg;
+
+			if(last)
+			{
+				UE_LOG(LogTemp, Display, TEXT("FibworkBusyWait Awaitable Performance:          Spawning %d tasks took %f (%f avg) secs total, %f (%f avg) secs spawning and %f (%f avg) secs dispatching"), NumTasks, Total, TotalAvg, Spawned, SpawnedAvg, End, EndAvg);
+			}
+		}
+
+		if(last)
+		{
+			UE_LOG(LogTemp, Display, TEXT("--------------------------END SIMPLE Perf Test--------------------------------"));
+		}
+		SGarbage++;
+	}
+
+	template<int NumTasks>
+	void StealingPerfTest(bool last)
+	{
+		const double Alpha = 0.01;
+		static int64 SGarbage = 0;
+		static int64 STestVal = 13;
+		static int64 STheshold = 11;
+
+		using namespace LowLevelTasks;
+		struct Test
+		{
+			static int64 Fibonacci(int64 N)
+			{
+				if (N > 2)
+				{
+					int64 Res1 = Test::Fibonacci(N - 1);
+					int64 Res2 = Test::Fibonacci(N - 2);
+					return Res1 + Res2;
+				}
+				return 1;
+			}
+
+			static int64 FibonacciTasksBasic(int64 N)
+			{
+				if (N > STheshold)
+				{
+					int64 Res1, Res2;
+					FTask Task1, Task2;
+
+					Task1.Init(TEXT("Basic Fibonacci Test1"), ETaskPriority::Normal, [&Res1, N]()
+					{
+						Res1 = Test::FibonacciTasksBasic(N - 1);
+					});
+					TryLaunch(Task1);
+
+					Task2.Init(TEXT("Basic Fibonacci Test2"), ETaskPriority::Normal, [&Res2, N]()
+					{
+						Res2 = Test::FibonacciTasksBasic(N - 2);
+					});
+					TryLaunch(Task2);
+
+					BusyWaitForTask(Task2);
+					BusyWaitForTask(Task1);
+					return Res1 + Res2;
+				}
+				return Fibonacci(N);
+			}
+
+			static int64 FibonacciTasksAwaitable(int64 N)
+			{
+				if (N > STheshold)
+				{
+					TAwaitableTask<int64> Task1, Task2;
+					Task1.InitAndLaunch(TEXT("Awaitable Fibonacci Test1"), [N]()
+					{
+						return Test::FibonacciTasksBasic(N - 1);
+					});
+
+					Task2.InitAndLaunch(TEXT("Awaitable Fibonacci Test2"), [N]()
+					{
+						return Test::FibonacciTasksBasic(N - 2);
+					});
+
+					return Task2.Await() + Task1.Await();
+				}
+				return Fibonacci(N);
+			}
+		};
+
+		if(last)
+		{
+			UE_LOG(LogTemp, Display, TEXT("--------------------------STEALING Perf Test--------------------------------"));
+		}
+		volatile int64& Garbage = SGarbage;
+		volatile int64& TestVal = STestVal;
+
+		{
+			TArray<FTask> Tasks;
+			Tasks.AddDefaulted(NumTasks);
+
+			double StartTime = FPlatformTime::Seconds();
+
+			for (int i = 0; i < NumTasks; i++)
+			{
+				Tasks[i].Init(TEXT("Perf Test Stealing"), ETaskPriority::Normal, [&Garbage, &TestVal]()
+				{
+					Garbage = Test::FibonacciTasksBasic(TestVal);
+				});
+			}
+
+			FTask SignalTask;
+			SignalTask.Init(TEXT("Signal Task"), [](){}, [&Tasks]()
+			{
+				for (int i = 0; i < NumTasks; i++)
+				{
+					TryLaunch(Tasks[i]);
+				}
+			});
+
+			double SpawnedTime = FPlatformTime::Seconds();
+
+			TryLaunch(SignalTask);
+			BusyWaitForTasks<FTask>(Tasks);
+			BusyWaitForTask(SignalTask);
+
+			double EndTime = FPlatformTime::Seconds();
+			double Total = EndTime - StartTime;
+			static double TotalAvg = Total;
+			TotalAvg = (Alpha * Total) + (1.0 - Alpha) * TotalAvg;
+
+			double Spawned = SpawnedTime - StartTime;
+			static double SpawnedAvg = Spawned;
+			SpawnedAvg = (Alpha * Spawned) + (1.0 - Alpha) * SpawnedAvg;
+
+			double End = EndTime - SpawnedTime;
+			static double EndAvg = End;
+			EndAvg = (Alpha * End) + (1.0 - Alpha) * EndAvg;
+
+			if(last)
+			{
+				UE_LOG(LogTemp, Display, TEXT("FibworkBusyWaitStealing Basic Performance:      Spawning %d tasks took %f (%f avg) secs total, %f (%f avg) secs spawning and %f (%f avg) secs dispatching"), NumTasks, Total, TotalAvg, Spawned, SpawnedAvg, End, EndAvg);
+			}
+		}
+
+		{
+			TArray<TAwaitableTask<int64>> Tasks;
+			Tasks.AddDefaulted(NumTasks);
+
+			double StartTime = FPlatformTime::Seconds();
+
+			TAwaitableTask<void> SignalTask;
+			SignalTask.InitAndLaunch(TEXT("Stealing Awaitable Signal Task"), [&Tasks, &TestVal]()
+			{
+				for (int i = 0; i < NumTasks; i++)
+				{
+					Tasks[i].InitAndLaunch(TEXT("Stealing Awaitable Perf Test"), [&TestVal]()
+					{
+						return Test::FibonacciTasksAwaitable(TestVal);
+					});
+					//Tasks[i].Await();
+				}
+			});
+
+			double SpawnedTime = FPlatformTime::Seconds();
+
+			SignalTask.Await();
+			for (int i = 0; i < NumTasks; i++)
+			{
+				Garbage = Tasks[i].Await();
+			}
+
+			double EndTime = FPlatformTime::Seconds();
+			double Total = EndTime - StartTime;
+			static double TotalAvg = Total;
+			TotalAvg = (Alpha * Total) + (1.0 - Alpha) * TotalAvg;
+
+			double Spawned = SpawnedTime - StartTime;
+			static double SpawnedAvg = Spawned;
+			SpawnedAvg = (Alpha * Spawned) + (1.0 - Alpha) * SpawnedAvg;
+
+			double End = EndTime - SpawnedTime;
+			static double EndAvg = End;
+			EndAvg = (Alpha * End) + (1.0 - Alpha) * EndAvg;
+
+			if(last)
+			{
+				UE_LOG(LogTemp, Display, TEXT("FibworkBusyWaitStealing Awaitable Performance:  Spawning %d tasks took %f (%f avg) secs total, %f (%f avg) secs spawning and %f (%f avg) secs dispatching"), NumTasks, Total, TotalAvg, Spawned, SpawnedAvg, End, EndAvg);
+			}
+		}
+
+		if(last)
+		{
+			UE_LOG(LogTemp, Display, TEXT("--------------------------END STEALING Perf Test--------------------------------"));
+		}
+		SGarbage++;
+	}
+
+	IMPLEMENT_SIMPLE_AUTOMATION_TEST(FTasksPerfTests2, "System.Core.LowLevelTasks.PerfTests", EAutomationTestFlags::ApplicationContextMask | EAutomationTestFlags::EngineFilter);
+
+	bool FTasksPerfTests2::RunTest(const FString& Parameters)
+	{
+		for(int i = 0; i <= 100; i++)
+		{
+			EmptyPerfTest<10000>(i == 100);
+			BasicPerfTest<10000>(i == 100);
+			StealingPerfTest<10000>(i == 100);
+		}
+		return true;
+	}
+
+	IMPLEMENT_SIMPLE_AUTOMATION_TEST(FTasksUnitTests2, "System.Core.LowLevelTasks.UnitTests", EAutomationTestFlags::ApplicationContextMask | EAutomationTestFlags::EngineFilter);
+	bool FTasksUnitTests2::RunTest(const FString& Parameters)
+	{
+		using namespace LowLevelTasks;
+
+		for (uint32 TestCancel = 0; TestCancel < 2; TestCancel++)
+		{
+			//Basic Low LevelAPi
+			{
+				uint32 TestValue = 1337;
+
+				FTask Task;
+				Task.Init(TEXT("Basic Test"), [&TestValue]()
+				{
+					TestValue = 42;
+				});
+				TryLaunch(Task);
+
+				bool WasCanceled = TestCancel && Task.TryCancel();
+				BusyWaitForTask(Task);
+
+				if (WasCanceled)
+				{
+					verify(TestValue == 1337);
+				}
+				else
+				{
+					verify(TestValue == 42);
+				}
+			}
+
+			//optional Refcounting
+			{
+				uint32 TestValue = 1337;
+
+				using FTaskHandle = TSharedPtr<FTask, ESPMode::ThreadSafe>;
+				FTaskHandle TaskHandle = MakeShared<FTask, ESPMode::ThreadSafe>();
+
+				TaskHandle->Init(TEXT("Refcounted Test"), [&TestValue]()
+				{
+					TestValue = 42;
+				}, [TaskHandle](){});
+				TryLaunch(*TaskHandle);
+
+				bool WasCanceled = TestCancel && TaskHandle->TryCancel();
+				BusyWaitForTask(*TaskHandle);
+
+				if (WasCanceled)
+				{
+					verify(TestValue == 1337);
+				}
+				else
+				{
+					verify(TestValue == 42);
+				}
+			}
+		}
+
+		//example Awaitable Tasks
+		{
+			uint32 TestValue = 1337;
+
+			TAwaitableTask<void> Task;
+			Task.InitAndLaunch(TEXT("Awaitable Test"), [&TestValue]()
+			{
+				TestValue = 42;
+			});
+
+			Task.Await();
+			verify(TestValue == 42);
+		}
+
+		//optional uniqueptr
+		{
+			uint32* TestValue = new uint32(1337);
+
+			using FTaskHandle = TUniquePtr<FTask>;
+			FTaskHandle TaskHandle = MakeUnique<FTask>();
+			FTask& Task = *TaskHandle;
+			Task.Init(TEXT("Uniqueptr Test"), [TestValue]()
+			{
+				*TestValue = 42;
+			}, [TestValue, TaskHandle = MoveTemp(TaskHandle)]()
+			{
+				verify(*TestValue == 42);
+				delete TestValue;
+			});
+			TryLaunch(Task);
+		}
+
+		return true;
+	}
+}
+
+#endif
\ No newline at end of file
diff --git a/Engine/Source/Runtime/Core/Public/Experimental/Async/AwaitableTask.h b/Engine/Source/Runtime/Core/Public/Experimental/Async/AwaitableTask.h
--- a/Engine/Source/Runtime/Core/Public/Experimental/Async/AwaitableTask.h
+++ b/Engine/Source/Runtime/Core/Public/Experimental/Async/AwaitableTask.h
@@ -200,9 +200,13 @@ namespace AwaitableTask_Detail
 		}
 	};
 
-	template<typename ReturnType, typename CallableType>
+	template<typename TReturnType, typename CallableType>
 	class TPromise final : public FPromiseBase
 	{
+	public:
+		using ReturnType = TReturnType;
+
+	private:
 		using ThisType = TPromise<ReturnType, CallableType>;
 
 	private:
@@ -254,10 +258,10 @@ namespace AwaitableTask_Detail
 		inline void Finish()
 		{
 			int LocalCounter = ReferenceCounter.fetch_sub(1, std::memory_order_release);
-			if (LocalCounter == 2) //this is 2 because fetch_sub returns the value before the decrement (value is actually 1)
+			if (LocalCounter == 2 && !IsLaunched()) //this is 2 because fetch_sub returns the value before the decrement (value is actually 1)
 			{
 				//Cancel the Task and Try to Launch the Continuation if we are the last reference
-				verify(!IsLaunched() == LowLevelTasks::TryCancelAndLaunchContinuation(Task));
+				verify(LowLevelTasks::TryCancelAndLaunchContinuation(Task));
 			}
 			else if (LocalCounter == 1) //this is 1 because fetch_sub returns the value before the decrement (value is actually 0)
 			{
@@ -283,11 +287,12 @@ namespace AwaitableTask_Detail
 	template<typename CallableType>
 	class TPromise<void, CallableType> final : public FPromiseBase
 	{
-		using ThisType = TPromise<void, CallableType>;
-
 	public:
 		using ReturnType = void;
 
+	private:
+		using ThisType = TPromise<void, CallableType>;
+
 	private:
 		CallableType Callable;
 		LowLevelTasks::FTask Task;
@@ -336,10 +341,10 @@ namespace AwaitableTask_Detail
 		inline void Finish()
 		{
 			int LocalCounter = ReferenceCounter.fetch_sub(1, std::memory_order_release);
-			if (LocalCounter == 2) //this is 2 because fetch_sub returns the value before the decrement (value is actually 1)
+			if (LocalCounter == 2 && !IsLaunched()) //this is 2 because fetch_sub returns the value before the decrement (value is actually 1)
 			{
-				//Cancel the Task and Try to Launch the continuation if we are the last reference
-				verify(!IsLaunched() == LowLevelTasks::TryCancelAndLaunchContinuation(Task));
+				//Cancel the Task and Try to Launch the Continuation if we are the last reference
+				verify(LowLevelTasks::TryCancelAndLaunchContinuation(Task));
 			}
 			else if (LocalCounter == 1) //this is 1 because fetch_sub returns the value before the decrement (value is actually 0)
 			{
