Use HashCombine in FHashBuilder instead of a faster (but worst) version to minimize hash collisions.

diff --git a/Engine/Source/Runtime/Core/Public/Misc/HashBuilder.h b/Engine/Source/Runtime/Core/Public/Misc/HashBuilder.h
--- a/Engine/Source/Runtime/Core/Public/Misc/HashBuilder.h
+++ b/Engine/Source/Runtime/Core/Public/Misc/HashBuilder.h
@@ -42,7 +42,7 @@ public:
 	template <typename T>
 	typename TEnableIf<TModels<CGetTypeHashable, T>::Value, FHashBuilder&>::Type Append(const T& InData)
 	{
-		Hash ^= GetTypeHash(InData) + 0x9e3779b9 + (Hash << 6) + (Hash >> 2);
+		Hash = HashCombine(Hash, GetTypeHash(InData));
 		return *this;
 	}
 
