CompactBinary: Rename Reference to CompactBinaryReference for clarity and consistency

#rnx

diff --git a/Engine/Source/Developer/DerivedDataCache/Private/DerivedDataCache.cpp b/Engine/Source/Developer/DerivedDataCache/Private/DerivedDataCache.cpp
--- a/Engine/Source/Developer/DerivedDataCache/Private/DerivedDataCache.cpp
+++ b/Engine/Source/Developer/DerivedDataCache/Private/DerivedDataCache.cpp
@@ -827,9 +827,9 @@ void FCache::Get(
 
 	// Validate that the record was serialized in the expected format.
 	if ((MetaField && !MetaField.IsObject()) ||
-		(ValueField && !ValueField.IsAnyReference()) ||
+		(ValueField && !ValueField.IsReference()) ||
 		(AttachmentsField && !AttachmentsField.IsArray()) ||
-		!Algo::AllOf(AttachmentsArray, &FCbField::IsAnyReference))
+		!Algo::AllOf(AttachmentsArray, &FCbField::IsReference))
 	{
 		UE_LOG(LogDerivedDataCache, Verbose, TEXT("Cache: Get cache miss with invalid format for %s from '%.*s'"),
 			*TToString<96>(Key), Context.Len(), Context.GetData());
@@ -839,9 +839,9 @@ void FCache::Get(
 	// Check for existence of the value and attachments if they are being skipped.
 	{
 		TArray<FBlake3Hash, TInlineAllocator<1>> KeysToCheck;
-		if (EnumHasAnyFlags(Policy, ECachePolicy::SkipValue) && ValueField.IsAnyReference())
+		if (EnumHasAnyFlags(Policy, ECachePolicy::SkipValue) && ValueField.IsReference())
 		{
-			KeysToCheck.Add(ValueField.AsAnyReference());
+			KeysToCheck.Add(ValueField.AsReference());
 			if (KeysToCheck.Last().IsZero())
 			{
 				KeysToCheck.Pop();
@@ -852,7 +852,7 @@ void FCache::Get(
 			KeysToCheck.Reserve(KeysToCheck.Num() + AttachmentsArray.Num());
 			for (FCbField AttachmentField : AttachmentsArray)
 			{
-				KeysToCheck.Add(AttachmentField.AsAnyReference());
+				KeysToCheck.Add(AttachmentField.AsReference());
 			}
 		}
 		if (KeysToCheck.Num())
@@ -874,7 +874,7 @@ void FCache::Get(
 	const auto GetContent = [&Key, Context](FCbField ReferenceField) -> FCbAttachment
 	{
 		TArray<uint8> ContentData;
-		const FBlake3Hash Hash = ReferenceField.AsAnyReference();
+		const FBlake3Hash Hash = ReferenceField.AsReference();
 		check(!ReferenceField.HasError());
 		const FCacheAttachmentKey ContentKey(Key, Hash);
 		if (GetDerivedDataCacheRef().GetSynchronous(*MakeContentKey(ContentKey), ContentData, Context))
@@ -882,7 +882,7 @@ void FCache::Get(
 			if (FBlake3::HashBuffer(MakeMemoryView(ContentData)) == Hash)
 			{
 				FSharedBuffer ContentBuffer = FSharedBuffer::Clone(MakeMemoryView(ContentData));
-				if (ReferenceField.IsReference())
+				if (ReferenceField.IsCompactBinaryReference())
 				{
 					return FCbAttachment(FCbFieldRefIterator::MakeRange(MoveTemp(ContentBuffer)), Hash);
 				}
@@ -906,7 +906,7 @@ void FCache::Get(
 	};
 
 	FCbAttachment Value;
-	if (!EnumHasAnyFlags(Policy, ECachePolicy::SkipValue) && ValueField.IsAnyReference() && !(Value = GetContent(ValueField)))
+	if (!EnumHasAnyFlags(Policy, ECachePolicy::SkipValue) && ValueField.IsReference() && !(Value = GetContent(ValueField)))
 	{
 		return;
 	}
@@ -1052,7 +1052,7 @@ void FCache::Put(
 	{
 		if (Value.IsCompactBinary())
 		{
-			Writer.Name("Value"_ASV).Reference(Value.GetHash());
+			Writer.Name("Value"_ASV).CompactBinaryReference(Value.GetHash());
 		}
 		else if (Value.IsBinary())
 		{
@@ -1067,7 +1067,7 @@ void FCache::Put(
 			{
 				if (Attachment.IsCompactBinary())
 				{
-					Writer.Reference(Attachment.GetHash());
+					Writer.CompactBinaryReference(Attachment.GetHash());
 				}
 				else
 				{
diff --git a/Engine/Source/Runtime/Core/Private/Serialization/CompactBinary.cpp b/Engine/Source/Runtime/Core/Private/Serialization/CompactBinary.cpp
--- a/Engine/Source/Runtime/Core/Private/Serialization/CompactBinary.cpp
+++ b/Engine/Source/Runtime/Core/Private/Serialization/CompactBinary.cpp
@@ -64,10 +64,10 @@ void FCbFieldType::StaticAssertTypeConstants()
 	static_assert(!(BoolMask & (BoolBase ^ ECbFieldType::BoolTrue)), "BoolMask or BoolBase is invalid!");
 	static_assert(TypeMask == (BoolMask | (BoolBase ^ ECbFieldType::BoolTrue)), "BoolMask or BoolBase is invalid!");
 
-	static_assert(AnyReferenceBase == ECbFieldType::Reference, "AnyReferenceBase is invalid!");
-	static_assert((AnyReferenceMask & (AllFlags | ECbFieldType::BinaryReference)) == ECbFieldType::Reference, "AnyReferenceMask is invalid!");
-	static_assert(!(AnyReferenceMask & (AnyReferenceBase ^ ECbFieldType::BinaryReference)), "AnyReferenceMask or AnyReferenceBase is invalid!");
-	static_assert(TypeMask == (AnyReferenceMask | (AnyReferenceBase ^ ECbFieldType::BinaryReference)), "AnyReferenceMask or AnyReferenceBase is invalid!");
+	static_assert(ReferenceBase == ECbFieldType::CompactBinaryReference, "ReferenceBase is invalid!");
+	static_assert((ReferenceMask & (AllFlags | ECbFieldType::BinaryReference)) == ECbFieldType::CompactBinaryReference, "ReferenceMask is invalid!");
+	static_assert(!(ReferenceMask & (ReferenceBase ^ ECbFieldType::BinaryReference)), "ReferenceMask or ReferenceBase is invalid!");
+	static_assert(TypeMask == (ReferenceMask | (ReferenceBase ^ ECbFieldType::BinaryReference)), "ReferenceMask or ReferenceBase is invalid!");
 }
 
 ///////////////////////////////////////////////////////////////////////////////////////////////////////////////////////
@@ -253,9 +253,9 @@ bool FCbField::AsBool(const bool bDefault)
 	return (uint8(bIsBool) & uint8(LocalType) & 1) | ((!bIsBool) & bDefault);
 }
 
-FBlake3Hash FCbField::AsReference(const FBlake3Hash& Default)
+FBlake3Hash FCbField::AsCompactBinaryReference(const FBlake3Hash& Default)
 {
-	if (FCbFieldType::IsReference(Type))
+	if (FCbFieldType::IsCompactBinaryReference(Type))
 	{
 		Error = ECbFieldError::None;
 		return FBlake3Hash(*static_cast<const FBlake3Hash::ByteArray*>(Payload));
@@ -281,9 +281,9 @@ FBlake3Hash FCbField::AsBinaryReference(const FBlake3Hash& Default)
 	}
 }
 
-FBlake3Hash FCbField::AsAnyReference(const FBlake3Hash& Default)
+FBlake3Hash FCbField::AsReference(const FBlake3Hash& Default)
 {
-	if (FCbFieldType::IsAnyReference(Type))
+	if (FCbFieldType::IsReference(Type))
 	{
 		Error = ECbFieldError::None;
 		return FBlake3Hash(*static_cast<const FBlake3Hash::ByteArray*>(Payload));
@@ -415,7 +415,7 @@ uint64 FCbField::GetPayloadSize() const
 	case ECbFieldType::BoolFalse:
 	case ECbFieldType::BoolTrue:
 		return 0;
-	case ECbFieldType::Reference:
+	case ECbFieldType::CompactBinaryReference:
 	case ECbFieldType::BinaryReference:
 	case ECbFieldType::Hash:
 		return 32;
@@ -479,7 +479,7 @@ void FCbField::IterateReferences(FCbFieldVisitor Visitor) const
 	case ECbFieldType::Array:
 	case ECbFieldType::UniformArray:
 		return FCbArray::FromField(*this).IterateReferences(Visitor);
-	case ECbFieldType::Reference:
+	case ECbFieldType::CompactBinaryReference:
 	case ECbFieldType::BinaryReference:
 		return Visitor(*this);
 	default:
diff --git a/Engine/Source/Runtime/Core/Private/Serialization/CompactBinaryPackage.cpp b/Engine/Source/Runtime/Core/Private/Serialization/CompactBinaryPackage.cpp
--- a/Engine/Source/Runtime/Core/Private/Serialization/CompactBinaryPackage.cpp
+++ b/Engine/Source/Runtime/Core/Private/Serialization/CompactBinaryPackage.cpp
@@ -95,9 +95,9 @@ void FCbAttachment::Load(FCbFieldRefIterator& Fields)
 		Buffer = FSharedBuffer::MakeView(View, Fields.GetOuterBuffer());
 		Buffer.MakeOwned();
 		++Fields;
-		Hash = Fields.AsAnyReference();
+		Hash = Fields.AsReference();
 		checkf(!Fields.HasError(), TEXT("Attachments must be a non-empty binary value with a content hash."));
-		if (Fields.IsReference())
+		if (Fields.IsCompactBinaryReference())
 		{
 			CompactBinary = FCbFieldIterator::MakeRange(Buffer);
 		}
@@ -130,9 +130,9 @@ void FCbAttachment::Load(FArchive& Ar, FCbBufferAllocator Allocator)
 				HashBuffer.SetNumUninitialized(int32(Size));
 				return FUniqueBuffer::MakeView(HashBuffer.GetData(), Size);
 			});
-		Hash = HashField.AsAnyReference();
+		Hash = HashField.AsReference();
 		checkf(!HashField.HasError(), TEXT("Attachments must be a non-empty binary value with a content hash."));
-		if (HashField.IsReference())
+		if (HashField.IsCompactBinaryReference())
 		{
 			CompactBinary = FCbFieldIterator::MakeRange(Buffer);
 		}
@@ -158,7 +158,7 @@ void FCbAttachment::Save(FCbWriter& Writer) const
 		{
 			Writer.Binary(AsBinary());
 		}
-		Writer.Reference(Hash);
+		Writer.CompactBinaryReference(Hash);
 	}
 	else if (Buffer.GetSize())
 	{
@@ -258,10 +258,10 @@ void FCbPackage::GatherAttachments(const FCbFieldIterator& Fields, FAttachmentRe
 {
 	Fields.IterateRangeReferences([this, &Resolver](FCbField Field)
 		{
-			const FBlake3Hash& Hash = Field.AsAnyReference();
+			const FBlake3Hash& Hash = Field.AsReference();
 			if (FSharedBuffer Buffer = Resolver(Hash))
 			{
-				if (Field.IsReference())
+				if (Field.IsCompactBinaryReference())
 				{
 					AddAttachment(FCbAttachment(FCbFieldRefIterator::MakeRange(MoveTemp(Buffer)), Hash), &Resolver);
 				}
@@ -297,8 +297,8 @@ void FCbPackage::Load(FCbFieldRefIterator& Fields)
 			++Fields;
 			if (Object.CreateIterator())
 			{
-				ObjectHash = Fields.AsReference();
-				checkf(!Fields.HasError(), TEXT("Object must be followed by a Reference with the object hash."));
+				ObjectHash = Fields.AsCompactBinaryReference();
+				checkf(!Fields.HasError(), TEXT("Object must be followed by a CompactBinaryReference with the object hash."));
 				++Fields;
 			}
 			else
@@ -333,9 +333,9 @@ void FCbPackage::Load(FArchive& Ar, FCbBufferAllocator Allocator)
 				FSharedBuffer Buffer = FSharedBuffer::MakeView(View, ValueField.GetOuterBuffer());
 				Buffer.MakeOwned();
 				FCbFieldRef HashField = LoadCompactBinary(Ar, StackAllocator);
-				const FBlake3Hash& Hash = HashField.AsAnyReference();
+				const FBlake3Hash& Hash = HashField.AsReference();
 				checkf(!HashField.HasError(), TEXT("Attachments must be a non-empty binary value with a content hash."));
-				if (HashField.IsReference())
+				if (HashField.IsCompactBinaryReference())
 				{
 					AddAttachment(FCbAttachment(FCbFieldRefIterator::MakeRange(MoveTemp(Buffer)), Hash));
 				}
@@ -353,8 +353,8 @@ void FCbPackage::Load(FArchive& Ar, FCbBufferAllocator Allocator)
 			if (Object.CreateIterator())
 			{
 				FCbFieldRef HashField = LoadCompactBinary(Ar, StackAllocator);
-				ObjectHash = HashField.AsReference();
-				checkf(!HashField.HasError(), TEXT("Object must be followed by a Reference with the object hash."));
+				ObjectHash = HashField.AsCompactBinaryReference();
+				checkf(!HashField.HasError(), TEXT("Object must be followed by a CompactBinaryReference with the object hash."));
 			}
 			else
 			{
@@ -369,7 +369,7 @@ void FCbPackage::Save(FCbWriter& Writer) const
 	if (Object.CreateIterator())
 	{
 		Writer.Object(Object);
-		Writer.Reference(ObjectHash);
+		Writer.CompactBinaryReference(ObjectHash);
 	}
 	for (const FCbAttachment& Attachment : Attachments)
 	{
diff --git a/Engine/Source/Runtime/Core/Private/Serialization/CompactBinarySerialization.cpp b/Engine/Source/Runtime/Core/Private/Serialization/CompactBinarySerialization.cpp
--- a/Engine/Source/Runtime/Core/Private/Serialization/CompactBinarySerialization.cpp
+++ b/Engine/Source/Runtime/Core/Private/Serialization/CompactBinarySerialization.cpp
@@ -59,7 +59,7 @@ bool TryMeasureCompactBinary(FMemoryView View, ECbFieldType& OutType, uint64& Ou
 	case ECbFieldType::BoolFalse:
 	case ECbFieldType::BoolTrue:
 		break;
-	case ECbFieldType::Reference:
+	case ECbFieldType::CompactBinaryReference:
 	case ECbFieldType::BinaryReference:
 	case ECbFieldType::Hash:
 		FixedSize = 32;
diff --git a/Engine/Source/Runtime/Core/Private/Serialization/CompactBinaryValidation.cpp b/Engine/Source/Runtime/Core/Private/Serialization/CompactBinaryValidation.cpp
--- a/Engine/Source/Runtime/Core/Private/Serialization/CompactBinaryValidation.cpp
+++ b/Engine/Source/Runtime/Core/Private/Serialization/CompactBinaryValidation.cpp
@@ -352,7 +352,7 @@ static FCbField ValidateCbField(FMemoryView& View, ECbValidateMode Mode, ECbVali
 	case ECbFieldType::Float64:
 		ValidateCbFloat64(View, Mode, Error);
 		break;
-	case ECbFieldType::Reference:
+	case ECbFieldType::CompactBinaryReference:
 	case ECbFieldType::BinaryReference:
 	case ECbFieldType::Hash:
 		ValidateFixedPayload(32);
@@ -409,7 +409,7 @@ static FBlake3Hash ValidateCbPackageAttachment(FCbField& Value, FMemoryView& Vie
 	{
 		if (FCbField HashField = ValidateCbPackageField(View, Mode, Error))
 		{
-			const FBlake3Hash Hash = HashField.AsAnyReference();
+			const FBlake3Hash Hash = HashField.AsReference();
 			if (EnumHasAnyFlags(Mode, ECbValidateMode::Package))
 			{
 				if (HashField.HasError())
@@ -439,7 +439,7 @@ static FBlake3Hash ValidateCbPackageObject(FCbField& Value, FMemoryView& View, E
 	}
 	else if (FCbField HashField = ValidateCbPackageField(View, Mode, Error))
 	{
-		const FBlake3Hash Hash = HashField.AsAnyReference();
+		const FBlake3Hash Hash = HashField.AsReference();
 		if (EnumHasAnyFlags(Mode, ECbValidateMode::Package))
 		{
 			if (!Object.CreateIterator())
diff --git a/Engine/Source/Runtime/Core/Private/Serialization/CompactBinaryWriter.cpp b/Engine/Source/Runtime/Core/Private/Serialization/CompactBinaryWriter.cpp
--- a/Engine/Source/Runtime/Core/Private/Serialization/CompactBinaryWriter.cpp
+++ b/Engine/Source/Runtime/Core/Private/Serialization/CompactBinaryWriter.cpp
@@ -477,11 +477,11 @@ void FCbWriter::Bool(const bool bValue)
 	EndField(bValue ? ECbFieldType::BoolTrue : ECbFieldType::BoolFalse);
 }
 
-void FCbWriter::Reference(const FBlake3Hash& Value)
+void FCbWriter::CompactBinaryReference(const FBlake3Hash& Value)
 {
 	BeginField();
 	Data.Append(Value.Bytes(), sizeof(decltype(Value.Bytes())));
-	EndField(ECbFieldType::Reference);
+	EndField(ECbFieldType::CompactBinaryReference);
 }
 
 void FCbWriter::BinaryReference(const FBlake3Hash& Value)
diff --git a/Engine/Source/Runtime/Core/Private/Tests/Serialization/CompactBinaryPackageTest.cpp b/Engine/Source/Runtime/Core/Private/Tests/Serialization/CompactBinaryPackageTest.cpp
--- a/Engine/Source/Runtime/Core/Private/Tests/Serialization/CompactBinaryPackageTest.cpp
+++ b/Engine/Source/Runtime/Core/Private/Tests/Serialization/CompactBinaryPackageTest.cpp
@@ -372,7 +372,7 @@ bool FCbPackageTest::RunTest(const FString& Parameters)
 		TCbWriter<256> Writer;
 		Writer.Name("Level3");
 		Writer.BeginArray();
-		Writer.Reference(Level3Hash);
+		Writer.CompactBinaryReference(Level3Hash);
 		Writer.EndArray();
 		Level2 = Writer.Save().AsArrayRef();
 	}
@@ -382,7 +382,7 @@ bool FCbPackageTest::RunTest(const FString& Parameters)
 	{
 		TCbWriter<256> Writer;
 		Writer.BeginObject();
-		Writer.Name("Level2").Reference(Level2Hash);
+		Writer.Name("Level2").CompactBinaryReference(Level2Hash);
 		Writer.EndObject();
 		Level1 = Writer.Save().AsObjectRef();
 	}
@@ -453,13 +453,13 @@ bool FCbPackageTest::RunTest(const FString& Parameters)
 	{
 		TCbWriter<384> Writer;
 		Writer.Binary(Level2.GetBuffer());
-		Writer.Reference(Level2Hash);
+		Writer.CompactBinaryReference(Level2Hash);
 		Writer.Binary(Level4);
 		Writer.BinaryReference(Level4Hash);
 		Writer.Object(Level1);
-		Writer.Reference(Level1Hash);
+		Writer.CompactBinaryReference(Level1Hash);
 		Writer.Binary(Level3.GetBuffer());
-		Writer.Reference(Level3Hash);
+		Writer.CompactBinaryReference(Level3Hash);
 		Writer.Null();
 
 		FCbFieldRefIterator Fields = Writer.Save();
@@ -507,15 +507,15 @@ bool FCbPackageTest::RunTest(const FString& Parameters)
 		FCbFieldRefIterator Saved = Writer.Save();
 		TestTrue(TEXT("FCbPackage(OutOfOrder).Save()->Level1"), Saved.AsObjectRef().Equals(Level1));
 		++Saved;
-		TestEqual(TEXT("FCbPackage(OutOfOrder).Save()->Level1Hash"), Saved.AsReference(), Level1Hash);
+		TestEqual(TEXT("FCbPackage(OutOfOrder).Save()->Level1Hash"), Saved.AsCompactBinaryReference(), Level1Hash);
 		++Saved;
 		TestTrue(TEXT("FCbPackage(OutOfOrder).Save()->Level3"), Saved.AsBinary().EqualBytes(Level3.GetView()));
 		++Saved;
-		TestEqual(TEXT("FCbPackage(OutOfOrder).Save()->Level3Hash"), Saved.AsReference(), Level3Hash);
+		TestEqual(TEXT("FCbPackage(OutOfOrder).Save()->Level3Hash"), Saved.AsCompactBinaryReference(), Level3Hash);
 		++Saved;
 		TestTrue(TEXT("FCbPackage(OutOfOrder).Save()->Level2"), Saved.AsBinary().EqualBytes(Level2.GetView()));
 		++Saved;
-		TestEqual(TEXT("FCbPackage(OutOfOrder).Save()->Level2Hash"), Saved.AsReference(), Level2Hash);
+		TestEqual(TEXT("FCbPackage(OutOfOrder).Save()->Level2Hash"), Saved.AsCompactBinaryReference(), Level2Hash);
 		++Saved;
 		TestTrue(TEXT("FCbPackage(OutOfOrder).Save()->Level4"), Saved.AsBinary().EqualBytes(Level4.GetView()));
 		++Saved;
diff --git a/Engine/Source/Runtime/Core/Private/Tests/Serialization/CompactBinarySerializationTest.cpp b/Engine/Source/Runtime/Core/Private/Tests/Serialization/CompactBinarySerializationTest.cpp
--- a/Engine/Source/Runtime/Core/Private/Tests/Serialization/CompactBinarySerializationTest.cpp
+++ b/Engine/Source/Runtime/Core/Private/Tests/Serialization/CompactBinarySerializationTest.cpp
@@ -115,7 +115,7 @@ bool FCbMeasureTest::RunTest(const FString& Parameters)
 	TestMeasure(TEXT("BoolFalse"), {uint8(EType::BoolFalse)}, true, 1, EType::BoolFalse);
 	TestMeasure(TEXT("BoolTrue"), {uint8(EType::BoolTrue)}, true, 1, EType::BoolTrue);
 
-	TestMeasure(TEXT("Reference"), {uint8(EType::Reference)}, true, 33, EType::Reference);
+	TestMeasure(TEXT("CompactBinaryReference"), {uint8(EType::CompactBinaryReference)}, true, 33, EType::CompactBinaryReference);
 	TestMeasure(TEXT("BinaryReference"), {uint8(EType::BinaryReference)}, true, 33, EType::BinaryReference);
 
 	TestMeasure(TEXT("Hash"), {uint8(EType::Hash)}, true, 33, EType::Hash);
diff --git a/Engine/Source/Runtime/Core/Private/Tests/Serialization/CompactBinaryTest.cpp b/Engine/Source/Runtime/Core/Private/Tests/Serialization/CompactBinaryTest.cpp
--- a/Engine/Source/Runtime/Core/Private/Tests/Serialization/CompactBinaryTest.cpp
+++ b/Engine/Source/Runtime/Core/Private/Tests/Serialization/CompactBinaryTest.cpp
@@ -60,7 +60,7 @@ UE_CBFIELD_TYPE_ACCESSOR(Float32, IsFloat, AsFloat, float);
 UE_CBFIELD_TYPE_ACCESSOR(Float64, IsFloat, AsDouble, double);
 UE_CBFIELD_TYPE_ACCESSOR(BoolFalse, IsBool, AsBool, bool);
 UE_CBFIELD_TYPE_ACCESSOR(BoolTrue, IsBool, AsBool, bool);
-UE_CBFIELD_TYPE_ACCESSOR(Reference, IsReference, AsReference, FBlake3Hash);
+UE_CBFIELD_TYPE_ACCESSOR(CompactBinaryReference, IsCompactBinaryReference, AsCompactBinaryReference, FBlake3Hash);
 UE_CBFIELD_TYPE_ACCESSOR(BinaryReference, IsBinaryReference, AsBinaryReference, FBlake3Hash);
 UE_CBFIELD_TYPE_ACCESSOR(Hash, IsHash, AsHash, FBlake3Hash);
 UE_CBFIELD_TYPE_ACCESSOR_TYPED(Uuid, IsUuid, AsUuid, FGuid, const FGuid&);
@@ -941,28 +941,28 @@ bool FCbFieldBoolTest::RunTest(const FString& Parameters)
 	return true;
 }
 
-IMPLEMENT_CUSTOM_SIMPLE_AUTOMATION_TEST(FCbFieldReferenceTest, FCbFieldTestBase, "System.Core.Serialization.CbField.Reference", CompactBinaryTestFlags)
-bool FCbFieldReferenceTest::RunTest(const FString& Parameters)
+IMPLEMENT_CUSTOM_SIMPLE_AUTOMATION_TEST(FCbFieldCompactBinaryReferenceTest, FCbFieldTestBase, "System.Core.Serialization.CbField.CompactBinaryReference", CompactBinaryTestFlags)
+bool FCbFieldCompactBinaryReferenceTest::RunTest(const FString& Parameters)
 {
 	const FBlake3Hash::ByteArray ZeroBytes{};
 	const FBlake3Hash::ByteArray SequentialBytes{1, 2, 3, 4, 5, 6, 7, 8, 9, 10, 11, 12, 13, 14, 15, 16, 17, 18, 19, 20, 21, 22, 23, 24, 25, 26, 27, 28, 29, 30, 31, 32};
 
-	// Test FCbField(Reference, Zero)
-	TestField<ECbFieldType::Reference>(TEXT("Reference, Zero"), ZeroBytes);
+	// Test FCbField(CompactBinaryReference, Zero)
+	TestField<ECbFieldType::CompactBinaryReference>(TEXT("CompactBinaryReference, Zero"), ZeroBytes);
 
-	// Test FCbField(Reference, NonZero)
-	TestField<ECbFieldType::Reference>(TEXT("Reference, NonZero"), SequentialBytes, FBlake3Hash(SequentialBytes));
+	// Test FCbField(CompactBinaryReference, NonZero)
+	TestField<ECbFieldType::CompactBinaryReference>(TEXT("CompactBinaryReference, NonZero"), SequentialBytes, FBlake3Hash(SequentialBytes));
 
-	// Test FCbField(Reference, NonZero) AsAnyReference
+	// Test FCbField(CompactBinaryReference, NonZero) AsReference
 	{
-		FCbField Field(SequentialBytes, ECbFieldType::Reference);
-		TestFieldAsType(TEXT("Reference, NonZero, AsAnyReference"), Field, &FCbField::AsAnyReference, FBlake3Hash(SequentialBytes));
+		FCbField Field(SequentialBytes, ECbFieldType::CompactBinaryReference);
+		TestFieldAsType(TEXT("CompactBinaryReference, NonZero, AsReference"), Field, &FCbField::AsReference, FBlake3Hash(SequentialBytes));
 	}
 
-	// Test FCbField(None) as Reference
+	// Test FCbField(None) as CompactBinaryReference
 	{
 		FCbField DefaultField;
-		TestFieldError<ECbFieldType::Reference>(TEXT("Reference, None"), DefaultField, ECbFieldError::TypeError, FBlake3Hash(SequentialBytes));
+		TestFieldError<ECbFieldType::CompactBinaryReference>(TEXT("CompactBinaryReference, None"), DefaultField, ECbFieldError::TypeError, FBlake3Hash(SequentialBytes));
 	}
 
 	return true;
@@ -983,7 +983,7 @@ bool FCbFieldBinaryReferenceTest::RunTest(const FString& Parameters)
 	// Test FCbField(BinaryReference, NonZero) AsAnyReference
 	{
 		FCbField Field(SequentialBytes, ECbFieldType::BinaryReference);
-		TestFieldAsType(TEXT("BinaryReference, NonZero, AsAnyReference"), Field, &FCbField::AsAnyReference, FBlake3Hash(SequentialBytes));
+		TestFieldAsType(TEXT("BinaryReference, NonZero, AsReference"), Field, &FCbField::AsReference, FBlake3Hash(SequentialBytes));
 	}
 
 	// Test FCbField(None) as BinaryReference
@@ -1102,9 +1102,9 @@ bool FCbFieldIterateReferencesTest::RunTest(const FString& Parameters)
 		FCbWriter Writer;
 
 		Writer.Name("IgnoredTypeInRoot").Hash(MakeTestHash(100));
-		Writer.Reference(MakeTestHash(0));
+		Writer.CompactBinaryReference(MakeTestHash(0));
 		Writer.BinaryReference(MakeTestHash(1));
-		Writer.Name("RefInRoot").Reference(MakeTestHash(2));
+		Writer.Name("RefInRoot").CompactBinaryReference(MakeTestHash(2));
 		Writer.Name("BinRefInRoot").BinaryReference(MakeTestHash(3));
 
 		// Uniform array of type to ignore.
@@ -1142,7 +1142,7 @@ bool FCbFieldIterateReferencesTest::RunTest(const FString& Parameters)
 			Writer << false;
 			Writer.EndArray();
 			Writer.BeginArray();
-			Writer.Reference(MakeTestHash(10));
+			Writer.CompactBinaryReference(MakeTestHash(10));
 			Writer << false;
 			Writer.EndArray();
 		}
@@ -1151,12 +1151,12 @@ bool FCbFieldIterateReferencesTest::RunTest(const FString& Parameters)
 		Writer.BeginArray();
 		{
 			Writer.BeginObject();
-			Writer.Name("RefInUniObjInUniObj1").Reference(MakeTestHash(11));
-			Writer.Name("RefInUniObjInUniObj2").Reference(MakeTestHash(12));
+			Writer.Name("RefInUniObjInUniObj1").CompactBinaryReference(MakeTestHash(11));
+			Writer.Name("RefInUniObjInUniObj2").CompactBinaryReference(MakeTestHash(12));
 			Writer.EndObject();
 			Writer.BeginObject();
-			Writer.Name("RefInUniObjInUniObj3").Reference(MakeTestHash(13));
-			Writer.Name("RefInUniObjInUniObj4").Reference(MakeTestHash(14));
+			Writer.Name("RefInUniObjInUniObj3").CompactBinaryReference(MakeTestHash(13));
+			Writer.Name("RefInUniObjInUniObj4").CompactBinaryReference(MakeTestHash(14));
 			Writer.EndObject();
 		}
 		Writer.EndArray();
@@ -1168,7 +1168,7 @@ bool FCbFieldIterateReferencesTest::RunTest(const FString& Parameters)
 			Writer << "Bool" << false;
 			Writer.EndObject();
 			Writer.BeginObject();
-			Writer.Name("RefInNonUniObjInUniObj").Reference(MakeTestHash(15));
+			Writer.Name("RefInNonUniObjInUniObj").CompactBinaryReference(MakeTestHash(15));
 			Writer << "Bool" << false;
 			Writer.EndObject();
 		}
@@ -1213,7 +1213,7 @@ bool FCbFieldIterateReferencesTest::RunTest(const FString& Parameters)
 			Writer.EndArray();
 			Writer.Name("Array2");
 			Writer.BeginArray();
-			Writer.Reference(MakeTestHash(22));
+			Writer.CompactBinaryReference(MakeTestHash(22));
 			Writer << false;
 			Writer.EndArray();
 		}
@@ -1223,13 +1223,13 @@ bool FCbFieldIterateReferencesTest::RunTest(const FString& Parameters)
 		{
 			Writer.Name("Object1");
 			Writer.BeginObject();
-			Writer.Name("RefInUniObjInUniObj1").Reference(MakeTestHash(23));
-			Writer.Name("RefInUniObjInUniObj2").Reference(MakeTestHash(24));
+			Writer.Name("RefInUniObjInUniObj1").CompactBinaryReference(MakeTestHash(23));
+			Writer.Name("RefInUniObjInUniObj2").CompactBinaryReference(MakeTestHash(24));
 			Writer.EndObject();
 			Writer.Name("Object2");
 			Writer.BeginObject();
-			Writer.Name("RefInUniObjInUniObj3").Reference(MakeTestHash(25));
-			Writer.Name("RefInUniObjInUniObj4").Reference(MakeTestHash(26));
+			Writer.Name("RefInUniObjInUniObj3").CompactBinaryReference(MakeTestHash(25));
+			Writer.Name("RefInUniObjInUniObj4").CompactBinaryReference(MakeTestHash(26));
 			Writer.EndObject();
 		}
 		Writer.EndObject();
@@ -1243,7 +1243,7 @@ bool FCbFieldIterateReferencesTest::RunTest(const FString& Parameters)
 			Writer.EndObject();
 			Writer.Name("Object2");
 			Writer.BeginObject();
-			Writer.Name("RefInNonUniObjInUniObj").Reference(MakeTestHash(27));
+			Writer.Name("RefInNonUniObjInUniObj").CompactBinaryReference(MakeTestHash(27));
 			Writer << "Bool" << false;
 			Writer.EndObject();
 		}
@@ -1257,8 +1257,8 @@ bool FCbFieldIterateReferencesTest::RunTest(const FString& Parameters)
 	uint32 ReferenceIndex = 0;
 	Fields.IterateRangeReferences([this, &ReferenceIndex, &MakeTestHash](FCbField Field)
 		{
-			TestTrue(FString::Printf(TEXT("FCbField::IterateReferences(%u)::IsAnyReference"), ReferenceIndex), Field.IsAnyReference());
-			TestEqual(FString::Printf(TEXT("FCbField::IterateReferences(%u)"), ReferenceIndex), Field.AsReference(Field.AsBinaryReference()), MakeTestHash(ReferenceIndex));
+			TestTrue(FString::Printf(TEXT("FCbField::IterateReferences(%u)::IsAnyReference"), ReferenceIndex), Field.IsReference());
+			TestEqual(FString::Printf(TEXT("FCbField::IterateReferences(%u)"), ReferenceIndex), Field.AsCompactBinaryReference(Field.AsBinaryReference()), MakeTestHash(ReferenceIndex));
 			++ReferenceIndex;
 		});
 	TestEqual(TEXT("FCbField::IterateReferences"), ReferenceIndex, 28);
diff --git a/Engine/Source/Runtime/Core/Private/Tests/Serialization/CompactBinaryValidationTest.cpp b/Engine/Source/Runtime/Core/Private/Tests/Serialization/CompactBinaryValidationTest.cpp
--- a/Engine/Source/Runtime/Core/Private/Tests/Serialization/CompactBinaryValidationTest.cpp
+++ b/Engine/Source/Runtime/Core/Private/Tests/Serialization/CompactBinaryValidationTest.cpp
@@ -140,10 +140,10 @@ bool FCbValidateTest::RunTest(const FString& Parameters)
 	TestEqual(TEXT("ValidateCompactBinary(Valid, BoolFalse)"), Validate({uint8(ECbFieldType::BoolFalse)}), ECbValidateError::None);
 	TestEqual(TEXT("ValidateCompactBinary(Valid, BoolTrue)"), Validate({uint8(ECbFieldType::BoolTrue)}), ECbValidateError::None);
 
-	TestEqual(TEXT("ValidateCompactBinary(Valid, Reference)"), Validate({uint8(ECbFieldType::Reference), 0, 0, 0, 0, 0, 0, 0, 0, 0, 0, 0, 0, 0, 0, 0, 0, 0, 0, 0, 0, 0, 0, 0, 0, 0, 0, 0, 0, 0, 0, 0, 0}), ECbValidateError::None);
-	TestEqual(TEXT("ValidateCompactBinary(Valid, Reference, NoType)"), Validate({0, 0, 0, 0, 0, 0, 0, 0, 0, 0, 0, 0, 0, 0, 0, 0, 0, 0, 0, 0, 0, 0, 0, 0, 0, 0, 0, 0, 0, 0, 0, 0}, ECbFieldType::Reference), ECbValidateError::None);
-	TestEqual(TEXT("ValidateCompactBinary(OutOfBounds, Reference)"), Validate({uint8(ECbFieldType::Reference), 0, 0, 0, 0, 0, 0, 0, 0, 0, 0, 0, 0, 0, 0, 0, 0, 0, 0, 0, 0, 0, 0, 0, 0, 0, 0, 0, 0, 0, 0, 0}), ECbValidateError::OutOfBounds);
-	TestEqual(TEXT("ValidateCompactBinary(OutOfBounds, Reference, NoType)"), Validate({0, 0, 0, 0, 0, 0, 0, 0, 0, 0, 0, 0, 0, 0, 0, 0, 0, 0, 0, 0, 0, 0, 0, 0, 0, 0, 0, 0, 0, 0, 0}, ECbFieldType::Reference), ECbValidateError::OutOfBounds);
+	TestEqual(TEXT("ValidateCompactBinary(Valid, CompactBinaryReference)"), Validate({uint8(ECbFieldType::CompactBinaryReference), 0, 0, 0, 0, 0, 0, 0, 0, 0, 0, 0, 0, 0, 0, 0, 0, 0, 0, 0, 0, 0, 0, 0, 0, 0, 0, 0, 0, 0, 0, 0, 0}), ECbValidateError::None);
+	TestEqual(TEXT("ValidateCompactBinary(Valid, CompactBinaryReference, NoType)"), Validate({0, 0, 0, 0, 0, 0, 0, 0, 0, 0, 0, 0, 0, 0, 0, 0, 0, 0, 0, 0, 0, 0, 0, 0, 0, 0, 0, 0, 0, 0, 0, 0}, ECbFieldType::CompactBinaryReference), ECbValidateError::None);
+	TestEqual(TEXT("ValidateCompactBinary(OutOfBounds, CompactBinaryReference)"), Validate({uint8(ECbFieldType::CompactBinaryReference), 0, 0, 0, 0, 0, 0, 0, 0, 0, 0, 0, 0, 0, 0, 0, 0, 0, 0, 0, 0, 0, 0, 0, 0, 0, 0, 0, 0, 0, 0, 0}), ECbValidateError::OutOfBounds);
+	TestEqual(TEXT("ValidateCompactBinary(OutOfBounds, CompactBinaryReference, NoType)"), Validate({0, 0, 0, 0, 0, 0, 0, 0, 0, 0, 0, 0, 0, 0, 0, 0, 0, 0, 0, 0, 0, 0, 0, 0, 0, 0, 0, 0, 0, 0, 0}, ECbFieldType::CompactBinaryReference), ECbValidateError::OutOfBounds);
 
 	TestEqual(TEXT("ValidateCompactBinary(Valid, BinaryReference)"), Validate({uint8(ECbFieldType::BinaryReference), 0, 0, 0, 0, 0, 0, 0, 0, 0, 0, 0, 0, 0, 0, 0, 0, 0, 0, 0, 0, 0, 0, 0, 0, 0, 0, 0, 0, 0, 0, 0, 0}), ECbValidateError::None);
 	TestEqual(TEXT("ValidateCompactBinary(Valid, BinaryReference, NoType)"), Validate({0, 0, 0, 0, 0, 0, 0, 0, 0, 0, 0, 0, 0, 0, 0, 0, 0, 0, 0, 0, 0, 0, 0, 0, 0, 0, 0, 0, 0, 0, 0, 0}, ECbFieldType::BinaryReference), ECbValidateError::None);
@@ -351,7 +351,7 @@ bool FCbValidateAttachmentTest::RunTest(const FString& Parameters)
 	{
 		FBufferArchive Buffer;
 		FCbWriter Writer;
-		Writer.Reference(FBlake3::HashBuffer(BinaryView));
+		Writer.CompactBinaryReference(FBlake3::HashBuffer(BinaryView));
 		Writer.Binary(BinaryView);
 		Writer.Save(Buffer);
 		TestEqual(TEXT("ValidateCompactBinaryAttachment(HashBeforeValue)"), ValidateCompactBinaryAttachment(MakeMemoryView(Buffer), ECbValidateMode::All), ECbValidateError::InvalidPackageFormat | ECbValidateError::Padding);
@@ -361,7 +361,7 @@ bool FCbValidateAttachmentTest::RunTest(const FString& Parameters)
 		FBufferArchive Buffer;
 		FCbWriter Writer;
 		Writer.Name("Name"_ASV).Binary(BinaryView);
-		Writer.Reference(FBlake3::HashBuffer(BinaryView));
+		Writer.CompactBinaryReference(FBlake3::HashBuffer(BinaryView));
 		Writer.Save(Buffer);
 		TestEqual(TEXT("ValidateCompactBinaryAttachment(NameOnValue)"), ValidateCompactBinaryAttachment(MakeMemoryView(Buffer), ECbValidateMode::All), ECbValidateError::InvalidPackageFormat);
 		TestEqual(TEXT("ValidateCompactBinaryAttachment(NameOnValue, Mode)"), ValidateCompactBinaryAttachment(MakeMemoryView(Buffer), ECbValidateMode::All & ~ECbValidateMode::Package), ECbValidateError::None);
@@ -370,7 +370,7 @@ bool FCbValidateAttachmentTest::RunTest(const FString& Parameters)
 		FBufferArchive Buffer;
 		FCbWriter Writer;
 		Writer.Binary(BinaryView);
-		Writer.Name("Name"_ASV).Reference(FBlake3::HashBuffer(BinaryView));
+		Writer.Name("Name"_ASV).CompactBinaryReference(FBlake3::HashBuffer(BinaryView));
 		Writer.Save(Buffer);
 		TestEqual(TEXT("ValidateCompactBinaryAttachment(NameOnHash)"), ValidateCompactBinaryAttachment(MakeMemoryView(Buffer), ECbValidateMode::All), ECbValidateError::InvalidPackageFormat);
 		TestEqual(TEXT("ValidateCompactBinaryAttachment(NameOnHash, Mode)"), ValidateCompactBinaryAttachment(MakeMemoryView(Buffer), ECbValidateMode::All & ~ECbValidateMode::Package), ECbValidateError::None);
@@ -401,7 +401,7 @@ bool FCbValidatePackageTest::RunTest(const FString& Parameters)
 		Writer.BeginObject();
 		Writer << "Field" << 42;
 		Writer.EndObject();
-		Writer.Reference(Writer.Save().GetHash());
+		Writer.CompactBinaryReference(Writer.Save().GetHash());
 		Writer.Null();
 		Writer.Save(Buffer);
 		TestEqual(TEXT("ValidateCompactBinaryPackage(Object)"), ValidateCompactBinaryPackage(MakeMemoryView(Buffer), ECbValidateMode::All), ECbValidateError::None);
@@ -429,10 +429,10 @@ bool FCbValidatePackageTest::RunTest(const FString& Parameters)
 		Writer.BeginObject();
 		Writer << "Field" << 42;
 		Writer.EndObject();
-		Writer.Reference(Writer.Save().GetHash());
+		Writer.CompactBinaryReference(Writer.Save().GetHash());
 		const uint64 AttachmentOffset = Writer.GetSaveSize();
 		Writer.Binary(BinaryView);
-		Writer.Reference(FBlake3::HashBuffer(BinaryView));
+		Writer.CompactBinaryReference(FBlake3::HashBuffer(BinaryView));
 		Writer.Null();
 		Writer.Save(Buffer);
 		TestEqual(TEXT("ValidateCompactBinaryPackage(Object+Attachment)"), ValidateCompactBinaryPackage(MakeMemoryView(Buffer), ECbValidateMode::All), ECbValidateError::None);
@@ -465,7 +465,7 @@ bool FCbValidatePackageTest::RunTest(const FString& Parameters)
 		Writer.BeginObject();
 		Writer << "Field" << 42;
 		Writer.EndObject();
-		Writer.Reference(Writer.Save().GetHash());
+		Writer.CompactBinaryReference(Writer.Save().GetHash());
 		Writer.Null();
 		Writer.Save(Buffer);
 		TestEqual(TEXT("ValidateCompactBinaryPackage(NameOnObject)"), ValidateCompactBinaryPackage(MakeMemoryView(Buffer), ECbValidateMode::All), ECbValidateError::InvalidPackageFormat);
@@ -479,7 +479,7 @@ bool FCbValidatePackageTest::RunTest(const FString& Parameters)
 		Writer.EndObject();
 		const FBlake3Hash ObjectHash = Writer.Save().GetHash();
 		Writer.Name("Name"_ASV);
-		Writer.Reference(ObjectHash);
+		Writer.CompactBinaryReference(ObjectHash);
 		Writer.Null();
 		Writer.Save(Buffer);
 		TestEqual(TEXT("ValidateCompactBinaryPackage(NameOnHash)"), ValidateCompactBinaryPackage(MakeMemoryView(Buffer), ECbValidateMode::All), ECbValidateError::InvalidPackageFormat);
@@ -491,9 +491,9 @@ bool FCbValidatePackageTest::RunTest(const FString& Parameters)
 		Writer.BeginObject();
 		Writer << "Field" << 42;
 		Writer.EndObject();
-		Writer.Reference(Writer.Save().GetHash());
+		Writer.CompactBinaryReference(Writer.Save().GetHash());
 		Writer.Name("Name"_ASV).Binary(BinaryView);
-		Writer.Reference(FBlake3::HashBuffer(BinaryView));
+		Writer.CompactBinaryReference(FBlake3::HashBuffer(BinaryView));
 		Writer.Null();
 		Writer.Save(Buffer);
 		TestEqual(TEXT("ValidateCompactBinaryPackage(NameOnAttachment)"), ValidateCompactBinaryPackage(MakeMemoryView(Buffer), ECbValidateMode::All), ECbValidateError::InvalidPackageFormat);
@@ -506,11 +506,11 @@ bool FCbValidatePackageTest::RunTest(const FString& Parameters)
 		Writer << "Field" << 42;
 		Writer.EndObject();
 		const FBlake3Hash ObjectHash = Writer.Save().GetHash();
-		Writer.Reference(ObjectHash);
+		Writer.CompactBinaryReference(ObjectHash);
 		Writer.BeginObject();
 		Writer << "Field" << 42;
 		Writer.EndObject();
-		Writer.Reference(ObjectHash);
+		Writer.CompactBinaryReference(ObjectHash);
 		Writer.Null();
 		Writer.Save(Buffer);
 		TestEqual(TEXT("ValidateCompactBinaryPackage(MultipleObjects)"), ValidateCompactBinaryPackage(MakeMemoryView(Buffer), ECbValidateMode::All), ECbValidateError::MultiplePackageObjects);
@@ -520,9 +520,9 @@ bool FCbValidatePackageTest::RunTest(const FString& Parameters)
 		FBufferArchive Buffer;
 		TCbWriter<256> Writer;
 		Writer.Binary(BinaryView);
-		Writer.Reference(FBlake3::HashBuffer(BinaryView));
+		Writer.CompactBinaryReference(FBlake3::HashBuffer(BinaryView));
 		Writer.Binary(BinaryView);
-		Writer.Reference(FBlake3::HashBuffer(BinaryView));
+		Writer.CompactBinaryReference(FBlake3::HashBuffer(BinaryView));
 		Writer.Null();
 		Writer.Save(Buffer);
 		TestEqual(TEXT("ValidateCompactBinaryPackage(DuplicateAttachments)"), ValidateCompactBinaryPackage(MakeMemoryView(Buffer), ECbValidateMode::All), ECbValidateError::DuplicateAttachments);
@@ -533,7 +533,7 @@ bool FCbValidatePackageTest::RunTest(const FString& Parameters)
 		TCbWriter<256> Writer;
 		Writer.BeginObject();
 		Writer.EndObject();
-		Writer.Reference(FCbObject().GetHash());
+		Writer.CompactBinaryReference(FCbObject().GetHash());
 		Writer.Null();
 		Writer.Save(Buffer);
 		TestEqual(TEXT("ValidateCompactBinaryPackage(NullObject)"), ValidateCompactBinaryPackage(MakeMemoryView(Buffer), ECbValidateMode::All), ECbValidateError::NullPackageObject);
diff --git a/Engine/Source/Runtime/Core/Private/Tests/Serialization/CompactBinaryWriterTest.cpp b/Engine/Source/Runtime/Core/Private/Tests/Serialization/CompactBinaryWriterTest.cpp
--- a/Engine/Source/Runtime/Core/Private/Tests/Serialization/CompactBinaryWriterTest.cpp
+++ b/Engine/Source/Runtime/Core/Private/Tests/Serialization/CompactBinaryWriterTest.cpp
@@ -563,8 +563,8 @@ bool FCbWriterBoolTest::RunTest(const FString& Parameters)
 
 ///////////////////////////////////////////////////////////////////////////////////////////////////////////////////////
 
-IMPLEMENT_SIMPLE_AUTOMATION_TEST(FCbWriterReferenceTest, "System.Core.Serialization.CbWriter.Reference", CompactBinaryWriterTestFlags)
-bool FCbWriterReferenceTest::RunTest(const FString& Parameters)
+IMPLEMENT_SIMPLE_AUTOMATION_TEST(FCbWriterCompactBinaryReferenceTest, "System.Core.Serialization.CbWriter.CompactBinaryReference", CompactBinaryWriterTestFlags)
+bool FCbWriterCompactBinaryReferenceTest::RunTest(const FString& Parameters)
 {
 	TCbWriter<256> Writer;
 
@@ -574,17 +574,17 @@ bool FCbWriterReferenceTest::RunTest(const FString& Parameters)
 	const FBlake3Hash Values[] = { FBlake3Hash(ZeroBytes), FBlake3Hash(SequentialBytes) };
 	for (const FBlake3Hash& Value : Values)
 	{
-		Writer.Reference(Value);
+		Writer.CompactBinaryReference(Value);
 	}
 
 	FCbFieldRefIterator Fields = Writer.Save();
-	if (TestEqual(TEXT("FCbWriter(Reference) Validate"), ValidateCompactBinaryRange(Fields.GetBuffer(), ECbValidateMode::All), ECbValidateError::None))
+	if (TestEqual(TEXT("FCbWriter(CompactBinaryReference) Validate"), ValidateCompactBinaryRange(Fields.GetBuffer(), ECbValidateMode::All), ECbValidateError::None))
 	{
 		const FBlake3Hash* CheckValue = Values;
 		for (FCbField Field : Fields)
 		{
-			TestEqual(TEXT("FCbWriter(Reference).AsReference()"), Field.AsReference(), *CheckValue++);
-			TestFalse(TEXT("FCbWriter(Reference) Error"), Field.HasError());
+			TestEqual(TEXT("FCbWriter(CompactBinaryReference).AsCompactBinaryReference()"), Field.AsCompactBinaryReference(), *CheckValue++);
+			TestFalse(TEXT("FCbWriter(CompactBinaryReference) Error"), Field.HasError());
 		}
 	}
 
@@ -816,7 +816,7 @@ bool FCbWriterComplexTest::RunTest(const FString& Parameters)
 		Writer.Name("False"_ASV).Bool(false);
 		Writer.Name("True"_ASV).Bool(true);
 
-		Writer.Name("Reference"_ASV).Reference(FBlake3Hash());
+		Writer.Name("CompactBinaryReference"_ASV).CompactBinaryReference(FBlake3Hash());
 		Writer.Name("BinaryReference"_ASV).BinaryReference(FBlake3Hash());
 
 		Writer.Name("Hash"_ASV).Hash(FBlake3Hash());
diff --git a/Engine/Source/Runtime/Core/Public/Serialization/CompactBinary.h b/Engine/Source/Runtime/Core/Public/Serialization/CompactBinary.h
--- a/Engine/Source/Runtime/Core/Public/Serialization/CompactBinary.h
+++ b/Engine/Source/Runtime/Core/Public/Serialization/CompactBinary.h
@@ -145,15 +145,15 @@ enum class ECbFieldType : uint8
 	BoolTrue        = 0x0d,
 
 	/**
-	 * Reference is a reference to compact binary stored externally.
+	 * CompactBinaryReference is a reference to compact binary data stored externally.
 	 *
-	 * Payload is a 256-bit hash digest of the compact binary.
+	 * Payload is a 256-bit hash digest of the compact binary data.
 	 */
-	Reference       = 0x0e,
+	CompactBinaryReference = 0x0e,
 	/**
-	 * BinaryReference is a reference to a blob stored externally.
+	 * BinaryReference is a reference to binary data stored externally.
 	 *
-	 * Payload is a 256-bit hash digest of the blob.
+	 * Payload is a 256-bit hash digest of the binary data.
 	 */
 	BinaryReference = 0x0f,
 
@@ -212,8 +212,8 @@ class FCbFieldType
 	static constexpr ECbFieldType BoolMask              = ECbFieldType(0b0011'1110);
 	static constexpr ECbFieldType BoolBase              = ECbFieldType(0b0000'1100);
 
-	static constexpr ECbFieldType AnyReferenceMask      = ECbFieldType(0b0011'1110);
-	static constexpr ECbFieldType AnyReferenceBase      = ECbFieldType(0b0000'1110);
+	static constexpr ECbFieldType ReferenceMask      = ECbFieldType(0b0011'1110);
+	static constexpr ECbFieldType ReferenceBase      = ECbFieldType(0b0000'1110);
 
 	static void StaticAssertTypeConstants();
 
@@ -240,9 +240,9 @@ public:
 	static constexpr inline bool IsFloat(ECbFieldType Type)      { return (Type & FloatMask) == FloatBase; }
 	static constexpr inline bool IsBool(ECbFieldType Type)       { return (Type & BoolMask) == BoolBase; }
 
-	static constexpr inline bool IsReference(ECbFieldType Type)       { return GetType(Type) == ECbFieldType::Reference; }
-	static constexpr inline bool IsBinaryReference(ECbFieldType Type) { return GetType(Type) == ECbFieldType::BinaryReference; }
-	static constexpr inline bool IsAnyReference(ECbFieldType Type)    { return (Type & AnyReferenceMask) == AnyReferenceBase; }
+	static constexpr inline bool IsCompactBinaryReference(ECbFieldType Type) { return GetType(Type) == ECbFieldType::CompactBinaryReference; }
+	static constexpr inline bool IsBinaryReference(ECbFieldType Type)        { return GetType(Type) == ECbFieldType::BinaryReference; }
+	static constexpr inline bool IsReference(ECbFieldType Type)   { return (Type & ReferenceMask) == ReferenceBase; }
 
 	static constexpr inline bool IsHash(ECbFieldType Type)       { return GetType(Type) == ECbFieldType::Hash; }
 	static constexpr inline bool IsUuid(ECbFieldType Type)       { return GetType(Type) == ECbFieldType::Uuid; }
@@ -254,7 +254,7 @@ public:
 	static constexpr inline bool MayContainReferences(ECbFieldType Type)
 	{
 		// The use of !! will suppress V792 from static analysis. Using //-V792 did not work.
-		return !!IsObject(Type) | !!IsArray(Type) | !!IsAnyReference(Type);
+		return !!IsObject(Type) | !!IsArray(Type) | !!IsReference(Type);
 	}
 };
 
@@ -495,12 +495,12 @@ public:
 	/** Access the field as a bool. Returns the provided default on error. */
 	CORE_API bool AsBool(bool bDefault = false);
 
-	/** Access the field as a hash referencing compact binary. Returns the provided default on error. */
-	CORE_API FBlake3Hash AsReference(const FBlake3Hash& Default = FBlake3Hash());
-	/** Access the field as a hash referencing a blob. Returns the provided default on error. */
+	/** Access the field as a hash referencing compact binary data. Returns the provided default on error. */
+	CORE_API FBlake3Hash AsCompactBinaryReference(const FBlake3Hash& Default = FBlake3Hash());
+	/** Access the field as a hash referencing binary data. Returns the provided default on error. */
 	CORE_API FBlake3Hash AsBinaryReference(const FBlake3Hash& Default = FBlake3Hash());
-	/** Access the field as a hash referencing a blob or compact binary. Returns the provided default on error. */
-	CORE_API FBlake3Hash AsAnyReference(const FBlake3Hash& Default = FBlake3Hash());
+	/** Access the field as a hash referencing compact binary data or binary data. Returns the provided default on error. */
+	CORE_API FBlake3Hash AsReference(const FBlake3Hash& Default = FBlake3Hash());
 
 	/** Access the field as a hash. Returns the provided default on error. */
 	CORE_API FBlake3Hash AsHash(const FBlake3Hash& Default = FBlake3Hash());
@@ -543,9 +543,9 @@ public:
 	constexpr inline bool IsFloat() const           { return FCbFieldType::IsFloat(Type); }
 	constexpr inline bool IsBool() const            { return FCbFieldType::IsBool(Type); }
 
-	constexpr inline bool IsReference() const       { return FCbFieldType::IsReference(Type); }
+	constexpr inline bool IsCompactBinaryReference() const { return FCbFieldType::IsCompactBinaryReference(Type); }
 	constexpr inline bool IsBinaryReference() const { return FCbFieldType::IsBinaryReference(Type); }
-	constexpr inline bool IsAnyReference() const    { return FCbFieldType::IsAnyReference(Type); }
+	constexpr inline bool IsReference() const       { return FCbFieldType::IsReference(Type); }
 
 	constexpr inline bool IsHash() const            { return FCbFieldType::IsHash(Type); }
 	constexpr inline bool IsUuid() const            { return FCbFieldType::IsUuid(Type); }
diff --git a/Engine/Source/Runtime/Core/Public/Serialization/CompactBinaryPackage.h b/Engine/Source/Runtime/Core/Public/Serialization/CompactBinaryPackage.h
--- a/Engine/Source/Runtime/Core/Public/Serialization/CompactBinaryPackage.h
+++ b/Engine/Source/Runtime/Core/Public/Serialization/CompactBinaryPackage.h
@@ -20,8 +20,8 @@ template <typename FuncType> class TFunctionRef;
  * A compact binary attachment is also a valid binary attachment and may be accessed as binary.
  *
  * Attachments are serialized as one or two compact binary fields with no name. A Binary field is
- * written first with its content. The content hash non-zero as a Reference or BinaryReference if
- * the content size is non-zero, and is otherwise omitted.
+ * written first with its content. The content hash is omitted when the content size is zero, and
+ * is otherwise written as a BinaryReference or CompactBinaryReference depending on the type.
  */
 class FCbAttachment
 {
@@ -144,9 +144,11 @@ inline uint32 GetTypeHash(const FCbAttachment& Attachment)
  *
  * A package is serialized as a sequence of compact binary fields with no name. The object may be
  * both preceded and followed by attachments. The object itself is written as an Object field and
- * then its hash in a Reference field. The end of the package is marked by a Null field. Standard
- * ordering of components is the object, then the attachments ordered by hash, though this is not
- * a requirement, and it is valid for a package to contain components in an any order.
+ * followed by its hash in a CompactBinaryReference field when the object is non-empty. A package
+ * ends with a Null field. The canonical order of components is the object and its hash, followed
+ * by the attachments ordered by hash, followed by a Null field. It is valid for the a package to
+ * have its components serialized in any order, provided there is at most one object and the null
+ * field is written last.
  */
 class FCbPackage
 {
diff --git a/Engine/Source/Runtime/Core/Public/Serialization/CompactBinaryWriter.h b/Engine/Source/Runtime/Core/Public/Serialization/CompactBinaryWriter.h
--- a/Engine/Source/Runtime/Core/Public/Serialization/CompactBinaryWriter.h
+++ b/Engine/Source/Runtime/Core/Public/Serialization/CompactBinaryWriter.h
@@ -165,9 +165,9 @@ public:
 	/** Write a bool field. */
 	CORE_API void Bool(bool bValue);
 
-	/** Write a field referencing compact binary. */
-	CORE_API void Reference(const FBlake3Hash& Value);
-	/** Write a field referencing a blob. */
+	/** Write a field referencing compact binary data by its hash. */
+	CORE_API void CompactBinaryReference(const FBlake3Hash& Value);
+	/** Write a field referencing binary data by its hash. */
 	CORE_API void BinaryReference(const FBlake3Hash& Value);
 
 	/** Write a hash field. */
