U5 - disable PVS Warnings

diff --git a/Engine/Source/Runtime/Core/Public/Experimental/Async/AwaitableTask.h b/Engine/Source/Runtime/Core/Public/Experimental/Async/AwaitableTask.h
--- a/Engine/Source/Runtime/Core/Public/Experimental/Async/AwaitableTask.h
+++ b/Engine/Source/Runtime/Core/Public/Experimental/Async/AwaitableTask.h
@@ -132,14 +132,14 @@ namespace AwaitableTask_Detail
 	template<typename ReturnType>
 	inline TPromiseVTableBase<ReturnType>& TPromiseVTableBase<ReturnType>::operator= (const TPromiseVTableBase<ReturnType>& Other)
 	{
-		memcpy((void*)this, (void*)&Other, sizeof(TPromiseVTableBase<ReturnType>));
+		memcpy((void*)this, (void*)&Other, sizeof(TPromiseVTableBase<ReturnType>)); //-V598
 		return *this;
 	}
 
 	template<typename ReturnType>
 	inline TPromiseVTableBase<ReturnType>& TPromiseVTableBase<ReturnType>::operator= (TPromiseVTableBase<ReturnType>&& Other)
 	{
-		memcpy((void*)this, (void*)&Other, sizeof(TPromiseVTableBase<ReturnType>));
+		memcpy((void*)this, (void*)&Other, sizeof(TPromiseVTableBase<ReturnType>)); //-V598
 		new (&Other) TPromiseVTableDummy<ReturnType>();
 		return *this;
 	}
