Create a separate thread pool for the I/O portion of DDC requests to avoid clogging the workers with long I/O waits
Add queued work priority support to auto delete async tasks and schedule DDC Puts as Low priority
ThreadPool threads are now honoring the names they've been given on creation
Move the I/O only async DDC Get into the DDC ThreadPool with higher priority (Normal) than DCC Puts

diff --git a/Engine/Source/Developer/DerivedDataCache/Private/DerivedDataBackendAsyncPutWrapper.cpp b/Engine/Source/Developer/DerivedDataCache/Private/DerivedDataBackendAsyncPutWrapper.cpp
--- a/Engine/Source/Developer/DerivedDataCache/Private/DerivedDataBackendAsyncPutWrapper.cpp
+++ b/Engine/Source/Developer/DerivedDataCache/Private/DerivedDataBackendAsyncPutWrapper.cpp
@@ -262,7 +262,7 @@ FDerivedDataBackendInterface::EPutStatus FDerivedDataBackendAsyncPutWrapper::Put
 	UE_LOG(LogDerivedDataCache, Verbose, TEXT("%s queueing %s for put"), *GetName(), CacheKey);
 
 	FDerivedDataBackend::Get().AddToAsyncCompletionCounter(1);
-	(new FAutoDeleteAsyncTask<FCachePutAsyncWorker>(CacheKey, InData, InnerBackend, bPutEvenIfExists, InflightCache.Get(), &FilesInFlight, UsageStats))->StartBackgroundTask();
+	(new FAutoDeleteAsyncTask<FCachePutAsyncWorker>(CacheKey, InData, InnerBackend, bPutEvenIfExists, InflightCache.Get(), &FilesInFlight, UsageStats))->StartBackgroundTask(GDDCIOThreadPool, EQueuedWorkPriority::Low);
 
 	return EPutStatus::Executing;
 }
diff --git a/Engine/Source/Developer/DerivedDataCache/Private/DerivedDataCache.cpp b/Engine/Source/Developer/DerivedDataCache/Private/DerivedDataCache.cpp
--- a/Engine/Source/Developer/DerivedDataCache/Private/DerivedDataCache.cpp
+++ b/Engine/Source/Developer/DerivedDataCache/Private/DerivedDataCache.cpp
@@ -505,7 +505,8 @@ public:
 		check(!PendingTasks.Contains(Handle));
 		PendingTasks.Add(Handle, AsyncTask);
 		AddToAsyncCompletionCounter(1);
-		AsyncTask->StartBackgroundTask();
+		// This request is I/O only, doesn't do any processing, send it to the I/O only thread-pool to avoid wasting worker threads on long I/O waits.
+		AsyncTask->StartBackgroundTask(GDDCIOThreadPool);
 		return Handle;
 	}
 
diff --git a/Engine/Source/Runtime/Core/Private/HAL/ThreadingBase.cpp b/Engine/Source/Runtime/Core/Private/HAL/ThreadingBase.cpp
--- a/Engine/Source/Runtime/Core/Private/HAL/ThreadingBase.cpp
+++ b/Engine/Source/Runtime/Core/Private/HAL/ThreadingBase.cpp
@@ -37,6 +37,8 @@ FQueuedThreadPool* GThreadPool = nullptr;
 
 FQueuedThreadPool* GIOThreadPool = nullptr;
 
+FQueuedThreadPool* GDDCIOThreadPool = nullptr;
+
 FQueuedThreadPool* GBackgroundPriorityThreadPool = nullptr;
 
 #if WITH_EDITOR
@@ -940,11 +942,10 @@ public:
 	 * @param ThreadPriority priority of new thread
 	 * @return True if the thread and all of its initialization was successful, false otherwise
 	 */
-	virtual bool Create(class FQueuedThreadPoolBase* InPool,uint32 InStackSize = 0,EThreadPriority ThreadPriority=TPri_Normal)
+	virtual bool Create(class FQueuedThreadPoolBase* InPool,uint32 InStackSize = 0, EThreadPriority ThreadPriority=TPri_Normal, const TCHAR* ThreadName = nullptr)
 	{
 		static int32 PoolThreadIndex = 0;
-		const FString PoolThreadName = FString::Printf( TEXT( "PoolThread %d" ), PoolThreadIndex );
-		PoolThreadIndex++;
+		const FString PoolThreadName = ThreadName ? FString(ThreadName) : FString::Printf( TEXT( "PoolThread %d" ), PoolThreadIndex++ );
 
 		OwningThreadPool = InPool;
 		DoWorkEvent = FPlatformProcess::GetSynchEventFromPool();
@@ -1061,7 +1062,8 @@ public:
 			// Create a new queued thread
 			FQueuedThread* pThread = new FQueuedThread();
 			// Now create the thread and add it if ok
-			if (pThread->Create(this,StackSize,ThreadPriority) == true)
+			const FString ThreadName = FString::Printf(TEXT("%s #%d"), Name, Count);
+			if (pThread->Create(this, StackSize, ThreadPriority, *ThreadName) == true)
 			{
 				QueuedThreads.Add(pThread);
 				AllThreads.Add(pThread);
diff --git a/Engine/Source/Runtime/Core/Public/Async/AsyncWork.h b/Engine/Source/Runtime/Core/Public/Async/AsyncWork.h
--- a/Engine/Source/Runtime/Core/Public/Async/AsyncWork.h
+++ b/Engine/Source/Runtime/Core/Public/Async/AsyncWork.h
@@ -71,7 +71,7 @@ class FAutoDeleteAsyncTask
 	/* Generic start function, not called directly
 	 * @param bForceSynchronous if true, this job will be started synchronously, now, on this thread
 	 **/
-	void Start(bool bForceSynchronous, FQueuedThreadPool* InQueuedPool)
+	void Start(bool bForceSynchronous, FQueuedThreadPool* InQueuedPool, EQueuedWorkPriority InPriority = EQueuedWorkPriority::Normal)
 	{
 		LLM(InheritedLLMTag = FLowLevelMemTracker::bIsDisabled ? nullptr : FLowLevelMemTracker::Get().GetActiveTagData(ELLMTracker::Default));
 #if USE_MEMORY_TRACE_TAGS
@@ -86,7 +86,7 @@ class FAutoDeleteAsyncTask
 		}
 		if (QueuedPool)
 		{
-			QueuedPool->AddQueuedWork(this);
+			QueuedPool->AddQueuedWork(this, InPriority);
 		}
 		else
 		{
@@ -145,15 +145,15 @@ public:
 	**/
 	void StartSynchronousTask()
 	{
-		Start(true, GThreadPool);
+		Start(true, nullptr);
 	}
 
 	/** 
 	* Run this task on the lo priority thread pool. It is not safe to use this object after this call.
 	**/
-	void StartBackgroundTask(FQueuedThreadPool* InQueuedPool = GThreadPool)
+	void StartBackgroundTask(FQueuedThreadPool* InQueuedPool = GThreadPool, EQueuedWorkPriority InPriority = EQueuedWorkPriority::Normal)
 	{
-		Start(false, InQueuedPool);
+		Start(false, InQueuedPool, InPriority);
 	}
 };
 
diff --git a/Engine/Source/Runtime/Core/Public/Misc/QueuedThreadPool.h b/Engine/Source/Runtime/Core/Public/Misc/QueuedThreadPool.h
--- a/Engine/Source/Runtime/Core/Public/Misc/QueuedThreadPool.h
+++ b/Engine/Source/Runtime/Core/Public/Misc/QueuedThreadPool.h
@@ -164,6 +164,8 @@ extern CORE_API FQueuedThreadPool* GThreadPool;
 
 extern CORE_API FQueuedThreadPool* GIOThreadPool;
 
+extern CORE_API FQueuedThreadPool* GDDCIOThreadPool;
+
 extern CORE_API FQueuedThreadPool* GBackgroundPriorityThreadPool;
 
 #if WITH_EDITOR
diff --git a/Engine/Source/Runtime/Launch/Private/LaunchEngineLoop.cpp b/Engine/Source/Runtime/Launch/Private/LaunchEngineLoop.cpp
--- a/Engine/Source/Runtime/Launch/Private/LaunchEngineLoop.cpp
+++ b/Engine/Source/Runtime/Launch/Private/LaunchEngineLoop.cpp
@@ -2350,6 +2350,18 @@ PRAGMA_ENABLE_DEPRECATION_WARNINGS
 			}
 			verify(GIOThreadPool->Create(NumThreadsInThreadPool, 96 * 1024, TPri_AboveNormal, TEXT("IOThreadPool")));
 		}
+ 
+		if (!FPlatformProperties::RequiresCookedData())
+		{
+			SCOPED_BOOT_TIMING("GDDCIOThreadPool->Create");
+			GDDCIOThreadPool = FQueuedThreadPool::Allocate();
+			int32 NumThreadsInThreadPool = FPlatformMisc::NumberOfIOWorkerThreadsToSpawn();
+			if (FPlatformProperties::IsServerOnly())
+			{
+				NumThreadsInThreadPool = 2;
+			}
+			verify(GDDCIOThreadPool->Create(NumThreadsInThreadPool, 96 * 1024, TPri_AboveNormal, TEXT("DDC IO ThreadPool")));
+		}
 	}
 
 	FEmbeddedCommunication::ForceTick(1);
