BLAKE3: Renamed Bytes() to GetBytes() for consistency

#rnx

diff --git a/Engine/Source/Runtime/Core/Private/Serialization/CompactBinaryWriter.cpp b/Engine/Source/Runtime/Core/Private/Serialization/CompactBinaryWriter.cpp
--- a/Engine/Source/Runtime/Core/Private/Serialization/CompactBinaryWriter.cpp
+++ b/Engine/Source/Runtime/Core/Private/Serialization/CompactBinaryWriter.cpp
@@ -480,21 +480,21 @@ void FCbWriter::Bool(const bool bValue)
 void FCbWriter::CompactBinaryReference(const FBlake3Hash& Value)
 {
 	BeginField();
-	Data.Append(Value.Bytes(), sizeof(decltype(Value.Bytes())));
+	Data.Append(Value.GetBytes(), sizeof(decltype(Value.GetBytes())));
 	EndField(ECbFieldType::CompactBinaryReference);
 }
 
 void FCbWriter::BinaryReference(const FBlake3Hash& Value)
 {
 	BeginField();
-	Data.Append(Value.Bytes(), sizeof(decltype(Value.Bytes())));
+	Data.Append(Value.GetBytes(), sizeof(decltype(Value.GetBytes())));
 	EndField(ECbFieldType::BinaryReference);
 }
 
 void FCbWriter::Hash(const FBlake3Hash& Value)
 {
 	BeginField();
-	Data.Append(Value.Bytes(), sizeof(decltype(Value.Bytes())));
+	Data.Append(Value.GetBytes(), sizeof(decltype(Value.GetBytes())));
 	EndField(ECbFieldType::Hash);
 }
 
diff --git a/Engine/Source/Runtime/Core/Public/Misc/Blake3.h b/Engine/Source/Runtime/Core/Public/Misc/Blake3.h
--- a/Engine/Source/Runtime/Core/Public/Misc/Blake3.h
+++ b/Engine/Source/Runtime/Core/Public/Misc/Blake3.h
@@ -32,7 +32,7 @@ public:
 	inline void Reset() { *this = FBlake3Hash(); }
 
 	/** Returns a reference to the raw byte array for the hash. */
-	inline const ByteArray& Bytes() const { return Hash; }
+	inline const ByteArray& GetBytes() const { return Hash; }
 
 	/** Returns whether this is a zero hash. */
 	inline bool IsZero() const;
@@ -103,32 +103,32 @@ inline bool FBlake3Hash::IsZero() const
 
 inline bool operator==(const FBlake3Hash& A, const FBlake3Hash& B)
 {
-	return FMemory::Memcmp(A.Bytes(), B.Bytes(), sizeof(decltype(A.Bytes()))) == 0;
+	return FMemory::Memcmp(A.GetBytes(), B.GetBytes(), sizeof(decltype(A.GetBytes()))) == 0;
 }
 
 inline bool operator!=(const FBlake3Hash& A, const FBlake3Hash& B)
 {
-	return FMemory::Memcmp(A.Bytes(), B.Bytes(), sizeof(decltype(A.Bytes()))) != 0;
+	return FMemory::Memcmp(A.GetBytes(), B.GetBytes(), sizeof(decltype(A.GetBytes()))) != 0;
 }
 
 inline bool operator<(const FBlake3Hash& A, const FBlake3Hash& B)
 {
-	return FMemory::Memcmp(A.Bytes(), B.Bytes(), sizeof(decltype(A.Bytes()))) < 0;
+	return FMemory::Memcmp(A.GetBytes(), B.GetBytes(), sizeof(decltype(A.GetBytes()))) < 0;
 }
 
 inline uint32 GetTypeHash(const FBlake3Hash& Hash)
 {
-	return *reinterpret_cast<const uint32*>(Hash.Bytes());
+	return *reinterpret_cast<const uint32*>(Hash.GetBytes());
 }
 
 inline FAnsiStringBuilderBase& operator<<(FAnsiStringBuilderBase& Builder, const FBlake3Hash& Hash)
 {
-	UE::String::BytesToHexLower(Hash.Bytes(), Builder);
+	UE::String::BytesToHexLower(Hash.GetBytes(), Builder);
 	return Builder;
 }
 
 inline FWideStringBuilderBase& operator<<(FWideStringBuilderBase& Builder, const FBlake3Hash& Hash)
 {
-	UE::String::BytesToHexLower(Hash.Bytes(), Builder);
+	UE::String::BytesToHexLower(Hash.GetBytes(), Builder);
 	return Builder;
 }
