Respect GIgnoreDebugger on Apple platforms.

#rnx

diff --git a/Engine/Source/Runtime/Core/Public/Apple/ApplePlatformMisc.h b/Engine/Source/Runtime/Core/Public/Apple/ApplePlatformMisc.h
--- a/Engine/Source/Runtime/Core/Public/Apple/ApplePlatformMisc.h
+++ b/Engine/Source/Runtime/Core/Public/Apple/ApplePlatformMisc.h
@@ -22,6 +22,8 @@
 
 #define UE_DEBUG_BREAK_IMPL() PLATFORM_BREAK()
 
+extern CORE_API bool GIgnoreDebugger;
+
 #ifdef __OBJC__
 
 class FScopeAutoreleasePool
@@ -64,6 +66,11 @@ struct CORE_API FApplePlatformMisc : public FGenericPlatformMisc
 	{
 		// Based on http://developer.apple.com/library/mac/#qa/qa1361/_index.html
 
+		if (GIgnoreDebugger)
+		{
+			return false;
+		}
+
 		struct kinfo_proc Info;
 		int32 Mib[] = { CTL_KERN, KERN_PROC, KERN_PROC_PID, getpid() };
 		SIZE_T Size = sizeof(Info);
