Use SymFromAddr() instead of SymGetSymFromAddr64()

Fix runtime Access violation coming from dbghelp.dll when opening Test Automation tab from Editor in recent Visual Studio Community 2019 version 16.8.3

As noted in microsoft docs https://docs.microsoft.com/en-us/windows/win32/api/dbghelp/nf-dbghelp-symgetsymfromaddr applications should use SymFromAddr()

Pull request #7592 from tteoz

diff --git a/Engine/Source/Runtime/Core/Private/Windows/WindowsPlatformStackWalk.cpp b/Engine/Source/Runtime/Core/Private/Windows/WindowsPlatformStackWalk.cpp
--- a/Engine/Source/Runtime/Core/Private/Windows/WindowsPlatformStackWalk.cpp
+++ b/Engine/Source/Runtime/Core/Private/Windows/WindowsPlatformStackWalk.cpp
@@ -447,13 +447,13 @@ void FWindowsPlatformStackWalk::ProgramCounterToSymbolInfo( uint64 ProgramCounte
 	HANDLE ProcessHandle = GProcessHandle;
 
 	// Initialize symbol.
-	ANSICHAR SymbolBuffer[sizeof( IMAGEHLP_SYMBOL64 ) + FProgramCounterSymbolInfo::MAX_NAME_LENGTH] = {0};
-	IMAGEHLP_SYMBOL64* Symbol = (IMAGEHLP_SYMBOL64*)SymbolBuffer;
-	Symbol->SizeOfStruct = sizeof(SymbolBuffer);
-	Symbol->MaxNameLength = FProgramCounterSymbolInfo::MAX_NAME_LENGTH;
+	ANSICHAR SymbolBuffer[sizeof( SYMBOL_INFO ) + FProgramCounterSymbolInfo::MAX_NAME_LENGTH] = {0};
+	SYMBOL_INFO* Symbol = (SYMBOL_INFO*)SymbolBuffer;
+	Symbol->SizeOfStruct = sizeof(SYMBOL_INFO);
+	Symbol->MaxNameLen = FProgramCounterSymbolInfo::MAX_NAME_LENGTH;
 
 	// Get function name.
-	if( SymGetSymFromAddr64( ProcessHandle, ProgramCounter, nullptr, Symbol ) )
+	if( SymFromAddr( ProcessHandle, ProgramCounter, nullptr, Symbol ) )
 	{
 		// Skip any funky chars in the beginning of a function name.
 		int32 Offset = 0;
@@ -513,14 +513,14 @@ void FWindowsPlatformStackWalk::ProgramCounterToSymbolInfoEx(uint64 ProgramCount
 	HANDLE ProcessHandle = GProcessHandle;
 
 	// Initialize symbol.
-	ANSICHAR SymbolBuffer[sizeof(IMAGEHLP_SYMBOL64) + FProgramCounterSymbolInfo::MAX_NAME_LENGTH] = { 0 };
-	IMAGEHLP_SYMBOL64* Symbol = (IMAGEHLP_SYMBOL64*)SymbolBuffer;
-	Symbol->SizeOfStruct = sizeof(SymbolBuffer);
-	Symbol->MaxNameLength = FProgramCounterSymbolInfo::MAX_NAME_LENGTH;
+	ANSICHAR SymbolBuffer[sizeof(SYMBOL_INFO) + FProgramCounterSymbolInfo::MAX_NAME_LENGTH] = { 0 };
+	SYMBOL_INFO* Symbol = (SYMBOL_INFO*)SymbolBuffer;
+	Symbol->SizeOfStruct = sizeof(SYMBOL_INFO);
+	Symbol->MaxNameLen = FProgramCounterSymbolInfo::MAX_NAME_LENGTH;
 
 	// Get function name.
 	ANSICHAR FunctionName[FProgramCounterSymbolInfo::MAX_NAME_LENGTH] = { 0 };
-	if (SymGetSymFromAddr64(ProcessHandle, ProgramCounter, nullptr, Symbol))
+	if (SymFromAddr(ProcessHandle, ProgramCounter, nullptr, Symbol))
 	{
 		// Skip any funky chars in the beginning of a function name.
 		int32 Offset = 0;
