At some point in the preceeding eleven years MS have addressed whatever this code triggered in their compiler.

diff --git a/Engine/Source/Runtime/Core/Private/Windows/WindowsPlatformStackWalk.cpp b/Engine/Source/Runtime/Core/Private/Windows/WindowsPlatformStackWalk.cpp
--- a/Engine/Source/Runtime/Core/Private/Windows/WindowsPlatformStackWalk.cpp
+++ b/Engine/Source/Runtime/Core/Private/Windows/WindowsPlatformStackWalk.cpp
@@ -207,9 +207,6 @@ static int32 CaptureStackTraceHelper(uint64 *BackTrace, uint32 MaxDepth, FWindow
 	return EXCEPTION_EXECUTE_HANDLER;
 }
 
-PRAGMA_DISABLE_OPTIMIZATION // Work around "flow in or out of inline asm code suppresses global optimization" warning C4740.
-
-
 int32 CaptureStackTraceHelper(uint64* BackTrace, uint32 MaxDepth, CONTEXT* Context)
 {
 	FWindowsThreadContextWrapper HelperContext;
@@ -428,8 +425,6 @@ uint32 FWindowsPlatformStackWalk::CaptureStackBackTrace( uint64* BackTrace, uint
 	return Depth;
 }
 
-PRAGMA_ENABLE_OPTIMIZATION
-
 void FWindowsPlatformStackWalk::ProgramCounterToSymbolInfo( uint64 ProgramCounter, FProgramCounterSymbolInfo& out_SymbolInfo )
 {
 	// Initialize stack walking as it loads up symbol information which we require.
