BLAKE3: Made FBlake3Hash entirely inline

#rnx

diff --git a/Engine/Source/Runtime/Core/Private/Misc/Blake3.cpp b/Engine/Source/Runtime/Core/Private/Misc/Blake3.cpp
--- a/Engine/Source/Runtime/Core/Private/Misc/Blake3.cpp
+++ b/Engine/Source/Runtime/Core/Private/Misc/Blake3.cpp
@@ -2,35 +2,10 @@
 
 #include "Misc/Blake3.h"
 
-#include "Containers/StringView.h"
-#include "Memory/MemoryView.h"
-#include "Serialization/Archive.h"
-#include "String/HexToBytes.h"
-
 #include "blake3.h"
 
 ///////////////////////////////////////////////////////////////////////////////////////////////////////////////////////
 
-FBlake3Hash::FBlake3Hash(const FAnsiStringView HexHash)
-{
-	check(HexHash.Len() == sizeof(ByteArray) * 2);
-	UE::String::HexToBytes(HexHash, Hash);
-}
-
-FBlake3Hash::FBlake3Hash(const FWideStringView HexHash)
-{
-	check(HexHash.Len() == sizeof(ByteArray) * 2);
-	UE::String::HexToBytes(HexHash, Hash);
-}
-
-FArchive& operator<<(FArchive& Ar, FBlake3Hash& Hash)
-{
-	Ar.Serialize(Hash.Hash, sizeof(Hash.Hash));
-	return Ar;
-}
-
-///////////////////////////////////////////////////////////////////////////////////////////////////////////////////////
-
 static_assert(sizeof(FBlake3) == sizeof(blake3_hasher), "Adjust the allocation in FBlake3 to match blake3_hasher");
 
 void FBlake3::Reset()
@@ -47,10 +22,12 @@ void FBlake3::Update(const void* Data, uint64 Size)
 
 FBlake3Hash FBlake3::Finalize() const
 {
-	uint8 Output[BLAKE3_OUT_LEN];
+	FBlake3Hash Hash;
+	FBlake3Hash::ByteArray& Output = Hash.GetBytes();
+	static_assert(sizeof(decltype(Output)) == BLAKE3_OUT_LEN, "Mismatch in BLAKE3 hash size.");
 	const blake3_hasher& Hasher = reinterpret_cast<const blake3_hasher&>(HasherBytes);
 	blake3_hasher_finalize(&Hasher, Output, BLAKE3_OUT_LEN);
-	return FBlake3Hash(Output);
+	return Hash;
 }
 
 FBlake3Hash FBlake3::HashBuffer(const void* Data, uint64 Size)
diff --git a/Engine/Source/Runtime/Core/Public/Misc/Blake3.h b/Engine/Source/Runtime/Core/Public/Misc/Blake3.h
--- a/Engine/Source/Runtime/Core/Public/Misc/Blake3.h
+++ b/Engine/Source/Runtime/Core/Public/Misc/Blake3.h
@@ -3,10 +3,12 @@
 #pragma once
 
 #include "Containers/StringFwd.h"
+#include "Containers/StringView.h"
 #include "HAL/UnrealMemory.h"
 #include "Memory/MemoryView.h"
+#include "Serialization/Archive.h"
 #include "String/BytesToHex.h"
-#include "Templates/TypeCompatibleBytes.h"
+#include "String/HexToBytes.h"
 
 class FArchive;
 
@@ -20,23 +22,24 @@ public:
 
 	/** Construct a zero hash. */
 	FBlake3Hash() = default;
-	explicit FBlake3Hash(const ByteArray& Hash);
 
-	/** Construct a hash from a 64-character hex string. */
-	CORE_API explicit FBlake3Hash(FAnsiStringView HexHash);
-	CORE_API explicit FBlake3Hash(FWideStringView HexHash);
+	/** Construct a hash from an array of 32 bytes. */
+	inline explicit FBlake3Hash(const ByteArray& Hash);
 
-	CORE_API friend FArchive& operator<<(FArchive& Ar, FBlake3Hash& Hash);
+	/** Construct a hash from a 64-character hex string. */
+	inline explicit FBlake3Hash(FAnsiStringView HexHash);
+	inline explicit FBlake3Hash(FWideStringView HexHash);
 
 	/** Reset this to a zero hash. */
 	inline void Reset() { *this = FBlake3Hash(); }
 
-	/** Returns a reference to the raw byte array for the hash. */
-	inline const ByteArray& GetBytes() const { return Hash; }
-
 	/** Returns whether this is a zero hash. */
 	inline bool IsZero() const;
 
+	/** Returns a reference to the raw byte array for the hash. */
+	inline ByteArray& GetBytes() { return Hash; }
+	inline const ByteArray& GetBytes() const { return Hash; }
+
 private:
 	alignas(uint32) ByteArray Hash{};
 };
@@ -87,6 +90,18 @@ inline FBlake3Hash::FBlake3Hash(const ByteArray& InHash)
 	FMemory::Memcpy(Hash, InHash, sizeof(ByteArray));
 }
 
+inline FBlake3Hash::FBlake3Hash(const FAnsiStringView HexHash)
+{
+	check(HexHash.Len() == sizeof(ByteArray) * 2);
+	UE::String::HexToBytes(HexHash, Hash);
+}
+
+inline FBlake3Hash::FBlake3Hash(const FWideStringView HexHash)
+{
+	check(HexHash.Len() == sizeof(ByteArray) * 2);
+	UE::String::HexToBytes(HexHash, Hash);
+}
+
 inline bool FBlake3Hash::IsZero() const
 {
 	using UInt32Array = uint32[8];
@@ -116,6 +131,12 @@ inline bool operator<(const FBlake3Hash& A, const FBlake3Hash& B)
 	return FMemory::Memcmp(A.GetBytes(), B.GetBytes(), sizeof(decltype(A.GetBytes()))) < 0;
 }
 
+inline FArchive& operator<<(FArchive& Ar, FBlake3Hash& Hash)
+{
+	Ar.Serialize(Hash.GetBytes(), sizeof(decltype(Hash.GetBytes())));
+	return Ar;
+}
+
 inline uint32 GetTypeHash(const FBlake3Hash& Hash)
 {
 	return *reinterpret_cast<const uint32*>(Hash.GetBytes());
