U5 - pump threadpool during retraction to avoid starvation. And also BusyWait during Destruction to avoid taking worker-threads down in the DistanceFieldBuilder Use-Case

diff --git a/Engine/Source/Runtime/Core/Public/Misc/QueuedThreadPoolWrapper.h b/Engine/Source/Runtime/Core/Public/Misc/QueuedThreadPoolWrapper.h
--- a/Engine/Source/Runtime/Core/Public/Misc/QueuedThreadPoolWrapper.h
+++ b/Engine/Source/Runtime/Core/Public/Misc/QueuedThreadPoolWrapper.h
@@ -199,9 +199,17 @@ protected:
 	void Destroy() override
 	{
 		bIsExiting = true;
-		while (TaskCount != 0)
+
+		if (LowLevelTasks::FScheduler::Get().GetActiveTask())
 		{
-			FPlatformProcess::Sleep(0.01f);
+			LowLevelTasks::BusyWaitUntil([this]() { return TaskCount == 0; });
+		}
+		else
+		{
+			while (TaskCount != 0)
+			{
+				FPlatformProcess::Sleep(0.01f);
+			}
 		}
 	}
 private:
@@ -288,19 +296,19 @@ public:
 		}
 
 		bool bWakeUpWorker = true;
-		ScheduleTasks(bWakeUpWorker);
+		ScheduleTasks(bWakeUpWorker, TaskCount.load(std::memory_order_relaxed));
 	}
 
 private:
-	void ScheduleTasks(bool &bWakeUpWorker)
+	void ScheduleTasks(bool &bWakeUpWorker, uint32 LocalTaskCount)
 	{
-		while (!bIsPaused && TaskCount < MaxConcurrency)
+		while (!bIsPaused && LocalTaskCount < MaxConcurrency)
 		{
 			FQueuedWorkInternalData* QueuedWork = Dequeue();
 			if (QueuedWork)
 			{
-				TaskCount++;
 				verifySlow(LowLevelTasks::TryLaunch(QueuedWork->Task, bWakeUpWorker ? LowLevelTasks::EQueuePreference::GlobalQueuePreference : LowLevelTasks::EQueuePreference::LocalQueuePreference, bWakeUpWorker));
+				LocalTaskCount = TaskCount.fetch_add(1, std::memory_order_acquire) + 1;
 				bWakeUpWorker = true;
 			}
 			else
@@ -326,25 +334,28 @@ private:
 		},
 		[this, InternalData = InQueuedWork->InternalData]()
 		{
-			--TaskCount;
+			uint32 LocalTaskCount = TaskCount.fetch_sub(1, std::memory_order_release) + 1;
 			bool bWakeUpWorker = false;
-			ScheduleTasks(bWakeUpWorker);
+			ScheduleTasks(bWakeUpWorker, LocalTaskCount);
 		});
 
 		Enqueue(Priority, QueuedWorkInternalData);
 		bool bWakeUpWorker = true;
-		ScheduleTasks(bWakeUpWorker);	
+		ScheduleTasks(bWakeUpWorker, TaskCount.load(std::memory_order_relaxed));	
 	}
 
 	bool RetractQueuedWork(IQueuedWork* InQueuedWork) override
 	{
+		bool bCancelled = false;
 		if(InQueuedWork->InternalData.IsValid())
 		{
-			bool bCancelled = InQueuedWork->InternalData->Retract();
+			bCancelled = InQueuedWork->InternalData->Retract();
 			InQueuedWork->InternalData = nullptr;
-			return bCancelled;
 		}
-		return false;
+
+		bool bWakeUpWorker = true;
+		ScheduleTasks(bWakeUpWorker, TaskCount.load(std::memory_order_relaxed));	
+		return bCancelled;
 	}
 
 	int32 GetNumThreads() const override
@@ -376,9 +387,16 @@ protected:
 			verifySlow(LowLevelTasks::TryLaunch(QueuedWork->Task, LowLevelTasks::EQueuePreference::GlobalQueuePreference));
 		}
 
-		while (TaskCount != 0)
+		if (LowLevelTasks::FScheduler::Get().GetActiveTask())
 		{
-			FPlatformProcess::Sleep(0.01f);
+			LowLevelTasks::BusyWaitUntil([this]() { return TaskCount == 0; });
+		}
+		else
+		{
+			while (TaskCount != 0)
+			{
+				FPlatformProcess::Sleep(0.01f);
+			}
 		}
 	}
 
