Added FIoHash as BLAKE3-160

diff --git a/Engine/Source/Runtime/Core/Private/IO/IoHash.cpp b/Engine/Source/Runtime/Core/Private/IO/IoHash.cpp
new file mode 100644
--- /dev/null
+++ b/Engine/Source/Runtime/Core/Private/IO/IoHash.cpp
@@ -0,0 +1,15 @@
+// Copyright Epic Games, Inc. All Rights Reserved.
+
+#include "IO/IoHash.h"
+
+#include "Containers/UnrealString.h"
+
+FString FIoHash::ToString() const
+{
+	FString Output;
+	TArray<TCHAR>& CharArray = Output.GetCharArray();
+	CharArray.AddUninitialized(sizeof(ByteArray) * 2 + 1);
+	UE::String::BytesToHexLower(Hash, CharArray.GetData());
+	CharArray.Last() = 0;
+	return Output;
+}
diff --git a/Engine/Source/Runtime/Core/Public/IO/IoHash.h b/Engine/Source/Runtime/Core/Public/IO/IoHash.h
new file mode 100644
--- /dev/null
+++ b/Engine/Source/Runtime/Core/Public/IO/IoHash.h
@@ -0,0 +1,148 @@
+// Copyright Epic Games, Inc. All Rights Reserved.
+
+#pragma once
+
+#include "Containers/StringFwd.h"
+#include "Containers/StringView.h"
+#include "HAL/UnrealMemory.h"
+#include "Misc/Blake3.h"
+#include "Memory/MemoryView.h"
+#include "Serialization/Archive.h"
+#include "String/BytesToHex.h"
+#include "String/HexToBytes.h"
+
+/**
+ * Stores a BLAKE3-160 hash, taken from the first 20 bytes of a BLAKE3-256 hash.
+ *
+ * The BLAKE3 hash function was selected for its high performance and its ability to parallelize.
+ * Only the leading 160 bits of the 256-bit hash are used, to provide strong collision resistance
+ * while minimizing the size of the hash.
+ *
+ * When the data to hash is not in a contiguous region of memory, FBlake3 can be used to hash the
+ * data in blocks with FBlake3::Update(Block) followed by FIoHash(FBlake3::Finalize()).
+ */
+struct FIoHash
+{
+public:
+	using ByteArray = uint8[20];
+
+	/** Construct a zero hash. */
+	FIoHash() = default;
+
+	/** Construct a hash from an array of 20 bytes. */
+	inline explicit FIoHash(const ByteArray& Hash);
+
+	/** Construct a hash from a BLAKE3-256 hash. */
+	inline explicit FIoHash(const FBlake3Hash& Hash);
+
+	/** Construct a hash from a 40-character hex string. */
+	inline explicit FIoHash(FAnsiStringView HexHash);
+	inline explicit FIoHash(FWideStringView HexHash);
+
+	/** Reset this to a zero hash. */
+	inline void Reset() { *this = FIoHash(); }
+
+	/** Returns whether this is a zero hash. */
+	inline bool IsZero() const;
+
+	/** Returns a reference to the raw byte array for the hash. */
+	inline ByteArray& GetBytes() { return Hash; }
+	inline const ByteArray& GetBytes() const { return Hash; }
+
+	/** Returns the hash as a 40-character hex string. */
+	CORE_API FString ToString() const;
+
+	/** Calculate the hash of the input data. */
+	static inline FIoHash HashBuffer(const void* Data, uint64 Size);
+
+	/** Calculate the hash of the input view. */
+	static inline FIoHash HashBuffer(FMemoryView View);
+
+private:
+	alignas(uint32) ByteArray Hash{};
+};
+
+inline FIoHash::FIoHash(const ByteArray& InHash)
+{
+	FMemory::Memcpy(Hash, InHash, sizeof(ByteArray));
+}
+
+inline FIoHash::FIoHash(const FBlake3Hash& InHash)
+{
+	static_assert(sizeof(ByteArray) <= sizeof(decltype(InHash.GetBytes())), "Reading too many bytes from source.");
+	FMemory::Memcpy(Hash, InHash.GetBytes(), sizeof(ByteArray));
+}
+
+inline FIoHash::FIoHash(const FAnsiStringView HexHash)
+{
+	check(HexHash.Len() == sizeof(ByteArray) * 2);
+	UE::String::HexToBytes(HexHash, Hash);
+}
+
+inline FIoHash::FIoHash(const FWideStringView HexHash)
+{
+	check(HexHash.Len() == sizeof(ByteArray) * 2);
+	UE::String::HexToBytes(HexHash, Hash);
+}
+
+inline bool FIoHash::IsZero() const
+{
+	using UInt32Array = uint32[5];
+	static_assert(sizeof(UInt32Array) == sizeof(ByteArray), "Invalid size for UInt32Array");
+	for (uint32 Value : reinterpret_cast<const UInt32Array&>(Hash))
+	{
+		if (Value != 0)
+		{
+			return false;
+		}
+	}
+	return true;
+}
+
+inline FIoHash FIoHash::HashBuffer(const void* Data, uint64 Size)
+{
+	return FIoHash(FBlake3::HashBuffer(Data, Size));
+}
+
+inline FIoHash FIoHash::HashBuffer(FMemoryView View)
+{
+	return FIoHash(FBlake3::HashBuffer(View));
+}
+
+inline bool operator==(const FIoHash& A, const FIoHash& B)
+{
+	return FMemory::Memcmp(A.GetBytes(), B.GetBytes(), sizeof(decltype(A.GetBytes()))) == 0;
+}
+
+inline bool operator!=(const FIoHash& A, const FIoHash& B)
+{
+	return FMemory::Memcmp(A.GetBytes(), B.GetBytes(), sizeof(decltype(A.GetBytes()))) != 0;
+}
+
+inline bool operator<(const FIoHash& A, const FIoHash& B)
+{
+	return FMemory::Memcmp(A.GetBytes(), B.GetBytes(), sizeof(decltype(A.GetBytes()))) < 0;
+}
+
+inline FArchive& operator<<(FArchive& Ar, FIoHash& Hash)
+{
+	Ar.Serialize(Hash.GetBytes(), sizeof(decltype(Hash.GetBytes())));
+	return Ar;
+}
+
+inline uint32 GetTypeHash(const FIoHash& Hash)
+{
+	return *reinterpret_cast<const uint32*>(Hash.GetBytes());
+}
+
+inline FAnsiStringBuilderBase& operator<<(FAnsiStringBuilderBase& Builder, const FIoHash& Hash)
+{
+	UE::String::BytesToHexLower(Hash.GetBytes(), Builder);
+	return Builder;
+}
+
+inline FWideStringBuilderBase& operator<<(FWideStringBuilderBase& Builder, const FIoHash& Hash)
+{
+	UE::String::BytesToHexLower(Hash.GetBytes(), Builder);
+	return Builder;
+}
