CompactBinary: Converted from FBlake3Hash to FIoHash

#rnx

diff --git a/Engine/Source/Developer/DerivedDataCache/Private/DerivedDataCache.cpp b/Engine/Source/Developer/DerivedDataCache/Private/DerivedDataCache.cpp
--- a/Engine/Source/Developer/DerivedDataCache/Private/DerivedDataCache.cpp
+++ b/Engine/Source/Developer/DerivedDataCache/Private/DerivedDataCache.cpp
@@ -769,7 +769,7 @@ private:
 	{
 		check(AttachmentKey);
 		TStringBuilder<128> Out;
-		Out << AttachmentKey.GetKey().GetBucket() << TEXT("_BLAKE3_") << AttachmentKey.GetHash();
+		Out << AttachmentKey.GetKey().GetBucket() << TEXT("_CAS_") << AttachmentKey.GetHash();
 		return FDerivedDataCacheInterface::SanitizeCacheKey(Out.ToString());
 	}
 };
@@ -839,7 +839,7 @@ void FCache::Get(
 
 	// Check for existence of the value and attachments if they are being skipped.
 	{
-		TArray<FBlake3Hash, TInlineAllocator<1>> KeysToCheck;
+		TArray<FIoHash, TInlineAllocator<1>> KeysToCheck;
 		if (EnumHasAnyFlags(Policy, ECachePolicy::SkipValue) && ValueField.IsReference())
 		{
 			KeysToCheck.Add(ValueField.AsReference());
@@ -861,7 +861,7 @@ void FCache::Get(
 			TArray<FString, TInlineAllocator<1>> ContentKeysToCheck;
 			ContentKeysToCheck.Reserve(KeysToCheck.Num());
 			Algo::Transform(KeysToCheck, ContentKeysToCheck,
-				[&Key](const FBlake3Hash& Hash) -> FString { return MakeContentKey(FCacheAttachmentKey(Key, Hash)); });
+				[&Key](const FIoHash& Hash) -> FString { return MakeContentKey(FCacheAttachmentKey(Key, Hash)); });
 			if (!GetDerivedDataCacheRef().AllCachedDataProbablyExists(ContentKeysToCheck))
 			{
 				UE_LOG(LogDerivedDataCache, Verbose, TEXT("Cache: Get cache miss with missing content for %s from '%.*s'"),
@@ -875,12 +875,12 @@ void FCache::Get(
 	const auto GetContent = [&Key, Context](FCbField ReferenceField) -> FCbAttachment
 	{
 		TArray<uint8> ContentData;
-		const FBlake3Hash Hash = ReferenceField.AsReference();
+		const FIoHash Hash = ReferenceField.AsReference();
 		check(!ReferenceField.HasError());
 		const FCacheAttachmentKey ContentKey(Key, Hash);
 		if (GetDerivedDataCacheRef().GetSynchronous(*MakeContentKey(ContentKey), ContentData, Context))
 		{
-			if (FBlake3::HashBuffer(MakeMemoryView(ContentData)) == Hash)
+			if (FIoHash::HashBuffer(MakeMemoryView(ContentData)) == Hash)
 			{
 				FSharedBuffer ContentBuffer = FSharedBuffer::Clone(MakeMemoryView(ContentData));
 				if (ReferenceField.IsCompactBinaryReference())
@@ -1030,7 +1030,7 @@ void FCache::Put(
 		break;
 	}
 
-	const auto PutContent = [&Key = Params.Key, Context](const FSharedBuffer& Buffer, const FBlake3Hash& Hash)
+	const auto PutContent = [&Key = Params.Key, Context](const FSharedBuffer& Buffer, const FIoHash& Hash)
 	{
 		check(Buffer.GetSize() <= MAX_int32);
 		TConstArrayView<uint8> BufferView = MakeArrayView(
@@ -1165,7 +1165,7 @@ void FCache::GetAttachment(
 		return;
 	}
 
-	if (FBlake3::HashBuffer(MakeMemoryView(Data)) != Key.GetHash())
+	if (FIoHash::HashBuffer(MakeMemoryView(Data)) != Key.GetHash())
 	{
 		UE_LOG(LogDerivedDataCache, Verbose, TEXT("Cache: GetAttachment cache miss with corrupted content for %s from '%.*s'"),
 			*TToString<160>(Key), Context.Len(), Context.GetData());
diff --git a/Engine/Source/Developer/DerivedDataCache/Public/DerivedDataCacheKey.h b/Engine/Source/Developer/DerivedDataCache/Public/DerivedDataCacheKey.h
--- a/Engine/Source/Developer/DerivedDataCache/Public/DerivedDataCacheKey.h
+++ b/Engine/Source/Developer/DerivedDataCache/Public/DerivedDataCacheKey.h
@@ -5,7 +5,7 @@
 #include "CoreTypes.h"
 #include "Containers/StringFwd.h"
 #include "Containers/StringView.h"
-#include "Misc/Blake3.h"
+#include "IO/IoHash.h"
 #include "Templates/TypeHash.h"
 
 #define UE_API DERIVEDDATACACHE_API
@@ -97,7 +97,7 @@ public:
 	FCacheKey() = default;
 
 	/** Construct a cache key from the non-null bucket and non-zero hash. */
-	inline FCacheKey(FCacheBucket InBucket, const FBlake3Hash& InHash)
+	inline FCacheKey(FCacheBucket InBucket, const FIoHash& InHash)
 		: Bucket(InBucket)
 		, Hash(InHash)
 	{
@@ -107,7 +107,7 @@ public:
 	inline FCacheBucket GetBucket() const { return Bucket; }
 
 	/** Returns the hash component of the cache key. */
-	inline const FBlake3Hash& GetHash() const { return Hash; }
+	inline const FIoHash& GetHash() const { return Hash; }
 
 	/** Append the cache key to the string builder. */
 	UE_API void ToString(FAnsiStringBuilderBase& Builder) const;
@@ -123,7 +123,7 @@ public:
 
 private:
 	FCacheBucket Bucket;
-	FBlake3Hash Hash;
+	FIoHash Hash;
 };
 
 inline bool operator==(const FCacheKey& A, const FCacheKey& B)
@@ -176,7 +176,7 @@ public:
 	FCacheAttachmentKey() = default;
 
 	/** Construct a cache attachment key from the non-null key and non-zero hash. */
-	inline FCacheAttachmentKey(const FCacheKey& InKey, const FBlake3Hash& InHash)
+	inline FCacheAttachmentKey(const FCacheKey& InKey, const FIoHash& InHash)
 		: Key(InKey)
 		, Hash(InHash)
 	{
@@ -186,7 +186,7 @@ public:
 	inline const FCacheKey& GetKey() const { return Key; }
 
 	/** Returns the hash of the attachment. */
-	inline const FBlake3Hash& GetHash() const { return Hash; }
+	inline const FIoHash& GetHash() const { return Hash; }
 
 	/** Append the cache attachment key to the string builder. */
 	UE_API void ToString(FAnsiStringBuilderBase& Builder) const;
@@ -202,7 +202,7 @@ public:
 
 private:
 	FCacheKey Key;
-	FBlake3Hash Hash;
+	FIoHash Hash;
 };
 
 inline bool operator==(const FCacheAttachmentKey& A, const FCacheAttachmentKey& B)
diff --git a/Engine/Source/Runtime/Core/Private/Serialization/CompactBinary.cpp b/Engine/Source/Runtime/Core/Private/Serialization/CompactBinary.cpp
--- a/Engine/Source/Runtime/Core/Private/Serialization/CompactBinary.cpp
+++ b/Engine/Source/Runtime/Core/Private/Serialization/CompactBinary.cpp
@@ -3,6 +3,7 @@
 #include "Serialization/CompactBinary.h"
 
 #include "Math/UnrealMathUtility.h"
+#include "Misc/Blake3.h"
 #include "Misc/ByteSwap.h"
 #include "Misc/DateTime.h"
 #include "Misc/Guid.h"
@@ -253,12 +254,12 @@ bool FCbField::AsBool(const bool bDefault)
 	return (uint8(bIsBool) & uint8(LocalType) & 1) | ((!bIsBool) & bDefault);
 }
 
-FBlake3Hash FCbField::AsCompactBinaryReference(const FBlake3Hash& Default)
+FIoHash FCbField::AsCompactBinaryReference(const FIoHash& Default)
 {
 	if (FCbFieldType::IsCompactBinaryReference(Type))
 	{
 		Error = ECbFieldError::None;
-		return FBlake3Hash(*static_cast<const FBlake3Hash::ByteArray*>(Payload));
+		return FIoHash(*static_cast<const FIoHash::ByteArray*>(Payload));
 	}
 	else
 	{
@@ -267,12 +268,12 @@ FBlake3Hash FCbField::AsCompactBinaryReference(const FBlake3Hash& Default)
 	}
 }
 
-FBlake3Hash FCbField::AsBinaryReference(const FBlake3Hash& Default)
+FIoHash FCbField::AsBinaryReference(const FIoHash& Default)
 {
 	if (FCbFieldType::IsBinaryReference(Type))
 	{
 		Error = ECbFieldError::None;
-		return FBlake3Hash(*static_cast<const FBlake3Hash::ByteArray*>(Payload));
+		return FIoHash(*static_cast<const FIoHash::ByteArray*>(Payload));
 	}
 	else
 	{
@@ -281,12 +282,12 @@ FBlake3Hash FCbField::AsBinaryReference(const FBlake3Hash& Default)
 	}
 }
 
-FBlake3Hash FCbField::AsReference(const FBlake3Hash& Default)
+FIoHash FCbField::AsReference(const FIoHash& Default)
 {
 	if (FCbFieldType::IsReference(Type))
 	{
 		Error = ECbFieldError::None;
-		return FBlake3Hash(*static_cast<const FBlake3Hash::ByteArray*>(Payload));
+		return FIoHash(*static_cast<const FIoHash::ByteArray*>(Payload));
 	}
 	else
 	{
@@ -295,12 +296,12 @@ FBlake3Hash FCbField::AsReference(const FBlake3Hash& Default)
 	}
 }
 
-FBlake3Hash FCbField::AsHash(const FBlake3Hash& Default)
+FIoHash FCbField::AsHash(const FIoHash& Default)
 {
 	if (FCbFieldType::IsHash(Type))
 	{
 		Error = ECbFieldError::None;
-		return FBlake3Hash(*static_cast<const FBlake3Hash::ByteArray*>(Payload));
+		return FIoHash(*static_cast<const FIoHash::ByteArray*>(Payload));
 	}
 	else
 	{
@@ -418,7 +419,7 @@ uint64 FCbField::GetPayloadSize() const
 	case ECbFieldType::CompactBinaryReference:
 	case ECbFieldType::BinaryReference:
 	case ECbFieldType::Hash:
-		return 32;
+		return 20;
 	case ECbFieldType::Uuid:
 		return 16;
 	case ECbFieldType::DateTime:
@@ -429,11 +430,11 @@ uint64 FCbField::GetPayloadSize() const
 	}
 }
 
-FBlake3Hash FCbField::GetHash() const
+FIoHash FCbField::GetHash() const
 {
 	FBlake3 Hash;
 	GetHash(Hash);
-	return Hash.Finalize();
+	return FIoHash(Hash.Finalize());
 }
 
 void FCbField::GetHash(FBlake3& Hash) const
@@ -540,11 +541,11 @@ uint64 FCbArray::GetSize() const
 	return sizeof(ECbFieldType) + GetPayloadSize();
 }
 
-FBlake3Hash FCbArray::GetHash() const
+FIoHash FCbArray::GetHash() const
 {
 	FBlake3 Hash;
 	GetHash(Hash);
-	return Hash.Finalize();
+	return FIoHash(Hash.Finalize());
 }
 
 void FCbArray::GetHash(FBlake3& Hash) const
@@ -632,11 +633,11 @@ uint64 FCbObject::GetSize() const
 	return sizeof(ECbFieldType) + GetPayloadSize();
 }
 
-FBlake3Hash FCbObject::GetHash() const
+FIoHash FCbObject::GetHash() const
 {
 	FBlake3 Hash;
 	GetHash(Hash);
-	return Hash.Finalize();
+	return FIoHash(Hash.Finalize());
 }
 
 void FCbObject::GetHash(FBlake3& Hash) const
@@ -694,11 +695,11 @@ uint64 TCbFieldIterator<FieldType>::GetRangeSize() const
 }
 
 template <typename FieldType>
-FBlake3Hash TCbFieldIterator<FieldType>::GetRangeHash() const
+FIoHash TCbFieldIterator<FieldType>::GetRangeHash() const
 {
 	FBlake3 Hash;
 	GetRangeHash(Hash);
-	return Hash.Finalize();
+	return FIoHash(Hash.Finalize());
 }
 
 template <typename FieldType>
diff --git a/Engine/Source/Runtime/Core/Private/Serialization/CompactBinaryPackage.cpp b/Engine/Source/Runtime/Core/Private/Serialization/CompactBinaryPackage.cpp
--- a/Engine/Source/Runtime/Core/Private/Serialization/CompactBinaryPackage.cpp
+++ b/Engine/Source/Runtime/Core/Private/Serialization/CompactBinaryPackage.cpp
@@ -8,7 +8,7 @@
 
 ///////////////////////////////////////////////////////////////////////////////////////////////////////////////////////
 
-FCbAttachment::FCbAttachment(FCbFieldRefIterator InValue, const FBlake3Hash* const InHash)
+FCbAttachment::FCbAttachment(FCbFieldRefIterator InValue, const FIoHash* const InHash)
 {
 	if (InValue)
 	{
@@ -39,7 +39,7 @@ FCbAttachment::FCbAttachment(FCbFieldRefIterator InValue, const FBlake3Hash* con
 	}
 }
 
-FCbAttachment::FCbAttachment(FSharedBuffer InBuffer, const FBlake3Hash* const InHash)
+FCbAttachment::FCbAttachment(FSharedBuffer InBuffer, const FIoHash* const InHash)
 	: Buffer(MoveTemp(InBuffer))
 {
 	Buffer.MakeOwned();
@@ -48,7 +48,7 @@ FCbAttachment::FCbAttachment(FSharedBuffer InBuffer, const FBlake3Hash* const In
 		Hash = *InHash;
 		if (Buffer.GetSize())
 		{
-			checkSlow(Hash == FBlake3::HashBuffer(Buffer));
+			checkSlow(Hash == FIoHash::HashBuffer(Buffer));
 		}
 		else
 		{
@@ -57,7 +57,7 @@ FCbAttachment::FCbAttachment(FSharedBuffer InBuffer, const FBlake3Hash* const In
 	}
 	else if (Buffer.GetSize())
 	{
-		Hash = FBlake3::HashBuffer(Buffer);
+		Hash = FIoHash::HashBuffer(Buffer);
 	}
 	else
 	{
@@ -180,7 +180,7 @@ void FCbAttachment::Save(FArchive& Ar) const
 
 ///////////////////////////////////////////////////////////////////////////////////////////////////////////////////////
 
-void FCbPackage::SetObject(FCbObjectRef InObject, const FBlake3Hash* InObjectHash, FAttachmentResolver* InResolver)
+void FCbPackage::SetObject(FCbObjectRef InObject, const FIoHash* InObjectHash, FAttachmentResolver* InResolver)
 {
 	if (InObject.CreateIterator())
 	{
@@ -230,10 +230,10 @@ void FCbPackage::AddAttachment(const FCbAttachment& Attachment, FAttachmentResol
 	}
 }
 
-int32 FCbPackage::RemoveAttachment(const FBlake3Hash& Hash)
+int32 FCbPackage::RemoveAttachment(const FIoHash& Hash)
 {
 	const int32 Index = Algo::BinarySearchBy(Attachments, Hash,
-		[](const FCbAttachment& Attachment) -> const FBlake3Hash& { return Attachment.GetHash(); });
+		[](const FCbAttachment& Attachment) -> const FIoHash& { return Attachment.GetHash(); });
 	if (Attachments.IsValidIndex(Index))
 	{
 		Attachments.RemoveAt(Index);
@@ -247,10 +247,10 @@ bool FCbPackage::Equals(const FCbPackage& Package) const
 	return ObjectHash == Package.ObjectHash && Attachments == Package.Attachments;
 }
 
-const FCbAttachment* FCbPackage::FindAttachment(const FBlake3Hash& Hash) const
+const FCbAttachment* FCbPackage::FindAttachment(const FIoHash& Hash) const
 {
 	const int32 Index = Algo::BinarySearchBy(Attachments, Hash,
-		[](const FCbAttachment& Attachment) -> const FBlake3Hash& { return Attachment.GetHash(); });
+		[](const FCbAttachment& Attachment) -> const FIoHash& { return Attachment.GetHash(); });
 	return Attachments.IsValidIndex(Index) ? &Attachments[Index] : nullptr;
 }
 
@@ -258,7 +258,7 @@ void FCbPackage::GatherAttachments(const FCbFieldIterator& Fields, FAttachmentRe
 {
 	Fields.IterateRangeReferences([this, &Resolver](FCbField Field)
 		{
-			const FBlake3Hash& Hash = Field.AsReference();
+			const FIoHash& Hash = Field.AsReference();
 			if (FSharedBuffer Buffer = Resolver(Hash))
 			{
 				if (Field.IsCompactBinaryReference())
@@ -333,7 +333,7 @@ void FCbPackage::Load(FArchive& Ar, FCbBufferAllocator Allocator)
 				FSharedBuffer Buffer = FSharedBuffer::MakeView(View, ValueField.GetOuterBuffer());
 				Buffer.MakeOwned();
 				FCbFieldRef HashField = LoadCompactBinary(Ar, StackAllocator);
-				const FBlake3Hash& Hash = HashField.AsReference();
+				const FIoHash& Hash = HashField.AsReference();
 				checkf(!HashField.HasError(), TEXT("Attachments must be a non-empty binary value with a content hash."));
 				if (HashField.IsCompactBinaryReference())
 				{
diff --git a/Engine/Source/Runtime/Core/Private/Serialization/CompactBinarySerialization.cpp b/Engine/Source/Runtime/Core/Private/Serialization/CompactBinarySerialization.cpp
--- a/Engine/Source/Runtime/Core/Private/Serialization/CompactBinarySerialization.cpp
+++ b/Engine/Source/Runtime/Core/Private/Serialization/CompactBinarySerialization.cpp
@@ -62,7 +62,7 @@ bool TryMeasureCompactBinary(FMemoryView View, ECbFieldType& OutType, uint64& Ou
 	case ECbFieldType::CompactBinaryReference:
 	case ECbFieldType::BinaryReference:
 	case ECbFieldType::Hash:
-		FixedSize = 32;
+		FixedSize = 20;
 		break;
 	case ECbFieldType::Uuid:
 		FixedSize = 16;
diff --git a/Engine/Source/Runtime/Core/Private/Serialization/CompactBinaryValidation.cpp b/Engine/Source/Runtime/Core/Private/Serialization/CompactBinaryValidation.cpp
--- a/Engine/Source/Runtime/Core/Private/Serialization/CompactBinaryValidation.cpp
+++ b/Engine/Source/Runtime/Core/Private/Serialization/CompactBinaryValidation.cpp
@@ -2,6 +2,7 @@
 
 #include "Serialization/CompactBinaryValidation.h"
 
+#include "IO/IoHash.h"
 #include "Misc/ByteSwap.h"
 #include "Serialization/CompactBinary.h"
 #include "Serialization/VarInt.h"
@@ -355,7 +356,7 @@ static FCbField ValidateCbField(FMemoryView& View, ECbValidateMode Mode, ECbVali
 	case ECbFieldType::CompactBinaryReference:
 	case ECbFieldType::BinaryReference:
 	case ECbFieldType::Hash:
-		ValidateFixedPayload(32);
+		ValidateFixedPayload(20);
 		break;
 	case ECbFieldType::Uuid:
 		ValidateFixedPayload(16);
@@ -395,7 +396,7 @@ static FCbField ValidateCbPackageField(FMemoryView& View, ECbValidateMode Mode,
 	return FCbField();
 }
 
-static FBlake3Hash ValidateCbPackageAttachment(FCbField& Value, FMemoryView& View, ECbValidateMode Mode, ECbValidateError& Error)
+static FIoHash ValidateCbPackageAttachment(FCbField& Value, FMemoryView& View, ECbValidateMode Mode, ECbValidateError& Error)
 {
 	const FMemoryView ValueView = Value.AsBinary();
 	if (Value.HasError() && EnumHasAnyFlags(Mode, ECbValidateMode::Package))
@@ -409,14 +410,14 @@ static FBlake3Hash ValidateCbPackageAttachment(FCbField& Value, FMemoryView& Vie
 	{
 		if (FCbField HashField = ValidateCbPackageField(View, Mode, Error))
 		{
-			const FBlake3Hash Hash = HashField.AsReference();
+			const FIoHash Hash = HashField.AsReference();
 			if (EnumHasAnyFlags(Mode, ECbValidateMode::Package))
 			{
 				if (HashField.HasError())
 				{
 					AddError(Error, ECbValidateError::InvalidPackageFormat);
 				}
-				else if (Hash != FBlake3::HashBuffer(ValueView))
+				else if (Hash != FIoHash::HashBuffer(ValueView))
 				{
 					AddError(Error, ECbValidateError::InvalidPackageHash);
 				}
@@ -424,10 +425,10 @@ static FBlake3Hash ValidateCbPackageAttachment(FCbField& Value, FMemoryView& Vie
 			return Hash;
 		}
 	}
-	return FBlake3Hash();
+	return FIoHash();
 }
 
-static FBlake3Hash ValidateCbPackageObject(FCbField& Value, FMemoryView& View, ECbValidateMode Mode, ECbValidateError& Error)
+static FIoHash ValidateCbPackageObject(FCbField& Value, FMemoryView& View, ECbValidateMode Mode, ECbValidateError& Error)
 {
 	FCbObject Object = Value.AsObject();
 	if (Value.HasError())
@@ -439,7 +440,7 @@ static FBlake3Hash ValidateCbPackageObject(FCbField& Value, FMemoryView& View, E
 	}
 	else if (FCbField HashField = ValidateCbPackageField(View, Mode, Error))
 	{
-		const FBlake3Hash Hash = HashField.AsReference();
+		const FIoHash Hash = HashField.AsReference();
 		if (EnumHasAnyFlags(Mode, ECbValidateMode::Package))
 		{
 			if (!Object.CreateIterator())
@@ -457,7 +458,7 @@ static FBlake3Hash ValidateCbPackageObject(FCbField& Value, FMemoryView& View, E
 		}
 		return Hash;
 	}
-	return FBlake3Hash();
+	return FIoHash();
 }
 
 ECbValidateError ValidateCompactBinary(FMemoryView View, ECbValidateMode Mode, ECbFieldType Type)
@@ -506,7 +507,7 @@ ECbValidateError ValidateCompactBinaryAttachment(FMemoryView View, ECbValidateMo
 
 ECbValidateError ValidateCompactBinaryPackage(FMemoryView View, ECbValidateMode Mode)
 {
-	TArray<FBlake3Hash, TInlineAllocator<16>> Attachments;
+	TArray<FIoHash, TInlineAllocator<16>> Attachments;
 	ECbValidateError Error = ECbValidateError::None;
 	if (EnumHasAnyFlags(Mode, ECbValidateMode::All))
 	{
@@ -515,7 +516,7 @@ ECbValidateError ValidateCompactBinaryPackage(FMemoryView View, ECbValidateMode
 		{
 			if (Value.IsBinary())
 			{
-				const FBlake3Hash Hash = ValidateCbPackageAttachment(Value, View, Mode, Error);
+				const FIoHash Hash = ValidateCbPackageAttachment(Value, View, Mode, Error);
 				if (EnumHasAnyFlags(Mode, ECbValidateMode::Package))
 				{
 					Attachments.Add(Hash);
@@ -556,7 +557,7 @@ ECbValidateError ValidateCompactBinaryPackage(FMemoryView View, ECbValidateMode
 		if (Attachments.Num() && EnumHasAnyFlags(Mode, ECbValidateMode::Package))
 		{
 			Algo::Sort(Attachments);
-			for (const FBlake3Hash* It = Attachments.GetData(), *End = It + Attachments.Num() - 1; It != End; ++It)
+			for (const FIoHash* It = Attachments.GetData(), *End = It + Attachments.Num() - 1; It != End; ++It)
 			{
 				if (It[0] == It[1])
 				{
diff --git a/Engine/Source/Runtime/Core/Private/Serialization/CompactBinaryWriter.cpp b/Engine/Source/Runtime/Core/Private/Serialization/CompactBinaryWriter.cpp
--- a/Engine/Source/Runtime/Core/Private/Serialization/CompactBinaryWriter.cpp
+++ b/Engine/Source/Runtime/Core/Private/Serialization/CompactBinaryWriter.cpp
@@ -3,7 +3,7 @@
 #include "Serialization/CompactBinaryWriter.h"
 
 #include "Containers/StringConv.h"
-#include "Misc/Blake3.h"
+#include "IO/IoHash.h"
 #include "Misc/ByteSwap.h"
 #include "Misc/DateTime.h"
 #include "Misc/Guid.h"
@@ -477,21 +477,21 @@ void FCbWriter::Bool(const bool bValue)
 	EndField(bValue ? ECbFieldType::BoolTrue : ECbFieldType::BoolFalse);
 }
 
-void FCbWriter::CompactBinaryReference(const FBlake3Hash& Value)
+void FCbWriter::CompactBinaryReference(const FIoHash& Value)
 {
 	BeginField();
 	Data.Append(Value.GetBytes(), sizeof(decltype(Value.GetBytes())));
 	EndField(ECbFieldType::CompactBinaryReference);
 }
 
-void FCbWriter::BinaryReference(const FBlake3Hash& Value)
+void FCbWriter::BinaryReference(const FIoHash& Value)
 {
 	BeginField();
 	Data.Append(Value.GetBytes(), sizeof(decltype(Value.GetBytes())));
 	EndField(ECbFieldType::BinaryReference);
 }
 
-void FCbWriter::Hash(const FBlake3Hash& Value)
+void FCbWriter::Hash(const FIoHash& Value)
 {
 	BeginField();
 	Data.Append(Value.GetBytes(), sizeof(decltype(Value.GetBytes())));
diff --git a/Engine/Source/Runtime/Core/Private/Tests/Serialization/CompactBinaryPackageTest.cpp b/Engine/Source/Runtime/Core/Private/Tests/Serialization/CompactBinaryPackageTest.cpp
--- a/Engine/Source/Runtime/Core/Private/Tests/Serialization/CompactBinaryPackageTest.cpp
+++ b/Engine/Source/Runtime/Core/Private/Tests/Serialization/CompactBinaryPackageTest.cpp
@@ -57,7 +57,7 @@ bool FCbAttachmentTest::RunTest(const FString& Parameters)
 		TestFalse(TEXT("FCbAttachment(Null).AsCompactBinary()"), Attachment.AsCompactBinary().HasValue());
 		TestFalse(TEXT("FCbAttachment(Null).IsBinary()"), Attachment.IsBinary());
 		TestFalse(TEXT("FCbAttachment(Null).IsCompactBinary()"), Attachment.IsCompactBinary());
-		TestEqual(TEXT("FCbAttachment(Null).GetHash()"), Attachment.GetHash(), FBlake3Hash());
+		TestEqual(TEXT("FCbAttachment(Null).GetHash()"), Attachment.GetHash(), FIoHash());
 		TestSaveLoadValidate(TEXT("Null"), Attachment);
 	}
 
@@ -71,7 +71,7 @@ bool FCbAttachmentTest::RunTest(const FString& Parameters)
 		TestFalse(TEXT("FCbAttachment(Binary).AsCompactBinary()"), Attachment.AsCompactBinary().HasValue());
 		TestTrue(TEXT("FCbAttachment(Binary).IsBinary()"), Attachment.IsBinary());
 		TestFalse(TEXT("FCbAttachment(Binary).IsCompactBinary()"), Attachment.IsCompactBinary());
-		TestEqual(TEXT("FCbAttachment(Binary).GetHash()"), Attachment.GetHash(), FBlake3::HashBuffer(Buffer));
+		TestEqual(TEXT("FCbAttachment(Binary).GetHash()"), Attachment.GetHash(), FIoHash::HashBuffer(Buffer));
 		TestSaveLoadValidate(TEXT("Binary"), Attachment);
 	}
 
@@ -105,7 +105,7 @@ bool FCbAttachmentTest::RunTest(const FString& Parameters)
 		TestTrue(TEXT("FCbAttachment(BinaryView).IsBinary()"), Attachment.IsBinary());
 		TestFalse(TEXT("FCbAttachment(BinaryView).IsCompactBinary()"), Attachment.IsCompactBinary());
 		TestEqual(TEXT("FCbAttachment(BinaryView).GetHash()"),
-			Attachment.GetHash(), FBlake3::HashBuffer(Buffer));
+			Attachment.GetHash(), FIoHash::HashBuffer(Buffer));
 	}
 
 	// Compact Binary View
@@ -148,7 +148,7 @@ bool FCbAttachmentTest::RunTest(const FString& Parameters)
 		TestTrue(TEXT("FCbAttachment(LoadBinaryView).IsBinary()"), Attachment.IsBinary());
 		TestFalse(TEXT("FCbAttachment(LoadBinaryView).IsCompactBinary()"), Attachment.IsCompactBinary());
 		TestEqual(TEXT("FCbAttachment(LoadBinaryView).GetHash()"),
-			Attachment.GetHash(), FBlake3::HashBuffer(MakeMemoryView(Value)));
+			Attachment.GetHash(), FIoHash::HashBuffer(MakeMemoryView(Value)));
 	}
 
 	// Compact Binary Load from View
@@ -202,7 +202,7 @@ bool FCbAttachmentTest::RunTest(const FString& Parameters)
 		TestTrue(TEXT("FCbAttachment(BinaryNull).IsNull()"), Attachment.IsNull());
 		TestFalse(TEXT("FCbAttachment(BinaryNull).IsBinary()"), Attachment.IsBinary());
 		TestFalse(TEXT("FCbAttachment(BinaryNull).IsCompactBinary()"), Attachment.IsCompactBinary());
-		TestEqual(TEXT("FCbAttachment(BinaryNull).GetHash()"), Attachment.GetHash(), FBlake3Hash());
+		TestEqual(TEXT("FCbAttachment(BinaryNull).GetHash()"), Attachment.GetHash(), FIoHash());
 	}
 
 	// Binary Empty
@@ -211,7 +211,7 @@ bool FCbAttachmentTest::RunTest(const FString& Parameters)
 		TestTrue(TEXT("FCbAttachment(BinaryEmpty).IsNull()"), Attachment.IsNull());
 		TestFalse(TEXT("FCbAttachment(BinaryEmpty).IsBinary()"), Attachment.IsBinary());
 		TestFalse(TEXT("FCbAttachment(BinaryEmpty).IsCompactBinary()"), Attachment.IsCompactBinary());
-		TestEqual(TEXT("FCbAttachment(BinaryEmpty).GetHash()"), Attachment.GetHash(), FBlake3Hash());
+		TestEqual(TEXT("FCbAttachment(BinaryEmpty).GetHash()"), Attachment.GetHash(), FIoHash());
 	}
 
 	// Compact Binary Empty
@@ -220,7 +220,7 @@ bool FCbAttachmentTest::RunTest(const FString& Parameters)
 		TestTrue(TEXT("FCbAttachment(CompactBinaryEmpty).IsNull()"), Attachment.IsNull());
 		TestFalse(TEXT("FCbAttachment(CompactBinaryEmpty).IsBinary()"), Attachment.IsBinary());
 		TestFalse(TEXT("FCbAttachment(CompactBinaryEmpty).IsCompactBinary()"), Attachment.IsCompactBinary());
-		TestEqual(TEXT("FCbAttachment(CompactBinaryEmpty).GetHash()"), Attachment.GetHash(), FBlake3Hash());
+		TestEqual(TEXT("FCbAttachment(CompactBinaryEmpty).GetHash()"), Attachment.GetHash(), FIoHash());
 	}
 
 	return true;
@@ -321,7 +321,7 @@ bool FCbPackageTest::RunTest(const FString& Parameters)
 		TestTrue(TEXT("FCbPackage(Attachments) as bool"), bool(Package));
 		TestEqual(TEXT("FCbPackage(Attachments).GetAttachments()"), Package.GetAttachments().Num(), 2);
 		TestTrue(TEXT("FCbPackage(Attachments).GetObject()"), Package.GetObject().Equals(FCbObjectRef()));
-		TestEqual(TEXT("FCbPackage(Attachments).GetObjectHash()"), Package.GetObjectHash(), FBlake3Hash());
+		TestEqual(TEXT("FCbPackage(Attachments).GetObjectHash()"), Package.GetObjectHash(), FIoHash());
 		TestSaveLoadValidate(TEXT("Attachments"), Package);
 
 		const FCbAttachment* const ObjectAttachment = Package.FindAttachment(Object.GetHash());
@@ -357,7 +357,7 @@ bool FCbPackageTest::RunTest(const FString& Parameters)
 	// Shared Values
 	const uint8 Level4Values[]{0, 1, 2, 3};
 	FSharedBuffer Level4 = FSharedBuffer::MakeView(MakeMemoryView(Level4Values));
-	const FBlake3Hash Level4Hash = FBlake3::HashBuffer(Level4);
+	const FIoHash Level4Hash = FIoHash::HashBuffer(Level4);
 
 	FCbFieldRef Level3;
 	{
@@ -365,7 +365,7 @@ bool FCbPackageTest::RunTest(const FString& Parameters)
 		Writer.Name("Level4").BinaryReference(Level4Hash);
 		Level3 = Writer.Save();
 	}
-	const FBlake3Hash Level3Hash = Level3.GetHash();
+	const FIoHash Level3Hash = Level3.GetHash();
 
 	FCbArrayRef Level2;
 	{
@@ -376,7 +376,7 @@ bool FCbPackageTest::RunTest(const FString& Parameters)
 		Writer.EndArray();
 		Level2 = Writer.Save().AsArrayRef();
 	}
-	const FBlake3Hash Level2Hash = Level2.AsField().GetHash();
+	const FIoHash Level2Hash = Level2.AsField().GetHash();
 
 	FCbObjectRef Level1;
 	{
@@ -386,10 +386,10 @@ bool FCbPackageTest::RunTest(const FString& Parameters)
 		Writer.EndObject();
 		Level1 = Writer.Save().AsObjectRef();
 	}
-	const FBlake3Hash Level1Hash = Level1.AsField().GetHash();
+	const FIoHash Level1Hash = Level1.AsField().GetHash();
 
 	const auto Resolver = [&Level2, &Level2Hash, &Level3, &Level3Hash, &Level4, &Level4Hash]
-		(const FBlake3Hash& Hash) -> FSharedBuffer
+		(const FIoHash& Hash) -> FSharedBuffer
 		{
 			return
 				Hash == Level2Hash ? Level2.GetBuffer() :
@@ -446,7 +446,7 @@ bool FCbPackageTest::RunTest(const FString& Parameters)
 		TestEqual(TEXT("FCbPackage(Object+Attachments).Equals(EqualAssigned)"), PackageCopy, Package);
 		Package.SetObject(FCbObjectRef());
 		TestNotEqual(TEXT("FCbPackage(Object+Attachments).Equals(ObjectNotEqual)"), PackageCopy, Package);
-		TestEqual(TEXT("FCbPackage(Object+Attachments).GetObjectHash(Null)"), Package.GetObjectHash(), FBlake3Hash());
+		TestEqual(TEXT("FCbPackage(Object+Attachments).GetObjectHash(Null)"), Package.GetObjectHash(), FIoHash());
 	}
 
 	// Out of Order
@@ -509,14 +509,14 @@ bool FCbPackageTest::RunTest(const FString& Parameters)
 		++Saved;
 		TestEqual(TEXT("FCbPackage(OutOfOrder).Save()->Level1Hash"), Saved.AsCompactBinaryReference(), Level1Hash);
 		++Saved;
-		TestTrue(TEXT("FCbPackage(OutOfOrder).Save()->Level3"), Saved.AsBinary().EqualBytes(Level3.GetView()));
-		++Saved;
-		TestEqual(TEXT("FCbPackage(OutOfOrder).Save()->Level3Hash"), Saved.AsCompactBinaryReference(), Level3Hash);
-		++Saved;
 		TestTrue(TEXT("FCbPackage(OutOfOrder).Save()->Level2"), Saved.AsBinary().EqualBytes(Level2.GetView()));
 		++Saved;
 		TestEqual(TEXT("FCbPackage(OutOfOrder).Save()->Level2Hash"), Saved.AsCompactBinaryReference(), Level2Hash);
 		++Saved;
+		TestTrue(TEXT("FCbPackage(OutOfOrder).Save()->Level3"), Saved.AsBinary().EqualBytes(Level3.GetView()));
+		++Saved;
+		TestEqual(TEXT("FCbPackage(OutOfOrder).Save()->Level3Hash"), Saved.AsCompactBinaryReference(), Level3Hash);
+		++Saved;
 		TestTrue(TEXT("FCbPackage(OutOfOrder).Save()->Level4"), Saved.AsBinary().EqualBytes(Level4.GetView()));
 		++Saved;
 		TestEqual(TEXT("FCbPackage(OutOfOrder).Save()->Level4Hash"), Saved.AsBinaryReference(), Level4Hash);
@@ -543,7 +543,7 @@ bool FCbPackageTest::RunTest(const FString& Parameters)
 		FCbPackage Package;
 		Package.AddAttachment(FCbAttachment(Level3.GetBuffer()));
 		Package.AddAttachment(FCbAttachment(FCbFieldRefIterator::MakeSingle(Level3)),
-			[&bResolved](const FBlake3Hash& Hash) -> FSharedBuffer
+			[&bResolved](const FIoHash& Hash) -> FSharedBuffer
 			{
 				bResolved = true;
 				return FSharedBuffer();
diff --git a/Engine/Source/Runtime/Core/Private/Tests/Serialization/CompactBinarySerializationTest.cpp b/Engine/Source/Runtime/Core/Private/Tests/Serialization/CompactBinarySerializationTest.cpp
--- a/Engine/Source/Runtime/Core/Private/Tests/Serialization/CompactBinarySerializationTest.cpp
+++ b/Engine/Source/Runtime/Core/Private/Tests/Serialization/CompactBinarySerializationTest.cpp
@@ -115,10 +115,10 @@ bool FCbMeasureTest::RunTest(const FString& Parameters)
 	TestMeasure(TEXT("BoolFalse"), {uint8(EType::BoolFalse)}, true, 1, EType::BoolFalse);
 	TestMeasure(TEXT("BoolTrue"), {uint8(EType::BoolTrue)}, true, 1, EType::BoolTrue);
 
-	TestMeasure(TEXT("CompactBinaryReference"), {uint8(EType::CompactBinaryReference)}, true, 33, EType::CompactBinaryReference);
-	TestMeasure(TEXT("BinaryReference"), {uint8(EType::BinaryReference)}, true, 33, EType::BinaryReference);
+	TestMeasure(TEXT("CompactBinaryReference"), {uint8(EType::CompactBinaryReference)}, true, 21, EType::CompactBinaryReference);
+	TestMeasure(TEXT("BinaryReference"), {uint8(EType::BinaryReference)}, true, 21, EType::BinaryReference);
 
-	TestMeasure(TEXT("Hash"), {uint8(EType::Hash)}, true, 33, EType::Hash);
+	TestMeasure(TEXT("Hash"), {uint8(EType::Hash)}, true, 21, EType::Hash);
 	TestMeasure(TEXT("Uuid"), {uint8(EType::Uuid)}, true, 17, EType::Uuid);
 
 	TestMeasure(TEXT("DateTime"), {uint8(EType::DateTime)}, true, 9, EType::DateTime);
diff --git a/Engine/Source/Runtime/Core/Private/Tests/Serialization/CompactBinaryTest.cpp b/Engine/Source/Runtime/Core/Private/Tests/Serialization/CompactBinaryTest.cpp
--- a/Engine/Source/Runtime/Core/Private/Tests/Serialization/CompactBinaryTest.cpp
+++ b/Engine/Source/Runtime/Core/Private/Tests/Serialization/CompactBinaryTest.cpp
@@ -2,8 +2,8 @@
 
 #include "Serialization/CompactBinary.h"
 
+#include "IO/IoHash.h"
 #include "Misc/AutomationTest.h"
-#include "Misc/Blake3.h"
 #include "Misc/DateTime.h"
 #include "Misc/Guid.h"
 #include "Misc/Timespan.h"
@@ -60,9 +60,9 @@ UE_CBFIELD_TYPE_ACCESSOR(Float32, IsFloat, AsFloat, float);
 UE_CBFIELD_TYPE_ACCESSOR(Float64, IsFloat, AsDouble, double);
 UE_CBFIELD_TYPE_ACCESSOR(BoolFalse, IsBool, AsBool, bool);
 UE_CBFIELD_TYPE_ACCESSOR(BoolTrue, IsBool, AsBool, bool);
-UE_CBFIELD_TYPE_ACCESSOR(CompactBinaryReference, IsCompactBinaryReference, AsCompactBinaryReference, FBlake3Hash);
-UE_CBFIELD_TYPE_ACCESSOR(BinaryReference, IsBinaryReference, AsBinaryReference, FBlake3Hash);
-UE_CBFIELD_TYPE_ACCESSOR(Hash, IsHash, AsHash, FBlake3Hash);
+UE_CBFIELD_TYPE_ACCESSOR(CompactBinaryReference, IsCompactBinaryReference, AsCompactBinaryReference, FIoHash);
+UE_CBFIELD_TYPE_ACCESSOR(BinaryReference, IsBinaryReference, AsBinaryReference, FIoHash);
+UE_CBFIELD_TYPE_ACCESSOR(Hash, IsHash, AsHash, FIoHash);
 UE_CBFIELD_TYPE_ACCESSOR_TYPED(Uuid, IsUuid, AsUuid, FGuid, const FGuid&);
 UE_CBFIELD_TYPE_ACCESSOR(DateTime, IsDateTime, AsDateTimeTicks, int64);
 UE_CBFIELD_TYPE_ACCESSOR(TimeSpan, IsTimeSpan, AsTimeSpanTicks, int64);
@@ -191,7 +191,7 @@ bool FCbFieldNoneTest::RunTest(const FString& Parameters)
 		TestFalse(TEXT("!FCbField()::HasValue()"), DefaultField.HasValue());
 		TestFalse(TEXT("!FCbField()::HasError()"), DefaultField.HasError());
 		TestEqual(TEXT("FCbField()::GetError() == None"), DefaultField.GetError(), ECbFieldError::None);
-		TestEqual(TEXT("FCbField()::GetHash()"), DefaultField.GetHash(), FBlake3::HashBuffer(MakeMemoryView<uint8>({uint8(ECbFieldType::None)})));
+		TestEqual(TEXT("FCbField()::GetHash()"), DefaultField.GetHash(), FIoHash::HashBuffer(MakeMemoryView<uint8>({uint8(ECbFieldType::None)})));
 		TestEqual(TEXT("FCbField()::GetView()"), DefaultField.GetView(), FMemoryView());
 		FMemoryView SerializedView;
 		TestFalse(TEXT("FCbField()::TryGetSerializedView()"), DefaultField.TryGetSerializedView(SerializedView));
@@ -221,7 +221,7 @@ bool FCbFieldNoneTest::RunTest(const FString& Parameters)
 		TestEqual(TEXT("FCbField(None|Type|Name)::GetName()"), NoneField.GetName(), "Name"_ASV);
 		TestTrue(TEXT("FCbField(None|Type|Name)::HasName()"), NoneField.HasName());
 		TestFalse(TEXT("!FCbField(None|Type|Name)::HasValue()"), NoneField.HasValue());
-		TestEqual(TEXT("FCbField(None|Type|Name)::GetHash()"), NoneField.GetHash(), FBlake3::HashBuffer(MakeMemoryView(NoneBytes)));
+		TestEqual(TEXT("FCbField(None|Type|Name)::GetHash()"), NoneField.GetHash(), FIoHash::HashBuffer(MakeMemoryView(NoneBytes)));
 		TestEqual(TEXT("FCbField(None|Type|Name)::GetView()"), NoneField.GetView(), MakeMemoryView(NoneBytes));
 		FMemoryView SerializedView;
 		TestTrue(TEXT("FCbField(None|Type|Name)::TryGetSerializedView()"), NoneField.TryGetSerializedView(SerializedView) && SerializedView == MakeMemoryView(NoneBytes));
@@ -255,7 +255,7 @@ bool FCbFieldNoneTest::RunTest(const FString& Parameters)
 		TestEqual(TEXT("FCbField(None|Name)::GetName()"), NoneField.GetName(), "Name"_ASV);
 		TestTrue(TEXT("FCbField(None|Name)::HasName()"), NoneField.HasName());
 		TestFalse(TEXT("!FCbField(None|Name)::HasValue()"), NoneField.HasValue());
-		TestEqual(TEXT("FCbField(None|Name)::GetHash()"), NoneField.GetHash(), FBlake3::HashBuffer(MakeMemoryView(NoneBytes)));
+		TestEqual(TEXT("FCbField(None|Name)::GetHash()"), NoneField.GetHash(), FIoHash::HashBuffer(MakeMemoryView(NoneBytes)));
 		TestEqual(TEXT("FCbField(None|Name)::GetView()"), NoneField.GetView(), MakeMemoryView(NoneBytes) + 1);
 		FMemoryView SerializedView;
 		TestFalse(TEXT("FCbField(None|Name)::TryGetSerializedView()"), NoneField.TryGetSerializedView(SerializedView));
@@ -274,7 +274,7 @@ bool FCbFieldNoneTest::RunTest(const FString& Parameters)
 		TestEqual(TEXT("FCbField(None|EmptyName)::GetName()"), NoneField.GetName(), ""_ASV);
 		TestTrue(TEXT("FCbField(None|EmptyName)::HasName()"), NoneField.HasName());
 		TestFalse(TEXT("!FCbField(None|EmptyName)::HasValue()"), NoneField.HasValue());
-		TestEqual(TEXT("FCbField(None|EmptyName)::GetHash()"), NoneField.GetHash(), FBlake3::HashBuffer(MakeMemoryView(NoneBytes)));
+		TestEqual(TEXT("FCbField(None|EmptyName)::GetHash()"), NoneField.GetHash(), FIoHash::HashBuffer(MakeMemoryView(NoneBytes)));
 		TestEqual(TEXT("FCbField(None|EmptyName)::GetView()"), NoneField.GetView(), MakeMemoryView(NoneBytes) + 1);
 		FMemoryView SerializedView;
 		TestFalse(TEXT("FCbField(None|EmptyName)::TryGetSerializedView()"), NoneField.TryGetSerializedView(SerializedView));
@@ -294,7 +294,7 @@ bool FCbFieldNullTest::RunTest(const FString& Parameters)
 		TestTrue(TEXT("FCbField(Null)::HasValue()"), NullField.HasValue());
 		TestFalse(TEXT("!FCbField(Null)::HasError()"), NullField.HasError());
 		TestEqual(TEXT("FCbField(Null)::GetError() == None"), NullField.GetError(), ECbFieldError::None);
-		TestEqual(TEXT("FCbField(Null)::GetHash()"), NullField.GetHash(), FBlake3::HashBuffer(MakeMemoryView<uint8>({uint8(ECbFieldType::Null)})));
+		TestEqual(TEXT("FCbField(Null)::GetHash()"), NullField.GetHash(), FIoHash::HashBuffer(MakeMemoryView<uint8>({uint8(ECbFieldType::Null)})));
 	}
 
 	// Test FCbField(None) as Null
@@ -944,25 +944,25 @@ bool FCbFieldBoolTest::RunTest(const FString& Parameters)
 IMPLEMENT_CUSTOM_SIMPLE_AUTOMATION_TEST(FCbFieldCompactBinaryReferenceTest, FCbFieldTestBase, "System.Core.Serialization.CbField.CompactBinaryReference", CompactBinaryTestFlags)
 bool FCbFieldCompactBinaryReferenceTest::RunTest(const FString& Parameters)
 {
-	const FBlake3Hash::ByteArray ZeroBytes{};
-	const FBlake3Hash::ByteArray SequentialBytes{1, 2, 3, 4, 5, 6, 7, 8, 9, 10, 11, 12, 13, 14, 15, 16, 17, 18, 19, 20, 21, 22, 23, 24, 25, 26, 27, 28, 29, 30, 31, 32};
+	const FIoHash::ByteArray ZeroBytes{};
+	const FIoHash::ByteArray SequentialBytes{1, 2, 3, 4, 5, 6, 7, 8, 9, 10, 11, 12, 13, 14, 15, 16, 17, 18, 19, 20};
 
 	// Test FCbField(CompactBinaryReference, Zero)
 	TestField<ECbFieldType::CompactBinaryReference>(TEXT("CompactBinaryReference, Zero"), ZeroBytes);
 
 	// Test FCbField(CompactBinaryReference, NonZero)
-	TestField<ECbFieldType::CompactBinaryReference>(TEXT("CompactBinaryReference, NonZero"), SequentialBytes, FBlake3Hash(SequentialBytes));
+	TestField<ECbFieldType::CompactBinaryReference>(TEXT("CompactBinaryReference, NonZero"), SequentialBytes, FIoHash(SequentialBytes));
 
 	// Test FCbField(CompactBinaryReference, NonZero) AsReference
 	{
 		FCbField Field(SequentialBytes, ECbFieldType::CompactBinaryReference);
-		TestFieldAsType(TEXT("CompactBinaryReference, NonZero, AsReference"), Field, &FCbField::AsReference, FBlake3Hash(SequentialBytes));
+		TestFieldAsType(TEXT("CompactBinaryReference, NonZero, AsReference"), Field, &FCbField::AsReference, FIoHash(SequentialBytes));
 	}
 
 	// Test FCbField(None) as CompactBinaryReference
 	{
 		FCbField DefaultField;
-		TestFieldError<ECbFieldType::CompactBinaryReference>(TEXT("CompactBinaryReference, None"), DefaultField, ECbFieldError::TypeError, FBlake3Hash(SequentialBytes));
+		TestFieldError<ECbFieldType::CompactBinaryReference>(TEXT("CompactBinaryReference, None"), DefaultField, ECbFieldError::TypeError, FIoHash(SequentialBytes));
 	}
 
 	return true;
@@ -971,25 +971,25 @@ bool FCbFieldCompactBinaryReferenceTest::RunTest(const FString& Parameters)
 IMPLEMENT_CUSTOM_SIMPLE_AUTOMATION_TEST(FCbFieldBinaryReferenceTest, FCbFieldTestBase, "System.Core.Serialization.CbField.BinaryReference", CompactBinaryTestFlags)
 bool FCbFieldBinaryReferenceTest::RunTest(const FString& Parameters)
 {
-	const FBlake3Hash::ByteArray ZeroBytes{};
-	const FBlake3Hash::ByteArray SequentialBytes{1, 2, 3, 4, 5, 6, 7, 8, 9, 10, 11, 12, 13, 14, 15, 16, 17, 18, 19, 20, 21, 22, 23, 24, 25, 26, 27, 28, 29, 30, 31, 32};
+	const FIoHash::ByteArray ZeroBytes{};
+	const FIoHash::ByteArray SequentialBytes{1, 2, 3, 4, 5, 6, 7, 8, 9, 10, 11, 12, 13, 14, 15, 16, 17, 18, 19, 20};
 
 	// Test FCbField(BinaryReference, Zero)
 	TestField<ECbFieldType::BinaryReference>(TEXT("BinaryReference, Zero"), ZeroBytes);
 
 	// Test FCbField(BinaryReference, NonZero)
-	TestField<ECbFieldType::BinaryReference>(TEXT("BinaryReference, NonZero"), SequentialBytes, FBlake3Hash(SequentialBytes));
+	TestField<ECbFieldType::BinaryReference>(TEXT("BinaryReference, NonZero"), SequentialBytes, FIoHash(SequentialBytes));
 
 	// Test FCbField(BinaryReference, NonZero) AsAnyReference
 	{
 		FCbField Field(SequentialBytes, ECbFieldType::BinaryReference);
-		TestFieldAsType(TEXT("BinaryReference, NonZero, AsReference"), Field, &FCbField::AsReference, FBlake3Hash(SequentialBytes));
+		TestFieldAsType(TEXT("BinaryReference, NonZero, AsReference"), Field, &FCbField::AsReference, FIoHash(SequentialBytes));
 	}
 
 	// Test FCbField(None) as BinaryReference
 	{
 		FCbField DefaultField;
-		TestFieldError<ECbFieldType::BinaryReference>(TEXT("BinaryReference, None"), DefaultField, ECbFieldError::TypeError, FBlake3Hash(SequentialBytes));
+		TestFieldError<ECbFieldType::BinaryReference>(TEXT("BinaryReference, None"), DefaultField, ECbFieldError::TypeError, FIoHash(SequentialBytes));
 	}
 
 	return true;
@@ -998,19 +998,19 @@ bool FCbFieldBinaryReferenceTest::RunTest(const FString& Parameters)
 IMPLEMENT_CUSTOM_SIMPLE_AUTOMATION_TEST(FCbFieldHashTest, FCbFieldTestBase, "System.Core.Serialization.CbField.Hash", CompactBinaryTestFlags)
 bool FCbFieldHashTest::RunTest(const FString& Parameters)
 {
-	const FBlake3Hash::ByteArray ZeroBytes{};
-	const FBlake3Hash::ByteArray SequentialBytes{1, 2, 3, 4, 5, 6, 7, 8, 9, 10, 11, 12, 13, 14, 15, 16, 17, 18, 19, 20, 21, 22, 23, 24, 25, 26, 27, 28, 29, 30, 31, 32};
+	const FIoHash::ByteArray ZeroBytes{};
+	const FIoHash::ByteArray SequentialBytes{1, 2, 3, 4, 5, 6, 7, 8, 9, 10, 11, 12, 13, 14, 15, 16, 17, 18, 19, 20};
 
 	// Test FCbField(Hash, Zero)
 	TestField<ECbFieldType::Hash>(TEXT("Hash, Zero"), ZeroBytes);
 
 	// Test FCbField(Hash, NonZero)
-	TestField<ECbFieldType::Hash>(TEXT("Hash, NonZero"), SequentialBytes, FBlake3Hash(SequentialBytes));
+	TestField<ECbFieldType::Hash>(TEXT("Hash, NonZero"), SequentialBytes, FIoHash(SequentialBytes));
 
 	// Test FCbField(None) as Hash
 	{
 		FCbField DefaultField;
-		TestFieldError<ECbFieldType::Hash>(TEXT("Hash, None"), DefaultField, ECbFieldError::TypeError, FBlake3Hash(SequentialBytes));
+		TestFieldError<ECbFieldType::Hash>(TEXT("Hash, None"), DefaultField, ECbFieldError::TypeError, FIoHash(SequentialBytes));
 	}
 
 	return true;
@@ -1095,7 +1095,7 @@ bool FCbFieldTimeSpanTest::RunTest(const FString& Parameters)
 IMPLEMENT_CUSTOM_SIMPLE_AUTOMATION_TEST(FCbFieldIterateReferencesTest, FCbFieldTestBase, "System.Core.Serialization.CbField.IterateReferences", CompactBinaryTestFlags)
 bool FCbFieldIterateReferencesTest::RunTest(const FString& Parameters)
 {
-	const auto MakeTestHash = [](uint32 Index) { return FBlake3::HashBuffer(&Index, sizeof(Index)); };
+	const auto MakeTestHash = [](uint32 Index) { return FIoHash::HashBuffer(&Index, sizeof(Index)); };
 
 	FCbFieldRefIterator Fields;
 	{
@@ -1489,8 +1489,8 @@ bool FCbFieldRefIteratorTest::RunTest(const FString& Parameters)
 		const FCbFieldIterator FieldIt = FCbFieldIterator::MakeRange(View);
 		const FCbFieldRefIterator FieldRefIt = FCbFieldRefIterator::MakeRange(View);
 
-		TestEqual(TEXT("FCbFieldIterator::GetRangeHash()"), FieldIt.GetRangeHash(), FBlake3::HashBuffer(View));
-		TestEqual(TEXT("FCbFieldRefIterator::GetRangeHash()"), FieldRefIt.GetRangeHash(), FBlake3::HashBuffer(View));
+		TestEqual(TEXT("FCbFieldIterator::GetRangeHash()"), FieldIt.GetRangeHash(), FIoHash::HashBuffer(View));
+		TestEqual(TEXT("FCbFieldRefIterator::GetRangeHash()"), FieldRefIt.GetRangeHash(), FIoHash::HashBuffer(View));
 
 		TestEqual(TEXT("FCbFieldIterator::GetRangeView()"), FieldIt.GetRangeView(), MakeMemoryView(Payload));
 		TestEqual(TEXT("FCbFieldRefIterator::GetRangeView()"), FieldRefIt.GetRangeView(), MakeMemoryView(Payload));
diff --git a/Engine/Source/Runtime/Core/Private/Tests/Serialization/CompactBinaryValidationTest.cpp b/Engine/Source/Runtime/Core/Private/Tests/Serialization/CompactBinaryValidationTest.cpp
--- a/Engine/Source/Runtime/Core/Private/Tests/Serialization/CompactBinaryValidationTest.cpp
+++ b/Engine/Source/Runtime/Core/Private/Tests/Serialization/CompactBinaryValidationTest.cpp
@@ -140,20 +140,20 @@ bool FCbValidateTest::RunTest(const FString& Parameters)
 	TestEqual(TEXT("ValidateCompactBinary(Valid, BoolFalse)"), Validate({uint8(ECbFieldType::BoolFalse)}), ECbValidateError::None);
 	TestEqual(TEXT("ValidateCompactBinary(Valid, BoolTrue)"), Validate({uint8(ECbFieldType::BoolTrue)}), ECbValidateError::None);
 
-	TestEqual(TEXT("ValidateCompactBinary(Valid, CompactBinaryReference)"), Validate({uint8(ECbFieldType::CompactBinaryReference), 0, 0, 0, 0, 0, 0, 0, 0, 0, 0, 0, 0, 0, 0, 0, 0, 0, 0, 0, 0, 0, 0, 0, 0, 0, 0, 0, 0, 0, 0, 0, 0}), ECbValidateError::None);
-	TestEqual(TEXT("ValidateCompactBinary(Valid, CompactBinaryReference, NoType)"), Validate({0, 0, 0, 0, 0, 0, 0, 0, 0, 0, 0, 0, 0, 0, 0, 0, 0, 0, 0, 0, 0, 0, 0, 0, 0, 0, 0, 0, 0, 0, 0, 0}, ECbFieldType::CompactBinaryReference), ECbValidateError::None);
-	TestEqual(TEXT("ValidateCompactBinary(OutOfBounds, CompactBinaryReference)"), Validate({uint8(ECbFieldType::CompactBinaryReference), 0, 0, 0, 0, 0, 0, 0, 0, 0, 0, 0, 0, 0, 0, 0, 0, 0, 0, 0, 0, 0, 0, 0, 0, 0, 0, 0, 0, 0, 0, 0}), ECbValidateError::OutOfBounds);
-	TestEqual(TEXT("ValidateCompactBinary(OutOfBounds, CompactBinaryReference, NoType)"), Validate({0, 0, 0, 0, 0, 0, 0, 0, 0, 0, 0, 0, 0, 0, 0, 0, 0, 0, 0, 0, 0, 0, 0, 0, 0, 0, 0, 0, 0, 0, 0}, ECbFieldType::CompactBinaryReference), ECbValidateError::OutOfBounds);
+	TestEqual(TEXT("ValidateCompactBinary(Valid, CompactBinaryReference)"), Validate({uint8(ECbFieldType::CompactBinaryReference), 0, 0, 0, 0, 0, 0, 0, 0, 0, 0, 0, 0, 0, 0, 0, 0, 0, 0, 0, 0}), ECbValidateError::None);
+	TestEqual(TEXT("ValidateCompactBinary(Valid, CompactBinaryReference, NoType)"), Validate({0, 0, 0, 0, 0, 0, 0, 0, 0, 0, 0, 0, 0, 0, 0, 0, 0, 0, 0, 0}, ECbFieldType::CompactBinaryReference), ECbValidateError::None);
+	TestEqual(TEXT("ValidateCompactBinary(OutOfBounds, CompactBinaryReference)"), Validate({uint8(ECbFieldType::CompactBinaryReference), 0, 0, 0, 0, 0, 0, 0, 0, 0, 0, 0, 0, 0, 0, 0, 0, 0, 0, 0}), ECbValidateError::OutOfBounds);
+	TestEqual(TEXT("ValidateCompactBinary(OutOfBounds, CompactBinaryReference, NoType)"), Validate({0, 0, 0, 0, 0, 0, 0, 0, 0, 0, 0, 0, 0, 0, 0, 0, 0, 0, 0}, ECbFieldType::CompactBinaryReference), ECbValidateError::OutOfBounds);
 
-	TestEqual(TEXT("ValidateCompactBinary(Valid, BinaryReference)"), Validate({uint8(ECbFieldType::BinaryReference), 0, 0, 0, 0, 0, 0, 0, 0, 0, 0, 0, 0, 0, 0, 0, 0, 0, 0, 0, 0, 0, 0, 0, 0, 0, 0, 0, 0, 0, 0, 0, 0}), ECbValidateError::None);
-	TestEqual(TEXT("ValidateCompactBinary(Valid, BinaryReference, NoType)"), Validate({0, 0, 0, 0, 0, 0, 0, 0, 0, 0, 0, 0, 0, 0, 0, 0, 0, 0, 0, 0, 0, 0, 0, 0, 0, 0, 0, 0, 0, 0, 0, 0}, ECbFieldType::BinaryReference), ECbValidateError::None);
-	TestEqual(TEXT("ValidateCompactBinary(OutOfBounds, BinaryReference)"), Validate({uint8(ECbFieldType::BinaryReference), 0, 0, 0, 0, 0, 0, 0, 0, 0, 0, 0, 0, 0, 0, 0, 0, 0, 0, 0, 0, 0, 0, 0, 0, 0, 0, 0, 0, 0, 0, 0}), ECbValidateError::OutOfBounds);
-	TestEqual(TEXT("ValidateCompactBinary(OutOfBounds, BinaryReference, NoType)"), Validate({0, 0, 0, 0, 0, 0, 0, 0, 0, 0, 0, 0, 0, 0, 0, 0, 0, 0, 0, 0, 0, 0, 0, 0, 0, 0, 0, 0, 0, 0, 0}, ECbFieldType::BinaryReference), ECbValidateError::OutOfBounds);
+	TestEqual(TEXT("ValidateCompactBinary(Valid, BinaryReference)"), Validate({uint8(ECbFieldType::BinaryReference), 0, 0, 0, 0, 0, 0, 0, 0, 0, 0, 0, 0, 0, 0, 0, 0, 0, 0, 0, 0}), ECbValidateError::None);
+	TestEqual(TEXT("ValidateCompactBinary(Valid, BinaryReference, NoType)"), Validate({0, 0, 0, 0, 0, 0, 0, 0, 0, 0, 0, 0, 0, 0, 0, 0, 0, 0, 0, 0}, ECbFieldType::BinaryReference), ECbValidateError::None);
+	TestEqual(TEXT("ValidateCompactBinary(OutOfBounds, BinaryReference)"), Validate({uint8(ECbFieldType::BinaryReference), 0, 0, 0, 0, 0, 0, 0, 0, 0, 0, 0, 0, 0, 0, 0, 0, 0, 0, 0}), ECbValidateError::OutOfBounds);
+	TestEqual(TEXT("ValidateCompactBinary(OutOfBounds, BinaryReference, NoType)"), Validate({0, 0, 0, 0, 0, 0, 0, 0, 0, 0, 0, 0, 0, 0, 0, 0, 0, 0, 0}, ECbFieldType::BinaryReference), ECbValidateError::OutOfBounds);
 
-	TestEqual(TEXT("ValidateCompactBinary(Valid, Hash)"), Validate({uint8(ECbFieldType::Hash), 0, 0, 0, 0, 0, 0, 0, 0, 0, 0, 0, 0, 0, 0, 0, 0, 0, 0, 0, 0, 0, 0, 0, 0, 0, 0, 0, 0, 0, 0, 0, 0}), ECbValidateError::None);
-	TestEqual(TEXT("ValidateCompactBinary(Valid, Hash, NoType)"), Validate({0, 0, 0, 0, 0, 0, 0, 0, 0, 0, 0, 0, 0, 0, 0, 0, 0, 0, 0, 0, 0, 0, 0, 0, 0, 0, 0, 0, 0, 0, 0, 0}, ECbFieldType::Hash), ECbValidateError::None);
-	TestEqual(TEXT("ValidateCompactBinary(OutOfBounds, Hash)"), Validate({uint8(ECbFieldType::Hash), 0, 0, 0, 0, 0, 0, 0, 0, 0, 0, 0, 0, 0, 0, 0, 0, 0, 0, 0, 0, 0, 0, 0, 0, 0, 0, 0, 0, 0, 0, 0}), ECbValidateError::OutOfBounds);
-	TestEqual(TEXT("ValidateCompactBinary(OutOfBounds, Hash, NoType)"), Validate({0, 0, 0, 0, 0, 0, 0, 0, 0, 0, 0, 0, 0, 0, 0, 0, 0, 0, 0, 0, 0, 0, 0, 0, 0, 0, 0, 0, 0, 0, 0}, ECbFieldType::Hash), ECbValidateError::OutOfBounds);
+	TestEqual(TEXT("ValidateCompactBinary(Valid, Hash)"), Validate({uint8(ECbFieldType::Hash), 0, 0, 0, 0, 0, 0, 0, 0, 0, 0, 0, 0, 0, 0, 0, 0, 0, 0, 0, 0}), ECbValidateError::None);
+	TestEqual(TEXT("ValidateCompactBinary(Valid, Hash, NoType)"), Validate({0, 0, 0, 0, 0, 0, 0, 0, 0, 0, 0, 0, 0, 0, 0, 0, 0, 0, 0, 0}, ECbFieldType::Hash), ECbValidateError::None);
+	TestEqual(TEXT("ValidateCompactBinary(OutOfBounds, Hash)"), Validate({uint8(ECbFieldType::Hash), 0, 0, 0, 0, 0, 0, 0, 0, 0, 0, 0, 0, 0, 0, 0, 0, 0, 0, 0}), ECbValidateError::OutOfBounds);
+	TestEqual(TEXT("ValidateCompactBinary(OutOfBounds, Hash, NoType)"), Validate({0, 0, 0, 0, 0, 0, 0, 0, 0, 0, 0, 0, 0, 0, 0, 0, 0, 0, 0}, ECbFieldType::Hash), ECbValidateError::OutOfBounds);
 
 	TestEqual(TEXT("ValidateCompactBinary(Valid, Uuid)"), Validate({uint8(ECbFieldType::Uuid), 0, 0, 0, 0, 0, 0, 0, 0, 0, 0, 0, 0, 0, 0, 0, 0}), ECbValidateError::None);
 	TestEqual(TEXT("ValidateCompactBinary(Valid, Uuid, NoType)"), Validate({0, 0, 0, 0, 0, 0, 0, 0, 0, 0, 0, 0, 0, 0, 0, 0}, ECbFieldType::Uuid), ECbValidateError::None);
@@ -351,7 +351,7 @@ bool FCbValidateAttachmentTest::RunTest(const FString& Parameters)
 	{
 		FBufferArchive Buffer;
 		FCbWriter Writer;
-		Writer.CompactBinaryReference(FBlake3::HashBuffer(BinaryView));
+		Writer.CompactBinaryReference(FIoHash::HashBuffer(BinaryView));
 		Writer.Binary(BinaryView);
 		Writer.Save(Buffer);
 		TestEqual(TEXT("ValidateCompactBinaryAttachment(HashBeforeValue)"), ValidateCompactBinaryAttachment(MakeMemoryView(Buffer), ECbValidateMode::All), ECbValidateError::InvalidPackageFormat | ECbValidateError::Padding);
@@ -361,7 +361,7 @@ bool FCbValidateAttachmentTest::RunTest(const FString& Parameters)
 		FBufferArchive Buffer;
 		FCbWriter Writer;
 		Writer.Name("Name"_ASV).Binary(BinaryView);
-		Writer.CompactBinaryReference(FBlake3::HashBuffer(BinaryView));
+		Writer.CompactBinaryReference(FIoHash::HashBuffer(BinaryView));
 		Writer.Save(Buffer);
 		TestEqual(TEXT("ValidateCompactBinaryAttachment(NameOnValue)"), ValidateCompactBinaryAttachment(MakeMemoryView(Buffer), ECbValidateMode::All), ECbValidateError::InvalidPackageFormat);
 		TestEqual(TEXT("ValidateCompactBinaryAttachment(NameOnValue, Mode)"), ValidateCompactBinaryAttachment(MakeMemoryView(Buffer), ECbValidateMode::All & ~ECbValidateMode::Package), ECbValidateError::None);
@@ -370,7 +370,7 @@ bool FCbValidateAttachmentTest::RunTest(const FString& Parameters)
 		FBufferArchive Buffer;
 		FCbWriter Writer;
 		Writer.Binary(BinaryView);
-		Writer.Name("Name"_ASV).CompactBinaryReference(FBlake3::HashBuffer(BinaryView));
+		Writer.Name("Name"_ASV).CompactBinaryReference(FIoHash::HashBuffer(BinaryView));
 		Writer.Save(Buffer);
 		TestEqual(TEXT("ValidateCompactBinaryAttachment(NameOnHash)"), ValidateCompactBinaryAttachment(MakeMemoryView(Buffer), ECbValidateMode::All), ECbValidateError::InvalidPackageFormat);
 		TestEqual(TEXT("ValidateCompactBinaryAttachment(NameOnHash, Mode)"), ValidateCompactBinaryAttachment(MakeMemoryView(Buffer), ECbValidateMode::All & ~ECbValidateMode::Package), ECbValidateError::None);
@@ -432,12 +432,12 @@ bool FCbValidatePackageTest::RunTest(const FString& Parameters)
 		Writer.CompactBinaryReference(Writer.Save().GetHash());
 		const uint64 AttachmentOffset = Writer.GetSaveSize();
 		Writer.Binary(BinaryView);
-		Writer.CompactBinaryReference(FBlake3::HashBuffer(BinaryView));
+		Writer.CompactBinaryReference(FIoHash::HashBuffer(BinaryView));
 		Writer.Null();
 		Writer.Save(Buffer);
 		TestEqual(TEXT("ValidateCompactBinaryPackage(Object+Attachment)"), ValidateCompactBinaryPackage(MakeMemoryView(Buffer), ECbValidateMode::All), ECbValidateError::None);
-		TestEqual(TEXT("ValidateCompactBinaryPackage(Object+Attachment, MissingHash)"), ValidateCompactBinaryPackage(MakeMemoryView(Buffer).LeftChop(34), ECbValidateMode::All), ECbValidateError::InvalidPackageFormat);
-		TestEqual(TEXT("ValidateCompactBinaryPackage(Object+Attachment, MissingHash, Mode)"), ValidateCompactBinaryPackage(MakeMemoryView(Buffer).LeftChop(34), ECbValidateMode::All & ~ECbValidateMode::Package), ECbValidateError::None);
+		TestEqual(TEXT("ValidateCompactBinaryPackage(Object+Attachment, MissingHash)"), ValidateCompactBinaryPackage(MakeMemoryView(Buffer).LeftChop(22), ECbValidateMode::All), ECbValidateError::InvalidPackageFormat);
+		TestEqual(TEXT("ValidateCompactBinaryPackage(Object+Attachment, MissingHash, Mode)"), ValidateCompactBinaryPackage(MakeMemoryView(Buffer).LeftChop(22), ECbValidateMode::All & ~ECbValidateMode::Package), ECbValidateError::None);
 		TestEqual(TEXT("ValidateCompactBinaryPackage(Object+Attachment, MissingNull)"), ValidateCompactBinaryPackage(MakeMemoryView(Buffer).LeftChop(1), ECbValidateMode::All), ECbValidateError::InvalidPackageFormat);
 		TestEqual(TEXT("ValidateCompactBinaryPackage(Object+Attachment, MissingNull, Mode)"), ValidateCompactBinaryPackage(MakeMemoryView(Buffer).LeftChop(1), ECbValidateMode::All & ~ECbValidateMode::Package), ECbValidateError::None);
 		++Buffer.Last(1);
@@ -477,7 +477,7 @@ bool FCbValidatePackageTest::RunTest(const FString& Parameters)
 		Writer.BeginObject();
 		Writer << "Field" << 42;
 		Writer.EndObject();
-		const FBlake3Hash ObjectHash = Writer.Save().GetHash();
+		const FIoHash ObjectHash = Writer.Save().GetHash();
 		Writer.Name("Name"_ASV);
 		Writer.CompactBinaryReference(ObjectHash);
 		Writer.Null();
@@ -493,7 +493,7 @@ bool FCbValidatePackageTest::RunTest(const FString& Parameters)
 		Writer.EndObject();
 		Writer.CompactBinaryReference(Writer.Save().GetHash());
 		Writer.Name("Name"_ASV).Binary(BinaryView);
-		Writer.CompactBinaryReference(FBlake3::HashBuffer(BinaryView));
+		Writer.CompactBinaryReference(FIoHash::HashBuffer(BinaryView));
 		Writer.Null();
 		Writer.Save(Buffer);
 		TestEqual(TEXT("ValidateCompactBinaryPackage(NameOnAttachment)"), ValidateCompactBinaryPackage(MakeMemoryView(Buffer), ECbValidateMode::All), ECbValidateError::InvalidPackageFormat);
@@ -505,7 +505,7 @@ bool FCbValidatePackageTest::RunTest(const FString& Parameters)
 		Writer.BeginObject();
 		Writer << "Field" << 42;
 		Writer.EndObject();
-		const FBlake3Hash ObjectHash = Writer.Save().GetHash();
+		const FIoHash ObjectHash = Writer.Save().GetHash();
 		Writer.CompactBinaryReference(ObjectHash);
 		Writer.BeginObject();
 		Writer << "Field" << 42;
@@ -520,9 +520,9 @@ bool FCbValidatePackageTest::RunTest(const FString& Parameters)
 		FBufferArchive Buffer;
 		TCbWriter<256> Writer;
 		Writer.Binary(BinaryView);
-		Writer.CompactBinaryReference(FBlake3::HashBuffer(BinaryView));
+		Writer.CompactBinaryReference(FIoHash::HashBuffer(BinaryView));
 		Writer.Binary(BinaryView);
-		Writer.CompactBinaryReference(FBlake3::HashBuffer(BinaryView));
+		Writer.CompactBinaryReference(FIoHash::HashBuffer(BinaryView));
 		Writer.Null();
 		Writer.Save(Buffer);
 		TestEqual(TEXT("ValidateCompactBinaryPackage(DuplicateAttachments)"), ValidateCompactBinaryPackage(MakeMemoryView(Buffer), ECbValidateMode::All), ECbValidateError::DuplicateAttachments);
diff --git a/Engine/Source/Runtime/Core/Private/Tests/Serialization/CompactBinaryWriterTest.cpp b/Engine/Source/Runtime/Core/Private/Tests/Serialization/CompactBinaryWriterTest.cpp
--- a/Engine/Source/Runtime/Core/Private/Tests/Serialization/CompactBinaryWriterTest.cpp
+++ b/Engine/Source/Runtime/Core/Private/Tests/Serialization/CompactBinaryWriterTest.cpp
@@ -2,8 +2,8 @@
 
 #include "Serialization/CompactBinaryWriter.h"
 
+#include "IO/IoHash.h"
 #include "Misc/AutomationTest.h"
-#include "Misc/Blake3.h"
 #include "Misc/DateTime.h"
 #include "Misc/Guid.h"
 #include "Misc/Timespan.h"
@@ -568,11 +568,11 @@ bool FCbWriterCompactBinaryReferenceTest::RunTest(const FString& Parameters)
 {
 	TCbWriter<256> Writer;
 
-	const FBlake3Hash::ByteArray ZeroBytes{};
-	const FBlake3Hash::ByteArray SequentialBytes{1, 2, 3, 4, 5, 6, 7, 8, 9, 10, 11, 12, 13, 14, 15, 16, 17, 18, 19, 20, 21, 22, 23, 24, 25, 26, 27, 28, 29, 30, 31, 32};
+	const FIoHash::ByteArray ZeroBytes{};
+	const FIoHash::ByteArray SequentialBytes{1, 2, 3, 4, 5, 6, 7, 8, 9, 10, 11, 12, 13, 14, 15, 16, 17, 18, 19, 20};
 
-	const FBlake3Hash Values[] = { FBlake3Hash(ZeroBytes), FBlake3Hash(SequentialBytes) };
-	for (const FBlake3Hash& Value : Values)
+	const FIoHash Values[] = { FIoHash(ZeroBytes), FIoHash(SequentialBytes) };
+	for (const FIoHash& Value : Values)
 	{
 		Writer.CompactBinaryReference(Value);
 	}
@@ -580,7 +580,7 @@ bool FCbWriterCompactBinaryReferenceTest::RunTest(const FString& Parameters)
 	FCbFieldRefIterator Fields = Writer.Save();
 	if (TestEqual(TEXT("FCbWriter(CompactBinaryReference) Validate"), ValidateCompactBinaryRange(Fields.GetBuffer(), ECbValidateMode::All), ECbValidateError::None))
 	{
-		const FBlake3Hash* CheckValue = Values;
+		const FIoHash* CheckValue = Values;
 		for (FCbField Field : Fields)
 		{
 			TestEqual(TEXT("FCbWriter(CompactBinaryReference).AsCompactBinaryReference()"), Field.AsCompactBinaryReference(), *CheckValue++);
@@ -598,11 +598,11 @@ bool FCbWriterBinaryReferenceTest::RunTest(const FString& Parameters)
 {
 	TCbWriter<256> Writer;
 
-	const FBlake3Hash::ByteArray ZeroBytes{};
-	const FBlake3Hash::ByteArray SequentialBytes{1, 2, 3, 4, 5, 6, 7, 8, 9, 10, 11, 12, 13, 14, 15, 16, 17, 18, 19, 20, 21, 22, 23, 24, 25, 26, 27, 28, 29, 30, 31, 32};
+	const FIoHash::ByteArray ZeroBytes{};
+	const FIoHash::ByteArray SequentialBytes{1, 2, 3, 4, 5, 6, 7, 8, 9, 10, 11, 12, 13, 14, 15, 16, 17, 18, 19, 20};
 
-	const FBlake3Hash Values[] = { FBlake3Hash(ZeroBytes), FBlake3Hash(SequentialBytes) };
-	for (const FBlake3Hash& Value : Values)
+	const FIoHash Values[] = { FIoHash(ZeroBytes), FIoHash(SequentialBytes) };
+	for (const FIoHash& Value : Values)
 	{
 		Writer.BinaryReference(Value);
 	}
@@ -610,7 +610,7 @@ bool FCbWriterBinaryReferenceTest::RunTest(const FString& Parameters)
 	FCbFieldRefIterator Fields = Writer.Save();
 	if (TestEqual(TEXT("FCbWriter(BinaryReference) Validate"), ValidateCompactBinaryRange(Fields.GetBuffer(), ECbValidateMode::All), ECbValidateError::None))
 	{
-		const FBlake3Hash* CheckValue = Values;
+		const FIoHash* CheckValue = Values;
 		for (FCbField Field : Fields)
 		{
 			TestEqual(TEXT("FCbWriter(BinaryReference).AsBinaryReference()"), Field.AsBinaryReference(), *CheckValue++);
@@ -628,11 +628,11 @@ bool FCbWriterHashTest::RunTest(const FString& Parameters)
 {
 	TCbWriter<256> Writer;
 
-	const FBlake3Hash::ByteArray ZeroBytes{};
-	const FBlake3Hash::ByteArray SequentialBytes{1, 2, 3, 4, 5, 6, 7, 8, 9, 10, 11, 12, 13, 14, 15, 16, 17, 18, 19, 20, 21, 22, 23, 24, 25, 26, 27, 28, 29, 30, 31, 32};
+	const FIoHash::ByteArray ZeroBytes{};
+	const FIoHash::ByteArray SequentialBytes{1, 2, 3, 4, 5, 6, 7, 8, 9, 10, 11, 12, 13, 14, 15, 16, 17, 18, 19, 20};
 
-	const FBlake3Hash Values[] = { FBlake3Hash(ZeroBytes), FBlake3Hash(SequentialBytes) };
-	for (const FBlake3Hash& Value : Values)
+	const FIoHash Values[] = { FIoHash(ZeroBytes), FIoHash(SequentialBytes) };
+	for (const FIoHash& Value : Values)
 	{
 		Writer.Hash(Value);
 	}
@@ -640,7 +640,7 @@ bool FCbWriterHashTest::RunTest(const FString& Parameters)
 	FCbFieldRefIterator Fields = Writer.Save();
 	if (TestEqual(TEXT("FCbWriter(Hash) Validate"), ValidateCompactBinaryRange(Fields.GetBuffer(), ECbValidateMode::All), ECbValidateError::None))
 	{
-		const FBlake3Hash* CheckValue = Values;
+		const FIoHash* CheckValue = Values;
 		for (FCbField Field : Fields)
 		{
 			TestEqual(TEXT("FCbWriter(Hash).AsHash()"), Field.AsHash(), *CheckValue++);
@@ -816,10 +816,10 @@ bool FCbWriterComplexTest::RunTest(const FString& Parameters)
 		Writer.Name("False"_ASV).Bool(false);
 		Writer.Name("True"_ASV).Bool(true);
 
-		Writer.Name("CompactBinaryReference"_ASV).CompactBinaryReference(FBlake3Hash());
-		Writer.Name("BinaryReference"_ASV).BinaryReference(FBlake3Hash());
+		Writer.Name("CompactBinaryReference"_ASV).CompactBinaryReference(FIoHash());
+		Writer.Name("BinaryReference"_ASV).BinaryReference(FIoHash());
 
-		Writer.Name("Hash"_ASV).Hash(FBlake3Hash());
+		Writer.Name("Hash"_ASV).Hash(FIoHash());
 		Writer.Name("Uuid"_ASV).Uuid(FGuid());
 
 		Writer.Name("DateTimeZero"_ASV).DateTimeTicks(0);
@@ -982,7 +982,7 @@ bool FCbWriterStreamTest::RunTest(const FString& Parameters)
 
 		Writer << "False"_ASV << false << "True"_ASV << true;
 
-		Writer << "Hash"_ASV << FBlake3Hash();
+		Writer << "Hash"_ASV << FIoHash();
 		Writer << "Uuid"_ASV << FGuid();
 
 		Writer << "DateTime"_ASV << FDateTime(2020, 5, 13, 15, 10);
diff --git a/Engine/Source/Runtime/Core/Public/Serialization/CompactBinary.h b/Engine/Source/Runtime/Core/Public/Serialization/CompactBinary.h
--- a/Engine/Source/Runtime/Core/Public/Serialization/CompactBinary.h
+++ b/Engine/Source/Runtime/Core/Public/Serialization/CompactBinary.h
@@ -4,9 +4,9 @@
 
 #include "CoreTypes.h"
 #include "Containers/StringView.h"
+#include "IO/IoHash.h"
 #include "Memory/MemoryView.h"
 #include "Memory/SharedBuffer.h"
-#include "Misc/Blake3.h"
 #include "Misc/EnumClassFlags.h"
 #include "Templates/IsTriviallyDestructible.h"
 
@@ -61,6 +61,7 @@
  */
 
 class FArchive;
+class FBlake3;
 class FCbArray;
 class FCbField;
 class FCbObject;
@@ -309,7 +310,7 @@ public:
 	CORE_API uint64 GetRangeSize() const;
 
 	/** Calculate the hash of every field in the range. */
-	CORE_API FBlake3Hash GetRangeHash() const;
+	CORE_API FIoHash GetRangeHash() const;
 	/** Calculate the hash of every field in the range. */
 	CORE_API void GetRangeHash(FBlake3& Hash) const;
 
@@ -496,14 +497,14 @@ public:
 	CORE_API bool AsBool(bool bDefault = false);
 
 	/** Access the field as a hash referencing compact binary data. Returns the provided default on error. */
-	CORE_API FBlake3Hash AsCompactBinaryReference(const FBlake3Hash& Default = FBlake3Hash());
+	CORE_API FIoHash AsCompactBinaryReference(const FIoHash& Default = FIoHash());
 	/** Access the field as a hash referencing binary data. Returns the provided default on error. */
-	CORE_API FBlake3Hash AsBinaryReference(const FBlake3Hash& Default = FBlake3Hash());
+	CORE_API FIoHash AsBinaryReference(const FIoHash& Default = FIoHash());
 	/** Access the field as a hash referencing compact binary data or binary data. Returns the provided default on error. */
-	CORE_API FBlake3Hash AsReference(const FBlake3Hash& Default = FBlake3Hash());
+	CORE_API FIoHash AsReference(const FIoHash& Default = FIoHash());
 
 	/** Access the field as a hash. Returns the provided default on error. */
-	CORE_API FBlake3Hash AsHash(const FBlake3Hash& Default = FBlake3Hash());
+	CORE_API FIoHash AsHash(const FIoHash& Default = FIoHash());
 
 	/** Access the field as a UUID. Returns a nil UUID on error. */
 	CORE_API FGuid AsUuid();
@@ -574,7 +575,7 @@ public:
 	CORE_API uint64 GetSize() const;
 
 	/** Calculate the hash of the field, including the type and name. */
-	CORE_API FBlake3Hash GetHash() const;
+	CORE_API FIoHash GetHash() const;
 	/** Calculate the hash of the field, including the type and name. */
 	CORE_API void GetHash(FBlake3& Hash) const;
 
@@ -763,7 +764,7 @@ public:
 	CORE_API uint64 GetSize() const;
 
 	/** Calculate the hash of the array if serialized by itself with no name. */
-	CORE_API FBlake3Hash GetHash() const;
+	CORE_API FIoHash GetHash() const;
 	/** Calculate the hash of the array if serialized by itself with no name. */
 	CORE_API void GetHash(FBlake3& Hash) const;
 
@@ -861,7 +862,7 @@ public:
 	CORE_API uint64 GetSize() const;
 
 	/** Calculate the hash of the object if serialized by itself with no name. */
-	CORE_API FBlake3Hash GetHash() const;
+	CORE_API FIoHash GetHash() const;
 	/** Calculate the hash of the object if serialized by itself with no name. */
 	CORE_API void GetHash(FBlake3& Hash) const;
 
diff --git a/Engine/Source/Runtime/Core/Public/Serialization/CompactBinaryPackage.h b/Engine/Source/Runtime/Core/Public/Serialization/CompactBinaryPackage.h
--- a/Engine/Source/Runtime/Core/Public/Serialization/CompactBinaryPackage.h
+++ b/Engine/Source/Runtime/Core/Public/Serialization/CompactBinaryPackage.h
@@ -4,8 +4,8 @@
 
 #include "CoreTypes.h"
 #include "Containers/Array.h"
+#include "IO/IoHash.h"
 #include "Memory/SharedBuffer.h"
-#include "Misc/Blake3.h"
 #include "Serialization/CompactBinary.h"
 
 class FArchive;
@@ -36,7 +36,7 @@ public:
 	}
 
 	/** Construct a compact binary attachment. Value is cloned if not owned. Hash must match Value. */
-	inline explicit FCbAttachment(FCbFieldRefIterator Value, const FBlake3Hash& Hash)
+	inline explicit FCbAttachment(FCbFieldRefIterator Value, const FIoHash& Hash)
 		: FCbAttachment(MoveTemp(Value), &Hash)
 	{
 	}
@@ -48,7 +48,7 @@ public:
 	}
 
 	/** Construct a binary attachment. Value is cloned if not owned. Hash must match Value. */
-	inline explicit FCbAttachment(FSharedBuffer Value, const FBlake3Hash& Hash)
+	inline explicit FCbAttachment(FSharedBuffer Value, const FIoHash& Hash)
 		: FCbAttachment(MoveTemp(Value), &Hash)
 	{
 	}
@@ -75,7 +75,7 @@ public:
 	inline bool IsCompactBinary() const { return CompactBinary.HasValue(); }
 
 	/** Returns the hash of the attachment value. */
-	inline const FBlake3Hash& GetHash() const { return Hash; }
+	inline const FIoHash& GetHash() const { return Hash; }
 
 	/** Compares attachments by their hash. Any discrepancy in type must be handled externally. */
 	inline bool operator==(const FCbAttachment& Attachment) const { return Hash == Attachment.Hash; }
@@ -108,15 +108,15 @@ public:
 	CORE_API void Save(FArchive& Ar) const;
 
 private:
-	CORE_API FCbAttachment(FCbFieldRefIterator Value, const FBlake3Hash* Hash);
-	CORE_API FCbAttachment(FSharedBuffer Value, const FBlake3Hash* Hash);
+	CORE_API FCbAttachment(FCbFieldRefIterator Value, const FIoHash* Hash);
+	CORE_API FCbAttachment(FSharedBuffer Value, const FIoHash* Hash);
 
 	/** An owned buffer containing the binary or compact binary data. */
 	FSharedBuffer Buffer;
 	/** A field iterator that is valid only for compact binary attachments. */
 	FCbFieldIterator CompactBinary;
 	/** A hash of the attachment value. */
-	FBlake3Hash Hash;
+	FIoHash Hash;
 };
 
 /** Hashes attachments by their hash. Any discrepancy in type must be handled externally. */
@@ -158,7 +158,7 @@ public:
 	 *
 	 * The resolver may return a null buffer to skip resolving an attachment for the hash.
 	 */
-	using FAttachmentResolver = TFunctionRef<FSharedBuffer (const FBlake3Hash& Hash)>;
+	using FAttachmentResolver = TFunctionRef<FSharedBuffer (const FIoHash& Hash)>;
 
 	/** Construct a null package. */
 	FCbPackage() = default;
@@ -190,7 +190,7 @@ public:
 	 * @param InObject The root object, which will be cloned unless it is owned.
 	 * @param InObjectHash The hash of the object, which must match to avoid validation errors.
 	 */
-	inline explicit FCbPackage(FCbObjectRef InObject, const FBlake3Hash& InObjectHash)
+	inline explicit FCbPackage(FCbObjectRef InObject, const FIoHash& InObjectHash)
 	{
 		SetObject(MoveTemp(InObject), InObjectHash);
 	}
@@ -202,7 +202,7 @@ public:
 	 * @param InObjectHash The hash of the object, which must match to avoid validation errors.
 	 * @param InResolver A function that is invoked for every reference and binary reference field.
 	 */
-	inline explicit FCbPackage(FCbObjectRef InObject, const FBlake3Hash& InObjectHash, FAttachmentResolver InResolver)
+	inline explicit FCbPackage(FCbObjectRef InObject, const FIoHash& InObjectHash, FAttachmentResolver InResolver)
 	{
 		SetObject(MoveTemp(InObject), InObjectHash, InResolver);
 	}
@@ -223,7 +223,7 @@ public:
 	inline const FCbObjectRef& GetObject() const { return Object; }
 
 	/** Returns the has of the compact binary object for the package. */
-	inline const FBlake3Hash& GetObjectHash() const { return ObjectHash; }
+	inline const FIoHash& GetObjectHash() const { return ObjectHash; }
 
 	/**
 	 * Set the root object without gathering attachments.
@@ -252,7 +252,7 @@ public:
 	 * @param InObject The root object, which will be cloned unless it is owned.
 	 * @param InObjectHash The hash of the object, which must match to avoid validation errors.
 	 */
-	inline void SetObject(FCbObjectRef InObject, const FBlake3Hash& InObjectHash)
+	inline void SetObject(FCbObjectRef InObject, const FIoHash& InObjectHash)
 	{
 		SetObject(MoveTemp(InObject), &InObjectHash, nullptr);
 	}
@@ -264,7 +264,7 @@ public:
 	 * @param InObjectHash The hash of the object, which must match to avoid validation errors.
 	 * @param InResolver A function that is invoked for every reference and binary reference field.
 	 */
-	inline void SetObject(FCbObjectRef InObject, const FBlake3Hash& InObjectHash, FAttachmentResolver InResolver)
+	inline void SetObject(FCbObjectRef InObject, const FIoHash& InObjectHash, FAttachmentResolver InResolver)
 	{
 		SetObject(MoveTemp(InObject), &InObjectHash, &InResolver);
 	}
@@ -278,7 +278,7 @@ public:
 	 * @return The attachment, or null if the attachment is not found.
 	 * @note The returned pointer is only valid until the attachments on this package are modified.
 	 */
-	CORE_API const FCbAttachment* FindAttachment(const FBlake3Hash& Hash) const;
+	CORE_API const FCbAttachment* FindAttachment(const FIoHash& Hash) const;
 
 	/** Find an attachment if it exists in the package. */
 	inline const FCbAttachment* FindAttachment(const FCbAttachment& Attachment) const
@@ -303,7 +303,7 @@ public:
 	 *
 	 * @return Number of attachments removed, which will be either 0 or 1.
 	 */
-	CORE_API int32 RemoveAttachment(const FBlake3Hash& Hash);
+	CORE_API int32 RemoveAttachment(const FIoHash& Hash);
 	inline int32 RemoveAttachment(const FCbAttachment& Attachment) { return RemoveAttachment(Attachment.GetHash()); }
 
 	/** Compares packages by their object and attachment hashes. */
@@ -338,7 +338,7 @@ public:
 	CORE_API void Save(FArchive& Ar) const;
 
 private:
-	CORE_API void SetObject(FCbObjectRef Object, const FBlake3Hash* Hash, FAttachmentResolver* Resolver);
+	CORE_API void SetObject(FCbObjectRef Object, const FIoHash* Hash, FAttachmentResolver* Resolver);
 	CORE_API void AddAttachment(const FCbAttachment& Attachment, FAttachmentResolver* Resolver);
 
 	void GatherAttachments(const FCbFieldIterator& Fields, FAttachmentResolver Resolver);
@@ -346,7 +346,7 @@ private:
 	/** Attachments ordered by their hash. */
 	TArray<FCbAttachment> Attachments;
 	FCbObjectRef Object;
-	FBlake3Hash ObjectHash;
+	FIoHash ObjectHash;
 };
 
 ///////////////////////////////////////////////////////////////////////////////////////////////////////////////////////
diff --git a/Engine/Source/Runtime/Core/Public/Serialization/CompactBinaryWriter.h b/Engine/Source/Runtime/Core/Public/Serialization/CompactBinaryWriter.h
--- a/Engine/Source/Runtime/Core/Public/Serialization/CompactBinaryWriter.h
+++ b/Engine/Source/Runtime/Core/Public/Serialization/CompactBinaryWriter.h
@@ -10,9 +10,9 @@
 ///////////////////////////////////////////////////////////////////////////////////////////////////////////////////////
 
 class FArchive;
-struct FBlake3Hash;
 struct FDateTime;
 struct FGuid;
+struct FIoHash;
 struct FTimespan;
 
 /**
@@ -166,12 +166,12 @@ public:
 	CORE_API void Bool(bool bValue);
 
 	/** Write a field referencing compact binary data by its hash. */
-	CORE_API void CompactBinaryReference(const FBlake3Hash& Value);
+	CORE_API void CompactBinaryReference(const FIoHash& Value);
 	/** Write a field referencing binary data by its hash. */
-	CORE_API void BinaryReference(const FBlake3Hash& Value);
+	CORE_API void BinaryReference(const FIoHash& Value);
 
 	/** Write a hash field. */
-	CORE_API void Hash(const FBlake3Hash& Value);
+	CORE_API void Hash(const FIoHash& Value);
 	/** Write a UUID field. */
 	CORE_API void Uuid(const FGuid& Value);
 
@@ -372,7 +372,7 @@ inline FCbWriter& operator<<(FCbWriter& Writer, bool Value)
 	return Writer;
 }
 
-inline FCbWriter& operator<<(FCbWriter& Writer, const FBlake3Hash& Value)
+inline FCbWriter& operator<<(FCbWriter& Writer, const FIoHash& Value)
 {
 	Writer.Hash(Value);
 	return Writer;
