Add Stall Detector API, enabled only for Editor builds for Windows

This code is meant to help locate and send reports/telemetry for slow code pathways that create unresponsive conditions
FGameThreadHitchHeartBeatThreaded was considered, but doesn't fit Editor's needs because its designed around general GameThread deadlines
Editor workloads are much less homogenous, and proper async support for a consistent GameThread deadline in Editor is a ways away
This necessitates a more focused approach where we can instrument specific routines such that each issue their own telemetry report
Add a "Stalls" counter in the Frame Rate and Memory title bar stats
Add LogStall Log category for viewing details about stalls that have occurred
Introduces a stall counter object on the GameThread to collect statistical data about stalls (this will not report to telemetry)
Future changes will introduce report objects into specific routines to upload to crashreporter
Future changes will introduce support for non-Windows OSes

diff --git a/Engine/Source/Editor/MainFrame/Private/MainFrameModule.cpp b/Engine/Source/Editor/MainFrame/Private/MainFrameModule.cpp
--- a/Engine/Source/Editor/MainFrame/Private/MainFrameModule.cpp
+++ b/Engine/Source/Editor/MainFrame/Private/MainFrameModule.cpp
@@ -31,6 +31,7 @@
 #include "HAL/PlatformFileManager.h"
 #include "Toolkits/FConsoleCommandExecutor.h"
 #include "Misc/MessageDialog.h"
+#include "ProfilingDebugging/StallDetector.h"
 
 DEFINE_LOG_CATEGORY(LogMainFrame);
 #define LOCTEXT_NAMESPACE "FMainFrameModule"
@@ -395,6 +396,13 @@ TSharedRef<SWidget> FMainFrameModule::MakeDeveloperTools( const TArray<FMainFram
 			return FText::AsNumber(GUObjectArray.GetObjectArrayNumMinusAvailable());
 		}
 
+#if STALL_DETECTOR
+		static FText GetStallCountAsString()
+		{
+			return FText::AsNumber(UE::FStallDetectorStats::TotalTriggeredCount);
+		}
+#endif // STALL_DETECTOR
+
 		static void OpenVideo( FString SourceFilePath )
 		{
 			FPlatformProcess::ExploreFolder( *( FPaths::GetPath(SourceFilePath) ) );
@@ -477,6 +485,15 @@ TSharedRef<SWidget> FMainFrameModule::MakeDeveloperTools( const TArray<FMainFram
 		ObjDeveloperTool.Value = TAttribute<FText>::Create(&Local::GetUObjectCountAsString);
 		Local::AddSlot(DeveloperToolWidget, LabelFont, NormalFixedFont, ObjDeveloperTool);
 	}
+	{
+#if STALL_DETECTOR
+		FMainFrameDeveloperTool StallsDeveloperTool;
+		StallsDeveloperTool.Visibility = TAttribute<EVisibility>::Create(&Local::ShouldShowFrameRateAndMemory);
+		StallsDeveloperTool.Label = LOCTEXT("StallsLabel", "Stalls: ");
+		StallsDeveloperTool.Value = TAttribute<FText>::Create(&Local::GetStallCountAsString);
+		Local::AddSlot(DeveloperToolWidget, LabelFont, NormalFixedFont, StallsDeveloperTool);
+#endif // STALL_DETECTOR
+	}
 
 
 	// Invisible border, so that we can animate our box panel size
diff --git a/Engine/Source/Programs/CrashReportClient/Private/CrashReportClientApp.cpp b/Engine/Source/Programs/CrashReportClient/Private/CrashReportClientApp.cpp
--- a/Engine/Source/Programs/CrashReportClient/Private/CrashReportClientApp.cpp
+++ b/Engine/Source/Programs/CrashReportClient/Private/CrashReportClientApp.cpp
@@ -512,7 +512,7 @@ FPlatformErrorReport CollectErrorReport(FRecoveryService* RecoveryService, uint3
 	if (RecoveryService && 
 		DirectoryExists && 
 		SharedCrashContext.UserSettings.bSendUsageData && 
-		SharedCrashContext.CrashType != ECrashContextType::Ensure)
+		!FGenericCrashContext::IsTypeContinuable(SharedCrashContext.CrashType)
 	{
 		RecoveryService->CollectFiles(ReportDirectoryAbsolutePath);
 	}
@@ -655,6 +655,10 @@ bool IsCrashReportAvailable(uint32 WatchedProcess, FSharedCrashContext& CrashCon
 					FDiagnosticLogger::Get().LogEvent(TEXT("Pipe/Ensure"));
 					break;
 
+				case ECrashContextType::Stall:
+					FDiagnosticLogger::Get().LogEvent(TEXT("Pipe/Stall"));
+					break;
+
 				case ECrashContextType::Crash:
 					FDiagnosticLogger::Get().LogEvent(TEXT("Pipe/Crash"));
 					bMonitoredAppCrashed = true;
diff --git a/Engine/Source/Runtime/Core/Private/GenericPlatform/GenericPlatformCrashContext.cpp b/Engine/Source/Runtime/Core/Private/GenericPlatform/GenericPlatformCrashContext.cpp
--- a/Engine/Source/Runtime/Core/Private/GenericPlatform/GenericPlatformCrashContext.cpp
+++ b/Engine/Source/Runtime/Core/Private/GenericPlatform/GenericPlatformCrashContext.cpp
@@ -62,6 +62,7 @@ const TCHAR* const FGenericCrashContext::NewLineTag = TEXT( "&nl;" );
 const TCHAR* const FGenericCrashContext::CrashTypeCrash = TEXT("Crash");
 const TCHAR* const FGenericCrashContext::CrashTypeAssert = TEXT("Assert");
 const TCHAR* const FGenericCrashContext::CrashTypeEnsure = TEXT("Ensure");
+const TCHAR* const FGenericCrashContext::CrashTypeStall = TEXT("Stall");
 const TCHAR* const FGenericCrashContext::CrashTypeGPU = TEXT("GPUCrash");
 const TCHAR* const FGenericCrashContext::CrashTypeHang = TEXT("Hang");
 const TCHAR* const FGenericCrashContext::CrashTypeAbnormalShutdown = TEXT("AbnormalShutdown");
@@ -618,6 +619,7 @@ void FGenericCrashContext::SerializeContentToBuffer() const
 	AddCrashProperty( TEXT( "CrashGUID" ), (const TCHAR*)CrashGUID);
 
 	AddCrashProperty( TEXT( "IsEnsure" ), (Type == ECrashContextType::Ensure) );
+	AddCrashProperty( TEXT( "IsStall"), (Type == ECrashContextType::Stall) );
 	AddCrashProperty( TEXT( "IsAssert" ), (Type == ECrashContextType::Assert) );
 	AddCrashProperty( TEXT( "CrashType" ), GetCrashTypeString(Type) );
 	AddCrashProperty( TEXT( "ErrorMessage" ), ErrorMessage );
@@ -726,7 +728,7 @@ void FGenericCrashContext::GetUniqueCrashName(TCHAR* GUIDBuffer, int32 BufferSiz
 
 const bool FGenericCrashContext::IsFullCrashDump() const
 {
-	if(Type == ECrashContextType::Ensure)
+	if (FGenericCrashContext::IsTypeContinuable(Type))
 	{
 		return (NCached::Session.CrashDumpMode == (int32)ECrashDumpMode::FullDumpAlways);
 	}
@@ -920,6 +922,8 @@ const TCHAR* FGenericCrashContext::GetCrashTypeString(ECrashContextType Type)
 		return CrashTypeGPU;
 	case ECrashContextType::Ensure:
 		return CrashTypeEnsure;
+	case ECrashContextType::Stall:
+		return CrashTypeStall;
 	case ECrashContextType::Assert:
 		return CrashTypeAssert;
 	case ECrashContextType::AbnormalShutdown:
diff --git a/Engine/Source/Runtime/Core/Private/ProfilingDebugging/StallDetector.cpp b/Engine/Source/Runtime/Core/Private/ProfilingDebugging/StallDetector.cpp
new file mode 100644
--- /dev/null
+++ b/Engine/Source/Runtime/Core/Private/ProfilingDebugging/StallDetector.cpp
@@ -0,0 +1,379 @@
+// Copyright Epic Games, Inc. All Rights Reserved.
+#include "ProfilingDebugging/StallDetector.h"
+
+#if STALL_DETECTOR
+
+#include "HAL/ExceptionHandling.h"
+#include "HAL/PlatformProcess.h"
+#include "HAL/PlatformStackWalk.h"
+#include "HAL/PlatformTime.h"
+#include "HAL/PlatformTLS.h"
+#include "HAL/Runnable.h"
+#include "HAL/RunnableThread.h"
+#include "ProfilingDebugging/CpuProfilerTrace.h"
+
+// force normal behavior in the face of debug configuration and debugger attached
+#define STALL_DETECTOR_DEBUG 0
+
+// use the heart beat clock to account for process suspend
+#define STALL_DETECTOR_HEART_BEAT_CLOCK 1
+
+#if STALL_DETECTOR_DEBUG && _MSC_VER
+#pragma optimize( "", off )
+#endif // STALL_DETECTOR_DEBUG && _MSC_VER
+
+#if STALL_DETECTOR_HEART_BEAT_CLOCK
+ #include "HAL/ThreadHeartBeat.h"
+#endif // STALL_DETECTOR_HEART_BEAT_CLOCK
+
+DEFINE_LOG_CATEGORY(LogStall);
+
+/**
+* The reference count for the resources for this API
+**/
+static uint32 InitCount = 0;
+
+/**
+* Stall Detector Thread
+**/
+
+namespace UE
+{
+	class FStallDetectorRunnable : public FRunnable
+	{
+	public:
+		FStallDetectorRunnable();
+
+		// FRunnable implementation
+		virtual uint32 Run() override;
+		virtual void Stop() override
+		{
+			StopThread = true;
+		}
+		virtual void Exit() override
+		{
+			Stop();
+		}
+
+		bool GetStartedThread()
+		{
+			return StartedThread;
+		}
+		
+#if STALL_DETECTOR_HEART_BEAT_CLOCK
+		FThreadHeartBeatClock& GetClock()
+		{
+			return Clock;
+		}
+#endif
+
+	private:
+		bool StartedThread;
+		bool StopThread;
+
+#if STALL_DETECTOR_HEART_BEAT_CLOCK
+		FThreadHeartBeatClock Clock;
+#endif
+	};
+
+	static FStallDetectorRunnable* Runnable = nullptr;
+	static FRunnableThread* Thread = nullptr;
+}
+
+UE::FStallDetectorRunnable::FStallDetectorRunnable()
+	: StartedThread(false)
+	, StopThread(false)
+#if STALL_DETECTOR_HEART_BEAT_CLOCK
+	, Clock(50.0/1000.0) // the clamped time interval that each tick of the clock can possibly advance
+#endif
+{
+}
+
+uint32 UE::FStallDetectorRunnable::Run()
+{
+	while (!StopThread)
+	{
+		TRACE_CPUPROFILER_EVENT_SCOPE(FStallDetectorRunnable::Run);
+
+#if STALL_DETECTOR_HEART_BEAT_CLOCK
+		Clock.Tick();
+#endif
+
+		// Clock has been ticked
+		StartedThread = true;
+
+		// Use this timestamp to try to avoid marginal triggering
+		double Seconds = FStallDetector::Seconds();
+
+		// Check the detectors
+		{
+			FScopeLock ScopeLock(&FStallDetector::GetCriticalSection());
+			for (FStallDetector* Detector : FStallDetector::GetInstances())
+			{
+				Detector->Check(false, Seconds);
+			}
+		}
+
+		// Sleep a an interval, the resolution at which we want to detect an overage
+		FPlatformProcess::Sleep(0.005);
+	}
+
+	return 0;
+}
+
+/**
+* Stall Detector Stats
+**/
+
+FCriticalSection UE::FStallDetectorStats::CriticalSection;
+TSet<UE::FStallDetectorStats*> UE::FStallDetectorStats::Instances;
+std::atomic<uint32> UE::FStallDetectorStats::TotalTriggeredCount = 0;
+std::atomic<uint32> UE::FStallDetectorStats::TotalReportedCount = 0;
+
+UE::FStallDetectorStats::FStallDetectorStats(const TCHAR* InName, const double InBudgetSeconds, const EStallDetectorReportingMode InReportingMode)
+	: Name(InName)
+	, BudgetSeconds(InBudgetSeconds)
+	, ReportingMode(InReportingMode)
+	, TriggerCount(0)
+	, OverageSeconds(0.0)
+{
+	// Add at the end of construction
+	FScopeLock ScopeLock(&CriticalSection);
+	Instances.Add(this);
+}
+
+UE::FStallDetectorStats::~FStallDetectorStats()
+{
+	// Remove at the beginning of destruction
+	FScopeLock ScopeLock(&CriticalSection);
+	Instances.Remove(this);
+}
+
+/**
+* Stall Detector
+**/
+
+FCriticalSection UE::FStallDetector::CriticalSection;
+TSet<UE::FStallDetector*> UE::FStallDetector::Instances;
+
+UE::FStallDetector::FStallDetector(FStallDetectorStats& InStats)
+	: Stats(InStats)
+	, bPersistent(false)
+	, Triggered(false)
+{
+	check(InitCount);
+
+	ThreadId = FPlatformTLS::GetCurrentThreadId();
+	StartSeconds = FStallDetector::Seconds();
+
+	// Add at the end of construction
+	FScopeLock ScopeLock(&CriticalSection);
+	Instances.Add(this);
+}
+
+UE::FStallDetector::~FStallDetector()
+{
+	// Remove at the beginning of destruction
+	{
+		FScopeLock ScopeLock(&CriticalSection);
+		Instances.Remove(this);
+	}
+
+	if (!bPersistent)
+	{
+		Check(true);
+	}
+}
+
+void UE::FStallDetector::Check(bool bIsComplete, double InWhenToCheckSeconds)
+{
+	double CheckSeconds = InWhenToCheckSeconds;
+	if (InWhenToCheckSeconds == 0.0)
+	{
+		CheckSeconds = FStallDetector::Seconds();
+	}
+
+	double DeltaSeconds = CheckSeconds - StartSeconds;
+	double OverageSeconds = DeltaSeconds - Stats.BudgetSeconds;
+
+	if (Triggered)
+	{
+		if (bIsComplete)
+		{
+			Stats.OverageSeconds += OverageSeconds;
+
+#if STALL_DETECTOR_DEBUG
+			FString OverageString = FString::Printf(TEXT("[FStallDetector] [%s] Overage of %f\n"), Stats.Name, OverageSeconds);
+			FPlatformMisc::LocalPrint(OverageString.GetCharArray().GetData());
+#endif
+			UE_LOG(LogStall, Display, TEXT("Stall detector '%s' exceeded budget of %fs, and completed in %fs"), Stats.Name, Stats.BudgetSeconds, OverageSeconds);
+		}
+	}
+	else
+	{
+		if (OverageSeconds > 0.0)
+		{
+			bool PreviousTriggered = false;
+			if (Triggered.compare_exchange_strong(PreviousTriggered, true, std::memory_order_acquire, std::memory_order_relaxed))
+			{
+#if STALL_DETECTOR_DEBUG
+				FString OverageString = FString::Printf(TEXT("[FStallDetector] [%s] Triggered at %f\n"), Stats.Name, CheckSeconds);
+				FPlatformMisc::LocalPrint(OverageString.GetCharArray().GetData());
+#endif
+				Stats.TriggerCount++;
+				OnStallDetected(ThreadId, DeltaSeconds);
+			}
+		}
+	}
+}
+
+void UE::FStallDetector::CheckAndReset()
+{
+	double CheckSeconds = FStallDetector::Seconds();
+
+	// if this is the first call to CheckAndReset
+	if (!bPersistent)
+	{
+		// never the first call again, because the timespan between construction and the first call isn't valid don't perform a check
+		bPersistent = true;
+	}
+	else
+	{
+		// only perform the check on the second call
+		Check(true, CheckSeconds);
+	}
+
+	StartSeconds = CheckSeconds;
+	Triggered = false;
+}
+
+void UE::FStallDetector::OnStallDetected(uint32 InThreadId, const double InElapsedSeconds)
+{
+	TRACE_CPUPROFILER_EVENT_SCOPE(FStallDetector::OnStallDetected);
+
+	Stats.TotalTriggeredCount++;
+
+	//
+	// Determine if we want to undermine the specified reporting mode
+	//
+
+	EStallDetectorReportingMode ReportingMode = Stats.ReportingMode;
+
+	bool bDisableReporting = false;
+#if UE_BUILD_DEBUG
+	bDisableReporting |= true; // Do not generate a report in debug configurations due to performance characteristics
+#endif
+	bDisableReporting |= FPlatformMisc::IsDebuggerPresent(); // Do not generate a report if we detect the debugger mucking with things
+
+	if (bDisableReporting)
+	{
+#if !STALL_DETECTOR_DEBUG
+		ReportingMode = EStallDetectorReportingMode::Never;
+#endif
+	}
+
+	//
+	// Resolve reporting mode to whether we should send a report for this call
+	//
+
+	bool bSendReport = false;
+	switch (ReportingMode)
+	{
+	case EStallDetectorReportingMode::First:
+		bSendReport = Stats.TriggerCount == 1;
+		break;
+
+	case EStallDetectorReportingMode::Always:
+		bSendReport = true;
+		break;
+
+	default:
+		break;
+	}
+
+	//
+	// Send the report
+	//
+
+	if (bSendReport)
+	{
+		Stats.TotalReportedCount++;
+		const int NumStackFramesToIgnore = FPlatformTLS::GetCurrentThreadId() == InThreadId ? 2 : 0;
+		ReportStall(Stats.Name, InThreadId, NumStackFramesToIgnore);
+		UE_LOG(LogStall, Warning, TEXT("Stall detector '%s' exceeded budget of %fs, and was reported"), Stats.Name, Stats.BudgetSeconds);
+	}
+	else
+	{
+		UE_LOG(LogStall, Warning, TEXT("Stall detector '%s' exceeded budget of %fs seconds"), Stats.Name, Stats.BudgetSeconds);
+	}
+}
+
+double UE::FStallDetector::Seconds()
+{
+	TRACE_CPUPROFILER_EVENT_SCOPE(FStallDetector::Seconds);
+
+	check(InitCount);
+
+	double Result;
+
+#if STALL_DETECTOR_HEART_BEAT_CLOCK
+	Result = Runnable->GetClock().Seconds();
+#else
+	Result = FPlatformTime::Seconds();
+#endif
+
+#if STALL_DETECTOR_DEBUG
+	static double ClockStartSeconds = Result;
+	static double PlatformStartSeconds = FPlatformTime::Seconds();
+	double ClockDelta = Result - ClockStartSeconds;
+	double PlatformDelta = FPlatformTime::Seconds() - PlatformStartSeconds;
+	double Drift = PlatformDelta - ClockDelta;
+	static double LastDrift = Drift;
+	double DriftDelta = Drift - LastDrift;
+	if (DriftDelta > 0.001)
+	{
+		FString ResultString = FString::Printf(TEXT("[FStallDetector] Thread %5d / Platform: %f / Clock: %f / Drift: %f (%f)\n"), FPlatformTLS::GetCurrentThreadId(), PlatformDelta, ClockDelta, Drift, DriftDelta);
+		FPlatformMisc::LocalPrint(ResultString.GetCharArray().GetData());
+		LastDrift = Drift;
+	}
+#endif
+
+	return Result;
+}
+
+void UE::FStallDetector::Startup()
+{
+	if (++InitCount == 1)
+	{
+		check(FPlatformTime::GetSecondsPerCycle());
+
+		// Cannot be a global due to clock member
+		Runnable = new FStallDetectorRunnable();
+
+		if (Thread == nullptr)
+		{
+			Thread = FRunnableThread::Create(Runnable, TEXT("StallDetectorThread"));
+			check(Thread);
+
+			// Poll until we have ticked the clock
+			while (!Runnable->GetStartedThread())
+			{
+				FPlatformProcess::YieldThread();
+			}
+		}
+	}
+}
+
+void UE::FStallDetector::Shutdown()
+{
+	if (--InitCount == 0)
+	{
+		delete Thread;
+		Thread = nullptr;
+
+		delete Runnable;
+		Runnable = nullptr;
+	}
+}
+
+#endif // STALL_DETECTOR
\ No newline at end of file
diff --git a/Engine/Source/Runtime/Core/Private/Windows/WindowsPlatformCrashContext.cpp b/Engine/Source/Runtime/Core/Private/Windows/WindowsPlatformCrashContext.cpp
--- a/Engine/Source/Runtime/Core/Private/Windows/WindowsPlatformCrashContext.cpp
+++ b/Engine/Source/Runtime/Core/Private/Windows/WindowsPlatformCrashContext.cpp
@@ -1244,18 +1244,18 @@ public:
 		return ThreadId;
 	}
 
-	/** Ensures are passed trough this. */
-	FORCEINLINE int32 OnEnsure(LPEXCEPTION_POINTERS InExceptionInfo, int NumStackFramesToIgnore, const TCHAR* ErrorMessage, EErrorReportUI ReportUI)
+	/** Ensures and Stalls are passed through this. */
+	FORCEINLINE int32 OnContinuableEvent(ECrashContextType InType, LPEXCEPTION_POINTERS InExceptionInfo, HANDLE InThreadHandle, uint32 InThreadId, int NumStackFramesToIgnore, const TCHAR* ErrorMessage, EErrorReportUI ReportUI)
 	{
 		if (CrashClientHandle.IsValid() && FPlatformProcess::IsProcRunning(CrashClientHandle))
 		{
 			return ReportCrashForMonitor(
 				InExceptionInfo, 
-				ECrashContextType::Ensure, 
+				InType,
 				ErrorMessage, 
 				NumStackFramesToIgnore, 
-				GetCurrentThread(), 
-				GetCurrentThreadId(), 
+				InThreadHandle,
+				InThreadId,
 				CrashClientHandle, 
 				&SharedContext, 
 				CrashMonitorWritePipe, 
@@ -1265,10 +1265,10 @@ public:
 		}
 		else
 		{
-			FWindowsPlatformCrashContext CrashContext(ECrashContextType::Ensure, ErrorMessage);
+			FWindowsPlatformCrashContext CrashContext(InType, ErrorMessage);
 			CrashContext.SetCrashedProcess(FProcHandle(::GetCurrentProcess()));
-			CrashContext.SetCrashedThreadId(GetCurrentThreadId());
-			void* ContextWrapper = FWindowsPlatformStackWalk::MakeThreadContextWrapper(InExceptionInfo->ContextRecord, GetCurrentThread());
+			CrashContext.SetCrashedThreadId(InThreadId);
+			void* ContextWrapper = FWindowsPlatformStackWalk::MakeThreadContextWrapper(InExceptionInfo->ContextRecord, InThreadHandle);
 			CrashContext.CapturePortableCallStack(NumStackFramesToIgnore, ContextWrapper);
 
 			return ReportCrashUsingCrashReportClient(CrashContext, InExceptionInfo, ReportUI);
@@ -1591,24 +1591,21 @@ static bool bReentranceGuard = false;
 /**
  * A wrapper for ReportCrashUsingCrashReportClient that creates a new ensure crash context
  */
-int32 ReportEnsureUsingCrashReportClient(EXCEPTION_POINTERS* ExceptionInfo, int NumStackFramesToIgnore, const TCHAR* ErrorMessage, EErrorReportUI ReportUI)
+int32 ReportContinuableEventUsingCrashReportClient(ECrashContextType InType, EXCEPTION_POINTERS* ExceptionInfo, HANDLE InThreadHandle, uint32 InThreadId, int NumStackFramesToIgnore, const TCHAR* ErrorMessage, EErrorReportUI ReportUI)
 {
 #if !NOINITCRASHREPORTER
-	return GCrashReportingThread->OnEnsure(ExceptionInfo, NumStackFramesToIgnore, ErrorMessage, ReportUI);
+	return GCrashReportingThread->OnContinuableEvent(InType, ExceptionInfo, InThreadHandle, InThreadId, NumStackFramesToIgnore, ErrorMessage, ReportUI);
 #else 
 	return EXCEPTION_EXECUTE_HANDLER;
 #endif
 }
 #endif
 
-FORCENOINLINE void ReportEnsureInner(const TCHAR* ErrorMessage, int NumStackFramesToIgnore)
+static void ReportEventOnCallingThread(ECrashContextType InType, const TCHAR* ErrorMessage, int NumStackFramesToIgnore)
 {
 	// Skip this frame and the ::RaiseException call itself
 	NumStackFramesToIgnore += 2;
 
-	/** This is the last place to gather memory stats before exception. */
-	FGenericCrashContext::SetMemoryStats(FPlatformMemory::GetStats());
-
 #if WINVER > 0x502	// Windows Error Reporting is not supported on Windows XP
 #if !PLATFORM_SEH_EXCEPTIONS_DISABLED
 	__try
@@ -1617,7 +1614,7 @@ FORCENOINLINE void ReportEnsureInner(const TCHAR* ErrorMessage, int NumStackFram
 		::RaiseException(EnsureExceptionCode, 0, 0, nullptr);
 	}
 #if !PLATFORM_SEH_EXCEPTIONS_DISABLED
-	__except (ReportEnsureUsingCrashReportClient( GetExceptionInformation(), NumStackFramesToIgnore, ErrorMessage, IsInteractiveEnsureMode() ? EErrorReportUI::ShowDialog : EErrorReportUI::ReportInUnattendedMode))
+	__except (ReportContinuableEventUsingCrashReportClient( InType, GetExceptionInformation(), GetCurrentThread(), GetCurrentThreadId(), NumStackFramesToIgnore, ErrorMessage, IsInteractiveEnsureMode() ? EErrorReportUI::ShowDialog : EErrorReportUI::ReportInUnattendedMode))
 		CA_SUPPRESS(6322)
 	{
 	}
@@ -1625,6 +1622,89 @@ FORCENOINLINE void ReportEnsureInner(const TCHAR* ErrorMessage, int NumStackFram
 #endif	// WINVER
 }
 
+static void ReportEvent(ECrashContextType InType, const TCHAR* ErrorMessage, uint32 InThreadId, int NumStackFramesToIgnore)
+{
+	if (ReportCrashCallCount > 0 || FDebug::HasAsserted())
+	{
+		// Don't report ensures after we've crashed/asserted, they simply may be a result of the crash as
+		// the engine is already in a bad state.
+		return;
+	}
+
+	// Serialize concurrent ensures (from concurrent threads).
+	FScopeLock ScopeLock(&EnsureLock);
+
+	// Ignore any ensure that could be fired by the code reporting an ensure.
+	TGuardValue<bool> ReentranceGuard(bReentranceGuard, true);
+	if (*ReentranceGuard) // Read the old value.
+	{
+		return; // Already handling an ensure.
+	}
+
+	// Stop checking heartbeat for this thread (and stop the gamethread hitch detector if we're the game thread).
+	// Ensure can take a lot of time (when stackwalking), so we don't want hitches/hangs firing.
+	// These are no-ops on threads that didn't already have a heartbeat etc.
+	FSlowHeartBeatScope SuspendHeartBeat(true);
+	FDisableHitchDetectorScope SuspendGameThreadHitch;
+
+	/** This is the last place to gather memory stats before exception. */
+	FGenericCrashContext::SetMemoryStats(FPlatformMemory::GetStats());
+
+	if (FPlatformTLS::GetCurrentThreadId() == InThreadId)
+	{
+		ReportEventOnCallingThread(InType, ErrorMessage, NumStackFramesToIgnore);
+	}
+	else
+	{
+		// This is what is normally returned by GetExceptionInformation(), which is a special compiler intrinsic only valid inside __except parenthesis
+		//  Here we construct one for a perfectly valid suspended thread of execution, and pass it into the CrashReportClient pathway
+		EXCEPTION_POINTERS ExceptionInformation = {};
+
+		HANDLE ThreadHandle = OpenThread(THREAD_SUSPEND_RESUME|THREAD_GET_CONTEXT|THREAD_QUERY_INFORMATION, 0, InThreadId);
+		if (ThreadHandle != nullptr)
+		{
+			// This is the thread state data, represents the register file state once suspended
+			CONTEXT ThreadContext;
+			ZeroMemory(&ThreadContext, sizeof(ThreadContext));
+
+			// This include segments and debug registers, just in case
+			ThreadContext.ContextFlags |= CONTEXT_ALL;
+
+			DWORD SuspendCount = SuspendThread(ThreadHandle);
+			if (SuspendCount != -1)
+			{
+				// Read the context block from the thread now its suspended
+				BOOL bGotThreadContext = GetThreadContext(ThreadHandle, &ThreadContext);
+				if (bGotThreadContext)
+				{
+					// Success! Set the pointer to this state in the exception information block
+					ExceptionInformation.ContextRecord = &ThreadContext;
+				}
+			}
+
+			// If we were able to suspend and capture the thread state
+			if (ExceptionInformation.ContextRecord)
+			{
+				EXCEPTION_RECORD ExceptionRecord;
+				ZeroMemory(&ExceptionRecord, sizeof(ExceptionRecord));
+				ExceptionRecord.ExceptionCode = STILL_ACTIVE; // normally where EXCEPTION_ACCESS_VIOLATION would be
+
+				// The thread context and exception recard are the two things expected
+				ExceptionInformation.ExceptionRecord = &ExceptionRecord;
+
+				(void)ReportContinuableEventUsingCrashReportClient(InType, &ExceptionInformation, ThreadHandle, InThreadId, NumStackFramesToIgnore, ErrorMessage, IsInteractiveEnsureMode() ? EErrorReportUI::ShowDialog : EErrorReportUI::ReportInUnattendedMode);
+			}
+
+			if (SuspendCount != -1)
+			{
+				ResumeThread(ThreadHandle);
+			}
+
+			CloseHandle(ThreadHandle);
+		}
+	}
+}
+
 FORCENOINLINE void ReportAssert(const TCHAR* ErrorMessage, int NumStackFramesToIgnore)
 {
 	/** This is the last place to gather memory stats before exception. */
@@ -1683,32 +1763,16 @@ void ReportHang(const TCHAR* ErrorMessage, const uint64* StackFrames, int32 NumS
  */
 FORCENOINLINE void ReportEnsure(const TCHAR* ErrorMessage, int NumStackFramesToIgnore)
 {
-	if (ReportCrashCallCount > 0 || FDebug::HasAsserted())
-	{
-		// Don't report ensures after we've crashed/asserted, they simply may be a result of the crash as
-		// the engine is already in a bad state.
-		return;
-	}
-
-	// Serialize concurrent ensures (from concurrent threads).
-	FScopeLock ScopedEnsureLock(&EnsureLock);
-
-	// Ignore any ensure that could be fired by the code reporting an ensure.
-	TGuardValue<bool> ReentranceGuard(bReentranceGuard, true);
-	if (*ReentranceGuard) // Read the old value.
-	{
-		return; // Already handling an ensure.
-	}
-
-	// Stop checking heartbeat for this thread (and stop the gamethread hitch detector if we're the game thread).
-	// Ensure can take a lot of time (when stackwalking), so we don't want hitches/hangs firing.
-	// These are no-ops on threads that didn't already have a heartbeat etc.
-	FSlowHeartBeatScope SuspendHeartBeat(true);
-	FDisableHitchDetectorScope SuspendGameThreadHitch;
-
-	ReportEnsureInner(ErrorMessage, NumStackFramesToIgnore + 1);
+	ReportEvent(ECrashContextType::Ensure, ErrorMessage, FPlatformTLS::GetCurrentThreadId(), NumStackFramesToIgnore + 1);
 }
 
+/**
+ * Report a hitch to the crash reporting system
+ */
+FORCENOINLINE void ReportStall(const TCHAR* ErrorMessage, uint32 HitchThreadId, int NumStackFramesToIgnore)
+{
+	ReportEvent(ECrashContextType::Stall, ErrorMessage, HitchThreadId, FPlatformTLS::GetCurrentThreadId() == HitchThreadId ? NumStackFramesToIgnore + 1 : NumStackFramesToIgnore);
+}
 
 #if !IS_PROGRAM && 0
 namespace {
diff --git a/Engine/Source/Runtime/Core/Private/Windows/WindowsPlatformStackWalk.cpp b/Engine/Source/Runtime/Core/Private/Windows/WindowsPlatformStackWalk.cpp
--- a/Engine/Source/Runtime/Core/Private/Windows/WindowsPlatformStackWalk.cpp
+++ b/Engine/Source/Runtime/Core/Private/Windows/WindowsPlatformStackWalk.cpp
@@ -307,11 +307,21 @@ uint32 FWindowsPlatformStackWalk::CaptureThreadStackBackTrace(uint64 ThreadId, u
 	InitStackWalking();
 
 	if (BackTrace == nullptr || MaxDepth == 0)
+	{
 		return 0;
+	}
+
+	// Don't suspend the calling thread, capture it's context directly and trace
+	if (ThreadId == FPlatformTLS::GetCurrentThreadId())
+	{
+		return CaptureStackBackTrace(BackTrace, MaxDepth);
+	}
 
 	HANDLE ThreadHandle = OpenThread(THREAD_GET_CONTEXT | THREAD_SET_CONTEXT | THREAD_TERMINATE | THREAD_SUSPEND_RESUME, false, (DWORD)ThreadId);
 	if (!ThreadHandle)
+	{
 		return 0;
+	}
 
 	// Suspend the thread before grabbing its context
 	SuspendThread(ThreadHandle);
diff --git a/Engine/Source/Runtime/Core/Public/GenericPlatform/GenericPlatformCrashContext.h b/Engine/Source/Runtime/Core/Public/GenericPlatform/GenericPlatformCrashContext.h
--- a/Engine/Source/Runtime/Core/Public/GenericPlatform/GenericPlatformCrashContext.h
+++ b/Engine/Source/Runtime/Core/Public/GenericPlatform/GenericPlatformCrashContext.h
@@ -95,6 +95,7 @@ enum class ECrashContextType
 	Crash,
 	Assert,
 	Ensure,
+	Stall,
 	GPUCrash,
 	Hang,
 	OutOfMemory,
@@ -247,6 +248,7 @@ public:
 	static const TCHAR* const CrashTypeCrash;
 	static const TCHAR* const CrashTypeAssert;
 	static const TCHAR* const CrashTypeEnsure;
+	static const TCHAR* const CrashTypeStall;
 	static const TCHAR* const CrashTypeGPU;
 	static const TCHAR* const CrashTypeHang;
 	static const TCHAR* const CrashTypeAbnormalShutdown;
@@ -443,10 +445,26 @@ public:
 	static void CleanupPlatformSpecificFiles();
 
 	/**
-	 * @return whether this crash is a non-crash event
+	 * @return the type of this crash
 	 */
 	ECrashContextType GetType() const { return Type; }
 
+	/**
+	 * @return whether a crash context type is continable
+	 */
+	static bool IsTypeContinuable(ECrashContextType Type)
+	{
+		switch (Type)
+		{
+		case ECrashContextType::Ensure:
+			return true;
+		case ECrashContextType::Stall:
+			return true;
+		default:
+			return false;
+		}
+	}
+
 	/**
 	 * Set the current deployment name (ie. EpicApp)
 	 */
diff --git a/Engine/Source/Runtime/Core/Public/HAL/ExceptionHandling.h b/Engine/Source/Runtime/Core/Public/HAL/ExceptionHandling.h
--- a/Engine/Source/Runtime/Core/Public/HAL/ExceptionHandling.h
+++ b/Engine/Source/Runtime/Core/Public/HAL/ExceptionHandling.h
@@ -27,6 +27,7 @@ extern CORE_API int32 ReportCrash( Windows::LPEXCEPTION_POINTERS ExceptionInfo )
 extern CORE_API void ReportAssert(const TCHAR* ErrorMessage, int NumStackFramesToIgnore);
 extern CORE_API void ReportGPUCrash(const TCHAR* ErrorMessage, int NumStackFramesToIgnore);
 extern CORE_API void ReportEnsure(const TCHAR* ErrorMessage, int NumStackFramesToIgnore);
+extern CORE_API void ReportStall(const TCHAR* ErrorMessage, uint32 HitchThreadId, int NumStackFramesToIgnore);
 extern CORE_API void ReportHang(const TCHAR*, const uint64* StackFrames, int32 NumStackFrames, uint32 HungThreadId);
 #elif PLATFORM_MAC
 // #CrashReport: 2014-10-09 Should be move to another file
diff --git a/Engine/Source/Runtime/Core/Public/ProfilingDebugging/StallDetector.h b/Engine/Source/Runtime/Core/Public/ProfilingDebugging/StallDetector.h
new file mode 100644
--- /dev/null
+++ b/Engine/Source/Runtime/Core/Public/ProfilingDebugging/StallDetector.h
@@ -0,0 +1,206 @@
+// Copyright Epic Games, Inc. All Rights Reserved.
+
+#pragma once
+
+#include "CoreTypes.h"
+
+#include "Containers/UnrealString.h"
+#include "Containers/Set.h"
+#include "HAL/CriticalSection.h"
+#include "HAL/PreprocessorHelpers.h"
+#include "Logging/LogMacros.h"
+#include "Misc/ScopeLock.h"
+
+#include <atomic>
+
+#ifndef STALL_DETECTOR
+ #if WITH_EDITOR && PLATFORM_WINDOWS
+  #define STALL_DETECTOR 1
+ #else // WITH_EDITOR && PLATFORM_WINDOWS
+  #define STALL_DETECTOR 0
+ #endif // WITH_EDITOR && PLATFORM_WINDOWS
+#endif // STALL_DETECTOR
+
+#if STALL_DETECTOR
+
+CORE_API DECLARE_LOG_CATEGORY_EXTERN(LogStall, Log, All);
+
+/**
+* This code is meant to instrument code to send reports for slow pathways that create unresponsive conditions
+**/
+
+namespace UE
+{
+	/**
+	* The reporting behavior of the detector
+	**/
+	enum EStallDetectorReportingMode
+	{
+		Never,
+		First,
+		Always
+	};
+
+	/**
+	* Information to track per-callsite (vs. lifetime of a single interval)
+	**/
+	class CORE_API FStallDetectorStats
+	{
+	public:
+		/**
+		* Construct a stats object for use with a stall detector
+		*/
+		FStallDetectorStats(const TCHAR* InName, const double InBudgetSeconds, const EStallDetectorReportingMode InReportingMode);
+
+		/**
+		* Destruct a stats object, just removes it from the instance tracking
+		*/
+		~FStallDetectorStats();
+
+		/**
+		* Retrieve the critical section (lock) for accessing the instances list
+		*/
+		static FCriticalSection& GetCriticalSection()
+		{
+			return CriticalSection;
+		}
+
+		/**
+		* The container of all constructed instances of this class
+		*/
+		static const TSet<FStallDetectorStats*>& GetInstances()
+		{
+			return Instances;
+		}
+
+		// The name to refer to this callsite in the codebase
+		const TCHAR* Name;
+
+		// The budgeted time we tolerate without sending a report
+		const double BudgetSeconds;
+
+		// Drives the behavior about what to do when a budget is violated (triggered)
+		const EStallDetectorReportingMode ReportingMode;
+
+		// The number of times this callsite has been triggered
+		uint32 TriggerCount;
+
+		// The cumulative overage time for this callsite
+		double OverageSeconds;
+
+		// The total number of times all callsites have been triggered
+		static std::atomic<uint32> TotalTriggeredCount;
+
+		// The total number of reports that have been sent
+		static std::atomic<uint32> TotalReportedCount;
+
+	private:
+		// Guards access to the instances container from multiple threads
+		static FCriticalSection CriticalSection;
+
+		// The pointers to every constructed instance of this class
+		static TSet<FStallDetectorStats*> Instances;
+	};
+
+	/**
+	* FStallDetector is meant to be constructed and destructed across a single scope to measure that specific timespan.
+	**/
+	class CORE_API FStallDetector
+	{
+	public:
+		/**
+		* Construct the stall detector linking it to it's stats object and take a timestamp
+		*/
+		FStallDetector(FStallDetectorStats& InStats);
+
+		/**
+		* Destruct the stall detector and perform a check if appropriate (the default behavior)
+		*/
+		~FStallDetector();
+
+		/**
+		* Check the detector to see if we are over budget, called from both construction and background thread
+		* @param bIsComplete			Is this check the one performed at completion of the scope or measurement period?
+		* @param InWhenToCheckSeconds	The timestamp at which to compare the start time to
+		**/
+		void Check(bool bIsComplete, double InWhenToCheckSeconds = 0.0);
+
+		/**
+		* Perform a Check and then reset the timer to the current time.
+		* This changes the mode of this class to defer the first check, and skip the destruction check
+		* This mode enables this class to used as a global or static, and CheckAndReset() callable periodically
+		**/
+		void CheckAndReset();
+
+		/**
+		* Sample the same clock that the stall detector uses internally
+		**/
+		static double Seconds();
+
+		/**
+		* Retrieve the critical section (lock) for accessing the instances list
+		*/
+		static FCriticalSection& GetCriticalSection()
+		{
+			return CriticalSection;
+		}
+
+		/**
+		* The container of all constructed instances of this class
+		*/
+		static const TSet<FStallDetector*>& GetInstances()
+		{
+			return Instances;
+		}
+
+		/**
+		* Initialize Stall Detector API resources
+		**/
+		static void Startup();
+
+		/**
+		* Cleanup Stall Detector API resources
+		**/
+		static void Shutdown();
+
+	private:
+		// The time provided, the deadline
+		FStallDetectorStats& Stats;
+
+		// Thread it was constructed on, and to stack trace when budgeted time is elapsed
+		uint32 ThreadId;
+
+		// The time of construction of this object
+		double StartSeconds;
+
+		// Persistent usage mode
+		bool bPersistent;
+
+		// Track the triggered state, atomic to mitigate foreground and background thread checking
+		std::atomic<bool> Triggered;
+
+		// The pinch point for tracking counts and emitting a report for a triggered hitch detector
+		void OnStallDetected(uint32 InThreadId, const double InElapsedSeconds);
+
+		// Guards access to the instances container from multiple threads
+		static FCriticalSection CriticalSection;
+
+		// The pointers to every constructed instance of this class
+		static TSet<FStallDetector*> Instances;
+	};
+}
+
+#define SCOPE_STALL_REPORTER(InName, InBudgetSeconds) \
+static UE::FStallDetectorStats PREPROCESSOR_JOIN(StallDetectorStats, __LINE__) (TEXT(#InName), InBudgetSeconds, UE::EStallDetectorReportingMode::First); \
+UE::FStallDetector PREPROCESSOR_JOIN(ScopeStallDetector, __LINE__)(PREPROCESSOR_JOIN(StallDetectorStats, __LINE__));
+
+#define SCOPE_STALL_COUNTER(InName, InBudgetSeconds) \
+static UE::FStallDetectorStats PREPROCESSOR_JOIN(StallDetectorStats, __LINE__) (TEXT(#InName), InBudgetSeconds, UE::EStallDetectorReportingMode::Never); \
+UE::FStallDetector PREPROCESSOR_JOIN(ScopeStallDetector, __LINE__)(PREPROCESSOR_JOIN(StallDetectorStats, __LINE__));
+
+#else // STALL_DETECTOR
+
+#define SCOPE_STALL_REPORTER(InName, InBudgetSeconds)
+#define SCOPE_STALL_COUNTER(InName, InBudgetSeconds)
+
+#endif // STALL_DETECTOR
\ No newline at end of file
diff --git a/Engine/Source/Runtime/Launch/Private/LaunchEngineLoop.cpp b/Engine/Source/Runtime/Launch/Private/LaunchEngineLoop.cpp
--- a/Engine/Source/Runtime/Launch/Private/LaunchEngineLoop.cpp
+++ b/Engine/Source/Runtime/Launch/Private/LaunchEngineLoop.cpp
@@ -43,6 +43,7 @@
 #endif
 #include "HAL/ThreadManager.h"
 #include "ProfilingDebugging/ExternalProfiler.h"
+#include "ProfilingDebugging/StallDetector.h"
 #include "Containers/StringView.h"
 #include "Containers/Ticker.h"
 
@@ -4782,6 +4783,8 @@ static inline void EndFrameRenderThread(FRHICommandListImmediate& RHICmdList, ui
 
 void FEngineLoop::Tick()
 {
+	SCOPE_STALL_COUNTER(FEngineLoop::Tick, 2.0)
+
 	// make sure to catch any FMemStack uses outside of UWorld::Tick
 	FMemMark MemStackMark(FMemStack::Get());
 
@@ -5925,6 +5928,10 @@ bool FEngineLoop::AppInit( )
 
 	FEmbeddedCommunication::ForceTick(19);
 
+#if STALL_DETECTOR
+	UE::FStallDetector::Startup();
+#endif
+
 	return true;
 }
 
@@ -6014,6 +6021,10 @@ void FEngineLoop::AppPreExit( )
 
 void FEngineLoop::AppExit( )
 {
+#if STALL_DETECTOR
+	UE::FStallDetector::Shutdown();
+#endif
+
 #if !WITH_ENGINE
 	// when compiled WITH_ENGINE, this will happen in FEngineLoop::Exit()
 	FTaskGraphInterface::Shutdown();
