[TraceLog] Re-add trace code that was lost in TaskGraph refactoring

#rnx

diff --git a/Engine/Source/Runtime/Core/Public/Async/AsyncWork.h b/Engine/Source/Runtime/Core/Public/Async/AsyncWork.h
--- a/Engine/Source/Runtime/Core/Public/Async/AsyncWork.h
+++ b/Engine/Source/Runtime/Core/Public/Async/AsyncWork.h
@@ -16,6 +16,7 @@
 #include "HAL/LowLevelMemTracker.h"
 #include "Misc/IQueuedWork.h"
 #include "Misc/QueuedThreadPool.h"
+#include "ProfilingDebugging/TagTrace.h"
 
 /**
 	FAutoDeleteAsyncTask - template task for jobs that delete themselves when complete
@@ -227,6 +228,10 @@ class FAsyncTask
 	EQueuedWorkPriority Priority;
 	/** optional LLM tag */
 	LLM(const UE::LLMPrivate::FTagData* InheritedLLMTag);
+	/** Memory trace tag */
+#if USE_MEMORY_TRACE_TAGS
+	int32 InheritedTraceTag;
+#endif
 
 	/* Internal function to destroy the completion event
 	**/
@@ -244,6 +249,9 @@ class FAsyncTask
 		FScopeCycleCounter Scope( Task.GetStatId(), true );
 		DECLARE_SCOPE_CYCLE_COUNTER( TEXT( "FAsyncTask::Start" ), STAT_FAsyncTask_Start, STATGROUP_ThreadPoolAsyncTasks );
 		LLM(InheritedLLMTag = FLowLevelMemTracker::bIsDisabled ? nullptr : FLowLevelMemTracker::Get().GetActiveTagData(ELLMTracker::Default));
+#if USE_MEMORY_TRACE_TAGS
+		InheritedTraceTag = MemoryTrace_GetActiveTag();
+#endif
 
 		FPlatformMisc::MemoryBarrier();
 		CheckIdle();  // can't start a job twice without it being completed first
@@ -277,6 +285,7 @@ class FAsyncTask
 	void DoWork()
 	{	
 		LLM_SCOPE(InheritedLLMTag);
+		UE_MEMSCOPE(InheritedTraceTag, ELLMTracker::Default);
 		FScopeCycleCounter Scope(Task.GetStatId(), true); 
 
 		Task.DoWork();		
