IoStore: Changes to enable usage from the cooker
#rnx

diff --git a/Engine/Source/Developer/IoStoreUtilities/Private/IoStoreUtilities.cpp b/Engine/Source/Developer/IoStoreUtilities/Private/IoStoreUtilities.cpp
--- a/Engine/Source/Developer/IoStoreUtilities/Private/IoStoreUtilities.cpp
+++ b/Engine/Source/Developer/IoStoreUtilities/Private/IoStoreUtilities.cpp
@@ -3448,11 +3448,16 @@ private:
 			Manager.Schedule(this);
 		}
 
-		FIoBuffer ConsumeSourceBuffer() override
+		const FIoBuffer* GetSourceBuffer() override
 		{
+			return SourceBuffer;
+		}
+
+		void FreeSourceBuffer() override
+		{
+			delete SourceBuffer;
+			SourceBuffer = nullptr;
 			Manager.OnBufferMemoryFreed(TargetFile.SourceSize);
-			FIoBuffer Result = MoveTemp(SourceBuffer);
-			return Result;
 		}
 
 		uint64 GetOrderHint() override
@@ -3467,8 +3472,9 @@ private:
 
 		FIoBuffer& PrepareSourceBuffer()
 		{
-			SourceBuffer = FIoBuffer(TargetFile.SourceSize);
-			return SourceBuffer;
+			check(!SourceBuffer);
+			SourceBuffer = new FIoBuffer(TargetFile.SourceSize);
+			return *SourceBuffer;
 		}
 
 		void AsyncReadCallback()
@@ -3477,7 +3483,7 @@ private:
 			if (!TargetFile.bIsBulkData)
 			{
 				check(ObjectExports);
-				SourceBuffer = CreateExportBundleBuffer(TargetFile, *ObjectExports, SourceBuffer, bHasUpdatedExportBundleRegions ? nullptr : &FileRegions);
+				*SourceBuffer = CreateExportBundleBuffer(TargetFile, *ObjectExports, *SourceBuffer, bHasUpdatedExportBundleRegions ? nullptr : &FileRegions);
 				bHasUpdatedExportBundleRegions = true;
 			}
 			TArray<FBaseGraphTask*> NewTasks;
@@ -3490,7 +3496,7 @@ private:
 		const TArray<FObjectExport>* ObjectExports;
 		TArray<FFileRegion> FileRegions;
 		FGraphEventRef CompletionEvent;
-		FIoBuffer SourceBuffer;
+		FIoBuffer* SourceBuffer = nullptr;
 		bool bHasUpdatedExportBundleRegions = false;
 	};
 
@@ -3798,8 +3804,9 @@ int32 CreateTarget(const FIoStoreArguments& Arguments, const FIoStoreWriterSetti
 					ContainerSettings.ContainerFlags |= EIoContainerFlags::Signed;
 				}
 				ContainerSettings.bGenerateDiffPatch = ContainerTarget->bGenerateDiffPatch;
-				IoStatus = ContainerTarget->IoStoreWriter->Initialize(*IoStoreWriterContext, ContainerSettings, ContainerTarget->PatchSourceReaders);
+				IoStatus = ContainerTarget->IoStoreWriter->Initialize(*IoStoreWriterContext, ContainerSettings);
 				check(IoStatus.IsOk());
+				ContainerTarget->IoStoreWriter->EnableDiskLayoutOrdering(ContainerTarget->PatchSourceReaders);
 			}
 		}
 	}
diff --git a/Engine/Source/Runtime/Core/Private/IO/IoStore.cpp b/Engine/Source/Runtime/Core/Private/IO/IoStore.cpp
--- a/Engine/Source/Runtime/Core/Private/IO/IoStore.cpp
+++ b/Engine/Source/Runtime/Core/Private/IO/IoStore.cpp
@@ -39,7 +39,11 @@ static IEngineCrypto* GetEngineCrypto()
 
 static bool IsSigningEnabled()
 {
+#if UE_BUILD_SHIPPING
 	return FCoreDelegates::GetPakSigningKeysDelegate().IsBound();
+#else
+	return false;
+#endif
 }
 
 static FRSAKeyHandle GetPublicSigningKey()
@@ -136,7 +140,8 @@ static FIoStatus ValidateContainerSignature(
 
 struct FChunkBlock
 {
-	uint64 Offset = 0;
+	const uint8* UncompressedData = nullptr;
+	FIoBuffer* IoBuffer = nullptr;
 	uint64 Size = 0;
 	uint64 CompressedSize = 0;
 	uint64 UncompressedSize = 0;
@@ -147,23 +152,24 @@ struct FChunkBlock
 struct FIoStoreWriteQueueEntry
 {
 	FIoStoreWriteQueueEntry* Next = nullptr;
+	FIoStoreWriterImpl* Writer = nullptr;
 	IIoStoreWriteRequest* Request = nullptr;
 	FIoChunkId ChunkId;
 	FIoChunkHash ChunkHash;
-	FIoBuffer ChunkBuffer;
 	uint64 Sequence = 0;
 	uint64 UncompressedSize = 0;
 	uint64 CompressedSize = 0;
 	uint64 Padding = 0;
 	uint64 Offset = 0;
-	FArchive* ContainerArchive = nullptr;
 	TArray<FChunkBlock> ChunkBlocks;
 	FIoWriteOptions Options;
 	FGraphEventRef HashBarrier;
 	FGraphEventRef HashTask;
-	FGraphEventRef CreateChunkBlocksBarrier;
-	FGraphEventRef CreateChunkBlocksTask;
-	FEvent* WriteCompletedEvent = nullptr;
+	FGraphEventRef BeginCompressionBarrier;
+	FGraphEventRef FinishCompressionBarrier;
+	FGraphEventRef BeginWriteBarrier;
+	FGraphEventRef WriteFinishedEvent;
+	TAtomic<int32> FinishedBlocksCount{ 0 };
 	int32 PartitionIndex = -1;
 	bool bAdded = false;
 	bool bModified = false;
@@ -250,7 +256,7 @@ private:
 
 class FIoStoreWriterContextImpl
 {
-	static constexpr uint64 DefaultMemoryLimit = 5ull * (2ull << 30ull);
+	static constexpr uint64 DefaultMemoryLimit = 1ull << 30ull;
 public:
 	FIoStoreWriterContextImpl()
 	{
@@ -258,23 +264,43 @@ public:
 
 	~FIoStoreWriterContextImpl()
 	{
-		CompressionQueue.CompleteAdding();
-		WriteQueue.CompleteAdding();
-		CompressorThread.Wait();
+		BeginCompressionQueue.CompleteAdding();
+		FinishCompressionQueue.CompleteAdding();
+		WriterQueue.CompleteAdding();
+		BeginCompressionThread.Wait();
+		FinishCompressionThread.Wait();
 		WriterThread.Wait();
-		if (MemoryFreedEvent)
+		if (CompressionBufferAvailableEvent)
+		{
+			FPlatformProcess::ReturnSynchEventToPool(CompressionBufferAvailableEvent);
+		}
+		for (FIoBuffer* IoBuffer : AvailableCompressionBuffers)
 		{
-			FPlatformProcess::ReturnSynchEventToPool(MemoryFreedEvent);
+			delete IoBuffer;
 		}
 	}
 
 	UE_NODISCARD FIoStatus Initialize(const FIoStoreWriterSettings& InWriterSettings)
 	{
 		WriterSettings = InWriterSettings;
-		MemoryFreedEvent = FPlatformProcess::GetSynchEventFromPool(false);
+		CompressionBufferAvailableEvent = FPlatformProcess::GetSynchEventFromPool(false);
+
+		if (WriterSettings.CompressionMethod != NAME_None)
+		{
+			CompressionBufferSize = FCompression::CompressMemoryBound(WriterSettings.CompressionMethod, static_cast<int32>(WriterSettings.CompressionBlockSize));
+		}
+		CompressionBufferSize = FMath::Max(CompressionBufferSize, static_cast<int32>(WriterSettings.CompressionBlockSize));
+		CompressionBufferSize = Align(CompressionBufferSize, FAES::AESBlockSize);
 
-		PaddingBuffer.SetNumZeroed(int32(WriterSettings.CompressionBlockSize));
-		CompressorThread = Async(EAsyncExecution::Thread, [this]() { CompressorThreadFunc(); });
+		int32 CompressionBufferCount = DefaultMemoryLimit / CompressionBufferSize;
+		AvailableCompressionBuffers.Reserve(CompressionBufferCount);
+		for (int32 BufferIndex = 0; BufferIndex < CompressionBufferCount; ++BufferIndex)
+		{
+			AvailableCompressionBuffers.Add(new FIoBuffer(CompressionBufferSize));
+		}
+
+		BeginCompressionThread = Async(EAsyncExecution::Thread, [this]() { BeginCompressionThreadFunc(); });
+		FinishCompressionThread = Async(EAsyncExecution::Thread, [this]() { FinishCompressionThreadFunc(); });
 		WriterThread = Async(EAsyncExecution::Thread, [this]() { WriterThreadFunc(); });
 
 		return FIoStatus::Ok;
@@ -295,106 +321,68 @@ public:
 		return WriterSettings;
 	}
 
-	void BeginCompress(FIoStoreWriteQueueEntry* QueueEntry)
-	{
-		CompressionQueue.Enqueue(QueueEntry);
-	}
-
-	void BeginWrite(FIoStoreWriteQueueEntry* QueueEntry)
-	{
-		WriteQueue.Enqueue(QueueEntry);
-	}
-
-private:
-	void BeginCompressEntry(FIoStoreWriteQueueEntry* Entry)
+	void ScheduleCompression(FIoStoreWriteQueueEntry* QueueEntry)
 	{
-		check(!Entry->CompressedSize);
-		check(!Entry->ChunkBuffer.DataSize());
-		uint64 LocalUsedBufferMemory = UsedBufferMemory.Load();
-		while (LocalUsedBufferMemory > 0 && LocalUsedBufferMemory + Entry->UncompressedSize > DefaultMemoryLimit)
-		{
-			TRACE_CPUPROFILER_EVENT_SCOPE(WaitForBufferMemory);
-			MemoryFreedEvent->Wait();
-			LocalUsedBufferMemory = UsedBufferMemory.Load();
-		}
-		UsedBufferMemory.AddExchange(Entry->UncompressedSize);
-		Entry->Request->PrepareSourceBufferAsync(Entry->CreateChunkBlocksBarrier);
+		BeginCompressionQueue.Enqueue(QueueEntry);
+		QueueEntry->Request->PrepareSourceBufferAsync(QueueEntry->BeginCompressionBarrier);
 	}
 
-	void WriteEntry(FIoStoreWriteQueueEntry* Entry)
+	FIoBuffer* AllocCompressionBuffer()
 	{
-		TRACE_CPUPROFILER_EVENT_SCOPE(WriteToContainerFile);
-		if (Entry->Padding > 0)
+		FIoBuffer* AllocatedBuffer = nullptr;
+		while (!AllocatedBuffer)
 		{
-			if (PaddingBuffer.Num() < Entry->Padding)
 			{
-				PaddingBuffer.SetNumZeroed(int32(Entry->Padding));
-			}
-			Entry->ContainerArchive->Serialize(PaddingBuffer.GetData(), Entry->Padding);
-		}
-		check(Entry->Offset == Entry->ContainerArchive->Tell());
-		Entry->ContainerArchive->Serialize(Entry->ChunkBuffer.Data(), Entry->ChunkBuffer.DataSize());
-		Entry->ChunkBuffer = FIoBuffer();
-		uint64 UsedBefore = UsedBufferMemory.SubExchange(Entry->UncompressedSize);
-		check(UsedBefore >= Entry->UncompressedSize);
-		MemoryFreedEvent->Trigger();
-		if (Entry->WriteCompletedEvent)
-		{
-			Entry->WriteCompletedEvent->Trigger();
-		}
-		SerializedChunksCount.IncrementExchange();
-	}
-
-	void CompressorThreadFunc()
-	{
-		for (;;)
-		{
-			FIoStoreWriteQueueEntry* Entry = CompressionQueue.DequeueOrWait();
-			if (!Entry)
-			{
-				return;
+				FScopeLock Lock(&AvailableCompressionBuffersCritical);
+				if (AvailableCompressionBuffers.Num())
+				{
+					AllocatedBuffer = AvailableCompressionBuffers.Pop();
+				}
 			}
-			while (Entry)
+			if (!AllocatedBuffer)
 			{
-				FIoStoreWriteQueueEntry* Next = Entry->Next;
-				BeginCompressEntry(Entry);
-				Entry = Next;
+				TRACE_CPUPROFILER_EVENT_SCOPE(WaitForCompressionBuffer);
+				CompressionBufferAvailableEvent->Wait();
 			}
 		}
+		return AllocatedBuffer;
 	}
 
-	void WriterThreadFunc()
+	void FreeCompressionBuffer(FIoBuffer* Buffer)
 	{
-		for (;;)
+		bool bTriggerEvent;
 		{
-			FIoStoreWriteQueueEntry* Entry = WriteQueue.DequeueOrWait();
-			if (!Entry)
-			{
-				return;
-			}
-			while (Entry)
-			{
-				FIoStoreWriteQueueEntry* Next = Entry->Next;
-				WriteEntry(Entry);
-				Entry = Next;
-			}
+			FScopeLock Lock(&AvailableCompressionBuffersCritical);
+			bTriggerEvent = AvailableCompressionBuffers.Num() == 0;
+			AvailableCompressionBuffers.Push(Buffer);
+		}
+		if (bTriggerEvent)
+		{
+			CompressionBufferAvailableEvent->Trigger();
 		}
 	}
 
+private:
+	void BeginCompressionThreadFunc();
+	void FinishCompressionThreadFunc();
+	void WriterThreadFunc();
+
 	FIoStoreWriterSettings WriterSettings;
-	FCriticalSection CriticalSection;
-	FEvent* MemoryFreedEvent = nullptr;
-	TAtomic<uint64> UsedBufferMemory{ 0 };
-	TFuture<void> CompressorThread;
+	FEvent* CompressionBufferAvailableEvent = nullptr;
+	TFuture<void> BeginCompressionThread;
+	TFuture<void> FinishCompressionThread;
 	TFuture<void> WriterThread;
-	FIoStoreWriteQueue CompressionQueue;
-	FIoStoreWriteQueue WriteQueue;
-	TArray<uint8> PaddingBuffer;
+	FIoStoreWriteQueue BeginCompressionQueue;
+	FIoStoreWriteQueue FinishCompressionQueue;
+	FIoStoreWriteQueue WriterQueue;
 	TAtomic<uint64> TotalChunksCount{ 0 };
 	TAtomic<uint64> HashedChunksCount{ 0 };
 	TAtomic<uint64> CompressedChunksCount{ 0 };
 	TAtomic<uint64> SerializedChunksCount{ 0 };
-	
+	FCriticalSection AvailableCompressionBuffersCritical;
+	TArray<FIoBuffer*> AvailableCompressionBuffers;
+	int32 CompressionBufferSize = -1;
+
 	friend class FIoStoreWriterImpl;
 };
 
@@ -539,7 +527,7 @@ public:
 	{
 	}
 
-	UE_NODISCARD FIoStatus Initialize(FIoStoreWriterContextImpl& InContext, const FIoContainerSettings& InContainerSettings, const TArray<TUniquePtr<FIoStoreReader>>& PatchSourceReaders)
+	UE_NODISCARD FIoStatus Initialize(FIoStoreWriterContextImpl& InContext, const FIoContainerSettings& InContainerSettings)
 	{
 		WriterContext = &InContext;
 		ContainerSettings = InContainerSettings;
@@ -555,7 +543,8 @@ public:
 			Status = EnableCsvOutput();
 		}
 
-		PrepareLayout(PatchSourceReaders);
+		FPartition& Partition = Partitions.AddDefaulted_GetRef();
+		Partition.Index = 0;
 
 		return Status;
 	}
@@ -574,13 +563,51 @@ public:
 		return FIoStatus::Ok;
 	}
 
+	void EnableDiskLayoutOrdering(const TArray<TUniquePtr<FIoStoreReader>>& PatchSourceReaders)
+	{
+		check(!LayoutEntriesHead);
+		check(!Entries.Num());
+		LayoutEntriesHead = new FLayoutEntry();
+		LayoutEntries.Add(LayoutEntriesHead);
+		FLayoutEntry* PrevEntryLink = LayoutEntriesHead;
+
+		for (const TUniquePtr<FIoStoreReader>& PatchSourceReader : PatchSourceReaders)
+		{
+			PatchSourceReader->EnumerateChunks([this, &PrevEntryLink](const FIoStoreTocChunkInfo& ChunkInfo)
+				{
+					FLayoutEntry* PreviousBuildEntry = new FLayoutEntry();
+					LayoutEntries.Add(PreviousBuildEntry);
+
+					PreviousBuildEntry->Hash = ChunkInfo.Hash;
+					PreviousBuildEntry->PartitionIndex = ChunkInfo.PartitionIndex;
+					PreviousBuildEntry->CompressedSize = ChunkInfo.CompressedSize;
+					PrevEntryLink->Next = PreviousBuildEntry;
+					PreviousBuildEntry->Prev = PrevEntryLink;
+					PrevEntryLink = PreviousBuildEntry;
+					PreviousBuildLayoutEntryByChunkId.Add(ChunkInfo.Id, PreviousBuildEntry);
+					return true;
+				});
+
+			if (!ContainerSettings.bGenerateDiffPatch)
+			{
+				break;
+			}
+		}
+
+		LayoutEntriesTail = new FLayoutEntry();
+		LayoutEntries.Add(LayoutEntriesTail);
+		PrevEntryLink->Next = LayoutEntriesTail;
+		LayoutEntriesTail->Prev = PrevEntryLink;
+	}
+
 	void Append(const FIoChunkId& ChunkId, IIoStoreWriteRequest* Request, const FIoWriteOptions& WriteOptions)
 	{
 		TRACE_CPUPROFILER_EVENT_SCOPE(AppendWriteRequest);
+		check(!bHasFlushed);
 		checkf(ChunkId.IsValid(), TEXT("ChunkId is not valid!"));
-		IsMetadataDirty = true;
 
 		FIoStoreWriteQueueEntry* Entry = new FIoStoreWriteQueueEntry();
+		Entry->Writer = this;
 		Entry->Sequence = Entries.Num();
 		WriterContext->TotalChunksCount.IncrementExchange();
 		Entries.Add(Entry);
@@ -588,208 +615,85 @@ public:
 		Entry->Options = WriteOptions;
 		Entry->Request = Request;
 		Entry->HashBarrier = FGraphEvent::CreateGraphEvent();
+		Entry->BeginCompressionBarrier = FGraphEvent::CreateGraphEvent();
+		Entry->FinishCompressionBarrier = FGraphEvent::CreateGraphEvent();
+		Entry->BeginWriteBarrier = FGraphEvent::CreateGraphEvent();
+		Entry->WriteFinishedEvent = FGraphEvent::CreateGraphEvent();
+		
 		FGraphEventArray HashPrereqs;
 		HashPrereqs.Add(Entry->HashBarrier);
 		Entry->HashTask = FFunctionGraphTask::CreateAndDispatchWhenReady([this, Entry]()
 		{
 			TRACE_CPUPROFILER_EVENT_SCOPE(HashChunk);
-			FIoBuffer ChunkBuffer = Entry->Request->ConsumeSourceBuffer();
-			Entry->UncompressedSize = ChunkBuffer.DataSize();
-			Entry->ChunkHash = FIoChunkHash::HashBuffer(ChunkBuffer.Data(), ChunkBuffer.DataSize());
+			const FIoBuffer* SourceBuffer = Entry->Request->GetSourceBuffer();
+			Entry->ChunkHash = FIoChunkHash::HashBuffer(SourceBuffer->Data(), SourceBuffer->DataSize());
 			WriterContext->HashedChunksCount.IncrementExchange();
+			if (LayoutEntriesHead)
+			{
+				// Release the source data buffer if disk layout ordering is enabled, it will be reloaded
+				// later when we start compressing the chunk
+				Entry->Request->FreeSourceBuffer();
+			}
 		}, TStatId(), &HashPrereqs, ENamedThreads::AnyHiPriThreadHiPriTask);
-		
-		Entry->CreateChunkBlocksBarrier = FGraphEvent::CreateGraphEvent();
-		FGraphEventArray CreateBlocksPrereqs;
-		CreateBlocksPrereqs.Add(Entry->CreateChunkBlocksBarrier);
-		Entry->CreateChunkBlocksTask = FFunctionGraphTask::CreateAndDispatchWhenReady([this, Entry]()
-		{
-			TRACE_CPUPROFILER_EVENT_SCOPE(CreateChunkBlocks);
-			Entry->ChunkBuffer = Entry->Request->ConsumeSourceBuffer();
-			CreateChunkBlocks(Entry, ContainerSettings, WriterContext->GetSettings());
-			WriterContext->CompressedChunksCount.IncrementExchange();
-		}, TStatId(), &CreateBlocksPrereqs, ENamedThreads::AnyHiPriThreadHiPriTask);
-		
-		Entry->Request->PrepareSourceBufferAsync(Entry->HashBarrier);
+		// If disk layout ordering is enabled we need to hash each chunk before we can determine the
+		// layout, otherwise we can start compressing the chunk immediately  
+		if (LayoutEntriesHead)
+		{
+			Entry->Request->PrepareSourceBufferAsync(Entry->HashBarrier);
+		}
+		else
+		{
+			WriterContext->ScheduleCompression(Entry);
+		}
 	}
 
 	UE_NODISCARD TIoStatusOr<FIoStoreWriterResult> Flush()
 	{
 		TRACE_CPUPROFILER_EVENT_SCOPE(FlushContainer);
 
-		if (!IsMetadataDirty)
+		if (bHasFlushed)
 		{
 			return Result;
 		}
+		bHasFlushed = true;
 
-		IsMetadataDirty = false;
-
-		const FIoStoreWriterSettings& Settings = WriterContext->GetSettings();
-		uint64 UncompressedFileOffset = 0;
-
-		{
-			TRACE_CPUPROFILER_EVENT_SCOPE(WaitForChunkHashes);
-			Algo::Reverse(Entries);
-			for (FIoStoreWriteQueueEntry* Entry : Entries)
-			{
-				FTaskGraphInterface::Get().WaitUntilTaskCompletes(Entry->HashTask);
-				Entry = Entry->Next;
-			}
-		}
-		{
-			TRACE_CPUPROFILER_EVENT_SCOPE(FinalizeLayout);
-			FinalizeLayout(Entries);
-		}
-
-		for (FIoStoreWriteQueueEntry* Entry : Entries)
-		{
-			WriterContext->BeginCompress(Entry);
-		}
-
-		const uint64 MaxPartitionSize = Settings.MaxPartitionSize > 0 ? Settings.MaxPartitionSize : MAX_uint64;
-		uint64 TotalEntryUncompressedSize = 0;
-		int32 CurrentPartitionIndex = 0;
-		FIoStoreWriteQueueEntry* LastEntry = Entries.Num() ? Entries.Last() : nullptr;
-		if (LastEntry)
-		{
-			LastEntry->WriteCompletedEvent = FPlatformProcess::GetSynchEventFromPool(false);
-		}
-		bool bHasMemoryMappedEntry = false;
-		for (FIoStoreWriteQueueEntry* Entry : Entries)
+		if (LayoutEntriesHead)
 		{
-			TRACE_CPUPROFILER_EVENT_SCOPE(WriteChunk);
-			FTaskGraphInterface::Get().WaitUntilTaskCompletes(Entry->CreateChunkBlocksTask);
-			
-			uint64 EntryWriteSize = Entry->ChunkBuffer.DataSize();
-			FPartition* TargetPartition = &Partitions[CurrentPartitionIndex];
-			int32 NextPartitionIndexToTry = CurrentPartitionIndex + 1;
-			if (Entry->PartitionIndex >= 0)
 			{
-				TargetPartition = &Partitions[Entry->PartitionIndex];
-				if (TargetPartition->ReservedSpace > Entry->CompressedSize)
+				TRACE_CPUPROFILER_EVENT_SCOPE(WaitForChunkHashes);
+				Algo::Reverse(Entries);
+				for (FIoStoreWriteQueueEntry* Entry : Entries)
 				{
-					TargetPartition->ReservedSpace -= Entry->CompressedSize;
+					FTaskGraphInterface::Get().WaitUntilTaskCompletes(Entry->HashTask);
+					Entry = Entry->Next;
 				}
-				else
-				{
-					TargetPartition->ReservedSpace = 0;
-				}
-				NextPartitionIndexToTry = CurrentPartitionIndex;
 			}
-			bHasMemoryMappedEntry |= Entry->Options.bIsMemoryMapped;
-			const uint64 ChunkAlignment = Entry->Options.bIsMemoryMapped ? Settings.MemoryMappingAlignment : 0;
-			checkf(EntryWriteSize <= MaxPartitionSize, TEXT("Chunk is too large, increase max partition size!"));
-			for (;;)
-			{
-				uint64 OffsetBeforePadding = TargetPartition->Offset;
-				if (ChunkAlignment)
-				{
-					TargetPartition->Offset = Align(TargetPartition->Offset, ChunkAlignment);
-				}
-				if (Settings.CompressionBlockAlignment)
-				{
-					bool bCrossesBlockBoundary = Align(TargetPartition->Offset, Settings.CompressionBlockAlignment) != Align(TargetPartition->Offset + EntryWriteSize - 1, Settings.CompressionBlockAlignment);
-					if (bCrossesBlockBoundary)
-					{
-						TargetPartition->Offset = Align(TargetPartition->Offset, Settings.CompressionBlockAlignment);
-					}
-				}
-
-				if (TargetPartition->Offset + EntryWriteSize + TargetPartition->ReservedSpace > MaxPartitionSize)
-				{
-					TargetPartition->Offset = OffsetBeforePadding;
-					while (Partitions.Num() <= NextPartitionIndexToTry)
-					{
-						FPartition& NewPartition = Partitions.AddDefaulted_GetRef();
-						NewPartition.Index = Partitions.Num() - 1;
-					}
-					CurrentPartitionIndex = NextPartitionIndexToTry;
-					TargetPartition = &Partitions[CurrentPartitionIndex];
-					++NextPartitionIndexToTry;
-				}
-				else
-				{
-					Entry->Padding = TargetPartition->Offset - OffsetBeforePadding;
-					TotalPaddedBytes += Entry->Padding;
-					break;
-				}
-			}
-
-			if (!TargetPartition->ContainerFileHandle)
-			{
-				TRACE_CPUPROFILER_EVENT_SCOPE(CreatePartitionContainerFile);
-				CreatePartitionContainerFile(*TargetPartition);
-			}
-			Entry->ContainerArchive = TargetPartition->ContainerFileHandle.Get();
-			Entry->Offset = TargetPartition->Offset;
-			WriterContext->BeginWrite(Entry);
-			
-			FIoOffsetAndLength OffsetLength;
-			OffsetLength.SetOffset(UncompressedFileOffset);
-			OffsetLength.SetLength(Entry->UncompressedSize);
-
-			FIoStoreTocEntryMeta ChunkMeta{ Entry->ChunkHash, FIoStoreTocEntryMetaFlags::None };
-			if (Entry->Options.bIsMemoryMapped)
-			{
-				ChunkMeta.Flags |= FIoStoreTocEntryMetaFlags::MemoryMapped;
-			}
-
-			for (const FChunkBlock& ChunkBlock : Entry->ChunkBlocks)
-			{
-				check(ChunkBlock.Offset + ChunkBlock.Size <= EntryWriteSize);
-
-				FIoStoreTocCompressedBlockEntry& BlockEntry = Toc.AddCompressionBlockEntry();
-				BlockEntry.SetOffset(TargetPartition->Index * Settings.MaxPartitionSize + TargetPartition->Offset + ChunkBlock.Offset);
-				BlockEntry.SetCompressedSize(uint32(ChunkBlock.CompressedSize));
-				BlockEntry.SetUncompressedSize(uint32(ChunkBlock.UncompressedSize));
-				BlockEntry.SetCompressionMethodIndex(Toc.AddCompressionMethodEntry(ChunkBlock.CompressionMethod));
-
-				if (!ChunkBlock.CompressionMethod.IsNone())
-				{
-					ChunkMeta.Flags |= FIoStoreTocEntryMetaFlags::Compressed;
-				}
-
-				if (ContainerSettings.IsSigned())
-				{
-					FSHAHash& Signature = Toc.AddBlockSignatureEntry();
-					Signature = ChunkBlock.Signature;
-				}
-			}
-
-			const int32 TocEntryIndex = Toc.AddChunkEntry(Entry->ChunkId, OffsetLength, ChunkMeta);
-			check(TocEntryIndex != INDEX_NONE);
-
-			if (ContainerSettings.IsIndexed() && Entry->Options.FileName.Len() > 0)
-			{
-				Toc.AddToFileIndex(Entry->Options.FileName, TocEntryIndex);
-			}
-
-			const uint64 RegionStartOffset = TargetPartition->Offset;
-			TargetPartition->Offset += EntryWriteSize;
-			UncompressedFileOffset += Align(Entry->UncompressedSize, Settings.CompressionBlockSize);
-			TotalEntryUncompressedSize += Entry->UncompressedSize;
+			FinalizeLayout(Entries);
 
-			if (Settings.bEnableFileRegions)
+			for (FIoStoreWriteQueueEntry* Entry : Entries)
 			{
-				FFileRegion::AccumulateFileRegions(TargetPartition->AllFileRegions, RegionStartOffset, RegionStartOffset, TargetPartition->Offset, Entry->Request->GetRegions());
+				WriterContext->ScheduleCompression(Entry);
 			}
 		}
-		if (LastEntry)
+
 		{
 			TRACE_CPUPROFILER_EVENT_SCOPE(WaitForWritesToComplete);
-			LastEntry->WriteCompletedEvent->Wait();
-			FPlatformProcess::ReturnSynchEventToPool(LastEntry->WriteCompletedEvent);
+			for (FIoStoreWriteQueueEntry* Entry : Entries)
+			{
+				Entry->WriteFinishedEvent->Wait();
+			}
 		}
-		
-		UncompressedContainerSize = TotalEntryUncompressedSize + TotalPaddedBytes;
+
+		UncompressedContainerSize = TotalEntryUncompressedSize + TotalPaddingSize;
 		CompressedContainerSize = 0;
+		const FIoStoreWriterSettings& WriterSettings = WriterContext->GetSettings();
 		for (FPartition& Partition : Partitions)
 		{
 			CompressedContainerSize += Partition.Offset;
 
 			if (bHasMemoryMappedEntry)
 			{
-				uint64 ExtraPaddingBytes = Align(Partition.Offset, Settings.MemoryMappingAlignment) - Partition.Offset;
+				uint64 ExtraPaddingBytes = Align(Partition.Offset, WriterSettings.MemoryMappingAlignment) - Partition.Offset;
 				if (ExtraPaddingBytes)
 				{
 					TArray<uint8> Padding;
@@ -802,6 +706,7 @@ public:
 			
 			if (Partition.ContainerFileHandle)
 			{
+				TRACE_CPUPROFILER_EVENT_SCOPE(FlushContainerFile);
 				Partition.ContainerFileHandle->Flush();
 				check(Partition.ContainerFileHandle->Tell() == Partition.Offset);
 			}
@@ -817,6 +722,7 @@ public:
 
 		if (ContainerSettings.IsIndexed())
 		{
+			TRACE_CPUPROFILER_EVENT_SCOPE(BuildIndex);
 			const TArray<FString>& FilesToIndex = Toc.GetFilesToIndex();
 			const TArray<uint32>& FileTocEntryIndices = Toc.GetFileTocEntryIndices();
 
@@ -848,7 +754,7 @@ public:
 		Result.ContainerFlags = ContainerSettings.ContainerFlags;
 		Result.TocSize = TocSize.ConsumeValueOrDie();
 		Result.TocEntryCount = TocResource.Header.TocEntryCount;
-		Result.PaddingSize = TotalPaddedBytes;
+		Result.PaddingSize = TotalPaddingSize;
 		Result.UncompressedContainerSize = UncompressedContainerSize;
 		Result.CompressedContainerSize = CompressedContainerSize;
 		Result.DirectoryIndexSize = TocResource.Header.DirectoryIndexSize;
@@ -859,20 +765,22 @@ public:
 		Result.AddedChunksCount = 0;
 		Result.ModifiedChunksSize= 0;
 		Result.AddedChunksSize = 0;
-		for (FIoStoreWriteQueueEntry* Entry : Entries)
 		{
-			if (Entry->bModified)
-			{
-				++Result.ModifiedChunksCount;
-				Result.ModifiedChunksSize += Entry->CompressedSize;
-			}
-			else if (Entry->bAdded)
+			TRACE_CPUPROFILER_EVENT_SCOPE(Cleanup);
+			for (FIoStoreWriteQueueEntry* Entry : Entries)
 			{
-				++Result.AddedChunksCount;
-				Result.AddedChunksSize += Entry->CompressedSize;
+				if (Entry->bModified)
+				{
+					++Result.ModifiedChunksCount;
+					Result.ModifiedChunksSize += Entry->CompressedSize;
+				}
+				else if (Entry->bAdded)
+				{
+					++Result.AddedChunksCount;
+					Result.AddedChunksSize += Entry->CompressedSize;
+				}
+				delete Entry;
 			}
-			delete Entry->Request;
-			delete Entry;
 		}
 
 		Entries.Empty();
@@ -902,46 +810,10 @@ private:
 		int32 PartitionIndex = -1;
 	};
 
-	void PrepareLayout(const TArray<TUniquePtr<FIoStoreReader>>& PatchSourceReaders)
-	{
-		LayoutEntriesHead = new FLayoutEntry();
-		LayoutEntries.Add(LayoutEntriesHead);
-		FLayoutEntry* PrevEntryLink = LayoutEntriesHead;
-
-		for (const TUniquePtr<FIoStoreReader>& PatchSourceReader : PatchSourceReaders)
-		{
-			PatchSourceReader->EnumerateChunks([this, &PrevEntryLink](const FIoStoreTocChunkInfo& ChunkInfo)
-			{
-				FLayoutEntry* PreviousBuildEntry = new FLayoutEntry();
-				LayoutEntries.Add(PreviousBuildEntry);
-
-				PreviousBuildEntry->Hash = ChunkInfo.Hash;
-				PreviousBuildEntry->PartitionIndex = ChunkInfo.PartitionIndex;
-				PreviousBuildEntry->CompressedSize = ChunkInfo.CompressedSize;
-				PrevEntryLink->Next = PreviousBuildEntry;
-				PreviousBuildEntry->Prev = PrevEntryLink;
-				PrevEntryLink = PreviousBuildEntry;
-				PreviousBuildLayoutEntryByChunkId.Add(ChunkInfo.Id, PreviousBuildEntry);
-				return true;
-			});
-
-			if (!ContainerSettings.bGenerateDiffPatch)
-			{
-				break;
-			}
-		}
-
-		LayoutEntriesTail = new FLayoutEntry();
-		LayoutEntries.Add(LayoutEntriesTail);
-		PrevEntryLink->Next = LayoutEntriesTail;
-		LayoutEntriesTail->Prev = PrevEntryLink;
-	}
-
 	void FinalizeLayout(TArray<FIoStoreWriteQueueEntry*>& WriteQueueEntries)
 	{
-		FPartition& Partition = Partitions.AddDefaulted_GetRef();
-		Partition.Index = 0;
-
+		TRACE_CPUPROFILER_EVENT_SCOPE(FinalizeLayout);
+		
 		Algo::Sort(WriteQueueEntries, [](const FIoStoreWriteQueueEntry* A, const FIoStoreWriteQueueEntry* B)
 		{
 			uint64 AOrderHint = A->Request->GetOrderHint();
@@ -976,7 +848,7 @@ private:
 			{
 				WriteQueueEntry->bAdded = true;
 			}
-			if (WriteQueueEntry->bModified | WriteQueueEntry->bAdded)
+			if (WriteQueueEntry->bModified || WriteQueueEntry->bAdded)
 			{
 				FLayoutEntry* NewLayoutEntry = new FLayoutEntry();
 				NewLayoutEntry->QueueEntry = WriteQueueEntry;
@@ -1066,8 +938,6 @@ private:
 		ContainerFilePath += TEXT(".ucas");
 		
 		Partition.ContainerFileHandle.Reset(IFileManager::Get().CreateFileWriter(*ContainerFilePath));
-		//IPlatformFile& Ipf = IPlatformFile::GetPlatformPhysical();
-		//Partition.ContainerFileHandle.Reset(Ipf.OpenWrite(*ContainerFilePath, /* append */ false, /* allowread */ true));
 		if (!Partition.ContainerFileHandle)
 		{
 			return FIoStatusBuilder(EIoErrorCode::FileOpenFailed) << TEXT("Failed to open IoStore container file '") << *ContainerFilePath << TEXT("'");
@@ -1085,150 +955,246 @@ private:
 		return FIoStatus::Ok;
 	}
 
-	static void CreateChunkBlocks(
-		FIoStoreWriteQueueEntry* Entry,
-		const FIoContainerSettings& ContainerSettings,
-		const FIoStoreWriterSettings& WriterSettings)
+	void BeginCompress(FIoStoreWriteQueueEntry* Entry)
 	{
-		check(WriterSettings.CompressionBlockSize > 0);
-
-		const uint64 NumChunkBlocks = Align(Entry->ChunkBuffer.DataSize(), WriterSettings.CompressionBlockSize) / WriterSettings.CompressionBlockSize;
-		Entry->ChunkBlocks.Reserve(int32(NumChunkBlocks));
-
-		auto CreateUncompressedBlocks = [](FIoStoreWriteQueueEntry* UncompressedEntry, const uint64 BlockSize) -> void
+		FName CompressionMethod = NAME_None;
+		const FIoStoreWriterSettings& WriterSettings = WriterContext->WriterSettings;
+		if (ContainerSettings.IsCompressed() && !Entry->Options.bForceUncompressed && !Entry->Options.bIsMemoryMapped)
 		{
-			UncompressedEntry->ChunkBlocks.Empty();
+			CompressionMethod = WriterSettings.CompressionMethod;
+		}
 
-			FIoBuffer& ChunkBuffer = UncompressedEntry->ChunkBuffer;
+		const FIoBuffer* SourceBuffer = Entry->Request->GetSourceBuffer();
+		Entry->UncompressedSize = SourceBuffer->DataSize();
 
-			uint64 UncompressedSize = ChunkBuffer.DataSize();
-			uint64 RawSize = UncompressedSize;
-			if (!IsAligned(RawSize, FAES::AESBlockSize))
+		check(WriterSettings.CompressionBlockSize > 0);
+		const uint64 NumChunkBlocks = Align(Entry->UncompressedSize, WriterSettings.CompressionBlockSize) / WriterSettings.CompressionBlockSize;
+		Entry->ChunkBlocks.SetNum(int32(NumChunkBlocks));
+
+		uint64 BytesToProcess = Entry->UncompressedSize;
+		const uint8* UncompressedData = SourceBuffer->Data();
+		for (int32 BlockIndex = 0; BlockIndex < NumChunkBlocks; ++BlockIndex)
+		{
+			FChunkBlock& Block = Entry->ChunkBlocks[BlockIndex];
+			Block.IoBuffer = WriterContext->AllocCompressionBuffer();
+			Block.CompressionMethod = CompressionMethod;
+			Block.UncompressedSize = FMath::Min(BytesToProcess, WriterSettings.CompressionBlockSize);
+			Block.UncompressedData = UncompressedData;
+			BytesToProcess -= Block.UncompressedSize;
+			UncompressedData += Block.UncompressedSize;
+
+			FFunctionGraphTask::CreateAndDispatchWhenReady([this, Entry, BlockPtr = &Block]()
 			{
-				RawSize = Align(RawSize, FAES::AESBlockSize);
-				FIoBuffer AlignedBuffer(RawSize);
-				FMemory::Memcpy(AlignedBuffer.Data(), ChunkBuffer.Data(), UncompressedSize);
-
-				for (uint64 FillIndex = UncompressedSize; FillIndex < RawSize ; ++FillIndex)
+				CompressBlock(BlockPtr);
+				int32 FinishedBlocksCount = Entry->FinishedBlocksCount.IncrementExchange();
+				if (FinishedBlocksCount + 1 == Entry->ChunkBlocks.Num())
 				{
-					AlignedBuffer.Data()[FillIndex] = AlignedBuffer.Data()[(FillIndex - UncompressedSize) % UncompressedSize];
+					WriterContext->CompressedChunksCount.IncrementExchange();
+					Entry->FinishCompressionBarrier->DispatchSubsequents();
 				}
+			}, TStatId(), nullptr, ENamedThreads::AnyHiPriThreadHiPriTask);
+		}
+	}
 
-				ChunkBuffer = AlignedBuffer;
+	void CompressBlock(FChunkBlock* Block)
+	{
+		if (Block->CompressionMethod != NAME_None)
+		{
+			int32 CompressedBlockSize = int32(Block->IoBuffer->DataSize());
+			{
+				TRACE_CPUPROFILER_EVENT_SCOPE(CompressMemory);
+				const bool bCompressed = FCompression::CompressMemory(
+					Block->CompressionMethod,
+					Block->IoBuffer->Data(),
+					CompressedBlockSize,
+					Block->UncompressedData,
+					static_cast<int32>(Block->UncompressedSize));
+				check(bCompressed);
 			}
-
-			UncompressedEntry->CompressedSize = ChunkBuffer.DataSize();
-
-			uint64 UncompressedOffset = 0;
-			uint64 RemainingSize = UncompressedSize;
-			while (RemainingSize)
+			check(CompressedBlockSize > 0);
+			if (CompressedBlockSize >= Block->UncompressedSize)
 			{
-				const uint64 UncompressedBlockSize = FMath::Min<uint64>(RemainingSize, BlockSize);
-				const uint64 RawBlockSize = Align(UncompressedBlockSize, FAES::AESBlockSize);
-				UncompressedEntry->ChunkBlocks.Add(FChunkBlock { UncompressedOffset, RawBlockSize, UncompressedBlockSize, UncompressedBlockSize, NAME_None });
-				RemainingSize -= UncompressedBlockSize;
-				UncompressedOffset += RawBlockSize;
+				Block->CompressionMethod = NAME_None;
 			}
-		};
+			else
+			{
+				Block->CompressedSize = CompressedBlockSize;
+			}
+		}
+		if (Block->CompressionMethod == NAME_None)
+		{
+			Block->CompressedSize = Block->UncompressedSize;
+			FMemory::Memcpy(Block->IoBuffer->Data(), Block->UncompressedData, Block->UncompressedSize);
+		}
 
-		if (ContainerSettings.IsCompressed() && !Entry->Options.bForceUncompressed && !Entry->Options.bIsMemoryMapped)
+		// Always align each compressed block to AES block size but store the compressed block size in the TOC
+		Block->Size = Block->CompressedSize;
+		if (!IsAligned(Block->Size, FAES::AESBlockSize))
 		{
-			check(!WriterSettings.CompressionMethod.IsNone());
+			uint64 AlignedCompressedBlockSize = Align(Block->Size, FAES::AESBlockSize);
+			uint8* CompressedData = Block->IoBuffer->Data();
+			for (uint64 FillIndex = Block->Size; FillIndex < AlignedCompressedBlockSize; ++FillIndex)
+			{
+				check(FillIndex < Block->IoBuffer->DataSize());
+				CompressedData[FillIndex] = CompressedData[(FillIndex - Block->CompressedSize) % Block->Size];
+			}
+			Block->Size = AlignedCompressedBlockSize;
+		}
 
-			const uint8* UncompressedBlock = Entry->ChunkBuffer.Data();
-			TArray<TUniquePtr<uint8[]>> CompressedBlocks;
-			CompressedBlocks.Reserve(int32(NumChunkBlocks));
+		if (ContainerSettings.IsEncrypted())
+		{
+			FAES::EncryptData(Block->IoBuffer->Data(), static_cast<uint32>(Block->Size), ContainerSettings.EncryptionKey);
+		}
 
-			uint64 BytesToProcess	= Entry->ChunkBuffer.DataSize();
-			uint64 BlockOffset		= 0;
+		if (ContainerSettings.IsSigned())
+		{
+			FSHA1::HashBuffer(Block->IoBuffer->Data(), Block->Size, Block->Signature.Hash);
+		}
+	}
 
-			while (BytesToProcess > 0)
+	void WriteEntry(FIoStoreWriteQueueEntry* Entry)
+	{
+		TRACE_CPUPROFILER_EVENT_SCOPE(WriteEntry);
+		FPartition* TargetPartition = &Partitions[CurrentPartitionIndex];
+		int32 NextPartitionIndexToTry = CurrentPartitionIndex + 1;
+		if (Entry->PartitionIndex >= 0)
+		{
+			TargetPartition = &Partitions[Entry->PartitionIndex];
+			if (TargetPartition->ReservedSpace > Entry->CompressedSize)
+			{
+				TargetPartition->ReservedSpace -= Entry->CompressedSize;
+			}
+			else
 			{
-				const int32 UncompressedBlockSize = static_cast<int32>(FMath::Min(BytesToProcess, WriterSettings.CompressionBlockSize));
-				int32 CompressedBlockSize = FCompression::CompressMemoryBound(WriterSettings.CompressionMethod, UncompressedBlockSize);
-				TUniquePtr<uint8[]>& CompressedBlock = CompressedBlocks.AddDefaulted_GetRef();
-				CompressedBlock = MakeUnique<uint8[]>(CompressedBlockSize);
+				TargetPartition->ReservedSpace = 0;
+			}
+			NextPartitionIndexToTry = CurrentPartitionIndex;
+		}
 
-				FName CompressionMethod = WriterSettings.CompressionMethod;
+		const FIoStoreWriterSettings& WriterSettings = WriterContext->WriterSettings;
+		bHasMemoryMappedEntry |= Entry->Options.bIsMemoryMapped;
+		const uint64 ChunkAlignment = Entry->Options.bIsMemoryMapped ? WriterSettings.MemoryMappingAlignment : 0;
+		const uint64 PartitionSizeLimit = WriterSettings.MaxPartitionSize > 0 ? WriterSettings.MaxPartitionSize : MAX_uint64;
+		checkf(Entry->CompressedSize <= PartitionSizeLimit, TEXT("Chunk is too large, increase max partition size!"));
+		for (;;)
+		{
+			uint64 OffsetBeforePadding = TargetPartition->Offset;
+			if (ChunkAlignment)
+			{
+				TargetPartition->Offset = Align(TargetPartition->Offset, ChunkAlignment);
+			}
+			if (WriterSettings.CompressionBlockAlignment)
+			{
+				bool bCrossesBlockBoundary = Align(TargetPartition->Offset, WriterSettings.CompressionBlockAlignment) != Align(TargetPartition->Offset + Entry->CompressedSize - 1, WriterSettings.CompressionBlockAlignment);
+				if (bCrossesBlockBoundary)
 				{
-					TRACE_CPUPROFILER_EVENT_SCOPE(CompressMemory);
-					const bool bCompressed = FCompression::CompressMemory(
-						CompressionMethod,
-						CompressedBlock.Get(),
-						CompressedBlockSize,
-						UncompressedBlock,
-						UncompressedBlockSize);
-					check(bCompressed);
+					TargetPartition->Offset = Align(TargetPartition->Offset, WriterSettings.CompressionBlockAlignment);
 				}
-				check(CompressedBlockSize > 0);
+			}
 
-				if (CompressedBlockSize >= UncompressedBlockSize)
+			if (TargetPartition->Offset + Entry->CompressedSize + TargetPartition->ReservedSpace > PartitionSizeLimit)
+			{
+				TargetPartition->Offset = OffsetBeforePadding;
+				while (Partitions.Num() <= NextPartitionIndexToTry)
 				{
-					memcpy(CompressedBlock.Get(), UncompressedBlock, UncompressedBlockSize);
-					CompressedBlockSize = UncompressedBlockSize;
-					CompressionMethod = NAME_None;
+					FPartition& NewPartition = Partitions.AddDefaulted_GetRef();
+					NewPartition.Index = Partitions.Num() - 1;
 				}
+				CurrentPartitionIndex = NextPartitionIndexToTry;
+				TargetPartition = &Partitions[CurrentPartitionIndex];
+				++NextPartitionIndexToTry;
+			}
+			else
+			{
+				Entry->Padding = TargetPartition->Offset - OffsetBeforePadding;
+				TotalPaddingSize += Entry->Padding;
+				break;
+			}
+		}
 
-				// Always align each compressed block to AES block size but store the compressed block size in the TOC
-				uint64 AlignedCompressedBlockSize = CompressedBlockSize;
-				if (!IsAligned(CompressedBlockSize, FAES::AESBlockSize))
-				{
-					AlignedCompressedBlockSize = Align(CompressedBlockSize, FAES::AESBlockSize);
-					TUniquePtr<uint8[]> AlignedBlock = MakeUnique<uint8[]>(AlignedCompressedBlockSize);
-
-					FMemory::Memcpy(AlignedBlock.Get(), CompressedBlock.Get(), CompressedBlockSize);
-
-					for (uint64 FillIndex = CompressedBlockSize; FillIndex < AlignedCompressedBlockSize ; ++FillIndex)
-					{
-						AlignedBlock.Get()[FillIndex] = AlignedBlock.Get()[(FillIndex - CompressedBlockSize) % CompressedBlockSize];
-					}
+		if (!TargetPartition->ContainerFileHandle)
+		{
+			TRACE_CPUPROFILER_EVENT_SCOPE(CreatePartitionContainerFile);
+			CreatePartitionContainerFile(*TargetPartition);
+		}
+		Entry->Offset = TargetPartition->Offset;
+		
+		FIoOffsetAndLength OffsetLength;
+		OffsetLength.SetOffset(UncompressedFileOffset);
+		OffsetLength.SetLength(Entry->UncompressedSize);
 
-					CompressedBlock.Reset(AlignedBlock.Release());
-				}
+		FIoStoreTocEntryMeta ChunkMeta{ Entry->ChunkHash, FIoStoreTocEntryMetaFlags::None };
+		if (Entry->Options.bIsMemoryMapped)
+		{
+			ChunkMeta.Flags |= FIoStoreTocEntryMetaFlags::MemoryMapped;
+		}
 
-				Entry->ChunkBlocks.Add(FChunkBlock { BlockOffset, AlignedCompressedBlockSize, uint64(CompressedBlockSize), uint64(UncompressedBlockSize), CompressionMethod });
+		uint64 OffsetInChunk = 0;
+		for (const FChunkBlock& ChunkBlock : Entry->ChunkBlocks)
+		{
+			FIoStoreTocCompressedBlockEntry& BlockEntry = Toc.AddCompressionBlockEntry();
+			BlockEntry.SetOffset(TargetPartition->Index * WriterSettings.MaxPartitionSize + TargetPartition->Offset + OffsetInChunk);
+			OffsetInChunk += ChunkBlock.Size;
+			BlockEntry.SetCompressedSize(uint32(ChunkBlock.CompressedSize));
+			BlockEntry.SetUncompressedSize(uint32(ChunkBlock.UncompressedSize));
+			BlockEntry.SetCompressionMethodIndex(Toc.AddCompressionMethodEntry(ChunkBlock.CompressionMethod));
 
-				BytesToProcess		-= UncompressedBlockSize;
-				BlockOffset			+= AlignedCompressedBlockSize;
-				UncompressedBlock	+= UncompressedBlockSize;
+			if (!ChunkBlock.CompressionMethod.IsNone())
+			{
+				ChunkMeta.Flags |= FIoStoreTocEntryMetaFlags::Compressed;
 			}
 
-			Entry->CompressedSize = BlockOffset;
-			
-			Entry->ChunkBuffer = FIoBuffer(Entry->CompressedSize);
-			
-			uint8* CompressedChunkBuffer = Entry->ChunkBuffer.Data();
-			FMemory::Memzero(CompressedChunkBuffer, Entry->CompressedSize); 
-
-			for (int32 BlockIndex = 0; BlockIndex < CompressedBlocks.Num(); ++BlockIndex)
+			if (ContainerSettings.IsSigned())
 			{
-				TUniquePtr<uint8[]>& CompressedBlock = CompressedBlocks[BlockIndex];
-				const FChunkBlock& ChunkBlock = Entry->ChunkBlocks[BlockIndex];
-				FMemory::Memcpy(CompressedChunkBuffer, CompressedBlock.Get(), ChunkBlock.Size);
-				CompressedChunkBuffer += ChunkBlock.Size;
+				FSHAHash& Signature = Toc.AddBlockSignatureEntry();
+				Signature = ChunkBlock.Signature;
 			}
 		}
-		else
+
+		const int32 TocEntryIndex = Toc.AddChunkEntry(Entry->ChunkId, OffsetLength, ChunkMeta);
+		check(TocEntryIndex != INDEX_NONE);
+
+		if (ContainerSettings.IsIndexed() && Entry->Options.FileName.Len() > 0)
 		{
-			CreateUncompressedBlocks(Entry, WriterSettings.CompressionBlockSize);
+			Toc.AddToFileIndex(Entry->Options.FileName, TocEntryIndex);
 		}
 
-		if (ContainerSettings.IsEncrypted())
+		const uint64 RegionStartOffset = TargetPartition->Offset;
+		TargetPartition->Offset += Entry->CompressedSize;
+		UncompressedFileOffset += Align(Entry->UncompressedSize, WriterSettings.CompressionBlockSize);
+		TotalEntryUncompressedSize += Entry->UncompressedSize;
+
+		if (WriterSettings.bEnableFileRegions)
 		{
-			for (FChunkBlock& ChunkBlock : Entry->ChunkBlocks)
+			FFileRegion::AccumulateFileRegions(TargetPartition->AllFileRegions, RegionStartOffset, RegionStartOffset, TargetPartition->Offset, Entry->Request->GetRegions());
+		}
+		delete Entry->Request;
+		Entry->Request = nullptr;
+
+		if (Entry->Padding > 0)
+		{
+			if (PaddingBuffer.Num() < Entry->Padding)
+			{
+				PaddingBuffer.SetNumZeroed(int32(Entry->Padding));
+			}
 			{
-				check(IsAligned(ChunkBlock.Size, FAES::AESBlockSize));
-				FAES::EncryptData(Entry->ChunkBuffer.Data() + ChunkBlock.Offset, static_cast<uint32>(ChunkBlock.Size), ContainerSettings.EncryptionKey);
+				TRACE_CPUPROFILER_EVENT_SCOPE(WritePaddingToContainer);
+				TargetPartition->ContainerFileHandle->Serialize(PaddingBuffer.GetData(), Entry->Padding);
 			}
 		}
-
-		if (ContainerSettings.IsSigned())
+		check(Entry->Offset == TargetPartition->ContainerFileHandle->Tell());
+		for (FChunkBlock& ChunkBlock : Entry->ChunkBlocks)
 		{
-			for (FChunkBlock& ChunkBlock : Entry->ChunkBlocks)
 			{
-				FSHA1::HashBuffer(Entry->ChunkBuffer.Data() + ChunkBlock.Offset, ChunkBlock.Size, ChunkBlock.Signature.Hash);
+				TRACE_CPUPROFILER_EVENT_SCOPE(WriteBlockToContainer);
+				TargetPartition->ContainerFileHandle->Serialize(ChunkBlock.IoBuffer->Data(), ChunkBlock.Size);
 			}
+			// TODO: Write each block when it becomes available, currently this limits the maximum
+			// chunk size to the total compression buffer size
+			WriterContext->FreeCompressionBuffer(ChunkBlock.IoBuffer);
+			ChunkBlock.IoBuffer = nullptr;
 		}
+		WriterContext->SerializedChunksCount.IncrementExchange();
 	}
 
 	const FString				ContainerPath;
@@ -1236,6 +1202,7 @@ private:
 	FIoContainerSettings		ContainerSettings;
 	FString						TocFilePath;
 	FIoStoreToc					Toc;
+	TArray<uint8>				PaddingBuffer;
 	TArray<FPartition>			Partitions;
 	TArray<FIoStoreWriteQueueEntry*> Entries;
 	TArray<FLayoutEntry*>		LayoutEntries;
@@ -1244,12 +1211,98 @@ private:
 	TMap<FIoChunkId, FLayoutEntry*> PreviousBuildLayoutEntryByChunkId;
 	TUniquePtr<FArchive>		CsvArchive;
 	FIoStoreWriterResult		Result;
-	uint64						TotalPaddedBytes = 0;
+	uint64						UncompressedFileOffset = 0;
+	uint64						TotalEntryUncompressedSize = 0;
+	uint64						TotalPaddingSize = 0;
 	uint64						UncompressedContainerSize = 0;
 	uint64						CompressedContainerSize = 0;
-	bool						IsMetadataDirty = true;
+	int32						CurrentPartitionIndex = 0;
+	bool						bHasMemoryMappedEntry = false;
+	bool						bHasFlushed = false;
+
+	friend class FIoStoreWriterContextImpl;
 };
 
+void FIoStoreWriterContextImpl::BeginCompressionThreadFunc()
+{
+	for (;;)
+	{
+		FIoStoreWriteQueueEntry* Entry = BeginCompressionQueue.DequeueOrWait();
+		if (!Entry)
+		{
+			return;
+		}
+		while (Entry)
+		{
+			FIoStoreWriteQueueEntry* Next = Entry->Next;
+			Entry->BeginCompressionBarrier->Wait();
+
+			FinishCompressionQueue.Enqueue(Entry);
+
+			if (!Entry->HashBarrier->IsComplete())
+			{
+				Entry->HashBarrier->DispatchSubsequents();
+			}
+			
+			Entry->Writer->BeginCompress(Entry);
+
+			Entry = Next;
+		}
+	}
+}
+
+void FIoStoreWriterContextImpl::FinishCompressionThreadFunc()
+{
+	for (;;)
+	{
+		FIoStoreWriteQueueEntry* Entry = FinishCompressionQueue.DequeueOrWait();
+		if (!Entry)
+		{
+			return;
+		}
+		while (Entry)
+		{
+			FIoStoreWriteQueueEntry* Next = Entry->Next;
+			Entry->FinishCompressionBarrier->Wait();
+			Entry->HashTask->Wait();
+
+			WriterQueue.Enqueue(Entry);
+
+			Entry->Request->FreeSourceBuffer();
+
+			Entry->CompressedSize = 0;
+			for (const FChunkBlock& ChunkBlock : Entry->ChunkBlocks)
+			{
+				Entry->CompressedSize += ChunkBlock.Size;
+			}
+			Entry->BeginWriteBarrier->DispatchSubsequents();
+
+			Entry = Next;
+		}
+	}
+}
+
+void FIoStoreWriterContextImpl::WriterThreadFunc()
+{
+	for (;;)
+	{
+		FIoStoreWriteQueueEntry* Entry = WriterQueue.DequeueOrWait();
+		if (!Entry)
+		{
+			return;
+		}
+		while (Entry)
+		{
+			FIoStoreWriteQueueEntry* Next = Entry->Next;
+			Entry->BeginWriteBarrier->Wait();
+			Entry->Writer->WriteEntry(Entry);
+			Entry->WriteFinishedEvent->DispatchSubsequents();
+			Entry = Next;
+		}
+	}
+}
+
+
 FIoStoreWriter::FIoStoreWriter(const TCHAR* ContainerPath)
 :	Impl(new FIoStoreWriterImpl(ContainerPath))
 {
@@ -1260,9 +1313,14 @@ FIoStoreWriter::~FIoStoreWriter()
 	(void)Impl->Flush();
 }
 
-FIoStatus FIoStoreWriter::Initialize(const FIoStoreWriterContext& Context, const FIoContainerSettings& ContainerSettings, const TArray<TUniquePtr<FIoStoreReader>>& PatchSourceReaders)
+FIoStatus FIoStoreWriter::Initialize(const FIoStoreWriterContext& Context, const FIoContainerSettings& ContainerSettings)
 {
-	return Impl->Initialize(*Context.Impl, ContainerSettings, PatchSourceReaders);
+	return Impl->Initialize(*Context.Impl, ContainerSettings);
+}
+
+void FIoStoreWriter::EnableDiskLayoutOrdering(const TArray<TUniquePtr<FIoStoreReader>>& PatchSourceReaders)
+{
+	Impl->EnableDiskLayoutOrdering(PatchSourceReaders);
 }
 
 void FIoStoreWriter::Append(const FIoChunkId& ChunkId, FIoBuffer Chunk, const FIoWriteOptions& WriteOptions)
@@ -1280,13 +1338,16 @@ void FIoStoreWriter::Append(const FIoChunkId& ChunkId, FIoBuffer Chunk, const FI
 
 		void PrepareSourceBufferAsync(FGraphEventRef CompletionEvent) override
 		{
-			TArray<FBaseGraphTask*> NewTasks;
-			CompletionEvent->DispatchSubsequents(NewTasks);
+			CompletionEvent->DispatchSubsequents();
+		}
+
+		const FIoBuffer* GetSourceBuffer() override
+		{
+			return &SourceBuffer;
 		}
 
-		FIoBuffer ConsumeSourceBuffer() override
+		void FreeSourceBuffer() override
 		{
-			return SourceBuffer;
 		}
 
 		uint64 GetOrderHint() override
diff --git a/Engine/Source/Runtime/Core/Public/IO/IoDispatcher.h b/Engine/Source/Runtime/Core/Public/IO/IoDispatcher.h
--- a/Engine/Source/Runtime/Core/Public/IO/IoDispatcher.h
+++ b/Engine/Source/Runtime/Core/Public/IO/IoDispatcher.h
@@ -973,7 +973,7 @@ private:
 struct FIoStoreWriterSettings
 {
 	FName CompressionMethod = NAME_None;
-	uint64 CompressionBlockSize = 0;
+	uint64 CompressionBlockSize = 64 << 10;
 	uint64 CompressionBlockAlignment = 0;
 	uint64 MemoryMappingAlignment = 0;
 	uint64 MaxPartitionSize = 0;
@@ -1077,7 +1077,8 @@ public:
 	virtual void PrepareSourceBufferAsync(FGraphEventRef CompletionEvent) = 0;
 	virtual uint64 GetOrderHint() = 0;
 	virtual TArrayView<const FFileRegion> GetRegions() = 0;
-	virtual FIoBuffer ConsumeSourceBuffer() = 0;
+	virtual const FIoBuffer* GetSourceBuffer() = 0;
+	virtual void FreeSourceBuffer() = 0;
 };
 
 class FIoStoreWriter
@@ -1089,7 +1090,8 @@ public:
 	FIoStoreWriter(const FIoStoreWriter&) = delete;
 	FIoStoreWriter& operator=(const FIoStoreWriter&) = delete;
 
-	UE_NODISCARD CORE_API FIoStatus	Initialize(const FIoStoreWriterContext& Context, const FIoContainerSettings& ContainerSettings, const TArray<TUniquePtr<FIoStoreReader>>& PatchSourceReaders = TArray<TUniquePtr<FIoStoreReader>>());
+	UE_NODISCARD CORE_API FIoStatus	Initialize(const FIoStoreWriterContext& Context, const FIoContainerSettings& ContainerSettings);
+	CORE_API void EnableDiskLayoutOrdering(const TArray<TUniquePtr<FIoStoreReader>>& PatchSourceReaders = TArray<TUniquePtr<FIoStoreReader>>());
 	CORE_API void Append(const FIoChunkId& ChunkId, FIoBuffer Chunk, const FIoWriteOptions& WriteOptions);
 	CORE_API void Append(const FIoChunkId& ChunkId, IIoStoreWriteRequest* Request, const FIoWriteOptions& WriteOptions);
 	UE_NODISCARD CORE_API TIoStatusOr<FIoStoreWriterResult> Flush();
