Fix clang compile error. Clang's std::atomic copy ctor is deleted

diff --git a/Engine/Source/Runtime/Core/Private/ProfilingDebugging/StallDetector.cpp b/Engine/Source/Runtime/Core/Private/ProfilingDebugging/StallDetector.cpp
--- a/Engine/Source/Runtime/Core/Private/ProfilingDebugging/StallDetector.cpp
+++ b/Engine/Source/Runtime/Core/Private/ProfilingDebugging/StallDetector.cpp
@@ -127,8 +127,8 @@ uint32 UE::FStallDetectorRunnable::Run()
 
 FCriticalSection UE::FStallDetectorStats::CriticalSection;
 TSet<UE::FStallDetectorStats*> UE::FStallDetectorStats::Instances;
-std::atomic<uint32> UE::FStallDetectorStats::TotalTriggeredCount = 0;
-std::atomic<uint32> UE::FStallDetectorStats::TotalReportedCount = 0;
+std::atomic<uint32> UE::FStallDetectorStats::TotalTriggeredCount;
+std::atomic<uint32> UE::FStallDetectorStats::TotalReportedCount;
 
 UE::FStallDetectorStats::FStallDetectorStats(const TCHAR* InName, const double InBudgetSeconds, const EStallDetectorReportingMode InReportingMode)
 	: Name(InName)
