U5 - Add YieldCycles implementation as was recommended by intel

diff --git a/Engine/Source/Runtime/Core/Public/GenericPlatform/GenericPlatformProcess.h b/Engine/Source/Runtime/Core/Public/GenericPlatform/GenericPlatformProcess.h
--- a/Engine/Source/Runtime/Core/Public/GenericPlatform/GenericPlatformProcess.h
+++ b/Engine/Source/Runtime/Core/Public/GenericPlatform/GenericPlatformProcess.h
@@ -708,4 +708,33 @@ struct CORE_API FGenericPlatformProcess
 #	error Unsupported architecture!
 #endif
 	}
+
+	/**
+	* Read the cycle counter of the CPU
+	*/
+	static FORCEINLINE uint64 ReadCycleCounter()
+	{
+#if defined(_MSC_VER)
+		return __rdtsc();
+#elif __has_builtin(__builtin_readcyclecounter)
+		return __builtin_readcyclecounter();
+#else
+#	error Unsupported architecture!
+#endif
+	}
+
+	/**
+	* Tells the processor to pause for at least the amount of cycles given. Is used for spin-loops to improve the speed at 
+	* which the code detects the release of the lock and power-consumption.
+	*/
+	static FORCEINLINE void YieldCycles(uint64 cycles)
+	{
+		uint64 start = ReadCycleCounter();
+		//some 32bit implementations return 0 for __builtin_readcyclecounter just to be on the safe side we protect against this.
+		uint64 end = start != 0 ? start + cycles : 0; 
+		do
+		{
+			Yield();
+		} while(ReadCycleCounter() <= end);
+	}
 };
