Crash early and with a clearer messaging if global shadermap cannot be loaded (UE-107173).

diff --git a/Engine/Source/Runtime/Core/Private/Serialization/MemoryImage.cpp b/Engine/Source/Runtime/Core/Private/Serialization/MemoryImage.cpp
--- a/Engine/Source/Runtime/Core/Private/Serialization/MemoryImage.cpp
+++ b/Engine/Source/Runtime/Core/Private/Serialization/MemoryImage.cpp
@@ -1539,6 +1539,13 @@ FMemoryImageObject FMemoryImageResult::LoadFromArchive(FArchive& Ar, const FType
 
 		FMemoryUnfreezeContent UnfreezeContext(PointerTable, LayoutParameters);
 		UnfreezeContext.UnfreezeObject(FrozenObject, TypeDesc, UnfrozenObject);
+#else
+		FPlatformTypeLayoutParameters CurrentPlatform;
+		CurrentPlatform.InitializeForCurrent();
+
+		UE_LOG(LogMemoryImage, Error, TEXT("Attempting to load an incompatible memory image and conversion isn't supported. Image layout: MaxFieldAlignment=%u, Flags=0x%x, CurrentPlatform: MaxFieldAlignment=%u, Flags=0x%x"),
+			LayoutParameters.MaxFieldAlignment, LayoutParameters.Flags, CurrentPlatform.MaxFieldAlignment, CurrentPlatform.Flags
+			);
 #endif // UE_MEMORYIMAGE_TRACK_TYPE_DEPENDENCIES
 		FMemory::Free(FrozenObject);
 		return FMemoryImageObject(TypeDesc, UnfrozenObject, 0u);
diff --git a/Engine/Source/Runtime/RenderCore/Private/GlobalShader.cpp b/Engine/Source/Runtime/RenderCore/Private/GlobalShader.cpp
--- a/Engine/Source/Runtime/RenderCore/Private/GlobalShader.cpp
+++ b/Engine/Source/Runtime/RenderCore/Private/GlobalShader.cpp
@@ -448,5 +448,9 @@ void FGlobalShaderMap::LoadFromGlobalArchive(FArchive& Ar)
 		{
 			AddSection(Section);
 		}
+		else
+		{
+			UE_LOG(LogShaders, Fatal, TEXT("Could not load section %d (of %d) of the global shadermap."), i, NumSections);
+		}
 	}
 }
