U5 - account for uint overflow

diff --git a/Engine/Source/Runtime/Core/Public/GenericPlatform/GenericPlatformProcess.h b/Engine/Source/Runtime/Core/Public/GenericPlatform/GenericPlatformProcess.h
--- a/Engine/Source/Runtime/Core/Public/GenericPlatform/GenericPlatformProcess.h
+++ b/Engine/Source/Runtime/Core/Public/GenericPlatform/GenericPlatformProcess.h
@@ -731,10 +731,10 @@ struct CORE_API FGenericPlatformProcess
 	{
 		uint64 start = ReadCycleCounter();
 		//some 32bit implementations return 0 for __builtin_readcyclecounter just to be on the safe side we protect against this.
-		uint64 end = start != 0 ? start + cycles : 0; 
+		cycles = start != 0 ? cycles : 0;
 		do
 		{
 			Yield();
-		} while(ReadCycleCounter() <= end);
+		} while((ReadCycleCounter() - start) <= cycles);
 	}
 };
