[TraceLog] Clean shutdown of callstack tracing worker thread

On process exit the worker thread is killed, but we need orderly shutdown to avoid dropped callstacks and deadlocks when allocating during static deinitialization.

#rnx

diff --git a/Engine/Source/Runtime/Core/Private/ProfilingDebugging/Windows/CallstackTrace.cpp b/Engine/Source/Runtime/Core/Private/ProfilingDebugging/Windows/CallstackTrace.cpp
--- a/Engine/Source/Runtime/Core/Private/ProfilingDebugging/Windows/CallstackTrace.cpp
+++ b/Engine/Source/Runtime/Core/Private/ProfilingDebugging/Windows/CallstackTrace.cpp
@@ -147,6 +147,7 @@ public:
 							~FBacktracer();
 	static FBacktracer*		Get();
 	static void				StartWorker();
+	static void				StopWorker();
 	void					AddModule(UPTRINT Base, const TCHAR* Name);
 	void					RemoveModule(UPTRINT Base);
 	void*					GetBacktraceId(void* AddressOfReturnAddress) const;
@@ -183,6 +184,7 @@ private:
 	};
 
 	void					StartWorkerThread();
+	void					StopWorkerThread();
 	const FFunction*		LookupFunction(UPTRINT Address, FLookupState& State) const;
 	static FBacktracer*		Instance;
 	mutable FRWLock			Lock;
@@ -215,7 +217,8 @@ FBacktracer::FBacktracer(FMalloc* InMalloc)
 	// We cannot start worker threads directly on creation. Delay them until the
 	// engine has gotten a little bit further.
 	FCoreDelegates::GetPreMainInitDelegate().AddStatic(FBacktracer::StartWorker);
-
+	// The exit() call will just kill all threads, need controlled shutdown
+	FCoreDelegates::OnExit.AddStatic(FBacktracer::StopWorker);
 }
 
 ////////////////////////////////////////////////////////////////////////////////
@@ -227,11 +230,7 @@ FBacktracer::~FBacktracer()
 		Malloc->Free(Module.Functions);
 	}
 
-	ProcessingThreadRunnable->Stop();
-	if (FCallstackProcWorker::Thread)
-	{
-		FCallstackProcWorker::Thread->WaitForCompletion();
-	} 
+	StopWorkerThread();
 }
 
 ////////////////////////////////////////////////////////////////////////////////
@@ -248,15 +247,38 @@ void FBacktracer::StartWorker()
 	}
 }
 
+////////////////////////////////////////////////////////////////////////////////
+void FBacktracer::StopWorker()
+{
+	if (Instance)
+	{
+		Instance->StopWorkerThread();
+	}
+}
+
 ////////////////////////////////////////////////////////////////////////////////
 void FBacktracer::StartWorkerThread()
 {
-	if (ProcessingThreadRunnable)
+	check(ProcessingThreadRunnable);
+	if (!FCallstackProcWorker::Thread)
 	{
 		FCallstackProcWorker::Thread = FRunnableThread::Create(ProcessingThreadRunnable, TEXT("TraceMemCallstacks"), 0, TPri_BelowNormal);
 	}
 }
 
+////////////////////////////////////////////////////////////////////////////////
+void FBacktracer::StopWorkerThread()
+{
+	check(ProcessingThreadRunnable);
+	if (FCallstackProcWorker::Thread)
+	{
+		FRunnableThread* WorkerThread = FCallstackProcWorker::Thread;
+		FCallstackProcWorker::Thread = nullptr;
+		ProcessingThreadRunnable->Stop();
+		WorkerThread->WaitForCompletion();
+	}
+}
+
 ////////////////////////////////////////////////////////////////////////////////
 void FBacktracer::AddModule(UPTRINT ModuleBase, const TCHAR* Name)
 {
