U5 - expose information if we are running on a workerthread vs just any busy waiting scope

diff --git a/Engine/Source/Runtime/Core/Private/Async/Fundamental/Scheduler.cpp b/Engine/Source/Runtime/Core/Private/Async/Fundamental/Scheduler.cpp
--- a/Engine/Source/Runtime/Core/Private/Async/Fundamental/Scheduler.cpp
+++ b/Engine/Source/Runtime/Core/Private/Async/Fundamental/Scheduler.cpp
@@ -12,7 +12,7 @@ namespace LowLevelTasks
 
 	thread_local FScheduler::FLocalQueueType* FScheduler::LocalQueue = nullptr;
 	thread_local FTask* FScheduler::ActiveTask = nullptr;
-	thread_local bool FScheduler::bIsBackgroundWorker = false;
+	thread_local FScheduler::EWorkerType FScheduler::WorkerType = FScheduler::EWorkerType::None;
 
 	FScheduler FScheduler::Singleton;
 
@@ -116,6 +116,7 @@ namespace LowLevelTasks
 		if (ActiveWorkers.load(std::memory_order_relaxed))
 		{			
 			const bool bIsBackgroundTask = Task.IsBackgroundTask();
+			const bool bIsBackgroundWorker = WorkerType == EWorkerType::Background;
 			if (bIsBackgroundTask && !bIsBackgroundWorker)
 			{
 				QueuePreference = EQueuePreference::GlobalQueuePreference;
@@ -155,6 +156,11 @@ namespace LowLevelTasks
 		return ActiveTask;
 	}
 
+	bool FScheduler::IsWorkerThread() const
+	{
+		return WorkerType != EWorkerType::None;
+	}
+
 	template<FTask* (FScheduler::FLocalQueueType::*DequeueFunction)(bool)>
 	bool FScheduler::TryExecuteTaskFrom(FLocalQueueType* Queue, FQueueRegistry::FOutOfWork& OutOfWork, bool bPermitBackgroundWork)
 	{
@@ -176,8 +182,10 @@ namespace LowLevelTasks
 
 	void FScheduler::WorkerMain(FSleepEvent* WorkerEvent, FLocalQueueType* ExternalWorkerLocalQueue, uint32 WaitCycles, bool bPermitBackgroundWork)
 	{
+		FTaskTagScope WorkerScope(ETaskTag::EWorkerThread);
+
 		FMemory::SetupTLSCachesOnCurrentThread();
-		bIsBackgroundWorker = bPermitBackgroundWork;
+		WorkerType = bPermitBackgroundWork ? EWorkerType::Background : EWorkerType::Foreground;
 
 		checkSlow(LocalQueue == nullptr);
 		if(ExternalWorkerLocalQueue)
@@ -216,7 +224,7 @@ namespace LowLevelTasks
 
 			if (WaitCount < WorkerSpinCycles)
 			{
-				OutOfWork.Start();			
+				OutOfWork.Start();
 				for (uint32 i = 0; i < WaitCycles; i++)
 				{
 					FPlatformProcess::Yield();
@@ -233,14 +241,14 @@ namespace LowLevelTasks
 		FLocalQueueType::DeleteLocalQueue(WorkerLocalQueue, bPermitBackgroundWork, ExternalWorkerLocalQueue != nullptr);
 		LocalQueue = nullptr;
 
-		bIsBackgroundWorker = false;
+		WorkerType = EWorkerType::None;
 		FMemory::ClearAndDisableTLSCachesOnCurrentThread();
 	}
 
 	void FScheduler::BusyWaitInternal(const FConditional& Conditional)
 	{
 		TRACE_CPUPROFILER_EVENT_SCOPE(FScheduler::BusyWaitInternal);
-		FTaskTagScope NoneScope(ETaskTag::EBusyWait);
+		FTaskTagScope WorkerScope(ETaskTag::EWorkerThread);
 
 		checkSlow(LocalQueue != nullptr);
 		check(ActiveWorkers.load(std::memory_order_relaxed));
@@ -276,6 +284,7 @@ namespace LowLevelTasks
 				return;
 			}
 
+			const bool bIsBackgroundWorker = WorkerType == EWorkerType::Background;
 			if (WaitCount < WorkerSpinCycles)
 			{
 				OutOfWork.Start();
diff --git a/Engine/Source/Runtime/Core/Private/HAL/ThreadingBase.cpp b/Engine/Source/Runtime/Core/Private/HAL/ThreadingBase.cpp
--- a/Engine/Source/Runtime/Core/Private/HAL/ThreadingBase.cpp
+++ b/Engine/Source/Runtime/Core/Private/HAL/ThreadingBase.cpp
@@ -92,7 +92,7 @@ FTaskTagScope::FTaskTagScope(bool InTagOnlyIfNone, ETaskTag InTag) : Tag(InTag),
 	}
 	
 	ParentTag = ActiveTaskTag;
-	if (!TagOnlyIfNone || ActiveTaskTag == ETaskTag::ENone || ActiveTaskTag == ETaskTag::EBusyWait)
+	if (!TagOnlyIfNone || ActiveTaskTag == ETaskTag::ENone || ActiveTaskTag == ETaskTag::EWorkerThread)
 	{
 		ActiveTaskTag = Tag;
 	}
@@ -113,7 +113,7 @@ FTaskTagScope::FTaskTagScope(bool InTagOnlyIfNone, ETaskTag InTag) : Tag(InTag),
 FTaskTagScope::~FTaskTagScope()
 {
 	checkf(TagOnlyIfNone || ActiveTaskTag == Tag, TEXT("ActiveTaskTag(%x) corrupted needs to be Tag(%x)"), FTaskTagScope::GetCurrentTag(), Tag);
-	if (!TagOnlyIfNone || ParentTag == ETaskTag::ENone || ParentTag == ETaskTag::EBusyWait)
+	if (!TagOnlyIfNone || ParentTag == ETaskTag::ENone || ParentTag == ETaskTag::EWorkerThread)
 	{
 		ActiveTaskTag = ParentTag;
 	}
diff --git a/Engine/Source/Runtime/Core/Public/Async/Fundamental/Scheduler.h b/Engine/Source/Runtime/Core/Public/Async/Fundamental/Scheduler.h
--- a/Engine/Source/Runtime/Core/Public/Async/Fundamental/Scheduler.h
+++ b/Engine/Source/Runtime/Core/Public/Async/Fundamental/Scheduler.h
@@ -29,7 +29,15 @@ namespace LowLevelTasks
 
 		static thread_local FLocalQueueType* LocalQueue;
 		static thread_local FTask* ActiveTask;
-		static thread_local bool bIsBackgroundWorker;
+
+		enum class EWorkerType
+		{
+			None,
+			Background,
+			Foreground,
+		};
+
+		static thread_local EWorkerType WorkerType;
 		static CORE_API FScheduler Singleton;
 
 		// using 16 bytes here because it fits the vtable and one additional pointer
@@ -78,6 +86,8 @@ namespace LowLevelTasks
 		//get the active task if any
 		CORE_API const FTask* GetActiveTask() const;
 
+		CORE_API bool IsWorkerThread() const;
+
 	private: //Private Interface of the Scheduler	
 		enum class ESleepState
 		{
diff --git a/Engine/Source/Runtime/Core/Public/CoreGlobals.h b/Engine/Source/Runtime/Core/Public/CoreGlobals.h
--- a/Engine/Source/Runtime/Core/Public/CoreGlobals.h
+++ b/Engine/Source/Runtime/Core/Public/CoreGlobals.h
@@ -507,10 +507,10 @@ enum class ETaskTag : int32
 	ERenderingThread			= 1 << 4,
 	ERhiThread					= 1 << 5,
 	EAsyncLoadingThread			= 1 << 6,
-	EBusyWait					= 1 << 7,
 
 	ENamedThreadBits			= (EAsyncLoadingThread << 1) - 1,
-	EParallelThread				= 1 << 8,										//This can be used when multipe threads or jobs are involved (usually a parallel for) It will avoid the check for uniqieness of the named thread tag.
+	EParallelThread				= 1 << 8, //This can be used when multipe threads or jobs are involved (usually a parallel for) It will avoid the check for uniqieness of the named thread tag.
+	EWorkerThread				= 1 << 7 | EParallelThread,
 	EParallelRenderingThread	= ERenderingThread | EParallelThread,
 	EParallelGameThread			= EGameThread | EParallelThread,
 	EParallelRhiThread			= ERhiThread | EParallelThread,
