Fixed
Ensure condition failed: oldValue == newValue launching on Android [ThreadingBase.cpp] [Line: 238]

diff --git a/Engine/Source/Runtime/Core/Core.Build.cs b/Engine/Source/Runtime/Core/Core.Build.cs
--- a/Engine/Source/Runtime/Core/Core.Build.cs
+++ b/Engine/Source/Runtime/Core/Core.Build.cs
@@ -229,7 +229,8 @@ public class Core : ModuleRules
 			PublicDefinitions.Add("WITH_DIRECTXMATH=0");
 		}
 
-		if ((Target.Platform == UnrealTargetPlatform.Mac) || (Target.Platform == UnrealTargetPlatform.IOS) || (Target.Platform == UnrealTargetPlatform.TVOS) || (Target.Platform == UnrealTargetPlatform.HoloLens))
+		if ((Target.Platform == UnrealTargetPlatform.Mac) || (Target.Platform == UnrealTargetPlatform.IOS) || (Target.Platform == UnrealTargetPlatform.TVOS) 
+			|| (Target.Platform == UnrealTargetPlatform.HoloLens) || (Target.Platform == UnrealTargetPlatform.Android))
 		{
 			PublicDefinitions.Add("IS_RUNNING_GAMETHREAD_ON_EXTERNAL_THREAD=1");
 		}
