clarifying comment for VectorSin, so that people don't think it's accurate (like I did).

diff --git a/Engine/Source/Runtime/Core/Public/Math/UnrealMathSSE.h b/Engine/Source/Runtime/Core/Public/Math/UnrealMathSSE.h
--- a/Engine/Source/Runtime/Core/Public/Math/UnrealMathSSE.h
+++ b/Engine/Source/Runtime/Core/Public/Math/UnrealMathSSE.h
@@ -1277,7 +1277,16 @@ FORCEINLINE VectorRegister VectorSin(const VectorRegister& X)
 	//After approx 2.5 million tests comparing to sin(): 
 	//Average error of 0.000128
 	//Max error of 0.001091
-
+	//
+	// Error clarification - the *relative* error rises above 1.2% near
+	// 0 and PI (as the result nears 0). This is enough to introduce 
+	// harmonic distortion when used as an oscillator - VectorSinCos
+	// doesn't cost that much more and is significantly more accurate.
+	// (though don't use either for an oscillator if you care about perf)
+	//
+	// Also - other platforms (NEON) don't vectorize this as of January 2021,
+	// so you'll get accuracy differences between platforms.
+	//
 	VectorRegister y = VectorMultiply(X, GlobalVectorConstants::OneOverTwoPi);
 	y = VectorSubtract(y, VectorFloor(VectorAdd(y, GlobalVectorConstants::FloatOneHalf)));
 	y = VectorMultiply(VectorSinConstantsSSE::A, VectorMultiply(y, VectorSubtract(GlobalVectorConstants::FloatOneHalf, VectorAbs(y))));
