Replaced BuildSettings::IsLicenseeVersion() by FEngineBuildSettings::IsInternalBuild() to check if analytic data and ensures should be sent from  licensee packaged games/apps because:
  - BuildSettings::IsLicenseeVersion() returns true when the engine is built by the licensee (from GIT) but false when the licensee package a game using the officially released engine dll (built by Epic).
  - FEngineBuildSettings::IsInternalBuild() returns false when the engine is built by Epic (in Epic products) and also false in official engine release because it is built from a special licensee-like workspace since 4.25 ish.

diff --git a/Engine/Source/Runtime/Core/Private/Mac/MacPlatformMisc.cpp b/Engine/Source/Runtime/Core/Private/Mac/MacPlatformMisc.cpp
--- a/Engine/Source/Runtime/Core/Private/Mac/MacPlatformMisc.cpp
+++ b/Engine/Source/Runtime/Core/Private/Mac/MacPlatformMisc.cpp
@@ -11,6 +11,7 @@
 #include "VarargsHelper.h"
 #include "Mac/CocoaThread.h"
 #include "Misc/EngineVersion.h"
+#include "Misc/EngineBuildSettings.h"
 #include "Mac/MacMallocZone.h"
 #include "Apple/ApplePlatformSymbolication.h"
 #include "Mac/MacPlatformCrashContext.h"
@@ -2199,7 +2200,7 @@ void FMacCrashContext::GenerateCrashInfoAndLaunchReporter() const
 		GConfig->GetBool(TEXT("/Script/UnrealEd.AnalyticsPrivacySettings"), TEXT("bSendUsageData"), bSendUsageData, GEditorSettingsIni);
 	}
 
-	if (BuildSettings::IsLicenseeVersion() && !UE_EDITOR)
+	if (!UE_EDITOR && !FEngineBuildSettings::IsInternalBuild())
 	{
 		// do not send unattended reports in licensees' builds except for the editor, where it is governed by the above setting
 		bSendUnattendedBugReports = false;
@@ -2304,7 +2305,7 @@ void FMacCrashContext::GenerateEnsureInfoAndLaunchReporter() const
 		GConfig->GetBool(TEXT("/Script/UnrealEd.AnalyticsPrivacySettings"), TEXT("bSendUsageData"), bSendUsageData, GEditorSettingsIni);
 	}
 
-	if (BuildSettings::IsLicenseeVersion() && !UE_EDITOR)
+	if (!UE_EDITOR && !FEngineBuildSettings::IsInternalBuild())
 	{
 		// do not send unattended reports in licensees' builds except for the editor, where it is governed by the above setting
 		bSendUnattendedBugReports = false;
diff --git a/Engine/Source/Runtime/Core/Private/Windows/WindowsPlatformCrashContext.cpp b/Engine/Source/Runtime/Core/Private/Windows/WindowsPlatformCrashContext.cpp
--- a/Engine/Source/Runtime/Core/Private/Windows/WindowsPlatformCrashContext.cpp
+++ b/Engine/Source/Runtime/Core/Private/Windows/WindowsPlatformCrashContext.cpp
@@ -693,7 +693,7 @@ int32 ReportCrashForMonitor(
 	}
 
 #if !UE_EDITOR
-	if (BuildSettings::IsLicenseeVersion())
+	if (!FEngineBuildSettings::IsInternalBuild())
 	{
 		// do not send unattended reports in licensees' builds except for the editor, where it is governed by the above setting
 		bSendUnattendedBugReports = false;
@@ -860,7 +860,7 @@ int32 ReportCrashUsingCrashReportClient(FWindowsPlatformCrashContext& InContext,
 	}
 
 #if !UE_EDITOR
-	if (BuildSettings::IsLicenseeVersion())
+	if (!FEngineBuildSettings::IsInternalBuild())
 	{
 		// do not send unattended reports in licensees' builds except for the editor, where it is governed by the above setting
 		bSendUnattendedBugReports = false;
