SPSC/MPSC concurrent queues, with tests

perf test comparison with existing TCircularQueue and TQueue (for reference):

-------------------------------
TestTCircularQueueSingleThread<5'000'000>
[[REDACTED_PHONE]:244][  0]LogTemp: Display: #0: 0.129468 secs
[[REDACTED_PHONE]:374][  0]LogTemp: Display: #1: 0.129917 secs
[[REDACTED_PHONE]:505][  0]LogTemp: Display: #2: 0.130142 secs
[[REDACTED_PHONE]:635][  0]LogTemp: Display: #3: 0.130430 secs
[[REDACTED_PHONE]:766][  0]LogTemp: Display: #4: 0.130163 secs
[[REDACTED_PHONE]:766][  0]LogTemp: Display: min: 0.129468 secs, avg: 0.130024 secs
-------------------------------

[[REDACTED_PHONE]:766][  0]LogTemp: Display:
-------------------------------
TestQueueSingleThread<5'000'000, TQueueAdapter<TQueue<uint32, EQueueMode::Spsc>>>
[[REDACTED_PHONE]:883][  0]LogTemp: Display: #0: 0.115902 secs
[[REDACTED_PHONE]:000][  0]LogTemp: Display: #1: 0.117038 secs
[[REDACTED_PHONE]:117][  0]LogTemp: Display: #2: 0.116632 secs
[[REDACTED_PHONE]:233][  0]LogTemp: Display: #3: 0.115814 secs
[[REDACTED_PHONE]:349][  0]LogTemp: Display: #4: 0.115581 secs
[[REDACTED_PHONE]:350][  0]LogTemp: Display: min: 0.115581 secs, avg: 0.116193 secs
-------------------------------

[[REDACTED_PHONE]:351][  0]LogTemp: Display:
-------------------------------
TestQueueSingleThread<5'000'000, TQueueAdapter<TQueue<uint32, EQueueMode::Mpsc>>>
[[REDACTED_PHONE]:471][  0]LogTemp: Display: #0: 0.119073 secs
[[REDACTED_PHONE]:591][  0]LogTemp: Display: #1: 0.119660 secs
[[REDACTED_PHONE]:710][  0]LogTemp: Display: #2: 0.119417 secs
[[REDACTED_PHONE]:833][  0]LogTemp: Display: #3: 0.122411 secs
[[REDACTED_PHONE]:956][  0]LogTemp: Display: #4: 0.122508 secs
[[REDACTED_PHONE]:958][  0]LogTemp: Display: min: 0.119073 secs, avg: 0.120614 secs
-------------------------------

[[REDACTED_PHONE]:959][  0]LogTemp: Display:
-------------------------------
TestQueueSingleThread<5'000'000, TMpscQueue<uint32>>
[[REDACTED_PHONE]:072][  0]LogTemp: Display: #0: 0.113381 secs
[[REDACTED_PHONE]:186][  0]LogTemp: Display: #1: 0.113742 secs
[[REDACTED_PHONE]:300][  0]LogTemp: Display: #2: 0.113434 secs
[[REDACTED_PHONE]:414][  0]LogTemp: Display: #3: 0.113194 secs
[[REDACTED_PHONE]:527][  0]LogTemp: Display: #4: 0.112951 secs
[[REDACTED_PHONE]:529][  0]LogTemp: Display: min: 0.112951 secs, avg: 0.113340 secs
-------------------------------

[[REDACTED_PHONE]:530][  0]LogTemp: Display:
-------------------------------
TestQueueSingleThread<5'000'000, TSpscQueue<uint32>>
[[REDACTED_PHONE]:541][  0]LogTemp: Display: #0: 0.011667 secs
[[REDACTED_PHONE]:553][  0]LogTemp: Display: #1: 0.011426 secs
[[REDACTED_PHONE]:565][  0]LogTemp: Display: #2: 0.011434 secs
[[REDACTED_PHONE]:576][  0]LogTemp: Display: #3: 0.011443 secs
[[REDACTED_PHONE]:588][  0]LogTemp: Display: #4: 0.011427 secs
[[REDACTED_PHONE]:589][  0]LogTemp: Display: min: 0.011426 secs, avg: 0.011479 secs
-------------------------------

[[REDACTED_PHONE]:812][  0]LogTemp: Display:
-------------------------------
TestTCircularQueue<5'000'000>
[[REDACTED_PHONE]:241][  0]LogTemp: Display: #0: 0.428893 secs
[[REDACTED_PHONE]:671][  0]LogTemp: Display: #1: 0.429693 secs
[[REDACTED_PHONE]:060][  0]LogTemp: Display: #2: 0.388387 secs
[[REDACTED_PHONE]:490][  0]LogTemp: Display: #3: 0.429410 secs
[[REDACTED_PHONE]:918][  0]LogTemp: Display: #4: 0.427447 secs
[[REDACTED_PHONE]:919][  0]LogTemp: Display: min: 0.388387 secs, avg: 0.420766 secs
-------------------------------

[[REDACTED_PHONE]:920][  0]LogTemp: Display:
-------------------------------
TestSpscQueue<5'000'000, TQueueAdapter<TQueue<uint32, EQueueMode::Spsc>>>
[[REDACTED_PHONE]:169][  0]LogTemp: Display: #0: 0.248199 secs
[[REDACTED_PHONE]:417][  0]LogTemp: Display: #1: 0.248265 secs
[[REDACTED_PHONE]:666][  0]LogTemp: Display: #2: 0.248310 secs
[[REDACTED_PHONE]:923][  0]LogTemp: Display: #3: 0.256853 secs
[[REDACTED_PHONE]:211][  0]LogTemp: Display: #4: 0.288116 secs
[[REDACTED_PHONE]:213][  0]LogTemp: Display: min: 0.248199 secs, avg: 0.257949 secs
-------------------------------

[[REDACTED_PHONE]:214][  0]LogTemp: Display:
-------------------------------
TestSpscQueue<5'000'000, TQueueAdapter<TQueue<uint32, EQueueMode::Mpsc>>>
[[REDACTED_PHONE]:741][  0]LogTemp: Display: #0: 0.526950 secs
[[REDACTED_PHONE]:286][  0]LogTemp: Display: #1: 0.544205 secs
[[REDACTED_PHONE]:811][  0]LogTemp: Display: #2: 0.524776 secs
[[REDACTED_PHONE]:355][  0]LogTemp: Display: #3: 0.543133 secs
[[REDACTED_PHONE]:885][  0]LogTemp: Display: #4: 0.528939 secs
[[REDACTED_PHONE]:886][  0]LogTemp: Display: min: 0.524776 secs, avg: 0.533601 secs
-------------------------------

[[REDACTED_PHONE]:888][  0]LogTemp: Display:
-------------------------------
TestSpscQueue<5'000'000, TMpscQueue<uint32>>
[[REDACTED_PHONE]:100][  0]LogTemp: Display: #0: 0.211752 secs
[[REDACTED_PHONE]:312][  0]LogTemp: Display: #1: 0.211619 secs
[[REDACTED_PHONE]:526][  0]LogTemp: Display: #2: 0.213645 secs
[[REDACTED_PHONE]:740][  0]LogTemp: Display: #3: 0.213254 secs
[[REDACTED_PHONE]:957][  0]LogTemp: Display: #4: 0.216837 secs
[[REDACTED_PHONE]:958][  0]LogTemp: Display: min: 0.211619 secs, avg: 0.213421 secs
-------------------------------

[[REDACTED_PHONE]:959][  0]LogTemp: Display:
-------------------------------
TestSpscQueue<5'000'000, TSpscQueue<uint32>>
[[REDACTED_PHONE]:038][  0]LogTemp: Display: #0: 0.078628 secs
[[REDACTED_PHONE]:119][  0]LogTemp: Display: #1: 0.080822 secs
[[REDACTED_PHONE]:182][  0]LogTemp: Display: #2: 0.062727 secs
[[REDACTED_PHONE]:226][  0]LogTemp: Display: #3: 0.043279 secs
[[REDACTED_PHONE]:307][  0]LogTemp: Display: #4: 0.081072 secs
[[REDACTED_PHONE]:308][  0]LogTemp: Display: min: 0.043279 secs, avg: 0.069306 secs
-------------------------------

[[REDACTED_PHONE]:904][  0]LogTemp: Display:
-------------------------------
TestMpscQueue<1'000'000, TQueueAdapter<TQueue<uint32, EQueueMode::Mpsc>>>
[[REDACTED_PHONE]:653][  0]LogTemp: Display: #0: 0.748177 secs
[[REDACTED_PHONE]:354][  0]LogTemp: Display: #1: 0.700353 secs
[[REDACTED_PHONE]:069][  0]LogTemp: Display: #2: 0.715220 secs
[[REDACTED_PHONE]:826][  0]LogTemp: Display: #3: 0.755602 secs
[[REDACTED_PHONE]:660][  0]LogTemp: Display: #4: 0.832969 secs
[[REDACTED_PHONE]:661][  0]LogTemp: Display: min: 0.700353 secs, avg: 0.750464 secs
-------------------------------

[[REDACTED_PHONE]:662][  0]LogTemp: Display:
-------------------------------
TestMpscQueue<1'000'000, TMpscQueue<uint32>>
[[REDACTED_PHONE]:740][  0]LogTemp: Display: #0: 0.078090 secs
[[REDACTED_PHONE]:818][  0]LogTemp: Display: #1: 0.077926 secs
[[REDACTED_PHONE]:897][  0]LogTemp: Display: #2: 0.077738 secs
[[REDACTED_PHONE]:980][  0]LogTemp: Display: #3: 0.082891 secs
[[REDACTED_PHONE]:063][  0]LogTemp: Display: #4: 0.082600 secs
[[REDACTED_PHONE]:064][  0]LogTemp: Display: min: 0.077738 secs, avg: 0.079849 secs
-------------------------------

diff --git a/Engine/Source/Runtime/Core/Private/Tests/Containers/ConcurrentQueuesTest.cpp b/Engine/Source/Runtime/Core/Private/Tests/Containers/ConcurrentQueuesTest.cpp
new file mode 100644
--- /dev/null
+++ b/Engine/Source/Runtime/Core/Private/Tests/Containers/ConcurrentQueuesTest.cpp
@@ -0,0 +1,429 @@
+// Copyright Epic Games, Inc. All Rights Reserved.
+
+#include "CoreTypes.h"
+#include "Misc/AutomationTest.h"
+#include "Async/TaskGraphInterfaces.h"
+#include "Containers/CircularQueue.h"
+#include "Containers/Queue.h"
+#include "Containers/SpscQueue.h"
+#include "Containers/MpscQueue.h"
+#include "Tests/Benchmark.h"
+
+#include <atomic>
+
+#if WITH_DEV_AUTOMATION_TESTS
+
+namespace ConcurrentQueuesTests
+{
+	template<typename QueueType>
+	class TQueueAdapter : public QueueType
+	{
+	public:
+		template<typename... ArgTypes>
+		TQueueAdapter(ArgTypes... Args)
+			: QueueType(Forward<ArgTypes>(Args)...)
+		{
+		}
+
+		TOptional<typename QueueType::FElementType> Dequeue()
+		{
+			typename QueueType::FElementType Value;
+			if (QueueType::Dequeue(Value))
+			{
+				return Value;
+			}
+			else
+			{
+				return {};
+			}
+		}
+	};
+
+	const uint32 CircularQueueSize = 1024;
+
+	// measures performance of a queue when the producer and the consumer are run in the same thread
+	template<uint32 Num, typename QueueType>
+	void TestSpscQueueSingleThreadImp(QueueType& Queue)
+	{
+		const uint32 BatchSize = CircularQueueSize;
+		const uint32 BatchNum = Num / BatchSize;
+		for (uint32 BatchIndex = 0; BatchIndex != BatchNum; ++BatchIndex)
+		{
+			for (uint32 i = 0; i != BatchSize; ++i)
+			{
+				Queue.Enqueue(i);
+			}
+
+			for (uint32 i = 0; i != BatchSize; ++i)
+			{
+				TOptional<uint32> Consumed{ Queue.Dequeue() };
+				checkSlow(Consumed.IsSet() && Consumed.GetValue() == i);
+			}
+		}
+	}
+
+	template<uint32 Num>
+	void TestTCircularQueueSingleThread()
+	{
+		TQueueAdapter<TCircularQueue<uint32>> Queue{ CircularQueueSize + 1 };
+		TestSpscQueueSingleThreadImp<Num>(Queue);
+	}
+
+	template<uint32 Num, typename QueueType>
+	void TestQueueSingleThread()
+	{
+		QueueType Queue;
+		TestSpscQueueSingleThreadImp<Num>(Queue);
+	}
+
+	template<uint32 Num, typename QueueType>
+	void TestSpscQueue_Impl(QueueType& Queue)
+	{
+		std::atomic<bool> bStop{ false };
+
+		FGraphEventRef Producer = FFunctionGraphTask::CreateAndDispatchWhenReady(
+			[&bStop, &Queue]
+			{
+				while (!bStop)
+				{
+					Queue.Enqueue(0);
+				}
+			}
+			);
+
+		// consumer
+		uint32 It = 0;
+		while (It != Num)
+		{
+			if (Queue.Dequeue().IsSet())
+			{
+				++It;
+			}
+		}
+
+		bStop = true;
+
+		Producer->Wait(ENamedThreads::GameThread);
+	}
+
+	template<uint32 Num>
+	void TestTCircularQueue()
+	{
+		TQueueAdapter<TCircularQueue<uint32>> Queue{ CircularQueueSize + 1 };
+		TestSpscQueue_Impl<Num>(Queue);
+	}
+
+	template<uint32 Num, typename QueueType>
+	void TestSpscQueue()
+	{
+		QueueType Queue;
+		TestSpscQueue_Impl<Num>(Queue);
+	}
+
+	template<uint32 Num, typename QueueType>
+	void TestSpscQueueCorrectness()
+	{
+		QueueType Queue;
+		int32 NumProduced = 0;
+
+		// producer
+		FGraphEventRef Task = FFunctionGraphTask::CreateAndDispatchWhenReady(
+			[&Queue, &NumProduced]
+			{
+				do
+				{
+					Queue.Enqueue(0);
+				} while (++NumProduced != Num);
+			}
+		);
+
+		// consumer
+		uint32 NumConcumed = 0;
+		while (NumConcumed != Num)
+		{
+			if (Queue.Dequeue().IsSet())
+			{
+				++NumConcumed;
+			}
+		}
+
+		while (Queue.Dequeue().IsSet())
+		{
+			++NumConcumed;
+		}
+
+		Task->Wait();
+
+		check(NumProduced == NumConcumed);
+	}
+
+	template<uint32 Num, typename QueueType>
+	void TestMpscQueue()
+	{
+		QueueType Queue;
+		std::atomic<bool> bStop{ false };
+
+		int32 NumProducers = FTaskGraphInterface::Get().GetNumWorkerThreads();
+		FGraphEventArray Producers;
+		for (int32 i = 0; i != NumProducers; ++i)
+		{
+			Producers.Add(FFunctionGraphTask::CreateAndDispatchWhenReady(
+				[&bStop, &Queue]
+				{
+					while (!bStop)
+					{
+						Queue.Enqueue(0);
+					}
+				}
+				));
+		}
+
+		uint32 It = 0;
+		while (It != Num)
+		{
+			if (Queue.Dequeue().IsSet())
+			{
+				++It;
+			}
+		}
+
+		bStop = true;
+
+		FTaskGraphInterface::Get().WaitUntilTasksComplete(MoveTemp(Producers), ENamedThreads::GameThread);
+	}
+
+	template<uint32 Num, typename QueueType>
+	void TestMpscQueueCorrectness()
+	{
+		struct alignas(PLATFORM_CACHE_LINE_SIZE) FCounter
+		{
+			uint32 Count = 0;
+		};
+
+		QueueType Queue;
+
+		const int32 NumProducers = FTaskGraphInterface::Get().GetNumWorkerThreads();
+		const uint32 NumPerProducer = Num / NumProducers;
+		TArray<FCounter> NumsProduced;
+		NumsProduced.AddDefaulted(NumProducers);
+		for (int32 i = 0; i != NumProducers; ++i)
+		{
+			FFunctionGraphTask::CreateAndDispatchWhenReady(
+				[&Queue, &NumsProduced, i, NumPerProducer]()
+			{
+				do
+				{
+					Queue.Enqueue(0);
+				} while (++NumsProduced[i].Count < NumPerProducer);
+			}
+			);
+		}
+
+		// consumer
+		uint32 NumConcumed = 0;
+		while (NumConcumed != NumPerProducer * NumProducers)
+		{
+			if (Queue.Dequeue().IsSet())
+			{
+				++NumConcumed;
+			}
+		}
+
+		while (Queue.Dequeue().IsSet())
+		{
+			++NumConcumed;
+		}
+
+		uint32 Produced = 0;
+		for (int i = 0; i != NumProducers; ++i)
+		{
+			Produced += NumsProduced[i].Count;
+		}
+		check(Produced == NumConcumed);
+	}
+
+	IMPLEMENT_SIMPLE_AUTOMATION_TEST(FConcurrentQueuesTest, "System.Core.Async.ConcurrentQueuesTest", EAutomationTestFlags::ApplicationContextMask | EAutomationTestFlags::EngineFilter);
+
+	bool FConcurrentQueuesTest::RunTest(const FString& Parameters)
+	{
+		{	// test for support of not default constructible types
+			struct FNonDefaultConstructable
+			{
+				int Value;
+
+				explicit FNonDefaultConstructable(int InValue)
+					: Value(InValue)
+				{
+					UE_LOG(LogTemp, Display, TEXT("ctor"));
+				}
+
+				~FNonDefaultConstructable()
+				{
+					UE_LOG(LogTemp, Display, TEXT("dctor"));
+				}
+
+				FNonDefaultConstructable(const FNonDefaultConstructable&) = delete;
+				FNonDefaultConstructable& operator=(const FNonDefaultConstructable&) = delete;
+
+				FNonDefaultConstructable(FNonDefaultConstructable&& Other)
+					: Value(Other.Value)
+				{
+					UE_LOG(LogTemp, Display, TEXT("move-ctor"));
+				}
+
+				FNonDefaultConstructable& operator=(FNonDefaultConstructable&& Other)
+				{
+					Value = Other.Value;
+					UE_LOG(LogTemp, Display, TEXT("move="));
+					return *this;
+				}
+			};
+
+			{
+				TSpscQueue<FNonDefaultConstructable> Q;
+				Q.Enqueue(1);
+				TOptional<FNonDefaultConstructable> Res{ Q.Dequeue() };
+				verify(Res.IsSet() && Res.GetValue().Value == 1);
+			}
+
+			{
+				TMpscQueue<FNonDefaultConstructable> Q;
+				Q.Enqueue(1);
+				TOptional<FNonDefaultConstructable> Res{ Q.Dequeue() };
+				verify(Res.IsSet() && Res.GetValue().Value == 1);
+			}
+		}
+
+		{	// test queue destruction and not default movable types
+			struct FNonTrivial
+			{
+				int* Value;
+
+				explicit FNonTrivial(TUniquePtr<int> InValue)
+					: Value(InValue.Release())
+				{}
+
+				~FNonTrivial()
+				{
+					// check for double delete
+					check(Value != (int*)0xcdcdcdcdcdcdcdcd);
+					delete Value;
+					Value = (int*)0xcdcdcdcdcdcdcdcd;
+				}
+
+				FNonTrivial(const FNonTrivial&) = delete;
+				FNonTrivial& operator=(const FNonTrivial&) = delete;
+
+				FNonTrivial(FNonTrivial&& Other)
+				{
+					Value = Other.Value;
+					Other.Value = nullptr;
+				}
+
+				FNonTrivial& operator=(FNonTrivial&& Other)
+				{
+					Swap(Value, Other.Value);
+					return *this;
+				}
+			};
+
+			// SPSC
+
+			{	// destroy queue while it's holding one unconsumed item
+				TSpscQueue<FNonTrivial> Q;
+				Q.Enqueue(MakeUnique<int>(1));
+			}
+
+			{	// destroy queue while it's holding one cached consumed time
+				TSpscQueue<FNonTrivial> Q;
+				Q.Enqueue(MakeUnique<int>(1));
+				TOptional<FNonTrivial> Res{ Q.Dequeue() };
+				verify(Res.IsSet() && *Res.GetValue().Value == 1);
+			}
+
+			{	// destroy queue while it's holding one chached consumed item and one unconsumed item
+				TSpscQueue<FNonTrivial> Q;
+				Q.Enqueue(MakeUnique<int>(1));
+				Q.Enqueue(MakeUnique<int>(2));
+				TOptional<FNonTrivial> Res{ Q.Dequeue() };
+				verify(Res.IsSet() && *Res.GetValue().Value == 1);
+			}
+
+			// MPSC
+			{	// destroy untouched queue
+				TMpscQueue<FNonTrivial> Q;
+			}
+
+			{	// destroy never consumed queue with one unconsumed item
+				TMpscQueue<FNonTrivial> Q;
+				Q.Enqueue(MakeUnique<int>(1));
+			}
+
+			{	// destroy empty queue
+				TMpscQueue<FNonTrivial> Q;
+				Q.Enqueue(MakeUnique<int>(1));
+				TOptional<FNonTrivial> Res{ Q.Dequeue() };
+				verify(Res.IsSet() && *Res.GetValue().Value == 1);
+			}
+
+			{	// destroy queue with one unconsumed item
+				TMpscQueue<FNonTrivial> Q;
+				Q.Enqueue(MakeUnique<int>(1));
+				Q.Enqueue(MakeUnique<int>(2));
+				TOptional<FNonTrivial> Res{ Q.Dequeue() };
+				verify(Res.IsSet() && *Res.GetValue().Value == 1);
+			}
+
+			{	// destroy queue with two items
+				TMpscQueue<FNonTrivial> Q;
+				Q.Enqueue(MakeUnique<int>(1));
+				Q.Enqueue(MakeUnique<int>(2));
+			}
+
+			{	// enqueue and dequeue multiple items
+				TMpscQueue<FNonTrivial> Q;
+				Q.Enqueue(MakeUnique<int>(1));
+				Q.Enqueue(MakeUnique<int>(2));
+				TOptional<FNonTrivial> Res{ Q.Dequeue() };
+				verify(Res.IsSet() && *Res.GetValue().Value == 1);
+				Res = Q.Dequeue();
+				verify(Res.IsSet() && *Res.GetValue().Value == 2);
+			}
+
+			{	// enqueue and dequeue (interleaved) multiple items
+				TMpscQueue<FNonTrivial> Q;
+				Q.Enqueue(MakeUnique<int>(1));
+				TOptional<FNonTrivial> Res{ Q.Dequeue() };
+				verify(Res.IsSet() && *Res.GetValue().Value == 1);
+				Q.Enqueue(MakeUnique<int>(2));
+				Res = Q.Dequeue();
+				verify(Res.IsSet() && *Res.GetValue().Value == 2);
+			}
+		}
+
+		UE_BENCHMARK(5, TestTCircularQueueSingleThread<5'000'000>);
+		UE_BENCHMARK(5, TestQueueSingleThread<5'000'000, TQueueAdapter<TQueue<uint32, EQueueMode::Spsc>>>);
+		UE_BENCHMARK(5, TestQueueSingleThread<5'000'000, TQueueAdapter<TQueue<uint32, EQueueMode::Mpsc>>>);
+		UE_BENCHMARK(5, TestQueueSingleThread<5'000'000, TMpscQueue<uint32>>);
+		UE_BENCHMARK(5, TestQueueSingleThread<5'000'000, TSpscQueue<uint32>>);
+
+		UE_BENCHMARK(5, TestSpscQueueCorrectness<5'000'000, TMpscQueue<uint32>>);
+		UE_BENCHMARK(5, TestSpscQueueCorrectness<5'000'000, TSpscQueue<uint32>>);
+
+		UE_BENCHMARK(5, TestTCircularQueue<5'000'000>);
+		UE_BENCHMARK(5, TestSpscQueue<5'000'000, TQueueAdapter<TQueue<uint32, EQueueMode::Spsc>>>);
+		UE_BENCHMARK(5, TestSpscQueue<5'000'000, TQueueAdapter<TQueue<uint32, EQueueMode::Mpsc>>>);
+		UE_BENCHMARK(5, TestSpscQueue<5'000'000, TMpscQueue<uint32>>);
+		UE_BENCHMARK(5, TestSpscQueue<5'000'000, TSpscQueue<uint32>>);
+
+		UE_BENCHMARK(5, TestMpscQueueCorrectness<5'000'000, TQueueAdapter<TQueue<uint32, EQueueMode::Mpsc>>>);
+		UE_BENCHMARK(5, TestMpscQueueCorrectness<5'000'000, TMpscQueue<uint32>>);
+
+		UE_BENCHMARK(5, TestMpscQueue<1'000'000, TQueueAdapter<TQueue<uint32, EQueueMode::Mpsc>>>);
+		UE_BENCHMARK(5, TestMpscQueue<1'000'000, TMpscQueue<uint32>>);
+
+		return true;
+	}
+}
+
+#endif // WITH_DEV_AUTOMATION_TESTS
diff --git a/Engine/Source/Runtime/Core/Public/Containers/CircularQueue.h b/Engine/Source/Runtime/Core/Public/Containers/CircularQueue.h
--- a/Engine/Source/Runtime/Core/Public/Containers/CircularQueue.h
+++ b/Engine/Source/Runtime/Core/Public/Containers/CircularQueue.h
@@ -18,11 +18,12 @@
  * the implications for all of our target platforms need further analysis, so
  * we're using the simpler sequentially consistent model for now.
  *
- * @param ElementType The type of elements held in the queue.
+ * @param T The type of elements held in the queue.
  */
-template<typename ElementType> class TCircularQueue
+template<typename T> class TCircularQueue
 {
 public:
+	using FElementType = T;
 
 	/**
 	 * Constructor.
@@ -64,7 +65,7 @@ public:
 	 * @return true if an element has been returned, false if the queue was empty.
 	 * @note To be called only from consumer thread.
 	 */
-	bool Dequeue(ElementType& OutElement)
+	bool Dequeue(FElementType& OutElement)
 	{
 		const uint32 CurrentHead = Head.Load();
 
@@ -117,7 +118,7 @@ public:
 	 * @return true if the item was added, false if the queue was full.
 	 * @note To be called only from producer thread.
 	 */
-	bool Enqueue(const ElementType& Element)
+	bool Enqueue(const FElementType& Element)
 	{
 		const uint32 CurrentTail = Tail.Load();
 		uint32 NewTail = Buffer.GetNextIndex(CurrentTail);
@@ -140,7 +141,7 @@ public:
 	 * @return true if the item was added, false if the queue was full.
 	 * @note To be called only from producer thread.
 	 */
-	bool Enqueue(ElementType&& Element)
+	bool Enqueue(FElementType&& Element)
 	{
 		const uint32 CurrentTail = Tail.Load();
 		uint32 NewTail = Buffer.GetNextIndex(CurrentTail);
@@ -191,7 +192,7 @@ public:
 	 * @return true if an item has been returned, false if the queue was empty.
 	 * @note To be called only from consumer thread.
 	 */
-	bool Peek(ElementType& OutItem) const
+	bool Peek(FElementType& OutItem) const
 	{
 		const uint32 CurrentHead = Head.Load();
 
@@ -208,11 +209,11 @@ public:
 	/**
 	 * Returns the oldest item in the queue without removing it.
 	 *
-	 * @return an ElementType pointer if an item has been returned, nullptr if the queue was empty.
+	 * @return an FElementType pointer if an item has been returned, nullptr if the queue was empty.
 	 * @note To be called only from consumer thread.
 	 * @note The return value is only valid until Dequeue, Empty, or the destructor has been called.
 	 */
-	const ElementType* Peek() const
+	const FElementType* Peek() const
 	{
 		const uint32 CurrentHead = Head.Load();
 
@@ -227,7 +228,7 @@ public:
 private:
 
 	/** Holds the buffer. */
-	TCircularBuffer<ElementType> Buffer;
+	TCircularBuffer<FElementType> Buffer;
 
 	/** Holds the index to the first item in the buffer. */
 	TAtomic<uint32> Head;
diff --git a/Engine/Source/Runtime/Core/Public/Containers/MpscQueue.h b/Engine/Source/Runtime/Core/Public/Containers/MpscQueue.h
new file mode 100644
--- /dev/null
+++ b/Engine/Source/Runtime/Core/Public/Containers/MpscQueue.h
@@ -0,0 +1,84 @@
+#pragma once
+
+#include "CoreTypes.h"
+#include "Templates/TypeCompatibleBytes.h"
+#include "Templates/UnrealTemplate.h"
+#include "Misc/Optional.h"
+#include <atomic>
+
+/** 
+ * Fast multi-producer/single-consumer unbounded concurrent queue.
+ * Based on http://www.1024cores.net/home/lock-free-algorithms/queues/non-intrusive-mpsc-node-based-queue
+ */
+template<typename T>
+class TMpscQueue final
+{
+public:
+	using ElementType = T;
+
+	UE_NONCOPYABLE(TMpscQueue);
+
+	TMpscQueue()
+	{
+		FNode* Sentinel = new FNode;
+		Head.store(Sentinel, std::memory_order_relaxed);
+		Tail = Sentinel;
+	}
+
+	~TMpscQueue()
+	{
+		FNode* Next = Tail->Next.load(std::memory_order_relaxed);
+
+		// sentinel's value is already destroyed
+		delete Tail;
+
+		while (Next != nullptr)
+		{
+			Tail = Next;
+			Next = Tail->Next.load(std::memory_order_relaxed);
+
+			DestructItem((ElementType*)&Tail->Value);
+			delete Tail;
+		}
+	}
+
+	template <typename... ArgTypes>
+	void Enqueue(ArgTypes&&... Args)
+	{
+		FNode* New = new FNode;
+		new (&New->Value) ElementType(Forward<ArgTypes>(Args)...);
+
+		FNode* Prev = Head.exchange(New, std::memory_order_acq_rel);
+		Prev->Next.store(New, std::memory_order_release);
+	}
+
+	TOptional<ElementType> Dequeue()
+	{
+		FNode* Next = Tail->Next.load(std::memory_order_acquire);
+
+		if (Next == nullptr)
+		{
+			return {};
+		}
+
+		ElementType* ValuePtr = (ElementType*)&Next->Value;
+		TOptional<ElementType> Res{ MoveTemp(*ValuePtr) };
+		DestructItem(ValuePtr);
+
+		delete Tail; // current sentinel
+
+		Tail = Next; // new sentinel
+		return Res;
+	}
+
+private:
+	struct FNode
+	{
+		std::atomic<FNode*> Next{ nullptr };
+		TTypeCompatibleBytes<ElementType> Value;
+	};
+
+private:
+	std::atomic<FNode*> Head; // accessed only by producers
+	alignas(PLATFORM_CACHE_LINE_SIZE) FNode* Tail; // accessed only by consumer, hence should be on a different cache line than `Head`
+};
diff --git a/Engine/Source/Runtime/Core/Public/Containers/Queue.h b/Engine/Source/Runtime/Core/Public/Containers/Queue.h
--- a/Engine/Source/Runtime/Core/Public/Containers/Queue.h
+++ b/Engine/Source/Runtime/Core/Public/Containers/Queue.h
@@ -31,14 +31,15 @@ enum class EQueueMode
  * writing it in a way that does not depend on possible instruction reordering on the CPU.
  * The Enqueue() method uses an atomic compare-and-swap in multiple-producers scenarios.
  *
- * @param ItemType The type of items stored in the queue.
+ * @param T The type of items stored in the queue.
  * @param Mode The queue mode (single-producer, single-consumer by default).
  * @todo gmp: Implement node pooling.
  */
-template<typename ItemType, EQueueMode Mode = EQueueMode::Spsc>
+template<typename T, EQueueMode Mode = EQueueMode::Spsc>
 class TQueue
 {
 public:
+	using FElementType = T;
 
 	/** Default constructor. */
 	TQueue()
@@ -66,7 +67,7 @@ public:
 	 * @note To be called only from consumer thread.
 	 * @see Empty, Enqueue, IsEmpty, Peek, Pop
 	 */
-	bool Dequeue(ItemType& OutItem)
+	bool Dequeue(FElementType& OutItem)
 	{
 		TNode* Popped = Tail->NextNode;
 
@@ -80,7 +81,7 @@ public:
 
 		TNode* OldTail = Tail;
 		Tail = Popped;
-		Tail->Item = ItemType();
+		Tail->Item = FElementType();
 		delete OldTail;
 
 		return true;
@@ -105,7 +106,7 @@ public:
 	 * @note To be called only from producer thread(s).
 	 * @see Dequeue, Pop
 	 */
-	bool Enqueue(const ItemType& Item)
+	bool Enqueue(const FElementType& Item)
 	{
 		TNode* NewNode = new TNode(Item);
 
@@ -142,7 +143,7 @@ public:
 	 * @note To be called only from producer thread(s).
 	 * @see Dequeue, Pop
 	 */
-	bool Enqueue(ItemType&& Item)
+	bool Enqueue(FElementType&& Item)
 	{
 		TNode* NewNode = new TNode(MoveTemp(Item));
 
@@ -191,7 +192,7 @@ public:
 	 * @note To be called only from consumer thread.
 	 * @see Dequeue, Empty, IsEmpty, Pop
 	 */
-	bool Peek(ItemType& OutItem) const
+	bool Peek(FElementType& OutItem) const
 	{
 		if (Tail->NextNode == nullptr)
 		{
@@ -211,7 +212,7 @@ public:
 	 *
 	 * @return Pointer to the item, or nullptr if queue is empty
 	 */
-	ItemType* Peek()
+	FElementType* Peek()
 	{
 		if (Tail->NextNode == nullptr)
 		{
@@ -221,7 +222,7 @@ public:
 		return &Tail->NextNode->Item;
 	}
 
-	FORCEINLINE const ItemType* Peek() const
+	FORCEINLINE const FElementType* Peek() const
 	{
 		return const_cast<TQueue*>(this)->Peek();
 	}
@@ -246,7 +247,7 @@ public:
 
 		TNode* OldTail = Tail;
 		Tail = Popped;
-		Tail->Item = ItemType();
+		Tail->Item = FElementType();
 		delete OldTail;
 
 		return true;
@@ -261,7 +262,7 @@ private:
 		TNode* volatile NextNode;
 
 		/** Holds the node's item. */
-		ItemType Item;
+		FElementType Item;
 
 		/** Default constructor. */
 		TNode()
@@ -269,13 +270,13 @@ private:
 		{ }
 
 		/** Creates and initializes a new node. */
-		explicit TNode(const ItemType& InItem)
+		explicit TNode(const FElementType& InItem)
 			: NextNode(nullptr)
 			, Item(InItem)
 		{ }
 
 		/** Creates and initializes a new node. */
-		explicit TNode(ItemType&& InItem)
+		explicit TNode(FElementType&& InItem)
 			: NextNode(nullptr)
 			, Item(MoveTemp(InItem))
 		{ }
diff --git a/Engine/Source/Runtime/Core/Public/Containers/SpscQueue.h b/Engine/Source/Runtime/Core/Public/Containers/SpscQueue.h
new file mode 100644
--- /dev/null
+++ b/Engine/Source/Runtime/Core/Public/Containers/SpscQueue.h
@@ -0,0 +1,126 @@
+#pragma once
+
+#include "CoreTypes.h"
+#include "Templates/TypeCompatibleBytes.h"
+#include "Templates/UnrealTemplate.h"
+#include "Misc/Optional.h"
+#include <atomic>
+
+/**
+ * Fast single-producer/single-consumer unbounded concurrent queue. Doesn't free memory until destruction but recycles consumed items.
+ * Based on http://www.1024cores.net/home/lock-free-algorithms/queues/unbounded-spsc-queue
+ */
+template<typename T>
+class TSpscQueue final
+{
+public:
+	using ElementType = T;
+
+	UE_NONCOPYABLE(TSpscQueue);
+
+	TSpscQueue()
+	{
+		FNode* Node = new FNode;
+		Tail.store(Node, std::memory_order_relaxed);
+		Head = First = TailCopy = Node;
+	}
+
+	~TSpscQueue()
+	{
+		FNode* Node = First;
+		FNode* LocalTail = Tail.load(std::memory_order_relaxed);
+
+		// Delete all nodes which are the sentinel or unoccupied
+		bool bContinue = false;
+		do
+		{
+			FNode* Next = Node->Next.load(std::memory_order_relaxed);
+			delete Node;
+			bContinue = Node != LocalTail;
+			Node = Next;
+		} while (bContinue);
+
+		// Delete all nodes which are occupied, destroying the element first
+		while (Node != nullptr)
+		{
+			FNode* Next = Node->Next.load(std::memory_order_relaxed);
+			DestructItem((ElementType*)&Node->Value);
+			delete Node;
+			Node = Next;
+		}
+	}
+
+	template <typename... ArgTypes>
+	void Enqueue(ArgTypes&&... Args)
+	{
+		FNode* Node = AllocNode();
+		new (&Node->Value) ElementType(Forward<ArgTypes>(Args)...);
+
+		Head->Next.store(Node, std::memory_order_release);
+		Head = Node;
+	}
+
+	// returns empty TOptional if queue is empty 
+	TOptional<ElementType> Dequeue()
+	{
+		FNode* LocalTail = Tail.load(std::memory_order_relaxed);
+		FNode* LocalTailNext = LocalTail->Next.load(std::memory_order_acquire);
+		if (LocalTailNext == nullptr)
+		{
+			return {};
+		}
+
+		ElementType* TailNextValue = (ElementType*)&LocalTailNext->Value;
+		TOptional<ElementType> Value{ MoveTemp(*TailNextValue) };
+		DestructItem(TailNextValue);
+
+		Tail.store(LocalTailNext, std::memory_order_release);
+		return Value;
+	}
+
+private:
+	struct FNode
+	{
+		std::atomic<FNode*> Next{ nullptr };
+		TTypeCompatibleBytes<ElementType> Value;
+	};
+
+private:
+	FNode* AllocNode()
+	{
+		// first tries to allocate node from internal node cache, 
+		// if attempt fails, allocates node via ::operator new() 
+
+		auto AllocFromCache = [this]()
+		{
+			FNode* Node = First;
+			First = First->Next;
+			Node->Next.store(nullptr, std::memory_order_relaxed);
+			return Node;
+		};
+
+		if (First != TailCopy)
+		{
+			return AllocFromCache();
+		}
+
+		TailCopy = Tail.load(std::memory_order_acquire);
+		if (First != TailCopy)
+		{
+			return AllocFromCache();
+		}
+
+		return new FNode();
+	}
+
+private:
+	// consumer part 
+	// accessed mainly by consumer, infrequently by producer 
+	alignas(PLATFORM_CACHE_LINE_SIZE) std::atomic<FNode*> Tail; // tail of the queue 
+
+	// producer part 
+	// accessed only by producer 
+	alignas(PLATFORM_CACHE_LINE_SIZE) FNode* Head; // head of the queue 
+	FNode* First; // last unused node (tail of node cache) 
+	FNode* TailCopy; // helper (points somewhere between First and Tail)
+};
\ No newline at end of file
diff --git a/Engine/Source/Runtime/Core/Public/Tests/Benchmark.h b/Engine/Source/Runtime/Core/Public/Tests/Benchmark.h
new file mode 100644
--- /dev/null
+++ b/Engine/Source/Runtime/Core/Public/Tests/Benchmark.h
@@ -0,0 +1,36 @@
+#pragma once
+
+#include "CoreTypes.h"
+#include "Logging/LogMacros.h"
+#include "HAL/PlatformTime.h"
+#include "Math/NumericLimits.h"
+
+template<uint32 NumRuns, typename TestT>
+void Benchmark(const TCHAR* TestName, TestT&& TestBody)
+{
+	UE_LOG(LogTemp, Display, TEXT("\n-------------------------------\n%s"), TestName);
+	double MinTime = TNumericLimits<double>::Max();
+	double TotalTime = 0;
+	for (uint32 RunNo = 0; RunNo != NumRuns; ++RunNo)
+	{
+		double Time = FPlatformTime::Seconds();
+		TestBody();
+		Time = FPlatformTime::Seconds() - Time;
+
+		UE_LOG(LogTemp, Display, TEXT("#%d: %f secs"), RunNo, Time);
+
+		TotalTime += Time;
+		if (MinTime > Time)
+		{
+			MinTime = Time;
+		}
+	}
+	UE_LOG(LogTemp, Display, TEXT("min: %f secs, avg: %f secs\n-------------------------------\n"), MinTime, TotalTime / NumRuns);
+
+#if NO_LOGGING
+	printf("min: %f secs, avg: %f secs\n-------------------------------\n\n", MinTime, TotalTime / NumRuns);
+#endif
+}
+
+#define UE_BENCHMARK(NumRuns, ...) Benchmark<NumRuns>(TEXT(#__VA_ARGS__), __VA_ARGS__)
+
