Add an object version for the //UE5/Release-* streams.

#rnx

diff --git a/Engine/Source/Runtime/Core/Private/UObject/DevObjectVersion.cpp b/Engine/Source/Runtime/Core/Private/UObject/DevObjectVersion.cpp
--- a/Engine/Source/Runtime/Core/Private/UObject/DevObjectVersion.cpp
+++ b/Engine/Source/Runtime/Core/Private/UObject/DevObjectVersion.cpp
@@ -203,6 +203,12 @@ const FGuid FUE5MainStreamObjectVersion::GUID(0x697DD581, 0xE64f41AB, 0xAA4A51EC
 // Register UE5 main stream custom version with Core
 FDevVersionRegistration GRegisterUE5MainStreamObjectVersion(FUE5MainStreamObjectVersion::GUID, FUE5MainStreamObjectVersion::LatestVersion, TEXT("UE5-Main"));
 
+// Unique UE5 release version id
+const FGuid FUE5ReleaseStreamObjectVersion::GUID(0xD89B5E42, 0x24BD4D46, 0x8412ACA8, 0xDF641779);
+// Register UE5 release stream custom version with Core
+FDevVersionRegistration GRegisterUE5ReleaseStreamObjectVersion(FUE5ReleaseStreamObjectVersion::GUID, FUE5ReleaseStreamObjectVersion::LatestVersion, TEXT("UE5-Release"));
+
+
 // Unique Cooker version id
 const FGuid FUE5CookerObjectVersion::GUID(0x26075A32, 0x730F4708, 0x88E98C32, 0xF1599D05);
 // Register Cooker custom version with Core
diff --git a/Engine/Source/Runtime/Core/Public/UObject/UE5ReleaseStreamObjectVersion.h b/Engine/Source/Runtime/Core/Public/UObject/UE5ReleaseStreamObjectVersion.h
new file mode 100644
--- /dev/null
+++ b/Engine/Source/Runtime/Core/Public/UObject/UE5ReleaseStreamObjectVersion.h
@@ -0,0 +1,24 @@
+// Copyright Epic Games, Inc. All Rights Reserved.
+#pragma once
+
+#include "CoreTypes.h"
+#include "Misc/Guid.h"
+
+// Custom serialization version for changes made in //UE5/Release-* stream
+struct CORE_API FUE5ReleaseStreamObjectVersion
+{
+	enum Type
+	{
+		// Before any version changes were made
+		BeforeCustomVersionWasAdded = 0,
+
+		// -----<new versions can be added above this line>-------------------------------------------------
+		VersionPlusOne,
+		LatestVersion = VersionPlusOne - 1
+	};
+
+	// The GUID for this custom version number
+	const static FGuid GUID;
+
+	FUE5ReleaseStreamObjectVersion() = delete;
+};
