SharedBuffer: Made default buffers owned/materialized and improved documentation

#rnx

diff --git a/Engine/Source/Runtime/Core/Private/Tests/Memory/SharedBufferTest.cpp b/Engine/Source/Runtime/Core/Private/Tests/Memory/SharedBufferTest.cpp
--- a/Engine/Source/Runtime/Core/Private/Tests/Memory/SharedBufferTest.cpp
+++ b/Engine/Source/Runtime/Core/Private/Tests/Memory/SharedBufferTest.cpp
@@ -129,7 +129,7 @@ bool FUniqueBufferTest::RunTest(const FString& Parameters)
 	// Test Null
 	{
 		FUniqueBuffer Buffer;
-		TestFalse(TEXT("FUniqueBuffer().IsOwned()"), Buffer.IsOwned());
+		TestTrue(TEXT("FUniqueBuffer().IsOwned()"), Buffer.IsOwned());
 		TestEqual(TEXT("FUniqueBuffer().GetSize()"), Buffer.GetSize(), uint64(0));
 		TestEqual(TEXT("FUniqueBuffer().GetView().GetSize()"), Buffer.GetView().GetSize(), uint64(0));
 		TestEqual(TEXT("FUniqueBuffer().GetView().GetData()"), Buffer.GetView().GetData(), Buffer.GetData());
@@ -186,7 +186,7 @@ bool FUniqueBufferTest::RunTest(const FString& Parameters)
 	// Test MakeUnique
 	{
 		FUniqueBuffer Buffer = FUniqueBuffer::MakeUnique(FSharedBuffer());
-		TestFalse(TEXT("FUniqueBuffer::MakeUnique(Null).IsOwned()"), Buffer.IsOwned());
+		TestTrue(TEXT("FUniqueBuffer::MakeUnique(Null).IsOwned()"), Buffer.IsOwned());
 		TestEqual(TEXT("FUniqueBuffer::MakeUnique(Null).GetSize()"), Buffer.GetSize(), uint64(0));
 	}
 	{
@@ -247,7 +247,7 @@ bool FSharedBufferTest::RunTest(const FString& Parameters)
 	// Test Null
 	{
 		FSharedBuffer Buffer;
-		TestFalse(TEXT("FSharedBuffer().IsOwned()"), Buffer.IsOwned());
+		TestTrue(TEXT("FSharedBuffer().IsOwned()"), Buffer.IsOwned());
 		TestEqual(TEXT("FSharedBuffer().GetSize()"), Buffer.GetSize(), uint64(0));
 		TestEqual(TEXT("FSharedBuffer().GetView().GetSize()"), Buffer.GetView().GetSize(), uint64(0));
 		TestEqual(TEXT("FSharedBuffer().GetView().GetData()"), Buffer.GetView().GetData(), Buffer.GetData());
diff --git a/Engine/Source/Runtime/Core/Public/Memory/SharedBuffer.h b/Engine/Source/Runtime/Core/Public/Memory/SharedBuffer.h
--- a/Engine/Source/Runtime/Core/Public/Memory/SharedBuffer.h
+++ b/Engine/Source/Runtime/Core/Public/Memory/SharedBuffer.h
@@ -20,10 +20,12 @@ namespace BufferOwnerPrivate { struct FWeakOps; }
 ///////////////////////////////////////////////////////////////////////////////////////////////////////////////////////
 
 /**
- * A reference-counted buffer.
+ * A reference-counted owner for a buffer, which is a raw pointer and size.
  *
  * A buffer owner may own its memory or provide a view into memory owned externally. When used as
- * a non-owning view, the viewed memory must be guaranteed to outlive the buffer owner.
+ * a non-owning view, the viewed memory must be guaranteed to outlive the buffer owner. When this
+ * lifetime guarantee cannot be satisfied, MakeOwned may be called on the reference to the buffer
+ * to clone into a new buffer owner that owns the memory.
  *
  * A buffer owner must be referenced and accessed through one of its three reference types:
  * FUniqueBuffer, FSharedBuffer, or FWeakSharedBuffer.
@@ -187,7 +189,7 @@ private:
 ///////////////////////////////////////////////////////////////////////////////////////////////////////////////////////
 
 /**
- * A single-ownership mutable buffer.
+ * A reference to a single-ownership mutable buffer.
  *
  * Ownership can be transferred by moving to FUniqueBuffer or it can be converted to an immutable
  * shared buffer by moving to FSharedBuffer.
@@ -197,7 +199,7 @@ private:
 class FUniqueBuffer
 {
 public:
-	/** Make an owned buffer of Size bytes. */
+	/** Make an uninitialized owned buffer of the specified size. */
 	CORE_API static FUniqueBuffer Alloc(uint64 Size);
 
 	/** Make an owned clone of the input. */
@@ -262,22 +264,30 @@ public:
 	inline operator FMutableMemoryView() { return GetView(); }
 	inline operator FMemoryView() const { return GetView(); }
 
-	/** Returns true if this points to a buffer. */
+	/** Returns true if this points to a buffer owner. */
 	inline explicit operator bool() const { return !IsNull(); }
 
-	/** Returns true if this does not point to a buffer. */
+	/**
+	 * Returns true if this does not point to a buffer owner.
+	 *
+	 * A null buffer is always owned, materialized, and empty.
+	 */
 	inline bool IsNull() const { return Owner.IsNull(); }
 
-	/** Returns true if this owns the memory for its buffer. */
-	inline bool IsOwned() const { return Owner && Owner->IsOwned(); }
+	/** Returns true if this keeps the referenced buffer alive. */
+	inline bool IsOwned() const { return !Owner || Owner->IsOwned(); }
 
-	/** Clone into a new buffer if this does not own the memory for its buffer. */
+	/** Clone into a new buffer if the buffer is not owned. */
 	CORE_API void MakeOwned();
 
-	/** Returns true if this points to a buffer that has been materialized. */
-	inline bool IsMaterialized() const { return Owner && Owner->IsMaterialized(); }
+	/** Returns true if the referenced buffer has been materialized. */
+	inline bool IsMaterialized() const { return !Owner || Owner->IsMaterialized(); }
 
-	/** Make the data and size of the buffer available. GetData and GetSize do this automatically. */
+	/**
+	 * Materialize the buffer by making its data and size available.
+	 *
+	 * The buffer is automatically materialized by GetData, GetSize, GetView.
+	 */
 	CORE_API void Materialize() const;
 
 private:
@@ -294,7 +304,7 @@ private:
 ///////////////////////////////////////////////////////////////////////////////////////////////////////////////////////
 
 /**
- * A shared-ownership immutable buffer.
+ * A reference to a shared-ownership immutable buffer.
  *
  * @see FBufferOwner
  */
@@ -337,13 +347,13 @@ public:
 	/** Construct a shared buffer from a new unreferenced buffer owner. */
 	CORE_API explicit FSharedBuffer(FBufferOwner* Owner);
 
-	/** Construct a shared buffer by taking ownership of a unique buffer. */
+	/** Construct a shared buffer from a unique buffer, making it immutable. */
 	inline explicit FSharedBuffer(FUniqueBuffer&& Buffer)
 		: Owner(ToPrivateOwnerPtr(MoveTemp(Buffer)))
 	{
 	}
 
-	/** Assign a shared buffer by taking ownership of a unique buffer. */
+	/** Assign a shared buffer from a unique buffer, making it immutable. */
 	inline FSharedBuffer& operator=(FUniqueBuffer&& Buffer)
 	{
 		Owner = ToPrivateOwnerPtr(MoveTemp(Buffer));
@@ -363,22 +373,30 @@ public:
 	inline FMemoryView GetView() const { return FMemoryView(GetData(), GetSize()); }
 	inline operator FMemoryView() const { return GetView(); }
 
-	/** Returns true if this points to a buffer. */
+	/** Returns true if this points to a buffer owner. */
 	inline explicit operator bool() const { return !IsNull(); }
 
-	/** Returns true if this does not point to a buffer. */
+	/**
+	 * Returns true if this does not point to a buffer owner.
+	 *
+	 * A null buffer is always owned, materialized, and empty.
+	 */
 	inline bool IsNull() const { return Owner.IsNull(); }
 
-	/** Returns true if this owns the memory for its buffer. */
-	inline bool IsOwned() const { return Owner && Owner->IsOwned(); }
+	/** Returns true if this keeps the referenced buffer alive. */
+	inline bool IsOwned() const { return !Owner || Owner->IsOwned(); }
 
-	/** Clone into a new buffer if this does not own the memory for its buffer. */
+	/** Clone into a new buffer if the buffer is not owned. */
 	CORE_API void MakeOwned();
 
-	/** Returns true if this points to a buffer that has been materialized. */
-	inline bool IsMaterialized() const { return Owner && Owner->IsMaterialized(); }
+	/** Returns true if the referenced buffer has been materialized. */
+	inline bool IsMaterialized() const { return !Owner || Owner->IsMaterialized(); }
 
-	/** Make the data and size of the buffer available. GetData and GetSize do this automatically. */
+	/**
+	 * Materialize the buffer by making its data and size available.
+	 *
+	 * The buffer is automatically materialized by GetData, GetSize, GetView.
+	 */
 	CORE_API void Materialize() const;
 
 	friend class FWeakSharedBuffer;
@@ -416,7 +434,7 @@ public:
 	/** Reset this to null. */
 	CORE_API void Reset();
 
-	/** Convert this to a shared buffer. */
+	/** Convert this to a shared buffer if it has any remaining shared references. */
 	CORE_API FSharedBuffer Pin() const;
 
 private:
diff --git a/Engine/Source/Runtime/Core/Public/Serialization/CompactBinary.h b/Engine/Source/Runtime/Core/Public/Serialization/CompactBinary.h
--- a/Engine/Source/Runtime/Core/Public/Serialization/CompactBinary.h
+++ b/Engine/Source/Runtime/Core/Public/Serialization/CompactBinary.h
@@ -961,7 +961,7 @@ public:
 	inline void Reset() { *this = TCbBufferRef(); }
 
 	/** Whether this reference has ownership of the memory in its buffer. */
-	inline bool IsOwned() const { return Buffer.IsOwned(); }
+	inline bool IsOwned() const { return Buffer && Buffer.IsOwned(); }
 
 	/** Clone the value, if necessary, to a buffer that this reference has ownership of. */
 	inline void MakeOwned()
