Missing include.

#rnx

diff --git a/Engine/Source/Runtime/Core/Private/UObject/DevObjectVersion.cpp b/Engine/Source/Runtime/Core/Private/UObject/DevObjectVersion.cpp
--- a/Engine/Source/Runtime/Core/Private/UObject/DevObjectVersion.cpp
+++ b/Engine/Source/Runtime/Core/Private/UObject/DevObjectVersion.cpp
@@ -30,6 +30,7 @@
 #include "UObject/CineCameraObjectVersion.h"
 #include "UObject/VirtualProductionObjectVersion.h"
 #include "UObject/UE5MainStreamObjectVersion.h"
+#include "UObject/UE5ReleaseStreamObjectVersion.h"
 #include "UObject/UE5CookerObjectVersion.h"
 #include "UObject/MediaFrameWorkObjectVersion.h"
 
