U5 - have MeshDistanceFields BusyWait and not take down all Task Workers.

diff --git a/Engine/Source/Developer/MeshUtilities/Private/MeshDistanceFieldUtilities.cpp b/Engine/Source/Developer/MeshUtilities/Private/MeshDistanceFieldUtilities.cpp
--- a/Engine/Source/Developer/MeshUtilities/Private/MeshDistanceFieldUtilities.cpp
+++ b/Engine/Source/Developer/MeshUtilities/Private/MeshDistanceFieldUtilities.cpp
@@ -305,7 +305,8 @@ void FMeshUtilities::GenerateSignedDistanceFieldVolumeData(
 			for (int32 TaskIndex = 0; TaskIndex < AsyncTasks.Num(); TaskIndex++)
 			{
 				FAsyncTask<FMeshDistanceFieldAsyncTask>& Task = AsyncTasks[TaskIndex];
-				Task.EnsureCompletion(false);
+				extern CORE_API int32 GUseNewTaskBackend;
+				Task.EnsureCompletion(!!GUseNewTaskBackend);
 				bNegativeAtBorder = bNegativeAtBorder || Task.GetTask().WasNegativeAtBorder();
 			}
 
diff --git a/Engine/Source/Runtime/Core/Public/Async/AsyncWork.h b/Engine/Source/Runtime/Core/Public/Async/AsyncWork.h
--- a/Engine/Source/Runtime/Core/Public/Async/AsyncWork.h
+++ b/Engine/Source/Runtime/Core/Public/Async/AsyncWork.h
@@ -17,6 +17,7 @@
 #include "Misc/IQueuedWork.h"
 #include "Misc/QueuedThreadPool.h"
 #include "ProfilingDebugging/TagTrace.h"
+#include "Async/Fundamental/Scheduler.h"
 
 /**
 	FAutoDeleteAsyncTask - template task for jobs that delete themselves when complete
@@ -355,6 +356,11 @@ class FAsyncTask
 			FScopeCycleCounter Scope( Task.GetStatId() );
 			DECLARE_SCOPE_CYCLE_COUNTER( TEXT( "FAsyncTask::SyncCompletion" ), STAT_FAsyncTask_SyncCompletion, STATGROUP_ThreadPoolAsyncTasks );
 
+			if (LowLevelTasks::FScheduler::Get().GetActiveTask())
+			{
+				LowLevelTasks::BusyWaitUntil([this]() { return IsWorkDone(); });
+			}
+
 			check(DoneEvent); // if it is not done yet, we must have an event
 			DoneEvent->Wait();
 			QueuedPool = 0;
