task trace fix for a problem detected by unit tests: "empty" FGraphEvent doesn't have TraceId generated

diff --git a/Engine/Source/Runtime/Core/Private/Async/TaskGraph.cpp b/Engine/Source/Runtime/Core/Private/Async/TaskGraph.cpp
--- a/Engine/Source/Runtime/Core/Private/Async/TaskGraph.cpp
+++ b/Engine/Source/Runtime/Core/Private/Async/TaskGraph.cpp
@@ -2251,7 +2251,11 @@ static TLockFreeClassAllocator_TLSCache<FGraphEvent, PLATFORM_CACHE_LINE_SIZE> T
 
 FGraphEventRef FGraphEvent::CreateGraphEvent()
 {
-	return new(TheGraphEventAllocator.Allocate()) FGraphEvent{};
+	FGraphEvent* Instance = new(TheGraphEventAllocator.Allocate()) FGraphEvent{};
+	// `CreateGraphEvent()` creates an "empty" event that doesn't have an associated task and so won't get its TraceId set automatically
+	// to track such events we need to set a unique TraceId
+	Instance->SetTraceId(TaskTrace::GenerateTaskId());
+	return Instance;
 }
 
 void FGraphEvent::Recycle(FGraphEvent* ToRecycle)
diff --git a/Engine/Source/Runtime/Core/Private/Async/TaskTrace.cpp b/Engine/Source/Runtime/Core/Private/Async/TaskTrace.cpp
--- a/Engine/Source/Runtime/Core/Private/Async/TaskTrace.cpp
+++ b/Engine/Source/Runtime/Core/Private/Async/TaskTrace.cpp
@@ -74,7 +74,9 @@ namespace TaskTrace
 	FId GenerateTaskId()
 	{
 		static std::atomic<FId> UId{ 0 };
-		return UId.fetch_add(1, std::memory_order_relaxed);
+		FId Id = UId.fetch_add(1, std::memory_order_relaxed);
+		checkf(Id != InvalidId, TEXT("TraceId overflow"));
+		return Id;
 	}
 
 	static bool bGTaskTraceInitialized = false;
