TaskGraph tests: ported from Dev-TaskGraph stream

diff --git a/Engine/Source/Runtime/Core/Private/Tests/Async/TaskGraphTest.cpp b/Engine/Source/Runtime/Core/Private/Tests/Async/TaskGraphTest.cpp
--- a/Engine/Source/Runtime/Core/Private/Tests/Async/TaskGraphTest.cpp
+++ b/Engine/Source/Runtime/Core/Private/Tests/Async/TaskGraphTest.cpp
@@ -13,6 +13,8 @@
 #include "Math/RandomStream.h"
 #include "Containers/CircularQueue.h"
 #include "Containers/Queue.h"
+#include "Tests/Benchmark.h"
+
 #include <atomic>
 
 #if WITH_DEV_AUTOMATION_TESTS
@@ -724,7 +726,8 @@ namespace TaskGraphTests
 					//UE_LOG(LogTemp, Log, TEXT("Main task"));
 				}
 			);
-			while (!Event->IsComplete()) /* NOOP */;
+			while (!Event->IsComplete() && FTaskGraphInterface::Get().GetNumWorkerThreads() != 0) // in single-threaded mode tasks are executed only when waited for
+			{}
 			Event->Wait(ENamedThreads::GameThread);
 		}
 
@@ -747,26 +750,26 @@ namespace TaskGraphTests
 					Event->DispatchSubsequents();
 				}
 			);
-			while (!Event->IsComplete()) /* NOOP */;
+			while (!Event->IsComplete() && FTaskGraphInterface::Get().GetNumWorkerThreads() != 0) // in single-threaded mode tasks are executed only when waited for
+			{}
 			Event->Wait(ENamedThreads::GameThread);
 		}
 
 		{	// event w/o a task, signaled by explicit call to DispatchSubsequents after it's waited for
 			FGraphEventRef Event = FGraphEvent::CreateGraphEvent();
-			FFunctionGraphTask::CreateAndDispatchWhenReady(
-				[&Event]
+			auto Lambda = [&Event]
 				{
 					FPlatformProcess::Sleep(0.1f); // pause for a bit to let waiting start
 					Event->DispatchSubsequents();
-				}
-			);
+			};
+			FGraphEventRef Task = FFunctionGraphTask::CreateAndDispatchWhenReady(MoveTemp(Lambda));
 			check(!Event->IsComplete());
-			Event->Wait(ENamedThreads::GameThread);
+			Event->Wait();
+			Task->Wait();
 		}
 
 		{	// wait for prereq by DontCompleteUntil
-			FGraphEventRef Event = FFunctionGraphTask::CreateAndDispatchWhenReady(
-				[](ENamedThreads::Type CurrentThread, const FGraphEventRef& MyCompletionGraphEvent)
+			auto Lambda = [](ENamedThreads::Type CurrentThread, const FGraphEventRef& MyCompletionGraphEvent)
 				{
 					//UE_LOG(LogTemp, Log, TEXT("Main task"));
 
@@ -788,14 +791,15 @@ namespace TaskGraphTests
 
 					// now that Prereq was registered in DontCompleteUntil, unlock it
 					PrereqHolder->DispatchSubsequents();
-				}
-			);
+			};
+
+			FGraphEventRef Event = FFunctionGraphTask::CreateAndDispatchWhenReady(MoveTemp(Lambda));
 			Event->SetDebugName(TEXT("MainEvent"));
 			check(!Event->IsComplete());
 			Event->Wait(ENamedThreads::GameThread);
 		}
 
-		{	// prereq is completed when DontCompleteUntil is called
+		{	// prereq is completed before DontCompleteUntil is called
 			FGraphEventRef Prereq = FFunctionGraphTask::CreateAndDispatchWhenReady(
 				[]
 				{
@@ -813,14 +817,11 @@ namespace TaskGraphTests
 				}
 			);
 			Event->SetDebugName(TEXT("MainEvent"));
-			while (!Event->IsComplete()) /* NOOP */;
+			while (!Event->IsComplete() && FTaskGraphInterface::Get().GetNumWorkerThreads() != 0) // in single-threaded mode tasks are executed only when waited for
+			{}
 			Event->Wait(ENamedThreads::GameThread);
 		}
 
-		{	// "taskless" event with prereq
-			// forget about it, it's illegal as DontCompleteUntil() can be called only in associated task execution context
-		}
-
 		return true;
 	}
 
@@ -854,50 +855,198 @@ namespace TaskGraphTests
 		return true;
 	}
 
-	template<uint32 NumRuns, typename TestT>
-	void Benchmark(const TCHAR* TestName, TestT&& TestBody)
+	IMPLEMENT_SIMPLE_AUTOMATION_TEST(FTaskGraphBasicTest, "System.Core.Async.TaskGraph.BasicTest", EAutomationTestFlags::ApplicationContextMask | EAutomationTestFlags::EngineFilter);
+
+	bool FTaskGraphBasicTest::RunTest(const FString& Parameters)
 	{
-		UE_LOG(LogTemp, Display, TEXT("\n-------------------------------\n%s"), TestName);
-		double MinTime = TNumericLimits<double>::Max();
-		double TotalTime = 0;
-		for (uint32 RunNo = 0; RunNo != NumRuns; ++RunNo)
-		{
-			double Time = FPlatformTime::Seconds();
-			TestBody();
-			Time = FPlatformTime::Seconds() - Time;
+		// thread and task priorities
 
-			UE_LOG(LogTemp, Display, TEXT("#%d: %f secs"), RunNo, Time);
+		{	// AnyNormalThreadNormalTask
+			bool bExecuted = false;
+			FFunctionGraphTask::CreateAndDispatchWhenReady([&bExecuted] { bExecuted = true; }, TStatId{}, nullptr, ENamedThreads::AnyNormalThreadNormalTask)->Wait();
+			check(bExecuted);
+		}
 
-			TotalTime += Time;
-			if (MinTime > Time)
-			{
-				MinTime = Time;
-			}
+		{	// AnyNormalThreadHiPriTask
+			bool bExecuted = false;
+			FFunctionGraphTask::CreateAndDispatchWhenReady([&bExecuted] { bExecuted = true; }, TStatId{}, nullptr, ENamedThreads::AnyNormalThreadHiPriTask)->Wait();
+			check(bExecuted);
+		}
+
+		{	// AnyBackgroundThreadNormalTask
+			bool bExecuted = false;
+			FFunctionGraphTask::CreateAndDispatchWhenReady([&bExecuted] { bExecuted = true; }, TStatId{}, nullptr, ENamedThreads::AnyBackgroundThreadNormalTask)->Wait();
+			check(bExecuted);
+		}
+
+		{	// AnyHiPriThreadNormalTask
+			bool bExecuted = false;
+			FFunctionGraphTask::CreateAndDispatchWhenReady([&bExecuted] { bExecuted = true; }, TStatId{}, nullptr, ENamedThreads::AnyHiPriThreadNormalTask)->Wait();
+			check(bExecuted);
+		}
+
+		{	// AnyHiPriThreadHiPriTask
+			bool bExecuted = false;
+			FFunctionGraphTask::CreateAndDispatchWhenReady([&bExecuted] { bExecuted = true; }, TStatId{}, nullptr, ENamedThreads::AnyHiPriThreadHiPriTask)->Wait();
+			check(bExecuted);
 		}
-		UE_LOG(LogTemp, Display, TEXT("min: %f secs, avg: %f secs\n-------------------------------\n"), MinTime, TotalTime / NumRuns);
 
-#if NO_LOGGING
-		printf("min: %f\n", MinTime);
+		// named threads and local queues
+
+#if STATS
+		{	// StatsThread
+			bool bExecuted = false;
+			FFunctionGraphTask::CreateAndDispatchWhenReady([&bExecuted] { bExecuted = true; }, TStatId{}, nullptr, ENamedThreads::StatsThread)->Wait();
+			check(bExecuted);
+		}
 #endif
-	}
 
-#define BENCHMARK(NumRuns, ...) Benchmark<NumRuns>(TEXT(#__VA_ARGS__), __VA_ARGS__)
+		if (IsRHIThreadRunning())
+		{	// RHIThread
+			bool bExecuted = false;
+			FFunctionGraphTask::CreateAndDispatchWhenReady([&bExecuted] { bExecuted = true; }, TStatId{}, nullptr, ENamedThreads::RHIThread)->Wait();
+			check(bExecuted);
+		}
+
+		if (IsAudioThreadRunning())
+		{	// AudioThread
+			bool bExecuted = false;
+			FFunctionGraphTask::CreateAndDispatchWhenReady([&bExecuted] { bExecuted = true; }, TStatId{}, nullptr, ENamedThreads::AudioThread)->Wait();
+			check(bExecuted);
+		}
+
+		{	// GameThread
+			bool bExecuted = false;
+			FFunctionGraphTask::CreateAndDispatchWhenReady([&bExecuted] { bExecuted = true; }, TStatId{}, nullptr, ENamedThreads::GameThread);
+			// GT must be executed explicitly
+			FTaskGraphInterface::Get().ProcessThreadUntilIdle(ENamedThreads::GameThread);
+			check(bExecuted);
+		}
+
+PRAGMA_DISABLE_DEPRECATION_WARNINGS
+		if (GRenderThreadId != 0)
+PRAGMA_ENABLE_DEPRECATION_WARNINGS
+		{	// ActualRenderingThread
+			bool bExecuted = false;
+			FFunctionGraphTask::CreateAndDispatchWhenReady([&bExecuted] { bExecuted = true; }, TStatId{}, nullptr, ENamedThreads::ActualRenderingThread)->Wait();
+			check(bExecuted);
+		}
+
+		{	// GameThread_Local
+			bool bExecuted = false;
+			FFunctionGraphTask::CreateAndDispatchWhenReady([&bExecuted] { bExecuted = true; }, TStatId{}, nullptr, ENamedThreads::GameThread_Local)->Wait(ENamedThreads::GameThread_Local);
+			check(bExecuted);
+		}
+
+PRAGMA_DISABLE_DEPRECATION_WARNINGS
+		if (GRenderThreadId != 0)
+PRAGMA_ENABLE_DEPRECATION_WARNINGS
+		{	// ActualRenderingThread_Local
+			bool bExecuted = false;
+			FFunctionGraphTask::CreateAndDispatchWhenReady(
+				[&bExecuted] 
+				{ 
+					FFunctionGraphTask::CreateAndDispatchWhenReady([&bExecuted] { bExecuted = true; }, TStatId{}, nullptr, ENamedThreads::ActualRenderingThread_Local)->Wait(ENamedThreads::ActualRenderingThread_Local);
+				}, 
+				TStatId{}, nullptr, ENamedThreads::ActualRenderingThread)->Wait();
+			check(bExecuted);
+		}
+
+		// dependencies
+
+		{	// a task is not executed until its prerequisite is completed
+			bool bExecuted = false;
+			FGraphEventRef Prereq = FGraphEvent::CreateGraphEvent();
+			FGraphEventRef MainTask = FFunctionGraphTask::CreateAndDispatchWhenReady([&bExecuted] { bExecuted = true; }, TStatId{}, Prereq);
+			// dummy task that is executed while the main task is waiting for its prereq
+			FFunctionGraphTask::CreateAndDispatchWhenReady([] {})->Wait();
+			check(!bExecuted);
+			Prereq->DispatchSubsequents();
+			MainTask->Wait();
+			check(bExecuted);
+		}
+
+		{	// a task is not executed until all its prerequisites are completed
+			bool bExecuted = false;
+			FGraphEventArray Prereqs{ FGraphEvent::CreateGraphEvent(), FGraphEvent::CreateGraphEvent() };
+			FGraphEventRef MainTask = FFunctionGraphTask::CreateAndDispatchWhenReady([&bExecuted] { bExecuted = true; }, TStatId{}, &Prereqs);
+			// dummy task that is executed while the main task is waiting for its prereqs
+			FFunctionGraphTask::CreateAndDispatchWhenReady([] {})->Wait();
+			check(!bExecuted);
+
+			Prereqs[0]->DispatchSubsequents();
+			FFunctionGraphTask::CreateAndDispatchWhenReady([] {})->Wait();
+			check(!bExecuted);
+			
+			Prereqs[1]->DispatchSubsequents();
+			MainTask->Wait();
+			check(bExecuted);
+		}
+
+		{	// ParallelFor
+			std::atomic<int32> Total{ 0 };
+			int32 Num = 1000;
+			ParallelFor(Num, [&Total](int32 i) { Total += i; });
+			check(Total == (Num - 1) * (Num / 2));
+		}
+
+		{	// holding a task
+			struct FTask
+			{
+				TStatId GetStatId() const
+				{
+					return TStatId{};
+				}
+
+				ENamedThreads::Type GetDesiredThread()
+				{
+					return ENamedThreads::AnyThread;
+				}
+
+				static ESubsequentsMode::Type GetSubsequentsMode()
+				{
+					return ESubsequentsMode::TrackSubsequents;
+				}
+
+				void DoTask(ENamedThreads::Type CurrentThread, const FGraphEventRef& MyCompletionGraphEvent)
+				{}
+			};
+
+			TGraphTask<FTask>* Task = TGraphTask<FTask>::CreateTask().ConstructAndHold();
+			FGraphEventRef Event = Task->GetCompletionEvent();
+			check(!Event->IsComplete());
+			Task->Unlock();
+			Event->Wait();
+			check(Event->IsComplete());
+		}
+
+		//for (int i = 0; i != 100000; ++i)
+		{	// a particular real-life case that doesn't work in the old TaskGraph if run in single-threaded mode.
+			// the culprit is that when a task is waited for, in single-threaded mode the queue it was pushed to is executed. 
+			// Here the (local queue) task  depends on a task that is not in the same queue and so it doesn't get executed
+
+			FGraphEventRef AnyTask = FFunctionGraphTask::CreateAndDispatchWhenReady([] {});
+			FGraphEventRef LocalQueueTask = FFunctionGraphTask::CreateAndDispatchWhenReady([] {}, TStatId{}, AnyTask, ENamedThreads::GameThread_Local);
+			LocalQueueTask->Wait(ENamedThreads::GameThread_Local);
+			check(LocalQueueTask.GetRefCount() == 1);
+		}
+
+		return true;
+	}
 
 	// it's fast because tasks are too lightweight and so are executed almost as fast
 	template<int NumTasks>
 	void TestPerfBasic()
 	{
-		int32 CompletedTasks = 0;
+		FGraphEventArray Tasks;
+		Tasks.Reserve(NumTasks);
 
 		for (int32 TaskIndex = 0; TaskIndex < NumTasks; ++TaskIndex)
 		{
-			FFunctionGraphTask::CreateAndDispatchWhenReady([&CompletedTasks] { FPlatformAtomics::InterlockedIncrement(&CompletedTasks); });
+			Tasks.Emplace(FFunctionGraphTask::CreateAndDispatchWhenReady([] {}));
 		}
 
-		while ((CompletedTasks < NumTasks))
-		{
-			FPlatformProcess::Yield();
-		}
+		FTaskGraphInterface::Get().WaitUntilTasksComplete(Tasks);
 	}
 
 	template<int32 NumTasks, int32 BatchSize>
@@ -905,26 +1054,27 @@ namespace TaskGraphTests
 	{
 		static_assert(NumTasks % BatchSize == 0, "`NumTasks` must be divisible by `BatchSize`");
 		constexpr int32 NumBatches = NumTasks / BatchSize;
-		
-		int32 CompletedTasks = 0;
+
+		FGraphEventArray Batches;
+		Batches.Reserve(NumBatches);
+		FGraphEventArray Tasks;
+		Tasks.AddDefaulted(NumTasks);
 
 		for (int32 BatchIndex = 0; BatchIndex < NumBatches; ++BatchIndex)
 		{
-			FFunctionGraphTask::CreateAndDispatchWhenReady(
-				[&CompletedTasks]
+			Batches.Add(FFunctionGraphTask::CreateAndDispatchWhenReady(
+				[&Tasks, BatchIndex]
 				{
 					for (int32 TaskIndex = 0; TaskIndex < BatchSize; ++TaskIndex)
 					{
-						FFunctionGraphTask::CreateAndDispatchWhenReady([&CompletedTasks] { FPlatformAtomics::InterlockedIncrement(&CompletedTasks); });
+						Tasks[BatchIndex * BatchSize + TaskIndex] = FFunctionGraphTask::CreateAndDispatchWhenReady([] {});
 					}
 				}
-			);
+			));
 		}
 
-		while ((CompletedTasks < NumTasks))
-		{
-			FPlatformProcess::Yield();
-		}
+		FTaskGraphInterface::Get().WaitUntilTasksComplete(Batches);
+		FTaskGraphInterface::Get().WaitUntilTasksComplete(Tasks);
 	}
 
 	template<int32 NumTasks, int32 BatchSize>
@@ -932,14 +1082,14 @@ namespace TaskGraphTests
 	{
 		static_assert(NumTasks % BatchSize == 0, "`NumTasks` must be divisible by `BatchSize`");
 		constexpr int32 NumBatches = NumTasks / BatchSize;
-		
+
 		FGraphEventRef SpawnSignal = FGraphEvent::CreateGraphEvent();
 		FGraphEventArray AllDone;
 
 		for (int32 BatchIndex = 0; BatchIndex < NumBatches; ++BatchIndex)
 		{
 			AllDone.Add(FFunctionGraphTask::CreateAndDispatchWhenReady(
-				[] (ENamedThreads::Type CurrentThread, const FGraphEventRef& CompletionEvent)
+				[](ENamedThreads::Type CurrentThread, const FGraphEventRef& CompletionEvent)
 				{
 					FGraphEventRef RunSignal = FGraphEvent::CreateGraphEvent();
 					for (int32 TaskIndex = 0; TaskIndex < BatchSize; ++TaskIndex)
@@ -1042,120 +1192,11 @@ namespace TaskGraphTests
 		UE_LOG(LogTemp, Display, TEXT("Fibonacci(%d) = %d"), N, *Res);
 	}
 
-	namespace Queues
-	{
-		template<uint32 Num>
-		void TestTCircularQueue()
-		{
-			TCircularQueue<uint32> Queue{ 100 };
-			std::atomic<bool> bStop{ false };
-
-			FGraphEventRef Task = FFunctionGraphTask::CreateAndDispatchWhenReady(
-				[&bStop, &Queue]
-				{
-					while (!bStop)
-					{
-						Queue.Enqueue(0);
-					}
-				}
-				);
-
-			uint32 It = 0;
-			while (It != Num)
-			{
-				uint32 El;
-				if (Queue.Dequeue(El))
-				{
-					++It;
-				}
-			}
-
-			bStop = true;
-
-			Task->Wait(ENamedThreads::GameThread);
-		}
-
-		template<uint32 Num, EQueueMode Mode>
-		void TestTQueue()
-		{
-			TQueue<uint32, Mode> Queue;
-			std::atomic<bool> bStop{ false };
-
-			FGraphEventRef Task = FFunctionGraphTask::CreateAndDispatchWhenReady(
-				[&bStop, &Queue]
-				{
-					while (!bStop)
-					{
-						Queue.Enqueue(0);
-					}
-				}
-				);
-
-			uint32 It = 0;
-			while (It != Num)
-			{
-				uint32 El;
-				if (Queue.Dequeue(El))
-				{
-					++It;
-				}
-			}
-
-			bStop = true;
-
-			Task->Wait(ENamedThreads::GameThread);
-		}
-
-		template<uint32 Num>
-		void TestMpscTQueue()
-		{
-			TQueue<uint32, EQueueMode::Mpsc> Queue;
-			std::atomic<bool> bStop{ false };
-
-			int32 NumProducers = FPlatformMisc::NumberOfCoresIncludingHyperthreads() - 1;
-			FGraphEventArray Tasks;
-			for (int32 i = 0; i != NumProducers; ++i)
-			{
-				Tasks.Add(FFunctionGraphTask::CreateAndDispatchWhenReady(
-					[&bStop, &Queue]
-					{
-						while (!bStop)
-						{
-							Queue.Enqueue(0);
-						}
-					}
-					));
-			}
-
-			uint32 It = 0;
-			while (It != Num)
-			{
-				uint32 El;
-				if (Queue.Dequeue(El))
-				{
-					++It;
-				}
-			}
-
-			bStop = true;
-
-			FTaskGraphInterface::Get().WaitUntilTasksComplete(MoveTemp(Tasks), ENamedThreads::GameThread);
-		}
-
-		void Test()
-		{
-			BENCHMARK(5, TestTCircularQueue<10'000'000>);
-			BENCHMARK(5, TestTQueue<10'000'000, EQueueMode::Spsc>);
-			BENCHMARK(5, TestTQueue<10'000'000, EQueueMode::Mpsc>);
-			BENCHMARK(5, TestMpscTQueue<1'000'000>);
-		}
-	}
-
-	template<int NumTasks>
+	template<uint32 NumTasks>
 	void TestFGraphEventPerf()
 	{
 		FGraphEventRef Prereq = FGraphEvent::CreateGraphEvent();
-		int32 CompletedTasks = 0;
+		std::atomic<uint32> CompletedTasks{ 0 };
 
 		FGraphEventArray Tasks;
 		for (int i = 0; i != NumTasks; ++i)
@@ -1164,14 +1205,14 @@ namespace TaskGraphTests
 				[&Prereq, &CompletedTasks](ENamedThreads::Type CurrentThread, const FGraphEventRef& MyCompletionGraphEvent)
 				{
 					MyCompletionGraphEvent->DontCompleteUntil(Prereq);
-					FPlatformAtomics::InterlockedIncrement(&CompletedTasks);
+					++CompletedTasks;
 				}
 			));
 		}
 
 		Prereq->DispatchSubsequents(ENamedThreads::GameThread);
 
-		FTaskGraphInterface::Get().WaitUntilTasksComplete(MoveTemp(Tasks), ENamedThreads::GameThread);
+		FTaskGraphInterface::Get().WaitUntilTasksComplete(Tasks, ENamedThreads::GameThread);
 
 		check(CompletedTasks == NumTasks);
 	}
@@ -1182,14 +1223,14 @@ namespace TaskGraphTests
 		{
 			FGraphEventArray Tasks;
 			Tasks.Reserve(NumTasks);
-			double StartTime = FPlatformTime::Seconds();
+			//double StartTime = FPlatformTime::Seconds();
 			for (uint32 TaskNo = 0; TaskNo != NumTasks; ++TaskNo)
 			{
 				Tasks.Add(FFunctionGraphTask::CreateAndDispatchWhenReady([] {}));
 			}
 
-			double Duration = FPlatformTime::Seconds() - StartTime;
-			UE_LOG(LogTemp, Display, TEXT("Spawning %d empty trackable tasks took %f secs"), NumTasks, Duration);
+			//double Duration = FPlatformTime::Seconds() - StartTime;
+			//UE_LOG(LogTemp, Display, TEXT("Spawning %d empty trackable tasks took %f secs"), NumTasks, Duration);
 
 			FTaskGraphInterface::Get().WaitUntilTasksComplete(MoveTemp(Tasks));
 		}
@@ -1200,60 +1241,351 @@ namespace TaskGraphTests
 				FFunctionGraphTask::CreateAndDispatchWhenReady([] {});
 			}
 
-			double Duration = FPlatformTime::Seconds() - StartTime;
-			UE_LOG(LogTemp, Display, TEXT("Spawning %d empty non-trackable tasks took %f secs"), NumTasks, Duration);
+			//double Duration = FPlatformTime::Seconds() - StartTime;
+			//UE_LOG(LogTemp, Display, TEXT("Spawning %d empty non-trackable tasks took %f secs"), NumTasks, Duration);
 		}
 	}
 
 	template<int NumTasks>
 	void TestBatchSpawning()
 	{
-			double StartTime = FPlatformTime::Seconds();
-			FGraphEventRef Trigger = FGraphEvent::CreateGraphEvent();
-			for (uint32 TaskNo = 0; TaskNo != NumTasks; ++TaskNo)
-			{
-				FFunctionGraphTask::CreateAndDispatchWhenReady([] {}, TStatId{}, Trigger);
-			}
+		//double StartTime = FPlatformTime::Seconds();
+		FGraphEventRef Trigger = FGraphEvent::CreateGraphEvent();
+		for (uint32 TaskNo = 0; TaskNo != NumTasks; ++TaskNo)
+		{
+			FFunctionGraphTask::CreateAndDispatchWhenReady([] {}, TStatId{}, Trigger);
+		}
 
-			double SpawnedTime = FPlatformTime::Seconds();
-			Trigger->DispatchSubsequents();
+		//double SpawnedTime = FPlatformTime::Seconds();
+		Trigger->DispatchSubsequents();
 
-			double EndTime = FPlatformTime::Seconds();
-			UE_LOG(LogTemp, Display, TEXT("Spawning %d empty non-trackable tasks took %f secs total, %f secs spawning and %f secs dispatching"), NumTasks, EndTime - StartTime, SpawnedTime - StartTime, EndTime - SpawnedTime);
+		//double EndTime = FPlatformTime::Seconds();
+		//UE_LOG(LogTemp, Display, TEXT("Spawning %d empty non-trackable tasks took %f secs total, %f secs spawning and %f secs dispatching"), NumTasks, EndTime - StartTime, SpawnedTime - StartTime, EndTime - SpawnedTime);
 
-			//FTaskGraphInterface::Get().WaitUntilTasksComplete(MoveTemp(Tasks));
+		//FTaskGraphInterface::Get().WaitUntilTasksComplete(MoveTemp(Tasks));
+	}
+
+	template<int64 NumBatches, int64 NumTasksPerBatch>
+	void TestWorkStealing()
+	{
+		FGraphEventArray Batches;
+		Batches.Reserve(NumBatches);
+
+		FGraphEventArray Tasks[NumBatches];
+		for (int32 BatchIndex = 0; BatchIndex != NumBatches; ++BatchIndex)
+		{
+			Tasks[BatchIndex].Reserve(NumBatches * NumTasksPerBatch);
+			Batches.Add(FFunctionGraphTask::CreateAndDispatchWhenReady(
+				[&Tasks, BatchIndex]()
+				{
+					for (int32 TaskIndex = 0; TaskIndex < NumTasksPerBatch; ++TaskIndex)
+					{
+						Tasks[BatchIndex].Add(FFunctionGraphTask::CreateAndDispatchWhenReady([] {}));
+					}
+				}
+			));
+		}
+
+		FTaskGraphInterface::Get().WaitUntilTasksComplete(Batches);
+		for (int32 BatchIndex = 0; BatchIndex != NumBatches; ++BatchIndex)
+		{
+			FTaskGraphInterface::Get().WaitUntilTasksComplete(Tasks[BatchIndex]);
+		}
 	}
 
 	IMPLEMENT_SIMPLE_AUTOMATION_TEST(FPerfTest, "System.Core.Async.TaskGraph.PerfTest", EAutomationTestFlags::ApplicationContextMask | EAutomationTestFlags::EngineFilter);
 
 	bool FPerfTest::RunTest(const FString& Parameters)
 	{
-		// profiling
-		//TestBatchSpawning<10000000>();
-		//BENCHMARK(5, TestPerfBasic<1 << 25>);
-		//BENCHMARK(5, TestPerfOptimised<1 << 24>);
-		//BENCHMARK(10, TestFGraphEventPerf<1 << 22>);
-		//BENCHMARK(10, TestSpawning<100000000>); // for profiling
-		//BENCHMARK(10, Fib<36>); // for profiling
+		//UE_BENCHMARK(1, TestSpawning<100000>);
+		//return true;
+
+		//UE_BENCHMARK(5, Fib<18>);
+		//UE_BENCHMARK(5, [] { Fibonacci(15); });
 
-		BENCHMARK(5, Fib<18>);
-		//BENCHMARK(5, [] { Fibonacci(15); });
+		UE_BENCHMARK(5, TestPerfBasic<100000>);
+		UE_BENCHMARK(5, TestPerfBatch<100000, 100>);
+		UE_BENCHMARK(5, TestPerfBatchOptimised<100000, 100>);
+		UE_BENCHMARK(5, TestLatency<10000>);
+		UE_BENCHMARK(5, TestFGraphEventPerf<100000>);
+		UE_BENCHMARK(5, TestWorkStealing<100, 1000>);
+		UE_BENCHMARK(5, TestSpawning<100000>);
+		UE_BENCHMARK(5, TestBatchSpawning<100000>);
 
-		BENCHMARK(5, TestFGraphEventPerf<1 << 16>);
-		BENCHMARK(5, TestPerfBasic<1 << 17>);
-		BENCHMARK(5, TestPerfBatch<1 << 17, 1 << 13>);
-		BENCHMARK(5, TestPerfBatchOptimised<1 << 17, 1 << 13>);
-		BENCHMARK(5, TestLatency<10'000>);
+		return true;
+	}
 
-		//BENCHMARK(10, TestSpawning<1>);
-		BENCHMARK(5, TestSpawning<100'000>);
-		BENCHMARK(5, TestBatchSpawning<100'000>);
+	void PipeStressTest()
+	{
+#if STATS
+		constexpr ENamedThreads::Type Pipe = ENamedThreads::StatsThread;
+		constexpr ENamedThreads::Type ManualPipe = ENamedThreads::Type(ENamedThreads::StatsThread | ENamedThreads::LocalQueue);
+
+		{	// check that piped task is executed
+			std::atomic<bool> bExecuted{ false };
+			FFunctionGraphTask::CreateAndDispatchWhenReady([&bExecuted] { bExecuted = true; }, TStatId{}, nullptr, Pipe)->Wait();
+			check(bExecuted);
+		}
 
-		//Queues::Test();
+		{	// check that waiting inside a piped task works
+			FGraphEventRef ContinueTask = FGraphEvent::CreateGraphEvent();
 
-		return true;
+			FGraphEventRef Task = FFunctionGraphTask::CreateAndDispatchWhenReady(
+				[&ContinueTask]
+				{
+					check(!ContinueTask->IsComplete() || FTaskGraphInterface::Get().GetNumWorkerThreads() == 0);
+					ContinueTask->Wait();
+				},
+				TStatId{}, nullptr, Pipe
+			);
+
+			FPlatformProcess::Sleep(0.1f);
+			ContinueTask->DispatchSubsequents();
+			Task->Wait();
+		}
+
+		// DEADLOCK:
+		// `T1Launcher->DispatchSubsequents()` -> resolves T1 which is executed and unlocks T2. Then `T1->Wait()` (which is currently implemented
+		// as busy-waiting) starts executing T2 and blocks waiting for T2Launcher, but T2Launcher is triggered only the next line. BOOM
+		// Conclusion: mixing blocking waiting and busy waiting is explosive
+		if (false)
+		{	// check that piped tasks don't execute in parallel
+			std::atomic<bool> bExecutingPipe{ false };
+
+			FGraphEventRef T1Launcher = FGraphEvent::CreateGraphEvent();
+			std::atomic<bool> bT1Executed{ false };
+			FGraphEventRef T1 = FFunctionGraphTask::CreateAndDispatchWhenReady(
+				[&bExecutingPipe, T1Launcher, &bT1Executed]
+				{
+					check(!bExecutingPipe);
+					bExecutingPipe = true;
+					T1Launcher->Wait();
+					bT1Executed = true;
+					bExecutingPipe = false;
+				}, 
+				TStatId{}, nullptr, Pipe
+			);
+
+			FGraphEventRef T2Launcher = FGraphEvent::CreateGraphEvent();
+			std::atomic<bool> bT2Executed{ false };
+			FGraphEventRef T2 = FFunctionGraphTask::CreateAndDispatchWhenReady(
+				[&bExecutingPipe, T2Launcher, &bT2Executed]
+				{
+					check(!bExecutingPipe);
+					bExecutingPipe = true;
+					T2Launcher->Wait();
+					bT2Executed = true;
+					bExecutingPipe = false;
+				}, 
+				TStatId{}, nullptr, Pipe
+			);
+
+			FGraphEventRef T3Launcher = FGraphEvent::CreateGraphEvent();
+			std::atomic<bool> bT3Executed{ false };
+			FGraphEventRef T3 = FFunctionGraphTask::CreateAndDispatchWhenReady(
+				[&bExecutingPipe, T3Launcher, &bT3Executed]
+				{
+					check(!bExecutingPipe);
+					bExecutingPipe = true;
+					T3Launcher->Wait();
+					bT3Executed = true;
+					bExecutingPipe = false;
+				}, 
+				TStatId{}, nullptr, Pipe
+			);
+
+			check(!T1->IsComplete());
+			check(!T2->IsComplete());
+			check(!T3->IsComplete());
+			T1Launcher->DispatchSubsequents();
+			T1->Wait();
+			check(!T2->IsComplete());
+			check(!T3->IsComplete());
+			T2Launcher->DispatchSubsequents();
+			T2->Wait();
+			check(!T3->IsComplete());
+			T3Launcher->DispatchSubsequents();
+			T3->Wait();
+		}
+
+		{	// test that manual pipe tasks are executed only manually
+			std::atomic<bool> bExecuted{ false };
+			FGraphEventRef ManualPipeTask = FFunctionGraphTask::CreateAndDispatchWhenReady([&bExecuted] { bExecuted = true; }, TStatId{}, nullptr, ManualPipe);
+			FPlatformProcess::Sleep(0.1f);
+			check(!bExecuted);
+			check(!ManualPipeTask->IsComplete());
+
+			// manual pipes must be executed from corresponding pipe
+			FFunctionGraphTask::CreateAndDispatchWhenReady([ManualPipe] { FTaskGraphInterface::Get().ProcessThreadUntilIdle(ManualPipe); }, TStatId{}, nullptr, Pipe)
+				->Wait();
+
+			check(bExecuted);
+			check(ManualPipeTask->IsComplete());
+		}
+
+		{	// test ProcessThreadUntilIdle
+			std::atomic<bool> bExecuted{ false };
+
+			FFunctionGraphTask::CreateAndDispatchWhenReady(
+				[ManualPipe, &bExecuted]
+				{
+					FGraphEventRef ManualPipeTask = FFunctionGraphTask::CreateAndDispatchWhenReady([&bExecuted] { bExecuted = true; }, TStatId{}, nullptr, ManualPipe);
+					FTaskGraphInterface::Get().ProcessThreadUntilIdle(ManualPipe);
+					check(bExecuted);
+					check(ManualPipeTask->IsComplete());
+				}, 
+				TStatId{}, nullptr, Pipe
+			)->Wait();
+			check(bExecuted);
+		}
+
+		{	// test Wait
+			std::atomic<bool> bExecuted{ false };
+
+			FFunctionGraphTask::CreateAndDispatchWhenReady(
+				[ManualPipe, &bExecuted]
+				{
+					FFunctionGraphTask::CreateAndDispatchWhenReady([&bExecuted] { bExecuted = true; }, TStatId{}, nullptr, ManualPipe)
+						->Wait(ManualPipe);
+					check(bExecuted);
+				},
+				TStatId{}, nullptr, Pipe
+			)->Wait();
+			check(bExecuted);
+		}
+
+		if (FTaskGraphInterface::Get().GetNumWorkerThreads() != 0)
+		{	// auto-pipe with fire-and-forget task
+			std::atomic<bool> bExecuted{ false };
+
+			TGraphTask<TFunctionGraphTaskImpl<void(), ESubsequentsMode::FireAndForget>>::CreateTask().ConstructAndDispatchWhenReady([&bExecuted] { bExecuted = true; }, TStatId{}, Pipe);
+
+			FPlatformProcess::Sleep(0.1f);
+			check(bExecuted);
+		}
+
+		{	// manual pipe with fire-and-forget task
+			std::atomic<bool> bExecuted{ false };
+
+			TGraphTask<TFunctionGraphTaskImpl<void(), ESubsequentsMode::FireAndForget>>::CreateTask().ConstructAndDispatchWhenReady([&bExecuted] { bExecuted = true; }, TStatId{}, ManualPipe);
+
+			FFunctionGraphTask::CreateAndDispatchWhenReady([ManualPipe] { FTaskGraphInterface::Get().ProcessThreadUntilIdle(ManualPipe); }, TStatId{}, nullptr, Pipe)
+				->Wait();
+			check(bExecuted);
+		}
+
+		{	// test GetCurrentThread() from inside manual pipe
+			FFunctionGraphTask::CreateAndDispatchWhenReady(
+				[ManualPipe]
+				{ 
+					checkf(FTaskGraphInterface::Get().GetCurrentThreadIfKnown(true) == ManualPipe, TEXT("%d - %d"), FTaskGraphInterface::Get().GetCurrentThreadIfKnown(true), ManualPipe);
+				},
+			TStatId{}, nullptr, ManualPipe);
+
+			FFunctionGraphTask::CreateAndDispatchWhenReady([ManualPipe] { FTaskGraphInterface::Get().ProcessThreadUntilIdle(ManualPipe); }, TStatId{}, nullptr, Pipe)
+				->Wait();
+		}
+
+		{	// test IsThreadProcessingTasks() from inside manual pipe
+			FFunctionGraphTask::CreateAndDispatchWhenReady(
+				[ManualPipe]
+				{
+					check(FTaskGraphInterface::Get().IsThreadProcessingTasks(ManualPipe));
+				},
+			TStatId{}, nullptr, ManualPipe);
+
+			FFunctionGraphTask::CreateAndDispatchWhenReady(
+				[ManualPipe]
+				{
+					check(!FTaskGraphInterface::Get().IsThreadProcessingTasks(ManualPipe));
+					FTaskGraphInterface::Get().ProcessThreadUntilIdle(ManualPipe); 
+				}, 
+			TStatId{}, nullptr, Pipe)->Wait();
+		}
+
+		{	// check that independent task can bypass a task with pending prereq
+			FGraphEventRef Prereq = FGraphEvent::CreateGraphEvent();
+			FGraphEventRef Dependent = FFunctionGraphTask::CreateAndDispatchWhenReady([] {}, TStatId{}, Prereq, Pipe);
+			FGraphEventRef Independent = FFunctionGraphTask::CreateAndDispatchWhenReady([] {}, TStatId{}, nullptr, Pipe);
+			Independent->Wait();
+			check(Independent->IsComplete() && !Dependent->IsComplete());
+			Prereq->DispatchSubsequents();
+			Dependent->Wait();
+		}
+
+		{	// stress test for piped tasks, checks spawning a large number of piped tasks from multiple threads and executing them
+			//UE_LOG(LogTemp, Display, TEXT("Pipe stress test"));
+
+			constexpr uint32 SpawnerGroupsNum = 50;
+			FGraphEventArray SpawnerGroups;
+			SpawnerGroups.Reserve(SpawnerGroupsNum);
+
+			constexpr uint32 SpawnersPerGroupNum = 50;
+			constexpr uint32 TasksNum = SpawnerGroupsNum * SpawnersPerGroupNum;
+			FGraphEventArray Spawners;
+			Spawners.AddDefaulted(TasksNum);
+			FGraphEventArray Tasks;
+			Tasks.AddDefaulted(TasksNum);
+
+			std::atomic<bool> bExecuting{ false };
+			std::atomic<uint32> TasksExecutedNum{ 0 };
+
+			for (uint32 GroupIndex = 0; GroupIndex != SpawnerGroupsNum; ++GroupIndex)
+			{
+				SpawnerGroups.Add(FFunctionGraphTask::CreateAndDispatchWhenReady(
+					[Spawners = &Spawners[GroupIndex * SpawnersPerGroupNum],
+					Tasks = &Tasks[GroupIndex * SpawnersPerGroupNum],
+					&bExecuting,
+					&TasksExecutedNum,
+					SpawnersPerGroupNum,
+					Pipe]
+					{
+						for (uint32 SpawnerIndex = 0; SpawnerIndex != SpawnersPerGroupNum; ++SpawnerIndex)
+						{
+							Spawners[SpawnerIndex] = FFunctionGraphTask::CreateAndDispatchWhenReady(
+								[Task = &Tasks[SpawnerIndex], &bExecuting, &TasksExecutedNum, Pipe]
+								{
+									*Task = FFunctionGraphTask::CreateAndDispatchWhenReady(
+										[&bExecuting, &TasksExecutedNum, Pipe]
+										{
+											check(!bExecuting);
+											bExecuting = true;
+											check(FTaskGraphInterface::Get().GetCurrentThreadIfKnown() == Pipe);
+											FPlatformProcess::SleepNoStats(0.0f);
+											++TasksExecutedNum;
+											bExecuting = false;
+										},
+										TStatId{}, nullptr, Pipe
+									);
+								}
+							);
+						}
+					}
+				));
+			}
+
+			//UE_LOG(LogTemp, Display, TEXT("Spawners groups spawned"));
+			FTaskGraphInterface::Get().WaitUntilTasksComplete(SpawnerGroups);
+			//UE_LOG(LogTemp, Display, TEXT("Spawners groups completed"));
+
+			FTaskGraphInterface::Get().WaitUntilTasksComplete(Spawners);
+			//UE_LOG(LogTemp, Display, TEXT("Spawners completed"));
+
+			FTaskGraphInterface::Get().WaitUntilTasksComplete(Tasks);
+			//UE_LOG(LogTemp, Display, TEXT("Tasks completed"));
+			check(TasksExecutedNum == TasksNum);
+		}
+#endif
 	}
 
-#undef BENCHMARK
+	IMPLEMENT_SIMPLE_AUTOMATION_TEST(FPipeTest, "System.Core.Async.TaskGraph.Pipe", EAutomationTestFlags::ApplicationContextMask | EAutomationTestFlags::EngineFilter);
+
+	bool FPipeTest::RunTest(const FString& Parameters)
+	{
+		UE_BENCHMARK(5, PipeStressTest);
+
+		return true;
+	}
 }
+
 #endif //WITH_DEV_AUTOMATION_TESTS
diff --git a/Engine/Source/Runtime/Core/Public/Async/TaskGraphInterfaces.h b/Engine/Source/Runtime/Core/Public/Async/TaskGraphInterfaces.h
--- a/Engine/Source/Runtime/Core/Public/Async/TaskGraphInterfaces.h
+++ b/Engine/Source/Runtime/Core/Public/Async/TaskGraphInterfaces.h
@@ -778,6 +778,11 @@ public:
 		return RefCount;
 	}
 
+	uint32 GetRefCount() const
+	{
+		return ReferenceCount.GetValue();
+	}
+
 private:
 
 	/** Threadsafe list of subsequents for the event **/
