TaskTrace adding missing includes and forward declaration

diff --git a/Engine/Source/Runtime/Core/Public/Async/TaskTrace.h b/Engine/Source/Runtime/Core/Public/Async/TaskTrace.h
--- a/Engine/Source/Runtime/Core/Public/Async/TaskTrace.h
+++ b/Engine/Source/Runtime/Core/Public/Async/TaskTrace.h
@@ -3,6 +3,7 @@
 #pragma once
 
 #include "Trace/Config.h"
+#include "Containers/Array.h"
 
 #if UE_TRACE_ENABLED && !IS_PROGRAM && !UE_BUILD_SHIPPING
 #define UE_TASK_TRACE_ENABLED 1
@@ -10,6 +11,12 @@
 #define UE_TASK_TRACE_ENABLED 0
 #endif
 
+namespace ENamedThreads
+{
+	// Forward declare
+	enum Type : int32;
+}
+
 namespace TaskTrace
 {
 	using FId = uint32;
