Fix for disliked comment style in Clang

diff --git a/Engine/Source/Runtime/Core/Public/ProfilingDebugging/StallDetector.h b/Engine/Source/Runtime/Core/Public/ProfilingDebugging/StallDetector.h
--- a/Engine/Source/Runtime/Core/Public/ProfilingDebugging/StallDetector.h
+++ b/Engine/Source/Runtime/Core/Public/ProfilingDebugging/StallDetector.h
@@ -219,14 +219,14 @@ namespace UE
 }
 
 /**
-/* Counts, but does not report, a stall.
+* Counts, but does not report, a stall.
 **/
 #define SCOPE_STALL_COUNTER(InName, InBudgetSeconds) \
 static UE::FStallDetectorStats PREPROCESSOR_JOIN(StallDetectorStats, __LINE__) (TEXT(#InName), InBudgetSeconds, UE::EStallDetectorReportingMode::Disabled); \
 UE::FStallDetector PREPROCESSOR_JOIN(ScopeStallDetector, __LINE__)(PREPROCESSOR_JOIN(StallDetectorStats, __LINE__));
 
 /**
-/* Counts and sends a report of the first occurence of a stall.
+* Counts and sends a report of the first occurence of a stall.
 **/
 #define SCOPE_STALL_REPORTER(InName, InBudgetSeconds) \
 static UE::FStallDetectorStats PREPROCESSOR_JOIN(StallDetectorStats, __LINE__) (TEXT(#InName), InBudgetSeconds, UE::EStallDetectorReportingMode::First); \
