CompactBinary: Renamed Reference to Attachment to clarify the expected relationship between reference fields and package attachments

#rnx

diff --git a/Engine/Source/Developer/DerivedDataCache/Private/DerivedDataCache.cpp b/Engine/Source/Developer/DerivedDataCache/Private/DerivedDataCache.cpp
--- a/Engine/Source/Developer/DerivedDataCache/Private/DerivedDataCache.cpp
+++ b/Engine/Source/Developer/DerivedDataCache/Private/DerivedDataCache.cpp
@@ -828,9 +828,9 @@ void FCache::Get(
 
 	// Validate that the record was serialized in the expected format.
 	if ((MetaField && !MetaField.IsObject()) ||
-		(ValueField && !ValueField.IsReference()) ||
+		(ValueField && !ValueField.IsAttachment()) ||
 		(AttachmentsField && !AttachmentsField.IsArray()) ||
-		!Algo::AllOf(AttachmentsArray, &FCbField::IsReference))
+		!Algo::AllOf(AttachmentsArray, &FCbField::IsAttachment))
 	{
 		UE_LOG(LogDerivedDataCache, Verbose, TEXT("Cache: Get cache miss with invalid format for %s from '%.*s'"),
 			*TToString<96>(Key), Context.Len(), Context.GetData());
@@ -840,9 +840,9 @@ void FCache::Get(
 	// Check for existence of the value and attachments if they are being skipped.
 	{
 		TArray<FIoHash, TInlineAllocator<1>> KeysToCheck;
-		if (EnumHasAnyFlags(Policy, ECachePolicy::SkipValue) && ValueField.IsReference())
+		if (EnumHasAnyFlags(Policy, ECachePolicy::SkipValue) && ValueField.IsAttachment())
 		{
-			KeysToCheck.Add(ValueField.AsReference());
+			KeysToCheck.Add(ValueField.AsAttachment());
 			if (KeysToCheck.Last().IsZero())
 			{
 				KeysToCheck.Pop();
@@ -853,7 +853,7 @@ void FCache::Get(
 			KeysToCheck.Reserve(KeysToCheck.Num() + AttachmentsArray.Num());
 			for (FCbField AttachmentField : AttachmentsArray)
 			{
-				KeysToCheck.Add(AttachmentField.AsReference());
+				KeysToCheck.Add(AttachmentField.AsAttachment());
 			}
 		}
 		if (KeysToCheck.Num())
@@ -872,18 +872,18 @@ void FCache::Get(
 	}
 
 	// Read the value and attachments if they have been requested.
-	const auto GetContent = [&Key, Context](FCbField ReferenceField) -> FCbAttachment
+	const auto GetContent = [&Key, Context](FCbField AttachmentField) -> FCbAttachment
 	{
 		TArray<uint8> ContentData;
-		const FIoHash Hash = ReferenceField.AsReference();
-		check(!ReferenceField.HasError());
+		const FIoHash Hash = AttachmentField.AsAttachment();
+		check(!AttachmentField.HasError());
 		const FCacheAttachmentKey ContentKey(Key, Hash);
 		if (GetDerivedDataCacheRef().GetSynchronous(*MakeContentKey(ContentKey), ContentData, Context))
 		{
 			if (FIoHash::HashBuffer(MakeMemoryView(ContentData)) == Hash)
 			{
 				FSharedBuffer ContentBuffer = FSharedBuffer::Clone(MakeMemoryView(ContentData));
-				if (ReferenceField.IsCompactBinaryReference())
+				if (AttachmentField.IsCompactBinaryAttachment())
 				{
 					return FCbAttachment(FCbFieldRefIterator::MakeRange(MoveTemp(ContentBuffer)), Hash);
 				}
@@ -907,7 +907,7 @@ void FCache::Get(
 	};
 
 	FCbAttachment Value;
-	if (!EnumHasAnyFlags(Policy, ECachePolicy::SkipValue) && ValueField.IsReference() && !(Value = GetContent(ValueField)))
+	if (!EnumHasAnyFlags(Policy, ECachePolicy::SkipValue) && ValueField.IsAttachment() && !(Value = GetContent(ValueField)))
 	{
 		return;
 	}
@@ -932,7 +932,7 @@ void FCache::Get(
 		FCbPackage Package;
 		if (Value && Value.IsCompactBinary())
 		{
-			Package.SetObject(Value.AsCompactBinary().AsObjectRef(), ValueField.AsReference());
+			Package.SetObject(Value.AsCompactBinary().AsObjectRef(), ValueField.AsAttachment());
 		}
 		for (FCbAttachment Attachment : Attachments)
 		{
@@ -945,12 +945,12 @@ void FCache::Get(
 		// Object
 		if (Value.IsCompactBinary())
 		{
-			Params.Record.SetObject(Value.AsCompactBinary().AsObjectRef(), ValueField.AsReference());
+			Params.Record.SetObject(Value.AsCompactBinary().AsObjectRef(), ValueField.AsAttachment());
 		}
 		// Binary
 		else
 		{
-			Params.Record.SetBinary(Value.AsBinary(), ValueField.AsReference());
+			Params.Record.SetBinary(Value.AsBinary(), ValueField.AsAttachment());
 		}
 	}
 
@@ -1051,11 +1051,11 @@ void FCache::Put(
 	{
 		if (Value.IsCompactBinary())
 		{
-			Writer.Name("Value"_ASV).CompactBinaryReference(Value.GetHash());
+			Writer.Name("Value"_ASV).CompactBinaryAttachment(Value.GetHash());
 		}
 		else if (Value.IsBinary())
 		{
-			Writer.Name("Value"_ASV).BinaryReference(Value.GetHash());
+			Writer.Name("Value"_ASV).BinaryAttachment(Value.GetHash());
 		}
 
 		if (Record.GetType() == ECacheRecordType::Package)
@@ -1066,11 +1066,11 @@ void FCache::Put(
 			{
 				if (Attachment.IsCompactBinary())
 				{
-					Writer.CompactBinaryReference(Attachment.GetHash());
+					Writer.CompactBinaryAttachment(Attachment.GetHash());
 				}
 				else
 				{
-					Writer.BinaryReference(Attachment.GetHash());
+					Writer.BinaryAttachment(Attachment.GetHash());
 				}
 			}
 			Writer.EndArray();
diff --git a/Engine/Source/Runtime/Core/Private/Serialization/CompactBinary.cpp b/Engine/Source/Runtime/Core/Private/Serialization/CompactBinary.cpp
--- a/Engine/Source/Runtime/Core/Private/Serialization/CompactBinary.cpp
+++ b/Engine/Source/Runtime/Core/Private/Serialization/CompactBinary.cpp
@@ -65,10 +65,10 @@ void FCbFieldType::StaticAssertTypeConstants()
 	static_assert(!(BoolMask & (BoolBase ^ ECbFieldType::BoolTrue)), "BoolMask or BoolBase is invalid!");
 	static_assert(TypeMask == (BoolMask | (BoolBase ^ ECbFieldType::BoolTrue)), "BoolMask or BoolBase is invalid!");
 
-	static_assert(ReferenceBase == ECbFieldType::CompactBinaryReference, "ReferenceBase is invalid!");
-	static_assert((ReferenceMask & (AllFlags | ECbFieldType::BinaryReference)) == ECbFieldType::CompactBinaryReference, "ReferenceMask is invalid!");
-	static_assert(!(ReferenceMask & (ReferenceBase ^ ECbFieldType::BinaryReference)), "ReferenceMask or ReferenceBase is invalid!");
-	static_assert(TypeMask == (ReferenceMask | (ReferenceBase ^ ECbFieldType::BinaryReference)), "ReferenceMask or ReferenceBase is invalid!");
+	static_assert(AttachmentBase == ECbFieldType::CompactBinaryAttachment, "AttachmentBase is invalid!");
+	static_assert((AttachmentMask & (AllFlags | ECbFieldType::BinaryAttachment)) == ECbFieldType::CompactBinaryAttachment, "AttachmentMask is invalid!");
+	static_assert(!(AttachmentMask & (AttachmentBase ^ ECbFieldType::BinaryAttachment)), "AttachmentMask or AttachmentBase is invalid!");
+	static_assert(TypeMask == (AttachmentMask | (AttachmentBase ^ ECbFieldType::BinaryAttachment)), "AttachmentMask or AttachmentBase is invalid!");
 }
 
 ///////////////////////////////////////////////////////////////////////////////////////////////////////////////////////
@@ -254,9 +254,9 @@ bool FCbField::AsBool(const bool bDefault)
 	return (uint8(bIsBool) & uint8(LocalType) & 1) | ((!bIsBool) & bDefault);
 }
 
-FIoHash FCbField::AsCompactBinaryReference(const FIoHash& Default)
+FIoHash FCbField::AsCompactBinaryAttachment(const FIoHash& Default)
 {
-	if (FCbFieldType::IsCompactBinaryReference(Type))
+	if (FCbFieldType::IsCompactBinaryAttachment(Type))
 	{
 		Error = ECbFieldError::None;
 		return FIoHash(*static_cast<const FIoHash::ByteArray*>(Payload));
@@ -268,9 +268,9 @@ FIoHash FCbField::AsCompactBinaryReference(const FIoHash& Default)
 	}
 }
 
-FIoHash FCbField::AsBinaryReference(const FIoHash& Default)
+FIoHash FCbField::AsBinaryAttachment(const FIoHash& Default)
 {
-	if (FCbFieldType::IsBinaryReference(Type))
+	if (FCbFieldType::IsBinaryAttachment(Type))
 	{
 		Error = ECbFieldError::None;
 		return FIoHash(*static_cast<const FIoHash::ByteArray*>(Payload));
@@ -282,9 +282,9 @@ FIoHash FCbField::AsBinaryReference(const FIoHash& Default)
 	}
 }
 
-FIoHash FCbField::AsReference(const FIoHash& Default)
+FIoHash FCbField::AsAttachment(const FIoHash& Default)
 {
-	if (FCbFieldType::IsReference(Type))
+	if (FCbFieldType::IsAttachment(Type))
 	{
 		Error = ECbFieldError::None;
 		return FIoHash(*static_cast<const FIoHash::ByteArray*>(Payload));
@@ -416,8 +416,8 @@ uint64 FCbField::GetPayloadSize() const
 	case ECbFieldType::BoolFalse:
 	case ECbFieldType::BoolTrue:
 		return 0;
-	case ECbFieldType::CompactBinaryReference:
-	case ECbFieldType::BinaryReference:
+	case ECbFieldType::CompactBinaryAttachment:
+	case ECbFieldType::BinaryAttachment:
 	case ECbFieldType::Hash:
 		return 20;
 	case ECbFieldType::Uuid:
@@ -470,18 +470,18 @@ void FCbField::CopyTo(FArchive& Ar) const
 	Ar.Serialize(const_cast<void*>(Source.GetData()), static_cast<int64>(Source.GetSize()));
 }
 
-void FCbField::IterateReferences(FCbFieldVisitor Visitor) const
+void FCbField::IterateAttachments(FCbFieldVisitor Visitor) const
 {
 	switch (FCbFieldType::GetType(Type))
 	{
 	case ECbFieldType::Object:
 	case ECbFieldType::UniformObject:
-		return FCbObject::FromField(*this).IterateReferences(Visitor);
+		return FCbObject::FromField(*this).IterateAttachments(Visitor);
 	case ECbFieldType::Array:
 	case ECbFieldType::UniformArray:
-		return FCbArray::FromField(*this).IterateReferences(Visitor);
-	case ECbFieldType::CompactBinaryReference:
-	case ECbFieldType::BinaryReference:
+		return FCbArray::FromField(*this).IterateAttachments(Visitor);
+	case ECbFieldType::CompactBinaryAttachment:
+	case ECbFieldType::BinaryAttachment:
 		return Visitor(*this);
 	default:
 		return;
@@ -760,27 +760,27 @@ void TCbFieldIterator<FieldType>::CopyRangeTo(FArchive& Ar) const
 }
 
 template <typename FieldType>
-void TCbFieldIterator<FieldType>::IterateRangeReferences(FCbFieldVisitor Visitor) const
+void TCbFieldIterator<FieldType>::IterateRangeAttachments(FCbFieldVisitor Visitor) const
 {
-	// Always iterate over non-uniform ranges because we do not know if they contain a reference.
+	// Always iterate over non-uniform ranges because we do not know if they contain an attachment.
 	if (FCbFieldType::HasFieldType(FieldType::GetType()))
 	{
 		for (TCbFieldIterator It(*this); It; ++It)
 		{
-			if (FCbFieldType::MayContainReferences(It.GetType()))
+			if (FCbFieldType::MayContainAttachments(It.GetType()))
 			{
-				It.IterateReferences(Visitor);
+				It.IterateAttachments(Visitor);
 			}
 		}
 	}
-	// Only iterate over uniform ranges if the uniform type may contain a reference.
+	// Only iterate over uniform ranges if the uniform type may contain an attachment.
 	else
 	{
-		if (FCbFieldType::MayContainReferences(FieldType::GetType()))
+		if (FCbFieldType::MayContainAttachments(FieldType::GetType()))
 		{
 			for (TCbFieldIterator It(*this); It; ++It)
 			{
-				It.IterateReferences(Visitor);
+				It.IterateAttachments(Visitor);
 			}
 		}
 	}
diff --git a/Engine/Source/Runtime/Core/Private/Serialization/CompactBinaryPackage.cpp b/Engine/Source/Runtime/Core/Private/Serialization/CompactBinaryPackage.cpp
--- a/Engine/Source/Runtime/Core/Private/Serialization/CompactBinaryPackage.cpp
+++ b/Engine/Source/Runtime/Core/Private/Serialization/CompactBinaryPackage.cpp
@@ -95,9 +95,9 @@ void FCbAttachment::Load(FCbFieldRefIterator& Fields)
 		Buffer = FSharedBuffer::MakeView(View, Fields.GetOuterBuffer());
 		Buffer.MakeOwned();
 		++Fields;
-		Hash = Fields.AsReference();
+		Hash = Fields.AsAttachment();
 		checkf(!Fields.HasError(), TEXT("Attachments must be a non-empty binary value with a content hash."));
-		if (Fields.IsCompactBinaryReference())
+		if (Fields.IsCompactBinaryAttachment())
 		{
 			CompactBinary = FCbFieldIterator::MakeRange(Buffer);
 		}
@@ -130,9 +130,9 @@ void FCbAttachment::Load(FArchive& Ar, FCbBufferAllocator Allocator)
 				HashBuffer.SetNumUninitialized(int32(Size));
 				return FUniqueBuffer::MakeView(HashBuffer.GetData(), Size);
 			});
-		Hash = HashField.AsReference();
+		Hash = HashField.AsAttachment();
 		checkf(!HashField.HasError(), TEXT("Attachments must be a non-empty binary value with a content hash."));
-		if (HashField.IsCompactBinaryReference())
+		if (HashField.IsCompactBinaryAttachment())
 		{
 			CompactBinary = FCbFieldIterator::MakeRange(Buffer);
 		}
@@ -158,12 +158,12 @@ void FCbAttachment::Save(FCbWriter& Writer) const
 		{
 			Writer.Binary(AsBinary());
 		}
-		Writer.CompactBinaryReference(Hash);
+		Writer.CompactBinaryAttachment(Hash);
 	}
 	else if (Buffer.GetSize())
 	{
 		Writer.Binary(Buffer);
-		Writer.BinaryReference(Hash);
+		Writer.BinaryAttachment(Hash);
 	}
 	else // Null
 	{
@@ -256,12 +256,12 @@ const FCbAttachment* FCbPackage::FindAttachment(const FIoHash& Hash) const
 
 void FCbPackage::GatherAttachments(const FCbFieldIterator& Fields, FAttachmentResolver Resolver)
 {
-	Fields.IterateRangeReferences([this, &Resolver](FCbField Field)
+	Fields.IterateRangeAttachments([this, &Resolver](FCbField Field)
 		{
-			const FIoHash& Hash = Field.AsReference();
+			const FIoHash& Hash = Field.AsAttachment();
 			if (FSharedBuffer Buffer = Resolver(Hash))
 			{
-				if (Field.IsCompactBinaryReference())
+				if (Field.IsCompactBinaryAttachment())
 				{
 					AddAttachment(FCbAttachment(FCbFieldRefIterator::MakeRange(MoveTemp(Buffer)), Hash), &Resolver);
 				}
@@ -297,8 +297,8 @@ void FCbPackage::Load(FCbFieldRefIterator& Fields)
 			++Fields;
 			if (Object.CreateIterator())
 			{
-				ObjectHash = Fields.AsCompactBinaryReference();
-				checkf(!Fields.HasError(), TEXT("Object must be followed by a CompactBinaryReference with the object hash."));
+				ObjectHash = Fields.AsCompactBinaryAttachment();
+				checkf(!Fields.HasError(), TEXT("Object must be followed by a CompactBinaryAttachment with the object hash."));
 				++Fields;
 			}
 			else
@@ -333,9 +333,9 @@ void FCbPackage::Load(FArchive& Ar, FCbBufferAllocator Allocator)
 				FSharedBuffer Buffer = FSharedBuffer::MakeView(View, ValueField.GetOuterBuffer());
 				Buffer.MakeOwned();
 				FCbFieldRef HashField = LoadCompactBinary(Ar, StackAllocator);
-				const FIoHash& Hash = HashField.AsReference();
+				const FIoHash& Hash = HashField.AsAttachment();
 				checkf(!HashField.HasError(), TEXT("Attachments must be a non-empty binary value with a content hash."));
-				if (HashField.IsCompactBinaryReference())
+				if (HashField.IsCompactBinaryAttachment())
 				{
 					AddAttachment(FCbAttachment(FCbFieldRefIterator::MakeRange(MoveTemp(Buffer)), Hash));
 				}
@@ -353,8 +353,8 @@ void FCbPackage::Load(FArchive& Ar, FCbBufferAllocator Allocator)
 			if (Object.CreateIterator())
 			{
 				FCbFieldRef HashField = LoadCompactBinary(Ar, StackAllocator);
-				ObjectHash = HashField.AsCompactBinaryReference();
-				checkf(!HashField.HasError(), TEXT("Object must be followed by a CompactBinaryReference with the object hash."));
+				ObjectHash = HashField.AsCompactBinaryAttachment();
+				checkf(!HashField.HasError(), TEXT("Object must be followed by a CompactBinaryAttachment with the object hash."));
 			}
 			else
 			{
@@ -369,7 +369,7 @@ void FCbPackage::Save(FCbWriter& Writer) const
 	if (Object.CreateIterator())
 	{
 		Writer.Object(Object);
-		Writer.CompactBinaryReference(ObjectHash);
+		Writer.CompactBinaryAttachment(ObjectHash);
 	}
 	for (const FCbAttachment& Attachment : Attachments)
 	{
diff --git a/Engine/Source/Runtime/Core/Private/Serialization/CompactBinarySerialization.cpp b/Engine/Source/Runtime/Core/Private/Serialization/CompactBinarySerialization.cpp
--- a/Engine/Source/Runtime/Core/Private/Serialization/CompactBinarySerialization.cpp
+++ b/Engine/Source/Runtime/Core/Private/Serialization/CompactBinarySerialization.cpp
@@ -59,8 +59,8 @@ bool TryMeasureCompactBinary(FMemoryView View, ECbFieldType& OutType, uint64& Ou
 	case ECbFieldType::BoolFalse:
 	case ECbFieldType::BoolTrue:
 		break;
-	case ECbFieldType::CompactBinaryReference:
-	case ECbFieldType::BinaryReference:
+	case ECbFieldType::CompactBinaryAttachment:
+	case ECbFieldType::BinaryAttachment:
 	case ECbFieldType::Hash:
 		FixedSize = 20;
 		break;
diff --git a/Engine/Source/Runtime/Core/Private/Serialization/CompactBinaryValidation.cpp b/Engine/Source/Runtime/Core/Private/Serialization/CompactBinaryValidation.cpp
--- a/Engine/Source/Runtime/Core/Private/Serialization/CompactBinaryValidation.cpp
+++ b/Engine/Source/Runtime/Core/Private/Serialization/CompactBinaryValidation.cpp
@@ -353,8 +353,8 @@ static FCbField ValidateCbField(FMemoryView& View, ECbValidateMode Mode, ECbVali
 	case ECbFieldType::Float64:
 		ValidateCbFloat64(View, Mode, Error);
 		break;
-	case ECbFieldType::CompactBinaryReference:
-	case ECbFieldType::BinaryReference:
+	case ECbFieldType::CompactBinaryAttachment:
+	case ECbFieldType::BinaryAttachment:
 	case ECbFieldType::Hash:
 		ValidateFixedPayload(20);
 		break;
@@ -410,7 +410,7 @@ static FIoHash ValidateCbPackageAttachment(FCbField& Value, FMemoryView& View, E
 	{
 		if (FCbField HashField = ValidateCbPackageField(View, Mode, Error))
 		{
-			const FIoHash Hash = HashField.AsReference();
+			const FIoHash Hash = HashField.AsAttachment();
 			if (EnumHasAnyFlags(Mode, ECbValidateMode::Package))
 			{
 				if (HashField.HasError())
@@ -440,7 +440,7 @@ static FIoHash ValidateCbPackageObject(FCbField& Value, FMemoryView& View, ECbVa
 	}
 	else if (FCbField HashField = ValidateCbPackageField(View, Mode, Error))
 	{
-		const FIoHash Hash = HashField.AsReference();
+		const FIoHash Hash = HashField.AsAttachment();
 		if (EnumHasAnyFlags(Mode, ECbValidateMode::Package))
 		{
 			if (!Object.CreateIterator())
diff --git a/Engine/Source/Runtime/Core/Private/Serialization/CompactBinaryWriter.cpp b/Engine/Source/Runtime/Core/Private/Serialization/CompactBinaryWriter.cpp
--- a/Engine/Source/Runtime/Core/Private/Serialization/CompactBinaryWriter.cpp
+++ b/Engine/Source/Runtime/Core/Private/Serialization/CompactBinaryWriter.cpp
@@ -477,18 +477,18 @@ void FCbWriter::Bool(const bool bValue)
 	EndField(bValue ? ECbFieldType::BoolTrue : ECbFieldType::BoolFalse);
 }
 
-void FCbWriter::CompactBinaryReference(const FIoHash& Value)
+void FCbWriter::CompactBinaryAttachment(const FIoHash& Value)
 {
 	BeginField();
 	Data.Append(Value.GetBytes(), sizeof(decltype(Value.GetBytes())));
-	EndField(ECbFieldType::CompactBinaryReference);
+	EndField(ECbFieldType::CompactBinaryAttachment);
 }
 
-void FCbWriter::BinaryReference(const FIoHash& Value)
+void FCbWriter::BinaryAttachment(const FIoHash& Value)
 {
 	BeginField();
 	Data.Append(Value.GetBytes(), sizeof(decltype(Value.GetBytes())));
-	EndField(ECbFieldType::BinaryReference);
+	EndField(ECbFieldType::BinaryAttachment);
 }
 
 void FCbWriter::Hash(const FIoHash& Value)
diff --git a/Engine/Source/Runtime/Core/Private/Tests/Serialization/CompactBinaryPackageTest.cpp b/Engine/Source/Runtime/Core/Private/Tests/Serialization/CompactBinaryPackageTest.cpp
--- a/Engine/Source/Runtime/Core/Private/Tests/Serialization/CompactBinaryPackageTest.cpp
+++ b/Engine/Source/Runtime/Core/Private/Tests/Serialization/CompactBinaryPackageTest.cpp
@@ -362,7 +362,7 @@ bool FCbPackageTest::RunTest(const FString& Parameters)
 	FCbFieldRef Level3;
 	{
 		TCbWriter<256> Writer;
-		Writer.Name("Level4").BinaryReference(Level4Hash);
+		Writer.Name("Level4").BinaryAttachment(Level4Hash);
 		Level3 = Writer.Save();
 	}
 	const FIoHash Level3Hash = Level3.GetHash();
@@ -372,7 +372,7 @@ bool FCbPackageTest::RunTest(const FString& Parameters)
 		TCbWriter<256> Writer;
 		Writer.Name("Level3");
 		Writer.BeginArray();
-		Writer.CompactBinaryReference(Level3Hash);
+		Writer.CompactBinaryAttachment(Level3Hash);
 		Writer.EndArray();
 		Level2 = Writer.Save().AsArrayRef();
 	}
@@ -382,7 +382,7 @@ bool FCbPackageTest::RunTest(const FString& Parameters)
 	{
 		TCbWriter<256> Writer;
 		Writer.BeginObject();
-		Writer.Name("Level2").CompactBinaryReference(Level2Hash);
+		Writer.Name("Level2").CompactBinaryAttachment(Level2Hash);
 		Writer.EndObject();
 		Level1 = Writer.Save().AsObjectRef();
 	}
@@ -453,13 +453,13 @@ bool FCbPackageTest::RunTest(const FString& Parameters)
 	{
 		TCbWriter<384> Writer;
 		Writer.Binary(Level2.GetBuffer());
-		Writer.CompactBinaryReference(Level2Hash);
+		Writer.CompactBinaryAttachment(Level2Hash);
 		Writer.Binary(Level4);
-		Writer.BinaryReference(Level4Hash);
+		Writer.BinaryAttachment(Level4Hash);
 		Writer.Object(Level1);
-		Writer.CompactBinaryReference(Level1Hash);
+		Writer.CompactBinaryAttachment(Level1Hash);
 		Writer.Binary(Level3.GetBuffer());
-		Writer.CompactBinaryReference(Level3Hash);
+		Writer.CompactBinaryAttachment(Level3Hash);
 		Writer.Null();
 
 		FCbFieldRefIterator Fields = Writer.Save();
@@ -507,19 +507,19 @@ bool FCbPackageTest::RunTest(const FString& Parameters)
 		FCbFieldRefIterator Saved = Writer.Save();
 		TestTrue(TEXT("FCbPackage(OutOfOrder).Save()->Level1"), Saved.AsObjectRef().Equals(Level1));
 		++Saved;
-		TestEqual(TEXT("FCbPackage(OutOfOrder).Save()->Level1Hash"), Saved.AsCompactBinaryReference(), Level1Hash);
+		TestEqual(TEXT("FCbPackage(OutOfOrder).Save()->Level1Hash"), Saved.AsCompactBinaryAttachment(), Level1Hash);
 		++Saved;
 		TestTrue(TEXT("FCbPackage(OutOfOrder).Save()->Level2"), Saved.AsBinary().EqualBytes(Level2.GetView()));
 		++Saved;
-		TestEqual(TEXT("FCbPackage(OutOfOrder).Save()->Level2Hash"), Saved.AsCompactBinaryReference(), Level2Hash);
+		TestEqual(TEXT("FCbPackage(OutOfOrder).Save()->Level2Hash"), Saved.AsCompactBinaryAttachment(), Level2Hash);
 		++Saved;
 		TestTrue(TEXT("FCbPackage(OutOfOrder).Save()->Level3"), Saved.AsBinary().EqualBytes(Level3.GetView()));
 		++Saved;
-		TestEqual(TEXT("FCbPackage(OutOfOrder).Save()->Level3Hash"), Saved.AsCompactBinaryReference(), Level3Hash);
+		TestEqual(TEXT("FCbPackage(OutOfOrder).Save()->Level3Hash"), Saved.AsCompactBinaryAttachment(), Level3Hash);
 		++Saved;
 		TestTrue(TEXT("FCbPackage(OutOfOrder).Save()->Level4"), Saved.AsBinary().EqualBytes(Level4.GetView()));
 		++Saved;
-		TestEqual(TEXT("FCbPackage(OutOfOrder).Save()->Level4Hash"), Saved.AsBinaryReference(), Level4Hash);
+		TestEqual(TEXT("FCbPackage(OutOfOrder).Save()->Level4Hash"), Saved.AsBinaryAttachment(), Level4Hash);
 		++Saved;
 		TestTrue(TEXT("FCbPackage(OutOfOrder).Save()->Null"), Saved.IsNull());
 		++Saved;
diff --git a/Engine/Source/Runtime/Core/Private/Tests/Serialization/CompactBinarySerializationTest.cpp b/Engine/Source/Runtime/Core/Private/Tests/Serialization/CompactBinarySerializationTest.cpp
--- a/Engine/Source/Runtime/Core/Private/Tests/Serialization/CompactBinarySerializationTest.cpp
+++ b/Engine/Source/Runtime/Core/Private/Tests/Serialization/CompactBinarySerializationTest.cpp
@@ -115,8 +115,8 @@ bool FCbMeasureTest::RunTest(const FString& Parameters)
 	TestMeasure(TEXT("BoolFalse"), {uint8(EType::BoolFalse)}, true, 1, EType::BoolFalse);
 	TestMeasure(TEXT("BoolTrue"), {uint8(EType::BoolTrue)}, true, 1, EType::BoolTrue);
 
-	TestMeasure(TEXT("CompactBinaryReference"), {uint8(EType::CompactBinaryReference)}, true, 21, EType::CompactBinaryReference);
-	TestMeasure(TEXT("BinaryReference"), {uint8(EType::BinaryReference)}, true, 21, EType::BinaryReference);
+	TestMeasure(TEXT("CompactBinaryAttachment"), {uint8(EType::CompactBinaryAttachment)}, true, 21, EType::CompactBinaryAttachment);
+	TestMeasure(TEXT("BinaryAttachment"), {uint8(EType::BinaryAttachment)}, true, 21, EType::BinaryAttachment);
 
 	TestMeasure(TEXT("Hash"), {uint8(EType::Hash)}, true, 21, EType::Hash);
 	TestMeasure(TEXT("Uuid"), {uint8(EType::Uuid)}, true, 17, EType::Uuid);
diff --git a/Engine/Source/Runtime/Core/Private/Tests/Serialization/CompactBinaryTest.cpp b/Engine/Source/Runtime/Core/Private/Tests/Serialization/CompactBinaryTest.cpp
--- a/Engine/Source/Runtime/Core/Private/Tests/Serialization/CompactBinaryTest.cpp
+++ b/Engine/Source/Runtime/Core/Private/Tests/Serialization/CompactBinaryTest.cpp
@@ -60,8 +60,8 @@ UE_CBFIELD_TYPE_ACCESSOR(Float32, IsFloat, AsFloat, float);
 UE_CBFIELD_TYPE_ACCESSOR(Float64, IsFloat, AsDouble, double);
 UE_CBFIELD_TYPE_ACCESSOR(BoolFalse, IsBool, AsBool, bool);
 UE_CBFIELD_TYPE_ACCESSOR(BoolTrue, IsBool, AsBool, bool);
-UE_CBFIELD_TYPE_ACCESSOR(CompactBinaryReference, IsCompactBinaryReference, AsCompactBinaryReference, FIoHash);
-UE_CBFIELD_TYPE_ACCESSOR(BinaryReference, IsBinaryReference, AsBinaryReference, FIoHash);
+UE_CBFIELD_TYPE_ACCESSOR(CompactBinaryAttachment, IsCompactBinaryAttachment, AsCompactBinaryAttachment, FIoHash);
+UE_CBFIELD_TYPE_ACCESSOR(BinaryAttachment, IsBinaryAttachment, AsBinaryAttachment, FIoHash);
 UE_CBFIELD_TYPE_ACCESSOR(Hash, IsHash, AsHash, FIoHash);
 UE_CBFIELD_TYPE_ACCESSOR_TYPED(Uuid, IsUuid, AsUuid, FGuid, const FGuid&);
 UE_CBFIELD_TYPE_ACCESSOR(DateTime, IsDateTime, AsDateTimeTicks, int64);
@@ -941,55 +941,55 @@ bool FCbFieldBoolTest::RunTest(const FString& Parameters)
 	return true;
 }
 
-IMPLEMENT_CUSTOM_SIMPLE_AUTOMATION_TEST(FCbFieldCompactBinaryReferenceTest, FCbFieldTestBase, "System.Core.Serialization.CbField.CompactBinaryReference", CompactBinaryTestFlags)
-bool FCbFieldCompactBinaryReferenceTest::RunTest(const FString& Parameters)
+IMPLEMENT_CUSTOM_SIMPLE_AUTOMATION_TEST(FCbFieldCompactBinaryAttachmentTest, FCbFieldTestBase, "System.Core.Serialization.CbField.CompactBinaryAttachment", CompactBinaryTestFlags)
+bool FCbFieldCompactBinaryAttachmentTest::RunTest(const FString& Parameters)
 {
 	const FIoHash::ByteArray ZeroBytes{};
 	const FIoHash::ByteArray SequentialBytes{1, 2, 3, 4, 5, 6, 7, 8, 9, 10, 11, 12, 13, 14, 15, 16, 17, 18, 19, 20};
 
-	// Test FCbField(CompactBinaryReference, Zero)
-	TestField<ECbFieldType::CompactBinaryReference>(TEXT("CompactBinaryReference, Zero"), ZeroBytes);
+	// Test FCbField(CompactBinaryAttachment, Zero)
+	TestField<ECbFieldType::CompactBinaryAttachment>(TEXT("CompactBinaryAttachment, Zero"), ZeroBytes);
 
-	// Test FCbField(CompactBinaryReference, NonZero)
-	TestField<ECbFieldType::CompactBinaryReference>(TEXT("CompactBinaryReference, NonZero"), SequentialBytes, FIoHash(SequentialBytes));
+	// Test FCbField(CompactBinaryAttachment, NonZero)
+	TestField<ECbFieldType::CompactBinaryAttachment>(TEXT("CompactBinaryAttachment, NonZero"), SequentialBytes, FIoHash(SequentialBytes));
 
-	// Test FCbField(CompactBinaryReference, NonZero) AsReference
+	// Test FCbField(CompactBinaryAttachment, NonZero) AsAttachment
 	{
-		FCbField Field(SequentialBytes, ECbFieldType::CompactBinaryReference);
-		TestFieldAsType(TEXT("CompactBinaryReference, NonZero, AsReference"), Field, &FCbField::AsReference, FIoHash(SequentialBytes));
+		FCbField Field(SequentialBytes, ECbFieldType::CompactBinaryAttachment);
+		TestFieldAsType(TEXT("CompactBinaryAttachment, NonZero, AsAttachment"), Field, &FCbField::AsAttachment, FIoHash(SequentialBytes));
 	}
 
-	// Test FCbField(None) as CompactBinaryReference
+	// Test FCbField(None) as CompactBinaryAttachment
 	{
 		FCbField DefaultField;
-		TestFieldError<ECbFieldType::CompactBinaryReference>(TEXT("CompactBinaryReference, None"), DefaultField, ECbFieldError::TypeError, FIoHash(SequentialBytes));
+		TestFieldError<ECbFieldType::CompactBinaryAttachment>(TEXT("CompactBinaryAttachment, None"), DefaultField, ECbFieldError::TypeError, FIoHash(SequentialBytes));
 	}
 
 	return true;
 }
 
-IMPLEMENT_CUSTOM_SIMPLE_AUTOMATION_TEST(FCbFieldBinaryReferenceTest, FCbFieldTestBase, "System.Core.Serialization.CbField.BinaryReference", CompactBinaryTestFlags)
-bool FCbFieldBinaryReferenceTest::RunTest(const FString& Parameters)
+IMPLEMENT_CUSTOM_SIMPLE_AUTOMATION_TEST(FCbFieldBinaryAttachmentTest, FCbFieldTestBase, "System.Core.Serialization.CbField.BinaryAttachment", CompactBinaryTestFlags)
+bool FCbFieldBinaryAttachmentTest::RunTest(const FString& Parameters)
 {
 	const FIoHash::ByteArray ZeroBytes{};
 	const FIoHash::ByteArray SequentialBytes{1, 2, 3, 4, 5, 6, 7, 8, 9, 10, 11, 12, 13, 14, 15, 16, 17, 18, 19, 20};
 
-	// Test FCbField(BinaryReference, Zero)
-	TestField<ECbFieldType::BinaryReference>(TEXT("BinaryReference, Zero"), ZeroBytes);
+	// Test FCbField(BinaryAttachment, Zero)
+	TestField<ECbFieldType::BinaryAttachment>(TEXT("BinaryAttachment, Zero"), ZeroBytes);
 
-	// Test FCbField(BinaryReference, NonZero)
-	TestField<ECbFieldType::BinaryReference>(TEXT("BinaryReference, NonZero"), SequentialBytes, FIoHash(SequentialBytes));
+	// Test FCbField(BinaryAttachment, NonZero)
+	TestField<ECbFieldType::BinaryAttachment>(TEXT("BinaryAttachment, NonZero"), SequentialBytes, FIoHash(SequentialBytes));
 
-	// Test FCbField(BinaryReference, NonZero) AsAnyReference
+	// Test FCbField(BinaryAttachment, NonZero) AsAttachment
 	{
-		FCbField Field(SequentialBytes, ECbFieldType::BinaryReference);
-		TestFieldAsType(TEXT("BinaryReference, NonZero, AsReference"), Field, &FCbField::AsReference, FIoHash(SequentialBytes));
+		FCbField Field(SequentialBytes, ECbFieldType::BinaryAttachment);
+		TestFieldAsType(TEXT("BinaryAttachment, NonZero, AsAttachment"), Field, &FCbField::AsAttachment, FIoHash(SequentialBytes));
 	}
 
-	// Test FCbField(None) as BinaryReference
+	// Test FCbField(None) as BinaryAttachment
 	{
 		FCbField DefaultField;
-		TestFieldError<ECbFieldType::BinaryReference>(TEXT("BinaryReference, None"), DefaultField, ECbFieldError::TypeError, FIoHash(SequentialBytes));
+		TestFieldError<ECbFieldType::BinaryAttachment>(TEXT("BinaryAttachment, None"), DefaultField, ECbFieldError::TypeError, FIoHash(SequentialBytes));
 	}
 
 	return true;
@@ -1092,8 +1092,8 @@ bool FCbFieldTimeSpanTest::RunTest(const FString& Parameters)
 	return true;
 }
 
-IMPLEMENT_CUSTOM_SIMPLE_AUTOMATION_TEST(FCbFieldIterateReferencesTest, FCbFieldTestBase, "System.Core.Serialization.CbField.IterateReferences", CompactBinaryTestFlags)
-bool FCbFieldIterateReferencesTest::RunTest(const FString& Parameters)
+IMPLEMENT_CUSTOM_SIMPLE_AUTOMATION_TEST(FCbFieldIterateAttachmentsTest, FCbFieldTestBase, "System.Core.Serialization.CbField.IterateAttachments", CompactBinaryTestFlags)
+bool FCbFieldIterateAttachmentsTest::RunTest(const FString& Parameters)
 {
 	const auto MakeTestHash = [](uint32 Index) { return FIoHash::HashBuffer(&Index, sizeof(Index)); };
 
@@ -1102,10 +1102,10 @@ bool FCbFieldIterateReferencesTest::RunTest(const FString& Parameters)
 		FCbWriter Writer;
 
 		Writer.Name("IgnoredTypeInRoot").Hash(MakeTestHash(100));
-		Writer.CompactBinaryReference(MakeTestHash(0));
-		Writer.BinaryReference(MakeTestHash(1));
-		Writer.Name("RefInRoot").CompactBinaryReference(MakeTestHash(2));
-		Writer.Name("BinRefInRoot").BinaryReference(MakeTestHash(3));
+		Writer.CompactBinaryAttachment(MakeTestHash(0));
+		Writer.BinaryAttachment(MakeTestHash(1));
+		Writer.Name("CbAttachmentInRoot").CompactBinaryAttachment(MakeTestHash(2));
+		Writer.Name("BinAttachmentInRoot").BinaryAttachment(MakeTestHash(3));
 
 		// Uniform array of type to ignore.
 		Writer.BeginArray();
@@ -1114,23 +1114,23 @@ bool FCbFieldIterateReferencesTest::RunTest(const FString& Parameters)
 			Writer << 2;
 		}
 		Writer.EndArray();
-		// Uniform array of binary references.
+		// Uniform array of binary attachments.
 		Writer.BeginArray();
 		{
-			Writer.BinaryReference(MakeTestHash(4));
-			Writer.BinaryReference(MakeTestHash(5));
+			Writer.BinaryAttachment(MakeTestHash(4));
+			Writer.BinaryAttachment(MakeTestHash(5));
 		}
 		Writer.EndArray();
 		// Uniform array of uniform arrays.
 		Writer.BeginArray();
 		{
 			Writer.BeginArray();
-			Writer.BinaryReference(MakeTestHash(6));
-			Writer.BinaryReference(MakeTestHash(7));
+			Writer.BinaryAttachment(MakeTestHash(6));
+			Writer.BinaryAttachment(MakeTestHash(7));
 			Writer.EndArray();
 			Writer.BeginArray();
-			Writer.BinaryReference(MakeTestHash(8));
-			Writer.BinaryReference(MakeTestHash(9));
+			Writer.BinaryAttachment(MakeTestHash(8));
+			Writer.BinaryAttachment(MakeTestHash(9));
 			Writer.EndArray();
 		}
 		Writer.EndArray();
@@ -1142,7 +1142,7 @@ bool FCbFieldIterateReferencesTest::RunTest(const FString& Parameters)
 			Writer << false;
 			Writer.EndArray();
 			Writer.BeginArray();
-			Writer.CompactBinaryReference(MakeTestHash(10));
+			Writer.CompactBinaryAttachment(MakeTestHash(10));
 			Writer << false;
 			Writer.EndArray();
 		}
@@ -1151,12 +1151,12 @@ bool FCbFieldIterateReferencesTest::RunTest(const FString& Parameters)
 		Writer.BeginArray();
 		{
 			Writer.BeginObject();
-			Writer.Name("RefInUniObjInUniObj1").CompactBinaryReference(MakeTestHash(11));
-			Writer.Name("RefInUniObjInUniObj2").CompactBinaryReference(MakeTestHash(12));
+			Writer.Name("CbAttachmentInUniObjInUniObj1").CompactBinaryAttachment(MakeTestHash(11));
+			Writer.Name("CbAttachmentInUniObjInUniObj2").CompactBinaryAttachment(MakeTestHash(12));
 			Writer.EndObject();
 			Writer.BeginObject();
-			Writer.Name("RefInUniObjInUniObj3").CompactBinaryReference(MakeTestHash(13));
-			Writer.Name("RefInUniObjInUniObj4").CompactBinaryReference(MakeTestHash(14));
+			Writer.Name("CbAttachmentInUniObjInUniObj3").CompactBinaryAttachment(MakeTestHash(13));
+			Writer.Name("CbAttachmentInUniObjInUniObj4").CompactBinaryAttachment(MakeTestHash(14));
 			Writer.EndObject();
 		}
 		Writer.EndArray();
@@ -1168,7 +1168,7 @@ bool FCbFieldIterateReferencesTest::RunTest(const FString& Parameters)
 			Writer << "Bool" << false;
 			Writer.EndObject();
 			Writer.BeginObject();
-			Writer.Name("RefInNonUniObjInUniObj").CompactBinaryReference(MakeTestHash(15));
+			Writer.Name("CbAttachmentInNonUniObjInUniObj").CompactBinaryAttachment(MakeTestHash(15));
 			Writer << "Bool" << false;
 			Writer.EndObject();
 		}
@@ -1181,11 +1181,11 @@ bool FCbFieldIterateReferencesTest::RunTest(const FString& Parameters)
 			Writer << "Int2" << 2;
 		}
 		Writer.EndObject();
-		// Uniform object of binary references.
+		// Uniform object of binary attachments.
 		Writer.BeginObject();
 		{
-			Writer.Name("BinRefInUniObj1").BinaryReference(MakeTestHash(16));
-			Writer.Name("BinRefInUniObj2").BinaryReference(MakeTestHash(17));
+			Writer.Name("BinAttachmentInUniObj1").BinaryAttachment(MakeTestHash(16));
+			Writer.Name("BinAttachmentInUniObj2").BinaryAttachment(MakeTestHash(17));
 		}
 		Writer.EndObject();
 		// Uniform object of uniform arrays.
@@ -1193,13 +1193,13 @@ bool FCbFieldIterateReferencesTest::RunTest(const FString& Parameters)
 		{
 			Writer.Name("Array1");
 			Writer.BeginArray();
-			Writer.BinaryReference(MakeTestHash(18));
-			Writer.BinaryReference(MakeTestHash(19));
+			Writer.BinaryAttachment(MakeTestHash(18));
+			Writer.BinaryAttachment(MakeTestHash(19));
 			Writer.EndArray();
 			Writer.Name("Array2");
 			Writer.BeginArray();
-			Writer.BinaryReference(MakeTestHash(20));
-			Writer.BinaryReference(MakeTestHash(21));
+			Writer.BinaryAttachment(MakeTestHash(20));
+			Writer.BinaryAttachment(MakeTestHash(21));
 			Writer.EndArray();
 		}
 		Writer.EndObject();
@@ -1213,7 +1213,7 @@ bool FCbFieldIterateReferencesTest::RunTest(const FString& Parameters)
 			Writer.EndArray();
 			Writer.Name("Array2");
 			Writer.BeginArray();
-			Writer.CompactBinaryReference(MakeTestHash(22));
+			Writer.CompactBinaryAttachment(MakeTestHash(22));
 			Writer << false;
 			Writer.EndArray();
 		}
@@ -1223,13 +1223,13 @@ bool FCbFieldIterateReferencesTest::RunTest(const FString& Parameters)
 		{
 			Writer.Name("Object1");
 			Writer.BeginObject();
-			Writer.Name("RefInUniObjInUniObj1").CompactBinaryReference(MakeTestHash(23));
-			Writer.Name("RefInUniObjInUniObj2").CompactBinaryReference(MakeTestHash(24));
+			Writer.Name("CbAttachmentInUniObjInUniObj1").CompactBinaryAttachment(MakeTestHash(23));
+			Writer.Name("CbAttachmentInUniObjInUniObj2").CompactBinaryAttachment(MakeTestHash(24));
 			Writer.EndObject();
 			Writer.Name("Object2");
 			Writer.BeginObject();
-			Writer.Name("RefInUniObjInUniObj3").CompactBinaryReference(MakeTestHash(25));
-			Writer.Name("RefInUniObjInUniObj4").CompactBinaryReference(MakeTestHash(26));
+			Writer.Name("CbAttachmentInUniObjInUniObj3").CompactBinaryAttachment(MakeTestHash(25));
+			Writer.Name("CbAttachmentInUniObjInUniObj4").CompactBinaryAttachment(MakeTestHash(26));
 			Writer.EndObject();
 		}
 		Writer.EndObject();
@@ -1243,7 +1243,7 @@ bool FCbFieldIterateReferencesTest::RunTest(const FString& Parameters)
 			Writer.EndObject();
 			Writer.Name("Object2");
 			Writer.BeginObject();
-			Writer.Name("RefInNonUniObjInUniObj").CompactBinaryReference(MakeTestHash(27));
+			Writer.Name("CbAttachmentInNonUniObjInUniObj").CompactBinaryAttachment(MakeTestHash(27));
 			Writer << "Bool" << false;
 			Writer.EndObject();
 		}
@@ -1252,16 +1252,16 @@ bool FCbFieldIterateReferencesTest::RunTest(const FString& Parameters)
 		Fields = Writer.Save();
 	}
 
-	TestEqual(TEXT("FCbField::IterateReferences Validate"), ValidateCompactBinaryRange(Fields.GetBuffer(), ECbValidateMode::All), ECbValidateError::None);
+	TestEqual(TEXT("FCbField::IterateAttachments Validate"), ValidateCompactBinaryRange(Fields.GetBuffer(), ECbValidateMode::All), ECbValidateError::None);
 
-	uint32 ReferenceIndex = 0;
-	Fields.IterateRangeReferences([this, &ReferenceIndex, &MakeTestHash](FCbField Field)
+	uint32 AttachmentIndex = 0;
+	Fields.IterateRangeAttachments([this, &AttachmentIndex, &MakeTestHash](FCbField Field)
 		{
-			TestTrue(FString::Printf(TEXT("FCbField::IterateReferences(%u)::IsAnyReference"), ReferenceIndex), Field.IsReference());
-			TestEqual(FString::Printf(TEXT("FCbField::IterateReferences(%u)"), ReferenceIndex), Field.AsCompactBinaryReference(Field.AsBinaryReference()), MakeTestHash(ReferenceIndex));
-			++ReferenceIndex;
+			TestTrue(FString::Printf(TEXT("FCbField::IterateAttachments(%u)::IsAttachment"), AttachmentIndex), Field.IsAttachment());
+			TestEqual(FString::Printf(TEXT("FCbField::IterateAttachments(%u)"), AttachmentIndex), Field.AsAttachment(), MakeTestHash(AttachmentIndex));
+			++AttachmentIndex;
 		});
-	TestEqual(TEXT("FCbField::IterateReferences"), ReferenceIndex, 28);
+	TestEqual(TEXT("FCbField::IterateAttachments"), AttachmentIndex, 28);
 
 	return true;
 }
diff --git a/Engine/Source/Runtime/Core/Private/Tests/Serialization/CompactBinaryValidationTest.cpp b/Engine/Source/Runtime/Core/Private/Tests/Serialization/CompactBinaryValidationTest.cpp
--- a/Engine/Source/Runtime/Core/Private/Tests/Serialization/CompactBinaryValidationTest.cpp
+++ b/Engine/Source/Runtime/Core/Private/Tests/Serialization/CompactBinaryValidationTest.cpp
@@ -140,15 +140,15 @@ bool FCbValidateTest::RunTest(const FString& Parameters)
 	TestEqual(TEXT("ValidateCompactBinary(Valid, BoolFalse)"), Validate({uint8(ECbFieldType::BoolFalse)}), ECbValidateError::None);
 	TestEqual(TEXT("ValidateCompactBinary(Valid, BoolTrue)"), Validate({uint8(ECbFieldType::BoolTrue)}), ECbValidateError::None);
 
-	TestEqual(TEXT("ValidateCompactBinary(Valid, CompactBinaryReference)"), Validate({uint8(ECbFieldType::CompactBinaryReference), 0, 0, 0, 0, 0, 0, 0, 0, 0, 0, 0, 0, 0, 0, 0, 0, 0, 0, 0, 0}), ECbValidateError::None);
-	TestEqual(TEXT("ValidateCompactBinary(Valid, CompactBinaryReference, NoType)"), Validate({0, 0, 0, 0, 0, 0, 0, 0, 0, 0, 0, 0, 0, 0, 0, 0, 0, 0, 0, 0}, ECbFieldType::CompactBinaryReference), ECbValidateError::None);
-	TestEqual(TEXT("ValidateCompactBinary(OutOfBounds, CompactBinaryReference)"), Validate({uint8(ECbFieldType::CompactBinaryReference), 0, 0, 0, 0, 0, 0, 0, 0, 0, 0, 0, 0, 0, 0, 0, 0, 0, 0, 0}), ECbValidateError::OutOfBounds);
-	TestEqual(TEXT("ValidateCompactBinary(OutOfBounds, CompactBinaryReference, NoType)"), Validate({0, 0, 0, 0, 0, 0, 0, 0, 0, 0, 0, 0, 0, 0, 0, 0, 0, 0, 0}, ECbFieldType::CompactBinaryReference), ECbValidateError::OutOfBounds);
+	TestEqual(TEXT("ValidateCompactBinary(Valid, CompactBinaryAttachment)"), Validate({uint8(ECbFieldType::CompactBinaryAttachment), 0, 0, 0, 0, 0, 0, 0, 0, 0, 0, 0, 0, 0, 0, 0, 0, 0, 0, 0, 0}), ECbValidateError::None);
+	TestEqual(TEXT("ValidateCompactBinary(Valid, CompactBinaryAttachment, NoType)"), Validate({0, 0, 0, 0, 0, 0, 0, 0, 0, 0, 0, 0, 0, 0, 0, 0, 0, 0, 0, 0}, ECbFieldType::CompactBinaryAttachment), ECbValidateError::None);
+	TestEqual(TEXT("ValidateCompactBinary(OutOfBounds, CompactBinaryAttachment)"), Validate({uint8(ECbFieldType::CompactBinaryAttachment), 0, 0, 0, 0, 0, 0, 0, 0, 0, 0, 0, 0, 0, 0, 0, 0, 0, 0, 0}), ECbValidateError::OutOfBounds);
+	TestEqual(TEXT("ValidateCompactBinary(OutOfBounds, CompactBinaryAttachment, NoType)"), Validate({0, 0, 0, 0, 0, 0, 0, 0, 0, 0, 0, 0, 0, 0, 0, 0, 0, 0, 0}, ECbFieldType::CompactBinaryAttachment), ECbValidateError::OutOfBounds);
 
-	TestEqual(TEXT("ValidateCompactBinary(Valid, BinaryReference)"), Validate({uint8(ECbFieldType::BinaryReference), 0, 0, 0, 0, 0, 0, 0, 0, 0, 0, 0, 0, 0, 0, 0, 0, 0, 0, 0, 0}), ECbValidateError::None);
-	TestEqual(TEXT("ValidateCompactBinary(Valid, BinaryReference, NoType)"), Validate({0, 0, 0, 0, 0, 0, 0, 0, 0, 0, 0, 0, 0, 0, 0, 0, 0, 0, 0, 0}, ECbFieldType::BinaryReference), ECbValidateError::None);
-	TestEqual(TEXT("ValidateCompactBinary(OutOfBounds, BinaryReference)"), Validate({uint8(ECbFieldType::BinaryReference), 0, 0, 0, 0, 0, 0, 0, 0, 0, 0, 0, 0, 0, 0, 0, 0, 0, 0, 0}), ECbValidateError::OutOfBounds);
-	TestEqual(TEXT("ValidateCompactBinary(OutOfBounds, BinaryReference, NoType)"), Validate({0, 0, 0, 0, 0, 0, 0, 0, 0, 0, 0, 0, 0, 0, 0, 0, 0, 0, 0}, ECbFieldType::BinaryReference), ECbValidateError::OutOfBounds);
+	TestEqual(TEXT("ValidateCompactBinary(Valid, BinaryAttachment)"), Validate({uint8(ECbFieldType::BinaryAttachment), 0, 0, 0, 0, 0, 0, 0, 0, 0, 0, 0, 0, 0, 0, 0, 0, 0, 0, 0, 0}), ECbValidateError::None);
+	TestEqual(TEXT("ValidateCompactBinary(Valid, BinaryAttachment, NoType)"), Validate({0, 0, 0, 0, 0, 0, 0, 0, 0, 0, 0, 0, 0, 0, 0, 0, 0, 0, 0, 0}, ECbFieldType::BinaryAttachment), ECbValidateError::None);
+	TestEqual(TEXT("ValidateCompactBinary(OutOfBounds, BinaryAttachment)"), Validate({uint8(ECbFieldType::BinaryAttachment), 0, 0, 0, 0, 0, 0, 0, 0, 0, 0, 0, 0, 0, 0, 0, 0, 0, 0, 0}), ECbValidateError::OutOfBounds);
+	TestEqual(TEXT("ValidateCompactBinary(OutOfBounds, BinaryAttachment, NoType)"), Validate({0, 0, 0, 0, 0, 0, 0, 0, 0, 0, 0, 0, 0, 0, 0, 0, 0, 0, 0}, ECbFieldType::BinaryAttachment), ECbValidateError::OutOfBounds);
 
 	TestEqual(TEXT("ValidateCompactBinary(Valid, Hash)"), Validate({uint8(ECbFieldType::Hash), 0, 0, 0, 0, 0, 0, 0, 0, 0, 0, 0, 0, 0, 0, 0, 0, 0, 0, 0, 0}), ECbValidateError::None);
 	TestEqual(TEXT("ValidateCompactBinary(Valid, Hash, NoType)"), Validate({0, 0, 0, 0, 0, 0, 0, 0, 0, 0, 0, 0, 0, 0, 0, 0, 0, 0, 0, 0}, ECbFieldType::Hash), ECbValidateError::None);
@@ -351,7 +351,7 @@ bool FCbValidateAttachmentTest::RunTest(const FString& Parameters)
 	{
 		FBufferArchive Buffer;
 		FCbWriter Writer;
-		Writer.CompactBinaryReference(FIoHash::HashBuffer(BinaryView));
+		Writer.CompactBinaryAttachment(FIoHash::HashBuffer(BinaryView));
 		Writer.Binary(BinaryView);
 		Writer.Save(Buffer);
 		TestEqual(TEXT("ValidateCompactBinaryAttachment(HashBeforeValue)"), ValidateCompactBinaryAttachment(MakeMemoryView(Buffer), ECbValidateMode::All), ECbValidateError::InvalidPackageFormat | ECbValidateError::Padding);
@@ -361,7 +361,7 @@ bool FCbValidateAttachmentTest::RunTest(const FString& Parameters)
 		FBufferArchive Buffer;
 		FCbWriter Writer;
 		Writer.Name("Name"_ASV).Binary(BinaryView);
-		Writer.CompactBinaryReference(FIoHash::HashBuffer(BinaryView));
+		Writer.CompactBinaryAttachment(FIoHash::HashBuffer(BinaryView));
 		Writer.Save(Buffer);
 		TestEqual(TEXT("ValidateCompactBinaryAttachment(NameOnValue)"), ValidateCompactBinaryAttachment(MakeMemoryView(Buffer), ECbValidateMode::All), ECbValidateError::InvalidPackageFormat);
 		TestEqual(TEXT("ValidateCompactBinaryAttachment(NameOnValue, Mode)"), ValidateCompactBinaryAttachment(MakeMemoryView(Buffer), ECbValidateMode::All & ~ECbValidateMode::Package), ECbValidateError::None);
@@ -370,7 +370,7 @@ bool FCbValidateAttachmentTest::RunTest(const FString& Parameters)
 		FBufferArchive Buffer;
 		FCbWriter Writer;
 		Writer.Binary(BinaryView);
-		Writer.Name("Name"_ASV).CompactBinaryReference(FIoHash::HashBuffer(BinaryView));
+		Writer.Name("Name"_ASV).CompactBinaryAttachment(FIoHash::HashBuffer(BinaryView));
 		Writer.Save(Buffer);
 		TestEqual(TEXT("ValidateCompactBinaryAttachment(NameOnHash)"), ValidateCompactBinaryAttachment(MakeMemoryView(Buffer), ECbValidateMode::All), ECbValidateError::InvalidPackageFormat);
 		TestEqual(TEXT("ValidateCompactBinaryAttachment(NameOnHash, Mode)"), ValidateCompactBinaryAttachment(MakeMemoryView(Buffer), ECbValidateMode::All & ~ECbValidateMode::Package), ECbValidateError::None);
@@ -401,7 +401,7 @@ bool FCbValidatePackageTest::RunTest(const FString& Parameters)
 		Writer.BeginObject();
 		Writer << "Field" << 42;
 		Writer.EndObject();
-		Writer.CompactBinaryReference(Writer.Save().GetHash());
+		Writer.CompactBinaryAttachment(Writer.Save().GetHash());
 		Writer.Null();
 		Writer.Save(Buffer);
 		TestEqual(TEXT("ValidateCompactBinaryPackage(Object)"), ValidateCompactBinaryPackage(MakeMemoryView(Buffer), ECbValidateMode::All), ECbValidateError::None);
@@ -429,10 +429,10 @@ bool FCbValidatePackageTest::RunTest(const FString& Parameters)
 		Writer.BeginObject();
 		Writer << "Field" << 42;
 		Writer.EndObject();
-		Writer.CompactBinaryReference(Writer.Save().GetHash());
+		Writer.CompactBinaryAttachment(Writer.Save().GetHash());
 		const uint64 AttachmentOffset = Writer.GetSaveSize();
 		Writer.Binary(BinaryView);
-		Writer.CompactBinaryReference(FIoHash::HashBuffer(BinaryView));
+		Writer.CompactBinaryAttachment(FIoHash::HashBuffer(BinaryView));
 		Writer.Null();
 		Writer.Save(Buffer);
 		TestEqual(TEXT("ValidateCompactBinaryPackage(Object+Attachment)"), ValidateCompactBinaryPackage(MakeMemoryView(Buffer), ECbValidateMode::All), ECbValidateError::None);
@@ -465,7 +465,7 @@ bool FCbValidatePackageTest::RunTest(const FString& Parameters)
 		Writer.BeginObject();
 		Writer << "Field" << 42;
 		Writer.EndObject();
-		Writer.CompactBinaryReference(Writer.Save().GetHash());
+		Writer.CompactBinaryAttachment(Writer.Save().GetHash());
 		Writer.Null();
 		Writer.Save(Buffer);
 		TestEqual(TEXT("ValidateCompactBinaryPackage(NameOnObject)"), ValidateCompactBinaryPackage(MakeMemoryView(Buffer), ECbValidateMode::All), ECbValidateError::InvalidPackageFormat);
@@ -479,7 +479,7 @@ bool FCbValidatePackageTest::RunTest(const FString& Parameters)
 		Writer.EndObject();
 		const FIoHash ObjectHash = Writer.Save().GetHash();
 		Writer.Name("Name"_ASV);
-		Writer.CompactBinaryReference(ObjectHash);
+		Writer.CompactBinaryAttachment(ObjectHash);
 		Writer.Null();
 		Writer.Save(Buffer);
 		TestEqual(TEXT("ValidateCompactBinaryPackage(NameOnHash)"), ValidateCompactBinaryPackage(MakeMemoryView(Buffer), ECbValidateMode::All), ECbValidateError::InvalidPackageFormat);
@@ -491,9 +491,9 @@ bool FCbValidatePackageTest::RunTest(const FString& Parameters)
 		Writer.BeginObject();
 		Writer << "Field" << 42;
 		Writer.EndObject();
-		Writer.CompactBinaryReference(Writer.Save().GetHash());
+		Writer.CompactBinaryAttachment(Writer.Save().GetHash());
 		Writer.Name("Name"_ASV).Binary(BinaryView);
-		Writer.CompactBinaryReference(FIoHash::HashBuffer(BinaryView));
+		Writer.CompactBinaryAttachment(FIoHash::HashBuffer(BinaryView));
 		Writer.Null();
 		Writer.Save(Buffer);
 		TestEqual(TEXT("ValidateCompactBinaryPackage(NameOnAttachment)"), ValidateCompactBinaryPackage(MakeMemoryView(Buffer), ECbValidateMode::All), ECbValidateError::InvalidPackageFormat);
@@ -506,11 +506,11 @@ bool FCbValidatePackageTest::RunTest(const FString& Parameters)
 		Writer << "Field" << 42;
 		Writer.EndObject();
 		const FIoHash ObjectHash = Writer.Save().GetHash();
-		Writer.CompactBinaryReference(ObjectHash);
+		Writer.CompactBinaryAttachment(ObjectHash);
 		Writer.BeginObject();
 		Writer << "Field" << 42;
 		Writer.EndObject();
-		Writer.CompactBinaryReference(ObjectHash);
+		Writer.CompactBinaryAttachment(ObjectHash);
 		Writer.Null();
 		Writer.Save(Buffer);
 		TestEqual(TEXT("ValidateCompactBinaryPackage(MultipleObjects)"), ValidateCompactBinaryPackage(MakeMemoryView(Buffer), ECbValidateMode::All), ECbValidateError::MultiplePackageObjects);
@@ -520,9 +520,9 @@ bool FCbValidatePackageTest::RunTest(const FString& Parameters)
 		FBufferArchive Buffer;
 		TCbWriter<256> Writer;
 		Writer.Binary(BinaryView);
-		Writer.CompactBinaryReference(FIoHash::HashBuffer(BinaryView));
+		Writer.CompactBinaryAttachment(FIoHash::HashBuffer(BinaryView));
 		Writer.Binary(BinaryView);
-		Writer.CompactBinaryReference(FIoHash::HashBuffer(BinaryView));
+		Writer.CompactBinaryAttachment(FIoHash::HashBuffer(BinaryView));
 		Writer.Null();
 		Writer.Save(Buffer);
 		TestEqual(TEXT("ValidateCompactBinaryPackage(DuplicateAttachments)"), ValidateCompactBinaryPackage(MakeMemoryView(Buffer), ECbValidateMode::All), ECbValidateError::DuplicateAttachments);
@@ -533,7 +533,7 @@ bool FCbValidatePackageTest::RunTest(const FString& Parameters)
 		TCbWriter<256> Writer;
 		Writer.BeginObject();
 		Writer.EndObject();
-		Writer.CompactBinaryReference(FCbObject().GetHash());
+		Writer.CompactBinaryAttachment(FCbObject().GetHash());
 		Writer.Null();
 		Writer.Save(Buffer);
 		TestEqual(TEXT("ValidateCompactBinaryPackage(NullObject)"), ValidateCompactBinaryPackage(MakeMemoryView(Buffer), ECbValidateMode::All), ECbValidateError::NullPackageObject);
diff --git a/Engine/Source/Runtime/Core/Private/Tests/Serialization/CompactBinaryWriterTest.cpp b/Engine/Source/Runtime/Core/Private/Tests/Serialization/CompactBinaryWriterTest.cpp
--- a/Engine/Source/Runtime/Core/Private/Tests/Serialization/CompactBinaryWriterTest.cpp
+++ b/Engine/Source/Runtime/Core/Private/Tests/Serialization/CompactBinaryWriterTest.cpp
@@ -563,8 +563,8 @@ bool FCbWriterBoolTest::RunTest(const FString& Parameters)
 
 ///////////////////////////////////////////////////////////////////////////////////////////////////////////////////////
 
-IMPLEMENT_SIMPLE_AUTOMATION_TEST(FCbWriterCompactBinaryReferenceTest, "System.Core.Serialization.CbWriter.CompactBinaryReference", CompactBinaryWriterTestFlags)
-bool FCbWriterCompactBinaryReferenceTest::RunTest(const FString& Parameters)
+IMPLEMENT_SIMPLE_AUTOMATION_TEST(FCbWriterCompactBinaryAttachmentTest, "System.Core.Serialization.CbWriter.CompactBinaryAttachment", CompactBinaryWriterTestFlags)
+bool FCbWriterCompactBinaryAttachmentTest::RunTest(const FString& Parameters)
 {
 	TCbWriter<256> Writer;
 
@@ -574,17 +574,17 @@ bool FCbWriterCompactBinaryReferenceTest::RunTest(const FString& Parameters)
 	const FIoHash Values[] = { FIoHash(ZeroBytes), FIoHash(SequentialBytes) };
 	for (const FIoHash& Value : Values)
 	{
-		Writer.CompactBinaryReference(Value);
+		Writer.CompactBinaryAttachment(Value);
 	}
 
 	FCbFieldRefIterator Fields = Writer.Save();
-	if (TestEqual(TEXT("FCbWriter(CompactBinaryReference) Validate"), ValidateCompactBinaryRange(Fields.GetBuffer(), ECbValidateMode::All), ECbValidateError::None))
+	if (TestEqual(TEXT("FCbWriter(CompactBinaryAttachment) Validate"), ValidateCompactBinaryRange(Fields.GetBuffer(), ECbValidateMode::All), ECbValidateError::None))
 	{
 		const FIoHash* CheckValue = Values;
 		for (FCbField Field : Fields)
 		{
-			TestEqual(TEXT("FCbWriter(CompactBinaryReference).AsCompactBinaryReference()"), Field.AsCompactBinaryReference(), *CheckValue++);
-			TestFalse(TEXT("FCbWriter(CompactBinaryReference) Error"), Field.HasError());
+			TestEqual(TEXT("FCbWriter(CompactBinaryAttachment).AsCompactBinaryAttachment()"), Field.AsCompactBinaryAttachment(), *CheckValue++);
+			TestFalse(TEXT("FCbWriter(CompactBinaryAttachment) Error"), Field.HasError());
 		}
 	}
 
@@ -593,8 +593,8 @@ bool FCbWriterCompactBinaryReferenceTest::RunTest(const FString& Parameters)
 
 ///////////////////////////////////////////////////////////////////////////////////////////////////////////////////////
 
-IMPLEMENT_SIMPLE_AUTOMATION_TEST(FCbWriterBinaryReferenceTest, "System.Core.Serialization.CbWriter.BinaryReference", CompactBinaryWriterTestFlags)
-bool FCbWriterBinaryReferenceTest::RunTest(const FString& Parameters)
+IMPLEMENT_SIMPLE_AUTOMATION_TEST(FCbWriterBinaryAttachmentTest, "System.Core.Serialization.CbWriter.BinaryAttachment", CompactBinaryWriterTestFlags)
+bool FCbWriterBinaryAttachmentTest::RunTest(const FString& Parameters)
 {
 	TCbWriter<256> Writer;
 
@@ -604,17 +604,17 @@ bool FCbWriterBinaryReferenceTest::RunTest(const FString& Parameters)
 	const FIoHash Values[] = { FIoHash(ZeroBytes), FIoHash(SequentialBytes) };
 	for (const FIoHash& Value : Values)
 	{
-		Writer.BinaryReference(Value);
+		Writer.BinaryAttachment(Value);
 	}
 
 	FCbFieldRefIterator Fields = Writer.Save();
-	if (TestEqual(TEXT("FCbWriter(BinaryReference) Validate"), ValidateCompactBinaryRange(Fields.GetBuffer(), ECbValidateMode::All), ECbValidateError::None))
+	if (TestEqual(TEXT("FCbWriter(BinaryAttachment) Validate"), ValidateCompactBinaryRange(Fields.GetBuffer(), ECbValidateMode::All), ECbValidateError::None))
 	{
 		const FIoHash* CheckValue = Values;
 		for (FCbField Field : Fields)
 		{
-			TestEqual(TEXT("FCbWriter(BinaryReference).AsBinaryReference()"), Field.AsBinaryReference(), *CheckValue++);
-			TestFalse(TEXT("FCbWriter(BinaryReference) Error"), Field.HasError());
+			TestEqual(TEXT("FCbWriter(BinaryAttachment).AsBinaryAttachment()"), Field.AsBinaryAttachment(), *CheckValue++);
+			TestFalse(TEXT("FCbWriter(BinaryAttachment) Error"), Field.HasError());
 		}
 	}
 
@@ -816,8 +816,8 @@ bool FCbWriterComplexTest::RunTest(const FString& Parameters)
 		Writer.Name("False"_ASV).Bool(false);
 		Writer.Name("True"_ASV).Bool(true);
 
-		Writer.Name("CompactBinaryReference"_ASV).CompactBinaryReference(FIoHash());
-		Writer.Name("BinaryReference"_ASV).BinaryReference(FIoHash());
+		Writer.Name("CompactBinaryAttachment"_ASV).CompactBinaryAttachment(FIoHash());
+		Writer.Name("BinaryAttachment"_ASV).BinaryAttachment(FIoHash());
 
 		Writer.Name("Hash"_ASV).Hash(FIoHash());
 		Writer.Name("Uuid"_ASV).Uuid(FGuid());
diff --git a/Engine/Source/Runtime/Core/Public/Serialization/CompactBinary.h b/Engine/Source/Runtime/Core/Public/Serialization/CompactBinary.h
--- a/Engine/Source/Runtime/Core/Public/Serialization/CompactBinary.h
+++ b/Engine/Source/Runtime/Core/Public/Serialization/CompactBinary.h
@@ -38,7 +38,7 @@
  *
  * void BeginBuild(FCbObjectRef Params)
  * {
- *     if (FSharedBuffer Data = Storage().Load(Params["Data"].AsBinaryReference()))
+ *     if (FSharedBuffer Data = Storage().Load(Params["Data"].AsBinaryAttachment()))
  *     {
  *         SetData(Data);
  *     }
@@ -146,19 +146,19 @@ enum class ECbFieldType : uint8
 	BoolTrue        = 0x0d,
 
 	/**
-	 * CompactBinaryReference is a reference to compact binary data stored externally.
+	 * CompactBinaryAttachment is a reference to a compact binary attachment stored externally.
 	 *
-	 * Payload is a 256-bit hash digest of the compact binary data.
+	 * Payload is a 160-bit hash digest of the referenced compact binary data.
 	 */
-	CompactBinaryReference = 0x0e,
+	CompactBinaryAttachment = 0x0e,
 	/**
-	 * BinaryReference is a reference to binary data stored externally.
+	 * BinaryAttachment is a reference to a binary attachment stored externally.
 	 *
-	 * Payload is a 256-bit hash digest of the binary data.
+	 * Payload is a 160-bit hash digest of the referenced binary data.
 	 */
-	BinaryReference = 0x0f,
+	BinaryAttachment = 0x0f,
 
-	/** Hash. Payload is a 256-bit hash digest. */
+	/** Hash. Payload is a 160-bit hash digest. */
 	Hash            = 0x10,
 	/** UUID/GUID. Payload is a 128-bit UUID as defined by RFC 4122. */
 	Uuid            = 0x11,
@@ -213,8 +213,8 @@ class FCbFieldType
 	static constexpr ECbFieldType BoolMask              = ECbFieldType(0b0011'1110);
 	static constexpr ECbFieldType BoolBase              = ECbFieldType(0b0000'1100);
 
-	static constexpr ECbFieldType ReferenceMask      = ECbFieldType(0b0011'1110);
-	static constexpr ECbFieldType ReferenceBase      = ECbFieldType(0b0000'1110);
+	static constexpr ECbFieldType AttachmentMask        = ECbFieldType(0b0011'1110);
+	static constexpr ECbFieldType AttachmentBase        = ECbFieldType(0b0000'1110);
 
 	static void StaticAssertTypeConstants();
 
@@ -241,9 +241,9 @@ public:
 	static constexpr inline bool IsFloat(ECbFieldType Type)      { return (Type & FloatMask) == FloatBase; }
 	static constexpr inline bool IsBool(ECbFieldType Type)       { return (Type & BoolMask) == BoolBase; }
 
-	static constexpr inline bool IsCompactBinaryReference(ECbFieldType Type) { return GetType(Type) == ECbFieldType::CompactBinaryReference; }
-	static constexpr inline bool IsBinaryReference(ECbFieldType Type)        { return GetType(Type) == ECbFieldType::BinaryReference; }
-	static constexpr inline bool IsReference(ECbFieldType Type)   { return (Type & ReferenceMask) == ReferenceBase; }
+	static constexpr inline bool IsCompactBinaryAttachment(ECbFieldType Type) { return GetType(Type) == ECbFieldType::CompactBinaryAttachment; }
+	static constexpr inline bool IsBinaryAttachment(ECbFieldType Type)        { return GetType(Type) == ECbFieldType::BinaryAttachment; }
+	static constexpr inline bool IsAttachment(ECbFieldType Type) { return (Type & AttachmentMask) == AttachmentBase; }
 
 	static constexpr inline bool IsHash(ECbFieldType Type)       { return GetType(Type) == ECbFieldType::Hash; }
 	static constexpr inline bool IsUuid(ECbFieldType Type)       { return GetType(Type) == ECbFieldType::Uuid; }
@@ -251,11 +251,11 @@ public:
 	static constexpr inline bool IsDateTime(ECbFieldType Type)   { return GetType(Type) == ECbFieldType::DateTime; }
 	static constexpr inline bool IsTimeSpan(ECbFieldType Type)   { return GetType(Type) == ECbFieldType::TimeSpan; }
 
-	/** Whether the type is or may contain fields of any reference type. */
-	static constexpr inline bool MayContainReferences(ECbFieldType Type)
+	/** Whether the type is or may contain fields of any attachment type. */
+	static constexpr inline bool MayContainAttachments(ECbFieldType Type)
 	{
 		// The use of !! will suppress V792 from static analysis. Using //-V792 did not work.
-		return !!IsObject(Type) | !!IsArray(Type) | !!IsReference(Type);
+		return !!IsObject(Type) | !!IsArray(Type) | !!IsAttachment(Type);
 	}
 };
 
@@ -340,8 +340,8 @@ public:
 	/** Copy the field range into an archive, as if calling CopyTo on every field. */
 	CORE_API void CopyRangeTo(FArchive& Ar) const;
 
-	/** Invoke the visitor for every reference or binary reference in the field range. */
-	CORE_API void IterateRangeReferences(FCbFieldVisitor Visitor) const;
+	/** Invoke the visitor for every attachment in the field range. */
+	CORE_API void IterateRangeAttachments(FCbFieldVisitor Visitor) const;
 
 	/** Create a view of every field in the range. */
 	inline FMemoryView GetRangeView() const
@@ -496,12 +496,12 @@ public:
 	/** Access the field as a bool. Returns the provided default on error. */
 	CORE_API bool AsBool(bool bDefault = false);
 
-	/** Access the field as a hash referencing compact binary data. Returns the provided default on error. */
-	CORE_API FIoHash AsCompactBinaryReference(const FIoHash& Default = FIoHash());
-	/** Access the field as a hash referencing binary data. Returns the provided default on error. */
-	CORE_API FIoHash AsBinaryReference(const FIoHash& Default = FIoHash());
-	/** Access the field as a hash referencing compact binary data or binary data. Returns the provided default on error. */
-	CORE_API FIoHash AsReference(const FIoHash& Default = FIoHash());
+	/** Access the field as a hash referencing a compact binary attachment. Returns the provided default on error. */
+	CORE_API FIoHash AsCompactBinaryAttachment(const FIoHash& Default = FIoHash());
+	/** Access the field as a hash referencing a binary attachment. Returns the provided default on error. */
+	CORE_API FIoHash AsBinaryAttachment(const FIoHash& Default = FIoHash());
+	/** Access the field as a hash referencing an attachment. Returns the provided default on error. */
+	CORE_API FIoHash AsAttachment(const FIoHash& Default = FIoHash());
 
 	/** Access the field as a hash. Returns the provided default on error. */
 	CORE_API FIoHash AsHash(const FIoHash& Default = FIoHash());
@@ -544,9 +544,9 @@ public:
 	constexpr inline bool IsFloat() const           { return FCbFieldType::IsFloat(Type); }
 	constexpr inline bool IsBool() const            { return FCbFieldType::IsBool(Type); }
 
-	constexpr inline bool IsCompactBinaryReference() const { return FCbFieldType::IsCompactBinaryReference(Type); }
-	constexpr inline bool IsBinaryReference() const { return FCbFieldType::IsBinaryReference(Type); }
-	constexpr inline bool IsReference() const       { return FCbFieldType::IsReference(Type); }
+	constexpr inline bool IsCompactBinaryAttachment() const { return FCbFieldType::IsCompactBinaryAttachment(Type); }
+	constexpr inline bool IsBinaryAttachment() const        { return FCbFieldType::IsBinaryAttachment(Type); }
+	constexpr inline bool IsAttachment() const              { return FCbFieldType::IsAttachment(Type); }
 
 	constexpr inline bool IsHash() const            { return FCbFieldType::IsHash(Type); }
 	constexpr inline bool IsUuid() const            { return FCbFieldType::IsUuid(Type); }
@@ -596,8 +596,8 @@ public:
 	/** Copy the field into an archive, including its type and name. */
 	CORE_API void CopyTo(FArchive& Ar) const;
 
-	/** Invoke the visitor for every reference or binary reference in the field. */
-	CORE_API void IterateReferences(FCbFieldVisitor Visitor) const;
+	/** Invoke the visitor for every attachment in the field. */
+	CORE_API void IterateAttachments(FCbFieldVisitor Visitor) const;
 
 	/** Returns a view of the field, including the type and name when present. */
 	CORE_API FMemoryView GetView() const;
@@ -785,8 +785,8 @@ public:
 	/** Copy the array into an archive. This will write GetSize() bytes, with no name. */
 	CORE_API void CopyTo(FArchive& Ar) const;
 
-	/** Invoke the visitor for every reference or binary reference in the array. */
-	inline void IterateReferences(FCbFieldVisitor Visitor) const { CreateIterator().IterateRangeReferences(Visitor); }
+	/** Invoke the visitor for every attachment in the array. */
+	inline void IterateAttachments(FCbFieldVisitor Visitor) const { CreateIterator().IterateRangeAttachments(Visitor); }
 
 	/** Returns a view of the array, including the type and name when present. */
 	using FCbField::GetView;
@@ -883,8 +883,8 @@ public:
 	/** Copy the object into an archive. This will write GetSize() bytes, with no name. */
 	CORE_API void CopyTo(FArchive& Ar) const;
 
-	/** Invoke the visitor for every reference or binary reference in the object. */
-	inline void IterateReferences(FCbFieldVisitor Visitor) const { CreateIterator().IterateRangeReferences(Visitor); }
+	/** Invoke the visitor for every attachment in the object. */
+	inline void IterateAttachments(FCbFieldVisitor Visitor) const { CreateIterator().IterateRangeAttachments(Visitor); }
 
 	/** Returns a view of the object, including the type and name when present. */
 	using FCbField::GetView;
diff --git a/Engine/Source/Runtime/Core/Public/Serialization/CompactBinaryPackage.h b/Engine/Source/Runtime/Core/Public/Serialization/CompactBinaryPackage.h
--- a/Engine/Source/Runtime/Core/Public/Serialization/CompactBinaryPackage.h
+++ b/Engine/Source/Runtime/Core/Public/Serialization/CompactBinaryPackage.h
@@ -21,7 +21,7 @@ template <typename FuncType> class TFunctionRef;
  *
  * Attachments are serialized as one or two compact binary fields with no name. A Binary field is
  * written first with its content. The content hash is omitted when the content size is zero, and
- * is otherwise written as a BinaryReference or CompactBinaryReference depending on the type.
+ * is otherwise written as a BinaryAttachment or CompactBinaryAttachment depending on the type.
  */
 class FCbAttachment
 {
@@ -144,11 +144,11 @@ inline uint32 GetTypeHash(const FCbAttachment& Attachment)
  *
  * A package is serialized as a sequence of compact binary fields with no name. The object may be
  * both preceded and followed by attachments. The object itself is written as an Object field and
- * followed by its hash in a CompactBinaryReference field when the object is non-empty. A package
+ * followed by its hash in a CompactBinaryAttachment field, if the object is non-empty. A package
  * ends with a Null field. The canonical order of components is the object and its hash, followed
  * by the attachments ordered by hash, followed by a Null field. It is valid for the a package to
  * have its components serialized in any order, provided there is at most one object and the null
- * field is written last.
+ * field is written last, and the object hash immediately follows the object when present.
  */
 class FCbPackage
 {
diff --git a/Engine/Source/Runtime/Core/Public/Serialization/CompactBinaryWriter.h b/Engine/Source/Runtime/Core/Public/Serialization/CompactBinaryWriter.h
--- a/Engine/Source/Runtime/Core/Public/Serialization/CompactBinaryWriter.h
+++ b/Engine/Source/Runtime/Core/Public/Serialization/CompactBinaryWriter.h
@@ -166,9 +166,9 @@ public:
 	CORE_API void Bool(bool bValue);
 
 	/** Write a field referencing compact binary data by its hash. */
-	CORE_API void CompactBinaryReference(const FIoHash& Value);
+	CORE_API void CompactBinaryAttachment(const FIoHash& Value);
 	/** Write a field referencing binary data by its hash. */
-	CORE_API void BinaryReference(const FIoHash& Value);
+	CORE_API void BinaryAttachment(const FIoHash& Value);
 
 	/** Write a hash field. */
 	CORE_API void Hash(const FIoHash& Value);
