CompactBinary: Renamed and added overloads to FCbWriter to clarify behavior and allow specifying name and value in the same function

#rnx

diff --git a/Engine/Source/Developer/DerivedDataCache/Private/DerivedDataCache.cpp b/Engine/Source/Developer/DerivedDataCache/Private/DerivedDataCache.cpp
--- a/Engine/Source/Developer/DerivedDataCache/Private/DerivedDataCache.cpp
+++ b/Engine/Source/Developer/DerivedDataCache/Private/DerivedDataCache.cpp
@@ -1051,26 +1051,26 @@ void FCache::Put(
 	{
 		if (Value.IsCompactBinary())
 		{
-			Writer.Name("Value"_ASV).CompactBinaryAttachment(Value.GetHash());
+			Writer.SetName("Value"_ASV).AddCompactBinaryAttachment(Value.GetHash());
 		}
 		else if (Value.IsBinary())
 		{
-			Writer.Name("Value"_ASV).BinaryAttachment(Value.GetHash());
+			Writer.SetName("Value"_ASV).AddBinaryAttachment(Value.GetHash());
 		}
 
 		if (Record.GetType() == ECacheRecordType::Package)
 		{
-			Writer.Name("Attachments"_ASV);
+			Writer.SetName("Attachments"_ASV);
 			Writer.BeginArray();
 			for (const FCbAttachment& Attachment : Attachments)
 			{
 				if (Attachment.IsCompactBinary())
 				{
-					Writer.CompactBinaryAttachment(Attachment.GetHash());
+					Writer.AddCompactBinaryAttachment(Attachment.GetHash());
 				}
 				else
 				{
-					Writer.BinaryAttachment(Attachment.GetHash());
+					Writer.AddBinaryAttachment(Attachment.GetHash());
 				}
 			}
 			Writer.EndArray();
@@ -1079,7 +1079,7 @@ void FCache::Put(
 		const FCbObjectRef& Meta = Record.GetMeta();
 		if (Meta.CreateIterator())
 		{
-			Writer.Name("Meta"_ASV).Object(Meta);
+			Writer.SetName("Meta"_ASV).AddObject(Meta);
 		}
 	}
 	Writer.EndObject();
diff --git a/Engine/Source/Runtime/Core/Private/Serialization/CompactBinaryPackage.cpp b/Engine/Source/Runtime/Core/Private/Serialization/CompactBinaryPackage.cpp
--- a/Engine/Source/Runtime/Core/Private/Serialization/CompactBinaryPackage.cpp
+++ b/Engine/Source/Runtime/Core/Private/Serialization/CompactBinaryPackage.cpp
@@ -152,22 +152,22 @@ void FCbAttachment::Save(FCbWriter& Writer) const
 		FMemoryView SerializedView;
 		if (CompactBinary.TryGetSerializedRangeView(SerializedView))
 		{
-			Writer.Binary(SerializedView);
+			Writer.AddBinary(SerializedView);
 		}
 		else
 		{
-			Writer.Binary(AsBinary());
+			Writer.AddBinary(AsBinary());
 		}
-		Writer.CompactBinaryAttachment(Hash);
+		Writer.AddCompactBinaryAttachment(Hash);
 	}
 	else if (Buffer.GetSize())
 	{
-		Writer.Binary(Buffer);
-		Writer.BinaryAttachment(Hash);
+		Writer.AddBinary(Buffer);
+		Writer.AddBinaryAttachment(Hash);
 	}
 	else // Null
 	{
-		Writer.Binary(FMemoryView());
+		Writer.AddBinary(FMemoryView());
 	}
 }
 
@@ -368,14 +368,14 @@ void FCbPackage::Save(FCbWriter& Writer) const
 {
 	if (Object.CreateIterator())
 	{
-		Writer.Object(Object);
-		Writer.CompactBinaryAttachment(ObjectHash);
+		Writer.AddObject(Object);
+		Writer.AddCompactBinaryAttachment(ObjectHash);
 	}
 	for (const FCbAttachment& Attachment : Attachments)
 	{
 		Attachment.Save(Writer);
 	}
-	Writer.Null();
+	Writer.AddNull();
 }
 
 void FCbPackage::Save(FArchive& Ar) const
diff --git a/Engine/Source/Runtime/Core/Private/Serialization/CompactBinaryWriter.cpp b/Engine/Source/Runtime/Core/Private/Serialization/CompactBinaryWriter.cpp
--- a/Engine/Source/Runtime/Core/Private/Serialization/CompactBinaryWriter.cpp
+++ b/Engine/Source/Runtime/Core/Private/Serialization/CompactBinaryWriter.cpp
@@ -9,6 +9,7 @@
 #include "Misc/Guid.h"
 #include "Misc/Timespan.h"
 #include "Serialization/Archive.h"
+#include "Serialization/CompactBinaryPackage.h"
 #include "Serialization/CompactBinarySerialization.h"
 #include "Serialization/VarInt.h"
 
@@ -129,7 +130,7 @@ void FCbWriter::BeginField()
 	{
 		checkf((State.Flags & EStateFlags::Name) == EStateFlags::Name,
 			TEXT("A new field cannot be written until the previous field '%.*hs' is finished."),
-			CurrentName().Len(), CurrentName().GetData());
+			GetActiveName().Len(), GetActiveName().GetData());
 	}
 }
 
@@ -161,7 +162,7 @@ void FCbWriter::EndField(ECbFieldType Type)
 	Data[State.Offset] = uint8(Type);
 }
 
-FCbWriter& FCbWriter::Name(const FAnsiStringView Name)
+FCbWriter& FCbWriter::SetName(const FAnsiStringView Name)
 {
 	FState& State = States.Last();
 	checkf((State.Flags & EStateFlags::Array) != EStateFlags::Array,
@@ -171,7 +172,7 @@ FCbWriter& FCbWriter::Name(const FAnsiStringView Name)
 		: TEXT("It is invalid to write an empty name for a top-level field. Specify a name or avoid this call."));
 	checkf((State.Flags & (EStateFlags::Name | EStateFlags::Field)) == EStateFlags::None,
 		TEXT("A new field '%.*hs' cannot be written until the previous field '%.*hs' is finished."),
-		Name.Len(), Name.GetData(), CurrentName().Len(), CurrentName().GetData());
+		Name.Len(), Name.GetData(), GetActiveName().Len(), GetActiveName().GetData());
 
 	BeginField();
 	State.Flags |= EStateFlags::Name;
@@ -182,20 +183,20 @@ FCbWriter& FCbWriter::Name(const FAnsiStringView Name)
 	return *this;
 }
 
-void FCbWriter::NameOrString(const FAnsiStringView NameOrValue)
+void FCbWriter::SetNameOrAddString(const FAnsiStringView NameOrValue)
 {
 	// A name is only written if it would begin a new field inside of an object.
 	if ((States.Last().Flags & (EStateFlags::Name | EStateFlags::Field | EStateFlags::Object)) == EStateFlags::Object)
 	{
-		Name(NameOrValue);
+		SetName(NameOrValue);
 	}
 	else
 	{
-		String(NameOrValue);
+		AddString(NameOrValue);
 	}
 }
 
-FAnsiStringView FCbWriter::CurrentName() const
+FAnsiStringView FCbWriter::GetActiveName() const
 {
 	const FState& State = States.Last();
 	if ((State.Flags & EStateFlags::Name) == EStateFlags::Name)
@@ -231,16 +232,16 @@ void FCbWriter::MakeFieldsUniform(const int64 FieldBeginOffset, const int64 Fiel
 	}
 }
 
-void FCbWriter::Field(const FCbField& Value)
+void FCbWriter::AddField(const FCbField& Value)
 {
 	checkf(Value.HasValue(), TEXT("It is invalid to write a field with no value."));
 	BeginField();
 	EndField(AppendCompactBinary(Value, Data));
 }
 
-void FCbWriter::Field(const FCbFieldRef& Value)
+void FCbWriter::AddField(const FCbFieldRef& Value)
 {
-	Field(FCbField(Value));
+	AddField(FCbField(Value));
 }
 
 void FCbWriter::BeginObject()
@@ -285,15 +286,15 @@ void FCbWriter::EndObject()
 	EndField(bUniform ? ECbFieldType::UniformObject : ECbFieldType::Object);
 }
 
-void FCbWriter::Object(const FCbObject& Value)
+void FCbWriter::AddObject(const FCbObject& Value)
 {
 	BeginField();
 	EndField(AppendCompactBinary(Value.AsField(), Data));
 }
 
-void FCbWriter::Object(const FCbObjectRef& Value)
+void FCbWriter::AddObject(const FCbObjectRef& Value)
 {
-	Object(FCbObject(Value));
+	AddObject(FCbObject(Value));
 }
 
 void FCbWriter::BeginArray()
@@ -340,24 +341,24 @@ void FCbWriter::EndArray()
 	EndField(bUniform ? ECbFieldType::UniformArray : ECbFieldType::Array);
 }
 
-void FCbWriter::Array(const FCbArray& Value)
+void FCbWriter::AddArray(const FCbArray& Value)
 {
 	BeginField();
 	EndField(AppendCompactBinary(Value.AsField(), Data));
 }
 
-void FCbWriter::Array(const FCbArrayRef& Value)
+void FCbWriter::AddArray(const FCbArrayRef& Value)
 {
-	Array(FCbArray(Value));
+	AddArray(FCbArray(Value));
 }
 
-void FCbWriter::Null()
+void FCbWriter::AddNull()
 {
 	BeginField();
 	EndField(ECbFieldType::Null);
 }
 
-void FCbWriter::Binary(const void* const Value, const uint64 Size)
+void FCbWriter::AddBinary(const void* const Value, const uint64 Size)
 {
 	BeginField();
 	const uint32 SizeByteCount = MeasureVarUInt(Size);
@@ -367,12 +368,12 @@ void FCbWriter::Binary(const void* const Value, const uint64 Size)
 	EndField(ECbFieldType::Binary);
 }
 
-void FCbWriter::Binary(const FSharedBuffer& Buffer)
+void FCbWriter::AddBinary(FSharedBuffer Buffer)
 {
-	Binary(Buffer.GetData(), Buffer.GetSize());
+	AddBinary(Buffer.GetData(), Buffer.GetSize());
 }
 
-void FCbWriter::String(const FAnsiStringView Value)
+void FCbWriter::AddString(const FAnsiStringView Value)
 {
 	BeginField();
 	const uint64 Size = uint64(Value.Len());
@@ -388,7 +389,7 @@ void FCbWriter::String(const FAnsiStringView Value)
 	EndField(ECbFieldType::String);
 }
 
-void FCbWriter::String(const FWideStringView Value)
+void FCbWriter::AddString(const FWideStringView Value)
 {
 	BeginField();
 	const uint32 Size = uint32(FTCHARToUTF8_Convert::ConvertedLength(Value.GetData(), Value.Len()));
@@ -404,11 +405,11 @@ void FCbWriter::String(const FWideStringView Value)
 	EndField(ECbFieldType::String);
 }
 
-void FCbWriter::Integer(const int32 Value)
+void FCbWriter::AddInteger(const int32 Value)
 {
 	if (Value >= 0)
 	{
-		return Integer(uint32(Value));
+		return AddInteger(uint32(Value));
 	}
 	BeginField();
 	const uint32 Magnitude = ~uint32(Value);
@@ -418,11 +419,11 @@ void FCbWriter::Integer(const int32 Value)
 	EndField(ECbFieldType::IntegerNegative);
 }
 
-void FCbWriter::Integer(const int64 Value)
+void FCbWriter::AddInteger(const int64 Value)
 {
 	if (Value >= 0)
 	{
-		return Integer(uint64(Value));
+		return AddInteger(uint64(Value));
 	}
 	BeginField();
 	const uint64 Magnitude = ~uint64(Value);
@@ -432,7 +433,7 @@ void FCbWriter::Integer(const int64 Value)
 	EndField(ECbFieldType::IntegerNegative);
 }
 
-void FCbWriter::Integer(const uint32 Value)
+void FCbWriter::AddInteger(const uint32 Value)
 {
 	BeginField();
 	const uint32 ValueByteCount = MeasureVarUInt(Value);
@@ -441,7 +442,7 @@ void FCbWriter::Integer(const uint32 Value)
 	EndField(ECbFieldType::IntegerPositive);
 }
 
-void FCbWriter::Integer(const uint64 Value)
+void FCbWriter::AddInteger(const uint64 Value)
 {
 	BeginField();
 	const uint32 ValueByteCount = MeasureVarUInt(Value);
@@ -450,7 +451,7 @@ void FCbWriter::Integer(const uint64 Value)
 	EndField(ECbFieldType::IntegerPositive);
 }
 
-void FCbWriter::Float(const float Value)
+void FCbWriter::AddFloat(const float Value)
 {
 	BeginField();
 	const uint32 RawValue = NETWORK_ORDER32(reinterpret_cast<const uint32&>(Value));
@@ -458,12 +459,12 @@ void FCbWriter::Float(const float Value)
 	EndField(ECbFieldType::Float32);
 }
 
-void FCbWriter::Float(const double Value)
+void FCbWriter::AddFloat(const double Value)
 {
 	const float Value32 = float(Value);
 	if (Value == double(Value32))
 	{
-		return Float(Value32);
+		return AddFloat(Value32);
 	}
 	BeginField();
 	const uint64 RawValue = NETWORK_ORDER64(reinterpret_cast<const uint64&>(Value));
@@ -471,34 +472,42 @@ void FCbWriter::Float(const double Value)
 	EndField(ECbFieldType::Float64);
 }
 
-void FCbWriter::Bool(const bool bValue)
+void FCbWriter::AddBool(const bool bValue)
 {
 	BeginField();
 	EndField(bValue ? ECbFieldType::BoolTrue : ECbFieldType::BoolFalse);
 }
 
-void FCbWriter::CompactBinaryAttachment(const FIoHash& Value)
+void FCbWriter::AddCompactBinaryAttachment(const FIoHash& Value)
 {
 	BeginField();
 	Data.Append(Value.GetBytes(), sizeof(decltype(Value.GetBytes())));
 	EndField(ECbFieldType::CompactBinaryAttachment);
 }
 
-void FCbWriter::BinaryAttachment(const FIoHash& Value)
+void FCbWriter::AddBinaryAttachment(const FIoHash& Value)
 {
 	BeginField();
 	Data.Append(Value.GetBytes(), sizeof(decltype(Value.GetBytes())));
 	EndField(ECbFieldType::BinaryAttachment);
 }
 
-void FCbWriter::Hash(const FIoHash& Value)
+void FCbWriter::AddAttachment(const FCbAttachment& Attachment)
+{
+	BeginField();
+	const FIoHash& Value = Attachment.GetHash();
+	Data.Append(Value.GetBytes(), sizeof(decltype(Value.GetBytes())));
+	EndField(ECbFieldType::BinaryAttachment);
+}
+
+void FCbWriter::AddHash(const FIoHash& Value)
 {
 	BeginField();
 	Data.Append(Value.GetBytes(), sizeof(decltype(Value.GetBytes())));
 	EndField(ECbFieldType::Hash);
 }
 
-void FCbWriter::Uuid(const FGuid& Value)
+void FCbWriter::AddUuid(const FGuid& Value)
 {
 	const auto AppendSwappedBytes = [this](uint32 In)
 	{
@@ -513,7 +522,7 @@ void FCbWriter::Uuid(const FGuid& Value)
 	EndField(ECbFieldType::Uuid);
 }
 
-void FCbWriter::DateTimeTicks(const int64 Ticks)
+void FCbWriter::AddDateTimeTicks(const int64 Ticks)
 {
 	BeginField();
 	const uint64 RawValue = NETWORK_ORDER64(uint64(Ticks));
@@ -521,12 +530,12 @@ void FCbWriter::DateTimeTicks(const int64 Ticks)
 	EndField(ECbFieldType::DateTime);
 }
 
-void FCbWriter::DateTime(const FDateTime Value)
+void FCbWriter::AddDateTime(const FDateTime Value)
 {
-	DateTimeTicks(Value.GetTicks());
+	AddDateTimeTicks(Value.GetTicks());
 }
 
-void FCbWriter::TimeSpanTicks(const int64 Ticks)
+void FCbWriter::AddTimeSpanTicks(const int64 Ticks)
 {
 	BeginField();
 	const uint64 RawValue = NETWORK_ORDER64(uint64(Ticks));
@@ -534,22 +543,22 @@ void FCbWriter::TimeSpanTicks(const int64 Ticks)
 	EndField(ECbFieldType::TimeSpan);
 }
 
-void FCbWriter::TimeSpan(const FTimespan Value)
+void FCbWriter::AddTimeSpan(const FTimespan Value)
 {
-	TimeSpanTicks(Value.GetTicks());
+	AddTimeSpanTicks(Value.GetTicks());
 }
 
 ///////////////////////////////////////////////////////////////////////////////////////////////////////////////////////
 
 FCbWriter& operator<<(FCbWriter& Writer, const FDateTime Value)
 {
-	Writer.DateTime(Value);
+	Writer.AddDateTime(Value);
 	return Writer;
 }
 
 FCbWriter& operator<<(FCbWriter& Writer, const FTimespan Value)
 {
-	Writer.TimeSpan(Value);
+	Writer.AddTimeSpan(Value);
 	return Writer;
 }
 
diff --git a/Engine/Source/Runtime/Core/Private/Tests/Serialization/CompactBinaryPackageTest.cpp b/Engine/Source/Runtime/Core/Private/Tests/Serialization/CompactBinaryPackageTest.cpp
--- a/Engine/Source/Runtime/Core/Private/Tests/Serialization/CompactBinaryPackageTest.cpp
+++ b/Engine/Source/Runtime/Core/Private/Tests/Serialization/CompactBinaryPackageTest.cpp
@@ -362,7 +362,7 @@ bool FCbPackageTest::RunTest(const FString& Parameters)
 	FCbFieldRef Level3;
 	{
 		TCbWriter<256> Writer;
-		Writer.Name("Level4").BinaryAttachment(Level4Hash);
+		Writer.SetName("Level4").AddBinaryAttachment(Level4Hash);
 		Level3 = Writer.Save();
 	}
 	const FIoHash Level3Hash = Level3.GetHash();
@@ -370,9 +370,9 @@ bool FCbPackageTest::RunTest(const FString& Parameters)
 	FCbArrayRef Level2;
 	{
 		TCbWriter<256> Writer;
-		Writer.Name("Level3");
+		Writer.SetName("Level3");
 		Writer.BeginArray();
-		Writer.CompactBinaryAttachment(Level3Hash);
+		Writer.AddCompactBinaryAttachment(Level3Hash);
 		Writer.EndArray();
 		Level2 = Writer.Save().AsArrayRef();
 	}
@@ -382,7 +382,7 @@ bool FCbPackageTest::RunTest(const FString& Parameters)
 	{
 		TCbWriter<256> Writer;
 		Writer.BeginObject();
-		Writer.Name("Level2").CompactBinaryAttachment(Level2Hash);
+		Writer.SetName("Level2").AddCompactBinaryAttachment(Level2Hash);
 		Writer.EndObject();
 		Level1 = Writer.Save().AsObjectRef();
 	}
@@ -452,15 +452,15 @@ bool FCbPackageTest::RunTest(const FString& Parameters)
 	// Out of Order
 	{
 		TCbWriter<384> Writer;
-		Writer.Binary(Level2.GetBuffer());
-		Writer.CompactBinaryAttachment(Level2Hash);
-		Writer.Binary(Level4);
-		Writer.BinaryAttachment(Level4Hash);
-		Writer.Object(Level1);
-		Writer.CompactBinaryAttachment(Level1Hash);
-		Writer.Binary(Level3.GetBuffer());
-		Writer.CompactBinaryAttachment(Level3Hash);
-		Writer.Null();
+		Writer.AddBinary(Level2.GetBuffer());
+		Writer.AddCompactBinaryAttachment(Level2Hash);
+		Writer.AddBinary(Level4);
+		Writer.AddBinaryAttachment(Level4Hash);
+		Writer.AddObject(Level1);
+		Writer.AddCompactBinaryAttachment(Level1Hash);
+		Writer.AddBinary(Level3.GetBuffer());
+		Writer.AddCompactBinaryAttachment(Level3Hash);
+		Writer.AddNull();
 
 		FCbFieldRefIterator Fields = Writer.Save();
 		FCbPackage FromFields;
diff --git a/Engine/Source/Runtime/Core/Private/Tests/Serialization/CompactBinaryTest.cpp b/Engine/Source/Runtime/Core/Private/Tests/Serialization/CompactBinaryTest.cpp
--- a/Engine/Source/Runtime/Core/Private/Tests/Serialization/CompactBinaryTest.cpp
+++ b/Engine/Source/Runtime/Core/Private/Tests/Serialization/CompactBinaryTest.cpp
@@ -1101,11 +1101,11 @@ bool FCbFieldIterateAttachmentsTest::RunTest(const FString& Parameters)
 	{
 		FCbWriter Writer;
 
-		Writer.Name("IgnoredTypeInRoot").Hash(MakeTestHash(100));
-		Writer.CompactBinaryAttachment(MakeTestHash(0));
-		Writer.BinaryAttachment(MakeTestHash(1));
-		Writer.Name("CbAttachmentInRoot").CompactBinaryAttachment(MakeTestHash(2));
-		Writer.Name("BinAttachmentInRoot").BinaryAttachment(MakeTestHash(3));
+		Writer.SetName("IgnoredTypeInRoot").AddHash(MakeTestHash(100));
+		Writer.AddCompactBinaryAttachment(MakeTestHash(0));
+		Writer.AddBinaryAttachment(MakeTestHash(1));
+		Writer.SetName("CbAttachmentInRoot").AddCompactBinaryAttachment(MakeTestHash(2));
+		Writer.SetName("BinAttachmentInRoot").AddBinaryAttachment(MakeTestHash(3));
 
 		// Uniform array of type to ignore.
 		Writer.BeginArray();
@@ -1117,20 +1117,20 @@ bool FCbFieldIterateAttachmentsTest::RunTest(const FString& Parameters)
 		// Uniform array of binary attachments.
 		Writer.BeginArray();
 		{
-			Writer.BinaryAttachment(MakeTestHash(4));
-			Writer.BinaryAttachment(MakeTestHash(5));
+			Writer.AddBinaryAttachment(MakeTestHash(4));
+			Writer.AddBinaryAttachment(MakeTestHash(5));
 		}
 		Writer.EndArray();
 		// Uniform array of uniform arrays.
 		Writer.BeginArray();
 		{
 			Writer.BeginArray();
-			Writer.BinaryAttachment(MakeTestHash(6));
-			Writer.BinaryAttachment(MakeTestHash(7));
+			Writer.AddBinaryAttachment(MakeTestHash(6));
+			Writer.AddBinaryAttachment(MakeTestHash(7));
 			Writer.EndArray();
 			Writer.BeginArray();
-			Writer.BinaryAttachment(MakeTestHash(8));
-			Writer.BinaryAttachment(MakeTestHash(9));
+			Writer.AddBinaryAttachment(MakeTestHash(8));
+			Writer.AddBinaryAttachment(MakeTestHash(9));
 			Writer.EndArray();
 		}
 		Writer.EndArray();
@@ -1142,7 +1142,7 @@ bool FCbFieldIterateAttachmentsTest::RunTest(const FString& Parameters)
 			Writer << false;
 			Writer.EndArray();
 			Writer.BeginArray();
-			Writer.CompactBinaryAttachment(MakeTestHash(10));
+			Writer.AddCompactBinaryAttachment(MakeTestHash(10));
 			Writer << false;
 			Writer.EndArray();
 		}
@@ -1151,12 +1151,12 @@ bool FCbFieldIterateAttachmentsTest::RunTest(const FString& Parameters)
 		Writer.BeginArray();
 		{
 			Writer.BeginObject();
-			Writer.Name("CbAttachmentInUniObjInUniObj1").CompactBinaryAttachment(MakeTestHash(11));
-			Writer.Name("CbAttachmentInUniObjInUniObj2").CompactBinaryAttachment(MakeTestHash(12));
+			Writer.SetName("CbAttachmentInUniObjInUniObj1").AddCompactBinaryAttachment(MakeTestHash(11));
+			Writer.SetName("CbAttachmentInUniObjInUniObj2").AddCompactBinaryAttachment(MakeTestHash(12));
 			Writer.EndObject();
 			Writer.BeginObject();
-			Writer.Name("CbAttachmentInUniObjInUniObj3").CompactBinaryAttachment(MakeTestHash(13));
-			Writer.Name("CbAttachmentInUniObjInUniObj4").CompactBinaryAttachment(MakeTestHash(14));
+			Writer.SetName("CbAttachmentInUniObjInUniObj3").AddCompactBinaryAttachment(MakeTestHash(13));
+			Writer.SetName("CbAttachmentInUniObjInUniObj4").AddCompactBinaryAttachment(MakeTestHash(14));
 			Writer.EndObject();
 		}
 		Writer.EndArray();
@@ -1168,7 +1168,7 @@ bool FCbFieldIterateAttachmentsTest::RunTest(const FString& Parameters)
 			Writer << "Bool" << false;
 			Writer.EndObject();
 			Writer.BeginObject();
-			Writer.Name("CbAttachmentInNonUniObjInUniObj").CompactBinaryAttachment(MakeTestHash(15));
+			Writer.SetName("CbAttachmentInNonUniObjInUniObj").AddCompactBinaryAttachment(MakeTestHash(15));
 			Writer << "Bool" << false;
 			Writer.EndObject();
 		}
@@ -1184,36 +1184,36 @@ bool FCbFieldIterateAttachmentsTest::RunTest(const FString& Parameters)
 		// Uniform object of binary attachments.
 		Writer.BeginObject();
 		{
-			Writer.Name("BinAttachmentInUniObj1").BinaryAttachment(MakeTestHash(16));
-			Writer.Name("BinAttachmentInUniObj2").BinaryAttachment(MakeTestHash(17));
+			Writer.SetName("BinAttachmentInUniObj1").AddBinaryAttachment(MakeTestHash(16));
+			Writer.SetName("BinAttachmentInUniObj2").AddBinaryAttachment(MakeTestHash(17));
 		}
 		Writer.EndObject();
 		// Uniform object of uniform arrays.
 		Writer.BeginObject();
 		{
-			Writer.Name("Array1");
+			Writer.SetName("Array1");
 			Writer.BeginArray();
-			Writer.BinaryAttachment(MakeTestHash(18));
-			Writer.BinaryAttachment(MakeTestHash(19));
+			Writer.AddBinaryAttachment(MakeTestHash(18));
+			Writer.AddBinaryAttachment(MakeTestHash(19));
 			Writer.EndArray();
-			Writer.Name("Array2");
+			Writer.SetName("Array2");
 			Writer.BeginArray();
-			Writer.BinaryAttachment(MakeTestHash(20));
-			Writer.BinaryAttachment(MakeTestHash(21));
+			Writer.AddBinaryAttachment(MakeTestHash(20));
+			Writer.AddBinaryAttachment(MakeTestHash(21));
 			Writer.EndArray();
 		}
 		Writer.EndObject();
 		// Uniform object of non-uniform arrays.
 		Writer.BeginObject();
 		{
-			Writer.Name("Array1");
+			Writer.SetName("Array1");
 			Writer.BeginArray();
 			Writer << 0;
 			Writer << false;
 			Writer.EndArray();
-			Writer.Name("Array2");
+			Writer.SetName("Array2");
 			Writer.BeginArray();
-			Writer.CompactBinaryAttachment(MakeTestHash(22));
+			Writer.AddCompactBinaryAttachment(MakeTestHash(22));
 			Writer << false;
 			Writer.EndArray();
 		}
@@ -1221,29 +1221,29 @@ bool FCbFieldIterateAttachmentsTest::RunTest(const FString& Parameters)
 		// Uniform object of uniform objects.
 		Writer.BeginObject();
 		{
-			Writer.Name("Object1");
+			Writer.SetName("Object1");
 			Writer.BeginObject();
-			Writer.Name("CbAttachmentInUniObjInUniObj1").CompactBinaryAttachment(MakeTestHash(23));
-			Writer.Name("CbAttachmentInUniObjInUniObj2").CompactBinaryAttachment(MakeTestHash(24));
+			Writer.SetName("CbAttachmentInUniObjInUniObj1").AddCompactBinaryAttachment(MakeTestHash(23));
+			Writer.SetName("CbAttachmentInUniObjInUniObj2").AddCompactBinaryAttachment(MakeTestHash(24));
 			Writer.EndObject();
-			Writer.Name("Object2");
+			Writer.SetName("Object2");
 			Writer.BeginObject();
-			Writer.Name("CbAttachmentInUniObjInUniObj3").CompactBinaryAttachment(MakeTestHash(25));
-			Writer.Name("CbAttachmentInUniObjInUniObj4").CompactBinaryAttachment(MakeTestHash(26));
+			Writer.SetName("CbAttachmentInUniObjInUniObj3").AddCompactBinaryAttachment(MakeTestHash(25));
+			Writer.SetName("CbAttachmentInUniObjInUniObj4").AddCompactBinaryAttachment(MakeTestHash(26));
 			Writer.EndObject();
 		}
 		Writer.EndObject();
 		// Uniform object of non-uniform objects.
 		Writer.BeginObject();
 		{
-			Writer.Name("Object1");
+			Writer.SetName("Object1");
 			Writer.BeginObject();
 			Writer << "Int" << 0;
 			Writer << "Bool" << false;
 			Writer.EndObject();
-			Writer.Name("Object2");
+			Writer.SetName("Object2");
 			Writer.BeginObject();
-			Writer.Name("CbAttachmentInNonUniObjInUniObj").CompactBinaryAttachment(MakeTestHash(27));
+			Writer.SetName("CbAttachmentInNonUniObjInUniObj").AddCompactBinaryAttachment(MakeTestHash(27));
 			Writer << "Bool" << false;
 			Writer.EndObject();
 		}
diff --git a/Engine/Source/Runtime/Core/Private/Tests/Serialization/CompactBinaryValidationTest.cpp b/Engine/Source/Runtime/Core/Private/Tests/Serialization/CompactBinaryValidationTest.cpp
--- a/Engine/Source/Runtime/Core/Private/Tests/Serialization/CompactBinaryValidationTest.cpp
+++ b/Engine/Source/Runtime/Core/Private/Tests/Serialization/CompactBinaryValidationTest.cpp
@@ -351,8 +351,8 @@ bool FCbValidateAttachmentTest::RunTest(const FString& Parameters)
 	{
 		FBufferArchive Buffer;
 		FCbWriter Writer;
-		Writer.CompactBinaryAttachment(FIoHash::HashBuffer(BinaryView));
-		Writer.Binary(BinaryView);
+		Writer.AddCompactBinaryAttachment(FIoHash::HashBuffer(BinaryView));
+		Writer.AddBinary(BinaryView);
 		Writer.Save(Buffer);
 		TestEqual(TEXT("ValidateCompactBinaryAttachment(HashBeforeValue)"), ValidateCompactBinaryAttachment(MakeMemoryView(Buffer), ECbValidateMode::All), ECbValidateError::InvalidPackageFormat | ECbValidateError::Padding);
 		TestEqual(TEXT("ValidateCompactBinaryAttachment(HashBeforeValue, Mode)"), ValidateCompactBinaryAttachment(MakeMemoryView(Buffer), ECbValidateMode::All & ~ECbValidateMode::Package), ECbValidateError::Padding);
@@ -360,8 +360,8 @@ bool FCbValidateAttachmentTest::RunTest(const FString& Parameters)
 	{
 		FBufferArchive Buffer;
 		FCbWriter Writer;
-		Writer.Name("Name"_ASV).Binary(BinaryView);
-		Writer.CompactBinaryAttachment(FIoHash::HashBuffer(BinaryView));
+		Writer.SetName("Name"_ASV).AddBinary(BinaryView);
+		Writer.AddCompactBinaryAttachment(FIoHash::HashBuffer(BinaryView));
 		Writer.Save(Buffer);
 		TestEqual(TEXT("ValidateCompactBinaryAttachment(NameOnValue)"), ValidateCompactBinaryAttachment(MakeMemoryView(Buffer), ECbValidateMode::All), ECbValidateError::InvalidPackageFormat);
 		TestEqual(TEXT("ValidateCompactBinaryAttachment(NameOnValue, Mode)"), ValidateCompactBinaryAttachment(MakeMemoryView(Buffer), ECbValidateMode::All & ~ECbValidateMode::Package), ECbValidateError::None);
@@ -369,8 +369,8 @@ bool FCbValidateAttachmentTest::RunTest(const FString& Parameters)
 	{
 		FBufferArchive Buffer;
 		FCbWriter Writer;
-		Writer.Binary(BinaryView);
-		Writer.Name("Name"_ASV).CompactBinaryAttachment(FIoHash::HashBuffer(BinaryView));
+		Writer.AddBinary(BinaryView);
+		Writer.SetName("Name"_ASV).AddCompactBinaryAttachment(FIoHash::HashBuffer(BinaryView));
 		Writer.Save(Buffer);
 		TestEqual(TEXT("ValidateCompactBinaryAttachment(NameOnHash)"), ValidateCompactBinaryAttachment(MakeMemoryView(Buffer), ECbValidateMode::All), ECbValidateError::InvalidPackageFormat);
 		TestEqual(TEXT("ValidateCompactBinaryAttachment(NameOnHash, Mode)"), ValidateCompactBinaryAttachment(MakeMemoryView(Buffer), ECbValidateMode::All & ~ECbValidateMode::Package), ECbValidateError::None);
@@ -401,8 +401,8 @@ bool FCbValidatePackageTest::RunTest(const FString& Parameters)
 		Writer.BeginObject();
 		Writer << "Field" << 42;
 		Writer.EndObject();
-		Writer.CompactBinaryAttachment(Writer.Save().GetHash());
-		Writer.Null();
+		Writer.AddCompactBinaryAttachment(Writer.Save().GetHash());
+		Writer.AddNull();
 		Writer.Save(Buffer);
 		TestEqual(TEXT("ValidateCompactBinaryPackage(Object)"), ValidateCompactBinaryPackage(MakeMemoryView(Buffer), ECbValidateMode::All), ECbValidateError::None);
 		TestEqual(TEXT("ValidateCompactBinaryPackage(Object, MissingHash)"), ValidateCompactBinaryPackage(MakeMemoryView(Buffer).LeftChop(34), ECbValidateMode::All), ECbValidateError::InvalidPackageFormat);
@@ -429,11 +429,11 @@ bool FCbValidatePackageTest::RunTest(const FString& Parameters)
 		Writer.BeginObject();
 		Writer << "Field" << 42;
 		Writer.EndObject();
-		Writer.CompactBinaryAttachment(Writer.Save().GetHash());
+		Writer.AddCompactBinaryAttachment(Writer.Save().GetHash());
 		const uint64 AttachmentOffset = Writer.GetSaveSize();
-		Writer.Binary(BinaryView);
-		Writer.CompactBinaryAttachment(FIoHash::HashBuffer(BinaryView));
-		Writer.Null();
+		Writer.AddBinary(BinaryView);
+		Writer.AddCompactBinaryAttachment(FIoHash::HashBuffer(BinaryView));
+		Writer.AddNull();
 		Writer.Save(Buffer);
 		TestEqual(TEXT("ValidateCompactBinaryPackage(Object+Attachment)"), ValidateCompactBinaryPackage(MakeMemoryView(Buffer), ECbValidateMode::All), ECbValidateError::None);
 		TestEqual(TEXT("ValidateCompactBinaryPackage(Object+Attachment, MissingHash)"), ValidateCompactBinaryPackage(MakeMemoryView(Buffer).LeftChop(22), ECbValidateMode::All), ECbValidateError::InvalidPackageFormat);
@@ -461,12 +461,12 @@ bool FCbValidatePackageTest::RunTest(const FString& Parameters)
 	{
 		FBufferArchive Buffer;
 		TCbWriter<256> Writer;
-		Writer.Name("Name"_ASV);
+		Writer.SetName("Name"_ASV);
 		Writer.BeginObject();
 		Writer << "Field" << 42;
 		Writer.EndObject();
-		Writer.CompactBinaryAttachment(Writer.Save().GetHash());
-		Writer.Null();
+		Writer.AddCompactBinaryAttachment(Writer.Save().GetHash());
+		Writer.AddNull();
 		Writer.Save(Buffer);
 		TestEqual(TEXT("ValidateCompactBinaryPackage(NameOnObject)"), ValidateCompactBinaryPackage(MakeMemoryView(Buffer), ECbValidateMode::All), ECbValidateError::InvalidPackageFormat);
 		TestEqual(TEXT("ValidateCompactBinaryPackage(NameOnObject, Mode)"), ValidateCompactBinaryPackage(MakeMemoryView(Buffer), ECbValidateMode::All & ~ECbValidateMode::Package), ECbValidateError::None);
@@ -478,9 +478,9 @@ bool FCbValidatePackageTest::RunTest(const FString& Parameters)
 		Writer << "Field" << 42;
 		Writer.EndObject();
 		const FIoHash ObjectHash = Writer.Save().GetHash();
-		Writer.Name("Name"_ASV);
-		Writer.CompactBinaryAttachment(ObjectHash);
-		Writer.Null();
+		Writer.SetName("Name"_ASV);
+		Writer.AddCompactBinaryAttachment(ObjectHash);
+		Writer.AddNull();
 		Writer.Save(Buffer);
 		TestEqual(TEXT("ValidateCompactBinaryPackage(NameOnHash)"), ValidateCompactBinaryPackage(MakeMemoryView(Buffer), ECbValidateMode::All), ECbValidateError::InvalidPackageFormat);
 		TestEqual(TEXT("ValidateCompactBinaryPackage(NameOnHash, Mode)"), ValidateCompactBinaryPackage(MakeMemoryView(Buffer), ECbValidateMode::All & ~ECbValidateMode::Package), ECbValidateError::None);
@@ -491,10 +491,10 @@ bool FCbValidatePackageTest::RunTest(const FString& Parameters)
 		Writer.BeginObject();
 		Writer << "Field" << 42;
 		Writer.EndObject();
-		Writer.CompactBinaryAttachment(Writer.Save().GetHash());
-		Writer.Name("Name"_ASV).Binary(BinaryView);
-		Writer.CompactBinaryAttachment(FIoHash::HashBuffer(BinaryView));
-		Writer.Null();
+		Writer.AddCompactBinaryAttachment(Writer.Save().GetHash());
+		Writer.SetName("Name"_ASV).AddBinary(BinaryView);
+		Writer.AddCompactBinaryAttachment(FIoHash::HashBuffer(BinaryView));
+		Writer.AddNull();
 		Writer.Save(Buffer);
 		TestEqual(TEXT("ValidateCompactBinaryPackage(NameOnAttachment)"), ValidateCompactBinaryPackage(MakeMemoryView(Buffer), ECbValidateMode::All), ECbValidateError::InvalidPackageFormat);
 		TestEqual(TEXT("ValidateCompactBinaryPackage(NameOnAttachment, Mode)"), ValidateCompactBinaryPackage(MakeMemoryView(Buffer), ECbValidateMode::All & ~ECbValidateMode::Package), ECbValidateError::None);
@@ -506,12 +506,12 @@ bool FCbValidatePackageTest::RunTest(const FString& Parameters)
 		Writer << "Field" << 42;
 		Writer.EndObject();
 		const FIoHash ObjectHash = Writer.Save().GetHash();
-		Writer.CompactBinaryAttachment(ObjectHash);
+		Writer.AddCompactBinaryAttachment(ObjectHash);
 		Writer.BeginObject();
 		Writer << "Field" << 42;
 		Writer.EndObject();
-		Writer.CompactBinaryAttachment(ObjectHash);
-		Writer.Null();
+		Writer.AddCompactBinaryAttachment(ObjectHash);
+		Writer.AddNull();
 		Writer.Save(Buffer);
 		TestEqual(TEXT("ValidateCompactBinaryPackage(MultipleObjects)"), ValidateCompactBinaryPackage(MakeMemoryView(Buffer), ECbValidateMode::All), ECbValidateError::MultiplePackageObjects);
 		TestEqual(TEXT("ValidateCompactBinaryPackage(MultipleObjects, Mode)"), ValidateCompactBinaryPackage(MakeMemoryView(Buffer), ECbValidateMode::All & ~ECbValidateMode::Package), ECbValidateError::None);
@@ -519,11 +519,11 @@ bool FCbValidatePackageTest::RunTest(const FString& Parameters)
 	{
 		FBufferArchive Buffer;
 		TCbWriter<256> Writer;
-		Writer.Binary(BinaryView);
-		Writer.CompactBinaryAttachment(FIoHash::HashBuffer(BinaryView));
-		Writer.Binary(BinaryView);
-		Writer.CompactBinaryAttachment(FIoHash::HashBuffer(BinaryView));
-		Writer.Null();
+		Writer.AddBinary(BinaryView);
+		Writer.AddCompactBinaryAttachment(FIoHash::HashBuffer(BinaryView));
+		Writer.AddBinary(BinaryView);
+		Writer.AddCompactBinaryAttachment(FIoHash::HashBuffer(BinaryView));
+		Writer.AddNull();
 		Writer.Save(Buffer);
 		TestEqual(TEXT("ValidateCompactBinaryPackage(DuplicateAttachments)"), ValidateCompactBinaryPackage(MakeMemoryView(Buffer), ECbValidateMode::All), ECbValidateError::DuplicateAttachments);
 		TestEqual(TEXT("ValidateCompactBinaryPackage(DuplicateAttachments, Mode)"), ValidateCompactBinaryPackage(MakeMemoryView(Buffer), ECbValidateMode::All & ~ECbValidateMode::Package), ECbValidateError::None);
@@ -533,8 +533,8 @@ bool FCbValidatePackageTest::RunTest(const FString& Parameters)
 		TCbWriter<256> Writer;
 		Writer.BeginObject();
 		Writer.EndObject();
-		Writer.CompactBinaryAttachment(FCbObject().GetHash());
-		Writer.Null();
+		Writer.AddCompactBinaryAttachment(FCbObject().GetHash());
+		Writer.AddNull();
 		Writer.Save(Buffer);
 		TestEqual(TEXT("ValidateCompactBinaryPackage(NullObject)"), ValidateCompactBinaryPackage(MakeMemoryView(Buffer), ECbValidateMode::All), ECbValidateError::NullPackageObject);
 		TestEqual(TEXT("ValidateCompactBinaryPackage(NullObject, Mode)"), ValidateCompactBinaryPackage(MakeMemoryView(Buffer), ECbValidateMode::All & ~ECbValidateMode::Package), ECbValidateError::None);
@@ -542,8 +542,8 @@ bool FCbValidatePackageTest::RunTest(const FString& Parameters)
 	{
 		FBufferArchive Buffer;
 		TCbWriter<256> Writer;
-		Writer.Binary(FMemoryView());
-		Writer.Null();
+		Writer.AddBinary(FMemoryView());
+		Writer.AddNull();
 		Writer.Save(Buffer);
 		TestEqual(TEXT("ValidateCompactBinaryPackage(NullAttachment)"), ValidateCompactBinaryPackage(MakeMemoryView(Buffer), ECbValidateMode::All), ECbValidateError::NullPackageAttachment);
 		TestEqual(TEXT("ValidateCompactBinaryPackage(NullAttachment, Mode)"), ValidateCompactBinaryPackage(MakeMemoryView(Buffer), ECbValidateMode::All & ~ECbValidateMode::Package), ECbValidateError::None);
diff --git a/Engine/Source/Runtime/Core/Private/Tests/Serialization/CompactBinaryWriterTest.cpp b/Engine/Source/Runtime/Core/Private/Tests/Serialization/CompactBinaryWriterTest.cpp
--- a/Engine/Source/Runtime/Core/Private/Tests/Serialization/CompactBinaryWriterTest.cpp
+++ b/Engine/Source/Runtime/Core/Private/Tests/Serialization/CompactBinaryWriterTest.cpp
@@ -9,6 +9,7 @@
 #include "Misc/Timespan.h"
 #include "Misc/StringBuilder.h"
 #include "Serialization/BufferArchive.h"
+#include "Serialization/CompactBinaryPackage.h"
 #include "Serialization/CompactBinaryValidation.h"
 
 #if WITH_DEV_AUTOMATION_TESTS
@@ -38,7 +39,7 @@ bool FCbWriterObjectTest::RunTest(const FString& Parameters)
 	// Test Named Empty Object
 	{
 		Writer.Reset();
-		Writer.Name("Object"_ASV);
+		Writer.SetName("Object"_ASV);
 		Writer.BeginObject();
 		Writer.EndObject();
 		FCbFieldRef Field = Writer.Save();
@@ -53,8 +54,8 @@ bool FCbWriterObjectTest::RunTest(const FString& Parameters)
 	{
 		Writer.Reset();
 		Writer.BeginObject();
-		Writer.Name("Integer"_ASV).Integer(0);
-		Writer.Name("Float"_ASV).Float(0.0f);
+		Writer.SetName("Integer"_ASV).AddInteger(0);
+		Writer.SetName("Float"_ASV).AddFloat(0.0f);
 		Writer.EndObject();
 		FCbFieldRef Field = Writer.Save();
 		if (TestEqual(TEXT("FCbWriter(Object, Basic) Validate"), ValidateCompactBinary(Field.GetBuffer(), ECbValidateMode::All), ECbValidateError::None))
@@ -70,8 +71,8 @@ bool FCbWriterObjectTest::RunTest(const FString& Parameters)
 	{
 		Writer.Reset();
 		Writer.BeginObject();
-		Writer.Name("Field1"_ASV).Integer(0);
-		Writer.Name("Field2"_ASV).Integer(1);
+		Writer.SetName("Field1"_ASV).AddInteger(0);
+		Writer.SetName("Field2"_ASV).AddInteger(1);
 		Writer.EndObject();
 		FCbFieldRef Field = Writer.Save();
 		if (TestEqual(TEXT("FCbWriter(Object, Uniform) Validate"), ValidateCompactBinary(Field.GetBuffer(), ECbValidateMode::All), ECbValidateError::None))
@@ -109,7 +110,7 @@ bool FCbWriterArrayTest::RunTest(const FString& Parameters)
 	// Test Named Empty Array
 	{
 		Writer.Reset();
-		Writer.Name("Array"_ASV);
+		Writer.SetName("Array"_ASV);
 		Writer.BeginArray();
 		Writer.EndArray();
 		FCbFieldRef Field = Writer.Save();
@@ -124,8 +125,8 @@ bool FCbWriterArrayTest::RunTest(const FString& Parameters)
 	{
 		Writer.Reset();
 		Writer.BeginArray();
-		Writer.Integer(0);
-		Writer.Float(0.0f);
+		Writer.AddInteger(0);
+		Writer.AddFloat(0.0f);
 		Writer.EndArray();
 		FCbFieldRef Field = Writer.Save();
 		if (TestEqual(TEXT("FCbWriter(Array, Basic) Validate"), ValidateCompactBinary(Field.GetBuffer(), ECbValidateMode::All), ECbValidateError::None))
@@ -144,8 +145,8 @@ bool FCbWriterArrayTest::RunTest(const FString& Parameters)
 	{
 		Writer.Reset();
 		Writer.BeginArray();
-		Writer.Integer(0);
-		Writer.Integer(1);
+		Writer.AddInteger(0);
+		Writer.AddInteger(1);
 		Writer.EndArray();
 		FCbFieldRef Field = Writer.Save();
 		if (TestEqual(TEXT("FCbWriter(Array, Uniform) Validate"), ValidateCompactBinary(Field.GetBuffer(), ECbValidateMode::All), ECbValidateError::None))
@@ -173,7 +174,7 @@ bool FCbWriterNullTest::RunTest(const FString& Parameters)
 	// Test Null
 	{
 		Writer.Reset();
-		Writer.Null();
+		Writer.AddNull();
 		FCbFieldRef Field = Writer.Save();
 		if (TestEqual(TEXT("FCbWriter(Null) Validate"), ValidateCompactBinary(Field.GetBuffer(), ECbValidateMode::All), ECbValidateError::None))
 		{
@@ -185,8 +186,8 @@ bool FCbWriterNullTest::RunTest(const FString& Parameters)
 	// Test Null with Name
 	{
 		Writer.Reset();
-		Writer.Name("Null"_ASV);
-		Writer.Null();
+		Writer.SetName("Null"_ASV);
+		Writer.AddNull();
 		FCbFieldRef Field = Writer.Save();
 		if (TestEqual(TEXT("FCbWriter(Null, Name) Validate"), ValidateCompactBinary(Field.GetBuffer(), ECbValidateMode::All), ECbValidateError::None))
 		{
@@ -201,15 +202,15 @@ bool FCbWriterNullTest::RunTest(const FString& Parameters)
 		Writer.Reset();
 
 		Writer.BeginArray();
-		Writer.Null();
-		Writer.Null();
-		Writer.Null();
+		Writer.AddNull();
+		Writer.AddNull();
+		Writer.AddNull();
 		Writer.EndArray();
 
 		Writer.BeginObject();
-		Writer.Name("N1"_ASV).Null();
-		Writer.Name("N2"_ASV).Null();
-		Writer.Name("N3"_ASV).Null();
+		Writer.SetName("N1"_ASV).AddNull();
+		Writer.SetName("N2"_ASV).AddNull();
+		Writer.SetName("N3"_ASV).AddNull();
 		Writer.EndObject();
 
 		FCbFieldRefIterator Fields = Writer.Save();
@@ -222,7 +223,7 @@ bool FCbWriterNullTest::RunTest(const FString& Parameters)
 		constexpr int NullCount = 3;
 		for (int Index = 0; Index < NullCount; ++Index)
 		{
-			Writer.Null();
+			Writer.AddNull();
 		}
 		uint8 Buffer[NullCount]{};
 		FCbFieldIterator Fields = Writer.Save(MakeMemoryView(Buffer));
@@ -263,7 +264,7 @@ bool FCbWriterBinaryTest::RunTest(const FString& Parameters)
 	// Test Empty Binary
 	{
 		Writer.Reset();
-		Writer.Binary(nullptr, 0);
+		Writer.AddBinary(nullptr, 0);
 		FCbFieldRef Field = Writer.Save();
 		if (TestEqual(TEXT("FCbWriter(Binary, Empty) Validate"), ValidateCompactBinary(Field.GetBuffer(), ECbValidateMode::All), ECbValidateError::None))
 		{
@@ -277,8 +278,8 @@ bool FCbWriterBinaryTest::RunTest(const FString& Parameters)
 	{
 		Writer.Reset();
 		const uint8 BinaryValue[] = { 1, 2, 3, 4, 5, 6 };
-		Writer.Name("Binary"_ASV);
-		Writer.Binary(BinaryValue, sizeof(BinaryValue));
+		Writer.SetName("Binary"_ASV);
+		Writer.AddBinary(BinaryValue, sizeof(BinaryValue));
 		FCbFieldRef Field = Writer.Save();
 		if (TestEqual(TEXT("FCbWriter(Binary, Array) Validate"), ValidateCompactBinary(Field.GetBuffer(), ECbValidateMode::All), ECbValidateError::None))
 		{
@@ -302,8 +303,8 @@ bool FCbWriterStringTest::RunTest(const FString& Parameters)
 	// Test Empty Strings
 	{
 		Writer.Reset();
-		Writer.String(FAnsiStringView());
-		Writer.String(FWideStringView());
+		Writer.AddString(FAnsiStringView());
+		Writer.AddString(FWideStringView());
 		FCbFieldRefIterator Fields = Writer.Save();
 		if (TestEqual(TEXT("FCbWriter(String, Empty) Validate"), ValidateCompactBinaryRange(Fields.GetBuffer(), ECbValidateMode::All), ECbValidateError::None))
 		{
@@ -319,8 +320,8 @@ bool FCbWriterStringTest::RunTest(const FString& Parameters)
 	// Test Basic Strings
 	{
 		Writer.Reset();
-		Writer.Name("String"_ASV).String("Value"_ASV);
-		Writer.Name("String"_ASV).String(TEXT("Value"_SV));
+		Writer.SetName("String"_ASV).AddString("Value"_ASV);
+		Writer.SetName("String"_ASV).AddString(TEXT("Value"_SV));
 		FCbFieldRefIterator Fields = Writer.Save();
 		if (TestEqual(TEXT("FCbWriter(String, Basic) Validate"), ValidateCompactBinaryRange(Fields.GetBuffer(), ECbValidateMode::All), ECbValidateError::None))
 		{
@@ -343,8 +344,8 @@ bool FCbWriterStringTest::RunTest(const FString& Parameters)
 		{
 			Dots << '.';
 		}
-		Writer.String(Dots);
-		Writer.String(FString::ChrN(DotCount, TEXT('.')));
+		Writer.AddString(Dots);
+		Writer.AddString(FString::ChrN(DotCount, TEXT('.')));
 		FCbFieldRefIterator Fields = Writer.Save();
 		if (TestEqual(TEXT("FCbWriter(String, Long) Validate"), ValidateCompactBinaryRange(Fields.GetBuffer(), ECbValidateMode::All), ECbValidateError::None))
 		{
@@ -359,8 +360,8 @@ bool FCbWriterStringTest::RunTest(const FString& Parameters)
 	{
 		Writer.Reset();
 		WIDECHAR Value[2] = { 0xd83d, 0xde00 };
-		Writer.String("\xf0\x9f\x98\x80"_ASV);
-		Writer.String(FWideStringView(Value, UE_ARRAY_COUNT(Value)));
+		Writer.AddString("\xf0\x9f\x98\x80"_ASV);
+		Writer.AddString(FWideStringView(Value, UE_ARRAY_COUNT(Value)));
 		FCbFieldRefIterator Fields = Writer.Save();
 		if (TestEqual(TEXT("FCbWriter(String, Unicode) Validate"), ValidateCompactBinaryRange(Fields.GetBuffer(), ECbValidateMode::All), ECbValidateError::None))
 		{
@@ -384,7 +385,7 @@ bool FCbWriterIntegerTest::RunTest(const FString& Parameters)
 	auto TestInt32 = [this, &Writer](int32 Value)
 	{
 		Writer.Reset();
-		Writer.Integer(Value);
+		Writer.AddInteger(Value);
 		FCbFieldRef Field = Writer.Save();
 		if (TestEqual(TEXT("FCbWriter(Integer, Int32) Validate"), ValidateCompactBinary(Field.GetBuffer(), ECbValidateMode::All), ECbValidateError::None))
 		{
@@ -396,7 +397,7 @@ bool FCbWriterIntegerTest::RunTest(const FString& Parameters)
 	auto TestUInt32 = [this, &Writer](uint32 Value)
 	{
 		Writer.Reset();
-		Writer.Integer(Value);
+		Writer.AddInteger(Value);
 		FCbFieldRef Field = Writer.Save();
 		if (TestEqual(TEXT("FCbWriter(Integer, UInt32) Validate"), ValidateCompactBinary(Field.GetBuffer(), ECbValidateMode::All), ECbValidateError::None))
 		{
@@ -408,7 +409,7 @@ bool FCbWriterIntegerTest::RunTest(const FString& Parameters)
 	auto TestInt64 = [this, &Writer](int64 Value)
 	{
 		Writer.Reset();
-		Writer.Integer(Value);
+		Writer.AddInteger(Value);
 		FCbFieldRef Field = Writer.Save();
 		if (TestEqual(TEXT("FCbWriter(Integer, Int64) Validate"), ValidateCompactBinary(Field.GetBuffer(), ECbValidateMode::All), ECbValidateError::None))
 		{
@@ -420,7 +421,7 @@ bool FCbWriterIntegerTest::RunTest(const FString& Parameters)
 	auto TestUInt64 = [this, &Writer](uint64 Value)
 	{
 		Writer.Reset();
-		Writer.Integer(Value);
+		Writer.AddInteger(Value);
 		FCbFieldRef Field = Writer.Save();
 		if (TestEqual(TEXT("FCbWriter(Integer, UInt64) Validate"), ValidateCompactBinary(Field.GetBuffer(), ECbValidateMode::All), ECbValidateError::None))
 		{
@@ -478,7 +479,7 @@ bool FCbWriterFloatTest::RunTest(const FString& Parameters)
 		constexpr float Values[] = { 0.0f, 1.0f, -1.0f, PI };
 		for (float Value : Values)
 		{
-			Writer.Float(Value);
+			Writer.AddFloat(Value);
 		}
 		FCbFieldRefIterator Fields = Writer.Save();
 		if (TestEqual(TEXT("FCbWriter(Float, Single) Validate"), ValidateCompactBinaryRange(Fields.GetBuffer(), ECbValidateMode::All), ECbValidateError::None))
@@ -498,7 +499,7 @@ bool FCbWriterFloatTest::RunTest(const FString& Parameters)
 		constexpr double Values[] = { 0.0f, 1.0f, -1.0f, PI, 1.9999998807907104, 1.9999999403953552, 3.4028234663852886e38, 6.8056469327705771e38 };
 		for (double Value : Values)
 		{
-			Writer.Float(Value);
+			Writer.AddFloat(Value);
 		}
 		FCbFieldRefIterator Fields = Writer.Save();
 		if (TestEqual(TEXT("FCbWriter(Float, Double) Validate"), ValidateCompactBinaryRange(Fields.GetBuffer(), ECbValidateMode::All), ECbValidateError::None))
@@ -524,8 +525,8 @@ bool FCbWriterBoolTest::RunTest(const FString& Parameters)
 
 	// Test Bool Values
 	{
-		Writer.Bool(true);
-		Writer.Bool(false);
+		Writer.AddBool(true);
+		Writer.AddBool(false);
 
 		FCbFieldRefIterator Fields = Writer.Save();
 		if (TestEqual(TEXT("FCbWriter(Bool) Validate"), ValidateCompactBinaryRange(Fields.GetBuffer(), ECbValidateMode::All), ECbValidateError::None))
@@ -543,15 +544,15 @@ bool FCbWriterBoolTest::RunTest(const FString& Parameters)
 		Writer.Reset();
 
 		Writer.BeginArray();
-		Writer.Bool(false);
-		Writer.Bool(false);
-		Writer.Bool(false);
+		Writer.AddBool(false);
+		Writer.AddBool(false);
+		Writer.AddBool(false);
 		Writer.EndArray();
 
 		Writer.BeginObject();
-		Writer.Name("B1"_ASV).Bool(false);
-		Writer.Name("B2"_ASV).Bool(false);
-		Writer.Name("B3"_ASV).Bool(false);
+		Writer.SetName("B1"_ASV).AddBool(false);
+		Writer.SetName("B2"_ASV).AddBool(false);
+		Writer.SetName("B3"_ASV).AddBool(false);
 		Writer.EndObject();
 
 		FCbFieldRefIterator Fields = Writer.Save();
@@ -574,7 +575,7 @@ bool FCbWriterCompactBinaryAttachmentTest::RunTest(const FString& Parameters)
 	const FIoHash Values[] = { FIoHash(ZeroBytes), FIoHash(SequentialBytes) };
 	for (const FIoHash& Value : Values)
 	{
-		Writer.CompactBinaryAttachment(Value);
+		Writer.AddCompactBinaryAttachment(Value);
 	}
 
 	FCbFieldRefIterator Fields = Writer.Save();
@@ -604,7 +605,7 @@ bool FCbWriterBinaryAttachmentTest::RunTest(const FString& Parameters)
 	const FIoHash Values[] = { FIoHash(ZeroBytes), FIoHash(SequentialBytes) };
 	for (const FIoHash& Value : Values)
 	{
-		Writer.BinaryAttachment(Value);
+		Writer.AddBinaryAttachment(Value);
 	}
 
 	FCbFieldRefIterator Fields = Writer.Save();
@@ -634,7 +635,7 @@ bool FCbWriterHashTest::RunTest(const FString& Parameters)
 	const FIoHash Values[] = { FIoHash(ZeroBytes), FIoHash(SequentialBytes) };
 	for (const FIoHash& Value : Values)
 	{
-		Writer.Hash(Value);
+		Writer.AddHash(Value);
 	}
 
 	FCbFieldRefIterator Fields = Writer.Save();
@@ -661,7 +662,7 @@ bool FCbWriterUuidTest::RunTest(const FString& Parameters)
 	const FGuid Values[] = { FGuid(), FGuid::NewGuid() };
 	for (const FGuid& Value : Values)
 	{
-		Writer.Uuid(Value);
+		Writer.AddUuid(Value);
 	}
 
 	FCbFieldRefIterator Fields = Writer.Save();
@@ -688,7 +689,7 @@ bool FCbWriterDateTimeTest::RunTest(const FString& Parameters)
 	const FDateTime Values[] = { FDateTime(0), FDateTime(2020, 5, 13, 15, 10) };
 	for (FDateTime Value : Values)
 	{
-		Writer.DateTime(Value);
+		Writer.AddDateTime(Value);
 	}
 
 	FCbFieldRefIterator Fields = Writer.Save();
@@ -715,7 +716,7 @@ bool FCbWriterTimeSpanTest::RunTest(const FString& Parameters)
 	const FTimespan Values[] = { FTimespan(0), FTimespan(1, 2, 4, 8) };
 	for (FTimespan Value : Values)
 	{
-		Writer.TimeSpan(Value);
+		Writer.AddTimeSpan(Value);
 	}
 
 	FCbFieldRefIterator Fields = Writer.Save();
@@ -744,166 +745,155 @@ bool FCbWriterComplexTest::RunTest(const FString& Parameters)
 		Writer.BeginObject();
 
 		const uint8 LocalField[] = { uint8(ECbFieldType::IntegerPositive | ECbFieldType::HasFieldName), 1, 'I', 42 };
-		Writer.Name("FieldCopy"_ASV).Field(FCbField(LocalField));
-		Writer.Name("FieldRefCopy"_ASV).Field(FCbFieldRef(FSharedBuffer::Clone(MakeMemoryView(LocalField))));
+		Writer.AddField("FieldCopy"_ASV, FCbField(LocalField));
+		Writer.AddField("FieldRefCopy"_ASV, FCbFieldRef(FSharedBuffer::Clone(MakeMemoryView(LocalField))));
 
 		const uint8 LocalObject[] = { uint8(ECbFieldType::Object | ECbFieldType::HasFieldName), 1, 'O', 7,
 			uint8(ECbFieldType::IntegerPositive | ECbFieldType::HasFieldName), 1, 'I', 42,
 			uint8(ECbFieldType::Null | ECbFieldType::HasFieldName), 1, 'N' };
-		Writer.Name("ObjectCopy"_ASV).Object(FCbObject(LocalObject));
-		Writer.Name("ObjectRefCopy"_ASV).Object(FCbObjectRef(FSharedBuffer::Clone(MakeMemoryView(LocalObject))));
+		Writer.AddObject("ObjectCopy"_ASV, FCbObject(LocalObject));
+		Writer.AddObject("ObjectRefCopy"_ASV, FCbObjectRef(FSharedBuffer::Clone(MakeMemoryView(LocalObject))));
 
 		const uint8 LocalArray[] = { uint8(ECbFieldType::UniformArray | ECbFieldType::HasFieldName), 1, 'A', 4, 2,
 			uint8(ECbFieldType::IntegerPositive), 42, 21 };
-		Writer.Name("ArrayCopy"_ASV).Array(FCbArray(LocalArray));
-		Writer.Name("ArrayRefCopy"_ASV).Array(FCbArrayRef(FSharedBuffer::Clone(MakeMemoryView(LocalArray))));
+		Writer.AddArray("ArrayCopy"_ASV, FCbArray(LocalArray));
+		Writer.AddArray("ArrayRefCopy"_ASV, FCbArrayRef(FSharedBuffer::Clone(MakeMemoryView(LocalArray))));
 
-		Writer.Name("Null"_ASV).Null();
+		Writer.AddNull("Null"_ASV);
 
-		Writer.Name("Binary"_ASV);
-		Writer.BeginObject();
+		Writer.BeginObject("Binary"_ASV);
 		{
-			Writer.Name("Empty"_ASV).Binary(FMemoryView());
-			Writer.Name("Value"_ASV).Binary(MakeMemoryView("BinaryValue"));
-			Writer.Name("LargeValue"_ASV).Binary(MakeMemoryView(FString::ChrN(256, TEXT('.'))));
-			Writer.Name("LargeRefValue"_ASV).Binary(FSharedBuffer::Clone(MakeMemoryView(FString::ChrN(256, TEXT('!')))));
+			Writer.AddBinary("Empty"_ASV, FMemoryView());
+			Writer.AddBinary("Value"_ASV, MakeMemoryView("BinaryValue"));
+			Writer.AddBinary("LargeValue"_ASV, MakeMemoryView(FString::ChrN(256, TEXT('.'))));
+			Writer.AddBinary("LargeRefValue"_ASV, FSharedBuffer::Clone(MakeMemoryView(FString::ChrN(256, TEXT('!')))));
 		}
 		Writer.EndObject();
 
-		Writer.Name("Strings"_ASV);
-		Writer.BeginObject();
+		Writer.BeginObject("Strings"_ASV);
 		{
-			Writer.Name("AnsiString"_ASV).String("AnsiValue"_ASV);
-			Writer.Name("WideString"_ASV).String(FString::ChrN(256, TEXT('.')));
-			Writer.Name("EmptyAnsiString"_ASV).String(FAnsiStringView());
-			Writer.Name("EmptyWideString"_ASV).String(FWideStringView());
-			Writer.Name("AnsiStringLiteral").String("AnsiValue");
-			Writer.Name("WideStringLiteral").String(TEXT("AnsiValue"));
+			Writer.AddString("AnsiString"_ASV, "AnsiValue"_ASV);
+			Writer.AddString("WideString"_ASV, FString::ChrN(256, TEXT('.')));
+			Writer.AddString("EmptyAnsiString"_ASV, FAnsiStringView());
+			Writer.AddString("EmptyWideString"_ASV, FWideStringView());
+			Writer.AddString("AnsiStringLiteral", "AnsiValue");
+			Writer.AddString("WideStringLiteral", TEXT("AnsiValue"));
 		}
 		Writer.EndObject();
 
-		Writer.Name("Integers"_ASV);
-		Writer.BeginArray();
+		Writer.BeginArray("Integers"_ASV);
 		{
-			Writer.Integer(int32(-1));
-			Writer.Integer(int64(-1));
-			Writer.Integer(uint32(1));
-			Writer.Integer(uint64(1));
-			Writer.Integer(MIN_int32);
-			Writer.Integer(MAX_int32);
-			Writer.Integer(MAX_uint32);
-			Writer.Integer(MIN_int64);
-			Writer.Integer(MAX_int64);
-			Writer.Integer(MAX_uint64);
+			Writer.AddInteger(int32(-1));
+			Writer.AddInteger(int64(-1));
+			Writer.AddInteger(uint32(1));
+			Writer.AddInteger(uint64(1));
+			Writer.AddInteger(MIN_int32);
+			Writer.AddInteger(MAX_int32);
+			Writer.AddInteger(MAX_uint32);
+			Writer.AddInteger(MIN_int64);
+			Writer.AddInteger(MAX_int64);
+			Writer.AddInteger(MAX_uint64);
 		}
 		Writer.EndArray();
 
-		Writer.Name("UniformIntegers"_ASV);
-		Writer.BeginArray();
+		Writer.BeginArray("UniformIntegers"_ASV);
 		{
-			Writer.Integer(0);
-			Writer.Integer(MAX_int32);
-			Writer.Integer(MAX_uint32);
-			Writer.Integer(MAX_int64);
-			Writer.Integer(MAX_uint64);
+			Writer.AddInteger(0);
+			Writer.AddInteger(MAX_int32);
+			Writer.AddInteger(MAX_uint32);
+			Writer.AddInteger(MAX_int64);
+			Writer.AddInteger(MAX_uint64);
 		}
 		Writer.EndArray();
 
-		Writer.Name("Float32"_ASV).Float(1.0f);
-		Writer.Name("Float64as32"_ASV).Float(2.0);
-		Writer.Name("Float64"_ASV).Float(3.0e100);
+		Writer.AddFloat("Float32"_ASV, 1.0f);
+		Writer.AddFloat("Float64as32"_ASV, 2.0);
+		Writer.AddFloat("Float64"_ASV, 3.0e100);
 
-		Writer.Name("False"_ASV).Bool(false);
-		Writer.Name("True"_ASV).Bool(true);
+		Writer.AddBool("False"_ASV, false);
+		Writer.AddBool("True"_ASV, true);
 
-		Writer.Name("CompactBinaryAttachment"_ASV).CompactBinaryAttachment(FIoHash());
-		Writer.Name("BinaryAttachment"_ASV).BinaryAttachment(FIoHash());
+		Writer.AddCompactBinaryAttachment("CompactBinaryAttachment"_ASV, FIoHash());
+		Writer.AddBinaryAttachment("BinaryAttachment"_ASV, FIoHash());
+		Writer.AddAttachment("Attachment"_ASV, FCbAttachment());
 
-		Writer.Name("Hash"_ASV).Hash(FIoHash());
-		Writer.Name("Uuid"_ASV).Uuid(FGuid());
+		Writer.AddHash("Hash"_ASV, FIoHash());
+		Writer.AddUuid("Uuid"_ASV, FGuid());
 
-		Writer.Name("DateTimeZero"_ASV).DateTimeTicks(0);
-		Writer.Name("DateTime2020"_ASV).DateTime(FDateTime(2020, 5, 13, 15, 10));
+		Writer.AddDateTimeTicks("DateTimeZero"_ASV, 0);
+		Writer.AddDateTime("DateTime2020"_ASV, FDateTime(2020, 5, 13, 15, 10));
 
-		Writer.Name("TimeSpanZero"_ASV).TimeSpanTicks(0);
-		Writer.Name("TimeSpan"_ASV).TimeSpan(FTimespan(1, 2, 4, 8));
+		Writer.AddTimeSpanTicks("TimeSpanZero"_ASV, 0);
+		Writer.AddTimeSpan("TimeSpan"_ASV, FTimespan(1, 2, 4, 8));
 
-		Writer.Name("NestedObjects"_ASV);
-		Writer.BeginObject();
+		Writer.BeginObject("NestedObjects"_ASV);
 		{
-			Writer.Name("Empty"_ASV);
-			Writer.BeginObject();
+			Writer.BeginObject("Empty"_ASV);
 			Writer.EndObject();
 
-			Writer.Name("Null"_ASV);
-			Writer.BeginObject();
-			Writer.Name("Null"_ASV).Null();
+			Writer.BeginObject("Null"_ASV);
+			Writer.AddNull("Null"_ASV);
 			Writer.EndObject();
 		}
 		Writer.EndObject();
 
-		Writer.Name("NestedArrays"_ASV);
-		Writer.BeginArray();
+		Writer.BeginArray("NestedArrays"_ASV);
 		{
 			Writer.BeginArray();
 			Writer.EndArray();
 
 			Writer.BeginArray();
-			Writer.Null();
-			Writer.Null();
-			Writer.Null();
+			Writer.AddNull();
+			Writer.AddNull();
+			Writer.AddNull();
 			Writer.EndArray();
 
 			Writer.BeginArray();
-			Writer.Bool(false);
-			Writer.Bool(false);
-			Writer.Bool(false);
+			Writer.AddBool(false);
+			Writer.AddBool(false);
+			Writer.AddBool(false);
 			Writer.EndArray();
 
 			Writer.BeginArray();
-			Writer.Bool(true);
-			Writer.Bool(true);
-			Writer.Bool(true);
+			Writer.AddBool(true);
+			Writer.AddBool(true);
+			Writer.AddBool(true);
 			Writer.EndArray();
 		}
 		Writer.EndArray();
 
-		Writer.Name("ArrayOfObjects"_ASV);
-		Writer.BeginArray();
+		Writer.BeginArray("ArrayOfObjects"_ASV);
 		{
 			Writer.BeginObject();
 			Writer.EndObject();
 
 			Writer.BeginObject();
-			Writer.Name("Null"_ASV).Null();
+			Writer.AddNull("Null"_ASV);
 			Writer.EndObject();
 		}
 		Writer.EndArray();
 
-		Writer.Name("LargeArray"_ASV);
-		Writer.BeginArray();
+		Writer.BeginArray("LargeArray"_ASV);
 		for (int Index = 0; Index < 256; ++Index)
 		{
-			Writer.Integer(Index - 128);
+			Writer.AddInteger(Index - 128);
 		}
 		Writer.EndArray();
 
-		Writer.Name("LargeUniformArray"_ASV);
-		Writer.BeginArray();
+		Writer.BeginArray("LargeUniformArray"_ASV);
 		for (int Index = 0; Index < 256; ++Index)
 		{
-			Writer.Integer(Index);
+			Writer.AddInteger(Index);
 		}
 		Writer.EndArray();
 
-		Writer.Name("NestedUniformArray"_ASV);
-		Writer.BeginArray();
+		Writer.BeginArray("NestedUniformArray"_ASV);
 		for (int Index = 0; Index < 16; ++Index)
 		{
 			Writer.BeginArray();
 			for (int Value = 0; Value < 4; ++Value)
 			{
-				Writer.Integer(Value);
+				Writer.AddInteger(Value);
 			}
 			Writer.EndArray();
 		}
@@ -982,6 +972,8 @@ bool FCbWriterStreamTest::RunTest(const FString& Parameters)
 
 		Writer << "False"_ASV << false << "True"_ASV << true;
 
+		Writer << "Attachment"_ASV << FCbAttachment();
+
 		Writer << "Hash"_ASV << FIoHash();
 		Writer << "Uuid"_ASV << FGuid();
 
@@ -1014,29 +1006,29 @@ bool FCbWriterStateTest::RunTest(const FString& Parameters)
 	// Assert on under-sized save buffer.
 	//uint8 ZeroFieldSmall[1];
 	//Writer.Reset();
-	//Writer.Integer(0);
+	//Writer.AddInteger(0);
 	//Writer.Save(MakeMemoryView(ZeroFieldSmall));
 
 	// Assert on over-sized save buffer.
 	//uint8 ZeroFieldLarge[3];
 	//Writer.Reset();
-	//Writer.Integer(0);
+	//Writer.AddInteger(0);
 	//Writer.Save(MakeMemoryView(ZeroFieldLarge));
 
 	// Assert on empty name.
-	//Writer.Name(""_ASV);
+	//Writer.SetName(""_ASV);
 
 	// Assert on name after name.
-	//Writer.Name("Field"_ASV).Name("Field"_ASV);
+	//Writer.SetName("Field"_ASV).SetName("Field"_ASV);
 
 	// Assert on missing name.
 	//Writer.BeginObject();
-	//Writer.Null();
+	//Writer.AddNull();
 	//Writer.EndObject();
 
 	// Assert on name in array.
 	//Writer.BeginArray();
-	//Writer.Name("Field"_ASV);
+	//Writer.SetName("Field"_ASV);
 	//Writer.EndArray();
 
 	// Assert on save in object.
@@ -1061,7 +1053,7 @@ bool FCbWriterStateTest::RunTest(const FString& Parameters)
 
 	// Assert on object end after name with no value.
 	//Writer.BeginObject();
-	//Writer.Name("Field"_ASV);
+	//Writer.SetName("Field"_ASV);
 	//Writer.EndObject();
 
 	// Assert on writing a field with no value.
diff --git a/Engine/Source/Runtime/Core/Public/Serialization/CompactBinaryWriter.h b/Engine/Source/Runtime/Core/Public/Serialization/CompactBinaryWriter.h
--- a/Engine/Source/Runtime/Core/Public/Serialization/CompactBinaryWriter.h
+++ b/Engine/Source/Runtime/Core/Public/Serialization/CompactBinaryWriter.h
@@ -10,6 +10,7 @@
 ///////////////////////////////////////////////////////////////////////////////////////////////////////////////////////
 
 class FArchive;
+class FCbAttachment;
 struct FDateTime;
 struct FGuid;
 struct FIoHash;
@@ -101,91 +102,123 @@ public:
 	CORE_API uint64 GetSaveSize() const;
 
 	/**
-	 * Write a field name.
+	 * Sets the name of the next field to be written.
 	 *
-	 * It is not valid to call this function when writing an array field.
+	 * It is not valid to call this function when writing a field inside an array.
 	 * Names must be valid UTF-8 and must be unique within an object.
 	 */
-	CORE_API FCbWriter& Name(FAnsiStringView Name);
+	CORE_API FCbWriter& SetName(FAnsiStringView Name);
 
 	/** Copy the value (not the name) of an existing field. */
-	CORE_API void Field(const FCbField& Value);
+	inline void AddField(FAnsiStringView Name, const FCbField& Value) { SetName(Name); AddField(Value); }
+	CORE_API void AddField(const FCbField& Value);
 	/** Copy the value (not the name) of an existing field. Holds a reference if owned. */
-	CORE_API void Field(const FCbFieldRef& Value);
+	inline void AddField(FAnsiStringView Name, const FCbFieldRef& Value) { SetName(Name); AddField(Value); }
+	CORE_API void AddField(const FCbFieldRef& Value);
 
 	/** Begin a new object. Must have a matching call to EndObject. */
+	inline void BeginObject(FAnsiStringView Name) { SetName(Name); BeginObject(); }
 	CORE_API void BeginObject();
 	/** End an object after its fields have been written. */
 	CORE_API void EndObject();
 
 	/** Copy the value (not the name) of an existing object. */
-	CORE_API void Object(const FCbObject& Value);
+	inline void AddObject(FAnsiStringView Name, const FCbObject& Value) { SetName(Name); AddObject(Value); }
+	CORE_API void AddObject(const FCbObject& Value);
 	/** Copy the value (not the name) of an existing object. Holds a reference if owned. */
-	CORE_API void Object(const FCbObjectRef& Value);
+	inline void AddObject(FAnsiStringView Name, const FCbObjectRef& Value) { SetName(Name); AddObject(Value); }
+	CORE_API void AddObject(const FCbObjectRef& Value);
 
 	/** Begin a new array. Must have a matching call to EndArray. */
+	inline void BeginArray(FAnsiStringView Name) { SetName(Name); BeginArray(); }
 	CORE_API void BeginArray();
 	/** End an array after its fields have been written. */
 	CORE_API void EndArray();
 
 	/** Copy the value (not the name) of an existing array. */
-	CORE_API void Array(const FCbArray& Value);
+	inline void AddArray(FAnsiStringView Name, const FCbArray& Value) { SetName(Name); AddArray(Value); }
+	CORE_API void AddArray(const FCbArray& Value);
 	/** Copy the value (not the name) of an existing array. Holds a reference if owned. */
-	CORE_API void Array(const FCbArrayRef& Value);
+	inline void AddArray(FAnsiStringView Name, const FCbArrayRef& Value) { SetName(Name); AddArray(Value); }
+	CORE_API void AddArray(const FCbArrayRef& Value);
 
 	/** Write a null field. */
-	CORE_API void Null();
+	inline void AddNull(FAnsiStringView Name) { SetName(Name); AddNull(); }
+	CORE_API void AddNull();
 
 	/** Write a binary field by copying Size bytes from Value. */
-	CORE_API void Binary(const void* Value, uint64 Size);
+	inline void AddBinary(FAnsiStringView Name, const void* Value, uint64 Size) { SetName(Name); AddBinary(Value, Size); }
+	CORE_API void AddBinary(const void* Value, uint64 Size);
 	/** Write a binary field by copying the view. */
-	inline void Binary(FMemoryView Value) { Binary(Value.GetData(), Value.GetSize()); }
+	inline void AddBinary(FAnsiStringView Name, FMemoryView Value) { SetName(Name); AddBinary(Value); }
+	inline void AddBinary(FMemoryView Value) { AddBinary(Value.GetData(), Value.GetSize()); }
 	/** Write a binary field by copying the buffer. Holds a reference if owned. */
-	CORE_API void Binary(const FSharedBuffer& Buffer);
+	inline void AddBinary(FAnsiStringView Name, FSharedBuffer Value) { SetName(Name); AddBinary(MoveTemp(Value)); }
+	CORE_API void AddBinary(FSharedBuffer Value);
 
 	/** Write a string field by copying the UTF-8 value. */
-	CORE_API void String(FAnsiStringView Value);
+	inline void AddString(FAnsiStringView Name, FAnsiStringView Value) { SetName(Name); AddString(Value); }
+	CORE_API void AddString(FAnsiStringView Value);
 	/** Write a string field by converting the UTF-16 value to UTF-8. */
-	CORE_API void String(FWideStringView Value);
+	inline void AddString(FAnsiStringView Name, FWideStringView Value) { SetName(Name); AddString(Value); }
+	CORE_API void AddString(FWideStringView Value);
 
 	/** Write an integer field. */
-	CORE_API void Integer(int32 Value);
+	inline void AddInteger(FAnsiStringView Name, int32 Value) { SetName(Name); AddInteger(Value); }
+	CORE_API void AddInteger(int32 Value);
 	/** Write an integer field. */
-	CORE_API void Integer(int64 Value);
+	inline void AddInteger(FAnsiStringView Name, int64 Value) { SetName(Name); AddInteger(Value); }
+	CORE_API void AddInteger(int64 Value);
 	/** Write an integer field. */
-	CORE_API void Integer(uint32 Value);
+	inline void AddInteger(FAnsiStringView Name, uint32 Value) { SetName(Name); AddInteger(Value); }
+	CORE_API void AddInteger(uint32 Value);
 	/** Write an integer field. */
-	CORE_API void Integer(uint64 Value);
+	inline void AddInteger(FAnsiStringView Name, uint64 Value) { SetName(Name); AddInteger(Value); }
+	CORE_API void AddInteger(uint64 Value);
 
 	/** Write a float field from a 32-bit float value. */
-	CORE_API void Float(float Value);
+	inline void AddFloat(FAnsiStringView Name, float Value) { SetName(Name); AddFloat(Value); }
+	CORE_API void AddFloat(float Value);
 	/** Write a float field from a 64-bit float value. */
-	CORE_API void Float(double Value);
+	inline void AddFloat(FAnsiStringView Name, double Value) { SetName(Name); AddFloat(Value); }
+	CORE_API void AddFloat(double Value);
 
 	/** Write a bool field. */
-	CORE_API void Bool(bool bValue);
-
-	/** Write a field referencing compact binary data by its hash. */
-	CORE_API void CompactBinaryAttachment(const FIoHash& Value);
-	/** Write a field referencing binary data by its hash. */
-	CORE_API void BinaryAttachment(const FIoHash& Value);
+	inline void AddBool(FAnsiStringView Name, bool bValue) { SetName(Name); AddBool(bValue); }
+	CORE_API void AddBool(bool bValue);
+
+	/** Write a field referencing a compact binary attachment by its hash. */
+	inline void AddCompactBinaryAttachment(FAnsiStringView Name, const FIoHash& Value) { SetName(Name); AddCompactBinaryAttachment(Value); }
+	CORE_API void AddCompactBinaryAttachment(const FIoHash& Value);
+	/** Write a field referencing a binary attachment by its hash. */
+	inline void AddBinaryAttachment(FAnsiStringView Name, const FIoHash& Value) { SetName(Name); AddBinaryAttachment(Value); }
+	CORE_API void AddBinaryAttachment(const FIoHash& Value);
+	/** Write a field referencing the attachment by its hash. */
+	inline void AddAttachment(FAnsiStringView Name, const FCbAttachment& Attachment) { SetName(Name); AddAttachment(Attachment); }
+	CORE_API void AddAttachment(const FCbAttachment& Attachment);
 
 	/** Write a hash field. */
-	CORE_API void Hash(const FIoHash& Value);
+	inline void AddHash(FAnsiStringView Name, const FIoHash& Value) { SetName(Name); AddHash(Value); }
+	CORE_API void AddHash(const FIoHash& Value);
 	/** Write a UUID field. */
-	CORE_API void Uuid(const FGuid& Value);
+	inline void AddUuid(FAnsiStringView Name, const FGuid& Value) { SetName(Name); AddUuid(Value); }
+	CORE_API void AddUuid(const FGuid& Value);
 
 	/** Write a date/time field with the specified count of 100ns ticks since the epoch. */
-	CORE_API void DateTimeTicks(int64 Ticks);
+	inline void AddDateTimeTicks(FAnsiStringView Name, int64 Ticks) { SetName(Name); AddDateTimeTicks(Ticks); }
+	CORE_API void AddDateTimeTicks(int64 Ticks);
 
 	/** Write a date/time field. */
-	CORE_API void DateTime(FDateTime Value);
+	inline void AddDateTime(FAnsiStringView Name, FDateTime Value) { SetName(Name); AddDateTime(Value); }
+	CORE_API void AddDateTime(FDateTime Value);
 
 	/** Write a time span field with the specified count of 100ns ticks. */
-	CORE_API void TimeSpanTicks(int64 Ticks);
+	inline void AddTimeSpanTicks(FAnsiStringView Name, int64 Ticks) { SetName(Name); AddTimeSpanTicks(Ticks); }
+	CORE_API void AddTimeSpanTicks(int64 Ticks);
 
 	/** Write a time span field. */
-	CORE_API void TimeSpan(FTimespan Value);
+	inline void AddTimeSpan(FAnsiStringView Name, FTimespan Value) { SetName(Name); AddTimeSpan(Value); }
+	CORE_API void AddTimeSpan(FTimespan Value);
 
 	/** Private flags that are public to work with ENUM_CLASS_FLAGS. */
 	enum class EStateFlags : uint8;
@@ -197,17 +230,17 @@ protected:
 private:
 	friend FCbWriter& operator<<(FCbWriter& Writer, FAnsiStringView NameOrValue);
 
-	/** Begin writing a field. Called twice for fields with names and once otherwise. */
+	/** Begin writing a field. May be called twice for named fields. */
 	void BeginField();
 
 	/** Finish writing a field by writing its type. */
 	void EndField(ECbFieldType Type);
 
-	/** Write the field name if valid in this state, otherwise write the string value. */
-	CORE_API void NameOrString(FAnsiStringView NameOrValue);
+	/** Set the field name if valid in this state, otherwise write add a string field. */
+	CORE_API void SetNameOrAddString(FAnsiStringView NameOrValue);
 
-	/** A view of the name of the current field (if any), otherwise the empty view. */
-	FAnsiStringView CurrentName() const;
+	/** Returns a view of the name of the active field, if any, otherwise the empty view. */
+	FAnsiStringView GetActiveName() const;
 
 	/** Remove field types after the first to make the sequence uniform. */
 	void MakeFieldsUniform(int64 FieldBeginOffset, int64 FieldEndOffset);
@@ -266,7 +299,7 @@ private:
 /** Write the field name if valid in this state, otherwise write the string value. */
 inline FCbWriter& operator<<(FCbWriter& Writer, FAnsiStringView NameOrValue)
 {
-	Writer.NameOrString(NameOrValue);
+	Writer.SetNameOrAddString(NameOrValue);
 	return Writer;
 }
 
@@ -278,109 +311,115 @@ inline FCbWriter& operator<<(FCbWriter& Writer, const ANSICHAR* NameOrValue)
 
 inline FCbWriter& operator<<(FCbWriter& Writer, const FCbField& Value)
 {
-	Writer.Field(Value);
+	Writer.AddField(Value);
 	return Writer;
 }
 
 inline FCbWriter& operator<<(FCbWriter& Writer, const FCbFieldRef& Value)
 {
-	Writer.Field(Value);
+	Writer.AddField(Value);
 	return Writer;
 }
 
 inline FCbWriter& operator<<(FCbWriter& Writer, const FCbObject& Value)
 {
-	Writer.Object(Value);
+	Writer.AddObject(Value);
 	return Writer;
 }
 
 inline FCbWriter& operator<<(FCbWriter& Writer, const FCbObjectRef& Value)
 {
-	Writer.Object(Value);
+	Writer.AddObject(Value);
 	return Writer;
 }
 
 inline FCbWriter& operator<<(FCbWriter& Writer, const FCbArray& Value)
 {
-	Writer.Array(Value);
+	Writer.AddArray(Value);
 	return Writer;
 }
 
 inline FCbWriter& operator<<(FCbWriter& Writer, const FCbArrayRef& Value)
 {
-	Writer.Array(Value);
+	Writer.AddArray(Value);
 	return Writer;
 }
 
 inline FCbWriter& operator<<(FCbWriter& Writer, nullptr_t)
 {
-	Writer.Null();
+	Writer.AddNull();
 	return Writer;
 }
 
 inline FCbWriter& operator<<(FCbWriter& Writer, FWideStringView Value)
 {
-	Writer.String(Value);
+	Writer.AddString(Value);
 	return Writer;
 }
 
 inline FCbWriter& operator<<(FCbWriter& Writer, const WIDECHAR* Value)
 {
-	Writer.String(Value);
+	Writer.AddString(Value);
 	return Writer;
 }
 
 inline FCbWriter& operator<<(FCbWriter& Writer, int32 Value)
 {
-	Writer.Integer(Value);
+	Writer.AddInteger(Value);
 	return Writer;
 }
 
 inline FCbWriter& operator<<(FCbWriter& Writer, int64 Value)
 {
-	Writer.Integer(Value);
+	Writer.AddInteger(Value);
 	return Writer;
 }
 
 inline FCbWriter& operator<<(FCbWriter& Writer, uint32 Value)
 {
-	Writer.Integer(Value);
+	Writer.AddInteger(Value);
 	return Writer;
 }
 
 inline FCbWriter& operator<<(FCbWriter& Writer, uint64 Value)
 {
-	Writer.Integer(Value);
+	Writer.AddInteger(Value);
 	return Writer;
 }
 
 inline FCbWriter& operator<<(FCbWriter& Writer, float Value)
 {
-	Writer.Float(Value);
+	Writer.AddFloat(Value);
 	return Writer;
 }
 
 inline FCbWriter& operator<<(FCbWriter& Writer, double Value)
 {
-	Writer.Float(Value);
+	Writer.AddFloat(Value);
 	return Writer;
 }
 
 inline FCbWriter& operator<<(FCbWriter& Writer, bool Value)
 {
-	Writer.Bool(Value);
+	Writer.AddBool(Value);
+	return Writer;
+}
+
+inline FCbWriter& operator<<(FCbWriter& Writer, const FCbAttachment& Attachment)
+{
+	Writer.AddAttachment(Attachment);
 	return Writer;
 }
 
 inline FCbWriter& operator<<(FCbWriter& Writer, const FIoHash& Value)
 {
-	Writer.Hash(Value);
+	Writer.AddHash(Value);
 	return Writer;
 }
 
 inline FCbWriter& operator<<(FCbWriter& Writer, const FGuid& Value)
 {
-	Writer.Uuid(Value);
+	Writer.AddUuid(Value);
 	return Writer;
 }
 
