Make FindFilesRecursive parallel by using the low level recursive file iteration
Fix crash when trying to use ParallelFor before the task graph is ready

- 12.3s to 1.3s for SearchForPackageOnDisk when running world partition commandlet

diff --git a/Engine/Source/Runtime/Core/Private/GenericPlatform/GenericPlatformFile.cpp b/Engine/Source/Runtime/Core/Private/GenericPlatform/GenericPlatformFile.cpp
--- a/Engine/Source/Runtime/Core/Private/GenericPlatform/GenericPlatformFile.cpp
+++ b/Engine/Source/Runtime/Core/Private/GenericPlatform/GenericPlatformFile.cpp
@@ -577,6 +577,7 @@ bool IPlatformFile::IterateDirectoryRecursively(const TCHAR* Directory, FDirecto
 	TArray<FString> DirectoriesToVisitNext;
 	DirectoriesToVisitNext.Add(Directory);
 
+	const bool IsTaskGraphReady = FTaskGraphInterface::Get().IsRunning();
 	TAtomic<bool> bResult(true);
 	FRecurse Recurse(*this, Visitor, DirectoriesToVisitNext);
 	while (bResult && DirectoriesToVisitNext.Num() > 0)
@@ -591,7 +592,7 @@ bool IPlatformFile::IterateDirectoryRecursively(const TCHAR* Directory, FDirecto
 					bResult = false;
 				}
 			},
-			Visitor.IsThreadSafe() ? EParallelForFlags::Unbalanced : EParallelForFlags::ForceSingleThread
+			(IsTaskGraphReady && Visitor.IsThreadSafe()) ? EParallelForFlags::Unbalanced : EParallelForFlags::ForceSingleThread
 		);
 	}
 
diff --git a/Engine/Source/Runtime/Core/Private/HAL/FileManagerGeneric.cpp b/Engine/Source/Runtime/Core/Private/HAL/FileManagerGeneric.cpp
--- a/Engine/Source/Runtime/Core/Private/HAL/FileManagerGeneric.cpp
+++ b/Engine/Source/Runtime/Core/Private/HAL/FileManagerGeneric.cpp
@@ -16,6 +16,7 @@
 #include "Misc/CommandLine.h"
 #include "Misc/Paths.h"
 #include "Misc/SecureHash.h"
+#include "Misc/ScopeRWLock.h"
 #include "Templates/UniquePtr.h"
 #include <time.h>
 
@@ -405,40 +406,50 @@ FFileStatData FFileManagerGeneric::GetStatData(const TCHAR* FilenameOrDirectory)
 	return GetLowLevel().GetStatData(FilenameOrDirectory);
 }
 
-void FFileManagerGeneric::FindFiles( TArray<FString>& Result, const TCHAR* InFilename, bool Files, bool Directories )
+namespace FileManagerGenericImpl
 {
 	class FFileMatch : public IPlatformFile::FDirectoryVisitor
 	{
 	public:
 		TArray<FString>& Result;
+		FRWLock ResultLock;
 		FString WildCard;
 		bool bFiles;
 		bool bDirectories;
-		FFileMatch( TArray<FString>& InResult, const FString& InWildCard, bool bInFiles, bool bInDirectories )
-			: Result( InResult )
-			, WildCard( InWildCard )
-			, bFiles( bInFiles )
-			, bDirectories( bInDirectories )
+		bool bStoreFullPath;
+		FFileMatch(TArray<FString>& InResult, const FString& InWildCard, bool bInFiles, bool bInDirectories, bool bInStoreFullPath = false)
+			: IPlatformFile::FDirectoryVisitor(EDirectoryVisitorFlags::ThreadSafe)
+			, Result(InResult)
+			, WildCard(InWildCard)
+			, bFiles(bInFiles)
+			, bDirectories(bInDirectories)
+			, bStoreFullPath(bInStoreFullPath)
 		{
 		}
-		virtual bool Visit( const TCHAR* FilenameOrDirectory, bool bIsDirectory )
+		virtual bool Visit(const TCHAR* FilenameOrDirectory, bool bIsDirectory)
 		{
 			if ((bIsDirectory && bDirectories) || (!bIsDirectory && bFiles))
 			{
-				const FString Filename = FPaths::GetCleanFilename(FilenameOrDirectory);
+				FString Filename = FPaths::GetCleanFilename(FilenameOrDirectory);
 				if (Filename.MatchesWildcard(WildCard))
 				{
-					new( Result ) FString(Filename);
+					FString FullPath = bStoreFullPath ? FString(FilenameOrDirectory) : MoveTemp(Filename);
+					FWriteScopeLock ScopeLock(ResultLock);
+					Result.Add(MoveTemp(FullPath));
 				}
 			}
 			return true;
 		}
 	};
+}
+
+void FFileManagerGeneric::FindFiles( TArray<FString>& Result, const TCHAR* InFilename, bool Files, bool Directories )
+{
 	FString Filename( InFilename );
 	FPaths::NormalizeFilename( Filename );
 	const FString CleanFilename = FPaths::GetCleanFilename(Filename);
 	const bool bFindAllFiles = CleanFilename == TEXT("*") || CleanFilename == TEXT("*.*");
-	FFileMatch FileMatch( Result, bFindAllFiles ? TEXT("*") : CleanFilename, Files, Directories );
+	FileManagerGenericImpl::FFileMatch FileMatch( Result, bFindAllFiles ? TEXT("*") : CleanFilename, Files, Directories );
 	GetLowLevel().IterateDirectory( *FPaths::GetPath(Filename), FileMatch );
 }
 
@@ -631,26 +642,15 @@ void FFileManagerGeneric::FindFilesRecursive( TArray<FString>& FileNames, const
 	FindFilesRecursiveInternal(FileNames, StartDirectory, Filename, Files, Directories);
 }
 
-void FFileManagerGeneric::FindFilesRecursiveInternal( TArray<FString>& FileNames, const TCHAR* StartDirectory, const TCHAR* Filename, bool Files, bool Directories)
+void FFileManagerGeneric::FindFilesRecursiveInternal( TArray<FString>& FileNames, const TCHAR* StartDirectory, const TCHAR* InFilename, bool Files, bool Directories)
 {
-	FString CurrentSearch = FString(StartDirectory) / Filename;
-	TArray<FString> Result;
-	FindFiles(Result, *CurrentSearch, Files, Directories);
-
-	for (int32 i=0; i<Result.Num(); i++)
-	{
-		FileNames.Add(FString(StartDirectory) / Result[i]);
-	}
-
-	TArray<FString> SubDirs;
-	FString RecursiveDirSearch = FString(StartDirectory) / TEXT("*");
-	FindFiles(SubDirs, *RecursiveDirSearch, false, true);
-
-	for (int32 SubDirIdx=0; SubDirIdx<SubDirs.Num(); SubDirIdx++)
-	{
-		FString SubDir = FString(StartDirectory) / SubDirs[SubDirIdx];
-		FindFilesRecursiveInternal(FileNames, *SubDir, Filename, Files, Directories);
-	}
+	FString Filename(InFilename);
+	FPaths::NormalizeFilename(Filename);
+	const FString CleanFilename = FPaths::GetCleanFilename(Filename);
+	const bool bFindAllFiles = CleanFilename == TEXT("*") || CleanFilename == TEXT("*.*");
+	const bool bStoreFullPath = true;
+	FileManagerGenericImpl::FFileMatch FileMatch(FileNames, bFindAllFiles ? TEXT("*") : CleanFilename, Files, Directories, bStoreFullPath);
+	GetLowLevel().IterateDirectoryRecursively(StartDirectory, FileMatch);
 }
 
 FArchiveFileReaderGeneric::FArchiveFileReaderGeneric( IFileHandle* InHandle, const TCHAR* InFilename, int64 InSize, uint32 InBufferSize )
