Work around MSVC compiler bug when /ZW is used. The forward-declaration of complex templates must use the same template argument names due to presumed quirks in how MSVC tokenises the template statement.

#rnx

diff --git a/Engine/Source/Runtime/Core/Public/Containers/ContainersFwd.h b/Engine/Source/Runtime/Core/Public/Containers/ContainersFwd.h
--- a/Engine/Source/Runtime/Core/Public/Containers/ContainersFwd.h
+++ b/Engine/Source/Runtime/Core/Public/Containers/ContainersFwd.h
@@ -23,8 +23,8 @@ template<typename T> using TArrayView64 = TArrayView<T, int64>;
 template<typename T, typename SizeType = int32> using TConstArrayView = TArrayView<const T, SizeType>;
 template<typename T> using TConstArrayView64 = TConstArrayView<T, int64>;
 template<typename T> class TTransArray;
-template<typename KeyType, typename ValueType, bool bInAllowDuplicateKeys> struct TDefaultMapHashableKeyFuncs;
-template<typename KeyType, typename ValueType, typename SetAllocator = FDefaultSetAllocator, typename KeyFuncs = TDefaultMapHashableKeyFuncs<KeyType, ValueType, false> > class TMap;
+template<typename InKeyType, typename InValueType, bool bInAllowDuplicateKeys> struct TDefaultMapHashableKeyFuncs;
+template<typename InKeyType, typename InValueType, typename SetAllocator = FDefaultSetAllocator, typename KeyFuncs = TDefaultMapHashableKeyFuncs<InKeyType, InValueType, false> > class TMap;
 template<typename KeyType, typename ValueType, typename SetAllocator = FDefaultSetAllocator, typename KeyFuncs = TDefaultMapHashableKeyFuncs<KeyType, ValueType, true > > class TMultiMap;
 template <typename T = void > struct TLess;
 template <typename> struct TTypeTraits;
diff --git a/Engine/Source/Runtime/Core/Public/Containers/UnrealString.h b/Engine/Source/Runtime/Core/Public/Containers/UnrealString.h
--- a/Engine/Source/Runtime/Core/Public/Containers/UnrealString.h
+++ b/Engine/Source/Runtime/Core/Public/Containers/UnrealString.h
@@ -27,7 +27,7 @@
 #include "Templates/IsFloatingPoint.h"
 
 struct FStringFormatArg;
-template<typename KeyType,typename ValueType,typename SetAllocator ,typename KeyFuncs > class TMap;
+template<typename InKeyType,typename InValueType,typename SetAllocator ,typename KeyFuncs > class TMap;
 
 typedef TMap<FString, FStringFormatArg> FStringFormatNamedArguments;
 typedef TArray<FStringFormatArg> FStringFormatOrderedArguments;
