Fix crash in engine startup. FTaskGraphInterface::IsRunning() is static and shouldn't be called via FTaskGraphInterface::Get().IsRunning(), otherwise we hit the nullptr check inside Get() when the task graph has not yet started.

diff --git a/Engine/Source/Runtime/Core/Private/GenericPlatform/GenericPlatformFile.cpp b/Engine/Source/Runtime/Core/Private/GenericPlatform/GenericPlatformFile.cpp
--- a/Engine/Source/Runtime/Core/Private/GenericPlatform/GenericPlatformFile.cpp
+++ b/Engine/Source/Runtime/Core/Private/GenericPlatform/GenericPlatformFile.cpp
@@ -577,7 +577,7 @@ bool IPlatformFile::IterateDirectoryRecursively(const TCHAR* Directory, FDirecto
 	TArray<FString> DirectoriesToVisitNext;
 	DirectoriesToVisitNext.Add(Directory);
 
-	const bool IsTaskGraphReady = FTaskGraphInterface::Get().IsRunning();
+	const bool IsTaskGraphReady = FTaskGraphInterface::IsRunning();
 	TAtomic<bool> bResult(true);
 	FRecurse Recurse(*this, Visitor, DirectoriesToVisitNext);
 	while (bResult && DirectoriesToVisitNext.Num() > 0)
