Implements r.ResetRenderTargetsExtent command

diff --git a/Engine/Source/Runtime/Core/Private/HAL/ConsoleManager.cpp b/Engine/Source/Runtime/Core/Private/HAL/ConsoleManager.cpp
--- a/Engine/Source/Runtime/Core/Private/HAL/ConsoleManager.cpp
+++ b/Engine/Source/Runtime/Core/Private/HAL/ConsoleManager.cpp
@@ -2004,6 +2004,7 @@ void CreateConsoleVariables()
 	IConsoleManager::Get().RegisterConsoleCommand(TEXT("Stat MapBuildData"),	TEXT(""), ECVF_Cheat);
 	IConsoleManager::Get().RegisterConsoleCommand(TEXT("r.ResetViewState"), TEXT("Reset some state (e.g. TemporalAA index) to make rendering more deterministic (for automated screenshot verification)"), ECVF_Cheat);
 	IConsoleManager::Get().RegisterConsoleCommand(TEXT("r.RHI.Name"),		TEXT("Show current RHI's name"), ECVF_Cheat);
+	IConsoleManager::Get().RegisterConsoleCommand(TEXT("r.ResetRenderTargetsExtent"), TEXT("To reset internal render target extents"), ECVF_Cheat);
 #endif // !(UE_BUILD_SHIPPING || UE_BUILD_TEST)
 
 
diff --git a/Engine/Source/Runtime/Renderer/Private/Renderer.cpp b/Engine/Source/Runtime/Renderer/Private/Renderer.cpp
--- a/Engine/Source/Runtime/Renderer/Private/Renderer.cpp
+++ b/Engine/Source/Runtime/Renderer/Private/Renderer.cpp
@@ -451,6 +451,7 @@ static void VisualizeTextureExec( const TCHAR* Cmd, FOutputDevice &Ar )
 }
 
 extern void NaniteStatsFilterExec(const TCHAR* Cmd, FOutputDevice& Ar);
+extern void ResetRenderTargetsExtent(FOutputDevice& Ar);
 
 static bool RendererExec( UWorld* InWorld, const TCHAR* Cmd, FOutputDevice& Ar )
 {
@@ -479,6 +480,11 @@ static bool RendererExec( UWorld* InWorld, const TCHAR* Cmd, FOutputDevice& Ar )
 			: TEXT("<NULL DynamicRHI>"));
 		return true;
 	}
+	else if (FParse::Command(&Cmd, TEXT("r.ResetRenderTargetsExtent")))
+	{
+		ResetRenderTargetsExtent(Ar);
+		return true;
+	}
 #endif
 
 	return false;
diff --git a/Engine/Source/Runtime/Renderer/Private/SceneTextures.cpp b/Engine/Source/Runtime/Renderer/Private/SceneTextures.cpp
--- a/Engine/Source/Runtime/Renderer/Private/SceneTextures.cpp
+++ b/Engine/Source/Runtime/Renderer/Private/SceneTextures.cpp
@@ -510,6 +510,12 @@ public:
 		return DesiredExtent;
 	}
 
+	void ResetRenderTargetsExtent()
+	{
+		LastStereoExtent = FIntPoint(0, 0);
+		LastExtent = FIntPoint(0, 0);
+	}
+
 private:
 	enum class ERenderTargetHistory
 	{
@@ -558,6 +564,13 @@ private:
 	uint32 ThisFrameNumber = 0;
 };
 
+
+void ResetRenderTargetsExtent(FOutputDevice& Ar)
+{
+	FSceneTextureExtentState::Get().ResetRenderTargetsExtent();
+	Ar.Logf(TEXT("Internal render target extents reseted"));
+}
+
 ENUM_CLASS_FLAGS(FSceneTextureExtentState::ERenderTargetHistory);
 
 FSceneTexturesConfig FSceneTexturesConfig::GlobalInstance;
