added missing copyright notice

diff --git a/Engine/Source/Runtime/Core/Public/Containers/MpscQueue.h b/Engine/Source/Runtime/Core/Public/Containers/MpscQueue.h
--- a/Engine/Source/Runtime/Core/Public/Containers/MpscQueue.h
+++ b/Engine/Source/Runtime/Core/Public/Containers/MpscQueue.h
@@ -1,3 +1,5 @@
+// Copyright Epic Games, Inc. All Rights Reserved.
+
 #pragma once
 
 #include "CoreTypes.h"
diff --git a/Engine/Source/Runtime/Core/Public/Containers/SpscQueue.h b/Engine/Source/Runtime/Core/Public/Containers/SpscQueue.h
--- a/Engine/Source/Runtime/Core/Public/Containers/SpscQueue.h
+++ b/Engine/Source/Runtime/Core/Public/Containers/SpscQueue.h
@@ -1,3 +1,5 @@
+// Copyright Epic Games, Inc. All Rights Reserved.
+
 #pragma once
 
 #include "CoreTypes.h"
diff --git a/Engine/Source/Runtime/Core/Public/Tests/Benchmark.h b/Engine/Source/Runtime/Core/Public/Tests/Benchmark.h
--- a/Engine/Source/Runtime/Core/Public/Tests/Benchmark.h
+++ b/Engine/Source/Runtime/Core/Public/Tests/Benchmark.h
@@ -1,3 +1,5 @@
+// Copyright Epic Games, Inc. All Rights Reserved.
+
 #pragma once
 
 #include "CoreTypes.h"
