On Windows, fixed CRC (out of process mode for Editor) generating an incomplete portable callstack when the crash occurred because a null function pointer was invoked
  - When CRC runs out of process, instead of reading the current thread context of the crashed thread, read and use the crash context that was reported during the crash (which is different).
  - Added an optional context parameter to FGenericPlatformStackWalk::CaptureThreadStackBackTrace(), implemented it across all platforms, but only used on Windows.

On Windows, fixed InitStackWalking() and InitStackWalkingForProcess() to reset the process that needs to be walked.
  - CRC, running out of process may run its own process or the Editor process and which ever was walking first ruled out the other.

diff --git a/Engine/Source/Programs/CrashReportClient/Private/CrashReportClientApp.cpp b/Engine/Source/Programs/CrashReportClient/Private/CrashReportClientApp.cpp
--- a/Engine/Source/Programs/CrashReportClient/Private/CrashReportClientApp.cpp
+++ b/Engine/Source/Programs/CrashReportClient/Private/CrashReportClientApp.cpp
@@ -451,17 +451,47 @@ FPlatformErrorReport CollectErrorReport(FRecoveryService* RecoveryService, uint3
 	CrashContext.SetCrashedThreadId(SharedCrashContext.CrashingThreadId);
 	CrashContext.SetNumMinidumpFramesToIgnore(SharedCrashContext.NumStackFramesToIgnore);
 
-	// Initialize the stack walking for the monitored process
+	// Initialize the stack walking for the monitored process (effectively overriding this process stack walking functionality)
 	FPlatformStackWalk::InitStackWalkingForProcess(ProcessHandle);
 
 	for (uint32 ThreadIdx = 0; ThreadIdx < SharedCrashContext.NumThreads; ThreadIdx++)
 	{
 		const uint32 ThreadId = SharedCrashContext.ThreadIds[ThreadIdx];
+		TSharedPtr<void> PlatformContext;
+
+#if PLATFORM_WINDOWS
+		// This code let us acquire the complete portable callstack of the remote process after it crashed on a null pointer function invokation. To successfully walk the
+		// stack where a null pointer function is called, we need to provide the thread context reported in the crash (the pointer passed to minidump function), otherwise,
+		// the portable callstack is incomplete.
+		if (ThreadId == SharedCrashContext.CrashingThreadId)
+		{
+			SIZE_T ReadCount = 0;
+
+			// On Windows, 'PlatformCrashContext' is a pointer to the EXCEPTION_POINTERS struct. Try to read it from the monitored process memory.
+			EXCEPTION_POINTERS ExceptPtrs{nullptr, nullptr};
+			if (::ReadProcessMemory(ProcessHandle.Get(), SharedCrashContext.PlatformCrashContext, &ExceptPtrs, sizeof(EXCEPTION_POINTERS), &ReadCount) && ReadCount == sizeof(EXCEPTION_POINTERS))
+			{
+				// Try to read memory of the CONTEXT member from the monitored process.
+				CONTEXT WindowsContext;
+				FMemory::Memzero(WindowsContext);
+				if (::ReadProcessMemory(ProcessHandle.Get(), ExceptPtrs.ContextRecord, &WindowsContext, sizeof(CONTEXT), &ReadCount) && ReadCount == sizeof(CONTEXT))
+				{
+					// NOTE: CaptureThreadStackBackTrace() will open and supply the thread handle specified as null here.
+					PlatformContext = TSharedPtr<void>(FWindowsPlatformStackWalk::MakeThreadContextWrapper(&WindowsContext, nullptr), [](void* Ptr)
+					{
+						FWindowsPlatformStackWalk::ReleaseThreadContextWrapper(Ptr);
+					});
+				}
+			}
+		}
+#endif
+
 		uint64 StackFrames[CR_MAX_STACK_FRAMES] = {0};
 		const uint32 StackFrameCount = FPlatformStackWalk::CaptureThreadStackBackTrace(
 			ThreadId, 
 			StackFrames,
-			CR_MAX_STACK_FRAMES
+			CR_MAX_STACK_FRAMES,
+			PlatformContext.Get()
 		);
 
 		CrashContext.AddPortableThreadCallStack(
@@ -533,6 +563,10 @@ FPlatformErrorReport CollectErrorReport(FRecoveryService* RecoveryService, uint3
 	// Link the crash to the Editor summary event to help diagnose the abnormal termination quickly.
 	FDiagnosticLogger::Get().LogEvent(*FPrimaryCrashProperties::Get()->CrashGUID);
 
+	// Reset stack walking to allow CRC to implicitly walk its own process and close the monitored process handle.
+	FPlatformStackWalk::InitStackWalkingForProcess(FProcHandle());
+	FPlatformProcess::CloseProc(ProcessHandle);
+
 #if CRASH_REPORT_UNATTENDED_ONLY
 	return ErrorReport;
 #else
diff --git a/Engine/Source/Programs/CrashReportClient/Private/Windows/CrashReportClientMainWindows.cpp b/Engine/Source/Programs/CrashReportClient/Private/Windows/CrashReportClientMainWindows.cpp
--- a/Engine/Source/Programs/CrashReportClient/Private/Windows/CrashReportClientMainWindows.cpp
+++ b/Engine/Source/Programs/CrashReportClient/Private/Windows/CrashReportClientMainWindows.cpp
@@ -51,11 +51,13 @@ void SaveCrcCrashException(EXCEPTION_POINTERS* ExceptionInfo)
 		{
 			// Try to get the exception callstack to log to figure out why CRC crashed. This is not robust because this runs
 			// in the crashing processs and it allocates memory/use callstack, but we may still be able to get some useful data.
-			FPlatformStackWalk::InitStackWalking();
-			FPlatformStackWalk::StackWalkAndDump(CrashStackTrace, UE_ARRAY_COUNT(CrashStackTrace), 0);
-			if (CrashStackTrace[0] != 0)
+			if (FPlatformStackWalk::InitStackWalkingForProcess(FProcHandle()))
 			{
-				FDiagnosticLogger::Get().LogEvent(ANSI_TO_TCHAR(CrashStackTrace));
+				FPlatformStackWalk::StackWalkAndDump(CrashStackTrace, UE_ARRAY_COUNT(CrashStackTrace), 0);
+				if (CrashStackTrace[0] != 0)
+				{
+					FDiagnosticLogger::Get().LogEvent(ANSI_TO_TCHAR(CrashStackTrace));
+				}
 			}
 		}
 	}
diff --git a/Engine/Source/Runtime/Core/Private/Android/AndroidPlatformStackWalk.cpp b/Engine/Source/Runtime/Core/Private/Android/AndroidPlatformStackWalk.cpp
--- a/Engine/Source/Runtime/Core/Private/Android/AndroidPlatformStackWalk.cpp
+++ b/Engine/Source/Runtime/Core/Private/Android/AndroidPlatformStackWalk.cpp
@@ -331,7 +331,7 @@ void FAndroidPlatformStackWalk::HandleBackTraceSignal(siginfo* Info, void* Conte
 // Sends a signal to ThreadId, wait AndroidPlatformThreadStackWalk.RequestMaxWait seconds for result or time out and return 0.
 // if callstack capture begins, but takes > AndroidPlatformThreadStackWalk.MaxWait the process will be killed.
 // Is not thread safe, returns 0 if a CaptureThreadStackBackTrace is occurring on another thread.
-uint32 FAndroidPlatformStackWalk::CaptureThreadStackBackTrace(uint64 ThreadId, uint64* BackTrace, uint32 MaxDepth)
+uint32 FAndroidPlatformStackWalk::CaptureThreadStackBackTrace(uint64 ThreadId, uint64* BackTrace, uint32 MaxDepth, void* Context)
 {
 	static TAtomic<bool> bHasReentered(false);
 	bool bExpected = false;
@@ -408,7 +408,7 @@ uint32 FAndroidPlatformStackWalk::CaptureThreadStackBackTrace(uint64 ThreadId, u
 	return GatherCallstackFromThread(ThreadId);
 }
 #else
-uint32 FAndroidPlatformStackWalk::CaptureThreadStackBackTrace(uint64 ThreadId, uint64* BackTrace, uint32 MaxDepth)
+uint32 FAndroidPlatformStackWalk::CaptureThreadStackBackTrace(uint64 ThreadId, uint64* BackTrace, uint32 MaxDepth, void* Context)
 {
 	return 0;
 }
diff --git a/Engine/Source/Runtime/Core/Private/Apple/ApplePlatformStackWalk.cpp b/Engine/Source/Runtime/Core/Private/Apple/ApplePlatformStackWalk.cpp
--- a/Engine/Source/Runtime/Core/Private/Apple/ApplePlatformStackWalk.cpp
+++ b/Engine/Source/Runtime/Core/Private/Apple/ApplePlatformStackWalk.cpp
@@ -367,7 +367,7 @@ void FApplePlatformStackWalk::ThreadStackWalkAndDump(ANSICHAR* HumanReadableStri
 #endif
 }
 
-uint32 FApplePlatformStackWalk::CaptureThreadStackBackTrace(uint64 ThreadId, uint64* BackTrace, uint32 MaxDepth)
+uint32 FApplePlatformStackWalk::CaptureThreadStackBackTrace(uint64 ThreadId, uint64* BackTrace, uint32 MaxDepth, void* Context)
 {
 // Consider this for iOS
 #ifdef PLATFORM_MAC
diff --git a/Engine/Source/Runtime/Core/Private/GenericPlatform/GenericPlatformStackWalk.cpp b/Engine/Source/Runtime/Core/Private/GenericPlatform/GenericPlatformStackWalk.cpp
--- a/Engine/Source/Runtime/Core/Private/GenericPlatform/GenericPlatformStackWalk.cpp
+++ b/Engine/Source/Runtime/Core/Private/GenericPlatform/GenericPlatformStackWalk.cpp
@@ -169,7 +169,7 @@ uint32 FGenericPlatformStackWalk::CaptureStackBackTrace( uint64* BackTrace, uint
 	return 0;
 }
 
-uint32 FGenericPlatformStackWalk::CaptureThreadStackBackTrace(uint64 ThreadId, uint64* BackTrace, uint32 MaxDepth)
+uint32 FGenericPlatformStackWalk::CaptureThreadStackBackTrace(uint64 ThreadId, uint64* BackTrace, uint32 MaxDepth, void* Context)
 {
 	return 0;
 }
diff --git a/Engine/Source/Runtime/Core/Private/Unix/UnixPlatformStackWalk.cpp b/Engine/Source/Runtime/Core/Private/Unix/UnixPlatformStackWalk.cpp
--- a/Engine/Source/Runtime/Core/Private/Unix/UnixPlatformStackWalk.cpp
+++ b/Engine/Source/Runtime/Core/Private/Unix/UnixPlatformStackWalk.cpp
@@ -772,7 +772,7 @@ void FUnixPlatformStackWalk::ThreadStackWalkAndDump(ANSICHAR* HumanReadableStrin
 	GatherCallstackFromThread(ThreadCallStack, ThreadId);
 }
 
-uint32 FUnixPlatformStackWalk::CaptureThreadStackBackTrace(uint64 ThreadId, uint64* BackTrace, uint32 MaxDepth)
+uint32 FUnixPlatformStackWalk::CaptureThreadStackBackTrace(uint64 ThreadId, uint64* BackTrace, uint32 MaxDepth, void* Context)
 {
 	ThreadStackUserData ThreadBackTrace;
 	ThreadBackTrace.bCaptureCallStack = false;
diff --git a/Engine/Source/Runtime/Core/Private/Windows/WindowsPlatformStackWalk.cpp b/Engine/Source/Runtime/Core/Private/Windows/WindowsPlatformStackWalk.cpp
--- a/Engine/Source/Runtime/Core/Private/Windows/WindowsPlatformStackWalk.cpp
+++ b/Engine/Source/Runtime/Core/Private/Windows/WindowsPlatformStackWalk.cpp
@@ -107,6 +107,7 @@ struct FWindowsThreadContextWrapper
 		check(Magic == MAGIC_VAL);
 	}
 };
+
 /**
  * Helper function performing the actual stack walk. This code relies on the symbols being loaded for best results
  * walking the stack albeit at a significant performance penalty.
@@ -118,7 +119,6 @@ struct FWindowsThreadContextWrapper
  * @param	Context				Thread context information
  * @return	EXCEPTION_EXECUTE_HANDLER
  */
-
 static int32 CaptureStackTraceHelper(uint64 *BackTrace, uint32 MaxDepth, FWindowsThreadContextWrapper* ContextWapper, uint32* Depth)
 {
 	STACKFRAME64		StackFrame64;
@@ -138,7 +138,7 @@ static int32 CaptureStackTraceHelper(uint64 *BackTrace, uint32 MaxDepth, FWindow
 #endif
 	{
 		// Get context, process and thread information.
-		ProcessHandle	= GProcessHandle;
+		ProcessHandle = GProcessHandle;
 
 		// Zero out stack frame.
 		FMemory::Memzero( StackFrame64 );
@@ -302,7 +302,7 @@ void FWindowsPlatformStackWalk::ThreadStackWalkAndDump(ANSICHAR* HumanReadableSt
 	}
 }
 
-uint32 FWindowsPlatformStackWalk::CaptureThreadStackBackTrace(uint64 ThreadId, uint64* BackTrace, uint32 MaxDepth)
+uint32 FWindowsPlatformStackWalk::CaptureThreadStackBackTrace(uint64 ThreadId, uint64* BackTrace, uint32 MaxDepth, void* Context)
 {
 	InitStackWalking();
 
@@ -326,14 +326,23 @@ uint32 FWindowsPlatformStackWalk::CaptureThreadStackBackTrace(uint64 ThreadId, u
 	// Suspend the thread before grabbing its context
 	SuspendThread(ThreadHandle);
 
-	FWindowsThreadContextWrapper ContextWrapper;
-	ContextWrapper.Context.ContextFlags = CONTEXT_CONTROL;
-	ContextWrapper.ThreadHandle = ThreadHandle;
-
 	uint32 Depth = 0;
-	if (GetThreadContext(ThreadHandle, &ContextWrapper.Context))
+	if (Context)
 	{
-		CaptureStackTraceHelper(BackTrace, MaxDepth, &ContextWrapper, &Depth);
+		FWindowsThreadContextWrapper* ContextWrapper = reinterpret_cast<FWindowsThreadContextWrapper*>(Context);
+		ContextWrapper->ThreadHandle = ThreadHandle; // Use the thread handle open above to be sure it has the correct permissions.
+		CaptureStackTraceHelper(BackTrace, MaxDepth, ContextWrapper, &Depth);
+	}
+	else
+	{
+		FWindowsThreadContextWrapper ContextWrapper;
+		ContextWrapper.Context.ContextFlags = CONTEXT_CONTROL;
+		ContextWrapper.ThreadHandle = ThreadHandle;
+		
+		if (::GetThreadContext(ThreadHandle, &ContextWrapper.Context))
+		{
+			CaptureStackTraceHelper(BackTrace, MaxDepth, &ContextWrapper, &Depth);
+		}
 	}
 
 	ResumeThread(ThreadHandle);
@@ -349,7 +358,7 @@ uint32 FWindowsPlatformStackWalk::CaptureThreadStackBackTrace(uint64 ThreadId, u
  *
  * @param	BackTrace			[out] Pointer to array to take backtrace
  * @param	MaxDepth			Entries in BackTrace array
- * @param	Context				Optional thread context information
+ * @param	Context				Optional thread context information (FWindowsThreadContextWrapper instance)
  */
 uint32 FWindowsPlatformStackWalk::CaptureStackBackTrace( uint64* BackTrace, uint32 MaxDepth, void* Context )
 {
@@ -1036,22 +1045,45 @@ FString GetSymbolSearchPath()
 /**
  * Initializes the symbol engine if needed.
  */
-bool FWindowsPlatformStackWalk::InitStackWalkingInternal(void* Process)
+bool FWindowsPlatformStackWalk::InitStackWalkingInternal(void* Process, bool bForceReinitOnProcessMismatch)
 {
-	if (GProcessHandle != INVALID_HANDLE_VALUE && !GNeedToRefreshSymbols)
-	{
-		return true;
-	}
-	GProcessHandle = Process;
-
 	// DbgHelp functions are not thread safe, but this function can potentially be called from different
 	// threads in our engine, so we take a critical section
 	static FCriticalSection CriticalSection;
-	FScopeLock Lock( &CriticalSection );
+	FScopeLock Lock(&CriticalSection);
 
-	// Only initialize once.
-	if( !GStackWalkingInitialized )
+	if (GStackWalkingInitialized)
 	{
+		if (Process == GProcessHandle && !GNeedToRefreshSymbols)
+		{
+			return true;
+		}
+		else if (Process != GProcessHandle) // Case: CrashReportClient walked its own process to log one of its ensure, then need to remote walk the Editor process to report an ensure/crash.
+		{
+			if (bForceReinitOnProcessMismatch)
+			{
+				SymCleanup(GProcessHandle);
+				GStackWalkingInitialized = false;
+				GNeedToRefreshSymbols = false;
+				GProcessHandle = Process;
+			}
+			else
+			{
+				return false; // Stack walking is already initialized for another process.
+			}
+		}
+		// Fallthrough.
+	}
+	else // Not initialized yet.
+	{
+		GProcessHandle = Process;
+	}
+
+	// Only do this code once, it never changes.
+	static bool bHasRunOnce = false;
+	if (!bHasRunOnce)
+	{
+		bHasRunOnce = true;
 		void* DllHandle = FPlatformProcess::GetDllHandle( TEXT("PSAPI.DLL") );
 		if( DllHandle == NULL )
 		{
@@ -1069,13 +1101,16 @@ bool FWindowsPlatformStackWalk::InitStackWalkingInternal(void* Process)
 		FGetModuleBaseName		= (TFGetModuleBaseName)		FPlatformProcess::GetDllExport( DllHandle,TEXT("GetModuleBaseNameA"));
 #endif
 		FGetModuleInformation	= (TFGetModuleInformation)	FPlatformProcess::GetDllExport( DllHandle,TEXT("GetModuleInformation"));
+	}
 
-		// Abort if we can't look up the functions.
-		if( !FEnumProcesses || !FEnumProcessModules || !FGetModuleFileNameEx || !FGetModuleBaseName || !FGetModuleInformation )
-		{
-			return false;
-		}
+	// Abort if the required function pointers were not be loaded successfully.
+	if( !FEnumProcesses || !FEnumProcessModules || !FGetModuleFileNameEx || !FGetModuleBaseName || !FGetModuleInformation )
+	{
+		return false;
+	}
 
+	if (!GStackWalkingInitialized)
+	{
 		// Set up the symbol engine.
 		uint32 SymOpts = SymGetOptions();
 
@@ -1096,15 +1131,14 @@ bool FWindowsPlatformStackWalk::InitStackWalkingInternal(void* Process)
 
 		SymSetOptions( SymOpts );
 
-		FString SymbolSearchPath = GetSymbolSearchPath();
-	
 		// Initialize the symbol engine.
 #if WINVER > 0x502
+		FString SymbolSearchPath = GetSymbolSearchPath();
 		SymInitializeW( GProcessHandle, SymbolSearchPath.IsEmpty() ? nullptr : *SymbolSearchPath, true );
 #else
 		SymInitialize( GProcessHandle, nullptr, true );
 #endif
-	
+
 		GNeedToRefreshSymbols = false;
 		GStackWalkingInitialized = true;
 
@@ -1120,7 +1154,7 @@ bool FWindowsPlatformStackWalk::InitStackWalkingInternal(void* Process)
 	else if (GNeedToRefreshSymbols)
 	{
 		// Refresh and reload symbols
-		SymRefreshModuleList( GProcessHandle );
+		SymRefreshModuleList(GProcessHandle);
 
 		GNeedToRefreshSymbols = false;
 
@@ -1140,15 +1174,19 @@ bool FWindowsPlatformStackWalk::InitStackWalkingInternal(void* Process)
 
 bool FWindowsPlatformStackWalk::InitStackWalking()
 {
-	return FWindowsPlatformStackWalk::InitStackWalkingInternal(GetCurrentProcess());
+	// When stack walking is already initialized, the code often implicitly assumes that the same process is going to
+	// be walked, but that's not always true since CrashReportClient can walk its own process to log an ensure and walk
+	// the Editor process to report an ensure/crash. If stack walking is already initialized, keep walking the same process.
+	return FWindowsPlatformStackWalk::InitStackWalkingInternal(GetCurrentProcess(), /*bForceReinitOnProcessMismatch*/false);
 }
 
 bool FWindowsPlatformStackWalk::InitStackWalkingForProcess(const FProcHandle& Process)
 {
-	return FWindowsPlatformStackWalk::InitStackWalkingInternal(Process.Get());
+	// When the caller specifies a process, that process become the 'official' one and this will reinitialize the stack walking
+	// if the current process and the new process don't match. Using an invalid process handle means 'force to reinitialize to the current process'.
+	return FWindowsPlatformStackWalk::InitStackWalkingInternal(Process.IsValid() ? Process.Get() : GetCurrentProcess(), /*bForceReinitOnProcessMismatch*/true);
 }
 
-
 void FWindowsPlatformStackWalk::RegisterOnModulesChanged()
 {
 	// Register for callback so we can reload symbols when new modules are loaded
diff --git a/Engine/Source/Runtime/Core/Public/Android/AndroidPlatformStackWalk.h b/Engine/Source/Runtime/Core/Public/Android/AndroidPlatformStackWalk.h
--- a/Engine/Source/Runtime/Core/Public/Android/AndroidPlatformStackWalk.h
+++ b/Engine/Source/Runtime/Core/Public/Android/AndroidPlatformStackWalk.h
@@ -19,7 +19,7 @@ struct CORE_API FAndroidPlatformStackWalk : public FGenericPlatformStackWalk
 	static uint32 CaptureStackBackTrace(uint64* BackTrace, uint32 MaxDepth, void* Context = nullptr);
 	static bool SymbolInfoToHumanReadableString(const FProgramCounterSymbolInfo& SymbolInfo, ANSICHAR* HumanReadableString, SIZE_T HumanReadableStringSize);
 
-	static uint32 CaptureThreadStackBackTrace(uint64 ThreadId, uint64* BackTrace, uint32 MaxDepth);
+	static uint32 CaptureThreadStackBackTrace(uint64 ThreadId, uint64* BackTrace, uint32 MaxDepth, void* Context = nullptr);
 
 	static void HandleBackTraceSignal(siginfo* Info, void* Context);
 
diff --git a/Engine/Source/Runtime/Core/Public/Apple/ApplePlatformStackWalk.h b/Engine/Source/Runtime/Core/Public/Apple/ApplePlatformStackWalk.h
--- a/Engine/Source/Runtime/Core/Public/Apple/ApplePlatformStackWalk.h
+++ b/Engine/Source/Runtime/Core/Public/Apple/ApplePlatformStackWalk.h
@@ -12,10 +12,10 @@
  */
 struct CORE_API FApplePlatformStackWalk : public FGenericPlatformStackWalk
 {
-	static bool ProgramCounterToHumanReadableString( int32 CurrentCallDepth, uint64 ProgramCounter, ANSICHAR* HumanReadableString, SIZE_T HumanReadableStringSize, FGenericCrashContext* Context = nullptr );
+	static bool ProgramCounterToHumanReadableString( int32 CurrentCallDepth, uint64 ProgramCounter, ANSICHAR* HumanReadableString, SIZE_T HumanReadableStringSize, FGenericCrashContext* Context = nullptr);
 	static void ProgramCounterToSymbolInfo( uint64 ProgramCounter, FProgramCounterSymbolInfo&  out_SymbolInfo);
-	static uint32 CaptureStackBackTrace( uint64* BackTrace, uint32 MaxDepth, void* Context = nullptr );
-	static uint32 CaptureThreadStackBackTrace(uint64 ThreadId, uint64* BackTrace, uint32 MaxDepth);
+	static uint32 CaptureStackBackTrace( uint64* BackTrace, uint32 MaxDepth, void* Context = nullptr);
+	static uint32 CaptureThreadStackBackTrace(uint64 ThreadId, uint64* BackTrace, uint32 MaxDepth, void* Context = nullptr);
 	static void ThreadStackWalkAndDump(ANSICHAR* HumanReadableString, SIZE_T HumanReadableStringSize, int32 IgnoreCount, uint32 ThreadId);
 	static int32 GetProcessModuleCount();
 	static int32 GetProcessModuleSignatures(FStackWalkModuleInfo *ModuleSignatures, const int32 ModuleSignaturesSize);
diff --git a/Engine/Source/Runtime/Core/Public/GenericPlatform/GenericPlatformStackWalk.h b/Engine/Source/Runtime/Core/Public/GenericPlatform/GenericPlatformStackWalk.h
--- a/Engine/Source/Runtime/Core/Public/GenericPlatform/GenericPlatformStackWalk.h
+++ b/Engine/Source/Runtime/Core/Public/GenericPlatform/GenericPlatformStackWalk.h
@@ -220,8 +220,9 @@ struct CORE_API FGenericPlatformStackWalk
 	 * @param	ThreadId			ID of the thread to trace
 	 * @param	BackTrace			[out] Pointer to array to take backtrace
 	 * @param	MaxDepth			Entries in BackTrace array
+	 * @param	Context				Optional thread context information
 	 */
-	static uint32 CaptureThreadStackBackTrace(uint64 ThreadId, uint64* BackTrace, uint32 MaxDepth);
+	static uint32 CaptureThreadStackBackTrace(uint64 ThreadId, uint64* BackTrace, uint32 MaxDepth, void* Context = nullptr);
 
 	/**
 	 * Walks the stack and appends the human readable string to the passed in one.
diff --git a/Engine/Source/Runtime/Core/Public/Unix/UnixPlatformStackWalk.h b/Engine/Source/Runtime/Core/Public/Unix/UnixPlatformStackWalk.h
--- a/Engine/Source/Runtime/Core/Public/Unix/UnixPlatformStackWalk.h
+++ b/Engine/Source/Runtime/Core/Public/Unix/UnixPlatformStackWalk.h
@@ -20,7 +20,7 @@ struct CORE_API FUnixPlatformStackWalk : public FGenericPlatformStackWalk
 	static void StackWalkAndDump( ANSICHAR* HumanReadableString, SIZE_T HumanReadableStringSize, void* ProgramCounter, void* Context = nullptr );
 	static void StackWalkAndDumpEx(ANSICHAR* HumanReadableString, SIZE_T HumanReadableStringSize, int32 IgnoreCount, uint32 Flags, void* Context = nullptr);
 	static void StackWalkAndDumpEx(ANSICHAR* HumanReadableString, SIZE_T HumanReadableStringSize, void* ProgramCounter, uint32 Flags, void* Context = nullptr);
-	static uint32 CaptureThreadStackBackTrace(uint64 ThreadId, uint64* BackTrace, uint32 MaxDepth);
+	static uint32 CaptureThreadStackBackTrace(uint64 ThreadId, uint64* BackTrace, uint32 MaxDepth, void* Context = nullptr);
 	static void ThreadStackWalkAndDump(ANSICHAR* HumanReadableString, SIZE_T HumanReadableStringSize, int32 IgnoreCount, uint32 ThreadId);
 	static int32 GetProcessModuleCount();
 	static int32 GetProcessModuleSignatures(FStackWalkModuleInfo *ModuleSignatures, const int32 ModuleSignaturesSize);
diff --git a/Engine/Source/Runtime/Core/Public/Windows/WindowsPlatformStackWalk.h b/Engine/Source/Runtime/Core/Public/Windows/WindowsPlatformStackWalk.h
--- a/Engine/Source/Runtime/Core/Public/Windows/WindowsPlatformStackWalk.h
+++ b/Engine/Source/Runtime/Core/Public/Windows/WindowsPlatformStackWalk.h
@@ -19,7 +19,7 @@ struct CORE_API FWindowsPlatformStackWalk
 	static void ProgramCounterToSymbolInfo( uint64 ProgramCounter, FProgramCounterSymbolInfo& out_SymbolInfo );
 	static void ProgramCounterToSymbolInfoEx( uint64 ProgramCounter, FProgramCounterSymbolInfoEx& out_SymbolInfo );
 	static uint32 CaptureStackBackTrace( uint64* BackTrace, uint32 MaxDepth, void* Context = nullptr );
-	static uint32 CaptureThreadStackBackTrace( uint64 ThreadId, uint64* BackTrace, uint32 MaxDepth );
+	static uint32 CaptureThreadStackBackTrace( uint64 ThreadId, uint64* BackTrace, uint32 MaxDepth, void* Context = nullptr);
 	static void StackWalkAndDump( ANSICHAR* HumanReadableString, SIZE_T HumanReadableStringSize, int32 IgnoreCount, void* Context = nullptr );
 	static void StackWalkAndDump( ANSICHAR* HumanReadableString, SIZE_T HumanReadableStringSize, void* ProgramCounter, void* Context = nullptr );
 	static void ThreadStackWalkAndDump(ANSICHAR* HumanReadableString, SIZE_T HumanReadableStringSize, int32 IgnoreCount, uint32 ThreadId);
@@ -48,7 +48,7 @@ struct CORE_API FWindowsPlatformStackWalk
 	static void ReleaseThreadContextWrapper(void* ThreadContext);
 
 private:
-	static bool InitStackWalkingInternal(void* Process);
+	static bool InitStackWalkingInternal(void* Process, bool bForceReinitOnProcessMismatch);
 };
 
 
