Moved MsgLogfStaticBufferGuard into a static function to make sure it is not sensitive to static initialization order.

diff --git a/Engine/Source/Runtime/Core/Private/Logging/LogMacros.cpp b/Engine/Source/Runtime/Core/Private/Logging/LogMacros.cpp
--- a/Engine/Source/Runtime/Core/Private/Logging/LogMacros.cpp
+++ b/Engine/Source/Runtime/Core/Private/Logging/LogMacros.cpp
@@ -13,7 +13,12 @@
 void StaticFailDebug( const TCHAR* Error, const FDebug::FFailureInfo& Info, const TCHAR* Description, bool bIsEnsure, int32 NumStackFramesToIgnore );
 
 /** Statics to prevent FMsg::Logf from allocating too much stack memory. */
-static FCriticalSection					MsgLogfStaticBufferGuard;
+static FCriticalSection* GetMsgLogfStaticBufferGuard()
+{
+	static FCriticalSection CS;
+	return &CS;
+}
+
 /** Increased from 4096 to fix crashes in the renderthread without autoreporter. */
 static TCHAR							MsgLogfStaticBuffer[8192];
 
@@ -54,7 +59,7 @@ void FMsg::LogfImpl(const ANSICHAR* File, int32 Line, const FLogCategoryName& Ca
 			// Simulate Sprintf_s
 			// @todo: implement platform independent sprintf_S
 			// We're using one big shared static buffer here, so guard against re-entry
-			FScopeLock MsgLock(&MsgLogfStaticBufferGuard);
+			FScopeLock MsgLock(GetMsgLogfStaticBufferGuard());
 			// Print to a large static buffer so we can keep the stack allocation below 16K
 			GET_VARARGS(MsgLogfStaticBuffer, UE_ARRAY_COUNT(MsgLogfStaticBuffer), UE_ARRAY_COUNT(MsgLogfStaticBuffer) - 1, Fmt, Fmt);
 			// Copy the message to the stack-allocated buffer)
@@ -102,7 +107,7 @@ void FMsg::Logf_InternalImpl(const ANSICHAR* File, int32 Line, const FLogCategor
 			// Simulate Sprintf_s
 			// @todo: implement platform independent sprintf_S
 			// We're using one big shared static buffer here, so guard against re-entry
-			FScopeLock MsgLock(&MsgLogfStaticBufferGuard);
+			FScopeLock MsgLock(GetMsgLogfStaticBufferGuard());
 			// Print to a large static buffer so we can keep the stack allocation below 16K
 			GET_VARARGS(MsgLogfStaticBuffer, UE_ARRAY_COUNT(MsgLogfStaticBuffer), UE_ARRAY_COUNT(MsgLogfStaticBuffer) - 1, Fmt, Fmt);
 			// Copy the message to the stack-allocated buffer)
