U5 - only do broadcast on background workers if it was explicitly requested

diff --git a/Engine/Source/Runtime/Core/Private/Async/TaskGraph.cpp b/Engine/Source/Runtime/Core/Private/Async/TaskGraph.cpp
--- a/Engine/Source/Runtime/Core/Private/Async/TaskGraph.cpp
+++ b/Engine/Source/Runtime/Core/Private/Async/TaskGraph.cpp
@@ -2455,7 +2455,7 @@ void FTaskGraphInterface::BroadcastSlow_OnlyUseForSpecialPurposes(bool bDoTaskTh
 		int32 Workers = 0;
 		if(GUseNewTaskBackend)
 		{
-			Workers = TaskGraphImplementationSingleton->GetNumWorkerThreads();
+			Workers = bDoBackgroundThreads ? TaskGraphImplementationSingleton->GetNumWorkerThreads() : GNumForegroundWorkers;
 			StallForTaskThread.Add(Workers);
 		}
 		else
