U5 - removing this bad sleep because it could trigger a context switch and keep the worker unresponsive for a long time.

diff --git a/Engine/Source/Runtime/Core/Public/Async/Fundamental/Scheduler.h b/Engine/Source/Runtime/Core/Public/Async/Fundamental/Scheduler.h
--- a/Engine/Source/Runtime/Core/Public/Async/Fundamental/Scheduler.h
+++ b/Engine/Source/Runtime/Core/Public/Async/Fundamental/Scheduler.h
@@ -238,8 +238,7 @@ namespace LowLevelTasks
 		if(!Drowsing && WorkerEvent->SleepState.compare_exchange_strong(DrowsingState1, ESleepState::Drowsing, std::memory_order_relaxed)) //continue drowsing
 		{
 			verifySlow(OutOfWork.Stop());
-			Drowsing = true;
-			FPlatformProcess::SleepNoStats(0);	// Alternative State one: ((Running -> Drowsing) -> Drowsing)
+			Drowsing = true; // Alternative State one: ((Running -> Drowsing) -> Drowsing)
 		}
 		else if(WorkerEvent->SleepState.compare_exchange_strong(DrowsingState2, ESleepState::Sleeping, std::memory_order_relaxed))
 		{
@@ -252,8 +251,7 @@ namespace LowLevelTasks
 		{
 			verifySlow(OutOfWork.Stop() == !Drowsing);
 			Drowsing = true;
-			SleepEventQueue[bBackgroundWorker].enqueue(WorkerEvent);
-			FPlatformProcess::SleepNoStats(0); // State one: (Running -> Drowsing)
+			SleepEventQueue[bBackgroundWorker].enqueue(WorkerEvent); // State one: (Running -> Drowsing)
 		}
 		else
 		{
