U5 - don't start the out of work scope when waking a worker.

diff --git a/Engine/Source/Runtime/Core/Public/Async/Fundamental/Scheduler.h b/Engine/Source/Runtime/Core/Public/Async/Fundamental/Scheduler.h
--- a/Engine/Source/Runtime/Core/Public/Async/Fundamental/Scheduler.h
+++ b/Engine/Source/Runtime/Core/Public/Async/Fundamental/Scheduler.h
@@ -245,7 +245,6 @@ namespace LowLevelTasks
 			verifySlow(!OutOfWork.Stop());
 			Drowsing = false;
 			WorkerEvent->SleepEvent->Wait(); // State two: ((Running -> Drowsing) -> Sleeping)
-			verifySlow(OutOfWork.Start());
 		}
 		else if(WorkerEvent->SleepState.compare_exchange_strong(RunningState, ESleepState::Drowsing, std::memory_order_relaxed))
 		{
