Use a platform specific time function instead of unix time.

diff --git a/Engine/Source/Runtime/Core/Public/Linux/LinuxPlatformTime.h b/Engine/Source/Runtime/Core/Public/Linux/LinuxPlatformTime.h
--- a/Engine/Source/Runtime/Core/Public/Linux/LinuxPlatformTime.h
+++ b/Engine/Source/Runtime/Core/Public/Linux/LinuxPlatformTime.h
@@ -3,3 +3,5 @@
 #pragma once
 
 #include "Unix/UnixPlatformTime.h"
+
+typedef FUnixTime FPlatformTime;
\ No newline at end of file
diff --git a/Engine/Source/Runtime/Core/Public/Unix/UnixPlatformTime.h b/Engine/Source/Runtime/Core/Public/Unix/UnixPlatformTime.h
--- a/Engine/Source/Runtime/Core/Public/Unix/UnixPlatformTime.h
+++ b/Engine/Source/Runtime/Core/Public/Unix/UnixPlatformTime.h
@@ -65,5 +65,3 @@ private:
 	 */
 	static uint64 CallsPerSecondBenchmark(clockid_t BenchClockId, const char * BenchClockIdName);
 };
-
-typedef FUnixTime FPlatformTime;
