U5 - the valid case where the woker wakes up and finds no task can trigger this

diff --git a/Engine/Source/Runtime/Core/Public/Async/Fundamental/Scheduler.h b/Engine/Source/Runtime/Core/Public/Async/Fundamental/Scheduler.h
--- a/Engine/Source/Runtime/Core/Public/Async/Fundamental/Scheduler.h
+++ b/Engine/Source/Runtime/Core/Public/Async/Fundamental/Scheduler.h
@@ -248,7 +248,7 @@ namespace LowLevelTasks
 		}
 		else if(WorkerEvent->SleepState.compare_exchange_strong(RunningState, ESleepState::Drowsing, std::memory_order_relaxed))
 		{
-			verifySlow(OutOfWork.Stop() == !Drowsing);
+			OutOfWork.Stop();
 			Drowsing = true;
 			SleepEventQueue[bBackgroundWorker].enqueue(WorkerEvent); // State one: (Running -> Drowsing)
 		}
