CompactBinary: Fixed NonUnity compile errors

#rnx

diff --git a/Engine/Source/Runtime/Core/Private/Serialization/CompactBinaryWriter.cpp b/Engine/Source/Runtime/Core/Private/Serialization/CompactBinaryWriter.cpp
--- a/Engine/Source/Runtime/Core/Private/Serialization/CompactBinaryWriter.cpp
+++ b/Engine/Source/Runtime/Core/Private/Serialization/CompactBinaryWriter.cpp
@@ -530,6 +530,12 @@ void FCbWriter::AddDateTimeTicks(const int64 Ticks)
 	EndField(ECbFieldType::DateTime);
 }
 
+void FCbWriter::AddDateTime(FAnsiStringView Name, FDateTime Value)
+{
+	SetName(Name);
+	AddDateTime(Value);
+}
+
 void FCbWriter::AddDateTime(const FDateTime Value)
 {
 	AddDateTimeTicks(Value.GetTicks());
@@ -543,6 +549,12 @@ void FCbWriter::AddTimeSpanTicks(const int64 Ticks)
 	EndField(ECbFieldType::TimeSpan);
 }
 
+void FCbWriter::AddTimeSpan(FAnsiStringView Name, FTimespan Value)
+{
+	SetName(Name);
+	AddTimeSpan(Value);
+}
+
 void FCbWriter::AddTimeSpan(const FTimespan Value)
 {
 	AddTimeSpanTicks(Value.GetTicks());
diff --git a/Engine/Source/Runtime/Core/Public/Serialization/CompactBinaryWriter.h b/Engine/Source/Runtime/Core/Public/Serialization/CompactBinaryWriter.h
--- a/Engine/Source/Runtime/Core/Public/Serialization/CompactBinaryWriter.h
+++ b/Engine/Source/Runtime/Core/Public/Serialization/CompactBinaryWriter.h
@@ -209,7 +209,7 @@ public:
 	CORE_API void AddDateTimeTicks(int64 Ticks);
 
 	/** Write a date/time field. */
-	inline void AddDateTime(FAnsiStringView Name, FDateTime Value) { SetName(Name); AddDateTime(Value); }
+	CORE_API void AddDateTime(FAnsiStringView Name, FDateTime Value);
 	CORE_API void AddDateTime(FDateTime Value);
 
 	/** Write a time span field with the specified count of 100ns ticks. */
@@ -217,7 +217,7 @@ public:
 	CORE_API void AddTimeSpanTicks(int64 Ticks);
 
 	/** Write a time span field. */
-	inline void AddTimeSpan(FAnsiStringView Name, FTimespan Value) { SetName(Name); AddTimeSpan(Value); }
+	CORE_API void AddTimeSpan(FAnsiStringView Name, FTimespan Value);
 	CORE_API void AddTimeSpan(FTimespan Value);
 
 	/** Private flags that are public to work with ENUM_CLASS_FLAGS. */
