U5 - classify the out of work scope during busy waiting on backgroundworkers as backgroundwork

diff --git a/Engine/Source/Runtime/Core/Private/Async/Fundamental/Scheduler.cpp b/Engine/Source/Runtime/Core/Private/Async/Fundamental/Scheduler.cpp
--- a/Engine/Source/Runtime/Core/Private/Async/Fundamental/Scheduler.cpp
+++ b/Engine/Source/Runtime/Core/Private/Async/Fundamental/Scheduler.cpp
@@ -255,8 +255,9 @@ namespace LowLevelTasks
 		FLocalQueueType* WorkerLocalQueue = LocalQueue;
 
 		uint32 WaitCount = 0;
+		const bool bIsBackgroundWorker = WorkerType == EWorkerType::Background;
 		bool bPermitBackgroundWork = ActiveTask && ActiveTask->IsBackgroundTask();
-		FQueueRegistry::FOutOfWork OutOfWork = QueueRegistry.GetOutOfWorkScope(bPermitBackgroundWork);
+		FQueueRegistry::FOutOfWork OutOfWork = QueueRegistry.GetOutOfWorkScope(bIsBackgroundWorker);
 		while (true)
 		{
 			while(TryExecuteTaskFrom<&FLocalQueueType::DequeueLocal>(WorkerLocalQueue, OutOfWork, bPermitBackgroundWork)
@@ -284,7 +285,6 @@ namespace LowLevelTasks
 				return;
 			}
 
-			const bool bIsBackgroundWorker = WorkerType == EWorkerType::Background;
 			if (WaitCount < WorkerSpinCycles)
 			{
 				OutOfWork.Start();
