u5 - never sleep while busy waiting

diff --git a/Engine/Source/Runtime/Core/Private/Async/Fundamental/Scheduler.cpp b/Engine/Source/Runtime/Core/Private/Async/Fundamental/Scheduler.cpp
--- a/Engine/Source/Runtime/Core/Private/Async/Fundamental/Scheduler.cpp
+++ b/Engine/Source/Runtime/Core/Private/Async/Fundamental/Scheduler.cpp
@@ -296,13 +296,11 @@ namespace LowLevelTasks
 			{
 				bPermitBackgroundWork = true;
 			}
-			else
+			else if(!bIsBackgroundWorker)
 			{
-				TRACE_CPUPROFILER_EVENT_SCOPE(BusyWaitInternal::SleepNoStats);
-				WakeUpWorker(true); //wake up a backgoundworker in case we were waiting on background work.
-				FPlatformProcess::SleepNoStats(0);
-				WaitCount = 0;
+				WakeUpWorker(true);
 			}
+			WaitCount = 0;
 		}
 	}
 }
\ No newline at end of file
