U5 - limit number of threads until propper process and thread group affinity is supported. :-(

diff --git a/Engine/Source/Runtime/Core/Private/Windows/WindowsPlatformMisc.cpp b/Engine/Source/Runtime/Core/Private/Windows/WindowsPlatformMisc.cpp
--- a/Engine/Source/Runtime/Core/Private/Windows/WindowsPlatformMisc.cpp
+++ b/Engine/Source/Runtime/Core/Private/Windows/WindowsPlatformMisc.cpp
@@ -1801,7 +1801,7 @@ int32 FWindowsPlatformMisc::NumberOfWorkerThreadsToSpawn()
 	static int32 MaxServerWorkerThreads = 4;
 
 	extern CORE_API int32 GUseNewTaskBackend;
-	int32 MaxWorkerThreads = GUseNewTaskBackend ? INT32_MAX : 26;
+	int32 MaxWorkerThreads = GUseNewTaskBackend ? 62 : 26; //Windows Thread Group limit is 64
 
 	int32 NumberOfCores = FWindowsPlatformMisc::NumberOfCores();
 	int32 NumberOfCoresIncludingHyperthreads = FWindowsPlatformMisc::NumberOfCoresIncludingHyperthreads();
