U5 - fix WaitCount logic

diff --git a/Engine/Source/Runtime/Core/Private/Async/Fundamental/Scheduler.cpp b/Engine/Source/Runtime/Core/Private/Async/Fundamental/Scheduler.cpp
--- a/Engine/Source/Runtime/Core/Private/Async/Fundamental/Scheduler.cpp
+++ b/Engine/Source/Runtime/Core/Private/Async/Fundamental/Scheduler.cpp
@@ -296,11 +296,14 @@ namespace LowLevelTasks
 			{
 				bPermitBackgroundWork = true;
 			}
-			else if(!bIsBackgroundWorker)
+			else
 			{
-				WakeUpWorker(true);
-			}
-			WaitCount = 0;
+				if (!bIsBackgroundWorker)
+				{
+					WakeUpWorker(true);
+				}
+				WaitCount = 0;
+			}			
 		}
 	}
 }
\ No newline at end of file
