Never perform fast reference collector on background thread in editor to reduce latency
Latency on the GC could be as much as >5 minutes in some cases when loading P_Construct because the background task were already at full capacity for long periods of time

diff --git a/Engine/Source/Runtime/Core/Private/GenericPlatform/GenericPlatformProcess.cpp b/Engine/Source/Runtime/Core/Private/GenericPlatform/GenericPlatformProcess.cpp
--- a/Engine/Source/Runtime/Core/Private/GenericPlatform/GenericPlatformProcess.cpp
+++ b/Engine/Source/Runtime/Core/Private/GenericPlatform/GenericPlatformProcess.cpp
@@ -648,6 +648,10 @@ void FGenericPlatformProcess::ModifyThreadAssignmentForUObjectReferenceCollector
 	// On devices with overridden affinity only HiPri threads can run on big cores
 	NormalThreadName = ENamedThreads::AnyHiPriThreadHiPriTask; 
 	NumBackgroundThreads = 0; // run on single group
+#elif WITH_EDITOR
+	// Avoid the ReferenceCollector being slowed down by long running background tasks, async compilation, etc...
+	NormalThreadName = ENamedThreads::AnyHiPriThreadHiPriTask;
+	BackgroundThreadName = ENamedThreads::AnyHiPriThreadHiPriTask;
 #endif
 }
 
