U5 - add the ability to provide a Scheduler instance to the LowLevel TaskPool so that it can use a seperate instance of a Scheduler.

diff --git a/Engine/Source/Runtime/Core/Private/Async/Fundamental/Scheduler.cpp b/Engine/Source/Runtime/Core/Private/Async/Fundamental/Scheduler.cpp
--- a/Engine/Source/Runtime/Core/Private/Async/Fundamental/Scheduler.cpp
+++ b/Engine/Source/Runtime/Core/Private/Async/Fundamental/Scheduler.cpp
@@ -3,7 +3,7 @@
 #include "Async/Fundamental/Scheduler.h"
 #include "Async/Fundamental/Task.h"
 #include "Logging/LogMacros.h"
-#include "Misc/ScopeLock.h" 
+#include "Misc/ScopeLock.h"
 #include "CoreGlobals.h"
 
 namespace LowLevelTasks
@@ -12,6 +12,7 @@ namespace LowLevelTasks
 
 	thread_local FScheduler::FLocalQueueType* FScheduler::LocalQueue = nullptr;
 	thread_local FTask* FScheduler::ActiveTask = nullptr;
+	thread_local FScheduler* FScheduler::ActiveScheduler = nullptr;
 	thread_local FScheduler::EWorkerType FScheduler::WorkerType = FScheduler::EWorkerType::None;
 
 	FScheduler FScheduler::Singleton;
@@ -53,26 +54,26 @@ namespace LowLevelTasks
 		);
 	}
 
-	void FScheduler::StartWorkers(uint32 NumWorkers, uint32 NumBackgroundWorkers, EThreadPriority WorkerPriority,  EThreadPriority BackgroundPriority, bool bIsForkable)
+	void FScheduler::StartWorkers(uint32 NumForegroundWorkers, uint32 NumBackgroundWorkers, EThreadPriority WorkerPriority,  EThreadPriority BackgroundPriority, bool bIsForkable)
 	{
-		if (NumWorkers == 0)
+		if (NumForegroundWorkers == 0 && NumBackgroundWorkers == 0)
 		{
-			NumWorkers = 2;
-			NumBackgroundWorkers = FMath::Max(1, FPlatformMisc::NumberOfWorkerThreadsToSpawn() - 2);
+			NumForegroundWorkers = FMath::Max<int32>(1, FMath::Min<int32>(2, FPlatformMisc::NumberOfWorkerThreadsToSpawn() - 1));
+			NumBackgroundWorkers = FMath::Max<int32>(1, FPlatformMisc::NumberOfWorkerThreadsToSpawn() - NumForegroundWorkers);
 		}
 
 		uint32 OldActiveWorkers = ActiveWorkers.load(std::memory_order_relaxed);
-		if(OldActiveWorkers == 0 && FPlatformProcess::SupportsMultithreading() && ActiveWorkers.compare_exchange_strong(OldActiveWorkers, NumWorkers + NumBackgroundWorkers, std::memory_order_relaxed))
+		if(OldActiveWorkers == 0 && FPlatformProcess::SupportsMultithreading() && ActiveWorkers.compare_exchange_strong(OldActiveWorkers, NumForegroundWorkers + NumBackgroundWorkers, std::memory_order_relaxed))
 		{
 			FScopeLock Lock(&WorkerThreadsCS);
 			check(!WorkerThreads.Num());
 			check(!WorkerLocalQueues.Num());
 			check(NextWorkerId == 0);
 
-			WorkerThreads.Reserve(NumWorkers + NumBackgroundWorkers);
-			WorkerLocalQueues.Reserve(NumWorkers + NumBackgroundWorkers);
+			WorkerThreads.Reserve(NumForegroundWorkers + NumBackgroundWorkers);
+			WorkerLocalQueues.Reserve(NumForegroundWorkers + NumBackgroundWorkers);
 			UE::Trace::ThreadGroupBegin(TEXT("Foreground Workers"));
-			for (uint32 WorkerId = 0; WorkerId < NumWorkers; ++WorkerId)
+			for (uint32 WorkerId = 0; WorkerId < NumForegroundWorkers; ++WorkerId)
 			{
 				WorkerLocalQueues.Emplace(QueueRegistry, false);
 				WorkerThreads.Add(CreateWorker(&WorkerLocalQueues.Last(), WorkerPriority, false, bIsForkable));
@@ -151,14 +152,14 @@ namespace LowLevelTasks
 		}
 	}
 
-	const FTask* FScheduler::GetActiveTask() const
+	const FTask* FScheduler::GetActiveTask() 
 	{
 		return ActiveTask;
 	}
 
 	bool FScheduler::IsWorkerThread() const
 	{
-		return WorkerType != EWorkerType::None;
+		return WorkerType != EWorkerType::None && ActiveScheduler == this;
 	}
 
 	template<FTask* (FScheduler::FLocalQueueType::*DequeueFunction)(bool)>
@@ -183,6 +184,7 @@ namespace LowLevelTasks
 	void FScheduler::WorkerMain(FSleepEvent* WorkerEvent, FLocalQueueType* ExternalWorkerLocalQueue, uint32 WaitCycles, bool bPermitBackgroundWork)
 	{
 		FTaskTagScope WorkerScope(ETaskTag::EWorkerThread);
+		ActiveScheduler = this;
 
 		FMemory::SetupTLSCachesOnCurrentThread();
 		WorkerType = bPermitBackgroundWork ? EWorkerType::Background : EWorkerType::Foreground;
@@ -241,6 +243,7 @@ namespace LowLevelTasks
 		FLocalQueueType::DeleteLocalQueue(WorkerLocalQueue, bPermitBackgroundWork, ExternalWorkerLocalQueue != nullptr);
 		LocalQueue = nullptr;
 
+		ActiveScheduler = nullptr;
 		WorkerType = EWorkerType::None;
 		FMemory::ClearAndDisableTLSCachesOnCurrentThread();
 	}
diff --git a/Engine/Source/Runtime/Core/Private/Async/TaskGraph.cpp b/Engine/Source/Runtime/Core/Private/Async/TaskGraph.cpp
--- a/Engine/Source/Runtime/Core/Private/Async/TaskGraph.cpp
+++ b/Engine/Source/Runtime/Core/Private/Async/TaskGraph.cpp
@@ -1937,7 +1937,7 @@ private:
 
 	bool IsCurrentThreadKnown() final override
 	{
-		return FPlatformTLS::GetTlsValue(PerThreadIDTLSSlot) != nullptr || LowLevelTasks::FScheduler::Get().GetActiveTask() != nullptr;
+		return FPlatformTLS::GetTlsValue(PerThreadIDTLSSlot) != nullptr || LowLevelTasks::FScheduler::GetActiveTask() != nullptr;
 	}
 
 	ENamedThreads::Type GetCurrentThreadIfKnown(bool bLocalQueue) final override
@@ -2040,7 +2040,7 @@ private:
 			TGraphTask<FReturnGraphTask>::CreateTask(&Tasks, CurrentThread).ConstructAndDispatchWhenReady(CurrentThread);
 			ProcessThreadUntilRequestReturn(CurrentThread);
 		}
-		else if (LowLevelTasks::FScheduler::Get().GetActiveTask() != nullptr)
+		else if (LowLevelTasks::FScheduler::Get().IsWorkerThread())
 		{
 			LowLevelTasks::BusyWaitUntil([Index(0), &Tasks]() mutable
 			{
@@ -2151,7 +2151,7 @@ private:
 		}
 		else 
 		{
-			const LowLevelTasks::FTask* Task = LowLevelTasks::FScheduler::Get().GetActiveTask();
+			const LowLevelTasks::FTask* Task = LowLevelTasks::FScheduler::GetActiveTask();
 			if (Task != nullptr)
 			{
 				ENamedThreads::Type ThreadConversion[int(LowLevelTasks::ETaskPriority::Count)] = { ENamedThreads::HighThreadPriority, ENamedThreads::NormalThreadPriority, ENamedThreads::BackgroundThreadPriority, ENamedThreads::BackgroundThreadPriority };
diff --git a/Engine/Source/Runtime/Core/Public/Async/AsyncWork.h b/Engine/Source/Runtime/Core/Public/Async/AsyncWork.h
--- a/Engine/Source/Runtime/Core/Public/Async/AsyncWork.h
+++ b/Engine/Source/Runtime/Core/Public/Async/AsyncWork.h
@@ -356,7 +356,7 @@ class FAsyncTask
 			FScopeCycleCounter Scope( Task.GetStatId() );
 			DECLARE_SCOPE_CYCLE_COUNTER( TEXT( "FAsyncTask::SyncCompletion" ), STAT_FAsyncTask_SyncCompletion, STATGROUP_ThreadPoolAsyncTasks );
 
-			if (LowLevelTasks::FScheduler::Get().GetActiveTask())
+			if (LowLevelTasks::FScheduler::Get().IsWorkerThread())
 			{
 				LowLevelTasks::BusyWaitUntil([this]() { return IsWorkDone(); });
 			}
diff --git a/Engine/Source/Runtime/Core/Public/Async/Fundamental/Scheduler.h b/Engine/Source/Runtime/Core/Public/Async/Fundamental/Scheduler.h
--- a/Engine/Source/Runtime/Core/Public/Async/Fundamental/Scheduler.h
+++ b/Engine/Source/Runtime/Core/Public/Async/Fundamental/Scheduler.h
@@ -29,6 +29,7 @@ namespace LowLevelTasks
 
 		static thread_local FLocalQueueType* LocalQueue;
 		static thread_local FTask* ActiveTask;
+		static thread_local FScheduler* ActiveScheduler;
 
 		enum class EWorkerType
 		{
@@ -56,7 +57,7 @@ namespace LowLevelTasks
 		FORCEINLINE_DEBUGGABLE static FScheduler& Get();
 
 		//start number of workers where 0 is the system default
-		CORE_API void StartWorkers(uint32 NumWorkers = 0, uint32 NumBackgroundWorkers = 0, EThreadPriority WorkerPriority = EThreadPriority::TPri_Normal,  EThreadPriority BackgroundPriority = EThreadPriority::TPri_BelowNormal, bool bIsForkable = false);
+		CORE_API void StartWorkers(uint32 NumForegroundWorkers = 0, uint32 NumBackgroundWorkers = 0, EThreadPriority WorkerPriority = EThreadPriority::TPri_Normal,  EThreadPriority BackgroundPriority = EThreadPriority::TPri_BelowNormal, bool bIsForkable = false);
 		CORE_API void StopWorkers();
 
 		//try to launch the task, the return value will specify if the task was in the ready state and has been launced
@@ -84,7 +85,7 @@ namespace LowLevelTasks
 		inline uint32 GetNumWorkers() const;
 
 		//get the active task if any
-		CORE_API const FTask* GetActiveTask() const;
+		CORE_API static const FTask* GetActiveTask();
 
 		CORE_API bool IsWorkerThread() const;
 
diff --git a/Engine/Source/Runtime/Core/Public/Misc/QueuedThreadPoolWrapper.h b/Engine/Source/Runtime/Core/Public/Misc/QueuedThreadPoolWrapper.h
--- a/Engine/Source/Runtime/Core/Public/Misc/QueuedThreadPoolWrapper.h
+++ b/Engine/Source/Runtime/Core/Public/Misc/QueuedThreadPoolWrapper.h
@@ -200,7 +200,7 @@ protected:
 	{
 		bIsExiting = true;
 
-		if (LowLevelTasks::FScheduler::Get().GetActiveTask())
+		if (LowLevelTasks::FScheduler::Get().IsWorkerThread())
 		{
 			LowLevelTasks::BusyWaitUntil([this]() { return TaskCount == 0; });
 		}
@@ -237,8 +237,8 @@ public:
     * InMaxConcurrency           Maximum number of concurrent tasks allowed, -1 will limit concurrency to number of threads available in the underlying thread pool.
 	* InPriorityMapper           Thread-safe function used to map any priority from this Queue to the priority that should be used when scheduling the task on the underlying thread pool.
 	**/
-	FQueuedLowLevelThreadPool(uint32 InMaxConcurrency = ~0, TFunction<EQueuedWorkPriority(EQueuedWorkPriority)> InPriorityMapper = [](EQueuedWorkPriority InPriority) { return InPriority; }) 
-		: PriorityMapper(InPriorityMapper)
+	FQueuedLowLevelThreadPool(uint32 InMaxConcurrency = ~0, TFunction<EQueuedWorkPriority(EQueuedWorkPriority)> InPriorityMapper = [](EQueuedWorkPriority InPriority) { return InPriority; }, LowLevelTasks::FScheduler* InScheduler = &LowLevelTasks::FScheduler::Get()) 
+		: PriorityMapper(InPriorityMapper), Scheduler(InScheduler)
 	{
 		MaxConcurrency = InMaxConcurrency;
 	}
@@ -287,7 +287,7 @@ public:
 				break;
 			}
 			TaskCount++;
-			verifySlow(LowLevelTasks::TryLaunch(QueuedWork->Task, LowLevelTasks::EQueuePreference::GlobalQueuePreference));
+			verifySlow(Scheduler->TryLaunch(QueuedWork->Task, LowLevelTasks::EQueuePreference::GlobalQueuePreference));
 		}
 
 		if (InNumQueuedWork == -1)
@@ -307,7 +307,7 @@ private:
 			FQueuedWorkInternalData* QueuedWork = Dequeue();
 			if (QueuedWork)
 			{
-				verifySlow(LowLevelTasks::TryLaunch(QueuedWork->Task, bWakeUpWorker ? LowLevelTasks::EQueuePreference::GlobalQueuePreference : LowLevelTasks::EQueuePreference::LocalQueuePreference, bWakeUpWorker));
+				verifySlow(Scheduler->TryLaunch(QueuedWork->Task, bWakeUpWorker ? LowLevelTasks::EQueuePreference::GlobalQueuePreference : LowLevelTasks::EQueuePreference::LocalQueuePreference, bWakeUpWorker));
 				LocalTaskCount = TaskCount.fetch_add(1, std::memory_order_acquire) + 1;
 				bWakeUpWorker = true;
 			}
@@ -360,7 +360,7 @@ private:
 
 	int32 GetNumThreads() const override
 	{
-		return LowLevelTasks::FScheduler::Get().GetNumWorkers();
+		return Scheduler->GetNumWorkers();
 	}
 
 protected:
@@ -384,12 +384,12 @@ protected:
 
 			verify(QueuedWork->Retract());
 			TaskCount++;
-			verifySlow(LowLevelTasks::TryLaunch(QueuedWork->Task, LowLevelTasks::EQueuePreference::GlobalQueuePreference));
+			verifySlow(Scheduler->TryLaunch(QueuedWork->Task, LowLevelTasks::EQueuePreference::GlobalQueuePreference));
 		}
 
-		if (LowLevelTasks::FScheduler::Get().GetActiveTask())
+		if (Scheduler->IsWorkerThread())
 		{
-			LowLevelTasks::BusyWaitUntil([this]() { return TaskCount == 0; });
+			Scheduler->BusyWaitUntil([this]() { return TaskCount == 0; });
 		}
 		else
 		{
@@ -421,6 +421,7 @@ private:
 		PendingWork[int32(Priority)].enqueue(Item);
 	}
 
+	LowLevelTasks::FScheduler* Scheduler = nullptr;
 	TFunction<EQueuedWorkPriority(EQueuedWorkPriority)> PriorityMapper;
 	std::atomic_uint MaxConcurrency{~0u};
 	std::atomic_uint TaskCount{0};
