U5 - fix warning

diff --git a/Engine/Source/Runtime/Core/Public/Misc/QueuedThreadPoolWrapper.h b/Engine/Source/Runtime/Core/Public/Misc/QueuedThreadPoolWrapper.h
--- a/Engine/Source/Runtime/Core/Public/Misc/QueuedThreadPoolWrapper.h
+++ b/Engine/Source/Runtime/Core/Public/Misc/QueuedThreadPoolWrapper.h
@@ -238,7 +238,7 @@ public:
 	* InPriorityMapper           Thread-safe function used to map any priority from this Queue to the priority that should be used when scheduling the task on the underlying thread pool.
 	**/
 	FQueuedLowLevelThreadPool(uint32 InMaxConcurrency = ~0, TFunction<EQueuedWorkPriority(EQueuedWorkPriority)> InPriorityMapper = [](EQueuedWorkPriority InPriority) { return InPriority; }, LowLevelTasks::FScheduler* InScheduler = &LowLevelTasks::FScheduler::Get()) 
-		: PriorityMapper(InPriorityMapper), Scheduler(InScheduler)
+		: Scheduler(InScheduler), PriorityMapper(InPriorityMapper)
 	{
 		MaxConcurrency = InMaxConcurrency;
 	}
