Fix for asynchronous unlock of UAT process serializer on POSIX platforms.

#rnx

diff --git a/Engine/Source/Runtime/Core/Private/Misc/MonitoredProcess.cpp b/Engine/Source/Runtime/Core/Private/Misc/MonitoredProcess.cpp
--- a/Engine/Source/Runtime/Core/Private/Misc/MonitoredProcess.cpp
+++ b/Engine/Source/Runtime/Core/Private/Misc/MonitoredProcess.cpp
@@ -188,10 +188,8 @@ void FMonitoredProcess::Tick()
 }
 
 
-
-
-FCriticalSection FSerializedUATProcess::Serializer;
-
+/* FSerializedUATProcess
+*****************************************************************************/
 FSerializedUATProcess::FSerializedUATProcess(const FString& RunUATCommandline, bool InHidden, bool InCreatePipes)
 	// we will modify URL and Params in this constructor, so there's no need to pass anything up to base
 	: FMonitoredProcess("", "", InHidden, InCreatePipes)
@@ -212,12 +210,20 @@ FSerializedUATProcess::FSerializedUATProcess(const FString& RunUATCommandline, b
 
 bool FSerializedUATProcess::Launch()
 {
-	Async(EAsyncExecution::Thread, [this]()
+	static FCriticalSection SerializerCS;
+	static bool bHasSucceededOnceFlag = false;
+
+	FCriticalSection* Serializer = &SerializerCS;
+	bool* bHasSucceededOnce = &bHasSucceededOnceFlag;
+
+	Async(EAsyncExecution::Thread, [this, Serializer, bHasSucceededOnce]()
 		{
 			// don't let this do anything if another process is running
-			Serializer.Lock();
+			FScopeLock Lock(Serializer);
+
+			FEventRef Event;
 
-			if (bHasSucceededOnce)
+			if (*bHasSucceededOnce)
 			{
 				Params += TEXT(" -nocompile");
 			}
@@ -225,23 +231,21 @@ bool FSerializedUATProcess::Launch()
 			FOnMonitoredProcessCompleted OriginalCompletedDelegate = CompletedDelegate;
 			FSimpleDelegate OriginalCanceledDelegate = CanceledDelegate;
 
-			CompletedDelegate.BindLambda([this, OriginalCompletedDelegate](int32 ExitCode)
+			CompletedDelegate.BindLambda([this, &Event, bHasSucceededOnce, OriginalCompletedDelegate](int32 ExitCode)
 				{
 					if (ExitCode == 0 || ExitCode == 10)
 					{
-						bHasSucceededOnce = true;
+						*bHasSucceededOnce = true;
 					}
 					OriginalCompletedDelegate.ExecuteIfBound(ExitCode);
 
-					// let another one in
-					Serializer.Unlock();
+					Event->Trigger();
 				});
-			CanceledDelegate.BindLambda([this, OriginalCanceledDelegate]()
+			CanceledDelegate.BindLambda([this, &Event, OriginalCanceledDelegate]()
 				{
 					OriginalCanceledDelegate.ExecuteIfBound();
 
-					// let another one in
-					Serializer.Unlock();
+					Event->Trigger();
 				});
 
 			UE_LOG(LogMonitoredProcess, Log, TEXT("Running Serialized UAT: '%s %s'"), *URL, *Params);
@@ -250,17 +254,15 @@ bool FSerializedUATProcess::Launch()
 			{
 				LaunchFailedDelegate.ExecuteIfBound();
 
-				// let another one in
-				Serializer.Unlock();
+				Event->Trigger();
 			}
+
+			Event->Wait();
 		});
 
 	return true;
 }
 
-
-bool FSerializedUATProcess::bHasSucceededOnce = false;
-
 FString FSerializedUATProcess::GetUATPath()
 {
 #if PLATFORM_WINDOWS
diff --git a/Engine/Source/Runtime/Core/Public/Misc/MonitoredProcess.h b/Engine/Source/Runtime/Core/Public/Misc/MonitoredProcess.h
--- a/Engine/Source/Runtime/Core/Public/Misc/MonitoredProcess.h
+++ b/Engine/Source/Runtime/Core/Public/Misc/MonitoredProcess.h
@@ -288,7 +288,4 @@ private:
 	// Holds a delegate that is executed when the process fails to launch (delayed, in a thread). Used in place of the return value
 	// of Launch in the parent class, since it's async
 	FSimpleDelegate LaunchFailedDelegate;
-
-	static FCriticalSection Serializer;
-	static bool bHasSucceededOnce;
 };
