UE-108091 //UE5/Release-5.0-EarlyAccess - Compile UnrealGame HoloLens - GenericPlatformProcess.h(718): error C3861: '__rdtsc': identifier not found
UE-107675 HoloLens 2 fails to package for UE5 due to 'undeclared identifiers'
-Fixed arm64 hololens build break in ReadCycleCounter(), implemented it for arm & arm64.
-Fixed a logic error in the case where the counter fetch always returns zero (possible on certain cpus).

diff --git a/Engine/Source/Runtime/Core/Public/GenericPlatform/GenericPlatformProcess.h b/Engine/Source/Runtime/Core/Public/GenericPlatform/GenericPlatformProcess.h
--- a/Engine/Source/Runtime/Core/Public/GenericPlatform/GenericPlatformProcess.h
+++ b/Engine/Source/Runtime/Core/Public/GenericPlatform/GenericPlatformProcess.h
@@ -714,8 +714,24 @@ struct CORE_API FGenericPlatformProcess
 	*/
 	static FORCEINLINE uint64 ReadCycleCounter()
 	{
-#if defined(_MSC_VER)
+#if defined(_MSC_VER) && PLATFORM_CPU_X86_FAMILY
 		return __rdtsc();
+#elif defined(_MSC_VER) && defined(_M_ARM)
+		return __rdpmccntr64();
+#elif defined(_MSC_VER) && defined(_M_ARM64)
+		// Duplicate define of ARM64_PMCCNTR_EL0 from winnt.h to avoid including the entire header
+#		define UE_ARM64_SYSREG(op0, op1, crn, crm, op2) \
+		( ((op0 & 1) << 14) | \
+		  ((op1 & 7) << 11) | \
+		  ((crn & 15) << 7) | \
+		  ((crm & 15) << 3) | \
+		  ((op2 & 7) << 0) )
+#		define UE_ARM64_PMCCNTR_EL0 UE_ARM64_SYSREG(3,3, 9,13,0)  // Cycle Count Register [CP15_PMCCNTR]	
+
+		return _ReadStatusReg(UE_ARM64_PMCCNTR_EL0);
+
+#		undef UE_ARM64_PMCCNTR_EL0
+#		undef UE_ARM64_SYSREG
 #elif __has_builtin(__builtin_readcyclecounter)
 		return __builtin_readcyclecounter();
 #else
@@ -735,6 +751,6 @@ struct CORE_API FGenericPlatformProcess
 		do
 		{
 			Yield();
-		} while((ReadCycleCounter() - start) <= cycles);
+		} while((ReadCycleCounter() - start) < cycles);
 	}
 };
