Move up StallDetector shutdown into PreExit to shutdown before exceptions in OnExit broadcast.
Use SleepNoStats to prevent stat accumulation of sleep time
Remove CPUTRACE macros to remove dependency on that system (they weren't really useful)
Improve StallDetector logging

diff --git a/Engine/Source/Runtime/Core/Private/ProfilingDebugging/StallDetector.cpp b/Engine/Source/Runtime/Core/Private/ProfilingDebugging/StallDetector.cpp
--- a/Engine/Source/Runtime/Core/Private/ProfilingDebugging/StallDetector.cpp
+++ b/Engine/Source/Runtime/Core/Private/ProfilingDebugging/StallDetector.cpp
@@ -10,7 +10,6 @@
 #include "HAL/PlatformTLS.h"
 #include "HAL/Runnable.h"
 #include "HAL/RunnableThread.h"
-#include "ProfilingDebugging/CpuProfilerTrace.h"
 
 // force normal behavior in the face of debug configuration and debugger attached
 #define STALL_DETECTOR_DEBUG 0
@@ -93,8 +92,6 @@ uint32 UE::FStallDetectorRunnable::Run()
 {
 	while (!StopThread)
 	{
-		TRACE_CPUPROFILER_EVENT_SCOPE(FStallDetectorRunnable::Run);
-
 #if STALL_DETECTOR_HEART_BEAT_CLOCK
 		Clock.Tick();
 #endif
@@ -114,8 +111,8 @@ uint32 UE::FStallDetectorRunnable::Run()
 			}
 		}
 
-		// Sleep a an interval, the resolution at which we want to detect an overage
-		FPlatformProcess::Sleep(0.005);
+		// Sleep an interval, the resolution at which we want to detect an overage
+		FPlatformProcess::SleepNoStats(0.005);
 	}
 
 	return 0;
@@ -271,7 +268,7 @@ void UE::FStallDetector::Check(bool bIsComplete, double InWhenToCheckSeconds)
 			FString OverageString = FString::Printf(TEXT("[FStallDetector] [%s] Overage of %f\n"), Stats.Name, OverageSeconds);
 			FPlatformMisc::LocalPrint(OverageString.GetCharArray().GetData());
 #endif
-			UE_LOG(LogStall, Display, TEXT("Stall detector '%s' exceeded budget of %fs, and completed in %fs"), Stats.Name, Stats.BudgetSeconds, OverageSeconds);
+			UE_LOG(LogStall, Display, TEXT("Stall detector '%s' complete in %fs (%fs overbudget)"), Stats.Name, DeltaSeconds, OverageSeconds);
 		}
 	}
 	else
@@ -313,8 +310,6 @@ void UE::FStallDetector::CheckAndReset()
 
 void UE::FStallDetector::OnStallDetected(uint32 InThreadId, const double InElapsedSeconds)
 {
-	TRACE_CPUPROFILER_EVENT_SCOPE(FStallDetector::OnStallDetected);
-
 	Stats.TotalTriggeredCount++;
 
 	//
@@ -368,19 +363,20 @@ void UE::FStallDetector::OnStallDetected(uint32 InThreadId, const double InElaps
 		Stats.bReported = true;
 		Stats.TotalReportedCount++;
 		const int NumStackFramesToIgnore = FPlatformTLS::GetCurrentThreadId() == InThreadId ? 2 : 0;
+		UE_LOG(LogStall, Warning, TEXT("Stall detector '%s' exceeded budget of %fs, reporting..."), Stats.Name, Stats.BudgetSeconds);
+		double ReportSeconds = FStallDetector::Seconds();
 		ReportStall(Stats.Name, InThreadId, NumStackFramesToIgnore);
-		UE_LOG(LogStall, Warning, TEXT("Stall detector '%s' exceeded budget of %fs, and was reported"), Stats.Name, Stats.BudgetSeconds);
+		ReportSeconds = FStallDetector::Seconds() - ReportSeconds;
+		UE_LOG(LogStall, Log, TEXT("Stall detector '%s' report submitted, and took %fs"), Stats.Name, ReportSeconds);
 	}
 	else
 	{
-		UE_LOG(LogStall, Warning, TEXT("Stall detector '%s' exceeded budget of %fs seconds"), Stats.Name, Stats.BudgetSeconds);
+		UE_LOG(LogStall, Warning, TEXT("Stall detector '%s' exceeded budget of %fs"), Stats.Name, Stats.BudgetSeconds);
 	}
 }
 
 double UE::FStallDetector::Seconds()
 {
-	TRACE_CPUPROFILER_EVENT_SCOPE(FStallDetector::Seconds);
-
 	check(InitCount);
 
 	double Result;
@@ -414,6 +410,8 @@ void UE::FStallDetector::Startup()
 {
 	if (++InitCount == 1)
 	{
+		UE_LOG(LogStall, Log, TEXT("Startup..."));
+
 		check(FPlatformTime::GetSecondsPerCycle());
 
 		// Cannot be a global due to clock member
@@ -430,6 +428,8 @@ void UE::FStallDetector::Startup()
 				FPlatformProcess::YieldThread();
 			}
 		}
+
+		UE_LOG(LogStall, Log, TEXT("Startup complete."));
 	}
 }
 
@@ -437,11 +437,15 @@ void UE::FStallDetector::Shutdown()
 {
 	if (--InitCount == 0)
 	{
+		UE_LOG(LogStall, Log, TEXT("Shutdown..."));
+
 		delete Thread;
 		Thread = nullptr;
 
 		delete Runnable;
 		Runnable = nullptr;
+		
+		UE_LOG(LogStall, Log, TEXT("Shutdown complete."));
 	}
 }
 
diff --git a/Engine/Source/Runtime/Launch/Private/LaunchEngineLoop.cpp b/Engine/Source/Runtime/Launch/Private/LaunchEngineLoop.cpp
--- a/Engine/Source/Runtime/Launch/Private/LaunchEngineLoop.cpp
+++ b/Engine/Source/Runtime/Launch/Private/LaunchEngineLoop.cpp
@@ -4779,7 +4779,7 @@ static inline void EndFrameRenderThread(FRHICommandListImmediate& RHICmdList, ui
 
 void FEngineLoop::Tick()
 {
-	SCOPE_STALL_COUNTER(FEngineLoop::Tick, 2.0)
+	SCOPE_STALL_COUNTER(FEngineLoop::Tick, 2.0);
 
 	// make sure to catch any FMemStack uses outside of UWorld::Tick
 	FMemMark MemStackMark(FMemStack::Get());
@@ -5942,6 +5942,10 @@ void FEngineLoop::AppPreExit( )
 	FRemoteConfig::Flush();
 #endif
 
+#if STALL_DETECTOR
+	UE::FStallDetector::Shutdown();
+#endif
+
 	FCoreDelegates::OnExit.Broadcast();
 
 	// Clean up the thread pool
@@ -6000,10 +6004,6 @@ void FEngineLoop::AppPreExit( )
 
 void FEngineLoop::AppExit( )
 {
-#if STALL_DETECTOR
-	UE::FStallDetector::Shutdown();
-#endif
-
 #if !WITH_ENGINE
 	// when compiled WITH_ENGINE, this will happen in FEngineLoop::Exit()
 	FTaskGraphInterface::Shutdown();
