U5 - Implement ProcessorGroup support to enable more than 64 cores on Windows.

diff --git a/Engine/Source/Runtime/Core/Private/Async/Fundamental/Scheduler.cpp b/Engine/Source/Runtime/Core/Private/Async/Fundamental/Scheduler.cpp
--- a/Engine/Source/Runtime/Core/Private/Async/Fundamental/Scheduler.cpp
+++ b/Engine/Source/Runtime/Core/Private/Async/Fundamental/Scheduler.cpp
@@ -43,6 +43,29 @@ namespace LowLevelTasks
 		const uint32 WaitTimes[8] = { 23, 31, 41, 37, 47, 29, 19, 43 };
 		uint32 WaitTime = WaitTimes[WorkerId % 8];
 		uint64 ThreadAffinityMask = FPlatformAffinity::GetTaskGraphThreadMask();
+
+		const FProcessorGroupDesc& ProcessorGroups = FPlatformMisc::GetProcessorGroupDesc();
+		int32 CpuGroupCount = ProcessorGroups.NumProcessorGroups;
+		uint16 CpuGroup = 0;
+
+		//offset the firt set of workers to leave space for Game, RHI and Renderthread.
+		uint64 GroupWorkerId = WorkerId + 2;
+		for (uint16 GroupIndex = 0; GroupIndex < CpuGroupCount; GroupIndex++)
+		{
+			CpuGroup = GroupIndex;
+
+			uint32 CpusInGroup = FMath::CountBits(ProcessorGroups.ThreadAffinities[GroupIndex]);
+			if(GroupWorkerId < CpusInGroup)
+			{
+				if (CpuGroup != 0) //pin larger groups workers to a core and leave first group as is for legacy reasons
+				{
+					ThreadAffinityMask = 1ull << GroupWorkerId;
+				}		
+				break;
+			}
+			GroupWorkerId -= CpusInGroup;
+		}
+		
 		return MakeUnique<FThread>
 		(
 			bPermitBackgroundWork ? *FString::Printf(TEXT("Background Worker #%d"), WorkerId) : *FString::Printf(TEXT("Foreground Worker #%d"), WorkerId),
@@ -50,7 +73,7 @@ namespace LowLevelTasks
 			{ 
 				FSleepEvent Event;
 				WorkerMain(&Event, ExternalWorkerLocalQueue, WaitTime, bPermitBackgroundWork);
-			}, 0, Priority, ThreadAffinityMask
+			}, 0, Priority, FThreadAffinity{ ThreadAffinityMask & ProcessorGroups.ThreadAffinities[CpuGroup], CpuGroup }, bIsForkable
 		);
 	}
 
diff --git a/Engine/Source/Runtime/Core/Private/GenericPlatform/GenericPlatformMisc.cpp b/Engine/Source/Runtime/Core/Private/GenericPlatform/GenericPlatformMisc.cpp
--- a/Engine/Source/Runtime/Core/Private/GenericPlatform/GenericPlatformMisc.cpp
+++ b/Engine/Source/Runtime/Core/Private/GenericPlatform/GenericPlatformMisc.cpp
@@ -1178,6 +1178,20 @@ int32 FGenericPlatformMisc::NumberOfCoresIncludingHyperthreads()
 	return FPlatformMisc::NumberOfCores();
 }
 
+FProcessorGroupDesc InternalGetProcessorGroupDesc()
+{
+	FProcessorGroupDesc Desc;
+	Desc.NumProcessorGroups = 1;
+	memset(Desc.ThreadAffinities, 0xFF, sizeof(Desc.ThreadAffinities));
+	return Desc;
+}
+
+const FProcessorGroupDesc& FGenericPlatformMisc::GetProcessorGroupDesc()
+{
+	static FProcessorGroupDesc Desc = InternalGetProcessorGroupDesc();
+	return Desc;
+}
+
 int32 FGenericPlatformMisc::NumberOfWorkerThreadsToSpawn()
 {
 	static int32 MaxGameThreads = 4;
diff --git a/Engine/Source/Runtime/Core/Private/HAL/Thread.cpp b/Engine/Source/Runtime/Core/Private/HAL/Thread.cpp
--- a/Engine/Source/Runtime/Core/Private/HAL/Thread.cpp
+++ b/Engine/Source/Runtime/Core/Private/HAL/Thread.cpp
@@ -18,13 +18,14 @@ public:
 		TUniqueFunction<void()>&& InThreadFunction,
 		uint32 StackSize,
 		EThreadPriority ThreadPriority,
-		uint64 ThreadAffinityMask,
+		FThreadAffinity ThreadAffinity,
 		bool bIsForkable
 	) 
 		: ThreadFunction(MoveTemp(InThreadFunction))
-		, RunnableThread(bIsForkable ? FForkProcessHelper::CreateForkableThread(this, ThreadName, StackSize, ThreadPriority, ThreadAffinityMask)
-			: FRunnableThread::Create(this, ThreadName, StackSize, ThreadPriority, ThreadAffinityMask))
+		, RunnableThread(bIsForkable ? FForkProcessHelper::CreateForkableThread(this, ThreadName, StackSize, ThreadPriority, ThreadAffinity.ThreadAffinityMask)
+			: FRunnableThread::Create(this, ThreadName, StackSize, ThreadPriority, ThreadAffinity.ThreadAffinityMask))
 	{
+		RunnableThread->SetThreadAffinity(ThreadAffinity);
 		check(IsJoinable());
 	}
 
@@ -96,10 +97,10 @@ FThread::FThread(
 	TUniqueFunction<void()>&& ThreadFunction,
 	uint32 StackSize,
 	EThreadPriority ThreadPriority,
-	uint64 ThreadAffinityMask,
+	FThreadAffinity ThreadAffinity,
 	bool bIsForkable
 )
-	: Impl(MakeShared<FThreadImpl, ESPMode::ThreadSafe>(ThreadName, MoveTemp(ThreadFunction), StackSize, ThreadPriority, ThreadAffinityMask, bIsForkable))
+	: Impl(MakeShared<FThreadImpl, ESPMode::ThreadSafe>(ThreadName, MoveTemp(ThreadFunction), StackSize, ThreadPriority, ThreadAffinity, bIsForkable))
 {
 	Impl->Initialize(Impl);
 }
diff --git a/Engine/Source/Runtime/Core/Private/Windows/WindowsPlatformMisc.cpp b/Engine/Source/Runtime/Core/Private/Windows/WindowsPlatformMisc.cpp
--- a/Engine/Source/Runtime/Core/Private/Windows/WindowsPlatformMisc.cpp
+++ b/Engine/Source/Runtime/Core/Private/Windows/WindowsPlatformMisc.cpp
@@ -1638,7 +1638,7 @@ bool FWindowsPlatformMisc::IsValidAbsolutePathFormat(const FString& Path)
 	return bIsValid;
 }
 
-static void QueryCpuInformation(uint32& OutGroupCount, uint32& OutNumaNodeCount, uint32& OutCoreCount, uint32& OutLogicalProcessorCount, bool bForceSingleNumaNode = false)
+static void QueryCpuInformation(FProcessorGroupDesc& OutGroupDesc, uint32& OutNumaNodeCount, uint32& OutCoreCount, uint32& OutLogicalProcessorCount, bool bForceSingleNumaNode = false)
 {
 	GROUP_AFFINITY FilterGroupAffinity = {};
 
@@ -1652,7 +1652,7 @@ static void QueryCpuInformation(uint32& OutGroupCount, uint32& OutNumaNodeCount,
 		GetNumaNodeProcessorMaskEx(NodeNumber, &FilterGroupAffinity);
 	}
 
-	OutGroupCount = OutNumaNodeCount = OutCoreCount = OutLogicalProcessorCount = 0;
+	OutNumaNodeCount = OutCoreCount = OutLogicalProcessorCount = 0;
 	uint8* BufferPtr = nullptr;
 	DWORD BufferBytes = 0;
 
@@ -1711,7 +1711,11 @@ static void QueryCpuInformation(uint32& OutGroupCount, uint32& OutNumaNodeCount,
 
 					if (ProcessorInfo->Relationship == RelationGroup)
 					{
-						OutGroupCount = ProcessorInfo->Group.ActiveGroupCount;
+						OutGroupDesc.NumProcessorGroups = FMath::Min<uint16>(FProcessorGroupDesc::MaxNumProcessorGroups, ProcessorInfo->Group.ActiveGroupCount);
+						for(int32 GroupIndex = 0; GroupIndex < OutGroupDesc.NumProcessorGroups; GroupIndex++)
+						{
+							OutGroupDesc.ThreadAffinities[GroupIndex] = ProcessorInfo->Group.GroupInfo[GroupIndex].ActiveProcessorMask;
+						}
 					}
 
 					InfoPtr += ProcessorInfo->Size;
@@ -1728,11 +1732,11 @@ int32 FWindowsPlatformMisc::NumberOfCores()
 	static int32 CoreCount = 0;
 	if (CoreCount == 0)
 	{
-		uint32 NumGroups = 0;
+		FProcessorGroupDesc GroupDesc;
 		uint32 NumaNodeCount = 0;
 		uint32 NumCores = 0;
 		uint32 LogicalProcessorCount = 0;
-		QueryCpuInformation(NumGroups, NumaNodeCount, NumCores, LogicalProcessorCount);
+		QueryCpuInformation(GroupDesc, NumaNodeCount, NumCores, LogicalProcessorCount);
 
 		if (FCommandLine::IsInitialized() && FParse::Param(FCommandLine::Get(), TEXT("usehyperthreading")))
 		{
@@ -1755,16 +1759,32 @@ int32 FWindowsPlatformMisc::NumberOfCores()
 	return CoreCount;
 }
 
+FProcessorGroupDesc NumberOfProcessorGroupsInternal()
+{
+	FProcessorGroupDesc GroupDesc;
+	uint32 NumaNodeCount = 0;
+	uint32 NumCores = 0;
+	uint32 LogicalProcessorCount = 0;
+	QueryCpuInformation(GroupDesc, NumaNodeCount, NumCores, LogicalProcessorCount);
+	return GroupDesc;
+}
+
+const FProcessorGroupDesc& FWindowsPlatformMisc::GetProcessorGroupDesc()
+{
+	static FProcessorGroupDesc GroupDesc(NumberOfProcessorGroupsInternal());
+	return GroupDesc;
+}
+
 int32 FWindowsPlatformMisc::NumberOfCoresIncludingHyperthreads()
 {
 	static int32 CoreCount = 0;
 	if (CoreCount == 0)
 	{
-		uint32 NumGroups = 0;
+		FProcessorGroupDesc GroupDesc;
 		uint32 NumaNodeCount = 0;
 		uint32 NumCores = 0;
 		uint32 LogicalProcessorCount = 0;
-		QueryCpuInformation(NumGroups, NumaNodeCount, NumCores, LogicalProcessorCount);
+		QueryCpuInformation(GroupDesc, NumaNodeCount, NumCores, LogicalProcessorCount);
 
 		CoreCount = LogicalProcessorCount;
 
@@ -1801,7 +1821,7 @@ int32 FWindowsPlatformMisc::NumberOfWorkerThreadsToSpawn()
 	static int32 MaxServerWorkerThreads = 4;
 
 	extern CORE_API int32 GUseNewTaskBackend;
-	int32 MaxWorkerThreads = GUseNewTaskBackend ? 62 : 26; //Windows Thread Group limit is 64
+	int32 MaxWorkerThreads = GUseNewTaskBackend ? INT32_MAX : 26;
 
 	int32 NumberOfCores = FWindowsPlatformMisc::NumberOfCores();
 	int32 NumberOfCoresIncludingHyperthreads = FWindowsPlatformMisc::NumberOfCoresIncludingHyperthreads();
diff --git a/Engine/Source/Runtime/Core/Private/Windows/WindowsRunnableThread.cpp b/Engine/Source/Runtime/Core/Private/Windows/WindowsRunnableThread.cpp
--- a/Engine/Source/Runtime/Core/Private/Windows/WindowsRunnableThread.cpp
+++ b/Engine/Source/Runtime/Core/Private/Windows/WindowsRunnableThread.cpp
@@ -108,6 +108,25 @@ uint32 FRunnableThreadWin::GuardedRun()
 	return ExitCode;
 }
 
+bool FRunnableThreadWin::SetThreadAffinity(const FThreadAffinity& Affinity)
+{
+	const FProcessorGroupDesc& ProcessorGroups = FPlatformMisc::GetProcessorGroupDesc();
+	int32 CpuGroupCount = ProcessorGroups.NumProcessorGroups;
+	check(Affinity.ProcessorGroup < CpuGroupCount);
+
+	GROUP_AFFINITY GroupAffinity = {};
+	GROUP_AFFINITY PreviousGroupAffinity = {};
+	GroupAffinity.Mask = Affinity.ThreadAffinityMask & ProcessorGroups.ThreadAffinities[Affinity.ProcessorGroup];
+	GroupAffinity.Group = Affinity.ProcessorGroup;
+	if (SetThreadGroupAffinity(Thread, &GroupAffinity, &PreviousGroupAffinity) == 0)
+	{
+		DWORD LastError = GetLastError();
+		UE_LOG( LogThreadingWindows, Warning, TEXT( "Runnable thread %s call to SetThreadAffinity failed with: 0x%x" ), *ThreadName, LastError);
+		return  false;
+	}
+	ThreadAffinityMask = Affinity.ThreadAffinityMask;
+	return PreviousGroupAffinity.Mask != GroupAffinity.Mask || PreviousGroupAffinity.Group != GroupAffinity.Group;
+};
 
 uint32 FRunnableThreadWin::Run()
 {
diff --git a/Engine/Source/Runtime/Core/Private/Windows/WindowsRunnableThread.h b/Engine/Source/Runtime/Core/Private/Windows/WindowsRunnableThread.h
--- a/Engine/Source/Runtime/Core/Private/Windows/WindowsRunnableThread.h
+++ b/Engine/Source/Runtime/Core/Private/Windows/WindowsRunnableThread.h
@@ -110,6 +110,8 @@ public:
 		WaitForSingleObject(Thread, INFINITE);
 	}
 
+	virtual bool SetThreadAffinity(const FThreadAffinity& Affinity) override;
+
 protected:
  	static int TranslateThreadPriority(EThreadPriority Priority);
 
diff --git a/Engine/Source/Runtime/Core/Public/GenericPlatform/GenericPlatformMisc.h b/Engine/Source/Runtime/Core/Public/GenericPlatform/GenericPlatformMisc.h
--- a/Engine/Source/Runtime/Core/Public/GenericPlatform/GenericPlatformMisc.h
+++ b/Engine/Source/Runtime/Core/Public/GenericPlatform/GenericPlatformMisc.h
@@ -333,6 +333,13 @@ enum class ENetworkConnectionType : uint8
  */
 CORE_API const TCHAR* LexToString( ENetworkConnectionType Target );
 
+struct FProcessorGroupDesc 
+{
+	static constexpr uint16 MaxNumProcessorGroups = 16;
+	uint64 ThreadAffinities[MaxNumProcessorGroups] = {}; 
+	uint16 NumProcessorGroups = 0;
+};
+
 /**
  * Generic implementation for most platforms
  */
@@ -860,6 +867,11 @@ public:
 		return 1;
 	}
 
+	/**
+	* @return a description of all the groups in the current system and their affinities
+	*/
+	static const FProcessorGroupDesc& GetProcessorGroupDesc();
+
 	/**
 	 * return the number of logical CPU cores
 	 */
diff --git a/Engine/Source/Runtime/Core/Public/HAL/PlatformAffinity.h b/Engine/Source/Runtime/Core/Public/HAL/PlatformAffinity.h
--- a/Engine/Source/Runtime/Core/Public/HAL/PlatformAffinity.h
+++ b/Engine/Source/Runtime/Core/Public/HAL/PlatformAffinity.h
@@ -5,3 +5,9 @@
 #include "CoreTypes.h"
 
 #include COMPILED_PLATFORM_HEADER(PlatformAffinity.h)
+
+struct FThreadAffinity 
+{ 
+	uint64 ThreadAffinityMask = FPlatformAffinity::GetNoAffinityMask(); 
+	uint16 ProcessorGroup = 0; 
+};
\ No newline at end of file
diff --git a/Engine/Source/Runtime/Core/Public/HAL/RunnableThread.h b/Engine/Source/Runtime/Core/Public/HAL/RunnableThread.h
--- a/Engine/Source/Runtime/Core/Public/HAL/RunnableThread.h
+++ b/Engine/Source/Runtime/Core/Public/HAL/RunnableThread.h
@@ -56,6 +56,15 @@ public:
 	 */
 	virtual void SetThreadPriority( EThreadPriority NewPriority ) = 0;
 
+	/**
+	* Changes the thread affinity of the currently running thread
+	*
+	* @param ThreadAffinityMask The thread affinity to change to (can be 0 to keep previously set affinity mask)
+	* @param ProcessorGroup The thread group to change to
+	* @return returns true if the affinity changed, false if affinity did not change
+	*/
+	virtual bool SetThreadAffinity( const FThreadAffinity& Affinity ) { return false; };
+
 	/**
 	 * Tells the thread to either pause execution or resume depending on the
 	 * passed in value.
diff --git a/Engine/Source/Runtime/Core/Public/HAL/Thread.h b/Engine/Source/Runtime/Core/Public/HAL/Thread.h
--- a/Engine/Source/Runtime/Core/Public/HAL/Thread.h
+++ b/Engine/Source/Runtime/Core/Public/HAL/Thread.h
@@ -30,19 +30,21 @@ public:
 	{}
 
 	/**
-	 * Creates and immediately starts a new system thread that will execute `ThreadFunction` argument.
-	 * Can return before the thread is actually started or when it already finished execution.
-	 * @param ThreadName Name of the thread
-	 * @param ThreadFunction The function that will be executed by the newly created thread
-	 * @param StackSize The size of the stack to create. 0 means use the current thread's stack size
-	 * @param ThreadPriority Tells the thread whether it needs to adjust its priority or not. Defaults to normal priority
-	 */
+	* Creates and immediately starts a new system thread that will execute `ThreadFunction` argument.
+	* Can return before the thread is actually started or when it already finished execution.
+	* @param ThreadName Name of the thread
+	* @param ThreadFunction The function that will be executed by the newly created thread
+	* @param StackSize The size of the stack to create. 0 means use the current thread's stack size
+	* @param ThreadPriority Tells the thread whether it needs to adjust its priority or not. Defaults to normal priority
+	* @param ThreadAffinity Tells the thread whether it needs to adjust its affinity or not. Defaults to no affinity
+	* @param bIsForkable Tells the thread whether it can be forked. Defaults to false
+	*/
 	FThread(
 		TCHAR const* ThreadName,
 		TUniqueFunction<void()>&& ThreadFunction,
 		uint32 StackSize = 0,
 		EThreadPriority ThreadPriority = TPri_Normal,
-		uint64 ThreadAffinityMask = FPlatformAffinity::GetNoAffinityMask(),
+		FThreadAffinity ThreadAffinity = FThreadAffinity(),
 		bool bIsForkable = false
 	);
 
diff --git a/Engine/Source/Runtime/Core/Public/Windows/WindowsPlatformMisc.h b/Engine/Source/Runtime/Core/Public/Windows/WindowsPlatformMisc.h
--- a/Engine/Source/Runtime/Core/Public/Windows/WindowsPlatformMisc.h
+++ b/Engine/Source/Runtime/Core/Public/Windows/WindowsPlatformMisc.h
@@ -76,6 +76,7 @@ struct CORE_API FWindowsPlatformMisc
 	static bool Is64bitOperatingSystem();
 	static bool IsValidAbsolutePathFormat(const FString& Path);
 	static int32 NumberOfCores();
+	static const FProcessorGroupDesc& GetProcessorGroupDesc();
 	static int32 NumberOfCoresIncludingHyperthreads();
 	static int32 NumberOfWorkerThreadsToSpawn();
 
