U5 - change constructor visibility

diff --git a/Engine/Source/Runtime/Core/Public/Async/Fundamental/Scheduler.h b/Engine/Source/Runtime/Core/Public/Async/Fundamental/Scheduler.h
--- a/Engine/Source/Runtime/Core/Public/Async/Fundamental/Scheduler.h
+++ b/Engine/Source/Runtime/Core/Public/Async/Fundamental/Scheduler.h
@@ -104,7 +104,7 @@ namespace LowLevelTasks
 		};
 		using FEventQueueType = FAAArrayQueue<FSleepEvent>;
 
-	private:
+	public:
 		FScheduler() = default;
 		~FScheduler();
 
