CompactBinary: Simplified the core compact binary tests

#rnx

diff --git a/Engine/Source/Runtime/Core/Private/Tests/Serialization/CompactBinaryTest.cpp b/Engine/Source/Runtime/Core/Private/Tests/Serialization/CompactBinaryTest.cpp
--- a/Engine/Source/Runtime/Core/Private/Tests/Serialization/CompactBinaryTest.cpp
+++ b/Engine/Source/Runtime/Core/Private/Tests/Serialization/CompactBinaryTest.cpp
@@ -26,32 +26,31 @@ struct TCbFieldTypeAccessors;
 template <ECbFieldType FieldType>
 using TCbFieldValueType = typename TCbFieldTypeAccessors<FCbFieldType::GetType(FieldType)>::ValueType;
 
-template <ECbFieldType FieldType>
-constexpr bool (FCbField::*TCbFieldIsTypeFn)() const = TCbFieldTypeAccessors<FCbFieldType::GetType(FieldType)>::IsTypeFn;
-
-template <ECbFieldType FieldType>
-constexpr auto TCbFieldAsTypeFn = TCbFieldTypeAccessors<FCbFieldType::GetType(FieldType)>::AsTypeFn;
-
-#define UE_CBFIELD_TYPE_ACCESSOR(FieldType, InIsTypeFn, InAsTypeFn, InValueType)                                      \
+#define UE_CBFIELD_TYPE_ACCESSOR_NO_DEFAULT(FieldType, InIsTypeFn, InAsTypeFn, InValueType)                           \
 	template <>                                                                                                       \
 	struct TCbFieldTypeAccessors<ECbFieldType::FieldType>                                                             \
 	{                                                                                                                 \
 		using ValueType = InValueType;                                                                                \
-		static constexpr bool (FCbField::*IsTypeFn)() const = &FCbField::InIsTypeFn;                                  \
-		static constexpr auto AsTypeFn = &FCbField::InAsTypeFn;                                                       \
+		static constexpr bool (FCbField::*IsType)() const = &FCbField::InIsTypeFn;                                    \
+		static auto AsType(FCbField& Field, ValueType) { return Field.InAsTypeFn(); }                                 \
 	};
 
-#define UE_CBFIELD_TYPE_ACCESSOR_TYPED(FieldType, InIsTypeFn, InAsTypeFn, InValueType, InDefaultType)                 \
+#define UE_CBFIELD_TYPE_ACCESSOR_EX(FieldType, InIsTypeFn, InAsTypeFn, InValueType, InDefaultType)                    \
 	template <>                                                                                                       \
 	struct TCbFieldTypeAccessors<ECbFieldType::FieldType>                                                             \
 	{                                                                                                                 \
 		using ValueType = InValueType;                                                                                \
-		static constexpr bool (FCbField::*IsTypeFn)() const = &FCbField::InIsTypeFn;                                  \
-		static constexpr InValueType (FCbField::*AsTypeFn)(InDefaultType) = &FCbField::InAsTypeFn;                    \
+		static constexpr bool (FCbField::*IsType)() const = &FCbField::InIsTypeFn;                                    \
+		static constexpr InValueType (FCbField::*AsType)(InDefaultType) = &FCbField::InAsTypeFn;                      \
 	};
 
-UE_CBFIELD_TYPE_ACCESSOR(Object, IsObject, AsObject, FCbObject);
-UE_CBFIELD_TYPE_ACCESSOR(Array, IsArray, AsArray, FCbArray);
+#define UE_CBFIELD_TYPE_ACCESSOR(FieldType, InIsTypeFn, InAsTypeFn, InValueType)                                      \
+	UE_CBFIELD_TYPE_ACCESSOR_EX(FieldType, InIsTypeFn, InAsTypeFn, InValueType, InValueType)
+
+UE_CBFIELD_TYPE_ACCESSOR_NO_DEFAULT(Object, IsObject, AsObject, FCbObject);
+UE_CBFIELD_TYPE_ACCESSOR_NO_DEFAULT(UniformObject, IsObject, AsObject, FCbObject);
+UE_CBFIELD_TYPE_ACCESSOR_NO_DEFAULT(Array, IsArray, AsArray, FCbArray);
+UE_CBFIELD_TYPE_ACCESSOR_NO_DEFAULT(UniformArray, IsArray, AsArray, FCbArray);
 UE_CBFIELD_TYPE_ACCESSOR(Binary, IsBinary, AsBinary, FMemoryView);
 UE_CBFIELD_TYPE_ACCESSOR(String, IsString, AsString, FAnsiStringView);
 UE_CBFIELD_TYPE_ACCESSOR(IntegerPositive, IsInteger, AsUInt64, uint64);
@@ -60,13 +59,19 @@ UE_CBFIELD_TYPE_ACCESSOR(Float32, IsFloat, AsFloat, float);
 UE_CBFIELD_TYPE_ACCESSOR(Float64, IsFloat, AsDouble, double);
 UE_CBFIELD_TYPE_ACCESSOR(BoolFalse, IsBool, AsBool, bool);
 UE_CBFIELD_TYPE_ACCESSOR(BoolTrue, IsBool, AsBool, bool);
-UE_CBFIELD_TYPE_ACCESSOR(CompactBinaryAttachment, IsCompactBinaryAttachment, AsCompactBinaryAttachment, FIoHash);
-UE_CBFIELD_TYPE_ACCESSOR(BinaryAttachment, IsBinaryAttachment, AsBinaryAttachment, FIoHash);
-UE_CBFIELD_TYPE_ACCESSOR(Hash, IsHash, AsHash, FIoHash);
-UE_CBFIELD_TYPE_ACCESSOR_TYPED(Uuid, IsUuid, AsUuid, FGuid, const FGuid&);
+UE_CBFIELD_TYPE_ACCESSOR_EX(CompactBinaryAttachment, IsCompactBinaryAttachment, AsCompactBinaryAttachment, FIoHash, const FIoHash&);
+UE_CBFIELD_TYPE_ACCESSOR_EX(BinaryAttachment, IsBinaryAttachment, AsBinaryAttachment, FIoHash, const FIoHash&);
+UE_CBFIELD_TYPE_ACCESSOR_EX(Hash, IsHash, AsHash, FIoHash, const FIoHash&);
+UE_CBFIELD_TYPE_ACCESSOR_EX(Uuid, IsUuid, AsUuid, FGuid, const FGuid&);
 UE_CBFIELD_TYPE_ACCESSOR(DateTime, IsDateTime, AsDateTimeTicks, int64);
 UE_CBFIELD_TYPE_ACCESSOR(TimeSpan, IsTimeSpan, AsTimeSpanTicks, int64);
 
+struct FCbAttachmentAccessors
+{
+	static constexpr bool (FCbField::*IsType)() const = &FCbField::IsAttachment;
+	static constexpr FIoHash (FCbField::*AsType)(const FIoHash&) = &FCbField::AsAttachment;
+};
+
 ///////////////////////////////////////////////////////////////////////////////////////////////////////////////////////
 
 class FCbFieldTestBase : public FAutomationTestBase
@@ -75,101 +80,67 @@ protected:
 	using FAutomationTestBase::FAutomationTestBase;
 	using FAutomationTestBase::TestEqual;
 
-	void TestEqualBytes(const TCHAR* What, FMemoryView Actual, TArrayView<const uint8> Expected)
+	bool TestEqual(const FString& What, const FCbArray& Actual, const FCbArray& Expected)
 	{
-		TestTrue(What, Actual.EqualBytes(MakeMemoryView(Expected)));
+		return TestTrue(What, Actual.Equals(Expected));
 	}
 
-	template <typename T, typename Default>
-	void TestFieldAsTypeNoClone(const TCHAR* What, FCbField& Field, T (FCbField::*AsTypeFn)(Default), T ExpectedValue = T(), T DefaultValue = T(), ECbFieldError ExpectedError = ECbFieldError::None)
+	bool TestEqual(const FString& What, const FCbObject& Actual, const FCbObject& Expected)
 	{
-		TestTypeError(What, Field, ExpectedError);
-		TestEqual(FString::Printf(TEXT("FCbField::As[Type](%s) -> Equal"), What), (Field.*AsTypeFn)(DefaultValue), ExpectedValue);
-		TestEqual(FString::Printf(TEXT("FCbField::As[Type](%s) -> HasError()"), What), Field.HasError(), ExpectedError != ECbFieldError::None);
-		TestEqual(FString::Printf(TEXT("FCbField::As[Type](%s) -> GetError()"), What), Field.GetError(), ExpectedError);
+		return TestTrue(What, Actual.Equals(Expected));
 	}
 
-	template <typename T, typename Default>
-	void TestFieldAsType(const TCHAR* What, FCbField& Field, T (FCbField::*AsTypeFn)(Default), T ExpectedValue = T(), T DefaultValue = T(), ECbFieldError ExpectedError = ECbFieldError::None)
+	bool TestEqualBytes(const TCHAR* What, FMemoryView Actual, TArrayView<const uint8> Expected)
 	{
-		TestFieldAsTypeNoClone(What, Field, AsTypeFn, ExpectedValue, DefaultValue, ExpectedError);
-		FCbFieldRef FieldClone = FCbFieldRef::Clone(Field);
-		TestFieldAsTypeNoClone(*FString::Printf(TEXT("%s, Clone"), What), FieldClone, AsTypeFn, ExpectedValue, DefaultValue, ExpectedError);
-		TestTrue(FString::Printf(TEXT("FCbField::Equals(%s)"), What), Field.Equals(FieldClone));
+		return TestTrue(What, Actual.EqualBytes(MakeMemoryView(Expected)));
 	}
 
-	template <typename T>
-	void TestFieldAsTypeNoClone(const TCHAR* What, FCbField& Field, T (FCbField::*AsTypeFn)(), T ExpectedValue = T(), T DefaultValue = T(), ECbFieldError ExpectedError = ECbFieldError::None)
+	template <ECbFieldType FieldType, typename T = TCbFieldValueType<FieldType>, typename AccessorsType = TCbFieldTypeAccessors<FieldType>>
+	void TestFieldNoClone(const TCHAR* What, FCbField& Field, T ExpectedValue = T(), T DefaultValue = T(), ECbFieldError ExpectedError = ECbFieldError::None, const AccessorsType& Accessors = AccessorsType())
 	{
-		TestTypeError(What, Field, ExpectedError);
-		(Field.*AsTypeFn)();
+		TestEqual(FString::Printf(TEXT("FCbField::Is[Type](%s)"), What), Invoke(Accessors.IsType, Field), ExpectedError != ECbFieldError::TypeError);
+		if (ExpectedError == ECbFieldError::None && !Field.IsBool())
+		{
+			TestFalse(FString::Printf(TEXT("FCbField::AsBool(%s) == false"), What), Field.AsBool());
+			TestTrue(FString::Printf(TEXT("FCbField::AsBool(%s) -> HasError()"), What), Field.HasError());
+			TestEqual(FString::Printf(TEXT("FCbField::AsBool(%s) -> GetError() == TypeError"), What), Field.GetError(), ECbFieldError::TypeError);
+		}
+		TestEqual(FString::Printf(TEXT("FCbField::As[Type](%s) -> Equal"), What), Invoke(Accessors.AsType, Field, DefaultValue), ExpectedValue);
 		TestEqual(FString::Printf(TEXT("FCbField::As[Type](%s) -> HasError()"), What), Field.HasError(), ExpectedError != ECbFieldError::None);
 		TestEqual(FString::Printf(TEXT("FCbField::As[Type](%s) -> GetError()"), What), Field.GetError(), ExpectedError);
 	}
 
-	template <typename T>
-	void TestFieldAsType(const TCHAR* What, FCbField& Field, T (FCbField::*AsTypeFn)(), T ExpectedValue = T(), T DefaultValue = T(), ECbFieldError ExpectedError = ECbFieldError::None)
+	template <ECbFieldType FieldType, typename T = TCbFieldValueType<FieldType>, typename AccessorsType = TCbFieldTypeAccessors<FieldType>>
+	void TestField(const TCHAR* What, FCbField& Field, T ExpectedValue = T(), T DefaultValue = T(), ECbFieldError ExpectedError = ECbFieldError::None, const AccessorsType& Accessors = AccessorsType())
 	{
-		TestFieldAsTypeNoClone(What, Field, AsTypeFn, ExpectedValue, DefaultValue, ExpectedError);
+		TestFieldNoClone<FieldType>(What, Field, ExpectedValue, DefaultValue, ExpectedError, Accessors);
 		FCbFieldRef FieldClone = FCbFieldRef::Clone(Field);
-		TestFieldAsTypeNoClone(*FString::Printf(TEXT("%s, Clone"), What), FieldClone, AsTypeFn, ExpectedValue, DefaultValue, ExpectedError);
+		TestFieldNoClone<FieldType>(*FString::Printf(TEXT("%s, Clone"), What), FieldClone, ExpectedValue, DefaultValue, ExpectedError, Accessors);
 		TestTrue(FString::Printf(TEXT("FCbField::Equals(%s)"), What), Field.Equals(FieldClone));
 	}
 
-	template <ECbFieldType FieldType, typename T = TCbFieldValueType<FieldType>>
-	void TestField(const TCHAR* What, FCbField& Field, T ExpectedValue = T(), T DefaultValue = T(), ECbFieldError ExpectedError = ECbFieldError::None)
-	{
-		TestTrue(FString::Printf(TEXT("FCbField::Is[Type](%s)"), What), (Field.*TCbFieldIsTypeFn<FieldType>)());
-		TestFieldAsType(What, Field, TCbFieldAsTypeFn<FieldType>, ExpectedValue, DefaultValue, ExpectedError);
-	}
-
-	template <ECbFieldType FieldType, typename T = TCbFieldValueType<FieldType>>
-	void TestField(const TCHAR* What, TArrayView<const uint8> Payload, T ExpectedValue = T(), T DefaultValue = T())
+	template <ECbFieldType FieldType, typename T = TCbFieldValueType<FieldType>, typename AccessorsType = TCbFieldTypeAccessors<FieldType>>
+	void TestField(const TCHAR* What, TArrayView<const uint8> Payload, T ExpectedValue = T(), T DefaultValue = T(), ECbFieldError ExpectedError = ECbFieldError::None, const AccessorsType& Accessors = AccessorsType())
 	{
 		FCbField Field(Payload.GetData(), FieldType);
 		TestEqual(FString::Printf(TEXT("FCbField::GetSize(%s)"), What), Field.GetSize(), uint64(Payload.Num() + !FCbFieldType::HasFieldType(FieldType)));
 		TestTrue(FString::Printf(TEXT("FCbField::HasValue(%s)"), What), Field.HasValue());
 		TestFalse(FString::Printf(TEXT("FCbField::HasError(%s) == false"), What), Field.HasError());
 		TestEqual(FString::Printf(TEXT("FCbField::GetError(%s) == None"), What), Field.GetError(), ECbFieldError::None);
-		TestField<FieldType>(What, Field, ExpectedValue, DefaultValue);
-	}
-
-	template <typename T, typename Default>
-	void TestFieldAsTypeError(const TCHAR* What, FCbField& Field, T (FCbField::*AsTypeFn)(Default), ECbFieldError ExpectedError, T ExpectedValue = T())
-	{
-		TestFieldAsTypeNoClone(What, Field, AsTypeFn, ExpectedValue, ExpectedValue, ExpectedError);
+		TestField<FieldType>(What, Field, ExpectedValue, DefaultValue, ExpectedError, Accessors);
 	}
 
-	template <typename T>
-	void TestFieldAsTypeError(const TCHAR* What, FCbField& Field, T (FCbField::*AsTypeFn)(), ECbFieldError ExpectedError, T ExpectedValue = T())
+	template <ECbFieldType FieldType, typename T = TCbFieldValueType<FieldType>, typename AccessorsType = TCbFieldTypeAccessors<FieldType>>
+	void TestFieldError(const TCHAR* What, FCbField& Field, ECbFieldError ExpectedError, T ExpectedValue = T(), const AccessorsType& Accessors = AccessorsType())
 	{
-		TestFieldAsTypeNoClone(What, Field, AsTypeFn, ExpectedValue, ExpectedValue, ExpectedError);
+		TestFieldNoClone<FieldType>(What, Field, ExpectedValue, ExpectedValue, ExpectedError, Accessors);
 	}
 
-	template <ECbFieldType FieldType, typename T = TCbFieldValueType<FieldType>>
-	void TestFieldError(const TCHAR* What, FCbField& Field, ECbFieldError ExpectedError, T ExpectedValue = T())
-	{
-		TestEqual(FString::Printf(TEXT("FCbField::Is[Type](%s)"), What), (Field.*TCbFieldIsTypeFn<FieldType>)(), ExpectedError != ECbFieldError::TypeError);
-		TestFieldAsTypeError(What, Field, TCbFieldAsTypeFn<FieldType>, ExpectedError, ExpectedValue);
-	}
-
-	template <ECbFieldType FieldType, typename T = TCbFieldValueType<FieldType>>
-	void TestFieldError(const TCHAR* What, TArrayView<const uint8> Payload, ECbFieldError ExpectedError, T ExpectedValue = T())
+	template <ECbFieldType FieldType, typename T = TCbFieldValueType<FieldType>, typename AccessorsType = TCbFieldTypeAccessors<FieldType>>
+	void TestFieldError(const TCHAR* What, TArrayView<const uint8> Payload, ECbFieldError ExpectedError, T ExpectedValue = T(), const AccessorsType& Accessors = AccessorsType())
 	{
 		FCbField Field(Payload.GetData(), FieldType);
-		TestFieldError<FieldType>(What, Field, ExpectedError, ExpectedValue);
-	}
-
-private:
-	void TestTypeError(const TCHAR* What, FCbField& Field, ECbFieldError ExpectedError)
-	{
-		if (ExpectedError == ECbFieldError::None && !Field.IsBool())
-		{
-			TestFalse(FString::Printf(TEXT("FCbField::IsBool(%s) == false"), What), Field.IsBool());
-			TestFalse(FString::Printf(TEXT("FCbField::AsBool(%s) == false"), What), Field.AsBool());
-			TestTrue(FString::Printf(TEXT("FCbField::AsBool(%s) -> HasError()"), What), Field.HasError());
-			TestEqual(FString::Printf(TEXT("FCbField::AsBool(%s) -> GetError() == TypeError"), What), Field.GetError(), ECbFieldError::TypeError);
-		}
+		TestFieldError<FieldType>(What, Field, ExpectedError, ExpectedValue, Accessors);
 	}
 };
 
@@ -368,7 +339,7 @@ bool FCbFieldObjectTest::RunTest(const FString& Parameters)
 		constexpr uint8 IntType = uint8(ECbFieldType::HasFieldName | ECbFieldType::IntegerPositive);
 		const uint8 Payload[] = { 12, IntType, 1, 'A', 1, IntType, 1, 'B', 2, IntType, 1, 'C', 3 };
 		FCbField Field(Payload, ECbFieldType::Object);
-		TestFieldAsType(TEXT("Object, NotEmpty"), Field, &FCbField::AsObject);
+		TestField<ECbFieldType::Object>(TEXT("Object, NotEmpty"), Field, FCbObject(Payload, ECbFieldType::Object));
 		FCbObjectRef Object = FCbObjectRef::Clone(Field.AsObject());
 		TestIntObject(Object, 3, sizeof(Payload));
 		TestIntObject(Field.AsObject(), 3, sizeof(Payload));
@@ -387,7 +358,7 @@ bool FCbFieldObjectTest::RunTest(const FString& Parameters)
 		constexpr uint8 IntType = uint8(ECbFieldType::HasFieldName | ECbFieldType::IntegerPositive);
 		const uint8 Payload[] = { 10, IntType, 1, 'A', 1, 1, 'B', 2, 1, 'C', 3 };
 		FCbField Field(Payload, ECbFieldType::UniformObject);
-		TestFieldAsType(TEXT("UniformObject, NotEmpty"), Field, &FCbField::AsObject);
+		TestField<ECbFieldType::UniformObject>(TEXT("UniformObject, NotEmpty"), Field, FCbObject(Payload, ECbFieldType::UniformObject));
 		FCbObjectRef Object = FCbObjectRef::Clone(Field.AsObject());
 		TestIntObject(Object, 3, sizeof(Payload));
 		TestIntObject(Field.AsObject(), 3, sizeof(Payload));
@@ -549,7 +520,7 @@ bool FCbFieldArrayTest::RunTest(const FString& Parameters)
 		constexpr uint8 IntType = uint8(ECbFieldType::IntegerPositive);
 		const uint8 Payload[] = { 7, 3, IntType, 1, IntType, 2, IntType, 3 };
 		FCbField Field(Payload, ECbFieldType::Array);
-		TestFieldAsType(TEXT("Array, NotEmpty"), Field, &FCbField::AsArray);
+		TestField<ECbFieldType::Array>(TEXT("Array, NotEmpty"), Field, FCbArray(Payload, ECbFieldType::Array));
 		FCbArrayRef Array = FCbArrayRef::Clone(Field.AsArray());
 		TestIntArray(Array, 3, sizeof(Payload));
 		TestIntArray(Field.AsArray(), 3, sizeof(Payload));
@@ -562,7 +533,7 @@ bool FCbFieldArrayTest::RunTest(const FString& Parameters)
 		constexpr uint8 IntType = uint8(ECbFieldType::IntegerPositive);
 		const uint8 Payload[] = { 5, 3, IntType, 1, 2, 3 };
 		FCbField Field(Payload, ECbFieldType::UniformArray);
-		TestFieldAsType(TEXT("UniformArray"), Field, &FCbField::AsArray);
+		TestField<ECbFieldType::UniformArray>(TEXT("UniformArray"), Field, FCbArray(Payload, ECbFieldType::UniformArray));
 		FCbArrayRef Array = FCbArrayRef::Clone(Field.AsArray());
 		TestIntArray(Array, 3, sizeof(Payload));
 		TestIntArray(Field.AsArray(), 3, sizeof(Payload));
@@ -667,7 +638,7 @@ bool FCbFieldBinaryTest::RunTest(const FString& Parameters)
 	{
 		const uint8 Payload[] = { 3, 4, 5, 6 }; // Size: 3, Data: 4/5/6
 		FCbField Field(Payload, ECbFieldType::Binary);
-		TestFieldAsTypeNoClone(TEXT("Binary, Value"), Field, &FCbField::AsBinary, MakeMemoryView(Payload + 1, 3));
+		TestFieldNoClone<ECbFieldType::Binary>(TEXT("Binary, Value"), Field, MakeMemoryView(Payload + 1, 3));
 	}
 
 	// Test FCbField(None) as Binary
@@ -740,6 +711,13 @@ protected:
 		Neg7  = Neg15 | Int8,
 	};
 
+	template <typename T, T (FCbField::*InAsTypeFn)(T)>
+	struct TCbIntegerAccessors
+	{
+		static constexpr bool (FCbField::*IsType)() const = &FCbField::IsInteger;
+		static constexpr T (FCbField::*AsType)(T) = InAsTypeFn;
+	};
+
 	void TestIntegerField(ECbFieldType FieldType, EIntType ExpectedMask, uint64 Magnitude)
 	{
 		uint8 Payload[9];
@@ -748,22 +726,22 @@ protected:
 		constexpr uint64 DefaultValue = 8;
 		const uint64 ExpectedValue = Negative ? uint64(-int64(Magnitude)) : Magnitude;
 		FCbField Field(Payload, FieldType);
-		TestFieldAsType(TEXT("Int8"), Field, &FCbField::AsInt8, int8(EnumHasAnyFlags(ExpectedMask, EIntType::Int8) ? ExpectedValue : DefaultValue),
-			int8(DefaultValue), EnumHasAnyFlags(ExpectedMask, EIntType::Int8) ? ECbFieldError::None : ECbFieldError::RangeError);
-		TestFieldAsType(TEXT("Int16"), Field, &FCbField::AsInt16, int16(EnumHasAnyFlags(ExpectedMask, EIntType::Int16) ? ExpectedValue : DefaultValue),
-			int16(DefaultValue), EnumHasAnyFlags(ExpectedMask, EIntType::Int16) ? ECbFieldError::None : ECbFieldError::RangeError);
-		TestFieldAsType(TEXT("Int32"), Field, &FCbField::AsInt32, int32(EnumHasAnyFlags(ExpectedMask, EIntType::Int32) ? ExpectedValue : DefaultValue),
-			int32(DefaultValue), EnumHasAnyFlags(ExpectedMask, EIntType::Int32) ? ECbFieldError::None : ECbFieldError::RangeError);
-		TestFieldAsType(TEXT("Int64"), Field, &FCbField::AsInt64, int64(EnumHasAnyFlags(ExpectedMask, EIntType::Int64) ? ExpectedValue : DefaultValue),
-			int64(DefaultValue), EnumHasAnyFlags(ExpectedMask, EIntType::Int64) ? ECbFieldError::None : ECbFieldError::RangeError);
-		TestFieldAsType(TEXT("UInt8"), Field, &FCbField::AsUInt8, uint8(EnumHasAnyFlags(ExpectedMask, EIntType::UInt8) ? ExpectedValue : DefaultValue),
-			uint8(DefaultValue), EnumHasAnyFlags(ExpectedMask, EIntType::UInt8) ? ECbFieldError::None : ECbFieldError::RangeError);
-		TestFieldAsType(TEXT("UInt16"), Field, &FCbField::AsUInt16, uint16(EnumHasAnyFlags(ExpectedMask, EIntType::UInt16) ? ExpectedValue : DefaultValue),
-			uint16(DefaultValue), EnumHasAnyFlags(ExpectedMask, EIntType::UInt16) ? ECbFieldError::None : ECbFieldError::RangeError);
-		TestFieldAsType(TEXT("UInt32"), Field, &FCbField::AsUInt32, uint32(EnumHasAnyFlags(ExpectedMask, EIntType::UInt32) ? ExpectedValue : DefaultValue),
-			uint32(DefaultValue), EnumHasAnyFlags(ExpectedMask, EIntType::UInt32) ? ECbFieldError::None : ECbFieldError::RangeError);
-		TestFieldAsType(TEXT("UInt64"), Field, &FCbField::AsUInt64, uint64(EnumHasAnyFlags(ExpectedMask, EIntType::UInt64) ? ExpectedValue : DefaultValue),
-			uint64(DefaultValue), EnumHasAnyFlags(ExpectedMask, EIntType::UInt64) ? ECbFieldError::None : ECbFieldError::RangeError);
+		TestField<ECbFieldType::IntegerNegative>(TEXT("Int8"), Field, int8(EnumHasAnyFlags(ExpectedMask, EIntType::Int8) ? ExpectedValue : DefaultValue),
+			int8(DefaultValue), EnumHasAnyFlags(ExpectedMask, EIntType::Int8) ? ECbFieldError::None : ECbFieldError::RangeError, TCbIntegerAccessors<int8, &FCbField::AsInt8>());
+		TestField<ECbFieldType::IntegerNegative>(TEXT("Int16"), Field, int16(EnumHasAnyFlags(ExpectedMask, EIntType::Int16) ? ExpectedValue : DefaultValue),
+			int16(DefaultValue), EnumHasAnyFlags(ExpectedMask, EIntType::Int16) ? ECbFieldError::None : ECbFieldError::RangeError, TCbIntegerAccessors<int16, &FCbField::AsInt16>());
+		TestField<ECbFieldType::IntegerNegative>(TEXT("Int32"), Field, int32(EnumHasAnyFlags(ExpectedMask, EIntType::Int32) ? ExpectedValue : DefaultValue),
+			int32(DefaultValue), EnumHasAnyFlags(ExpectedMask, EIntType::Int32) ? ECbFieldError::None : ECbFieldError::RangeError, TCbIntegerAccessors<int32, &FCbField::AsInt32>());
+		TestField<ECbFieldType::IntegerNegative>(TEXT("Int64"), Field, int64(EnumHasAnyFlags(ExpectedMask, EIntType::Int64) ? ExpectedValue : DefaultValue),
+			int64(DefaultValue), EnumHasAnyFlags(ExpectedMask, EIntType::Int64) ? ECbFieldError::None : ECbFieldError::RangeError, TCbIntegerAccessors<int64, &FCbField::AsInt64>());
+		TestField<ECbFieldType::IntegerPositive>(TEXT("UInt8"), Field, uint8(EnumHasAnyFlags(ExpectedMask, EIntType::UInt8) ? ExpectedValue : DefaultValue),
+			uint8(DefaultValue), EnumHasAnyFlags(ExpectedMask, EIntType::UInt8) ? ECbFieldError::None : ECbFieldError::RangeError, TCbIntegerAccessors<uint8, &FCbField::AsUInt8>());
+		TestField<ECbFieldType::IntegerPositive>(TEXT("UInt16"), Field, uint16(EnumHasAnyFlags(ExpectedMask, EIntType::UInt16) ? ExpectedValue : DefaultValue),
+			uint16(DefaultValue), EnumHasAnyFlags(ExpectedMask, EIntType::UInt16) ? ECbFieldError::None : ECbFieldError::RangeError, TCbIntegerAccessors<uint16, &FCbField::AsUInt16>());
+		TestField<ECbFieldType::IntegerPositive>(TEXT("UInt32"), Field, uint32(EnumHasAnyFlags(ExpectedMask, EIntType::UInt32) ? ExpectedValue : DefaultValue),
+			uint32(DefaultValue), EnumHasAnyFlags(ExpectedMask, EIntType::UInt32) ? ECbFieldError::None : ECbFieldError::RangeError, TCbIntegerAccessors<uint32, &FCbField::AsUInt32>());
+		TestField<ECbFieldType::IntegerPositive>(TEXT("UInt64"), Field, uint64(EnumHasAnyFlags(ExpectedMask, EIntType::UInt64) ? ExpectedValue : DefaultValue),
+			uint64(DefaultValue), EnumHasAnyFlags(ExpectedMask, EIntType::UInt64) ? ECbFieldError::None : ECbFieldError::RangeError, TCbIntegerAccessors<uint64, &FCbField::AsUInt64>());
 	}
 };
 
@@ -819,7 +797,7 @@ bool FCbFieldFloatTest::RunTest(const FString& Parameters)
 		TestField<ECbFieldType::Float32>(TEXT("Float32"), Payload, -2.28444433f);
 
 		FCbField Field(Payload, ECbFieldType::Float32);
-		TestFieldAsType(TEXT("Float32, AsDouble"), Field, &FCbField::AsDouble, -2.28444433);
+		TestField<ECbFieldType::Float64>(TEXT("Float32, AsDouble"), Field, -2.28444433);
 	}
 
 	// Test FCbField(Float, 64-bit)
@@ -828,7 +806,7 @@ bool FCbFieldFloatTest::RunTest(const FString& Parameters)
 		TestField<ECbFieldType::Float64>(TEXT("Float64"), Payload, -631475.76888888876);
 
 		FCbField Field(Payload, ECbFieldType::Float64);
-		TestFieldAsTypeError(TEXT("Float64, AsFloat"), Field, &FCbField::AsFloat, ECbFieldError::RangeError, 8.0f);
+		TestFieldError<ECbFieldType::Float32>(TEXT("Float64, AsFloat"), Field, ECbFieldError::RangeError, 8.0f);
 	}
 
 	// Test FCbField(Integer+, MaxBinary32) as Float
@@ -956,7 +934,7 @@ bool FCbFieldCompactBinaryAttachmentTest::RunTest(const FString& Parameters)
 	// Test FCbField(CompactBinaryAttachment, NonZero) AsAttachment
 	{
 		FCbField Field(SequentialBytes, ECbFieldType::CompactBinaryAttachment);
-		TestFieldAsType(TEXT("CompactBinaryAttachment, NonZero, AsAttachment"), Field, &FCbField::AsAttachment, FIoHash(SequentialBytes));
+		TestField<ECbFieldType::CompactBinaryAttachment>(TEXT("CompactBinaryAttachment, NonZero, AsAttachment"), Field, FIoHash(SequentialBytes), FIoHash(), ECbFieldError::None, FCbAttachmentAccessors());
 	}
 
 	// Test FCbField(None) as CompactBinaryAttachment
@@ -983,7 +961,7 @@ bool FCbFieldBinaryAttachmentTest::RunTest(const FString& Parameters)
 	// Test FCbField(BinaryAttachment, NonZero) AsAttachment
 	{
 		FCbField Field(SequentialBytes, ECbFieldType::BinaryAttachment);
-		TestFieldAsType(TEXT("BinaryAttachment, NonZero, AsAttachment"), Field, &FCbField::AsAttachment, FIoHash(SequentialBytes));
+		TestField<ECbFieldType::BinaryAttachment>(TEXT("BinaryAttachment, NonZero, AsAttachment"), Field, FIoHash(SequentialBytes), FIoHash(), ECbFieldError::None, FCbAttachmentAccessors());
 	}
 
 	// Test FCbField(None) as BinaryAttachment
