IoStore:
Fixed infinite hangs when appending extremely large or zero-sized chunks
#rnx

diff --git a/Engine/Source/Runtime/Core/Private/IO/IoStore.cpp b/Engine/Source/Runtime/Core/Private/IO/IoStore.cpp
--- a/Engine/Source/Runtime/Core/Private/IO/IoStore.cpp
+++ b/Engine/Source/Runtime/Core/Private/IO/IoStore.cpp
@@ -274,6 +274,7 @@ public:
 		{
 			FPlatformProcess::ReturnSynchEventToPool(CompressionBufferAvailableEvent);
 		}
+		check(AvailableCompressionBuffers.Num() == TotalCompressionBufferCount);
 		for (FIoBuffer* IoBuffer : AvailableCompressionBuffers)
 		{
 			delete IoBuffer;
@@ -292,9 +293,9 @@ public:
 		CompressionBufferSize = FMath::Max(CompressionBufferSize, static_cast<int32>(WriterSettings.CompressionBlockSize));
 		CompressionBufferSize = Align(CompressionBufferSize, FAES::AESBlockSize);
 
-		int32 CompressionBufferCount = DefaultMemoryLimit / CompressionBufferSize;
-		AvailableCompressionBuffers.Reserve(CompressionBufferCount);
-		for (int32 BufferIndex = 0; BufferIndex < CompressionBufferCount; ++BufferIndex)
+		TotalCompressionBufferCount = DefaultMemoryLimit / CompressionBufferSize;
+		AvailableCompressionBuffers.Reserve(TotalCompressionBufferCount);
+		for (int32 BufferIndex = 0; BufferIndex < TotalCompressionBufferCount; ++BufferIndex)
 		{
 			AvailableCompressionBuffers.Add(new FIoBuffer(CompressionBufferSize));
 		}
@@ -327,8 +328,12 @@ public:
 		QueueEntry->Request->PrepareSourceBufferAsync(QueueEntry->BeginCompressionBarrier);
 	}
 
-	FIoBuffer* AllocCompressionBuffer()
+	FIoBuffer* AllocCompressionBuffer(int32 TotalEntryChunkBlocksCount)
 	{
+		if (TotalEntryChunkBlocksCount > TotalCompressionBufferCount)
+		{
+			return new FIoBuffer(CompressionBufferSize);
+		}
 		FIoBuffer* AllocatedBuffer = nullptr;
 		while (!AllocatedBuffer)
 		{
@@ -348,8 +353,13 @@ public:
 		return AllocatedBuffer;
 	}
 
-	void FreeCompressionBuffer(FIoBuffer* Buffer)
+	void FreeCompressionBuffer(FIoBuffer* Buffer, int32 TotalEntryChunkBlocksCount)
 	{
+		if (TotalEntryChunkBlocksCount > TotalCompressionBufferCount)
+		{
+			delete Buffer;
+			return;
+		}
 		bool bTriggerEvent;
 		{
 			FScopeLock Lock(&AvailableCompressionBuffersCritical);
@@ -382,6 +392,7 @@ private:
 	FCriticalSection AvailableCompressionBuffersCritical;
 	TArray<FIoBuffer*> AvailableCompressionBuffers;
 	int32 CompressionBufferSize = -1;
+	int32 TotalCompressionBufferCount = -1;
 
 	friend class FIoStoreWriterImpl;
 };
@@ -969,6 +980,13 @@ private:
 
 		check(WriterSettings.CompressionBlockSize > 0);
 		const uint64 NumChunkBlocks = Align(Entry->UncompressedSize, WriterSettings.CompressionBlockSize) / WriterSettings.CompressionBlockSize;
+		if (NumChunkBlocks == 0)
+		{
+			WriterContext->CompressedChunksCount.IncrementExchange();
+			Entry->FinishCompressionBarrier->DispatchSubsequents();
+			return;
+		}
+		
 		Entry->ChunkBlocks.SetNum(int32(NumChunkBlocks));
 
 		uint64 BytesToProcess = Entry->UncompressedSize;
@@ -976,7 +994,7 @@ private:
 		for (int32 BlockIndex = 0; BlockIndex < NumChunkBlocks; ++BlockIndex)
 		{
 			FChunkBlock& Block = Entry->ChunkBlocks[BlockIndex];
-			Block.IoBuffer = WriterContext->AllocCompressionBuffer();
+			Block.IoBuffer = WriterContext->AllocCompressionBuffer(int32(NumChunkBlocks));
 			Block.CompressionMethod = CompressionMethod;
 			Block.UncompressedSize = FMath::Min(BytesToProcess, WriterSettings.CompressionBlockSize);
 			Block.UncompressedData = UncompressedData;
@@ -1189,9 +1207,7 @@ private:
 				TRACE_CPUPROFILER_EVENT_SCOPE(WriteBlockToContainer);
 				TargetPartition->ContainerFileHandle->Serialize(ChunkBlock.IoBuffer->Data(), ChunkBlock.Size);
 			}
-			// TODO: Write each block when it becomes available, currently this limits the maximum
-			// chunk size to the total compression buffer size
-			WriterContext->FreeCompressionBuffer(ChunkBlock.IoBuffer);
+			WriterContext->FreeCompressionBuffer(ChunkBlock.IoBuffer, Entry->ChunkBlocks.Num());
 			ChunkBlock.IoBuffer = nullptr;
 		}
 		WriterContext->SerializedChunksCount.IncrementExchange();
