CompactBinary: Added custom field types to allow forward compatibility when new types are added

#rnx

diff --git a/Engine/Source/Runtime/Core/Private/Serialization/CompactBinary.cpp b/Engine/Source/Runtime/Core/Private/Serialization/CompactBinary.cpp
--- a/Engine/Source/Runtime/Core/Private/Serialization/CompactBinary.cpp
+++ b/Engine/Source/Runtime/Core/Private/Serialization/CompactBinary.cpp
@@ -91,7 +91,7 @@ FCbObject FCbField::AsObject()
 	if (FCbFieldType::IsObject(Type))
 	{
 		Error = ECbFieldError::None;
-		return FCbObject::FromField(*this);
+		return FCbObject::FromFieldNoCheck(*this);
 	}
 	else
 	{
@@ -105,7 +105,7 @@ FCbArray FCbField::AsArray()
 	if (FCbFieldType::IsArray(Type))
 	{
 		Error = ECbFieldError::None;
-		return FCbArray::FromField(*this);
+		return FCbArray::FromFieldNoCheck(*this);
 	}
 	else
 	{
@@ -383,6 +383,85 @@ FTimespan FCbField::AsTimeSpan(FTimespan Default)
 	return FTimespan(AsTimeSpanTicks(Default.GetTicks()));
 }
 
+FCbCustomById FCbField::AsCustomById(FCbCustomById Default)
+{
+	if (FCbFieldType::IsCustomById(Type))
+	{
+		const uint8* PayloadBytes = static_cast<const uint8*>(Payload);
+		uint32 PayloadSizeByteCount;
+		const uint64 PayloadSize = ReadVarUInt(PayloadBytes, PayloadSizeByteCount);
+		PayloadBytes += PayloadSizeByteCount;
+
+		FCbCustomById Value;
+		uint32 TypeIdByteCount;
+		Value.Id = ReadVarUInt(PayloadBytes, TypeIdByteCount);
+		Value.Data = MakeMemoryView(PayloadBytes + TypeIdByteCount, PayloadSize - TypeIdByteCount);
+		Error = ECbFieldError::None;
+		return Value;
+	}
+	else
+	{
+		Error = ECbFieldError::TypeError;
+		return Default;
+	}
+}
+
+FCbCustomByName FCbField::AsCustomByName(FCbCustomByName Default)
+{
+	if (FCbFieldType::IsCustomByName(Type))
+	{
+		const uint8* PayloadBytes = static_cast<const uint8*>(Payload);
+		uint32 PayloadSizeByteCount;
+		const uint64 PayloadSize = ReadVarUInt(PayloadBytes, PayloadSizeByteCount);
+		PayloadBytes += PayloadSizeByteCount;
+
+		uint32 TypeNameLenByteCount;
+		const uint64 TypeNameLen = ReadVarUInt(PayloadBytes, TypeNameLenByteCount);
+		PayloadBytes += TypeNameLenByteCount;
+
+		FCbCustomByName Value;
+		Value.Name = FAnsiStringView(
+			reinterpret_cast<const ANSICHAR*>(PayloadBytes),
+			static_cast<FAnsiStringView::SizeType>(TypeNameLen));
+		Value.Data = MakeMemoryView(PayloadBytes + TypeNameLen, PayloadSize - TypeNameLen - TypeNameLenByteCount);
+		Error = ECbFieldError::None;
+		return Value;
+	}
+	else
+	{
+		Error = ECbFieldError::TypeError;
+		return Default;
+	}
+}
+
+FMemoryView FCbField::AsCustom(uint64 Id, FMemoryView Default)
+{
+	FCbCustomById Custom = AsCustomById(FCbCustomById{Id, Default});
+	if (Custom.Id == Id)
+	{
+		return Custom.Data;
+	}
+	else
+	{
+		Error = ECbFieldError::RangeError;
+		return Default;
+	}
+}
+
+FMemoryView FCbField::AsCustom(FAnsiStringView Name, FMemoryView Default)
+{
+	const FCbCustomByName Custom = AsCustomByName(FCbCustomByName{Name, Default});
+	if (Custom.Name.Equals(Name, ESearchCase::CaseSensitive))
+	{
+		return Custom.Data;
+	}
+	else
+	{
+		Error = ECbFieldError::RangeError;
+		return Default;
+	}
+}
+
 uint64 FCbField::GetSize() const
 {
 	return sizeof(ECbFieldType) + GetViewNoType().GetSize();
@@ -401,6 +480,8 @@ uint64 FCbField::GetPayloadSize() const
 	case ECbFieldType::UniformArray:
 	case ECbFieldType::Binary:
 	case ECbFieldType::String:
+	case ECbFieldType::CustomById:
+	case ECbFieldType::CustomByName:
 	{
 		uint32 PayloadSizeByteCount;
 		const uint64 PayloadSize = ReadVarUInt(Payload, PayloadSizeByteCount);
@@ -476,10 +557,10 @@ void FCbField::IterateAttachments(FCbFieldVisitor Visitor) const
 	{
 	case ECbFieldType::Object:
 	case ECbFieldType::UniformObject:
-		return FCbObject::FromField(*this).IterateAttachments(Visitor);
+		return FCbObject::FromFieldNoCheck(*this).IterateAttachments(Visitor);
 	case ECbFieldType::Array:
 	case ECbFieldType::UniformArray:
-		return FCbArray::FromField(*this).IterateAttachments(Visitor);
+		return FCbArray::FromFieldNoCheck(*this).IterateAttachments(Visitor);
 	case ECbFieldType::CompactBinaryAttachment:
 	case ECbFieldType::BinaryAttachment:
 		return Visitor(*this);
diff --git a/Engine/Source/Runtime/Core/Private/Serialization/CompactBinarySerialization.cpp b/Engine/Source/Runtime/Core/Private/Serialization/CompactBinarySerialization.cpp
--- a/Engine/Source/Runtime/Core/Private/Serialization/CompactBinarySerialization.cpp
+++ b/Engine/Source/Runtime/Core/Private/Serialization/CompactBinarySerialization.cpp
@@ -48,6 +48,8 @@ bool TryMeasureCompactBinary(FMemoryView View, ECbFieldType& OutType, uint64& Ou
 	case ECbFieldType::String:
 	case ECbFieldType::IntegerPositive:
 	case ECbFieldType::IntegerNegative:
+	case ECbFieldType::CustomById:
+	case ECbFieldType::CustomByName:
 		bDynamicSize = true;
 		break;
 	case ECbFieldType::Float32:
@@ -116,6 +118,8 @@ bool TryMeasureCompactBinary(FMemoryView View, ECbFieldType& OutType, uint64& Ou
 	case ECbFieldType::UniformArray:
 	case ECbFieldType::Binary:
 	case ECbFieldType::String:
+	case ECbFieldType::CustomById:
+	case ECbFieldType::CustomByName:
 		if (View.GetSize() == 0)
 		{
 			OutSize = Size + 1;
diff --git a/Engine/Source/Runtime/Core/Private/Serialization/CompactBinaryValidation.cpp b/Engine/Source/Runtime/Core/Private/Serialization/CompactBinaryValidation.cpp
--- a/Engine/Source/Runtime/Core/Private/Serialization/CompactBinaryValidation.cpp
+++ b/Engine/Source/Runtime/Core/Private/Serialization/CompactBinaryValidation.cpp
@@ -47,10 +47,10 @@ static ECbFieldType ValidateCbFieldType(FMemoryView& View, ECbValidateMode Mode,
 {
 	if (FCbFieldType::HasFieldType(Type))
 	{
-		if (View.GetSize() >= 1)
+		if (View.GetSize() >= sizeof(ECbFieldType))
 		{
 			Type = *static_cast<const ECbFieldType*>(View.GetData());
-			View += 1;
+			View += sizeof(ECbFieldType);
 			if (FCbFieldType::HasFieldType(Type))
 			{
 				AddError(Error, ECbValidateError::InvalidType);
@@ -64,7 +64,7 @@ static ECbFieldType ValidateCbFieldType(FMemoryView& View, ECbValidateMode Mode,
 		}
 	}
 
-	if (FCbFieldType::GetType(Type) > ECbFieldType::TimeSpan)
+	if (FCbFieldType::GetSerializedType(Type) != Type)
 	{
 		AddError(Error, ECbValidateError::InvalidType);
 		View.Reset();
@@ -127,25 +127,42 @@ static void ValidateCbFloat64(FMemoryView& View, ECbValidateMode Mode, ECbValida
 }
 
 /**
- * Validate and read a string from the view.
+ * Validate and read a fixed-size payload from the view.
  *
- * Modifies the view to start at the end of the string, and adds error flags if applicable.
+ * Modifies the view to start at the end of the payload, and adds error flags if applicable.
  */
-static FAnsiStringView ValidateCbString(FMemoryView& View, ECbValidateMode Mode, ECbValidateError& Error)
+static FMemoryView ValidateCbFixedPayload(FMemoryView& View, ECbValidateMode Mode, ECbValidateError& Error, uint64 Size)
 {
-	const uint64 NameSize = ValidateCbUInt(View, Mode, Error);
-	if (View.GetSize() >= NameSize)
-	{
-		const FAnsiStringView Name(static_cast<const ANSICHAR*>(View.GetData()), static_cast<int32>(NameSize));
-		View += NameSize;
-		return Name;
-	}
-	else
+	const FMemoryView Payload = View.Left(Size);
+	View += Size;
+	if (Payload.GetSize() < Size)
 	{
 		AddError(Error, ECbValidateError::OutOfBounds);
-		View.Reset();
-		return FAnsiStringView();
 	}
+	return Payload;
+};
+
+/**
+ * Validate and read a payload from the view where the view begins with the payload size.
+ *
+ * Modifies the view to start at the end of the payload, and adds error flags if applicable.
+ */
+static FMemoryView ValidateCbDynamicPayload(FMemoryView& View, ECbValidateMode Mode, ECbValidateError& Error)
+{
+	const uint64 ValueSize = ValidateCbUInt(View, Mode, Error);
+	return ValidateCbFixedPayload(View, Mode, Error, ValueSize);
+}
+
+/**
+ * Validate and read a string from the view.
+ *
+ * Modifies the view to start at the end of the string, and adds error flags if applicable.
+ */
+static FAnsiStringView ValidateCbString(FMemoryView& View, ECbValidateMode Mode, ECbValidateError& Error)
+{
+	const FMemoryView Payload = ValidateCbDynamicPayload(View, Mode, Error);
+	const FAnsiStringView Value(static_cast<const ANSICHAR*>(Payload.GetData()), static_cast<int32>(Payload.GetSize()));
+	return Value;
 }
 
 static FCbField ValidateCbField(FMemoryView& View, ECbValidateMode Mode, ECbValidateError& Error, ECbFieldType ExternalType);
@@ -281,19 +298,6 @@ static FCbField ValidateCbField(FMemoryView& View, ECbValidateMode Mode, ECbVali
 	const ECbFieldType Type = ValidateCbFieldType(View, Mode, Error, ExternalType);
 	const FAnsiStringView Name = FCbFieldType::HasFieldName(Type) ? ValidateCbString(View, Mode, Error) : FAnsiStringView();
 
-	auto ValidateFixedPayload = [&View, &Error](uint32 PayloadSize)
-	{
-		if (View.GetSize() >= PayloadSize)
-		{
-			View += PayloadSize;
-		}
-		else
-		{
-			AddError(Error, ECbValidateError::OutOfBounds);
-			View.Reset();
-		}
-	};
-
 	if (EnumHasAnyFlags(Error, ECbValidateError::OutOfBounds | ECbValidateError::InvalidType))
 	{
 		return FCbField();
@@ -325,19 +329,8 @@ static FCbField ValidateCbField(FMemoryView& View, ECbValidateMode Mode, ECbVali
 		ValidateCbArray(View, Mode, Error, FCbFieldType::GetType(Type));
 		break;
 	case ECbFieldType::Binary:
-	{
-		const uint64 ValueSize = ValidateCbUInt(View, Mode, Error);
-		if (View.GetSize() < ValueSize)
-		{
-			AddError(Error, ECbValidateError::OutOfBounds);
-			View.Reset();
-		}
-		else
-		{
-			View += ValueSize;
-		}
+		ValidateCbDynamicPayload(View, Mode, Error);
 		break;
-	}
 	case ECbFieldType::String:
 		ValidateCbString(View, Mode, Error);
 		break;
@@ -348,7 +341,7 @@ static FCbField ValidateCbField(FMemoryView& View, ECbValidateMode Mode, ECbVali
 		ValidateCbUInt(View, Mode, Error);
 		break;
 	case ECbFieldType::Float32:
-		ValidateFixedPayload(4);
+		ValidateCbFixedPayload(View, Mode, Error, 4);
 		break;
 	case ECbFieldType::Float64:
 		ValidateCbFloat64(View, Mode, Error);
@@ -356,16 +349,32 @@ static FCbField ValidateCbField(FMemoryView& View, ECbValidateMode Mode, ECbVali
 	case ECbFieldType::CompactBinaryAttachment:
 	case ECbFieldType::BinaryAttachment:
 	case ECbFieldType::Hash:
-		ValidateFixedPayload(20);
+		ValidateCbFixedPayload(View, Mode, Error, 20);
 		break;
 	case ECbFieldType::Uuid:
-		ValidateFixedPayload(16);
+		ValidateCbFixedPayload(View, Mode, Error, 16);
 		break;
 	case ECbFieldType::DateTime:
 	case ECbFieldType::TimeSpan:
-		ValidateFixedPayload(8);
+		ValidateCbFixedPayload(View, Mode, Error, 8);
+		break;
+	case ECbFieldType::CustomById:
+	{
+		FMemoryView Value = ValidateCbDynamicPayload(View, Mode, Error);
+		ValidateCbUInt(Value, Mode, Error);
+		break;
+	}
+	case ECbFieldType::CustomByName:
+	{
+		FMemoryView Value = ValidateCbDynamicPayload(View, Mode, Error);
+		const FAnsiStringView TypeName = ValidateCbString(Value, Mode, Error);
+		if (TypeName.IsEmpty() && !EnumHasAnyFlags(Error, ECbValidateError::OutOfBounds))
+		{
+			AddError(Error, ECbValidateError::InvalidType);
+		}
 		break;
 	}
+	}
 
 	if (EnumHasAnyFlags(Error, ECbValidateError::OutOfBounds | ECbValidateError::InvalidType))
 	{
diff --git a/Engine/Source/Runtime/Core/Private/Serialization/CompactBinaryWriter.cpp b/Engine/Source/Runtime/Core/Private/Serialization/CompactBinaryWriter.cpp
--- a/Engine/Source/Runtime/Core/Private/Serialization/CompactBinaryWriter.cpp
+++ b/Engine/Source/Runtime/Core/Private/Serialization/CompactBinaryWriter.cpp
@@ -560,6 +560,41 @@ void FCbWriter::AddTimeSpan(const FTimespan Value)
 	AddTimeSpanTicks(Value.GetTicks());
 }
 
+void FCbWriter::AddCustom(const uint64 TypeId, const FMemoryView Value)
+{
+	BeginField();
+	const uint32 TypeByteCount = MeasureVarUInt(TypeId);
+	const uint64 Size = TypeByteCount + Value.GetSize();
+	const uint32 SizeByteCount = MeasureVarUInt(Size);
+	const int64 PayloadOffset = Data.AddUninitialized(SizeByteCount + TypeByteCount);
+	WriteVarUInt(Size, Data.GetData() + PayloadOffset);
+	WriteVarUInt(TypeId, Data.GetData() + PayloadOffset + SizeByteCount);
+	Data.Append(static_cast<const uint8*>(Value.GetData()), Value.GetSize());
+	EndField(ECbFieldType::CustomById);
+}
+
+void FCbWriter::AddCustom(const FAnsiStringView TypeName, const FMemoryView Value)
+{
+	checkf(!TypeName.IsEmpty(), TEXT("Field '%.*hs' requires a non-empty type name for its custom type."),
+		GetActiveName().Len(), GetActiveName().GetData());
+
+	BeginField();
+
+	const uint64 TypeNameLen = TypeName.Len();
+	const uint32 TypeNameLenByteCount = MeasureVarUInt(TypeNameLen);
+
+	const uint64 Size = TypeNameLenByteCount + TypeNameLen + Value.GetSize();
+	const uint32 SizeByteCount = MeasureVarUInt(Size);
+
+	const int64 PayloadOffset = Data.AddUninitialized(SizeByteCount + TypeNameLenByteCount);
+	WriteVarUInt(Size, Data.GetData() + PayloadOffset);
+	WriteVarUInt(TypeNameLen, Data.GetData() + PayloadOffset + SizeByteCount);
+	Data.Append(reinterpret_cast<const uint8*>(TypeName.GetData()), TypeNameLen);
+	Data.Append(static_cast<const uint8*>(Value.GetData()), Value.GetSize());
+
+	EndField(ECbFieldType::CustomByName);
+}
+
 ///////////////////////////////////////////////////////////////////////////////////////////////////////////////////////
 
 FCbWriter& operator<<(FCbWriter& Writer, const FDateTime Value)
diff --git a/Engine/Source/Runtime/Core/Private/Tests/Serialization/CompactBinarySerializationTest.cpp b/Engine/Source/Runtime/Core/Private/Tests/Serialization/CompactBinarySerializationTest.cpp
--- a/Engine/Source/Runtime/Core/Private/Tests/Serialization/CompactBinarySerializationTest.cpp
+++ b/Engine/Source/Runtime/Core/Private/Tests/Serialization/CompactBinarySerializationTest.cpp
@@ -124,6 +124,18 @@ bool FCbMeasureTest::RunTest(const FString& Parameters)
 	TestMeasure(TEXT("DateTime"), {uint8(EType::DateTime)}, true, 9, EType::DateTime);
 	TestMeasure(TEXT("TimeSpan"), {uint8(EType::TimeSpan)}, true, 9, EType::TimeSpan);
 
+	TestMeasure(TEXT("CustomById, Size1B"), {uint8(EType::CustomById), 0x7f}, true, 129, EType::CustomById);
+	TestMeasure(TEXT("CustomById, Size2B"), {uint8(EType::CustomById), 0x80, 0x80}, true, 131, EType::CustomById);
+	TestMeasure(TEXT("CustomById, Size1BShort"), {uint8(EType::CustomById)}, false, 2, EType::CustomById);
+	TestMeasure(TEXT("CustomById, Size2BShort"), {uint8(EType::CustomById), 0x80}, false, 3, EType::CustomById);
+	TestMeasure(TEXT("CustomById, Size9BShort"), {uint8(EType::CustomById), 0xff}, false, 10, EType::CustomById);
+
+	TestMeasure(TEXT("CustomByName, Size1B"), {uint8(EType::CustomByName), 0x7f}, true, 129, EType::CustomByName);
+	TestMeasure(TEXT("CustomByName, Size2B"), {uint8(EType::CustomByName), 0x80, 0x80}, true, 131, EType::CustomByName);
+	TestMeasure(TEXT("CustomByName, Size1BShort"), {uint8(EType::CustomByName)}, false, 2, EType::CustomByName);
+	TestMeasure(TEXT("CustomByName, Size2BShort"), {uint8(EType::CustomByName), 0x80}, false, 3, EType::CustomByName);
+	TestMeasure(TEXT("CustomByName, Size9BShort"), {uint8(EType::CustomByName), 0xff}, false, 10, EType::CustomByName);
+
 	return true;
 }
 
diff --git a/Engine/Source/Runtime/Core/Private/Tests/Serialization/CompactBinaryTest.cpp b/Engine/Source/Runtime/Core/Private/Tests/Serialization/CompactBinaryTest.cpp
--- a/Engine/Source/Runtime/Core/Private/Tests/Serialization/CompactBinaryTest.cpp
+++ b/Engine/Source/Runtime/Core/Private/Tests/Serialization/CompactBinaryTest.cpp
@@ -65,6 +65,8 @@ UE_CBFIELD_TYPE_ACCESSOR_EX(Hash, IsHash, AsHash, FIoHash, const FIoHash&);
 UE_CBFIELD_TYPE_ACCESSOR_EX(Uuid, IsUuid, AsUuid, FGuid, const FGuid&);
 UE_CBFIELD_TYPE_ACCESSOR(DateTime, IsDateTime, AsDateTimeTicks, int64);
 UE_CBFIELD_TYPE_ACCESSOR(TimeSpan, IsTimeSpan, AsTimeSpanTicks, int64);
+UE_CBFIELD_TYPE_ACCESSOR(CustomById, IsCustomById, AsCustomById, FCbCustomById);
+UE_CBFIELD_TYPE_ACCESSOR(CustomByName, IsCustomByName, AsCustomByName, FCbCustomByName);
 
 struct FCbAttachmentAccessors
 {
@@ -90,6 +92,16 @@ protected:
 		return TestTrue(What, Actual.Equals(Expected));
 	}
 
+	bool TestEqual(const FString& What, const FCbCustomById& Actual, const FCbCustomById& Expected)
+	{
+		return TestTrue(What, Actual.Id == Expected.Id && Actual.Data == Expected.Data);
+	}
+
+	bool TestEqual(const FString& What, const FCbCustomByName& Actual, const FCbCustomByName& Expected)
+	{
+		return TestTrue(What, Actual.Name.Equals(Expected.Name, ESearchCase::CaseSensitive) && Actual.Data == Expected.Data);
+	}
+
 	bool TestEqualBytes(const TCHAR* What, FMemoryView Actual, TArrayView<const uint8> Expected)
 	{
 		return TestTrue(What, Actual.EqualBytes(MakeMemoryView(Expected)));
@@ -110,6 +122,13 @@ protected:
 		TestEqual(FString::Printf(TEXT("FCbField::As[Type](%s) -> GetError()"), What), Field.GetError(), ExpectedError);
 	}
 
+	template <ECbFieldType FieldType, typename T = TCbFieldValueType<FieldType>, typename AccessorsType = TCbFieldTypeAccessors<FieldType>>
+	void TestFieldNoClone(const TCHAR* What, TArrayView<const uint8> Payload, T ExpectedValue = T(), T DefaultValue = T(), ECbFieldError ExpectedError = ECbFieldError::None, const AccessorsType& Accessors = AccessorsType())
+	{
+		FCbField Field(Payload.GetData(), FieldType);
+		TestFieldNoClone<FieldType>(What, Field, ExpectedValue, DefaultValue, ExpectedError, Accessors);
+	}
+
 	template <ECbFieldType FieldType, typename T = TCbFieldValueType<FieldType>, typename AccessorsType = TCbFieldTypeAccessors<FieldType>>
 	void TestField(const TCHAR* What, FCbField& Field, T ExpectedValue = T(), T DefaultValue = T(), ECbFieldError ExpectedError = ECbFieldError::None, const AccessorsType& Accessors = AccessorsType())
 	{
@@ -1070,6 +1089,106 @@ bool FCbFieldTimeSpanTest::RunTest(const FString& Parameters)
 	return true;
 }
 
+IMPLEMENT_CUSTOM_SIMPLE_AUTOMATION_TEST(FCbFieldCustomByIdTest, FCbFieldTestBase, "System.Core.Serialization.CbField.CustomById", CompactBinaryTestFlags)
+bool FCbFieldCustomByIdTest::RunTest(const FString& Parameters)
+{
+	struct FCustomByIdAccessor
+	{
+		explicit FCustomByIdAccessor(uint64 Id)
+			: AsType([Id](FCbField& Field, FMemoryView Default) { return Field.AsCustom(Id, Default); })
+		{
+		}
+
+		bool (FCbField::*IsType)() const = &FCbField::IsCustomById;
+		TUniqueFunction<FMemoryView (FCbField& Field, FMemoryView Default)> AsType;
+	};
+
+	// Test FCbField(CustomById, MinId, Empty)
+	{
+		const uint8 Payload[] = {1, 0};
+		TestField<ECbFieldType::CustomById>(TEXT("CustomById, MinId, Empty"), Payload, FCbCustomById{0});
+		TestField<ECbFieldType::CustomById>(TEXT("CustomById, MinId, Empty, View"), Payload, FMemoryView(), MakeMemoryView<uint8>({1, 2, 3}), ECbFieldError::None, FCustomByIdAccessor(0));
+		TestFieldError<ECbFieldType::CustomById>(TEXT("CustomById, MinId, Empty, InvalidId"), Payload, ECbFieldError::RangeError, MakeMemoryView<uint8>({1, 2, 3}), FCustomByIdAccessor(MAX_uint64));
+	}
+
+	// Test FCbField(CustomById, MinId, Value)
+	{
+		const uint8 Payload[] = {5, 0, 1, 2, 3, 4};
+		TestFieldNoClone<ECbFieldType::CustomById>(TEXT("CustomById, MinId, Value"), Payload, FCbCustomById{0, MakeMemoryView(Payload).Right(4)});
+		TestFieldNoClone<ECbFieldType::CustomById>(TEXT("CustomById, MinId, Value, View"), Payload, MakeMemoryView(Payload).Right(4), FMemoryView(), ECbFieldError::None, FCustomByIdAccessor(0));
+		TestFieldError<ECbFieldType::CustomById>(TEXT("CustomById, MinId, Value, InvalidId"), Payload, ECbFieldError::RangeError, MakeMemoryView<uint8>({1, 2, 3}), FCustomByIdAccessor(MAX_uint64));
+	}
+
+	// Test FCbField(CustomById, MaxId, Empty)
+	{
+		const uint8 Payload[] = {9, 0xff, 0xff, 0xff, 0xff, 0xff, 0xff, 0xff, 0xff, 0xff};
+		TestField<ECbFieldType::CustomById>(TEXT("CustomById, MaxId, Empty"), Payload, FCbCustomById{MAX_uint64});
+		TestField<ECbFieldType::CustomById>(TEXT("CustomById, MaxId, Empty, View"), Payload, FMemoryView(), MakeMemoryView<uint8>({1, 2, 3}), ECbFieldError::None, FCustomByIdAccessor(MAX_uint64));
+		TestFieldError<ECbFieldType::CustomById>(TEXT("CustomById, MaxId, Empty, InvalidId"), Payload, ECbFieldError::RangeError, MakeMemoryView<uint8>({1, 2, 3}), FCustomByIdAccessor(0));
+	}
+
+	// Test FCbField(CustomById, MaxId, Value)
+	{
+		const uint8 Payload[] = {13, 0xff, 0xff, 0xff, 0xff, 0xff, 0xff, 0xff, 0xff, 0xff, 1, 2, 3, 4};
+		TestFieldNoClone<ECbFieldType::CustomById>(TEXT("CustomById, MaxId, Value"), Payload, FCbCustomById{MAX_uint64, MakeMemoryView(Payload).Right(4)});
+		TestFieldNoClone<ECbFieldType::CustomById>(TEXT("CustomById, MaxId, Value, View"), Payload, MakeMemoryView(Payload).Right(4), FMemoryView(), ECbFieldError::None, FCustomByIdAccessor(MAX_uint64));
+		TestFieldError<ECbFieldType::CustomById>(TEXT("CustomById, MaxId, Value, InvalidId"), Payload, ECbFieldError::RangeError, MakeMemoryView<uint8>({1, 2, 3}), FCustomByIdAccessor(0));
+	}
+
+	// Test FCbField(None) as CustomById
+	{
+		FCbField DefaultField;
+		TestFieldError<ECbFieldType::CustomById>(TEXT("CustomById, None"), DefaultField, ECbFieldError::TypeError, FCbCustomById{4, MakeMemoryView<uint8>({1, 2, 3})});
+		TestFieldError<ECbFieldType::CustomById>(TEXT("CustomById, None, View"), DefaultField, ECbFieldError::TypeError, MakeMemoryView<uint8>({1, 2, 3}), FCustomByIdAccessor(0));
+		const uint8 DefaultValue[] = {1, 2, 3};
+		TestEqual(TEXT("FCbField()::AsCustom(Id)"), DefaultField.AsCustom(0, MakeMemoryView(DefaultValue)), MakeMemoryView(DefaultValue));
+	}
+
+	return true;
+}
+
+IMPLEMENT_CUSTOM_SIMPLE_AUTOMATION_TEST(FCbFieldCustomByNameTest, FCbFieldTestBase, "System.Core.Serialization.CbField.CustomByName", CompactBinaryTestFlags)
+bool FCbFieldCustomByNameTest::RunTest(const FString& Parameters)
+{
+	struct FCustomByNameAccessor
+	{
+		explicit FCustomByNameAccessor(FAnsiStringView Name)
+			: AsType([Name = FString(Name)](FCbField& Field, FMemoryView Default) { return Field.AsCustom(TCHAR_TO_ANSI(*Name), Default); })
+		{
+		}
+
+		bool (FCbField::*IsType)() const = &FCbField::IsCustomByName;
+		TUniqueFunction<FMemoryView (FCbField& Field, FMemoryView Default)> AsType;
+	};
+
+	// Test FCbField(CustomByName, ABC, Empty)
+	{
+		const uint8 Payload[] = {4, 3, 'A', 'B', 'C'};
+		TestField<ECbFieldType::CustomByName>(TEXT("CustomByName, MinId, Empty"), Payload, FCbCustomByName{"ABC"});
+		TestField<ECbFieldType::CustomByName>(TEXT("CustomByName, MinId, Empty, View"), Payload, FMemoryView(), MakeMemoryView<uint8>({1, 2, 3}), ECbFieldError::None, FCustomByNameAccessor("ABC"));
+		TestFieldError<ECbFieldType::CustomByName>(TEXT("CustomByName, MinId, Empty, InvalidCase"), Payload, ECbFieldError::RangeError, MakeMemoryView<uint8>({1, 2, 3}), FCustomByNameAccessor("abc"));
+	}
+
+	// Test FCbField(CustomByName, ABC, Value)
+	{
+		const uint8 Payload[] = {8, 3, 'A', 'B', 'C', 1, 2, 3, 4};
+		TestFieldNoClone<ECbFieldType::CustomByName>(TEXT("CustomByName, MinId, Value"), Payload, FCbCustomByName{"ABC", MakeMemoryView(Payload).Right(4)});
+		TestFieldNoClone<ECbFieldType::CustomByName>(TEXT("CustomByName, MinId, Value, View"), Payload, MakeMemoryView(Payload).Right(4), FMemoryView(), ECbFieldError::None, FCustomByNameAccessor("ABC"));
+		TestFieldError<ECbFieldType::CustomByName>(TEXT("CustomByName, MinId, Value, InvalidCase"), Payload, ECbFieldError::RangeError, MakeMemoryView<uint8>({1, 2, 3}), FCustomByNameAccessor("abc"));
+	}
+
+	// Test FCbField(None) as CustomByName
+	{
+		FCbField DefaultField;
+		TestFieldError<ECbFieldType::CustomByName>(TEXT("CustomByName, None"), DefaultField, ECbFieldError::TypeError, FCbCustomByName{"ABC", MakeMemoryView<uint8>({1, 2, 3})});
+		TestFieldError<ECbFieldType::CustomByName>(TEXT("CustomByName, None, View"), DefaultField, ECbFieldError::TypeError, MakeMemoryView<uint8>({1, 2, 3}), FCustomByNameAccessor("ABC"));
+		const uint8 DefaultValue[] = {1, 2, 3};
+		TestEqual(TEXT("FCbField()::AsCustom(Name)"), DefaultField.AsCustom("ABC", MakeMemoryView(DefaultValue)), MakeMemoryView(DefaultValue));
+	}
+
+	return true;
+}
+
 IMPLEMENT_CUSTOM_SIMPLE_AUTOMATION_TEST(FCbFieldIterateAttachmentsTest, FCbFieldTestBase, "System.Core.Serialization.CbField.IterateAttachments", CompactBinaryTestFlags)
 bool FCbFieldIterateAttachmentsTest::RunTest(const FString& Parameters)
 {
diff --git a/Engine/Source/Runtime/Core/Private/Tests/Serialization/CompactBinaryValidationTest.cpp b/Engine/Source/Runtime/Core/Private/Tests/Serialization/CompactBinaryValidationTest.cpp
--- a/Engine/Source/Runtime/Core/Private/Tests/Serialization/CompactBinaryValidationTest.cpp
+++ b/Engine/Source/Runtime/Core/Private/Tests/Serialization/CompactBinaryValidationTest.cpp
@@ -170,6 +170,42 @@ bool FCbValidateTest::RunTest(const FString& Parameters)
 	TestEqual(TEXT("ValidateCompactBinary(OutOfBounds, TimeSpan)"), Validate({uint8(ECbFieldType::TimeSpan), 0, 0, 0, 0, 0, 0, 0}), ECbValidateError::OutOfBounds);
 	TestEqual(TEXT("ValidateCompactBinary(OutOfBounds, TimeSpan, NoType)"), Validate({0, 0, 0, 0, 0, 0, 0}, ECbFieldType::TimeSpan), ECbValidateError::OutOfBounds);
 
+	TestEqual(TEXT("ValidateCompactBinary(Valid, CustomById, Empty)"), Validate({uint8(ECbFieldType::CustomById), 1, 0}), ECbValidateError::None);
+	TestEqual(TEXT("ValidateCompactBinary(Valid, CustomById, Empty, NoType)"), Validate({1, 0}, ECbFieldType::CustomById), ECbValidateError::None);
+	TestEqual(TEXT("ValidateCompactBinary(Valid, CustomById, Payload)"), Validate({uint8(ECbFieldType::CustomById), 2, 0, 0}), ECbValidateError::None);
+	TestEqual(TEXT("ValidateCompactBinary(Valid, CustomById, Payload, NoType)"), Validate({2, 0, 0}, ECbFieldType::CustomById), ECbValidateError::None);
+	TestEqual(TEXT("ValidateCompactBinary(Valid, CustomById, MaxTypeId)"), Validate({uint8(ECbFieldType::CustomById), 9, 0xff, 0xff, 0xff, 0xff, 0xff, 0xff, 0xff, 0xff, 0xff}), ECbValidateError::None);
+	TestEqual(TEXT("ValidateCompactBinary(Valid, CustomById, MaxTypeId, NoType)"), Validate({9, 0xff, 0xff, 0xff, 0xff, 0xff, 0xff, 0xff, 0xff, 0xff}, ECbFieldType::CustomById), ECbValidateError::None);
+	TestEqual(TEXT("ValidateCompactBinary(Valid, CustomById, MaxTypeId, Payload)"), Validate({uint8(ECbFieldType::CustomById), 10, 0xff, 0xff, 0xff, 0xff, 0xff, 0xff, 0xff, 0xff, 0xff, 0}), ECbValidateError::None);
+	TestEqual(TEXT("ValidateCompactBinary(Valid, CustomById, MaxTypeId, Payload, NoType)"), Validate({10, 0xff, 0xff, 0xff, 0xff, 0xff, 0xff, 0xff, 0xff, 0xff, 0}, ECbFieldType::CustomById), ECbValidateError::None);
+	TestEqual(TEXT("ValidateCompactBinary(OutOfBounds, CustomById, NoSize)"), Validate({uint8(ECbFieldType::CustomById)}), ECbValidateError::OutOfBounds);
+	TestEqual(TEXT("ValidateCompactBinary(OutOfBounds, CustomById, NoSize, NoType)"), Validate({}, ECbFieldType::CustomById), ECbValidateError::OutOfBounds);
+	TestEqual(TEXT("ValidateCompactBinary(OutOfBounds, CustomById, NoTypeId)"), Validate({uint8(ECbFieldType::CustomById), 0}), ECbValidateError::OutOfBounds);
+	TestEqual(TEXT("ValidateCompactBinary(OutOfBounds, CustomById, NoTypeId, NoType)"), Validate({0}, ECbFieldType::CustomById), ECbValidateError::OutOfBounds);
+	TestEqual(TEXT("ValidateCompactBinary(OutOfBounds, CustomById, ShortTypeId)"), Validate({uint8(ECbFieldType::CustomById), 1, 0x80}), ECbValidateError::OutOfBounds);
+	TestEqual(TEXT("ValidateCompactBinary(OutOfBounds, CustomById, ShortTypeId, NoType)"), Validate({1, 0x80}, ECbFieldType::CustomById), ECbValidateError::OutOfBounds);
+	TestEqual(TEXT("ValidateCompactBinary(OutOfBounds, CustomById, NoPayload)"), Validate({uint8(ECbFieldType::CustomById), 2, 0}), ECbValidateError::OutOfBounds);
+	TestEqual(TEXT("ValidateCompactBinary(OutOfBounds, CustomById, NoPayload, NoType)"), Validate({2, 0}, ECbFieldType::CustomById), ECbValidateError::OutOfBounds);
+	TestEqual(TEXT("ValidateCompactBinary(OutOfBounds, CustomById, ShortPayload)"), Validate({uint8(ECbFieldType::CustomById), 3, 0, 0}), ECbValidateError::OutOfBounds);
+	TestEqual(TEXT("ValidateCompactBinary(OutOfBounds, CustomById, ShortPayload, NoType)"), Validate({3, 0, 0}, ECbFieldType::CustomById), ECbValidateError::OutOfBounds);
+
+	TestEqual(TEXT("ValidateCompactBinary(Valid, CustomByName, Empty)"), Validate({uint8(ECbFieldType::CustomByName), 2, 1, 'A'}), ECbValidateError::None);
+	TestEqual(TEXT("ValidateCompactBinary(Valid, CustomByName, Empty, NoType)"), Validate({2, 1, 'A'}, ECbFieldType::CustomByName), ECbValidateError::None);
+	TestEqual(TEXT("ValidateCompactBinary(Valid, CustomByName, Payload)"), Validate({uint8(ECbFieldType::CustomByName), 3, 1, 'A', 0}), ECbValidateError::None);
+	TestEqual(TEXT("ValidateCompactBinary(Valid, CustomByName, Payload, NoType)"), Validate({3, 1, 'A', 0}, ECbFieldType::CustomByName), ECbValidateError::None);
+	TestEqual(TEXT("ValidateCompactBinary(OutOfBounds, CustomByName, NoSize)"), Validate({uint8(ECbFieldType::CustomByName)}), ECbValidateError::OutOfBounds);
+	TestEqual(TEXT("ValidateCompactBinary(OutOfBounds, CustomByName, NoSize, NoType)"), Validate({}, ECbFieldType::CustomByName), ECbValidateError::OutOfBounds);
+	TestEqual(TEXT("ValidateCompactBinary(OutOfBounds, CustomByName, NoTypeSize)"), Validate({uint8(ECbFieldType::CustomByName), 0}), ECbValidateError::OutOfBounds);
+	TestEqual(TEXT("ValidateCompactBinary(OutOfBounds, CustomByName, NoTypeSize, NoType)"), Validate({0}, ECbFieldType::CustomByName), ECbValidateError::OutOfBounds);
+	TestEqual(TEXT("ValidateCompactBinary(OutOfBounds, CustomByName, ShortTypeSize)"), Validate({uint8(ECbFieldType::CustomByName), 1, 0x80}), ECbValidateError::OutOfBounds);
+	TestEqual(TEXT("ValidateCompactBinary(OutOfBounds, CustomByName, ShortTypeSize, NoType)"), Validate({1, 0x80}, ECbFieldType::CustomByName), ECbValidateError::OutOfBounds);
+	TestEqual(TEXT("ValidateCompactBinary(OutOfBounds, CustomByName, ShortTypeName)"), Validate({uint8(ECbFieldType::CustomByName), 1, 1}), ECbValidateError::OutOfBounds);
+	TestEqual(TEXT("ValidateCompactBinary(OutOfBounds, CustomByName, ShortTypeName, NoType)"), Validate({1, 1}, ECbFieldType::CustomByName), ECbValidateError::OutOfBounds);
+	TestEqual(TEXT("ValidateCompactBinary(OutOfBounds, CustomByName, NoPayload)"), Validate({uint8(ECbFieldType::CustomByName), 3, 1, 'A'}), ECbValidateError::OutOfBounds);
+	TestEqual(TEXT("ValidateCompactBinary(OutOfBounds, CustomByName, NoPayload, NoType)"), Validate({3, 1, 'A'}, ECbFieldType::CustomByName), ECbValidateError::OutOfBounds);
+	TestEqual(TEXT("ValidateCompactBinary(OutOfBounds, CustomByName, ShortPayload)"), Validate({uint8(ECbFieldType::CustomByName), 4, 1, 'A', 0}), ECbValidateError::OutOfBounds);
+	TestEqual(TEXT("ValidateCompactBinary(OutOfBounds, CustomByName, ShortPayload, NoType)"), Validate({4, 1, 'A', 0}, ECbFieldType::CustomByName), ECbValidateError::OutOfBounds);
+
 	// Test InvalidType
 	TestEqual(TEXT("ValidateCompactBinary(InvalidType, Unknown)"), Validate({uint8(ECbFieldType::TimeSpan) + 1}), ECbValidateError::InvalidType);
 	TestEqual(TEXT("ValidateCompactBinary(InvalidType, Unknown)"), Validate({}, ECbFieldType(uint8(ECbFieldType::TimeSpan) + 1)), ECbValidateError::InvalidType);
@@ -182,6 +218,9 @@ bool FCbValidateTest::RunTest(const FString& Parameters)
 	TestEqual(TEXT("ValidateCompactBinary(InvalidType, ZeroSizeField, Array)"), Validate({uint8(ECbFieldType::UniformArray), 2, 2, NullNoName}), ECbValidateError::InvalidType);
 	TestEqual(TEXT("ValidateCompactBinary(InvalidType, ZeroSizeField, Object)"), Validate({uint8(ECbFieldType::UniformObject), 2, NullNoName, 0}), ECbValidateError::InvalidType);
 
+	TestEqual(TEXT("ValidateCompactBinary(InvalidType, CustomByName, EmptyTypeName)"), Validate({uint8(ECbFieldType::CustomByName), 1, 0}), ECbValidateError::InvalidType);
+	TestEqual(TEXT("ValidateCompactBinary(InvalidType, CustomByName, EmptyTypeName, NoType)"), Validate({1, 0}, ECbFieldType::CustomByName), ECbValidateError::InvalidType);
+
 	// Test DuplicateName
 	TestEqual(TEXT("ValidateCompactBinary(DuplicateName)"), Validate({uint8(ECbFieldType::UniformObject), 7, NullWithName, 1, 'A', 1, 'B', 1, 'A'}), ECbValidateError::DuplicateName);
 	TestEqual(TEXT("ValidateCompactBinary(DuplicateName, CaseSensitive)"), Validate({uint8(ECbFieldType::UniformObject), 7, NullWithName, 1, 'A', 1, 'B', 1, 'a'}), ECbValidateError::None);
@@ -214,10 +253,14 @@ bool FCbValidateTest::RunTest(const FString& Parameters)
 	TestEqual(TEXT("ValidateCompactBinary(InvalidInteger, ArraySize)"), Validate({uint8(ECbFieldType::Array), 0x80, 1, 0}), ECbValidateError::InvalidInteger);
 	TestEqual(TEXT("ValidateCompactBinary(InvalidInteger, ArrayCount)"), Validate({uint8(ECbFieldType::Array), 3, 0xc0, 0, 0}), ECbValidateError::InvalidInteger);
 	TestEqual(TEXT("ValidateCompactBinary(InvalidInteger, ObjectSize)"), Validate({uint8(ECbFieldType::Object), 0xe0, 0, 0, 0}), ECbValidateError::InvalidInteger);
+	TestEqual(TEXT("ValidateCompactBinary(InvalidInteger, CustomById)"), Validate({uint8(ECbFieldType::CustomById), 2, 0x80, 0}), ECbValidateError::InvalidInteger);
+	TestEqual(TEXT("ValidateCompactBinary(InvalidInteger, CustomByName)"), Validate({uint8(ECbFieldType::CustomByName), 3, 0x80, 1, 'A'}), ECbValidateError::InvalidInteger);
 
 	TestEqual(TEXT("ValidateCompactBinary(InvalidInteger, NameSize, Mode)"), ValidateMode({NullWithName, 0x80, 1, 'N'}, ECbValidateMode::All & ~ECbValidateMode::Format), ECbValidateError::None);
 	TestEqual(TEXT("ValidateCompactBinary(InvalidInteger, ArraySize, Mode)"), ValidateMode({uint8(ECbFieldType::Array), 0xc0, 0, 1, 0}, ECbValidateMode::All & ~ECbValidateMode::Format), ECbValidateError::None);
 	TestEqual(TEXT("ValidateCompactBinary(InvalidInteger, ObjectSize, Mode)"), ValidateMode({uint8(ECbFieldType::Object), 0xe0, 0, 0, 0}, ECbValidateMode::All & ~ECbValidateMode::Format), ECbValidateError::None);
+	TestEqual(TEXT("ValidateCompactBinary(InvalidInteger, CustomById, Mode)"), ValidateMode({uint8(ECbFieldType::CustomById), 2, 0x80, 0}, ECbValidateMode::All & ~ECbValidateMode::Format), ECbValidateError::None);
+	TestEqual(TEXT("ValidateCompactBinary(InvalidInteger, CustomByName, Mode)"), ValidateMode({uint8(ECbFieldType::CustomByName), 3, 0x80, 1, 'A'}, ECbValidateMode::All & ~ECbValidateMode::Format), ECbValidateError::None);
 
 	// Test InvalidFloat
 	TestEqual(TEXT("ValidateCompactBinary(InvalidFloat, MaxSignificant+1)"), Validate({uint8(ECbFieldType::Float64), 0x3f, 0xff, 0xff, 0xff, 0xf0, 0x00, 0x00, 0x00}), ECbValidateError::None); // 1.9999999403953552
diff --git a/Engine/Source/Runtime/Core/Private/Tests/Serialization/CompactBinaryWriterTest.cpp b/Engine/Source/Runtime/Core/Private/Tests/Serialization/CompactBinaryWriterTest.cpp
--- a/Engine/Source/Runtime/Core/Private/Tests/Serialization/CompactBinaryWriterTest.cpp
+++ b/Engine/Source/Runtime/Core/Private/Tests/Serialization/CompactBinaryWriterTest.cpp
@@ -214,7 +214,7 @@ bool FCbWriterNullTest::RunTest(const FString& Parameters)
 		Writer.EndObject();
 
 		FCbFieldRefIterator Fields = Writer.Save();
-		TestEqual(TEXT("FCbWriter(Null, Uniform) Validate"), ValidateCompactBinaryRange(Fields.GetBuffer(), ECbValidateMode::All), ECbValidateError::None);
+		TestEqual(TEXT("FCbWriter(Null, Uniform) Validate"), ValidateCompactBinaryRange(Fields.GetRangeBuffer(), ECbValidateMode::All), ECbValidateError::None);
 	}
 
 	// Test Null with Save(Buffer)
@@ -306,7 +306,7 @@ bool FCbWriterStringTest::RunTest(const FString& Parameters)
 		Writer.AddString(FAnsiStringView());
 		Writer.AddString(FWideStringView());
 		FCbFieldRefIterator Fields = Writer.Save();
-		if (TestEqual(TEXT("FCbWriter(String, Empty) Validate"), ValidateCompactBinaryRange(Fields.GetBuffer(), ECbValidateMode::All), ECbValidateError::None))
+		if (TestEqual(TEXT("FCbWriter(String, Empty) Validate"), ValidateCompactBinaryRange(Fields.GetRangeBuffer(), ECbValidateMode::All), ECbValidateError::None))
 		{
 			for (FCbField Field : Fields)
 			{
@@ -323,7 +323,7 @@ bool FCbWriterStringTest::RunTest(const FString& Parameters)
 		Writer.SetName("String"_ASV).AddString("Value"_ASV);
 		Writer.SetName("String"_ASV).AddString(TEXT("Value"_SV));
 		FCbFieldRefIterator Fields = Writer.Save();
-		if (TestEqual(TEXT("FCbWriter(String, Basic) Validate"), ValidateCompactBinaryRange(Fields.GetBuffer(), ECbValidateMode::All), ECbValidateError::None))
+		if (TestEqual(TEXT("FCbWriter(String, Basic) Validate"), ValidateCompactBinaryRange(Fields.GetRangeBuffer(), ECbValidateMode::All), ECbValidateError::None))
 		{
 			for (FCbField Field : Fields)
 			{
@@ -347,7 +347,7 @@ bool FCbWriterStringTest::RunTest(const FString& Parameters)
 		Writer.AddString(Dots);
 		Writer.AddString(FString::ChrN(DotCount, TEXT('.')));
 		FCbFieldRefIterator Fields = Writer.Save();
-		if (TestEqual(TEXT("FCbWriter(String, Long) Validate"), ValidateCompactBinaryRange(Fields.GetBuffer(), ECbValidateMode::All), ECbValidateError::None))
+		if (TestEqual(TEXT("FCbWriter(String, Long) Validate"), ValidateCompactBinaryRange(Fields.GetRangeBuffer(), ECbValidateMode::All), ECbValidateError::None))
 		{
 			for (FCbField Field : Fields)
 			{
@@ -363,7 +363,7 @@ bool FCbWriterStringTest::RunTest(const FString& Parameters)
 		Writer.AddString("\xf0\x9f\x98\x80"_ASV);
 		Writer.AddString(FWideStringView(Value, UE_ARRAY_COUNT(Value)));
 		FCbFieldRefIterator Fields = Writer.Save();
-		if (TestEqual(TEXT("FCbWriter(String, Unicode) Validate"), ValidateCompactBinaryRange(Fields.GetBuffer(), ECbValidateMode::All), ECbValidateError::None))
+		if (TestEqual(TEXT("FCbWriter(String, Unicode) Validate"), ValidateCompactBinaryRange(Fields.GetRangeBuffer(), ECbValidateMode::All), ECbValidateError::None))
 		{
 			for (FCbField Field : Fields)
 			{
@@ -482,7 +482,7 @@ bool FCbWriterFloatTest::RunTest(const FString& Parameters)
 			Writer.AddFloat(Value);
 		}
 		FCbFieldRefIterator Fields = Writer.Save();
-		if (TestEqual(TEXT("FCbWriter(Float, Single) Validate"), ValidateCompactBinaryRange(Fields.GetBuffer(), ECbValidateMode::All), ECbValidateError::None))
+		if (TestEqual(TEXT("FCbWriter(Float, Single) Validate"), ValidateCompactBinaryRange(Fields.GetRangeBuffer(), ECbValidateMode::All), ECbValidateError::None))
 		{
 			const float* CheckValue = Values;
 			for (FCbField Field : Fields)
@@ -502,7 +502,7 @@ bool FCbWriterFloatTest::RunTest(const FString& Parameters)
 			Writer.AddFloat(Value);
 		}
 		FCbFieldRefIterator Fields = Writer.Save();
-		if (TestEqual(TEXT("FCbWriter(Float, Double) Validate"), ValidateCompactBinaryRange(Fields.GetBuffer(), ECbValidateMode::All), ECbValidateError::None))
+		if (TestEqual(TEXT("FCbWriter(Float, Double) Validate"), ValidateCompactBinaryRange(Fields.GetRangeBuffer(), ECbValidateMode::All), ECbValidateError::None))
 		{
 			const double* CheckValue = Values;
 			for (FCbField Field : Fields)
@@ -529,7 +529,7 @@ bool FCbWriterBoolTest::RunTest(const FString& Parameters)
 		Writer.AddBool(false);
 
 		FCbFieldRefIterator Fields = Writer.Save();
-		if (TestEqual(TEXT("FCbWriter(Bool) Validate"), ValidateCompactBinaryRange(Fields.GetBuffer(), ECbValidateMode::All), ECbValidateError::None))
+		if (TestEqual(TEXT("FCbWriter(Bool) Validate"), ValidateCompactBinaryRange(Fields.GetRangeBuffer(), ECbValidateMode::All), ECbValidateError::None))
 		{
 			TestTrue(TEXT("FCbWriter(Bool).AsBool()"), Fields.AsBool());
 			TestFalse(TEXT("FCbWriter(Bool) Error"), Fields.HasError());
@@ -556,7 +556,7 @@ bool FCbWriterBoolTest::RunTest(const FString& Parameters)
 		Writer.EndObject();
 
 		FCbFieldRefIterator Fields = Writer.Save();
-		TestEqual(TEXT("FCbWriter(Bool, Uniform) Validate"), ValidateCompactBinaryRange(Fields.GetBuffer(), ECbValidateMode::All), ECbValidateError::None);
+		TestEqual(TEXT("FCbWriter(Bool, Uniform) Validate"), ValidateCompactBinaryRange(Fields.GetRangeBuffer(), ECbValidateMode::All), ECbValidateError::None);
 	}
 
 	return true;
@@ -579,7 +579,7 @@ bool FCbWriterCompactBinaryAttachmentTest::RunTest(const FString& Parameters)
 	}
 
 	FCbFieldRefIterator Fields = Writer.Save();
-	if (TestEqual(TEXT("FCbWriter(CompactBinaryAttachment) Validate"), ValidateCompactBinaryRange(Fields.GetBuffer(), ECbValidateMode::All), ECbValidateError::None))
+	if (TestEqual(TEXT("FCbWriter(CompactBinaryAttachment) Validate"), ValidateCompactBinaryRange(Fields.GetRangeBuffer(), ECbValidateMode::All), ECbValidateError::None))
 	{
 		const FIoHash* CheckValue = Values;
 		for (FCbField Field : Fields)
@@ -609,7 +609,7 @@ bool FCbWriterBinaryAttachmentTest::RunTest(const FString& Parameters)
 	}
 
 	FCbFieldRefIterator Fields = Writer.Save();
-	if (TestEqual(TEXT("FCbWriter(BinaryAttachment) Validate"), ValidateCompactBinaryRange(Fields.GetBuffer(), ECbValidateMode::All), ECbValidateError::None))
+	if (TestEqual(TEXT("FCbWriter(BinaryAttachment) Validate"), ValidateCompactBinaryRange(Fields.GetRangeBuffer(), ECbValidateMode::All), ECbValidateError::None))
 	{
 		const FIoHash* CheckValue = Values;
 		for (FCbField Field : Fields)
@@ -639,7 +639,7 @@ bool FCbWriterHashTest::RunTest(const FString& Parameters)
 	}
 
 	FCbFieldRefIterator Fields = Writer.Save();
-	if (TestEqual(TEXT("FCbWriter(Hash) Validate"), ValidateCompactBinaryRange(Fields.GetBuffer(), ECbValidateMode::All), ECbValidateError::None))
+	if (TestEqual(TEXT("FCbWriter(Hash) Validate"), ValidateCompactBinaryRange(Fields.GetRangeBuffer(), ECbValidateMode::All), ECbValidateError::None))
 	{
 		const FIoHash* CheckValue = Values;
 		for (FCbField Field : Fields)
@@ -666,7 +666,7 @@ bool FCbWriterUuidTest::RunTest(const FString& Parameters)
 	}
 
 	FCbFieldRefIterator Fields = Writer.Save();
-	if (TestEqual(TEXT("FCbWriter(Uuid) Validate"), ValidateCompactBinaryRange(Fields.GetBuffer(), ECbValidateMode::All), ECbValidateError::None))
+	if (TestEqual(TEXT("FCbWriter(Uuid) Validate"), ValidateCompactBinaryRange(Fields.GetRangeBuffer(), ECbValidateMode::All), ECbValidateError::None))
 	{
 		const FGuid* CheckValue = Values;
 		for (FCbField Field : Fields)
@@ -693,7 +693,7 @@ bool FCbWriterDateTimeTest::RunTest(const FString& Parameters)
 	}
 
 	FCbFieldRefIterator Fields = Writer.Save();
-	if (TestEqual(TEXT("FCbWriter(DateTime) Validate"), ValidateCompactBinaryRange(Fields.GetBuffer(), ECbValidateMode::All), ECbValidateError::None))
+	if (TestEqual(TEXT("FCbWriter(DateTime) Validate"), ValidateCompactBinaryRange(Fields.GetRangeBuffer(), ECbValidateMode::All), ECbValidateError::None))
 	{
 		const FDateTime* CheckValue = Values;
 		for (FCbField Field : Fields)
@@ -720,7 +720,7 @@ bool FCbWriterTimeSpanTest::RunTest(const FString& Parameters)
 	}
 
 	FCbFieldRefIterator Fields = Writer.Save();
-	if (TestEqual(TEXT("FCbWriter(TimeSpan) Validate"), ValidateCompactBinaryRange(Fields.GetBuffer(), ECbValidateMode::All), ECbValidateError::None))
+	if (TestEqual(TEXT("FCbWriter(TimeSpan) Validate"), ValidateCompactBinaryRange(Fields.GetRangeBuffer(), ECbValidateMode::All), ECbValidateError::None))
 	{
 		const FTimespan* CheckValue = Values;
 		for (FCbField Field : Fields)
@@ -735,6 +735,82 @@ bool FCbWriterTimeSpanTest::RunTest(const FString& Parameters)
 
 ///////////////////////////////////////////////////////////////////////////////////////////////////////////////////////
 
+IMPLEMENT_SIMPLE_AUTOMATION_TEST(FCbWriterCustomByIdTest, "System.Core.Serialization.CbWriter.CustomById", CompactBinaryWriterTestFlags)
+bool FCbWriterCustomByIdTest::RunTest(const FString& Parameters)
+{
+	TCbWriter<256> Writer;
+
+	struct FCustomValue
+	{
+		uint64 Type;
+		TArray<uint8, TInlineAllocator<16>> Bytes;
+	};
+	const FCustomValue Values[] =
+	{
+		{ 1, {1, 2, 3} },
+		{ MAX_uint64, {4, 5, 6} },
+	};
+
+	for (const FCustomValue& Value : Values)
+	{
+		Writer.AddCustom(Value.Type, MakeMemoryView(Value.Bytes));
+	}
+
+	FCbFieldRefIterator Fields = Writer.Save();
+	if (TestEqual(TEXT("FCbWriter(CustomById) Validate"), ValidateCompactBinaryRange(Fields.GetRangeBuffer(), ECbValidateMode::All), ECbValidateError::None))
+	{
+		const FCustomValue* CheckValue = Values;
+		for (FCbField Field : Fields)
+		{
+			TestTrue(TEXT("FCbWriter(CustomById).AsCustom()"), Field.AsCustom(CheckValue->Type).EqualBytes(MakeMemoryView(CheckValue->Bytes)));
+			TestFalse(TEXT("FCbWriter(CustomById) Error"), Field.HasError());
+			++CheckValue;
+		}
+	}
+
+	return true;
+}
+
+///////////////////////////////////////////////////////////////////////////////////////////////////////////////////////
+
+IMPLEMENT_SIMPLE_AUTOMATION_TEST(FCbWriterCustomByNameTest, "System.Core.Serialization.CbWriter.CustomByName", CompactBinaryWriterTestFlags)
+bool FCbWriterCustomByNameTest::RunTest(const FString& Parameters)
+{
+	TCbWriter<256> Writer;
+
+	struct FCustomValue
+	{
+		FAnsiStringView Type;
+		TArray<uint8, TInlineAllocator<16>> Bytes;
+	};
+	const FCustomValue Values[] =
+	{
+		{ "Type1", {1, 2, 3} },
+		{ "Type2", {4, 5, 6} },
+	};
+
+	for (const FCustomValue& Value : Values)
+	{
+		Writer.AddCustom(Value.Type, MakeMemoryView(Value.Bytes));
+	}
+
+	FCbFieldRefIterator Fields = Writer.Save();
+	if (TestEqual(TEXT("FCbWriter(CustomByName) Validate"), ValidateCompactBinaryRange(Fields.GetRangeBuffer(), ECbValidateMode::All), ECbValidateError::None))
+	{
+		const FCustomValue* CheckValue = Values;
+		for (FCbField Field : Fields)
+		{
+			TestTrue(TEXT("FCbWriter(CustomByName).AsCustom()"), Field.AsCustom(CheckValue->Type).EqualBytes(MakeMemoryView(CheckValue->Bytes)));
+			TestFalse(TEXT("FCbWriter(CustomByName) Error"), Field.HasError());
+			++CheckValue;
+		}
+	}
+
+	return true;
+}
+
+///////////////////////////////////////////////////////////////////////////////////////////////////////////////////////
+
 IMPLEMENT_SIMPLE_AUTOMATION_TEST(FCbWriterComplexTest, "System.Core.Serialization.CbWriter.Complex", CompactBinaryWriterTestFlags)
 bool FCbWriterComplexTest::RunTest(const FString& Parameters)
 {
diff --git a/Engine/Source/Runtime/Core/Public/Serialization/CompactBinary.h b/Engine/Source/Runtime/Core/Public/Serialization/CompactBinary.h
--- a/Engine/Source/Runtime/Core/Public/Serialization/CompactBinary.h
+++ b/Engine/Source/Runtime/Core/Public/Serialization/CompactBinary.h
@@ -75,8 +75,6 @@ template <typename FuncType> class TFunctionRef;
 /**
  * Field types and flags for FCbField.
  *
- * This is a private type and is only declared here to enable inline use below.
- *
  * DO NOT CHANGE THE VALUE OF ANY MEMBERS OF THIS ENUM!
  * BACKWARD COMPATIBILITY REQUIRES THAT THESE VALUES BE FIXED!
  * SERIALIZATION USES HARD-CODED CONSTANTS BASED ON THESE VALUES!
@@ -84,97 +82,116 @@ template <typename FuncType> class TFunctionRef;
 enum class ECbFieldType : uint8
 {
 	/** A field type that does not occur in a valid object. */
-	None            = 0x00,
+	None                             = 0x00,
 
 	/** Null. Payload is empty. */
-	Null            = 0x01,
+	Null                             = 0x01,
 
 	/**
 	 * Object is an array of fields with unique non-empty names.
 	 *
 	 * Payload is a VarUInt byte count for the encoded fields followed by the fields.
 	 */
-	Object          = 0x02,
+	Object                           = 0x02,
 	/**
 	 * UniformObject is an array of fields with the same field types and unique non-empty names.
 	 *
 	 * Payload is a VarUInt byte count for the encoded fields followed by the fields.
 	 */
-	UniformObject   = 0x03,
+	UniformObject                    = 0x03,
 
 	/**
 	 * Array is an array of fields with no name that may be of different types.
 	 *
 	 * Payload is a VarUInt byte count, followed by a VarUInt item count, followed by the fields.
 	 */
-	Array           = 0x04,
+	Array                            = 0x04,
 	/**
 	 * UniformArray is an array of fields with no name and with the same field type.
 	 *
 	 * Payload is a VarUInt byte count, followed by a VarUInt item count, followed by field type,
 	 * followed by the fields without their field type.
 	 */
-	UniformArray    = 0x05,
+	UniformArray                     = 0x05,
 
 	/** Binary. Payload is a VarUInt byte count followed by the data. */
-	Binary          = 0x06,
+	Binary                           = 0x06,
 
 	/** String in UTF-8. Payload is a VarUInt byte count then an unterminated UTF-8 string. */
-	String          = 0x07,
+	String                           = 0x07,
 
 	/**
 	 * Non-negative integer with the range of a 64-bit unsigned integer.
 	 *
 	 * Payload is the value encoded as a VarUInt.
 	 */
-	IntegerPositive = 0x08,
+	IntegerPositive                  = 0x08,
 	/**
 	 * Negative integer with the range of a 64-bit signed integer.
 	 *
 	 * Payload is the ones' complement of the value encoded as a VarUInt.
 	 */
-	IntegerNegative = 0x09,
+	IntegerNegative                  = 0x09,
 
 	/** Single precision float. Payload is one big endian IEEE 754 binary32 float. */
-	Float32         = 0x0a,
+	Float32                          = 0x0a,
 	/** Double precision float. Payload is one big endian IEEE 754 binary64 float. */
-	Float64         = 0x0b,
+	Float64                          = 0x0b,
 
 	/** Boolean false value. Payload is empty. */
-	BoolFalse       = 0x0c,
+	BoolFalse                        = 0x0c,
 	/** Boolean true value. Payload is empty. */
-	BoolTrue        = 0x0d,
+	BoolTrue                         = 0x0d,
 
 	/**
 	 * CompactBinaryAttachment is a reference to a compact binary attachment stored externally.
 	 *
 	 * Payload is a 160-bit hash digest of the referenced compact binary data.
 	 */
-	CompactBinaryAttachment = 0x0e,
+	CompactBinaryAttachment          = 0x0e,
 	/**
 	 * BinaryAttachment is a reference to a binary attachment stored externally.
 	 *
 	 * Payload is a 160-bit hash digest of the referenced binary data.
 	 */
-	BinaryAttachment = 0x0f,
+	BinaryAttachment                 = 0x0f,
 
 	/** Hash. Payload is a 160-bit hash digest. */
-	Hash            = 0x10,
+	Hash                             = 0x10,
 	/** UUID/GUID. Payload is a 128-bit UUID as defined by RFC 4122. */
-	Uuid            = 0x11,
+	Uuid                             = 0x11,
 
 	/**
 	 * Date and time between 0001-01-01 00:00:00.0000000 and 9999-12-31 23:59:59.9999999.
 	 *
 	 * Payload is a big endian int64 count of 100ns ticks since 0001-01-01 00:00:00.0000000.
 	 */
-	DateTime        = 0x12,
+	DateTime                         = 0x12,
 	/**
 	 * Difference between two date/time values.
 	 *
 	 * Payload is a big endian int64 count of 100ns ticks in the span, and may be negative.
 	 */
-	TimeSpan        = 0x13,
+	TimeSpan                         = 0x13,
+
+	/**
+	 * CustomById identifies the sub-type of its payload by an integer identifier.
+	 *
+	 * Payload is a VarUInt byte count of the sub-type identifier and the sub-type payload, followed
+	 * by a VarUInt of the sub-type identifier then the payload of the sub-type.
+	 */
+	CustomById                       = 0x1e,
+	/**
+	 * CustomByType identifies the sub-type of its payload by a string identifier.
+	 *
+	 * Payload is a VarUInt byte count of the sub-type identifier and the sub-type payload, followed
+	 * by a VarUInt byte count of the unterminated sub-type identifier, then the sub-type identifier
+	 * without termination, then the payload of the sub-type.
+	 */
+	CustomByName                     = 0x1f,
+
+	/** Reserved for future use as a flag. Do not add types in this range. */
+	Reserved                         = 0x20,
 
 	/**
 	 * A transient flag which indicates that the object or array containing this field has stored
@@ -182,10 +199,10 @@ enum class ECbFieldType : uint8
 	 *
 	 * Note: Since the flag must never be serialized, this bit may be repurposed in the future.
 	 */
-	HasFieldType    = 0x40,
+	HasFieldType                     = 0x40,
 
 	/** A persisted flag which indicates that the field has a name stored before the payload. */
-	HasFieldName    = 0x80,
+	HasFieldName                     = 0x80,
 };
 
 ENUM_CLASS_FLAGS(ECbFieldType);
@@ -195,25 +212,25 @@ ENUM_CLASS_FLAGS(ECbFieldType);
 /** Functions that operate on ECbFieldType. */
 class FCbFieldType
 {
-	static constexpr ECbFieldType SerializedTypeMask    = ECbFieldType(0b1011'1111);
-	static constexpr ECbFieldType TypeMask              = ECbFieldType(0b0011'1111);
+	static constexpr ECbFieldType SerializedTypeMask    = ECbFieldType(0b1001'1111);
+	static constexpr ECbFieldType TypeMask              = ECbFieldType(0b0001'1111);
 
-	static constexpr ECbFieldType ObjectMask            = ECbFieldType(0b0011'1110);
+	static constexpr ECbFieldType ObjectMask            = ECbFieldType(0b0001'1110);
 	static constexpr ECbFieldType ObjectBase            = ECbFieldType(0b0000'0010);
 
-	static constexpr ECbFieldType ArrayMask             = ECbFieldType(0b0011'1110);
+	static constexpr ECbFieldType ArrayMask             = ECbFieldType(0b0001'1110);
 	static constexpr ECbFieldType ArrayBase             = ECbFieldType(0b0000'0100);
 
-	static constexpr ECbFieldType IntegerMask           = ECbFieldType(0b0011'1110);
+	static constexpr ECbFieldType IntegerMask           = ECbFieldType(0b0001'1110);
 	static constexpr ECbFieldType IntegerBase           = ECbFieldType(0b0000'1000);
 
-	static constexpr ECbFieldType FloatMask             = ECbFieldType(0b0011'1100);
+	static constexpr ECbFieldType FloatMask             = ECbFieldType(0b0001'1100);
 	static constexpr ECbFieldType FloatBase             = ECbFieldType(0b0000'1000);
 
-	static constexpr ECbFieldType BoolMask              = ECbFieldType(0b0011'1110);
+	static constexpr ECbFieldType BoolMask              = ECbFieldType(0b0001'1110);
 	static constexpr ECbFieldType BoolBase              = ECbFieldType(0b0000'1100);
 
-	static constexpr ECbFieldType AttachmentMask        = ECbFieldType(0b0011'1110);
+	static constexpr ECbFieldType AttachmentMask        = ECbFieldType(0b0001'1110);
 	static constexpr ECbFieldType AttachmentBase        = ECbFieldType(0b0000'1110);
 
 	static void StaticAssertTypeConstants();
@@ -243,7 +260,7 @@ public:
 
 	static constexpr inline bool IsCompactBinaryAttachment(ECbFieldType Type) { return GetType(Type) == ECbFieldType::CompactBinaryAttachment; }
 	static constexpr inline bool IsBinaryAttachment(ECbFieldType Type)        { return GetType(Type) == ECbFieldType::BinaryAttachment; }
-	static constexpr inline bool IsAttachment(ECbFieldType Type) { return (Type & AttachmentMask) == AttachmentBase; }
+	static constexpr inline bool IsAttachment(ECbFieldType Type)              { return (Type & AttachmentMask) == AttachmentBase; }
 
 	static constexpr inline bool IsHash(ECbFieldType Type)       { return GetType(Type) == ECbFieldType::Hash; }
 	static constexpr inline bool IsUuid(ECbFieldType Type)       { return GetType(Type) == ECbFieldType::Uuid; }
@@ -251,6 +268,9 @@ public:
 	static constexpr inline bool IsDateTime(ECbFieldType Type)   { return GetType(Type) == ECbFieldType::DateTime; }
 	static constexpr inline bool IsTimeSpan(ECbFieldType Type)   { return GetType(Type) == ECbFieldType::TimeSpan; }
 
+	static constexpr inline bool IsCustomById(ECbFieldType Type)   { return GetType(Type) == ECbFieldType::CustomById; }
+	static constexpr inline bool IsCustomByName(ECbFieldType Type) { return GetType(Type) == ECbFieldType::CustomByName; }
+
 	/** Whether the type is or may contain fields of any attachment type. */
 	static constexpr inline bool MayContainAttachments(ECbFieldType Type)
 	{
@@ -419,6 +439,26 @@ enum class ECbFieldError : uint8
 
 ///////////////////////////////////////////////////////////////////////////////////////////////////////////////////////
 
+/** A custom compact binary field type with an integer identifier. */
+struct FCbCustomById
+{
+	/** An identifier for the sub-type of the field. */
+	uint64 Id = 0;
+	/** A view of the value. Lifetime is tied to the field that the value is associated with. */
+	FMemoryView Data;
+};
+
+/** A custom compact binary field type with a string identifier. */
+struct FCbCustomByName
+{
+	/** An identifier for the sub-type of the field. Lifetime is tied to the field that the name is associated with. */
+	FAnsiStringView Name;
+	/** A view of the value. Lifetime is tied to the field that the value is associated with. */
+	FMemoryView Data;
+};
+
+///////////////////////////////////////////////////////////////////////////////////////////////////////////////////////
+
 /**
  * An atom of data in the compact binary format.
  *
@@ -527,6 +567,16 @@ public:
 	/** Access the field as a timespan. Returns the provided default on error. */
 	CORE_API FTimespan AsTimeSpan(FTimespan Default);
 
+	/** Access the field as a custom sub-type with an integer identifier. Returns the provided default on error. */
+	CORE_API FCbCustomById AsCustomById(FCbCustomById Default = FCbCustomById());
+	/** Access the field as a custom sub-type with a string identifier. Returns the provided default on error. */
+	CORE_API FCbCustomByName AsCustomByName(FCbCustomByName Default = FCbCustomByName());
+
+	/** Access the field as a custom sub-type with an integer identifier. Returns the provided default on error. */
+	CORE_API FMemoryView AsCustom(uint64 Id, FMemoryView Default = FMemoryView());
+	/** Access the field as a custom sub-type with a string identifier. Returns the provided default on error. */
+	CORE_API FMemoryView AsCustom(FAnsiStringView Name, FMemoryView Default = FMemoryView());
+
 	/** True if the field has a name. */
 	constexpr inline bool HasName() const           { return FCbFieldType::HasFieldName(Type); }
 
@@ -554,6 +604,9 @@ public:
 	constexpr inline bool IsDateTime() const        { return FCbFieldType::IsDateTime(Type); }
 	constexpr inline bool IsTimeSpan() const        { return FCbFieldType::IsTimeSpan(Type); }
 
+	constexpr inline bool IsCustomById() const      { return FCbFieldType::IsCustomById(Type); }
+	constexpr inline bool IsCustomByName() const    { return FCbFieldType::IsCustomByName(Type); }
+
 	/** Whether the field has a value. */
 	constexpr inline explicit operator bool() const { return HasValue(); }
 
@@ -758,7 +811,7 @@ public:
 	inline FCbField AsField() const { return static_cast<const FCbField&>(*this); }
 
 	/** Construct an array from an array field. No type check is performed! */
-	static inline FCbArray FromField(const FCbField& Field) { return FCbArray(Field); }
+	static inline FCbArray FromFieldNoCheck(const FCbField& Field) { return FCbArray(Field); }
 
 	/** Returns the size of the array in bytes if serialized by itself with no name. */
 	CORE_API uint64 GetSize() const;
@@ -806,7 +859,7 @@ private:
 	friend inline FCbFieldIterator begin(const FCbArray& Array) { return Array.CreateIterator(); }
 	friend inline FCbFieldIterator end(const FCbArray&) { return FCbFieldIterator(); }
 
-	/** Construct an array from an array field. No type check is performed! Use via FromField. */
+	/** Construct an array from an array field. No type check is performed! Use via FromFieldNoCheck. */
 	inline explicit FCbArray(const FCbField& Field) : FCbField(Field) {}
 };
 
@@ -856,7 +909,7 @@ public:
 	inline FCbField AsField() const { return static_cast<const FCbField&>(*this); }
 
 	/** Construct an object from an object field. No type check is performed! */
-	static inline FCbObject FromField(const FCbField& Field) { return FCbObject(Field); }
+	static inline FCbObject FromFieldNoCheck(const FCbField& Field) { return FCbObject(Field); }
 
 	/** Returns the size of the object in bytes if serialized by itself with no name. */
 	CORE_API uint64 GetSize() const;
@@ -904,7 +957,7 @@ private:
 	friend inline FCbFieldIterator begin(const FCbObject& Object) { return Object.CreateIterator(); }
 	friend inline FCbFieldIterator end(const FCbObject&) { return FCbFieldIterator(); }
 
-	/** Construct an object from an object field. No type check is performed! Use via FromField. */
+	/** Construct an object from an object field. No type check is performed! Use via FromFieldNoCheck. */
 	inline explicit FCbObject(const FCbField& Field) : FCbField(Field) {}
 };
 
diff --git a/Engine/Source/Runtime/Core/Public/Serialization/CompactBinaryWriter.h b/Engine/Source/Runtime/Core/Public/Serialization/CompactBinaryWriter.h
--- a/Engine/Source/Runtime/Core/Public/Serialization/CompactBinaryWriter.h
+++ b/Engine/Source/Runtime/Core/Public/Serialization/CompactBinaryWriter.h
@@ -220,6 +220,14 @@ public:
 	CORE_API void AddTimeSpan(FAnsiStringView Name, FTimespan Value);
 	CORE_API void AddTimeSpan(FTimespan Value);
 
+	/** Write a custom field with an integer sub-type identifier. */
+	inline void AddCustom(FAnsiStringView FieldName, uint64 TypeId, FMemoryView Value) { SetName(FieldName); AddCustom(TypeId, Value); }
+	CORE_API void AddCustom(uint64 TypeId, FMemoryView Value);
+
+	/** Write a custom field with a string sub-type identifier. */
+	inline void AddCustom(FAnsiStringView FieldName, FAnsiStringView TypeName, FMemoryView Value) { SetName(FieldName); AddCustom(TypeName, Value); }
+	CORE_API void AddCustom(FAnsiStringView TypeName, FMemoryView Value);
+
 	/** Private flags that are public to work with ENUM_CLASS_FLAGS. */
 	enum class EStateFlags : uint8;
 
