U5 - Change order to not trigger oldvalue != newvalue ensure

diff --git a/Engine/Source/Runtime/Core/Public/Async/ParallelFor.h b/Engine/Source/Runtime/Core/Public/Async/ParallelFor.h
--- a/Engine/Source/Runtime/Core/Public/Async/ParallelFor.h
+++ b/Engine/Source/Runtime/Core/Public/Async/ParallelFor.h
@@ -326,7 +326,7 @@ namespace ParallelForImpl
 		// this thread can help too and this is important to prevent deadlock on recursion 
 		if (!Data->Process(0, Data, DesiredThread, true))
 		{
-			if (IsInRenderingThread() && (Flags & EParallelForFlags::PumpRenderingThread) != EParallelForFlags::None)
+			if ((Flags & EParallelForFlags::PumpRenderingThread) != EParallelForFlags::None && IsInRenderingThread())
 			{
 				while (!Data->Event->Wait(1))
 				{
