U5 - ParallelFor should always be slighty higher priority than the lowest background tasks possible

diff --git a/Engine/Source/Runtime/Core/Public/Async/ParallelFor.h b/Engine/Source/Runtime/Core/Public/Async/ParallelFor.h
--- a/Engine/Source/Runtime/Core/Public/Async/ParallelFor.h
+++ b/Engine/Source/Runtime/Core/Public/Async/ParallelFor.h
@@ -227,7 +227,7 @@ namespace ParallelForImpl
 		}
 
 		// Either the request comes from a totally unknown thread, or we've specifically been asked to be background
-		return ENamedThreads::AnyBackgroundThreadNormalTask;
+		return ENamedThreads::AnyBackgroundHiPriTask;
 	}
 
 	template<typename FunctionType>
