U5 - remove forcenoinline as otherwise the compiler will generate the call to the empty function in optimized builds.

diff --git a/Engine/Source/Runtime/Core/Public/Templates/Function.h b/Engine/Source/Runtime/Core/Public/Templates/Function.h
--- a/Engine/Source/Runtime/Core/Public/Templates/Function.h
+++ b/Engine/Source/Runtime/Core/Public/Templates/Function.h
@@ -665,7 +665,7 @@ namespace UE4Function_Private
 		TFunctionRefBase& operator=(const TFunctionRefBase&) = delete;
 
 		// Move all of the assert code out of line
-		FORCENOINLINE void CheckCallable() const
+		void CheckCallable() const
 		{
 			checkf(Callable, TEXT("Attempting to call an unbound TFunction!"));
 		}
