U5 - renamed  profiling scopes

diff --git a/Engine/Source/Runtime/Core/Public/Async/LocalWorkQueue.h b/Engine/Source/Runtime/Core/Public/Async/LocalWorkQueue.h
--- a/Engine/Source/Runtime/Core/Public/Async/LocalWorkQueue.h
+++ b/Engine/Source/Runtime/Core/Public/Async/LocalWorkQueue.h
@@ -66,7 +66,7 @@ public:
 			FTaskHandle TaskHandle = MakeShared<LowLevelTasks::FTask, ESPMode::ThreadSafe>();
 
 			TFunctionRef<void(TaskType*)>& LocalDoWork = *DoWork;
-			TaskHandle->Init(TEXT("TLocalWorkQueue::AddWork"), Priority, [LocalDoWork, InternalData = InternalData, TaskHandle]()
+			TaskHandle->Init(TEXT("TLocalWorkQueue::AddWorkers"), Priority, [LocalDoWork, InternalData = InternalData, TaskHandle]()
 			{
 				TRACE_CPUPROFILER_EVENT_SCOPE(TLocalWorkQueue::AddWork);
 				InternalData->ActiveWorkers.fetch_add(1, std::memory_order_acquire);
@@ -91,7 +91,7 @@ public:
 		DoWork = &InDoWork;
 		LowLevelTasks::BusyWaitUntil([&InDoWork, InternalData = InternalData]()
 		{
-			TRACE_CPUPROFILER_EVENT_SCOPE(TLocalWorkQueue::ProcessLocalWorkSynchronous);
+			TRACE_CPUPROFILER_EVENT_SCOPE(TLocalWorkQueue::Run);
 			bool Completed = false;
 			while(true)
 			{
