U5 - another rename

diff --git a/Engine/Source/Runtime/Core/Public/Async/LocalWorkQueue.h b/Engine/Source/Runtime/Core/Public/Async/LocalWorkQueue.h
--- a/Engine/Source/Runtime/Core/Public/Async/LocalWorkQueue.h
+++ b/Engine/Source/Runtime/Core/Public/Async/LocalWorkQueue.h
@@ -68,7 +68,7 @@ public:
 			TFunctionRef<void(TaskType*)>& LocalDoWork = *DoWork;
 			TaskHandle->Init(TEXT("TLocalWorkQueue::AddWorkers"), Priority, [LocalDoWork, InternalData = InternalData, TaskHandle]()
 			{
-				TRACE_CPUPROFILER_EVENT_SCOPE(TLocalWorkQueue::AddWork);
+				TRACE_CPUPROFILER_EVENT_SCOPE(TLocalWorkQueue::AddWorkers);
 				InternalData->ActiveWorkers.fetch_add(1, std::memory_order_acquire);
 				while(true)
 				{
