Fix missing bias in RGBE conversion

diff --git a/Engine/Source/Runtime/Core/Private/Math/Color.cpp b/Engine/Source/Runtime/Core/Private/Math/Color.cpp
--- a/Engine/Source/Runtime/Core/Private/Math/Color.cpp
+++ b/Engine/Source/Runtime/Core/Private/Math/Color.cpp
@@ -234,7 +234,12 @@ FLinearColor FColor::FromRGBE() const
 	{
 		// the extra 8 here does the /256
 		const float Scale = ldexpf(1.f, A - (128 + 8));
-		return FLinearColor( R * Scale, G * Scale, B * Scale, 1.0f );
+		// bias by 0.5 so [0,255] input shifts to [0.5,255.5] making it centered after /256
+		return FLinearColor( 
+			(R + 0.5f) * Scale, 
+			(G + 0.5f) * Scale, 
+			(B + 0.5f) * Scale, 
+			1.0f );
 	}
 }
 
