U5 - Removing MaxConcurreny from the LowLevelThreadPool and adding a BackgroundHigh Priority to the Backend, also rearranging the priories of the TaskGraph and Taskpool

diff --git a/Engine/Source/Runtime/Core/Private/Async/TaskGraph.cpp b/Engine/Source/Runtime/Core/Private/Async/TaskGraph.cpp
--- a/Engine/Source/Runtime/Core/Private/Async/TaskGraph.cpp
+++ b/Engine/Source/Runtime/Core/Private/Async/TaskGraph.cpp
@@ -1804,10 +1804,15 @@ public:
 
 		if (FTaskGraphInterface::IsMultithread())
 		{
+			if (NumWorkerThreads <= 3)
+			{
+				GNumForegroundWorkers = 1;
+			}
+
 			int32 NumBackgroundWorkers = FMath::Max(1, NumWorkerThreads - FMath::Min<int>(GNumForegroundWorkers, NumWorkerThreads));
-			int32 NumWorkers =  FMath::Max(1, NumWorkerThreads - NumBackgroundWorkers);
+			int32 NumForegroundWorkers =  FMath::Max(1, NumWorkerThreads - NumBackgroundWorkers);
 
-			LowLevelTasks::FScheduler::Get().StartWorkers(NumWorkers, NumBackgroundWorkers, EThreadPriority::TPri_Normal, EThreadPriority::TPri_BelowNormal, FForkProcessHelper::IsForkedMultithreadInstance());
+			LowLevelTasks::FScheduler::Get().StartWorkers(NumForegroundWorkers, NumBackgroundWorkers, EThreadPriority::TPri_Normal, EThreadPriority::TPri_BelowNormal, FForkProcessHelper::IsForkedMultithreadInstance());
 			NumNamedThreads = ENamedThreads::ActualRenderingThread + 1;
 			ENamedThreads::bHasBackgroundThreads = 1;
 			ENamedThreads::bHasHighPriorityThreads = 1;
@@ -1888,12 +1893,12 @@ private:
 		{
 			uint32 ThreadPriority = GetThreadPriorityIndex(InThreadToExecuteOn);
 			check(ThreadPriority < uint32(ENamedThreads::NumThreadPriorities));
-			LowLevelTasks::ETaskPriority Conversion[int(ENamedThreads::NumThreadPriorities)] = { LowLevelTasks::ETaskPriority::Normal, LowLevelTasks::ETaskPriority::High, LowLevelTasks::ETaskPriority::BackgroundLow };
+			LowLevelTasks::ETaskPriority Conversion[int(ENamedThreads::NumThreadPriorities)] = { LowLevelTasks::ETaskPriority::Normal, LowLevelTasks::ETaskPriority::High, LowLevelTasks::ETaskPriority::BackgroundNormal };
 			LowLevelTasks::ETaskPriority Priority = Conversion[ThreadPriority];
 
-			if (Priority == LowLevelTasks::ETaskPriority::BackgroundLow && GetTaskPriority(InThreadToExecuteOn))
+			if (Priority == LowLevelTasks::ETaskPriority::BackgroundNormal && GetTaskPriority(InThreadToExecuteOn))
 			{
-				Priority = LowLevelTasks::ETaskPriority::BackgroundNormal;
+				Priority = LowLevelTasks::ETaskPriority::BackgroundHigh;
 			}
 
 			Task->TaskHandle.Init(TEXT("TaskGraphTask"), Priority, [Task, InThreadToExecuteOn, Deleter(LowLevelTasks::TDeleter<FBaseGraphTask, &FBaseGraphTask::DeleteTask>(Task))]()
@@ -2154,8 +2159,8 @@ private:
 			const LowLevelTasks::FTask* Task = LowLevelTasks::FScheduler::GetActiveTask();
 			if (Task != nullptr)
 			{
-				ENamedThreads::Type ThreadConversion[int(LowLevelTasks::ETaskPriority::Count)] = { ENamedThreads::HighThreadPriority, ENamedThreads::NormalThreadPriority, ENamedThreads::BackgroundThreadPriority, ENamedThreads::BackgroundThreadPriority };
-				ENamedThreads::Type TaskConversion[int(LowLevelTasks::ETaskPriority::Count)] = { ENamedThreads::NormalTaskPriority, ENamedThreads::NormalTaskPriority, ENamedThreads::HighTaskPriority, ENamedThreads::NormalTaskPriority };
+				ENamedThreads::Type ThreadConversion[int(LowLevelTasks::ETaskPriority::Count)] = { ENamedThreads::HighThreadPriority, ENamedThreads::NormalThreadPriority, ENamedThreads::BackgroundThreadPriority, ENamedThreads::BackgroundThreadPriority, ENamedThreads::BackgroundThreadPriority };
+				ENamedThreads::Type TaskConversion[int(LowLevelTasks::ETaskPriority::Count)] = { ENamedThreads::NormalTaskPriority, ENamedThreads::NormalTaskPriority, ENamedThreads::HighTaskPriority, ENamedThreads::NormalTaskPriority, ENamedThreads::NormalTaskPriority };
 				CurrentThreadIfKnown = (ENamedThreads::Type)(ENamedThreads::AnyThread | ThreadConversion[int(Task->GetPriority())] | TaskConversion[int(Task->GetPriority())]);
 			}
 		}
diff --git a/Engine/Source/Runtime/Core/Public/Async/Fundamental/LocalQueue.h b/Engine/Source/Runtime/Core/Public/Async/Fundamental/LocalQueue.h
--- a/Engine/Source/Runtime/Core/Public/Async/Fundamental/LocalQueue.h
+++ b/Engine/Source/Runtime/Core/Public/Async/Fundamental/LocalQueue.h
@@ -287,7 +287,7 @@ public:
 
 		inline FTask* DequeueLocal(bool GetBackGroundTasks)
 		{
-			int32 MaxPriority = GetBackGroundTasks ? int32(ETaskPriority::Count) : int32(ETaskPriority::BackgroundNormal);
+			int32 MaxPriority = GetBackGroundTasks ? int32(ETaskPriority::Count) : int32(ETaskPriority::ForegroundCount);
 			for (int32 PriorityIndex = 0; PriorityIndex < MaxPriority; PriorityIndex++)
 			{
 				FTask* Item;
@@ -303,7 +303,7 @@ public:
 		{
 			if ((Registry->NumActiveWorkers[GetBackGroundTasks].load(std::memory_order_relaxed) >= (2 * Registry->NumWorkersLookingForWork[GetBackGroundTasks].load(std::memory_order_relaxed) - 1)) || ((Random.GetUnsignedInt() % 4) == 0))
 			{
-				int32 MaxPriority = GetBackGroundTasks ? int32(ETaskPriority::Count) : int32(ETaskPriority::BackgroundNormal);
+				int32 MaxPriority = GetBackGroundTasks ? int32(ETaskPriority::Count) : int32(ETaskPriority::ForegroundCount);
 				for (int32 PriorityIndex = 0; PriorityIndex < MaxPriority; PriorityIndex++)
 				{
 					FTask* Item = Registry->OverflowQueues[PriorityIndex].dequeue(DequeueHazards[PriorityIndex]);
@@ -404,7 +404,7 @@ private:
 	{
 		FLocalQueueCollection* Queues = Hazard.Get();
 		uint32 NumQueues = Queues->LocalQueues.Num();
-		uint32 MaxPriority = GetBackGroundTasks ? int32(ETaskPriority::Count) : int32(ETaskPriority::BackgroundNormal);
+		uint32 MaxPriority = GetBackGroundTasks ? int32(ETaskPriority::Count) : int32(ETaskPriority::ForegroundCount);
 		CachedRandomIndex = CachedRandomIndex % NumQueues;
 
 		for(uint32 i = 0; i < NumQueues; i++)
diff --git a/Engine/Source/Runtime/Core/Public/Async/Fundamental/Task.h b/Engine/Source/Runtime/Core/Public/Async/Fundamental/Task.h
--- a/Engine/Source/Runtime/Core/Public/Async/Fundamental/Task.h
+++ b/Engine/Source/Runtime/Core/Public/Async/Fundamental/Task.h
@@ -15,6 +15,8 @@ namespace LowLevelTasks
 		High,
 		Normal,
 		Default = Normal,
+		BackgroundHigh,
+		ForegroundCount = BackgroundHigh,
 		BackgroundNormal,
 		BackgroundLow,
 		Count
@@ -74,15 +76,15 @@ namespace LowLevelTasks
 #if PLATFORM_32BITS
 			uintptr_t DebugName : 32;
 #else
-			uintptr_t DebugName : 58;
+			uintptr_t DebugName : 57;
 #endif
-			uintptr_t Priority	: 2;
+			uintptr_t Priority	: 3;
 			uintptr_t State		: 4;
 			
 		public:
 			FPackedData() : FPackedData(nullptr,  ETaskPriority::Default, ETaskState::Completed)
 			{
-				static_assert(uintptr_t(ETaskPriority::Count) <= (1ull << 2), "Not enough bits to store ETaskPriority");
+				static_assert(uintptr_t(ETaskPriority::Count) <= (1ull << 3), "Not enough bits to store ETaskPriority");
 				static_assert(uintptr_t(ETaskState::Count) <= (1ull << 4), "Not enough bits to store ETaskState");
 			}
 
@@ -91,8 +93,8 @@ namespace LowLevelTasks
 				, Priority((uintptr_t)InPriority)
 				, State((uintptr_t)InState)
 			{
-				checkSlow(reinterpret_cast<uintptr_t>(InDebugName) < (1ull << 58));
-				checkSlow((uintptr_t)InPriority < (1ull << 2));
+				checkSlow(reinterpret_cast<uintptr_t>(InDebugName) < (1ull << 57));
+				checkSlow((uintptr_t)InPriority < (1ull << 3));
 				checkSlow((uintptr_t)InState < (1ull << 4));
 #if PLATFORM_32BITS
 				static_assert(sizeof(FPackedData) == 2 * sizeof(uintptr_t), "Packed data needs to be 8bytes in size");
@@ -318,6 +320,6 @@ namespace LowLevelTasks
 
 	inline bool FTask::IsBackgroundTask() const
 	{
-		return PackedData.load(std::memory_order_relaxed).GetPriority() >= ETaskPriority::BackgroundNormal;
+		return PackedData.load(std::memory_order_relaxed).GetPriority() >= ETaskPriority::ForegroundCount;
 	}
 }
\ No newline at end of file
diff --git a/Engine/Source/Runtime/Core/Public/Async/LocalWorkQueue.h b/Engine/Source/Runtime/Core/Public/Async/LocalWorkQueue.h
--- a/Engine/Source/Runtime/Core/Public/Async/LocalWorkQueue.h
+++ b/Engine/Source/Runtime/Core/Public/Async/LocalWorkQueue.h
@@ -36,6 +36,10 @@ public:
 				{
 					Priority = LowLevelTasks::ETaskPriority::BackgroundNormal;
 				}
+				else if (Priority == LowLevelTasks::ETaskPriority::BackgroundNormal)
+				{
+					Priority = LowLevelTasks::ETaskPriority::BackgroundHigh;
+				}
 			}
 			else
 			{
diff --git a/Engine/Source/Runtime/Core/Public/Async/ParallelFor.h b/Engine/Source/Runtime/Core/Public/Async/ParallelFor.h
--- a/Engine/Source/Runtime/Core/Public/Async/ParallelFor.h
+++ b/Engine/Source/Runtime/Core/Public/Async/ParallelFor.h
@@ -227,7 +227,7 @@ namespace ParallelForImpl
 		}
 
 		// Either the request comes from a totally unknown thread, or we've specifically been asked to be background
-		return ENamedThreads::AnyBackgroundHiPriTask;
+		return ENamedThreads::AnyBackgroundThreadNormalTask;
 	}
 
 	template<typename FunctionType>
diff --git a/Engine/Source/Runtime/Core/Public/Misc/QueuedThreadPoolWrapper.h b/Engine/Source/Runtime/Core/Public/Misc/QueuedThreadPoolWrapper.h
--- a/Engine/Source/Runtime/Core/Public/Misc/QueuedThreadPoolWrapper.h
+++ b/Engine/Source/Runtime/Core/Public/Misc/QueuedThreadPoolWrapper.h
@@ -229,18 +229,16 @@ class CORE_API FQueuedLowLevelThreadPool : public FQueuedThreadPool
 
 		virtual bool Retract()
 		{
-			return Task.WasCanceled() || Task.TryCancel();
+			return Task.TryCancel();
 		}
 	};
 public:
 	/**
-    * InMaxConcurrency           Maximum number of concurrent tasks allowed, -1 will limit concurrency to number of threads available in the underlying thread pool.
 	* InPriorityMapper           Thread-safe function used to map any priority from this Queue to the priority that should be used when scheduling the task on the underlying thread pool.
 	**/
-	FQueuedLowLevelThreadPool(uint32 InMaxConcurrency = ~0, TFunction<EQueuedWorkPriority(EQueuedWorkPriority)> InPriorityMapper = [](EQueuedWorkPriority InPriority) { return InPriority; }, LowLevelTasks::FScheduler* InScheduler = &LowLevelTasks::FScheduler::Get()) 
+	FQueuedLowLevelThreadPool(TFunction<EQueuedWorkPriority(EQueuedWorkPriority)> InPriorityMapper = [](EQueuedWorkPriority InPriority) { return InPriority; }, LowLevelTasks::FScheduler* InScheduler = &LowLevelTasks::FScheduler::Get()) 
 		: Scheduler(InScheduler), PriorityMapper(InPriorityMapper)
 	{
-		MaxConcurrency = InMaxConcurrency;
 	}
 
 	~FQueuedLowLevelThreadPool()
@@ -266,14 +264,6 @@ public:
 		bIsPaused = true;
 	}
 
-	/**
-	*  Dynamically adjust the maximum number of concurrent tasks, -1 for unlimited.
-	*/
-	void SetMaxConcurrency(uint32 InMaxConcurrency = ~0)
-	{
-		MaxConcurrency = InMaxConcurrency;
-	}
-
 	/**
 	*  Resume a specified amount of queued work, or -1 to unpause.
 	*/
@@ -286,7 +276,7 @@ public:
 			{
 				break;
 			}
-			TaskCount++;
+			TaskCount.fetch_add(1, std::memory_order_acquire);
 			verifySlow(Scheduler->TryLaunch(QueuedWork->Task, LowLevelTasks::EQueuePreference::GlobalQueuePreference));
 		}
 
@@ -296,19 +286,19 @@ public:
 		}
 
 		bool bWakeUpWorker = true;
-		ScheduleTasks(bWakeUpWorker, TaskCount.load(std::memory_order_relaxed));
+		ScheduleTasks(bWakeUpWorker);
 	}
 
 private:
-	void ScheduleTasks(bool &bWakeUpWorker, uint32 LocalTaskCount)
+	void ScheduleTasks(bool &bWakeUpWorker)
 	{
-		while (!bIsPaused && LocalTaskCount < MaxConcurrency)
+		while (!bIsPaused)
 		{
 			FQueuedWorkInternalData* QueuedWork = Dequeue();
 			if (QueuedWork)
 			{
 				verifySlow(Scheduler->TryLaunch(QueuedWork->Task, bWakeUpWorker ? LowLevelTasks::EQueuePreference::GlobalQueuePreference : LowLevelTasks::EQueuePreference::LocalQueuePreference, bWakeUpWorker));
-				LocalTaskCount = TaskCount.fetch_add(1, std::memory_order_acquire) + 1;
+				TaskCount.fetch_add(1, std::memory_order_acquire);
 				bWakeUpWorker = true;
 			}
 			else
@@ -325,7 +315,9 @@ private:
 		FQueuedWorkInternalData* QueuedWorkInternalData = new FQueuedWorkInternalData();
 		InQueuedWork->InternalData = QueuedWorkInternalData;
 		EQueuedWorkPriority Priority = PriorityMapper(InPriority);
-		LowLevelTasks::ETaskPriority TaskPriority = Priority > EQueuedWorkPriority::Normal ? LowLevelTasks::ETaskPriority::BackgroundLow : LowLevelTasks::ETaskPriority::BackgroundNormal;
+
+		LowLevelTasks::ETaskPriority TaskPriorityMapper[int(EQueuedWorkPriority::Count)] = { LowLevelTasks::ETaskPriority::High, LowLevelTasks::ETaskPriority::BackgroundHigh, LowLevelTasks::ETaskPriority::BackgroundNormal, LowLevelTasks::ETaskPriority::BackgroundLow, LowLevelTasks::ETaskPriority::BackgroundLow };
+		LowLevelTasks::ETaskPriority TaskPriority = TaskPriorityMapper[int(Priority)];
 
 		QueuedWorkInternalData->Task.Init(TEXT("FQueuedLowLevelThreadPoolTask"), TaskPriority, [InQueuedWork]
 		{
@@ -334,14 +326,20 @@ private:
 		},
 		[this, InternalData = InQueuedWork->InternalData]()
 		{
-			uint32 LocalTaskCount = TaskCount.fetch_sub(1, std::memory_order_release) + 1;
+			TaskCount.fetch_sub(1, std::memory_order_release);
 			bool bWakeUpWorker = false;
-			ScheduleTasks(bWakeUpWorker, LocalTaskCount);
+			ScheduleTasks(bWakeUpWorker);
 		});
 
-		Enqueue(Priority, QueuedWorkInternalData);
-		bool bWakeUpWorker = true;
-		ScheduleTasks(bWakeUpWorker, TaskCount.load(std::memory_order_relaxed));	
+		if (!bIsPaused)
+		{
+			TaskCount.fetch_add(1, std::memory_order_acquire);
+			verifySlow(Scheduler->TryLaunch(QueuedWorkInternalData->Task, LowLevelTasks::EQueuePreference::GlobalQueuePreference));
+		}
+		else
+		{
+			Enqueue(Priority, QueuedWorkInternalData);
+		}
 	}
 
 	bool RetractQueuedWork(IQueuedWork* InQueuedWork) override
@@ -354,7 +352,7 @@ private:
 		}
 
 		bool bWakeUpWorker = true;
-		ScheduleTasks(bWakeUpWorker, TaskCount.load(std::memory_order_relaxed));	
+		ScheduleTasks(bWakeUpWorker);	
 		return bCancelled;
 	}
 
@@ -366,7 +364,6 @@ private:
 protected:
 	bool Create(uint32 InNumQueuedThreads, uint32 InStackSize, EThreadPriority InThreadPriority, const TCHAR* InName) override
 	{
-		MaxConcurrency = InNumQueuedThreads;
 		return true;
 	}
 
@@ -424,7 +421,6 @@ private:
 	LowLevelTasks::FScheduler* Scheduler = nullptr;
 	TFunction<EQueuedWorkPriority(EQueuedWorkPriority)> PriorityMapper;
 
-	std::atomic_uint MaxConcurrency{~0u};
 	std::atomic_uint TaskCount{0};
 	std::atomic_bool bIsExiting{false};
 	std::atomic_bool bIsPaused{false};
diff --git a/Engine/Source/Runtime/Engine/Private/AsyncCompilationHelpers.cpp b/Engine/Source/Runtime/Engine/Private/AsyncCompilationHelpers.cpp
--- a/Engine/Source/Runtime/Engine/Private/AsyncCompilationHelpers.cpp
+++ b/Engine/Source/Runtime/Engine/Private/AsyncCompilationHelpers.cpp
@@ -136,7 +136,7 @@ namespace AsyncCompilationHelpers
 		for (int32 Index = 0; Index < Num; ++Index)
 		{
 			PendingTasks[Index].Job = &Getter(Index);
-			GThreadPool->AddQueuedWork(&PendingTasks[Index], EQueuedWorkPriority::High);
+			GThreadPool->AddQueuedWork(&PendingTasks[Index], EQueuedWorkPriority::Highest);
 		}
 
 		auto FormatProgress =
diff --git a/Engine/Source/Runtime/Launch/Private/LaunchEngineLoop.cpp b/Engine/Source/Runtime/Launch/Private/LaunchEngineLoop.cpp
--- a/Engine/Source/Runtime/Launch/Private/LaunchEngineLoop.cpp
+++ b/Engine/Source/Runtime/Launch/Private/LaunchEngineLoop.cpp
@@ -2185,31 +2185,18 @@ int32 FEngineLoop::PreInitPreStartupScreen(const TCHAR* CmdLine)
 
 #if WITH_EDITOR
 		{
-			int32 NumThreadsInLargeThreadPool = 1;
-			int32 NumThreadsInThreadPool = 1;
+			int32 NumThreadsInLargeThreadPool = FMath::Max(FPlatformMisc::NumberOfCoresIncludingHyperthreads() - 2, 2);
+			int32 NumThreadsInThreadPool = FPlatformMisc::NumberOfWorkerThreadsToSpawn();
 			// when we are in the editor we like to do things like build lighting and such
 			// this thread pool can be used for those purposes
 			extern CORE_API int32 GUseNewTaskBackend;
 			if (!GUseNewTaskBackend)
 			{
-				NumThreadsInLargeThreadPool = FMath::Max(FPlatformMisc::NumberOfCoresIncludingHyperthreads() - 2, 2);
-				NumThreadsInThreadPool = FPlatformMisc::NumberOfWorkerThreadsToSpawn();
 				GLargeThreadPool = FQueuedThreadPool::Allocate();
 			}
 			else
 			{
-				NumThreadsInLargeThreadPool = FMath::Max(LowLevelTasks::FScheduler::Get().GetNumWorkers() - 2, 2u);			
-				if(NumThreadsInLargeThreadPool < int(LowLevelTasks::FScheduler::Get().GetNumWorkers()))
-				{
-					NumThreadsInThreadPool = FMath::Max(NumThreadsInLargeThreadPool - 1, 1);
-					GLargeThreadPool = new FQueuedLowLevelThreadPool(NumThreadsInLargeThreadPool);
-				}
-				else
-				{
-					NumThreadsInLargeThreadPool = FMath::Max(FPlatformMisc::NumberOfCoresIncludingHyperthreads() - 2, 2);
-					NumThreadsInThreadPool = FPlatformMisc::NumberOfWorkerThreadsToSpawn();
-					GLargeThreadPool = FQueuedThreadPool::Allocate();
-				}
+				GLargeThreadPool = new FQueuedLowLevelThreadPool();
 			}
 		
 			// TaskGraph has it's HP threads slightly below normal, we want to be below the taskgraph HP threads to avoid interfering with the game-thread.
@@ -2241,24 +2228,7 @@ int32 FEngineLoop::PreInitPreStartupScreen(const TCHAR* CmdLine)
 			}
 			else
 			{
-				int NumThreadsInWorkerPool = FMath::Max(LowLevelTasks::FScheduler::Get().GetNumWorkers() - 2, 2u);			
-				if(NumThreadsInWorkerPool < int(LowLevelTasks::FScheduler::Get().GetNumWorkers()))
-				{
-					NumThreadsInWorkerPool = FMath::Max(NumThreadsInWorkerPool - 1, 1);
-					GThreadPool = new FQueuedLowLevelThreadPool(NumThreadsInWorkerPool);
-				}
-				else
-				{
-					GThreadPool = FQueuedThreadPool::Allocate();
-					int32 NumThreadsInThreadPool = FPlatformMisc::NumberOfWorkerThreadsToSpawn();
-
-					// we are only going to give dedicated servers one pool thread
-					if (FPlatformProperties::IsServerOnly())
-					{
-						NumThreadsInThreadPool = 1;
-					}
-					verify(GThreadPool->Create(NumThreadsInThreadPool, StackSize, TPri_SlightlyBelowNormal, TEXT("ThreadPool")));
-				}
+				GThreadPool = new FQueuedLowLevelThreadPool();
 			}
 		}
 #endif
