TSet::FindOrAdd
#rnx

diff --git a/Engine/Source/Runtime/Core/Private/Containers/ContainersTest.cpp b/Engine/Source/Runtime/Core/Private/Containers/ContainersTest.cpp
--- a/Engine/Source/Runtime/Core/Private/Containers/ContainersTest.cpp
+++ b/Engine/Source/Runtime/Core/Private/Containers/ContainersTest.cpp
@@ -4,6 +4,7 @@
 #include "Misc/AssertionMacros.h"
 #include "Containers/Array.h"
 #include "Containers/Map.h"
+#include "Containers/Set.h"
 #include "Containers/SortedMap.h"
 #include "Containers/ArrayView.h"
 #include "Misc/AutomationTest.h"
@@ -578,6 +579,179 @@ bool FContainerPerformanceTest::RunTest(const FString& Parameters)
 	return true;
 }
 
+namespace
+{
+	struct FRecorder
+	{
+		FRecorder(uint32 InKey = 0, uint32 InPayload=0)
+			: Id(NextId++)
+			, Key(InKey)
+			, Payload(InPayload)
+		{
+		}
+		FRecorder(const FRecorder& Other)
+			: Id(NextId++)
+			, Key(Other.Key)
+			, Payload(Other.Payload)
+			, NumCopies(Other.NumCopies+1)
+			, NumMoves(Other.NumMoves)
+		{
+		}
+		FRecorder(FRecorder&& Other)
+			: Id(NextId++)
+			, Key(Other.Key)
+			, Payload(Other.Payload)
+			, NumCopies(Other.NumCopies)
+			, NumMoves(Other.NumMoves+1)
+		{
+		}
+
+		bool operator==(const FRecorder& Other) const
+		{
+			return Key == Other.Key;
+		}
+
+		uint32 Id;
+		uint32 Key;
+		uint32 Payload;
+		uint32 NumCopies = 0;
+		uint32 NumMoves = 0;
+
+		static uint32 NextId;
+	};
+	uint32 FRecorder::NextId = 0;
+
+	uint32 GetTypeHash(const FRecorder& Recorder)
+	{
+		return Recorder.Payload;
+	}
+
+}
+IMPLEMENT_SIMPLE_AUTOMATION_TEST(FContainersTSetTest, "System.Core.Containers.TSet", EAutomationTestFlags::ApplicationContextMask | EAutomationTestFlags::EngineFilter)
+bool FContainersTSetTest::RunTest(const FString& Parameters)
+{
+	enum class EArgType
+	{
+		Copy,
+		Move
+	};
+	enum class EHashType
+	{
+		Internal,
+		PassedIn
+	};
+
+	for (EHashType HashType : { EHashType::Internal, EHashType::PassedIn})
+	{
+		for (EArgType ArgType : {EArgType::Copy, EArgType::Move})
+		{
+			FString FuncName;
+			auto FullText = [HashType,ArgType,&FuncName](const TCHAR* Message)
+			{
+				const TCHAR* HashText = HashType == EHashType::Internal ? TEXT("") : TEXT("ByHash");
+				const TCHAR* ArgText = ArgType == EArgType::Copy ? TEXT("const&") : TEXT("&&");
+				return FString::Printf(TEXT("TSet::%s%s(%s) %s"), *FuncName, HashText, ArgText, Message);
+			};
+
+			// Test TSet::Add(const&), Add(&&), AddByHash(const&), AddByHash(&&)
+			FuncName = TEXT("Add");
+			{
+				TSet<FRecorder> Set;
+				FRecorder First(37, 43);
+				bool bAlreadyInSet = true;
+
+				if (HashType == EHashType::Internal)
+					if (ArgType == EArgType::Copy)
+						Set.Add(First, &bAlreadyInSet);
+					else
+						Set.Add(MoveTemp(First), &bAlreadyInSet);
+				else
+					if (ArgType == EArgType::Copy)
+						Set.AddByHash(GetTypeHash(First), First, &bAlreadyInSet);
+					else
+						Set.AddByHash(GetTypeHash(First), MoveTemp(First), &bAlreadyInSet);
+				TestFalse(FullText(TEXT("returns bAlreadyInSet==false for first add")), bAlreadyInSet);
+
+				FRecorder* Found = Set.Find(First);
+				if (ArgType == EArgType::Copy)
+					TestTrue(FullText(TEXT("constructs a copy")), Found && Found->Id > First.Id && Found->NumCopies > 0 && Found->Payload == First.Payload);
+				else
+					TestTrue(FullText(TEXT("constructs a move")), Found && Found->Id > First.Id && Found->NumCopies == 0 && Found->NumMoves >= 1 && Found->Payload == First.Payload);
+
+				uint32 FoundId = Found ? Found->Id : 0;
+				Found = Set.Find(First);
+				TestTrue(TEXT("Finding an element returns a reference, no copies"), Found && Found->Id == FoundId);
+
+				FRecorder Second(37, 56);
+				if (HashType == EHashType::Internal)
+					if (ArgType == EArgType::Copy)
+						Set.Add(Second, &bAlreadyInSet);
+					else
+						Set.Add(MoveTemp(Second), &bAlreadyInSet);
+				else
+					if (ArgType == EArgType::Copy)
+						Set.AddByHash(GetTypeHash(Second), Second, &bAlreadyInSet);
+					else
+						Set.AddByHash(GetTypeHash(Second), MoveTemp(Second), &bAlreadyInSet);
+				TestTrue(FullText(TEXT("returns bAlreadyInSet==true for second add")), bAlreadyInSet);
+				Found = Set.Find(First);
+				TestTrue(FullText(TEXT("with a duplicate key constructs a copy of the new key")), Found && Found->Id > Second.Id && Found->Payload == Second.Payload);
+			}
+
+			// Test TSet::FindOrAdd(const&), FindOrAdd(&&), FindOrAddByHash(const&), FindOrAddByHash(&&)
+			FuncName = TEXT("FindOrAdd");
+			{
+				TSet<FRecorder> Set;
+				FRecorder First(37, 43);
+				bool bAlreadyInSet = true;
+
+				FRecorder* FindOrAddResult;
+				if (HashType == EHashType::Internal)
+					if (ArgType == EArgType::Copy)
+						FindOrAddResult = &Set.FindOrAdd(First, &bAlreadyInSet);
+					else
+						FindOrAddResult = &Set.FindOrAdd(MoveTemp(First), &bAlreadyInSet);
+				else
+					if (ArgType == EArgType::Copy)
+						FindOrAddResult = &Set.FindOrAddByHash(GetTypeHash(First), First, &bAlreadyInSet);
+					else
+						FindOrAddResult = &Set.FindOrAddByHash(GetTypeHash(First), MoveTemp(First), &bAlreadyInSet);
+				TestFalse(FullText(TEXT("returns bAlreadyInSet==false for first add")), bAlreadyInSet);
+				if (ArgType == EArgType::Copy)
+					TestTrue(FullText(TEXT("on the first constructs a copy")), FindOrAddResult->Id > First.Id && FindOrAddResult->NumCopies > 0 && FindOrAddResult->Payload == First.Payload);
+				else
+					TestTrue(FullText(TEXT("on the first constructs a move")), FindOrAddResult->Id > First.Id && FindOrAddResult->NumCopies == 0 && FindOrAddResult->NumMoves >= 1 && FindOrAddResult->Payload == First.Payload);
+				uint32 FoundId = FindOrAddResult->Id;
+
+				FRecorder* Found = Set.Find(First);
+				TestTrue(FullText(TEXT("returns same value as future find")), Found&& Found->Id == FindOrAddResult->Id);
+				Found = Set.Find(First);
+				TestTrue(TEXT("Finding an element returns a reference, no copies"), Found && Found->Id == FoundId);
+
+				FRecorder Second(37, 56);
+				if (HashType == EHashType::Internal)
+					if (ArgType == EArgType::Copy)
+						FindOrAddResult = &Set.FindOrAdd(Second, &bAlreadyInSet);
+					else
+						FindOrAddResult = &Set.FindOrAdd(MoveTemp(Second), &bAlreadyInSet);
+				else
+					if (ArgType == EArgType::Copy)
+						FindOrAddResult = &Set.FindOrAddByHash(GetTypeHash(Second), Second, &bAlreadyInSet);
+					else
+						FindOrAddResult = &Set.FindOrAddByHash(GetTypeHash(Second), MoveTemp(Second), &bAlreadyInSet);
+				TestTrue(FullText(TEXT("returns bAlreadyInSet==true for second add")), bAlreadyInSet);
+				TestTrue(FullText(TEXT("with a duplicate key keeps the original key, returned from FindOrAdd")), FindOrAddResult->Id == FoundId && FindOrAddResult->Payload == First.Payload);
+				Found = Set.Find(First);
+				TestTrue(FullText(TEXT("with a duplicate key keeps the original key, returned from future Finds")), Found->Id == FoundId && Found->Payload == First.Payload);
+			}
+		}
+	}
+
+
+	return !HasAnyErrors();
+}
+
+
 namespace ArrayViewTests
 {
 	// commented out lines shouldn't compile
diff --git a/Engine/Source/Runtime/Core/Public/Containers/Set.h b/Engine/Source/Runtime/Core/Public/Containers/Set.h
--- a/Engine/Source/Runtime/Core/Public/Containers/Set.h
+++ b/Engine/Source/Runtime/Core/Public/Containers/Set.h
@@ -604,6 +604,22 @@ public:
 	FORCEINLINE FSetElementId Add(const InElementType&  InElement, bool* bIsAlreadyInSetPtr = nullptr) { return Emplace(                   InElement , bIsAlreadyInSetPtr); }
 	FORCEINLINE FSetElementId Add(      InElementType&& InElement, bool* bIsAlreadyInSetPtr = nullptr) { return Emplace(MoveTempIfPossible(InElement), bIsAlreadyInSetPtr); }
 
+	/**
+	 * Adds an element to the set if not already present and returns a reference to the added or existing element.
+	 *
+	 * @param	InElement					Element to add to set
+	 * @param	bIsAlreadyInSetPtr	[out]	Optional pointer to bool that will be set depending on whether element is already in set
+	 * @return	A reference to the element stored in the set.
+	 */
+	FORCEINLINE ElementType& FindOrAdd(const InElementType& InElement, bool* bIsAlreadyInSetPtr = nullptr)
+	{
+		return FindOrAddByHash(KeyFuncs::GetKeyHash(KeyFuncs::GetSetKey(InElement)), InElement, bIsAlreadyInSetPtr);
+	}
+	FORCEINLINE ElementType& FindOrAdd(InElementType&& InElement, bool* bIsAlreadyInSetPtr = nullptr)
+	{
+		return FindOrAddByHash(KeyFuncs::GetKeyHash(KeyFuncs::GetSetKey(InElement)), MoveTempIfPossible(InElement), bIsAlreadyInSetPtr);
+	}
+
 	/**
 	 * Adds an element to the set.
 	 *
@@ -621,8 +637,37 @@ public:
 		return EmplaceByHash(KeyHash, MoveTempIfPossible(InElement), bIsAlreadyInSetPtr);
 	}
 
+	/**
+	 * Adds an element to the set if not already present and returns a reference to the added or existing element.
+	 *
+	 * @see		Class documentation section on ByHash() functions
+	 * @param	InElement					Element to add to set
+	 * @param	bIsAlreadyInSetPtr	[out]	Optional pointer to bool that will be set depending on whether element is already in set
+	 * @return  A reference to the element stored in the set
+	 */
+	template <typename ElementReferenceType>
+	ElementType& FindOrAddByHash(uint32 KeyHash, ElementReferenceType&& InElement, bool* bIsAlreadyInSetPtr = nullptr)
+	{
+		FSetElementId ExistingId = FindIdByHash(KeyHash, KeyFuncs::GetSetKey(InElement));
+		bool bIsAlreadyInSet = ExistingId.IsValidId();
+		if (bIsAlreadyInSetPtr)
+		{
+			*bIsAlreadyInSetPtr = bIsAlreadyInSet;
+		}
+		if (bIsAlreadyInSet)
+		{
+			return Elements[ExistingId].Value;
+		}
+
+		// Create a new element.
+		FSparseArrayAllocationInfo ElementAllocation = Elements.AddUninitialized();
+		SetElementType& Element = *new (ElementAllocation) SetElementType(Forward<ElementReferenceType>(InElement));
+		RehashOrLink(KeyHash, Element, ElementAllocation.Index);
+		return Element.Value;
+	}
+
 private:
-	FSetElementId EmplaceImpl(uint32 KeyHash, SetElementType& Element, FSetElementId ElementId, bool* bIsAlreadyInSetPtr)
+	bool TryReplaceExisting(uint32 KeyHash, SetElementType& Element, FSetElementId& InOutElementId, bool* bIsAlreadyInSetPtr)
 	{
 		bool bIsAlreadyInSet = false;
 		if (!KeyFuncs::bAllowDuplicateKeys)
@@ -640,30 +685,28 @@ private:
 					MoveByRelocate(Elements[ExistingId].Value, Element.Value);
 
 					// Then remove the new element.
-					Elements.RemoveAtUninitialized(ElementId);
+					Elements.RemoveAtUninitialized(InOutElementId);
 
 					// Then point the return value at the replaced element.
-					ElementId = ExistingId;
+					InOutElementId = ExistingId;
 				}
 			}
 		}
-
-		if (!bIsAlreadyInSet)
-		{
-			// Check if the hash needs to be resized.
-			if (!ConditionalRehash(Elements.Num()))
-			{
-				// If the rehash didn't add the new element to the hash, add it.
-				LinkElement(ElementId, Element, KeyHash);
-			}
-		}
-
 		if (bIsAlreadyInSetPtr)
 		{
 			*bIsAlreadyInSetPtr = bIsAlreadyInSet;
 		}
+		return bIsAlreadyInSet;
+	}
 
-		return ElementId;
+	FORCEINLINE void RehashOrLink(uint32 KeyHash, SetElementType& Element, FSetElementId ElementId)
+	{
+		// Check if the hash needs to be resized.
+		if (!ConditionalRehash(Elements.Num()))
+		{
+			// If the rehash didn't add the new element to the hash, add it.
+			LinkElement(ElementId, Element, KeyHash);
+		}
 	}
 
 public:
@@ -680,9 +723,14 @@ public:
 		// Create a new element.
 		FSparseArrayAllocationInfo ElementAllocation = Elements.AddUninitialized();
 		SetElementType& Element = *new (ElementAllocation) SetElementType(Forward<ArgsType>(Args));
+		FSetElementId ElementId = ElementAllocation.Index;
 
 		uint32 KeyHash = KeyFuncs::GetKeyHash(KeyFuncs::GetSetKey(Element.Value));
-		return EmplaceImpl(KeyHash, Element, ElementAllocation.Index, bIsAlreadyInSetPtr);
+		if (!TryReplaceExisting(KeyHash, Element, ElementId, bIsAlreadyInSetPtr))
+		{
+			RehashOrLink(KeyHash, Element, ElementId);
+		}
+		return ElementId;
 	}
 	
 	/**
@@ -699,8 +747,13 @@ public:
 		// Create a new element.
 		FSparseArrayAllocationInfo ElementAllocation = Elements.AddUninitialized();
 		SetElementType& Element = *new (ElementAllocation) SetElementType(Forward<ArgsType>(Args));
+		FSetElementId ElementId = ElementAllocation.Index;
 
-		return EmplaceImpl(KeyHash, Element, ElementAllocation.Index, bIsAlreadyInSetPtr);
+		if (!TryReplaceExisting(KeyHash, Element, ElementId, bIsAlreadyInSetPtr))
+		{
+			RehashOrLink(KeyHash, Element, ElementId);
+		}
+		return ElementId;
 	}
 
 	template<typename ArrayAllocator>
