Move the ExploreFolder Finder event(s) to the UI thread.

- This workaround places the event to open the Finder window onto the UI thread runloop, effectively deferring the action until the current event handler is complete.

#rnx

diff --git a/Engine/Source/Runtime/Core/Private/Mac/MacPlatformProcess.cpp b/Engine/Source/Runtime/Core/Private/Mac/MacPlatformProcess.cpp
--- a/Engine/Source/Runtime/Core/Private/Mac/MacPlatformProcess.cpp
+++ b/Engine/Source/Runtime/Core/Private/Mac/MacPlatformProcess.cpp
@@ -1257,8 +1257,12 @@ void FMacPlatformProcess::LaunchFileInDefaultExternalApplication( const TCHAR* F
 
 void FMacPlatformProcess::ExploreFolder( const TCHAR* FilePath )
 {
+    extern void MainThreadCall(dispatch_block_t Block, NSString* WaitMode, bool const bWait);
+
 	SCOPED_AUTORELEASE_POOL;
-	CFStringRef CFFilePath = FPlatformString::TCHARToCFString( FilePath );
+	__block CFStringRef CFFilePath = FPlatformString::TCHARToCFString( FilePath );
+
+    MainThreadCall(^{
 	BOOL IsDirectory = NO;
 	if([[NSFileManager defaultManager] fileExistsAtPath:(NSString *)CFFilePath isDirectory:&IsDirectory])
 	{
@@ -1273,6 +1277,7 @@ void FMacPlatformProcess::ExploreFolder( const TCHAR* FilePath )
 		}
 	}
 	CFRelease( CFFilePath );
+    }, NSDefaultRunLoopMode, false);
 }
 
 void FMacPlatformProcess::ClosePipe( void* ReadPipe, void* WritePipe )
