Support inert FStallDetector objects.

Add support for detectors that are constructed outside of when the StallDetector system is started up.
This gives us the flexibility to startup and shutdown within a finer scope of process startup and shutdown coming in a future change.

diff --git a/Engine/Source/Runtime/Core/Private/ProfilingDebugging/StallDetector.cpp b/Engine/Source/Runtime/Core/Private/ProfilingDebugging/StallDetector.cpp
--- a/Engine/Source/Runtime/Core/Private/ProfilingDebugging/StallDetector.cpp
+++ b/Engine/Source/Runtime/Core/Private/ProfilingDebugging/StallDetector.cpp
@@ -28,8 +28,10 @@
 DEFINE_LOG_CATEGORY(LogStall);
 
 /**
-* The reference count for the resources for this API
+* Globals
 **/
+
+// The reference count for the resources for this API
 static uint32 InitCount = 0;
 
 /**
@@ -103,6 +105,7 @@ uint32 UE::FStallDetectorRunnable::Run()
 		double Seconds = FStallDetector::Seconds();
 
 		// Check the detectors
+		if (Seconds != FStallDetector::InvalidSeconds)
 		{
 			FScopeLock ScopeLock(&FStallDetector::GetInstancesSection());
 			for (FStallDetector* Detector : FStallDetector::GetInstances())
@@ -215,18 +218,22 @@ void UE::FStallDetectorStats::TabulateStats(TArray<TabulatedResult>& TabulatedRe
 * Stall Detector
 **/
 
+const double UE::FStallDetector::InvalidSeconds = -1.0;
 FCriticalSection UE::FStallDetector::InstancesSection;
 TSet<UE::FStallDetector*> UE::FStallDetector::Instances;
 
 UE::FStallDetector::FStallDetector(FStallDetectorStats& InStats)
 	: Stats(InStats)
+	, ThreadId(0)
+	, StartSeconds(InvalidSeconds)
 	, bPersistent(false)
 	, Triggered(false)
 {
-	check(InitCount);
-
-	ThreadId = FPlatformTLS::GetCurrentThreadId();
-	StartSeconds = FStallDetector::Seconds();
+	if (InitCount)
+	{
+		ThreadId = FPlatformTLS::GetCurrentThreadId();
+		StartSeconds = FStallDetector::Seconds();
+	}
 
 	// Add at the end of construction
 	FScopeLock ScopeLock(&InstancesSection);
@@ -241,18 +248,32 @@ UE::FStallDetector::~FStallDetector()
 		Instances.Remove(this);
 	}
 
-	if (!bPersistent)
+	if (InitCount)
 	{
-		Check(true);
+		if (!bPersistent)
+		{
+			Check(true);
+		}
 	}
 }
 
 void UE::FStallDetector::Check(bool bIsComplete, double InWhenToCheckSeconds)
 {
+	// StartSeconds checks that system was started when this detector was constructed
+	bool bInitialized = InitCount && StartSeconds != InvalidSeconds;
+	if (!bInitialized)
+	{
+		return;
+	}
+
 	double CheckSeconds = InWhenToCheckSeconds;
-	if (InWhenToCheckSeconds == 0.0)
+	if (InWhenToCheckSeconds == InvalidSeconds)
 	{
 		CheckSeconds = FStallDetector::Seconds();
+		if (CheckSeconds == InvalidSeconds)
+		{
+			return;
+		}
 	}
 
 	double DeltaSeconds = CheckSeconds - StartSeconds;
@@ -290,7 +311,18 @@ void UE::FStallDetector::Check(bool bIsComplete, double InWhenToCheckSeconds)
 
 void UE::FStallDetector::CheckAndReset()
 {
+	// StartSeconds checks that system was started when this detector was constructed
+	bool bInitialized = InitCount && StartSeconds != InvalidSeconds;
+	if (!bInitialized)
+	{
+		return;
+	}
+
 	double CheckSeconds = FStallDetector::Seconds();
+	if (CheckSeconds == InvalidSeconds)
+	{
+		return;
+	}
 
 	// if this is the first call to CheckAndReset
 	if (!bPersistent)
@@ -377,31 +409,32 @@ void UE::FStallDetector::OnStallDetected(uint32 InThreadId, const double InElaps
 
 double UE::FStallDetector::Seconds()
 {
-	check(InitCount);
-
-	double Result;
+	double Result = InvalidSeconds;
 
+	if (InitCount)
+	{
 #if STALL_DETECTOR_HEART_BEAT_CLOCK
-	Result = Runnable->GetClock().Seconds();
+		Result = Runnable->GetClock().Seconds();
 #else
-	Result = FPlatformTime::Seconds();
+		Result = FPlatformTime::Seconds();
 #endif
 
 #if STALL_DETECTOR_DEBUG
-	static double ClockStartSeconds = Result;
-	static double PlatformStartSeconds = FPlatformTime::Seconds();
-	double ClockDelta = Result - ClockStartSeconds;
-	double PlatformDelta = FPlatformTime::Seconds() - PlatformStartSeconds;
-	double Drift = PlatformDelta - ClockDelta;
-	static double LastDrift = Drift;
-	double DriftDelta = Drift - LastDrift;
-	if (DriftDelta > 0.001)
-	{
-		FString ResultString = FString::Printf(TEXT("[FStallDetector] Thread %5d / Platform: %f / Clock: %f / Drift: %f (%f)\n"), FPlatformTLS::GetCurrentThreadId(), PlatformDelta, ClockDelta, Drift, DriftDelta);
-		FPlatformMisc::LocalPrint(ResultString.GetCharArray().GetData());
-		LastDrift = Drift;
-	}
+		static double ClockStartSeconds = Result;
+		static double PlatformStartSeconds = FPlatformTime::Seconds();
+		double ClockDelta = Result - ClockStartSeconds;
+		double PlatformDelta = FPlatformTime::Seconds() - PlatformStartSeconds;
+		double Drift = PlatformDelta - ClockDelta;
+		static double LastDrift = Drift;
+		double DriftDelta = Drift - LastDrift;
+		if (DriftDelta > 0.001)
+		{
+			FString ResultString = FString::Printf(TEXT("[FStallDetector] Thread %5d / Platform: %f / Clock: %f / Drift: %f (%f)\n"), FPlatformTLS::GetCurrentThreadId(), PlatformDelta, ClockDelta, Drift, DriftDelta);
+			FPlatformMisc::LocalPrint(ResultString.GetCharArray().GetData());
+			LastDrift = Drift;
+		}
 #endif
+	}
 
 	return Result;
 }
diff --git a/Engine/Source/Runtime/Core/Public/ProfilingDebugging/StallDetector.h b/Engine/Source/Runtime/Core/Public/ProfilingDebugging/StallDetector.h
--- a/Engine/Source/Runtime/Core/Public/ProfilingDebugging/StallDetector.h
+++ b/Engine/Source/Runtime/Core/Public/ProfilingDebugging/StallDetector.h
@@ -136,6 +136,11 @@ namespace UE
 	class CORE_API FStallDetector
 	{
 	public:
+		/**
+		* Special const value for invalid timestamps, used to detect races of startup/shutdown and threads
+		*/
+		static const double InvalidSeconds;
+
 		/**
 		* Construct the stall detector linking it to it's stats object and take a timestamp
 		*/
@@ -151,7 +156,7 @@ namespace UE
 		* @param bIsComplete			Is this check the one performed at completion of the scope or measurement period?
 		* @param InWhenToCheckSeconds	The timestamp at which to compare the start time to
 		**/
-		void Check(bool bIsComplete, double InWhenToCheckSeconds = 0.0);
+		void Check(bool bIsComplete, double InWhenToCheckSeconds = InvalidSeconds);
 
 		/**
 		* Perform a Check and then reset the timer to the current time.
