[TraceLog] Mark memory initialization event as important

Memory initialization event needs to be marked as important because we cannot receive time marker events before setting the base timestamp. It leads to negative time values.

diff --git a/Engine/Source/Runtime/Core/Private/ProfilingDebugging/Windows/MemoryTrace.cpp b/Engine/Source/Runtime/Core/Private/ProfilingDebugging/Windows/MemoryTrace.cpp
--- a/Engine/Source/Runtime/Core/Private/ProfilingDebugging/Windows/MemoryTrace.cpp
+++ b/Engine/Source/Runtime/Core/Private/ProfilingDebugging/Windows/MemoryTrace.cpp
@@ -82,7 +82,7 @@ static_assert(sizeof(FAddrPack) == sizeof(uint64), "");
 UE_TRACE_CHANNEL(MemSummaryChannel)
 UE_TRACE_CHANNEL_DEFINE(MemAllocChannel)
 
-UE_TRACE_EVENT_BEGIN(Memory, Init)
+UE_TRACE_EVENT_BEGIN(Memory, Init, NoSync|Important)
 	UE_TRACE_EVENT_FIELD(uint64, BaseCycle)
 	UE_TRACE_EVENT_FIELD(uint32, MarkerPeriod)
 	UE_TRACE_EVENT_FIELD(uint8, MinAlignment)
