fixed duplicated task trace ID

diff --git a/Engine/Source/Runtime/Core/Private/Async/TaskGraph.cpp b/Engine/Source/Runtime/Core/Private/Async/TaskGraph.cpp
--- a/Engine/Source/Runtime/Core/Private/Async/TaskGraph.cpp
+++ b/Engine/Source/Runtime/Core/Private/Async/TaskGraph.cpp
@@ -2257,9 +2257,6 @@ static TLockFreeClassAllocator_TLSCache<FGraphEvent, PLATFORM_CACHE_LINE_SIZE> T
 FGraphEventRef FGraphEvent::CreateGraphEvent()
 {
 	FGraphEvent* Instance = new(TheGraphEventAllocator.Allocate()) FGraphEvent{};
-	// `CreateGraphEvent()` creates an "empty" event that doesn't have an associated task and so won't get its TraceId set automatically
-	// to track such events we need to set a unique TraceId
-	Instance->SetTraceId(TaskTrace::GenerateTaskId());
 	return Instance;
 }
 
@@ -2309,6 +2306,11 @@ void FGraphEvent::DispatchSubsequents(TArray<FBaseGraphTask*>& NewTasks, ENamedT
 			{
 				LocalThreadToDoGatherOn = ThreadToDoGatherOn;
 			}
+
+#if UE_TASK_TRACE_ENABLED
+			// regenerate TraceId as we're going to use the same event with another task
+			TraceId = TaskTrace::GenerateTaskId();
+#endif
 			TGraphTask<FNullGraphTask>::CreateTask(FGraphEventRef(this), &TempEventsToWaitFor, CurrentThreadIfKnown).ConstructAndDispatchWhenReady(GET_STATID(STAT_FNullGraphTask_DontCompleteUntil), LocalThreadToDoGatherOn);
 			return;
 		}
diff --git a/Engine/Source/Runtime/Core/Public/Async/TaskGraphInterfaces.h b/Engine/Source/Runtime/Core/Public/Async/TaskGraphInterfaces.h
--- a/Engine/Source/Runtime/Core/Public/Async/TaskGraphInterfaces.h
+++ b/Engine/Source/Runtime/Core/Public/Async/TaskGraphInterfaces.h
@@ -521,6 +521,13 @@ protected:
 #endif
 	}
 
+	void SetTraceId(TaskTrace::FId InTraceId)
+	{
+#if UE_TASK_TRACE_ENABLED
+		TraceId = InTraceId;
+#endif
+	}
+
 private:
 	friend class FNamedTaskThread;
 	friend class FTaskThreadBase;
@@ -601,7 +608,7 @@ private:
 	LLM(const UE::LLMPrivate::FTagData* InheritedLLMTag);
 
 #if UE_TASK_TRACE_ENABLED
-	TaskTrace::FId TraceId{ TaskTrace::GenerateTaskId() };
+	TaskTrace::FId TraceId;
 #endif
 };
 
@@ -619,10 +626,7 @@ public:
 	 *	@return a reference counted pointer to the new graph event. Note this should be stored in a FGraphEventRef or it will be immediately destroyed!
 	**/
 	static CORE_API FGraphEventRef CreateGraphEvent();
-
-	// the returned event will have ref count zero; be sure to add one!
-	static CORE_API FGraphEvent* CreateGraphEventWithInlineStorage();
-
+	
 	/**
 	 *	Attempts to a new subsequent task. If this event has already fired, false is returned and action must be taken to ensure that the task will still fire even though this event cannot be a prerequisite (because it is already finished).
 	 *	@return true if the task was successfully set up as a subsequent. false if the event has already fired.
@@ -719,13 +723,6 @@ public:
 #endif
 	}
 
-	void SetTraceId(TaskTrace::FId InTraceId)
-	{
-#if UE_TASK_TRACE_ENABLED
-		TraceId = InTraceId;
-#endif
-	}
-
 private:
 	friend class TRefCountPtr<FGraphEvent>;
 	friend class TLockFreeClassAllocator_TLSCache<FGraphEvent, PLATFORM_CACHE_LINE_SIZE>;
@@ -798,7 +795,7 @@ private:
 #endif
 
 #if UE_TASK_TRACE_ENABLED
-	TaskTrace::FId TraceId = TaskTrace::InvalidId;
+	TaskTrace::FId TraceId = TaskTrace::GenerateTaskId();
 #endif
 };
 
@@ -917,9 +914,10 @@ public:
 
 	void Unlock(ENamedThreads::Type CurrentThreadIfKnown = ENamedThreads::AnyThread)
 	{
+		TaskTrace::Launched(GetTraceId(), nullptr, Subsequents.IsValid(), ((TTask*)&TaskStorage)->GetDesiredThread());
+
 		bool bWakeUpWorker = true;
 		ConditionalQueueTask(CurrentThreadIfKnown, bWakeUpWorker);
-		TaskTrace::Launched(GetTraceId(), nullptr, Subsequents.IsValid(), ((TTask*)&TaskStorage)->GetDesiredThread());
 	}
 
 	FGraphEventRef GetCompletionEvent()
@@ -1004,10 +1002,7 @@ private:
 		, TaskConstructed(false)
 	{
 		Subsequents.Swap(InSubsequents);
-		if (Subsequents.IsValid())
-		{
-			Subsequents->SetTraceId(GetTraceId());
-		}
+		SetTraceId(Subsequents.IsValid() ? Subsequents->GetTraceId() : TaskTrace::GenerateTaskId());
 	}
 
 	/** 
