U5 - Yield long running background pararallelFor Tasks so that higher priority work get's the chance to run.

diff --git a/Engine/Source/Runtime/Core/Private/Async/ParallelFor.cpp b/Engine/Source/Runtime/Core/Private/Async/ParallelFor.cpp
new file mode 100644
--- /dev/null
+++ b/Engine/Source/Runtime/Core/Private/Async/ParallelFor.cpp
@@ -0,0 +1,8 @@
+#include "Async/ParallelFor.h"
+
+CORE_API int32 GParallelForBackgroundYieldingTimeoutMs = 8;
+static FAutoConsoleVariableRef CVarParallelForBackgroundYieldingTimeout(
+	TEXT("Async.ParallelFor.YieldingTimeout"),
+	GParallelForBackgroundYieldingTimeoutMs,
+	TEXT("The timeout (in ms) when background priority parallel for task will yield execution to give higher priority tasks the chance to run.")
+);
\ No newline at end of file
diff --git a/Engine/Source/Runtime/Core/Public/Async/ParallelFor.h b/Engine/Source/Runtime/Core/Public/Async/ParallelFor.h
--- a/Engine/Source/Runtime/Core/Public/Async/ParallelFor.h
+++ b/Engine/Source/Runtime/Core/Public/Async/ParallelFor.h
@@ -17,6 +17,8 @@
 #include "Misc/App.h"
 #include "Misc/Fork.h"
 
+extern CORE_API int32 GParallelForBackgroundYieldingTimeoutMs;
+
 // Flags controlling the ParallelFor's behavior.
 enum class EParallelForFlags
 {
@@ -154,6 +156,19 @@ namespace ParallelForImpl
 		int32 LocalBlockSize = BlockSize;
 		int32 LocalNum = Num;
 		bool bLocalSaveLastBlockForMaster = bSaveLastBlockForMaster;
+
+		auto Now = [] { return FTimespan::FromSeconds(FPlatformTime::Seconds()); };
+		FTimespan Start = FTimespan::MinValue();
+			
+		bool bIsBackgroundPriority = !bMaster && ((InDesiredThread & ENamedThreads::ThreadPriorityMask) == ENamedThreads::BackgroundThreadPriority);
+
+		FTimespan YieldingThreshold;
+		if (bIsBackgroundPriority)
+		{
+			Start = Now();
+			YieldingThreshold = FTimespan::FromMilliseconds(FMath::Max(0, GParallelForBackgroundYieldingTimeoutMs));
+		}
+
 		TFunctionRef<void(int32)> LocalBody(Body);
 		while (true)
 		{
@@ -192,6 +207,16 @@ namespace ParallelForImpl
 			{
 				break;
 			}
+
+			if (bIsBackgroundPriority)
+			{
+				auto PassedTime = [Start, &Now]() { return Now() - Start; };
+				if (PassedTime() > YieldingThreshold)
+				{
+					TGraphTask<TParallelForTask<FunctionType>>::CreateTask().ConstructAndDispatchWhenReady(Data, InDesiredThread, 0);
+					return false;
+				}
+			}
 		}
 		return false;
 	}
