Demote StallDetector log prints from Warning to Log.

For now we don't want QA/Automation test results to write JIRAs for these events, we will triage them offline using analytics.

diff --git a/Engine/Source/Runtime/Core/Private/ProfilingDebugging/StallDetector.cpp b/Engine/Source/Runtime/Core/Private/ProfilingDebugging/StallDetector.cpp
--- a/Engine/Source/Runtime/Core/Private/ProfilingDebugging/StallDetector.cpp
+++ b/Engine/Source/Runtime/Core/Private/ProfilingDebugging/StallDetector.cpp
@@ -289,7 +289,7 @@ void UE::FStallDetector::Check(bool bIsComplete, double InWhenToCheckSeconds)
 			FString OverageString = FString::Printf(TEXT("[FStallDetector] [%s] Overage of %f\n"), Stats.Name, OverageSeconds);
 			FPlatformMisc::LocalPrint(OverageString.GetCharArray().GetData());
 #endif
-			UE_LOG(LogStall, Display, TEXT("Stall detector '%s' complete in %fs (%fs overbudget)"), Stats.Name, DeltaSeconds, OverageSeconds);
+			UE_LOG(LogStall, Log, TEXT("Stall detector '%s' complete in %fs (%fs overbudget)"), Stats.Name, DeltaSeconds, OverageSeconds);
 		}
 	}
 	else
@@ -395,7 +395,7 @@ void UE::FStallDetector::OnStallDetected(uint32 InThreadId, const double InElaps
 		Stats.bReported = true;
 		Stats.TotalReportedCount++;
 		const int NumStackFramesToIgnore = FPlatformTLS::GetCurrentThreadId() == InThreadId ? 2 : 0;
-		UE_LOG(LogStall, Warning, TEXT("Stall detector '%s' exceeded budget of %fs, reporting..."), Stats.Name, Stats.BudgetSeconds);
+		UE_LOG(LogStall, Log, TEXT("Stall detector '%s' exceeded budget of %fs, reporting..."), Stats.Name, Stats.BudgetSeconds);
 		double ReportSeconds = FStallDetector::Seconds();
 		ReportStall(Stats.Name, InThreadId, NumStackFramesToIgnore);
 		ReportSeconds = FStallDetector::Seconds() - ReportSeconds;
@@ -403,7 +403,7 @@ void UE::FStallDetector::OnStallDetected(uint32 InThreadId, const double InElaps
 	}
 	else
 	{
-		UE_LOG(LogStall, Warning, TEXT("Stall detector '%s' exceeded budget of %fs"), Stats.Name, Stats.BudgetSeconds);
+		UE_LOG(LogStall, Log, TEXT("Stall detector '%s' exceeded budget of %fs"), Stats.Name, Stats.BudgetSeconds);
 	}
 }
 
