Context can be null. We still need to respect the type derivation however so we'll create a temporary context

#rnx

diff --git a/Engine/Source/Runtime/Core/Private/Unix/UnixPlatformStackWalk.cpp b/Engine/Source/Runtime/Core/Private/Unix/UnixPlatformStackWalk.cpp
--- a/Engine/Source/Runtime/Core/Private/Unix/UnixPlatformStackWalk.cpp
+++ b/Engine/Source/Runtime/Core/Private/Unix/UnixPlatformStackWalk.cpp
@@ -664,8 +664,18 @@ void FUnixPlatformStackWalk::StackWalkAndDumpEx(ANSICHAR* HumanReadableString, S
 	GHandlingEnsure = bHandlingEnsure;
 	ECrashContextType HandlingType = bHandlingEnsure? ECrashContextType::Ensure : ECrashContextType::Crash;
 
-	FLocalGuardHelper Guard(reinterpret_cast<FUnixCrashContext*>(Context), HandlingType);
-	FPlatformStackWalk::StackWalkAndDump(HumanReadableString, HumanReadableStringSize, ProgramCounter, Context);
+	if (Context == nullptr)
+	{
+		FUnixCrashContext CrashContext(HandlingType, TEXT(""));
+		CrashContext.InitFromSignal(0, nullptr, nullptr);
+		CrashContext.FirstCrashHandlerFrame = nullptr; // ProgramCounter will trim the callstack instead.
+		FPlatformStackWalk::StackWalkAndDump(HumanReadableString, HumanReadableStringSize, ProgramCounter, &CrashContext);
+	}
+	else
+	{
+		FLocalGuardHelper Guard(reinterpret_cast<FUnixCrashContext*>(Context), HandlingType);
+		FPlatformStackWalk::StackWalkAndDump(HumanReadableString, HumanReadableStringSize, ProgramCounter, Context);
+	}
 
 	GHandlingEnsure = false;
 }
