Added copyright boilerplate to file.

diff --git a/Engine/Source/Runtime/Core/Private/Async/ParallelFor.cpp b/Engine/Source/Runtime/Core/Private/Async/ParallelFor.cpp
--- a/Engine/Source/Runtime/Core/Private/Async/ParallelFor.cpp
+++ b/Engine/Source/Runtime/Core/Private/Async/ParallelFor.cpp
@@ -1,3 +1,5 @@
+// Copyright Epic Games, Inc. All Rights Reserved.
+
 #include "Async/ParallelFor.h"
 
 CORE_API int32 GParallelForBackgroundYieldingTimeoutMs = 8;
