Fixed POSIX platform (Linux/Unix/Lumin) not honoring the -log= parameter in forked processes.
  - Reset the global variable 'CachedAbsoluteFilenamein child process to force reinitialization of the logger.
  - Added a callback to the OutputDeviceFile invoked every time the output file is opened for writing successfully to update the CachedAbsoluteFilename value in case the actual opened filename differ from CachedAbsoluteFilename generated before the file was opened to ensure CrashReporter picks the right log file in case of a crash/ensure.
  - Updated places where GetAbsoluteLogFilename() was called and the value cached to read the latest value instead.

diff --git a/Engine/Source/Runtime/Core/Private/GenericPlatform/GenericPlatformCrashContext.cpp b/Engine/Source/Runtime/Core/Private/GenericPlatform/GenericPlatformCrashContext.cpp
--- a/Engine/Source/Runtime/Core/Private/GenericPlatform/GenericPlatformCrashContext.cpp
+++ b/Engine/Source/Runtime/Core/Private/GenericPlatform/GenericPlatformCrashContext.cpp
@@ -158,7 +158,6 @@ void FGenericCrashContext::Initialize()
 		}
 
 		NCached::UserSettings.bNoDialog = FApp::IsUnattended() || IsRunningDedicatedServer();
-
 	}
 
 	// Create a unique base guid for bug report ids
@@ -597,7 +596,7 @@ void FGenericCrashContext::SerializeUserSettings(FString& Buffer)
 	AddCrashPropertyInternal(Buffer, TEXT("NoDialog"), NCached::UserSettings.bNoDialog);
 	AddCrashPropertyInternal(Buffer, TEXT("SendUnattendedBugReports"), NCached::UserSettings.bSendUnattendedBugReports);
 	AddCrashPropertyInternal(Buffer, TEXT("SendUsageData"), NCached::UserSettings.bSendUsageData);
-	AddCrashPropertyInternal(Buffer, TEXT("LogFilePath"), NCached::UserSettings.LogFilePath);
+	AddCrashPropertyInternal(Buffer, TEXT("LogFilePath"), FPlatformOutputDevices::GetAbsoluteLogFilename()); // Don't use the value cached, it may be out of date.
 }
 
 void FGenericCrashContext::SerializeContentToBuffer() const
diff --git a/Engine/Source/Runtime/Core/Private/GenericPlatform/GenericPlatformOutputDevices.cpp b/Engine/Source/Runtime/Core/Private/GenericPlatform/GenericPlatformOutputDevices.cpp
--- a/Engine/Source/Runtime/Core/Private/GenericPlatform/GenericPlatformOutputDevices.cpp
+++ b/Engine/Source/Runtime/Core/Private/GenericPlatform/GenericPlatformOutputDevices.cpp
@@ -11,18 +11,18 @@
 #include "Misc/OutputDeviceDebug.h"
 #include "Misc/OutputDeviceAnsiError.h"
 #include "Misc/App.h"
+#include "Misc/ScopeLock.h"
 #include "HAL/FeedbackContextAnsi.h"
 #include "Misc/OutputDeviceConsole.h"
 
 TCHAR FGenericPlatformOutputDevices::CachedAbsoluteFilename[FGenericPlatformOutputDevices::AbsoluteFileNameMaxLength] = { 0 };
-TUniquePtr<FOutputDeviceFile> FGenericPlatformOutputDevices::DefaultOutputDeviceFileTempHolder;
+FCriticalSection FGenericPlatformOutputDevices::LogFilenameLock;
 
 void FGenericPlatformOutputDevices::SetupOutputDevices()
 {
 	check(GLog);
 
-	InitDefaultOutputDeviceFile();
-
+	ResetCachedAbsoluteFilename();
 	GLog->AddOutputDevice(FPlatformOutputDevices::GetLog());
 
 	TArray<FOutputDevice*> ChannelFileOverrides;
@@ -51,71 +51,57 @@ void FGenericPlatformOutputDevices::SetupOutputDevices()
 	GLog->AddOutputDevice(FPlatformOutputDevices::GetEventLog());
 };
 
-
-void FGenericPlatformOutputDevices::InitDefaultOutputDeviceFile()
+void FGenericPlatformOutputDevices::ResetCachedAbsoluteFilename()
 {
-	if (CachedAbsoluteFilename[0] != 0)
-	{
-		return; // Already initialized
-	}
+	FScopeLock ScopeLock(&LogFilenameLock);
+	CachedAbsoluteFilename[0] = 0;
+}
 
-	FCString::Strcpy(CachedAbsoluteFilename, UE_ARRAY_COUNT(CachedAbsoluteFilename), *FPaths::ProjectLogDir());
-	FString LogFilename;
-	const bool bShouldStopOnSeparator = false;
-	if (!FParse::Value(FCommandLine::Get(), TEXT("LOG="), LogFilename, bShouldStopOnSeparator))
-	{
-		if (FParse::Value(FCommandLine::Get(), TEXT("ABSLOG="), LogFilename, bShouldStopOnSeparator))
-		{
-			CachedAbsoluteFilename[0] = 0;
-		}
-	}
+void FGenericPlatformOutputDevices::OnLogFileOpened(const TCHAR* Pathname)
+{
+	FScopeLock ScopeLock(&LogFilenameLock); // This function can be called on any thread, the first that serialize a log and lazily create the log file.
+	FCString::Strncpy(CachedAbsoluteFilename, Pathname, UE_ARRAY_COUNT(CachedAbsoluteFilename));
+}
 
-	FString Extension(FPaths::GetExtension(LogFilename));
-	if (Extension != TEXT("log") && Extension != TEXT("txt"))
-	{
-		// Ignoring the specified log filename because it doesn't have a .log extension			
-		LogFilename.Empty();
-	}
+FString FGenericPlatformOutputDevices::GetAbsoluteLogFilename()
+{
+	FScopeLock ScopeLock(&LogFilenameLock);
 
-	if (LogFilename.Len() == 0)
+	if (CachedAbsoluteFilename[0] == 0)
 	{
-		if (FCString::Strlen(FApp::GetProjectName()) != 0)
+		FCString::Strcpy(CachedAbsoluteFilename, UE_ARRAY_COUNT(CachedAbsoluteFilename), *FPaths::ProjectLogDir());
+		FString LogFilename;
+		const bool bShouldStopOnSeparator = false;
+		if (!FParse::Value(FCommandLine::Get(), TEXT("LOG="), LogFilename, bShouldStopOnSeparator))
 		{
-			LogFilename = FApp::GetProjectName();
+			if (FParse::Value(FCommandLine::Get(), TEXT("ABSLOG="), LogFilename, bShouldStopOnSeparator))
+			{
+				CachedAbsoluteFilename[0] = 0;
+			}
 		}
-		else
+
+		FString Extension(FPaths::GetExtension(LogFilename));
+		if (Extension != TEXT("log") && Extension != TEXT("txt"))
 		{
-			LogFilename = TEXT("Unreal");
+			// Ignoring the specified log filename because it doesn't have a .log extension.
+			LogFilename.Empty();
 		}
 
-		LogFilename += TEXT(".log");
-	}
-
-	FCString::Strcat(CachedAbsoluteFilename, UE_ARRAY_COUNT(CachedAbsoluteFilename) - FCString::Strlen(CachedAbsoluteFilename), *LogFilename);
-
-#if (!UE_BUILD_SHIPPING) || PRESERVE_LOG_BACKUPS_IN_SHIPPING
-	const bool bDisableBackup = false;
-#else
-	const bool bDisableBackup = true;
-#endif
-	// The desired 'CachedAbsoluteFilename' is likely ok unless a concurrent instance is already using it. To ensure the function returns real filename this process
-	// is going to use, create the log file, generate a new name if 'CachedAbsoluteFilename' is already used by another process and keep hold to it so than no other
-	// process can steal it.
-	DefaultOutputDeviceFileTempHolder = MakeUnique<FOutputDeviceFile>(CachedAbsoluteFilename, bDisableBackup, /*bAppendIfExists*/false, /*bCreateWriterLazily*/false);
+		if (LogFilename.Len() == 0)
+		{
+			if (FCString::Strlen(FApp::GetProjectName()) != 0)
+			{
+				LogFilename = FApp::GetProjectName();
+			}
+			else
+			{
+				LogFilename = TEXT("Unreal");
+			}
 
-	// Update the absolute log filename if 'CachedAbsoluteFilename' was already in use. FOutputDeviceFile automatically generated a new one.
-	if (FCString::Strcmp(CachedAbsoluteFilename, DefaultOutputDeviceFileTempHolder->GetFilename()) != 0)
-	{
-		FCString::Strcpy(CachedAbsoluteFilename, UE_ARRAY_COUNT(CachedAbsoluteFilename), DefaultOutputDeviceFileTempHolder->GetFilename());
-	}
-}
+			LogFilename += TEXT(".log");
+		}
 
-FString FGenericPlatformOutputDevices::GetAbsoluteLogFilename()
-{
-	if (!CachedAbsoluteFilename[0])
-	{
-		// Generates, cache and lock the filename, preventing concurrent instance to use it. Nothing will be logged to the file until the device get registered.
-		InitDefaultOutputDeviceFile();
+		FCString::Strcat(CachedAbsoluteFilename, UE_ARRAY_COUNT(CachedAbsoluteFilename) - FCString::Strlen(CachedAbsoluteFilename), *LogFilename);
 	}
 
 	return CachedAbsoluteFilename;
@@ -148,11 +134,15 @@ class FOutputDevice* FGenericPlatformOutputDevices::GetLog()
 #endif // WITH_LOGGING_TO_MEMORY
 			if (!LogDevice)
 			{
-				// Ensure the default file output device is created.
-				InitDefaultOutputDeviceFile();
-
-				// Activate the default file output device, transfering ownership.
-				LogDevice = MoveTemp(DefaultOutputDeviceFileTempHolder);
+#if (!UE_BUILD_SHIPPING) || PRESERVE_LOG_BACKUPS_IN_SHIPPING
+				const bool bDisableBackup = false;
+#else
+				const bool bDisableBackup = true;
+#endif
+				LogDevice = MakeUnique<FOutputDeviceFile>(nullptr, bDisableBackup, /*bAppendIfExists*/false, /*bCreateWriterLazily*/true, [](const TCHAR* AbsPathname)
+				{
+					FGenericPlatformOutputDevices::OnLogFileOpened(AbsPathname);
+				});
 			}
 		}
 
diff --git a/Engine/Source/Runtime/Core/Private/Lumin/LuminPlatformOutputDevices.cpp b/Engine/Source/Runtime/Core/Private/Lumin/LuminPlatformOutputDevices.cpp
--- a/Engine/Source/Runtime/Core/Private/Lumin/LuminPlatformOutputDevices.cpp
+++ b/Engine/Source/Runtime/Core/Private/Lumin/LuminPlatformOutputDevices.cpp
@@ -19,8 +19,7 @@ void FLuminOutputDevices::SetupOutputDevices()
 {
 	check(GLog);
 
-	InitDefaultOutputDeviceFile();
-
+	ResetCachedAbsoluteFilename();
 	GLog->AddOutputDevice(FPlatformOutputDevices::GetLog());
 
 #if !NO_LOGGING
@@ -34,7 +33,7 @@ void FLuminOutputDevices::SetupOutputDevices()
 	// unless logging is turned off
 	if (FPlatformMisc::HasSeparateChannelForDebugOutput())
 	{
-        // Use FLuminOutputDeviceDebug instead of the default FOutputDeviceDebug so that ml_log can respect the verbosity in the output
+		// Use FLuminOutputDeviceDebug instead of the default FOutputDeviceDebug so that ml_log can respect the verbosity in the output
 		GLog->AddOutputDevice(new FLuminOutputDeviceDebug());
 	}
 #endif
diff --git a/Engine/Source/Runtime/Core/Private/Mac/MacPlatformMisc.cpp b/Engine/Source/Runtime/Core/Private/Mac/MacPlatformMisc.cpp
--- a/Engine/Source/Runtime/Core/Private/Mac/MacPlatformMisc.cpp
+++ b/Engine/Source/Runtime/Core/Private/Mac/MacPlatformMisc.cpp
@@ -546,7 +546,11 @@ void FMacPlatformMisc::PlatformInit()
 	UE_LOG(LogInit, Log, TEXT("macOS %s (%s)"), *GMacAppInfo.OSVersion, *GMacAppInfo.OSBuild);
 	UE_LOG(LogInit, Log, TEXT("Model: %s"), *GMacAppInfo.MachineModel);
 	UE_LOG(LogInit, Log, TEXT("CPU: %s"), UTF8_TO_TCHAR(GMacAppInfo.MachineCPUString));
-	
+
+	// At the point, the log system is booted and the log file is likely created and the actual pathname used set. (Until the file is opened for writing, the name is subject to change)
+	FString LogPath = FGenericPlatformOutputDevices::GetAbsoluteLogFilename();
+	FCStringAnsi::Strcpy(GMacAppInfo.AppLogPath, PATH_MAX + 1, TCHAR_TO_UTF8(*LogPath));
+
 	const FPlatformMemoryConstants& MemoryConstants = FPlatformMemory::GetConstants();
 	UE_LOG(LogInit, Log, TEXT("CPU Page size=%i, Cores=%i, HT=%i"), MemoryConstants.PageSize, FPlatformMisc::NumberOfCores(), FPlatformMisc::NumberOfCoresIncludingHyperthreads() );
 
diff --git a/Engine/Source/Runtime/Core/Private/Misc/OutputDeviceFile.cpp b/Engine/Source/Runtime/Core/Private/Misc/OutputDeviceFile.cpp
--- a/Engine/Source/Runtime/Core/Private/Misc/OutputDeviceFile.cpp
+++ b/Engine/Source/Runtime/Core/Private/Misc/OutputDeviceFile.cpp
@@ -286,9 +286,10 @@ struct FOutputDeviceFile::FCategoryInclusionInternal
  * @param InFilename		Filename to use, can be NULL
  * @param bInDisableBackup	If true, existing files will not be backed up
  */
-FOutputDeviceFile::FOutputDeviceFile( const TCHAR* InFilename, bool bInDisableBackup, bool bInAppendIfExists, bool bCreateWriterLazily)
+FOutputDeviceFile::FOutputDeviceFile( const TCHAR* InFilename, bool bInDisableBackup, bool bInAppendIfExists, bool bCreateWriterLazily, TFunction<void(const TCHAR*)> FileOpenedCallback)
 : AsyncWriter(nullptr)
 , WriterArchive(nullptr)
+, OnFileOpenedFn(MoveTemp(FileOpenedCallback))
 , AppendIfExists(bInAppendIfExists)
 , Dead(false)
 , CategoryInclusionInternal(nullptr)
@@ -455,6 +456,10 @@ bool FOutputDeviceFile::CreateWriter(uint32 MaxAttempts)
 		WriterArchive = Ar;
 		AsyncWriter = new FAsyncWriter(*WriterArchive);
 		WriteByteOrderMarkToArchive(EByteOrderMark::UTF8);
+		if (OnFileOpenedFn)
+		{
+			OnFileOpenedFn(Filename);
+		}
 
 		if (!bSuppressEventTag)
 		{
diff --git a/Engine/Source/Runtime/Core/Private/Unix/UnixPlatformOutputDevices.cpp b/Engine/Source/Runtime/Core/Private/Unix/UnixPlatformOutputDevices.cpp
--- a/Engine/Source/Runtime/Core/Private/Unix/UnixPlatformOutputDevices.cpp
+++ b/Engine/Source/Runtime/Core/Private/Unix/UnixPlatformOutputDevices.cpp
@@ -18,7 +18,7 @@ void FUnixOutputDevices::SetupOutputDevices()
 {
 	check(GLog);
 
-	InitDefaultOutputDeviceFile();
+	ResetCachedAbsoluteFilename();
 
 	// add file log
 	GLog->AddOutputDevice(FPlatformOutputDevices::GetLog());
diff --git a/Engine/Source/Runtime/Core/Public/GenericPlatform/GenericPlatformOutputDevices.h b/Engine/Source/Runtime/Core/Public/GenericPlatform/GenericPlatformOutputDevices.h
--- a/Engine/Source/Runtime/Core/Public/GenericPlatform/GenericPlatformOutputDevices.h
+++ b/Engine/Source/Runtime/Core/Public/GenericPlatform/GenericPlatformOutputDevices.h
@@ -6,10 +6,10 @@
 #include "Containers/UnrealString.h"
 #include "GenericPlatform/GenericPlatform.h"
 #include "Templates/UniquePtr.h"
+#include "HAL/CriticalSection.h"
 
 class FOutputDeviceConsole;
 class FOutputDeviceError;
-class FOutputDeviceFile;
 
 /**
  * Generic implementation for most platforms
@@ -18,6 +18,13 @@ struct CORE_API FGenericPlatformOutputDevices
 {
 	/** Add output devices which can vary depending on platform, configuration, command line parameters. */
 	static void							SetupOutputDevices();
+
+	/**
+	 * Returns the absolute log filename generated from the project properties and/or command line parameters.
+	 * The returned value may change during the execution, may not exist yet or may be locked by another process.
+	 * It depends if the function is called before the log file was successfully opened or after. The log file is
+	 * open lazily.
+	 */
 	static FString						GetAbsoluteLogFilename();
 	static FOutputDevice*				GetLog();
 	static void							GetPerChannelFileOverrides(TArray<FOutputDevice*>& OutputDevices);
@@ -30,9 +37,12 @@ struct CORE_API FGenericPlatformOutputDevices
 	static FFeedbackContext*            GetFeedbackContext();
 
 protected:
-	static void InitDefaultOutputDeviceFile();
+	static void ResetCachedAbsoluteFilename();
 
-	static const SIZE_T AbsoluteFileNameMaxLength = 1024;
+private:
+	static constexpr SIZE_T AbsoluteFileNameMaxLength = 1024;
 	static TCHAR CachedAbsoluteFilename[AbsoluteFileNameMaxLength];
-	static TUniquePtr<FOutputDeviceFile> DefaultOutputDeviceFileTempHolder;
+
+	static void OnLogFileOpened(const TCHAR* Pathname);
+	static FCriticalSection LogFilenameLock;
 };
diff --git a/Engine/Source/Runtime/Core/Public/Misc/OutputDeviceFile.h b/Engine/Source/Runtime/Core/Public/Misc/OutputDeviceFile.h
--- a/Engine/Source/Runtime/Core/Public/Misc/OutputDeviceFile.h
+++ b/Engine/Source/Runtime/Core/Public/Misc/OutputDeviceFile.h
@@ -5,6 +5,7 @@
 #include "CoreTypes.h"
 #include "Misc/OutputDevice.h"
 #include "Templates/UniquePtr.h"
+#include "Templates/Function.h"
 #include "HAL/Runnable.h"
 #include "Serialization/Archive.h"
 #include "Misc/ScopeLock.h"
@@ -97,8 +98,9 @@ public:
 	*                   is created or the number of trials exhausted (32). 
 	* @param bDisableBackup If true, existing files will not be backed up
 	* @param bCreateWriterLazily If true, delay the creation of the file until something needs to be written, otherwise, open it immediatedly.
+	* @param FileOpenedCallback If bound, invoked when the output file is successfully opened, passing the actual filename.
 	*/
-	FOutputDeviceFile(const TCHAR* InFilename = nullptr, bool bDisableBackup = false, bool bAppendIfExists = false, bool bCreateWriterLazily = true);
+	FOutputDeviceFile(const TCHAR* InFilename = nullptr, bool bDisableBackup = false, bool bAppendIfExists = false, bool bCreateWriterLazily = true, TFunction<void(const TCHAR*)> FileOpenedCallback = TFunction<void(const TCHAR*)>());
 
 	/**
 	* Destructor to perform teardown
@@ -151,6 +153,8 @@ private:
 	FAsyncWriter* AsyncWriter;
 	/** Archive used by the async writer */
 	FArchive* WriterArchive;
+	/** In bound, invoked when the log file is open successfully for writing, reporting the actual log filename. */
+	TFunction<void(const TCHAR*)> OnFileOpenedFn;
 
 	TCHAR Filename[1024];
 	bool AppendIfExists;
