Fix assert on exit due to TLazySingleton failing to GetValue after modules have been cleaned up

diff --git a/Engine/Source/Runtime/Core/Private/Misc/ConfigCacheIni.cpp b/Engine/Source/Runtime/Core/Private/Misc/ConfigCacheIni.cpp
--- a/Engine/Source/Runtime/Core/Private/Misc/ConfigCacheIni.cpp
+++ b/Engine/Source/Runtime/Core/Private/Misc/ConfigCacheIni.cpp
@@ -5197,14 +5197,19 @@ void DeleteRecordedConfigReadsFromIni()
 
 #if ALLOW_OTHER_PLATFORM_CONFIG
 static TMap<FName, FConfigCacheIni> GConfigForPlatform;
-static TFuture<void> GConfigFuture;
 static FCriticalSection GConfigForPlatformLock;
+
+static TFuture<void>& GetConfigFuture()
+{
+	static TFuture<void> ConfigFuture;
+	return ConfigFuture;
+}
 #endif
 
 #if WITH_EDITOR
 void FConfigCacheIni::AsyncInitializeConfigForPlatforms()
 {
-	GConfigFuture = Async(EAsyncExecution::ThreadPool, []
+	GetConfigFuture() = Async(EAsyncExecution::ThreadPool, []
 	{
 		double Start = FPlatformTime::Seconds();
 
@@ -5236,12 +5241,12 @@ FConfigCacheIni* FConfigCacheIni::ForPlatform(FName PlatformName)
 	FScopeLock Lock(&GConfigForPlatformLock);
 
 	// they are likely already loaded, but just block to make sure
-	if (GConfigFuture.IsValid())
+	if (GetConfigFuture().IsValid())
 	{
 		double Start = FPlatformTime::Seconds();
 
-		GConfigFuture.Get();
-		GConfigFuture = TFuture<void>();
+		GetConfigFuture().Get();
+		GetConfigFuture() = TFuture<void>();
 
 		UE_LOG(LogConfig, Display, TEXT("Blocking on platform ini files took %.2f seconds"), FPlatformTime::Seconds() - Start);
 	}
@@ -5269,7 +5274,7 @@ void FConfigCacheIni::ClearOtherPlatformConfigs()
 
 	FScopeLock Lock(&GConfigForPlatformLock);
 
-	checkf(!GConfigFuture.IsValid(), TEXT("Trying to clear platform configs while still reading them in is unexpected"));
+	checkf(!GetConfigFuture().IsValid(), TEXT("Trying to clear platform configs while still reading them in is unexpected"));
 
 	GConfigForPlatform.Empty();
 #endif
