Only initialize stall detector for UnrealEditor.exe (remove it from cooks)
Add function to query if the static API is init and thread should be running
Tidy up some direct checks to the init count to use the new IsRunning function
Make init reference count signed and check for it going negative (mismatched startup/shutdown calls)
Disclaimer: init reference count isn't really used yet, but I'm hoping to propose using it more broadly in the future

diff --git a/Engine/Source/Runtime/Core/Private/ProfilingDebugging/StallDetector.cpp b/Engine/Source/Runtime/Core/Private/ProfilingDebugging/StallDetector.cpp
--- a/Engine/Source/Runtime/Core/Private/ProfilingDebugging/StallDetector.cpp
+++ b/Engine/Source/Runtime/Core/Private/ProfilingDebugging/StallDetector.cpp
@@ -32,7 +32,7 @@ DEFINE_LOG_CATEGORY(LogStall);
 **/
 
 // The reference count for the resources for this API
-static uint32 InitCount = 0;
+static int32 InitCount = 0;
 
 /**
 * Stall Detector Thread
@@ -229,7 +229,7 @@ UE::FStallDetector::FStallDetector(FStallDetectorStats& InStats)
 	, bPersistent(false)
 	, Triggered(false)
 {
-	if (InitCount)
+	if (FStallDetector::IsRunning())
 	{
 		ThreadId = FPlatformTLS::GetCurrentThreadId();
 		StartSeconds = FStallDetector::Seconds();
@@ -248,7 +248,7 @@ UE::FStallDetector::~FStallDetector()
 		Instances.Remove(this);
 	}
 
-	if (InitCount)
+	if (FStallDetector::IsRunning())
 	{
 		if (!bPersistent)
 		{
@@ -260,7 +260,7 @@ UE::FStallDetector::~FStallDetector()
 void UE::FStallDetector::Check(bool bIsComplete, double InWhenToCheckSeconds)
 {
 	// StartSeconds checks that system was started when this detector was constructed
-	bool bInitialized = InitCount && StartSeconds != InvalidSeconds;
+	bool bInitialized = FStallDetector::IsRunning() && StartSeconds != InvalidSeconds;
 	if (!bInitialized)
 	{
 		return;
@@ -312,7 +312,7 @@ void UE::FStallDetector::Check(bool bIsComplete, double InWhenToCheckSeconds)
 void UE::FStallDetector::CheckAndReset()
 {
 	// StartSeconds checks that system was started when this detector was constructed
-	bool bInitialized = InitCount && StartSeconds != InvalidSeconds;
+	bool bInitialized = FStallDetector::IsRunning() && StartSeconds != InvalidSeconds;
 	if (!bInitialized)
 	{
 		return;
@@ -411,7 +411,7 @@ double UE::FStallDetector::Seconds()
 {
 	double Result = InvalidSeconds;
 
-	if (InitCount)
+	if (FStallDetector::IsRunning())
 	{
 #if STALL_DETECTOR_HEART_BEAT_CLOCK
 		Result = Runnable->GetClock().Seconds();
@@ -441,6 +441,7 @@ double UE::FStallDetector::Seconds()
 
 void UE::FStallDetector::Startup()
 {
+	check(InitCount >= 0);
 	if (++InitCount == 1)
 	{
 		UE_LOG(LogStall, Log, TEXT("Startup..."));
@@ -480,6 +481,12 @@ void UE::FStallDetector::Shutdown()
 		
 		UE_LOG(LogStall, Log, TEXT("Shutdown complete."));
 	}
+	check(InitCount >= 0);
+}
+
+bool UE::FStallDetector::IsRunning()
+{
+	return InitCount > 0;
 }
 
 #endif // STALL_DETECTOR
\ No newline at end of file
diff --git a/Engine/Source/Runtime/Core/Public/ProfilingDebugging/StallDetector.h b/Engine/Source/Runtime/Core/Public/ProfilingDebugging/StallDetector.h
--- a/Engine/Source/Runtime/Core/Public/ProfilingDebugging/StallDetector.h
+++ b/Engine/Source/Runtime/Core/Public/ProfilingDebugging/StallDetector.h
@@ -196,6 +196,11 @@ namespace UE
 		**/
 		static void Shutdown();
 
+		/**
+		* Stall Detector API is running?
+		**/
+		static bool IsRunning();
+
 	private:
 		// The time provided, the deadline
 		FStallDetectorStats& Stats;
diff --git a/Engine/Source/Runtime/Launch/Private/LaunchEngineLoop.cpp b/Engine/Source/Runtime/Launch/Private/LaunchEngineLoop.cpp
--- a/Engine/Source/Runtime/Launch/Private/LaunchEngineLoop.cpp
+++ b/Engine/Source/Runtime/Launch/Private/LaunchEngineLoop.cpp
@@ -2619,17 +2619,26 @@ PRAGMA_ENABLE_DEPRECATION_WARNINGS
 		FAudioThread::SetUseThreadedAudio(bUseThreadedAudio);
 	}
 
-	if (FPlatformProcess::SupportsMultithreading() && !IsRunningDedicatedServer() && (bIsRegularClient || bHasEditorToken))
+	// Are we creating a slate application?
+	bool bSlateApplication = !IsRunningDedicatedServer() && (bIsRegularClient || bHasEditorToken);
+	if (bSlateApplication)
 	{
-		SCOPED_BOOT_TIMING("FPlatformSplash::Show()");
-		FPlatformSplash::Show();
-	}
+#if STALL_DETECTOR
+		if (GIsEditor)
+		{
+			UE::FStallDetector::Startup();
+		}
+#endif
+
+		if (FPlatformProcess::SupportsMultithreading())
+		{
+			SCOPED_BOOT_TIMING("FPlatformSplash::Show()");
+			FPlatformSplash::Show();
+		}
 
-	if (!IsRunningDedicatedServer() && (bHasEditorToken || bIsRegularClient))
-	{
 		// Init platform application
 		SCOPED_BOOT_TIMING("FSlateApplication::Create()");
-		FSlateApplication::Create();
+		FSlateApplication::Create();		
 	}
 	else
 	{
@@ -5877,10 +5886,6 @@ bool FEngineLoop::AppInit( )
 
 	FEmbeddedCommunication::ForceTick(19);
 
-#if STALL_DETECTOR
-	UE::FStallDetector::Startup();
-#endif
-
 	return true;
 }
 
@@ -5971,7 +5976,10 @@ void FEngineLoop::AppPreExit( )
 void FEngineLoop::AppExit()
 {
 #if STALL_DETECTOR
-	UE::FStallDetector::Shutdown();
+	if (UE::FStallDetector::IsRunning())
+	{
+		UE::FStallDetector::Shutdown();
+	}
 #endif
 
 #if !WITH_ENGINE
