U5 - make the YCombinator link inline

diff --git a/Engine/Source/Runtime/Core/Public/Async/LocalWorkQueue.h b/Engine/Source/Runtime/Core/Public/Async/LocalWorkQueue.h
--- a/Engine/Source/Runtime/Core/Public/Async/LocalWorkQueue.h
+++ b/Engine/Source/Runtime/Core/Public/Async/LocalWorkQueue.h
@@ -13,23 +13,23 @@ class TYCombinator
 	LAMBDA Lambda;
 
 public:
-	TYCombinator(LAMBDA&& InLambda) : Lambda(Forward<LAMBDA>(InLambda)) {}
+	inline TYCombinator(LAMBDA&& InLambda) : Lambda(Forward<LAMBDA>(InLambda)) {}
 
 	template <class... ARGS>
-	decltype(auto) operator()(ARGS&&... Args) const 
+	inline decltype(auto) operator()(ARGS&&... Args) const 
 	{
 		return Lambda(*this, Forward<ARGS>(Args)...);
 	}
 
 	template <class... ARGS>
-	decltype(auto) operator()(ARGS&&... Args) 
+	inline decltype(auto) operator()(ARGS&&... Args) 
 	{
 		return Lambda(*this, Forward<ARGS>(Args)...);
 	}
 };
 
 template<typename LAMBDA>
-TYCombinator<std::decay_t<LAMBDA>> MakeYCombinator(LAMBDA&& Lambda)
+inline TYCombinator<std::decay_t<LAMBDA>> MakeYCombinator(LAMBDA&& Lambda)
 {
 	return TYCombinator<std::decay_t<LAMBDA>>(Forward<LAMBDA>(Lambda));
 }
