U5 - inline AddWorkers

diff --git a/Engine/Source/Runtime/Core/Public/Async/LocalWorkQueue.h b/Engine/Source/Runtime/Core/Public/Async/LocalWorkQueue.h
--- a/Engine/Source/Runtime/Core/Public/Async/LocalWorkQueue.h
+++ b/Engine/Source/Runtime/Core/Public/Async/LocalWorkQueue.h
@@ -51,7 +51,7 @@ class TLocalWorkQueue
 	TFunctionRef<void(TaskType*)>* DoWork = nullptr;
 
 public:
-	TLocalWorkQueue(TaskType* InitialWork, LowLevelTasks::ETaskPriority InPriority = LowLevelTasks::ETaskPriority::Count) : Priority(InPriority)
+	inline TLocalWorkQueue(TaskType* InitialWork, LowLevelTasks::ETaskPriority InPriority = LowLevelTasks::ETaskPriority::Count) : Priority(InPriority)
 	{
 		if (Priority == LowLevelTasks::ETaskPriority::Count)
 		{
@@ -80,13 +80,13 @@ public:
 	}
 
 public:
-	void AddTask(TaskType* NewWork)
+	inline void AddTask(TaskType* NewWork)
 	{
 		check(!InternalData->CheckDone.load(std::memory_order_relaxed));
 		InternalData->TaskQueue.enqueue(NewWork);
 	}
 
-	void AddWorkers(uint16 NumWorkers)
+	inline void AddWorkers(uint16 NumWorkers)
 	{
 		check(!InternalData->CheckDone.load(std::memory_order_relaxed));
 		check(DoWork != nullptr);
@@ -117,7 +117,7 @@ public:
 		}
 	}
 
-	void Run(TFunctionRef<void(TaskType*)> InDoWork)
+	inline void Run(TFunctionRef<void(TaskType*)> InDoWork)
 	{
 		DoWork = &InDoWork;
 		LowLevelTasks::BusyWaitUntil([&InDoWork, InternalData = InternalData]()
