U5 - don't inline the operator of the ycombinator

diff --git a/Engine/Source/Runtime/Core/Public/Async/LocalWorkQueue.h b/Engine/Source/Runtime/Core/Public/Async/LocalWorkQueue.h
--- a/Engine/Source/Runtime/Core/Public/Async/LocalWorkQueue.h
+++ b/Engine/Source/Runtime/Core/Public/Async/LocalWorkQueue.h
@@ -16,13 +16,13 @@ public:
 	inline TYCombinator(LAMBDA&& InLambda) : Lambda(Forward<LAMBDA>(InLambda)) {}
 
 	template <class... ARGS>
-	inline decltype(auto) operator()(ARGS&&... Args) const 
+	FORCENOINLINE decltype(auto) operator()(ARGS&&... Args) const 
 	{
 		return Lambda(*this, Forward<ARGS>(Args)...);
 	}
 
 	template <class... ARGS>
-	inline decltype(auto) operator()(ARGS&&... Args) 
+	FORCENOINLINE decltype(auto) operator()(ARGS&&... Args) 
 	{
 		return Lambda(*this, Forward<ARGS>(Args)...);
 	}
