Fix FramePro integration to treat Unix platforms as Linux and some clang compile errors

diff --git a/Engine/Source/Runtime/Core/Private/FramePro/FrameProUE4.cpp b/Engine/Source/Runtime/Core/Private/FramePro/FrameProUE4.cpp
--- a/Engine/Source/Runtime/Core/Private/FramePro/FrameProUE4.cpp
+++ b/Engine/Source/Runtime/Core/Private/FramePro/FrameProUE4.cpp
@@ -150,6 +150,8 @@
 				return Platform::Windows;
 			#elif PLATFORM_LINUX
 				return Platform::Linux;
+			#elif PLATFORM_UNIX
+				return Platform::Linux;
 			#elif PLATFORM_PS4
 				return Platform::PS4;
 			#elif PLATFORM_ANDROID
@@ -467,7 +469,7 @@
 				p_event = new FPThreadEvent();
 				p_event->Create(!auto_reset);
 #else
-				checkf(false,"unsupported platform for -nothreading");
+				checkf(false, TEXT("unsupported platform for -nothreading"));
 				return;
 #endif
 			}
@@ -609,7 +611,7 @@
 					pthread_t thread_id;
 					pthread_create( &thread_id, nullptr, (PthreadEntryPoint)mp_ThreadMain, mp_Context );
 #else
-					checkf(false,"unsupported platform for -nothreading");
+					checkf(false,TEXT("unsupported platform for -nothreading"));
 #endif
 				}
 //@EPIC END
