Fixed out-of-process crash reporting used for the Editor on Windows to prevent deadlocking on allocation.
  - The function reporting the crash on the pipe doesn't need to suspend all the threads. The original purpose for suspending the threads was likely to preserve the state of the process as best as possible, but that is not required and prone to deadlocks.

Not suspending all the threads may would fix the hyphotetical case where CRC main thread is waiting for a prior ensure call stack to get resolved - I observed degenerated cases on my machine where this could take more than 15 minutes - preventing it to respond promptly to an incoming crash from the Editor. The flow was as following:
       - Editor fires an ensure, suspends all the thread, pipe a message to CRC to process the ensure.
       - CRC collects the ensure artefacts quickly, replies to the Editor, the Editor resumes, then CRC starts to resolve the call stack (blocking the main thread) from the minidump - degenerated cases can take several minutes.
       - Editor gets CRC messages and resumes its threads.
       - Editor fires a crash, suspends all the threads, pipes a message to CRC to process the crash.
       - CRC main thread is busy, waiting for the previous ensure call stack to be resolved... and doesn't respond promptly to the crash message.
       - Editor threads behing suspended, the code responsible to timeout if CRC takes too long never executes and Editor stalls until CRC dies or responds -> The user likely kills the Editor (and possibly CRC).

As a side effect from this change, if CRC doesn't respond promptly to a crash, the thread calling ReportCrash( )/ReportGPUCrash( ) will timeout and likely terminate the Editor before CRC could collect the crash artifacts or walk the thread to collect the call stacks.
  - Added a hint to the diagnostic logs reported with the Editor 'SummaryEvent' analytic event to indicate if the crash report was produced after the Editor died, so that the portable call stack wasn't captured.
  - Added a message displayed to the user by CRC saying that the the system failed to capture the callstack.

diff --git a/Engine/Source/Programs/CrashReportClient/Private/CrashReportClient.cpp b/Engine/Source/Programs/CrashReportClient/Private/CrashReportClient.cpp
--- a/Engine/Source/Programs/CrashReportClient/Private/CrashReportClient.cpp
+++ b/Engine/Source/Programs/CrashReportClient/Private/CrashReportClient.cpp
@@ -328,7 +328,14 @@ void FCrashReportClient::FinalizeDiagnoseReportWorker()
 	FString CallstackString = FPrimaryCrashProperties::Get()->CallStack.AsString();
 	if (CallstackString.IsEmpty())
 	{
-		DiagnosticText = LOCTEXT( "NoDebuggingSymbols", "You do not have any debugging symbols required to display the callstack for this crash." );
+		if (FPrimaryCrashProperties::Get()->PCallStackHash.IsEmpty())
+		{
+			DiagnosticText = LOCTEXT( "NoCallstack", "The system failed to capture the callstack for this crash." );
+		}
+		else
+		{
+			DiagnosticText = LOCTEXT( "NoDebuggingSymbols", "You do not have any debugging symbols required to display the callstack for this crash." );
+		}
 	}
 	else
 	{
diff --git a/Engine/Source/Programs/CrashReportClient/Private/CrashReportClientApp.cpp b/Engine/Source/Programs/CrashReportClient/Private/CrashReportClientApp.cpp
--- a/Engine/Source/Programs/CrashReportClient/Private/CrashReportClientApp.cpp
+++ b/Engine/Source/Programs/CrashReportClient/Private/CrashReportClientApp.cpp
@@ -433,8 +433,10 @@ SubmitCrashReportResult RunUnattended(FPlatformErrorReport ErrorReport, bool bIm
 	return FPrimaryCrashProperties::Get()->bIsEnsure ? SuccessContinue : SuccessClosed;
 }
 
-FPlatformErrorReport CollectErrorReport(FRecoveryService* RecoveryService, uint32 Pid, const FSharedCrashContext& SharedCrashContext, void* WritePipe)
+FPlatformErrorReport CollectErrorReport(FRecoveryService* RecoveryService, uint32 Pid, const FSharedCrashContext& SharedCrashContext, void* WritePipe, bool& bOutCrashPortableCallstackAvailable)
 {
+	bOutCrashPortableCallstackAvailable = false;
+
 	// @note: This API is only partially implemented on Mac OS and Linux.
 	FProcHandle ProcessHandle = FPlatformProcess::OpenProcess(Pid);
 
@@ -508,6 +510,9 @@ FPlatformErrorReport CollectErrorReport(FRecoveryService* RecoveryService, uint3
 				StackFrames,
 				StackFrameCount - SharedCrashContext.NumStackFramesToIgnore
 			);
+
+			// A completely missing portable callstack usually means that the crashing process died before CRC could walk the stack.
+			bOutCrashPortableCallstackAvailable = StackFrameCount > 0;
 		}
 	}
 
@@ -857,7 +862,8 @@ static void HandleAbnormalShutdown(FSharedCrashContext& CrashContext, uint64 Pro
 		// This crash is not a real one, but one to capture the Editor logs in case of abnormal termination.
 		FDiagnosticLogger::Get().LogEvent(TEXT("SyntheticCrash"));
 
-		FPlatformErrorReport ErrorReport = CollectErrorReport(RecoveryService.Get(), ProcessID, CrashContext, WritePipe);
+		bool bPortableCallstackAvailable = false; // Should always be false here. The process is dead.
+		FPlatformErrorReport ErrorReport = CollectErrorReport(RecoveryService.Get(), ProcessID, CrashContext, WritePipe, bPortableCallstackAvailable);
 		SendErrorReport(ErrorReport, /*bNoDialog*/ true);
 
 		// delete the temporary directory
@@ -1048,7 +1054,15 @@ void RunCrashReportClient(const TCHAR* CommandLine)
 					FDiagnosticLogger::Get().LogEvent(TEXT("Report/Collect"));
 
 					// Build error report in memory.
-					FPlatformErrorReport ErrorReport = CollectErrorReport(RecoveryServicePtr.Get(), MonitorPid, CrashContext, MonitorWritePipe);
+					bool bCrashedThreadCallstackAvailable = false;
+					FPlatformErrorReport ErrorReport = CollectErrorReport(RecoveryServicePtr.Get(), MonitorPid, CrashContext, MonitorWritePipe, bCrashedThreadCallstackAvailable);
+
+					// Log cases where the PCallstack is missing. For analytics, this event hints that the remote app exited before CRC could walk the process stack and possibly means
+					// that the timeout in the crashing process waiting for CRC to reply 'continue' is too short.
+					if (!bCrashedThreadCallstackAvailable)
+					{
+						FDiagnosticLogger::Get().LogEvent(FString::Printf(TEXT("Report/NoPCallstack")));
+					}
 
 #if CRASH_REPORT_WITH_RECOVERY
 					if (RecoveryServicePtr && !FPrimaryCrashProperties::Get()->bIsEnsure)
diff --git a/Engine/Source/Runtime/Core/Private/Windows/WindowsPlatformCrashContext.cpp b/Engine/Source/Runtime/Core/Private/Windows/WindowsPlatformCrashContext.cpp
--- a/Engine/Source/Runtime/Core/Private/Windows/WindowsPlatformCrashContext.cpp
+++ b/Engine/Source/Runtime/Core/Private/Windows/WindowsPlatformCrashContext.cpp
@@ -74,6 +74,7 @@ enum EConstants
 const uint32 EnsureExceptionCode = ECrashExitCodes::UnhandledEnsure; // Use a rather unique exception code in case SEH doesn't handle it as expected.
 const uint32 AssertExceptionCode = 0x4000;
 const uint32 GPUCrashExceptionCode = 0x8000;
+constexpr double CrashHandlingTimeoutSecs = 60.0;
 
 namespace {
 	/**
@@ -731,9 +732,7 @@ int32 ReportCrashForMonitor(
 	// query thread names from an external process.
 	uint32 ThreadIdx = 0;
 	DWORD CurrentProcessId = GetCurrentProcessId();
-	DWORD CurrentThreadId = GetCurrentThreadId();
 	HANDLE ThreadSnapshot = CreateToolhelp32Snapshot(TH32CS_SNAPTHREAD, 0);
-	TArray<HANDLE, TInlineAllocator<CR_MAX_THREADS>> ThreadHandles;
 	if (ThreadSnapshot != INVALID_HANDLE_VALUE)
 	{
 		THREADENTRY32 ThreadEntry;
@@ -744,17 +743,6 @@ int32 ReportCrashForMonitor(
 			{
 				if (ThreadEntry.th32OwnerProcessID == CurrentProcessId)
 				{
-					// Stop the thread (except current!). We will resume once the crash reporter is done.
-					if (ThreadEntry.th32ThreadID != CurrentThreadId)
-					{
-						HANDLE ThreadHandle = OpenThread(THREAD_SUSPEND_RESUME, FALSE, ThreadEntry.th32ThreadID);
-						if (ThreadHandle != NULL)
-						{
-							SuspendThread(ThreadHandle);
-							ThreadHandles.Push(ThreadHandle);
-						}
-					}
-
 					SharedContext->ThreadIds[ThreadIdx] = ThreadEntry.th32ThreadID;
 					const FString& TmThreadName = FThreadManager::GetThreadName(ThreadEntry.th32ThreadID);
 					const TCHAR* ThreadName = TmThreadName.IsEmpty() ? TEXT("Unknown") : *TmThreadName;
@@ -801,6 +789,7 @@ int32 ReportCrashForMonitor(
 
 	if (bPipeWriteSucceeded) // The receiver is not likely to respond if the shared context wasn't successfully written to the pipe.
 	{
+		double WaitResponseStartTimeSecs = FPlatformTime::Seconds();
 		// Wait for a response, saying it's ok to continue
 		bool bCanContinueExecution = false;
 		int32 ExitCode = 0;
@@ -818,13 +807,15 @@ int32 ReportCrashForMonitor(
 					bCanContinueExecution = true;
 				}
 			}
-		}
-	}
 
-	for (HANDLE ThreadHandle : ThreadHandles)
-	{
-		ResumeThread(ThreadHandle);
-		CloseHandle(ThreadHandle);
+			// In general, the crash monitor app (CRC) is expected to respond within ~5 seconds, but it might be busy sending an
+			// ensure/stall that occurred just before. In some degenerated cases, CRC may hang several minutes and cause the crash
+			// reporting thread to timeout and resume the crashing thread, likely resulting in this process to exit or to request it exit.
+			if (IsEngineExitRequested() && FPlatformTime::Seconds() - WaitResponseStartTimeSecs >= CrashHandlingTimeoutSecs)
+			{
+				break;
+			}
+		}
 	}
 
 	return EXCEPTION_CONTINUE_EXECUTION;
@@ -1118,7 +1109,6 @@ private:
 		__try
 #endif
 		{
-
 			while (StopTaskCounter.GetValue() == 0)
 			{
 				if (WaitForSingleObject(CrashEvent, 500) == WAIT_OBJECT_0)
@@ -1276,6 +1266,7 @@ public:
 			CrashContext.SetCrashedThreadId(InThreadId);
 			void* ContextWrapper = FWindowsPlatformStackWalk::MakeThreadContextWrapper(InExceptionInfo->ContextRecord, InThreadHandle);
 			CrashContext.CapturePortableCallStack(NumStackFramesToIgnore, ContextWrapper);
+			//CrashContext.CaptureAllThreadContexts(); -> For ensure/stall, don't capture all threads to report and resume quickly.
 
 			return ReportCrashUsingCrashReportClient(CrashContext, InExceptionInfo, ReportUI);
 		}
@@ -1294,7 +1285,7 @@ public:
 	FORCEINLINE bool WaitUntilCrashIsHandled()
 	{
 		// Wait 60s, it's more than enough to generate crash report. We don't want to stall forever otherwise.
-		return WaitForSingleObject(CrashHandledEvent, 60000) == WAIT_OBJECT_0;
+		return WaitForSingleObject(CrashHandledEvent, static_cast<DWORD>(CrashHandlingTimeoutSecs * 1000)) == WAIT_OBJECT_0;
 	}
 
 	/** Crashes during static init should be reported directly to crash monitor. */
@@ -1474,7 +1465,6 @@ private:
 		FPlatformStackWalk::UploadLocalSymbols();
 #endif
 	}
-
 };
 
 #include "Windows/HideWindowsPlatformTypes.h"
@@ -1741,7 +1731,6 @@ FORCENOINLINE void ReportGPUCrash(const TCHAR* ErrorMessage, int NumStackFramesT
 		FPlatformMisc::RequestExit(false);
 	}
 #endif
-
 }
 
 void ReportHang(const TCHAR* ErrorMessage, const uint64* StackFrames, int32 NumStackFrames, uint32 HungThreadId)
