CompactBinary: Added a 12-byte object identifier as a new field type

#rnx

diff --git a/Engine/Source/Runtime/Core/Private/Serialization/CompactBinary.cpp b/Engine/Source/Runtime/Core/Private/Serialization/CompactBinary.cpp
--- a/Engine/Source/Runtime/Core/Private/Serialization/CompactBinary.cpp
+++ b/Engine/Source/Runtime/Core/Private/Serialization/CompactBinary.cpp
@@ -10,6 +10,7 @@
 #include "Misc/Timespan.h"
 #include "Serialization/Archive.h"
 #include "Serialization/VarInt.h"
+#include "String/BytesToHex.h"
 
 ///////////////////////////////////////////////////////////////////////////////////////////////////////////////////////
 
@@ -73,6 +74,25 @@ void FCbFieldType::StaticAssertTypeConstants()
 
 ///////////////////////////////////////////////////////////////////////////////////////////////////////////////////////
 
+FCbObjectId::FCbObjectId(const FMemoryView ObjectId)
+{
+	checkf(ObjectId.GetSize() == sizeof(Bytes),
+		TEXT("FCbObjectId cannot be constructed from a view of %" UINT64_FMT " bytes."), ObjectId.GetSize());
+	FMemory::Memcpy(Bytes, ObjectId.GetData(), sizeof(Bytes));
+}
+
+void FCbObjectId::ToString(FAnsiStringBuilderBase& Builder) const
+{
+	UE::String::BytesToHexLower(Bytes, Builder);
+}
+
+void FCbObjectId::ToString(FWideStringBuilderBase& Builder) const
+{
+	UE::String::BytesToHexLower(Bytes, Builder);
+}
+
+///////////////////////////////////////////////////////////////////////////////////////////////////////////////////////
+
 FCbField::FCbField(const void* const InData, const ECbFieldType InType)
 {
 	const uint8* Bytes = static_cast<const uint8*>(InData);
@@ -383,6 +403,21 @@ FTimespan FCbField::AsTimeSpan(FTimespan Default)
 	return FTimespan(AsTimeSpanTicks(Default.GetTicks()));
 }
 
+FCbObjectId FCbField::AsObjectId(const FCbObjectId& Default)
+{
+	static_assert(sizeof(FCbObjectId) == 12, "FCbObjectId is expected to be 12 bytes.");
+	if (FCbFieldType::IsObjectId(Type))
+	{
+		Error = ECbFieldError::None;
+		return FCbObjectId(MakeMemoryView(Payload, sizeof(FCbObjectId)));
+	}
+	else
+	{
+		Error = ECbFieldError::TypeError;
+		return Default;
+	}
+}
+
 FCbCustomById FCbField::AsCustomById(FCbCustomById Default)
 {
 	if (FCbFieldType::IsCustomById(Type))
@@ -506,6 +541,8 @@ uint64 FCbField::GetPayloadSize() const
 	case ECbFieldType::DateTime:
 	case ECbFieldType::TimeSpan:
 		return 8;
+	case ECbFieldType::ObjectId:
+		return 12;
 	default:
 		return 0;
 	}
diff --git a/Engine/Source/Runtime/Core/Private/Serialization/CompactBinarySerialization.cpp b/Engine/Source/Runtime/Core/Private/Serialization/CompactBinarySerialization.cpp
--- a/Engine/Source/Runtime/Core/Private/Serialization/CompactBinarySerialization.cpp
+++ b/Engine/Source/Runtime/Core/Private/Serialization/CompactBinarySerialization.cpp
@@ -73,6 +73,9 @@ bool TryMeasureCompactBinary(FMemoryView View, ECbFieldType& OutType, uint64& Ou
 	case ECbFieldType::TimeSpan:
 		FixedSize = 8;
 		break;
+	case ECbFieldType::ObjectId:
+		FixedSize = 12;
+		break;
 	case ECbFieldType::None:
 	default:
 		OutType = ECbFieldType::None;
diff --git a/Engine/Source/Runtime/Core/Private/Serialization/CompactBinaryValidation.cpp b/Engine/Source/Runtime/Core/Private/Serialization/CompactBinaryValidation.cpp
--- a/Engine/Source/Runtime/Core/Private/Serialization/CompactBinaryValidation.cpp
+++ b/Engine/Source/Runtime/Core/Private/Serialization/CompactBinaryValidation.cpp
@@ -358,6 +358,9 @@ static FCbField ValidateCbField(FMemoryView& View, ECbValidateMode Mode, ECbVali
 	case ECbFieldType::TimeSpan:
 		ValidateCbFixedPayload(View, Mode, Error, 8);
 		break;
+	case ECbFieldType::ObjectId:
+		ValidateCbFixedPayload(View, Mode, Error, 12);
+		break;
 	case ECbFieldType::CustomById:
 	{
 		FMemoryView Value = ValidateCbDynamicPayload(View, Mode, Error);
diff --git a/Engine/Source/Runtime/Core/Private/Serialization/CompactBinaryWriter.cpp b/Engine/Source/Runtime/Core/Private/Serialization/CompactBinaryWriter.cpp
--- a/Engine/Source/Runtime/Core/Private/Serialization/CompactBinaryWriter.cpp
+++ b/Engine/Source/Runtime/Core/Private/Serialization/CompactBinaryWriter.cpp
@@ -560,6 +560,14 @@ void FCbWriter::AddTimeSpan(const FTimespan Value)
 	AddTimeSpanTicks(Value.GetTicks());
 }
 
+void FCbWriter::AddObjectId(const FCbObjectId& Value)
+{
+	static_assert(sizeof(FCbObjectId) == 12, "FCbObjectId is expected to be 12 bytes.");
+	BeginField();
+	Data.Append(reinterpret_cast<const uint8*>(&Value), sizeof(FCbObjectId));
+	EndField(ECbFieldType::ObjectId);
+}
+
 void FCbWriter::AddCustom(const uint64 TypeId, const FMemoryView Value)
 {
 	BeginField();
diff --git a/Engine/Source/Runtime/Core/Private/Tests/Serialization/CompactBinarySerializationTest.cpp b/Engine/Source/Runtime/Core/Private/Tests/Serialization/CompactBinarySerializationTest.cpp
--- a/Engine/Source/Runtime/Core/Private/Tests/Serialization/CompactBinarySerializationTest.cpp
+++ b/Engine/Source/Runtime/Core/Private/Tests/Serialization/CompactBinarySerializationTest.cpp
@@ -124,6 +124,8 @@ bool FCbMeasureTest::RunTest(const FString& Parameters)
 	TestMeasure(TEXT("DateTime"), {uint8(EType::DateTime)}, true, 9, EType::DateTime);
 	TestMeasure(TEXT("TimeSpan"), {uint8(EType::TimeSpan)}, true, 9, EType::TimeSpan);
 
+	TestMeasure(TEXT("ObjectId"), {uint8(EType::ObjectId)}, true, 13, EType::ObjectId);
+
 	TestMeasure(TEXT("CustomById, Size1B"), {uint8(EType::CustomById), 0x7f}, true, 129, EType::CustomById);
 	TestMeasure(TEXT("CustomById, Size2B"), {uint8(EType::CustomById), 0x80, 0x80}, true, 131, EType::CustomById);
 	TestMeasure(TEXT("CustomById, Size1BShort"), {uint8(EType::CustomById)}, false, 2, EType::CustomById);
diff --git a/Engine/Source/Runtime/Core/Private/Tests/Serialization/CompactBinaryTest.cpp b/Engine/Source/Runtime/Core/Private/Tests/Serialization/CompactBinaryTest.cpp
--- a/Engine/Source/Runtime/Core/Private/Tests/Serialization/CompactBinaryTest.cpp
+++ b/Engine/Source/Runtime/Core/Private/Tests/Serialization/CompactBinaryTest.cpp
@@ -65,6 +65,7 @@ UE_CBFIELD_TYPE_ACCESSOR_EX(Hash, IsHash, AsHash, FIoHash, const FIoHash&);
 UE_CBFIELD_TYPE_ACCESSOR_EX(Uuid, IsUuid, AsUuid, FGuid, const FGuid&);
 UE_CBFIELD_TYPE_ACCESSOR(DateTime, IsDateTime, AsDateTimeTicks, int64);
 UE_CBFIELD_TYPE_ACCESSOR(TimeSpan, IsTimeSpan, AsTimeSpanTicks, int64);
+UE_CBFIELD_TYPE_ACCESSOR_EX(ObjectId, IsObjectId, AsObjectId, FCbObjectId, const FCbObjectId&);
 UE_CBFIELD_TYPE_ACCESSOR(CustomById, IsCustomById, AsCustomById, FCbCustomById);
 UE_CBFIELD_TYPE_ACCESSOR(CustomByName, IsCustomByName, AsCustomByName, FCbCustomByName);
 
@@ -1089,6 +1090,34 @@ bool FCbFieldTimeSpanTest::RunTest(const FString& Parameters)
 	return true;
 }
 
+IMPLEMENT_CUSTOM_SIMPLE_AUTOMATION_TEST(FCbFieldObjectIdTest, FCbFieldTestBase, "System.Core.Serialization.CbField.ObjectId", CompactBinaryTestFlags)
+bool FCbFieldObjectIdTest::RunTest(const FString& Parameters)
+{
+	// Test FCbField(ObjectId, Zero)
+	TestField<ECbFieldType::ObjectId>(TEXT("ObjectId, Zero"), {0, 0, 0, 0, 0, 0, 0, 0, 0, 0, 0, 0});
+
+	// Test FCbField(ObjectId, 0x102030405060708090A0B0C0)
+	TestField<ECbFieldType::ObjectId>(TEXT("ObjectId, NonZero"), {0x10, 0x20, 0x30, 0x40, 0x50, 0x60, 0x70, 0x80, 0x90, 0xA0, 0xB0, 0xC0},
+		FCbObjectId(MakeMemoryView<uint8>({0x10, 0x20, 0x30, 0x40, 0x50, 0x60, 0x70, 0x80, 0x90, 0xA0, 0xB0, 0xC0})));
+
+	// Test FCbField(ObjectId, Zero) as FCbObjectId
+	{
+		const uint8 Payload[] = {0, 0, 0, 0, 0, 0, 0, 0, 0, 0, 0, 0};
+		FCbField Field(Payload, ECbFieldType::ObjectId);
+		TestEqual(TEXT("FCbField(ObjectId, Zero)::AsObjectId()"), Field.AsObjectId(), FCbObjectId());
+	}
+
+	// Test FCbField(None) as ObjectId
+	{
+		FCbField DefaultField;
+		TestFieldError<ECbFieldType::ObjectId>(TEXT("ObjectId, None"), DefaultField, ECbFieldError::TypeError);
+		const FCbObjectId DefaultValue(MakeMemoryView<uint8>({0x10, 0x20, 0x30, 0x40, 0x50, 0x60, 0x70, 0x80, 0x90, 0xA0, 0xB0, 0xC0}));
+		TestEqual(TEXT("FCbField(None)::AsObjectId()"), DefaultField.AsObjectId(DefaultValue), DefaultValue);
+	}
+
+	return true;
+}
+
 IMPLEMENT_CUSTOM_SIMPLE_AUTOMATION_TEST(FCbFieldCustomByIdTest, FCbFieldTestBase, "System.Core.Serialization.CbField.CustomById", CompactBinaryTestFlags)
 bool FCbFieldCustomByIdTest::RunTest(const FString& Parameters)
 {
diff --git a/Engine/Source/Runtime/Core/Private/Tests/Serialization/CompactBinaryValidationTest.cpp b/Engine/Source/Runtime/Core/Private/Tests/Serialization/CompactBinaryValidationTest.cpp
--- a/Engine/Source/Runtime/Core/Private/Tests/Serialization/CompactBinaryValidationTest.cpp
+++ b/Engine/Source/Runtime/Core/Private/Tests/Serialization/CompactBinaryValidationTest.cpp
@@ -170,6 +170,11 @@ bool FCbValidateTest::RunTest(const FString& Parameters)
 	TestEqual(TEXT("ValidateCompactBinary(OutOfBounds, TimeSpan)"), Validate({uint8(ECbFieldType::TimeSpan), 0, 0, 0, 0, 0, 0, 0}), ECbValidateError::OutOfBounds);
 	TestEqual(TEXT("ValidateCompactBinary(OutOfBounds, TimeSpan, NoType)"), Validate({0, 0, 0, 0, 0, 0, 0}, ECbFieldType::TimeSpan), ECbValidateError::OutOfBounds);
 
+	TestEqual(TEXT("ValidateCompactBinary(Valid, ObjectId)"), Validate({uint8(ECbFieldType::ObjectId), 0, 0, 0, 0, 0, 0, 0, 0, 0, 0, 0, 0}), ECbValidateError::None);
+	TestEqual(TEXT("ValidateCompactBinary(Valid, ObjectId, NoType)"), Validate({0, 0, 0, 0, 0, 0, 0, 0, 0, 0, 0, 0}, ECbFieldType::ObjectId), ECbValidateError::None);
+	TestEqual(TEXT("ValidateCompactBinary(OutOfBounds, ObjectId)"), Validate({uint8(ECbFieldType::ObjectId), 0, 0, 0, 0, 0, 0, 0, 0, 0, 0, 0}), ECbValidateError::OutOfBounds);
+	TestEqual(TEXT("ValidateCompactBinary(OutOfBounds, ObjectId, NoType)"), Validate({0, 0, 0, 0, 0, 0, 0, 0, 0, 0, 0}, ECbFieldType::ObjectId), ECbValidateError::OutOfBounds);
+
 	TestEqual(TEXT("ValidateCompactBinary(Valid, CustomById, Empty)"), Validate({uint8(ECbFieldType::CustomById), 1, 0}), ECbValidateError::None);
 	TestEqual(TEXT("ValidateCompactBinary(Valid, CustomById, Empty, NoType)"), Validate({1, 0}, ECbFieldType::CustomById), ECbValidateError::None);
 	TestEqual(TEXT("ValidateCompactBinary(Valid, CustomById, Payload)"), Validate({uint8(ECbFieldType::CustomById), 2, 0, 0}), ECbValidateError::None);
@@ -207,8 +212,8 @@ bool FCbValidateTest::RunTest(const FString& Parameters)
 	TestEqual(TEXT("ValidateCompactBinary(OutOfBounds, CustomByName, ShortPayload, NoType)"), Validate({4, 1, 'A', 0}, ECbFieldType::CustomByName), ECbValidateError::OutOfBounds);
 
 	// Test InvalidType
-	TestEqual(TEXT("ValidateCompactBinary(InvalidType, Unknown)"), Validate({uint8(ECbFieldType::TimeSpan) + 1}), ECbValidateError::InvalidType);
-	TestEqual(TEXT("ValidateCompactBinary(InvalidType, Unknown)"), Validate({}, ECbFieldType(uint8(ECbFieldType::TimeSpan) + 1)), ECbValidateError::InvalidType);
+	TestEqual(TEXT("ValidateCompactBinary(InvalidType, Unknown)"), Validate({uint8(ECbFieldType::ObjectId) + 1}), ECbValidateError::InvalidType);
+	TestEqual(TEXT("ValidateCompactBinary(InvalidType, Unknown)"), Validate({}, ECbFieldType(uint8(ECbFieldType::ObjectId) + 1)), ECbValidateError::InvalidType);
 	TestEqual(TEXT("ValidateCompactBinary(InvalidType, HasFieldType)"), Validate({uint8(ECbFieldType::Null | ECbFieldType::HasFieldType)}), ECbValidateError::InvalidType);
 
 	TestEqual(TEXT("ValidateCompactBinary(InvalidType, ZeroSizeField)"), Validate({}, ECbFieldType::Null), ECbValidateError::InvalidType);
diff --git a/Engine/Source/Runtime/Core/Private/Tests/Serialization/CompactBinaryWriterTest.cpp b/Engine/Source/Runtime/Core/Private/Tests/Serialization/CompactBinaryWriterTest.cpp
--- a/Engine/Source/Runtime/Core/Private/Tests/Serialization/CompactBinaryWriterTest.cpp
+++ b/Engine/Source/Runtime/Core/Private/Tests/Serialization/CompactBinaryWriterTest.cpp
@@ -735,6 +735,33 @@ bool FCbWriterTimeSpanTest::RunTest(const FString& Parameters)
 
 ///////////////////////////////////////////////////////////////////////////////////////////////////////////////////////
 
+IMPLEMENT_SIMPLE_AUTOMATION_TEST(FCbWriterObjectIdTest, "System.Core.Serialization.CbWriter.ObjectId", CompactBinaryWriterTestFlags)
+bool FCbWriterObjectIdTest::RunTest(const FString& Parameters)
+{
+	TCbWriter<256> Writer;
+
+	const FCbObjectId Values[] = { FCbObjectId(), FCbObjectId(MakeMemoryView<uint8>({1, 2, 3, 4, 5, 6, 7, 8, 9, 10, 11, 12})) };
+	for (const FCbObjectId& Value : Values)
+	{
+		Writer.AddObjectId(Value);
+	}
+
+	FCbFieldRefIterator Fields = Writer.Save();
+	if (TestEqual(TEXT("FCbWriter(ObjectId) Validate"), ValidateCompactBinaryRange(Fields.GetRangeBuffer(), ECbValidateMode::All), ECbValidateError::None))
+	{
+		const FCbObjectId* CheckValue = Values;
+		for (FCbField Field : Fields)
+		{
+			TestEqual(TEXT("FCbWriter(ObjectId).AsObjectId()"), Field.AsObjectId(), *CheckValue++);
+			TestFalse(TEXT("FCbWriter(ObjectId) Error"), Field.HasError());
+		}
+	}
+
+	return true;
+}
+
+///////////////////////////////////////////////////////////////////////////////////////////////////////////////////////
+
 IMPLEMENT_SIMPLE_AUTOMATION_TEST(FCbWriterCustomByIdTest, "System.Core.Serialization.CbWriter.CustomById", CompactBinaryWriterTestFlags)
 bool FCbWriterCustomByIdTest::RunTest(const FString& Parameters)
 {
@@ -902,6 +929,9 @@ bool FCbWriterComplexTest::RunTest(const FString& Parameters)
 		Writer.AddTimeSpanTicks("TimeSpanZero"_ASV, 0);
 		Writer.AddTimeSpan("TimeSpan"_ASV, FTimespan(1, 2, 4, 8));
 
+		Writer.AddObjectId("ObjectIdZero"_ASV, FCbObjectId());
+		Writer.AddObjectId("ObjectId"_ASV, FCbObjectId(MakeMemoryView<uint8>({1, 2, 3, 4, 5, 6, 7, 8, 9, 10, 11, 12})));
+
 		Writer.BeginObject("NestedObjects"_ASV);
 		{
 			Writer.BeginObject("Empty"_ASV);
@@ -1056,6 +1086,9 @@ bool FCbWriterStreamTest::RunTest(const FString& Parameters)
 		Writer << "DateTime"_ASV << FDateTime(2020, 5, 13, 15, 10);
 		Writer << "TimeSpan"_ASV << FTimespan(1, 2, 4, 8);
 
+		Writer << "ObjectIdZero"_ASV << FCbObjectId();
+		Writer << "ObjectId"_ASV << FCbObjectId(MakeMemoryView<uint8>({1, 2, 3, 4, 5, 6, 7, 8, 9, 10, 11, 12}));
+
 		Writer << "LiteralName" << nullptr;
 
 		Writer.EndObject();
diff --git a/Engine/Source/Runtime/Core/Public/Serialization/CompactBinary.h b/Engine/Source/Runtime/Core/Public/Serialization/CompactBinary.h
--- a/Engine/Source/Runtime/Core/Public/Serialization/CompactBinary.h
+++ b/Engine/Source/Runtime/Core/Public/Serialization/CompactBinary.h
@@ -3,6 +3,7 @@
 #pragma once
 
 #include "CoreTypes.h"
+#include "Containers/StringFwd.h"
 #include "Containers/StringView.h"
 #include "IO/IoHash.h"
 #include "Memory/MemoryView.h"
@@ -65,6 +66,7 @@ class FBlake3;
 class FCbArray;
 class FCbField;
 class FCbObject;
+class FCbObjectId;
 struct FDateTime;
 struct FGuid;
 struct FTimespan;
@@ -174,6 +176,13 @@ enum class ECbFieldType : uint8
 	 */
 	TimeSpan                         = 0x13,
 
+	/**
+	 * ObjectId is an opaque object identifier. See FCbObjectId.
+	 *
+	 * Payload is a 12-byte object identifier.
+	 */
+	ObjectId                         = 0x14,
+
 	/**
 	 * CustomById identifies the sub-type of its payload by an integer identifier.
 	 *
@@ -268,6 +277,8 @@ public:
 	static constexpr inline bool IsDateTime(ECbFieldType Type)   { return GetType(Type) == ECbFieldType::DateTime; }
 	static constexpr inline bool IsTimeSpan(ECbFieldType Type)   { return GetType(Type) == ECbFieldType::TimeSpan; }
 
+	static constexpr inline bool IsObjectId(ECbFieldType Type)   { return GetType(Type) == ECbFieldType::ObjectId; }
+
 	static constexpr inline bool IsCustomById(ECbFieldType Type)   { return GetType(Type) == ECbFieldType::CustomById; }
 	static constexpr inline bool IsCustomByName(ECbFieldType Type) { return GetType(Type) == ECbFieldType::CustomByName; }
 
@@ -439,6 +450,60 @@ enum class ECbFieldError : uint8
 
 ///////////////////////////////////////////////////////////////////////////////////////////////////////////////////////
 
+/**
+ * An opaque 12-byte object identifier.
+ *
+ * It has no intrinsic meaning, and can only be properly interpreted in the context of its usage.
+ */
+class FCbObjectId
+{
+public:
+	/** Construct an ObjectId with every byte initialized to zero. */
+	FCbObjectId() = default;
+
+	/** Construct an ObjectId from a view of 12 bytes. */
+	CORE_API explicit FCbObjectId(FMemoryView ObjectId);
+
+	/** Convert the ObjectId to a 24-character hex string. */
+	CORE_API void ToString(FAnsiStringBuilderBase& Builder) const;
+	CORE_API void ToString(FWideStringBuilderBase& Builder) const;
+
+	/** Returns a view of the raw byte array for the ObjectId. */
+	constexpr inline FMemoryView GetView() const { return MakeMemoryView(Bytes); }
+
+private:
+	alignas(uint32) uint8 Bytes[12]{};
+};
+
+inline bool operator==(const FCbObjectId& A, const FCbObjectId& B)
+{
+	return FMemory::Memcmp(&A, &B, sizeof(FCbObjectId)) == 0;
+}
+
+inline bool operator!=(const FCbObjectId& A, const FCbObjectId& B)
+{
+	return FMemory::Memcmp(&A, &B, sizeof(FCbObjectId)) != 0;
+}
+
+inline bool operator<(const FCbObjectId& A, const FCbObjectId& B)
+{
+	return FMemory::Memcmp(&A, &B, sizeof(FCbObjectId)) <= 0;
+}
+
+inline uint32 GetTypeHash(const FCbObjectId& Id)
+{
+	return *reinterpret_cast<const uint32*>(&Id);
+}
+
+template <typename CharType>
+inline TStringBuilderBase<CharType>& operator<<(TStringBuilderBase<CharType>& Builder, const FCbObjectId& Id)
+{
+	Id.ToString(Builder);
+	return Builder;
+}
+
+///////////////////////////////////////////////////////////////////////////////////////////////////////////////////////
+
 /** A custom compact binary field type with an integer identifier. */
 struct FCbCustomById
 {
@@ -567,6 +632,9 @@ public:
 	/** Access the field as a timespan. Returns the provided default on error. */
 	CORE_API FTimespan AsTimeSpan(FTimespan Default);
 
+	/** Access the field as an object identifier. Returns the provided default on error. */
+	CORE_API FCbObjectId AsObjectId(const FCbObjectId& Default = FCbObjectId());
+
 	/** Access the field as a custom sub-type with an integer identifier. Returns the provided default on error. */
 	CORE_API FCbCustomById AsCustomById(FCbCustomById Default = FCbCustomById());
 	/** Access the field as a custom sub-type with a string identifier. Returns the provided default on error. */
@@ -604,6 +672,8 @@ public:
 	constexpr inline bool IsDateTime() const        { return FCbFieldType::IsDateTime(Type); }
 	constexpr inline bool IsTimeSpan() const        { return FCbFieldType::IsTimeSpan(Type); }
 
+	constexpr inline bool IsObjectId() const        { return FCbFieldType::IsObjectId(Type); }
+
 	constexpr inline bool IsCustomById() const      { return FCbFieldType::IsCustomById(Type); }
 	constexpr inline bool IsCustomByName() const    { return FCbFieldType::IsCustomByName(Type); }
 
diff --git a/Engine/Source/Runtime/Core/Public/Serialization/CompactBinaryWriter.h b/Engine/Source/Runtime/Core/Public/Serialization/CompactBinaryWriter.h
--- a/Engine/Source/Runtime/Core/Public/Serialization/CompactBinaryWriter.h
+++ b/Engine/Source/Runtime/Core/Public/Serialization/CompactBinaryWriter.h
@@ -220,6 +220,10 @@ public:
 	CORE_API void AddTimeSpan(FAnsiStringView Name, FTimespan Value);
 	CORE_API void AddTimeSpan(FTimespan Value);
 
+	/** Write an ObjectId field. */
+	inline void AddObjectId(FAnsiStringView Name, const FCbObjectId& Value) { SetName(Name); AddObjectId(Value); }
+	CORE_API void AddObjectId(const FCbObjectId& Value);
+
 	/** Write a custom field with an integer sub-type identifier. */
 	inline void AddCustom(FAnsiStringView FieldName, uint64 TypeId, FMemoryView Value) { SetName(FieldName); AddCustom(TypeId, Value); }
 	CORE_API void AddCustom(uint64 TypeId, FMemoryView Value);
@@ -434,4 +438,10 @@ inline FCbWriter& operator<<(FCbWriter& Writer, const FGuid& Value)
 CORE_API FCbWriter& operator<<(FCbWriter& Writer, FDateTime Value);
 CORE_API FCbWriter& operator<<(FCbWriter& Writer, FTimespan Value);
 
+inline FCbWriter& operator<<(FCbWriter& Writer, const FCbObjectId& Value)
+{
+	Writer.AddObjectId(Value);
+	return Writer;
+}
+
 ///////////////////////////////////////////////////////////////////////////////////////////////////////////////////////
