Fixed missing FPlatformTime include reported by CIS

diff --git a/Engine/Source/Runtime/Core/Private/Windows/WindowsPlatformCrashContext.cpp b/Engine/Source/Runtime/Core/Private/Windows/WindowsPlatformCrashContext.cpp
--- a/Engine/Source/Runtime/Core/Private/Windows/WindowsPlatformCrashContext.cpp
+++ b/Engine/Source/Runtime/Core/Private/Windows/WindowsPlatformCrashContext.cpp
@@ -11,6 +11,7 @@
 #include "HAL/FileManager.h"
 #include "HAL/PlatformOutputDevices.h"
 #include "HAL/PlatformTLS.h"
+#include "HAL/PlatformTime.h"
 #include "Internationalization/Internationalization.h"
 #include "Misc/App.h"
 #include "Misc/Paths.h"
