Added FBox::Equals()

diff --git a/Engine/Source/Runtime/Core/Public/Math/Box.h b/Engine/Source/Runtime/Core/Public/Math/Box.h
--- a/Engine/Source/Runtime/Core/Public/Math/Box.h
+++ b/Engine/Source/Runtime/Core/Public/Math/Box.h
@@ -81,6 +81,28 @@ public:
 		return (Min == Other.Min) && (Max == Other.Max);
 	}
 
+	/**
+	 * Compares two boxes for inequality.
+	 *
+	 * @return false if the boxes are equal, true otherwise.
+	 */
+	FORCEINLINE bool operator!=( const FBox& Other) const
+	{
+		return (Min != Other.Min) || (Max != Other.Max);
+	}
+
+	/**
+	 * Check against another box for equality, within specified error limits.
+	 *
+	 * @param Other The box to check against.
+	 * @param Tolerance Error tolerance.
+	 * @return true if the boxes are equal within tolerance limits, false otherwise.
+	 */
+	bool Equals(const FBox& Other, float Tolerance=KINDA_SMALL_NUMBER) const
+	{
+		return Min.Equals(Other.Min, Tolerance) && Max.Equals(Other.Max, Tolerance);
+	}
+
 	/**
 	 * Adds to this bounding box to include a given point.
 	 *
