fix test_editor_level_library (python) tab error
+ add test context in log error

#rnx

diff --git a/Engine/Source/Runtime/Core/Private/Misc/AutomationTest.cpp b/Engine/Source/Runtime/Core/Private/Misc/AutomationTest.cpp
--- a/Engine/Source/Runtime/Core/Private/Misc/AutomationTest.cpp
+++ b/Engine/Source/Runtime/Core/Private/Misc/AutomationTest.cpp
@@ -103,7 +103,7 @@ void FAutomationTestFramework::FAutomationTestOutputDevice::Serialize( const TCH
 			{
 				if (!LoggedFailureCause.Contains(CurTest))
 				{
-					CurTest->AddError(FString::Printf(TEXT("%s will be marked as failing due to errors being logged"), *CurTest->GetBeautifiedTestName()), STACK_OFFSET);
+					CurTest->AddError(FString::Printf(TEXT("%s %s will be marked as failing due to errors being logged"), *CurTest->GetBeautifiedTestName(), *CurTest->GetTestContext()), STACK_OFFSET);
 					LoggedFailureCause.Add(CurTest);
 				}
 				CurTest->AddError(FString(V), STACK_OFFSET);
