SharedBuffer: Added MakeSharedBufferFromArray to construct a FSharedBuffer by taking ownership of an TArray

#rnx

diff --git a/Engine/Source/Runtime/Core/Private/Tests/Memory/SharedBufferTest.cpp b/Engine/Source/Runtime/Core/Private/Tests/Memory/SharedBufferTest.cpp
--- a/Engine/Source/Runtime/Core/Private/Tests/Memory/SharedBufferTest.cpp
+++ b/Engine/Source/Runtime/Core/Private/Tests/Memory/SharedBufferTest.cpp
@@ -450,4 +450,18 @@ bool FBufferOwnerTest::RunTest(const FString& Parameters)
 	return true;
 }
 
+IMPLEMENT_SIMPLE_AUTOMATION_TEST(FSharedBufferFromArrayTest, "System.Core.Memory.SharedBufferFromArray", EAutomationTestFlags::ApplicationContextMask | EAutomationTestFlags::SmokeFilter)
+bool FSharedBufferFromArrayTest::RunTest(const FString& Parameters)
+{
+	TArray<uint16> Array{1, 2, 3, 4, 5, 6, 7, 8};
+	const void* const ExpectedData = Array.GetData();
+	const uint64 ExpectedSize = Array.Num() * sizeof(uint16);
+	const FSharedBuffer Buffer = MakeSharedBufferFromArray(MoveTemp(Array));
+	TestTrue(TEXT("MakeSharedBufferFromArray -> Array.IsEmpty"), Array.IsEmpty());
+	TestEqual(TEXT("MakeSharedBufferFromArray -> Buffer.GetData()"), Buffer.GetData(), ExpectedData);
+	TestEqual(TEXT("MakeSharedBufferFromArray -> Buffer.GetSize()"), Buffer.GetSize(), ExpectedSize);
+
+	return true;
+}
+
 #endif // WITH_DEV_AUTOMATION_TESTS
diff --git a/Engine/Source/Runtime/Core/Public/Memory/SharedBuffer.h b/Engine/Source/Runtime/Core/Public/Memory/SharedBuffer.h
--- a/Engine/Source/Runtime/Core/Public/Memory/SharedBuffer.h
+++ b/Engine/Source/Runtime/Core/Public/Memory/SharedBuffer.h
@@ -3,6 +3,7 @@
 #pragma once
 
 #include "CoreTypes.h"
+#include "Containers/ContainersFwd.h"
 #include "Memory/MemoryView.h"
 #include "Misc/EnumClassFlags.h"
 #include "Templates/Invoke.h"
@@ -807,3 +808,38 @@ inline bool TBufferOwnerPtr<FOps>::operator!=(const TBufferOwnerPtr<FOtherOps>&
 } // BufferOwnerPrivate
 
 ///////////////////////////////////////////////////////////////////////////////////////////////////////////////////////
+
+namespace BufferOwnerPrivate
+{
+
+template <typename T, typename Allocator>
+class TBufferOwnerTArray final : public FBufferOwner
+{
+public:
+	explicit TBufferOwnerTArray(TArray<T, Allocator>&& InArray)
+		: Array(MoveTemp(InArray))
+	{
+		SetBuffer(Array.GetData(), uint64(Array.Num()) * sizeof(T));
+		SetIsMaterialized();
+		SetIsOwned();
+	}
+
+private:
+	virtual void FreeBuffer() final
+	{
+		Array.Empty();
+	}
+
+	TArray<T, Allocator> Array;
+};
+
+} // BufferOwnerPrivate
+
+/** Construct a shared buffer by taking ownership of an array. */
+template <typename T, typename Allocator>
+FSharedBuffer MakeSharedBufferFromArray(TArray<T, Allocator>&& Array)
+{
+	return FSharedBuffer(new BufferOwnerPrivate::TBufferOwnerTArray<T, Allocator>(MoveTemp(Array)));
+}
+
+///////////////////////////////////////////////////////////////////////////////////////////////////////////////////////
