CompactBinary: Renamed view types to use a View suffix, and remove the Ref suffix from types that hold a SharedBuffer

#rnx

diff --git a/Engine/Source/Developer/DerivedDataCache/Private/DerivedDataCache.cpp b/Engine/Source/Developer/DerivedDataCache/Private/DerivedDataCache.cpp
--- a/Engine/Source/Developer/DerivedDataCache/Private/DerivedDataCache.cpp
+++ b/Engine/Source/Developer/DerivedDataCache/Private/DerivedDataCache.cpp
@@ -859,8 +859,8 @@ void FCache::Put(
 			Writer.EndArray();
 		}
 
-		const FCbObjectRef& Meta = Record.GetMeta();
-		if (Meta.CreateIterator())
+		const FCbObject& Meta = Record.GetMeta();
+		if (Meta.CreateViewIterator())
 		{
 			Writer.AddObject("Meta"_ASV, Meta);
 		}
@@ -943,13 +943,13 @@ void FCache::Get(
 	}
 
 	// Read the record from its compact binary fields.
-	const FCbObject RecordObject(RecordData.GetData());
-	const FCbObject MetaObject = RecordObject["Meta"_ASV].AsObject();
-	const FCbObject ValueObject = RecordObject["Value"_ASV].AsObject();
-	const FCbArray AttachmentsArray = RecordObject["Attachments"_ASV].AsArray();
+	const FCbObjectView RecordObject(RecordData.GetData());
+	const FCbObjectView MetaObject = RecordObject["Meta"_ASV].AsObjectView();
+	const FCbObjectView ValueObject = RecordObject["Value"_ASV].AsObjectView();
+	const FCbArrayView AttachmentsArray = RecordObject["Attachments"_ASV].AsArrayView();
 
 	// Read the value and attachments if they have been requested.
-	const auto GetPayload = [&SkippedPayloads, &Key, Policy, Context](FCbObject PayloadObject, ECachePolicy SkipFlag) -> FCachePayload
+	const auto GetPayload = [&SkippedPayloads, &Key, Policy, Context](FCbObjectView PayloadObject, ECachePolicy SkipFlag) -> FCachePayload
 	{
 		TArray<uint8> PayloadData;
 		const FCachePayloadId PayloadId = FCachePayloadId::FromObjectId(PayloadObject["Id"_ASV].AsObjectId());
@@ -982,15 +982,15 @@ void FCache::Get(
 		}
 	};
 
-	if (ValueObject.CreateIterator() && !(Value = GetPayload(ValueObject, ECachePolicy::SkipValue)))
+	if (ValueObject.CreateViewIterator() && !(Value = GetPayload(ValueObject, ECachePolicy::SkipValue)))
 	{
 		return;
 	}
 
 	Attachments.Reserve(AttachmentsArray.Num());
-	for (FCbField AttachmentField : AttachmentsArray)
+	for (FCbFieldView AttachmentField : AttachmentsArray)
 	{
-		Attachments.Add(GetPayload(AttachmentField.AsObject(), ECachePolicy::SkipAttachments));
+		Attachments.Add(GetPayload(AttachmentField.AsObjectView(), ECachePolicy::SkipAttachments));
 		if (!Attachments.Last())
 		{
 			return;
@@ -1016,9 +1016,9 @@ void FCache::Get(
 	}
 
 	// Populate the builder now that any error conditions have been handled.
-	if (!EnumHasAnyFlags(Policy, ECachePolicy::SkipMeta) && MetaObject.CreateIterator())
+	if (!EnumHasAnyFlags(Policy, ECachePolicy::SkipMeta) && MetaObject.CreateViewIterator())
 	{
-		Builder.SetMeta(FCbObjectRef::Clone(MetaObject));
+		Builder.SetMeta(FCbObject::Clone(MetaObject));
 	}
 	if (Value)
 	{
@@ -1104,9 +1104,9 @@ void FCache::GetPayload(
 			}
 
 			// Read the record from its compact binary fields.
-			const FCbObject RecordObject(RecordData.GetData());
-			const FCbObject ValueObject = RecordObject["Value"_ASV].AsObject();
-			const FCbArray AttachmentsArray = RecordObject["Attachments"_ASV].AsArray();
+			const FCbObjectView RecordObject(RecordData.GetData());
+			const FCbObjectView ValueObject = RecordObject["Value"_ASV].AsObjectView();
+			const FCbArrayView AttachmentsArray = RecordObject["Attachments"_ASV].AsArrayView();
 
 			// Find the hash of the payload in the record.
 			const FCbObjectId KeyObjectId = Key.GetId().ToObjectId();
@@ -1116,11 +1116,11 @@ void FCache::GetPayload(
 			}
 			else
 			{
-				for (FCbField AttachmentField : AttachmentsArray)
+				for (FCbFieldView AttachmentField : AttachmentsArray)
 				{
-					if (AttachmentField.AsObject()["Id"_ASV].AsObjectId() == KeyObjectId)
+					if (AttachmentField.AsObjectView()["Id"_ASV].AsObjectId() == KeyObjectId)
 					{
-						Params.Payload = FCachePayload(Key.GetId(), AttachmentField.AsObject()["Hash"_ASV].AsBinaryAttachment());
+						Params.Payload = FCachePayload(Key.GetId(), AttachmentField.AsObjectView()["Hash"_ASV].AsBinaryAttachment());
 						break;
 					}
 				}
diff --git a/Engine/Source/Developer/DerivedDataCache/Private/DerivedDataCacheRecord.cpp b/Engine/Source/Developer/DerivedDataCache/Private/DerivedDataCacheRecord.cpp
--- a/Engine/Source/Developer/DerivedDataCache/Private/DerivedDataCacheRecord.cpp
+++ b/Engine/Source/Developer/DerivedDataCache/Private/DerivedDataCacheRecord.cpp
@@ -74,7 +74,7 @@ FSharedBuffer FCachePayload::Decompress(FSharedBuffer Buffer)
 struct FCacheRecordBuilderData
 {
 	FCacheKey Key;
-	FCbObjectRef Meta;
+	FCbObject Meta;
 	FCachePayload Value;
 	TArray<FCachePayload> Attachments;
 };
@@ -87,7 +87,7 @@ struct FCacheRecordData
 	explicit FCacheRecordData(FCacheRecordBuilderData&& Builder);
 
 	FCacheKey Key;
-	FCbObjectRef Meta;
+	FCbObject Meta;
 	FCachePayload Value;
 	FSharedBuffer ValueCache;
 	TArray<FCachePayload> Attachments;
@@ -143,7 +143,7 @@ const FCacheKey& FCacheRecord::GetKey() const
 	return (Data ? *Data : GetEmptyCacheRecord()).Key;
 }
 
-const FCbObjectRef& FCacheRecord::GetMeta() const
+const FCbObject& FCacheRecord::GetMeta() const
 {
 	return (Data ? *Data : GetEmptyCacheRecord()).Meta;
 }
@@ -215,7 +215,7 @@ void FCacheRecordBuilder::SetKey(const FCacheKey& Key)
 	Data->Key = Key;
 }
 
-void FCacheRecordBuilder::SetMeta(FCbObjectRef Meta)
+void FCacheRecordBuilder::SetMeta(FCbObject Meta)
 {
 	Meta.MakeOwned();
 	Data->Meta = MoveTemp(Meta);
diff --git a/Engine/Source/Developer/DerivedDataCache/Public/DerivedDataCacheRecord.h b/Engine/Source/Developer/DerivedDataCache/Public/DerivedDataCacheRecord.h
--- a/Engine/Source/Developer/DerivedDataCache/Public/DerivedDataCacheRecord.h
+++ b/Engine/Source/Developer/DerivedDataCache/Public/DerivedDataCacheRecord.h
@@ -114,7 +114,7 @@ public:
 	UE_API const FCacheKey& GetKey() const;
 
 	/** Returns the metadata for the cache record unless it was skipped in the request. */
-	UE_API const FCbObjectRef& GetMeta() const;
+	UE_API const FCbObject& GetMeta() const;
 
 	/** Returns the value. Null if the value was skipped in the request. */
 	UE_API FSharedBuffer GetValue() const;
@@ -183,7 +183,7 @@ public:
 	 *
 	 * The metadata is cloned if not owned.
 	 */
-	UE_API void SetMeta(FCbObjectRef Meta);
+	UE_API void SetMeta(FCbObject Meta);
 
 	/**
 	 * Set the value for the cache record.
diff --git a/Engine/Source/Editor/UnrealEd/Private/EditorDomain/EditorDomain.cpp b/Engine/Source/Editor/UnrealEd/Private/EditorDomain/EditorDomain.cpp
--- a/Engine/Source/Editor/UnrealEd/Private/EditorDomain/EditorDomain.cpp
+++ b/Engine/Source/Editor/UnrealEd/Private/EditorDomain/EditorDomain.cpp
@@ -185,8 +185,8 @@ int64 FEditorDomain::FileSize(const FPackagePath& PackagePath, EPackageSegment P
 			if ((PackageSource->Source == FEditorDomain::EPackageSource::Undecided || PackageSource->Source == FEditorDomain::EPackageSource::Editor) &&
 				Params.Status == UE::DerivedData::ECacheStatus::Cached)
 			{
-				const FCbObjectRef& MetaData = Params.Record.GetMeta();
-				FCbField FileSizeField = MetaData["FileSize"];
+				const FCbObject& MetaData = Params.Record.GetMeta();
+				FCbFieldView FileSizeField = MetaData["FileSize"];
 				FileSize = FileSizeField.AsInt64();
 				PackageSource->Source = EPackageSource::Editor;
 			}
diff --git a/Engine/Source/Editor/UnrealEd/Private/EditorDomain/EditorDomainArchive.cpp b/Engine/Source/Editor/UnrealEd/Private/EditorDomain/EditorDomainArchive.cpp
--- a/Engine/Source/Editor/UnrealEd/Private/EditorDomain/EditorDomainArchive.cpp
+++ b/Engine/Source/Editor/UnrealEd/Private/EditorDomain/EditorDomainArchive.cpp
@@ -183,7 +183,7 @@ void FEditorDomainReadArchive::OnCacheRequestComplete(UE::DerivedData::FCacheGet
 	if ((PackageSource->Source == FEditorDomain::EPackageSource::Undecided || PackageSource->Source == FEditorDomain::EPackageSource::Editor) &&
 		Params.Status == UE::DerivedData::ECacheStatus::Cached)
 	{
-		const FCbObjectRef& MetaData = Params.Record.GetMeta();
+		const FCbObject& MetaData = Params.Record.GetMeta();
 		int64 FileSize = MetaData["FileSize"].AsInt64(-1);
 		FSharedBuffer LocalBytes = Params.Record.GetValue();
 		if (static_cast<int64>(LocalBytes.GetSize()) != FileSize)
@@ -412,7 +412,7 @@ void FEditorDomainAsyncReadFileHandle::OnCacheRequestComplete(UE::DerivedData::F
 	if ((PackageSource->Source == FEditorDomain::EPackageSource::Undecided || PackageSource->Source == FEditorDomain::EPackageSource::Editor) &&
 		Params.Status == UE::DerivedData::ECacheStatus::Cached)
 	{
-		const FCbObjectRef& MetaData = Params.Record.GetMeta();
+		const FCbObject& MetaData = Params.Record.GetMeta();
 		int64 FileSize = MetaData["FileSize"].AsInt64(-1);
 		FSharedBuffer LocalBytes = Params.Record.GetValue();
 		if (static_cast<int64>(LocalBytes.GetSize()) != FileSize)
diff --git a/Engine/Source/Editor/UnrealEd/Private/EditorDomain/EditorDomainSave.cpp b/Engine/Source/Editor/UnrealEd/Private/EditorDomain/EditorDomainSave.cpp
--- a/Engine/Source/Editor/UnrealEd/Private/EditorDomain/EditorDomainSave.cpp
+++ b/Engine/Source/Editor/UnrealEd/Private/EditorDomain/EditorDomainSave.cpp
@@ -764,7 +764,7 @@ bool FEditorDomainSaveServer::TrySavePackage(const FPackagePath& PackagePath, FS
 	MetaData.BeginObject();
 	MetaData << "FileSize" << TempBytes.Num();
 	MetaData.EndObject();
-	RecordBuilder.SetMeta(MetaData.Save().AsObjectRef());
+	RecordBuilder.SetMeta(MetaData.Save().AsObject());
 	RecordBuilder.SetValue(FSharedBuffer::MakeView(TempBytes.GetData(), TempBytes.Num()));
 	UE::DerivedData::FCacheRecord Record = RecordBuilder.Build();
 	GetDerivedDataCacheRef().GetCache().Put(MakeArrayView(&Record, 1), PackagePath.GetDebugName());
diff --git a/Engine/Source/Runtime/Core/Private/Serialization/CompactBinary.cpp b/Engine/Source/Runtime/Core/Private/Serialization/CompactBinary.cpp
--- a/Engine/Source/Runtime/Core/Private/Serialization/CompactBinary.cpp
+++ b/Engine/Source/Runtime/Core/Private/Serialization/CompactBinary.cpp
@@ -93,7 +93,7 @@ void FCbObjectId::ToString(FWideStringBuilderBase& Builder) const
 
 ///////////////////////////////////////////////////////////////////////////////////////////////////////////////////////
 
-FCbField::FCbField(const void* const InData, const ECbFieldType InType)
+FCbFieldView::FCbFieldView(const void* const InData, const ECbFieldType InType)
 {
 	const uint8* Bytes = static_cast<const uint8*>(InData);
 	const ECbFieldType LocalType = FCbFieldType::HasFieldType(InType) ? (ECbFieldType(*Bytes++) | ECbFieldType::HasFieldType) : InType;
@@ -106,35 +106,35 @@ FCbField::FCbField(const void* const InData, const ECbFieldType InType)
 	Payload = Bytes;
 }
 
-FCbObject FCbField::AsObject()
+FCbObjectView FCbFieldView::AsObjectView()
 {
 	if (FCbFieldType::IsObject(Type))
 	{
 		Error = ECbFieldError::None;
-		return FCbObject::FromFieldNoCheck(*this);
+		return FCbObjectView::FromFieldNoCheck(*this);
 	}
 	else
 	{
 		Error = ECbFieldError::TypeError;
-		return FCbObject();
+		return FCbObjectView();
 	}
 }
 
-FCbArray FCbField::AsArray()
+FCbArrayView FCbFieldView::AsArrayView()
 {
 	if (FCbFieldType::IsArray(Type))
 	{
 		Error = ECbFieldError::None;
-		return FCbArray::FromFieldNoCheck(*this);
+		return FCbArrayView::FromFieldNoCheck(*this);
 	}
 	else
 	{
 		Error = ECbFieldError::TypeError;
-		return FCbArray();
+		return FCbArrayView();
 	}
 }
 
-FMemoryView FCbField::AsBinary(const FMemoryView Default)
+FMemoryView FCbFieldView::AsBinaryView(const FMemoryView Default)
 {
 	if (FCbFieldType::IsBinary(Type))
 	{
@@ -152,7 +152,7 @@ FMemoryView FCbField::AsBinary(const FMemoryView Default)
 	}
 }
 
-FAnsiStringView FCbField::AsString(const FAnsiStringView Default)
+FAnsiStringView FCbFieldView::AsString(const FAnsiStringView Default)
 {
 	if (FCbFieldType::IsString(Type))
 	{
@@ -176,7 +176,7 @@ FAnsiStringView FCbField::AsString(const FAnsiStringView Default)
 	}
 }
 
-uint64 FCbField::AsInteger(const uint64 Default, const FIntegerParams Params)
+uint64 FCbFieldView::AsInteger(const uint64 Default, const FIntegerParams Params)
 {
 	if (FCbFieldType::IsInteger(Type))
 	{
@@ -201,7 +201,7 @@ uint64 FCbField::AsInteger(const uint64 Default, const FIntegerParams Params)
 	}
 }
 
-float FCbField::AsFloat(const float Default)
+float FCbFieldView::AsFloat(const float Default)
 {
 	switch (FCbFieldType::GetType(Type))
 	{
@@ -232,7 +232,7 @@ float FCbField::AsFloat(const float Default)
 	}
 }
 
-double FCbField::AsDouble(const double Default)
+double FCbFieldView::AsDouble(const double Default)
 {
 	switch (FCbFieldType::GetType(Type))
 	{
@@ -266,7 +266,7 @@ double FCbField::AsDouble(const double Default)
 	}
 }
 
-bool FCbField::AsBool(const bool bDefault)
+bool FCbFieldView::AsBool(const bool bDefault)
 {
 	const ECbFieldType LocalType = Type;
 	const bool bIsBool = FCbFieldType::IsBool(LocalType);
@@ -274,7 +274,7 @@ bool FCbField::AsBool(const bool bDefault)
 	return (uint8(bIsBool) & uint8(LocalType) & 1) | ((!bIsBool) & bDefault);
 }
 
-FIoHash FCbField::AsCompactBinaryAttachment(const FIoHash& Default)
+FIoHash FCbFieldView::AsCompactBinaryAttachment(const FIoHash& Default)
 {
 	if (FCbFieldType::IsCompactBinaryAttachment(Type))
 	{
@@ -288,7 +288,7 @@ FIoHash FCbField::AsCompactBinaryAttachment(const FIoHash& Default)
 	}
 }
 
-FIoHash FCbField::AsBinaryAttachment(const FIoHash& Default)
+FIoHash FCbFieldView::AsBinaryAttachment(const FIoHash& Default)
 {
 	if (FCbFieldType::IsBinaryAttachment(Type))
 	{
@@ -302,7 +302,7 @@ FIoHash FCbField::AsBinaryAttachment(const FIoHash& Default)
 	}
 }
 
-FIoHash FCbField::AsAttachment(const FIoHash& Default)
+FIoHash FCbFieldView::AsAttachment(const FIoHash& Default)
 {
 	if (FCbFieldType::IsAttachment(Type))
 	{
@@ -316,7 +316,7 @@ FIoHash FCbField::AsAttachment(const FIoHash& Default)
 	}
 }
 
-FIoHash FCbField::AsHash(const FIoHash& Default)
+FIoHash FCbFieldView::AsHash(const FIoHash& Default)
 {
 	if (FCbFieldType::IsHash(Type))
 	{
@@ -330,12 +330,12 @@ FIoHash FCbField::AsHash(const FIoHash& Default)
 	}
 }
 
-FGuid FCbField::AsUuid()
+FGuid FCbFieldView::AsUuid()
 {
 	return AsUuid(FGuid());
 }
 
-FGuid FCbField::AsUuid(const FGuid& Default)
+FGuid FCbFieldView::AsUuid(const FGuid& Default)
 {
 	if (FCbFieldType::IsUuid(Type))
 	{
@@ -355,7 +355,7 @@ FGuid FCbField::AsUuid(const FGuid& Default)
 	}
 }
 
-int64 FCbField::AsDateTimeTicks(const int64 Default)
+int64 FCbFieldView::AsDateTimeTicks(const int64 Default)
 {
 	if (FCbFieldType::IsDateTime(Type))
 	{
@@ -369,17 +369,17 @@ int64 FCbField::AsDateTimeTicks(const int64 Default)
 	}
 }
 
-FDateTime FCbField::AsDateTime()
+FDateTime FCbFieldView::AsDateTime()
 {
 	return FDateTime(AsDateTimeTicks(0));
 }
 
-FDateTime FCbField::AsDateTime(FDateTime Default)
+FDateTime FCbFieldView::AsDateTime(FDateTime Default)
 {
 	return FDateTime(AsDateTimeTicks(Default.GetTicks()));
 }
 
-int64 FCbField::AsTimeSpanTicks(const int64 Default)
+int64 FCbFieldView::AsTimeSpanTicks(const int64 Default)
 {
 	if (FCbFieldType::IsTimeSpan(Type))
 	{
@@ -393,17 +393,17 @@ int64 FCbField::AsTimeSpanTicks(const int64 Default)
 	}
 }
 
-FTimespan FCbField::AsTimeSpan()
+FTimespan FCbFieldView::AsTimeSpan()
 {
 	return FTimespan(AsTimeSpanTicks(0));
 }
 
-FTimespan FCbField::AsTimeSpan(FTimespan Default)
+FTimespan FCbFieldView::AsTimeSpan(FTimespan Default)
 {
 	return FTimespan(AsTimeSpanTicks(Default.GetTicks()));
 }
 
-FCbObjectId FCbField::AsObjectId(const FCbObjectId& Default)
+FCbObjectId FCbFieldView::AsObjectId(const FCbObjectId& Default)
 {
 	static_assert(sizeof(FCbObjectId) == 12, "FCbObjectId is expected to be 12 bytes.");
 	if (FCbFieldType::IsObjectId(Type))
@@ -418,7 +418,7 @@ FCbObjectId FCbField::AsObjectId(const FCbObjectId& Default)
 	}
 }
 
-FCbCustomById FCbField::AsCustomById(FCbCustomById Default)
+FCbCustomById FCbFieldView::AsCustomById(FCbCustomById Default)
 {
 	if (FCbFieldType::IsCustomById(Type))
 	{
@@ -441,7 +441,7 @@ FCbCustomById FCbField::AsCustomById(FCbCustomById Default)
 	}
 }
 
-FCbCustomByName FCbField::AsCustomByName(FCbCustomByName Default)
+FCbCustomByName FCbFieldView::AsCustomByName(FCbCustomByName Default)
 {
 	if (FCbFieldType::IsCustomByName(Type))
 	{
@@ -469,7 +469,7 @@ FCbCustomByName FCbField::AsCustomByName(FCbCustomByName Default)
 	}
 }
 
-FMemoryView FCbField::AsCustom(uint64 Id, FMemoryView Default)
+FMemoryView FCbFieldView::AsCustom(uint64 Id, FMemoryView Default)
 {
 	FCbCustomById Custom = AsCustomById(FCbCustomById{Id, Default});
 	if (Custom.Id == Id)
@@ -483,7 +483,7 @@ FMemoryView FCbField::AsCustom(uint64 Id, FMemoryView Default)
 	}
 }
 
-FMemoryView FCbField::AsCustom(FAnsiStringView Name, FMemoryView Default)
+FMemoryView FCbFieldView::AsCustom(FAnsiStringView Name, FMemoryView Default)
 {
 	const FCbCustomByName Custom = AsCustomByName(FCbCustomByName{Name, Default});
 	if (Custom.Name.Equals(Name, ESearchCase::CaseSensitive))
@@ -497,12 +497,12 @@ FMemoryView FCbField::AsCustom(FAnsiStringView Name, FMemoryView Default)
 	}
 }
 
-uint64 FCbField::GetSize() const
+uint64 FCbFieldView::GetSize() const
 {
 	return sizeof(ECbFieldType) + GetViewNoType().GetSize();
 }
 
-uint64 FCbField::GetPayloadSize() const
+uint64 FCbFieldView::GetPayloadSize() const
 {
 	switch (FCbFieldType::GetType(Type))
 	{
@@ -548,27 +548,27 @@ uint64 FCbField::GetPayloadSize() const
 	}
 }
 
-FIoHash FCbField::GetHash() const
+FIoHash FCbFieldView::GetHash() const
 {
 	FBlake3 Hash;
 	GetHash(Hash);
 	return FIoHash(Hash.Finalize());
 }
 
-void FCbField::GetHash(FBlake3& Hash) const
+void FCbFieldView::GetHash(FBlake3& Hash) const
 {
 	const ECbFieldType SerializedType = FCbFieldType::GetSerializedType(Type);
 	Hash.Update(&SerializedType, sizeof(SerializedType));
 	Hash.Update(GetViewNoType());
 }
 
-bool FCbField::Equals(const FCbField& Other) const
+bool FCbFieldView::Equals(const FCbFieldView& Other) const
 {
 	return FCbFieldType::GetSerializedType(Type) == FCbFieldType::GetSerializedType(Other.Type) &&
 		GetViewNoType().EqualBytes(Other.GetViewNoType());
 }
 
-void FCbField::CopyTo(FMutableMemoryView Buffer) const
+void FCbFieldView::CopyTo(FMutableMemoryView Buffer) const
 {
 	const FMemoryView Source = GetViewNoType();
 	checkf(Buffer.GetSize() == sizeof(ECbFieldType) + Source.GetSize(),
@@ -579,7 +579,7 @@ void FCbField::CopyTo(FMutableMemoryView Buffer) const
 	FMemory::Memcpy(Buffer.GetData(), Source.GetData(), Source.GetSize());
 }
 
-void FCbField::CopyTo(FArchive& Ar) const
+void FCbFieldView::CopyTo(FArchive& Ar) const
 {
 	check(Ar.IsSaving());
 	const FMemoryView Source = GetViewNoType();
@@ -588,16 +588,16 @@ void FCbField::CopyTo(FArchive& Ar) const
 	Ar.Serialize(const_cast<void*>(Source.GetData()), static_cast<int64>(Source.GetSize()));
 }
 
-void FCbField::IterateAttachments(FCbFieldVisitor Visitor) const
+void FCbFieldView::IterateAttachments(FCbFieldVisitor Visitor) const
 {
 	switch (FCbFieldType::GetType(Type))
 	{
 	case ECbFieldType::Object:
 	case ECbFieldType::UniformObject:
-		return FCbObject::FromFieldNoCheck(*this).IterateAttachments(Visitor);
+		return FCbObjectView::FromFieldNoCheck(*this).IterateAttachments(Visitor);
 	case ECbFieldType::Array:
 	case ECbFieldType::UniformArray:
-		return FCbArray::FromFieldNoCheck(*this).IterateAttachments(Visitor);
+		return FCbArrayView::FromFieldNoCheck(*this).IterateAttachments(Visitor);
 	case ECbFieldType::CompactBinaryAttachment:
 	case ECbFieldType::BinaryAttachment:
 		return Visitor(*this);
@@ -606,7 +606,7 @@ void FCbField::IterateAttachments(FCbFieldVisitor Visitor) const
 	}
 }
 
-FMemoryView FCbField::GetView() const
+FMemoryView FCbFieldView::GetView() const
 {
 	const uint32 TypeSize = FCbFieldType::HasFieldType(Type) ? sizeof(ECbFieldType) : 0;
 	const uint32 NameSize = FCbFieldType::HasFieldName(Type) ? NameLen + MeasureVarUInt(NameLen) : 0;
@@ -614,7 +614,7 @@ FMemoryView FCbField::GetView() const
 	return MakeMemoryView(static_cast<const uint8*>(Payload) - TypeSize - NameSize, TypeSize + NameSize + PayloadSize);
 }
 
-FMemoryView FCbField::GetViewNoType() const
+FMemoryView FCbFieldView::GetViewNoType() const
 {
 	const uint32 NameSize = FCbFieldType::HasFieldName(Type) ? NameLen + MeasureVarUInt(NameLen) : 0;
 	const uint64 PayloadSize = GetPayloadSize();
@@ -623,12 +623,12 @@ FMemoryView FCbField::GetViewNoType() const
 
 ///////////////////////////////////////////////////////////////////////////////////////////////////////////////////////
 
-FCbArray::FCbArray()
-	: FCbField(CompactBinaryPrivate::GEmptyArrayPayload)
+FCbArrayView::FCbArrayView()
+	: FCbFieldView(CompactBinaryPrivate::GEmptyArrayPayload)
 {
 }
 
-uint64 FCbArray::Num() const
+uint64 FCbArrayView::Num() const
 {
 	const uint8* PayloadBytes = static_cast<const uint8*>(GetPayload());
 	PayloadBytes += MeasureVarUInt(PayloadBytes);
@@ -636,7 +636,7 @@ uint64 FCbArray::Num() const
 	return ReadVarUInt(PayloadBytes, NumByteCount);
 }
 
-FCbFieldIterator FCbArray::CreateIterator() const
+FCbFieldViewIterator FCbArrayView::CreateViewIterator() const
 {
 	const uint8* PayloadBytes = static_cast<const uint8*>(GetPayload());
 	uint32 PayloadSizeByteCount;
@@ -649,37 +649,37 @@ FCbFieldIterator FCbArray::CreateIterator() const
 		PayloadBytes += NumByteCount;
 		const ECbFieldType UniformType = FCbFieldType::GetType(GetType()) == ECbFieldType::UniformArray ?
 			ECbFieldType(*PayloadBytes++) : ECbFieldType::HasFieldType;
-		return FCbFieldIterator::MakeRange(MakeMemoryView(PayloadBytes, PayloadEnd), UniformType);
+		return FCbFieldViewIterator::MakeRange(MakeMemoryView(PayloadBytes, PayloadEnd), UniformType);
 	}
-	return FCbFieldIterator();
+	return FCbFieldViewIterator();
 }
 
-uint64 FCbArray::GetSize() const
+uint64 FCbArrayView::GetSize() const
 {
 	return sizeof(ECbFieldType) + GetPayloadSize();
 }
 
-FIoHash FCbArray::GetHash() const
+FIoHash FCbArrayView::GetHash() const
 {
 	FBlake3 Hash;
 	GetHash(Hash);
 	return FIoHash(Hash.Finalize());
 }
 
-void FCbArray::GetHash(FBlake3& Hash) const
+void FCbArrayView::GetHash(FBlake3& Hash) const
 {
 	const ECbFieldType SerializedType = FCbFieldType::GetType(GetType());
 	Hash.Update(&SerializedType, sizeof(SerializedType));
 	Hash.Update(GetPayloadView());
 }
 
-bool FCbArray::Equals(const FCbArray& Other) const
+bool FCbArrayView::Equals(const FCbArrayView& Other) const
 {
 	return FCbFieldType::GetType(GetType()) == FCbFieldType::GetType(Other.GetType()) &&
 		GetPayloadView().EqualBytes(Other.GetPayloadView());
 }
 
-void FCbArray::CopyTo(FMutableMemoryView Buffer) const
+void FCbArrayView::CopyTo(FMutableMemoryView Buffer) const
 {
 	const FMemoryView Source = GetPayloadView();
 	checkf(Buffer.GetSize() == sizeof(ECbFieldType) + Source.GetSize(),
@@ -690,7 +690,7 @@ void FCbArray::CopyTo(FMutableMemoryView Buffer) const
 	FMemory::Memcpy(Buffer.GetData(), Source.GetData(), Source.GetSize());
 }
 
-void FCbArray::CopyTo(FArchive& Ar) const
+void FCbArrayView::CopyTo(FArchive& Ar) const
 {
 	check(Ar.IsSaving());
 	const FMemoryView Source = GetPayloadView();
@@ -701,12 +701,12 @@ void FCbArray::CopyTo(FArchive& Ar) const
 
 ///////////////////////////////////////////////////////////////////////////////////////////////////////////////////////
 
-FCbObject::FCbObject()
-	: FCbField(CompactBinaryPrivate::GEmptyObjectPayload)
+FCbObjectView::FCbObjectView()
+	: FCbFieldView(CompactBinaryPrivate::GEmptyObjectPayload)
 {
 }
 
-FCbFieldIterator FCbObject::CreateIterator() const
+FCbFieldViewIterator FCbObjectView::CreateViewIterator() const
 {
 	const uint8* PayloadBytes = static_cast<const uint8*>(GetPayload());
 	uint32 PayloadSizeByteCount;
@@ -717,61 +717,61 @@ FCbFieldIterator FCbObject::CreateIterator() const
 		const void* const PayloadEnd = PayloadBytes + PayloadSize;
 		const ECbFieldType UniformType = FCbFieldType::GetType(GetType()) == ECbFieldType::UniformObject ?
 			ECbFieldType(*PayloadBytes++) : ECbFieldType::HasFieldType;
-		return FCbFieldIterator::MakeRange(MakeMemoryView(PayloadBytes, PayloadEnd), UniformType);
+		return FCbFieldViewIterator::MakeRange(MakeMemoryView(PayloadBytes, PayloadEnd), UniformType);
 	}
-	return FCbFieldIterator();
+	return FCbFieldViewIterator();
 }
 
-FCbField FCbObject::Find(const FAnsiStringView Name) const
+FCbFieldView FCbObjectView::FindView(const FAnsiStringView Name) const
 {
-	for (const FCbField Field : *this)
+	for (const FCbFieldView Field : *this)
 	{
 		if (Name.Equals(Field.GetName(), ESearchCase::CaseSensitive))
 		{
 			return Field;
 		}
 	}
-	return FCbField();
+	return FCbFieldView();
 }
 
-FCbField FCbObject::FindIgnoreCase(const FAnsiStringView Name) const
+FCbFieldView FCbObjectView::FindViewIgnoreCase(const FAnsiStringView Name) const
 {
-	for (const FCbField Field : *this)
+	for (const FCbFieldView Field : *this)
 	{
 		if (Name.Equals(Field.GetName(), ESearchCase::IgnoreCase))
 		{
 			return Field;
 		}
 	}
-	return FCbField();
+	return FCbFieldView();
 }
 
-uint64 FCbObject::GetSize() const
+uint64 FCbObjectView::GetSize() const
 {
 	return sizeof(ECbFieldType) + GetPayloadSize();
 }
 
-FIoHash FCbObject::GetHash() const
+FIoHash FCbObjectView::GetHash() const
 {
 	FBlake3 Hash;
 	GetHash(Hash);
 	return FIoHash(Hash.Finalize());
 }
 
-void FCbObject::GetHash(FBlake3& Hash) const
+void FCbObjectView::GetHash(FBlake3& Hash) const
 {
 	const ECbFieldType SerializedType = FCbFieldType::GetType(GetType());
 	Hash.Update(&SerializedType, sizeof(SerializedType));
 	Hash.Update(GetPayloadView());
 }
 
-bool FCbObject::Equals(const FCbObject& Other) const
+bool FCbObjectView::Equals(const FCbObjectView& Other) const
 {
 	return FCbFieldType::GetType(GetType()) == FCbFieldType::GetType(Other.GetType()) &&
 		GetPayloadView().EqualBytes(Other.GetPayloadView());
 }
 
-void FCbObject::CopyTo(FMutableMemoryView Buffer) const
+void FCbObjectView::CopyTo(FMutableMemoryView Buffer) const
 {
 	const FMemoryView Source = GetPayloadView();
 	checkf(Buffer.GetSize() == sizeof(ECbFieldType) + Source.GetSize(),
@@ -782,7 +782,7 @@ void FCbObject::CopyTo(FMutableMemoryView Buffer) const
 	FMemory::Memcpy(Buffer.GetData(), Source.GetData(), Source.GetSize());
 }
 
-void FCbObject::CopyTo(FArchive& Ar) const
+void FCbObjectView::CopyTo(FArchive& Ar) const
 {
 	check(Ar.IsSaving());
 	const FMemoryView Source = GetPayloadView();
@@ -904,12 +904,12 @@ void TCbFieldIterator<FieldType>::IterateRangeAttachments(FCbFieldVisitor Visito
 	}
 }
 
+template class TCbFieldIterator<FCbFieldView>;
 template class TCbFieldIterator<FCbField>;
-template class TCbFieldIterator<FCbFieldRef>;
 
 ///////////////////////////////////////////////////////////////////////////////////////////////////////////////////////
 
-FCbFieldRefIterator FCbFieldRefIterator::CloneRange(const FCbFieldIterator& It)
+FCbFieldIterator FCbFieldIterator::CloneRange(const FCbFieldViewIterator& It)
 {
 	FMemoryView View;
 	if (It.TryGetSerializedRangeView(View))
@@ -924,7 +924,7 @@ FCbFieldRefIterator FCbFieldRefIterator::CloneRange(const FCbFieldIterator& It)
 	}
 }
 
-FSharedBuffer FCbFieldRefIterator::GetRangeBuffer() const
+FSharedBuffer FCbFieldIterator::GetRangeBuffer() const
 {
 	const FMemoryView RangeView = GetRangeView();
 	const FSharedBuffer& OuterBuffer = GetOuterBuffer();
diff --git a/Engine/Source/Runtime/Core/Private/Serialization/CompactBinaryPackage.cpp b/Engine/Source/Runtime/Core/Private/Serialization/CompactBinaryPackage.cpp
--- a/Engine/Source/Runtime/Core/Private/Serialization/CompactBinaryPackage.cpp
+++ b/Engine/Source/Runtime/Core/Private/Serialization/CompactBinaryPackage.cpp
@@ -8,16 +8,16 @@
 
 ///////////////////////////////////////////////////////////////////////////////////////////////////////////////////////
 
-FCbAttachment::FCbAttachment(FCbFieldRefIterator InValue, const FIoHash* const InHash)
+FCbAttachment::FCbAttachment(FCbFieldIterator InValue, const FIoHash* const InHash)
 {
 	if (InValue)
 	{
 		if (!InValue.IsOwned())
 		{
-			InValue = FCbFieldRefIterator::CloneRange(InValue);
+			InValue = FCbFieldIterator::CloneRange(InValue);
 		}
 
-		CompactBinary = FCbFieldIterator(InValue);
+		CompactBinary = FCbFieldViewIterator(InValue);
 		Buffer = MoveTemp(InValue).GetOuterBuffer();
 	}
 
@@ -65,7 +65,7 @@ FCbAttachment::FCbAttachment(FSharedBuffer InBuffer, const FIoHash* const InHash
 	}
 }
 
-FSharedBuffer FCbAttachment::AsBinary() const
+FSharedBuffer FCbAttachment::AsBinaryView() const
 {
 	if (!CompactBinary)
 	{
@@ -78,18 +78,18 @@ FSharedBuffer FCbAttachment::AsBinary() const
 		return SerializedView == Buffer.GetView() ? Buffer : FSharedBuffer::MakeView(SerializedView, Buffer);
 	}
 
-	return FCbFieldRefIterator::CloneRange(CompactBinary).GetRangeBuffer();
+	return FCbFieldIterator::CloneRange(CompactBinary).GetRangeBuffer();
 }
 
-FCbFieldRefIterator FCbAttachment::AsCompactBinary() const
+FCbFieldIterator FCbAttachment::AsCompactBinary() const
 {
-	return CompactBinary ? FCbFieldRefIterator::MakeRangeView(CompactBinary, Buffer) : FCbFieldRefIterator();
+	return CompactBinary ? FCbFieldIterator::MakeRangeView(CompactBinary, Buffer) : FCbFieldIterator();
 }
 
-void FCbAttachment::Load(FCbFieldRefIterator& Fields)
+void FCbAttachment::Load(FCbFieldIterator& Fields)
 {
 	checkf(Fields.IsBinary(), TEXT("Attachments must start with a binary field."));
-	const FMemoryView View = Fields.AsBinary();
+	const FMemoryView View = Fields.AsBinaryView();
 	if (View.GetSize() > 0)
 	{
 		Buffer = FSharedBuffer::MakeView(View, Fields.GetOuterBuffer());
@@ -99,7 +99,7 @@ void FCbAttachment::Load(FCbFieldRefIterator& Fields)
 		checkf(!Fields.HasError(), TEXT("Attachments must be a non-empty binary value with a content hash."));
 		if (Fields.IsCompactBinaryAttachment())
 		{
-			CompactBinary = FCbFieldIterator::MakeRange(Buffer);
+			CompactBinary = FCbFieldViewIterator::MakeRange(Buffer);
 		}
 		++Fields;
 	}
@@ -114,17 +114,17 @@ void FCbAttachment::Load(FCbFieldRefIterator& Fields)
 
 void FCbAttachment::Load(FArchive& Ar, FCbBufferAllocator Allocator)
 {
-	FCbFieldRef BufferField = LoadCompactBinary(Ar, Allocator);
+	FCbField BufferField = LoadCompactBinary(Ar, Allocator);
 	checkf(BufferField.IsBinary(), TEXT("Attachments must start with a binary field."));
-	const FMemoryView View = BufferField.AsBinary();
+	const FMemoryView View = BufferField.AsBinaryView();
 	if (View.GetSize() > 0)
 	{
 		Buffer = FSharedBuffer::MakeView(View, BufferField.GetOuterBuffer());
 		Buffer.MakeOwned();
-		CompactBinary = FCbFieldIterator();
+		CompactBinary = FCbFieldViewIterator();
 
 		TArray<uint8, TInlineAllocator<64>> HashBuffer;
-		FCbFieldRef HashField = LoadCompactBinary(Ar,
+		FCbField HashField = LoadCompactBinary(Ar,
 			[&HashBuffer](uint64 Size) -> FUniqueBuffer
 			{
 				HashBuffer.SetNumUninitialized(int32(Size));
@@ -134,7 +134,7 @@ void FCbAttachment::Load(FArchive& Ar, FCbBufferAllocator Allocator)
 		checkf(!HashField.HasError(), TEXT("Attachments must be a non-empty binary value with a content hash."));
 		if (HashField.IsCompactBinaryAttachment())
 		{
-			CompactBinary = FCbFieldIterator::MakeRange(Buffer);
+			CompactBinary = FCbFieldViewIterator::MakeRange(Buffer);
 		}
 	}
 	else
@@ -156,7 +156,7 @@ void FCbAttachment::Save(FCbWriter& Writer) const
 		}
 		else
 		{
-			Writer.AddBinary(AsBinary());
+			Writer.AddBinary(AsBinaryView());
 		}
 		Writer.AddCompactBinaryAttachment(Hash);
 	}
@@ -180,11 +180,11 @@ void FCbAttachment::Save(FArchive& Ar) const
 
 ///////////////////////////////////////////////////////////////////////////////////////////////////////////////////////
 
-void FCbPackage::SetObject(FCbObjectRef InObject, const FIoHash* InObjectHash, FAttachmentResolver* InResolver)
+void FCbPackage::SetObject(FCbObject InObject, const FIoHash* InObjectHash, FAttachmentResolver* InResolver)
 {
-	if (InObject.CreateIterator())
+	if (InObject.CreateViewIterator())
 	{
-		Object = InObject.IsOwned() ? MoveTemp(InObject) : FCbObjectRef::Clone(InObject);
+		Object = InObject.IsOwned() ? MoveTemp(InObject) : FCbObject::Clone(InObject);
 		if (InObjectHash)
 		{
 			ObjectHash = *InObjectHash;
@@ -196,7 +196,7 @@ void FCbPackage::SetObject(FCbObjectRef InObject, const FIoHash* InObjectHash, F
 		}
 		if (InResolver)
 		{
-			GatherAttachments(Object.CreateIterator(), *InResolver);
+			GatherAttachments(Object.CreateViewIterator(), *InResolver);
 		}
 	}
 	else
@@ -216,7 +216,7 @@ void FCbPackage::AddAttachment(const FCbAttachment& Attachment, FAttachmentResol
 			FCbAttachment& Existing = Attachments[Index];
 			if (Attachment.IsCompactBinary() && !Existing.IsCompactBinary())
 			{
-				Existing = FCbAttachment(FCbFieldRefIterator::MakeRange(Existing.AsBinary()));
+				Existing = FCbAttachment(FCbFieldIterator::MakeRange(Existing.AsBinaryView()));
 			}
 		}
 		else
@@ -254,16 +254,16 @@ const FCbAttachment* FCbPackage::FindAttachment(const FIoHash& Hash) const
 	return Attachments.IsValidIndex(Index) ? &Attachments[Index] : nullptr;
 }
 
-void FCbPackage::GatherAttachments(const FCbFieldIterator& Fields, FAttachmentResolver Resolver)
+void FCbPackage::GatherAttachments(const FCbFieldViewIterator& Fields, FAttachmentResolver Resolver)
 {
-	Fields.IterateRangeAttachments([this, &Resolver](FCbField Field)
+	Fields.IterateRangeAttachments([this, &Resolver](FCbFieldView Field)
 		{
 			const FIoHash& Hash = Field.AsAttachment();
 			if (FSharedBuffer Buffer = Resolver(Hash))
 			{
 				if (Field.IsCompactBinaryAttachment())
 				{
-					AddAttachment(FCbAttachment(FCbFieldRefIterator::MakeRange(MoveTemp(Buffer)), Hash), &Resolver);
+					AddAttachment(FCbAttachment(FCbFieldIterator::MakeRange(MoveTemp(Buffer)), Hash), &Resolver);
 				}
 				else
 				{
@@ -273,7 +273,7 @@ void FCbPackage::GatherAttachments(const FCbFieldIterator& Fields, FAttachmentRe
 		});
 }
 
-void FCbPackage::Load(FCbFieldRefIterator& Fields)
+void FCbPackage::Load(FCbFieldIterator& Fields)
 {
 	*this = FCbPackage();
 	while (Fields)
@@ -292,10 +292,10 @@ void FCbPackage::Load(FCbFieldRefIterator& Fields)
 		else
 		{
 			checkf(Fields.IsObject(), TEXT("Expected Object, Binary, or Null field when loading a package."));
-			Object = Fields.AsObjectRef();
+			Object = Fields.AsObject();
 			Object.MakeOwned();
 			++Fields;
-			if (Object.CreateIterator())
+			if (Object.CreateViewIterator())
 			{
 				ObjectHash = Fields.AsCompactBinaryAttachment();
 				checkf(!Fields.HasError(), TEXT("Object must be followed by a CompactBinaryAttachment with the object hash."));
@@ -320,24 +320,24 @@ void FCbPackage::Load(FArchive& Ar, FCbBufferAllocator Allocator)
 	*this = FCbPackage();
 	for (;;)
 	{
-		FCbFieldRef ValueField = LoadCompactBinary(Ar, StackAllocator);
+		FCbField ValueField = LoadCompactBinary(Ar, StackAllocator);
 		if (ValueField.IsNull())
 		{
 			break;
 		}
 		else if (ValueField.IsBinary())
 		{
-			const FMemoryView View = ValueField.AsBinary();
+			const FMemoryView View = ValueField.AsBinaryView();
 			if (View.GetSize() > 0)
 			{
 				FSharedBuffer Buffer = FSharedBuffer::MakeView(View, ValueField.GetOuterBuffer());
 				Buffer.MakeOwned();
-				FCbFieldRef HashField = LoadCompactBinary(Ar, StackAllocator);
+				FCbField HashField = LoadCompactBinary(Ar, StackAllocator);
 				const FIoHash& Hash = HashField.AsAttachment();
 				checkf(!HashField.HasError(), TEXT("Attachments must be a non-empty binary value with a content hash."));
 				if (HashField.IsCompactBinaryAttachment())
 				{
-					AddAttachment(FCbAttachment(FCbFieldRefIterator::MakeRange(MoveTemp(Buffer)), Hash));
+					AddAttachment(FCbAttachment(FCbFieldIterator::MakeRange(MoveTemp(Buffer)), Hash));
 				}
 				else
 				{
@@ -348,11 +348,11 @@ void FCbPackage::Load(FArchive& Ar, FCbBufferAllocator Allocator)
 		else
 		{
 			checkf(ValueField.IsObject(), TEXT("Expected Object, Binary, or Null field when loading a package."));
-			Object = ValueField.AsObjectRef();
+			Object = ValueField.AsObject();
 			Object.MakeOwned();
-			if (Object.CreateIterator())
+			if (Object.CreateViewIterator())
 			{
-				FCbFieldRef HashField = LoadCompactBinary(Ar, StackAllocator);
+				FCbField HashField = LoadCompactBinary(Ar, StackAllocator);
 				ObjectHash = HashField.AsCompactBinaryAttachment();
 				checkf(!HashField.HasError(), TEXT("Object must be followed by a CompactBinaryAttachment with the object hash."));
 			}
@@ -366,7 +366,7 @@ void FCbPackage::Load(FArchive& Ar, FCbBufferAllocator Allocator)
 
 void FCbPackage::Save(FCbWriter& Writer) const
 {
-	if (Object.CreateIterator())
+	if (Object.CreateViewIterator())
 	{
 		Writer.AddObject(Object);
 		Writer.AddCompactBinaryAttachment(ObjectHash);
diff --git a/Engine/Source/Runtime/Core/Private/Serialization/CompactBinarySerialization.cpp b/Engine/Source/Runtime/Core/Private/Serialization/CompactBinarySerialization.cpp
--- a/Engine/Source/Runtime/Core/Private/Serialization/CompactBinarySerialization.cpp
+++ b/Engine/Source/Runtime/Core/Private/Serialization/CompactBinarySerialization.cpp
@@ -157,7 +157,7 @@ bool TryMeasureCompactBinary(FMemoryView View, ECbFieldType& OutType, uint64& Ou
 
 ///////////////////////////////////////////////////////////////////////////////////////////////////////////////////////
 
-FCbFieldRef LoadCompactBinary(FArchive& Ar, FCbBufferAllocator Allocator)
+FCbField LoadCompactBinary(FArchive& Ar, FCbBufferAllocator Allocator)
 {
 	TArray<uint8, TInlineAllocator<64>> HeaderBytes;
 	ECbFieldType FieldType;
@@ -188,24 +188,24 @@ FCbFieldRef LoadCompactBinary(FArchive& Ar, FCbBufferAllocator Allocator)
 	{
 		Ar.Serialize(View.GetData(), static_cast<int64>(View.GetSize()));
 	}
-	return FCbFieldRef(FSharedBuffer(MoveTemp(Buffer)));
+	return FCbField(FSharedBuffer(MoveTemp(Buffer)));
 }
 
 ///////////////////////////////////////////////////////////////////////////////////////////////////////////////////////
 
-void SaveCompactBinary(FArchive& Ar, const FCbField& Field)
+void SaveCompactBinary(FArchive& Ar, const FCbFieldView& Field)
 {
 	check(Ar.IsSaving());
 	Field.CopyTo(Ar);
 }
 
-void SaveCompactBinary(FArchive& Ar, const FCbArray& Array)
+void SaveCompactBinary(FArchive& Ar, const FCbArrayView& Array)
 {
 	check(Ar.IsSaving());
 	Array.CopyTo(Ar);
 }
 
-void SaveCompactBinary(FArchive& Ar, const FCbObject& Object)
+void SaveCompactBinary(FArchive& Ar, const FCbObjectView& Object)
 {
 	check(Ar.IsSaving());
 	Object.CopyTo(Ar);
@@ -231,19 +231,19 @@ static FArchive& SerializeCompactBinary(FArchive& Ar, T& Value, ConvertType&& Co
 	return Ar;
 }
 
-FArchive& operator<<(FArchive& Ar, FCbFieldRef& Field)
+FArchive& operator<<(FArchive& Ar, FCbField& Field)
 {
 	return SerializeCompactBinary(Ar, Field, FIdentityFunctor());
 }
 
-FArchive& operator<<(FArchive& Ar, FCbArrayRef& Array)
+FArchive& operator<<(FArchive& Ar, FCbArray& Array)
 {
-	return SerializeCompactBinary(Ar, Array, [](FCbFieldRef&& Field) { return MoveTemp(Field).AsArrayRef(); });
+	return SerializeCompactBinary(Ar, Array, [](FCbField&& Field) { return MoveTemp(Field).AsArray(); });
 }
 
-FArchive& operator<<(FArchive& Ar, FCbObjectRef& Object)
+FArchive& operator<<(FArchive& Ar, FCbObject& Object)
 {
-	return SerializeCompactBinary(Ar, Object, [](FCbFieldRef&& Field) { return MoveTemp(Field).AsObjectRef(); });
+	return SerializeCompactBinary(Ar, Object, [](FCbField&& Field) { return MoveTemp(Field).AsObject(); });
 }
 
 ///////////////////////////////////////////////////////////////////////////////////////////////////////////////////////
diff --git a/Engine/Source/Runtime/Core/Private/Serialization/CompactBinaryValidation.cpp b/Engine/Source/Runtime/Core/Private/Serialization/CompactBinaryValidation.cpp
--- a/Engine/Source/Runtime/Core/Private/Serialization/CompactBinaryValidation.cpp
+++ b/Engine/Source/Runtime/Core/Private/Serialization/CompactBinaryValidation.cpp
@@ -165,7 +165,7 @@ static FAnsiStringView ValidateCbString(FMemoryView& View, ECbValidateMode Mode,
 	return Value;
 }
 
-static FCbField ValidateCbField(FMemoryView& View, ECbValidateMode Mode, ECbValidateError& Error, ECbFieldType ExternalType);
+static FCbFieldView ValidateCbField(FMemoryView& View, ECbValidateMode Mode, ECbValidateError& Error, ECbFieldType ExternalType);
 
 /** A type that checks whether all validated fields are of the same type. */
 class FCbUniformFieldsValidator
@@ -176,10 +176,10 @@ public:
 	{
 	}
 
-	inline FCbField ValidateField(FMemoryView& View, ECbValidateMode Mode, ECbValidateError& Error)
+	inline FCbFieldView ValidateField(FMemoryView& View, ECbValidateMode Mode, ECbValidateError& Error)
 	{
 		const void* const FieldData = View.GetData();
-		if (FCbField Field = ValidateCbField(View, Mode, Error, ExternalType))
+		if (FCbFieldView Field = ValidateCbField(View, Mode, Error, ExternalType))
 		{
 			++FieldCount;
 			if (FCbFieldType::HasFieldType(ExternalType))
@@ -199,7 +199,7 @@ public:
 
 		// It may not safe to check for uniformity if the field was invalid.
 		bUniform = false;
-		return FCbField();
+		return FCbFieldView();
 	}
 
 	inline bool IsUniform() const { return FieldCount > 0 && bUniform; }
@@ -226,7 +226,7 @@ static void ValidateCbObject(FMemoryView& View, ECbValidateMode Mode, ECbValidat
 		FCbUniformFieldsValidator UniformValidator(ExternalType);
 		do
 		{
-			if (FCbField Field = UniformValidator.ValidateField(ObjectView, Mode, Error))
+			if (FCbFieldView Field = UniformValidator.ValidateField(ObjectView, Mode, Error))
 			{
 				if (EnumHasAnyFlags(Mode, ECbValidateMode::Names))
 				{
@@ -277,7 +277,7 @@ static void ValidateCbArray(FMemoryView& View, ECbValidateMode Mode, ECbValidate
 
 	for (uint64 Index = 0; Index < Count; ++Index)
 	{
-		if (FCbField Field = UniformValidator.ValidateField(ArrayView, Mode, Error))
+		if (FCbFieldView Field = UniformValidator.ValidateField(ArrayView, Mode, Error))
 		{
 			if (Field.HasName() && EnumHasAnyFlags(Mode, ECbValidateMode::Names))
 			{
@@ -292,7 +292,7 @@ static void ValidateCbArray(FMemoryView& View, ECbValidateMode Mode, ECbValidate
 	}
 }
 
-static FCbField ValidateCbField(FMemoryView& View, ECbValidateMode Mode, ECbValidateError& Error, const ECbFieldType ExternalType = ECbFieldType::HasFieldType)
+static FCbFieldView ValidateCbField(FMemoryView& View, ECbValidateMode Mode, ECbValidateError& Error, const ECbFieldType ExternalType = ECbFieldType::HasFieldType)
 {
 	const FMemoryView FieldView = View;
 	const ECbFieldType Type = ValidateCbFieldType(View, Mode, Error, ExternalType);
@@ -300,7 +300,7 @@ static FCbField ValidateCbField(FMemoryView& View, ECbValidateMode Mode, ECbVali
 
 	if (EnumHasAnyFlags(Error, ECbValidateError::OutOfBounds | ECbValidateError::InvalidType))
 	{
-		return FCbField();
+		return FCbFieldView();
 	}
 
 	switch (FCbFieldType::GetType(Type))
@@ -381,13 +381,13 @@ static FCbField ValidateCbField(FMemoryView& View, ECbValidateMode Mode, ECbVali
 
 	if (EnumHasAnyFlags(Error, ECbValidateError::OutOfBounds | ECbValidateError::InvalidType))
 	{
-		return FCbField();
+		return FCbFieldView();
 	}
 
-	return FCbField(FieldView.GetData(), ExternalType);
+	return FCbFieldView(FieldView.GetData(), ExternalType);
 }
 
-static FCbField ValidateCbPackageField(FMemoryView& View, ECbValidateMode Mode, ECbValidateError& Error)
+static FCbFieldView ValidateCbPackageField(FMemoryView& View, ECbValidateMode Mode, ECbValidateError& Error)
 {
 	if (View.IsEmpty())
 	{
@@ -395,9 +395,9 @@ static FCbField ValidateCbPackageField(FMemoryView& View, ECbValidateMode Mode,
 		{
 			AddError(Error, ECbValidateError::InvalidPackageFormat);
 		}
-		return FCbField();
+		return FCbFieldView();
 	}
-	if (FCbField Field = ValidateCbField(View, Mode, Error))
+	if (FCbFieldView Field = ValidateCbField(View, Mode, Error))
 	{
 		if (Field.HasName() && EnumHasAnyFlags(Mode, ECbValidateMode::Package))
 		{
@@ -405,12 +405,12 @@ static FCbField ValidateCbPackageField(FMemoryView& View, ECbValidateMode Mode,
 		}
 		return Field;
 	}
-	return FCbField();
+	return FCbFieldView();
 }
 
-static FIoHash ValidateCbPackageAttachment(FCbField& Value, FMemoryView& View, ECbValidateMode Mode, ECbValidateError& Error)
+static FIoHash ValidateCbPackageAttachment(FCbFieldView& Value, FMemoryView& View, ECbValidateMode Mode, ECbValidateError& Error)
 {
-	const FMemoryView ValueView = Value.AsBinary();
+	const FMemoryView ValueView = Value.AsBinaryView();
 	if (Value.HasError() && EnumHasAnyFlags(Mode, ECbValidateMode::Package))
 	{
 		if (EnumHasAnyFlags(Mode, ECbValidateMode::Package))
@@ -420,7 +420,7 @@ static FIoHash ValidateCbPackageAttachment(FCbField& Value, FMemoryView& View, E
 	}
 	else if (ValueView.GetSize())
 	{
-		if (FCbField HashField = ValidateCbPackageField(View, Mode, Error))
+		if (FCbFieldView HashField = ValidateCbPackageField(View, Mode, Error))
 		{
 			const FIoHash Hash = HashField.AsAttachment();
 			if (EnumHasAnyFlags(Mode, ECbValidateMode::Package))
@@ -440,9 +440,9 @@ static FIoHash ValidateCbPackageAttachment(FCbField& Value, FMemoryView& View, E
 	return FIoHash();
 }
 
-static FIoHash ValidateCbPackageObject(FCbField& Value, FMemoryView& View, ECbValidateMode Mode, ECbValidateError& Error)
+static FIoHash ValidateCbPackageObject(FCbFieldView& Value, FMemoryView& View, ECbValidateMode Mode, ECbValidateError& Error)
 {
-	FCbObject Object = Value.AsObject();
+	FCbObjectView Object = Value.AsObjectView();
 	if (Value.HasError())
 	{
 		if (EnumHasAnyFlags(Mode, ECbValidateMode::Package))
@@ -450,12 +450,12 @@ static FIoHash ValidateCbPackageObject(FCbField& Value, FMemoryView& View, ECbVa
 			AddError(Error, ECbValidateError::InvalidPackageFormat);
 		}
 	}
-	else if (FCbField HashField = ValidateCbPackageField(View, Mode, Error))
+	else if (FCbFieldView HashField = ValidateCbPackageField(View, Mode, Error))
 	{
 		const FIoHash Hash = HashField.AsAttachment();
 		if (EnumHasAnyFlags(Mode, ECbValidateMode::Package))
 		{
-			if (!Object.CreateIterator())
+			if (!Object.CreateViewIterator())
 			{
 				AddError(Error, ECbValidateError::NullPackageObject);
 			}
@@ -505,7 +505,7 @@ ECbValidateError ValidateCompactBinaryAttachment(FMemoryView View, ECbValidateMo
 	ECbValidateError Error = ECbValidateError::None;
 	if (EnumHasAnyFlags(Mode, ECbValidateMode::All))
 	{
-		if (FCbField Value = ValidateCbPackageField(View, Mode, Error))
+		if (FCbFieldView Value = ValidateCbPackageField(View, Mode, Error))
 		{
 			ValidateCbPackageAttachment(Value, View, Mode, Error);
 		}
@@ -524,7 +524,7 @@ ECbValidateError ValidateCompactBinaryPackage(FMemoryView View, ECbValidateMode
 	if (EnumHasAnyFlags(Mode, ECbValidateMode::All))
 	{
 		uint32 ObjectCount = 0;
-		while (FCbField Value = ValidateCbPackageField(View, Mode, Error))
+		while (FCbFieldView Value = ValidateCbPackageField(View, Mode, Error))
 		{
 			if (Value.IsBinary())
 			{
@@ -532,7 +532,7 @@ ECbValidateError ValidateCompactBinaryPackage(FMemoryView View, ECbValidateMode
 				if (EnumHasAnyFlags(Mode, ECbValidateMode::Package))
 				{
 					Attachments.Add(Hash);
-					if (Value.AsBinary().IsEmpty())
+					if (Value.AsBinaryView().IsEmpty())
 					{
 						AddError(Error, ECbValidateError::NullPackageAttachment);
 					}
diff --git a/Engine/Source/Runtime/Core/Private/Serialization/CompactBinaryWriter.cpp b/Engine/Source/Runtime/Core/Private/Serialization/CompactBinaryWriter.cpp
--- a/Engine/Source/Runtime/Core/Private/Serialization/CompactBinaryWriter.cpp
+++ b/Engine/Source/Runtime/Core/Private/Serialization/CompactBinaryWriter.cpp
@@ -49,12 +49,12 @@ static constexpr bool IsUniformType(const ECbFieldType Type)
 }
 
 /** Append the payload from the compact binary value to the array and return its type. */
-static inline ECbFieldType AppendCompactBinary(const FCbField& Value, TArray64<uint8>& OutData)
+static inline ECbFieldType AppendCompactBinary(const FCbFieldView& Value, TArray64<uint8>& OutData)
 {
-	struct FCopy : public FCbField
+	struct FCopy : public FCbFieldView
 	{
-		using FCbField::GetType;
-		using FCbField::GetPayloadView;
+		using FCbFieldView::GetType;
+		using FCbFieldView::GetPayloadView;
 	};
 	const FCopy& ValueCopy = static_cast<const FCopy&>(Value);
 	const FMemoryView SourceView = ValueCopy.GetPayloadView();
@@ -85,15 +85,15 @@ void FCbWriter::Reset()
 	States.Emplace();
 }
 
-FCbFieldRefIterator FCbWriter::Save() const
+FCbFieldIterator FCbWriter::Save() const
 {
 	const uint64 Size = GetSaveSize();
 	FUniqueBuffer Buffer = FUniqueBuffer::Alloc(Size);
-	const FCbFieldIterator Output = Save(Buffer);
-	return FCbFieldRefIterator::MakeRangeView(Output, FSharedBuffer(MoveTemp(Buffer)));
+	const FCbFieldViewIterator Output = Save(Buffer);
+	return FCbFieldIterator::MakeRangeView(Output, FSharedBuffer(MoveTemp(Buffer)));
 }
 
-FCbFieldIterator FCbWriter::Save(const FMutableMemoryView Buffer) const
+FCbFieldViewIterator FCbWriter::Save(const FMutableMemoryView Buffer) const
 {
 	checkf(States.Num() == 1 && States.Last().Flags == EStateFlags::None,
 		TEXT("It is invalid to save while there are incomplete write operations."));
@@ -101,7 +101,7 @@ FCbFieldIterator FCbWriter::Save(const FMutableMemoryView Buffer) const
 	checkf(Buffer.GetSize() == Data.Num(),
 		TEXT("Buffer is %" UINT64_FMT " bytes but %" INT64_FMT " is required."), Buffer.GetSize(), Data.Num());
 	FMemory::Memcpy(Buffer.GetData(), Data.GetData(), Data.Num());
-	return FCbFieldIterator::MakeRange(Buffer);
+	return FCbFieldViewIterator::MakeRange(Buffer);
 }
 
 void FCbWriter::Save(FArchive& Ar) const
@@ -232,16 +232,16 @@ void FCbWriter::MakeFieldsUniform(const int64 FieldBeginOffset, const int64 Fiel
 	}
 }
 
-void FCbWriter::AddField(const FCbField& Value)
+void FCbWriter::AddField(const FCbFieldView& Value)
 {
 	checkf(Value.HasValue(), TEXT("It is invalid to write a field with no value."));
 	BeginField();
 	EndField(AppendCompactBinary(Value, Data));
 }
 
-void FCbWriter::AddField(const FCbFieldRef& Value)
+void FCbWriter::AddField(const FCbField& Value)
 {
-	AddField(FCbField(Value));
+	AddField(FCbFieldView(Value));
 }
 
 void FCbWriter::BeginObject()
@@ -286,15 +286,15 @@ void FCbWriter::EndObject()
 	EndField(bUniform ? ECbFieldType::UniformObject : ECbFieldType::Object);
 }
 
-void FCbWriter::AddObject(const FCbObject& Value)
+void FCbWriter::AddObject(const FCbObjectView& Value)
 {
 	BeginField();
-	EndField(AppendCompactBinary(Value.AsField(), Data));
+	EndField(AppendCompactBinary(Value.AsFieldView(), Data));
 }
 
-void FCbWriter::AddObject(const FCbObjectRef& Value)
+void FCbWriter::AddObject(const FCbObject& Value)
 {
-	AddObject(FCbObject(Value));
+	AddObject(FCbObjectView(Value));
 }
 
 void FCbWriter::BeginArray()
@@ -341,15 +341,15 @@ void FCbWriter::EndArray()
 	EndField(bUniform ? ECbFieldType::UniformArray : ECbFieldType::Array);
 }
 
-void FCbWriter::AddArray(const FCbArray& Value)
+void FCbWriter::AddArray(const FCbArrayView& Value)
 {
 	BeginField();
-	EndField(AppendCompactBinary(Value.AsField(), Data));
+	EndField(AppendCompactBinary(Value.AsFieldView(), Data));
 }
 
-void FCbWriter::AddArray(const FCbArrayRef& Value)
+void FCbWriter::AddArray(const FCbArray& Value)
 {
-	AddArray(FCbArray(Value));
+	AddArray(FCbArrayView(Value));
 }
 
 void FCbWriter::AddNull()
diff --git a/Engine/Source/Runtime/Core/Private/Tests/Serialization/CompactBinaryPackageTest.cpp b/Engine/Source/Runtime/Core/Private/Tests/Serialization/CompactBinaryPackageTest.cpp
--- a/Engine/Source/Runtime/Core/Private/Tests/Serialization/CompactBinaryPackageTest.cpp
+++ b/Engine/Source/Runtime/Core/Private/Tests/Serialization/CompactBinaryPackageTest.cpp
@@ -27,7 +27,7 @@ bool FCbAttachmentTest::RunTest(const FString& Parameters)
 		FBufferArchive WriteAr;
 		Attachment.Save(Writer);
 		Attachment.Save(WriteAr);
-		FCbFieldRefIterator Fields = Writer.Save();
+		FCbFieldIterator Fields = Writer.Save();
 
 		TestTrue(FString::Printf(TEXT("FCbAttachment(%s).Save()->Equals"), Test),
 			MakeMemoryView(WriteAr).EqualBytes(Fields.GetRangeBuffer().GetView()));
@@ -53,7 +53,7 @@ bool FCbAttachmentTest::RunTest(const FString& Parameters)
 		FCbAttachment Attachment;
 		TestTrue(TEXT("FCbAttachment(Null).IsNull()"), Attachment.IsNull());
 		TestFalse(TEXT("FCbAttachment(Null) as bool"), bool(Attachment));
-		TestFalse(TEXT("FCbAttachment(Null).AsBinary()"), bool(Attachment.AsBinary()));
+		TestFalse(TEXT("FCbAttachment(Null).AsBinaryView()"), bool(Attachment.AsBinaryView()));
 		TestFalse(TEXT("FCbAttachment(Null).AsCompactBinary()"), Attachment.AsCompactBinary().HasValue());
 		TestFalse(TEXT("FCbAttachment(Null).IsBinary()"), Attachment.IsBinary());
 		TestFalse(TEXT("FCbAttachment(Null).IsCompactBinary()"), Attachment.IsCompactBinary());
@@ -67,7 +67,7 @@ bool FCbAttachmentTest::RunTest(const FString& Parameters)
 		FCbAttachment Attachment(Buffer);
 		TestFalse(TEXT("FCbAttachment(Binary).IsNull()"), Attachment.IsNull());
 		TestTrue(TEXT("FCbAttachment(Binary) as bool"), bool(Attachment));
-		TestEqual(TEXT("FCbAttachment(Binary).AsBinary()"), Attachment.AsBinary(), Buffer);
+		TestEqual(TEXT("FCbAttachment(Binary).AsBinaryView()"), Attachment.AsBinaryView(), Buffer);
 		TestFalse(TEXT("FCbAttachment(Binary).AsCompactBinary()"), Attachment.AsCompactBinary().HasValue());
 		TestTrue(TEXT("FCbAttachment(Binary).IsBinary()"), Attachment.IsBinary());
 		TestFalse(TEXT("FCbAttachment(Binary).IsCompactBinary()"), Attachment.IsCompactBinary());
@@ -79,11 +79,11 @@ bool FCbAttachmentTest::RunTest(const FString& Parameters)
 	{
 		FCbWriter Writer;
 		Writer << "Name"_ASV << 42;
-		FCbFieldRefIterator Fields = Writer.Save();
+		FCbFieldIterator Fields = Writer.Save();
 		FCbAttachment Attachment(Fields);
 		TestFalse(TEXT("FCbAttachment(CompactBinary).IsNull()"), Attachment.IsNull());
 		TestTrue(TEXT("FCbAttachment(CompactBinary) as bool"), bool(Attachment));
-		TestEqual(TEXT("FCbAttachment(CompactBinary).AsBinary()"), Attachment.AsBinary(), Fields.GetRangeBuffer());
+		TestEqual(TEXT("FCbAttachment(CompactBinary).AsBinaryView()"), Attachment.AsBinaryView(), Fields.GetRangeBuffer());
 		TestEqual(TEXT("FCbAttachment(CompactBinary).AsCompactBinary()"), Attachment.AsCompactBinary(), Fields);
 		TestTrue(TEXT("FCbAttachment(CompactBinary).IsBinary()"), Attachment.IsBinary());
 		TestTrue(TEXT("FCbAttachment(CompactBinary).IsCompactBinary()"), Attachment.IsCompactBinary());
@@ -98,9 +98,9 @@ bool FCbAttachmentTest::RunTest(const FString& Parameters)
 		FCbAttachment Attachment(Buffer);
 		TestFalse(TEXT("FCbAttachment(BinaryView).IsNull()"), Attachment.IsNull());
 		TestTrue(TEXT("FCbAttachment(BinaryView) as bool"), bool(Attachment));
-		TestNotEqual(TEXT("FCbAttachment(BinaryView).AsBinary()"), Attachment.AsBinary(), Buffer);
-		TestTrue(TEXT("FCbAttachment(BinaryView).AsBinary()"),
-			Attachment.AsBinary().GetView().EqualBytes(Buffer.GetView()));
+		TestNotEqual(TEXT("FCbAttachment(BinaryView).AsBinaryView()"), Attachment.AsBinaryView(), Buffer);
+		TestTrue(TEXT("FCbAttachment(BinaryView).AsBinaryView()"),
+			Attachment.AsBinaryView().GetView().EqualBytes(Buffer.GetView()));
 		TestFalse(TEXT("FCbAttachment(BinaryView).AsCompactBinary()"), Attachment.AsCompactBinary().HasValue());
 		TestTrue(TEXT("FCbAttachment(BinaryView).IsBinary()"), Attachment.IsBinary());
 		TestFalse(TEXT("FCbAttachment(BinaryView).IsCompactBinary()"), Attachment.IsCompactBinary());
@@ -112,13 +112,13 @@ bool FCbAttachmentTest::RunTest(const FString& Parameters)
 	{
 		FCbWriter Writer;
 		Writer << "Name"_ASV << 42;
-		FCbFieldRefIterator Fields = Writer.Save();
-		FCbFieldRefIterator FieldsView = FCbFieldRefIterator::MakeRangeView(FCbFieldIterator(Fields));
+		FCbFieldIterator Fields = Writer.Save();
+		FCbFieldIterator FieldsView = FCbFieldIterator::MakeRangeView(FCbFieldViewIterator(Fields));
 		FCbAttachment Attachment(FieldsView);
 		TestFalse(TEXT("FCbAttachment(CompactBinaryView).IsNull()"), Attachment.IsNull());
 		TestTrue(TEXT("FCbAttachment(CompactBinaryView) as bool"), bool(Attachment));
-		TestNotEqual(TEXT("FCbAttachment(CompactBinaryView).AsBinary()"),
-			Attachment.AsBinary(), FieldsView.GetRangeBuffer());
+		TestNotEqual(TEXT("FCbAttachment(CompactBinaryView).AsBinaryView()"),
+			Attachment.AsBinaryView(), FieldsView.GetRangeBuffer());
 		TestTrue(TEXT("FCbAttachment(CompactBinaryView).AsCompactBinary()"),
 			Attachment.AsCompactBinary().GetRangeView().EqualBytes(Fields.GetRangeView()));
 		TestTrue(TEXT("FCbAttachment(CompactBinaryView).IsBinary()"), Attachment.IsBinary());
@@ -134,16 +134,16 @@ bool FCbAttachmentTest::RunTest(const FString& Parameters)
 
 		FCbWriter Writer;
 		Attachment.Save(Writer);
-		FCbFieldRefIterator Fields = Writer.Save();
-		FCbFieldRefIterator FieldsView = FCbFieldRefIterator::MakeRangeView(FCbFieldIterator(Fields));
+		FCbFieldIterator Fields = Writer.Save();
+		FCbFieldIterator FieldsView = FCbFieldIterator::MakeRangeView(FCbFieldViewIterator(Fields));
 
 		Attachment.Load(FieldsView);
 		TestFalse(TEXT("FCbAttachment(LoadBinaryView).IsNull()"), Attachment.IsNull());
 		TestTrue(TEXT("FCbAttachment(LoadBinaryView) as bool"), bool(Attachment));
-		TestFalse(TEXT("FCbAttachment(LoadBinaryView).AsBinary()->!InView"),
-			FieldsView.GetRangeBuffer().GetView().Contains(Attachment.AsBinary().GetView()));
-		TestTrue(TEXT("FCbAttachment(LoadBinaryView).AsBinary()->EqualBytes"),
-			Attachment.AsBinary().GetView().EqualBytes(Buffer.GetView()));
+		TestFalse(TEXT("FCbAttachment(LoadBinaryView).AsBinaryView()->!InView"),
+			FieldsView.GetRangeBuffer().GetView().Contains(Attachment.AsBinaryView().GetView()));
+		TestTrue(TEXT("FCbAttachment(LoadBinaryView).AsBinaryView()->EqualBytes"),
+			Attachment.AsBinaryView().GetView().EqualBytes(Buffer.GetView()));
 		TestFalse(TEXT("FCbAttachment(LoadBinaryView).AsCompactBinary()"), Attachment.AsCompactBinary().HasValue());
 		TestTrue(TEXT("FCbAttachment(LoadBinaryView).IsBinary()"), Attachment.IsBinary());
 		TestFalse(TEXT("FCbAttachment(LoadBinaryView).IsCompactBinary()"), Attachment.IsCompactBinary());
@@ -155,21 +155,21 @@ bool FCbAttachmentTest::RunTest(const FString& Parameters)
 	{
 		FCbWriter ValueWriter;
 		ValueWriter << "Name"_ASV << 42;
-		const FCbFieldRefIterator Value = ValueWriter.Save();
+		const FCbFieldIterator Value = ValueWriter.Save();
 		TestEqual(TEXT("FCbAttachment(LoadCompactBinaryView).Validate"),
 			ValidateCompactBinaryRange(Value.GetRangeView(), ECbValidateMode::All), ECbValidateError::None);
 		FCbAttachment Attachment(Value);
 
 		FCbWriter Writer;
 		Attachment.Save(Writer);
-		FCbFieldRefIterator Fields = Writer.Save();
-		FCbFieldRefIterator FieldsView = FCbFieldRefIterator::MakeRangeView(FCbFieldIterator(Fields));
+		FCbFieldIterator Fields = Writer.Save();
+		FCbFieldIterator FieldsView = FCbFieldIterator::MakeRangeView(FCbFieldViewIterator(Fields));
 
 		Attachment.Load(FieldsView);
 		TestFalse(TEXT("FCbAttachment(LoadCompactBinaryView).IsNull()"), Attachment.IsNull());
 		TestTrue(TEXT("FCbAttachment(LoadCompactBinaryView) as bool"), bool(Attachment));
-		TestTrue(TEXT("FCbAttachment(LoadCompactBinaryView).AsBinary()->EqualBytes"),
-			Attachment.AsBinary().GetView().EqualBytes(Value.GetRangeView()));
+		TestTrue(TEXT("FCbAttachment(LoadCompactBinaryView).AsBinaryView()->EqualBytes"),
+			Attachment.AsBinaryView().GetView().EqualBytes(Value.GetRangeView()));
 		TestFalse(TEXT("FCbAttachment(LoadCompactBinaryView).AsCompactBinary()->!InView"),
 			FieldsView.GetRangeBuffer().GetView().Contains(Attachment.AsCompactBinary().GetRangeBuffer().GetView()));
 		TestTrue(TEXT("FCbAttachment(LoadCompactBinaryView).IsBinary()"), Attachment.IsBinary());
@@ -180,19 +180,19 @@ bool FCbAttachmentTest::RunTest(const FString& Parameters)
 	// Compact Binary Uniform Sub-View
 	{
 		const FSharedBuffer Buffer = FSharedBuffer::Clone(MakeMemoryView<uint8>({0, 1, 2, 3}));
-		const FCbFieldIterator Fields = FCbFieldIterator::MakeRange(Buffer.GetView().RightChop(2), ECbFieldType::IntegerPositive);
-		const FCbFieldRefIterator SavedFieldRefs = FCbFieldRefIterator::CloneRange(Fields);
-		FCbFieldRefIterator FieldRefs = FCbFieldRefIterator::MakeRangeView(Fields, Buffer);
-		FCbAttachment Attachment(FieldRefs);
-		const FSharedBuffer Binary = Attachment.AsBinary();
+		const FCbFieldViewIterator FieldViews = FCbFieldViewIterator::MakeRange(Buffer.GetView().RightChop(2), ECbFieldType::IntegerPositive);
+		const FCbFieldIterator SavedFields = FCbFieldIterator::CloneRange(FieldViews);
+		FCbFieldIterator Fields = FCbFieldIterator::MakeRangeView(FieldViews, Buffer);
+		FCbAttachment Attachment(Fields);
+		const FSharedBuffer Binary = Attachment.AsBinaryView();
 		TestEqual(TEXT("FCbAttachment(CompactBinaryUniformSubView).AsCompactBinary()->Equals()"),
-			Attachment.AsCompactBinary(), FieldRefs);
-		TestEqual(TEXT("FCbAttachment(CompactBinaryUniformSubView).AsBinary()->GetSize()"),
-			Binary.GetSize(), SavedFieldRefs.GetRangeSize());
-		TestTrue(TEXT("FCbAttachment(CompactBinaryUniformSubView).AsBinary()->EqualBytes()"),
-			Binary.GetView().EqualBytes(SavedFieldRefs.GetRangeView()));
+			Attachment.AsCompactBinary(), Fields);
+		TestEqual(TEXT("FCbAttachment(CompactBinaryUniformSubView).AsBinaryView()->GetSize()"),
+			Binary.GetSize(), SavedFields.GetRangeSize());
+		TestTrue(TEXT("FCbAttachment(CompactBinaryUniformSubView).AsBinaryView()->EqualBytes()"),
+			Binary.GetView().EqualBytes(SavedFields.GetRangeView()));
 		TestEqual(TEXT("FCbAttachment(CompactBinaryUniformSubView).GetHash()"),
-			Attachment.GetHash(), SavedFieldRefs.GetRangeHash());
+			Attachment.GetHash(), SavedFields.GetRangeHash());
 		TestSaveLoadValidate(TEXT("CompactBinaryUniformSubView"), Attachment);
 	}
 
@@ -216,7 +216,7 @@ bool FCbAttachmentTest::RunTest(const FString& Parameters)
 
 	// Compact Binary Empty
 	{
-		const FCbAttachment Attachment(FCbFieldRefIterator{});
+		const FCbAttachment Attachment(FCbFieldIterator{});
 		TestTrue(TEXT("FCbAttachment(CompactBinaryEmpty).IsNull()"), Attachment.IsNull());
 		TestFalse(TEXT("FCbAttachment(CompactBinaryEmpty).IsBinary()"), Attachment.IsBinary());
 		TestFalse(TEXT("FCbAttachment(CompactBinaryEmpty).IsCompactBinary()"), Attachment.IsCompactBinary());
@@ -235,7 +235,7 @@ bool FCbPackageTest::RunTest(const FString& Parameters)
 		FBufferArchive WriteAr;
 		Package.Save(Writer);
 		Package.Save(WriteAr);
-		FCbFieldRefIterator Fields = Writer.Save();
+		FCbFieldIterator Fields = Writer.Save();
 
 		TestTrue(FString::Printf(TEXT("FCbPackage(%s).Save()->Equals"), Test),
 			MakeMemoryView(WriteAr).EqualBytes(Fields.GetRangeBuffer().GetView()));
@@ -272,7 +272,7 @@ bool FCbPackageTest::RunTest(const FString& Parameters)
 		Writer << "Field" << 42;
 		Writer.EndObject();
 
-		const FCbObjectRef Object = Writer.Save().AsObjectRef();
+		const FCbObject Object = Writer.Save().AsObject();
 		FCbPackage Package(Object);
 		TestFalse(TEXT("FCbPackage(Object).IsNull()"), Package.IsNull());
 		TestTrue(TEXT("FCbPackage(Object) as bool"), bool(Package));
@@ -290,8 +290,8 @@ bool FCbPackageTest::RunTest(const FString& Parameters)
 		Writer << "Field" << 42;
 		Writer.EndObject();
 
-		const FCbObjectRef Object = Writer.Save().AsObjectRef();
-		FCbPackage Package(FCbObjectRef::MakeView(Object));
+		const FCbObject Object = Writer.Save().AsObject();
+		FCbPackage Package(FCbObject::MakeView(Object));
 		TestFalse(TEXT("FCbPackage(Object).IsNull()"), Package.IsNull());
 		TestTrue(TEXT("FCbPackage(Object) as bool"), bool(Package));
 		TestEqual(TEXT("FCbPackage(Object).GetAttachments()"), Package.GetAttachments().Num(), 0);
@@ -303,24 +303,24 @@ bool FCbPackageTest::RunTest(const FString& Parameters)
 
 	// Attachment Only
 	{
-		FCbObjectRef Object;
+		FCbObject Object;
 		{
 			TCbWriter<256> Writer;
 			Writer.BeginObject();
 			Writer << "Field" << 42;
 			Writer.EndObject();
-			Object = Writer.Save().AsObjectRef();
+			Object = Writer.Save().AsObject();
 		}
-		FCbFieldRef Field = FCbFieldRef::Clone(Object["Field"]);
+		FCbField Field = FCbField::Clone(Object["Field"]);
 
 		FCbPackage Package;
-		Package.AddAttachment(FCbAttachment(FCbFieldRefIterator::MakeSingle(Object.AsFieldRef())));
+		Package.AddAttachment(FCbAttachment(FCbFieldIterator::MakeSingle(Object.AsField())));
 		Package.AddAttachment(FCbAttachment(Field.GetBuffer()));
 
 		TestFalse(TEXT("FCbPackage(Attachments).IsNull()"), Package.IsNull());
 		TestTrue(TEXT("FCbPackage(Attachments) as bool"), bool(Package));
 		TestEqual(TEXT("FCbPackage(Attachments).GetAttachments()"), Package.GetAttachments().Num(), 2);
-		TestTrue(TEXT("FCbPackage(Attachments).GetObject()"), Package.GetObject().Equals(FCbObjectRef()));
+		TestTrue(TEXT("FCbPackage(Attachments).GetObject()"), Package.GetObject().Equals(FCbObject()));
 		TestEqual(TEXT("FCbPackage(Attachments).GetObjectHash()"), Package.GetObjectHash(), FIoHash());
 		TestSaveLoadValidate(TEXT("Attachments"), Package);
 
@@ -328,12 +328,12 @@ bool FCbPackageTest::RunTest(const FString& Parameters)
 		const FCbAttachment* const FieldAttachment = Package.FindAttachment(Field.GetHash());
 
 		TestTrue(TEXT("FCbPackage(Attachments).FindAttachment(Object)"),
-			ObjectAttachment && ObjectAttachment->AsCompactBinary().AsObjectRef().Equals(Object));
+			ObjectAttachment && ObjectAttachment->AsCompactBinary().AsObject().Equals(Object));
 		TestTrue(TEXT("FCbPackage(Attachments).FindAttachment(Field)"),
-			FieldAttachment && FieldAttachment->AsBinary() == Field.GetBuffer());
+			FieldAttachment && FieldAttachment->AsBinaryView() == Field.GetBuffer());
 
 		Package.AddAttachment(FCbAttachment(FSharedBuffer::Clone(Object.GetView())));
-		Package.AddAttachment(FCbAttachment(FCbFieldRefIterator::CloneRange(FCbFieldIterator::MakeSingle(Field))));
+		Package.AddAttachment(FCbAttachment(FCbFieldIterator::CloneRange(FCbFieldViewIterator::MakeSingle(Field))));
 
 		TestEqual(TEXT("FCbPackage(Attachments).GetAttachments()"), Package.GetAttachments().Num(), 2);
 		TestEqual(TEXT("FCbPackage(Attachments).FindAttachment(Object, Re-Add)"),
@@ -342,13 +342,13 @@ bool FCbPackageTest::RunTest(const FString& Parameters)
 			Package.FindAttachment(Field.GetHash()), FieldAttachment);
 
 		TestTrue(TEXT("FCbPackage(Attachments).FindAttachment(ObjectAsObject)"),
-			ObjectAttachment && ObjectAttachment->AsCompactBinary().AsObjectRef().Equals(Object));
+			ObjectAttachment && ObjectAttachment->AsCompactBinary().AsObject().Equals(Object));
 		TestTrue(TEXT("FCbPackage(Attachments).FindAttachment(ObjectAsBinary)"),
-			ObjectAttachment && ObjectAttachment->AsBinary() == Object.GetBuffer());
+			ObjectAttachment && ObjectAttachment->AsBinaryView() == Object.GetBuffer());
 		TestTrue(TEXT("FCbPackage(Attachments).FindAttachment(FieldAsField)"),
 			FieldAttachment && FieldAttachment->AsCompactBinary().Equals(Field));
 		TestTrue(TEXT("FCbPackage(Attachments).FindAttachment(FieldAsBinary)"),
-			FieldAttachment && FieldAttachment->AsBinary() == Field.GetBuffer());
+			FieldAttachment && FieldAttachment->AsBinaryView() == Field.GetBuffer());
 
 		TestTrue(TEXT("FCbPackage(Attachments).GetAttachments()->Sorted"),
 			Algo::IsSorted(Package.GetAttachments()));
@@ -359,7 +359,7 @@ bool FCbPackageTest::RunTest(const FString& Parameters)
 	FSharedBuffer Level4 = FSharedBuffer::MakeView(MakeMemoryView(Level4Values));
 	const FIoHash Level4Hash = FIoHash::HashBuffer(Level4);
 
-	FCbFieldRef Level3;
+	FCbField Level3;
 	{
 		TCbWriter<256> Writer;
 		Writer.SetName("Level4").AddBinaryAttachment(Level4Hash);
@@ -367,26 +367,26 @@ bool FCbPackageTest::RunTest(const FString& Parameters)
 	}
 	const FIoHash Level3Hash = Level3.GetHash();
 
-	FCbArrayRef Level2;
+	FCbArray Level2;
 	{
 		TCbWriter<256> Writer;
 		Writer.SetName("Level3");
 		Writer.BeginArray();
 		Writer.AddCompactBinaryAttachment(Level3Hash);
 		Writer.EndArray();
-		Level2 = Writer.Save().AsArrayRef();
+		Level2 = Writer.Save().AsArray();
 	}
-	const FIoHash Level2Hash = Level2.AsField().GetHash();
+	const FIoHash Level2Hash = Level2.AsFieldView().GetHash();
 
-	FCbObjectRef Level1;
+	FCbObject Level1;
 	{
 		TCbWriter<256> Writer;
 		Writer.BeginObject();
 		Writer.SetName("Level2").AddCompactBinaryAttachment(Level2Hash);
 		Writer.EndObject();
-		Level1 = Writer.Save().AsObjectRef();
+		Level1 = Writer.Save().AsObject();
 	}
-	const FIoHash Level1Hash = Level1.AsField().GetHash();
+	const FIoHash Level1Hash = Level1.AsFieldView().GetHash();
 
 	const auto Resolver = [&Level2, &Level2Hash, &Level3, &Level3Hash, &Level4, &Level4Hash]
 		(const FIoHash& Hash) -> FSharedBuffer
@@ -415,13 +415,13 @@ bool FCbPackageTest::RunTest(const FString& Parameters)
 		const FCbAttachment* const Level3Attachment = Package.FindAttachment(Level3Hash);
 		const FCbAttachment* const Level4Attachment = Package.FindAttachment(Level4Hash);
 		TestTrue(TEXT("FCbPackage(Object+Attachments).FindAttachment(Level2)"),
-			Level2Attachment && Level2Attachment->AsCompactBinary().AsArrayRef().Equals(Level2));
+			Level2Attachment && Level2Attachment->AsCompactBinary().AsArray().Equals(Level2));
 		TestTrue(TEXT("FCbPackage(Object+Attachments).FindAttachment(Level3)"),
 			Level3Attachment && Level3Attachment->AsCompactBinary().Equals(Level3));
 		TestTrue(TEXT("FCbPackage(Object+Attachments).FindAttachment(Level4)"),
 			Level4Attachment &&
-			Level4Attachment->AsBinary() != Level4 &&
-			Level4Attachment->AsBinary().GetView().EqualBytes(Level4.GetView()));
+			Level4Attachment->AsBinaryView() != Level4 &&
+			Level4Attachment->AsBinaryView().GetView().EqualBytes(Level4.GetView()));
 
 		TestTrue(TEXT("FCbPackage(Object+Attachments).GetAttachments()->Sorted"),
 			Algo::IsSorted(Package.GetAttachments()));
@@ -444,7 +444,7 @@ bool FCbPackageTest::RunTest(const FString& Parameters)
 		TestNotEqual(TEXT("FCbPackage(Object+Attachments).Equals(AttachmentsNotEqual)"), PackageCopy, Package);
 		Package = PackageCopy;
 		TestEqual(TEXT("FCbPackage(Object+Attachments).Equals(EqualAssigned)"), PackageCopy, Package);
-		Package.SetObject(FCbObjectRef());
+		Package.SetObject(FCbObject());
 		TestNotEqual(TEXT("FCbPackage(Object+Attachments).Equals(ObjectNotEqual)"), PackageCopy, Package);
 		TestEqual(TEXT("FCbPackage(Object+Attachments).GetObjectHash(Null)"), Package.GetObjectHash(), FIoHash());
 	}
@@ -462,7 +462,7 @@ bool FCbPackageTest::RunTest(const FString& Parameters)
 		Writer.AddCompactBinaryAttachment(Level3Hash);
 		Writer.AddNull();
 
-		FCbFieldRefIterator Fields = Writer.Save();
+		FCbFieldIterator Fields = Writer.Save();
 		FCbPackage FromFields;
 		FromFields.Load(Fields);
 
@@ -476,23 +476,23 @@ bool FCbPackageTest::RunTest(const FString& Parameters)
 		TestEqual(TEXT("FCbPackage(OutOfOrder).Load()->Level1Hash"), FromFields.GetObjectHash(), Level1Hash);
 
 		TestTrue(TEXT("FCbPackage(OutOfOrder).Load()->Level2"),
-			Level2Attachment && Level2Attachment->AsCompactBinary().AsArrayRef().Equals(Level2));
+			Level2Attachment && Level2Attachment->AsCompactBinary().AsArray().Equals(Level2));
 		TestTrue(TEXT("FCbPackage(OutOfOrder).Load()->Level2Buffer"),
-			Level2Attachment && Fields.GetOuterBuffer().GetView().Contains(Level2Attachment->AsBinary().GetView()));
+			Level2Attachment && Fields.GetOuterBuffer().GetView().Contains(Level2Attachment->AsBinaryView().GetView()));
 		TestTrue(TEXT("FCbPackage(OutOfOrder).Load()->Level2Hash"),
 			Level2Attachment && Level2Attachment->GetHash() == Level2Hash);
 
 		TestTrue(TEXT("FCbPackage(OutOfOrder).Load()->Level3"),
 			Level3Attachment && Level3Attachment->AsCompactBinary().Equals(Level3));
 		TestTrue(TEXT("FCbPackage(OutOfOrder).Load()->Level3Buffer"),
-			Level3Attachment && Fields.GetOuterBuffer().GetView().Contains(Level3Attachment->AsBinary().GetView()));
+			Level3Attachment && Fields.GetOuterBuffer().GetView().Contains(Level3Attachment->AsBinaryView().GetView()));
 		TestTrue(TEXT("FCbPackage(OutOfOrder).Load()->Level3Hash"),
 			Level3Attachment && Level3Attachment->GetHash() == Level3Hash);
 
 		TestTrue(TEXT("FCbPackage(OutOfOrder).Load()->Level4"),
-			Level4Attachment && Level4Attachment->AsBinary().GetView().EqualBytes(Level4.GetView()));
+			Level4Attachment && Level4Attachment->AsBinaryView().GetView().EqualBytes(Level4.GetView()));
 		TestTrue(TEXT("FCbPackage(OutOfOrder).Load()->Level4Buffer"),
-			Level4Attachment && Fields.GetOuterBuffer().GetView().Contains(Level4Attachment->AsBinary().GetView()));
+			Level4Attachment && Fields.GetOuterBuffer().GetView().Contains(Level4Attachment->AsBinaryView().GetView()));
 		TestTrue(TEXT("FCbPackage(OutOfOrder).Load()->Level4Hash"),
 			Level4Attachment && Level4Attachment->GetHash() == Level4Hash);
 
@@ -504,20 +504,20 @@ bool FCbPackageTest::RunTest(const FString& Parameters)
 
 		Writer.Reset();
 		FromArchive.Save(Writer);
-		FCbFieldRefIterator Saved = Writer.Save();
-		TestTrue(TEXT("FCbPackage(OutOfOrder).Save()->Level1"), Saved.AsObjectRef().Equals(Level1));
+		FCbFieldIterator Saved = Writer.Save();
+		TestTrue(TEXT("FCbPackage(OutOfOrder).Save()->Level1"), Saved.AsObject().Equals(Level1));
 		++Saved;
 		TestEqual(TEXT("FCbPackage(OutOfOrder).Save()->Level1Hash"), Saved.AsCompactBinaryAttachment(), Level1Hash);
 		++Saved;
-		TestTrue(TEXT("FCbPackage(OutOfOrder).Save()->Level2"), Saved.AsBinary().EqualBytes(Level2.GetView()));
+		TestTrue(TEXT("FCbPackage(OutOfOrder).Save()->Level2"), Saved.AsBinaryView().EqualBytes(Level2.GetView()));
 		++Saved;
 		TestEqual(TEXT("FCbPackage(OutOfOrder).Save()->Level2Hash"), Saved.AsCompactBinaryAttachment(), Level2Hash);
 		++Saved;
-		TestTrue(TEXT("FCbPackage(OutOfOrder).Save()->Level3"), Saved.AsBinary().EqualBytes(Level3.GetView()));
+		TestTrue(TEXT("FCbPackage(OutOfOrder).Save()->Level3"), Saved.AsBinaryView().EqualBytes(Level3.GetView()));
 		++Saved;
 		TestEqual(TEXT("FCbPackage(OutOfOrder).Save()->Level3Hash"), Saved.AsCompactBinaryAttachment(), Level3Hash);
 		++Saved;
-		TestTrue(TEXT("FCbPackage(OutOfOrder).Save()->Level4"), Saved.AsBinary().EqualBytes(Level4.GetView()));
+		TestTrue(TEXT("FCbPackage(OutOfOrder).Save()->Level4"), Saved.AsBinaryView().EqualBytes(Level4.GetView()));
 		++Saved;
 		TestEqual(TEXT("FCbPackage(OutOfOrder).Save()->Level4Hash"), Saved.AsBinaryAttachment(), Level4Hash);
 		++Saved;
@@ -542,7 +542,7 @@ bool FCbPackageTest::RunTest(const FString& Parameters)
 		bool bResolved = false;
 		FCbPackage Package;
 		Package.AddAttachment(FCbAttachment(Level3.GetBuffer()));
-		Package.AddAttachment(FCbAttachment(FCbFieldRefIterator::MakeSingle(Level3)),
+		Package.AddAttachment(FCbAttachment(FCbFieldIterator::MakeSingle(Level3)),
 			[&bResolved](const FIoHash& Hash) -> FSharedBuffer
 			{
 				bResolved = true;
diff --git a/Engine/Source/Runtime/Core/Private/Tests/Serialization/CompactBinarySerializationTest.cpp b/Engine/Source/Runtime/Core/Private/Tests/Serialization/CompactBinarySerializationTest.cpp
--- a/Engine/Source/Runtime/Core/Private/Tests/Serialization/CompactBinarySerializationTest.cpp
+++ b/Engine/Source/Runtime/Core/Private/Tests/Serialization/CompactBinarySerializationTest.cpp
@@ -155,7 +155,7 @@ bool FCbSaveTest::RunTest(const FString& Parameters)
 				TestTrue(FString::Printf(TEXT("SaveCompactBinary(%s)->EqualBytes"), Test), ExpectedData.EqualBytes(MakeMemoryView(WriteAr))))
 			{
 				FMemoryReader ReadAr(WriteAr);
-				FCbFieldRef Field = LoadCompactBinary(ReadAr);
+				FCbField Field = LoadCompactBinary(ReadAr);
 				TestTrue(FString::Printf(TEXT("LoadCompactBinary(%s)->EqualBytes"), Test), ExpectedData.EqualBytes(Field.GetView()));
 			}
 		}
@@ -175,20 +175,20 @@ bool FCbSaveTest::RunTest(const FString& Parameters)
 	// Field
 	{
 		const uint8 Payload[] = { uint8(ECbFieldType::IntegerPositive), 42 };
-		TestSave(TEXT("Field"), FCbFieldRef(FCbField(Payload), FSharedBuffer()), MakeMemoryView(Payload));
+		TestSave(TEXT("Field"), FCbField(FCbFieldView(Payload), FSharedBuffer()), MakeMemoryView(Payload));
 	}
 	{
 		const uint8 Payload[] = { uint8(ECbFieldType::IntegerPositive | ECbFieldType::HasFieldName), 1, 'A', 42 };
-		TestSave(TEXT("Field, Name"), FCbFieldRef(FCbField(Payload), FSharedBuffer()), MakeMemoryView(Payload));
+		TestSave(TEXT("Field, Name"), FCbField(FCbFieldView(Payload), FSharedBuffer()), MakeMemoryView(Payload));
 	}
 	{
 		const uint8 Payload[] = { 42 };
-		TestSave(TEXT("Field, NoType"), FCbFieldRef(FCbField(Payload, ECbFieldType::IntegerPositive), FSharedBuffer()),
+		TestSave(TEXT("Field, NoType"), FCbField(FCbFieldView(Payload, ECbFieldType::IntegerPositive), FSharedBuffer()),
 			MakeMemoryView<uint8>({ uint8(ECbFieldType::IntegerPositive), 42 }));
 	}
 	{
 		const uint8 Payload[] = { 1, 'I', 42 };
-		TestSave(TEXT("Field, NoType, Name"), FCbFieldRef(FCbField(Payload, ECbFieldType::IntegerPositive | ECbFieldType::HasFieldName), FSharedBuffer()),
+		TestSave(TEXT("Field, NoType, Name"), FCbField(FCbFieldView(Payload, ECbFieldType::IntegerPositive | ECbFieldType::HasFieldName), FSharedBuffer()),
 			MakeMemoryView<uint8>({ uint8(ECbFieldType::IntegerPositive | ECbFieldType::HasFieldName), 1, 'I', 42 }));
 	}
 
@@ -196,25 +196,25 @@ bool FCbSaveTest::RunTest(const FString& Parameters)
 	{
 		constexpr uint8 IntType = uint8(ECbFieldType::IntegerPositive);
 		const uint8 Payload[] = { uint8(ECbFieldType::UniformArray), 5, 3, IntType, 1, 2, 3 };
-		TestSave(TEXT("Array"), FCbArrayRef(FCbArray(Payload), FSharedBuffer()),
+		TestSave(TEXT("Array"), FCbArray(FCbArrayView(Payload), FSharedBuffer()),
 			MakeMemoryView(Payload));
 	}
 	{
 		constexpr uint8 IntType = uint8(ECbFieldType::IntegerPositive);
 		const uint8 Payload[] = { uint8(ECbFieldType::UniformArray | ECbFieldType::HasFieldName), 1, 'A', 5, 3, IntType, 1, 2, 3 };
-		TestSave(TEXT("Array, Name"), FCbArrayRef(FCbArray(Payload), FSharedBuffer()),
+		TestSave(TEXT("Array, Name"), FCbArray(FCbArrayView(Payload), FSharedBuffer()),
 			MakeMemoryView<uint8>({ uint8(ECbFieldType::UniformArray), 5, 3, IntType, 1, 2, 3 }));
 	}
 	{
 		constexpr uint8 IntType = uint8(ECbFieldType::IntegerPositive);
 		const uint8 Payload[] = { 5, 3, IntType, 1, 2, 3 };
-		TestSave(TEXT("Array, NoType"), FCbArrayRef(FCbArray(Payload, ECbFieldType::UniformArray), FSharedBuffer()),
+		TestSave(TEXT("Array, NoType"), FCbArray(FCbArrayView(Payload, ECbFieldType::UniformArray), FSharedBuffer()),
 			MakeMemoryView<uint8>({ uint8(ECbFieldType::UniformArray), 5, 3, IntType, 1, 2, 3 }));
 	}
 	{
 		constexpr uint8 IntType = uint8(ECbFieldType::IntegerPositive);
 		const uint8 Payload[] = { 1, 'A', 5, 3, IntType, 1, 2, 3 };
-		TestSave(TEXT("Array, NoType, Name"), FCbArrayRef(FCbArray(Payload, ECbFieldType::UniformArray | ECbFieldType::HasFieldName), FSharedBuffer()),
+		TestSave(TEXT("Array, NoType, Name"), FCbArray(FCbArrayView(Payload, ECbFieldType::UniformArray | ECbFieldType::HasFieldName), FSharedBuffer()),
 			MakeMemoryView<uint8>({ uint8(ECbFieldType::UniformArray), 5, 3, IntType, 1, 2, 3 }));
 	}
 
@@ -222,25 +222,25 @@ bool FCbSaveTest::RunTest(const FString& Parameters)
 	{
 		constexpr uint8 IntType = uint8(ECbFieldType::HasFieldName | ECbFieldType::IntegerPositive);
 		const uint8 Payload[] = { uint8(ECbFieldType::UniformObject), 10, IntType, 1, 'A', 1, 1, 'B', 2, 1, 'C', 3 };
-		TestSave(TEXT("Object"), FCbObjectRef(FCbObject(Payload), FSharedBuffer()),
+		TestSave(TEXT("Object"), FCbObject(FCbObjectView(Payload), FSharedBuffer()),
 			MakeMemoryView(Payload));
 	}
 	{
 		constexpr uint8 IntType = uint8(ECbFieldType::HasFieldName | ECbFieldType::IntegerPositive);
 		const uint8 Payload[] = { uint8(ECbFieldType::UniformObject | ECbFieldType::HasFieldName), 1, 'O', 10, IntType, 1, 'A', 1, 1, 'B', 2, 1, 'C', 3 };
-		TestSave(TEXT("Object, Name"), FCbObjectRef(FCbObject(Payload), FSharedBuffer()),
+		TestSave(TEXT("Object, Name"), FCbObject(FCbObjectView(Payload), FSharedBuffer()),
 			MakeMemoryView<uint8>({ uint8(ECbFieldType::UniformObject), 10, IntType, 1, 'A', 1, 1, 'B', 2, 1, 'C', 3 }));
 	}
 	{
 		constexpr uint8 IntType = uint8(ECbFieldType::HasFieldName | ECbFieldType::IntegerPositive);
 		const uint8 Payload[] = { 10, IntType, 1, 'A', 1, 1, 'B', 2, 1, 'C', 3 };
-		TestSave(TEXT("Object, NoType"), FCbObjectRef(FCbObject(Payload, ECbFieldType::UniformObject), FSharedBuffer()),
+		TestSave(TEXT("Object, NoType"), FCbObject(FCbObjectView(Payload, ECbFieldType::UniformObject), FSharedBuffer()),
 			MakeMemoryView<uint8>({ uint8(ECbFieldType::UniformObject), 10, IntType, 1, 'A', 1, 1, 'B', 2, 1, 'C', 3 }));
 	}
 	{
 		constexpr uint8 IntType = uint8(ECbFieldType::HasFieldName | ECbFieldType::IntegerPositive);
 		const uint8 Payload[] = { 1, 'O', 10, IntType, 1, 'A', 1, 1, 'B', 2, 1, 'C', 3 };
-		TestSave(TEXT("Object, NoType, Name"), FCbObjectRef(FCbObject(Payload, ECbFieldType::UniformObject | ECbFieldType::HasFieldName), FSharedBuffer()),
+		TestSave(TEXT("Object, NoType, Name"), FCbObject(FCbObjectView(Payload, ECbFieldType::UniformObject | ECbFieldType::HasFieldName), FSharedBuffer()),
 			MakeMemoryView<uint8>({ uint8(ECbFieldType::UniformObject), 10, IntType, 1, 'A', 1, 1, 'B', 2, 1, 'C', 3 }));
 	}
 
diff --git a/Engine/Source/Runtime/Core/Private/Tests/Serialization/CompactBinaryTest.cpp b/Engine/Source/Runtime/Core/Private/Tests/Serialization/CompactBinaryTest.cpp
--- a/Engine/Source/Runtime/Core/Private/Tests/Serialization/CompactBinaryTest.cpp
+++ b/Engine/Source/Runtime/Core/Private/Tests/Serialization/CompactBinaryTest.cpp
@@ -31,8 +31,8 @@ using TCbFieldValueType = typename TCbFieldTypeAccessors<FCbFieldType::GetType(F
 	struct TCbFieldTypeAccessors<ECbFieldType::FieldType>                                                             \
 	{                                                                                                                 \
 		using ValueType = InValueType;                                                                                \
-		static constexpr bool (FCbField::*IsType)() const = &FCbField::InIsTypeFn;                                    \
-		static auto AsType(FCbField& Field, ValueType) { return Field.InAsTypeFn(); }                                 \
+		static constexpr bool (FCbFieldView::*IsType)() const = &FCbFieldView::InIsTypeFn;                                    \
+		static auto AsType(FCbFieldView& Field, ValueType) { return Field.InAsTypeFn(); }                                 \
 	};
 
 #define UE_CBFIELD_TYPE_ACCESSOR_EX(FieldType, InIsTypeFn, InAsTypeFn, InValueType, InDefaultType)                    \
@@ -40,18 +40,18 @@ using TCbFieldValueType = typename TCbFieldTypeAccessors<FCbFieldType::GetType(F
 	struct TCbFieldTypeAccessors<ECbFieldType::FieldType>                                                             \
 	{                                                                                                                 \
 		using ValueType = InValueType;                                                                                \
-		static constexpr bool (FCbField::*IsType)() const = &FCbField::InIsTypeFn;                                    \
-		static constexpr InValueType (FCbField::*AsType)(InDefaultType) = &FCbField::InAsTypeFn;                      \
+		static constexpr bool (FCbFieldView::*IsType)() const = &FCbFieldView::InIsTypeFn;                                    \
+		static constexpr InValueType (FCbFieldView::*AsType)(InDefaultType) = &FCbFieldView::InAsTypeFn;                      \
 	};
 
 #define UE_CBFIELD_TYPE_ACCESSOR(FieldType, InIsTypeFn, InAsTypeFn, InValueType)                                      \
 	UE_CBFIELD_TYPE_ACCESSOR_EX(FieldType, InIsTypeFn, InAsTypeFn, InValueType, InValueType)
 
-UE_CBFIELD_TYPE_ACCESSOR_NO_DEFAULT(Object, IsObject, AsObject, FCbObject);
-UE_CBFIELD_TYPE_ACCESSOR_NO_DEFAULT(UniformObject, IsObject, AsObject, FCbObject);
-UE_CBFIELD_TYPE_ACCESSOR_NO_DEFAULT(Array, IsArray, AsArray, FCbArray);
-UE_CBFIELD_TYPE_ACCESSOR_NO_DEFAULT(UniformArray, IsArray, AsArray, FCbArray);
-UE_CBFIELD_TYPE_ACCESSOR(Binary, IsBinary, AsBinary, FMemoryView);
+UE_CBFIELD_TYPE_ACCESSOR_NO_DEFAULT(Object, IsObject, AsObjectView, FCbObjectView);
+UE_CBFIELD_TYPE_ACCESSOR_NO_DEFAULT(UniformObject, IsObject, AsObjectView, FCbObjectView);
+UE_CBFIELD_TYPE_ACCESSOR_NO_DEFAULT(Array, IsArray, AsArrayView, FCbArrayView);
+UE_CBFIELD_TYPE_ACCESSOR_NO_DEFAULT(UniformArray, IsArray, AsArrayView, FCbArrayView);
+UE_CBFIELD_TYPE_ACCESSOR(Binary, IsBinary, AsBinaryView, FMemoryView);
 UE_CBFIELD_TYPE_ACCESSOR(String, IsString, AsString, FAnsiStringView);
 UE_CBFIELD_TYPE_ACCESSOR(IntegerPositive, IsInteger, AsUInt64, uint64);
 UE_CBFIELD_TYPE_ACCESSOR(IntegerNegative, IsInteger, AsInt64, int64);
@@ -71,8 +71,8 @@ UE_CBFIELD_TYPE_ACCESSOR(CustomByName, IsCustomByName, AsCustomByName, FCbCustom
 
 struct FCbAttachmentAccessors
 {
-	static constexpr bool (FCbField::*IsType)() const = &FCbField::IsAttachment;
-	static constexpr FIoHash (FCbField::*AsType)(const FIoHash&) = &FCbField::AsAttachment;
+	static constexpr bool (FCbFieldView::*IsType)() const = &FCbFieldView::IsAttachment;
+	static constexpr FIoHash (FCbFieldView::*AsType)(const FIoHash&) = &FCbFieldView::AsAttachment;
 };
 
 ///////////////////////////////////////////////////////////////////////////////////////////////////////////////////////
@@ -83,12 +83,12 @@ protected:
 	using FAutomationTestBase::FAutomationTestBase;
 	using FAutomationTestBase::TestEqual;
 
-	bool TestEqual(const FString& What, const FCbArray& Actual, const FCbArray& Expected)
+	bool TestEqual(const FString& What, const FCbArrayView& Actual, const FCbArrayView& Expected)
 	{
 		return TestTrue(What, Actual.Equals(Expected));
 	}
 
-	bool TestEqual(const FString& What, const FCbObject& Actual, const FCbObject& Expected)
+	bool TestEqual(const FString& What, const FCbObjectView& Actual, const FCbObjectView& Expected)
 	{
 		return TestTrue(What, Actual.Equals(Expected));
 	}
@@ -109,49 +109,49 @@ protected:
 	}
 
 	template <ECbFieldType FieldType, typename T = TCbFieldValueType<FieldType>, typename AccessorsType = TCbFieldTypeAccessors<FieldType>>
-	void TestFieldNoClone(const TCHAR* What, FCbField& Field, T ExpectedValue = T(), T DefaultValue = T(), ECbFieldError ExpectedError = ECbFieldError::None, const AccessorsType& Accessors = AccessorsType())
+	void TestFieldNoClone(const TCHAR* What, FCbFieldView& Field, T ExpectedValue = T(), T DefaultValue = T(), ECbFieldError ExpectedError = ECbFieldError::None, const AccessorsType& Accessors = AccessorsType())
 	{
-		TestEqual(FString::Printf(TEXT("FCbField::Is[Type](%s)"), What), Invoke(Accessors.IsType, Field), ExpectedError != ECbFieldError::TypeError);
+		TestEqual(FString::Printf(TEXT("FCbFieldView::Is[Type](%s)"), What), Invoke(Accessors.IsType, Field), ExpectedError != ECbFieldError::TypeError);
 		if (ExpectedError == ECbFieldError::None && !Field.IsBool())
 		{
-			TestFalse(FString::Printf(TEXT("FCbField::AsBool(%s) == false"), What), Field.AsBool());
-			TestTrue(FString::Printf(TEXT("FCbField::AsBool(%s) -> HasError()"), What), Field.HasError());
-			TestEqual(FString::Printf(TEXT("FCbField::AsBool(%s) -> GetError() == TypeError"), What), Field.GetError(), ECbFieldError::TypeError);
+			TestFalse(FString::Printf(TEXT("FCbFieldView::AsBool(%s) == false"), What), Field.AsBool());
+			TestTrue(FString::Printf(TEXT("FCbFieldView::AsBool(%s) -> HasError()"), What), Field.HasError());
+			TestEqual(FString::Printf(TEXT("FCbFieldView::AsBool(%s) -> GetError() == TypeError"), What), Field.GetError(), ECbFieldError::TypeError);
 		}
-		TestEqual(FString::Printf(TEXT("FCbField::As[Type](%s) -> Equal"), What), Invoke(Accessors.AsType, Field, DefaultValue), ExpectedValue);
-		TestEqual(FString::Printf(TEXT("FCbField::As[Type](%s) -> HasError()"), What), Field.HasError(), ExpectedError != ECbFieldError::None);
-		TestEqual(FString::Printf(TEXT("FCbField::As[Type](%s) -> GetError()"), What), Field.GetError(), ExpectedError);
+		TestEqual(FString::Printf(TEXT("FCbFieldView::As[Type](%s) -> Equal"), What), Invoke(Accessors.AsType, Field, DefaultValue), ExpectedValue);
+		TestEqual(FString::Printf(TEXT("FCbFieldView::As[Type](%s) -> HasError()"), What), Field.HasError(), ExpectedError != ECbFieldError::None);
+		TestEqual(FString::Printf(TEXT("FCbFieldView::As[Type](%s) -> GetError()"), What), Field.GetError(), ExpectedError);
 	}
 
 	template <ECbFieldType FieldType, typename T = TCbFieldValueType<FieldType>, typename AccessorsType = TCbFieldTypeAccessors<FieldType>>
 	void TestFieldNoClone(const TCHAR* What, TArrayView<const uint8> Payload, T ExpectedValue = T(), T DefaultValue = T(), ECbFieldError ExpectedError = ECbFieldError::None, const AccessorsType& Accessors = AccessorsType())
 	{
-		FCbField Field(Payload.GetData(), FieldType);
+		FCbFieldView Field(Payload.GetData(), FieldType);
 		TestFieldNoClone<FieldType>(What, Field, ExpectedValue, DefaultValue, ExpectedError, Accessors);
 	}
 
 	template <ECbFieldType FieldType, typename T = TCbFieldValueType<FieldType>, typename AccessorsType = TCbFieldTypeAccessors<FieldType>>
-	void TestField(const TCHAR* What, FCbField& Field, T ExpectedValue = T(), T DefaultValue = T(), ECbFieldError ExpectedError = ECbFieldError::None, const AccessorsType& Accessors = AccessorsType())
+	void TestField(const TCHAR* What, FCbFieldView& Field, T ExpectedValue = T(), T DefaultValue = T(), ECbFieldError ExpectedError = ECbFieldError::None, const AccessorsType& Accessors = AccessorsType())
 	{
 		TestFieldNoClone<FieldType>(What, Field, ExpectedValue, DefaultValue, ExpectedError, Accessors);
-		FCbFieldRef FieldClone = FCbFieldRef::Clone(Field);
+		FCbField FieldClone = FCbField::Clone(Field);
 		TestFieldNoClone<FieldType>(*FString::Printf(TEXT("%s, Clone"), What), FieldClone, ExpectedValue, DefaultValue, ExpectedError, Accessors);
-		TestTrue(FString::Printf(TEXT("FCbField::Equals(%s)"), What), Field.Equals(FieldClone));
+		TestTrue(FString::Printf(TEXT("FCbFieldView::Equals(%s)"), What), Field.Equals(FieldClone));
 	}
 
 	template <ECbFieldType FieldType, typename T = TCbFieldValueType<FieldType>, typename AccessorsType = TCbFieldTypeAccessors<FieldType>>
 	void TestField(const TCHAR* What, TArrayView<const uint8> Payload, T ExpectedValue = T(), T DefaultValue = T(), ECbFieldError ExpectedError = ECbFieldError::None, const AccessorsType& Accessors = AccessorsType())
 	{
-		FCbField Field(Payload.GetData(), FieldType);
-		TestEqual(FString::Printf(TEXT("FCbField::GetSize(%s)"), What), Field.GetSize(), uint64(Payload.Num() + !FCbFieldType::HasFieldType(FieldType)));
-		TestTrue(FString::Printf(TEXT("FCbField::HasValue(%s)"), What), Field.HasValue());
-		TestFalse(FString::Printf(TEXT("FCbField::HasError(%s) == false"), What), Field.HasError());
-		TestEqual(FString::Printf(TEXT("FCbField::GetError(%s) == None"), What), Field.GetError(), ECbFieldError::None);
+		FCbFieldView Field(Payload.GetData(), FieldType);
+		TestEqual(FString::Printf(TEXT("FCbFieldView::GetSize(%s)"), What), Field.GetSize(), uint64(Payload.Num() + !FCbFieldType::HasFieldType(FieldType)));
+		TestTrue(FString::Printf(TEXT("FCbFieldView::HasValue(%s)"), What), Field.HasValue());
+		TestFalse(FString::Printf(TEXT("FCbFieldView::HasError(%s) == false"), What), Field.HasError());
+		TestEqual(FString::Printf(TEXT("FCbFieldView::GetError(%s) == None"), What), Field.GetError(), ECbFieldError::None);
 		TestField<FieldType>(What, Field, ExpectedValue, DefaultValue, ExpectedError, Accessors);
 	}
 
 	template <ECbFieldType FieldType, typename T = TCbFieldValueType<FieldType>, typename AccessorsType = TCbFieldTypeAccessors<FieldType>>
-	void TestFieldError(const TCHAR* What, FCbField& Field, ECbFieldError ExpectedError, T ExpectedValue = T(), const AccessorsType& Accessors = AccessorsType())
+	void TestFieldError(const TCHAR* What, FCbFieldView& Field, ECbFieldError ExpectedError, T ExpectedValue = T(), const AccessorsType& Accessors = AccessorsType())
 	{
 		TestFieldNoClone<FieldType>(What, Field, ExpectedValue, ExpectedValue, ExpectedError, Accessors);
 	}
@@ -159,7 +159,7 @@ protected:
 	template <ECbFieldType FieldType, typename T = TCbFieldValueType<FieldType>, typename AccessorsType = TCbFieldTypeAccessors<FieldType>>
 	void TestFieldError(const TCHAR* What, TArrayView<const uint8> Payload, ECbFieldError ExpectedError, T ExpectedValue = T(), const AccessorsType& Accessors = AccessorsType())
 	{
-		FCbField Field(Payload.GetData(), FieldType);
+		FCbFieldView Field(Payload.GetData(), FieldType);
 		TestFieldError<FieldType>(What, Field, ExpectedError, ExpectedValue, Accessors);
 	}
 };
@@ -169,106 +169,106 @@ protected:
 IMPLEMENT_CUSTOM_SIMPLE_AUTOMATION_TEST(FCbFieldNoneTest, FCbFieldTestBase, "System.Core.Serialization.CbField.None", CompactBinaryTestFlags)
 bool FCbFieldNoneTest::RunTest(const FString& Parameters)
 {
-	// Test FCbField()
+	// Test FCbFieldView()
 	{
-		constexpr FCbField DefaultField;
+		constexpr FCbFieldView DefaultField;
 		static_assert(!DefaultField.HasName(), "Error in HasName()");
 		static_assert(!DefaultField.HasValue(), "Error in HasValue()");
 		static_assert(!DefaultField.HasError(), "Error in HasError()");
 		static_assert(DefaultField.GetError() == ECbFieldError::None, "Error in GetError()");
-		TestEqual(TEXT("FCbField()::GetSize() == 1"), DefaultField.GetSize(), uint64(1));
-		TestEqual(TEXT("FCbField()::GetName().Len() == 0"), DefaultField.GetName().Len(), 0);
-		TestFalse(TEXT("!FCbField()::HasName()"), DefaultField.HasName());
-		TestFalse(TEXT("!FCbField()::HasValue()"), DefaultField.HasValue());
-		TestFalse(TEXT("!FCbField()::HasError()"), DefaultField.HasError());
-		TestEqual(TEXT("FCbField()::GetError() == None"), DefaultField.GetError(), ECbFieldError::None);
-		TestEqual(TEXT("FCbField()::GetHash()"), DefaultField.GetHash(), FIoHash::HashBuffer(MakeMemoryView<uint8>({uint8(ECbFieldType::None)})));
-		TestEqual(TEXT("FCbField()::GetView()"), DefaultField.GetView(), FMemoryView());
+		TestEqual(TEXT("FCbFieldView()::GetSize() == 1"), DefaultField.GetSize(), uint64(1));
+		TestEqual(TEXT("FCbFieldView()::GetName().Len() == 0"), DefaultField.GetName().Len(), 0);
+		TestFalse(TEXT("!FCbFieldView()::HasName()"), DefaultField.HasName());
+		TestFalse(TEXT("!FCbFieldView()::HasValue()"), DefaultField.HasValue());
+		TestFalse(TEXT("!FCbFieldView()::HasError()"), DefaultField.HasError());
+		TestEqual(TEXT("FCbFieldView()::GetError() == None"), DefaultField.GetError(), ECbFieldError::None);
+		TestEqual(TEXT("FCbFieldView()::GetHash()"), DefaultField.GetHash(), FIoHash::HashBuffer(MakeMemoryView<uint8>({uint8(ECbFieldType::None)})));
+		TestEqual(TEXT("FCbFieldView()::GetView()"), DefaultField.GetView(), FMemoryView());
 		FMemoryView SerializedView;
-		TestFalse(TEXT("FCbField()::TryGetSerializedView()"), DefaultField.TryGetSerializedView(SerializedView));
+		TestFalse(TEXT("FCbFieldView()::TryGetSerializedView()"), DefaultField.TryGetSerializedView(SerializedView));
 	}
 
-	// Test FCbField(None)
+	// Test FCbFieldView(None)
 	{
-		FCbField NoneField(nullptr, ECbFieldType::None);
-		TestEqual(TEXT("FCbField(None)::GetSize() == 1"), NoneField.GetSize(), uint64(1));
-		TestEqual(TEXT("FCbField(None)::GetName().Len() == 0"), NoneField.GetName().Len(), 0);
-		TestFalse(TEXT("!FCbField(None)::HasName()"), NoneField.HasName());
-		TestFalse(TEXT("!FCbField(None)::HasValue()"), NoneField.HasValue());
-		TestFalse(TEXT("!FCbField(None)::HasError()"), NoneField.HasError());
-		TestEqual(TEXT("FCbField(None)::GetError() == None"), NoneField.GetError(), ECbFieldError::None);
-		TestEqual(TEXT("FCbField(None)::GetHash()"), NoneField.GetHash(), FCbField().GetHash());
-		TestEqual(TEXT("FCbField(None)::GetView()"), NoneField.GetView(), FMemoryView());
+		FCbFieldView NoneField(nullptr, ECbFieldType::None);
+		TestEqual(TEXT("FCbFieldView(None)::GetSize() == 1"), NoneField.GetSize(), uint64(1));
+		TestEqual(TEXT("FCbFieldView(None)::GetName().Len() == 0"), NoneField.GetName().Len(), 0);
+		TestFalse(TEXT("!FCbFieldView(None)::HasName()"), NoneField.HasName());
+		TestFalse(TEXT("!FCbFieldView(None)::HasValue()"), NoneField.HasValue());
+		TestFalse(TEXT("!FCbFieldView(None)::HasError()"), NoneField.HasError());
+		TestEqual(TEXT("FCbFieldView(None)::GetError() == None"), NoneField.GetError(), ECbFieldError::None);
+		TestEqual(TEXT("FCbFieldView(None)::GetHash()"), NoneField.GetHash(), FCbFieldView().GetHash());
+		TestEqual(TEXT("FCbFieldView(None)::GetView()"), NoneField.GetView(), FMemoryView());
 		FMemoryView SerializedView;
-		TestFalse(TEXT("FCbField(None)::TryGetSerializedView()"), NoneField.TryGetSerializedView(SerializedView));
+		TestFalse(TEXT("FCbFieldView(None)::TryGetSerializedView()"), NoneField.TryGetSerializedView(SerializedView));
 	}
 
-	// Test FCbField(None|Type|Name)
+	// Test FCbFieldView(None|Type|Name)
 	{
 		constexpr ECbFieldType FieldType = ECbFieldType::None | ECbFieldType::HasFieldName;
 		constexpr const ANSICHAR NoneBytes[] = { ANSICHAR(FieldType), 4, 'N', 'a', 'm', 'e' };
-		FCbField NoneField(NoneBytes);
-		TestEqual(TEXT("FCbField(None|Type|Name)::GetSize()"), NoneField.GetSize(), uint64(sizeof(NoneBytes)));
-		TestEqual(TEXT("FCbField(None|Type|Name)::GetName()"), NoneField.GetName(), "Name"_ASV);
-		TestTrue(TEXT("FCbField(None|Type|Name)::HasName()"), NoneField.HasName());
-		TestFalse(TEXT("!FCbField(None|Type|Name)::HasValue()"), NoneField.HasValue());
-		TestEqual(TEXT("FCbField(None|Type|Name)::GetHash()"), NoneField.GetHash(), FIoHash::HashBuffer(MakeMemoryView(NoneBytes)));
-		TestEqual(TEXT("FCbField(None|Type|Name)::GetView()"), NoneField.GetView(), MakeMemoryView(NoneBytes));
+		FCbFieldView NoneField(NoneBytes);
+		TestEqual(TEXT("FCbFieldView(None|Type|Name)::GetSize()"), NoneField.GetSize(), uint64(sizeof(NoneBytes)));
+		TestEqual(TEXT("FCbFieldView(None|Type|Name)::GetName()"), NoneField.GetName(), "Name"_ASV);
+		TestTrue(TEXT("FCbFieldView(None|Type|Name)::HasName()"), NoneField.HasName());
+		TestFalse(TEXT("!FCbFieldView(None|Type|Name)::HasValue()"), NoneField.HasValue());
+		TestEqual(TEXT("FCbFieldView(None|Type|Name)::GetHash()"), NoneField.GetHash(), FIoHash::HashBuffer(MakeMemoryView(NoneBytes)));
+		TestEqual(TEXT("FCbFieldView(None|Type|Name)::GetView()"), NoneField.GetView(), MakeMemoryView(NoneBytes));
 		FMemoryView SerializedView;
-		TestTrue(TEXT("FCbField(None|Type|Name)::TryGetSerializedView()"), NoneField.TryGetSerializedView(SerializedView) && SerializedView == MakeMemoryView(NoneBytes));
+		TestTrue(TEXT("FCbFieldView(None|Type|Name)::TryGetSerializedView()"), NoneField.TryGetSerializedView(SerializedView) && SerializedView == MakeMemoryView(NoneBytes));
 
 		uint8 CopyBytes[sizeof(NoneBytes)];
 		NoneField.CopyTo(MakeMemoryView(CopyBytes));
-		TestTrue(TEXT("FCbField(None|Type|Name)::CopyTo()"), MakeMemoryView(NoneBytes).EqualBytes(MakeMemoryView(CopyBytes)));
+		TestTrue(TEXT("FCbFieldView(None|Type|Name)::CopyTo()"), MakeMemoryView(NoneBytes).EqualBytes(MakeMemoryView(CopyBytes)));
 	}
 
-	// Test FCbField(None|Type)
+	// Test FCbFieldView(None|Type)
 	{
 		constexpr ECbFieldType FieldType = ECbFieldType::None;
 		constexpr const ANSICHAR NoneBytes[] = { ANSICHAR(FieldType) };
-		FCbField NoneField(NoneBytes);
-		TestEqual(TEXT("FCbField(None|Type)::GetSize()"), NoneField.GetSize(), uint64(sizeof(NoneBytes)));
-		TestEqual(TEXT("FCbField(None|Type)::GetName()"), NoneField.GetName().Len(), 0);
-		TestFalse(TEXT("FCbField(None|Type)::HasName()"), NoneField.HasName());
-		TestFalse(TEXT("!FCbField(None|Type)::HasValue()"), NoneField.HasValue());
-		TestEqual(TEXT("FCbField(None|Type)::GetHash()"), NoneField.GetHash(), FCbField().GetHash());
-		TestEqual(TEXT("FCbField(None|Type)::GetView()"), NoneField.GetView(), MakeMemoryView(NoneBytes));
+		FCbFieldView NoneField(NoneBytes);
+		TestEqual(TEXT("FCbFieldView(None|Type)::GetSize()"), NoneField.GetSize(), uint64(sizeof(NoneBytes)));
+		TestEqual(TEXT("FCbFieldView(None|Type)::GetName()"), NoneField.GetName().Len(), 0);
+		TestFalse(TEXT("FCbFieldView(None|Type)::HasName()"), NoneField.HasName());
+		TestFalse(TEXT("!FCbFieldView(None|Type)::HasValue()"), NoneField.HasValue());
+		TestEqual(TEXT("FCbFieldView(None|Type)::GetHash()"), NoneField.GetHash(), FCbFieldView().GetHash());
+		TestEqual(TEXT("FCbFieldView(None|Type)::GetView()"), NoneField.GetView(), MakeMemoryView(NoneBytes));
 		FMemoryView SerializedView;
-		TestTrue(TEXT("FCbField(None|Type)::TryGetSerializedView()"), NoneField.TryGetSerializedView(SerializedView) && SerializedView == MakeMemoryView(NoneBytes));
+		TestTrue(TEXT("FCbFieldView(None|Type)::TryGetSerializedView()"), NoneField.TryGetSerializedView(SerializedView) && SerializedView == MakeMemoryView(NoneBytes));
 	}
 
-	// Test FCbField(None|Name)
+	// Test FCbFieldView(None|Name)
 	{
 		constexpr ECbFieldType FieldType = ECbFieldType::None | ECbFieldType::HasFieldName;
 		constexpr const ANSICHAR NoneBytes[] = { ANSICHAR(FieldType), 4, 'N', 'a', 'm', 'e' };
-		FCbField NoneField(NoneBytes + 1, FieldType);
-		TestEqual(TEXT("FCbField(None|Name)::GetSize()"), NoneField.GetSize(), uint64(sizeof(NoneBytes)));
-		TestEqual(TEXT("FCbField(None|Name)::GetName()"), NoneField.GetName(), "Name"_ASV);
-		TestTrue(TEXT("FCbField(None|Name)::HasName()"), NoneField.HasName());
-		TestFalse(TEXT("!FCbField(None|Name)::HasValue()"), NoneField.HasValue());
-		TestEqual(TEXT("FCbField(None|Name)::GetHash()"), NoneField.GetHash(), FIoHash::HashBuffer(MakeMemoryView(NoneBytes)));
-		TestEqual(TEXT("FCbField(None|Name)::GetView()"), NoneField.GetView(), MakeMemoryView(NoneBytes) + 1);
+		FCbFieldView NoneField(NoneBytes + 1, FieldType);
+		TestEqual(TEXT("FCbFieldView(None|Name)::GetSize()"), NoneField.GetSize(), uint64(sizeof(NoneBytes)));
+		TestEqual(TEXT("FCbFieldView(None|Name)::GetName()"), NoneField.GetName(), "Name"_ASV);
+		TestTrue(TEXT("FCbFieldView(None|Name)::HasName()"), NoneField.HasName());
+		TestFalse(TEXT("!FCbFieldView(None|Name)::HasValue()"), NoneField.HasValue());
+		TestEqual(TEXT("FCbFieldView(None|Name)::GetHash()"), NoneField.GetHash(), FIoHash::HashBuffer(MakeMemoryView(NoneBytes)));
+		TestEqual(TEXT("FCbFieldView(None|Name)::GetView()"), NoneField.GetView(), MakeMemoryView(NoneBytes) + 1);
 		FMemoryView SerializedView;
-		TestFalse(TEXT("FCbField(None|Name)::TryGetSerializedView()"), NoneField.TryGetSerializedView(SerializedView));
+		TestFalse(TEXT("FCbFieldView(None|Name)::TryGetSerializedView()"), NoneField.TryGetSerializedView(SerializedView));
 
 		uint8 CopyBytes[sizeof(NoneBytes)];
 		NoneField.CopyTo(MakeMemoryView(CopyBytes));
-		TestTrue(TEXT("FCbField(None|Name)::CopyTo()"), MakeMemoryView(NoneBytes).EqualBytes(MakeMemoryView(CopyBytes)));
+		TestTrue(TEXT("FCbFieldView(None|Name)::CopyTo()"), MakeMemoryView(NoneBytes).EqualBytes(MakeMemoryView(CopyBytes)));
 	}
 
-	// Test FCbField(None|EmptyName)
+	// Test FCbFieldView(None|EmptyName)
 	{
 		constexpr ECbFieldType FieldType = ECbFieldType::None | ECbFieldType::HasFieldName;
 		constexpr const uint8 NoneBytes[] = { uint8(FieldType), 0 };
-		FCbField NoneField(NoneBytes + 1, FieldType);
-		TestEqual(TEXT("FCbField(None|EmptyName)::GetSize()"), NoneField.GetSize(), uint64(sizeof(NoneBytes)));
-		TestEqual(TEXT("FCbField(None|EmptyName)::GetName()"), NoneField.GetName(), ""_ASV);
-		TestTrue(TEXT("FCbField(None|EmptyName)::HasName()"), NoneField.HasName());
-		TestFalse(TEXT("!FCbField(None|EmptyName)::HasValue()"), NoneField.HasValue());
-		TestEqual(TEXT("FCbField(None|EmptyName)::GetHash()"), NoneField.GetHash(), FIoHash::HashBuffer(MakeMemoryView(NoneBytes)));
-		TestEqual(TEXT("FCbField(None|EmptyName)::GetView()"), NoneField.GetView(), MakeMemoryView(NoneBytes) + 1);
+		FCbFieldView NoneField(NoneBytes + 1, FieldType);
+		TestEqual(TEXT("FCbFieldView(None|EmptyName)::GetSize()"), NoneField.GetSize(), uint64(sizeof(NoneBytes)));
+		TestEqual(TEXT("FCbFieldView(None|EmptyName)::GetName()"), NoneField.GetName(), ""_ASV);
+		TestTrue(TEXT("FCbFieldView(None|EmptyName)::HasName()"), NoneField.HasName());
+		TestFalse(TEXT("!FCbFieldView(None|EmptyName)::HasValue()"), NoneField.HasValue());
+		TestEqual(TEXT("FCbFieldView(None|EmptyName)::GetHash()"), NoneField.GetHash(), FIoHash::HashBuffer(MakeMemoryView(NoneBytes)));
+		TestEqual(TEXT("FCbFieldView(None|EmptyName)::GetView()"), NoneField.GetView(), MakeMemoryView(NoneBytes) + 1);
 		FMemoryView SerializedView;
-		TestFalse(TEXT("FCbField(None|EmptyName)::TryGetSerializedView()"), NoneField.TryGetSerializedView(SerializedView));
+		TestFalse(TEXT("FCbFieldView(None|EmptyName)::TryGetSerializedView()"), NoneField.TryGetSerializedView(SerializedView));
 	}
 
 	return true;
@@ -277,21 +277,21 @@ bool FCbFieldNoneTest::RunTest(const FString& Parameters)
 IMPLEMENT_CUSTOM_SIMPLE_AUTOMATION_TEST(FCbFieldNullTest, FCbFieldTestBase, "System.Core.Serialization.CbField.Null", CompactBinaryTestFlags)
 bool FCbFieldNullTest::RunTest(const FString& Parameters)
 {
-	// Test FCbField(Null)
+	// Test FCbFieldView(Null)
 	{
-		FCbField NullField(nullptr, ECbFieldType::Null);
-		TestEqual(TEXT("FCbField(Null)::GetSize() == 1"), NullField.GetSize(), uint64(1));
-		TestTrue(TEXT("FCbField(Null)::IsNull()"), NullField.IsNull());
-		TestTrue(TEXT("FCbField(Null)::HasValue()"), NullField.HasValue());
-		TestFalse(TEXT("!FCbField(Null)::HasError()"), NullField.HasError());
-		TestEqual(TEXT("FCbField(Null)::GetError() == None"), NullField.GetError(), ECbFieldError::None);
-		TestEqual(TEXT("FCbField(Null)::GetHash()"), NullField.GetHash(), FIoHash::HashBuffer(MakeMemoryView<uint8>({uint8(ECbFieldType::Null)})));
+		FCbFieldView NullField(nullptr, ECbFieldType::Null);
+		TestEqual(TEXT("FCbFieldView(Null)::GetSize() == 1"), NullField.GetSize(), uint64(1));
+		TestTrue(TEXT("FCbFieldView(Null)::IsNull()"), NullField.IsNull());
+		TestTrue(TEXT("FCbFieldView(Null)::HasValue()"), NullField.HasValue());
+		TestFalse(TEXT("!FCbFieldView(Null)::HasError()"), NullField.HasError());
+		TestEqual(TEXT("FCbFieldView(Null)::GetError() == None"), NullField.GetError(), ECbFieldError::None);
+		TestEqual(TEXT("FCbFieldView(Null)::GetHash()"), NullField.GetHash(), FIoHash::HashBuffer(MakeMemoryView<uint8>({uint8(ECbFieldType::Null)})));
 	}
 
-	// Test FCbField(None) as Null
+	// Test FCbFieldView(None) as Null
 	{
-		FCbField Field;
-		TestFalse(TEXT("FCbField(None)::IsNull()"), Field.IsNull());
+		FCbFieldView Field;
+		TestFalse(TEXT("FCbFieldView(None)::IsNull()"), Field.IsNull());
 	}
 
 	return true;
@@ -300,185 +300,185 @@ bool FCbFieldNullTest::RunTest(const FString& Parameters)
 IMPLEMENT_CUSTOM_SIMPLE_AUTOMATION_TEST(FCbFieldObjectTest, FCbFieldTestBase, "System.Core.Serialization.CbField.Object", CompactBinaryTestFlags)
 bool FCbFieldObjectTest::RunTest(const FString& Parameters)
 {
-	static_assert(!std::is_constructible<FCbField, const FCbObject&>::value, "Invalid constructor for FCbField");
-	static_assert(!std::is_assignable<FCbField, const FCbObject&>::value, "Invalid assignment for FCbField");
-	static_assert(!std::is_convertible<FCbField, FCbObject>::value, "Invalid conversion to FCbObject");
-	static_assert(!std::is_assignable<FCbObject, const FCbField&>::value, "Invalid assignment for FCbObject");
+	static_assert(!std::is_constructible<FCbFieldView, const FCbObjectView&>::value, "Invalid constructor for FCbFieldView");
+	static_assert(!std::is_assignable<FCbFieldView, const FCbObjectView&>::value, "Invalid assignment for FCbFieldView");
+	static_assert(!std::is_convertible<FCbFieldView, FCbObjectView>::value, "Invalid conversion to FCbObjectView");
+	static_assert(!std::is_assignable<FCbObjectView, const FCbFieldView&>::value, "Invalid assignment for FCbObjectView");
 
-	auto TestIntObject = [this](const FCbObject& Object, int32 ExpectedNum, uint64 ExpectedPayloadSize)
+	auto TestIntObject = [this](const FCbObjectView& Object, int32 ExpectedNum, uint64 ExpectedPayloadSize)
 	{
-		TestEqual(TEXT("FCbField(Object)::AsObject().GetSize()"), Object.GetSize(), ExpectedPayloadSize + sizeof(ECbFieldType));
+		TestEqual(TEXT("FCbFieldView(Object)::AsObjectView().GetSize()"), Object.GetSize(), ExpectedPayloadSize + sizeof(ECbFieldType));
 
 		int32 ActualNum = 0;
-		for (FCbFieldIterator It = Object.CreateIterator(); It; ++It)
+		for (FCbFieldViewIterator It = Object.CreateViewIterator(); It; ++It)
 		{
 			++ActualNum;
-			TestNotEqual(TEXT("FCbField(Object) Iterator Name"), It->GetName().Len(), 0);
-			TestEqual(TEXT("FCbField(Object) Iterator"), It->AsInt32(), ActualNum);
+			TestNotEqual(TEXT("FCbFieldView(Object) Iterator Name"), It->GetName().Len(), 0);
+			TestEqual(TEXT("FCbFieldView(Object) Iterator"), It->AsInt32(), ActualNum);
 		}
-		TestEqual(TEXT("FCbField(Object)::AsObject().CreateIterator() -> Count"), ActualNum, ExpectedNum);
+		TestEqual(TEXT("FCbFieldView(Object)::AsObjectView().CreateViewIterator() -> Count"), ActualNum, ExpectedNum);
 
 		ActualNum = 0;
-		for (FCbField Field : Object)
+		for (FCbFieldView Field : Object)
 		{
 			++ActualNum;
-			TestNotEqual(TEXT("FCbField(Object) Iterator Name"), Field.GetName().Len(), 0);
-			TestEqual(TEXT("FCbField(Object) Range"), Field.AsInt32(), ActualNum);
+			TestNotEqual(TEXT("FCbFieldView(Object) Iterator Name"), Field.GetName().Len(), 0);
+			TestEqual(TEXT("FCbFieldView(Object) Range"), Field.AsInt32(), ActualNum);
 		}
-		TestEqual(TEXT("FCbField(Object)::AsObject() Range -> Count"), ActualNum, ExpectedNum);
+		TestEqual(TEXT("FCbFieldView(Object)::AsObjectView() Range -> Count"), ActualNum, ExpectedNum);
 	};
 
-	// Test FCbField(Object, Empty)
+	// Test FCbFieldView(Object, Empty)
 	TestField<ECbFieldType::Object>(TEXT("Object, Empty"), {0});
 
-	// Test FCbField(Object, Empty)
+	// Test FCbFieldView(Object, Empty)
 	{
-		FCbObject Object;
+		FCbObjectView Object;
 		TestIntObject(Object, 0, 1);
 
 		// Find fields that do not exist.
-		TestFalse(TEXT("FCbObject()::Find(Missing)"), Object.Find("Field"_ASV).HasValue());
-		TestFalse(TEXT("FCbObject()::FindIgnoreCase(Missing)"), Object.FindIgnoreCase("Field"_ASV).HasValue());
-		TestFalse(TEXT("FCbObject()::operator[](Missing)"), Object["Field"_ASV].HasValue());
+		TestFalse(TEXT("FCbObjectView()::Find(Missing)"), Object.FindView("Field"_ASV).HasValue());
+		TestFalse(TEXT("FCbObjectView()::FindViewIgnoreCase(Missing)"), Object.FindViewIgnoreCase("Field"_ASV).HasValue());
+		TestFalse(TEXT("FCbObjectView()::operator[](Missing)"), Object["Field"_ASV].HasValue());
 
 		// Advance an iterator past the last field.
-		FCbFieldIterator It = Object.CreateIterator();
-		TestFalse(TEXT("FCbObject()::CreateIterator() At End"), bool(It));
-		TestTrue(TEXT("FCbObject()::CreateIterator() At End"), !It);
+		FCbFieldViewIterator It = Object.CreateViewIterator();
+		TestFalse(TEXT("FCbObjectView()::CreateViewIterator() At End"), bool(It));
+		TestTrue(TEXT("FCbObjectView()::CreateViewIterator() At End"), !It);
 		for (int Count = 16; Count > 0; --Count)
 		{
 			++It;
 			It->AsInt32();
 		}
-		TestFalse(TEXT("FCbObject()::CreateIterator() At End"), bool(It));
-		TestTrue(TEXT("FCbObject()::CreateIterator() At End"), !It);
+		TestFalse(TEXT("FCbObjectView()::CreateViewIterator() At End"), bool(It));
+		TestTrue(TEXT("FCbObjectView()::CreateViewIterator() At End"), !It);
 	}
 
-	// Test FCbField(Object, NotEmpty)
+	// Test FCbFieldView(Object, NotEmpty)
 	{
 		constexpr uint8 IntType = uint8(ECbFieldType::HasFieldName | ECbFieldType::IntegerPositive);
 		const uint8 Payload[] = { 12, IntType, 1, 'A', 1, IntType, 1, 'B', 2, IntType, 1, 'C', 3 };
-		FCbField Field(Payload, ECbFieldType::Object);
-		TestField<ECbFieldType::Object>(TEXT("Object, NotEmpty"), Field, FCbObject(Payload, ECbFieldType::Object));
-		FCbObjectRef Object = FCbObjectRef::Clone(Field.AsObject());
+		FCbFieldView Field(Payload, ECbFieldType::Object);
+		TestField<ECbFieldType::Object>(TEXT("Object, NotEmpty"), Field, FCbObjectView(Payload, ECbFieldType::Object));
+		FCbObject Object = FCbObject::Clone(Field.AsObjectView());
 		TestIntObject(Object, 3, sizeof(Payload));
-		TestIntObject(Field.AsObject(), 3, sizeof(Payload));
-		TestTrue(TEXT("FCbObject::Equals()"), Object.Equals(Field.AsObject()));
-		TestEqual(TEXT("FCbObject::Find()"), Object.Find("B"_ASV).AsInt32(), 2);
-		TestEqual(TEXT("FCbObject::Find()"), Object.Find("b"_ASV).AsInt32(4), 4);
-		TestEqual(TEXT("FCbObject::FindIgnoreCase()"), Object.FindIgnoreCase("B"_ASV).AsInt32(), 2);
-		TestEqual(TEXT("FCbObject::FindIgnoreCase()"), Object.FindIgnoreCase("b"_ASV).AsInt32(), 2);
-		TestEqual(TEXT("FCbObject::operator[]"), Object["B"_ASV].AsInt32(), 2);
-		TestEqual(TEXT("FCbObject::operator[]"), Object["b"_ASV].AsInt32(4), 4);
-		TestEqual(TEXT("FCbObject::GetView()"), Field.AsObject().GetView(), Field.GetView());
+		TestIntObject(Field.AsObjectView(), 3, sizeof(Payload));
+		TestTrue(TEXT("FCbObjectView::Equals()"), Object.Equals(Field.AsObjectView()));
+		TestEqual(TEXT("FCbObjectView::Find()"), Object.FindView("B"_ASV).AsInt32(), 2);
+		TestEqual(TEXT("FCbObjectView::Find()"), Object.FindView("b"_ASV).AsInt32(4), 4);
+		TestEqual(TEXT("FCbObjectView::FindViewIgnoreCase()"), Object.FindViewIgnoreCase("B"_ASV).AsInt32(), 2);
+		TestEqual(TEXT("FCbObjectView::FindViewIgnoreCase()"), Object.FindViewIgnoreCase("b"_ASV).AsInt32(), 2);
+		TestEqual(TEXT("FCbObjectView::operator[]"), Object["B"_ASV].AsInt32(), 2);
+		TestEqual(TEXT("FCbObjectView::operator[]"), Object["b"_ASV].AsInt32(4), 4);
+		TestEqual(TEXT("FCbObjectView::GetView()"), Field.AsObjectView().GetView(), Field.GetView());
 	}
 
-	// Test FCbField(UniformObject, NotEmpty)
+	// Test FCbFieldView(UniformObject, NotEmpty)
 	{
 		constexpr uint8 IntType = uint8(ECbFieldType::HasFieldName | ECbFieldType::IntegerPositive);
 		const uint8 Payload[] = { 10, IntType, 1, 'A', 1, 1, 'B', 2, 1, 'C', 3 };
-		FCbField Field(Payload, ECbFieldType::UniformObject);
-		TestField<ECbFieldType::UniformObject>(TEXT("UniformObject, NotEmpty"), Field, FCbObject(Payload, ECbFieldType::UniformObject));
-		FCbObjectRef Object = FCbObjectRef::Clone(Field.AsObject());
+		FCbFieldView Field(Payload, ECbFieldType::UniformObject);
+		TestField<ECbFieldType::UniformObject>(TEXT("UniformObject, NotEmpty"), Field, FCbObjectView(Payload, ECbFieldType::UniformObject));
+		FCbObject Object = FCbObject::Clone(Field.AsObjectView());
 		TestIntObject(Object, 3, sizeof(Payload));
-		TestIntObject(Field.AsObject(), 3, sizeof(Payload));
-		TestTrue(TEXT("FCbObject(Uniform)::Equals()"), Object.Equals(Field.AsObject()));
-		TestEqual(TEXT("FCbObject(Uniform)::Find()"), Object.Find("B"_ASV).AsInt32(), 2);
-		TestEqual(TEXT("FCbObject(Uniform)::Find()"), Object.FindRef("B"_ASV).AsInt32(), 2);
-		TestEqual(TEXT("FCbObject(Uniform)::Find()"), Object.Find("b"_ASV).AsInt32(4), 4);
-		TestEqual(TEXT("FCbObject(Uniform)::Find()"), Object.FindRef("b"_ASV).AsInt32(4), 4);
-		TestEqual(TEXT("FCbObject(Uniform)::FindIgnoreCase()"), Object.FindIgnoreCase("B"_ASV).AsInt32(), 2);
-		TestEqual(TEXT("FCbObject(Uniform)::FindIgnoreCase()"), Object.FindRefIgnoreCase("B"_ASV).AsInt32(), 2);
-		TestEqual(TEXT("FCbObject(Uniform)::FindIgnoreCase()"), Object.FindIgnoreCase("b"_ASV).AsInt32(), 2);
-		TestEqual(TEXT("FCbObject(Uniform)::FindIgnoreCase()"), Object.FindRefIgnoreCase("b"_ASV).AsInt32(), 2);
-		TestEqual(TEXT("FCbObject(Uniform)::operator[]"), Object["B"_ASV].AsInt32(), 2);
-		TestEqual(TEXT("FCbObject(Uniform)::operator[]"), Object["b"_ASV].AsInt32(4), 4);
-		TestEqual(TEXT("FCbObject(Uniform)::GetView()"), Field.AsObject().GetView(), Field.GetView());
-
-		TestTrue(TEXT("FCbObjectRef::AsFieldRef()"), Object.GetBuffer() == Object.AsFieldRef().AsObjectRef().GetBuffer());
+		TestIntObject(Field.AsObjectView(), 3, sizeof(Payload));
+		TestTrue(TEXT("FCbObjectView(Uniform)::Equals()"), Object.Equals(Field.AsObjectView()));
+		TestEqual(TEXT("FCbObjectView(Uniform)::Find()"), Object.FindView("B"_ASV).AsInt32(), 2);
+		TestEqual(TEXT("FCbObjectView(Uniform)::Find()"), Object.Find("B"_ASV).AsInt32(), 2);
+		TestEqual(TEXT("FCbObjectView(Uniform)::Find()"), Object.FindView("b"_ASV).AsInt32(4), 4);
+		TestEqual(TEXT("FCbObjectView(Uniform)::Find()"), Object.Find("b"_ASV).AsInt32(4), 4);
+		TestEqual(TEXT("FCbObjectView(Uniform)::FindViewIgnoreCase()"), Object.FindViewIgnoreCase("B"_ASV).AsInt32(), 2);
+		TestEqual(TEXT("FCbObjectView(Uniform)::FindViewIgnoreCase()"), Object.FindIgnoreCase("B"_ASV).AsInt32(), 2);
+		TestEqual(TEXT("FCbObjectView(Uniform)::FindViewIgnoreCase()"), Object.FindViewIgnoreCase("b"_ASV).AsInt32(), 2);
+		TestEqual(TEXT("FCbObjectView(Uniform)::FindViewIgnoreCase()"), Object.FindIgnoreCase("b"_ASV).AsInt32(), 2);
+		TestEqual(TEXT("FCbObjectView(Uniform)::operator[]"), Object["B"_ASV].AsInt32(), 2);
+		TestEqual(TEXT("FCbObjectView(Uniform)::operator[]"), Object["b"_ASV].AsInt32(4), 4);
+		TestEqual(TEXT("FCbObjectView(Uniform)::GetView()"), Field.AsObjectView().GetView(), Field.GetView());
+
+		TestTrue(TEXT("FCbObject::AsField()"), Object.GetBuffer() == Object.AsField().AsObject().GetBuffer());
 
 		// Equals
 		const uint8 NamedPayload[] = { 1, 'O', 10, IntType, 1, 'A', 1, 1, 'B', 2, 1, 'C', 3 };
-		FCbField NamedField(NamedPayload, ECbFieldType::UniformObject | ECbFieldType::HasFieldName);
-		TestTrue(TEXT("FCbObject::Equals()"), Field.AsObject().Equals(NamedField.AsObject()));
-		TestTrue(TEXT("FCbObject::AsField().Equals()"), Field.Equals(Field.AsObject().AsField()));
-		TestTrue(TEXT("FCbObject::AsField().Equals()"), NamedField.Equals(NamedField.AsObject().AsField()));
-		TestEqual(TEXT("FCbObject(Name)::GetView()"), NamedField.AsObject().GetView(), NamedField.GetView());
+		FCbFieldView NamedField(NamedPayload, ECbFieldType::UniformObject | ECbFieldType::HasFieldName);
+		TestTrue(TEXT("FCbObjectView::Equals()"), Field.AsObjectView().Equals(NamedField.AsObjectView()));
+		TestTrue(TEXT("FCbObjectView::AsFieldView().Equals()"), Field.Equals(Field.AsObjectView().AsFieldView()));
+		TestTrue(TEXT("FCbObjectView::AsFieldView().Equals()"), NamedField.Equals(NamedField.AsObjectView().AsFieldView()));
+		TestEqual(TEXT("FCbObjectView(Name)::GetView()"), NamedField.AsObjectView().GetView(), NamedField.GetView());
 
 		// CopyTo
 		uint8 CopyBytes[sizeof(Payload) + 1];
-		Field.AsObject().CopyTo(MakeMemoryView(CopyBytes));
-		TestTrue(TEXT("FCbObject(NoType)::CopyTo()"), MakeMemoryView(Payload).EqualBytes(MakeMemoryView(CopyBytes) + 1));
-		NamedField.AsObject().CopyTo(MakeMemoryView(CopyBytes));
-		TestTrue(TEXT("FCbObject(NoType, Name)::CopyTo()"), MakeMemoryView(Payload).EqualBytes(MakeMemoryView(CopyBytes) + 1));
+		Field.AsObjectView().CopyTo(MakeMemoryView(CopyBytes));
+		TestTrue(TEXT("FCbObjectView(NoType)::CopyTo()"), MakeMemoryView(Payload).EqualBytes(MakeMemoryView(CopyBytes) + 1));
+		NamedField.AsObjectView().CopyTo(MakeMemoryView(CopyBytes));
+		TestTrue(TEXT("FCbObjectView(NoType, Name)::CopyTo()"), MakeMemoryView(Payload).EqualBytes(MakeMemoryView(CopyBytes) + 1));
 
 		// GetSerializedView
 		FMemoryView SerializedView;
-		TestTrue(TEXT("FCbObject(Clone)::TryGetSerializedView()"), Object.TryGetSerializedView(SerializedView) && SerializedView == Object.AsField().GetView());
-		TestFalse(TEXT("FCbObject(NoType)::TryGetSerializedView()"), Field.AsObject().TryGetSerializedView(SerializedView));
-		TestFalse(TEXT("FCbObject(Name)::TryGetSerializedView()"), NamedField.AsObject().TryGetSerializedView(SerializedView));
+		TestTrue(TEXT("FCbObjectView(Clone)::TryGetSerializedView()"), Object.TryGetSerializedView(SerializedView) && SerializedView == Object.AsFieldView().GetView());
+		TestFalse(TEXT("FCbObjectView(NoType)::TryGetSerializedView()"), Field.AsObjectView().TryGetSerializedView(SerializedView));
+		TestFalse(TEXT("FCbObjectView(Name)::TryGetSerializedView()"), NamedField.AsObjectView().TryGetSerializedView(SerializedView));
 	}
 
-	// Test FCbField(None) as Object
+	// Test FCbFieldView(None) as Object
 	{
-		FCbField Field;
+		FCbFieldView Field;
 		TestFieldError<ECbFieldType::Object>(TEXT("Object, None"), Field, ECbFieldError::TypeError);
-		FCbFieldRef::MakeView(Field).AsObjectRef();
+		FCbField::MakeView(Field).AsObject();
 	}
 
-	// Test FCbObject(ObjectWithName) and CreateRefIterator
+	// Test FCbObjectView(ObjectWithName) and CreateIterator
 	{
 		const uint8 ObjectType = uint8(ECbFieldType::Object | ECbFieldType::HasFieldName);
 		const uint8 Buffer[] = { ObjectType, 3, 'K', 'e', 'y', 4, uint8(ECbFieldType::HasFieldName | ECbFieldType::IntegerPositive), 1, 'F', 8 };
-		const FCbObject Object(Buffer);
-		TestEqual(TEXT("FCbObject(ObjectWithName)::GetSize()"), Object.GetSize(), uint64(6));
-		const FCbObjectRef ObjectClone = FCbObjectRef::Clone(Object);
-		TestEqual(TEXT("FCbObjectRef(ObjectWithName)::GetSize()"), ObjectClone.GetSize(), uint64(6));
-		TestTrue(TEXT("FCbObject::Equals()"), Object.Equals(ObjectClone));
-		TestEqual(TEXT("FCbObject::GetHash()"), ObjectClone.GetHash(), Object.GetHash());
-		for (FCbFieldRefIterator It = ObjectClone.CreateRefIterator(); It; ++It)
+		const FCbObjectView Object(Buffer);
+		TestEqual(TEXT("FCbObjectView(ObjectWithName)::GetSize()"), Object.GetSize(), uint64(6));
+		const FCbObject ObjectClone = FCbObject::Clone(Object);
+		TestEqual(TEXT("FCbObject(ObjectWithName)::GetSize()"), ObjectClone.GetSize(), uint64(6));
+		TestTrue(TEXT("FCbObjectView::Equals()"), Object.Equals(ObjectClone));
+		TestEqual(TEXT("FCbObjectView::GetHash()"), ObjectClone.GetHash(), Object.GetHash());
+		for (FCbFieldIterator It = ObjectClone.CreateIterator(); It; ++It)
 		{
-			FCbFieldRef Field = *It;
-			TestEqual(TEXT("FCbObjectRef::CreateRefIterator().GetName()"), Field.GetName(), "F"_ASV);
-			TestEqual(TEXT("FCbObjectRef::CreateRefIterator().AsInt32()"), Field.AsInt32(), 8);
-			TestTrue(TEXT("FCbObjectRef::CreateRefIterator().IsOwned()"), Field.IsOwned());
-			TestEqual(TEXT("FCbObjectRef::CreateRefIterator().GetBuffer()"), Field.GetBuffer().GetView(), Field.GetView());
+			FCbField Field = *It;
+			TestEqual(TEXT("FCbObject::CreateIterator().GetName()"), Field.GetName(), "F"_ASV);
+			TestEqual(TEXT("FCbObject::CreateIterator().AsInt32()"), Field.AsInt32(), 8);
+			TestTrue(TEXT("FCbObject::CreateIterator().IsOwned()"), Field.IsOwned());
+			TestEqual(TEXT("FCbObject::CreateIterator().GetBuffer()"), Field.GetBuffer().GetView(), Field.GetView());
 		}
-		for (FCbFieldRefIterator It = ObjectClone.CreateRefIterator(), End; It != End; ++It)
+		for (FCbFieldIterator It = ObjectClone.CreateIterator(), End; It != End; ++It)
 		{
 		}
 
 		// CopyTo
 		uint8 CopyBytes[6];
 		Object.CopyTo(MakeMemoryView(CopyBytes));
-		TestTrue(TEXT("FCbObject(Name)::CopyTo()"), ObjectClone.GetView().EqualBytes(MakeMemoryView(CopyBytes)));
+		TestTrue(TEXT("FCbObjectView(Name)::CopyTo()"), ObjectClone.GetView().EqualBytes(MakeMemoryView(CopyBytes)));
 		ObjectClone.CopyTo(MakeMemoryView(CopyBytes));
-		TestTrue(TEXT("FCbObject()::CopyTo()"), ObjectClone.GetView().EqualBytes(MakeMemoryView(CopyBytes)));
+		TestTrue(TEXT("FCbObjectView()::CopyTo()"), ObjectClone.GetView().EqualBytes(MakeMemoryView(CopyBytes)));
 	}
 
-	// Test FCbObject as FCbFieldIterator
+	// Test FCbObjectView as FCbFieldViewIterator
 	{
 		uint32 Count = 0;
-		FCbObject Object;
-		for (FCbField Field : FCbFieldIterator::MakeSingle(Object.AsField()))
+		FCbObjectView Object;
+		for (FCbFieldView Field : FCbFieldViewIterator::MakeSingle(Object.AsFieldView()))
 		{
-			TestTrue(TEXT("FCbObject::AsField() as Iterator"), Field.IsObject());
+			TestTrue(TEXT("FCbObjectView::AsFieldView() as Iterator"), Field.IsObject());
 			++Count;
 		}
-		TestEqual(TEXT("FCbObject::AsField() as Iterator Count"), Count, 1u);
+		TestEqual(TEXT("FCbObjectView::AsFieldView() as Iterator Count"), Count, 1u);
 	}
 
-	// Test FCbObjectRef as FCbFieldRefIterator
+	// Test FCbObject as FCbFieldIterator
 	{
 		uint32 Count = 0;
-		FCbObjectRef Object;
+		FCbObject Object;
 		Object.MakeOwned();
-		for (FCbFieldRef Field : FCbFieldRefIterator::MakeSingle(Object.AsFieldRef()))
+		for (FCbField Field : FCbFieldIterator::MakeSingle(Object.AsField()))
 		{
-			TestTrue(TEXT("FCbObjectRef::AsField() as Iterator"), Field.IsObject());
+			TestTrue(TEXT("FCbObject::AsFieldView() as Iterator"), Field.IsObject());
 			++Count;
 		}
-		TestEqual(TEXT("FCbObjectRef::AsField() as Iterator Count"), Count, 1u);
+		TestEqual(TEXT("FCbObject::AsFieldView() as Iterator Count"), Count, 1u);
 	}
 
 	return true;
@@ -487,162 +487,162 @@ bool FCbFieldObjectTest::RunTest(const FString& Parameters)
 IMPLEMENT_CUSTOM_SIMPLE_AUTOMATION_TEST(FCbFieldArrayTest, FCbFieldTestBase, "System.Core.Serialization.CbField.Array", CompactBinaryTestFlags)
 bool FCbFieldArrayTest::RunTest(const FString& Parameters)
 {
-	static_assert(!std::is_constructible<FCbField, const FCbArray&>::value, "Invalid constructor for FCbField");
-	static_assert(!std::is_assignable<FCbField, const FCbArray&>::value, "Invalid assignment for FCbField");
-	static_assert(!std::is_convertible<FCbField, FCbArray>::value, "Invalid conversion to FCbArray");
-	static_assert(!std::is_assignable<FCbArray, const FCbField&>::value, "Invalid assignment for FCbArray");
+	static_assert(!std::is_constructible<FCbFieldView, const FCbArrayView&>::value, "Invalid constructor for FCbFieldView");
+	static_assert(!std::is_assignable<FCbFieldView, const FCbArrayView&>::value, "Invalid assignment for FCbFieldView");
+	static_assert(!std::is_convertible<FCbFieldView, FCbArrayView>::value, "Invalid conversion to FCbArrayView");
+	static_assert(!std::is_assignable<FCbArrayView, const FCbFieldView&>::value, "Invalid assignment for FCbArrayView");
 
-	auto TestIntArray = [this](FCbArray Array, int32 ExpectedNum, uint64 ExpectedPayloadSize)
+	auto TestIntArray = [this](FCbArrayView Array, int32 ExpectedNum, uint64 ExpectedPayloadSize)
 	{
-		TestEqual(TEXT("FCbField(Array)::AsArray().GetSize()"), Array.GetSize(), ExpectedPayloadSize + sizeof(ECbFieldType));
-		TestEqual(TEXT("FCbField(Array)::AsArray().Num()"), Array.Num(), uint64(ExpectedNum));
+		TestEqual(TEXT("FCbFieldView(Array)::AsArrayView().GetSize()"), Array.GetSize(), ExpectedPayloadSize + sizeof(ECbFieldType));
+		TestEqual(TEXT("FCbFieldView(Array)::AsArrayView().Num()"), Array.Num(), uint64(ExpectedNum));
 
 		int32 ActualNum = 0;
-		for (FCbFieldIterator It = Array.CreateIterator(); It; ++It)
+		for (FCbFieldViewIterator It = Array.CreateViewIterator(); It; ++It)
 		{
 			++ActualNum;
-			TestEqual(TEXT("FCbField(Array) Iterator"), It->AsInt32(), ActualNum);
+			TestEqual(TEXT("FCbFieldView(Array) Iterator"), It->AsInt32(), ActualNum);
 		}
-		TestEqual(TEXT("FCbField(Array)::AsArray().CreateIterator() -> Count"), ActualNum, ExpectedNum);
+		TestEqual(TEXT("FCbFieldView(Array)::AsArrayView().CreateViewIterator() -> Count"), ActualNum, ExpectedNum);
 
 		ActualNum = 0;
-		for (FCbField Field : Array)
+		for (FCbFieldView Field : Array)
 		{
 			++ActualNum;
-			TestEqual(TEXT("FCbField(Array) Range"), Field.AsInt32(), ActualNum);
+			TestEqual(TEXT("FCbFieldView(Array) Range"), Field.AsInt32(), ActualNum);
 		}
-		TestEqual(TEXT("FCbField(Array)::AsArray() Range -> Count"), ActualNum, ExpectedNum);
+		TestEqual(TEXT("FCbFieldView(Array)::AsArrayView() Range -> Count"), ActualNum, ExpectedNum);
 	};
 
-	// Test FCbField(Array, Empty)
+	// Test FCbFieldView(Array, Empty)
 	TestField<ECbFieldType::Array>(TEXT("Array, Empty"), {1, 0});
 
-	// Test FCbField(Array, Empty)
+	// Test FCbFieldView(Array, Empty)
 	{
-		FCbArray Array;
+		FCbArrayView Array;
 		TestIntArray(Array, 0, 2);
 
 		// Advance an iterator past the last field.
-		FCbFieldIterator It = Array.CreateIterator();
-		TestFalse(TEXT("FCbArray()::CreateIterator() At End"), bool(It));
-		TestTrue(TEXT("FCbArray()::CreateIterator() At End"), !It);
+		FCbFieldViewIterator It = Array.CreateViewIterator();
+		TestFalse(TEXT("FCbArrayView()::CreateViewIterator() At End"), bool(It));
+		TestTrue(TEXT("FCbArrayView()::CreateViewIterator() At End"), !It);
 		for (int Count = 16; Count > 0; --Count)
 		{
 			++It;
 			It->AsInt32();
 		}
-		TestFalse(TEXT("FCbArray()::CreateIterator() At End"), bool(It));
-		TestTrue(TEXT("FCbArray()::CreateIterator() At End"), !It);
+		TestFalse(TEXT("FCbArrayView()::CreateViewIterator() At End"), bool(It));
+		TestTrue(TEXT("FCbArrayView()::CreateViewIterator() At End"), !It);
 	}
 
-	// Test FCbField(Array, NotEmpty)
+	// Test FCbFieldView(Array, NotEmpty)
 	{
 		constexpr uint8 IntType = uint8(ECbFieldType::IntegerPositive);
 		const uint8 Payload[] = { 7, 3, IntType, 1, IntType, 2, IntType, 3 };
-		FCbField Field(Payload, ECbFieldType::Array);
-		TestField<ECbFieldType::Array>(TEXT("Array, NotEmpty"), Field, FCbArray(Payload, ECbFieldType::Array));
-		FCbArrayRef Array = FCbArrayRef::Clone(Field.AsArray());
+		FCbFieldView Field(Payload, ECbFieldType::Array);
+		TestField<ECbFieldType::Array>(TEXT("Array, NotEmpty"), Field, FCbArrayView(Payload, ECbFieldType::Array));
+		FCbArray Array = FCbArray::Clone(Field.AsArrayView());
 		TestIntArray(Array, 3, sizeof(Payload));
-		TestIntArray(Field.AsArray(), 3, sizeof(Payload));
-		TestTrue(TEXT("FCbArray::Equals()"), Array.Equals(Field.AsArray()));
-		TestEqual(TEXT("FCbArray::GetView()"), Field.AsArray().GetView(), Field.GetView());
+		TestIntArray(Field.AsArrayView(), 3, sizeof(Payload));
+		TestTrue(TEXT("FCbArrayView::Equals()"), Array.Equals(Field.AsArrayView()));
+		TestEqual(TEXT("FCbArrayView::GetView()"), Field.AsArrayView().GetView(), Field.GetView());
 	}
 
-	// Test FCbField(UniformArray)
+	// Test FCbFieldView(UniformArray)
 	{
 		constexpr uint8 IntType = uint8(ECbFieldType::IntegerPositive);
 		const uint8 Payload[] = { 5, 3, IntType, 1, 2, 3 };
-		FCbField Field(Payload, ECbFieldType::UniformArray);
-		TestField<ECbFieldType::UniformArray>(TEXT("UniformArray"), Field, FCbArray(Payload, ECbFieldType::UniformArray));
-		FCbArrayRef Array = FCbArrayRef::Clone(Field.AsArray());
+		FCbFieldView Field(Payload, ECbFieldType::UniformArray);
+		TestField<ECbFieldType::UniformArray>(TEXT("UniformArray"), Field, FCbArrayView(Payload, ECbFieldType::UniformArray));
+		FCbArray Array = FCbArray::Clone(Field.AsArrayView());
 		TestIntArray(Array, 3, sizeof(Payload));
-		TestIntArray(Field.AsArray(), 3, sizeof(Payload));
-		TestTrue(TEXT("FCbArray(Uniform)::Equals()"), Array.Equals(Field.AsArray()));
-		TestEqual(TEXT("FCbArray(Uniform)::GetView()"), Field.AsArray().GetView(), Field.GetView());
+		TestIntArray(Field.AsArrayView(), 3, sizeof(Payload));
+		TestTrue(TEXT("FCbArrayView(Uniform)::Equals()"), Array.Equals(Field.AsArrayView()));
+		TestEqual(TEXT("FCbArrayView(Uniform)::GetView()"), Field.AsArrayView().GetView(), Field.GetView());
 
-		TestTrue(TEXT("FCbArrayRef::AsFieldRef()"), Array.GetBuffer() == Array.AsFieldRef().AsArrayRef().GetBuffer());
+		TestTrue(TEXT("FCbArray::AsField()"), Array.GetBuffer() == Array.AsField().AsArray().GetBuffer());
 
 		// Equals
 		const uint8 NamedPayload[] = { 1, 'A', 5, 3, IntType, 1, 2, 3 };
-		FCbField NamedField(NamedPayload, ECbFieldType::UniformArray | ECbFieldType::HasFieldName);
-		TestTrue(TEXT("FCbArray::Equals()"), Field.AsArray().Equals(NamedField.AsArray()));
-		TestTrue(TEXT("FCbArray::AsField().Equals()"), Field.Equals(Field.AsArray().AsField()));
-		TestTrue(TEXT("FCbArray::AsField().Equals()"), NamedField.Equals(NamedField.AsArray().AsField()));
-		TestEqual(TEXT("FCbArray(Name)::GetView()"), NamedField.AsArray().GetView(), NamedField.GetView());
+		FCbFieldView NamedField(NamedPayload, ECbFieldType::UniformArray | ECbFieldType::HasFieldName);
+		TestTrue(TEXT("FCbArrayView::Equals()"), Field.AsArrayView().Equals(NamedField.AsArrayView()));
+		TestTrue(TEXT("FCbArrayView::AsFieldView().Equals()"), Field.Equals(Field.AsArrayView().AsFieldView()));
+		TestTrue(TEXT("FCbArrayView::AsFieldView().Equals()"), NamedField.Equals(NamedField.AsArrayView().AsFieldView()));
+		TestEqual(TEXT("FCbArrayView(Name)::GetView()"), NamedField.AsArrayView().GetView(), NamedField.GetView());
 
 		// CopyTo
 		uint8 CopyBytes[sizeof(Payload) + 1];
-		Field.AsArray().CopyTo(MakeMemoryView(CopyBytes));
-		TestTrue(TEXT("FCbArray(NoType)::CopyTo()"), MakeMemoryView(Payload).EqualBytes(MakeMemoryView(CopyBytes) + 1));
-		NamedField.AsArray().CopyTo(MakeMemoryView(CopyBytes));
-		TestTrue(TEXT("FCbArray(NoType, Name)::CopyTo()"), MakeMemoryView(Payload).EqualBytes(MakeMemoryView(CopyBytes) + 1));
+		Field.AsArrayView().CopyTo(MakeMemoryView(CopyBytes));
+		TestTrue(TEXT("FCbArrayView(NoType)::CopyTo()"), MakeMemoryView(Payload).EqualBytes(MakeMemoryView(CopyBytes) + 1));
+		NamedField.AsArrayView().CopyTo(MakeMemoryView(CopyBytes));
+		TestTrue(TEXT("FCbArrayView(NoType, Name)::CopyTo()"), MakeMemoryView(Payload).EqualBytes(MakeMemoryView(CopyBytes) + 1));
 
 		// GetSerializedView
 		FMemoryView SerializedView;
-		TestTrue(TEXT("FCbArray(Clone)::TryGetSerializedView()"), Array.TryGetSerializedView(SerializedView) && SerializedView == Array.AsField().GetView());
-		TestFalse(TEXT("FCbArray(NoType)::TryGetSerializedView()"), Field.AsArray().TryGetSerializedView(SerializedView));
-		TestFalse(TEXT("FCbArray(Name)::TryGetSerializedView()"), NamedField.AsArray().TryGetSerializedView(SerializedView));
+		TestTrue(TEXT("FCbArrayView(Clone)::TryGetSerializedView()"), Array.TryGetSerializedView(SerializedView) && SerializedView == Array.AsFieldView().GetView());
+		TestFalse(TEXT("FCbArrayView(NoType)::TryGetSerializedView()"), Field.AsArrayView().TryGetSerializedView(SerializedView));
+		TestFalse(TEXT("FCbArrayView(Name)::TryGetSerializedView()"), NamedField.AsArrayView().TryGetSerializedView(SerializedView));
 	}
 
-	// Test FCbField(None) as Array
+	// Test FCbFieldView(None) as Array
 	{
-		FCbField Field;
+		FCbFieldView Field;
 		TestFieldError<ECbFieldType::Array>(TEXT("Array, None"), Field, ECbFieldError::TypeError);
-		FCbFieldRef::MakeView(Field).AsArrayRef();
+		FCbField::MakeView(Field).AsArray();
 	}
 
-	// Test FCbArray(ArrayWithName) and CreateRefIterator
+	// Test FCbArrayView(ArrayWithName) and CreateIterator
 	{
 		const uint8 ArrayType = uint8(ECbFieldType::Array | ECbFieldType::HasFieldName);
 		const uint8 Buffer[] = { ArrayType, 3, 'K', 'e', 'y', 3, 1, uint8(ECbFieldType::IntegerPositive), 8 };
-		const FCbArray Array(Buffer);
-		TestEqual(TEXT("FCbArray(ArrayWithName)::GetSize()"), Array.GetSize(), uint64(5));
-		const FCbArrayRef ArrayClone = FCbArrayRef::Clone(Array);
-		TestEqual(TEXT("FCbArrayRef(ArrayWithName)::GetSize()"), ArrayClone.GetSize(), uint64(5));
-		TestTrue(TEXT("FCbArray::Equals()"), Array.Equals(ArrayClone));
-		TestEqual(TEXT("FCbArray::GetHash()"), ArrayClone.GetHash(), Array.GetHash());
-		for (FCbFieldRefIterator It = ArrayClone.CreateRefIterator(); It; ++It)
+		const FCbArrayView Array(Buffer);
+		TestEqual(TEXT("FCbArrayView(ArrayWithName)::GetSize()"), Array.GetSize(), uint64(5));
+		const FCbArray ArrayClone = FCbArray::Clone(Array);
+		TestEqual(TEXT("FCbArray(ArrayWithName)::GetSize()"), ArrayClone.GetSize(), uint64(5));
+		TestTrue(TEXT("FCbArrayView::Equals()"), Array.Equals(ArrayClone));
+		TestEqual(TEXT("FCbArrayView::GetHash()"), ArrayClone.GetHash(), Array.GetHash());
+		for (FCbFieldIterator It = ArrayClone.CreateIterator(); It; ++It)
 		{
-			FCbFieldRef Field = *It;
-			TestEqual(TEXT("FCbArrayRef::CreateRefIterator().AsInt32()"), Field.AsInt32(), 8);
-			TestTrue(TEXT("FCbArrayRef::CreateRefIterator().IsOwned()"), Field.IsOwned());
-			TestEqual(TEXT("FCbArrayRef::CreateRefIterator().GetBuffer()"), Field.GetBuffer().GetView(), Field.GetView());
+			FCbField Field = *It;
+			TestEqual(TEXT("FCbArray::CreateIterator().AsInt32()"), Field.AsInt32(), 8);
+			TestTrue(TEXT("FCbArray::CreateIterator().IsOwned()"), Field.IsOwned());
+			TestEqual(TEXT("FCbArray::CreateIterator().GetBuffer()"), Field.GetBuffer().GetView(), Field.GetView());
 		}
-		for (FCbFieldRefIterator It = ArrayClone.CreateRefIterator(), End; It != End; ++It)
+		for (FCbFieldIterator It = ArrayClone.CreateIterator(), End; It != End; ++It)
 		{
 		}
 
 		// CopyTo
 		uint8 CopyBytes[5];
 		Array.CopyTo(MakeMemoryView(CopyBytes));
-		TestTrue(TEXT("FCbArray(Name)::CopyTo()"), ArrayClone.GetView().EqualBytes(MakeMemoryView(CopyBytes)));
+		TestTrue(TEXT("FCbArrayView(Name)::CopyTo()"), ArrayClone.GetView().EqualBytes(MakeMemoryView(CopyBytes)));
 		ArrayClone.CopyTo(MakeMemoryView(CopyBytes));
-		TestTrue(TEXT("FCbArray()::CopyTo()"), ArrayClone.GetView().EqualBytes(MakeMemoryView(CopyBytes)));
+		TestTrue(TEXT("FCbArrayView()::CopyTo()"), ArrayClone.GetView().EqualBytes(MakeMemoryView(CopyBytes)));
 	}
 
-	// Test FCbArray as FCbFieldIterator
+	// Test FCbArrayView as FCbFieldViewIterator
 	{
 		uint32 Count = 0;
-		FCbArray Array;
-		for (FCbField Field : FCbFieldIterator::MakeSingle(Array.AsField()))
+		FCbArrayView Array;
+		for (FCbFieldView Field : FCbFieldViewIterator::MakeSingle(Array.AsFieldView()))
 		{
-			TestTrue(TEXT("FCbArray::AsField() as Iterator"), Field.IsArray());
+			TestTrue(TEXT("FCbArrayView::AsFieldView() as Iterator"), Field.IsArray());
 			++Count;
 		}
-		TestEqual(TEXT("FCbArray::AsField() as Iterator Count"), Count, 1u);
+		TestEqual(TEXT("FCbArrayView::AsFieldView() as Iterator Count"), Count, 1u);
 	}
 
-	// Test FCbArrayRef as FCbFieldRefIterator
+	// Test FCbArray as FCbFieldIterator
 	{
 		uint32 Count = 0;
-		FCbArrayRef Array;
+		FCbArray Array;
 		Array.MakeOwned();
-		for (FCbFieldRef Field : FCbFieldRefIterator::MakeSingle(Array.AsFieldRef()))
+		for (FCbField Field : FCbFieldIterator::MakeSingle(Array.AsField()))
 		{
-			TestTrue(TEXT("FCbArrayRef::AsField() as Iterator"), Field.IsArray());
+			TestTrue(TEXT("FCbArray::AsFieldView() as Iterator"), Field.IsArray());
 			++Count;
 		}
-		TestEqual(TEXT("FCbArrayRef::AsField() as Iterator Count"), Count, 1u);
+		TestEqual(TEXT("FCbArray::AsFieldView() as Iterator Count"), Count, 1u);
 	}
 
 	return true;
@@ -651,19 +651,19 @@ bool FCbFieldArrayTest::RunTest(const FString& Parameters)
 IMPLEMENT_CUSTOM_SIMPLE_AUTOMATION_TEST(FCbFieldBinaryTest, FCbFieldTestBase, "System.Core.Serialization.CbField.Binary", CompactBinaryTestFlags)
 bool FCbFieldBinaryTest::RunTest(const FString& Parameters)
 {
-	// Test FCbField(Binary, Empty)
+	// Test FCbFieldView(Binary, Empty)
 	TestField<ECbFieldType::Binary>(TEXT("Binary, Empty"), {0});
 
-	// Test FCbField(Binary, Value)
+	// Test FCbFieldView(Binary, Value)
 	{
 		const uint8 Payload[] = { 3, 4, 5, 6 }; // Size: 3, Data: 4/5/6
-		FCbField Field(Payload, ECbFieldType::Binary);
+		FCbFieldView Field(Payload, ECbFieldType::Binary);
 		TestFieldNoClone<ECbFieldType::Binary>(TEXT("Binary, Value"), Field, MakeMemoryView(Payload + 1, 3));
 	}
 
-	// Test FCbField(None) as Binary
+	// Test FCbFieldView(None) as Binary
 	{
-		FCbField Field;
+		FCbFieldView Field;
 		const uint8 Default[] = { 1, 2, 3 };
 		TestFieldError<ECbFieldType::Binary>(TEXT("Binary, None"), Field, ECbFieldError::TypeError, MakeMemoryView(Default));
 	}
@@ -674,25 +674,25 @@ bool FCbFieldBinaryTest::RunTest(const FString& Parameters)
 IMPLEMENT_CUSTOM_SIMPLE_AUTOMATION_TEST(FCbFieldStringTest, FCbFieldTestBase, "System.Core.Serialization.CbField.String", CompactBinaryTestFlags)
 bool FCbFieldStringTest::RunTest(const FString& Parameters)
 {
-	// Test FCbField(String, Empty)
+	// Test FCbFieldView(String, Empty)
 	TestField<ECbFieldType::String>(TEXT("String, Empty"), {0});
 
-	// Test FCbField(String, Value)
+	// Test FCbFieldView(String, Value)
 	{
 		const uint8 Payload[] = { 3, 'A', 'B', 'C' }; // Size: 3, Data: ABC
 		TestField<ECbFieldType::String>(TEXT("String, Value"), Payload, FAnsiStringView(reinterpret_cast<const ANSICHAR*>(Payload) + 1, 3));
 	}
 
-	// Test FCbField(String, OutOfRangeSize)
+	// Test FCbFieldView(String, OutOfRangeSize)
 	{
 		uint8 Payload[9];
 		WriteVarUInt(uint64(1) << 31, Payload);
 		TestFieldError<ECbFieldType::String>(TEXT("String, OutOfRangeSize"), Payload, ECbFieldError::RangeError, "ABC"_ASV);
 	}
 
-	// Test FCbField(None) as String
+	// Test FCbFieldView(None) as String
 	{
-		FCbField Field;
+		FCbFieldView Field;
 		TestFieldError<ECbFieldType::String>(TEXT("String, None"), Field, ECbFieldError::TypeError, "ABC"_ASV);
 	}
 
@@ -731,11 +731,11 @@ protected:
 		Neg7  = Neg15 | Int8,
 	};
 
-	template <typename T, T (FCbField::*InAsTypeFn)(T)>
+	template <typename T, T (FCbFieldView::*InAsTypeFn)(T)>
 	struct TCbIntegerAccessors
 	{
-		static constexpr bool (FCbField::*IsType)() const = &FCbField::IsInteger;
-		static constexpr T (FCbField::*AsType)(T) = InAsTypeFn;
+		static constexpr bool (FCbFieldView::*IsType)() const = &FCbFieldView::IsInteger;
+		static constexpr T (FCbFieldView::*AsType)(T) = InAsTypeFn;
 	};
 
 	void TestIntegerField(ECbFieldType FieldType, EIntType ExpectedMask, uint64 Magnitude)
@@ -745,30 +745,30 @@ protected:
 		WriteVarUInt(Magnitude - Negative, Payload);
 		constexpr uint64 DefaultValue = 8;
 		const uint64 ExpectedValue = Negative ? uint64(-int64(Magnitude)) : Magnitude;
-		FCbField Field(Payload, FieldType);
+		FCbFieldView Field(Payload, FieldType);
 		TestField<ECbFieldType::IntegerNegative>(TEXT("Int8"), Field, int8(EnumHasAnyFlags(ExpectedMask, EIntType::Int8) ? ExpectedValue : DefaultValue),
-			int8(DefaultValue), EnumHasAnyFlags(ExpectedMask, EIntType::Int8) ? ECbFieldError::None : ECbFieldError::RangeError, TCbIntegerAccessors<int8, &FCbField::AsInt8>());
+			int8(DefaultValue), EnumHasAnyFlags(ExpectedMask, EIntType::Int8) ? ECbFieldError::None : ECbFieldError::RangeError, TCbIntegerAccessors<int8, &FCbFieldView::AsInt8>());
 		TestField<ECbFieldType::IntegerNegative>(TEXT("Int16"), Field, int16(EnumHasAnyFlags(ExpectedMask, EIntType::Int16) ? ExpectedValue : DefaultValue),
-			int16(DefaultValue), EnumHasAnyFlags(ExpectedMask, EIntType::Int16) ? ECbFieldError::None : ECbFieldError::RangeError, TCbIntegerAccessors<int16, &FCbField::AsInt16>());
+			int16(DefaultValue), EnumHasAnyFlags(ExpectedMask, EIntType::Int16) ? ECbFieldError::None : ECbFieldError::RangeError, TCbIntegerAccessors<int16, &FCbFieldView::AsInt16>());
 		TestField<ECbFieldType::IntegerNegative>(TEXT("Int32"), Field, int32(EnumHasAnyFlags(ExpectedMask, EIntType::Int32) ? ExpectedValue : DefaultValue),
-			int32(DefaultValue), EnumHasAnyFlags(ExpectedMask, EIntType::Int32) ? ECbFieldError::None : ECbFieldError::RangeError, TCbIntegerAccessors<int32, &FCbField::AsInt32>());
+			int32(DefaultValue), EnumHasAnyFlags(ExpectedMask, EIntType::Int32) ? ECbFieldError::None : ECbFieldError::RangeError, TCbIntegerAccessors<int32, &FCbFieldView::AsInt32>());
 		TestField<ECbFieldType::IntegerNegative>(TEXT("Int64"), Field, int64(EnumHasAnyFlags(ExpectedMask, EIntType::Int64) ? ExpectedValue : DefaultValue),
-			int64(DefaultValue), EnumHasAnyFlags(ExpectedMask, EIntType::Int64) ? ECbFieldError::None : ECbFieldError::RangeError, TCbIntegerAccessors<int64, &FCbField::AsInt64>());
+			int64(DefaultValue), EnumHasAnyFlags(ExpectedMask, EIntType::Int64) ? ECbFieldError::None : ECbFieldError::RangeError, TCbIntegerAccessors<int64, &FCbFieldView::AsInt64>());
 		TestField<ECbFieldType::IntegerPositive>(TEXT("UInt8"), Field, uint8(EnumHasAnyFlags(ExpectedMask, EIntType::UInt8) ? ExpectedValue : DefaultValue),
-			uint8(DefaultValue), EnumHasAnyFlags(ExpectedMask, EIntType::UInt8) ? ECbFieldError::None : ECbFieldError::RangeError, TCbIntegerAccessors<uint8, &FCbField::AsUInt8>());
+			uint8(DefaultValue), EnumHasAnyFlags(ExpectedMask, EIntType::UInt8) ? ECbFieldError::None : ECbFieldError::RangeError, TCbIntegerAccessors<uint8, &FCbFieldView::AsUInt8>());
 		TestField<ECbFieldType::IntegerPositive>(TEXT("UInt16"), Field, uint16(EnumHasAnyFlags(ExpectedMask, EIntType::UInt16) ? ExpectedValue : DefaultValue),
-			uint16(DefaultValue), EnumHasAnyFlags(ExpectedMask, EIntType::UInt16) ? ECbFieldError::None : ECbFieldError::RangeError, TCbIntegerAccessors<uint16, &FCbField::AsUInt16>());
+			uint16(DefaultValue), EnumHasAnyFlags(ExpectedMask, EIntType::UInt16) ? ECbFieldError::None : ECbFieldError::RangeError, TCbIntegerAccessors<uint16, &FCbFieldView::AsUInt16>());
 		TestField<ECbFieldType::IntegerPositive>(TEXT("UInt32"), Field, uint32(EnumHasAnyFlags(ExpectedMask, EIntType::UInt32) ? ExpectedValue : DefaultValue),
-			uint32(DefaultValue), EnumHasAnyFlags(ExpectedMask, EIntType::UInt32) ? ECbFieldError::None : ECbFieldError::RangeError, TCbIntegerAccessors<uint32, &FCbField::AsUInt32>());
+			uint32(DefaultValue), EnumHasAnyFlags(ExpectedMask, EIntType::UInt32) ? ECbFieldError::None : ECbFieldError::RangeError, TCbIntegerAccessors<uint32, &FCbFieldView::AsUInt32>());
 		TestField<ECbFieldType::IntegerPositive>(TEXT("UInt64"), Field, uint64(EnumHasAnyFlags(ExpectedMask, EIntType::UInt64) ? ExpectedValue : DefaultValue),
-			uint64(DefaultValue), EnumHasAnyFlags(ExpectedMask, EIntType::UInt64) ? ECbFieldError::None : ECbFieldError::RangeError, TCbIntegerAccessors<uint64, &FCbField::AsUInt64>());
+			uint64(DefaultValue), EnumHasAnyFlags(ExpectedMask, EIntType::UInt64) ? ECbFieldError::None : ECbFieldError::RangeError, TCbIntegerAccessors<uint64, &FCbFieldView::AsUInt64>());
 	}
 };
 
 IMPLEMENT_CUSTOM_SIMPLE_AUTOMATION_TEST(FCbFieldIntegerTest, FCbFieldIntegerTestBase, "System.Core.Serialization.CbField.Integer", CompactBinaryTestFlags)
 bool FCbFieldIntegerTest::RunTest(const FString& Parameters)
 {
-	// Test FCbField(IntegerPositive)
+	// Test FCbFieldView(IntegerPositive)
 	TestIntegerField(ECbFieldType::IntegerPositive, EIntType::Pos7,  0x00);
 	TestIntegerField(ECbFieldType::IntegerPositive, EIntType::Pos7,  0x7f);
 	TestIntegerField(ECbFieldType::IntegerPositive, EIntType::Pos8,  0x80);
@@ -786,7 +786,7 @@ bool FCbFieldIntegerTest::RunTest(const FString& Parameters)
 	TestIntegerField(ECbFieldType::IntegerPositive, EIntType::Pos64, 0x8000'0000'0000'0000);
 	TestIntegerField(ECbFieldType::IntegerPositive, EIntType::Pos64, 0xffff'ffff'ffff'ffff);
 
-	// Test FCbField(IntegerNegative)
+	// Test FCbFieldView(IntegerNegative)
 	TestIntegerField(ECbFieldType::IntegerNegative, EIntType::Neg7,  0x01);
 	TestIntegerField(ECbFieldType::IntegerNegative, EIntType::Neg7,  0x80);
 	TestIntegerField(ECbFieldType::IntegerNegative, EIntType::Neg15, 0x81);
@@ -798,9 +798,9 @@ bool FCbFieldIntegerTest::RunTest(const FString& Parameters)
 	TestIntegerField(ECbFieldType::IntegerNegative, EIntType::None,  0x8000'0000'0000'0001);
 	TestIntegerField(ECbFieldType::IntegerNegative, EIntType::None,  0xffff'ffff'ffff'ffff);
 
-	// Test FCbField(None) as Integer
+	// Test FCbFieldView(None) as Integer
 	{
-		FCbField Field;
+		FCbFieldView Field;
 		TestFieldError<ECbFieldType::IntegerPositive>(TEXT("Integer+, None"), Field, ECbFieldError::TypeError, uint64(8));
 		TestFieldError<ECbFieldType::IntegerNegative>(TEXT("Integer-, None"), Field, ECbFieldError::TypeError, int64(8));
 	}
@@ -811,108 +811,108 @@ bool FCbFieldIntegerTest::RunTest(const FString& Parameters)
 IMPLEMENT_CUSTOM_SIMPLE_AUTOMATION_TEST(FCbFieldFloatTest, FCbFieldTestBase, "System.Core.Serialization.CbField.Float", CompactBinaryTestFlags)
 bool FCbFieldFloatTest::RunTest(const FString& Parameters)
 {
-	// Test FCbField(Float, 32-bit)
+	// Test FCbFieldView(Float, 32-bit)
 	{
 		const uint8 Payload[] = { 0xc0, 0x12, 0x34, 0x56 }; // -2.28444433f
 		TestField<ECbFieldType::Float32>(TEXT("Float32"), Payload, -2.28444433f);
 
-		FCbField Field(Payload, ECbFieldType::Float32);
+		FCbFieldView Field(Payload, ECbFieldType::Float32);
 		TestField<ECbFieldType::Float64>(TEXT("Float32, AsDouble"), Field, -2.28444433);
 	}
 
-	// Test FCbField(Float, 64-bit)
+	// Test FCbFieldView(Float, 64-bit)
 	{
 		const uint8 Payload[] = { 0xc1, 0x23, 0x45, 0x67, 0x89, 0xab, 0xcd, 0xef }; // -631475.76888888876
 		TestField<ECbFieldType::Float64>(TEXT("Float64"), Payload, -631475.76888888876);
 
-		FCbField Field(Payload, ECbFieldType::Float64);
+		FCbFieldView Field(Payload, ECbFieldType::Float64);
 		TestFieldError<ECbFieldType::Float32>(TEXT("Float64, AsFloat"), Field, ECbFieldError::RangeError, 8.0f);
 	}
 
-	// Test FCbField(Integer+, MaxBinary32) as Float
+	// Test FCbFieldView(Integer+, MaxBinary32) as Float
 	{
 		uint8 Payload[9];
 		WriteVarUInt((uint64(1) << 24) - 1, Payload); // 16,777,215
-		FCbField Field(Payload, ECbFieldType::IntegerPositive);
+		FCbFieldView Field(Payload, ECbFieldType::IntegerPositive);
 		TestField<ECbFieldType::Float32>(TEXT("Integer+, MaxBinary32, AsFloat"), Field, 16'777'215.0f);
 		TestField<ECbFieldType::Float64>(TEXT("Integer+, MaxBinary32, AsDouble"), Field, 16'777'215.0);
 	}
 
-	// Test FCbField(Integer+, MaxBinary32+1) as Float
+	// Test FCbFieldView(Integer+, MaxBinary32+1) as Float
 	{
 		uint8 Payload[9];
 		WriteVarUInt(uint64(1) << 24, Payload); // 16,777,216
-		FCbField Field(Payload, ECbFieldType::IntegerPositive);
+		FCbFieldView Field(Payload, ECbFieldType::IntegerPositive);
 		TestFieldError<ECbFieldType::Float32>(TEXT("Integer+, MaxBinary32+1, AsFloat"), Field, ECbFieldError::RangeError, 8.0f);
 		TestField<ECbFieldType::Float64>(TEXT("Integer+, MaxBinary32+1, AsDouble"), Field, 16'777'216.0);
 	}
 
-	// Test FCbField(Integer+, MaxBinary64) as Float
+	// Test FCbFieldView(Integer+, MaxBinary64) as Float
 	{
 		uint8 Payload[9];
 		WriteVarUInt((uint64(1) << 53) - 1, Payload); // 9,007,199,254,740,991
-		FCbField Field(Payload, ECbFieldType::IntegerPositive);
+		FCbFieldView Field(Payload, ECbFieldType::IntegerPositive);
 		TestFieldError<ECbFieldType::Float32>(TEXT("Integer+, MaxBinary64, AsFloat"), Field, ECbFieldError::RangeError, 8.0f);
 		TestField<ECbFieldType::Float64>(TEXT("Integer+, MaxBinary64, AsDouble"), Field, 9'007'199'254'740'991.0);
 	}
 
-	// Test FCbField(Integer+, MaxBinary64+1) as Float
+	// Test FCbFieldView(Integer+, MaxBinary64+1) as Float
 	{
 		uint8 Payload[9];
 		WriteVarUInt(uint64(1) << 53, Payload); // 9,007,199,254,740,992
-		FCbField Field(Payload, ECbFieldType::IntegerPositive);
+		FCbFieldView Field(Payload, ECbFieldType::IntegerPositive);
 		TestFieldError<ECbFieldType::Float32>(TEXT("Integer+, MaxBinary64+1, AsFloat"), Field, ECbFieldError::RangeError, 8.0f);
 		TestFieldError<ECbFieldType::Float64>(TEXT("Integer+, MaxBinary64+1, AsDouble"), Field, ECbFieldError::RangeError, 8.0);
 	}
 
-	// Test FCbField(Integer+, MaxUInt64) as Float
+	// Test FCbFieldView(Integer+, MaxUInt64) as Float
 	{
 		uint8 Payload[9];
 		WriteVarUInt(uint64(-1), Payload); // Max uint64
-		FCbField Field(Payload, ECbFieldType::IntegerPositive);
+		FCbFieldView Field(Payload, ECbFieldType::IntegerPositive);
 		TestFieldError<ECbFieldType::Float32>(TEXT("Integer+, MaxUInt64, AsFloat"), Field, ECbFieldError::RangeError, 8.0f);
 		TestFieldError<ECbFieldType::Float64>(TEXT("Integer+, MaxUInt64, AsDouble"), Field, ECbFieldError::RangeError, 8.0);
 	}
 
-	// Test FCbField(Integer-, MaxBinary32) as Float
+	// Test FCbFieldView(Integer-, MaxBinary32) as Float
 	{
 		uint8 Payload[9];
 		WriteVarUInt((uint64(1) << 24) - 2, Payload); // -16,777,215
-		FCbField Field(Payload, ECbFieldType::IntegerNegative);
+		FCbFieldView Field(Payload, ECbFieldType::IntegerNegative);
 		TestField<ECbFieldType::Float32>(TEXT("Integer-, MaxBinary32, AsFloat"), Field, -16'777'215.0f);
 		TestField<ECbFieldType::Float64>(TEXT("Integer-, MaxBinary32, AsDouble"), Field, -16'777'215.0);
 	}
 
-	// Test FCbField(Integer-, MaxBinary32+1) as Float
+	// Test FCbFieldView(Integer-, MaxBinary32+1) as Float
 	{
 		uint8 Payload[9];
 		WriteVarUInt((uint64(1) << 24) - 1, Payload); // -16,777,216
-		FCbField Field(Payload, ECbFieldType::IntegerNegative);
+		FCbFieldView Field(Payload, ECbFieldType::IntegerNegative);
 		TestFieldError<ECbFieldType::Float32>(TEXT("Integer-, MaxBinary32+1, AsFloat"), Field, ECbFieldError::RangeError, 8.0f);
 		TestField<ECbFieldType::Float64>(TEXT("Integer-, MaxBinary32+1, AsDouble"), Field, -16'777'216.0);
 	}
 
-	// Test FCbField(Integer-, MaxBinary64) as Float
+	// Test FCbFieldView(Integer-, MaxBinary64) as Float
 	{
 		uint8 Payload[9];
 		WriteVarUInt((uint64(1) << 53) - 2, Payload); // -9,007,199,254,740,991
-		FCbField Field(Payload, ECbFieldType::IntegerNegative);
+		FCbFieldView Field(Payload, ECbFieldType::IntegerNegative);
 		TestFieldError<ECbFieldType::Float32>(TEXT("Integer-, MaxBinary64, AsFloat"), Field, ECbFieldError::RangeError, 8.0f);
 		TestField<ECbFieldType::Float64>(TEXT("Integer-, MaxBinary64, AsDouble"), Field, -9'007'199'254'740'991.0);
 	}
 
-	// Test FCbField(Integer-, MaxBinary64+1) as Float
+	// Test FCbFieldView(Integer-, MaxBinary64+1) as Float
 	{
 		uint8 Payload[9];
 		WriteVarUInt((uint64(1) << 53) - 1, Payload); // -9,007,199,254,740,992
-		FCbField Field(Payload, ECbFieldType::IntegerNegative);
+		FCbFieldView Field(Payload, ECbFieldType::IntegerNegative);
 		TestFieldError<ECbFieldType::Float32>(TEXT("Integer-, MaxBinary64+1, AsFloat"), Field, ECbFieldError::RangeError, 8.0f);
 		TestFieldError<ECbFieldType::Float64>(TEXT("Integer-, MaxBinary64+1, AsDouble"), Field, ECbFieldError::RangeError, 8.0);
 	}
 
-	// Test FCbField(None) as Float
+	// Test FCbFieldView(None) as Float
 	{
-		FCbField Field;
+		FCbFieldView Field;
 		TestFieldError<ECbFieldType::Float32>(TEXT("None, AsFloat"), Field, ECbFieldError::TypeError, 8.0f);
 		TestFieldError<ECbFieldType::Float64>(TEXT("None, AsDouble"), Field, ECbFieldError::TypeError, 8.0);
 	}
@@ -923,15 +923,15 @@ bool FCbFieldFloatTest::RunTest(const FString& Parameters)
 IMPLEMENT_CUSTOM_SIMPLE_AUTOMATION_TEST(FCbFieldBoolTest, FCbFieldTestBase, "System.Core.Serialization.CbField.Bool", CompactBinaryTestFlags)
 bool FCbFieldBoolTest::RunTest(const FString& Parameters)
 {
-	// Test FCbField(Bool, False)
+	// Test FCbFieldView(Bool, False)
 	TestField<ECbFieldType::BoolFalse>(TEXT("Bool, False"), {}, false, true);
 
-	// Test FCbField(Bool, True)
+	// Test FCbFieldView(Bool, True)
 	TestField<ECbFieldType::BoolTrue>(TEXT("Bool, True"), {}, true, false);
 
-	// Test FCbField(None) as Bool
+	// Test FCbFieldView(None) as Bool
 	{
-		FCbField DefaultField;
+		FCbFieldView DefaultField;
 		TestFieldError<ECbFieldType::BoolFalse>(TEXT("Bool, False, None"), DefaultField, ECbFieldError::TypeError, false);
 		TestFieldError<ECbFieldType::BoolTrue>(TEXT("Bool, True, None"), DefaultField, ECbFieldError::TypeError, true);
 	}
@@ -945,21 +945,21 @@ bool FCbFieldCompactBinaryAttachmentTest::RunTest(const FString& Parameters)
 	const FIoHash::ByteArray ZeroBytes{};
 	const FIoHash::ByteArray SequentialBytes{1, 2, 3, 4, 5, 6, 7, 8, 9, 10, 11, 12, 13, 14, 15, 16, 17, 18, 19, 20};
 
-	// Test FCbField(CompactBinaryAttachment, Zero)
+	// Test FCbFieldView(CompactBinaryAttachment, Zero)
 	TestField<ECbFieldType::CompactBinaryAttachment>(TEXT("CompactBinaryAttachment, Zero"), ZeroBytes);
 
-	// Test FCbField(CompactBinaryAttachment, NonZero)
+	// Test FCbFieldView(CompactBinaryAttachment, NonZero)
 	TestField<ECbFieldType::CompactBinaryAttachment>(TEXT("CompactBinaryAttachment, NonZero"), SequentialBytes, FIoHash(SequentialBytes));
 
-	// Test FCbField(CompactBinaryAttachment, NonZero) AsAttachment
+	// Test FCbFieldView(CompactBinaryAttachment, NonZero) AsAttachment
 	{
-		FCbField Field(SequentialBytes, ECbFieldType::CompactBinaryAttachment);
+		FCbFieldView Field(SequentialBytes, ECbFieldType::CompactBinaryAttachment);
 		TestField<ECbFieldType::CompactBinaryAttachment>(TEXT("CompactBinaryAttachment, NonZero, AsAttachment"), Field, FIoHash(SequentialBytes), FIoHash(), ECbFieldError::None, FCbAttachmentAccessors());
 	}
 
-	// Test FCbField(None) as CompactBinaryAttachment
+	// Test FCbFieldView(None) as CompactBinaryAttachment
 	{
-		FCbField DefaultField;
+		FCbFieldView DefaultField;
 		TestFieldError<ECbFieldType::CompactBinaryAttachment>(TEXT("CompactBinaryAttachment, None"), DefaultField, ECbFieldError::TypeError, FIoHash(SequentialBytes));
 	}
 
@@ -972,21 +972,21 @@ bool FCbFieldBinaryAttachmentTest::RunTest(const FString& Parameters)
 	const FIoHash::ByteArray ZeroBytes{};
 	const FIoHash::ByteArray SequentialBytes{1, 2, 3, 4, 5, 6, 7, 8, 9, 10, 11, 12, 13, 14, 15, 16, 17, 18, 19, 20};
 
-	// Test FCbField(BinaryAttachment, Zero)
+	// Test FCbFieldView(BinaryAttachment, Zero)
 	TestField<ECbFieldType::BinaryAttachment>(TEXT("BinaryAttachment, Zero"), ZeroBytes);
 
-	// Test FCbField(BinaryAttachment, NonZero)
+	// Test FCbFieldView(BinaryAttachment, NonZero)
 	TestField<ECbFieldType::BinaryAttachment>(TEXT("BinaryAttachment, NonZero"), SequentialBytes, FIoHash(SequentialBytes));
 
-	// Test FCbField(BinaryAttachment, NonZero) AsAttachment
+	// Test FCbFieldView(BinaryAttachment, NonZero) AsAttachment
 	{
-		FCbField Field(SequentialBytes, ECbFieldType::BinaryAttachment);
+		FCbFieldView Field(SequentialBytes, ECbFieldType::BinaryAttachment);
 		TestField<ECbFieldType::BinaryAttachment>(TEXT("BinaryAttachment, NonZero, AsAttachment"), Field, FIoHash(SequentialBytes), FIoHash(), ECbFieldError::None, FCbAttachmentAccessors());
 	}
 
-	// Test FCbField(None) as BinaryAttachment
+	// Test FCbFieldView(None) as BinaryAttachment
 	{
-		FCbField DefaultField;
+		FCbFieldView DefaultField;
 		TestFieldError<ECbFieldType::BinaryAttachment>(TEXT("BinaryAttachment, None"), DefaultField, ECbFieldError::TypeError, FIoHash(SequentialBytes));
 	}
 
@@ -999,15 +999,15 @@ bool FCbFieldHashTest::RunTest(const FString& Parameters)
 	const FIoHash::ByteArray ZeroBytes{};
 	const FIoHash::ByteArray SequentialBytes{1, 2, 3, 4, 5, 6, 7, 8, 9, 10, 11, 12, 13, 14, 15, 16, 17, 18, 19, 20};
 
-	// Test FCbField(Hash, Zero)
+	// Test FCbFieldView(Hash, Zero)
 	TestField<ECbFieldType::Hash>(TEXT("Hash, Zero"), ZeroBytes);
 
-	// Test FCbField(Hash, NonZero)
+	// Test FCbFieldView(Hash, NonZero)
 	TestField<ECbFieldType::Hash>(TEXT("Hash, NonZero"), SequentialBytes, FIoHash(SequentialBytes));
 
-	// Test FCbField(None) as Hash
+	// Test FCbFieldView(None) as Hash
 	{
-		FCbField DefaultField;
+		FCbFieldView DefaultField;
 		TestFieldError<ECbFieldType::Hash>(TEXT("Hash, None"), DefaultField, ECbFieldError::TypeError, FIoHash(SequentialBytes));
 	}
 
@@ -1021,15 +1021,15 @@ bool FCbFieldUuidTest::RunTest(const FString& Parameters)
 	const uint8 SequentialBytes[]{0, 1, 2, 3, 4, 5, 6, 7, 8, 9, 10, 11, 12, 13, 14, 15};
 	const FGuid SequentialGuid(TEXT("00010203-0405-0607-0809-0a0b0c0d0e0f"));
 
-	// Test FCbField(Uuid, Zero)
+	// Test FCbFieldView(Uuid, Zero)
 	TestField<ECbFieldType::Uuid>(TEXT("Uuid, Zero"), ZeroBytes, FGuid(), SequentialGuid);
 
-	// Test FCbField(Uuid, NonZero)
+	// Test FCbFieldView(Uuid, NonZero)
 	TestField<ECbFieldType::Uuid>(TEXT("Uuid, NonZero"), SequentialBytes, SequentialGuid, FGuid());
 
-	// Test FCbField(None) as Uuid
+	// Test FCbFieldView(None) as Uuid
 	{
-		FCbField DefaultField;
+		FCbFieldView DefaultField;
 		TestFieldError<ECbFieldType::Uuid>(TEXT("Uuid, None"), DefaultField, ECbFieldError::TypeError, FGuid::NewGuid());
 	}
 
@@ -1039,25 +1039,25 @@ bool FCbFieldUuidTest::RunTest(const FString& Parameters)
 IMPLEMENT_CUSTOM_SIMPLE_AUTOMATION_TEST(FCbFieldDateTimeTest, FCbFieldTestBase, "System.Core.Serialization.CbField.DateTime", CompactBinaryTestFlags)
 bool FCbFieldDateTimeTest::RunTest(const FString& Parameters)
 {
-	// Test FCbField(DateTime, Zero)
+	// Test FCbFieldView(DateTime, Zero)
 	TestField<ECbFieldType::DateTime>(TEXT("DateTime, Zero"), {0, 0, 0, 0, 0, 0, 0, 0});
 
-	// Test FCbField(DateTime, 0x1020'3040'5060'7080)
+	// Test FCbFieldView(DateTime, 0x1020'3040'5060'7080)
 	TestField<ECbFieldType::DateTime>(TEXT("DateTime, NonZero"), {0x10, 0x20, 0x30, 0x40, 0x50, 0x60, 0x70, 0x80}, int64(0x1020'3040'5060'7080));
 
-	// Test FCbField(DateTime, Zero) as FDateTime
+	// Test FCbFieldView(DateTime, Zero) as FDateTime
 	{
 		const uint8 Payload[] = {0, 0, 0, 0, 0, 0, 0, 0};
-		FCbField Field(Payload, ECbFieldType::DateTime);
-		TestEqual(TEXT("FCbField()::AsDateTime()"), Field.AsDateTime(), FDateTime(0));
+		FCbFieldView Field(Payload, ECbFieldType::DateTime);
+		TestEqual(TEXT("FCbFieldView()::AsDateTime()"), Field.AsDateTime(), FDateTime(0));
 	}
 
-	// Test FCbField(None) as DateTime
+	// Test FCbFieldView(None) as DateTime
 	{
-		FCbField DefaultField;
+		FCbFieldView DefaultField;
 		TestFieldError<ECbFieldType::DateTime>(TEXT("DateTime, None"), DefaultField, ECbFieldError::TypeError);
 		const FDateTime DefaultValue(0x1020'3040'5060'7080);
-		TestEqual(TEXT("FCbField()::AsDateTime()"), DefaultField.AsDateTime(DefaultValue), DefaultValue);
+		TestEqual(TEXT("FCbFieldView()::AsDateTime()"), DefaultField.AsDateTime(DefaultValue), DefaultValue);
 	}
 
 	return true;
@@ -1066,25 +1066,25 @@ bool FCbFieldDateTimeTest::RunTest(const FString& Parameters)
 IMPLEMENT_CUSTOM_SIMPLE_AUTOMATION_TEST(FCbFieldTimeSpanTest, FCbFieldTestBase, "System.Core.Serialization.CbField.TimeSpan", CompactBinaryTestFlags)
 bool FCbFieldTimeSpanTest::RunTest(const FString& Parameters)
 {
-	// Test FCbField(TimeSpan, Zero)
+	// Test FCbFieldView(TimeSpan, Zero)
 	TestField<ECbFieldType::TimeSpan>(TEXT("TimeSpan, Zero"), {0, 0, 0, 0, 0, 0, 0, 0});
 
-	// Test FCbField(TimeSpan, 0x1020'3040'5060'7080)
+	// Test FCbFieldView(TimeSpan, 0x1020'3040'5060'7080)
 	TestField<ECbFieldType::TimeSpan>(TEXT("TimeSpan, NonZero"), {0x10, 0x20, 0x30, 0x40, 0x50, 0x60, 0x70, 0x80}, int64(0x1020'3040'5060'7080));
 
-	// Test FCbField(TimeSpan, Zero) as FTimeSpan
+	// Test FCbFieldView(TimeSpan, Zero) as FTimeSpan
 	{
 		const uint8 Payload[] = {0, 0, 0, 0, 0, 0, 0, 0};
-		FCbField Field(Payload, ECbFieldType::TimeSpan);
-		TestEqual(TEXT("FCbField()::AsTimeSpan()"), Field.AsTimeSpan(), FTimespan(0));
+		FCbFieldView Field(Payload, ECbFieldType::TimeSpan);
+		TestEqual(TEXT("FCbFieldView()::AsTimeSpan()"), Field.AsTimeSpan(), FTimespan(0));
 	}
 
-	// Test FCbField(None) as TimeSpan
+	// Test FCbFieldView(None) as TimeSpan
 	{
-		FCbField DefaultField;
+		FCbFieldView DefaultField;
 		TestFieldError<ECbFieldType::TimeSpan>(TEXT("TimeSpan, None"), DefaultField, ECbFieldError::TypeError);
 		const FTimespan DefaultValue(0x1020'3040'5060'7080);
-		TestEqual(TEXT("FCbField()::AsTimeSpan()"), DefaultField.AsTimeSpan(DefaultValue), DefaultValue);
+		TestEqual(TEXT("FCbFieldView()::AsTimeSpan()"), DefaultField.AsTimeSpan(DefaultValue), DefaultValue);
 	}
 
 	return true;
@@ -1093,26 +1093,26 @@ bool FCbFieldTimeSpanTest::RunTest(const FString& Parameters)
 IMPLEMENT_CUSTOM_SIMPLE_AUTOMATION_TEST(FCbFieldObjectIdTest, FCbFieldTestBase, "System.Core.Serialization.CbField.ObjectId", CompactBinaryTestFlags)
 bool FCbFieldObjectIdTest::RunTest(const FString& Parameters)
 {
-	// Test FCbField(ObjectId, Zero)
+	// Test FCbFieldView(ObjectId, Zero)
 	TestField<ECbFieldType::ObjectId>(TEXT("ObjectId, Zero"), {0, 0, 0, 0, 0, 0, 0, 0, 0, 0, 0, 0});
 
-	// Test FCbField(ObjectId, 0x102030405060708090A0B0C0)
+	// Test FCbFieldView(ObjectId, 0x102030405060708090A0B0C0)
 	TestField<ECbFieldType::ObjectId>(TEXT("ObjectId, NonZero"), {0x10, 0x20, 0x30, 0x40, 0x50, 0x60, 0x70, 0x80, 0x90, 0xA0, 0xB0, 0xC0},
 		FCbObjectId(MakeMemoryView<uint8>({0x10, 0x20, 0x30, 0x40, 0x50, 0x60, 0x70, 0x80, 0x90, 0xA0, 0xB0, 0xC0})));
 
-	// Test FCbField(ObjectId, Zero) as FCbObjectId
+	// Test FCbFieldView(ObjectId, Zero) as FCbObjectId
 	{
 		const uint8 Payload[] = {0, 0, 0, 0, 0, 0, 0, 0, 0, 0, 0, 0};
-		FCbField Field(Payload, ECbFieldType::ObjectId);
-		TestEqual(TEXT("FCbField(ObjectId, Zero)::AsObjectId()"), Field.AsObjectId(), FCbObjectId());
+		FCbFieldView Field(Payload, ECbFieldType::ObjectId);
+		TestEqual(TEXT("FCbFieldView(ObjectId, Zero)::AsObjectId()"), Field.AsObjectId(), FCbObjectId());
 	}
 
-	// Test FCbField(None) as ObjectId
+	// Test FCbFieldView(None) as ObjectId
 	{
-		FCbField DefaultField;
+		FCbFieldView DefaultField;
 		TestFieldError<ECbFieldType::ObjectId>(TEXT("ObjectId, None"), DefaultField, ECbFieldError::TypeError);
 		const FCbObjectId DefaultValue(MakeMemoryView<uint8>({0x10, 0x20, 0x30, 0x40, 0x50, 0x60, 0x70, 0x80, 0x90, 0xA0, 0xB0, 0xC0}));
-		TestEqual(TEXT("FCbField(None)::AsObjectId()"), DefaultField.AsObjectId(DefaultValue), DefaultValue);
+		TestEqual(TEXT("FCbFieldView(None)::AsObjectId()"), DefaultField.AsObjectId(DefaultValue), DefaultValue);
 	}
 
 	return true;
@@ -1124,15 +1124,15 @@ bool FCbFieldCustomByIdTest::RunTest(const FString& Parameters)
 	struct FCustomByIdAccessor
 	{
 		explicit FCustomByIdAccessor(uint64 Id)
-			: AsType([Id](FCbField& Field, FMemoryView Default) { return Field.AsCustom(Id, Default); })
+			: AsType([Id](FCbFieldView& Field, FMemoryView Default) { return Field.AsCustom(Id, Default); })
 		{
 		}
 
-		bool (FCbField::*IsType)() const = &FCbField::IsCustomById;
-		TUniqueFunction<FMemoryView (FCbField& Field, FMemoryView Default)> AsType;
+		bool (FCbFieldView::*IsType)() const = &FCbFieldView::IsCustomById;
+		TUniqueFunction<FMemoryView (FCbFieldView& Field, FMemoryView Default)> AsType;
 	};
 
-	// Test FCbField(CustomById, MinId, Empty)
+	// Test FCbFieldView(CustomById, MinId, Empty)
 	{
 		const uint8 Payload[] = {1, 0};
 		TestField<ECbFieldType::CustomById>(TEXT("CustomById, MinId, Empty"), Payload, FCbCustomById{0});
@@ -1140,7 +1140,7 @@ bool FCbFieldCustomByIdTest::RunTest(const FString& Parameters)
 		TestFieldError<ECbFieldType::CustomById>(TEXT("CustomById, MinId, Empty, InvalidId"), Payload, ECbFieldError::RangeError, MakeMemoryView<uint8>({1, 2, 3}), FCustomByIdAccessor(MAX_uint64));
 	}
 
-	// Test FCbField(CustomById, MinId, Value)
+	// Test FCbFieldView(CustomById, MinId, Value)
 	{
 		const uint8 Payload[] = {5, 0, 1, 2, 3, 4};
 		TestFieldNoClone<ECbFieldType::CustomById>(TEXT("CustomById, MinId, Value"), Payload, FCbCustomById{0, MakeMemoryView(Payload).Right(4)});
@@ -1148,7 +1148,7 @@ bool FCbFieldCustomByIdTest::RunTest(const FString& Parameters)
 		TestFieldError<ECbFieldType::CustomById>(TEXT("CustomById, MinId, Value, InvalidId"), Payload, ECbFieldError::RangeError, MakeMemoryView<uint8>({1, 2, 3}), FCustomByIdAccessor(MAX_uint64));
 	}
 
-	// Test FCbField(CustomById, MaxId, Empty)
+	// Test FCbFieldView(CustomById, MaxId, Empty)
 	{
 		const uint8 Payload[] = {9, 0xff, 0xff, 0xff, 0xff, 0xff, 0xff, 0xff, 0xff, 0xff};
 		TestField<ECbFieldType::CustomById>(TEXT("CustomById, MaxId, Empty"), Payload, FCbCustomById{MAX_uint64});
@@ -1156,7 +1156,7 @@ bool FCbFieldCustomByIdTest::RunTest(const FString& Parameters)
 		TestFieldError<ECbFieldType::CustomById>(TEXT("CustomById, MaxId, Empty, InvalidId"), Payload, ECbFieldError::RangeError, MakeMemoryView<uint8>({1, 2, 3}), FCustomByIdAccessor(0));
 	}
 
-	// Test FCbField(CustomById, MaxId, Value)
+	// Test FCbFieldView(CustomById, MaxId, Value)
 	{
 		const uint8 Payload[] = {13, 0xff, 0xff, 0xff, 0xff, 0xff, 0xff, 0xff, 0xff, 0xff, 1, 2, 3, 4};
 		TestFieldNoClone<ECbFieldType::CustomById>(TEXT("CustomById, MaxId, Value"), Payload, FCbCustomById{MAX_uint64, MakeMemoryView(Payload).Right(4)});
@@ -1164,13 +1164,13 @@ bool FCbFieldCustomByIdTest::RunTest(const FString& Parameters)
 		TestFieldError<ECbFieldType::CustomById>(TEXT("CustomById, MaxId, Value, InvalidId"), Payload, ECbFieldError::RangeError, MakeMemoryView<uint8>({1, 2, 3}), FCustomByIdAccessor(0));
 	}
 
-	// Test FCbField(None) as CustomById
+	// Test FCbFieldView(None) as CustomById
 	{
-		FCbField DefaultField;
+		FCbFieldView DefaultField;
 		TestFieldError<ECbFieldType::CustomById>(TEXT("CustomById, None"), DefaultField, ECbFieldError::TypeError, FCbCustomById{4, MakeMemoryView<uint8>({1, 2, 3})});
 		TestFieldError<ECbFieldType::CustomById>(TEXT("CustomById, None, View"), DefaultField, ECbFieldError::TypeError, MakeMemoryView<uint8>({1, 2, 3}), FCustomByIdAccessor(0));
 		const uint8 DefaultValue[] = {1, 2, 3};
-		TestEqual(TEXT("FCbField()::AsCustom(Id)"), DefaultField.AsCustom(0, MakeMemoryView(DefaultValue)), MakeMemoryView(DefaultValue));
+		TestEqual(TEXT("FCbFieldView()::AsCustom(Id)"), DefaultField.AsCustom(0, MakeMemoryView(DefaultValue)), MakeMemoryView(DefaultValue));
 	}
 
 	return true;
@@ -1182,15 +1182,15 @@ bool FCbFieldCustomByNameTest::RunTest(const FString& Parameters)
 	struct FCustomByNameAccessor
 	{
 		explicit FCustomByNameAccessor(FAnsiStringView Name)
-			: AsType([Name = FString(Name)](FCbField& Field, FMemoryView Default) { return Field.AsCustom(TCHAR_TO_ANSI(*Name), Default); })
+			: AsType([Name = FString(Name)](FCbFieldView& Field, FMemoryView Default) { return Field.AsCustom(TCHAR_TO_ANSI(*Name), Default); })
 		{
 		}
 
-		bool (FCbField::*IsType)() const = &FCbField::IsCustomByName;
-		TUniqueFunction<FMemoryView (FCbField& Field, FMemoryView Default)> AsType;
+		bool (FCbFieldView::*IsType)() const = &FCbFieldView::IsCustomByName;
+		TUniqueFunction<FMemoryView (FCbFieldView& Field, FMemoryView Default)> AsType;
 	};
 
-	// Test FCbField(CustomByName, ABC, Empty)
+	// Test FCbFieldView(CustomByName, ABC, Empty)
 	{
 		const uint8 Payload[] = {4, 3, 'A', 'B', 'C'};
 		TestField<ECbFieldType::CustomByName>(TEXT("CustomByName, MinId, Empty"), Payload, FCbCustomByName{"ABC"});
@@ -1198,7 +1198,7 @@ bool FCbFieldCustomByNameTest::RunTest(const FString& Parameters)
 		TestFieldError<ECbFieldType::CustomByName>(TEXT("CustomByName, MinId, Empty, InvalidCase"), Payload, ECbFieldError::RangeError, MakeMemoryView<uint8>({1, 2, 3}), FCustomByNameAccessor("abc"));
 	}
 
-	// Test FCbField(CustomByName, ABC, Value)
+	// Test FCbFieldView(CustomByName, ABC, Value)
 	{
 		const uint8 Payload[] = {8, 3, 'A', 'B', 'C', 1, 2, 3, 4};
 		TestFieldNoClone<ECbFieldType::CustomByName>(TEXT("CustomByName, MinId, Value"), Payload, FCbCustomByName{"ABC", MakeMemoryView(Payload).Right(4)});
@@ -1206,13 +1206,13 @@ bool FCbFieldCustomByNameTest::RunTest(const FString& Parameters)
 		TestFieldError<ECbFieldType::CustomByName>(TEXT("CustomByName, MinId, Value, InvalidCase"), Payload, ECbFieldError::RangeError, MakeMemoryView<uint8>({1, 2, 3}), FCustomByNameAccessor("abc"));
 	}
 
-	// Test FCbField(None) as CustomByName
+	// Test FCbFieldView(None) as CustomByName
 	{
-		FCbField DefaultField;
+		FCbFieldView DefaultField;
 		TestFieldError<ECbFieldType::CustomByName>(TEXT("CustomByName, None"), DefaultField, ECbFieldError::TypeError, FCbCustomByName{"ABC", MakeMemoryView<uint8>({1, 2, 3})});
 		TestFieldError<ECbFieldType::CustomByName>(TEXT("CustomByName, None, View"), DefaultField, ECbFieldError::TypeError, MakeMemoryView<uint8>({1, 2, 3}), FCustomByNameAccessor("ABC"));
 		const uint8 DefaultValue[] = {1, 2, 3};
-		TestEqual(TEXT("FCbField()::AsCustom(Name)"), DefaultField.AsCustom("ABC", MakeMemoryView(DefaultValue)), MakeMemoryView(DefaultValue));
+		TestEqual(TEXT("FCbFieldView()::AsCustom(Name)"), DefaultField.AsCustom("ABC", MakeMemoryView(DefaultValue)), MakeMemoryView(DefaultValue));
 	}
 
 	return true;
@@ -1223,7 +1223,7 @@ bool FCbFieldIterateAttachmentsTest::RunTest(const FString& Parameters)
 {
 	const auto MakeTestHash = [](uint32 Index) { return FIoHash::HashBuffer(&Index, sizeof(Index)); };
 
-	FCbFieldRefIterator Fields;
+	FCbFieldIterator Fields;
 	{
 		FCbWriter Writer;
 
@@ -1378,47 +1378,47 @@ bool FCbFieldIterateAttachmentsTest::RunTest(const FString& Parameters)
 		Fields = Writer.Save();
 	}
 
-	TestEqual(TEXT("FCbField::IterateAttachments Validate"), ValidateCompactBinaryRange(Fields.GetBuffer(), ECbValidateMode::All), ECbValidateError::None);
+	TestEqual(TEXT("FCbFieldView::IterateAttachments Validate"), ValidateCompactBinaryRange(Fields.GetBuffer(), ECbValidateMode::All), ECbValidateError::None);
 
 	uint32 AttachmentIndex = 0;
-	Fields.IterateRangeAttachments([this, &AttachmentIndex, &MakeTestHash](FCbField Field)
+	Fields.IterateRangeAttachments([this, &AttachmentIndex, &MakeTestHash](FCbFieldView Field)
 		{
-			TestTrue(FString::Printf(TEXT("FCbField::IterateAttachments(%u)::IsAttachment"), AttachmentIndex), Field.IsAttachment());
-			TestEqual(FString::Printf(TEXT("FCbField::IterateAttachments(%u)"), AttachmentIndex), Field.AsAttachment(), MakeTestHash(AttachmentIndex));
+			TestTrue(FString::Printf(TEXT("FCbFieldView::IterateAttachments(%u)::IsAttachment"), AttachmentIndex), Field.IsAttachment());
+			TestEqual(FString::Printf(TEXT("FCbFieldView::IterateAttachments(%u)"), AttachmentIndex), Field.AsAttachment(), MakeTestHash(AttachmentIndex));
 			++AttachmentIndex;
 		});
-	TestEqual(TEXT("FCbField::IterateAttachments"), AttachmentIndex, 28);
+	TestEqual(TEXT("FCbFieldView::IterateAttachments"), AttachmentIndex, 28);
 
 	return true;
 }
 
 ///////////////////////////////////////////////////////////////////////////////////////////////////////////////////////
 
-IMPLEMENT_CUSTOM_SIMPLE_AUTOMATION_TEST(FCbFieldRefTest, FCbFieldTestBase, "System.Core.Serialization.CbFieldRef", CompactBinaryTestFlags)
-bool FCbFieldRefTest::RunTest(const FString& Parameters)
+IMPLEMENT_CUSTOM_SIMPLE_AUTOMATION_TEST(FCbFieldBufferTest, FCbFieldTestBase, "System.Core.Serialization.CbFieldBuffer", CompactBinaryTestFlags)
+bool FCbFieldBufferTest::RunTest(const FString& Parameters)
 {
-	static_assert(std::is_constructible<FCbFieldRef, const FSharedBuffer&>::value, "Missing constructor for FCbFieldRef");
-	static_assert(std::is_constructible<FCbFieldRef, FSharedBuffer&&>::value, "Missing constructor for FCbFieldRef");
-
-	static_assert(std::is_constructible<FCbFieldRef, const FCbField&, const FSharedBuffer&>::value, "Missing constructor for FCbFieldRef");
-	static_assert(std::is_constructible<FCbFieldRef, const FCbField&, const FCbFieldRefIterator&>::value, "Missing constructor for FCbFieldRef");
-	static_assert(std::is_constructible<FCbFieldRef, const FCbField&, const FCbFieldRef&>::value, "Missing constructor for FCbFieldRef");
-	static_assert(std::is_constructible<FCbFieldRef, const FCbField&, const FCbArrayRef&>::value, "Missing constructor for FCbFieldRef");
-	static_assert(std::is_constructible<FCbFieldRef, const FCbField&, const FCbObjectRef&>::value, "Missing constructor for FCbFieldRef");
-
-	static_assert(std::is_constructible<FCbFieldRef, const FCbField&, FSharedBuffer&&>::value, "Missing constructor for FCbFieldRef");
-	static_assert(std::is_constructible<FCbFieldRef, const FCbField&, FCbFieldRefIterator&&>::value, "Missing constructor for FCbFieldRef");
-	static_assert(std::is_constructible<FCbFieldRef, const FCbField&, FCbFieldRef&&>::value, "Missing constructor for FCbFieldRef");
-	static_assert(std::is_constructible<FCbFieldRef, const FCbField&, FCbArrayRef&&>::value, "Missing constructor for FCbFieldRef");
-	static_assert(std::is_constructible<FCbFieldRef, const FCbField&, FCbObjectRef&&>::value, "Missing constructor for FCbFieldRef");
-
-	// Test FCbFieldRef()
-	{
-		FCbFieldRef DefaultField;
-		TestFalse(TEXT("FCbFieldRef().HasValue()"), DefaultField.HasValue());
-		TestFalse(TEXT("FCbFieldRef().IsOwned()"), DefaultField.IsOwned());
+	static_assert(std::is_constructible<FCbField, const FSharedBuffer&>::value, "Missing constructor for FCbField");
+	static_assert(std::is_constructible<FCbField, FSharedBuffer&&>::value, "Missing constructor for FCbField");
+
+	static_assert(std::is_constructible<FCbField, const FCbFieldView&, const FSharedBuffer&>::value, "Missing constructor for FCbField");
+	static_assert(std::is_constructible<FCbField, const FCbFieldView&, const FCbFieldIterator&>::value, "Missing constructor for FCbField");
+	static_assert(std::is_constructible<FCbField, const FCbFieldView&, const FCbField&>::value, "Missing constructor for FCbField");
+	static_assert(std::is_constructible<FCbField, const FCbFieldView&, const FCbArray&>::value, "Missing constructor for FCbField");
+	static_assert(std::is_constructible<FCbField, const FCbFieldView&, const FCbObject&>::value, "Missing constructor for FCbField");
+
+	static_assert(std::is_constructible<FCbField, const FCbFieldView&, FSharedBuffer&&>::value, "Missing constructor for FCbField");
+	static_assert(std::is_constructible<FCbField, const FCbFieldView&, FCbFieldIterator&&>::value, "Missing constructor for FCbField");
+	static_assert(std::is_constructible<FCbField, const FCbFieldView&, FCbField&&>::value, "Missing constructor for FCbField");
+	static_assert(std::is_constructible<FCbField, const FCbFieldView&, FCbArray&&>::value, "Missing constructor for FCbField");
+	static_assert(std::is_constructible<FCbField, const FCbFieldView&, FCbObject&&>::value, "Missing constructor for FCbField");
+
+	// Test FCbField()
+	{
+		FCbField DefaultField;
+		TestFalse(TEXT("FCbField().HasValue()"), DefaultField.HasValue());
+		TestFalse(TEXT("FCbField().IsOwned()"), DefaultField.IsOwned());
 		DefaultField.MakeOwned();
-		TestTrue(TEXT("FCbFieldRef().MakeOwned().IsOwned()"), DefaultField.IsOwned());
+		TestTrue(TEXT("FCbField().MakeOwned().IsOwned()"), DefaultField.IsOwned());
 	}
 
 	// Test Field w/ Type from Shared Buffer
@@ -1427,92 +1427,92 @@ bool FCbFieldRefTest::RunTest(const FString& Parameters)
 		FSharedBuffer ViewBuffer = FSharedBuffer::MakeView(MakeMemoryView(Payload));
 		FSharedBuffer OwnedBuffer = FSharedBuffer::Clone(ViewBuffer);
 
-		FCbFieldRef View(ViewBuffer);
-		FCbFieldRef ViewMove{FSharedBuffer(ViewBuffer)};
-		FCbFieldRef ViewOuterField(ImplicitConv<FCbField>(View), ViewBuffer);
-		FCbFieldRef ViewOuterBuffer(ImplicitConv<FCbField>(View), View);
-		FCbFieldRef Owned(OwnedBuffer);
-		FCbFieldRef OwnedMove{FSharedBuffer(OwnedBuffer)};
-		FCbFieldRef OwnedOuterField(ImplicitConv<FCbField>(Owned), OwnedBuffer);
-		FCbFieldRef OwnedOuterBuffer(ImplicitConv<FCbField>(Owned), Owned);
+		FCbField View(ViewBuffer);
+		FCbField ViewMove{FSharedBuffer(ViewBuffer)};
+		FCbField ViewOuterField(ImplicitConv<FCbFieldView>(View), ViewBuffer);
+		FCbField ViewOuterBuffer(ImplicitConv<FCbFieldView>(View), View);
+		FCbField Owned(OwnedBuffer);
+		FCbField OwnedMove{FSharedBuffer(OwnedBuffer)};
+		FCbField OwnedOuterField(ImplicitConv<FCbFieldView>(Owned), OwnedBuffer);
+		FCbField OwnedOuterBuffer(ImplicitConv<FCbFieldView>(Owned), Owned);
 
 		// These lines are expected to assert when uncommented.
-		//FCbFieldRef InvalidOuterBuffer(ImplicitConv<FCbField>(Owned), ViewBuffer);
-		//FCbFieldRef InvalidOuterBufferMove(ImplicitConv<FCbField>(Owned), FSharedBuffer(ViewBuffer));
-
-		TestEqual(TEXT("FCbFieldRef(ViewBuffer)"), View.AsBinary(), ViewBuffer.GetView().Right(3));
-		TestEqual(TEXT("FCbFieldRef(ViewBuffer&&)"), ViewMove.AsBinary(), View.AsBinary());
-		TestEqual(TEXT("FCbFieldRef(ViewOuterField)"), ViewOuterField.AsBinary(), View.AsBinary());
-		TestEqual(TEXT("FCbFieldRef(ViewOuterBuffer)"), ViewOuterBuffer.AsBinary(), View.AsBinary());
-		TestEqual(TEXT("FCbFieldRef(OwnedBuffer)"), Owned.AsBinary(), OwnedBuffer.GetView().Right(3));
-		TestEqual(TEXT("FCbFieldRef(OwnedBuffer&&)"), OwnedMove.AsBinary(), Owned.AsBinary());
-		TestEqual(TEXT("FCbFieldRef(OwnedOuterField)"), OwnedOuterField.AsBinary(), Owned.AsBinary());
-		TestEqual(TEXT("FCbFieldRef(OwnedOuterBuffer)"), OwnedOuterBuffer.AsBinary(), Owned.AsBinary());
-
-		TestFalse(TEXT("FCbFieldRef(ViewBuffer).IsOwned()"), View.IsOwned());
-		TestFalse(TEXT("FCbFieldRef(ViewBuffer&&).IsOwned()"), ViewMove.IsOwned());
-		TestFalse(TEXT("FCbFieldRef(ViewOuterField).IsOwned()"), ViewOuterField.IsOwned());
-		TestFalse(TEXT("FCbFieldRef(ViewOuterBuffer).IsOwned()"), ViewOuterBuffer.IsOwned());
-		TestTrue(TEXT("FCbFieldRef(OwnedBuffer).IsOwned()"), Owned.IsOwned());
-		TestTrue(TEXT("FCbFieldRef(OwnedBuffer&&).IsOwned()"), OwnedMove.IsOwned());
-		TestTrue(TEXT("FCbFieldRef(OwnedOuterField).IsOwned()"), OwnedOuterField.IsOwned());
-		TestTrue(TEXT("FCbFieldRef(OwnedOuterBuffer).IsOwned()"), OwnedOuterBuffer.IsOwned());
+		//FCbField InvalidOuterBuffer(ImplicitConv<FCbFieldView>(Owned), ViewBuffer);
+		//FCbField InvalidOuterBufferMove(ImplicitConv<FCbFieldView>(Owned), FSharedBuffer(ViewBuffer));
+
+		TestEqual(TEXT("FCbField(ViewBuffer)"), View.AsBinaryView(), ViewBuffer.GetView().Right(3));
+		TestEqual(TEXT("FCbField(ViewBuffer&&)"), ViewMove.AsBinaryView(), View.AsBinaryView());
+		TestEqual(TEXT("FCbField(ViewOuterField)"), ViewOuterField.AsBinaryView(), View.AsBinaryView());
+		TestEqual(TEXT("FCbField(ViewOuterBuffer)"), ViewOuterBuffer.AsBinaryView(), View.AsBinaryView());
+		TestEqual(TEXT("FCbField(OwnedBuffer)"), Owned.AsBinaryView(), OwnedBuffer.GetView().Right(3));
+		TestEqual(TEXT("FCbField(OwnedBuffer&&)"), OwnedMove.AsBinaryView(), Owned.AsBinaryView());
+		TestEqual(TEXT("FCbField(OwnedOuterField)"), OwnedOuterField.AsBinaryView(), Owned.AsBinaryView());
+		TestEqual(TEXT("FCbField(OwnedOuterBuffer)"), OwnedOuterBuffer.AsBinaryView(), Owned.AsBinaryView());
+
+		TestFalse(TEXT("FCbField(ViewBuffer).IsOwned()"), View.IsOwned());
+		TestFalse(TEXT("FCbField(ViewBuffer&&).IsOwned()"), ViewMove.IsOwned());
+		TestFalse(TEXT("FCbField(ViewOuterField).IsOwned()"), ViewOuterField.IsOwned());
+		TestFalse(TEXT("FCbField(ViewOuterBuffer).IsOwned()"), ViewOuterBuffer.IsOwned());
+		TestTrue(TEXT("FCbField(OwnedBuffer).IsOwned()"), Owned.IsOwned());
+		TestTrue(TEXT("FCbField(OwnedBuffer&&).IsOwned()"), OwnedMove.IsOwned());
+		TestTrue(TEXT("FCbField(OwnedOuterField).IsOwned()"), OwnedOuterField.IsOwned());
+		TestTrue(TEXT("FCbField(OwnedOuterBuffer).IsOwned()"), OwnedOuterBuffer.IsOwned());
 
 		View.MakeOwned();
 		Owned.MakeOwned();
-		TestNotEqual(TEXT("FCbFieldRef(View).MakeOwned()"), View.AsBinary(), ViewBuffer.GetView().Right(3));
-		TestTrue(TEXT("FCbFieldRef(View).MakeOwned().IsOwned()"), View.IsOwned());
-		TestEqual(TEXT("FCbFieldRef(Owned).MakeOwned()"), Owned.AsBinary(), OwnedBuffer.GetView().Right(3));
-		TestTrue(TEXT("FCbFieldRef(Owned).MakeOwned().IsOwned()"), Owned.IsOwned());
+		TestNotEqual(TEXT("FCbField(View).MakeOwned()"), View.AsBinaryView(), ViewBuffer.GetView().Right(3));
+		TestTrue(TEXT("FCbField(View).MakeOwned().IsOwned()"), View.IsOwned());
+		TestEqual(TEXT("FCbField(Owned).MakeOwned()"), Owned.AsBinaryView(), OwnedBuffer.GetView().Right(3));
+		TestTrue(TEXT("FCbField(Owned).MakeOwned().IsOwned()"), Owned.IsOwned());
 	}
 
 	// Test Field w/ Type
 	{
 		const uint8 Payload[] = { uint8(ECbFieldType::Binary), 3, 4, 5, 6 }; // Size: 3, Data: 4/5/6
-		const FCbField Field(Payload);
-
-		FCbFieldRef VoidView = FCbFieldRef::MakeView(Payload);
-		FCbFieldRef VoidClone = FCbFieldRef::Clone(Payload);
-		FCbFieldRef FieldView = FCbFieldRef::MakeView(Field);
-		FCbFieldRef FieldClone = FCbFieldRef::Clone(Field);
-		FCbFieldRef FieldRefClone = FCbFieldRef::Clone(FieldView);
-
-		TestEqual(TEXT("FCbFieldRef::MakeView(Void)"), VoidView.AsBinary(), MakeMemoryView(Payload).Right(3));
-		TestNotEqual(TEXT("FCbFieldRef::Clone(Void)"), VoidClone.AsBinary(), MakeMemoryView(Payload).Right(3));
-		TestTrue(TEXT("FCbFieldRef::Clone(Void)->EqualBytes"), VoidClone.AsBinary().EqualBytes(VoidView.AsBinary()));
-		TestEqual(TEXT("FCbFieldRef::MakeView(Field)"), FieldView.AsBinary(), MakeMemoryView(Payload).Right(3));
-		TestNotEqual(TEXT("FCbFieldRef::Clone(Field)"), FieldClone.AsBinary(), MakeMemoryView(Payload).Right(3));
-		TestTrue(TEXT("FCbFieldRef::Clone(Field)->EqualBytes"), FieldClone.AsBinary().EqualBytes(VoidView.AsBinary()));
-		TestNotEqual(TEXT("FCbFieldRef::Clone(FieldRef)"), FieldRefClone.AsBinary(), FieldView.AsBinary());
-		TestTrue(TEXT("FCbFieldRef::Clone(FieldRef)->EqualBytes"), FieldRefClone.AsBinary().EqualBytes(VoidView.AsBinary()));
-
-		TestFalse(TEXT("FCbFieldRef::MakeView(Void).IsOwned()"), VoidView.IsOwned());
-		TestTrue(TEXT("FCbFieldRef::Clone(Void).IsOwned()"), VoidClone.IsOwned());
-		TestFalse(TEXT("FCbFieldRef::MakeView(Field).IsOwned()"), FieldView.IsOwned());
-		TestTrue(TEXT("FCbFieldRef::Clone(Field).IsOwned()"), FieldClone.IsOwned());
-		TestTrue(TEXT("FCbFieldRef::Clone(FieldRef).IsOwned()"), FieldRefClone.IsOwned());
+		const FCbFieldView Field(Payload);
+
+		FCbField VoidView = FCbField::MakeView(Payload);
+		FCbField VoidClone = FCbField::Clone(Payload);
+		FCbField FieldView = FCbField::MakeView(Field);
+		FCbField FieldClone = FCbField::Clone(Field);
+		FCbField FieldViewClone = FCbField::Clone(FieldView);
+
+		TestEqual(TEXT("FCbField::MakeView(Void)"), VoidView.AsBinaryView(), MakeMemoryView(Payload).Right(3));
+		TestNotEqual(TEXT("FCbField::Clone(Void)"), VoidClone.AsBinaryView(), MakeMemoryView(Payload).Right(3));
+		TestTrue(TEXT("FCbField::Clone(Void)->EqualBytes"), VoidClone.AsBinaryView().EqualBytes(VoidView.AsBinaryView()));
+		TestEqual(TEXT("FCbField::MakeView(Field)"), FieldView.AsBinaryView(), MakeMemoryView(Payload).Right(3));
+		TestNotEqual(TEXT("FCbField::Clone(Field)"), FieldClone.AsBinaryView(), MakeMemoryView(Payload).Right(3));
+		TestTrue(TEXT("FCbField::Clone(Field)->EqualBytes"), FieldClone.AsBinaryView().EqualBytes(VoidView.AsBinaryView()));
+		TestNotEqual(TEXT("FCbField::Clone(FieldView)"), FieldViewClone.AsBinaryView(), FieldView.AsBinaryView());
+		TestTrue(TEXT("FCbField::Clone(FieldView)->EqualBytes"), FieldViewClone.AsBinaryView().EqualBytes(VoidView.AsBinaryView()));
+
+		TestFalse(TEXT("FCbField::MakeView(Void).IsOwned()"), VoidView.IsOwned());
+		TestTrue(TEXT("FCbField::Clone(Void).IsOwned()"), VoidClone.IsOwned());
+		TestFalse(TEXT("FCbField::MakeView(Field).IsOwned()"), FieldView.IsOwned());
+		TestTrue(TEXT("FCbField::Clone(Field).IsOwned()"), FieldClone.IsOwned());
+		TestTrue(TEXT("FCbField::Clone(FieldView).IsOwned()"), FieldViewClone.IsOwned());
 	}
 
 	// Test Field w/o Type
 	{
 		const uint8 Payload[] = { 3, 4, 5, 6 }; // Size: 3, Data: 4/5/6
-		const FCbField Field(Payload, ECbFieldType::Binary);
+		const FCbFieldView Field(Payload, ECbFieldType::Binary);
 
-		FCbFieldRef FieldView = FCbFieldRef::MakeView(Field);
-		FCbFieldRef FieldClone = FCbFieldRef::Clone(Field);
-		FCbFieldRef FieldRefClone = FCbFieldRef::Clone(FieldView);
+		FCbField FieldView = FCbField::MakeView(Field);
+		FCbField FieldClone = FCbField::Clone(Field);
+		FCbField FieldViewClone = FCbField::Clone(FieldView);
 
-		TestEqual(TEXT("FCbFieldRef::MakeView(Field, NoType)"), FieldView.AsBinary(), MakeMemoryView(Payload).Right(3));
-		TestTrue(TEXT("FCbFieldRef::Clone(Field, NoType)"), FieldClone.AsBinary().EqualBytes(FieldView.AsBinary()));
-		TestTrue(TEXT("FCbFieldRef::Clone(FieldRef, NoType)"), FieldRefClone.AsBinary().EqualBytes(FieldView.AsBinary()));
+		TestEqual(TEXT("FCbField::MakeView(Field, NoType)"), FieldView.AsBinaryView(), MakeMemoryView(Payload).Right(3));
+		TestTrue(TEXT("FCbField::Clone(Field, NoType)"), FieldClone.AsBinaryView().EqualBytes(FieldView.AsBinaryView()));
+		TestTrue(TEXT("FCbField::Clone(FieldView, NoType)"), FieldViewClone.AsBinaryView().EqualBytes(FieldView.AsBinaryView()));
 
-		TestFalse(TEXT("FCbFieldRef::MakeView(Field, NoType).IsOwned()"), FieldView.IsOwned());
-		TestTrue(TEXT("FCbFieldRef::Clone(Field, NoType).IsOwned()"), FieldClone.IsOwned());
-		TestTrue(TEXT("FCbFieldRef::Clone(FieldRef, NoType).IsOwned()"), FieldRefClone.IsOwned());
+		TestFalse(TEXT("FCbField::MakeView(Field, NoType).IsOwned()"), FieldView.IsOwned());
+		TestTrue(TEXT("FCbField::Clone(Field, NoType).IsOwned()"), FieldClone.IsOwned());
+		TestTrue(TEXT("FCbField::Clone(FieldView, NoType).IsOwned()"), FieldViewClone.IsOwned());
 
 		FieldView.MakeOwned();
-		TestTrue(TEXT("FCbFieldRef::MakeView(NoType).MakeOwned()"), FieldView.AsBinary().EqualBytes(MakeMemoryView(Payload).Right(3)));
-		TestTrue(TEXT("FCbFieldRef::MakeView(NoType).MakeOwned().IsOwned()"), FieldView.IsOwned());
+		TestTrue(TEXT("FCbField::MakeView(NoType).MakeOwned()"), FieldView.AsBinaryView().EqualBytes(MakeMemoryView(Payload).Right(3)));
+		TestTrue(TEXT("FCbField::MakeView(NoType).MakeOwned().IsOwned()"), FieldView.IsOwned());
 	}
 
 	return true;
@@ -1520,27 +1520,27 @@ bool FCbFieldRefTest::RunTest(const FString& Parameters)
 
 ///////////////////////////////////////////////////////////////////////////////////////////////////////////////////////
 
-IMPLEMENT_SIMPLE_AUTOMATION_TEST(FCbArrayRefTest, "System.Core.Serialization.CbArrayRef", CompactBinaryTestFlags)
-bool FCbArrayRefTest::RunTest(const FString& Parameters)
+IMPLEMENT_SIMPLE_AUTOMATION_TEST(FCbArrayBufferTest, "System.Core.Serialization.CbArrayBuffer", CompactBinaryTestFlags)
+bool FCbArrayBufferTest::RunTest(const FString& Parameters)
 {
-	static_assert(std::is_constructible<FCbArrayRef, const FCbArray&, const FSharedBuffer&>::value, "Missing constructor for FCbArrayRef");
-	static_assert(std::is_constructible<FCbArrayRef, const FCbArray&, const FCbFieldRefIterator&>::value, "Missing constructor for FCbArrayRef");
-	static_assert(std::is_constructible<FCbArrayRef, const FCbArray&, const FCbFieldRef&>::value, "Missing constructor for FCbArrayRef");
-	static_assert(std::is_constructible<FCbArrayRef, const FCbArray&, const FCbArrayRef&>::value, "Missing constructor for FCbArrayRef");
-	static_assert(std::is_constructible<FCbArrayRef, const FCbArray&, const FCbObjectRef&>::value, "Missing constructor for FCbArrayRef");
-
-	static_assert(std::is_constructible<FCbArrayRef, const FCbArray&, FSharedBuffer&&>::value, "Missing constructor for FCbArrayRef");
-	static_assert(std::is_constructible<FCbArrayRef, const FCbArray&, FCbFieldRefIterator&&>::value, "Missing constructor for FCbArrayRef");
-	static_assert(std::is_constructible<FCbArrayRef, const FCbArray&, FCbFieldRef&&>::value, "Missing constructor for FCbArrayRef");
-	static_assert(std::is_constructible<FCbArrayRef, const FCbArray&, FCbArrayRef&&>::value, "Missing constructor for FCbArrayRef");
-	static_assert(std::is_constructible<FCbArrayRef, const FCbArray&, FCbObjectRef&&>::value, "Missing constructor for FCbArrayRef");
-
-	// Test FCbArrayRef()
-	{
-		FCbArrayRef DefaultArray;
-		TestFalse(TEXT("FCbArrayRef().IsOwned()"), DefaultArray.IsOwned());
+	static_assert(std::is_constructible<FCbArray, const FCbArrayView&, const FSharedBuffer&>::value, "Missing constructor for FCbArray");
+	static_assert(std::is_constructible<FCbArray, const FCbArrayView&, const FCbFieldIterator&>::value, "Missing constructor for FCbArray");
+	static_assert(std::is_constructible<FCbArray, const FCbArrayView&, const FCbField&>::value, "Missing constructor for FCbArray");
+	static_assert(std::is_constructible<FCbArray, const FCbArrayView&, const FCbArray&>::value, "Missing constructor for FCbArray");
+	static_assert(std::is_constructible<FCbArray, const FCbArrayView&, const FCbObject&>::value, "Missing constructor for FCbArray");
+
+	static_assert(std::is_constructible<FCbArray, const FCbArrayView&, FSharedBuffer&&>::value, "Missing constructor for FCbArray");
+	static_assert(std::is_constructible<FCbArray, const FCbArrayView&, FCbFieldIterator&&>::value, "Missing constructor for FCbArray");
+	static_assert(std::is_constructible<FCbArray, const FCbArrayView&, FCbField&&>::value, "Missing constructor for FCbArray");
+	static_assert(std::is_constructible<FCbArray, const FCbArrayView&, FCbArray&&>::value, "Missing constructor for FCbArray");
+	static_assert(std::is_constructible<FCbArray, const FCbArrayView&, FCbObject&&>::value, "Missing constructor for FCbArray");
+
+	// Test FCbArray()
+	{
+		FCbArray DefaultArray;
+		TestFalse(TEXT("FCbArray().IsOwned()"), DefaultArray.IsOwned());
 		DefaultArray.MakeOwned();
-		TestTrue(TEXT("FCbArrayRef().MakeOwned().IsOwned()"), DefaultArray.IsOwned());
+		TestTrue(TEXT("FCbArray().MakeOwned().IsOwned()"), DefaultArray.IsOwned());
 	}
 
 	return true;
@@ -1548,27 +1548,27 @@ bool FCbArrayRefTest::RunTest(const FString& Parameters)
 
 ///////////////////////////////////////////////////////////////////////////////////////////////////////////////////////
 
-IMPLEMENT_SIMPLE_AUTOMATION_TEST(FCbObjectRefTest, "System.Core.Serialization.CbObjectRef", CompactBinaryTestFlags)
-bool FCbObjectRefTest::RunTest(const FString& Parameters)
+IMPLEMENT_SIMPLE_AUTOMATION_TEST(FCbObjectBufferTest, "System.Core.Serialization.CbObjectBuffer", CompactBinaryTestFlags)
+bool FCbObjectBufferTest::RunTest(const FString& Parameters)
 {
-	static_assert(std::is_constructible<FCbObjectRef, const FCbObject&, const FSharedBuffer&>::value, "Missing constructor for FCbObjectRef");
-	static_assert(std::is_constructible<FCbObjectRef, const FCbObject&, const FCbFieldRefIterator&>::value, "Missing constructor for FCbObjectRef");
-	static_assert(std::is_constructible<FCbObjectRef, const FCbObject&, const FCbFieldRef&>::value, "Missing constructor for FCbObjectRef");
-	static_assert(std::is_constructible<FCbObjectRef, const FCbObject&, const FCbArrayRef&>::value, "Missing constructor for FCbObjectRef");
-	static_assert(std::is_constructible<FCbObjectRef, const FCbObject&, const FCbObjectRef&>::value, "Missing constructor for FCbObjectRef");
-
-	static_assert(std::is_constructible<FCbObjectRef, const FCbObject&, FSharedBuffer&&>::value, "Missing constructor for FCbObjectRef");
-	static_assert(std::is_constructible<FCbObjectRef, const FCbObject&, FCbFieldRefIterator&&>::value, "Missing constructor for FCbObjectRef");
-	static_assert(std::is_constructible<FCbObjectRef, const FCbObject&, FCbFieldRef&&>::value, "Missing constructor for FCbObjectRef");
-	static_assert(std::is_constructible<FCbObjectRef, const FCbObject&, FCbArrayRef&&>::value, "Missing constructor for FCbObjectRef");
-	static_assert(std::is_constructible<FCbObjectRef, const FCbObject&, FCbObjectRef&&>::value, "Missing constructor for FCbObjectRef");
-
-	// Test FCbObjectRef()
-	{
-		FCbObjectRef DefaultObject;
-		TestFalse(TEXT("FCbObjectRef().IsOwned()"), DefaultObject.IsOwned());
+	static_assert(std::is_constructible<FCbObject, const FCbObjectView&, const FSharedBuffer&>::value, "Missing constructor for FCbObject");
+	static_assert(std::is_constructible<FCbObject, const FCbObjectView&, const FCbFieldIterator&>::value, "Missing constructor for FCbObject");
+	static_assert(std::is_constructible<FCbObject, const FCbObjectView&, const FCbField&>::value, "Missing constructor for FCbObject");
+	static_assert(std::is_constructible<FCbObject, const FCbObjectView&, const FCbArray&>::value, "Missing constructor for FCbObject");
+	static_assert(std::is_constructible<FCbObject, const FCbObjectView&, const FCbObject&>::value, "Missing constructor for FCbObject");
+
+	static_assert(std::is_constructible<FCbObject, const FCbObjectView&, FSharedBuffer&&>::value, "Missing constructor for FCbObject");
+	static_assert(std::is_constructible<FCbObject, const FCbObjectView&, FCbFieldIterator&&>::value, "Missing constructor for FCbObject");
+	static_assert(std::is_constructible<FCbObject, const FCbObjectView&, FCbField&&>::value, "Missing constructor for FCbObject");
+	static_assert(std::is_constructible<FCbObject, const FCbObjectView&, FCbArray&&>::value, "Missing constructor for FCbObject");
+	static_assert(std::is_constructible<FCbObject, const FCbObjectView&, FCbObject&&>::value, "Missing constructor for FCbObject");
+
+	// Test FCbObject()
+	{
+		FCbObject DefaultObject;
+		TestFalse(TEXT("FCbObject().IsOwned()"), DefaultObject.IsOwned());
 		DefaultObject.MakeOwned();
-		TestTrue(TEXT("FCbObjectRef().MakeOwned().IsOwned()"), DefaultObject.IsOwned());
+		TestTrue(TEXT("FCbObject().MakeOwned().IsOwned()"), DefaultObject.IsOwned());
 	}
 
 	return true;
@@ -1576,14 +1576,14 @@ bool FCbObjectRefTest::RunTest(const FString& Parameters)
 
 ///////////////////////////////////////////////////////////////////////////////////////////////////////////////////////
 
-IMPLEMENT_SIMPLE_AUTOMATION_TEST(FCbFieldRefIteratorTest, "System.Core.Serialization.CbFieldRefIterator", CompactBinaryTestFlags)
-bool FCbFieldRefIteratorTest::RunTest(const FString& Parameters)
+IMPLEMENT_SIMPLE_AUTOMATION_TEST(FCbFieldBufferIteratorTest, "System.Core.Serialization.CbFieldBufferIterator", CompactBinaryTestFlags)
+bool FCbFieldBufferIteratorTest::RunTest(const FString& Parameters)
 {
-	static_assert(std::is_constructible<FCbFieldIterator, const FCbFieldRefIterator&>::value, "Missing constructor for FCbFieldIterator");
-	static_assert(std::is_constructible<FCbFieldIterator, FCbFieldRefIterator&&>::value, "Missing constructor for FCbFieldIterator");
+	static_assert(std::is_constructible<FCbFieldViewIterator, const FCbFieldIterator&>::value, "Missing constructor for FCbFieldViewIterator");
+	static_assert(std::is_constructible<FCbFieldViewIterator, FCbFieldIterator&&>::value, "Missing constructor for FCbFieldViewIterator");
 
-	static_assert(std::is_constructible<FCbFieldRefIterator, const FCbFieldRefIterator&>::value, "Missing constructor for FCbFieldRefIterator");
-	static_assert(std::is_constructible<FCbFieldRefIterator, FCbFieldRefIterator&&>::value, "Missing constructor for FCbFieldRefIterator");
+	static_assert(std::is_constructible<FCbFieldIterator, const FCbFieldIterator&>::value, "Missing constructor for FCbFieldIterator");
+	static_assert(std::is_constructible<FCbFieldIterator, FCbFieldIterator&&>::value, "Missing constructor for FCbFieldIterator");
 
 	const auto GetCount = [](auto It) -> uint32
 	{
@@ -1595,13 +1595,13 @@ bool FCbFieldRefIteratorTest::RunTest(const FString& Parameters)
 		return Count;
 	};
 
-	// Test FCbField[Ref]Iterator()
+	// Test FCbField[View]Iterator()
 	{
+		TestEqual(TEXT("FCbFieldViewIterator()"), GetCount(FCbFieldViewIterator()), 0);
 		TestEqual(TEXT("FCbFieldIterator()"), GetCount(FCbFieldIterator()), 0);
-		TestEqual(TEXT("FCbFieldRefIterator()"), GetCount(FCbFieldRefIterator()), 0);
 	}
 
-	// Test FCbField[Ref]Iterator(Range)
+	// Test FCbField[View]Iterator(Range)
 	{
 		constexpr uint8 T = uint8(ECbFieldType::IntegerPositive);
 		const uint8 Payload[] = { T, 0, T, 1, T, 2, T, 3 };
@@ -1612,104 +1612,104 @@ bool FCbFieldRefIteratorTest::RunTest(const FString& Parameters)
 		const FMemoryView EmptyView;
 		const FSharedBuffer NullBuffer;
 
+		const FCbFieldViewIterator FieldViewIt = FCbFieldViewIterator::MakeRange(View);
 		const FCbFieldIterator FieldIt = FCbFieldIterator::MakeRange(View);
-		const FCbFieldRefIterator FieldRefIt = FCbFieldRefIterator::MakeRange(View);
 
+		TestEqual(TEXT("FCbFieldViewIterator::GetRangeHash()"), FieldViewIt.GetRangeHash(), FIoHash::HashBuffer(View));
 		TestEqual(TEXT("FCbFieldIterator::GetRangeHash()"), FieldIt.GetRangeHash(), FIoHash::HashBuffer(View));
-		TestEqual(TEXT("FCbFieldRefIterator::GetRangeHash()"), FieldRefIt.GetRangeHash(), FIoHash::HashBuffer(View));
 
+		TestEqual(TEXT("FCbFieldViewIterator::GetRangeView()"), FieldViewIt.GetRangeView(), MakeMemoryView(Payload));
 		TestEqual(TEXT("FCbFieldIterator::GetRangeView()"), FieldIt.GetRangeView(), MakeMemoryView(Payload));
-		TestEqual(TEXT("FCbFieldRefIterator::GetRangeView()"), FieldRefIt.GetRangeView(), MakeMemoryView(Payload));
 		FMemoryView SerializedView;
+		TestTrue(TEXT("FCbFieldViewIterator::TryGetSerializedRangeView()"), FieldViewIt.TryGetSerializedRangeView(SerializedView) && SerializedView == MakeMemoryView(Payload));
 		TestTrue(TEXT("FCbFieldIterator::TryGetSerializedRangeView()"), FieldIt.TryGetSerializedRangeView(SerializedView) && SerializedView == MakeMemoryView(Payload));
-		TestTrue(TEXT("FCbFieldRefIterator::TryGetSerializedRangeView()"), FieldRefIt.TryGetSerializedRangeView(SerializedView) && SerializedView == MakeMemoryView(Payload));
-
-		TestEqual(TEXT("FCbFieldRefIterator::CloneRange(EmptyIt)"), GetCount(FCbFieldRefIterator::CloneRange(FCbFieldIterator())), 0);
-		TestEqual(TEXT("FCbFieldRefIterator::CloneRange(EmptyRefIt)"), GetCount(FCbFieldRefIterator::CloneRange(FCbFieldRefIterator())), 0);
-		const FCbFieldRefIterator FieldItClone = FCbFieldRefIterator::CloneRange(FieldIt);
-		const FCbFieldRefIterator FieldRefItClone = FCbFieldRefIterator::CloneRange(FieldRefIt);
-		TestEqual(TEXT("FCbFieldRefIterator::CloneRange(FieldIt)"), GetCount(FieldItClone), 4);
-		TestEqual(TEXT("FCbFieldRefIterator::CloneRange(FieldRefIt)"), GetCount(FieldRefItClone), 4);
-		TestNotEqual(TEXT("FCbFieldRefIterator::CloneRange(FieldIt).Equals()"), FieldItClone, FieldRefIt);
-		TestNotEqual(TEXT("FCbFieldRefIterator::CloneRange(FieldRefIt).Equals()"), FieldRefItClone, FieldRefIt);
-
-		TestEqual(TEXT("FCbFieldIterator::MakeRange(EmptyView)"), GetCount(FCbFieldIterator::MakeRange(EmptyView)), 0);
-		TestEqual(TEXT("FCbFieldRefIterator::MakeRange(BufferNullL)"), GetCount(FCbFieldRefIterator::MakeRange(NullBuffer)), 0);
-		TestEqual(TEXT("FCbFieldRefIterator::MakeRange(BufferNullR)"), GetCount(FCbFieldRefIterator::MakeRange(FSharedBuffer(NullBuffer))), 0);
-
-		TestEqual(TEXT("FCbFieldIterator::MakeRange(BufferView)"), GetCount(FCbFieldIterator::MakeRange(MakeMemoryView(Payload))), 4);
-		TestEqual(TEXT("FCbFieldRefIterator::MakeRange(BufferCloneL)"), GetCount(FCbFieldRefIterator::MakeRange(Clone)), 4);
-		TestEqual(TEXT("FCbFieldRefIterator::MakeRange(BufferCloneR)"), GetCount(FCbFieldRefIterator::MakeRange(FSharedBuffer(Clone))), 4);
-
-		TestEqual(TEXT("FCbFieldRefIterator::MakeRangeView(FieldIt, BufferNullL)"), GetCount(FCbFieldRefIterator::MakeRangeView(FCbFieldIterator::MakeRange(View), NullBuffer)), 4);
-		TestEqual(TEXT("FCbFieldRefIterator::MakeRangeView(FieldIt, BufferNullR)"), GetCount(FCbFieldRefIterator::MakeRangeView(FCbFieldIterator::MakeRange(View), FSharedBuffer(NullBuffer))), 4);
-		TestEqual(TEXT("FCbFieldRefIterator::MakeRangeView(FieldIt, BufferViewL)"), GetCount(FCbFieldRefIterator::MakeRangeView(FCbFieldIterator::MakeRange(View), View)), 4);
-		TestEqual(TEXT("FCbFieldRefIterator::MakeRangeView(FieldIt, BufferViewR)"), GetCount(FCbFieldRefIterator::MakeRangeView(FCbFieldIterator::MakeRange(View), FSharedBuffer(View))), 4);
-		TestEqual(TEXT("FCbFieldRefIterator::MakeRangeView(FieldIt, BufferCloneL)"), GetCount(FCbFieldRefIterator::MakeRangeView(FCbFieldIterator::MakeRange(Clone), Clone)), 4);
-		TestEqual(TEXT("FCbFieldRefIterator::MakeRangeView(FieldIt, BufferCloneR)"), GetCount(FCbFieldRefIterator::MakeRangeView(FCbFieldIterator::MakeRange(Clone), FSharedBuffer(Clone))), 4);
-
-		TestEqual(TEXT("FCbFieldIterator(FieldRefItL)"), GetCount(FCbFieldIterator(FieldRefIt)), 4);
-		TestEqual(TEXT("FCbFieldIterator(FieldRefItR)"), GetCount(FCbFieldIterator(FCbFieldRefIterator(FieldRefIt))), 4);
+
+		TestEqual(TEXT("FCbFieldIterator::CloneRange(EmptyViewIt)"), GetCount(FCbFieldIterator::CloneRange(FCbFieldViewIterator())), 0);
+		TestEqual(TEXT("FCbFieldIterator::CloneRange(EmptyIt)"), GetCount(FCbFieldIterator::CloneRange(FCbFieldIterator())), 0);
+		const FCbFieldIterator FieldViewItClone = FCbFieldIterator::CloneRange(FieldViewIt);
+		const FCbFieldIterator FieldItClone = FCbFieldIterator::CloneRange(FieldIt);
+		TestEqual(TEXT("FCbFieldIterator::CloneRange(FieldViewIt)"), GetCount(FieldViewItClone), 4);
+		TestEqual(TEXT("FCbFieldIterator::CloneRange(FieldIt)"), GetCount(FieldItClone), 4);
+		TestNotEqual(TEXT("FCbFieldIterator::CloneRange(FieldViewIt).Equals()"), FieldViewItClone, FieldIt);
+		TestNotEqual(TEXT("FCbFieldIterator::CloneRange(FieldIt).Equals()"), FieldItClone, FieldIt);
+
+		TestEqual(TEXT("FCbFieldViewIterator::MakeRange(EmptyView)"), GetCount(FCbFieldViewIterator::MakeRange(EmptyView)), 0);
+		TestEqual(TEXT("FCbFieldIterator::MakeRange(BufferNullL)"), GetCount(FCbFieldIterator::MakeRange(NullBuffer)), 0);
+		TestEqual(TEXT("FCbFieldIterator::MakeRange(BufferNullR)"), GetCount(FCbFieldIterator::MakeRange(FSharedBuffer(NullBuffer))), 0);
+
+		TestEqual(TEXT("FCbFieldViewIterator::MakeRange(BufferView)"), GetCount(FCbFieldViewIterator::MakeRange(MakeMemoryView(Payload))), 4);
+		TestEqual(TEXT("FCbFieldIterator::MakeRange(BufferCloneL)"), GetCount(FCbFieldIterator::MakeRange(Clone)), 4);
+		TestEqual(TEXT("FCbFieldIterator::MakeRange(BufferCloneR)"), GetCount(FCbFieldIterator::MakeRange(FSharedBuffer(Clone))), 4);
+
+		TestEqual(TEXT("FCbFieldIterator::MakeRangeView(FieldIt, BufferNullL)"), GetCount(FCbFieldIterator::MakeRangeView(FCbFieldViewIterator::MakeRange(View), NullBuffer)), 4);
+		TestEqual(TEXT("FCbFieldIterator::MakeRangeView(FieldIt, BufferNullR)"), GetCount(FCbFieldIterator::MakeRangeView(FCbFieldViewIterator::MakeRange(View), FSharedBuffer(NullBuffer))), 4);
+		TestEqual(TEXT("FCbFieldIterator::MakeRangeView(FieldIt, BufferViewL)"), GetCount(FCbFieldIterator::MakeRangeView(FCbFieldViewIterator::MakeRange(View), View)), 4);
+		TestEqual(TEXT("FCbFieldIterator::MakeRangeView(FieldIt, BufferViewR)"), GetCount(FCbFieldIterator::MakeRangeView(FCbFieldViewIterator::MakeRange(View), FSharedBuffer(View))), 4);
+		TestEqual(TEXT("FCbFieldIterator::MakeRangeView(FieldIt, BufferCloneL)"), GetCount(FCbFieldIterator::MakeRangeView(FCbFieldViewIterator::MakeRange(Clone), Clone)), 4);
+		TestEqual(TEXT("FCbFieldIterator::MakeRangeView(FieldIt, BufferCloneR)"), GetCount(FCbFieldIterator::MakeRangeView(FCbFieldViewIterator::MakeRange(Clone), FSharedBuffer(Clone))), 4);
+
+		TestEqual(TEXT("FCbFieldViewIterator(FieldItL)"), GetCount(FCbFieldViewIterator(FieldIt)), 4);
+		TestEqual(TEXT("FCbFieldViewIterator(FieldItR)"), GetCount(FCbFieldViewIterator(FCbFieldIterator(FieldIt))), 4);
 
 		// Uniform
 		const uint8 UniformPayload[] = { 0, 1, 2, 3 };
-		const FCbFieldIterator UniformFieldIt = FCbFieldIterator::MakeRange(MakeMemoryView(UniformPayload), ECbFieldType::IntegerPositive);
-		TestEqual(TEXT("FCbFieldIterator::MakeRange(Uniform).GetRangeHash()"), UniformFieldIt.GetRangeHash(), FieldIt.GetRangeHash());
+		const FCbFieldViewIterator UniformFieldViewIt = FCbFieldViewIterator::MakeRange(MakeMemoryView(UniformPayload), ECbFieldType::IntegerPositive);
+		TestEqual(TEXT("FCbFieldViewIterator::MakeRange(Uniform).GetRangeHash()"), UniformFieldViewIt.GetRangeHash(), FieldViewIt.GetRangeHash());
+		TestEqual(TEXT("FCbFieldViewIterator::MakeRange(Uniform).GetRangeView()"), UniformFieldViewIt.GetRangeView(), MakeMemoryView(UniformPayload));
+		TestFalse(TEXT("FCbFieldViewIterator::MakeRange(Uniform).TryGetSerializedRangeView()"), UniformFieldViewIt.TryGetSerializedRangeView(SerializedView));
+		const FSharedBuffer UniformView = FSharedBuffer::MakeView(MakeMemoryView(UniformPayload));
+		const FCbFieldIterator UniformFieldIt = FCbFieldIterator::MakeRange(UniformView, ECbFieldType::IntegerPositive);
+		TestEqual(TEXT("FCbFieldIterator::MakeRange(Uniform).GetRangeHash()"), UniformFieldIt.GetRangeHash(), FieldViewIt.GetRangeHash());
 		TestEqual(TEXT("FCbFieldIterator::MakeRange(Uniform).GetRangeView()"), UniformFieldIt.GetRangeView(), MakeMemoryView(UniformPayload));
 		TestFalse(TEXT("FCbFieldIterator::MakeRange(Uniform).TryGetSerializedRangeView()"), UniformFieldIt.TryGetSerializedRangeView(SerializedView));
-		const FSharedBuffer UniformView = FSharedBuffer::MakeView(MakeMemoryView(UniformPayload));
-		const FCbFieldRefIterator UniformFieldRefIt = FCbFieldRefIterator::MakeRange(UniformView, ECbFieldType::IntegerPositive);
-		TestEqual(TEXT("FCbFieldRefIterator::MakeRange(Uniform).GetRangeHash()"), UniformFieldRefIt.GetRangeHash(), FieldIt.GetRangeHash());
-		TestEqual(TEXT("FCbFieldRefIterator::MakeRange(Uniform).GetRangeView()"), UniformFieldRefIt.GetRangeView(), MakeMemoryView(UniformPayload));
-		TestFalse(TEXT("FCbFieldRefIterator::MakeRange(Uniform).TryGetSerializedRangeView()"), UniformFieldRefIt.TryGetSerializedRangeView(SerializedView));
 
 		// Equals
+		TestTrue(TEXT("FCbFieldViewIterator::Equals(Self)"), FieldViewIt.Equals(FieldViewIt));
+		TestTrue(TEXT("FCbFieldViewIterator::Equals(OtherType)"), FieldViewIt.Equals(FieldIt));
 		TestTrue(TEXT("FCbFieldIterator::Equals(Self)"), FieldIt.Equals(FieldIt));
-		TestTrue(TEXT("FCbFieldIterator::Equals(OtherType)"), FieldIt.Equals(FieldRefIt));
-		TestTrue(TEXT("FCbFieldRefIterator::Equals(Self)"), FieldRefIt.Equals(FieldRefIt));
-		TestTrue(TEXT("FCbFieldRefIterator::Equals(OtherType)"), FieldRefIt.Equals(FieldIt));
+		TestTrue(TEXT("FCbFieldIterator::Equals(OtherType)"), FieldIt.Equals(FieldViewIt));
+		TestFalse(TEXT("FCbFieldViewIterator::Equals(OtherRange)"), FieldViewIt.Equals(FieldViewItClone));
 		TestFalse(TEXT("FCbFieldIterator::Equals(OtherRange)"), FieldIt.Equals(FieldItClone));
-		TestFalse(TEXT("FCbFieldRefIterator::Equals(OtherRange)"), FieldRefIt.Equals(FieldRefItClone));
+		TestTrue(TEXT("FCbFieldViewIterator::Equals(Uniform, Self)"), UniformFieldViewIt.Equals(UniformFieldViewIt));
+		TestTrue(TEXT("FCbFieldViewIterator::Equals(Uniform, OtherType)"), UniformFieldViewIt.Equals(UniformFieldIt));
 		TestTrue(TEXT("FCbFieldIterator::Equals(Uniform, Self)"), UniformFieldIt.Equals(UniformFieldIt));
-		TestTrue(TEXT("FCbFieldIterator::Equals(Uniform, OtherType)"), UniformFieldIt.Equals(UniformFieldRefIt));
-		TestTrue(TEXT("FCbFieldRefIterator::Equals(Uniform, Self)"), UniformFieldRefIt.Equals(UniformFieldRefIt));
-		TestTrue(TEXT("FCbFieldRefIterator::Equals(Uniform, OtherType)"), UniformFieldRefIt.Equals(UniformFieldIt));
-		TestFalse(TEXT("FCbFieldIterator::Equals(SamePayload, DifferentEnd)"),
-			FCbFieldIterator::MakeRange(MakeMemoryView(UniformPayload), ECbFieldType::IntegerPositive)
-				.Equals(FCbFieldIterator::MakeRange(MakeMemoryView(UniformPayload).LeftChop(1), ECbFieldType::IntegerPositive)));
-		TestFalse(TEXT("FCbFieldIterator::Equals(DifferentPayload, SameEnd)"),
-			FCbFieldIterator::MakeRange(MakeMemoryView(UniformPayload), ECbFieldType::IntegerPositive)
-				.Equals(FCbFieldIterator::MakeRange(MakeMemoryView(UniformPayload).RightChop(1), ECbFieldType::IntegerPositive)));
+		TestTrue(TEXT("FCbFieldIterator::Equals(Uniform, OtherType)"), UniformFieldIt.Equals(UniformFieldViewIt));
+		TestFalse(TEXT("FCbFieldViewIterator::Equals(SamePayload, DifferentEnd)"),
+			FCbFieldViewIterator::MakeRange(MakeMemoryView(UniformPayload), ECbFieldType::IntegerPositive)
+				.Equals(FCbFieldViewIterator::MakeRange(MakeMemoryView(UniformPayload).LeftChop(1), ECbFieldType::IntegerPositive)));
+		TestFalse(TEXT("FCbFieldViewIterator::Equals(DifferentPayload, SameEnd)"),
+			FCbFieldViewIterator::MakeRange(MakeMemoryView(UniformPayload), ECbFieldType::IntegerPositive)
+				.Equals(FCbFieldViewIterator::MakeRange(MakeMemoryView(UniformPayload).RightChop(1), ECbFieldType::IntegerPositive)));
 
 		// CopyRangeTo
 		uint8 CopyBytes[sizeof(Payload)];
+		FieldViewIt.CopyRangeTo(MakeMemoryView(CopyBytes));
+		TestTrue(TEXT("FCbFieldViewIterator::MakeRange().CopyRangeTo()"), MakeMemoryView(CopyBytes).EqualBytes(MakeMemoryView(Payload)));
 		FieldIt.CopyRangeTo(MakeMemoryView(CopyBytes));
 		TestTrue(TEXT("FCbFieldIterator::MakeRange().CopyRangeTo()"), MakeMemoryView(CopyBytes).EqualBytes(MakeMemoryView(Payload)));
-		FieldRefIt.CopyRangeTo(MakeMemoryView(CopyBytes));
-		TestTrue(TEXT("FCbFieldRefIterator::MakeRange().CopyRangeTo()"), MakeMemoryView(CopyBytes).EqualBytes(MakeMemoryView(Payload)));
-		UniformFieldIt.CopyRangeTo(MakeMemoryView(CopyBytes));
-		TestTrue(TEXT("FCbFieldIterator::MakeRange(Uniform).CopyRangeTo()"), MakeMemoryView(CopyBytes).EqualBytes(MakeMemoryView(Payload)));
+		UniformFieldViewIt.CopyRangeTo(MakeMemoryView(CopyBytes));
+		TestTrue(TEXT("FCbFieldViewIterator::MakeRange(Uniform).CopyRangeTo()"), MakeMemoryView(CopyBytes).EqualBytes(MakeMemoryView(Payload)));
 
 		// MakeRangeOwned
-		FCbFieldRefIterator OwnedFromView = UniformFieldRefIt;
+		FCbFieldIterator OwnedFromView = UniformFieldIt;
 		OwnedFromView.MakeRangeOwned();
-		TestTrue(TEXT("FCbFieldRefIterator::MakeRangeOwned(View)"), OwnedFromView.GetRangeView().EqualBytes(MakeMemoryView(Payload)));
-		FCbFieldRefIterator OwnedFromOwned = OwnedFromView;
+		TestTrue(TEXT("FCbFieldIterator::MakeRangeOwned(View)"), OwnedFromView.GetRangeView().EqualBytes(MakeMemoryView(Payload)));
+		FCbFieldIterator OwnedFromOwned = OwnedFromView;
 		OwnedFromOwned.MakeRangeOwned();
-		TestEqual(TEXT("FCbFieldRefIterator::MakeRangeOwned(Owned)"), OwnedFromOwned, OwnedFromView);
+		TestEqual(TEXT("FCbFieldIterator::MakeRangeOwned(Owned)"), OwnedFromOwned, OwnedFromView);
 
 		// GetRangeBuffer
-		TestEqual(TEXT("FCbFieldRefIterator::GetRangeBuffer()"), UniformFieldRefIt.GetRangeBuffer(), UniformView);
-		const FCbFieldRefIterator UniformShortFieldRefIt = FCbFieldRefIterator::MakeRangeView(FCbFieldIterator::MakeRange(MakeMemoryView(UniformPayload).LeftChop(1), ECbFieldType::IntegerPositive), UniformView);
-		TestNotEqual(TEXT("FCbFieldRefIterator::GetRangeBuffer(Short)"), UniformShortFieldRefIt.GetRangeBuffer(), UniformView);
+		TestEqual(TEXT("FCbFieldIterator::GetRangeBuffer()"), UniformFieldIt.GetRangeBuffer(), UniformView);
+		const FCbFieldIterator UniformShortFieldIt = FCbFieldIterator::MakeRangeView(FCbFieldViewIterator::MakeRange(MakeMemoryView(UniformPayload).LeftChop(1), ECbFieldType::IntegerPositive), UniformView);
+		TestNotEqual(TEXT("FCbFieldIterator::GetRangeBuffer(Short)"), UniformShortFieldIt.GetRangeBuffer(), UniformView);
 
 		// These lines are expected to assert when uncommented.
 		//const FSharedBuffer ShortView = FSharedBuffer::MakeView(MakeMemoryView(Payload).LeftChop(2));
-		//TestEqual(TEXT("FCbFieldRefIterator::MakeRangeView(FieldIt, InvalidBufferL)"), GetCount(FCbFieldRefIterator::MakeRangeView(FCbFieldIterator::MakeRange(*View), ShortView)), 4);
-		//TestEqual(TEXT("FCbFieldRefIterator::MakeRangeView(FieldIt, InvalidBufferR)"), GetCount(FCbFieldRefIterator::MakeRangeView(FCbFieldIterator::MakeRange(*View), FSharedBuffer(ShortView))), 4);
+		//TestEqual(TEXT("FCbFieldIterator::MakeRangeView(FieldIt, InvalidBufferL)"), GetCount(FCbFieldIterator::MakeRangeView(FCbFieldViewIterator::MakeRange(*View), ShortView)), 4);
+		//TestEqual(TEXT("FCbFieldIterator::MakeRangeView(FieldIt, InvalidBufferR)"), GetCount(FCbFieldIterator::MakeRangeView(FCbFieldViewIterator::MakeRange(*View), FSharedBuffer(ShortView))), 4);
 	}
 
-	// Test FCbField[Ref]Iterator(Scalar)
+	// Test FCbField[View]Iterator(Scalar)
 	{
 		constexpr uint8 T = uint8(ECbFieldType::IntegerPositive);
 		const uint8 Payload[] = { T, 0 };
@@ -1717,13 +1717,13 @@ bool FCbFieldRefIteratorTest::RunTest(const FString& Parameters)
 		const FSharedBuffer View = FSharedBuffer::MakeView(MakeMemoryView(Payload));
 		const FSharedBuffer Clone = FSharedBuffer::Clone(View);
 
-		const FCbField Field(Payload);
-		const FCbFieldRef FieldRef(View);
+		const FCbFieldView FieldView(Payload);
+		const FCbField Field(View);
 
+		TestEqual(TEXT("FCbFieldViewIterator::MakeSingle(FieldViewL)"), GetCount(FCbFieldViewIterator::MakeSingle(FieldView)), 1);
+		TestEqual(TEXT("FCbFieldViewIterator::MakeSingle(FieldViewR)"), GetCount(FCbFieldViewIterator::MakeSingle(FCbFieldView(FieldView))), 1);
 		TestEqual(TEXT("FCbFieldIterator::MakeSingle(FieldL)"), GetCount(FCbFieldIterator::MakeSingle(Field)), 1);
 		TestEqual(TEXT("FCbFieldIterator::MakeSingle(FieldR)"), GetCount(FCbFieldIterator::MakeSingle(FCbField(Field))), 1);
-		TestEqual(TEXT("FCbFieldRefIterator::MakeSingle(FieldRefL)"), GetCount(FCbFieldRefIterator::MakeSingle(FieldRef)), 1);
-		TestEqual(TEXT("FCbFieldRefIterator::MakeSingle(FieldRefR)"), GetCount(FCbFieldRefIterator::MakeSingle(FCbFieldRef(FieldRef))), 1);
 	}
 
 	return true;
@@ -1738,11 +1738,11 @@ bool FCbFieldParseTest::RunTest(const FString& Parameters)
 	// Under ideal conditions, when the fields are in the expected order and there are no extra fields,
 	// the loop will execute once and only one comparison will be performed for each field name. Either
 	// way, each field will only be visited once even if the loop needs to execute several times.
-	auto ParseObject = [this](const FCbObject& Object, uint32& A, uint32& B, uint32& C, uint32& D)
+	auto ParseObject = [this](const FCbObjectView& Object, uint32& A, uint32& B, uint32& C, uint32& D)
 	{
-		for (FCbFieldIterator It = Object.CreateIterator(); It;)
+		for (FCbFieldViewIterator It = Object.CreateViewIterator(); It;)
 		{
-			const FCbFieldIterator Last = It;
+			const FCbFieldViewIterator Last = It;
 			if (It.GetName().Equals("A"_ASV))
 			{
 				A = It.AsUInt32();
@@ -1773,19 +1773,19 @@ bool FCbFieldParseTest::RunTest(const FString& Parameters)
 	auto TestParseObject = [this, &ParseObject](std::initializer_list<uint8> Data, uint32 A, uint32 B, uint32 C, uint32 D) -> bool
 	{
 		uint32 ParsedA = 0, ParsedB = 0, ParsedC = 0, ParsedD = 0;
-		ParseObject(FCbObject(GetData(Data), ECbFieldType::Object), ParsedA, ParsedB, ParsedC, ParsedD);
+		ParseObject(FCbObjectView(GetData(Data), ECbFieldType::Object), ParsedA, ParsedB, ParsedC, ParsedD);
 		return A == ParsedA && B == ParsedB && C == ParsedC && D == ParsedD;
 	};
 
 	constexpr uint8 T = uint8(ECbFieldType::IntegerPositive | ECbFieldType::HasFieldName);
-	TestTrue(TEXT("FCbObject Parse(None)"), TestParseObject({0}, 0, 0, 0, 0));
-	TestTrue(TEXT("FCbObject Parse(ABCD)"), TestParseObject({16, T, 1, 'A', 1, T, 1, 'B', 2, T, 1, 'C', 3, T, 1, 'D', 4}, 1, 2, 3, 4));
-	TestTrue(TEXT("FCbObject Parse(BCDA)"), TestParseObject({16, T, 1, 'B', 2, T, 1, 'C', 3, T, 1, 'D', 4, T, 1, 'A', 1}, 1, 2, 3, 4));
-	TestTrue(TEXT("FCbObject Parse(BCD)"), TestParseObject({12, T, 1, 'B', 2, T, 1, 'C', 3, T, 1, 'D', 4}, 0, 2, 3, 4));
-	TestTrue(TEXT("FCbObject Parse(BC)"), TestParseObject({8, T, 1, 'B', 2, T, 1, 'C', 3}, 0, 2, 3, 0));
-	TestTrue(TEXT("FCbObject Parse(ABCDE)"), TestParseObject({20, T, 1, 'A', 1, T, 1, 'B', 2, T, 1, 'C', 3, T, 1, 'D', 4, T, 1, 'E', 5}, 1, 2, 3, 4));
-	TestTrue(TEXT("FCbObject Parse(EABCD)"), TestParseObject({20, T, 1, 'E', 5, T, 1, 'A', 1, T, 1, 'B', 2, T, 1, 'C', 3, T, 1, 'D', 4}, 1, 2, 3, 4));
-	TestTrue(TEXT("FCbObject Parse(DCBA)"), TestParseObject({16, T, 1, 'D', 4, T, 1, 'C', 3, T, 1, 'B', 2, T, 1, 'A', 1}, 1, 2, 3, 4));
+	TestTrue(TEXT("FCbObjectView Parse(None)"), TestParseObject({0}, 0, 0, 0, 0));
+	TestTrue(TEXT("FCbObjectView Parse(ABCD)"), TestParseObject({16, T, 1, 'A', 1, T, 1, 'B', 2, T, 1, 'C', 3, T, 1, 'D', 4}, 1, 2, 3, 4));
+	TestTrue(TEXT("FCbObjectView Parse(BCDA)"), TestParseObject({16, T, 1, 'B', 2, T, 1, 'C', 3, T, 1, 'D', 4, T, 1, 'A', 1}, 1, 2, 3, 4));
+	TestTrue(TEXT("FCbObjectView Parse(BCD)"), TestParseObject({12, T, 1, 'B', 2, T, 1, 'C', 3, T, 1, 'D', 4}, 0, 2, 3, 4));
+	TestTrue(TEXT("FCbObjectView Parse(BC)"), TestParseObject({8, T, 1, 'B', 2, T, 1, 'C', 3}, 0, 2, 3, 0));
+	TestTrue(TEXT("FCbObjectView Parse(ABCDE)"), TestParseObject({20, T, 1, 'A', 1, T, 1, 'B', 2, T, 1, 'C', 3, T, 1, 'D', 4, T, 1, 'E', 5}, 1, 2, 3, 4));
+	TestTrue(TEXT("FCbObjectView Parse(EABCD)"), TestParseObject({20, T, 1, 'E', 5, T, 1, 'A', 1, T, 1, 'B', 2, T, 1, 'C', 3, T, 1, 'D', 4}, 1, 2, 3, 4));
+	TestTrue(TEXT("FCbObjectView Parse(DCBA)"), TestParseObject({16, T, 1, 'D', 4, T, 1, 'C', 3, T, 1, 'B', 2, T, 1, 'A', 1}, 1, 2, 3, 4));
 
 	return true;
 }
diff --git a/Engine/Source/Runtime/Core/Private/Tests/Serialization/CompactBinaryValidationTest.cpp b/Engine/Source/Runtime/Core/Private/Tests/Serialization/CompactBinaryValidationTest.cpp
--- a/Engine/Source/Runtime/Core/Private/Tests/Serialization/CompactBinaryValidationTest.cpp
+++ b/Engine/Source/Runtime/Core/Private/Tests/Serialization/CompactBinaryValidationTest.cpp
@@ -368,7 +368,7 @@ bool FCbValidateAttachmentTest::RunTest(const FString& Parameters)
 	// Test CompactBinary
 	{
 		FBufferArchive Buffer;
-		FCbAttachment(FCbFieldRefIterator::MakeSingle(FCbFieldRef::MakeView(FCbField(nullptr, ECbFieldType::Null)))).Save(Buffer);
+		FCbAttachment(FCbFieldIterator::MakeSingle(FCbField::MakeView(FCbFieldView(nullptr, ECbFieldType::Null)))).Save(Buffer);
 		TestEqual(TEXT("ValidateCompactBinaryAttachment(CompactBinary)"), ValidateCompactBinaryAttachment(MakeMemoryView(Buffer), ECbValidateMode::All), ECbValidateError::None);
 		TestEqual(TEXT("ValidateCompactBinaryAttachment(CompactBinary, MissingHash)"), ValidateCompactBinaryAttachment(MakeMemoryView(Buffer).Left(3), ECbValidateMode::All), ECbValidateError::InvalidPackageFormat);
 		TestEqual(TEXT("ValidateCompactBinaryAttachment(CompactBinary, MissingHash, Mode)"), ValidateCompactBinaryAttachment(MakeMemoryView(Buffer).Left(3), ECbValidateMode::All & ~ECbValidateMode::Package), ECbValidateError::None);
@@ -392,7 +392,7 @@ bool FCbValidateAttachmentTest::RunTest(const FString& Parameters)
 	}
 	{
 		FBufferArchive Buffer;
-		FCbFieldRefIterator::MakeSingle(FCbObjectRef().AsFieldRef()).CopyRangeTo(Buffer);
+		FCbFieldIterator::MakeSingle(FCbObject().AsField()).CopyRangeTo(Buffer);
 		TestEqual(TEXT("ValidateCompactBinaryAttachment(Object)"), ValidateCompactBinaryAttachment(MakeMemoryView(Buffer), ECbValidateMode::All), ECbValidateError::InvalidPackageFormat);
 		TestEqual(TEXT("ValidateCompactBinaryAttachment(Object, Mode)"), ValidateCompactBinaryAttachment(MakeMemoryView(Buffer), ECbValidateMode::All & ~ECbValidateMode::Package), ECbValidateError::None);
 	}
@@ -581,7 +581,7 @@ bool FCbValidatePackageTest::RunTest(const FString& Parameters)
 		TCbWriter<256> Writer;
 		Writer.BeginObject();
 		Writer.EndObject();
-		Writer.AddCompactBinaryAttachment(FCbObject().GetHash());
+		Writer.AddCompactBinaryAttachment(FCbObjectView().GetHash());
 		Writer.AddNull();
 		Writer.Save(Buffer);
 		TestEqual(TEXT("ValidateCompactBinaryPackage(NullObject)"), ValidateCompactBinaryPackage(MakeMemoryView(Buffer), ECbValidateMode::All), ECbValidateError::NullPackageObject);
diff --git a/Engine/Source/Runtime/Core/Private/Tests/Serialization/CompactBinaryWriterTest.cpp b/Engine/Source/Runtime/Core/Private/Tests/Serialization/CompactBinaryWriterTest.cpp
--- a/Engine/Source/Runtime/Core/Private/Tests/Serialization/CompactBinaryWriterTest.cpp
+++ b/Engine/Source/Runtime/Core/Private/Tests/Serialization/CompactBinaryWriterTest.cpp
@@ -28,11 +28,11 @@ bool FCbWriterObjectTest::RunTest(const FString& Parameters)
 		Writer.Reset();
 		Writer.BeginObject();
 		Writer.EndObject();
-		FCbFieldRef Field = Writer.Save();
+		FCbField Field = Writer.Save();
 		if (TestEqual(TEXT("FCbWriter(Object, Empty) Validate"), ValidateCompactBinary(Field.GetBuffer(), ECbValidateMode::All), ECbValidateError::None))
 		{
 			TestTrue(TEXT("FCbWriter(Object, Empty).IsObject()"), Field.IsObject());
-			TestFalse(TEXT("FCbWriter(Object, Empty).AsObject()"), Field.AsObject().CreateIterator().HasValue());
+			TestFalse(TEXT("FCbWriter(Object, Empty).AsObjectView()"), Field.AsObjectView().CreateViewIterator().HasValue());
 		}
 	}
 
@@ -42,11 +42,11 @@ bool FCbWriterObjectTest::RunTest(const FString& Parameters)
 		Writer.SetName("Object"_ASV);
 		Writer.BeginObject();
 		Writer.EndObject();
-		FCbFieldRef Field = Writer.Save();
+		FCbField Field = Writer.Save();
 		if (TestEqual(TEXT("FCbWriter(Object, Empty, Name) Validate"), ValidateCompactBinary(Field.GetBuffer(), ECbValidateMode::All), ECbValidateError::None))
 		{
 			TestTrue(TEXT("FCbWriter(Object, Empty, Name).IsObject()"), Field.IsObject());
-			TestFalse(TEXT("FCbWriter(Object, Empty, Name).AsObject()"), Field.AsObject().CreateIterator().HasValue());
+			TestFalse(TEXT("FCbWriter(Object, Empty, Name).AsObjectView()"), Field.AsObjectView().CreateViewIterator().HasValue());
 		}
 	}
 
@@ -57,13 +57,13 @@ bool FCbWriterObjectTest::RunTest(const FString& Parameters)
 		Writer.SetName("Integer"_ASV).AddInteger(0);
 		Writer.SetName("Float"_ASV).AddFloat(0.0f);
 		Writer.EndObject();
-		FCbFieldRef Field = Writer.Save();
+		FCbField Field = Writer.Save();
 		if (TestEqual(TEXT("FCbWriter(Object, Basic) Validate"), ValidateCompactBinary(Field.GetBuffer(), ECbValidateMode::All), ECbValidateError::None))
 		{
 			TestTrue(TEXT("FCbWriter(Object, Basic).IsObject()"), Field.IsObject());
-			FCbObject Object = Field.AsObject();
-			TestTrue(TEXT("FCbWriter(Object, Basic).AsObject()[Integer]"), Object["Integer"_ASV].IsInteger());
-			TestTrue(TEXT("FCbWriter(Object, Basic).AsObject()[Float]"), Object["Float"_ASV].IsFloat());
+			FCbObjectView Object = Field.AsObjectView();
+			TestTrue(TEXT("FCbWriter(Object, Basic).AsObjectView()[Integer]"), Object["Integer"_ASV].IsInteger());
+			TestTrue(TEXT("FCbWriter(Object, Basic).AsObjectView()[Float]"), Object["Float"_ASV].IsFloat());
 		}
 	}
 
@@ -74,13 +74,13 @@ bool FCbWriterObjectTest::RunTest(const FString& Parameters)
 		Writer.SetName("Field1"_ASV).AddInteger(0);
 		Writer.SetName("Field2"_ASV).AddInteger(1);
 		Writer.EndObject();
-		FCbFieldRef Field = Writer.Save();
+		FCbField Field = Writer.Save();
 		if (TestEqual(TEXT("FCbWriter(Object, Uniform) Validate"), ValidateCompactBinary(Field.GetBuffer(), ECbValidateMode::All), ECbValidateError::None))
 		{
 			TestTrue(TEXT("FCbWriter(Object, Uniform).IsObject()"), Field.IsObject());
-			FCbObject Object = Field.AsObject();
-			TestTrue(TEXT("FCbWriter(Object, Uniform).AsObject()[Field1]"), Object["Field1"_ASV].IsInteger());
-			TestTrue(TEXT("FCbWriter(Object, Uniform).AsObject()[Field2]"), Object["Field2"_ASV].IsInteger());
+			FCbObjectView Object = Field.AsObjectView();
+			TestTrue(TEXT("FCbWriter(Object, Uniform).AsObjectView()[Field1]"), Object["Field1"_ASV].IsInteger());
+			TestTrue(TEXT("FCbWriter(Object, Uniform).AsObjectView()[Field2]"), Object["Field2"_ASV].IsInteger());
 		}
 	}
 
@@ -99,11 +99,11 @@ bool FCbWriterArrayTest::RunTest(const FString& Parameters)
 		Writer.Reset();
 		Writer.BeginArray();
 		Writer.EndArray();
-		FCbFieldRef Field = Writer.Save();
+		FCbField Field = Writer.Save();
 		if (TestEqual(TEXT("FCbWriter(Array, Empty) Validate"), ValidateCompactBinary(Field.GetBuffer(), ECbValidateMode::All), ECbValidateError::None))
 		{
 			TestTrue(TEXT("FCbWriter(Array, Empty).IsArray()"), Field.IsArray());
-			TestEqual(TEXT("FCbWriter(Array, Empty).AsArray()"), Field.AsArray().Num(), uint64(0));
+			TestEqual(TEXT("FCbWriter(Array, Empty).AsArrayView()"), Field.AsArrayView().Num(), uint64(0));
 		}
 	}
 
@@ -113,11 +113,11 @@ bool FCbWriterArrayTest::RunTest(const FString& Parameters)
 		Writer.SetName("Array"_ASV);
 		Writer.BeginArray();
 		Writer.EndArray();
-		FCbFieldRef Field = Writer.Save();
+		FCbField Field = Writer.Save();
 		if (TestEqual(TEXT("FCbWriter(Array, Empty, Name) Validate"), ValidateCompactBinary(Field.GetBuffer(), ECbValidateMode::All), ECbValidateError::None))
 		{
 			TestTrue(TEXT("FCbWriter(Array, Empty, Name).IsArray()"), Field.IsArray());
-			TestEqual(TEXT("FCbWriter(Array, Empty, Name).AsArray()"), Field.AsArray().Num(), uint64(0));
+			TestEqual(TEXT("FCbWriter(Array, Empty, Name).AsArrayView()"), Field.AsArrayView().Num(), uint64(0));
 		}
 	}
 
@@ -128,16 +128,16 @@ bool FCbWriterArrayTest::RunTest(const FString& Parameters)
 		Writer.AddInteger(0);
 		Writer.AddFloat(0.0f);
 		Writer.EndArray();
-		FCbFieldRef Field = Writer.Save();
+		FCbField Field = Writer.Save();
 		if (TestEqual(TEXT("FCbWriter(Array, Basic) Validate"), ValidateCompactBinary(Field.GetBuffer(), ECbValidateMode::All), ECbValidateError::None))
 		{
 			TestTrue(TEXT("FCbWriter(Array, Basic).IsArray()"), Field.IsArray());
-			FCbFieldIterator Iterator = Field.AsArray().CreateIterator();
-			TestTrue(TEXT("FCbWriter(Array, Basic).AsArray()[Integer]"), Iterator.IsInteger());
+			FCbFieldViewIterator Iterator = Field.AsArrayView().CreateViewIterator();
+			TestTrue(TEXT("FCbWriter(Array, Basic).AsArrayView()[Integer]"), Iterator.IsInteger());
 			++Iterator;
-			TestTrue(TEXT("FCbWriter(Array, Basic).AsArray()[Float]"), Iterator.IsFloat());
+			TestTrue(TEXT("FCbWriter(Array, Basic).AsArrayView()[Float]"), Iterator.IsFloat());
 			++Iterator;
-			TestFalse(TEXT("FCbWriter(Array, Basic).AsArray()[End]"), Iterator.HasValue());
+			TestFalse(TEXT("FCbWriter(Array, Basic).AsArrayView()[End]"), Iterator.HasValue());
 		}
 	}
 
@@ -148,16 +148,16 @@ bool FCbWriterArrayTest::RunTest(const FString& Parameters)
 		Writer.AddInteger(0);
 		Writer.AddInteger(1);
 		Writer.EndArray();
-		FCbFieldRef Field = Writer.Save();
+		FCbField Field = Writer.Save();
 		if (TestEqual(TEXT("FCbWriter(Array, Uniform) Validate"), ValidateCompactBinary(Field.GetBuffer(), ECbValidateMode::All), ECbValidateError::None))
 		{
 			TestTrue(TEXT("FCbWriter(Array, Uniform).IsArray()"), Field.IsArray());
-			FCbFieldIterator Iterator = Field.AsArray().CreateIterator();
-			TestTrue(TEXT("FCbWriter(Array, Basic).AsArray()[Field1]"), Iterator.IsInteger());
+			FCbFieldViewIterator Iterator = Field.AsArrayView().CreateViewIterator();
+			TestTrue(TEXT("FCbWriter(Array, Basic).AsArrayView()[Field1]"), Iterator.IsInteger());
 			++Iterator;
-			TestTrue(TEXT("FCbWriter(Array, Basic).AsArray()[Field2]"), Iterator.IsInteger());
+			TestTrue(TEXT("FCbWriter(Array, Basic).AsArrayView()[Field2]"), Iterator.IsInteger());
 			++Iterator;
-			TestFalse(TEXT("FCbWriter(Array, Basic).AsArray()[End]"), Iterator.HasValue());
+			TestFalse(TEXT("FCbWriter(Array, Basic).AsArrayView()[End]"), Iterator.HasValue());
 		}
 	}
 
@@ -175,7 +175,7 @@ bool FCbWriterNullTest::RunTest(const FString& Parameters)
 	{
 		Writer.Reset();
 		Writer.AddNull();
-		FCbFieldRef Field = Writer.Save();
+		FCbField Field = Writer.Save();
 		if (TestEqual(TEXT("FCbWriter(Null) Validate"), ValidateCompactBinary(Field.GetBuffer(), ECbValidateMode::All), ECbValidateError::None))
 		{
 			TestFalse(TEXT("FCbWriter(Null).HasName()"), Field.HasName());
@@ -188,7 +188,7 @@ bool FCbWriterNullTest::RunTest(const FString& Parameters)
 		Writer.Reset();
 		Writer.SetName("Null"_ASV);
 		Writer.AddNull();
-		FCbFieldRef Field = Writer.Save();
+		FCbField Field = Writer.Save();
 		if (TestEqual(TEXT("FCbWriter(Null, Name) Validate"), ValidateCompactBinary(Field.GetBuffer(), ECbValidateMode::All), ECbValidateError::None))
 		{
 			TestEqual(TEXT("FCbWriter(Null, Name).GetName()"), Field.GetName(), "Null"_ASV);
@@ -213,7 +213,7 @@ bool FCbWriterNullTest::RunTest(const FString& Parameters)
 		Writer.SetName("N3"_ASV).AddNull();
 		Writer.EndObject();
 
-		FCbFieldRefIterator Fields = Writer.Save();
+		FCbFieldIterator Fields = Writer.Save();
 		TestEqual(TEXT("FCbWriter(Null, Uniform) Validate"), ValidateCompactBinaryRange(Fields.GetRangeBuffer(), ECbValidateMode::All), ECbValidateError::None);
 	}
 
@@ -226,7 +226,7 @@ bool FCbWriterNullTest::RunTest(const FString& Parameters)
 			Writer.AddNull();
 		}
 		uint8 Buffer[NullCount]{};
-		FCbFieldIterator Fields = Writer.Save(MakeMemoryView(Buffer));
+		FCbFieldViewIterator Fields = Writer.Save(MakeMemoryView(Buffer));
 		if (TestEqual(TEXT("FCbWriter(Null, Memory) Validate"), ValidateCompactBinaryRange(MakeMemoryView(Buffer), ECbValidateMode::All), ECbValidateError::None))
 		{
 			for (int Index = 0; Index < NullCount; ++Index)
@@ -265,12 +265,12 @@ bool FCbWriterBinaryTest::RunTest(const FString& Parameters)
 	{
 		Writer.Reset();
 		Writer.AddBinary(nullptr, 0);
-		FCbFieldRef Field = Writer.Save();
+		FCbField Field = Writer.Save();
 		if (TestEqual(TEXT("FCbWriter(Binary, Empty) Validate"), ValidateCompactBinary(Field.GetBuffer(), ECbValidateMode::All), ECbValidateError::None))
 		{
 			TestFalse(TEXT("FCbWriter(Binary, Empty).HasName()"), Field.HasName());
 			TestTrue(TEXT("FCbWriter(Binary, Empty).IsBinary()"), Field.IsBinary());
-			TestTrue(TEXT("FCbWriter(Binary, Empty).AsBinary()"), Field.AsBinary().IsEmpty());
+			TestTrue(TEXT("FCbWriter(Binary, Empty).AsBinaryView()"), Field.AsBinaryView().IsEmpty());
 		}
 	}
 
@@ -280,13 +280,13 @@ bool FCbWriterBinaryTest::RunTest(const FString& Parameters)
 		const uint8 BinaryValue[] = { 1, 2, 3, 4, 5, 6 };
 		Writer.SetName("Binary"_ASV);
 		Writer.AddBinary(BinaryValue, sizeof(BinaryValue));
-		FCbFieldRef Field = Writer.Save();
+		FCbField Field = Writer.Save();
 		if (TestEqual(TEXT("FCbWriter(Binary, Array) Validate"), ValidateCompactBinary(Field.GetBuffer(), ECbValidateMode::All), ECbValidateError::None))
 		{
 			TestEqual(TEXT("FCbWriter(Binary, Array).GetName()"), Field.GetName(), "Binary"_ASV);
 			TestTrue(TEXT("FCbWriter(Binary, Array).HasName()"), Field.HasName());
 			TestTrue(TEXT("FCbWriter(Binary, Array).IsBinary()"), Field.IsBinary());
-			TestTrue(TEXT("FCbWriter(Binary, Array).AsBinary()"), Field.AsBinary().EqualBytes(MakeMemoryView(BinaryValue)));
+			TestTrue(TEXT("FCbWriter(Binary, Array).AsBinaryView()"), Field.AsBinaryView().EqualBytes(MakeMemoryView(BinaryValue)));
 		}
 	}
 
@@ -305,10 +305,10 @@ bool FCbWriterStringTest::RunTest(const FString& Parameters)
 		Writer.Reset();
 		Writer.AddString(FAnsiStringView());
 		Writer.AddString(FWideStringView());
-		FCbFieldRefIterator Fields = Writer.Save();
+		FCbFieldIterator Fields = Writer.Save();
 		if (TestEqual(TEXT("FCbWriter(String, Empty) Validate"), ValidateCompactBinaryRange(Fields.GetRangeBuffer(), ECbValidateMode::All), ECbValidateError::None))
 		{
-			for (FCbField Field : Fields)
+			for (FCbFieldView Field : Fields)
 			{
 				TestFalse(TEXT("FCbWriter(String, Empty).HasName()"), Field.HasName());
 				TestTrue(TEXT("FCbWriter(String, Empty).IsString()"), Field.IsString());
@@ -322,10 +322,10 @@ bool FCbWriterStringTest::RunTest(const FString& Parameters)
 		Writer.Reset();
 		Writer.SetName("String"_ASV).AddString("Value"_ASV);
 		Writer.SetName("String"_ASV).AddString(TEXT("Value"_SV));
-		FCbFieldRefIterator Fields = Writer.Save();
+		FCbFieldIterator Fields = Writer.Save();
 		if (TestEqual(TEXT("FCbWriter(String, Basic) Validate"), ValidateCompactBinaryRange(Fields.GetRangeBuffer(), ECbValidateMode::All), ECbValidateError::None))
 		{
-			for (FCbField Field : Fields)
+			for (FCbFieldView Field : Fields)
 			{
 				TestEqual(TEXT("FCbWriter(String, Basic).GetName()"), Field.GetName(), "String"_ASV);
 				TestTrue(TEXT("FCbWriter(String, Basic).HasName()"), Field.HasName());
@@ -346,10 +346,10 @@ bool FCbWriterStringTest::RunTest(const FString& Parameters)
 		}
 		Writer.AddString(Dots);
 		Writer.AddString(FString::ChrN(DotCount, TEXT('.')));
-		FCbFieldRefIterator Fields = Writer.Save();
+		FCbFieldIterator Fields = Writer.Save();
 		if (TestEqual(TEXT("FCbWriter(String, Long) Validate"), ValidateCompactBinaryRange(Fields.GetRangeBuffer(), ECbValidateMode::All), ECbValidateError::None))
 		{
-			for (FCbField Field : Fields)
+			for (FCbFieldView Field : Fields)
 			{
 				TestEqual(TEXT("FCbWriter(String, Long).AsString()"), Field.AsString(), FAnsiStringView(Dots));
 			}
@@ -362,10 +362,10 @@ bool FCbWriterStringTest::RunTest(const FString& Parameters)
 		WIDECHAR Value[2] = { 0xd83d, 0xde00 };
 		Writer.AddString("\xf0\x9f\x98\x80"_ASV);
 		Writer.AddString(FWideStringView(Value, UE_ARRAY_COUNT(Value)));
-		FCbFieldRefIterator Fields = Writer.Save();
+		FCbFieldIterator Fields = Writer.Save();
 		if (TestEqual(TEXT("FCbWriter(String, Unicode) Validate"), ValidateCompactBinaryRange(Fields.GetRangeBuffer(), ECbValidateMode::All), ECbValidateError::None))
 		{
-			for (FCbField Field : Fields)
+			for (FCbFieldView Field : Fields)
 			{
 				TestEqual(TEXT("FCbWriter(String, Unicode).AsString()"), Field.AsString(), "\xf0\x9f\x98\x80"_ASV);
 			}
@@ -386,7 +386,7 @@ bool FCbWriterIntegerTest::RunTest(const FString& Parameters)
 	{
 		Writer.Reset();
 		Writer.AddInteger(Value);
-		FCbFieldRef Field = Writer.Save();
+		FCbField Field = Writer.Save();
 		if (TestEqual(TEXT("FCbWriter(Integer, Int32) Validate"), ValidateCompactBinary(Field.GetBuffer(), ECbValidateMode::All), ECbValidateError::None))
 		{
 			TestEqual(TEXT("FCbWriter(Integer, Int32) Value"), Field.AsInt32(), Value);
@@ -398,7 +398,7 @@ bool FCbWriterIntegerTest::RunTest(const FString& Parameters)
 	{
 		Writer.Reset();
 		Writer.AddInteger(Value);
-		FCbFieldRef Field = Writer.Save();
+		FCbField Field = Writer.Save();
 		if (TestEqual(TEXT("FCbWriter(Integer, UInt32) Validate"), ValidateCompactBinary(Field.GetBuffer(), ECbValidateMode::All), ECbValidateError::None))
 		{
 			TestEqual(TEXT("FCbWriter(Integer, UInt32) Value"), Field.AsUInt32(), Value);
@@ -410,7 +410,7 @@ bool FCbWriterIntegerTest::RunTest(const FString& Parameters)
 	{
 		Writer.Reset();
 		Writer.AddInteger(Value);
-		FCbFieldRef Field = Writer.Save();
+		FCbField Field = Writer.Save();
 		if (TestEqual(TEXT("FCbWriter(Integer, Int64) Validate"), ValidateCompactBinary(Field.GetBuffer(), ECbValidateMode::All), ECbValidateError::None))
 		{
 			TestEqual(TEXT("FCbWriter(Integer, Int64) Value"), Field.AsInt64(), Value);
@@ -422,7 +422,7 @@ bool FCbWriterIntegerTest::RunTest(const FString& Parameters)
 	{
 		Writer.Reset();
 		Writer.AddInteger(Value);
-		FCbFieldRef Field = Writer.Save();
+		FCbField Field = Writer.Save();
 		if (TestEqual(TEXT("FCbWriter(Integer, UInt64) Validate"), ValidateCompactBinary(Field.GetBuffer(), ECbValidateMode::All), ECbValidateError::None))
 		{
 			TestEqual(TEXT("FCbWriter(Integer, UInt64) Value"), Field.AsUInt64(), Value);
@@ -481,11 +481,11 @@ bool FCbWriterFloatTest::RunTest(const FString& Parameters)
 		{
 			Writer.AddFloat(Value);
 		}
-		FCbFieldRefIterator Fields = Writer.Save();
+		FCbFieldIterator Fields = Writer.Save();
 		if (TestEqual(TEXT("FCbWriter(Float, Single) Validate"), ValidateCompactBinaryRange(Fields.GetRangeBuffer(), ECbValidateMode::All), ECbValidateError::None))
 		{
 			const float* CheckValue = Values;
-			for (FCbField Field : Fields)
+			for (FCbFieldView Field : Fields)
 			{
 				TestEqual(TEXT("FCbWriter(Float, Single).AsFloat()"), Field.AsFloat(), *CheckValue++);
 				TestFalse(TEXT("FCbWriter(Float, Single) Error"), Field.HasError());
@@ -501,11 +501,11 @@ bool FCbWriterFloatTest::RunTest(const FString& Parameters)
 		{
 			Writer.AddFloat(Value);
 		}
-		FCbFieldRefIterator Fields = Writer.Save();
+		FCbFieldIterator Fields = Writer.Save();
 		if (TestEqual(TEXT("FCbWriter(Float, Double) Validate"), ValidateCompactBinaryRange(Fields.GetRangeBuffer(), ECbValidateMode::All), ECbValidateError::None))
 		{
 			const double* CheckValue = Values;
-			for (FCbField Field : Fields)
+			for (FCbFieldView Field : Fields)
 			{
 				TestEqual(TEXT("FCbWriter(Float, Double).AsDouble()"), Field.AsDouble(), *CheckValue++);
 				TestFalse(TEXT("FCbWriter(Float, Double) Error"), Field.HasError());
@@ -528,7 +528,7 @@ bool FCbWriterBoolTest::RunTest(const FString& Parameters)
 		Writer.AddBool(true);
 		Writer.AddBool(false);
 
-		FCbFieldRefIterator Fields = Writer.Save();
+		FCbFieldIterator Fields = Writer.Save();
 		if (TestEqual(TEXT("FCbWriter(Bool) Validate"), ValidateCompactBinaryRange(Fields.GetRangeBuffer(), ECbValidateMode::All), ECbValidateError::None))
 		{
 			TestTrue(TEXT("FCbWriter(Bool).AsBool()"), Fields.AsBool());
@@ -555,7 +555,7 @@ bool FCbWriterBoolTest::RunTest(const FString& Parameters)
 		Writer.SetName("B3"_ASV).AddBool(false);
 		Writer.EndObject();
 
-		FCbFieldRefIterator Fields = Writer.Save();
+		FCbFieldIterator Fields = Writer.Save();
 		TestEqual(TEXT("FCbWriter(Bool, Uniform) Validate"), ValidateCompactBinaryRange(Fields.GetRangeBuffer(), ECbValidateMode::All), ECbValidateError::None);
 	}
 
@@ -578,11 +578,11 @@ bool FCbWriterCompactBinaryAttachmentTest::RunTest(const FString& Parameters)
 		Writer.AddCompactBinaryAttachment(Value);
 	}
 
-	FCbFieldRefIterator Fields = Writer.Save();
+	FCbFieldIterator Fields = Writer.Save();
 	if (TestEqual(TEXT("FCbWriter(CompactBinaryAttachment) Validate"), ValidateCompactBinaryRange(Fields.GetRangeBuffer(), ECbValidateMode::All), ECbValidateError::None))
 	{
 		const FIoHash* CheckValue = Values;
-		for (FCbField Field : Fields)
+		for (FCbFieldView Field : Fields)
 		{
 			TestEqual(TEXT("FCbWriter(CompactBinaryAttachment).AsCompactBinaryAttachment()"), Field.AsCompactBinaryAttachment(), *CheckValue++);
 			TestFalse(TEXT("FCbWriter(CompactBinaryAttachment) Error"), Field.HasError());
@@ -608,11 +608,11 @@ bool FCbWriterBinaryAttachmentTest::RunTest(const FString& Parameters)
 		Writer.AddBinaryAttachment(Value);
 	}
 
-	FCbFieldRefIterator Fields = Writer.Save();
+	FCbFieldIterator Fields = Writer.Save();
 	if (TestEqual(TEXT("FCbWriter(BinaryAttachment) Validate"), ValidateCompactBinaryRange(Fields.GetRangeBuffer(), ECbValidateMode::All), ECbValidateError::None))
 	{
 		const FIoHash* CheckValue = Values;
-		for (FCbField Field : Fields)
+		for (FCbFieldView Field : Fields)
 		{
 			TestEqual(TEXT("FCbWriter(BinaryAttachment).AsBinaryAttachment()"), Field.AsBinaryAttachment(), *CheckValue++);
 			TestFalse(TEXT("FCbWriter(BinaryAttachment) Error"), Field.HasError());
@@ -638,11 +638,11 @@ bool FCbWriterHashTest::RunTest(const FString& Parameters)
 		Writer.AddHash(Value);
 	}
 
-	FCbFieldRefIterator Fields = Writer.Save();
+	FCbFieldIterator Fields = Writer.Save();
 	if (TestEqual(TEXT("FCbWriter(Hash) Validate"), ValidateCompactBinaryRange(Fields.GetRangeBuffer(), ECbValidateMode::All), ECbValidateError::None))
 	{
 		const FIoHash* CheckValue = Values;
-		for (FCbField Field : Fields)
+		for (FCbFieldView Field : Fields)
 		{
 			TestEqual(TEXT("FCbWriter(Hash).AsHash()"), Field.AsHash(), *CheckValue++);
 			TestFalse(TEXT("FCbWriter(Hash) Error"), Field.HasError());
@@ -665,11 +665,11 @@ bool FCbWriterUuidTest::RunTest(const FString& Parameters)
 		Writer.AddUuid(Value);
 	}
 
-	FCbFieldRefIterator Fields = Writer.Save();
+	FCbFieldIterator Fields = Writer.Save();
 	if (TestEqual(TEXT("FCbWriter(Uuid) Validate"), ValidateCompactBinaryRange(Fields.GetRangeBuffer(), ECbValidateMode::All), ECbValidateError::None))
 	{
 		const FGuid* CheckValue = Values;
-		for (FCbField Field : Fields)
+		for (FCbFieldView Field : Fields)
 		{
 			TestEqual(TEXT("FCbWriter(Uuid).AsUuid()"), Field.AsUuid(), *CheckValue++);
 			TestFalse(TEXT("FCbWriter(Uuid) Error"), Field.HasError());
@@ -692,11 +692,11 @@ bool FCbWriterDateTimeTest::RunTest(const FString& Parameters)
 		Writer.AddDateTime(Value);
 	}
 
-	FCbFieldRefIterator Fields = Writer.Save();
+	FCbFieldIterator Fields = Writer.Save();
 	if (TestEqual(TEXT("FCbWriter(DateTime) Validate"), ValidateCompactBinaryRange(Fields.GetRangeBuffer(), ECbValidateMode::All), ECbValidateError::None))
 	{
 		const FDateTime* CheckValue = Values;
-		for (FCbField Field : Fields)
+		for (FCbFieldView Field : Fields)
 		{
 			TestEqual(TEXT("FCbWriter(DateTime).AsDateTime()"), Field.AsDateTime(), *CheckValue++);
 			TestFalse(TEXT("FCbWriter(DateTime) Error"), Field.HasError());
@@ -719,11 +719,11 @@ bool FCbWriterTimeSpanTest::RunTest(const FString& Parameters)
 		Writer.AddTimeSpan(Value);
 	}
 
-	FCbFieldRefIterator Fields = Writer.Save();
+	FCbFieldIterator Fields = Writer.Save();
 	if (TestEqual(TEXT("FCbWriter(TimeSpan) Validate"), ValidateCompactBinaryRange(Fields.GetRangeBuffer(), ECbValidateMode::All), ECbValidateError::None))
 	{
 		const FTimespan* CheckValue = Values;
-		for (FCbField Field : Fields)
+		for (FCbFieldView Field : Fields)
 		{
 			TestEqual(TEXT("FCbWriter(TimeSpan).AsTimeSpan()"), Field.AsTimeSpan(), *CheckValue++);
 			TestFalse(TEXT("FCbWriter(TimeSpan) Error"), Field.HasError());
@@ -746,11 +746,11 @@ bool FCbWriterObjectIdTest::RunTest(const FString& Parameters)
 		Writer.AddObjectId(Value);
 	}
 
-	FCbFieldRefIterator Fields = Writer.Save();
+	FCbFieldIterator Fields = Writer.Save();
 	if (TestEqual(TEXT("FCbWriter(ObjectId) Validate"), ValidateCompactBinaryRange(Fields.GetRangeBuffer(), ECbValidateMode::All), ECbValidateError::None))
 	{
 		const FCbObjectId* CheckValue = Values;
-		for (FCbField Field : Fields)
+		for (FCbFieldView Field : Fields)
 		{
 			TestEqual(TEXT("FCbWriter(ObjectId).AsObjectId()"), Field.AsObjectId(), *CheckValue++);
 			TestFalse(TEXT("FCbWriter(ObjectId) Error"), Field.HasError());
@@ -783,11 +783,11 @@ bool FCbWriterCustomByIdTest::RunTest(const FString& Parameters)
 		Writer.AddCustom(Value.Type, MakeMemoryView(Value.Bytes));
 	}
 
-	FCbFieldRefIterator Fields = Writer.Save();
+	FCbFieldIterator Fields = Writer.Save();
 	if (TestEqual(TEXT("FCbWriter(CustomById) Validate"), ValidateCompactBinaryRange(Fields.GetRangeBuffer(), ECbValidateMode::All), ECbValidateError::None))
 	{
 		const FCustomValue* CheckValue = Values;
-		for (FCbField Field : Fields)
+		for (FCbFieldView Field : Fields)
 		{
 			TestTrue(TEXT("FCbWriter(CustomById).AsCustom()"), Field.AsCustom(CheckValue->Type).EqualBytes(MakeMemoryView(CheckValue->Bytes)));
 			TestFalse(TEXT("FCbWriter(CustomById) Error"), Field.HasError());
@@ -821,11 +821,11 @@ bool FCbWriterCustomByNameTest::RunTest(const FString& Parameters)
 		Writer.AddCustom(Value.Type, MakeMemoryView(Value.Bytes));
 	}
 
-	FCbFieldRefIterator Fields = Writer.Save();
+	FCbFieldIterator Fields = Writer.Save();
 	if (TestEqual(TEXT("FCbWriter(CustomByName) Validate"), ValidateCompactBinaryRange(Fields.GetRangeBuffer(), ECbValidateMode::All), ECbValidateError::None))
 	{
 		const FCustomValue* CheckValue = Values;
-		for (FCbField Field : Fields)
+		for (FCbFieldView Field : Fields)
 		{
 			TestTrue(TEXT("FCbWriter(CustomByName).AsCustom()"), Field.AsCustom(CheckValue->Type).EqualBytes(MakeMemoryView(CheckValue->Bytes)));
 			TestFalse(TEXT("FCbWriter(CustomByName) Error"), Field.HasError());
@@ -841,26 +841,26 @@ bool FCbWriterCustomByNameTest::RunTest(const FString& Parameters)
 IMPLEMENT_SIMPLE_AUTOMATION_TEST(FCbWriterComplexTest, "System.Core.Serialization.CbWriter.Complex", CompactBinaryWriterTestFlags)
 bool FCbWriterComplexTest::RunTest(const FString& Parameters)
 {
-	FCbObjectRef Object;
+	FCbObject Object;
 	FBufferArchive Archive;
 	{
 		FCbWriter Writer;
 		Writer.BeginObject();
 
 		const uint8 LocalField[] = { uint8(ECbFieldType::IntegerPositive | ECbFieldType::HasFieldName), 1, 'I', 42 };
-		Writer.AddField("FieldCopy"_ASV, FCbField(LocalField));
-		Writer.AddField("FieldRefCopy"_ASV, FCbFieldRef(FSharedBuffer::Clone(MakeMemoryView(LocalField))));
+		Writer.AddField("FieldViewCopy"_ASV, FCbFieldView(LocalField));
+		Writer.AddField("FieldCopy"_ASV, FCbField(FSharedBuffer::Clone(MakeMemoryView(LocalField))));
 
 		const uint8 LocalObject[] = { uint8(ECbFieldType::Object | ECbFieldType::HasFieldName), 1, 'O', 7,
 			uint8(ECbFieldType::IntegerPositive | ECbFieldType::HasFieldName), 1, 'I', 42,
 			uint8(ECbFieldType::Null | ECbFieldType::HasFieldName), 1, 'N' };
-		Writer.AddObject("ObjectCopy"_ASV, FCbObject(LocalObject));
-		Writer.AddObject("ObjectRefCopy"_ASV, FCbObjectRef(FSharedBuffer::Clone(MakeMemoryView(LocalObject))));
+		Writer.AddObject("ObjectViewCopy"_ASV, FCbObjectView(LocalObject));
+		Writer.AddObject("ObjectCopy"_ASV, FCbObject(FSharedBuffer::Clone(MakeMemoryView(LocalObject))));
 
 		const uint8 LocalArray[] = { uint8(ECbFieldType::UniformArray | ECbFieldType::HasFieldName), 1, 'A', 4, 2,
 			uint8(ECbFieldType::IntegerPositive), 42, 21 };
-		Writer.AddArray("ArrayCopy"_ASV, FCbArray(LocalArray));
-		Writer.AddArray("ArrayRefCopy"_ASV, FCbArrayRef(FSharedBuffer::Clone(MakeMemoryView(LocalArray))));
+		Writer.AddArray("ArrayViewCopy"_ASV, FCbArrayView(LocalArray));
+		Writer.AddArray("ArrayCopy"_ASV, FCbArray(FSharedBuffer::Clone(MakeMemoryView(LocalArray))));
 
 		Writer.AddNull("Null"_ASV);
 
@@ -868,8 +868,8 @@ bool FCbWriterComplexTest::RunTest(const FString& Parameters)
 		{
 			Writer.AddBinary("Empty"_ASV, FMemoryView());
 			Writer.AddBinary("Value"_ASV, MakeMemoryView("BinaryValue"));
-			Writer.AddBinary("LargeValue"_ASV, MakeMemoryView(FString::ChrN(256, TEXT('.'))));
-			Writer.AddBinary("LargeRefValue"_ASV, FSharedBuffer::Clone(MakeMemoryView(FString::ChrN(256, TEXT('!')))));
+			Writer.AddBinary("LargeViewValue"_ASV, MakeMemoryView(FString::ChrN(256, TEXT('.'))));
+			Writer.AddBinary("LargeValue"_ASV, FSharedBuffer::Clone(MakeMemoryView(FString::ChrN(256, TEXT('!')))));
 		}
 		Writer.EndObject();
 
@@ -1006,7 +1006,7 @@ bool FCbWriterComplexTest::RunTest(const FString& Parameters)
 		Writer.EndArray();
 
 		Writer.EndObject();
-		Object = Writer.Save().AsObjectRef();
+		Object = Writer.Save().AsObject();
 
 		Writer.Save(Archive);
 		TestEqual(TEXT("FCbWriter(Complex).Save(Ar)->Num()"), uint64(Archive.Num()), Writer.GetSaveSize());
@@ -1029,7 +1029,7 @@ bool FCbWriterOwnedReadOnlyTest::RunTest(const FString& Parameters)
 	FCbWriter Writer;
 	Writer.BeginObject();
 	Writer.EndObject();
-	FCbObjectRef Object = Writer.Save().AsObjectRef();
+	FCbObject Object = Writer.Save().AsObject();
 	TestTrue(TEXT("FCbWriter().Save().IsOwned()"), Object.IsOwned());
 
 	return true;
@@ -1040,22 +1040,22 @@ bool FCbWriterOwnedReadOnlyTest::RunTest(const FString& Parameters)
 IMPLEMENT_SIMPLE_AUTOMATION_TEST(FCbWriterStreamTest, "System.Core.Serialization.CbWriter.Stream", CompactBinaryWriterTestFlags)
 bool FCbWriterStreamTest::RunTest(const FString& Parameters)
 {
-	FCbObjectRef Object;
+	FCbObject Object;
 	{
 		FCbWriter Writer;
 		Writer.BeginObject();
 
 		const uint8 LocalField[] = { uint8(ECbFieldType::IntegerPositive | ECbFieldType::HasFieldName), 1, 'I', 42 };
-		Writer << "FieldCopy"_ASV << FCbField(LocalField);
+		Writer << "FieldCopy"_ASV << FCbFieldView(LocalField);
 
 		const uint8 LocalObject[] = { uint8(ECbFieldType::Object | ECbFieldType::HasFieldName), 1, 'O', 7,
 			uint8(ECbFieldType::IntegerPositive | ECbFieldType::HasFieldName), 1, 'I', 42,
 			uint8(ECbFieldType::Null | ECbFieldType::HasFieldName), 1, 'N' };
-		Writer << "ObjectCopy"_ASV << FCbObject(LocalObject);
+		Writer << "ObjectCopy"_ASV << FCbObjectView(LocalObject);
 
 		const uint8 LocalArray[] = { uint8(ECbFieldType::UniformArray | ECbFieldType::HasFieldName), 1, 'A', 4, 2,
 			uint8(ECbFieldType::IntegerPositive), 42, 21 };
-		Writer << "ArrayCopy"_ASV << FCbArray(LocalArray);
+		Writer << "ArrayCopy"_ASV << FCbArrayView(LocalArray);
 
 		Writer << "Null"_ASV << nullptr;
 
@@ -1092,7 +1092,7 @@ bool FCbWriterStreamTest::RunTest(const FString& Parameters)
 		Writer << "LiteralName" << nullptr;
 
 		Writer.EndObject();
-		Object = Writer.Save().AsObjectRef();
+		Object = Writer.Save().AsObject();
 	}
 
 	TestEqual(TEXT("FCbWriter(Stream) Validate"), ValidateCompactBinary(Object.GetBuffer(), ECbValidateMode::All), ECbValidateError::None);
@@ -1166,7 +1166,7 @@ bool FCbWriterStateTest::RunTest(const FString& Parameters)
 	//Writer.EndObject();
 
 	// Assert on writing a field with no value.
-	//Writer.Field(FCbField());
+	//Writer.Field(FCbFieldView());
 
 	return true;
 }
diff --git a/Engine/Source/Runtime/Core/Public/Serialization/CompactBinary.h b/Engine/Source/Runtime/Core/Public/Serialization/CompactBinary.h
--- a/Engine/Source/Runtime/Core/Public/Serialization/CompactBinary.h
+++ b/Engine/Source/Runtime/Core/Public/Serialization/CompactBinary.h
@@ -24,9 +24,9 @@
  * be cast to an object or array. This attribute means that a blob containing compact binary data
  * is always safe to interpret as a field, which allows for easy validation as described later.
  *
- * A field can be constructed as a view of the underlying memory with FCbField or can either view
- * or own (clone or take ownership) the underlying memory with FCbFieldRef. An array provides the
- * same behavior with FCbArray or FCbArrayRef, and an object uses FCbObject or FCbObjectRef.
+ * A field can be constructed as a view of the underlying memory with FCbFieldView, or a FCbField
+ * can used when ownership of the underlying memory is required. An object provides this behavior
+ * with FCbObjectView or FCbObject, and an array uses FCbArrayView or FCbArray.
  *
  * It is optimal use the view types when possible, and reference types only when they are needed,
  * to avoid the overhead of the atomic reference counting of the shared buffer.
@@ -37,7 +37,7 @@
  *
  * Example:
  *
- * void BeginBuild(FCbObjectRef Params)
+ * void BeginBuild(FCbObject Params)
  * {
  *     if (FSharedBuffer Data = Storage().Load(Params["Data"].AsBinaryAttachment()))
  *     {
@@ -46,15 +46,15 @@
  *
  *     if (Params["Resize"].AsBool())
  *     {
- *         FCbField MaxWidthField = Params["MaxWidth"]
- *         FCbField MaxHeightField = Params["MaxHeight"];
+ *         FCbFieldView MaxWidthField = Params["MaxWidth"]
+ *         FCbFieldView MaxHeightField = Params["MaxHeight"];
  *         if (MaxWidthField && MaxHeightField)
  *         {
  *             Resize(MaxWidthField.AsInt32(), MaxHeightField.AsInt32());
  *         }
  *     }
  *
- *     for (FCbField Format : Params["Formats"].AsArray())
+ *     for (FCbFieldView Format : Params["Formats"].AsArrayView())
  *     {
  *         BeginCompress(FName(Format.AsString()));
  *     }
@@ -63,10 +63,10 @@
 
 class FArchive;
 class FBlake3;
-class FCbArray;
-class FCbField;
-class FCbObject;
+class FCbArrayView;
+class FCbFieldView;
 class FCbObjectId;
+class FCbObjectView;
 struct FDateTime;
 struct FGuid;
 struct FTimespan;
@@ -75,7 +75,7 @@ template <typename FuncType> class TFunctionRef;
 ///////////////////////////////////////////////////////////////////////////////////////////////////////////////////////
 
 /**
- * Field types and flags for FCbField.
+ * Field types and flags for FCbField[View].
  *
  * DO NOT CHANGE THE VALUE OF ANY MEMBERS OF THIS ENUM!
  * BACKWARD COMPATIBILITY REQUIRES THAT THESE VALUES BE FIXED!
@@ -293,15 +293,15 @@ public:
 ///////////////////////////////////////////////////////////////////////////////////////////////////////////////////////
 
 /** A reference to a function that is used to visit fields. */
-using FCbFieldVisitor = TFunctionRef<void (FCbField)>;
+using FCbFieldVisitor = TFunctionRef<void (FCbFieldView)>;
 
 ///////////////////////////////////////////////////////////////////////////////////////////////////////////////////////
 
 /**
- * Iterator for FCbField[Ref] that can operate on any contiguous range of fields.
+ * Iterator for FCbField[View] that can operate on any contiguous range of fields.
  *
  * The iterator *is* the current field that the iterator points to and exposes the full interface
- * of FCbField[Ref]. An iterator that is at the end is equivalent to a field with no value.
+ * of FCbField[View]. An iterator that is at the end is equivalent to a field with no value.
  *
  * The iterator represents a range of fields from the current field to the last field.
  */
@@ -384,7 +384,7 @@ public:
 	 * Try to get a view of every field in the range as they would be serialized.
 	 *
 	 * A serialized view is not available if the underlying fields have an externally-provided type.
-	 * Access the serialized form of such ranges using FCbFieldRefIterator::CloneRange.
+	 * Access the serialized form of such ranges using FCbFieldIterator::CloneRange.
 	 */
 	inline bool TryGetSerializedRangeView(FMemoryView& OutView) const
 	{
@@ -542,13 +542,13 @@ struct FCbCustomByName
  * uniform array of fields with no payload, where the answer is to encode as a non-uniform array.
  *
  * This type only provides a view into memory and does not perform any memory management itself.
- * Use FCbFieldRef to hold a reference to the underlying memory when necessary.
+ * Use FCbField to hold a reference to the underlying memory when necessary.
  */
-class FCbField
+class FCbFieldView
 {
 public:
 	/** Construct a field with no name and no value. */
-	constexpr FCbField() = default;
+	constexpr FCbFieldView() = default;
 
 	/**
 	 * Construct a field from a pointer to its data and an optional externally-provided type.
@@ -556,7 +556,7 @@ public:
 	 * @param Data Pointer to the start of the field data.
 	 * @param Type HasFieldType means that Data contains the type. Otherwise, use the given type.
 	 */
-	CORE_API explicit FCbField(const void* Data, ECbFieldType Type = ECbFieldType::HasFieldType);
+	CORE_API explicit FCbFieldView(const void* Data, ECbFieldType Type = ECbFieldType::HasFieldType);
 
 	/** Returns the name of the field if it has a name, otherwise an empty view. */
 	constexpr inline FAnsiStringView GetName() const
@@ -565,13 +565,13 @@ public:
 	}
 
 	/** Access the field as an object. Defaults to an empty object on error. */
-	CORE_API FCbObject AsObject();
+	CORE_API FCbObjectView AsObjectView();
 
 	/** Access the field as an array. Defaults to an empty array on error. */
-	CORE_API FCbArray AsArray();
+	CORE_API FCbArrayView AsArrayView();
 
 	/** Access the field as binary. Returns the provided default on error. */
-	CORE_API FMemoryView AsBinary(FMemoryView Default = FMemoryView());
+	CORE_API FMemoryView AsBinaryView(FMemoryView Default = FMemoryView());
 
 	/** Access the field as a UTF-8 string. Returns the provided default on error. */
 	CORE_API FAnsiStringView AsString(FAnsiStringView Default = FAnsiStringView());
@@ -711,7 +711,7 @@ public:
 	 * these assumptions do not hold, this may return false for equivalent inputs. Validation can
 	 * be performed with ValidateCompactBinary, except for field order and field name case.
 	 */
-	CORE_API bool Equals(const FCbField& Other) const;
+	CORE_API bool Equals(const FCbFieldView& Other) const;
 
 	/** Copy the field into a buffer of exactly GetSize() bytes, including the type and name. */
 	CORE_API void CopyTo(FMutableMemoryView Buffer) const;
@@ -729,7 +729,7 @@ public:
 	 * Try to get a view of the field as it would be serialized, such as by CopyTo.
 	 *
 	 * A serialized view is not available if the field has an externally-provided type.
-	 * Access the serialized form of such fields using CopyTo or FCbFieldRef::Clone.
+	 * Access the serialized form of such fields using CopyTo or FCbField::Clone.
 	 */
 	inline bool TryGetSerializedView(FMemoryView& OutView) const
 	{
@@ -763,9 +763,9 @@ protected:
 	/** Assign a field from a pointer to its data and an optional externally-provided type. */
 	inline void Assign(const void* InData, const ECbFieldType InType)
 	{
-		static_assert(TIsTriviallyDestructible<FCbField>::Value,
-			"This optimization requires FCbField to be trivially destructible!");
-		new(this) FCbField(InData, InType);
+		static_assert(TIsTriviallyDestructible<FCbFieldView>::Value,
+			"This optimization requires FCbFieldView to be trivially destructible!");
+		new(this) FCbFieldView(InData, InType);
 	}
 
 private:
@@ -813,17 +813,17 @@ private:
 };
 
 /**
- * Iterator for FCbField.
+ * Iterator for FCbFieldView.
  *
  * @see TCbFieldIterator
  */
-class FCbFieldIterator : public TCbFieldIterator<FCbField>
+class FCbFieldViewIterator : public TCbFieldIterator<FCbFieldView>
 {
 public:
 	/** Construct a field range that contains exactly one field. */
-	static inline FCbFieldIterator MakeSingle(const FCbField& Field)
+	static inline FCbFieldViewIterator MakeSingle(const FCbFieldView& Field)
 	{
-		return FCbFieldIterator(Field);
+		return FCbFieldViewIterator(Field);
 	}
 
 	/**
@@ -832,18 +832,18 @@ public:
 	 * @param View A buffer containing zero or more valid fields.
 	 * @param Type HasFieldType means that View contains the type. Otherwise, use the given type.
 	 */
-	static inline FCbFieldIterator MakeRange(FMemoryView View, ECbFieldType Type = ECbFieldType::HasFieldType)
+	static inline FCbFieldViewIterator MakeRange(FMemoryView View, ECbFieldType Type = ECbFieldType::HasFieldType)
 	{
-		return !View.IsEmpty() ? FCbFieldIterator(FCbField(View.GetData(), Type), View.GetDataEnd()) : FCbFieldIterator();
+		return !View.IsEmpty() ? FCbFieldViewIterator(FCbFieldView(View.GetData(), Type), View.GetDataEnd()) : FCbFieldViewIterator();
 	}
 
 	/** Construct an empty field range. */
-	constexpr FCbFieldIterator() = default;
+	constexpr FCbFieldViewIterator() = default;
 
 	/** Construct an iterator from another iterator. */
 	template <typename OtherFieldType>
-	inline FCbFieldIterator(const TCbFieldIterator<OtherFieldType>& It)
-		: TCbFieldIterator(ImplicitConv<FCbField>(It), GetFieldsEnd(It))
+	inline FCbFieldViewIterator(const TCbFieldIterator<OtherFieldType>& It)
+		: TCbFieldIterator(ImplicitConv<FCbFieldView>(It), GetFieldsEnd(It))
 	{
 	}
 
@@ -854,34 +854,34 @@ private:
 ///////////////////////////////////////////////////////////////////////////////////////////////////////////////////////
 
 /**
- * Array of FCbField that have no names.
+ * Array of FCbField[View] that have no names.
  *
  * Accessing a field of the array requires iteration. Access by index is not provided because the
  * cost of accessing an item by index scales linearly with the index.
  *
  * This type only provides a view into memory and does not perform any memory management itself.
- * Use FCbArrayRef to hold a reference to the underlying memory when necessary.
+ * Use FCbArray to hold a reference to the underlying memory when necessary.
  */
-class FCbArray : protected FCbField
+class FCbArrayView : protected FCbFieldView
 {
 public:
-	/** @see FCbField::FCbField */
-	using FCbField::FCbField;
+	/** @see FCbFieldView::FCbFieldView */
+	using FCbFieldView::FCbFieldView;
 
 	/** Construct an array with no fields. */
-	CORE_API FCbArray();
+	CORE_API FCbArrayView();
 
 	/** Returns the number of items in the array. */
 	CORE_API uint64 Num() const;
 
 	/** Create an iterator for the fields of this array. */
-	CORE_API FCbFieldIterator CreateIterator() const;
+	CORE_API FCbFieldViewIterator CreateViewIterator() const;
 
 	/** Access the array as an array field. */
-	inline FCbField AsField() const { return static_cast<const FCbField&>(*this); }
+	inline FCbFieldView AsFieldView() const { return static_cast<const FCbFieldView&>(*this); }
 
 	/** Construct an array from an array field. No type check is performed! */
-	static inline FCbArray FromFieldNoCheck(const FCbField& Field) { return FCbArray(Field); }
+	static inline FCbArrayView FromFieldNoCheck(const FCbFieldView& Field) { return FCbArrayView(Field); }
 
 	/** Returns the size of the array in bytes if serialized by itself with no name. */
 	CORE_API uint64 GetSize() const;
@@ -900,7 +900,7 @@ public:
 	 * these assumptions do not hold, this may return false for equivalent inputs. Validation can
 	 * be done with the All mode to check these assumptions about the format of the inputs.
 	 */
-	CORE_API bool Equals(const FCbArray& Other) const;
+	CORE_API bool Equals(const FCbArrayView& Other) const;
 
 	/** Copy the array into a buffer of exactly GetSize() bytes, with no name. */
 	CORE_API void CopyTo(FMutableMemoryView Buffer) const;
@@ -909,77 +909,77 @@ public:
 	CORE_API void CopyTo(FArchive& Ar) const;
 
 	/** Invoke the visitor for every attachment in the array. */
-	inline void IterateAttachments(FCbFieldVisitor Visitor) const { CreateIterator().IterateRangeAttachments(Visitor); }
+	inline void IterateAttachments(FCbFieldVisitor Visitor) const { CreateViewIterator().IterateRangeAttachments(Visitor); }
 
 	/** Returns a view of the array, including the type and name when present. */
-	using FCbField::GetView;
+	using FCbFieldView::GetView;
 
 	/**
 	 * Try to get a view of the array as it would be serialized, such as by CopyTo.
 	 *
 	 * A serialized view is not available if the underlying field has an externally-provided type or
-	 * a name. Access the serialized form of such arrays using CopyTo or FCbArrayRef::Clone.
+	 * a name. Access the serialized form of such arrays using CopyTo or FCbArray::Clone.
 	 */
 	inline bool TryGetSerializedView(FMemoryView& OutView) const
 	{
-		return !FCbField::HasName() && FCbField::TryGetSerializedView(OutView);
+		return !FCbFieldView::HasName() && FCbFieldView::TryGetSerializedView(OutView);
 	}
 
 private:
-	friend inline FCbFieldIterator begin(const FCbArray& Array) { return Array.CreateIterator(); }
-	friend inline FCbFieldIterator end(const FCbArray&) { return FCbFieldIterator(); }
+	friend inline FCbFieldViewIterator begin(const FCbArrayView& Array) { return Array.CreateViewIterator(); }
+	friend inline FCbFieldViewIterator end(const FCbArrayView&) { return FCbFieldViewIterator(); }
 
 	/** Construct an array from an array field. No type check is performed! Use via FromFieldNoCheck. */
-	inline explicit FCbArray(const FCbField& Field) : FCbField(Field) {}
+	inline explicit FCbArrayView(const FCbFieldView& Field) : FCbFieldView(Field) {}
 };
 
 ///////////////////////////////////////////////////////////////////////////////////////////////////////////////////////
 
 /**
- * Array of FCbField that have unique names.
+ * Array of FCbField[View] that have unique names.
  *
  * Accessing the fields of an object is always a safe operation, even if the requested field does
  * not exist. Fields may be accessed by name or through iteration. When a field is requested that
  * is not found in the object, the field that it returns has no value (evaluates to false) though
- * attempting to access the empty field is also safe, as described by FCbField.
+ * attempting to access the empty field is also safe, as described by FCbFieldView.
  *
  * This type only provides a view into memory and does not perform any memory management itself.
- * Use FCbObjectRef to hold a reference to the underlying memory when necessary.
+ * Use FCbObject to hold a reference to the underlying memory when necessary.
  */
-class FCbObject : protected FCbField
+class FCbObjectView : protected FCbFieldView
 {
 public:
-	/** @see FCbField::FCbField */
-	using FCbField::FCbField;
+	/** @see FCbFieldView::FCbFieldView */
+	using FCbFieldView::FCbFieldView;
 
 	/** Construct an object with no fields. */
-	CORE_API FCbObject();
+	CORE_API FCbObjectView();
 
 	/** Create an iterator for the fields of this object. */
-	CORE_API FCbFieldIterator CreateIterator() const;
+	CORE_API FCbFieldViewIterator CreateViewIterator() const;
 
 	/**
 	 * Find a field by case-sensitive name comparison.
 	 *
-	 * The cost of this operation scales linearly with the number of fields in the object. Prefer
-	 * to iterate over the fields only once when consuming an object.
+	 * The cost of this operation scales linearly with the number of fields in the object. Prefer to
+	 * iterate over the fields only once when consuming an object.
 	 *
 	 * @param Name The name of the field.
 	 * @return The matching field if found, otherwise a field with no value.
 	 */
-	CORE_API FCbField Find(FAnsiStringView Name) const;
+	CORE_API FCbFieldView FindView(FAnsiStringView Name) const;
 
 	/** Find a field by case-insensitive name comparison. */
-	CORE_API FCbField FindIgnoreCase(FAnsiStringView Name) const;
+	CORE_API FCbFieldView FindViewIgnoreCase(FAnsiStringView Name) const;
 
 	/** Find a field by case-sensitive name comparison. */
-	inline FCbField operator[](FAnsiStringView Name) const { return Find(Name); }
+	inline FCbFieldView operator[](FAnsiStringView Name) const { return FindView(Name); }
 
 	/** Access the object as an object field. */
-	inline FCbField AsField() const { return static_cast<const FCbField&>(*this); }
+	inline FCbFieldView AsFieldView() const { return static_cast<const FCbFieldView&>(*this); }
 
 	/** Construct an object from an object field. No type check is performed! */
-	static inline FCbObject FromFieldNoCheck(const FCbField& Field) { return FCbObject(Field); }
+	static inline FCbObjectView FromFieldNoCheck(const FCbFieldView& Field) { return FCbObjectView(Field); }
 
 	/** Returns the size of the object in bytes if serialized by itself with no name. */
 	CORE_API uint64 GetSize() const;
@@ -998,7 +998,7 @@ public:
 	 * these assumptions do not hold, this may return false for equivalent inputs. Validation can
 	 * be done with the All mode to check these assumptions about the format of the inputs.
 	 */
-	CORE_API bool Equals(const FCbObject& Other) const;
+	CORE_API bool Equals(const FCbObjectView& Other) const;
 
 	/** Copy the object into a buffer of exactly GetSize() bytes, with no name. */
 	CORE_API void CopyTo(FMutableMemoryView Buffer) const;
@@ -1007,28 +1007,28 @@ public:
 	CORE_API void CopyTo(FArchive& Ar) const;
 
 	/** Invoke the visitor for every attachment in the object. */
-	inline void IterateAttachments(FCbFieldVisitor Visitor) const { CreateIterator().IterateRangeAttachments(Visitor); }
+	inline void IterateAttachments(FCbFieldVisitor Visitor) const { CreateViewIterator().IterateRangeAttachments(Visitor); }
 
 	/** Returns a view of the object, including the type and name when present. */
-	using FCbField::GetView;
+	using FCbFieldView::GetView;
 
 	/**
 	 * Try to get a view of the object as it would be serialized, such as by CopyTo.
 	 *
 	 * A serialized view is not available if the underlying field has an externally-provided type or
-	 * a name. Access the serialized form of such objects using CopyTo or FCbObjectRef::Clone.
+	 * a name. Access the serialized form of such objects using CopyTo or FCbObject::Clone.
 	 */
 	inline bool TryGetSerializedView(FMemoryView& OutView) const
 	{
-		return !FCbField::HasName() && FCbField::TryGetSerializedView(OutView);
+		return !FCbFieldView::HasName() && FCbFieldView::TryGetSerializedView(OutView);
 	}
 
 private:
-	friend inline FCbFieldIterator begin(const FCbObject& Object) { return Object.CreateIterator(); }
-	friend inline FCbFieldIterator end(const FCbObject&) { return FCbFieldIterator(); }
+	friend inline FCbFieldViewIterator begin(const FCbObjectView& Object) { return Object.CreateViewIterator(); }
+	friend inline FCbFieldViewIterator end(const FCbObjectView&) { return FCbFieldViewIterator(); }
 
 	/** Construct an object from an object field. No type check is performed! Use via FromFieldNoCheck. */
-	inline explicit FCbObject(const FCbField& Field) : FCbField(Field) {}
+	inline explicit FCbObjectView(const FCbFieldView& Field) : FCbFieldView(Field) {}
 };
 
 ///////////////////////////////////////////////////////////////////////////////////////////////////////////////////////
@@ -1039,12 +1039,12 @@ using FCbBufferAllocator = TFunctionRef<FUniqueBuffer (uint64 Size)>;
 ///////////////////////////////////////////////////////////////////////////////////////////////////////////////////////
 
 /** A wrapper that holds a reference to the buffer that contains its compact binary value. */
-template <typename BaseType>
-class TCbBufferRef : public BaseType
+template <typename ViewType>
+class TCbBuffer : public ViewType
 {
 public:
 	/** Construct a default value. */
-	TCbBufferRef() = default;
+	TCbBuffer() = default;
 
 	/**
 	 * Construct a value from a pointer to its data and an optional externally-provided type.
@@ -1052,36 +1052,36 @@ public:
 	 * @param ValueBuffer A buffer that exactly contains the value.
 	 * @param Type HasFieldType means that ValueBuffer contains the type. Otherwise, use the given type.
 	 */
-	inline explicit TCbBufferRef(FSharedBuffer ValueBuffer, ECbFieldType Type = ECbFieldType::HasFieldType)
+	inline explicit TCbBuffer(FSharedBuffer ValueBuffer, ECbFieldType Type = ECbFieldType::HasFieldType)
 	{
 		if (ValueBuffer)
 		{
-			BaseType::operator=(BaseType(ValueBuffer.GetData(), Type));
-			check(ValueBuffer.GetView().Contains(BaseType::GetView()));
+			ViewType::operator=(ViewType(ValueBuffer.GetData(), Type));
+			check(ValueBuffer.GetView().Contains(ViewType::GetView()));
 			Buffer = MoveTemp(ValueBuffer);
 		}
 	}
 
 	/** Construct a value that holds a reference to the buffer that contains it. */
-	inline TCbBufferRef(const BaseType& Value, FSharedBuffer OuterBuffer)
-		: BaseType(Value)
+	inline TCbBuffer(const ViewType& Value, FSharedBuffer OuterBuffer)
+		: ViewType(Value)
 	{
 		if (OuterBuffer)
 		{
-			check(OuterBuffer.GetView().Contains(BaseType::GetView()));
+			check(OuterBuffer.GetView().Contains(ViewType::GetView()));
 			Buffer = MoveTemp(OuterBuffer);
 		}
 	}
 
 	/** Construct a value that holds a reference to the buffer of the outer that contains it. */
-	template <typename OtherBaseType>
-	inline TCbBufferRef(const BaseType& Value, TCbBufferRef<OtherBaseType> OuterRef)
-		: TCbBufferRef(Value, MoveTemp(OuterRef.Buffer))
+	template <typename OtherViewType>
+	inline TCbBuffer(const ViewType& Value, TCbBuffer<OtherViewType> OuterBuffer)
+		: TCbBuffer(Value, MoveTemp(OuterBuffer.Buffer))
 	{
 	}
 
 	/** Reset this to a default value and null buffer. */
-	inline void Reset() { *this = TCbBufferRef(); }
+	inline void Reset() { *this = TCbBuffer(); }
 
 	/** Whether this reference has ownership of the memory in its buffer. */
 	inline bool IsOwned() const { return Buffer && Buffer.IsOwned(); }
@@ -1091,9 +1091,9 @@ public:
 	{
 		if (!IsOwned())
 		{
-			FUniqueBuffer MutableBuffer = FUniqueBuffer::Alloc(BaseType::GetSize());
-			BaseType::CopyTo(MutableBuffer);
-			BaseType::operator=(BaseType(MutableBuffer.GetData()));
+			FUniqueBuffer MutableBuffer = FUniqueBuffer::Alloc(ViewType::GetSize());
+			ViewType::CopyTo(MutableBuffer);
+			ViewType::operator=(ViewType(MutableBuffer.GetData()));
 			Buffer = MoveTemp(MutableBuffer);
 		}
 	}
@@ -1101,7 +1101,7 @@ public:
 	/** Returns a buffer that exactly contains this value. */
 	inline FSharedBuffer GetBuffer() const
 	{
-		const FMemoryView View = BaseType::GetView();
+		const FMemoryView View = ViewType::GetView();
 		const FSharedBuffer& OuterBuffer = GetOuterBuffer();
 		return View == OuterBuffer.GetView() ? OuterBuffer : FSharedBuffer::MakeView(View, OuterBuffer);
 	}
@@ -1112,7 +1112,7 @@ public:
 
 private:
 	template <typename OtherType>
-	friend class TCbBufferRef;
+	friend class TCbBuffer;
 
 	FSharedBuffer Buffer;
 };
@@ -1120,93 +1120,94 @@ private:
 ///////////////////////////////////////////////////////////////////////////////////////////////////////////////////////
 
 /**
- * Factory functions for types derived from TCbBufferRef.
+ * Factory functions for types derived from TCbBuffer.
  *
- * This uses the curiously recurring template pattern to construct the correct type of reference.
- * The derived type inherits from TCbBufferRef and this type to expose the factory functions.
+ * This uses the curiously recurring template pattern to construct the correct derived type, that
+ * must inherit from TCbBuffer and this type to expose the factory functions.
  */
-template <typename RefType, typename BaseType>
-class TCbBufferRefFactory
+template <typename Type, typename ViewType>
+class TCbBufferFactory
 {
 public:
 	/** Construct a value from an owned clone of its memory. */
-	static inline RefType Clone(const void* const Data)
+	static inline Type Clone(const void* const Data)
 	{
-		return Clone(BaseType(Data));
+		return Clone(ViewType(Data));
 	}
 
 	/** Construct a value from an owned clone of its memory. */
-	static inline RefType Clone(const BaseType& Value)
+	static inline Type Clone(const ViewType& Value)
 	{
-		RefType Ref = MakeView(Value);
-		Ref.MakeOwned();
-		return Ref;
+		Type Owned = MakeView(Value);
+		Owned.MakeOwned();
+		return Owned;
 	}
 
 	/** Construct a value from a read-only view of its memory and its optional outer buffer. */
-	static inline RefType MakeView(const void* const Data, FSharedBuffer OuterBuffer = FSharedBuffer())
+	static inline Type MakeView(const void* const Data, FSharedBuffer OuterBuffer = FSharedBuffer())
 	{
-		return MakeView(BaseType(Data), MoveTemp(OuterBuffer));
+		return MakeView(ViewType(Data), MoveTemp(OuterBuffer));
 	}
 
 	/** Construct a value from a read-only view of its memory and its optional outer buffer. */
-	static inline RefType MakeView(const BaseType& Value, FSharedBuffer OuterBuffer = FSharedBuffer())
+	static inline Type MakeView(const ViewType& Value, FSharedBuffer OuterBuffer = FSharedBuffer())
 	{
-		return RefType(Value, MoveTemp(OuterBuffer));
+		return Type(Value, MoveTemp(OuterBuffer));
 	}
 };
 
 ///////////////////////////////////////////////////////////////////////////////////////////////////////////////////////
 
-class FCbArrayRef;
-class FCbObjectRef;
+class FCbArray;
+class FCbObject;
 
 ///////////////////////////////////////////////////////////////////////////////////////////////////////////////////////
 
 /**
- * A field that can hold a reference to the memory that contains it.
+ * A field that includes a shared buffer for the memory that contains it.
  *
- * @see TCbBufferRef
+ * @see FCbFieldView
+ * @see TCbBuffer
  */
-class FCbFieldRef : public TCbBufferRef<FCbField>, public TCbBufferRefFactory<FCbFieldRef, FCbField>
+class FCbField : public TCbBuffer<FCbFieldView>, public TCbBufferFactory<FCbField, FCbFieldView>
 {
 public:
-	using TCbBufferRef::TCbBufferRef;
+	using TCbBuffer::TCbBuffer;
 
 	/** Access the field as an object. Defaults to an empty object on error. */
-	inline FCbObjectRef AsObjectRef() &;
+	inline FCbObject AsObject() &;
 
 	/** Access the field as an object. Defaults to an empty object on error. */
-	inline FCbObjectRef AsObjectRef() &&;
+	inline FCbObject AsObject() &&;
 
 	/** Access the field as an array. Defaults to an empty array on error. */
-	inline FCbArrayRef AsArrayRef() &;
+	inline FCbArray AsArray() &;
 
 	/** Access the field as an array. Defaults to an empty array on error. */
-	inline FCbArrayRef AsArrayRef() &&;
+	inline FCbArray AsArray() &&;
 };
 
 /**
- * Iterator for FCbFieldRef.
+ * Iterator for FCbField.
  *
  * @see TCbFieldIterator
  */
-class FCbFieldRefIterator : public TCbFieldIterator<FCbFieldRef>
+class FCbFieldIterator : public TCbFieldIterator<FCbField>
 {
 public:
 	/** Construct a field range from an owned clone of a range. */
-	CORE_API static FCbFieldRefIterator CloneRange(const FCbFieldIterator& It);
+	CORE_API static FCbFieldIterator CloneRange(const FCbFieldViewIterator& It);
 
 	/** Construct a field range from an owned clone of a range. */
-	static inline FCbFieldRefIterator CloneRange(const FCbFieldRefIterator& It)
+	static inline FCbFieldIterator CloneRange(const FCbFieldIterator& It)
 	{
-		return CloneRange(FCbFieldIterator(It));
+		return CloneRange(FCbFieldViewIterator(It));
 	}
 
 	/** Construct a field range that contains exactly one field. */
-	static inline FCbFieldRefIterator MakeSingle(FCbFieldRef Field)
+	static inline FCbFieldIterator MakeSingle(FCbField Field)
 	{
-		return FCbFieldRefIterator(MoveTemp(Field));
+		return FCbFieldIterator(MoveTemp(Field));
 	}
 
 	/**
@@ -1215,26 +1216,26 @@ public:
 	 * @param Buffer A buffer containing zero or more valid fields.
 	 * @param Type HasFieldType means that Buffer contains the type. Otherwise, use the given type.
 	 */
-	static inline FCbFieldRefIterator MakeRange(FSharedBuffer Buffer, ECbFieldType Type = ECbFieldType::HasFieldType)
+	static inline FCbFieldIterator MakeRange(FSharedBuffer Buffer, ECbFieldType Type = ECbFieldType::HasFieldType)
 	{
 		if (Buffer.GetSize())
 		{
 			const void* const DataEnd = Buffer.GetView().GetDataEnd();
-			return FCbFieldRefIterator(FCbFieldRef(MoveTemp(Buffer), Type), DataEnd);
+			return FCbFieldIterator(FCbField(MoveTemp(Buffer), Type), DataEnd);
 		}
-		return FCbFieldRefIterator();
+		return FCbFieldIterator();
 	}
 
 	/** Construct a field range from an iterator and its optional outer buffer. */
-	static inline FCbFieldRefIterator MakeRangeView(const FCbFieldIterator& It, FSharedBuffer OuterBuffer = FSharedBuffer())
+	static inline FCbFieldIterator MakeRangeView(const FCbFieldViewIterator& It, FSharedBuffer OuterBuffer = FSharedBuffer())
 	{
-		return FCbFieldRefIterator(FCbFieldRef(It, MoveTemp(OuterBuffer)), GetFieldsEnd(It));
+		return FCbFieldIterator(FCbField(It, MoveTemp(OuterBuffer)), GetFieldsEnd(It));
 	}
 
 	/** Construct an empty field range. */
-	constexpr FCbFieldRefIterator() = default;
+	constexpr FCbFieldIterator() = default;
 
-	/** Clone the range, if necessary, to a buffer that this reference has ownership of. */
+	/** Clone the range, if necessary, to a buffer that this has ownership of. */
 	inline void MakeRangeOwned()
 	{
 		if (!IsOwned())
@@ -1253,105 +1254,107 @@ private:
 ///////////////////////////////////////////////////////////////////////////////////////////////////////////////////////
 
 /**
- * An array that can hold a reference to the memory that contains it.
+ * An array that includes a shared buffer for the memory that contains it.
  *
- * @see TCbBufferRef
+ * @see FCbArrayView
+ * @see TCbBuffer
  */
-class FCbArrayRef : public TCbBufferRef<FCbArray>, public TCbBufferRefFactory<FCbArrayRef, FCbArray>
+class FCbArray : public TCbBuffer<FCbArrayView>, public TCbBufferFactory<FCbArray, FCbArrayView>
 {
 public:
-	using TCbBufferRef::TCbBufferRef;
+	using TCbBuffer::TCbBuffer;
 
 	/** Create an iterator for the fields of this array. */
-	inline FCbFieldRefIterator CreateRefIterator() const
+	inline FCbFieldIterator CreateIterator() const
 	{
-		return FCbFieldRefIterator::MakeRangeView(CreateIterator(), GetOuterBuffer());
+		return FCbFieldIterator::MakeRangeView(CreateViewIterator(), GetOuterBuffer());
 	}
 
 	/** Access the array as an array field. */
-	inline FCbFieldRef AsFieldRef() const &
+	inline FCbField AsField() const &
 	{
-		return FCbFieldRef(FCbArray::AsField(), *this);
+		return FCbField(FCbArrayView::AsFieldView(), *this);
 	}
 
 	/** Access the array as an array field. */
-	inline FCbFieldRef AsFieldRef() &&
+	inline FCbField AsField() &&
 	{
-		return FCbFieldRef(FCbArray::AsField(), MoveTemp(*this));
+		return FCbField(FCbArrayView::AsFieldView(), MoveTemp(*this));
 	}
 };
 
 ///////////////////////////////////////////////////////////////////////////////////////////////////////////////////////
 
 /**
- * An object that can hold a reference to the memory that contains it.
+ * An object that includes a shared buffer for the memory that contains it.
  *
- * @see TCbBufferRef
+ * @see FCbObjectView
+ * @see TCbBuffer
  */
-class FCbObjectRef : public TCbBufferRef<FCbObject>, public TCbBufferRefFactory<FCbObjectRef, FCbObject>
+class FCbObject : public TCbBuffer<FCbObjectView>, public TCbBufferFactory<FCbObject, FCbObjectView>
 {
 public:
-	using TCbBufferRef::TCbBufferRef;
+	using TCbBuffer::TCbBuffer;
 
 	/** Create an iterator for the fields of this object. */
-	inline FCbFieldRefIterator CreateRefIterator() const
+	inline FCbFieldIterator CreateIterator() const
 	{
-		return FCbFieldRefIterator::MakeRangeView(CreateIterator(), GetOuterBuffer());
+		return FCbFieldIterator::MakeRangeView(CreateViewIterator(), GetOuterBuffer());
 	}
 
 	/** Find a field by case-sensitive name comparison. */
-	inline FCbFieldRef FindRef(FAnsiStringView Name) const
+	inline FCbField Find(FAnsiStringView Name) const
 	{
-		if (::FCbField Field = Find(Name))
+		if (::FCbFieldView Field = FindView(Name))
 		{
-			return FCbFieldRef(Field, *this);
+			return FCbField(Field, *this);
 		}
-		return FCbFieldRef();
+		return FCbField();
 	}
 
 	/** Find a field by case-insensitive name comparison. */
-	inline FCbFieldRef FindRefIgnoreCase(FAnsiStringView Name) const
+	inline FCbField FindIgnoreCase(FAnsiStringView Name) const
 	{
-		if (::FCbField Field = FindIgnoreCase(Name))
+		if (::FCbFieldView Field = FindViewIgnoreCase(Name))
 		{
-			return FCbFieldRef(Field, *this);
+			return FCbField(Field, *this);
 		}
-		return FCbFieldRef();
+		return FCbField();
 	}
 
 	/** Access the object as an object field. */
-	inline FCbFieldRef AsFieldRef() const &
+	inline FCbField AsField() const &
 	{
-		return FCbFieldRef(FCbObject::AsField(), *this);
+		return FCbField(FCbObjectView::AsFieldView(), *this);
 	}
 
 	/** Access the object as an object field. */
-	inline FCbFieldRef AsFieldRef() &&
+	inline FCbField AsField() &&
 	{
-		return FCbFieldRef(FCbObject::AsField(), MoveTemp(*this));
+		return FCbField(FCbObjectView::AsFieldView(), MoveTemp(*this));
 	}
 };
 
 ///////////////////////////////////////////////////////////////////////////////////////////////////////////////////////
 
-inline FCbObjectRef FCbFieldRef::AsObjectRef() &
+inline FCbObject FCbField::AsObject() &
 {
-	return IsObject() ? FCbObjectRef(AsObject(), *this) : FCbObjectRef();
+	return IsObject() ? FCbObject(AsObjectView(), *this) : FCbObject();
 }
 
-inline FCbObjectRef FCbFieldRef::AsObjectRef() &&
+inline FCbObject FCbField::AsObject() &&
 {
-	return IsObject() ? FCbObjectRef(AsObject(), MoveTemp(*this)) : FCbObjectRef();
+	return IsObject() ? FCbObject(AsObjectView(), MoveTemp(*this)) : FCbObject();
 }
 
-inline FCbArrayRef FCbFieldRef::AsArrayRef() &
+inline FCbArray FCbField::AsArray() &
 {
-	return IsArray() ? FCbArrayRef(AsArray(), *this) : FCbArrayRef();
+	return IsArray() ? FCbArray(AsArrayView(), *this) : FCbArray();
 }
 
-inline FCbArrayRef FCbFieldRef::AsArrayRef() &&
+inline FCbArray FCbField::AsArray() &&
 {
-	return IsArray() ? FCbArrayRef(AsArray(), MoveTemp(*this)) : FCbArrayRef();
+	return IsArray() ? FCbArray(AsArrayView(), MoveTemp(*this)) : FCbArray();
 }
 
 ///////////////////////////////////////////////////////////////////////////////////////////////////////////////////////
diff --git a/Engine/Source/Runtime/Core/Public/Serialization/CompactBinaryPackage.h b/Engine/Source/Runtime/Core/Public/Serialization/CompactBinaryPackage.h
--- a/Engine/Source/Runtime/Core/Public/Serialization/CompactBinaryPackage.h
+++ b/Engine/Source/Runtime/Core/Public/Serialization/CompactBinaryPackage.h
@@ -30,13 +30,13 @@ public:
 	FCbAttachment() = default;
 
 	/** Construct a compact binary attachment. Value is cloned if not owned. */
-	inline explicit FCbAttachment(FCbFieldRefIterator Value)
+	inline explicit FCbAttachment(FCbFieldIterator Value)
 		: FCbAttachment(MoveTemp(Value), nullptr)
 	{
 	}
 
 	/** Construct a compact binary attachment. Value is cloned if not owned. Hash must match Value. */
-	inline explicit FCbAttachment(FCbFieldRefIterator Value, const FIoHash& Hash)
+	inline explicit FCbAttachment(FCbFieldIterator Value, const FIoHash& Hash)
 		: FCbAttachment(MoveTemp(Value), &Hash)
 	{
 	}
@@ -63,10 +63,10 @@ public:
 	inline bool IsNull() const { return !Buffer; }
 
 	/** Access the attachment as binary. Defaults to a null buffer on error. */
-	CORE_API FSharedBuffer AsBinary() const;
+	CORE_API FSharedBuffer AsBinaryView() const;
 
 	/** Access the attachment as compact binary. Defaults to a field iterator with no value on error. */
-	CORE_API FCbFieldRefIterator AsCompactBinary() const;
+	CORE_API FCbFieldIterator AsCompactBinary() const;
 
 	/** Returns whether the attachment is binary or compact binary. */
 	inline bool IsBinary() const { return !Buffer.IsNull(); }
@@ -89,7 +89,7 @@ public:
 	 *
 	 * The iterator is advanced as attachment fields are consumed from it.
 	 */
-	CORE_API void Load(FCbFieldRefIterator& Fields);
+	CORE_API void Load(FCbFieldIterator& Fields);
 
 	/**
 	 * Load the attachment from compact binary as written by Save.
@@ -108,13 +108,13 @@ public:
 	CORE_API void Save(FArchive& Ar) const;
 
 private:
-	CORE_API FCbAttachment(FCbFieldRefIterator Value, const FIoHash* Hash);
+	CORE_API FCbAttachment(FCbFieldIterator Value, const FIoHash* Hash);
 	CORE_API FCbAttachment(FSharedBuffer Value, const FIoHash* Hash);
 
 	/** An owned buffer containing the binary or compact binary data. */
 	FSharedBuffer Buffer;
 	/** A field iterator that is valid only for compact binary attachments. */
-	FCbFieldIterator CompactBinary;
+	FCbFieldViewIterator CompactBinary;
 	/** A hash of the attachment value. */
 	FIoHash Hash;
 };
@@ -168,7 +168,7 @@ public:
 	 *
 	 * @param InObject The root object, which will be cloned unless it is owned.
 	 */
-	inline explicit FCbPackage(FCbObjectRef InObject)
+	inline explicit FCbPackage(FCbObject InObject)
 	{
 		SetObject(MoveTemp(InObject));
 	}
@@ -179,7 +179,7 @@ public:
 	 * @param InObject The root object, which will be cloned unless it is owned.
 	 * @param InResolver A function that is invoked for every reference and binary reference field.
 	 */
-	inline explicit FCbPackage(FCbObjectRef InObject, FAttachmentResolver InResolver)
+	inline explicit FCbPackage(FCbObject InObject, FAttachmentResolver InResolver)
 	{
 		SetObject(MoveTemp(InObject), InResolver);
 	}
@@ -190,7 +190,7 @@ public:
 	 * @param InObject The root object, which will be cloned unless it is owned.
 	 * @param InObjectHash The hash of the object, which must match to avoid validation errors.
 	 */
-	inline explicit FCbPackage(FCbObjectRef InObject, const FIoHash& InObjectHash)
+	inline explicit FCbPackage(FCbObject InObject, const FIoHash& InObjectHash)
 	{
 		SetObject(MoveTemp(InObject), InObjectHash);
 	}
@@ -202,7 +202,7 @@ public:
 	 * @param InObjectHash The hash of the object, which must match to avoid validation errors.
 	 * @param InResolver A function that is invoked for every reference and binary reference field.
 	 */
-	inline explicit FCbPackage(FCbObjectRef InObject, const FIoHash& InObjectHash, FAttachmentResolver InResolver)
+	inline explicit FCbPackage(FCbObject InObject, const FIoHash& InObjectHash, FAttachmentResolver InResolver)
 	{
 		SetObject(MoveTemp(InObject), InObjectHash, InResolver);
 	}
@@ -216,11 +216,11 @@ public:
 	/** Whether the package has an empty object and no attachments. */
 	inline bool IsNull() const
 	{
-		return !Object.CreateIterator() && Attachments.Num() == 0;
+		return !Object.CreateViewIterator() && Attachments.Num() == 0;
 	}
 
 	/** Returns the compact binary object for the package. */
-	inline const FCbObjectRef& GetObject() const { return Object; }
+	inline const FCbObject& GetObject() const { return Object; }
 
 	/** Returns the has of the compact binary object for the package. */
 	inline const FIoHash& GetObjectHash() const { return ObjectHash; }
@@ -230,7 +230,7 @@ public:
 	 *
 	 * @param InObject The root object, which will be cloned unless it is owned.
 	 */
-	inline void SetObject(FCbObjectRef InObject)
+	inline void SetObject(FCbObject InObject)
 	{
 		SetObject(MoveTemp(InObject), nullptr, nullptr);
 	}
@@ -241,7 +241,7 @@ public:
 	 * @param InObject The root object, which will be cloned unless it is owned.
 	 * @param InResolver A function that is invoked for every reference and binary reference field.
 	 */
-	inline void SetObject(FCbObjectRef InObject, FAttachmentResolver InResolver)
+	inline void SetObject(FCbObject InObject, FAttachmentResolver InResolver)
 	{
 		SetObject(MoveTemp(InObject), nullptr, &InResolver);
 	}
@@ -252,7 +252,7 @@ public:
 	 * @param InObject The root object, which will be cloned unless it is owned.
 	 * @param InObjectHash The hash of the object, which must match to avoid validation errors.
 	 */
-	inline void SetObject(FCbObjectRef InObject, const FIoHash& InObjectHash)
+	inline void SetObject(FCbObject InObject, const FIoHash& InObjectHash)
 	{
 		SetObject(MoveTemp(InObject), &InObjectHash, nullptr);
 	}
@@ -264,7 +264,7 @@ public:
 	 * @param InObjectHash The hash of the object, which must match to avoid validation errors.
 	 * @param InResolver A function that is invoked for every reference and binary reference field.
 	 */
-	inline void SetObject(FCbObjectRef InObject, const FIoHash& InObjectHash, FAttachmentResolver InResolver)
+	inline void SetObject(FCbObject InObject, const FIoHash& InObjectHash, FAttachmentResolver InResolver)
 	{
 		SetObject(MoveTemp(InObject), &InObjectHash, &InResolver);
 	}
@@ -319,7 +319,7 @@ public:
 	 *
 	 * The iterator is advanced as object and attachment fields are consumed from it.
 	 */
-	CORE_API void Load(FCbFieldRefIterator& Fields);
+	CORE_API void Load(FCbFieldIterator& Fields);
 
 	/**
 	 * Load the object and attachments from compact binary as written by Save.
@@ -338,14 +338,14 @@ public:
 	CORE_API void Save(FArchive& Ar) const;
 
 private:
-	CORE_API void SetObject(FCbObjectRef Object, const FIoHash* Hash, FAttachmentResolver* Resolver);
+	CORE_API void SetObject(FCbObject Object, const FIoHash* Hash, FAttachmentResolver* Resolver);
 	CORE_API void AddAttachment(const FCbAttachment& Attachment, FAttachmentResolver* Resolver);
 
-	void GatherAttachments(const FCbFieldIterator& Fields, FAttachmentResolver Resolver);
+	void GatherAttachments(const FCbFieldViewIterator& Fields, FAttachmentResolver Resolver);
 
 	/** Attachments ordered by their hash. */
 	TArray<FCbAttachment> Attachments;
-	FCbObjectRef Object;
+	FCbObject Object;
 	FIoHash ObjectHash;
 };
 
diff --git a/Engine/Source/Runtime/Core/Public/Serialization/CompactBinarySerialization.h b/Engine/Source/Runtime/Core/Public/Serialization/CompactBinarySerialization.h
--- a/Engine/Source/Runtime/Core/Public/Serialization/CompactBinarySerialization.h
+++ b/Engine/Source/Runtime/Core/Public/Serialization/CompactBinarySerialization.h
@@ -54,24 +54,24 @@ CORE_API bool TryMeasureCompactBinary(
 /**
  * Load a compact binary field from an archive.
  *
- * The field may be an array or an object, which the caller can convert to by using AsArrayRef or
- * AsObjectRef as appropriate. The buffer allocator is called to provide the buffer for the field
+ * The field may be an array or an object, which the caller can convert to by using AsArray or
+ * AsObject as appropriate. The buffer allocator is called to provide the buffer for the field
  * to load into once its size has been determined.
  *
  * @param Ar Archive to read the field from.
  * @param Allocator Allocator for the buffer that the field is loaded into.
  * @return A field with a reference to the provided buffer if it is owned.
  */
-CORE_API FCbFieldRef LoadCompactBinary(FArchive& Ar, FCbBufferAllocator Allocator = FUniqueBuffer::Alloc);
+CORE_API FCbField LoadCompactBinary(FArchive& Ar, FCbBufferAllocator Allocator = FUniqueBuffer::Alloc);
 
 /** Save a compact binary value to an archive. */
-CORE_API void SaveCompactBinary(FArchive& Ar, const FCbField& Field);
-CORE_API void SaveCompactBinary(FArchive& Ar, const FCbArray& Array);
-CORE_API void SaveCompactBinary(FArchive& Ar, const FCbObject& Object);
+CORE_API void SaveCompactBinary(FArchive& Ar, const FCbFieldView& Field);
+CORE_API void SaveCompactBinary(FArchive& Ar, const FCbArrayView& Array);
+CORE_API void SaveCompactBinary(FArchive& Ar, const FCbObjectView& Object);
 
 /** Serialize a compact binary value to/from an archive. */
-CORE_API FArchive& operator<<(FArchive& Ar, FCbFieldRef& Field);
-CORE_API FArchive& operator<<(FArchive& Ar, FCbArrayRef& Array);
-CORE_API FArchive& operator<<(FArchive& Ar, FCbObjectRef& Object);
+CORE_API FArchive& operator<<(FArchive& Ar, FCbField& Field);
+CORE_API FArchive& operator<<(FArchive& Ar, FCbArray& Array);
+CORE_API FArchive& operator<<(FArchive& Ar, FCbObject& Object);
 
 ///////////////////////////////////////////////////////////////////////////////////////////////////////////////////////
diff --git a/Engine/Source/Runtime/Core/Public/Serialization/CompactBinaryWriter.h b/Engine/Source/Runtime/Core/Public/Serialization/CompactBinaryWriter.h
--- a/Engine/Source/Runtime/Core/Public/Serialization/CompactBinaryWriter.h
+++ b/Engine/Source/Runtime/Core/Public/Serialization/CompactBinaryWriter.h
@@ -21,7 +21,7 @@ struct FTimespan;
  *
  * The writer produces a sequence of fields that can be saved to a provided memory buffer or into
  * a new owned buffer. The typical use case is to write a single object, which can be accessed by
- * calling Save().AsObjectRef() or Save(Buffer).AsObject().
+ * calling Save().AsObject() or Save(Buffer).AsObjectView().
  *
  * The writer will assert on most incorrect usage and will always produce valid compact binary if
  * provided with valid input. The writer does not check for invalid UTF-8 string encoding, object
@@ -36,7 +36,7 @@ struct FTimespan;
  *
  * Example:
  *
- * FCbObjectRef WriteObject()
+ * FCbObject WriteObject()
  * {
  *     TCbWriter<256> Writer;
  *     Writer.BeginObject();
@@ -50,7 +50,7 @@ struct FTimespan;
  *     Writer.EndArray();
  * 
  *     Writer.EndObject();
- *     return Writer.Save().AsObjectRef();
+ *     return Writer.Save().AsObject();
  * }
  */
 class FCbWriter
@@ -71,7 +71,7 @@ public:
 	 * It is not valid to call this function in the middle of writing an object, array, or field.
 	 * The writer remains valid for further use when this function returns.
 	 */
-	CORE_API FCbFieldRefIterator Save() const;
+	CORE_API FCbFieldIterator Save() const;
 
 	/**
 	 * Serialize the field(s) to memory.
@@ -82,7 +82,7 @@ public:
 	 * @param Buffer A mutable memory view to write to. Must be exactly GetSaveSize() bytes.
 	 * @return An iterator for the field(s) written to the buffer.
 	 */
-	CORE_API FCbFieldIterator Save(FMutableMemoryView Buffer) const;
+	CORE_API FCbFieldViewIterator Save(FMutableMemoryView Buffer) const;
 
 	/**
 	 * Serialize the field(s) to an archive.
@@ -110,11 +110,11 @@ public:
 	CORE_API FCbWriter& SetName(FAnsiStringView Name);
 
 	/** Copy the value (not the name) of an existing field. */
+	inline void AddField(FAnsiStringView Name, const FCbFieldView& Value) { SetName(Name); AddField(Value); }
+	CORE_API void AddField(const FCbFieldView& Value);
+	/** Copy the value (not the name) of an existing field. Holds a reference if owned. */
 	inline void AddField(FAnsiStringView Name, const FCbField& Value) { SetName(Name); AddField(Value); }
 	CORE_API void AddField(const FCbField& Value);
-	/** Copy the value (not the name) of an existing field. Holds a reference if owned. */
-	inline void AddField(FAnsiStringView Name, const FCbFieldRef& Value) { SetName(Name); AddField(Value); }
-	CORE_API void AddField(const FCbFieldRef& Value);
 
 	/** Begin a new object. Must have a matching call to EndObject. */
 	inline void BeginObject(FAnsiStringView Name) { SetName(Name); BeginObject(); }
@@ -123,11 +123,11 @@ public:
 	CORE_API void EndObject();
 
 	/** Copy the value (not the name) of an existing object. */
+	inline void AddObject(FAnsiStringView Name, const FCbObjectView& Value) { SetName(Name); AddObject(Value); }
+	CORE_API void AddObject(const FCbObjectView& Value);
+	/** Copy the value (not the name) of an existing object. Holds a reference if owned. */
 	inline void AddObject(FAnsiStringView Name, const FCbObject& Value) { SetName(Name); AddObject(Value); }
 	CORE_API void AddObject(const FCbObject& Value);
-	/** Copy the value (not the name) of an existing object. Holds a reference if owned. */
-	inline void AddObject(FAnsiStringView Name, const FCbObjectRef& Value) { SetName(Name); AddObject(Value); }
-	CORE_API void AddObject(const FCbObjectRef& Value);
 
 	/** Begin a new array. Must have a matching call to EndArray. */
 	inline void BeginArray(FAnsiStringView Name) { SetName(Name); BeginArray(); }
@@ -136,11 +136,11 @@ public:
 	CORE_API void EndArray();
 
 	/** Copy the value (not the name) of an existing array. */
+	inline void AddArray(FAnsiStringView Name, const FCbArrayView& Value) { SetName(Name); AddArray(Value); }
+	CORE_API void AddArray(const FCbArrayView& Value);
+	/** Copy the value (not the name) of an existing array. Holds a reference if owned. */
 	inline void AddArray(FAnsiStringView Name, const FCbArray& Value) { SetName(Name); AddArray(Value); }
 	CORE_API void AddArray(const FCbArray& Value);
-	/** Copy the value (not the name) of an existing array. Holds a reference if owned. */
-	inline void AddArray(FAnsiStringView Name, const FCbArrayRef& Value) { SetName(Name); AddArray(Value); }
-	CORE_API void AddArray(const FCbArrayRef& Value);
 
 	/** Write a null field. */
 	inline void AddNull(FAnsiStringView Name) { SetName(Name); AddNull(); }
@@ -321,37 +321,37 @@ inline FCbWriter& operator<<(FCbWriter& Writer, const ANSICHAR* NameOrValue)
 	return Writer << FAnsiStringView(NameOrValue);
 }
 
-inline FCbWriter& operator<<(FCbWriter& Writer, const FCbField& Value)
+inline FCbWriter& operator<<(FCbWriter& Writer, const FCbFieldView& Value)
 {
 	Writer.AddField(Value);
 	return Writer;
 }
 
-inline FCbWriter& operator<<(FCbWriter& Writer, const FCbFieldRef& Value)
+inline FCbWriter& operator<<(FCbWriter& Writer, const FCbField& Value)
 {
 	Writer.AddField(Value);
 	return Writer;
 }
 
-inline FCbWriter& operator<<(FCbWriter& Writer, const FCbObject& Value)
+inline FCbWriter& operator<<(FCbWriter& Writer, const FCbObjectView& Value)
 {
 	Writer.AddObject(Value);
 	return Writer;
 }
 
-inline FCbWriter& operator<<(FCbWriter& Writer, const FCbObjectRef& Value)
+inline FCbWriter& operator<<(FCbWriter& Writer, const FCbObject& Value)
 {
 	Writer.AddObject(Value);
 	return Writer;
 }
 
-inline FCbWriter& operator<<(FCbWriter& Writer, const FCbArray& Value)
+inline FCbWriter& operator<<(FCbWriter& Writer, const FCbArrayView& Value)
 {
 	Writer.AddArray(Value);
 	return Writer;
 }
 
-inline FCbWriter& operator<<(FCbWriter& Writer, const FCbArrayRef& Value)
+inline FCbWriter& operator<<(FCbWriter& Writer, const FCbArray& Value)
 {
 	Writer.AddArray(Value);
 	return Writer;
