Fixed UnixCriticalSection not timing out properly.

diff --git a/Engine/Source/Runtime/Core/Private/Unix/UnixCriticalSection.cpp b/Engine/Source/Runtime/Core/Private/Unix/UnixCriticalSection.cpp
--- a/Engine/Source/Runtime/Core/Private/Unix/UnixCriticalSection.cpp
+++ b/Engine/Source/Runtime/Core/Private/Unix/UnixCriticalSection.cpp
@@ -3,6 +3,7 @@
 #include "Unix/UnixCriticalSection.h"
 #include "Misc/DateTime.h"
 #include "HAL/PlatformProcess.h"
+#include "HAL/PlatformTime.h"
 #include "Containers/StringConv.h"
 #include <spawn.h>
 #include <sys/wait.h>
@@ -22,7 +23,7 @@ FUnixSystemWideCriticalSection::FUnixSystemWideCriticalSection(const FString& In
 	NormalizedFilepath.ReplaceInline(TEXT("\\"), TEXT("/"));
 	FileHandle = -1;
 
-	FDateTime ExpireTime = FDateTime::UtcNow() + InTimeout;
+	double ExpireTimeSecs = FPlatformTime::Seconds() + InTimeout.GetTotalSeconds();
 	while (true)
 	{
 		if (FileHandle == -1)
@@ -32,13 +33,14 @@ FUnixSystemWideCriticalSection::FUnixSystemWideCriticalSection(const FString& In
 		}
 
 		// If the file is open, try to lock it, but don't block. If the file is already locked by another process or another thread, blocking here may not honor InTimeout.
+		// NOTE: In old Linux kernels, flock() wasn't always atomic, but in recent ones, that was fixed and flock() is expected to be an atomic operation.
 		if (FileHandle != -1 && flock(FileHandle, LOCK_EX | LOCK_NB) == 0)
 		{
 			return; // Lock was successfully taken.
 		}
 
 		// If the lock isn't acquired and no time is left to retry, clean up and set the state as 'invalid'
-		if (InTimeout != FTimespan::Zero() || FDateTime::UtcNow() < ExpireTime)
+		if (InTimeout == FTimespan::Zero() || FPlatformTime::Seconds() > ExpireTimeSecs)
 		{
 			if (FileHandle != -1)
 			{
