UE5EA - Disable GPU culling on Mac as we're out of vertex attributes

diff --git a/Engine/Shaders/Public/Platform/Metal/MetalCommon.ush b/Engine/Shaders/Public/Platform/Metal/MetalCommon.ush
--- a/Engine/Shaders/Public/Platform/Metal/MetalCommon.ush
+++ b/Engine/Shaders/Public/Platform/Metal/MetalCommon.ush
@@ -5,7 +5,7 @@
 =============================================================================*/
 
 // Update this GUID to improve shader recompilation for Metal only shaders
-// GUID = 73A7D47A-D254-474A-8D71-3AB6B4BB6717
+// GUID = 0475930F-29E9-4E32-98DF-B2D12F73C042
 
 #pragma once
 
diff --git a/Engine/Source/Runtime/Core/Core.Build.cs b/Engine/Source/Runtime/Core/Core.Build.cs
--- a/Engine/Source/Runtime/Core/Core.Build.cs
+++ b/Engine/Source/Runtime/Core/Core.Build.cs
@@ -266,7 +266,14 @@ public class Core : ModuleRules
 		PrivateDefinitions.Add("PLATFORM_SUPPORTS_BINARYCONFIG=" + (SupportsBinaryConfig(Target) ? "1" : "0"));
 
 		// temporary thing to enable backing out in case of disaster, remove after initial testing period.
-		PublicDefinitions.Add("GPUCULL_TODO=1");
+		if (Target.Platform == UnrealTargetPlatform.Mac)
+		{
+			PublicDefinitions.Add("GPUCULL_TODO=0");
+        }
+		else
+        {
+			PublicDefinitions.Add("GPUCULL_TODO=1");
+		}
 
 		PublicDefinitions.Add("WITH_MALLOC_STOMP=" + (bWithMallocStomp ? "1" : "0"));
 
