IoStore: Fixed incorrect size verification of container files with memory mapped chunks
#rnx

diff --git a/Engine/Source/Runtime/Core/Private/IO/IoStore.cpp b/Engine/Source/Runtime/Core/Private/IO/IoStore.cpp
--- a/Engine/Source/Runtime/Core/Private/IO/IoStore.cpp
+++ b/Engine/Source/Runtime/Core/Private/IO/IoStore.cpp
@@ -712,6 +712,7 @@ public:
 					Partition.ContainerFileHandle->Serialize(Padding.GetData(), ExtraPaddingBytes);
 					CompressedContainerSize += ExtraPaddingBytes;
 					UncompressedContainerSize += ExtraPaddingBytes;
+					Partition.Offset += ExtraPaddingBytes;
 				}
 			}
 			
