CompactBinary: Added FSharedBuffer FCbField::AsBinary()

#rnx

diff --git a/Engine/Source/Runtime/Core/Private/Tests/Serialization/CompactBinaryTest.cpp b/Engine/Source/Runtime/Core/Private/Tests/Serialization/CompactBinaryTest.cpp
--- a/Engine/Source/Runtime/Core/Private/Tests/Serialization/CompactBinaryTest.cpp
+++ b/Engine/Source/Runtime/Core/Private/Tests/Serialization/CompactBinaryTest.cpp
@@ -651,21 +651,37 @@ bool FCbFieldArrayTest::RunTest(const FString& Parameters)
 IMPLEMENT_CUSTOM_SIMPLE_AUTOMATION_TEST(FCbFieldBinaryTest, FCbFieldTestBase, "System.Core.Serialization.CbField.Binary", CompactBinaryTestFlags)
 bool FCbFieldBinaryTest::RunTest(const FString& Parameters)
 {
+	struct FCbBinaryAccessors
+	{
+		TUniqueFunction<bool (const FCbFieldView&)> IsType = [](const FCbFieldView& Field) { return Field.IsInteger(); };
+		TUniqueFunction<FSharedBuffer (FCbFieldView& Field, const FSharedBuffer& Default)> AsType =
+			[](FCbFieldView& Field, const FSharedBuffer& Default) -> FSharedBuffer { return static_cast<FCbField&>(Field).AsBinary(Default); };
+	};
+
 	// Test FCbFieldView(Binary, Empty)
 	TestField<ECbFieldType::Binary>(TEXT("Binary, Empty"), {0});
 
 	// Test FCbFieldView(Binary, Value)
 	{
 		const uint8 Payload[] = { 3, 4, 5, 6 }; // Size: 3, Data: 4/5/6
-		FCbFieldView Field(Payload, ECbFieldType::Binary);
-		TestFieldNoClone<ECbFieldType::Binary>(TEXT("Binary, Value"), Field, MakeMemoryView(Payload + 1, 3));
+		FCbFieldView FieldView(Payload, ECbFieldType::Binary);
+		TestFieldNoClone<ECbFieldType::Binary>(TEXT("Binary, Value, View"), FieldView, MakeMemoryView(Payload + 1, 3));
+
+		FCbField Field = FCbField::Clone(FieldView);
+		Field.AsBinary();
+		TestFalse(TEXT("Binary, Value, AsBinary -> GetOuterBuffer()"), Field.GetOuterBuffer().IsNull());
+		MoveTemp(Field).AsBinary();
+		TestTrue(TEXT("Binary, Value, AsBinary -> GetOuterBuffer()"), Field.GetOuterBuffer().IsNull());
 	}
 
 	// Test FCbFieldView(None) as Binary
 	{
-		FCbFieldView Field;
+		FCbFieldView FieldView;
 		const uint8 Default[] = { 1, 2, 3 };
-		TestFieldError<ECbFieldType::Binary>(TEXT("Binary, None"), Field, ECbFieldError::TypeError, MakeMemoryView(Default));
+		TestFieldError<ECbFieldType::Binary>(TEXT("Binary, None, View"), FieldView, ECbFieldError::TypeError, MakeMemoryView(Default));
+
+		FCbField Field = FCbField::Clone(FieldView);
+		TestFieldError<ECbFieldType::Binary, FSharedBuffer>(TEXT("Binary, None"), Field, ECbFieldError::TypeError, FSharedBuffer::MakeView(MakeMemoryView(Default)), FCbBinaryAccessors());
 	}
 
 	return true;
diff --git a/Engine/Source/Runtime/Core/Public/Serialization/CompactBinary.h b/Engine/Source/Runtime/Core/Public/Serialization/CompactBinary.h
--- a/Engine/Source/Runtime/Core/Public/Serialization/CompactBinary.h
+++ b/Engine/Source/Runtime/Core/Public/Serialization/CompactBinary.h
@@ -1176,15 +1176,15 @@ public:
 
 	/** Access the field as an object. Defaults to an empty object on error. */
 	inline FCbObject AsObject() &;
-
-	/** Access the field as an object. Defaults to an empty object on error. */
 	inline FCbObject AsObject() &&;
 
 	/** Access the field as an array. Defaults to an empty array on error. */
 	inline FCbArray AsArray() &;
-
-	/** Access the field as an array. Defaults to an empty array on error. */
 	inline FCbArray AsArray() &&;
+
+	/** Access the field as binary. Returns the provided default on error. */
+	inline FSharedBuffer AsBinary(const FSharedBuffer& Default = FSharedBuffer()) &;
+	inline FSharedBuffer AsBinary(const FSharedBuffer& Default = FSharedBuffer()) &&;
 };
 
 /**
@@ -1357,4 +1357,16 @@ inline FCbArray FCbField::AsArray() &&
 	return IsArray() ? FCbArray(AsArrayView(), MoveTemp(*this)) : FCbArray();
 }
 
+inline FSharedBuffer FCbField::AsBinary(const FSharedBuffer& Default) &
+{
+	const FMemoryView View = AsBinaryView();
+	return !HasError() ? FSharedBuffer::MakeView(View, GetOuterBuffer()) : Default;
+}
+
+inline FSharedBuffer FCbField::AsBinary(const FSharedBuffer& Default) &&
+{
+	const FMemoryView View = AsBinaryView();
+	return !HasError() ? FSharedBuffer::MakeView(View, MoveTemp(*this).GetOuterBuffer()) : Default;
+}
+
 ///////////////////////////////////////////////////////////////////////////////////////////////////////////////////////
