Fixed non-unity build

diff --git a/Engine/Source/Runtime/Core/Public/PixelFormat.h b/Engine/Source/Runtime/Core/Public/PixelFormat.h
--- a/Engine/Source/Runtime/Core/Public/PixelFormat.h
+++ b/Engine/Source/Runtime/Core/Public/PixelFormat.h
@@ -4,6 +4,7 @@
 #pragma once
 
 #include "CoreTypes.h"
+#include "Misc/EnumClassFlags.h"
 
 #if defined(PF_MAX)
 #undef PF_MAX
