Added an FChar::IsControl function to check whether a codepoint represents a control character (0x00-0x1F & 0x7F in ASCII) / (0x00-0x1F & 0x7F-0x9F in Unicode + whatever elese might get added).

diff --git a/Engine/Source/Runtime/Core/Public/Misc/Char.h b/Engine/Source/Runtime/Core/Public/Misc/Char.h
--- a/Engine/Source/Runtime/Core/Public/Misc/Char.h
+++ b/Engine/Source/Runtime/Core/Public/Misc/Char.h
@@ -99,6 +99,7 @@ struct TChar : TCharBase<CharType, sizeof(CharType)>
 	static bool IsDigit(CharType Char);
 	static bool IsHexDigit(CharType Char);
 	static bool IsWhitespace(CharType Char);
+	static bool IsControl(CharType Char);
 
 	static bool IsOctDigit(CharType Char)
 	{
@@ -152,6 +153,8 @@ template <> inline bool TChar<WIDECHAR>::IsAlnum(WIDECHAR Char) { return ::iswal
 template <> inline bool TChar<WIDECHAR>::IsDigit(WIDECHAR Char) { return ::iswdigit(Char) != 0; }
 template <> inline bool TChar<WIDECHAR>::IsHexDigit(WIDECHAR Char) { return ::iswxdigit(Char) != 0; }
 template <> inline bool TChar<WIDECHAR>::IsWhitespace(WIDECHAR Char) { return ::iswspace(Char) != 0; }
+template <> inline bool TChar<WIDECHAR>::IsControl(WIDECHAR Char) { return ::iswcntrl(Char) != 0; }
+
 
 /*-----------------------------------------------------------------------------
 	ANSICHAR specialized functions
@@ -166,3 +169,4 @@ template <> inline bool TChar<ANSICHAR>::IsAlnum(ANSICHAR Char) { return ::isaln
 template <> inline bool TChar<ANSICHAR>::IsDigit(ANSICHAR Char) { return ::isdigit((unsigned char)Char) != 0; }
 template <> inline bool TChar<ANSICHAR>::IsHexDigit(ANSICHAR Char) { return ::isxdigit((unsigned char)Char) != 0; }
 template <> inline bool TChar<ANSICHAR>::IsWhitespace(ANSICHAR Char) { return ::isspace((unsigned char)Char) != 0; }
+template <> inline bool TChar<ANSICHAR>::IsControl(ANSICHAR Char) { return ::iscntrl((unsigned char)Char) != 0; }
