File/INI/CVar access detection intended to help DDC2 isolate itself from local environment

diff --git a/Engine/Source/Runtime/Core/Private/HAL/FileManagerGeneric.cpp b/Engine/Source/Runtime/Core/Private/HAL/FileManagerGeneric.cpp
--- a/Engine/Source/Runtime/Core/Private/HAL/FileManagerGeneric.cpp
+++ b/Engine/Source/Runtime/Core/Private/HAL/FileManagerGeneric.cpp
@@ -12,6 +12,7 @@
 // for compression
 #include "HAL/FileManagerGeneric.h"
 #include "Logging/LogMacros.h"
+#include "Misc/AccessDetection.h"
 #include "Misc/Parse.h"
 #include "Misc/CommandLine.h"
 #include "Misc/Paths.h"
@@ -1032,6 +1033,7 @@ void FArchiveFileWriterGeneric::LogWriteError(const TCHAR* Message)
 
 IFileManager& IFileManager::Get()
 {
+	UE::AccessDetection::ReportAccess(UE::AccessDetection::EType::File);
 	static FFileManagerGeneric Singleton;
 	return Singleton;
 }
diff --git a/Engine/Source/Runtime/Core/Private/HAL/PlatformFileManager.cpp b/Engine/Source/Runtime/Core/Private/HAL/PlatformFileManager.cpp
--- a/Engine/Source/Runtime/Core/Private/HAL/PlatformFileManager.cpp
+++ b/Engine/Source/Runtime/Core/Private/HAL/PlatformFileManager.cpp
@@ -5,6 +5,7 @@
 =============================================================================*/
 
 #include "HAL/PlatformFileManager.h"
+#include "Misc/AccessDetection.h"
 #include "Misc/AssertionMacros.h"
 #include "GenericPlatform/GenericPlatformFile.h"
 #include "Modules/ModuleManager.h"
@@ -134,6 +135,7 @@ void FPlatformFileManager::InitializeNewAsyncIO()
 
 FPlatformFileManager& FPlatformFileManager::Get()
 {
+	UE::AccessDetection::ReportAccess(UE::AccessDetection::EType::File);
 	static FPlatformFileManager Singleton;
 	return Singleton;
 }
diff --git a/Engine/Source/Runtime/Core/Private/Misc/AccessDetection.cpp b/Engine/Source/Runtime/Core/Private/Misc/AccessDetection.cpp
new file mode 100644
--- /dev/null
+++ b/Engine/Source/Runtime/Core/Private/Misc/AccessDetection.cpp
@@ -0,0 +1,29 @@
+// Copyright Epic Games, Inc. All Rights Reserved.
+
+#include "Misc/AccessDetection.h"
+#include "Misc/AssertionMacros.h"
+#include "Templates/Atomic.h"
+
+#if WITH_EDITOR
+namespace UE { namespace AccessDetection { namespace Private {
+
+volatile int32 GNumLiveScopes;
+thread_local FScope* TScope;
+
+void SetCurrentThreadScope(FScope* Scope)
+{
+	checkf(!!Scope != !TScope, TEXT("Nested access detection scopes are not allowed"));
+	FPlatformAtomics::InterlockedAdd(&GNumLiveScopes, Scope ? 1 : -1);
+	TScope = Scope;
+}
+
+void ReportCurrentThreadAccess(EType Type)
+{
+	if (FScope* Scope = TScope)
+	{
+		Scope->AddAccess(Type);	
+	}
+}
+
+}}}
+#endif
\ No newline at end of file
diff --git a/Engine/Source/Runtime/Core/Public/HAL/IConsoleManager.h b/Engine/Source/Runtime/Core/Public/HAL/IConsoleManager.h
--- a/Engine/Source/Runtime/Core/Public/HAL/IConsoleManager.h
+++ b/Engine/Source/Runtime/Core/Public/HAL/IConsoleManager.h
@@ -3,6 +3,7 @@
 #pragma once
 
 #include "CoreTypes.h"
+#include "Misc/AccessDetection.h"
 #include "Misc/AssertionMacros.h"
 #include "Templates/UnrealTemplate.h"
 #include "Containers/UnrealString.h"
@@ -1302,6 +1303,7 @@ public:
 	// faster than GetValueOnAnyThread()
 	T GetValueOnGameThread() const
 	{
+		UE::AccessDetection::ReportAccess(UE::AccessDetection::EType::CVar);
 		// compiled out in shipping for performance (we can change in development later), if this get triggered you need to call GetValueOnRenderThread() or GetValueOnAnyThread(), the last one is a bit slower
 		cvarCheckCode(ensure(GetShadowIndex() == 0));	// ensure to not block content creators, #if to optimize in shipping
 		return ShadowedValue[0];
@@ -1310,6 +1312,7 @@ public:
 	// faster than GetValueOnAnyThread()
 	T GetValueOnRenderThread() const
 	{
+		UE::AccessDetection::ReportAccess(UE::AccessDetection::EType::CVar);
 #if !defined(__clang__) // @todo Mac: figure out how to make this compile
 		// compiled out in shipping for performance (we can change in development later), if this get triggered you need to call GetValueOnGameThread() or GetValueOnAnyThread(), the last one is a bit slower
 		cvarCheckCode(ensure(IsInParallelRenderingThread()));	// ensure to not block content creators, #if to optimize in shipping
@@ -1320,6 +1323,7 @@ public:
 	// convenient, for better performance consider using GetValueOnGameThread() or GetValueOnRenderThread()
 	T GetValueOnAnyThread(bool bForceGameThread = false) const
 	{
+		UE::AccessDetection::ReportAccess(UE::AccessDetection::EType::CVar);
 		return ShadowedValue[GetShadowIndex(bForceGameThread)];
 	}
 
@@ -1342,6 +1346,7 @@ private: // ----------------------------------------------------
 	// needed for FConsoleVariable and FConsoleVariableRef2, intentionally not public
 	T& GetReferenceOnAnyThread(bool bForceGameThread = false)
 	{
+		UE::AccessDetection::ReportAccess(UE::AccessDetection::EType::CVar);
 		return ShadowedValue[GetShadowIndex(bForceGameThread)];
 	}
 
diff --git a/Engine/Source/Runtime/Core/Public/Misc/AccessDetection.h b/Engine/Source/Runtime/Core/Public/Misc/AccessDetection.h
new file mode 100644
--- /dev/null
+++ b/Engine/Source/Runtime/Core/Public/Misc/AccessDetection.h
@@ -0,0 +1,60 @@
+// Copyright Epic Games, Inc. All Rights Reserved.
+
+#pragma once
+
+#include "CoreTypes.h"
+#include "Misc/EnumClassFlags.h"
+
+namespace UE { namespace AccessDetection {
+
+enum class EType : uint32
+{
+	None			= 0,
+	File			= 1 << 0, 
+	Ini 			= 1 << 1, 
+	CVar 			= 1 << 2,
+};
+ENUM_CLASS_FLAGS(EType);
+
+#if WITH_EDITOR
+
+class FScope;
+
+namespace Private
+{
+	extern CORE_API volatile int32 GNumLiveScopes;
+
+	CORE_API void ReportCurrentThreadAccess(EType Type);
+	CORE_API void SetCurrentThreadScope(FScope* Scope);
+}
+
+FORCEINLINE void ReportAccess(EType Type)
+{
+	// Note unsynchronized access is fine - we only need to see writes by current thread
+	if (Private::GNumLiveScopes > 0)
+	{
+		Private::ReportCurrentThreadAccess(Type);
+	}
+}
+
+/** Detects access to global core systems while running deterministic code that should be isolated from the local environment */
+class FScope
+{
+public:
+	FScope() { Private::SetCurrentThreadScope(this); }
+	~FScope() { Private::SetCurrentThreadScope(nullptr); }
+	
+	void AddAccess(EType Type) { AccessedTypes |= Type; }
+	EType GetAccesses() const { return AccessedTypes; }
+
+private:
+	EType AccessedTypes = EType::None;
+};
+
+#else
+
+FORCEINLINE void ReportAccess(EType Type) {}
+
+#endif
+
+}}
\ No newline at end of file
diff --git a/Engine/Source/Runtime/Core/Public/Misc/ConfigCacheIni.h b/Engine/Source/Runtime/Core/Public/Misc/ConfigCacheIni.h
--- a/Engine/Source/Runtime/Core/Public/Misc/ConfigCacheIni.h
+++ b/Engine/Source/Runtime/Core/Public/Misc/ConfigCacheIni.h
@@ -18,6 +18,7 @@
 #include "Delegates/Delegate.h"
 #include "Math/Vector.h"
 #include "Math/Rotator.h"
+#include "Misc/AccessDetection.h"
 #include "Misc/Paths.h"
 #include "Serialization/StructuredArchive.h"
 
@@ -150,7 +151,8 @@ public:
 
 	// Returns the ini setting with any macros expanded out
 	const FString& GetValue() const 
-	{ 
+	{
+		UE::AccessDetection::ReportAccess(UE::AccessDetection::EType::Ini);
 #if CONFIG_REMEMBER_ACCESS_PATTERN 
 		bRead = true; 
 #endif
@@ -159,7 +161,8 @@ public:
 
 	// Returns the original ini setting without macro expansion
 	const FString& GetSavedValue() const 
-	{ 
+	{
+		UE::AccessDetection::ReportAccess(UE::AccessDetection::EType::Ini);
 #if CONFIG_REMEMBER_ACCESS_PATTERN 
 		bRead = true; 
 #endif
