U5 - Ability to prevent certain tasks from being used during busy waiting

diff --git a/Engine/Source/Runtime/Core/Private/Async/Fundamental/Scheduler.cpp b/Engine/Source/Runtime/Core/Private/Async/Fundamental/Scheduler.cpp
--- a/Engine/Source/Runtime/Core/Private/Async/Fundamental/Scheduler.cpp
+++ b/Engine/Source/Runtime/Core/Private/Async/Fundamental/Scheduler.cpp
@@ -185,21 +185,30 @@ namespace LowLevelTasks
 		return WorkerType != EWorkerType::None && ActiveScheduler == this;
 	}
 
-	template<FTask* (FScheduler::FLocalQueueType::*DequeueFunction)(bool)>
+	template<FTask* (FScheduler::FLocalQueueType::*DequeueFunction)(bool), bool bIsBusyWaiting>
 	bool FScheduler::TryExecuteTaskFrom(FLocalQueueType* Queue, FQueueRegistry::FOutOfWork& OutOfWork, bool bPermitBackgroundWork)
 	{
-		FTask* Task = (Queue->*DequeueFunction)(bPermitBackgroundWork);
-		if (Task)
-		{	
-			OutOfWork.Stop();		
-			FTask* OldTask = ActiveTask;
-			ActiveTask = Task;
-			{
-				TRACE_CPUPROFILER_EVENT_SCOPE(ExecuteTask);
-				Task->ExecuteTask();
+		for(int i = 0; i < 2; i++) // one retry if we pick up a task that cannot used during busy waiting.
+		{
+			FTask* Task = (Queue->*DequeueFunction)(bPermitBackgroundWork);
+			if (Task)
+			{	
+				if (bIsBusyWaiting && !Task->AllowBusyWaiting())
+				{
+					QueueRegistry.Enqueue(Task, uint32(Task->GetPriority()));
+					continue;
+				}
+				OutOfWork.Stop();		
+				FTask* OldTask = ActiveTask;
+				ActiveTask = Task;
+				{
+					TRACE_CPUPROFILER_EVENT_SCOPE(ExecuteTask);
+					Task->ExecuteTask();
+				}
+				ActiveTask = OldTask;
+				return true;
 			}
-			ActiveTask = OldTask;
-			return true;
+			return false;
 		}
 		return false;
 	}
@@ -228,15 +237,15 @@ namespace LowLevelTasks
 		FQueueRegistry::FOutOfWork OutOfWork = QueueRegistry.GetOutOfWorkScope(bPermitBackgroundWork);
 		while (true)
 		{
-			while(TryExecuteTaskFrom<&FLocalQueueType::DequeueLocal>(WorkerLocalQueue, OutOfWork, bPermitBackgroundWork)
-			   || TryExecuteTaskFrom<&FLocalQueueType::DequeueGlobal>(WorkerLocalQueue, OutOfWork, bPermitBackgroundWork))
+			while(TryExecuteTaskFrom<&FLocalQueueType::DequeueLocal, false>(WorkerLocalQueue, OutOfWork, bPermitBackgroundWork)
+			   || TryExecuteTaskFrom<&FLocalQueueType::DequeueGlobal, false>(WorkerLocalQueue, OutOfWork, bPermitBackgroundWork))
 			{		
 				Drowsing = false;
 				WaitCount = 0;
 			}
 
-			while(TryExecuteTaskFrom<&FLocalQueueType::DequeueLocal>(WorkerLocalQueue, OutOfWork, bPermitBackgroundWork)
-			   || TryExecuteTaskFrom<&FLocalQueueType::DequeueSteal>(WorkerLocalQueue, OutOfWork, bPermitBackgroundWork))
+			while(TryExecuteTaskFrom<&FLocalQueueType::DequeueLocal, false>(WorkerLocalQueue, OutOfWork, bPermitBackgroundWork)
+			   || TryExecuteTaskFrom<&FLocalQueueType::DequeueSteal, false>(WorkerLocalQueue, OutOfWork, bPermitBackgroundWork))
 			{
 				Drowsing = false;
 				WaitCount = 0;
@@ -286,8 +295,8 @@ namespace LowLevelTasks
 		FQueueRegistry::FOutOfWork OutOfWork = QueueRegistry.GetOutOfWorkScope(bIsBackgroundWorker);
 		while (true)
 		{
-			while(TryExecuteTaskFrom<&FLocalQueueType::DequeueLocal>(WorkerLocalQueue, OutOfWork, bPermitBackgroundWork)
-			   || TryExecuteTaskFrom<&FLocalQueueType::DequeueGlobal>(WorkerLocalQueue, OutOfWork, bPermitBackgroundWork))
+			while(TryExecuteTaskFrom<&FLocalQueueType::DequeueLocal, true>(WorkerLocalQueue, OutOfWork, bPermitBackgroundWork)
+			   || TryExecuteTaskFrom<&FLocalQueueType::DequeueGlobal, true>(WorkerLocalQueue, OutOfWork, bPermitBackgroundWork))
 			{
 				if (Conditional())
 				{
@@ -296,8 +305,8 @@ namespace LowLevelTasks
 				WaitCount = 0;
 			}
 
-			while(TryExecuteTaskFrom<&FLocalQueueType::DequeueLocal>(WorkerLocalQueue, OutOfWork, bPermitBackgroundWork)
-			   || TryExecuteTaskFrom<&FLocalQueueType::DequeueSteal>(WorkerLocalQueue, OutOfWork, bPermitBackgroundWork))
+			while(TryExecuteTaskFrom<&FLocalQueueType::DequeueLocal, true>(WorkerLocalQueue, OutOfWork, bPermitBackgroundWork)
+			   || TryExecuteTaskFrom<&FLocalQueueType::DequeueSteal, true>(WorkerLocalQueue, OutOfWork, bPermitBackgroundWork))
 			{
 				if (Conditional())
 				{
diff --git a/Engine/Source/Runtime/Core/Public/Async/Fundamental/Scheduler.h b/Engine/Source/Runtime/Core/Public/Async/Fundamental/Scheduler.h
--- a/Engine/Source/Runtime/Core/Public/Async/Fundamental/Scheduler.h
+++ b/Engine/Source/Runtime/Core/Public/Async/Fundamental/Scheduler.h
@@ -116,8 +116,8 @@ namespace LowLevelTasks
 		FORCENOINLINE void TrySleeping(FSleepEvent* WorkerEvent, FQueueRegistry::FOutOfWork& OutOfWork, bool& Drowsing, bool bBackgroundWorker);
 		inline bool WakeUpWorker(bool bBackgroundWorker);
 
-		template<FTask* (FLocalQueueType::*DequeueFunction)(bool)>
-		static bool TryExecuteTaskFrom(FLocalQueueType* Queue, FQueueRegistry::FOutOfWork& OutOfWork, bool bPermitBackgroundWork);
+		template<FTask* (FLocalQueueType::*DequeueFunction)(bool), bool bIsBusyWaiting>
+		bool TryExecuteTaskFrom(FLocalQueueType* Queue, FQueueRegistry::FOutOfWork& OutOfWork, bool bPermitBackgroundWork);
 
 	private:
 		FEventQueueType 			SleepEventQueue[2];
diff --git a/Engine/Source/Runtime/Core/Public/Async/Fundamental/Task.h b/Engine/Source/Runtime/Core/Public/Async/Fundamental/Task.h
--- a/Engine/Source/Runtime/Core/Public/Async/Fundamental/Task.h
+++ b/Engine/Source/Runtime/Core/Public/Async/Fundamental/Task.h
@@ -73,38 +73,33 @@ namespace LowLevelTasks
 
 		class FPackedData
 		{
-#if PLATFORM_32BITS
-			uintptr_t DebugName : 32;
-#else
-			uintptr_t DebugName : 57;
-#endif
-			uintptr_t Priority	: 3;
-			uintptr_t State		: 4;
+			uintptr_t DebugName			: 56;
+			uintptr_t Priority			: 3;
+			uintptr_t State				: 4;
+			uintptr_t bAllowBusyWaiting : 1;
 			
 		public:
-			FPackedData() : FPackedData(nullptr,  ETaskPriority::Default, ETaskState::Completed)
+			FPackedData() : FPackedData(nullptr,  ETaskPriority::Default, ETaskState::Completed, true)
 			{
+				static_assert(!PLATFORM_32BITS, "32bit Platforms are not supported");
 				static_assert(uintptr_t(ETaskPriority::Count) <= (1ull << 3), "Not enough bits to store ETaskPriority");
 				static_assert(uintptr_t(ETaskState::Count) <= (1ull << 4), "Not enough bits to store ETaskState");
 			}
 
-			FPackedData(const TCHAR* InDebugName, ETaskPriority InPriority, ETaskState InState)
+			FPackedData(const TCHAR* InDebugName, ETaskPriority InPriority, ETaskState InState, bool bInAllowBusyWaiting)
 				: DebugName(reinterpret_cast<uintptr_t>(InDebugName))
 				, Priority((uintptr_t)InPriority)
 				, State((uintptr_t)InState)
+				, bAllowBusyWaiting(bInAllowBusyWaiting)
 			{
 				checkSlow(reinterpret_cast<uintptr_t>(InDebugName) < (1ull << 57));
 				checkSlow((uintptr_t)InPriority < (1ull << 3));
 				checkSlow((uintptr_t)InState < (1ull << 4));
-#if PLATFORM_32BITS
-				static_assert(sizeof(FPackedData) == 2 * sizeof(uintptr_t), "Packed data needs to be 8bytes in size");
-#else
 				static_assert(sizeof(FPackedData) == sizeof(uintptr_t), "Packed data needs to be pointer size");
-#endif
 			}
 
 			FPackedData(FPackedData& Other, ETaskState State)
-				: FPackedData(Other.GetDebugName(), Other.GetPriority(), State)
+				: FPackedData(Other.GetDebugName(), Other.GetPriority(), State, Other.AllowBusyWaiting())
 			{
 			}
 
@@ -122,6 +117,11 @@ namespace LowLevelTasks
 			{
 				return ETaskState(State);
 			}
+
+			bool AllowBusyWaiting() const
+			{
+				return bAllowBusyWaiting;
+			}
 		};
 
 	private:
@@ -175,6 +175,18 @@ namespace LowLevelTasks
 		//you can alternatively use FScheduler::TryCancelAndLaunch if you want to launch the task in this case automatically
 		inline bool TryCancel();
 
+		template<typename TRunnable, typename TContinuation>
+		inline void Init(const TCHAR* InDebugName, ETaskPriority InPriority, TRunnable&& InRunnable, TContinuation&& InContinuation, bool bAllowBusyWaiting);
+
+		template<typename TRunnable>
+		inline void Init(const TCHAR* InDebugName, ETaskPriority InPriority, TRunnable&& InRunnable, bool bAllowBusyWaiting);
+
+		template<typename TRunnable, typename TContinuation>
+		inline void Init(const TCHAR* InDebugName, TRunnable&& InRunnable, TContinuation&& InContinuation, bool bAllowBusyWaiting);
+
+		template<typename TRunnable>
+		inline void Init(const TCHAR* InDebugName, TRunnable&& InRunnable, bool bAllowBusyWaiting);
+
 		template<typename TRunnable, typename TContinuation>
 		inline void Init(const TCHAR* InDebugName, ETaskPriority InPriority, TRunnable&& InRunnable, TContinuation&& InContinuation);
 
@@ -202,6 +214,7 @@ namespace LowLevelTasks
 		inline bool TryPrepareLaunch();
 		//after calling this function the task can be considered dead
 		inline void ExecuteTask();
+		inline bool AllowBusyWaiting() const;
 	};
 
    /******************
@@ -209,7 +222,7 @@ namespace LowLevelTasks
 	******************/
 
 	template<typename TRunnable, typename TContinuation>
-	inline void FTask::Init(const TCHAR* InDebugName, ETaskPriority InPriority, TRunnable&& InRunnable, TContinuation&& InContinuation)
+	inline void FTask::Init(const TCHAR* InDebugName, ETaskPriority InPriority, TRunnable&& InRunnable, TContinuation&& InContinuation, bool bAllowBusyWaiting)
 	{
 		checkf(IsCompleted(), TEXT("State: %d"), PackedData.load(std::memory_order_relaxed).GetState());
 		checkSlow(!Runnable.IsSet());
@@ -231,11 +244,11 @@ namespace LowLevelTasks
 				return true;
 			}
 		};
-		PackedData.store(FPackedData(InDebugName, InPriority, ETaskState::Ready), std::memory_order_release);
+		PackedData.store(FPackedData(InDebugName, InPriority, ETaskState::Ready, bAllowBusyWaiting), std::memory_order_release);
 	}
 
 	template<typename TRunnable>
-	inline void FTask::Init(const TCHAR* InDebugName, ETaskPriority InPriority, TRunnable&& InRunnable)
+	inline void FTask::Init(const TCHAR* InDebugName, ETaskPriority InPriority, TRunnable&& InRunnable, bool bAllowBusyWaiting)
 	{
 		checkf(IsCompleted(), TEXT("State: %d"), PackedData.load(std::memory_order_relaxed).GetState());
 		checkSlow(!Runnable.IsSet());
@@ -255,19 +268,43 @@ namespace LowLevelTasks
 				return true;
 			}
 		};
-		PackedData.store(FPackedData(InDebugName, InPriority, ETaskState::Ready), std::memory_order_release);
+		PackedData.store(FPackedData(InDebugName, InPriority, ETaskState::Ready, bAllowBusyWaiting), std::memory_order_release);
+	}
+
+	template<typename TRunnable, typename TContinuation>
+	inline void FTask::Init(const TCHAR* InDebugName, TRunnable&& InRunnable, TContinuation&& InContinuation, bool bAllowBusyWaiting)
+	{
+		Init(InDebugName, ETaskPriority::Default, Forward<TRunnable>(InRunnable), Forward<TContinuation>(InContinuation), bAllowBusyWaiting);
+	}
+
+	template<typename TRunnable>
+	inline void FTask::Init(const TCHAR* InDebugName, TRunnable&& InRunnable, bool bAllowBusyWaiting)
+	{
+		Init(InDebugName, ETaskPriority::Default, Forward<TRunnable>(InRunnable), bAllowBusyWaiting);
+	}
+
+	template<typename TRunnable, typename TContinuation>
+	inline void FTask::Init(const TCHAR* InDebugName, ETaskPriority InPriority, TRunnable&& InRunnable, TContinuation&& InContinuation)
+	{
+		Init(InDebugName, InPriority, Forward<TRunnable>(InRunnable), Forward<TContinuation>(InContinuation), true);
+	}
+
+	template<typename TRunnable>
+	inline void FTask::Init(const TCHAR* InDebugName, ETaskPriority InPriority, TRunnable&& InRunnable)
+	{
+		Init(InDebugName, InPriority, Forward<TRunnable>(InRunnable), true);
 	}
 
 	template<typename TRunnable, typename TContinuation>
 	inline void FTask::Init(const TCHAR* InDebugName, TRunnable&& InRunnable, TContinuation&& InContinuation)
 	{
-		Init(InDebugName, ETaskPriority::Default, Forward<TRunnable>(InRunnable), Forward<TContinuation>(InContinuation));
+		Init(InDebugName, ETaskPriority::Default, Forward<TRunnable>(InRunnable), Forward<TContinuation>(InContinuation), true);
 	}
 
 	template<typename TRunnable>
 	inline void FTask::Init(const TCHAR* InDebugName, TRunnable&& InRunnable)
 	{
-		Init(InDebugName, ETaskPriority::Default, Forward<TRunnable>(InRunnable));
+		Init(InDebugName, ETaskPriority::Default, Forward<TRunnable>(InRunnable), true);
 	}
 
 	inline FTask::~FTask()
@@ -322,4 +359,9 @@ namespace LowLevelTasks
 	{
 		return PackedData.load(std::memory_order_relaxed).GetPriority() >= ETaskPriority::ForegroundCount;
 	}
+
+	inline bool FTask::AllowBusyWaiting() const
+	{
+		return PackedData.load(std::memory_order_relaxed).AllowBusyWaiting();
+	}	
 }
\ No newline at end of file
