Add support for displaying floats greater than e18 & remove cast losing double precision.

diff --git a/Engine/Source/Editor/PropertyEditor/Private/PropertyHandleImpl.cpp b/Engine/Source/Editor/PropertyEditor/Private/PropertyHandleImpl.cpp
--- a/Engine/Source/Editor/PropertyEditor/Private/PropertyHandleImpl.cpp
+++ b/Engine/Source/Editor/PropertyEditor/Private/PropertyHandleImpl.cpp
@@ -3639,7 +3639,7 @@ FPropertyAccess::Result FPropertyHandleDouble::SetValue( const double& NewValue,
 {
 	FPropertyAccess::Result Res;
 	// Clamp the value from any meta data ranges stored on the property value
-	float FinalValue = ClampValueFromMetaData<double>( NewValue, *Implementation->GetPropertyNode() );
+	double FinalValue = ClampValueFromMetaData<double>( NewValue, *Implementation->GetPropertyNode() );
 
 	const FString ValueStr = FString::Printf( TEXT("%f"), FinalValue );
 	Res = Implementation->ImportText( ValueStr, Flags );
diff --git a/Engine/Source/Runtime/Core/Private/Internationalization/FastDecimalFormat.cpp b/Engine/Source/Runtime/Core/Private/Internationalization/FastDecimalFormat.cpp
--- a/Engine/Source/Runtime/Core/Private/Internationalization/FastDecimalFormat.cpp
+++ b/Engine/Source/Runtime/Core/Private/Internationalization/FastDecimalFormat.cpp
@@ -2,6 +2,8 @@
 
 #include "Internationalization/FastDecimalFormat.h"
 #include "Misc/EnumClassFlags.h"
+#include "Misc/StringBuilder.h"
+#include "Containers/StringFwd.h"
 
 namespace FastDecimalFormat
 {
@@ -13,6 +15,12 @@ static const int32 MaxIntegralPrintLength = 20;
 static const int32 MaxFractionalPrintPrecision = 18;
 static const int32 MinRequiredIntegralBufferSize = (MaxIntegralPrintLength * 2) + 1; // *2 for an absolute worst case group separator scenario, +1 for null terminator
 
+static int32 bFastDecimalFormatLargeFloatSupport = 1;
+static FAutoConsoleVariableRef CVarFastDecimalFormatLargeFloatSupport(
+	TEXT("Core.bFastDecimalFormatLargeFloatSupport"),
+	bFastDecimalFormatLargeFloatSupport,
+	TEXT("True implies we perform additional processing for floating point types over 9223372036854775807 to prevent clipping to this value."));
+
 static const uint64 Pow10Table[] = {
 	1,						// 10^0
 	10,						// 10^1
@@ -395,6 +403,103 @@ void IntegralToString(const bool bIsNegative, const uint64 InVal, const FDecimal
 	BuildFinalString(bIsNegative, InFormattingOptions.AlwaysSign, InFormattingRules, IntegralPartBuffer, IntegralPartLen, FractionalPartBuffer, FractionalPartLen, OutString);
 }
 
+FString CultureInvariantDecimalToString(const double InVal, const TCHAR*& InBuffer, const int32 InBufferLen, const FDecimalNumberFormattingRules& InFormattingRules, const FNumberFormattingOptions& InFormattingOptions)
+{
+	if (!ensure(InBuffer && InBufferLen > 0))
+	{
+		return FString();
+	}
+
+	// Note: Does not consider max digits, this is by design as this method was created to support large floats greater than e18.
+	TStringBuilder<128> OutStr;
+
+	bool bUseGrouping = InFormattingOptions.UseGrouping && InFormattingRules.PrimaryGroupingSize > 0;
+	uint8 NumIntegralDigits = static_cast<uint8>(FMath::Abs((FMath::LogX(10.0, InVal)))) + 1;
+	uint8 NumUntilNextGroup = NumIntegralDigits % InFormattingRules.PrimaryGroupingSize;
+	const TCHAR* InBufferEnd = InBuffer + InBufferLen;
+
+	// Apply front padding
+	const int32 PaddingToApply = FMath::Max(InFormattingOptions.MinimumIntegralDigits - InBufferLen, 0);
+	for (int32 PaddingIndex = 0; PaddingIndex < PaddingToApply; ++PaddingIndex)
+	{
+		if (bUseGrouping && NumUntilNextGroup-- == 0)
+		{
+			OutStr += InFormattingRules.GroupingSeparatorCharacter;
+			NumUntilNextGroup = InFormattingRules.SecondaryGroupingSize;
+		}
+
+		OutStr += InFormattingRules.DigitCharacters[0];
+	}
+
+	// Scrape negative, apply at end
+	bool bIsNegative = false;
+	static const TCHAR EuropeanNegativePrefix = '-';
+	if (*InBuffer == EuropeanNegativePrefix)
+	{
+		bIsNegative = true;
+		++InBuffer;
+	}
+
+	// Parse digits & decimal, no grouping on fractional
+	bool bParsedFractional = false;
+	uint8 FractionalDigitsPrinted = 0;
+	while (InBuffer < InBufferEnd)
+	{
+		static const TCHAR EuropeanDecimal = '.';
+		if (*InBuffer == EuropeanDecimal && InFormattingOptions.MaximumFractionalDigits > 0)
+		{
+			bParsedFractional = true;
+			OutStr += InFormattingRules.DecimalSeparatorCharacter;
+			++InBuffer;
+			continue;
+		}
+
+		if (!bParsedFractional && bUseGrouping && NumUntilNextGroup-- == 0)
+		{
+			OutStr += InFormattingRules.GroupingSeparatorCharacter;
+			NumUntilNextGroup = InFormattingRules.SecondaryGroupingSize - 1; // -1 to account for the digit we're about to print
+		}
+
+		// 48 for raw ascii -> int
+		OutStr += InFormattingRules.DigitCharacters[(int32)InBuffer[0] - 48];
+		FractionalDigitsPrinted += bParsedFractional ? 1 : 0;
+		++InBuffer;
+	}
+
+	// Apply back padding, if back isn't just zero
+	double IntegralVal = 0.0;
+	double FractionalVal = FMath::Modf(InVal, &IntegralVal);
+	if (InFormattingOptions.MaximumFractionalDigits > FractionalDigitsPrinted)
+	{
+		if (!bParsedFractional)
+		{
+			OutStr += InFormattingRules.DecimalSeparatorCharacter;
+		}
+
+		if (FMath::Abs(FractionalVal) > 0.0)
+		{
+			const int32 BackPaddingToApply = InFormattingOptions.MaximumFractionalDigits - FractionalDigitsPrinted;
+			for (int32 PaddingIndex = 0; PaddingIndex < BackPaddingToApply; ++PaddingIndex)
+			{
+				OutStr += InFormattingRules.DigitCharacters[0];
+			}
+		}
+	}
+
+	const FDecimalNumberSigningStrings SigningStrings(InFormattingRules, InFormattingOptions.AlwaysSign ? EDecimalNumberSigningStringsFlags::AlwaysSign : EDecimalNumberSigningStringsFlags::None);
+
+	const FString& FinalPrefixStr = (bIsNegative) ? SigningStrings.GetNegativePrefixString() : SigningStrings.GetPositivePrefixString();
+	const FString& FinalSuffixStr = (bIsNegative) ? SigningStrings.GetNegativeSuffixString() : SigningStrings.GetPositiveSuffixString();
+
+	OutStr += FinalSuffixStr.IsEmpty() ? "" : FinalSuffixStr;
+	if (OutStr.LastChar() != '\0')
+	{
+		OutStr += '\0';
+	}
+
+	return FinalPrefixStr.IsEmpty() ? OutStr.GetData() : FinalPrefixStr + OutStr.GetData();
+}
+
 void FractionalToString(const double InVal, const FDecimalNumberFormattingRules& InFormattingRules, FNumberFormattingOptions InFormattingOptions, FString& OutString)
 {
 	SanitizeNumberFormattingOptions(InFormattingOptions);
@@ -417,9 +522,21 @@ void FractionalToString(const double InVal, const FDecimalNumberFormattingRules&
 		FractionalPart = -FractionalPart;
 	}
 
+	// Check for float-> int overflow, fallback on regular lex if occurs
+	// if fractional part overflows then we are losing precession but the number is still valid
+	uint64 IntIntegralPart = static_cast<uint64>(IntegralPart);
+	if (IntegralPart - static_cast<double>(IntIntegralPart) > SMALL_NUMBER && bFastDecimalFormatLargeFloatSupport)
+	{
+		OutString = LexToSanitizedString(InVal);
+
+		const TCHAR* CultureInvariantDecimalBuffer = *OutString;
+		OutString = CultureInvariantDecimalToString(InVal, CultureInvariantDecimalBuffer, OutString.Len(), InFormattingRules, InFormattingOptions);
+		return;
+	}
+
 	// Deal with the integral part (produces a string of the integral part, inserting group separators if requested and required, and padding as needed)
 	TCHAR IntegralPartBuffer[MinRequiredIntegralBufferSize];
-	const int32 IntegralPartLen = IntegralToString_Common(static_cast<uint64>(IntegralPart), InFormattingRules, InFormattingOptions, IntegralPartBuffer, UE_ARRAY_COUNT(IntegralPartBuffer));
+	const int32 IntegralPartLen = IntegralToString_Common(IntIntegralPart, InFormattingRules, InFormattingOptions, IntegralPartBuffer, UE_ARRAY_COUNT(IntegralPartBuffer));
 
 	// Deal with the fractional part (produces a string of the fractional part, potentially padding with zeros up to InFormattingOptions.MaximumFractionalDigits)
 	TCHAR FractionalPartBuffer[MinRequiredIntegralBufferSize];
@@ -704,6 +821,95 @@ bool StringToIntegral(const TCHAR* InStr, const int32 InStrLen, const FDecimalNu
 	return bResult;
 }
 
+bool StringToCultureInvariantDecimal(const TCHAR*& InBuffer, const TCHAR* InBufferEnd, const FDecimalNumberFormattingRules& InFormattingRules, const FNumberParsingOptions& InParsingOptions, const FDecimalNumberSignParser& InSignParser, TStringBuilder<128>& OutInvariantDecimal)
+{
+	// Empty string?
+	if (*InBuffer == 0)
+	{
+		return true;
+	}
+
+	EDecimalNumberParseFlags InParseFlags = EDecimalNumberParseFlags::AllowLeadingSign | EDecimalNumberParseFlags::AllowTrailingSign | EDecimalNumberParseFlags::AllowDecimalSeparators
+		| (InParsingOptions.UseGrouping ? EDecimalNumberParseFlags::AllowGroupSeparators : EDecimalNumberParseFlags::None)
+		| (InParsingOptions.InsideLimits ? EDecimalNumberParseFlags::TestLimits : EDecimalNumberParseFlags::None)
+		| (InParsingOptions.UseClamping ? EDecimalNumberParseFlags::ClampValue : EDecimalNumberParseFlags::None);
+
+	// Parse the leading sign (if present)
+	bool bIsNegative = false;
+	InSignParser.ParseLeadingSign(InBuffer, bIsNegative);
+
+	static const TCHAR InvariantNegativePrefix = '-';
+	if (bIsNegative)
+	{
+		OutInvariantDecimal += InvariantNegativePrefix;
+	}
+
+	// Parse the number, stopping once we find the end of the string or a decimal separator
+	static const TCHAR EuropeanNumerals[] = { '0', '1', '2', '3', '4', '5', '6', '7', '8', '9' };
+	bool bFoundUnexpectedNonNumericCharacter = false;
+	while (InBuffer < InBufferEnd)
+	{
+		// Skip group separators
+		if (*InBuffer == InFormattingRules.GroupingSeparatorCharacter)
+		{
+			if (!EnumHasAnyFlags(InParseFlags, EDecimalNumberParseFlags::AllowGroupSeparators))
+			{
+				return false;
+			}
+			++InBuffer;
+			continue;
+		}
+
+		// Walk over the decimal separator
+		static const TCHAR InvariantDecimal = '.';
+		if (*InBuffer == InFormattingRules.DecimalSeparatorCharacter)
+		{
+			if (!EnumHasAnyFlags(InParseFlags, EDecimalNumberParseFlags::AllowDecimalSeparators))
+			{
+				return false;
+			}
+			++InBuffer;
+			OutInvariantDecimal += InvariantDecimal;
+			continue;
+		}
+
+		// Process numeric characters (also test European numerals in case they were used by a language that doesn't normally use them)
+		bool bValidChar = false;
+		for (int32 CharIndex = 0; CharIndex < UE_ARRAY_COUNT(InFormattingRules.DigitCharacters); ++CharIndex)
+		{
+			if ((*InBuffer == InFormattingRules.DigitCharacters[CharIndex]) || (*InBuffer == EuropeanNumerals[CharIndex]))
+			{
+				// We don't consider MaxIntegralPrintLength, since this method is used to deal with large string to string values
+				++InBuffer;
+				OutInvariantDecimal += EuropeanNumerals[CharIndex];
+				bValidChar = true;
+				break;
+			}
+		}
+
+		// Found an non-numeric character?
+		if (!bValidChar)
+		{
+			bFoundUnexpectedNonNumericCharacter = true;
+			break;
+		}
+	}
+
+	// Parse the trailing sign (if present)
+	if (InSignParser.ParseTrailingSign(InBuffer, bIsNegative))
+	{
+		// The unexpected character was the trailing sign - clear that flag
+		bFoundUnexpectedNonNumericCharacter = false;
+	}
+
+	if (OutInvariantDecimal.LastChar() != '\0')
+	{
+		OutInvariantDecimal += '\0';
+	}
+
+	return !bFoundUnexpectedNonNumericCharacter;
+}
+
 bool StringToFractional(const TCHAR* InStr, const int32 InStrLen, const FDecimalNumberFormattingRules& InFormattingRules, const FNumberParsingOptions& InParsingOptions, const FDecimalNumberFractionalLimits& InLimits, double& OutVal, int32* OutParsedLen)
 {
 	const TCHAR* Buffer = InStr;
@@ -718,6 +924,20 @@ bool StringToFractional(const TCHAR* InStr, const int32 InStrLen, const FDecimal
 	bool bResult = StringToIntegral_Common(Buffer, BufferEnd, InFormattingRules, InParsingOptions, SignParser, bIntegralPartIsNegative, bIntegralPartIsOverflow, IntegralPart, IntegralPartDigitCount);
 	bResult &= !bIntegralPartIsOverflow;
 
+	if (bIntegralPartIsOverflow && bFastDecimalFormatLargeFloatSupport)
+	{
+		const TCHAR* InvariantBuffer = InStr;
+		TStringBuilder<128> InvariantDecimal;
+		if (StringToCultureInvariantDecimal(InvariantBuffer, BufferEnd, InFormattingRules, InParsingOptions, SignParser, InvariantDecimal))
+		{
+			LexFromString(OutVal, InvariantDecimal.GetData());
+
+			// We overflowed, so have callers act as if string length hasn't changed
+			*OutParsedLen = InStrLen;
+			return true;
+		}
+	}
+
 	// Parse the fractional part of the number
 	bool bFractionPartIsNegative = false;
 	bool bFractionPartIsOverflow = false;
diff --git a/Engine/Source/Runtime/Slate/Public/Widgets/Input/SNumericEntryBox.h b/Engine/Source/Runtime/Slate/Public/Widgets/Input/SNumericEntryBox.h
--- a/Engine/Source/Runtime/Slate/Public/Widgets/Input/SNumericEntryBox.h
+++ b/Engine/Source/Runtime/Slate/Public/Widgets/Input/SNumericEntryBox.h
@@ -182,6 +182,12 @@ public:
 		BorderImageFocused = &InArgs._EditableTextBoxStyle->BackgroundImageFocused;
 		Interface = InArgs._TypeInterface.IsValid() ? InArgs._TypeInterface : MakeShareable( new TDefaultNumericTypeInterface<NumericType> );
 
+		CachedExternalValue = ValueAttribute.Get();
+		if (CachedExternalValue.IsSet())
+		{
+			CachedValueString = Interface->ToString(CachedExternalValue.GetValue());
+		}
+
 		TAttribute<FMargin> TextMargin = InArgs._OverrideTextMargin.IsSet() ? InArgs._OverrideTextMargin : InArgs._EditableTextBoxStyle->Padding;
 		const bool bAllowSpin = InArgs._AllowSpin;
 		TSharedPtr<SWidget> FinalWidget;
@@ -407,12 +413,27 @@ private:
 		return 0;
 	}
 
+	void SetCachedString(const NumericType CurrentValue)
+	{
+		if (!CachedExternalValue.IsSet() || CachedExternalValue.GetValue() != CurrentValue)
+		{
+			CachedExternalValue = CurrentValue;
+			CachedValueString = Interface->ToString(CurrentValue);
+		}
+	}
+
+	FString GetCachedString(const NumericType CurrentValue) const
+	{
+		bool bUseCachedString = CachedExternalValue.IsSet() && CurrentValue == CachedExternalValue.GetValue();
+		return bUseCachedString ? CachedValueString : Interface->ToString(CurrentValue);
+	}
+
 	FText GetValueAsText() const
 	{
 		const TOptional<NumericType>& Value = ValueAttribute.Get();
 		if (Value.IsSet() == true)
 		{
-			return FText::FromString(Interface->ToString(Value.GetValue()));
+			return FText::FromString(GetCachedString(Value.GetValue()));
 		}
 		
 		return FText::GetEmpty();
@@ -432,7 +453,7 @@ private:
 			// If the value was set convert it to a string, otherwise the value cannot be determined
 			if( Value.IsSet() == true )
 			{
-				NewText = FText::FromString(Interface->ToString(Value.GetValue()));
+				NewText = FText::FromString(GetCachedString(Value.GetValue()));
 			}
 			else
 			{
@@ -526,6 +547,8 @@ private:
 
 		if (NumericValue.IsSet())
 		{
+			SetCachedString(NumericValue.GetValue());
+
 			if (bCommit)
 			{
 				OnValueCommitted.ExecuteIfBound(NumericValue.GetValue(), CommitInfo);
@@ -552,6 +575,8 @@ private:
 
 		if (Value.IsSet() == true)
 		{
+			SetCachedString(Value.GetValue());
+
 			if (SpinBox->GetVisibility() != EVisibility::Visible)
 			{
 				// Set the visibility of the spinbox to visible if we have a valid value
@@ -598,6 +623,10 @@ private:
 	TAttribute<float> MinDesiredValueWidth;
 	/** Type interface that defines how we should deal with the templated numeric type. Always valid after construction. */
 	TSharedPtr< INumericTypeInterface<NumericType> > Interface;
+	/** Cached value of entry box, updated on set & per tick */
+	TOptional<NumericType> CachedExternalValue;
+	/** Used to prevent per-frame re-conversion of the cached numeric value to a string. */
+	FString CachedValueString;
 };
 
 
diff --git a/Engine/Source/Runtime/Slate/Public/Widgets/Input/SSpinBox.h b/Engine/Source/Runtime/Slate/Public/Widgets/Input/SSpinBox.h
--- a/Engine/Source/Runtime/Slate/Public/Widgets/Input/SSpinBox.h
+++ b/Engine/Source/Runtime/Slate/Public/Widgets/Input/SSpinBox.h
@@ -199,6 +199,8 @@ public:
 		bPreventThrottling = InArgs._PreventThrottling;
 
 		CachedExternalValue = ValueAttribute.Get();
+		CachedValueString = Interface->ToString(CachedExternalValue);
+
 		InternalValue = (double)CachedExternalValue;
 
 		if (SupportDynamicSliderMaxValue.Get() && CachedExternalValue > GetMaxSliderValue())
@@ -718,7 +720,7 @@ public:
 	TAttribute<NumericType> GetValueAttribute() const { return ValueAttribute; }
 
 	/** See the Value attribute */
-	float GetValue() const { return ValueAttribute.Get(); }
+	NumericType GetValue() const { return ValueAttribute.Get(); }
 	void SetValue(const TAttribute<NumericType>& InValueAttribute) 
 	{
 		ValueAttribute = InValueAttribute; 
@@ -810,7 +812,13 @@ protected:
 	/** @return the value being observed by the spinbox as a string */
 	FString GetValueAsString() const
 	{
-		return Interface->ToString(ValueAttribute.Get());
+		NumericType CurrentValue = ValueAttribute.Get();
+		if (CurrentValue == CachedExternalValue)
+		{
+			return CachedValueString;
+		}
+
+		return Interface->ToString(CurrentValue);
 	}
 
 	/** @return the value being observed by the spinbox as FText - todo: spinbox FText support (reimplement me) */
@@ -959,7 +967,12 @@ protected:
 		}
 
 		// Update the cache of the external value to what the user believes the value is now.
-		CachedExternalValue = ValueAttribute.Get();
+		const NumericType CurrentValue = ValueAttribute.Get();
+		if (CachedExternalValue != CurrentValue)
+		{
+			CachedExternalValue = ValueAttribute.Get();
+			CachedValueString = Interface->ToString(CachedExternalValue);
+		}
 
 		// This ensures that dragging is cleared if focus has been removed from this widget in one of the delegate calls, such as when spawning a modal dialog.
 		if (!this->HasMouseCapture())
@@ -1105,6 +1118,9 @@ private:
 	 */
 	NumericType CachedExternalValue;
 
+	/** Used to prevent per-frame re-conversion of the cached numeric value to a string. */
+	FString CachedValueString;
+
 	/** Re-entrant guard for the text changed handler */
 	bool bIsTextChanging;
 
