Expose Running Commandlet Class for more precise logic when running a commandlet

diff --git a/Engine/Source/Runtime/Core/Private/Misc/CoreGlobals.cpp b/Engine/Source/Runtime/Core/Private/Misc/CoreGlobals.cpp
--- a/Engine/Source/Runtime/Core/Private/Misc/CoreGlobals.cpp
+++ b/Engine/Source/Runtime/Core/Private/Misc/CoreGlobals.cpp
@@ -166,6 +166,7 @@ bool					PRIVATE_GIsRunningCookCommandlet	= false;				/** Whether this executabl
 
 #if WITH_ENGINE
 bool					PRIVATE_GIsRunningCommandlet		= false;				/** Whether this executable is running a commandlet (custom command-line processing code) */
+UClass*					PRIVATE_GRunningCommandletClass		= nullptr;				/** Class of running cook commandlet */
 bool					PRIVATE_GAllowCommandletRendering	= false;				/** If true, initialise RHI and set up scene for rendering even when running a commandlet. */
 bool					PRIVATE_GAllowCommandletAudio 		= false;				/** If true, allow audio even when running a commandlet. */
 #endif	// WITH_ENGINE
diff --git a/Engine/Source/Runtime/Core/Public/CoreGlobals.h b/Engine/Source/Runtime/Core/Public/CoreGlobals.h
--- a/Engine/Source/Runtime/Core/Public/CoreGlobals.h
+++ b/Engine/Source/Runtime/Core/Public/CoreGlobals.h
@@ -152,6 +152,7 @@ extern CORE_API bool GIsRunningUnattendedScript;
 
 #if WITH_ENGINE
 extern CORE_API bool PRIVATE_GIsRunningCommandlet;
+extern CORE_API UClass* PRIVATE_GRunningCommandletClass;
 
 /** If true, initialize RHI and set up scene for rendering even when running a commandlet. */
 extern CORE_API bool PRIVATE_GAllowCommandletRendering;
@@ -223,6 +224,18 @@ FORCEINLINE bool IsRunningCookCommandlet()
 #endif
 }
 
+/** Returns running commandlet name 
+ * 
+ */
+FORCEINLINE UClass* GetRunningCommandletClass()
+{
+#if WITH_ENGINE
+	return PRIVATE_GRunningCommandletClass;
+#else
+	return nullptr;
+#endif
+}
+
 /**
  * Check to see if we should initialise RHI and set up scene for rendering even when running a commandlet.
  */
diff --git a/Engine/Source/Runtime/Launch/Private/LaunchEngineLoop.cpp b/Engine/Source/Runtime/Launch/Private/LaunchEngineLoop.cpp
--- a/Engine/Source/Runtime/Launch/Private/LaunchEngineLoop.cpp
+++ b/Engine/Source/Runtime/Launch/Private/LaunchEngineLoop.cpp
@@ -3520,7 +3520,15 @@ int32 FEngineLoop::PreInitPostStartupScreen(const TCHAR* CmdLine)
 			CycleCount_AfterStats.StopAndResetStatId();
 #endif // STATS
 			FStats::TickCommandletStats();
+
+#if WITH_ENGINE
+			PRIVATE_GRunningCommandletClass = CommandletClass;
+#endif
 			int32 ErrorLevel = Commandlet->Main( CommandletCommandLine );
+#if WITH_ENGINE
+			PRIVATE_GRunningCommandletClass = nullptr;
+#endif
+
 			FStats::TickCommandletStats();
 
 			RequestEngineExit(FString::Printf(TEXT("Commandlet %s finished execution (result %d)"), *Commandlet->GetName(), ErrorLevel));
