added support for mutable lambdas as task bodies

diff --git a/Engine/Source/Runtime/Core/Public/Async/Fundamental/Task.h b/Engine/Source/Runtime/Core/Public/Async/Fundamental/Task.h
--- a/Engine/Source/Runtime/Core/Public/Async/Fundamental/Task.h
+++ b/Engine/Source/Runtime/Core/Public/Async/Fundamental/Task.h
@@ -226,7 +226,7 @@ namespace LowLevelTasks
 	{
 		checkf(IsCompleted(), TEXT("State: %d"), PackedData.load(std::memory_order_relaxed).GetState());
 		checkSlow(!Runnable.IsSet());
-		Runnable = [this, LocalRunnable = Forward<TRunnable>(InRunnable), LocalContinuation = Forward<TContinuation>(InContinuation)]() -> bool
+		Runnable = [this, LocalRunnable = Forward<TRunnable>(InRunnable), LocalContinuation = Forward<TContinuation>(InContinuation)]() mutable -> bool
 		{
 			checkSlow(IsScheduled() || IsCanceled());
 			FPackedData LocalPackedData = PackedData.load(std::memory_order_relaxed);
