Fix one potential deadlock introduced by the busy wait design of the new task graph

 - Add unittest to repro the busy wait deadlock happening when th new task graph is activated
 - Add support for telling the lowlevel scheduler a task wants to opt out of being picked up during busy waits
 - Avoid FinishCompilation tasks being picked up by busy waits since that could lead to a deadlock

TESTS
  Run unittest and confirm deadlock when EQueuedWorkFlags::DoNotRunInsideBusyWait is missing
  Run unittest and confirm no deadlock when EQueuedWorkFlags::DoNotRunInsideBusyWait is present

diff --git a/Engine/Source/Runtime/Core/Private/Misc/QueuedThreadPoolWrapper.cpp b/Engine/Source/Runtime/Core/Private/Misc/QueuedThreadPoolWrapper.cpp
--- a/Engine/Source/Runtime/Core/Private/Misc/QueuedThreadPoolWrapper.cpp
+++ b/Engine/Source/Runtime/Core/Private/Misc/QueuedThreadPoolWrapper.cpp
@@ -24,6 +24,11 @@ struct FQueuedThreadPoolWrapper::FScheduledWork : public IQueuedWork
 		ParentPool->Schedule(this);
 	}
 
+	EQueuedWorkFlags GetQueuedWorkFlags() const override 
+	{
+		return InnerWork->GetQueuedWorkFlags();
+	}
+
 	IQueuedWork* InnerWork = nullptr;
 	EQueuedWorkPriority Priority = EQueuedWorkPriority::Normal;
 private:
diff --git a/Engine/Source/Runtime/Core/Private/Tests/Async/AsyncWorkTest.cpp b/Engine/Source/Runtime/Core/Private/Tests/Async/AsyncWorkTest.cpp
new file mode 100644
--- /dev/null
+++ b/Engine/Source/Runtime/Core/Private/Tests/Async/AsyncWorkTest.cpp
@@ -0,0 +1,137 @@
+// Copyright Epic Games, Inc. All Rights Reserved.
+
+#include "CoreTypes.h"
+#include "Misc/AutomationTest.h"
+#include "Async/AsyncWork.h"
+#include "Misc/QueuedThreadPool.h"
+
+#if WITH_DEV_AUTOMATION_TESTS
+
+IMPLEMENT_SIMPLE_AUTOMATION_TEST(FAsyncWorkEnsureCompletionBusyWaitDeadLockTest, "System.Core.AsyncWork.EnsureCompletionBusyWaitDeadlockTest", EAutomationTestFlags::EditorContext | EAutomationTestFlags::EngineFilter)
+
+namespace AsyncWorkTestImpl
+{
+	thread_local bool bIsOuterTaskRunning = false;
+
+	struct FInnerTask : public FNonAbandonableTask
+	{
+		void DoWork() 
+		{ 
+			TRACE_CPUPROFILER_EVENT_SCOPE(FInnerTask::DoWork);
+			FPlatformProcess::Sleep(0.01f); 
+		}
+
+		FORCEINLINE TStatId GetStatId() const
+		{
+			RETURN_QUICK_DECLARE_CYCLE_STAT(FInnerTask, STATGROUP_ThreadPoolAsyncTasks);
+		}
+	};
+
+	struct FOuterTask : public FNonAbandonableTask
+	{
+		FQueuedThreadPool* ThreadPool = nullptr;
+		int32 InnerTaskCount = 0;
+		void DoWork() 
+		{ 
+			TRACE_CPUPROFILER_EVENT_SCOPE(FOuterTask::DoWork);
+
+			bIsOuterTaskRunning = true;
+			TArray<FAsyncTask<FInnerTask>> InnerTasks;
+			InnerTasks.SetNum(InnerTaskCount);
+
+			for (int32 Index = 0; Index < InnerTasks.Num(); ++Index)
+			{
+				InnerTasks[Index].StartBackgroundTask(ThreadPool);
+			}
+
+			for (int32 Index = 0; Index < InnerTasks.Num(); ++Index)
+			{
+				TRACE_CPUPROFILER_EVENT_SCOPE(FOuterTask::EnsureCompletion);
+				InnerTasks[Index].EnsureCompletion();
+			}
+			bIsOuterTaskRunning = false;
+		}
+
+		FORCEINLINE TStatId GetStatId() const
+		{
+			RETURN_QUICK_DECLARE_CYCLE_STAT(FRootTask, STATGROUP_ThreadPoolAsyncTasks);
+		}
+	};
+
+	template <typename OtherTaskType>
+	struct FCompletionTask : public FNonAbandonableTask
+	{
+		OtherTaskType* OtherTask;
+		void DoWork() 
+		{ 
+			checkf(!bIsOuterTaskRunning, TEXT("We got picked by the busywait of the outer task we are going to wait on, we are now deadlocked"));
+
+			TRACE_CPUPROFILER_EVENT_SCOPE(FCompletionTask::DoWork);
+
+			// We need to wait a little bit before entering EnsureCompletion otherwise
+			// EnsureCompletion will busy wait into every completion task and there 
+			// will not be any completion tasks left for the outer task to pick up.
+			FPlatformProcess::Sleep(0.01f);
+
+			OtherTask->EnsureCompletion(); 
+		}
+
+		FORCEINLINE TStatId GetStatId() const
+		{
+			RETURN_QUICK_DECLARE_CYCLE_STAT(FCompletionTask, STATGROUP_ThreadPoolAsyncTasks);
+		}
+	};
+}
+
+bool FAsyncWorkEnsureCompletionBusyWaitDeadLockTest::RunTest(const FString& Parameters)
+{
+	using namespace AsyncWorkTestImpl;
+	TRACE_CPUPROFILER_EVENT_SCOPE(FAsyncWorkEnsureCompletionBusyWaitDeadLockTest::RunTest);
+	
+	FQueuedThreadPool* ThreadPool = GThreadPool;
+#if WITH_EDITOR
+	ThreadPool = GLargeThreadPool;
+#endif
+
+	FAsyncTask<FOuterTask> OuterTask;
+	OuterTask.GetTask().ThreadPool = ThreadPool;
+	OuterTask.GetTask().InnerTaskCount = ThreadPool->GetNumThreads() * 2 * 100;
+
+	TArray<FAsyncTask<FCompletionTask<FAsyncTask<FOuterTask>>>> CompletionTasks;
+	CompletionTasks.SetNum(ThreadPool->GetNumThreads());
+	for (int32 Index = 0; Index < CompletionTasks.Num(); ++Index)
+	{
+		CompletionTasks[Index].GetTask().OtherTask = &OuterTask;
+	}
+
+	// Start the task that all the other tasks will try to complete as fast as possible
+	OuterTask.StartBackgroundTask(ThreadPool, EQueuedWorkPriority::Normal, EQueuedWorkFlags::None);
+
+	// Wait until the outer task had time to schedule it's own tasks
+	FPlatformProcess::Sleep(0.1f);
+
+	// Now fill the thread pool with tasks that simply want to complete the first one
+	for (int32 Index = 0; Index < CompletionTasks.Num(); ++Index)
+	{
+		EQueuedWorkFlags QueuedWorkFlags = EQueuedWorkFlags::None;
+
+		// Busy waiting on our task could cause a circular dependency between tasks on the same thread, which cannot be resolved and results into a deadlock.
+		// We need to tell the scheduler we want to opt-out of being picked up by busy wait.
+		// Comment this line to repro the deadlock scenario on the low-level task graph (will not reproduce on the old task graph as busywait doesn't exists).
+		QueuedWorkFlags |= EQueuedWorkFlags::DoNotRunInsideBusyWait;
+
+		// Use Highest priority so we have a chance to run before all the inner tasks are finished
+		CompletionTasks[Index].StartBackgroundTask(ThreadPool, EQueuedWorkPriority::Highest, QueuedWorkFlags);
+	}
+
+	// Do not contribute to any work here, we normally tick slate for the game-thread to stay responsive while we're waiting on tasks to complete
+	for (int32 Index = 0; Index < CompletionTasks.Num(); ++Index)
+	{
+		CompletionTasks[Index].EnsureCompletion(false);
+	}
+	OuterTask.EnsureCompletion(false);
+
+	return true;
+}
+
+#endif //WITH_DEV_AUTOMATION_TESTS
diff --git a/Engine/Source/Runtime/Core/Public/Async/AsyncWork.h b/Engine/Source/Runtime/Core/Public/Async/AsyncWork.h
--- a/Engine/Source/Runtime/Core/Public/Async/AsyncWork.h
+++ b/Engine/Source/Runtime/Core/Public/Async/AsyncWork.h
@@ -227,6 +227,8 @@ class FAsyncTask
 	FQueuedThreadPool*	QueuedPool;
 	/** Current priority */
 	EQueuedWorkPriority Priority;
+	/** Current flags */
+	EQueuedWorkFlags Flags;
 	/** optional LLM tag */
 	LLM(const UE::LLMPrivate::FTagData* InheritedLLMTag);
 	/** Memory trace tag */
@@ -242,10 +244,15 @@ class FAsyncTask
 		DoneEvent = nullptr;
 	}
 
+	EQueuedWorkFlags GetQueuedWorkFlags() const override
+	{
+		return Flags;
+	}
+
 	/* Generic start function, not called directly
 		* @param bForceSynchronous if true, this job will be started synchronously, now, on this thread
 	**/
-	void Start(bool bForceSynchronous, FQueuedThreadPool* InQueuedPool, EQueuedWorkPriority InQueuedWorkPriority)
+	void Start(bool bForceSynchronous, FQueuedThreadPool* InQueuedPool, EQueuedWorkPriority InQueuedWorkPriority, EQueuedWorkFlags InQueuedWorkFlags)
 	{
 		FScopeCycleCounter Scope( Task.GetStatId(), true );
 		DECLARE_SCOPE_CYCLE_COUNTER( TEXT( "FAsyncTask::Start" ), STAT_FAsyncTask_Start, STATGROUP_ThreadPoolAsyncTasks );
@@ -259,6 +266,7 @@ class FAsyncTask
 		WorkNotFinishedCounter.Increment();
 		QueuedPool = InQueuedPool;
 		Priority = InQueuedWorkPriority;
+		Flags = InQueuedWorkFlags;
 		if (bForceSynchronous)
 		{
 			QueuedPool = 0;
@@ -350,6 +358,8 @@ class FAsyncTask
 	**/
 	void SyncCompletion()
 	{
+		TRACE_CPUPROFILER_EVENT_SCOPE(FAsyncTask::SyncCompletion);
+
 		FPlatformMisc::MemoryBarrier();
 		if (QueuedPool)
 		{
@@ -425,17 +435,17 @@ public:
 	* Run this task on this thread
 	* @param bDoNow if true then do the job now instead of at EnsureCompletion
 	**/
-	void StartSynchronousTask(EQueuedWorkPriority InQueuedWorkPriority = EQueuedWorkPriority::Normal)
+	void StartSynchronousTask(EQueuedWorkPriority InQueuedWorkPriority = EQueuedWorkPriority::Normal, EQueuedWorkFlags InQueuedWorkFlags = EQueuedWorkFlags::None)
 	{
-		Start(true, GThreadPool, InQueuedWorkPriority);
+		Start(true, GThreadPool, InQueuedWorkPriority, InQueuedWorkFlags);
 	}
 
 	/** 
 	* Queue this task for processing by the background thread pool
 	**/
-	void StartBackgroundTask(FQueuedThreadPool* InQueuedPool = GThreadPool, EQueuedWorkPriority InQueuedWorkPriority = EQueuedWorkPriority::Normal)
+	void StartBackgroundTask(FQueuedThreadPool* InQueuedPool = GThreadPool, EQueuedWorkPriority InQueuedWorkPriority = EQueuedWorkPriority::Normal, EQueuedWorkFlags InQueuedWorkFlags = EQueuedWorkFlags::None)
 	{
-		Start(false, InQueuedPool, InQueuedWorkPriority);
+		Start(false, InQueuedPool, InQueuedWorkPriority, InQueuedWorkFlags);
 	}
 
 	/** 
diff --git a/Engine/Source/Runtime/Core/Public/Misc/IQueuedWork.h b/Engine/Source/Runtime/Core/Public/Misc/IQueuedWork.h
--- a/Engine/Source/Runtime/Core/Public/Misc/IQueuedWork.h
+++ b/Engine/Source/Runtime/Core/Public/Misc/IQueuedWork.h
@@ -5,6 +5,26 @@
 #include "CoreTypes.h"
 #include "Templates/RefCounting.h"
 
+/** Special flags that can be associated with queued work. */
+enum class EQueuedWorkFlags : uint8
+{
+	None,
+
+	/**
+	 * Tells the scheduler if this task is allowed to run during another task's
+	 * busy wait. The default should be true for most case but it
+	 * is sometime useful to avoid it if this task is going to wait on another one,
+	 * and that other task busy waits, this could cause a cycle that could deadlock.
+	 * (i.e. T1 -> busywait -> picks T2 that then waits on T1 -> deadlock
+	 * In this case, we can decide that T2 should never be picked up by busy waits.
+	 */
+	 DoNotRunInsideBusyWait = (1 << 0),
+
+	 Count
+};
+
+ENUM_CLASS_FLAGS(EQueuedWorkFlags);
+
 /**
 * Interface for internal data of queued work objects.
 *
@@ -41,7 +61,6 @@ public:
 class IQueuedWork
 {
 public:
-
 	/**
 	 * This is where the real thread work is done. All work that is done for
 	 * this queued object should be done from within the call to this function.
@@ -56,6 +75,10 @@ public:
 	 */
 	virtual void Abandon() = 0;
 
+	/**
+	 * Returns any special work flags.
+	 */
+	virtual EQueuedWorkFlags GetQueuedWorkFlags() const { return EQueuedWorkFlags::None; }
 public:
 
 	/**
diff --git a/Engine/Source/Runtime/Core/Public/Misc/QueuedThreadPoolWrapper.h b/Engine/Source/Runtime/Core/Public/Misc/QueuedThreadPoolWrapper.h
--- a/Engine/Source/Runtime/Core/Public/Misc/QueuedThreadPoolWrapper.h
+++ b/Engine/Source/Runtime/Core/Public/Misc/QueuedThreadPoolWrapper.h
@@ -318,6 +318,7 @@ private:
 
 		LowLevelTasks::ETaskPriority TaskPriorityMapper[int(EQueuedWorkPriority::Count)] = { LowLevelTasks::ETaskPriority::High, LowLevelTasks::ETaskPriority::BackgroundHigh, LowLevelTasks::ETaskPriority::BackgroundNormal, LowLevelTasks::ETaskPriority::BackgroundLow, LowLevelTasks::ETaskPriority::BackgroundLow };
 		LowLevelTasks::ETaskPriority TaskPriority = TaskPriorityMapper[int(Priority)];
+		const bool bAllowBusyWaiting = (InQueuedWork->GetQueuedWorkFlags() & EQueuedWorkFlags::DoNotRunInsideBusyWait) == EQueuedWorkFlags::None;
 
 		QueuedWorkInternalData->Task.Init(TEXT("FQueuedLowLevelThreadPoolTask"), TaskPriority, [InQueuedWork]
 		{
@@ -329,7 +330,9 @@ private:
 			TaskCount.fetch_sub(1, std::memory_order_release);
 			bool bWakeUpWorker = false;
 			ScheduleTasks(bWakeUpWorker);
-		});
+		},
+		bAllowBusyWaiting
+		);
 
 		if (!bIsPaused)
 		{
diff --git a/Engine/Source/Runtime/Engine/Private/AsyncCompilationHelpers.cpp b/Engine/Source/Runtime/Engine/Private/AsyncCompilationHelpers.cpp
--- a/Engine/Source/Runtime/Engine/Private/AsyncCompilationHelpers.cpp
+++ b/Engine/Source/Runtime/Engine/Private/AsyncCompilationHelpers.cpp
@@ -129,6 +129,17 @@ namespace AsyncCompilationHelpers
 			}
 
 			void Abandon() override { }
+
+			// Since we're specifically waiting on other tasks here, we can't be run inside a busywait
+			// without risking a deadlock.
+			//
+			// (i.e. T1 -> busywait -> picks T2 that then waits on T1 -> deadlock
+			//
+			// Imagine we're T2 and we're going to wait on T1 here... 
+			// if we get picked up on the same thread inside T1's busy wait, we'll deadlock.
+			//
+			// See FAsyncWorkEnsureCompletionBusyWaitDeadLockTest for a 100% repro case
+			EQueuedWorkFlags GetQueuedWorkFlags() const override { return EQueuedWorkFlags::DoNotRunInsideBusyWait; }
 		};
 
 		// Perform forced compilation on as many thread as possible in high priority since the game-thread is waiting
