U5 - implement the ability to set and propagate Task UserData

diff --git a/Engine/Source/Runtime/Core/Private/Async/Fundamental/Scheduler.cpp b/Engine/Source/Runtime/Core/Private/Async/Fundamental/Scheduler.cpp
--- a/Engine/Source/Runtime/Core/Private/Async/Fundamental/Scheduler.cpp
+++ b/Engine/Source/Runtime/Core/Private/Async/Fundamental/Scheduler.cpp
@@ -185,6 +185,19 @@ namespace LowLevelTasks
 		return WorkerType != EWorkerType::None && ActiveScheduler == this;
 	}
 
+	void FTask::PropagateUserData()
+	{
+		const FTask* ActiveTask = FScheduler::GetActiveTask();
+		if (ActiveTask != nullptr)
+		{
+			UserData = ActiveTask->GetUserData();
+		}
+		else
+		{
+			UserData = nullptr;
+		}
+	}
+
 	template<FTask* (FScheduler::FLocalQueueType::*DequeueFunction)(bool), bool bIsBusyWaiting>
 	bool FScheduler::TryExecuteTaskFrom(FLocalQueueType* Queue, FQueueRegistry::FOutOfWork& OutOfWork, bool bPermitBackgroundWork)
 	{
diff --git a/Engine/Source/Runtime/Core/Public/Async/Fundamental/Task.h b/Engine/Source/Runtime/Core/Public/Async/Fundamental/Task.h
--- a/Engine/Source/Runtime/Core/Public/Async/Fundamental/Task.h
+++ b/Engine/Source/Runtime/Core/Public/Async/Fundamental/Task.h
@@ -73,9 +73,9 @@ namespace LowLevelTasks
 
 		class FPackedData
 		{
-			uintptr_t DebugName			: 56;
+			uintptr_t DebugName			: 57;
 			uintptr_t Priority			: 3;
-			uintptr_t State				: 4;
+			uintptr_t State				: 3;
 			uintptr_t bAllowBusyWaiting : 1;
 			
 		public:
@@ -83,7 +83,7 @@ namespace LowLevelTasks
 			{
 				static_assert(!PLATFORM_32BITS, "32bit Platforms are not supported");
 				static_assert(uintptr_t(ETaskPriority::Count) <= (1ull << 3), "Not enough bits to store ETaskPriority");
-				static_assert(uintptr_t(ETaskState::Count) <= (1ull << 4), "Not enough bits to store ETaskState");
+				static_assert(uintptr_t(ETaskState::Count) <= (1ull << 3), "Not enough bits to store ETaskState");
 			}
 
 			FPackedData(const TCHAR* InDebugName, ETaskPriority InPriority, ETaskState InState, bool bInAllowBusyWaiting)
@@ -94,7 +94,7 @@ namespace LowLevelTasks
 			{
 				checkSlow(reinterpret_cast<uintptr_t>(InDebugName) < (1ull << 57));
 				checkSlow((uintptr_t)InPriority < (1ull << 3));
-				checkSlow((uintptr_t)InState < (1ull << 4));
+				checkSlow((uintptr_t)InState < (1ull << 3));
 				static_assert(sizeof(FPackedData) == sizeof(uintptr_t), "Packed data needs to be pointer size");
 			}
 
@@ -125,8 +125,9 @@ namespace LowLevelTasks
 		};
 
 	private:
-		using FTaskDelegate = TTaskDelegate<PLATFORM_CACHE_LINE_SIZE - sizeof(FPackedData), bool>;
+		using FTaskDelegate = TTaskDelegate<PLATFORM_CACHE_LINE_SIZE - sizeof(FPackedData) - sizeof(void*), bool>;
 		FTaskDelegate Runnable;
+		mutable void* UserData;
 		std::atomic<FPackedData> PackedData { FPackedData() };
 
 	private:
@@ -203,6 +204,9 @@ namespace LowLevelTasks
 		inline ETaskPriority GetPriority() const;
 		inline bool IsBackgroundTask() const;
 
+		void* GetUserData() const { return UserData; }
+		void SetUserData(void* NewUserData) const { UserData = NewUserData; }
+
 	public:
 		FTask()
 		{
@@ -214,6 +218,7 @@ namespace LowLevelTasks
 		inline bool TryPrepareLaunch();
 		//after calling this function the task can be considered dead
 		inline void ExecuteTask();
+		CORE_API void PropagateUserData();
 		inline bool AllowBusyWaiting() const;
 	};
 
@@ -244,6 +249,7 @@ namespace LowLevelTasks
 				return true;
 			}
 		};
+		PropagateUserData();
 		PackedData.store(FPackedData(InDebugName, InPriority, ETaskState::Ready, bAllowBusyWaiting), std::memory_order_release);
 	}
 
@@ -268,6 +274,7 @@ namespace LowLevelTasks
 				return true;
 			}
 		};
+		PropagateUserData();
 		PackedData.store(FPackedData(InDebugName, InPriority, ETaskState::Ready, bAllowBusyWaiting), std::memory_order_release);
 	}
 
