Call LowLevelOutputDebugStringf instead of UE_LOG because InitializeConfigSystem() is too early to call UE_LOG.

GConfig->GetInt is called too early, before the config system is initialized. Some threads created during PreInitPreStartupScreen won't get the stack size set in the config.

diff --git a/Engine/Source/Runtime/Core/Private/Misc/ConfigCacheIni.cpp b/Engine/Source/Runtime/Core/Private/Misc/ConfigCacheIni.cpp
--- a/Engine/Source/Runtime/Core/Private/Misc/ConfigCacheIni.cpp
+++ b/Engine/Source/Runtime/Core/Private/Misc/ConfigCacheIni.cpp
@@ -4124,11 +4124,11 @@ void FConfigCacheIni::InitializeConfigSystem()
 		}
 		else
 		{
-			UE_LOG(LogInit, Display, TEXT("Unable to bootstrap from archive %s, will fallback on normal initialization"), *IniBootstrapFilename);
+			FPlatformMisc::LowLevelOutputDebugStringf(TEXT("Unable to bootstrap from archive %s, will fallback on normal initialization\n"), *IniBootstrapFilename);
 		}
 	}
 
-	UE_LOG(LogInit, Display, TEXT("Loading text-based GConfig..."));
+	FPlatformMisc::LowLevelOutputDebugStringf(TEXT("Loading text-based GConfig...\n"));
 
 	// Perform any upgrade we need before we load any configuration files
 	FConfigManifest::UpgradeFromPreviousVersions();
