Improve detection of slow running unit tests

* Sum execution time of tests instead of using wall time to remove FAutomationTestFramework overhead
* Don't emit warning in debug builds or when a debugger is attached.

diff --git a/Engine/Source/Runtime/Core/Private/Misc/AutomationTest.cpp b/Engine/Source/Runtime/Core/Private/Misc/AutomationTest.cpp
--- a/Engine/Source/Runtime/Core/Private/Misc/AutomationTest.cpp
+++ b/Engine/Source/Runtime/Core/Private/Misc/AutomationTest.cpp
@@ -220,6 +220,34 @@ bool FAutomationTestFramework::ContainsTest( const FString& InTestName ) const
 	return AutomationTestClassNameToInstanceMap.Contains( InTestName );
 }
 
+static double SumDurations(const TMap<FString, FAutomationTestExecutionInfo>& Executions)
+{
+	double Sum = 0;
+	for (const TPair<FString, FAutomationTestExecutionInfo>& Execution : Executions)
+	{
+		Sum += Execution.Value.Duration;
+	}
+	return Sum;
+}
+
+static const TCHAR* FindSlowestTest(const TMap<FString, FAutomationTestExecutionInfo>& Executions, double& OutMaxDuration)
+{
+	check(Executions.Num() > 0);
+
+	const TCHAR* OutName = nullptr;
+	OutMaxDuration = 0;
+	for (const TPair<FString, FAutomationTestExecutionInfo>& Execution : Executions)
+	{
+		if (OutMaxDuration <= Execution.Value.Duration)
+		{
+			OutMaxDuration = Execution.Value.Duration;
+			OutName = *Execution.Key;
+		}
+	}
+
+	return OutName;
+}
+
 bool FAutomationTestFramework::RunSmokeTests()
 {
 	TRACE_CPUPROFILER_EVENT_SCOPE(FAutomationTestFramework::RunSmokeTests);
@@ -242,8 +270,6 @@ bool FAutomationTestFramework::RunSmokeTests()
 
 		if ( TestInfo.Num() > 0 )
 		{
-			const double SmokeTestStartTime = FPlatformTime::Seconds();
-
 			// Output the results of running the automation tests
 			TMap<FString, FAutomationTestExecutionInfo> OutExecutionInfoMap;
 
@@ -253,8 +279,6 @@ bool FAutomationTestFramework::RunSmokeTests()
 			// We disable capturing the stack when running smoke tests, it adds too much overhead to do it at startup.
 			FAutomationTestFramework::Get().SetCaptureStack(false);
 
-			double SlowestTestDuration = 0.0f;
-			FString SlowestTestName;
 			for ( int TestIndex = 0; TestIndex < TestInfo.Num(); ++TestIndex )
 			{
 				SlowTask.EnterProgressFrame(1);
@@ -268,25 +292,23 @@ bool FAutomationTestFramework::RunSmokeTests()
 					const bool CurTestSuccessful = StopTest(CurExecutionInfo);
 
 					bAllSuccessful = bAllSuccessful && CurTestSuccessful;
-
-					if (CurTestSuccessful && CurExecutionInfo.Duration > SlowestTestDuration)
-					{
-						SlowestTestDuration = CurExecutionInfo.Duration;
-						SlowestTestName = MoveTemp(TestCommand);
-					}
 				}
 			}
 
 			FAutomationTestFramework::Get().SetCaptureStack(true);
 
-			const double TimeForTest = FPlatformTime::Seconds() - SmokeTestStartTime;
-			if (TimeForTest > 2.0f)
+#if !UE_BUILD_DEBUG
+			const double TotalDuration = SumDurations(OutExecutionInfoMap);
+			if (bAllSuccessful && !FPlatformMisc::IsDebuggerPresent() && TotalDuration > 2.0)
 			{
 				//force a failure if a smoke test takes too long
+				double SlowestTestDuration = 0;
+				const TCHAR* SlowestTestName = FindSlowestTest(OutExecutionInfoMap, /* out */ SlowestTestDuration);
 				UE_LOG(LogAutomationTest, Warning, TEXT("Smoke tests took >2s to run (%.2fs). '%s' took %dms. "
 					"SmokeFilter tier tests should take less than 1ms. Please optimize or move '%s' to a slower tier than SmokeFilter."), 
-					TimeForTest, *SlowestTestName, static_cast<int32>(1000*SlowestTestDuration), *SlowestTestName);
+					TotalDuration, SlowestTestName, static_cast<int32>(1000*SlowestTestDuration), SlowestTestName);
 			}
+#endif
 
 			FAutomationTestFramework::DumpAutomationTestExecutionInfo( OutExecutionInfoMap );
 		}
