Fix for android Non-Unity compile error

diff --git a/Engine/Source/Runtime/Core/Private/Internationalization/FastDecimalFormat.cpp b/Engine/Source/Runtime/Core/Private/Internationalization/FastDecimalFormat.cpp
--- a/Engine/Source/Runtime/Core/Private/Internationalization/FastDecimalFormat.cpp
+++ b/Engine/Source/Runtime/Core/Private/Internationalization/FastDecimalFormat.cpp
@@ -3,7 +3,7 @@
 #include "Internationalization/FastDecimalFormat.h"
 #include "Misc/EnumClassFlags.h"
 #include "Misc/StringBuilder.h"
-#include "Containers/StringFwd.h"
+#include "HAL/IConsoleManager.h"
 
 namespace FastDecimalFormat
 {
