Make sure the windows error output device print the message that triggered the error
Use a similar approach as the mac and unix version

diff --git a/Engine/Source/Runtime/Core/Private/Windows/WindowsErrorOutputDevice.cpp b/Engine/Source/Runtime/Core/Private/Windows/WindowsErrorOutputDevice.cpp
--- a/Engine/Source/Runtime/Core/Private/Windows/WindowsErrorOutputDevice.cpp
+++ b/Engine/Source/Runtime/Core/Private/Windows/WindowsErrorOutputDevice.cpp
@@ -49,6 +49,9 @@ void FWindowsErrorOutputDevice::Serialize( const TCHAR* Msg, ELogVerbosity::Type
 		{
 			UE_LOG(LogWindows, Error, TEXT("Windows GetLastError: %s (%i)"), FPlatformMisc::GetSystemErrorMessage(ErrorBuffer, 1024, LastError), LastError);
 		}
+
+		FCString::Strncpy(GErrorHist, Msg, UE_ARRAY_COUNT(GErrorHist) - 5);
+		FCString::Strncat(GErrorHist, TEXT("\r\n\r\n"), UE_ARRAY_COUNT(GErrorHist) - 1);
 	}
 	else
 	{
