Fix invalid buffer access static analysis CIS

diff --git a/Engine/Source/Runtime/Core/Private/Internationalization/FastDecimalFormat.cpp b/Engine/Source/Runtime/Core/Private/Internationalization/FastDecimalFormat.cpp
--- a/Engine/Source/Runtime/Core/Private/Internationalization/FastDecimalFormat.cpp
+++ b/Engine/Source/Runtime/Core/Private/Internationalization/FastDecimalFormat.cpp
@@ -461,9 +461,13 @@ FString CultureInvariantDecimalToString(const double InVal, const TCHAR*& InBuff
 		}
 
 		// 48 for raw ascii -> int
-		OutStr += InFormattingRules.DigitCharacters[(int32)InBuffer[0] - 48];
-		FractionalDigitsPrinted += bParsedFractional ? 1 : 0;
-		++InBuffer;
+		int32 CharacterIndex = (int32)InBuffer[0] - 48;
+		if (ensure(CharacterIndex >= 0 && CharacterIndex < sizeof(InFormattingRules.DigitCharacters) / sizeof(InFormattingRules.DigitCharacters[0])))
+		{
+			OutStr += InFormattingRules.DigitCharacters[CharacterIndex];
+			FractionalDigitsPrinted += bParsedFractional ? 1 : 0;
+			++InBuffer;
+		}
 	}
 
 	// Apply back padding, if back isn't just zero
