DDC: Added CreateOrGetDDC to the module and changed creation of DDC to occur there instead of on module startup

diff --git a/Engine/Source/Developer/DerivedDataCache/Private/DerivedDataCache.cpp b/Engine/Source/Developer/DerivedDataCache/Private/DerivedDataCache.cpp
--- a/Engine/Source/Developer/DerivedDataCache/Private/DerivedDataCache.cpp
+++ b/Engine/Source/Developer/DerivedDataCache/Private/DerivedDataCache.cpp
@@ -667,40 +667,47 @@ private:
 	TUniquePtr<UE::DerivedData::ICache> Cache;
 };
 
+static FDerivedDataCacheInterface* GDerivedDataCacheInstance;
+
+PRAGMA_DISABLE_DEPRECATION_WARNINGS
+
 /**
  * Module for the DDC
  */
-class FDerivedDataCacheModule : public IDerivedDataCacheModule
+class FDerivedDataCacheModule final : public IDerivedDataCacheModule
 {
 public:
-	virtual FDerivedDataCacheInterface& GetDDC() override
+	virtual FDerivedDataCacheInterface& GetDDC() final
 	{
-		check(Instance);
-		return *Instance;
+		return **CreateOrGetDDC();
 	}
 
-	virtual void StartupModule() override
+	virtual FDerivedDataCacheInterface* const* CreateOrGetDDC() final
 	{
-		if (!Instance)
+		FScopeLock Lock(&CreateLock);
+		if (!GDerivedDataCacheInstance)
 		{
-			Instance = MakeUnique<FDerivedDataCache>();
+			GDerivedDataCacheInstance = new FDerivedDataCache();
+			check(GDerivedDataCacheInstance);
 		}
-		check(Instance);
+		return &GDerivedDataCacheInstance;
 	}
 
-	virtual void ShutdownModule() override
+	virtual void ShutdownModule() final
 	{
 		FDDCCleanup::Shutdown();
 
-		check(Instance);
-		Instance.Release();
+		delete GDerivedDataCacheInstance;
+		GDerivedDataCacheInstance = nullptr;
 	}
 
 private:
-	TUniquePtr<FDerivedDataCache> Instance;
+	FCriticalSection CreateLock;
 };
 
-IMPLEMENT_MODULE( FDerivedDataCacheModule, DerivedDataCache);
+PRAGMA_ENABLE_DEPRECATION_WARNINGS
+
+IMPLEMENT_MODULE(FDerivedDataCacheModule, DerivedDataCache);
 
 ///////////////////////////////////////////////////////////////////////////////////////////////////////////////////////
 
diff --git a/Engine/Source/Developer/DerivedDataCache/Public/DerivedDataCacheInterface.h b/Engine/Source/Developer/DerivedDataCache/Public/DerivedDataCacheInterface.h
--- a/Engine/Source/Developer/DerivedDataCache/Public/DerivedDataCacheInterface.h
+++ b/Engine/Source/Developer/DerivedDataCache/Public/DerivedDataCacheInterface.h
@@ -295,6 +295,15 @@ class IDerivedDataCacheModule : public IModuleInterface
 {
 public:
 	/** Return the DDC interface **/
+	UE_DEPRECATED(4.27, "GetDDC has been replaced by CreateOrGetDDC.")
 	virtual FDerivedDataCacheInterface& GetDDC() = 0;
-};
 
+	/**
+	 * Returns the DDC interface, and creates it on the first call.
+	 *
+	 * This will never return a null interface, but an extra level of indirection is used to support
+	 * callers that store the return value. The returned interface pointer will become null when the
+	 * module shuts down, which destroys DDC.
+	 */
+	virtual FDerivedDataCacheInterface* const* CreateOrGetDDC() = 0;
+};
diff --git a/Engine/Source/Runtime/Core/Private/Misc/CoreMisc.cpp b/Engine/Source/Runtime/Core/Private/Misc/CoreMisc.cpp
--- a/Engine/Source/Runtime/Core/Private/Misc/CoreMisc.cpp
+++ b/Engine/Source/Runtime/Core/Private/Misc/CoreMisc.cpp
@@ -89,35 +89,34 @@ bool FStaticSelfRegisteringExec::Exec( UWorld* InWorld, const TCHAR* Cmd, FOutpu
 	Module singletons.
 -----------------------------------------------------------------------------*/
 
-class FDerivedDataCacheInterface* GetDerivedDataCache()
+FDerivedDataCacheInterface* GetDerivedDataCache()
 {
-	static class FDerivedDataCacheInterface* SingletonInterface = NULL;
-	if (!FPlatformProperties::RequiresCookedData())
+	static FDerivedDataCacheInterface* const* DDC;
+	static bool bInitialized = false;
+	if (!bInitialized)
 	{
-		static bool bInitialized = false;
-		if (!bInitialized)
+		if (!FPlatformProperties::RequiresCookedData())
 		{
 			check(IsInGameThread());
 			bInitialized = true;
-			class IDerivedDataCacheModule* Module = FModuleManager::LoadModulePtr<IDerivedDataCacheModule>("DerivedDataCache");
-			if (Module)
+			if (IDerivedDataCacheModule* Module = FModuleManager::LoadModulePtr<IDerivedDataCacheModule>("DerivedDataCache"))
 			{
-				SingletonInterface = &Module->GetDDC();
+				DDC = Module->CreateOrGetDDC();
 			}
 		}
 	}
-	return SingletonInterface;
+	return *DDC;
 }
 
-class FDerivedDataCacheInterface& GetDerivedDataCacheRef()
+FDerivedDataCacheInterface& GetDerivedDataCacheRef()
 {
-	class FDerivedDataCacheInterface* SingletonInterface = GetDerivedDataCache();
-	if (!SingletonInterface)
+	FDerivedDataCacheInterface* DDC = GetDerivedDataCache();
+	if (!DDC)
 	{
 		UE_LOG(LogInit, Fatal, TEXT("Derived Data Cache was requested, but not available."));
-		CA_ASSUME( SingletonInterface != NULL );	// Suppress static analysis warning in unreachable code (fatal error)
+		CA_ASSUME(DDC); // Suppress static analysis warning in unreachable code (fatal error)
 	}
-	return *SingletonInterface;
+	return *DDC;
 }
 
 class ITargetPlatformManagerModule* GetTargetPlatformManager(bool bFailOnInitErrors)
