DDC: Fixed a crash in builds that access DDC when it is disabled

#rnx

diff --git a/Engine/Source/Runtime/Core/Private/Misc/CoreMisc.cpp b/Engine/Source/Runtime/Core/Private/Misc/CoreMisc.cpp
--- a/Engine/Source/Runtime/Core/Private/Misc/CoreMisc.cpp
+++ b/Engine/Source/Runtime/Core/Private/Misc/CoreMisc.cpp
@@ -105,7 +105,7 @@ FDerivedDataCacheInterface* GetDerivedDataCache()
 			}
 		}
 	}
-	return *DDC;
+	return DDC ? *DDC : nullptr;
 }
 
 FDerivedDataCacheInterface& GetDerivedDataCacheRef()
