const correction for BlendFromIdentityAndAccumulate, for both (non-)vectorized versions

diff --git a/Engine/Source/Runtime/Core/Public/Math/TransformNonVectorized.h b/Engine/Source/Runtime/Core/Public/Math/TransformNonVectorized.h
--- a/Engine/Source/Runtime/Core/Public/Math/TransformNonVectorized.h
+++ b/Engine/Source/Runtime/Core/Public/Math/TransformNonVectorized.h
@@ -1032,34 +1032,35 @@ public:
 	/**
 	* Blends the Identity transform with a weighted source transform and accumulates that into a destination transform
 	*
-	* SourceAtom = Blend(Identity, SourceAtom, BlendWeight)
-	* FinalAtom.Rotation = SourceAtom.Rotation * FinalAtom.Rotation
-	* FinalAtom.Translation += SourceAtom.Translation
-	* FinalAtom.Scale3D *= SourceAtom.Scale3D
+	* DeltaAtom = Blend(Identity, SourceAtom, BlendWeight)
+	* FinalAtom.Rotation = DeltaAtom.Rotation * FinalAtom.Rotation
+	* FinalAtom.Translation += DeltaAtom.Translation
+	* FinalAtom.Scale3D *= DeltaAtom.Scale3D
 	*
 	* @param FinalAtom [in/out] The atom to accumulate the blended source atom into
 	* @param SourceAtom The target transformation (used when BlendWeight = 1); this is modified during the process
 	* @param BlendWeight The blend weight between Identity and SourceAtom
 	*/
-	FORCEINLINE static void BlendFromIdentityAndAccumulate(FTransform& FinalAtom, FTransform& SourceAtom, float BlendWeight)
+	FORCEINLINE static void BlendFromIdentityAndAccumulate(FTransform& FinalAtom, const FTransform& SourceAtom, float BlendWeight)
 	{
-		const  FTransform AdditiveIdentity(FQuat::Identity, FVector::ZeroVector, FVector::ZeroVector);
+		const FTransform AdditiveIdentity(FQuat::Identity, FVector::ZeroVector, FVector::ZeroVector);
 		const FVector DefaultScale(FVector::OneVector);
+		FTransform DeltaAtom = SourceAtom;
 
 		// Scale delta by weight
 		if (BlendWeight < (1.f - ZERO_ANIMWEIGHT_THRESH))
 		{
-			SourceAtom.Blend(AdditiveIdentity, SourceAtom, BlendWeight);
+			DeltaAtom.Blend(AdditiveIdentity, DeltaAtom, BlendWeight);
 		}
 
 		// Add ref pose relative animation to base animation, only if rotation is significant.
-		if (FMath::Square(SourceAtom.Rotation.W) < 1.f - DELTA * DELTA)
+		if (FMath::Square(DeltaAtom.Rotation.W) < 1.f - DELTA * DELTA)
 		{
-			FinalAtom.Rotation = SourceAtom.Rotation * FinalAtom.Rotation;
+			FinalAtom.Rotation = DeltaAtom.Rotation * FinalAtom.Rotation;
 		}
 
-		FinalAtom.Translation += SourceAtom.Translation;
-		FinalAtom.Scale3D *= (DefaultScale + SourceAtom.Scale3D);
+		FinalAtom.Translation += DeltaAtom.Translation;
+		FinalAtom.Scale3D *= (DefaultScale + DeltaAtom.Scale3D);
 
 		checkSlow(FinalAtom.IsRotationNormalized());
 	}
diff --git a/Engine/Source/Runtime/Core/Public/Math/TransformVectorized.h b/Engine/Source/Runtime/Core/Public/Math/TransformVectorized.h
--- a/Engine/Source/Runtime/Core/Public/Math/TransformVectorized.h
+++ b/Engine/Source/Runtime/Core/Public/Math/TransformVectorized.h
@@ -1098,7 +1098,7 @@ public:
 	 * @param SourceAtom The target transformation (used when BlendWeight = 1)
 	 * @param Alpha The blend weight between Identity and SourceAtom
 	 */
-	FORCEINLINE static void BlendFromIdentityAndAccumulate(FTransform& FinalAtom, FTransform& SourceAtom, const ScalarRegister& BlendWeight)
+	FORCEINLINE static void BlendFromIdentityAndAccumulate(FTransform& FinalAtom, const FTransform& SourceAtom, const ScalarRegister& BlendWeight)
 	{
 		const VectorRegister Const0001 = GlobalVectorConstants::Float0001;
 		const VectorRegister ConstNegative0001 = VectorSubtract(VectorZero(), Const0001);
