When capturing and reporting callstacks, use the return address of a failure instead of a count of stack frames to trim. The count approach was spread about in many places and fragile to maintain as code changed. This resulted in "noisy" callstacks with distracting boilerplate present like assert dispatch functions.

#rnx

diff --git a/Engine/Source/Programs/CrashReportClient/Private/CrashReportClientApp.cpp b/Engine/Source/Programs/CrashReportClient/Private/CrashReportClientApp.cpp
--- a/Engine/Source/Programs/CrashReportClient/Private/CrashReportClientApp.cpp
+++ b/Engine/Source/Programs/CrashReportClient/Private/CrashReportClientApp.cpp
@@ -451,7 +451,7 @@ FPlatformErrorReport CollectErrorReport(FRecoveryService* RecoveryService, uint3
 	FPlatformCrashContext CrashContext(SharedCrashContext.CrashType, SharedCrashContext.ErrorMessage);
 	CrashContext.SetCrashedProcess(ProcessHandle);
 	CrashContext.SetCrashedThreadId(SharedCrashContext.CrashingThreadId);
-	CrashContext.SetNumMinidumpFramesToIgnore(SharedCrashContext.NumStackFramesToIgnore);
+	CrashContext.SetNumMinidumpFramesToIgnore(0);
 
 	// Initialize the stack walking for the monitored process (effectively overriding this process stack walking functionality)
 	FPlatformStackWalk::InitStackWalkingForProcess(ProcessHandle);
@@ -489,7 +489,7 @@ FPlatformErrorReport CollectErrorReport(FRecoveryService* RecoveryService, uint3
 #endif
 
 		uint64 StackFrames[CR_MAX_STACK_FRAMES] = {0};
-		const uint32 StackFrameCount = FPlatformStackWalk::CaptureThreadStackBackTrace(
+		uint32 StackFrameCount = FPlatformStackWalk::CaptureThreadStackBackTrace(
 			ThreadId, 
 			StackFrames,
 			CR_MAX_STACK_FRAMES,
@@ -506,10 +506,38 @@ FPlatformErrorReport CollectErrorReport(FRecoveryService* RecoveryService, uint3
 		// Add the crashing stack specifically. Is this really needed?
 		if (ThreadId == SharedCrashContext.CrashingThreadId)
 		{
-			CrashContext.SetPortableCallStack(
-				StackFrames,
-				StackFrameCount - SharedCrashContext.NumStackFramesToIgnore
-			);
+			const uint64* StackFrameCursor = StackFrames;
+
+			// If the address of where the error occurred has been provided then
+			// we can remove the boilerplate noise from the callstack.
+			if (uint64 ErrorPc = uint64(SharedCrashContext.ErrorProgramCounter))
+			{
+				uint64 ExceptionPc = uint64(SharedCrashContext.ExceptionProgramCounter);
+				int32 ExceptionDepth = -1;
+				for (uint32 i = 0; i < StackFrameCount; ++i)
+				{
+					if (StackFrames[i] == ExceptionPc)
+					{
+						ExceptionDepth = i;
+					}
+
+					if (StackFrames[i] != ErrorPc)
+					{
+						continue;
+					}
+
+					if (ExceptionDepth >= 0)
+					{
+						CrashContext.SetNumMinidumpFramesToIgnore(i - ExceptionDepth);
+					}
+
+					StackFrameCursor = StackFrames + i;
+					StackFrameCount -= i;
+					break;
+				}
+			}
+
+			CrashContext.SetPortableCallStack(StackFrameCursor, StackFrameCount);
 
 			// A completely missing portable callstack usually means that the crashing process died before CRC could walk the stack.
 			bOutCrashPortableCallstackAvailable = StackFrameCount > 0;
diff --git a/Engine/Source/Programs/CrashReportClient/Private/Linux/CrashReportClientMainLinux.cpp b/Engine/Source/Programs/CrashReportClient/Private/Linux/CrashReportClientMainLinux.cpp
--- a/Engine/Source/Programs/CrashReportClient/Private/Linux/CrashReportClientMainLinux.cpp
+++ b/Engine/Source/Programs/CrashReportClient/Private/Linux/CrashReportClientMainLinux.cpp
@@ -23,7 +23,7 @@ void CrashReporterCrashHandler(const FGenericCrashContext& GenericContext)
 	const FUnixCrashContext& Context = static_cast< const FUnixCrashContext& >( GenericContext );
 
 	printf("CrashHandler: Signal=%d\n", Context.Signal);
-	const_cast< FUnixCrashContext& >(Context).CaptureStackTrace();
+	const_cast< FUnixCrashContext& >(Context).CaptureStackTrace(Context.FirstCrashHandlerFrame);
 	if (GLog)
 	{
 		GLog->Flush();
diff --git a/Engine/Source/Runtime/Core/Private/Android/AndroidPlatformCrashContext.cpp b/Engine/Source/Runtime/Core/Private/Android/AndroidPlatformCrashContext.cpp
--- a/Engine/Source/Runtime/Core/Private/Android/AndroidPlatformCrashContext.cpp
+++ b/Engine/Source/Runtime/Core/Private/Android/AndroidPlatformCrashContext.cpp
@@ -306,7 +306,7 @@ const TCHAR* FAndroidCrashContext::GetCallstackProperty() const
 
 void FAndroidCrashContext::CaptureCrashInfo()
 {
-	CapturePortableCallStack(0, Context);
+	CapturePortableCallStack(nullptr, Context);
 }
 
 static int64 GetAndroidLibraryBaseAddress()
diff --git a/Engine/Source/Runtime/Core/Private/Apple/ApplePlatformCrashContext.cpp b/Engine/Source/Runtime/Core/Private/Apple/ApplePlatformCrashContext.cpp
--- a/Engine/Source/Runtime/Core/Private/Apple/ApplePlatformCrashContext.cpp
+++ b/Engine/Source/Runtime/Core/Private/Apple/ApplePlatformCrashContext.cpp
@@ -24,7 +24,6 @@ FApplePlatformCrashContext::FApplePlatformCrashContext(ECrashContextType InType,
 ,	Signal(0)
 ,	Info(nullptr)
 ,	Context(nullptr)
-,	IgnoreDepth(6)
 {
 	SignalDescription[ 0 ] = 0;
 	MinidumpCallstackInfo[ 0 ] = 0;
@@ -39,6 +38,7 @@ void FApplePlatformCrashContext::InitFromSignal(int32 InSignal, siginfo_t* InInf
 {
 	Signal = InSignal;
 	Info = InInfo;
+	ErrorFrame = InInfo->si_addr;
 	Context = reinterpret_cast< ucontext_t* >( InContext );
 	
 #define HANDLE_CASE(a,b) case a: FCStringAnsi::Strncpy(SignalDescription, #a ": " b, UE_ARRAY_COUNT( SignalDescription ) - 1); break;
@@ -86,9 +86,9 @@ int32 FApplePlatformCrashContext::ReportCrash() const
 		*StackBuffer = 0;
 		
 		// Generate the portable callstack for the crash xml
-		const_cast<FApplePlatformCrashContext*>(this)->CapturePortableCallStack(IgnoreDepth, Context);
+		const_cast<FApplePlatformCrashContext*>(this)->CapturePortableCallStack(ErrorFrame, Context);
 		// Walk the stack and dump it to the allocated memory (ignore first 2 callstack lines as those are in stack walking code)
-		FPlatformStackWalk::StackWalkAndDump( StackBuffer, UE_ARRAY_COUNT(MinidumpCallstackInfo) - 1, IgnoreDepth, Context );
+		FPlatformStackWalk::StackWalkAndDump( StackBuffer, UE_ARRAY_COUNT(MinidumpCallstackInfo) - 1, ErrorFrame, Context );
 		
 		FUTF8ToTCHAR_Convert::Convert(GErrorHist, UE_ARRAY_COUNT(GErrorHist) - 1, MinidumpCallstackInfo, FCStringAnsi::Strlen(MinidumpCallstackInfo));
 		CreateExceptionInfoString(Signal, Info);
diff --git a/Engine/Source/Runtime/Core/Private/Apple/ApplePlatformStackWalk.cpp b/Engine/Source/Runtime/Core/Private/Apple/ApplePlatformStackWalk.cpp
--- a/Engine/Source/Runtime/Core/Private/Apple/ApplePlatformStackWalk.cpp
+++ b/Engine/Source/Runtime/Core/Private/Apple/ApplePlatformStackWalk.cpp
@@ -446,7 +446,7 @@ int32 ReportCrash(ucontext_t *Context, int32 Signal, struct __siginfo* Info)
 		StackTrace[0] = 0;
 		
 		// Walk the stack and dump it to the allocated memory.
-		FPlatformStackWalk::StackWalkAndDump( StackTrace, StackTraceSize, 0, Context );
+		FPlatformStackWalk::StackWalkAndDump( StackTrace, StackTraceSize, Info->si_addr, Context );
 #if WITH_EDITORONLY_DATA
         FCString::Strncat( GErrorHist, ANSI_TO_TCHAR(StackTrace), UE_ARRAY_COUNT(GErrorHist) - 1 );
 		CreateExceptionInfoString(Signal, Info);
diff --git a/Engine/Source/Runtime/Core/Private/GenericPlatform/GenericPlatformCrashContext.cpp b/Engine/Source/Runtime/Core/Private/GenericPlatform/GenericPlatformCrashContext.cpp
--- a/Engine/Source/Runtime/Core/Private/GenericPlatform/GenericPlatformCrashContext.cpp
+++ b/Engine/Source/Runtime/Core/Private/GenericPlatform/GenericPlatformCrashContext.cpp
@@ -1092,25 +1092,38 @@ void FGenericCrashContext::DumpLog(const FString& CrashFolderAbsolute)
 	
 }
 
-FORCENOINLINE void FGenericCrashContext::CapturePortableCallStack(int32 NumStackFramesToIgnore, void* Context)
+void FGenericCrashContext::CapturePortableCallStack(void* ErrorProgramCounter, void* Context)
 {
-	// If the callstack is for the executing thread, ignore this function
-	if(Context == nullptr)
-	{
-		NumStackFramesToIgnore++;
-	}
-
 	// Capture the stack trace
 	static const int StackTraceMaxDepth = 100;
 	uint64 StackTrace[StackTraceMaxDepth];
 	FMemory::Memzero(StackTrace);
 	int32 StackTraceDepth = FPlatformStackWalk::CaptureStackBackTrace(StackTrace, StackTraceMaxDepth, Context);
 
-	// Make sure we don't exceed the current stack depth
-	NumStackFramesToIgnore = FMath::Min(NumStackFramesToIgnore, StackTraceDepth);
+	const uint64* StackTraceCursor = StackTrace;
+	if (ErrorProgramCounter != nullptr)
+	{
+		for (int32 i = 0; i < StackTraceDepth; ++i)
+		{
+			if (StackTrace[i] != uint64(ErrorProgramCounter))
+			{
+				continue;
+			}
+
+			SetNumMinidumpFramesToIgnore(i);
+			StackTraceCursor = StackTrace + i;
+			StackTraceDepth -= i;
+			break;
+		}
+	}
 
 	// Generate the portable callstack from it
-	SetPortableCallStack(StackTrace + NumStackFramesToIgnore, StackTraceDepth - NumStackFramesToIgnore);
+	SetPortableCallStack(StackTraceCursor, StackTraceDepth);
+}
+
+void FGenericCrashContext::CapturePortableCallStack(int32 NumStackFramesToIgnore, void* Context)
+{
+	CapturePortableCallStack(nullptr, Context);
 }
 
 void FGenericCrashContext::SetPortableCallStack(const uint64* StackFrames, int32 NumStackFrames)
diff --git a/Engine/Source/Runtime/Core/Private/HAL/ExceptionHandling.cpp b/Engine/Source/Runtime/Core/Private/HAL/ExceptionHandling.cpp
--- a/Engine/Source/Runtime/Core/Private/HAL/ExceptionHandling.cpp
+++ b/Engine/Source/Runtime/Core/Private/HAL/ExceptionHandling.cpp
@@ -36,9 +36,7 @@ void ReportInteractiveEnsure(const TCHAR* InMessage)
 
 #if PLATFORM_DESKTOP
 	GLog->PanicFlushThreadedLogs();
-	// Skip macros and FDebug, we always want this to fire
-	const int32 NumStackFramesToIgnore = 1;
-	ReportEnsure(InMessage, NumStackFramesToIgnore);
+	ReportEnsure(InMessage, nullptr);
 #endif
 
 	GEnsureShowsCRC = false;
diff --git a/Engine/Source/Runtime/Core/Private/IOS/IOSPlatformMisc.cpp b/Engine/Source/Runtime/Core/Private/IOS/IOSPlatformMisc.cpp
--- a/Engine/Source/Runtime/Core/Private/IOS/IOSPlatformMisc.cpp
+++ b/Engine/Source/Runtime/Core/Private/IOS/IOSPlatformMisc.cpp
@@ -1391,9 +1391,6 @@ static void DefaultCrashHandler(FIOSCrashContext const& Context)
     return Context.GenerateCrashInfo();
 }
 
-// number of stack entries to ignore in backtrace
-static uint32 GIOSStackIgnoreDepth = 6;
-
 // true system specific crash handler that gets called first
 static FIOSCrashContext TempCrashContext(ECrashContextType::Crash, TEXT("Temp Context"));
 static void PlatformCrashHandler(int32 Signal, siginfo_t* Info, void* Context)
@@ -1403,7 +1400,6 @@ static void PlatformCrashHandler(int32 Signal, siginfo_t* Info, void* Context)
 	FIOSApplicationInfo::CrashMalloc->Enable(&TempCrashContext, FPlatformTLS::GetCurrentThreadId());
 	
     FIOSCrashContext CrashContext(ECrashContextType::Crash, TEXT("Caught signal"));
-    CrashContext.IgnoreDepth = GIOSStackIgnoreDepth;
     CrashContext.InitFromSignal(Signal, Info, Context);
 	
 	// switch to the crash malloc to the new context now that we have everything
@@ -1565,7 +1561,7 @@ void FIOSPlatformMisc::SetCrashHandler(void (* CrashHandler)(const FGenericCrash
         NSError* Error = nil;
         if ([FIOSApplicationInfo::CrashReporter enableCrashReporterAndReturnError: &Error])
         {
-            GIOSStackIgnoreDepth = 0;
+            /* no-op */
         }
         else
         {
diff --git a/Engine/Source/Runtime/Core/Private/Logging/LogMacros.cpp b/Engine/Source/Runtime/Core/Private/Logging/LogMacros.cpp
--- a/Engine/Source/Runtime/Core/Private/Logging/LogMacros.cpp
+++ b/Engine/Source/Runtime/Core/Private/Logging/LogMacros.cpp
@@ -10,7 +10,7 @@
 #include "Stats/Stats.h"
 #include "ProfilingDebugging/CsvProfiler.h"
 
-void StaticFailDebug( const TCHAR* Error, const FDebug::FFailureInfo& Info, const TCHAR* Description, bool bIsEnsure, int32 NumStackFramesToIgnore );
+void StaticFailDebug( const TCHAR* Error, const FDebug::FFailureInfo& Info, const TCHAR* Description, bool bIsEnsure );
 
 /** Statics to prevent FMsg::Logf from allocating too much stack memory. */
 static FCriticalSection* GetMsgLogfStaticBufferGuard()
@@ -69,7 +69,7 @@ void FMsg::LogfImpl(const ANSICHAR* File, int32 Line, const FLogCategoryName& Ca
 
 		const int32 NumStackFramesToIgnore = 1;
 		FDebug::FFailureInfo Info = { "FMsg::Logf(Fatal)", File, Line, PLATFORM_RETURN_ADDRESS() };
-		StaticFailDebug(TEXT("Fatal error:"), Info, Message, false, NumStackFramesToIgnore);
+		StaticFailDebug(TEXT("Fatal error:"), Info, Message, false);
 		FDebug::AssertFailed("", File, Line, Message);
 	}
 #endif
@@ -117,7 +117,7 @@ void FMsg::Logf_InternalImpl(const ANSICHAR* File, int32 Line, const FLogCategor
 
 		const int32 NumStackFramesToIgnore = 1;
 		FDebug::FFailureInfo Info = { "FMsg::Logf(Fatal)", File, Line, PLATFORM_RETURN_ADDRESS() };
-		StaticFailDebug(TEXT("Fatal error:"), Info, Message, false, NumStackFramesToIgnore);
+		StaticFailDebug(TEXT("Fatal error:"), Info, Message, false);
 	}
 #endif
 }
diff --git a/Engine/Source/Runtime/Core/Private/Mac/MacErrorOutputDevice.cpp b/Engine/Source/Runtime/Core/Private/Mac/MacErrorOutputDevice.cpp
--- a/Engine/Source/Runtime/Core/Private/Mac/MacErrorOutputDevice.cpp
+++ b/Engine/Source/Runtime/Core/Private/Mac/MacErrorOutputDevice.cpp
@@ -49,15 +49,14 @@ void FMacErrorOutputDevice::Serialize( const TCHAR* Msg, ELogVerbosity::Type Ver
 #if PLATFORM_EXCEPTIONS_DISABLED
 		UE_DEBUG_BREAK();
 #endif
-		const int32 NumStackFramesToIgnore = 0;
-
+		void* ErrorProgramCounter = GetErrorProgramCounter();
 		if (GIsGPUCrashed)
 		{
-			ReportGPUCrash(Msg, NumStackFramesToIgnore);
+			ReportGPUCrash(Msg, ErrorProgramCounter);
 		}
 		else
 		{
-			ReportAssert(Msg, NumStackFramesToIgnore);
+			ReportAssert(Msg, ErrorProgramCounter);
 		}
 	}
 	else
diff --git a/Engine/Source/Runtime/Core/Private/Mac/MacPlatformMisc.cpp b/Engine/Source/Runtime/Core/Private/Mac/MacPlatformMisc.cpp
--- a/Engine/Source/Runtime/Core/Private/Mac/MacPlatformMisc.cpp
+++ b/Engine/Source/Runtime/Core/Private/Mac/MacPlatformMisc.cpp
@@ -1817,11 +1817,9 @@ static void DefaultCrashHandler(FMacCrashContext const& Context)
 	return Context.GenerateCrashInfoAndLaunchReporter();
 }
 
-/** Number of stack entries to ignore in backtrace */
-static uint32 GMacStackIgnoreDepth = 6;
-
 /** Message for the assert triggered on this thread */
 thread_local const TCHAR* GCrashErrorMessage = nullptr;
+thread_local void* GCrashErrorProgramCounter = nullptr;
 thread_local ECrashContextType GCrashErrorType = ECrashContextType::Crash;
 thread_local uint8* GCrashContextMemory[sizeof(FMacCrashContext)];
 
@@ -1833,21 +1831,24 @@ static void PlatformCrashHandler(int32 Signal, siginfo_t* Info, void* Context)
 
 	ECrashContextType Type;
 	const TCHAR* ErrorMessage;
+	void* ErrorProgramCounter;
 
 	if (GCrashErrorMessage == nullptr)
 	{
 		Type = ECrashContextType::Crash;
 		ErrorMessage = TEXT("Caught signal");
+		ErrorProgramCounter = Info->si_addr;
 	}
 	else
 	{
 		Type = GCrashErrorType;
 		ErrorMessage = GCrashErrorMessage;
+		ErrorProgramCounter = GCrashErrorProgramCounter;
 	}
 
 	FMacCrashContext* CrashContext = new (GCrashContextMemory) FMacCrashContext(Type, ErrorMessage);
-	CrashContext->IgnoreDepth = GMacStackIgnoreDepth;
 	CrashContext->InitFromSignal(Signal, Info, Context);
+	CrashContext->ErrorFrame = ErrorProgramCounter;
 
 	// Switch to crash handler malloc to avoid malloc reentrancy
 	check(GCrashMalloc);
@@ -1977,7 +1978,7 @@ void FMacPlatformMisc::SetCrashHandler(void (* CrashHandler)(const FGenericCrash
 		NSError* Error = nil;
 		if ([FMacApplicationInfo::CrashReporter enableCrashReporterAndReturnError: &Error])
 		{
-			GMacStackIgnoreDepth = 0;
+			/* GMacStackIgnoreDepth = 0; */
 		}
 		else
 		{
@@ -2417,16 +2418,23 @@ bool FMacCrashContext::GetPlatformAllThreadContextsString(FString& OutStr) const
 	return !OutStr.IsEmpty();
 }
 
-void ReportAssert(const TCHAR* ErrorMessage, int NumStackFramesToIgnore)
+void ReportAssert(const TCHAR* ErrorMessage, void* ErrorProgramCounter)
 {
 	GCrashErrorMessage = ErrorMessage;
+	GCrashErrorProgramCounter = ErrorProgramCounter;
 	GCrashErrorType = ECrashContextType::Assert;
 	FPlatformMisc::RaiseException(1);
 }
 
-void ReportGPUCrash(const TCHAR* ErrorMessage, int NumStackFramesToIgnore)
+void ReportGPUCrash(const TCHAR* ErrorMessage, void* ErrorProgramCounter)
 {
+	if (ErrorProgramCounter == nullptr)
+	{
+		ErrorProgramCounter = PLATFORM_RETURN_ADDRESS();
+	}
+
 	GCrashErrorMessage = ErrorMessage;
+	GCrashErrorProgramCounter = ErrorProgramCounter;
 	GCrashErrorType = ECrashContextType::GPUCrash;
 	FPlatformMisc::RaiseException(1);
 }
@@ -2434,7 +2442,7 @@ void ReportGPUCrash(const TCHAR* ErrorMessage, int NumStackFramesToIgnore)
 static FCriticalSection EnsureLock;
 static bool bReentranceGuard = false;
 
-void ReportEnsure( const TCHAR* ErrorMessage, int NumStackFramesToIgnore )
+void ReportEnsure( const TCHAR* ErrorMessage, void* ErrorProgramCounter )
 {
 	// Simple re-entrance guard.
 	EnsureLock.Lock();
@@ -2452,10 +2460,11 @@ void ReportEnsure( const TCHAR* ErrorMessage, int NumStackFramesToIgnore )
 		siginfo_t Signal;
 		Signal.si_signo = SIGTRAP;
 		Signal.si_code = TRAP_TRACE;
-		Signal.si_addr = __builtin_return_address(0);
+		Signal.si_addr = ErrorProgramCounter;
 		
 		FMacCrashContext EnsureContext(ECrashContextType::Ensure, ErrorMessage);
 		EnsureContext.InitFromSignal(SIGTRAP, &Signal, nullptr);
+		EnsureContext.ErrorFrame = ErrorProgramCounter;
 		EnsureContext.GenerateEnsureInfoAndLaunchReporter();
 	}
 	
diff --git a/Engine/Source/Runtime/Core/Private/Misc/AssertionMacros.cpp b/Engine/Source/Runtime/Core/Private/Misc/AssertionMacros.cpp
--- a/Engine/Source/Runtime/Core/Private/Misc/AssertionMacros.cpp
+++ b/Engine/Source/Runtime/Core/Private/Misc/AssertionMacros.cpp
@@ -144,7 +144,7 @@ static void AssertFailedImplV(const FDebug::FFailureInfo& Info, const TCHAR* For
  *	prompts for the remote debugging if there is not debugger, breaks into the debugger 
  *	and copies the error into the global error message.
  */
-FORCENOINLINE void StaticFailDebug( const TCHAR* Error, const FDebug::FFailureInfo& Info, const TCHAR* Description, bool bIsEnsure, int NumStackFramesToIgnore )
+FORCENOINLINE void StaticFailDebug( const TCHAR* Error, const FDebug::FFailureInfo& Info, const TCHAR* Description, bool bIsEnsure )
 {
 	const ANSICHAR* File = Info.File;
 	int32 Line = Info.Line;
@@ -162,7 +162,7 @@ FORCENOINLINE void StaticFailDebug( const TCHAR* Error, const FDebug::FFailureIn
 	{
 		ANSICHAR StackTrace[4096];
 		StackTrace[0] = 0;
-		FPlatformStackWalk::StackWalkAndDump(StackTrace, UE_ARRAY_COUNT(StackTrace), NumStackFramesToIgnore + 1);
+		FPlatformStackWalk::StackWalkAndDump(StackTrace, UE_ARRAY_COUNT(StackTrace), Info.ProgramCounter);
 
 		FCString::Strncat(DescriptionAndTrace, TEXT("\n"), UE_ARRAY_COUNT(DescriptionAndTrace) - 1);
 		FCString::Strncat(DescriptionAndTrace, ANSI_TO_TCHAR(StackTrace), UE_ARRAY_COUNT(DescriptionAndTrace) - 1);
@@ -182,6 +182,11 @@ FORCENOINLINE void StaticFailDebug( const TCHAR* Error, const FDebug::FFailureIn
 	// Copy the error message to the error history.
 	FCString::Strncpy( GErrorHist, ErrorMessage, UE_ARRAY_COUNT( GErrorHist ) );
 	FCString::Strncat( GErrorHist, TEXT( "\r\n\r\n" ), UE_ARRAY_COUNT( GErrorHist ) );
+
+	if (GError != nullptr)
+	{
+		GError->SetErrorProgramCounter(Info.ProgramCounter);
+	}
 }
 
 
@@ -295,8 +300,7 @@ void FDebug::LogAssertFailedMessageImplV(const FFailureInfo& Info, const TCHAR*
 		TCHAR ErrorString[MAX_SPRINTF];
 		FCString::Sprintf( ErrorString, TEXT( "Assertion failed: %s" ), ANSI_TO_TCHAR( Info.Expr ) );
 
-		const int32 NumStackFramesToIgnore = 1;
-		StaticFailDebug( ErrorString, Info, DescriptionString, false, NumStackFramesToIgnore );
+		StaticFailDebug( ErrorString, Info, DescriptionString, false );
 	}
 }
 
@@ -309,7 +313,7 @@ void FDebug::LogAssertFailedMessageImplV(const FFailureInfo& Info, const TCHAR*
  * @param	Msg		Informative error message text
  * @param	NumStackFramesToIgnore	Number of stack frames to ignore in the callstack
  */
-FORCENOINLINE void FDebug::EnsureFailed(const FFailureInfo& Info, const TCHAR* Msg, int NumStackFramesToIgnore)
+FORCENOINLINE void FDebug::EnsureFailed(const FFailureInfo& Info, const TCHAR* Msg)
 {
 	const ANSICHAR* Expr = Info.Expr;
 	const ANSICHAR* File = Info.File;
@@ -346,7 +350,7 @@ FORCENOINLINE void FDebug::EnsureFailed(const FFailureInfo& Info, const TCHAR* M
 	TCHAR ErrorString[MAX_SPRINTF];
 	FCString::Sprintf(ErrorString,TEXT("Ensure condition failed: %s"),ANSI_TO_TCHAR(Info.Expr));
 
-	StaticFailDebug( ErrorString, Info, Msg, true, NumStackFramesToIgnore + 1 );
+	StaticFailDebug( ErrorString, Info, Msg, true );
 
 	// Is there a debugger attached?  If not we'll submit an error report.
 	if (FPlatformMisc::IsDebuggerPresent() && !GAlwaysReportCrash)
@@ -382,7 +386,7 @@ FORCENOINLINE void FDebug::EnsureFailed(const FFailureInfo& Info, const TCHAR* M
 				SCOPE_LOG_TIME_IN_SECONDS(*StackWalkPerfMessage, nullptr)
 #endif
 				StackTrace[0] = 0;
-				FPlatformStackWalk::StackWalkAndDumpEx(StackTrace, StackTraceSize, NumStackFramesToIgnore + 1, FGenericPlatformStackWalk::EStackWalkFlags::FlagsUsedWhenHandlingEnsure);
+				FPlatformStackWalk::StackWalkAndDumpEx(StackTrace, StackTraceSize, Info.ProgramCounter, FGenericPlatformStackWalk::EStackWalkFlags::FlagsUsedWhenHandlingEnsure);
 			}
 
 			// Also append the stack trace
@@ -465,7 +469,7 @@ FORCENOINLINE void FDebug::EnsureFailed(const FFailureInfo& Info, const TCHAR* M
 #if PLATFORM_DESKTOP
 			FScopeLock Lock(&GetFailDebugCriticalSection());
 
-			ReportEnsure(ErrorMsg, NumStackFramesToIgnore + 1);
+			ReportEnsure(ErrorMsg, Info.ProgramCounter);
 
 			GErrorHist[0] = 0;
 			GErrorExceptionDescription[0] = 0;
@@ -530,8 +534,7 @@ FORCENOINLINE bool VARARGS FDebug::OptionallyLogFormattedEnsureMessageReturningF
 		TCHAR TempStr[ TempStrSize ];
 		GET_VARARGS( TempStr, TempStrSize, TempStrSize - 1, FormattedMsg, FormattedMsg );
 
-		const int32 NumStackFramesToIgnore = 1; // Just ignore this frame
-		EnsureFailed( Info, TempStr, NumStackFramesToIgnore );
+		EnsureFailed( Info, TempStr );
 	}
 	
 	return false;
@@ -543,8 +546,7 @@ FORCENOINLINE void VARARGS LowLevelFatalErrorHandler(const FDebug::FFailureInfo&
 	TCHAR DescriptionString[4096];
 	GET_VARARGS( DescriptionString, UE_ARRAY_COUNT(DescriptionString), UE_ARRAY_COUNT(DescriptionString)-1, Format, Format );
 
-	const int32 NumStackFramesToIgnore = 1; // Just ignore this frame
-	StaticFailDebug(TEXT("LowLevelFatalError"), Info, DescriptionString, false, NumStackFramesToIgnore);
+	StaticFailDebug(TEXT("LowLevelFatalError"), Info, DescriptionString, false);
 }
 
 void FDebug::DumpStackTraceToLog(const ELogVerbosity::Type LogVerbosity)
diff --git a/Engine/Source/Runtime/Core/Private/ProfilingDebugging/StallDetector.cpp b/Engine/Source/Runtime/Core/Private/ProfilingDebugging/StallDetector.cpp
--- a/Engine/Source/Runtime/Core/Private/ProfilingDebugging/StallDetector.cpp
+++ b/Engine/Source/Runtime/Core/Private/ProfilingDebugging/StallDetector.cpp
@@ -397,7 +397,7 @@ void UE::FStallDetector::OnStallDetected(uint32 InThreadId, const double InElaps
 		const int NumStackFramesToIgnore = FPlatformTLS::GetCurrentThreadId() == InThreadId ? 2 : 0;
 		UE_LOG(LogStall, Log, TEXT("Stall detector '%s' exceeded budget of %fs, reporting..."), Stats.Name, Stats.BudgetSeconds);
 		double ReportSeconds = FStallDetector::Seconds();
-		ReportStall(Stats.Name, InThreadId, NumStackFramesToIgnore);
+		ReportStall(Stats.Name, InThreadId);
 		ReportSeconds = FStallDetector::Seconds() - ReportSeconds;
 		UE_LOG(LogStall, Log, TEXT("Stall detector '%s' report submitted, and took %fs"), Stats.Name, ReportSeconds);
 	}
diff --git a/Engine/Source/Runtime/Core/Private/Unix/UnixErrorOutputDevice.cpp b/Engine/Source/Runtime/Core/Private/Unix/UnixErrorOutputDevice.cpp
--- a/Engine/Source/Runtime/Core/Private/Unix/UnixErrorOutputDevice.cpp
+++ b/Engine/Source/Runtime/Core/Private/Unix/UnixErrorOutputDevice.cpp
@@ -55,19 +55,14 @@ void FUnixErrorOutputDevice::Serialize(const TCHAR* Msg, ELogVerbosity::Type Ver
 #if PLATFORM_EXCEPTIONS_DISABLED
 		UE_DEBUG_BREAK();
 #endif
-		// Generate the callstack.
-		// We do not ignore any stack frames since the optimization is
-		// brittle and the risk of trimming the valid frames is too high.
-		// The common frames will be instead filtered out in the web UI
-		const int32 NumStackFramesToIgnore = 0;
-
+		void* ErrorProgramCounter = GetErrorProgramCounter();
 		if (GIsGPUCrashed)
 		{
-			ReportGPUCrash(Msg, NumStackFramesToIgnore);
+			ReportGPUCrash(Msg, ErrorProgramCounter);
 		}
 		else
 		{
-			ReportAssert(Msg, NumStackFramesToIgnore);
+			ReportAssert(Msg, ErrorProgramCounter);
 		}
 	}
 	else
diff --git a/Engine/Source/Runtime/Core/Private/Unix/UnixPlatformCrashContext.cpp b/Engine/Source/Runtime/Core/Private/Unix/UnixPlatformCrashContext.cpp
--- a/Engine/Source/Runtime/Core/Private/Unix/UnixPlatformCrashContext.cpp
+++ b/Engine/Source/Runtime/Core/Private/Unix/UnixPlatformCrashContext.cpp
@@ -275,7 +275,7 @@ void FUnixCrashContext::GenerateReport(const FString & DiagnosticsPath) const
 	}
 }
 
-void FUnixCrashContext::CaptureStackTrace()
+void FUnixCrashContext::CaptureStackTrace(void* ErrorProgramCounter)
 {
 	// Only do work the first time this function is called - this is mainly a carry over from Windows where it can be called multiple times, left intact for extra safety.
 	if (!bCapturedBacktrace)
@@ -284,11 +284,10 @@ void FUnixCrashContext::CaptureStackTrace()
 		static ANSICHAR StackTrace[StackTraceSize];
 		StackTrace[0] = 0;
 
-		int32 IgnoreCount = NumMinidumpFramesToIgnore;
-		CapturePortableCallStack(IgnoreCount, this);
+		CapturePortableCallStack(ErrorProgramCounter, this);
 
 		// Walk the stack and dump it to the allocated memory (do not ignore any stack frames to be consistent with check()/ensure() handling)
-		FPlatformStackWalk::StackWalkAndDump( StackTrace, StackTraceSize, IgnoreCount, this);
+		FPlatformStackWalk::StackWalkAndDump( StackTrace, StackTraceSize, ErrorProgramCounter, this);
 
 #if !PLATFORM_LINUX
 		printf("StackTrace:\n%s\n", StackTrace);
@@ -731,7 +730,7 @@ void DefaultCrashHandler(const FUnixCrashContext & Context)
 	FThreadHeartBeat::Get().Stop();
 
 	// at this point we should already be using malloc crash handler (see PlatformCrashHandler)
-	const_cast<FUnixCrashContext&>(Context).CaptureStackTrace();
+	const_cast<FUnixCrashContext&>(Context).CaptureStackTrace(Context.ErrorFrame);
 	if (GLog)
 	{
 		GLog->SetCurrentThreadAsMasterThread();
@@ -756,6 +755,7 @@ void (* GCrashHandlerPointer)(const FGenericCrashContext & Context) = NULL;
 extern int32 CORE_API GMaxNumberFileMappingCache;
 
 extern thread_local const TCHAR* GCrashErrorMessage;
+extern thread_local void* GCrashErrorProgramCounter;
 extern thread_local ECrashContextType GCrashErrorType;
 
 namespace
@@ -808,6 +808,7 @@ void PlatformCrashHandler(int32 Signal, siginfo_t* Info, void* Context)
 	ECrashContextType Type;
 	TStringBuilder<128> DefaultErrorMessage;
 	const TCHAR* ErrorMessage;
+	void* ErrorProgramCounter;
 
 	if (GCrashErrorMessage == nullptr)
 	{
@@ -834,16 +835,19 @@ void PlatformCrashHandler(int32 Signal, siginfo_t* Info, void* Context)
 		}
 
 		ErrorMessage = *DefaultErrorMessage;
+		ErrorProgramCounter = __builtin_return_address(0);
 	}
 	else
 	{
 		Type = GCrashErrorType;
 		ErrorMessage = GCrashErrorMessage;
+		ErrorProgramCounter = GCrashErrorProgramCounter;
 	}
 
 	FUnixCrashContext CrashContext(Type, ErrorMessage);
 	CrashContext.InitFromSignal(Signal, Info, Context);
 	CrashContext.FirstCrashHandlerFrame = static_cast<uint64*>(__builtin_return_address(0));
+	CrashContext.ErrorFrame = ErrorProgramCounter;
 
 	// This will ungrab cursor/keyboard and bring down any pointer barriers which will be stuck on when opening the CRC
 	FPlatformMisc::UngrabAllInput();
diff --git a/Engine/Source/Runtime/Core/Private/Unix/UnixPlatformStackWalk.cpp b/Engine/Source/Runtime/Core/Private/Unix/UnixPlatformStackWalk.cpp
--- a/Engine/Source/Runtime/Core/Private/Unix/UnixPlatformStackWalk.cpp
+++ b/Engine/Source/Runtime/Core/Private/Unix/UnixPlatformStackWalk.cpp
@@ -891,19 +891,27 @@ int32 FUnixPlatformStackWalk::GetProcessModuleSignatures(FStackWalkModuleInfo *M
 }
 
 thread_local const TCHAR* GCrashErrorMessage = nullptr;
+thread_local void* GCrashErrorProgramCounter = nullptr;
 thread_local ECrashContextType GCrashErrorType = ECrashContextType::Crash;
 
-void ReportAssert(const TCHAR* ErrorMessage, int NumStackFramesToIgnore)
+void ReportAssert(const TCHAR* ErrorMessage, void* ProgramCounter)
 {
 	GCrashErrorMessage = ErrorMessage;
+	GCrashErrorProgramCounter = ProgramCounter;
 	GCrashErrorType = ECrashContextType::Assert;
 
 	FPlatformMisc::RaiseException(1);
 }
 
-void ReportGPUCrash(const TCHAR* ErrorMessage, int NumStackFramesToIgnore)
+void ReportGPUCrash(const TCHAR* ErrorMessage, void* ProgramCounter)
 {
+	if (ProgramCounter == nullptr)
+	{
+		ProgramCounter = PLATFORM_RETURN_ADDRESS();
+	}
+
 	GCrashErrorMessage = ErrorMessage;
+	GCrashErrorProgramCounter = ProgramCounter;
 	GCrashErrorType = ECrashContextType::GPUCrash;
 
 	FPlatformMisc::RaiseException(1);
@@ -912,7 +920,7 @@ void ReportGPUCrash(const TCHAR* ErrorMessage, int NumStackFramesToIgnore)
 static FCriticalSection EnsureLock;
 static bool bReentranceGuard = false;
 
-void ReportEnsure(const TCHAR* ErrorMessage, int NumStackFramesToIgnore)
+void ReportEnsure(const TCHAR* ErrorMessage, void* ProgramCounter)
 {
 	// Simple re-entrance guard.
 	EnsureLock.Lock();
@@ -926,9 +934,9 @@ void ReportEnsure(const TCHAR* ErrorMessage, int NumStackFramesToIgnore)
 	bReentranceGuard = true;
 
 	FUnixCrashContext EnsureContext(ECrashContextType::Ensure, ErrorMessage);
-	EnsureContext.InitFromEnsureHandler(ErrorMessage, __builtin_return_address(0));
+	EnsureContext.InitFromEnsureHandler(ErrorMessage, ProgramCounter);
 
-	EnsureContext.CaptureStackTrace();
+	EnsureContext.CaptureStackTrace(ProgramCounter);
 	EnsureContext.GenerateCrashInfoAndLaunchReporter(true);
 
 	bReentranceGuard = false;
diff --git a/Engine/Source/Runtime/Core/Private/Windows/WindowsErrorOutputDevice.cpp b/Engine/Source/Runtime/Core/Private/Windows/WindowsErrorOutputDevice.cpp
--- a/Engine/Source/Runtime/Core/Private/Windows/WindowsErrorOutputDevice.cpp
+++ b/Engine/Source/Runtime/Core/Private/Windows/WindowsErrorOutputDevice.cpp
@@ -64,19 +64,14 @@ void FWindowsErrorOutputDevice::Serialize( const TCHAR* Msg, ELogVerbosity::Type
 #if PLATFORM_EXCEPTIONS_DISABLED
 		UE_DEBUG_BREAK();
 #endif
-		// Generate the portable callstack. For asserts, we ignore the following frames:
-		// We do not ignore any stack frames since the optimization is
-		// brittle and the risk of trimming the valid frames is too high.
-		// The common frames will be instead filtered out in the web UI
-		const int32 NumStackFramesToIgnore = 0;
-
+		void* ErrorProgramCounter = GetErrorProgramCounter();
 		if (GIsGPUCrashed)
 		{
-			ReportGPUCrash(Msg, NumStackFramesToIgnore);
+			ReportGPUCrash(Msg, ErrorProgramCounter);
 		}
 		else
 		{
-			ReportAssert(Msg, NumStackFramesToIgnore);
+			ReportAssert(Msg, ErrorProgramCounter);
 		}
 	}
 	else
diff --git a/Engine/Source/Runtime/Core/Private/Windows/WindowsPlatformCrashContext.cpp b/Engine/Source/Runtime/Core/Private/Windows/WindowsPlatformCrashContext.cpp
--- a/Engine/Source/Runtime/Core/Private/Windows/WindowsPlatformCrashContext.cpp
+++ b/Engine/Source/Runtime/Core/Private/Windows/WindowsPlatformCrashContext.cpp
@@ -140,11 +140,11 @@ namespace {
 struct FAssertInfo
 {
 	const TCHAR* ErrorMessage;
-	int32 NumStackFramesToIgnore;
+	void* ProgramCounter;
 
-	FAssertInfo(const TCHAR* InErrorMessage, int32 InNumStackFramesToIgnore)
+	FAssertInfo(const TCHAR* InErrorMessage, void* InProgramCounter)
 		: ErrorMessage(InErrorMessage)
-		, NumStackFramesToIgnore(InNumStackFramesToIgnore)
+		, ProgramCounter(InProgramCounter)
 	{
 	}
 };
@@ -626,7 +626,7 @@ int32 ReportCrashForMonitor(
 	LPEXCEPTION_POINTERS ExceptionInfo,
 	ECrashContextType Type,
 	const TCHAR* ErrorMessage,
-	int NumStackFramesToIgnore,
+	void* ErrorProgramCounter,
 	HANDLE CrashingThreadHandle,
 	DWORD CrashingThreadId,
 	FProcHandle& CrashMonitorHandle,
@@ -647,7 +647,8 @@ int32 ReportCrashForMonitor(
 	// Setup up the shared memory area so that the crash report
 	SharedContext->CrashType = Type;
 	SharedContext->CrashingThreadId = CrashingThreadId;
-	SharedContext->NumStackFramesToIgnore = NumStackFramesToIgnore;
+	SharedContext->ErrorProgramCounter = ErrorProgramCounter;
+	SharedContext->ExceptionProgramCounter = ExceptionInfo->ExceptionRecord->ExceptionAddress;
 
 	// Determine UI settings for the crash report. Suppress the user input dialog if we're running in unattended mode
 	// Usage data controls if we want analytics in the crash report client
@@ -1242,7 +1243,7 @@ public:
 	}
 
 	/** Ensures and Stalls are passed through this. */
-	FORCEINLINE int32 OnContinuableEvent(ECrashContextType InType, LPEXCEPTION_POINTERS InExceptionInfo, HANDLE InThreadHandle, uint32 InThreadId, int NumStackFramesToIgnore, const TCHAR* ErrorMessage, EErrorReportUI ReportUI)
+	FORCEINLINE int32 OnContinuableEvent(ECrashContextType InType, LPEXCEPTION_POINTERS InExceptionInfo, HANDLE InThreadHandle, uint32 InThreadId, void* ProgramCounter, const TCHAR* ErrorMessage, EErrorReportUI ReportUI)
 	{
 		if (CrashClientHandle.IsValid() && FPlatformProcess::IsProcRunning(CrashClientHandle))
 		{
@@ -1250,7 +1251,7 @@ public:
 				InExceptionInfo, 
 				InType,
 				ErrorMessage, 
-				NumStackFramesToIgnore, 
+				ProgramCounter, 
 				InThreadHandle,
 				InThreadId,
 				CrashClientHandle, 
@@ -1266,7 +1267,7 @@ public:
 			CrashContext.SetCrashedProcess(FProcHandle(::GetCurrentProcess()));
 			CrashContext.SetCrashedThreadId(InThreadId);
 			void* ContextWrapper = FWindowsPlatformStackWalk::MakeThreadContextWrapper(InExceptionInfo->ContextRecord, InThreadHandle);
-			CrashContext.CapturePortableCallStack(NumStackFramesToIgnore, ContextWrapper);
+			CrashContext.CapturePortableCallStack(ProgramCounter, ContextWrapper);
 			//CrashContext.CaptureAllThreadContexts(); -> For ensure/stall, don't capture all threads to report and resume quickly.
 
 			return ReportCrashUsingCrashReportClient(CrashContext, InExceptionInfo, ReportUI);
@@ -1292,10 +1293,15 @@ public:
 	/** Crashes during static init should be reported directly to crash monitor. */
 	FORCEINLINE int32 OnCrashDuringStaticInit(LPEXCEPTION_POINTERS InExceptionInfo)
 	{
+		void* ErrorProgramCounter = nullptr;
+        if (InExceptionInfo && InExceptionInfo->ExceptionRecord)
+        {
+            ErrorProgramCounter = InExceptionInfo->ExceptionRecord->ExceptionAddress;
+        }
+
 		if (CrashClientHandle.IsValid() && FPlatformProcess::IsProcRunning(CrashClientHandle))
 		{
 			const ECrashContextType Type = ECrashContextType::Crash;
-			const int NumStackFramesToIgnore = 1;
 			const TCHAR* ErrorMessage = TEXT("Crash during static initialization");
 
 			if (!FPlatformCrashContext::IsInitalized())
@@ -1307,7 +1313,7 @@ public:
 				InExceptionInfo,
 				Type,
 				ErrorMessage,
-				NumStackFramesToIgnore,
+				ErrorProgramCounter,
 				CrashingThreadHandle,
 				CrashingThreadId,
 				CrashClientHandle,
@@ -1344,7 +1350,7 @@ private:
 		ECrashContextType Type = ECrashContextType::Crash;
 		const TCHAR* ErrorMessage = TEXT("Unhandled exception");
 		TCHAR ErrorMessageLocal[UE_ARRAY_COUNT(GErrorExceptionDescription)];
-		int NumStackFramesToIgnore = 2;
+		void* ErrorProgramCounter = ExceptionInfo->ExceptionRecord->ExceptionAddress;
 
 		void* ContextWrapper = nullptr;
 
@@ -1354,20 +1360,27 @@ private:
 			const FAssertInfo& Info = *(const FAssertInfo*)ExceptionInfo->ExceptionRecord->ExceptionInformation[0];
 			Type = ECrashContextType::Assert;
 			ErrorMessage = Info.ErrorMessage;
-			NumStackFramesToIgnore += Info.NumStackFramesToIgnore;
+			ErrorProgramCounter = Info.ProgramCounter;
 		}
 		else if (ExceptionInfo->ExceptionRecord->ExceptionCode == GPUCrashExceptionCode && ExceptionInfo->ExceptionRecord->NumberParameters == 1)
 		{
 			const FAssertInfo& Info = *(const FAssertInfo*)ExceptionInfo->ExceptionRecord->ExceptionInformation[0];
 			Type = ECrashContextType::GPUCrash;
 			ErrorMessage = Info.ErrorMessage;
-			NumStackFramesToIgnore += Info.NumStackFramesToIgnore;
+			ErrorProgramCounter = Info.ProgramCounter;
 		}
 		// Generic exception description is stored in GErrorExceptionDescription
-		else if (ExceptionInfo->ExceptionRecord->ExceptionCode != EnsureExceptionCode)
+		else if (ExceptionInfo->ExceptionRecord->ExceptionCode == EnsureExceptionCode)
+		{
+			const FAssertInfo& Info = *(const FAssertInfo*)ExceptionInfo->ExceptionRecord->ExceptionInformation[0];
+			Type = ECrashContextType::Ensure;
+			ErrorMessage = Info.ErrorMessage;
+			ErrorProgramCounter = Info.ProgramCounter;
+		}
+		// Generic exception description is stored in GErrorExceptionDescription
+		else
 		{
 			// When a generic exception is thrown, it is important to get all the stack frames
-			NumStackFramesToIgnore = 0;
 			CreateExceptionInfoString(ExceptionInfo->ExceptionRecord, ErrorMessageLocal, UE_ARRAY_COUNT(ErrorMessageLocal));
 			ErrorMessage = ErrorMessageLocal;
 
@@ -1384,7 +1397,7 @@ private:
 				ExceptionInfo,
 				Type,
 				ErrorMessage,
-				NumStackFramesToIgnore,
+				ErrorProgramCounter,
 				CrashingThreadHandle,
 				CrashingThreadId,
 				CrashClientHandle,
@@ -1404,12 +1417,9 @@ private:
 			// Thread context wrapper for stack operations
 			ContextWrapper = FWindowsPlatformStackWalk::MakeThreadContextWrapper(ExceptionInfo->ContextRecord, CrashingThreadHandle);
 			CrashContext.SetCrashedProcess(FProcHandle(::GetCurrentProcess()));
-			CrashContext.CapturePortableCallStack(NumStackFramesToIgnore, ContextWrapper);
 			CrashContext.SetCrashedThreadId(CrashingThreadId);
 			CrashContext.CaptureAllThreadContexts();
-
-			// Also mark the same number of frames to be ignored if we symbolicate from the minidump
-			CrashContext.SetNumMinidumpFramesToIgnore(NumStackFramesToIgnore);
+			CrashContext.CapturePortableCallStack(ErrorProgramCounter, ContextWrapper);
 
 			// First launch the crash reporter client.
 #if WINVER > 0x502	// Windows Error Reporting is not supported on Windows XP
@@ -1442,7 +1452,7 @@ private:
 				ContextWrapper = FWindowsPlatformStackWalk::MakeThreadContextWrapper(ExceptionInfo->ContextRecord, CrashingThreadHandle);
 			}
 			
-			FPlatformStackWalk::StackWalkAndDump(StackTrace, StackTraceSize, 0, ContextWrapper);
+			FPlatformStackWalk::StackWalkAndDump(StackTrace, StackTraceSize, ErrorProgramCounter, ContextWrapper);
 			
 			if (ExceptionInfo->ExceptionRecord->ExceptionCode != EnsureExceptionCode && ExceptionInfo->ExceptionRecord->ExceptionCode != AssertExceptionCode)
 			{
@@ -1588,30 +1598,29 @@ static bool bReentranceGuard = false;
 /**
  * A wrapper for ReportCrashUsingCrashReportClient that creates a new ensure crash context
  */
-int32 ReportContinuableEventUsingCrashReportClient(ECrashContextType InType, EXCEPTION_POINTERS* ExceptionInfo, HANDLE InThreadHandle, uint32 InThreadId, int NumStackFramesToIgnore, const TCHAR* ErrorMessage, EErrorReportUI ReportUI)
+int32 ReportContinuableEventUsingCrashReportClient(ECrashContextType InType, EXCEPTION_POINTERS* ExceptionInfo, HANDLE InThreadHandle, uint32 InThreadId, void* ProgramCounter, const TCHAR* ErrorMessage, EErrorReportUI ReportUI)
 {
 #if !NOINITCRASHREPORTER
-	return GCrashReportingThread->OnContinuableEvent(InType, ExceptionInfo, InThreadHandle, InThreadId, NumStackFramesToIgnore, ErrorMessage, ReportUI);
+	return GCrashReportingThread->OnContinuableEvent(InType, ExceptionInfo, InThreadHandle, InThreadId, ProgramCounter, ErrorMessage, ReportUI);
 #else 
 	return EXCEPTION_EXECUTE_HANDLER;
 #endif
 }
 #endif
 
-static void ReportEventOnCallingThread(ECrashContextType InType, const TCHAR* ErrorMessage, int NumStackFramesToIgnore)
+static void ReportEventOnCallingThread(ECrashContextType InType, const TCHAR* ErrorMessage, void* ProgramCounter)
 {
-	// Skip this frame and the ::RaiseException call itself
-	NumStackFramesToIgnore += 2;
-
 #if WINVER > 0x502	// Windows Error Reporting is not supported on Windows XP
 #if !PLATFORM_SEH_EXCEPTIONS_DISABLED
 	__try
 #endif
 	{
-		::RaiseException(EnsureExceptionCode, 0, 0, nullptr);
+		FAssertInfo Info(ErrorMessage, ProgramCounter);
+		ULONG_PTR Arguments[] = { (ULONG_PTR)&Info };
+		::RaiseException(EnsureExceptionCode, 0, UE_ARRAY_COUNT(Arguments), Arguments);
 	}
 #if !PLATFORM_SEH_EXCEPTIONS_DISABLED
-	__except (ReportContinuableEventUsingCrashReportClient( InType, GetExceptionInformation(), GetCurrentThread(), GetCurrentThreadId(), NumStackFramesToIgnore, ErrorMessage, IsInteractiveEnsureMode() ? EErrorReportUI::ShowDialog : EErrorReportUI::ReportInUnattendedMode))
+	__except (ReportContinuableEventUsingCrashReportClient( InType, GetExceptionInformation(), GetCurrentThread(), GetCurrentThreadId(), ProgramCounter, ErrorMessage, IsInteractiveEnsureMode() ? EErrorReportUI::ShowDialog : EErrorReportUI::ReportInUnattendedMode))
 		CA_SUPPRESS(6322)
 	{
 	}
@@ -1619,7 +1628,7 @@ static void ReportEventOnCallingThread(ECrashContextType InType, const TCHAR* Er
 #endif	// WINVER
 }
 
-static void ReportEvent(ECrashContextType InType, const TCHAR* ErrorMessage, uint32 InThreadId, int NumStackFramesToIgnore)
+static void ReportEvent(ECrashContextType InType, const TCHAR* ErrorMessage, uint32 InThreadId, void* ProgramCounter)
 {
 	if (ReportCrashCallCount > 0 || FDebug::HasAsserted())
 	{
@@ -1649,7 +1658,7 @@ static void ReportEvent(ECrashContextType InType, const TCHAR* ErrorMessage, uin
 
 	if (FPlatformTLS::GetCurrentThreadId() == InThreadId)
 	{
-		ReportEventOnCallingThread(InType, ErrorMessage, NumStackFramesToIgnore);
+		ReportEventOnCallingThread(InType, ErrorMessage, ProgramCounter);
 	}
 	else
 	{
@@ -1676,6 +1685,11 @@ static void ReportEvent(ECrashContextType InType, const TCHAR* ErrorMessage, uin
 				{
 					// Success! Set the pointer to this state in the exception information block
 					ExceptionInformation.ContextRecord = &ThreadContext;
+
+					if (ProgramCounter == nullptr)
+					{
+						ProgramCounter = (void*)(ThreadContext.Rip);
+					}
 				}
 			}
 
@@ -1689,7 +1703,7 @@ static void ReportEvent(ECrashContextType InType, const TCHAR* ErrorMessage, uin
 				// The thread context and exception recard are the two things expected
 				ExceptionInformation.ExceptionRecord = &ExceptionRecord;
 
-				(void)ReportContinuableEventUsingCrashReportClient(InType, &ExceptionInformation, ThreadHandle, InThreadId, NumStackFramesToIgnore, ErrorMessage, IsInteractiveEnsureMode() ? EErrorReportUI::ShowDialog : EErrorReportUI::ReportInUnattendedMode);
+				(void)ReportContinuableEventUsingCrashReportClient(InType, &ExceptionInformation, ThreadHandle, InThreadId, ProgramCounter, ErrorMessage, IsInteractiveEnsureMode() ? EErrorReportUI::ShowDialog : EErrorReportUI::ReportInUnattendedMode);
 			}
 
 			if (SuspendCount != -1)
@@ -1702,19 +1716,24 @@ static void ReportEvent(ECrashContextType InType, const TCHAR* ErrorMessage, uin
 	}
 }
 
-FORCENOINLINE void ReportAssert(const TCHAR* ErrorMessage, int NumStackFramesToIgnore)
+void ReportAssert(const TCHAR* ErrorMessage, void* ProgramCounter)
 {
 	/** This is the last place to gather memory stats before exception. */
 	FGenericCrashContext::SetMemoryStats(FPlatformMemory::GetStats());
 
-	FAssertInfo Info(ErrorMessage, NumStackFramesToIgnore + 2); // +2 for this function and RaiseException()
+	FAssertInfo Info(ErrorMessage, ProgramCounter);
 
 	ULONG_PTR Arguments[] = { (ULONG_PTR)&Info };
 	::RaiseException(AssertExceptionCode, 0, UE_ARRAY_COUNT(Arguments), Arguments);
 }
 
-FORCENOINLINE void ReportGPUCrash(const TCHAR* ErrorMessage, int NumStackFramesToIgnore)
+FORCENOINLINE void ReportGPUCrash(const TCHAR* ErrorMessage, void* ProgramCounter)
 {
+	if (ProgramCounter == nullptr)
+	{
+		ProgramCounter = PLATFORM_RETURN_ADDRESS();
+	}
+
 	/** This is the last place to gather memory stats before exception. */
 	FGenericCrashContext::SetMemoryStats(FPlatformMemory::GetStats());
 	
@@ -1722,7 +1741,7 @@ FORCENOINLINE void ReportGPUCrash(const TCHAR* ErrorMessage, int NumStackFramesT
 #if !PLATFORM_SEH_EXCEPTIONS_DISABLED
 	__try
 	{
-		FAssertInfo Info(ErrorMessage, NumStackFramesToIgnore + 2); // +2 for this function and RaiseException()
+		FAssertInfo Info(ErrorMessage, ProgramCounter);
 
 		ULONG_PTR Arguments[] = { (ULONG_PTR)&Info };
 		::RaiseException(GPUCrashExceptionCode, 0, UE_ARRAY_COUNT(Arguments), Arguments);
@@ -1757,17 +1776,17 @@ void ReportHang(const TCHAR* ErrorMessage, const uint64* StackFrames, int32 NumS
 /**
  * Report an ensure to the crash reporting system
  */
-FORCENOINLINE void ReportEnsure(const TCHAR* ErrorMessage, int NumStackFramesToIgnore)
+void ReportEnsure(const TCHAR* ErrorMessage, void* ProgramCounter)
 {
-	ReportEvent(ECrashContextType::Ensure, ErrorMessage, FPlatformTLS::GetCurrentThreadId(), NumStackFramesToIgnore + 1);
+	ReportEvent(ECrashContextType::Ensure, ErrorMessage, FPlatformTLS::GetCurrentThreadId(), ProgramCounter);
 }
 
 /**
  * Report a hitch to the crash reporting system
  */
-FORCENOINLINE void ReportStall(const TCHAR* ErrorMessage, uint32 HitchThreadId, int NumStackFramesToIgnore)
+FORCENOINLINE void ReportStall(const TCHAR* ErrorMessage, uint32 HitchThreadId)
 {
-	ReportEvent(ECrashContextType::Stall, ErrorMessage, HitchThreadId, FPlatformTLS::GetCurrentThreadId() == HitchThreadId ? NumStackFramesToIgnore + 1 : NumStackFramesToIgnore);
+	ReportEvent(ECrashContextType::Stall, ErrorMessage, HitchThreadId, nullptr);
 }
 
 #if !IS_PROGRAM && 0
@@ -1804,3 +1823,4 @@ namespace {
 
 
 
+
diff --git a/Engine/Source/Runtime/Core/Public/Apple/ApplePlatformCrashContext.h b/Engine/Source/Runtime/Core/Public/Apple/ApplePlatformCrashContext.h
--- a/Engine/Source/Runtime/Core/Public/Apple/ApplePlatformCrashContext.h
+++ b/Engine/Source/Runtime/Core/Public/Apple/ApplePlatformCrashContext.h
@@ -25,8 +25,9 @@ struct CORE_API FApplePlatformCrashContext : public FGenericCrashContext
 	/** Thread context */
 	ucontext_t*	Context;
 	
-	/** Number of stack frames to ignore as they are crash handling/stack walking */
-	uint32 IgnoreDepth;
+	/** The PC of where the error being reported occurred. This is used to figure
+	 * out the number of frames to ignore */
+	void* ErrorFrame = nullptr;
 	
 	/** Memory reserved for "exception" (signal) info */
 	char SignalDescription[128];
diff --git a/Engine/Source/Runtime/Core/Public/GenericPlatform/GenericPlatformCrashContext.h b/Engine/Source/Runtime/Core/Public/GenericPlatform/GenericPlatformCrashContext.h
--- a/Engine/Source/Runtime/Core/Public/GenericPlatform/GenericPlatformCrashContext.h
+++ b/Engine/Source/Runtime/Core/Public/GenericPlatform/GenericPlatformCrashContext.h
@@ -194,7 +194,6 @@ struct FSharedCrashContext
 	TCHAR					ThreadNames[CR_MAX_THREAD_NAME_CHARS * CR_MAX_THREADS];
 	uint32					NumThreads;
 	uint32					CrashingThreadId;
-	uint32					NumStackFramesToIgnore;
 	ECrashContextType		CrashType;
 
 	// Additional user settings.
@@ -217,6 +216,12 @@ struct FSharedCrashContext
 	uint32					GameDataOffset;
 	// Fixed size dynamic buffer
 	TCHAR					DynamicData[CR_MAX_DYNAMIC_BUFFER_CHARS];
+
+	// Program counter address where the error occurred.
+	void*					ErrorProgramCounter;
+
+	// Instruction address where the exception was raised that initiated crash reporting
+	void*					ExceptionProgramCounter;
 };
 
 /**
@@ -428,7 +433,14 @@ public:
 	/** Sets the number of stack frames to ignore when symbolicating from a minidump */
 	void SetNumMinidumpFramesToIgnore(int32 InNumMinidumpFramesToIgnore);
 
-	/** Generate raw call stack for crash report (image base + offset) */
+	/**
+	 * Generate raw call stack for crash report (image base + offset)
+	 * @param ErrorProgramCounter The program counter of where the occur occurred in the callstack being captured
+	 * @param Context Optional thread context information
+	 */
+	void CapturePortableCallStack(void* ErrorProgramCounter, void* Context);
+
+	UE_DEPRECATED(5.0, "")
 	void CapturePortableCallStack(int32 NumStackFramesToIgnore, void* Context);
 	
 	/** Sets the portable callstack to a specified stack */
diff --git a/Engine/Source/Runtime/Core/Public/HAL/ExceptionHandling.h b/Engine/Source/Runtime/Core/Public/HAL/ExceptionHandling.h
--- a/Engine/Source/Runtime/Core/Public/HAL/ExceptionHandling.h
+++ b/Engine/Source/Runtime/Core/Public/HAL/ExceptionHandling.h
@@ -24,29 +24,29 @@ extern CORE_API TCHAR MiniDumpFilenameW[1024];
 #include <excpt.h>
 // #CrashReport: 2014-10-09 These methods are specific to windows, remove from here.
 extern CORE_API int32 ReportCrash( Windows::LPEXCEPTION_POINTERS ExceptionInfo );
-extern CORE_API void ReportAssert(const TCHAR* ErrorMessage, int NumStackFramesToIgnore);
-extern CORE_API void ReportGPUCrash(const TCHAR* ErrorMessage, int NumStackFramesToIgnore);
-extern CORE_API void ReportEnsure(const TCHAR* ErrorMessage, int NumStackFramesToIgnore);
-extern CORE_API void ReportStall(const TCHAR* ErrorMessage, uint32 HitchThreadId, int NumStackFramesToIgnore);
+extern CORE_API void ReportAssert(const TCHAR* ErrorMessage, void* ProgramCounter);
+extern CORE_API void ReportGPUCrash(const TCHAR* ErrorMessage, void* ProgramCounter);
+extern CORE_API void ReportEnsure(const TCHAR* ErrorMessage, void* ProgramCounter);
+extern CORE_API void ReportStall(const TCHAR* ErrorMessage, uint32 HitchThreadId);
 extern CORE_API void ReportHang(const TCHAR*, const uint64* StackFrames, int32 NumStackFrames, uint32 HungThreadId);
 #elif PLATFORM_MAC
 // #CrashReport: 2014-10-09 Should be move to another file
 #include <signal.h>
 extern CORE_API int32 ReportCrash(ucontext_t *Context, int32 Signal, struct __siginfo* Info);
-extern CORE_API void ReportAssert(const TCHAR* ErrorMessage, int NumStackFramesToIgnore);
-extern CORE_API void ReportGPUCrash(const TCHAR* ErrorMessage, int NumStackFramesToIgnore);
-extern CORE_API void ReportEnsure(const TCHAR* ErrorMessage, int NumStackFramesToIgnore);
+extern CORE_API void ReportAssert(const TCHAR* ErrorMessage, void* ProgramCounter);
+extern CORE_API void ReportGPUCrash(const TCHAR* ErrorMessage, void* ProgramCounter);
+extern CORE_API void ReportEnsure(const TCHAR* ErrorMessage, void* ProgramCounter);
 extern CORE_API void ReportHang(const TCHAR*, const uint64* StackFrames, int32 NumStackFrames, uint32 HungThreadId);
 #elif PLATFORM_UNIX
-extern CORE_API void ReportAssert(const TCHAR* ErrorMessage, int NumStackFramesToIgnore);
-extern CORE_API void ReportGPUCrash(const TCHAR* ErrorMessage, int NumStackFramesToIgnore);
-extern CORE_API void ReportEnsure(const TCHAR* ErrorMessage, int NumStackFramesToIgnore);
+extern CORE_API void ReportAssert(const TCHAR* ErrorMessage, void* ProgramCounter);
+extern CORE_API void ReportGPUCrash(const TCHAR* ErrorMessage, void* ProgramCounter);
+extern CORE_API void ReportEnsure(const TCHAR* ErrorMessage, void* ProgramCounter);
 extern CORE_API void ReportHang(const TCHAR*, const uint64* StackFrames, int32 NumStackFrames, uint32 HungThreadId);
 #elif PLATFORM_HOLOLENS
 #include "HoloLens/HoloLensSystemIncludes.h"
 #include <excpt.h>
 extern CORE_API int32 ReportCrash(Windows::LPEXCEPTION_POINTERS ExceptionInfo);
-extern CORE_API void ReportEnsure(const TCHAR* ErrorMessage, int NumStackFramesToIgnore);
+extern CORE_API void ReportEnsure(const TCHAR* ErrorMessage, void* ProgramCounter);
 extern CORE_API void ReportHang(const TCHAR*, const uint64* StackFrames, int32 NumStackFrames, uint32 HungThreadId);
 extern CORE_API void NewReportEnsure(const TCHAR* ErrorMessage);
 #endif
diff --git a/Engine/Source/Runtime/Core/Public/Misc/AssertionMacros.h b/Engine/Source/Runtime/Core/Public/Misc/AssertionMacros.h
--- a/Engine/Source/Runtime/Core/Public/Misc/AssertionMacros.h
+++ b/Engine/Source/Runtime/Core/Public/Misc/AssertionMacros.h
@@ -98,7 +98,7 @@ public:
 	 * 
 	 * Don't change the name of this function, it's used to detect ensures by the crash reporter.
 	 */
-	static void EnsureFailed( const FFailureInfo& Info, const TCHAR* Msg, int NumStackFramesToIgnore );
+	static void EnsureFailed( const FFailureInfo& Info, const TCHAR* Msg );
 
 private:
 	static bool VARARGS OptionallyLogFormattedEnsureMessageReturningFalseImpl(bool bLog, const FFailureInfo& Info, const TCHAR* FormattedMsg, ...);
@@ -230,7 +230,7 @@ RetType FORCENOINLINE UE_DEBUG_SECTION DispatchCheckVerify(InnerType&& Inner, Ar
 		{ \
 			if(UNLIKELY(!(expr))) \
 			{ \
-				DispatchCheckVerify([&] () FORCENOINLINE UE_DEBUG_SECTION \
+				DispatchCheckVerify([&] () UE_DEBUG_SECTION \
 				{ \
 					FDebug::FFailureInfo CheckFailureInfo = { #expr, __FILE__, __LINE__, PLATFORM_RETURN_ADDRESS() }; \
 					FDebug::CheckVerifyFailed(CheckFailureInfo, format, ##__VA_ARGS__); \
diff --git a/Engine/Source/Runtime/Core/Public/Unix/UnixPlatformCrashContext.h b/Engine/Source/Runtime/Core/Public/Unix/UnixPlatformCrashContext.h
--- a/Engine/Source/Runtime/Core/Public/Unix/UnixPlatformCrashContext.h
+++ b/Engine/Source/Runtime/Core/Public/Unix/UnixPlatformCrashContext.h
@@ -51,6 +51,10 @@ struct CORE_API FUnixCrashContext : public FGenericCrashContext
 	/** The PC of the first function used when handling a crash. Used to figure out the number of frames to ignore */
 	uint64* FirstCrashHandlerFrame = nullptr;
 
+	/** The PC of where the error being reported occurred. Note that this could be different
+	 * from FirstCrashHandlerFrame */
+	void* ErrorFrame = nullptr;
+
 	FUnixCrashContext(ECrashContextType InType, const TCHAR* InErrorMessage)
 		:	FGenericCrashContext(InType, InErrorMessage)
 		,	Signal(0)
@@ -86,7 +90,7 @@ struct CORE_API FUnixCrashContext : public FGenericCrashContext
 	 * Populates crash context stack trace and a few related fields
 	 *
 	 */
-	void CaptureStackTrace();
+	void CaptureStackTrace(void* ErrorProgramCounter);
 
 	/**
 	 * Generates a new crash report containing information needed for the crash reporter and launches it; may not return.
diff --git a/Engine/Source/Runtime/D3D12RHI/Private/D3D12Util.cpp b/Engine/Source/Runtime/D3D12RHI/Private/D3D12Util.cpp
--- a/Engine/Source/Runtime/D3D12RHI/Private/D3D12Util.cpp
+++ b/Engine/Source/Runtime/D3D12RHI/Private/D3D12Util.cpp
@@ -622,7 +622,7 @@ namespace D3D12RHI
 			CloseHandle(FileHandle);
 
 			// Report the GPU crash which will raise the exception (only interesting if we have a GPU dump)
-			ReportGPUCrash(TEXT("Aftermath GPU Crash dump Triggered"), 0);
+			ReportGPUCrash(TEXT("Aftermath GPU Crash dump Triggered"), nullptr);
 
 			// Force shutdown, we can't do anything useful anymore.
 			FPlatformMisc::RequestExit(true);
diff --git a/Engine/Source/Runtime/Unix/UnixCommonStartup/Private/UnixCommonStartup.cpp b/Engine/Source/Runtime/Unix/UnixCommonStartup/Private/UnixCommonStartup.cpp
--- a/Engine/Source/Runtime/Unix/UnixCommonStartup/Private/UnixCommonStartup.cpp
+++ b/Engine/Source/Runtime/Unix/UnixCommonStartup/Private/UnixCommonStartup.cpp
@@ -32,7 +32,7 @@ void CommonUnixCrashHandler(const FGenericCrashContext& GenericContext)
 	FGenericCrashContext::SetMemoryStats(FPlatformMemory::GetStats());
 
 	// better than having mutable fields?
-	const_cast< FUnixCrashContext& >(Context).CaptureStackTrace();
+	const_cast< FUnixCrashContext& >(Context).CaptureStackTrace(Context.ProgramCounter);
 	if (GLog)
 	{
 		GLog->SetCurrentThreadAsMasterThread();
diff --git a/Engine/Source/Runtime/Windows/D3D11RHI/Private/Windows/WindowsD3D11Device.cpp b/Engine/Source/Runtime/Windows/D3D11RHI/Private/Windows/WindowsD3D11Device.cpp
--- a/Engine/Source/Runtime/Windows/D3D11RHI/Private/Windows/WindowsD3D11Device.cpp
+++ b/Engine/Source/Runtime/Windows/D3D11RHI/Private/Windows/WindowsD3D11Device.cpp
@@ -1270,7 +1270,7 @@ static void D3D11AftermathCrashCallback(const void* InGPUCrashDump, const uint32
 	CloseHandle(FileHandle);
 
 	// Report the GPU crash which will raise the exception
-	ReportGPUCrash(TEXT("Aftermath GPU Crash dump Triggered"), 0);
+	ReportGPUCrash(TEXT("Aftermath GPU Crash dump Triggered"), nullptr);
 }
 
 void EnableNVAftermathCrashDumps()
