Fixed trace for "TaskWorkerIsLookingForWork" cpu timer (ensured the event type is not missed if CpuChannel is initially disabled and also that Begin/End events are always traced in pairs -- even if cpu channel is enabled or disabled between begin-end).

diff --git a/Engine/Source/Runtime/Core/Public/Async/Fundamental/LocalQueue.h b/Engine/Source/Runtime/Core/Public/Async/Fundamental/LocalQueue.h
--- a/Engine/Source/Runtime/Core/Public/Async/Fundamental/LocalQueue.h
+++ b/Engine/Source/Runtime/Core/Public/Async/Fundamental/LocalQueue.h
@@ -145,16 +145,14 @@ public:
 		std::atomic_int& NumWorkersLookingForWork;
 		bool ActivelyLookingForWork = false;
 
+#if CPUPROFILERTRACE_ENABLED
+		bool bCpuBeginEventEmitted = false;
+#endif
+
 		static uint32 WorkerLookingForWorkTraceId;
 
 		inline FOutOfWork(std::atomic_int& InNumWorkersLookingForWork) : NumWorkersLookingForWork(InNumWorkersLookingForWork)
 		{
-#if CPUPROFILERTRACE_ENABLED
-			if (WorkerLookingForWorkTraceId == 0) 
-			{
-				WorkerLookingForWorkTraceId = FCpuProfilerTrace::OutputEventType("TaskWorkerIsLookingForWork");
-			}
-#endif
 		}
 
 	public:
@@ -168,7 +166,15 @@ public:
 			if (!ActivelyLookingForWork)
 			{
 #if CPUPROFILERTRACE_ENABLED
-				FCpuProfilerTrace::OutputBeginEvent(WorkerLookingForWorkTraceId);
+				if (CpuChannel)
+				{
+					if (WorkerLookingForWorkTraceId == 0)
+					{
+						WorkerLookingForWorkTraceId = FCpuProfilerTrace::OutputEventType("TaskWorkerIsLookingForWork");
+					}
+					FCpuProfilerTrace::OutputBeginEvent(WorkerLookingForWorkTraceId);
+					bCpuBeginEventEmitted = true;
+				}
 #endif
 				NumWorkersLookingForWork.fetch_add(1, std::memory_order_relaxed);
 				ActivelyLookingForWork = true;
@@ -182,7 +188,10 @@ public:
 			if (ActivelyLookingForWork)
 			{
 #if CPUPROFILERTRACE_ENABLED
-				FCpuProfilerTrace::OutputEndEvent();
+				if (bCpuBeginEventEmitted)
+				{
+					FCpuProfilerTrace::OutputEndEvent();
+				}
 #endif
 				NumWorkersLookingForWork.fetch_sub(1, std::memory_order_release);
 				ActivelyLookingForWork = false;
