LowLevelMemTracker: Fix missing declarations. Change FLLMTagDeclaration to allow multiple listeners.
#rnx

diff --git a/Engine/Source/Runtime/Core/Private/HAL/LowLevelMemTracker.cpp b/Engine/Source/Runtime/Core/Private/HAL/LowLevelMemTracker.cpp
--- a/Engine/Source/Runtime/Core/Private/HAL/LowLevelMemTracker.cpp
+++ b/Engine/Source/Runtime/Core/Private/HAL/LowLevelMemTracker.cpp
@@ -1313,13 +1313,13 @@ void FLowLevelMemTracker::InitialiseTagDatas()
 	// Construct the remaining startup tags; allocations when constructing these tags are known to consist only of the central list of ELLMTags so we do not need to bootstrap
 	{
 		// Construct LLM_DECLARE_TAGs
-		FLLMTagDeclaration*& List = FLLMTagDeclaration::GetList();
+		FLLMTagDeclaration* List = FLLMTagDeclaration::GetList();
 		while (List)
 		{
 			RegisterTagDeclaration(*List);
 			List = List->Next;
 		}
-		FLLMTagDeclaration::SetCreationCallback(GlobalRegisterTagDeclaration);
+		FLLMTagDeclaration::AddCreationCallback(GlobalRegisterTagDeclaration);
 	}
 
 	// now let the platform add any custom tags
@@ -1335,7 +1335,7 @@ void FLowLevelMemTracker::ClearTagDatas()
 	using namespace UE::LLMPrivate;
 
 	FWriteScopeLock ScopeLock(TagDataLock);
-	FLLMTagDeclaration::SetCreationCallback(nullptr);
+	FLLMTagDeclaration::ClearCreationCallbacks();
 
 	Allocator.Free(TagDataEnumMap, sizeof(FTagData*) * LLM_TAG_COUNT);
 	TagDataEnumMap = nullptr;
@@ -2073,36 +2073,69 @@ void FLLMTagDeclaration::ConstructUniqueName()
 	UniqueName = FName(*NameBuffer);
 }
 
-void FLLMTagDeclaration::SetCreationCallback(FCreationCallback InCallback)
+namespace UE
+{
+namespace LLMPrivate
+{
+namespace LLMTagDeclarationInternal
 {
-	GetCreationCallback() = InCallback;
-}
 
-FLLMTagDeclaration::FCreationCallback& FLLMTagDeclaration::GetCreationCallback()
+TArray<FLLMTagDeclaration::FCreationCallback, TInlineAllocator<2>>& GetCreationCallbacks()
 {
-	static FCreationCallback CreationCallback = nullptr;
-	return CreationCallback;
+	static TArray<FLLMTagDeclaration::FCreationCallback, TInlineAllocator<2>> CreationCallbacks;
+	return CreationCallbacks;
 }
 
-FLLMTagDeclaration*& FLLMTagDeclaration::GetList()
+FLLMTagDeclaration*& GetList()
 {
 	static FLLMTagDeclaration* List = nullptr;
 	return List;
 }
 
-void FLLMTagDeclaration::Register()
+}
+}
+}
+
+void FLLMTagDeclaration::AddCreationCallback(FCreationCallback InCallback)
 {
-	FCreationCallback& CreationCallback = GetCreationCallback();
-	if (CreationCallback)
+	TArray<FCreationCallback, TInlineAllocator<2>>& CreationCallbacks =
+		UE::LLMPrivate::LLMTagDeclarationInternal::GetCreationCallbacks();
+	if (CreationCallbacks.Num() >= CreationCallbacks.Max())
 	{
-		CreationCallback(*this);
+		check(false); // We are not allowed to allocate memory here; you need to increase the allocation size
 	}
 	else
 	{
-		FLLMTagDeclaration*& List = GetList();
-		Next = List;
-		List = this;
+		CreationCallbacks.Add(InCallback);
+	}
+}
+
+void FLLMTagDeclaration::ClearCreationCallbacks()
+{
+	TArray<FCreationCallback, TInlineAllocator<2>>& CreationCallbacks =
+		UE::LLMPrivate::LLMTagDeclarationInternal::GetCreationCallbacks();
+	CreationCallbacks.Empty();
+}
+
+TArrayView<FLLMTagDeclaration::FCreationCallback> FLLMTagDeclaration::GetCreationCallbacks()
+{
+	return UE::LLMPrivate::LLMTagDeclarationInternal::GetCreationCallbacks();
+}
+
+FLLMTagDeclaration* FLLMTagDeclaration::GetList()
+{
+	return UE::LLMPrivate::LLMTagDeclarationInternal::GetList();
+}
+
+void FLLMTagDeclaration::Register()
+{
+	for (FCreationCallback CreationCallback : GetCreationCallbacks())
+	{
+		CreationCallback(*this);
 	}
+	FLLMTagDeclaration*& List = UE::LLMPrivate::LLMTagDeclarationInternal::GetList();
+	Next = List;
+	List = this;
 }
 
 namespace UE
diff --git a/Engine/Source/Runtime/Core/Private/ProfilingDebugging/TagTrace.cpp b/Engine/Source/Runtime/Core/Private/ProfilingDebugging/TagTrace.cpp
--- a/Engine/Source/Runtime/Core/Private/ProfilingDebugging/TagTrace.cpp
+++ b/Engine/Source/Runtime/Core/Private/ProfilingDebugging/TagTrace.cpp
@@ -213,13 +213,13 @@ void FTagTrace::AnnounceGenericTags()
 ////////////////////////////////////////////////////////////////////////////////
 void FTagTrace::AnnounceTagDeclarations()
 {
-	FLLMTagDeclaration*& List = FLLMTagDeclaration::GetList();
+	FLLMTagDeclaration* List = FLLMTagDeclaration::GetList();
 	while (List)
 	{
 		OnAnnounceTagDeclaration(*List);
 		List = List->Next;
 	}
-	FLLMTagDeclaration::SetCreationCallback(FTagTrace::OnAnnounceTagDeclaration);
+	FLLMTagDeclaration::AddCreationCallback(FTagTrace::OnAnnounceTagDeclaration);
 }
 
 ////////////////////////////////////////////////////////////////////////////////
diff --git a/Engine/Source/Runtime/Core/Public/HAL/LowLevelMemTracker.h b/Engine/Source/Runtime/Core/Public/HAL/LowLevelMemTracker.h
--- a/Engine/Source/Runtime/Core/Public/HAL/LowLevelMemTracker.h
+++ b/Engine/Source/Runtime/Core/Public/HAL/LowLevelMemTracker.h
@@ -39,6 +39,8 @@
 // Enable stat tags if: (1) Stats are allowed or (2) Asset tags are allowed (asset tags use the stat macros to record asset scopes)
 #define LLM_ENABLED_STAT_TAGS LLM_ALLOW_STATS || LLM_ALLOW_ASSETS_TAGS
 
+#include "Containers/Array.h"
+#include "Containers/ArrayView.h"
 #include "HAL/CriticalSection.h"
 #include "Templates/AlignmentTemplates.h"
 #include "Templates/UnrealTemplate.h"
@@ -703,12 +705,13 @@ public:
 	FLLMTagDeclaration(const TCHAR* InCPPName, const FName InDisplayName=NAME_None, FName InParentTagName = NAME_None, FName InStatName = NAME_None, FName InSummaryStatName = NAME_None);
 	FName GetUniqueName() const { return UniqueName; }
 
-protected:
 	typedef void (*FCreationCallback)(FLLMTagDeclaration&);
+protected:
 
-	static void SetCreationCallback(FCreationCallback InCallback);
-	static FCreationCallback& GetCreationCallback();
-	static FLLMTagDeclaration*& GetList();
+	static void AddCreationCallback(FCreationCallback InCallback);
+	static void ClearCreationCallbacks();
+	static TArrayView<FCreationCallback> GetCreationCallbacks();
+	static FLLMTagDeclaration* GetList();
 
 	void Register();
 
