Don't log non-reporting stalls

diff --git a/Engine/Source/Runtime/Core/Private/ProfilingDebugging/StallDetector.cpp b/Engine/Source/Runtime/Core/Private/ProfilingDebugging/StallDetector.cpp
--- a/Engine/Source/Runtime/Core/Private/ProfilingDebugging/StallDetector.cpp
+++ b/Engine/Source/Runtime/Core/Private/ProfilingDebugging/StallDetector.cpp
@@ -289,7 +289,10 @@ void UE::FStallDetector::Check(bool bIsComplete, double InWhenToCheckSeconds)
 			FString OverageString = FString::Printf(TEXT("[FStallDetector] [%s] Overage of %f\n"), Stats.Name, OverageSeconds);
 			FPlatformMisc::LocalPrint(OverageString.GetCharArray().GetData());
 #endif
-			UE_LOG(LogStall, Log, TEXT("Stall detector '%s' complete in %fs (%fs overbudget)"), Stats.Name, DeltaSeconds, OverageSeconds);
+			if (Stats.ReportingMode != EStallDetectorReportingMode::Disabled)
+			{
+				UE_LOG(LogStall, Log, TEXT("Stall detector '%s' complete in %fs (%fs overbudget)"), Stats.Name, DeltaSeconds, OverageSeconds);
+			}
 		}
 	}
 	else
@@ -403,7 +406,10 @@ void UE::FStallDetector::OnStallDetected(uint32 InThreadId, const double InElaps
 	}
 	else
 	{
-		UE_LOG(LogStall, Log, TEXT("Stall detector '%s' exceeded budget of %fs"), Stats.Name, Stats.BudgetSeconds);
+		if (ReportingMode != EStallDetectorReportingMode::Disabled)
+		{
+			UE_LOG(LogStall, Log, TEXT("Stall detector '%s' exceeded budget of %fs"), Stats.Name, Stats.BudgetSeconds);
+		}
 	}
 }
 
