U5 - change scheduler to use cycles for it's spinloop because that is more predictable across differnt CPU architectures.

diff --git a/Engine/Source/Runtime/Core/Private/Async/Fundamental/Scheduler.cpp b/Engine/Source/Runtime/Core/Private/Async/Fundamental/Scheduler.cpp
--- a/Engine/Source/Runtime/Core/Private/Async/Fundamental/Scheduler.cpp
+++ b/Engine/Source/Runtime/Core/Private/Async/Fundamental/Scheduler.cpp
@@ -40,7 +40,7 @@ namespace LowLevelTasks
 	TUniquePtr<FThread> FScheduler::CreateWorker(FLocalQueueType* ExternalWorkerLocalQueue, EThreadPriority Priority, bool bPermitBackgroundWork, bool bIsForkable)
 	{
 		uint32 WorkerId = NextWorkerId++;
-		const uint32 WaitTimes[8] = { 23, 31, 41, 37, 47, 29, 19, 43 };
+		const uint32 WaitTimes[8] = { 719, 991, 1361, 1237, 1597, 953, 587, 1439 };
 		uint32 WaitTime = WaitTimes[WorkerId % 8];
 		uint64 ThreadAffinityMask = FPlatformAffinity::GetTaskGraphThreadMask();
 
@@ -272,10 +272,7 @@ namespace LowLevelTasks
 			if (WaitCount < WorkerSpinCycles)
 			{
 				OutOfWork.Start();
-				for (uint32 i = 0; i < WaitCycles; i++)
-				{
-					FPlatformProcess::Yield();
-				}
+				FPlatformProcess::YieldCycles(WaitCycles);
 				WaitCount++;
 				continue;
 			}
