Ensures can now optionally log out as Warning instead of Error, based on the "core.EnsuresAreErrors" cvar (defaults to true)

diff --git a/Engine/Source/Runtime/Core/Private/Misc/AssertionMacros.cpp b/Engine/Source/Runtime/Core/Private/Misc/AssertionMacros.cpp
--- a/Engine/Source/Runtime/Core/Private/Misc/AssertionMacros.cpp
+++ b/Engine/Source/Runtime/Core/Private/Misc/AssertionMacros.cpp
@@ -101,6 +101,16 @@ namespace
 
 */
 
+/** 
+ * Use this CVar to control whether ensures count as errors or warnings. 
+ * Errors will fail certain processes like cooks, whereas warnings will not.
+ */
+int32 GEnsuresAreErrors = 1;
+FAutoConsoleVariableRef CVarEnsuresAreErrors(
+	TEXT("core.EnsuresAreErrors"),
+	GEnsuresAreErrors,
+	TEXT("True means failed ensures are logged as errors. False means they are logged as warnings."),
+	ECVF_Default);
 
 
 CORE_API void (*GPrintScriptCallStackFn)() = nullptr;
@@ -356,7 +366,14 @@ FORCENOINLINE void FDebug::EnsureFailed(const FFailureInfo& Info, const TCHAR* M
 	if (FPlatformMisc::IsDebuggerPresent() && !GAlwaysReportCrash)
 	{
 #if !NO_LOGGING
-		UE_LOG(LogOutputDevice, Error, TEXT("%s") FILE_LINE_DESC TEXT("\n%s\n"), ErrorString, ANSI_TO_TCHAR(File), Line, Msg);
+		if (GEnsuresAreErrors)
+		{
+			UE_LOG(LogOutputDevice, Error, TEXT("%s") FILE_LINE_DESC TEXT("\n%s\n"), ErrorString, ANSI_TO_TCHAR(File), Line, Msg);
+		}
+		else
+		{
+			UE_LOG(LogOutputDevice, Warning, TEXT("%s") FILE_LINE_DESC TEXT("\n%s\n"), ErrorString, ANSI_TO_TCHAR(File), Line, Msg);
+		}
 #endif
 	}
 	else
@@ -395,7 +412,7 @@ FORCENOINLINE void FDebug::EnsureFailed(const FFailureInfo& Info, const TCHAR* M
 
 			// Dump the error and flush the log.
 #if !NO_LOGGING
-			FDebug::LogFormattedMessageWithCallstack(LogOutputDevice.GetCategoryName(), __FILE__, __LINE__, TEXT("=== Handled ensure: ==="), ErrorMsg, ELogVerbosity::Error);
+			FDebug::LogFormattedMessageWithCallstack(LogOutputDevice.GetCategoryName(), __FILE__, __LINE__, TEXT("=== Handled ensure: ==="), ErrorMsg, (GEnsuresAreErrors ? ELogVerbosity::Error : ELogVerbosity::Warning));
 #endif
 			GLog->Flush();
 
@@ -455,7 +472,14 @@ FORCENOINLINE void FDebug::EnsureFailed(const FFailureInfo& Info, const TCHAR* M
 
 			// Add message to log even without stacktrace. It is useful for testing fail on ensure.
 #if !NO_LOGGING
-			UE_LOG(LogOutputDevice, Error, TEXT("%s [File:%s] [Line: %i]"), ErrorString, ANSI_TO_TCHAR(File), Line);
+			if (GEnsuresAreErrors)
+			{
+				UE_LOG(LogOutputDevice, Error, TEXT("%s [File:%s] [Line: %i]"), ErrorString, ANSI_TO_TCHAR(File), Line);
+			}
+			else
+			{
+				UE_LOG(LogOutputDevice, Warning, TEXT("%s [File:%s] [Line: %i]"), ErrorString, ANSI_TO_TCHAR(File), Line);
+			}
 #endif
 		}
 
