Fix compile error in UnrealMultiUserServer - LLM defines need to be wrapped in (), and tagtrace needs to check for ENABLE_LOW_LEVEL_MEM_TRACKER.
#rnx

diff --git a/Engine/Source/Runtime/Core/Private/ProfilingDebugging/TagTrace.cpp b/Engine/Source/Runtime/Core/Private/ProfilingDebugging/TagTrace.cpp
--- a/Engine/Source/Runtime/Core/Private/ProfilingDebugging/TagTrace.cpp
+++ b/Engine/Source/Runtime/Core/Private/ProfilingDebugging/TagTrace.cpp
@@ -128,7 +128,9 @@ public:
 					FTagTrace(FMalloc* InMalloc);
 	void			AnnounceGenericTags();
 	void 			AnnounceTagDeclarations();
+#if ENABLE_LOW_LEVEL_MEM_TRACKER
 	static void 	OnAnnounceTagDeclaration(FLLMTagDeclaration& TagDeclaration);
+#endif
 	int32			AnnounceCustomTag(int32 Tag, int32 ParentTag, const ANSICHAR* Display);
 	int32 			AnnounceFNameTag(const FName& TagName);
 
@@ -198,6 +200,7 @@ FTagTrace::FTagTrace(FMalloc* InMalloc)
 ////////////////////////////////////////////////////////////////////////////////
 void FTagTrace::AnnounceGenericTags()
 {
+#if ENABLE_LOW_LEVEL_MEM_TRACKER
 	#define TRACE_TAG_SPEC(Enum,Str,Stat,Group,ParentTag)\
 	{\
 		const uint32 DisplayLen = FCStringAnsi::Strlen(Str);\
@@ -208,11 +211,13 @@ void FTagTrace::AnnounceGenericTags()
 	}
 	LLM_ENUM_GENERIC_TAGS(TRACE_TAG_SPEC);
 	#undef TRACE_TAG_SPEC
+#endif
 }
 
 ////////////////////////////////////////////////////////////////////////////////
 void FTagTrace::AnnounceTagDeclarations()
 {
+#if ENABLE_LOW_LEVEL_MEM_TRACKER
 	FLLMTagDeclaration* List = FLLMTagDeclaration::GetList();
 	while (List)
 	{
@@ -220,14 +225,17 @@ void FTagTrace::AnnounceTagDeclarations()
 		List = List->Next;
 	}
 	FLLMTagDeclaration::AddCreationCallback(FTagTrace::OnAnnounceTagDeclaration);
+#endif
 }
 
 ////////////////////////////////////////////////////////////////////////////////
+#if ENABLE_LOW_LEVEL_MEM_TRACKER
 void FTagTrace::OnAnnounceTagDeclaration(FLLMTagDeclaration& TagDeclaration)
 {
 	TagDeclaration.ConstructUniqueName();
 	GTagTrace->AnnounceFNameTag(TagDeclaration.GetUniqueName());
 }
+#endif
 
 ////////////////////////////////////////////////////////////////////////////////
 int32 FTagTrace::AnnounceFNameTag(const FName& Name)
diff --git a/Engine/Source/Runtime/Core/Public/HAL/LowLevelMemTracker.h b/Engine/Source/Runtime/Core/Public/HAL/LowLevelMemTracker.h
--- a/Engine/Source/Runtime/Core/Public/HAL/LowLevelMemTracker.h
+++ b/Engine/Source/Runtime/Core/Public/HAL/LowLevelMemTracker.h
@@ -11,11 +11,11 @@
 
 // LLM is currently incompatible with PLATFORM_USES_FIXED_GMalloc_CLASS, because LLM is activated way too early
 // Inability to use LLM with PLATFORM_USES_FIXED_GMalloc_CLASS is not a problem, because fixed GMalloc is only used in Test/Shipping builds
-#define LLM_ENABLED_ON_PLATFORM PLATFORM_SUPPORTS_LLM && !PLATFORM_USES_FIXED_GMalloc_CLASS
+#define LLM_ENABLED_ON_PLATFORM (PLATFORM_SUPPORTS_LLM && !PLATFORM_USES_FIXED_GMalloc_CLASS)
 
 // *** enable/disable LLM here ***
 #if !defined(ENABLE_LOW_LEVEL_MEM_TRACKER) || !LLM_ENABLED_ON_PLATFORM 
-	#define ENABLE_LOW_LEVEL_MEM_TRACKER LLM_ENABLED_ON_PLATFORM && !UE_BUILD_SHIPPING && (!UE_BUILD_TEST || ALLOW_LOW_LEVEL_MEM_TRACKER_IN_TEST) && WITH_ENGINE && 1
+	#define ENABLE_LOW_LEVEL_MEM_TRACKER (LLM_ENABLED_ON_PLATFORM && !UE_BUILD_SHIPPING && (!UE_BUILD_TEST || ALLOW_LOW_LEVEL_MEM_TRACKER_IN_TEST) && WITH_ENGINE && 1)
 #endif
 
 #if ENABLE_LOW_LEVEL_MEM_TRACKER
