Fixed missing include that was breaking nonunity BuildPatchTool compiles.

diff --git a/Engine/Source/Runtime/Core/Private/Misc/AssertionMacros.cpp b/Engine/Source/Runtime/Core/Private/Misc/AssertionMacros.cpp
--- a/Engine/Source/Runtime/Core/Private/Misc/AssertionMacros.cpp
+++ b/Engine/Source/Runtime/Core/Private/Misc/AssertionMacros.cpp
@@ -23,6 +23,7 @@
 #include "Misc/CoreDelegates.h"
 #include "HAL/ExceptionHandling.h"
 #include "HAL/ThreadHeartBeat.h"
+#include "HAL/IConsoleManager.h"
 
 namespace 
 {
