U5 - readcycle counter is a privilidged instruction on Androind (and maybe also IOS). Just Yield number of times in this case for now

diff --git a/Engine/Source/Runtime/Core/Public/GenericPlatform/GenericPlatformProcess.h b/Engine/Source/Runtime/Core/Public/GenericPlatform/GenericPlatformProcess.h
--- a/Engine/Source/Runtime/Core/Public/GenericPlatform/GenericPlatformProcess.h
+++ b/Engine/Source/Runtime/Core/Public/GenericPlatform/GenericPlatformProcess.h
@@ -745,6 +745,12 @@ struct CORE_API FGenericPlatformProcess
 	*/
 	static FORCEINLINE void YieldCycles(uint64 cycles)
 	{
+#if PLATFORM_IOS || PLATFORM_ANDROID
+		for (int i = 0; i < cycles; i++)
+		{
+			Yield();
+		}
+#else
 		uint64 start = ReadCycleCounter();
 		//some 32bit implementations return 0 for __builtin_readcyclecounter just to be on the safe side we protect against this.
 		cycles = start != 0 ? cycles : 0;
@@ -752,5 +758,6 @@ struct CORE_API FGenericPlatformProcess
 		{
 			Yield();
 		} while((ReadCycleCounter() - start) < cycles);
+#endif
 	}
 };
