U5 - int to int64 comparission

diff --git a/Engine/Source/Runtime/Core/Public/GenericPlatform/GenericPlatformProcess.h b/Engine/Source/Runtime/Core/Public/GenericPlatform/GenericPlatformProcess.h
--- a/Engine/Source/Runtime/Core/Public/GenericPlatform/GenericPlatformProcess.h
+++ b/Engine/Source/Runtime/Core/Public/GenericPlatform/GenericPlatformProcess.h
@@ -746,7 +746,7 @@ struct CORE_API FGenericPlatformProcess
 	static FORCEINLINE void YieldCycles(uint64 cycles)
 	{
 #if PLATFORM_IOS || PLATFORM_ANDROID
-		for (int i = 0; i < cycles; i++)
+		for (uint64 i = 0; i < cycles; i++)
 		{
 			Yield();
 		}
