String: Make TStringConversion and TStringPointer construct from and present as contiguous containers of characters

diff --git a/Engine/Source/Runtime/Core/Public/Containers/StringConv.h b/Engine/Source/Runtime/Core/Public/Containers/StringConv.h
--- a/Engine/Source/Runtime/Core/Public/Containers/StringConv.h
+++ b/Engine/Source/Runtime/Core/Public/Containers/StringConv.h
@@ -10,6 +10,9 @@
 #include "Containers/ContainerAllocationPolicies.h"
 #include "Containers/Array.h"
 #include "Misc/CString.h"
+#include "Templates/AndOrNot.h"
+#include "Traits/ElementType.h"
+#include "Traits/IsContiguousContainer.h"
 
 #define DEFAULT_STRING_CONVERSION_SIZE 128u
 #define UNICODE_BOGUS_CHAR_CODEPOINT '?'
@@ -1033,6 +1036,19 @@ public:
 		}
 	}
 
+	/**
+	 * Construct from a compatible character range such as TStringView or TStringBuilder.
+	 */
+	template <typename FromRangeType,
+		typename TEnableIf<TAnd<
+			TIsContiguousContainer<FromRangeType>,
+			TNot<TIsArray<typename TRemoveReference<FromRangeType>::Type>>,
+			TIsSame<FromType, typename TRemoveCV<typename TRemovePointer<decltype(GetData(DeclVal<FromRangeType>()))>::Type>::Type>>::Value>::Type* = nullptr>
+	TStringConversion(FromRangeType&& Source)
+		: TStringConversion(GetData(Source), GetNum(Source))
+	{
+	}
+
 	/**
 	 * Move constructor
 	 */
@@ -1077,6 +1093,30 @@ private:
 	int32   StringLength;
 };
 
+template <typename Converter, int32 DefaultConversionSize>
+inline auto GetData(const TStringConversion<Converter, DefaultConversionSize>& Conversion) -> decltype(Conversion.Get())
+{
+	return Conversion.Get();
+}
+
+template <typename Converter, int32 DefaultConversionSize>
+inline auto GetNum(const TStringConversion<Converter, DefaultConversionSize>& Conversion) -> decltype(Conversion.Length())
+{
+	return Conversion.Length();
+}
+
+template <typename Converter, int32 DefaultConversionSize>
+struct TIsContiguousContainer<TStringConversion<Converter, DefaultConversionSize>>
+{
+	static constexpr bool Value = true;
+};
+
+template <typename Converter, int32 DefaultConversionSize>
+struct TElementType<TStringConversion<Converter, DefaultConversionSize>>
+{
+	using Type = typename Converter::ToType;
+};
+
 
 /**
  * Class takes one type of string and and stores it as-is.
@@ -1126,6 +1166,19 @@ public:
 		}
 	}
 
+	/**
+	 * Construct from a compatible character range such as TStringView or TStringBuilder.
+	 */
+	template <typename FromRangeType,
+		typename TEnableIf<TAnd<
+			TIsContiguousContainer<FromRangeType>,
+			TNot<TIsArray<typename TRemoveReference<FromRangeType>::Type>>,
+			TIsSame<FromType, typename TRemoveCV<typename TRemovePointer<decltype(GetData(DeclVal<FromRangeType>()))>::Type>::Type>>::Value>::Type* = nullptr>
+	TStringPointer(FromRangeType&& Source)
+		: TStringPointer(GetData(Source), GetNum(Source))
+	{
+	}
+
 	/**
 	 * Move constructor
 	 */
@@ -1161,6 +1214,30 @@ private:
 	int32 StringLength;
 };
 
+template <typename FromType, typename ToType>
+inline auto GetData(const TStringPointer<FromType, ToType>& Pointer) -> decltype(Pointer.Get())
+{
+	return Pointer.Get();
+}
+
+template <typename FromType, typename ToType>
+inline auto GetNum(const TStringPointer<FromType, ToType>& Pointer) -> decltype(Pointer.Length())
+{
+	return Pointer.Length();
+}
+
+template <typename FromType, typename ToType>
+struct TIsContiguousContainer<TStringPointer<FromType, ToType>>
+{
+	static constexpr bool Value = true;
+};
+
+template <typename FromType, typename ToType>
+struct TElementType<TStringPointer<FromType, ToType>>
+{
+	using Type = ToType;
+};
+
 
 /**
  * NOTE: The objects these macros declare have very short lifetimes. They are
