Allow ExecProcess() to be called as a Job Process #fyi danny.couture #rb jack.porter #jira UE-109833 #okforgithub public

diff --git a/Engine/Source/Developer/IOS/IOSTargetPlatform/Private/IOSDeviceHelper.cpp b/Engine/Source/Developer/IOS/IOSTargetPlatform/Private/IOSDeviceHelper.cpp
--- a/Engine/Source/Developer/IOS/IOSTargetPlatform/Private/IOSDeviceHelper.cpp
+++ b/Engine/Source/Developer/IOS/IOSTargetPlatform/Private/IOSDeviceHelper.cpp
@@ -59,7 +59,7 @@ DEFINE_LOG_CATEGORY_STATIC(LogIOSDeviceHelper, Log, All);
         FString LibimobileDeviceId = GetLibImobileDeviceExe("idevice_id");
         int ReturnCode;
         // get the list of devices UDID
-        FPlatformProcess::ExecProcess(*LibimobileDeviceId, TEXT(""), &ReturnCode, &OutStdOut, &OutStdErr);
+        FPlatformProcess::ExecProcess(*LibimobileDeviceId, TEXT(""), &ReturnCode, &OutStdOut, &OutStdErr, NULL, true);
         
 
         TArray<LibIMobileDevice> ToReturn;
@@ -98,7 +98,7 @@ DEFINE_LOG_CATEGORY_STATIC(LogIOSDeviceHelper, Log, All);
                 Arguments = "-n -u " + DeviceID;
             }
 
-            FPlatformProcess::ExecProcess(*LibimobileDeviceInfo, *Arguments, &ReturnCodeInfo, &OutStdOutInfo, &OutStdErrInfo);
+            FPlatformProcess::ExecProcess(*LibimobileDeviceInfo, *Arguments, &ReturnCodeInfo, &OutStdOutInfo, &OutStdErrInfo, NULL, true);
             // parse product type and device name
             FString DeviceName;
             OutStdOutInfo.Split(TEXT("DeviceName: "), nullptr, &DeviceName, ESearchCase::CaseSensitive, ESearchDir::FromStart);
@@ -262,7 +262,7 @@ private:
         FString LibimobileDeviceId = GetLibImobileDeviceExe("idevice_id");
         int ReturnCode;
         // get the list of devices UDID
-        FPlatformProcess::ExecProcess(*LibimobileDeviceId, TEXT(""), &ReturnCode, &OutStdOut, &OutStdErr);
+        FPlatformProcess::ExecProcess(*LibimobileDeviceId, TEXT(""), &ReturnCode, &OutStdOut, &OutStdErr, NULL, true);
         if (OutStdOut.Len() == 0)
         {
             RetryQuery--;
diff --git a/Engine/Source/Runtime/Core/Private/GenericPlatform/GenericPlatformProcess.cpp b/Engine/Source/Runtime/Core/Private/GenericPlatform/GenericPlatformProcess.cpp
--- a/Engine/Source/Runtime/Core/Private/GenericPlatform/GenericPlatformProcess.cpp
+++ b/Engine/Source/Runtime/Core/Private/GenericPlatform/GenericPlatformProcess.cpp
@@ -297,7 +297,7 @@ FString FGenericPlatformProcess::GetApplicationName( uint32 ProcessId )
 	return FString(TEXT(""));
 }
 
-bool FGenericPlatformProcess::ExecProcess(const TCHAR* URL, const TCHAR* Params, int32* OutReturnCode, FString* OutStdOut, FString* OutStdErr, const TCHAR* OptionalWorkingDirectory)
+bool FGenericPlatformProcess::ExecProcess(const TCHAR* URL, const TCHAR* Params, int32* OutReturnCode, FString* OutStdOut, FString* OutStdErr, const TCHAR* OptionalWorkingDirectory, bool bShouldEndWithParentProcess)
 {
 	UE_LOG(LogHAL, Fatal, TEXT("FGenericPlatformProcess::ExecProcess not implemented on this platform"));
 	return false;
diff --git a/Engine/Source/Runtime/Core/Private/Mac/MacPlatformProcess.cpp b/Engine/Source/Runtime/Core/Private/Mac/MacPlatformProcess.cpp
--- a/Engine/Source/Runtime/Core/Private/Mac/MacPlatformProcess.cpp
+++ b/Engine/Source/Runtime/Core/Private/Mac/MacPlatformProcess.cpp
@@ -199,7 +199,7 @@ FString FMacPlatformProcess::GetGameBundleId()
 	return FString([[NSBundle mainBundle] bundleIdentifier]);
 }
 
-bool FMacPlatformProcess::ExecProcess( const TCHAR* URL, const TCHAR* Params, int32* OutReturnCode, FString* OutStdOut, FString* OutStdErr, const TCHAR* OptionalWorkingDirectory)
+bool FMacPlatformProcess::ExecProcess( const TCHAR* URL, const TCHAR* Params, int32* OutReturnCode, FString* OutStdOut, FString* OutStdErr, const TCHAR* OptionalWorkingDirectory, bool bShouldEndWithParentProcess)
 {
 	FString CmdLineParams = Params;
 	FString ExecutableFileName = URL;
diff --git a/Engine/Source/Runtime/Core/Private/Unix/UnixPlatformProcess.cpp b/Engine/Source/Runtime/Core/Private/Unix/UnixPlatformProcess.cpp
--- a/Engine/Source/Runtime/Core/Private/Unix/UnixPlatformProcess.cpp
+++ b/Engine/Source/Runtime/Core/Private/Unix/UnixPlatformProcess.cpp
@@ -1648,7 +1648,7 @@ bool FUnixPlatformProcess::IsApplicationRunning( const TCHAR* ProcName )
 	return !system(TCHAR_TO_UTF8(*Commandline));
 }
 
-bool FUnixPlatformProcess::ExecProcess(const TCHAR* URL, const TCHAR* Params, int32* OutReturnCode, FString* OutStdOut, FString* OutStdErr, const TCHAR* OptionalWorkingDirectory)
+bool FUnixPlatformProcess::ExecProcess(const TCHAR* URL, const TCHAR* Params, int32* OutReturnCode, FString* OutStdOut, FString* OutStdErr, const TCHAR* OptionalWorkingDirectory, bool bShouldEndWithParentProcess)
 {
 	FString CmdLineParams = Params;
 	FString ExecutableFileName = URL;
diff --git a/Engine/Source/Runtime/Core/Private/Windows/WindowsPlatformProcess.cpp b/Engine/Source/Runtime/Core/Private/Windows/WindowsPlatformProcess.cpp
--- a/Engine/Source/Runtime/Core/Private/Windows/WindowsPlatformProcess.cpp
+++ b/Engine/Source/Runtime/Core/Private/Windows/WindowsPlatformProcess.cpp
@@ -42,7 +42,8 @@ PRAGMA_DISABLE_UNSAFE_TYPECAST_WARNINGS
 // static variables
 TArray<FString> FWindowsPlatformProcess::DllDirectoryStack;
 TArray<FString> FWindowsPlatformProcess::DllDirectories;
-
+bool IsJobObjectSet = false;
+HANDLE GhJob = NULL;
 
 void FWindowsPlatformProcess::AddDllDirectory(const TCHAR* Directory)
 {
@@ -712,8 +713,28 @@ void FWindowsPlatformProcess::ReadFromPipes(FString* OutStrings[], HANDLE InPipe
  * Executes a process, returning the return code, stdout, and stderr. This
  * call blocks until the process has returned.
  */
-bool FWindowsPlatformProcess::ExecProcess(const TCHAR* URL, const TCHAR* Params, int32* OutReturnCode, FString* OutStdOut, FString* OutStdErr, const TCHAR* OptionalWorkingDirectory)
+bool FWindowsPlatformProcess::ExecProcess(const TCHAR* URL, const TCHAR* Params, int32* OutReturnCode, FString* OutStdOut, FString* OutStdErr, const TCHAR* OptionalWorkingDirectory, bool bShouldEndWithParentProcess)
 {
+	if (bShouldEndWithParentProcess && !IsJobObjectSet)
+	{
+		GhJob = CreateJobObject(NULL, NULL);
+		if (!GhJob)
+		{
+			UE_LOG(LogWindows, Warning, TEXT("Failed to create Job Object"));
+		}
+		else
+		{
+			JOBOBJECT_EXTENDED_LIMIT_INFORMATION LimitInformation = { 0 };
+
+			LimitInformation.BasicLimitInformation.LimitFlags = JOB_OBJECT_LIMIT_KILL_ON_JOB_CLOSE;
+			if (0 == SetInformationJobObject(GhJob, JobObjectExtendedLimitInformation, &LimitInformation, sizeof(LimitInformation)))
+			{
+				UE_LOG(LogWindows, Warning, TEXT("Could not SetInformationJobObject"));
+			}
+		}
+		IsJobObjectSet = true;
+	}
+
 	STARTUPINFOEX StartupInfoEx;
 	ZeroMemory(&StartupInfoEx, sizeof(StartupInfoEx));
 	StartupInfoEx.StartupInfo.cb = sizeof(StartupInfoEx);
@@ -777,6 +798,14 @@ bool FWindowsPlatformProcess::ExecProcess(const TCHAR* URL, const TCHAR* Params,
 	PROCESS_INFORMATION ProcInfo;
 	if (CreateProcess(NULL, CommandLine.GetCharArray().GetData(), NULL, NULL, TRUE, CreateFlags, NULL, OptionalWorkingDirectory, &StartupInfoEx.StartupInfo, &ProcInfo))
 	{
+		if (bShouldEndWithParentProcess && GhJob)
+		{
+			int RetVal = AssignProcessToJobObject(GhJob, ProcInfo.hProcess);
+			if (RetVal == 0)
+			{
+				UE_LOG(LogWindows, Warning, TEXT("AssignProcessToObject failed."));
+			}
+		}
 		if (hStdOutRead != NULL)
 		{
 			HANDLE ReadablePipes[2] = { hStdOutRead, hStdErrRead };
@@ -797,7 +826,7 @@ bool FWindowsPlatformProcess::ExecProcess(const TCHAR* URL, const TCHAR* Params,
 			};
 
 			FProcHandle ProcHandle(ProcInfo.hProcess);
-			do 
+			do
 			{
 				ReadPipes();
 				FPlatformProcess::Sleep(0);
@@ -817,7 +846,7 @@ bool FWindowsPlatformProcess::ExecProcess(const TCHAR* URL, const TCHAR* Params,
 		else
 		{
 			::WaitForSingleObject(ProcInfo.hProcess, INFINITE);
-		}		
+		}
 		if (OutReturnCode)
 		{
 			verify(::GetExitCodeProcess(ProcInfo.hProcess, (DWORD*)OutReturnCode));
diff --git a/Engine/Source/Runtime/Core/Public/GenericPlatform/GenericPlatformProcess.h b/Engine/Source/Runtime/Core/Public/GenericPlatform/GenericPlatformProcess.h
--- a/Engine/Source/Runtime/Core/Public/GenericPlatform/GenericPlatformProcess.h
+++ b/Engine/Source/Runtime/Core/Public/GenericPlatform/GenericPlatformProcess.h
@@ -468,8 +468,9 @@ struct CORE_API FGenericPlatformProcess
 	 * @param OutStdOut may be 0
 	 * @param OutStdErr may be 0
 	 * @OptionalWorkingDirectory may be 0
+	 * @OptionalbShouldEndWithParentProcess false by default. True to make sure the process is killed with the parent processor (Not Supported on all Platforms)
 	 */
-	static bool ExecProcess(const TCHAR* URL, const TCHAR* Params, int32* OutReturnCode, FString* OutStdOut, FString* OutStdErr, const TCHAR* OptionalWorkingDirectory = NULL);
+	static bool ExecProcess(const TCHAR* URL, const TCHAR* Params, int32* OutReturnCode, FString* OutStdOut, FString* OutStdErr, const TCHAR* OptionalWorkingDirectory = NULL, bool bShouldEndWithParentProcess = false);
 
 	/**
 	 * Executes a process as administrator, requesting elevation as necessary. This
diff --git a/Engine/Source/Runtime/Core/Public/Mac/MacPlatformProcess.h b/Engine/Source/Runtime/Core/Public/Mac/MacPlatformProcess.h
--- a/Engine/Source/Runtime/Core/Public/Mac/MacPlatformProcess.h
+++ b/Engine/Source/Runtime/Core/Public/Mac/MacPlatformProcess.h
@@ -261,7 +261,7 @@ struct CORE_API FMacPlatformProcess : public FGenericPlatformProcess
 	static bool ReadPipeToArray(void* ReadPipe, TArray<uint8> & Output);
 	static bool WritePipe(void* WritePipe, const FString& Message, FString* OutWritten = nullptr);
 	static bool WritePipe(void* WritePipe, const uint8* Data, const int32 DataLength, int32* OutDataLength = nullptr);
-	static bool ExecProcess(const TCHAR* URL, const TCHAR* Params, int32* OutReturnCode, FString* OutStdOut, FString* OutStdErr, const TCHAR* OptionalWorkingDirectory = NULL);
+	static bool ExecProcess(const TCHAR* URL, const TCHAR* Params, int32* OutReturnCode, FString* OutStdOut, FString* OutStdErr, const TCHAR* OptionalWorkingDirectory = NULL, bool bShouldEndWithParentProcess = false);
 	static void SetThreadAffinityMask(uint64 AffinityMask);
 
 	// Mac specific
diff --git a/Engine/Source/Runtime/Core/Public/Unix/UnixPlatformProcess.h b/Engine/Source/Runtime/Core/Public/Unix/UnixPlatformProcess.h
--- a/Engine/Source/Runtime/Core/Public/Unix/UnixPlatformProcess.h
+++ b/Engine/Source/Runtime/Core/Public/Unix/UnixPlatformProcess.h
@@ -292,7 +292,7 @@ struct CORE_API FUnixPlatformProcess : public FGenericPlatformProcess
 	static bool Daemonize();
 	static bool IsApplicationRunning( uint32 ProcessId );
 	static bool IsApplicationRunning( const TCHAR* ProcName );
-	static bool ExecProcess(const TCHAR* URL, const TCHAR* Params, int32* OutReturnCode, FString* OutStdOut, FString* OutStdErr, const TCHAR* OptionalWorkingDirectory = NULL);
+	static bool ExecProcess(const TCHAR* URL, const TCHAR* Params, int32* OutReturnCode, FString* OutStdOut, FString* OutStdErr, const TCHAR* OptionalWorkingDirectory = NULL, bool bShouldEndWithParentProcess = false);
 	static void ExploreFolder( const TCHAR* FilePath );
 	static void LaunchFileInDefaultExternalApplication( const TCHAR* FileName, const TCHAR* Parms = NULL, ELaunchVerb::Type Verb = ELaunchVerb::Open );
 	static bool IsFirstInstance();
diff --git a/Engine/Source/Runtime/Core/Public/Windows/WindowsPlatformProcess.h b/Engine/Source/Runtime/Core/Public/Windows/WindowsPlatformProcess.h
--- a/Engine/Source/Runtime/Core/Public/Windows/WindowsPlatformProcess.h
+++ b/Engine/Source/Runtime/Core/Public/Windows/WindowsPlatformProcess.h
@@ -167,7 +167,7 @@ public:
 	static bool IsApplicationRunning( uint32 ProcessId );
 	static bool IsApplicationRunning( const TCHAR* ProcName );
 	static FString GetApplicationName( uint32 ProcessId );	
-	static bool ExecProcess(const TCHAR* URL, const TCHAR* Params, int32* OutReturnCode, FString* OutStdOut, FString* OutStdErr, const TCHAR* OptionalWorkingDirectory = NULL);
+	static bool ExecProcess(const TCHAR* URL, const TCHAR* Params, int32* OutReturnCode, FString* OutStdOut, FString* OutStdErr, const TCHAR* OptionalWorkingDirectory = NULL, bool bShouldEndWithParentProcess  = false);
 	static bool ExecElevatedProcess(const TCHAR* URL, const TCHAR* Params, int32* OutReturnCode);
 	static void LaunchFileInDefaultExternalApplication( const TCHAR* FileName, const TCHAR* Parms = NULL, ELaunchVerb::Type Verb = ELaunchVerb::Open );
 	static void ExploreFolder( const TCHAR* FilePath );
