Change Frame events into a single event with absolute timestamps

diff --git a/Engine/Source/Runtime/Core/Private/ProfilingDebugging/MiscTrace.cpp b/Engine/Source/Runtime/Core/Private/ProfilingDebugging/MiscTrace.cpp
--- a/Engine/Source/Runtime/Core/Private/ProfilingDebugging/MiscTrace.cpp
+++ b/Engine/Source/Runtime/Core/Private/ProfilingDebugging/MiscTrace.cpp
@@ -23,26 +23,16 @@ UE_TRACE_EVENT_BEGIN(Misc, Bookmark)
 	UE_TRACE_EVENT_FIELD(const void*, BookmarkPoint)
 UE_TRACE_EVENT_END()
 
-UE_TRACE_EVENT_BEGIN(Misc, BeginGameFrame)
-UE_TRACE_EVENT_END()
-
-UE_TRACE_EVENT_BEGIN(Misc, EndGameFrame)
-UE_TRACE_EVENT_END()
-
-UE_TRACE_EVENT_BEGIN(Misc, BeginRenderFrame)
+UE_TRACE_EVENT_BEGIN(Misc, BeginFrame)
+	UE_TRACE_EVENT_FIELD(uint64, Cycle)
+	UE_TRACE_EVENT_FIELD(uint8, FrameType)
 UE_TRACE_EVENT_END()
 
-UE_TRACE_EVENT_BEGIN(Misc, EndRenderFrame)
+UE_TRACE_EVENT_BEGIN(Misc, EndFrame)
+	UE_TRACE_EVENT_FIELD(uint64, Cycle)
+	UE_TRACE_EVENT_FIELD(uint8, FrameType)
 UE_TRACE_EVENT_END()
 
-
-struct FMiscTraceInternal
-{
-	static uint64 LastFrameCycle[TraceFrameType_Count];
-};
-
-uint64 FMiscTraceInternal::LastFrameCycle[TraceFrameType_Count] = { 0, 0 };
-
 void FMiscTrace::OutputBookmarkSpec(const void* BookmarkPoint, const ANSICHAR* File, int32 Line, const TCHAR* Format)
 {
 	uint16 FileNameLen = uint16(strlen(File));
@@ -70,23 +60,11 @@ void FMiscTrace::OutputBeginFrame(ETraceFrameType FrameType)
 	{
 		return;
 	}
+
 	uint64 Cycle = FPlatformTime::Cycles64();
-	uint64 CycleDiff = Cycle - FMiscTraceInternal::LastFrameCycle[FrameType];
-	FMiscTraceInternal::LastFrameCycle[FrameType] = Cycle;
-	uint8 Buffer[10];
-	uint8* BufferPtr = Buffer;
-	FTraceUtils::Encode7bit(CycleDiff, BufferPtr);
-	uint16 BufferSize = (uint16)(BufferPtr - Buffer);
-	if (FrameType == TraceFrameType_Game)
-	{
-		UE_TRACE_LOG(Misc, BeginGameFrame, FrameChannel, BufferSize)
-			<< BeginGameFrame.Attachment(&Buffer, BufferSize);
-	}
-	else if (FrameType == TraceFrameType_Rendering)
-	{
-		UE_TRACE_LOG(Misc, BeginRenderFrame, FrameChannel, BufferSize)
-			<< BeginRenderFrame.Attachment(&Buffer, BufferSize);
-	}
+	UE_TRACE_LOG(Misc, BeginFrame, FrameChannel)
+		<< BeginFrame.Cycle(Cycle)
+		<< BeginFrame.FrameType((uint8)FrameType);
 }
 
 void FMiscTrace::OutputEndFrame(ETraceFrameType FrameType)
@@ -95,23 +73,11 @@ void FMiscTrace::OutputEndFrame(ETraceFrameType FrameType)
 	{
 		return;
 	}
+
 	uint64 Cycle = FPlatformTime::Cycles64();
-	uint64 CycleDiff = Cycle - FMiscTraceInternal::LastFrameCycle[FrameType];
-	FMiscTraceInternal::LastFrameCycle[FrameType] = Cycle;
-	uint8 Buffer[10];
-	uint8* BufferPtr = Buffer;
-	FTraceUtils::Encode7bit(CycleDiff, BufferPtr);
-	uint16 BufferSize = (uint16)(BufferPtr - Buffer);
-	if (FrameType == TraceFrameType_Game)
-	{
-		UE_TRACE_LOG(Misc, EndGameFrame, FrameChannel, BufferSize)
-			<< EndGameFrame.Attachment(&Buffer, BufferSize);
-	}
-	else if (FrameType == TraceFrameType_Rendering)
-	{
-		UE_TRACE_LOG(Misc, EndRenderFrame, FrameChannel, BufferSize)
-			<< EndRenderFrame.Attachment(&Buffer, BufferSize);
-	}
+	UE_TRACE_LOG(Misc, EndFrame, FrameChannel)
+		<< EndFrame.Cycle(Cycle)
+		<< EndFrame.FrameType((uint8)FrameType);
 }
 
 #endif
