Duplicated MakeAttributeLambda and others of its kind to static functions inside TAttribute.

diff --git a/Engine/Source/Runtime/Core/Public/Misc/Attribute.h b/Engine/Source/Runtime/Core/Public/Misc/Attribute.h
--- a/Engine/Source/Runtime/Core/Public/Misc/Attribute.h
+++ b/Engine/Source/Runtime/Core/Public/Misc/Attribute.h
@@ -110,6 +110,52 @@ public:
 		return TAttribute( FGetter::CreateStatic( InFuncPtr ), bExplicitConstructor );
 	}
 
+	/**
+	 * Helper function for creating TAttributes from a non-const member function pointer, accessed through a raw pointer
+	 */
+	template<typename SourceType, typename SourceTypeOrBase, typename... PayloadTypes>
+	FORCEINLINE static TAttribute CreateRaw(SourceType* InObject, ObjectType (SourceTypeOrBase::*InMethod)(PayloadTypes...), typename TDecay<PayloadTypes>::Type... InputPayload)
+	{
+		return Create(FGetter::CreateRaw(InObject, InMethod, MoveTemp(InputPayload)...));
+	}
+
+	/**
+	 * Helper function for creating TAttributes from a const member function pointer, accessed through a raw pointer
+	 */
+	template<typename SourceType, typename SourceTypeOrBase, typename... PayloadTypes>
+	FORCEINLINE static TAttribute CreateRaw(const SourceType* InObject, ObjectType (SourceTypeOrBase::*InMethod)(PayloadTypes...) const, typename TDecay<PayloadTypes>::Type... InputPayload)
+	{
+		return Create(FGetter::CreateRaw(InObject, InMethod, MoveTemp(InputPayload)...));
+	}
+
+	/**
+	 * Helper function for creating TAttributes from a non-const member function pointer, accessed through a weak pointer to the shared object
+	 */
+	template<typename SourceType, typename SourceTypeOrBase, typename... PayloadTypes>
+	FORCEINLINE static TAttribute CreateSP(SourceType* InObject, ObjectType (SourceTypeOrBase::*InMethod)(PayloadTypes...), typename TDecay<PayloadTypes>::Type... InputPayload)
+	{
+		return Create(FGetter::CreateSP(InObject, InMethod, MoveTemp(InputPayload)...));
+	}
+
+	/**
+	 * Helper function for creating TAttributes from a const member function pointer, accessed through a weak pointer to the shared object
+	 */
+	template<typename SourceType, typename SourceTypeOrBase, typename... PayloadTypes>
+	FORCEINLINE static TAttribute CreateSP(const SourceType* InObject, ObjectType (SourceTypeOrBase::*InMethod)(PayloadTypes...) const, typename TDecay<PayloadTypes>::Type... InputPayload)
+	{
+		return Create(FGetter::CreateSP(InObject, InMethod, MoveTemp(InputPayload)...));
+	}
+
+	/**
+	 * Helper function for creating TAttributes from a lambda
+	 * TAttribute<float> FloatAttribute = TAttribute<float>::CreateLambda([]{ return 10.f; });
+	 */
+	template<typename LambdaType, typename... PayloadTypes>
+	FORCEINLINE static TAttribute CreateLambda(LambdaType&& InCallable, PayloadTypes&&... InputPayload)
+	{
+		return Create(FGetter::CreateLambda(InCallable, Forward<PayloadTypes>(InputPayload)...));
+	}
+
 	/**
 	 * Sets the attribute's value
 	 *
