Add a generic topological sort algorithm

diff --git a/Engine/Source/Runtime/Core/Private/Tests/Algo/TopologicalSort.cpp b/Engine/Source/Runtime/Core/Private/Tests/Algo/TopologicalSort.cpp
new file mode 100644
--- /dev/null
+++ b/Engine/Source/Runtime/Core/Private/Tests/Algo/TopologicalSort.cpp
@@ -0,0 +1,91 @@
+// Copyright Epic Games, Inc. All Rights Reserved.
+
+#include "Algo/TopologicalSort.h"
+#include "Misc/AutomationTest.h"
+
+#if WITH_DEV_AUTOMATION_TESTS
+
+IMPLEMENT_SIMPLE_AUTOMATION_TEST(FSystemCoreAlgoTopologicalSortTest, "System.Core.Algo.TopologicalSort", EAutomationTestFlags::ApplicationContextMask | EAutomationTestFlags::EngineFilter)
+
+bool FSystemCoreAlgoTopologicalSortTest::RunTest(const FString& Parameters)
+{
+	using namespace Algo;
+
+	{
+		TArray<int32> Array{ 1, 2, 3 };
+
+		// Test the sort when each node depends on the previous one
+		bool bHasSucceeded = TopologicalSort(Array, [](int32 Element) { return Element > 1 ? TArray<int32>{ Element - 1 } : TArray<int32>{}; });
+		TestEqual(TEXT("TopologicalSort did not sort correctly"), bHasSucceeded, true);
+		TestEqual(TEXT("TopologicalSort did not sort correctly"), Array, TArray<int32>{1, 2, 3});
+	}
+	{
+		TArray<int32> Array{ 1, 2, 3 };
+		
+		// Test the sort when each node depends on the next one
+		bool bHasSucceeded = TopologicalSort(Array, [](int32 Element) { return Element < 3 ? TArray<int32>{ Element + 1 } : TArray<int32>{}; });
+		TestEqual(TEXT("TopologicalSort did not sort correctly"), bHasSucceeded, true);
+		TestEqual(TEXT("TopologicalSort did not sort correctly"), Array, TArray<int32>{3, 2, 1});
+	}
+	{
+		TArray<int32> Array{ 1, 2 };
+		
+		// Test the sort with a cycle between 1 and 2
+		bool bHasSucceeded = TopologicalSort(Array, [](int32 Element) { return TArray<int32>{ 1 + Element % 2 }; });
+		TestEqual(TEXT("TopologicalSort should not have succeeded when a cycle is detected"), bHasSucceeded, false);
+		TestEqual(TEXT("TopologicalSort should not have modified the array when failing"), Array, TArray<int32>{1, 2});
+	}
+	{
+		TArray<int32> Array;
+		for (int32 Index = 0; Index < 1000; ++Index)
+		{
+			Array.Add(Index);
+		}
+
+		// Make sure node 500 makes it on top if every other node depends on it
+		bool bHasSucceeded = TopologicalSort(Array, [](int32 Element) { return Element == 500 ? TArray<int32>{} : TArray<int32>{ 500 }; });
+		TestEqual(TEXT("TopologicalSort did not sort correctly"), bHasSucceeded, true);
+		TestEqual(TEXT("TopologicalSort did not sort correctly"), Array[0], 500);
+	}
+	{
+		TArray<int32> Array{ 1, 2, 3, 4, 5, 6, 7, 8, 9, 10 };
+		//              7
+		//             / \
+		//            6   8
+		//           / \   \
+		//          1   2   9
+ 		//           \ / \   \
+		//            4   5   |
+		//             \   \ /
+		//              \   10
+		//               \ /
+		//                3       
+		TMultiMap<int32, int32> Links;
+		Links.Add(6, 7);
+		Links.Add(1, 6);
+		Links.Add(4, 1);
+		Links.Add(3, 4);
+		Links.Add(3, 10);
+		Links.Add(10, 5);
+		Links.Add(10, 9);
+		Links.Add(9, 8);
+		Links.Add(8, 7);
+		Links.Add(2, 6);
+		Links.Add(5, 2);
+
+		bool bHasSucceeded = TopologicalSort(Array, [&Links](int32 Element) { TArray<int32> Dependencies; Links.MultiFind(Element, Dependencies); return Dependencies; });
+		TestEqual(TEXT("TopologicalSort did not sort correctly"), bHasSucceeded, true);
+
+		// There might be multiple valid answers, so test each condition separately to make sure they are all met.
+		for (const auto& Pair : Links)
+		{
+			int32 ChildIndex = Array.IndexOfByKey(Pair.Key);
+			int32 ParentIndex = Array.IndexOfByKey(Pair.Value);
+			TestEqual(TEXT("TopologicalSort did not sort correctly"), ParentIndex < ChildIndex, true);
+		}
+	}
+
+	return true;
+}
+
+#endif //WITH_DEV_AUTOMATION_TESTS
diff --git a/Engine/Source/Runtime/Core/Public/Algo/KahnTopologicalSort.h b/Engine/Source/Runtime/Core/Public/Algo/KahnTopologicalSort.h
new file mode 100644
--- /dev/null
+++ b/Engine/Source/Runtime/Core/Public/Algo/KahnTopologicalSort.h
@@ -0,0 +1,74 @@
+// Copyright Epic Games, Inc. All Rights Reserved.
+
+#pragma once
+
+#include "Traits/ElementType.h"
+#include "Containers/Array.h"
+#include "Containers/Map.h"
+
+namespace Algo
+{
+	template <typename RangeType, typename GetElementDependenciesType>
+	bool KahnTopologicalSort(RangeType& UniqueRange, GetElementDependenciesType GetElementDependencies)
+	{
+		using ElementType = typename TElementType<RangeType>::Type;
+		using ElementMultiMap = TMultiMap<ElementType, ElementType>;
+
+		ElementMultiMap Dependencies;
+		for (const ElementType& Element : UniqueRange)
+		{
+			for (const ElementType& Dependency : GetElementDependencies(Element))
+			{
+				Dependencies.Add(Element, Dependency);
+			}
+		}
+
+		TArray<ElementType> SortedRange;
+		TArray<ElementType> IndependentElements;
+		TArray<ElementType> UniqueElements;
+		UniqueElements.Reserve(GetNum(UniqueRange));
+		for (const ElementType& Element : UniqueRange)
+		{
+			UniqueElements.Add(Element);
+		}
+
+		// Sort graph so that vertices with no dependencies always go first.
+		while (UniqueElements.Num() > 0)
+		{
+			IndependentElements = UniqueElements.FilterByPredicate([&](const ElementType& InVertex) {
+				return (Dependencies.Num(InVertex) == 0);
+			});
+
+			if (0 == IndependentElements.Num())
+			{
+				// If there are no independent vertices, then likely there is a 
+				// cycle in the graph.
+				return false;
+			}
+
+			ElementMultiMap UpdatedDependencies;
+
+			// Remove independent vertices from dependency map.
+			for (const auto& Element : Dependencies)
+			{
+				if (!IndependentElements.Contains(Element.Value))
+				{
+					UpdatedDependencies.AddUnique(Element.Key, Element.Value);
+				}
+			}
+			Dependencies = MoveTemp(UpdatedDependencies);
+
+			// Remove independent vertices from node list
+			for (const ElementType& Vertex : IndependentElements)
+			{
+				UniqueElements.RemoveSwap(Vertex);
+			}
+
+			// Add independent vertices to output.
+			SortedRange.Append(MoveTemp(IndependentElements));
+		}
+
+		UniqueRange = MoveTemp(SortedRange);
+		return true;
+	}
+}
\ No newline at end of file
diff --git a/Engine/Source/Runtime/Core/Public/Algo/TopologicalSort.h b/Engine/Source/Runtime/Core/Public/Algo/TopologicalSort.h
new file mode 100644
--- /dev/null
+++ b/Engine/Source/Runtime/Core/Public/Algo/TopologicalSort.h
@@ -0,0 +1,14 @@
+// Copyright Epic Games, Inc. All Rights Reserved.
+
+#pragma once
+
+#include "Algo/KahnTopologicalSort.h"
+
+namespace Algo
+{
+	template <typename RangeType, typename GetElementDependenciesType>
+	FORCEINLINE bool TopologicalSort(RangeType& UniqueRange, GetElementDependenciesType GetElementDependencies)
+	{
+		return KahnTopologicalSort(UniqueRange, GetElementDependencies);
+	}
+}
\ No newline at end of file
