Two related parts:

1) TlsAutoCleanup: in the previous version `FTlsAutoCleanup` used `FRunnableThread` to store registered instances and to delete them on `FRunnableThread` destruction. This means that all threads that are not created by `FRunnableThread` (the game thread in particular, and all 3rd-party threads) never delete registered `FTlsAutoCleanup` instances. This change moves the storage from `FRunnableThread` to TLS to cover all threads

Thsi change can break systems that implemented a workaround for this problem, e.g. by manually deleting registered `FTlsAutoCleanup` instances on the game thread exit, as now this instances will be deleted twice. The solution is to remove the workaround and to rely on uniform `FTlsAutoCleanup` behaviour on all threads.

2) removed CpuProfilerTrace workaround to handle the fact that FTlsAutoCleanup leaks on the game thread, as this is fixed by (1).

Need to submit both parts together as they need each other

diff --git a/Engine/Source/Runtime/Core/Private/HAL/ThreadingBase.cpp b/Engine/Source/Runtime/Core/Private/HAL/ThreadingBase.cpp
--- a/Engine/Source/Runtime/Core/Private/HAL/ThreadingBase.cpp
+++ b/Engine/Source/Runtime/Core/Private/HAL/ThreadingBase.cpp
@@ -779,13 +779,6 @@ void FRunnableThread::FreeTls()
 	check( ThreadID == FPlatformTLS::GetCurrentThreadId() );
 	check( FPlatformTLS::IsValidTlsSlot(RunnableTlsSlot) );
 	FPlatformTLS::SetTlsValue( RunnableTlsSlot, nullptr );
-
-	// Delete all FTlsAutoCleanup objects created for this thread.
-	for( auto& Instance : TlsInstances )
-	{
-		delete Instance;
-		Instance = nullptr;
-	}
 }
 
 /*-----------------------------------------------------------------------------
@@ -1322,11 +1315,8 @@ FTlsAutoCleanup* FThreadSingletonInitializer::TryGet(uint32& TlsSlot)
 
 void FTlsAutoCleanup::Register()
 {
-	FRunnableThread* RunnableThread = FRunnableThread::GetRunnableThread();
-	if( RunnableThread )
-	{
-		RunnableThread->TlsInstances.Add( this );
-	}
+	static thread_local TArray<TUniquePtr<FTlsAutoCleanup>> TlsInstances;
+	TlsInstances.Add(TUniquePtr<FTlsAutoCleanup>(this));
 }
 
 //-------------------------------------------------------------------------------
diff --git a/Engine/Source/Runtime/Core/Private/ProfilingDebugging/CpuProfilerTrace.cpp b/Engine/Source/Runtime/Core/Private/ProfilingDebugging/CpuProfilerTrace.cpp
--- a/Engine/Source/Runtime/Core/Private/ProfilingDebugging/CpuProfilerTrace.cpp
+++ b/Engine/Source/Runtime/Core/Private/ProfilingDebugging/CpuProfilerTrace.cpp
@@ -224,13 +224,4 @@ uint32 FCpuProfilerTrace::OutputEventType(const ANSICHAR* Name)
 	return SpecId;
 }
 
-void FCpuProfilerTrace::Shutdown()
-{
-	if (FCpuProfilerTraceInternal::ThreadBuffer)
-	{
-		delete FCpuProfilerTraceInternal::ThreadBuffer;
-		FCpuProfilerTraceInternal::ThreadBuffer = nullptr;
-	}
-}
-
 #endif
diff --git a/Engine/Source/Runtime/Core/Public/HAL/RunnableThread.h b/Engine/Source/Runtime/Core/Public/HAL/RunnableThread.h
--- a/Engine/Source/Runtime/Core/Public/HAL/RunnableThread.h
+++ b/Engine/Source/Runtime/Core/Public/HAL/RunnableThread.h
@@ -165,9 +165,6 @@ protected:
 	/** The Affinity to run the thread with. */
 	uint64 ThreadAffinityMask;
 
-	/** An array of FTlsAutoCleanup based instances that needs to be deleted before the thread will die. */
-	TArray<FTlsAutoCleanup*> TlsInstances;
-
 	/** The priority to run the thread at. */
 	EThreadPriority ThreadPriority;
 
diff --git a/Engine/Source/Runtime/Core/Public/HAL/TlsAutoCleanup.h b/Engine/Source/Runtime/Core/Public/HAL/TlsAutoCleanup.h
--- a/Engine/Source/Runtime/Core/Public/HAL/TlsAutoCleanup.h
+++ b/Engine/Source/Runtime/Core/Public/HAL/TlsAutoCleanup.h
@@ -4,7 +4,10 @@
 
 #include "CoreTypes.h"
 
-/** Base class for objects in TLS that support auto-cleanup. */
+/** 
+Base class for objects in TLS that support auto-cleanup. 
+Polymorphically deletes registered instances on thread exit.
+*/
 class CORE_API FTlsAutoCleanup
 {
 public:
diff --git a/Engine/Source/Runtime/Core/Public/ProfilingDebugging/CpuProfilerTrace.h b/Engine/Source/Runtime/Core/Public/ProfilingDebugging/CpuProfilerTrace.h
--- a/Engine/Source/Runtime/Core/Public/ProfilingDebugging/CpuProfilerTrace.h
+++ b/Engine/Source/Runtime/Core/Public/ProfilingDebugging/CpuProfilerTrace.h
@@ -31,7 +31,6 @@ UE_TRACE_CHANNEL_EXTERN(CpuChannel, CORE_API);
  */
 struct FCpuProfilerTrace
 {
-	CORE_API static void Shutdown();
 	/*
 	 * Output cpu event definition (spec).
 	 * @param Name Event name
@@ -120,9 +119,6 @@ struct FCpuProfilerTrace
 	};
 };
 
-#define TRACE_CPUPROFILER_SHUTDOWN() \
-	FCpuProfilerTrace::Shutdown();
-
 // Trace a scoped cpu timing event providing a static string (const ANSICHAR* or const TCHAR*)
 // as the scope name and a trace channel.
 // Example: TRACE_CPUPROFILER_EVENT_SCOPE_ON_CHANNEL_STR("My Scoped Timer A", CpuChannel)
@@ -174,7 +170,6 @@ struct FCpuProfilerTrace
 
 #else
 
-#define TRACE_CPUPROFILER_SHUTDOWN()
 #define TRACE_CPUPROFILER_EVENT_SCOPE_ON_CHANNEL_STR(NameStr, Channel)
 #define TRACE_CPUPROFILER_EVENT_SCOPE_ON_CHANNEL(Name, Channel)
 #define TRACE_CPUPROFILER_EVENT_SCOPE_STR(NameStr)
diff --git a/Engine/Source/Runtime/Launch/Private/LaunchEngineLoop.cpp b/Engine/Source/Runtime/Launch/Private/LaunchEngineLoop.cpp
--- a/Engine/Source/Runtime/Launch/Private/LaunchEngineLoop.cpp
+++ b/Engine/Source/Runtime/Launch/Private/LaunchEngineLoop.cpp
@@ -4425,8 +4425,6 @@ void FEngineLoop::Exit()
 	FTaskGraphInterface::Shutdown();
 
 	FPlatformMisc::ShutdownTaggedStorage();
-
-	TRACE_CPUPROFILER_SHUTDOWN();
 }
 
 
@@ -6006,11 +6004,6 @@ void FEngineLoop::AppExit()
 	}
 
 	FInternationalization::TearDown();
-
-#if !WITH_ENGINE
-	// when compiled WITH_ENGINE, this will happen in FEngineLoop::Exit()
-	TRACE_CPUPROFILER_SHUTDOWN();
-#endif
 }
 
 void FEngineLoop::PostInitRHI()
