Ensure that the localization and internationalization systems shutdown in the correct order

diff --git a/Engine/Source/Runtime/Core/Private/Internationalization/Internationalization.cpp b/Engine/Source/Runtime/Core/Private/Internationalization/Internationalization.cpp
--- a/Engine/Source/Runtime/Core/Private/Internationalization/Internationalization.cpp
+++ b/Engine/Source/Runtime/Core/Private/Internationalization/Internationalization.cpp
@@ -34,7 +34,6 @@ bool FInternationalization::IsAvailable()
 void FInternationalization::TearDown()
 {
 	TLazySingleton<FInternationalization>::TearDown();
-	FTextCache::TearDown();
 }
 
 FText FInternationalization::ForUseOnlyByLocMacroAndGraphNodeTextLiterals_CreateText(const TCHAR* InTextLiteral, const TCHAR* InNamespace, const TCHAR* InKey)
diff --git a/Engine/Source/Runtime/Core/Private/Internationalization/TextLocalizationManager.cpp b/Engine/Source/Runtime/Core/Private/Internationalization/TextLocalizationManager.cpp
--- a/Engine/Source/Runtime/Core/Private/Internationalization/TextLocalizationManager.cpp
+++ b/Engine/Source/Runtime/Core/Private/Internationalization/TextLocalizationManager.cpp
@@ -18,6 +18,7 @@
 #include "Internationalization/Internationalization.h"
 #include "Internationalization/StringTableCore.h"
 #include "Internationalization/TextNamespaceUtil.h"
+#include "Internationalization/TextCache.h"
 #include "Stats/Stats.h"
 #include "Misc/ConfigCacheIni.h"
 #include "Misc/App.h"
@@ -463,6 +464,7 @@ FTextLocalizationManager& FTextLocalizationManager::Get()
 
 void FTextLocalizationManager::TearDown()
 {
+	FTextCache::TearDown();
 	TLazySingleton<FTextLocalizationManager>::TearDown();
 	FTextKey::TearDown();
 }
diff --git a/Engine/Source/Runtime/Launch/Private/LaunchEngineLoop.cpp b/Engine/Source/Runtime/Launch/Private/LaunchEngineLoop.cpp
--- a/Engine/Source/Runtime/Launch/Private/LaunchEngineLoop.cpp
+++ b/Engine/Source/Runtime/Launch/Private/LaunchEngineLoop.cpp
@@ -6003,6 +6003,7 @@ void FEngineLoop::AppExit()
 		GLog->TearDown();
 	}
 
+	FTextLocalizationManager::TearDown();
 	FInternationalization::TearDown();
 }
 
