Fix Clang build on Windows

diff --git a/Engine/Source/Runtime/Core/Private/Windows/WindowsPlatformProcess.cpp b/Engine/Source/Runtime/Core/Private/Windows/WindowsPlatformProcess.cpp
--- a/Engine/Source/Runtime/Core/Private/Windows/WindowsPlatformProcess.cpp
+++ b/Engine/Source/Runtime/Core/Private/Windows/WindowsPlatformProcess.cpp
@@ -724,7 +724,7 @@ bool FWindowsPlatformProcess::ExecProcess(const TCHAR* URL, const TCHAR* Params,
 		}
 		else
 		{
-			JOBOBJECT_EXTENDED_LIMIT_INFORMATION LimitInformation = { 0 };
+			JOBOBJECT_EXTENDED_LIMIT_INFORMATION LimitInformation = {{ 0 }};
 
 			LimitInformation.BasicLimitInformation.LimitFlags = JOB_OBJECT_LIMIT_KILL_ON_JOB_CLOSE;
 			if (0 == SetInformationJobObject(GhJob, JobObjectExtendedLimitInformation, &LimitInformation, sizeof(LimitInformation)))
