Fixup FSerializedUATProcess URL to shell out correctly on Linux and Mac.

#rnx

diff --git a/Engine/Source/Runtime/Core/Private/Misc/MonitoredProcess.cpp b/Engine/Source/Runtime/Core/Private/Misc/MonitoredProcess.cpp
--- a/Engine/Source/Runtime/Core/Private/Misc/MonitoredProcess.cpp
+++ b/Engine/Source/Runtime/Core/Private/Misc/MonitoredProcess.cpp
@@ -190,6 +190,10 @@ void FMonitoredProcess::Tick()
 
 /* FSerializedUATProcess
 *****************************************************************************/
+
+FCriticalSection FSerializedUATProcess::Serializer;
+bool FSerializedUATProcess::bHasSucceededOnce = false;
+
 FSerializedUATProcess::FSerializedUATProcess(const FString& RunUATCommandline)
 	// we will modify URL and Params in this constructor, so there's no need to pass anything up to base
 	: FMonitoredProcess("", "", true, true)
@@ -199,31 +203,22 @@ FSerializedUATProcess::FSerializedUATProcess(const FString& RunUATCommandline)
 #if PLATFORM_WINDOWS
 	URL = TEXT("cmd.exe");
 	Params = FString::Printf(TEXT("/c \"\"%s\" %s\""), *GetUATPath(), *RunUATCommandline);
-#elif PLATFORM_LINUX
+#elif PLATFORM_MAC || PLATFORM_LINUX
 	URL = TEXT("/bin/bash");
-	Params = FString::Printf(TEXT("\"%s\" %s"), *GetUATPath(), *RunUATCommandline);
-#elif PLATFORM_MAC
-	URL = TEXT("/bin/sh");
-	Params = FString::Printf(TEXT("\"%s\" %s"), *GetUATPath(), *RunUATCommandline);
+	Params = FString::Printf(TEXT("-c '\"%s\" %s'"), *GetUATPath(), *RunUATCommandline);
 #endif
 }
 
 bool FSerializedUATProcess::Launch()
 {
-	static FCriticalSection SerializerCS;
-	static bool bHasSucceededOnceFlag = false;
-
-	FCriticalSection* Serializer = &SerializerCS;
-	bool* bHasSucceededOnce = &bHasSucceededOnceFlag;
-
-	Async(EAsyncExecution::Thread, [this, Serializer, bHasSucceededOnce]()
+	Async(EAsyncExecution::Thread, [this]()
 		{
 			// don't let this do anything if another process is running
-			FScopeLock Lock(Serializer);
+			FScopeLock Lock(&Serializer);
 
 			FEventRef Event;
 
-			if (*bHasSucceededOnce)
+			if (bHasSucceededOnce)
 			{
 				Params += TEXT(" -nocompile");
 			}
@@ -231,11 +226,11 @@ bool FSerializedUATProcess::Launch()
 			FOnMonitoredProcessCompleted OriginalCompletedDelegate = CompletedDelegate;
 			FSimpleDelegate OriginalCanceledDelegate = CanceledDelegate;
 
-			CompletedDelegate.BindLambda([this, &Event, bHasSucceededOnce, OriginalCompletedDelegate](int32 ExitCode)
+			CompletedDelegate.BindLambda([this, &Event, OriginalCompletedDelegate](int32 ExitCode)
 				{
 					if (ExitCode == 0 || ExitCode == 10)
 					{
-						*bHasSucceededOnce = true;
+						bHasSucceededOnce = true;
 					}
 					OriginalCompletedDelegate.ExecuteIfBound(ExitCode);
 
diff --git a/Engine/Source/Runtime/Core/Public/Misc/MonitoredProcess.h b/Engine/Source/Runtime/Core/Public/Misc/MonitoredProcess.h
--- a/Engine/Source/Runtime/Core/Public/Misc/MonitoredProcess.h
+++ b/Engine/Source/Runtime/Core/Public/Misc/MonitoredProcess.h
@@ -288,4 +288,7 @@ private:
 	// Holds a delegate that is executed when the process fails to launch (delayed, in a thread). Used in place of the return value
 	// of Launch in the parent class, since it's async
 	FSimpleDelegate LaunchFailedDelegate;
+
+	static FCriticalSection Serializer;
+	static bool bHasSucceededOnce;
 };
