Disable HWCPipe on non-ARM architectures

#ue5
#android

diff --git a/Engine/Source/Runtime/Core/Private/Android/AndroidStats.cpp b/Engine/Source/Runtime/Core/Private/Android/AndroidStats.cpp
--- a/Engine/Source/Runtime/Core/Private/Android/AndroidStats.cpp
+++ b/Engine/Source/Runtime/Core/Private/Android/AndroidStats.cpp
@@ -7,7 +7,7 @@
 #include "HAL/IConsoleManager.h"
 #include "ProfilingDebugging/CsvProfiler.h"
 
-#if !UE_BUILD_SHIPPING && !PLATFORM_LUMIN
+#if !UE_BUILD_SHIPPING && !PLATFORM_LUMIN && PLATFORM_ANDROID_ARM64
 #	define HWCPIPE_SUPPORTED 1
 #else
 #	define HWCPIPE_SUPPORTED 0
