Stats system is still using 32-bit cycle counts, so don't use Cycles64() which may have a different resolution.  Fix issue with hitch detection resulting from use of 64-bit count.
#rnx

diff --git a/Engine/Source/Runtime/Core/Private/Stats/StatsCommand.cpp b/Engine/Source/Runtime/Core/Private/Stats/StatsCommand.cpp
--- a/Engine/Source/Runtime/Core/Private/Stats/StatsCommand.cpp
+++ b/Engine/Source/Runtime/Core/Private/Stats/StatsCommand.cpp
@@ -674,8 +674,9 @@ static void DumpHitch(int64 Frame)
 	FStatsThreadState& Stats = FStatsThreadState::GetLocalState();
 	SCOPE_CYCLE_COUNTER(STAT_HitchScan);
 
-	const double GameThreadTime = FPlatformTime::ToSeconds64(Stats.GetFastThreadFrameTime(Frame, EThreadType::Game));
-	const double RenderThreadTime = FPlatformTime::ToSeconds64(Stats.GetFastThreadFrameTime(Frame, EThreadType::Renderer));
+	// #todo_cycles64 - remove narrowing to 32-bit
+	const double GameThreadTime = FPlatformTime::ToSeconds((uint32)Stats.GetFastThreadFrameTime(Frame, EThreadType::Game));
+	const double RenderThreadTime = FPlatformTime::ToSeconds((uint32)Stats.GetFastThreadFrameTime(Frame, EThreadType::Renderer));
 	const float HitchThresholdSecs = GHitchThresholdMS * 0.001f;
 
 	if ((GameThreadTime > HitchThresholdSecs) || (RenderThreadTime > HitchThresholdSecs))
diff --git a/Engine/Source/Runtime/Core/Public/Stats/Stats2.h b/Engine/Source/Runtime/Core/Public/Stats/Stats2.h
--- a/Engine/Source/Runtime/Core/Public/Stats/Stats2.h
+++ b/Engine/Source/Runtime/Core/Public/Stats/Stats2.h
@@ -711,7 +711,7 @@ struct FStatMessage
 		// these branches are FORCEINLINE_STATS of constants in almost all cases, so they disappear
 		if (InStatOperation == EStatOperation::CycleScopeStart || InStatOperation == EStatOperation::CycleScopeEnd)
 		{
-			GetValue_int64()= int64(FPlatformTime::Cycles64());
+			GetValue_int64()= int64(FPlatformTime::Cycles());
 		}
 		else
 		{
