Improve chances of displaying message box information on configurations falling back to the generic implementation.

#rnx

diff --git a/Engine/Source/Runtime/Core/Private/GenericPlatform/GenericPlatformMisc.cpp b/Engine/Source/Runtime/Core/Private/GenericPlatform/GenericPlatformMisc.cpp
--- a/Engine/Source/Runtime/Core/Private/GenericPlatform/GenericPlatformMisc.cpp
+++ b/Engine/Source/Runtime/Core/Private/GenericPlatform/GenericPlatformMisc.cpp
@@ -796,10 +796,9 @@ const TCHAR* LexToString( EAppReturnType::Type Value )
 
 EAppReturnType::Type FGenericPlatformMisc::MessageBoxExt( EAppMsgType::Type MsgType, const TCHAR* Text, const TCHAR* Caption )
 {
-	if (GWarn)
-	{
-		UE_LOG(LogGenericPlatformMisc, Warning, TEXT("MessageBox: %s : %s"), Caption, Text);
-	}
+	// a message box typically conveys important information, so try to make sure at least one message reaches the user
+	UE_LOG(LogGenericPlatformMisc, Warning, TEXT("MessageBox: %s : %s"), Caption, Text);
+	FPlatformMisc::LowLevelOutputDebugStringf(TEXT("MessageBox: %s : %s\n"), Caption, Text);
 
 	switch(MsgType)
 	{
