Tasks: added optional memory order parameter to `LowLevelTasks::FTask::IsCompleted` for fine tuning cases where sequential consistency is not required

diff --git a/Engine/Source/Runtime/Core/Public/Async/Fundamental/Task.h b/Engine/Source/Runtime/Core/Public/Async/Fundamental/Task.h
--- a/Engine/Source/Runtime/Core/Public/Async/Fundamental/Task.h
+++ b/Engine/Source/Runtime/Core/Public/Async/Fundamental/Task.h
@@ -150,9 +150,9 @@ namespace LowLevelTasks
 
 	public: //Public Interface
 		//means the task is completed and this taskhandle can be recycled
-		inline bool IsCompleted() const
+		inline bool IsCompleted(std::memory_order MemoryOrder = std::memory_order_seq_cst) const
 		{
-			ETaskState State = PackedData.load(std::memory_order_seq_cst).GetState();
+			ETaskState State = PackedData.load(MemoryOrder).GetState();
 			return State == ETaskState::CanceledAndCompleted || State == ETaskState::Completed;
 		}
 		
