Use Invoke instead of calling the functor directly.

diff --git a/Engine/Source/Runtime/Core/Public/Algo/KahnTopologicalSort.h b/Engine/Source/Runtime/Core/Public/Algo/KahnTopologicalSort.h
--- a/Engine/Source/Runtime/Core/Public/Algo/KahnTopologicalSort.h
+++ b/Engine/Source/Runtime/Core/Public/Algo/KahnTopologicalSort.h
@@ -5,6 +5,7 @@
 #include "Traits/ElementType.h"
 #include "Containers/Array.h"
 #include "Containers/Map.h"
+#include "Templates/Invoke.h"
 
 namespace Algo
 {
@@ -17,7 +18,7 @@ namespace Algo
 		ElementMultiMap Dependencies;
 		for (const ElementType& Element : UniqueRange)
 		{
-			for (const ElementType& Dependency : GetElementDependencies(Element))
+			for (const ElementType& Dependency : Invoke(GetElementDependencies, Element))
 			{
 				Dependencies.Add(Element, Dependency);
 			}
