Add Steal function to TAttribute to return a moved getter or value. The attribute need to be set.

diff --git a/Engine/Source/Runtime/Core/Public/Misc/Attribute.h b/Engine/Source/Runtime/Core/Public/Misc/Attribute.h
--- a/Engine/Source/Runtime/Core/Public/Misc/Attribute.h
+++ b/Engine/Source/Runtime/Core/Public/Misc/Attribute.h
@@ -3,6 +3,7 @@
 #pragma once
 
 #include "CoreTypes.h"
+#include "Misc/TVariant.h"
 #include "Templates/Function.h"
 #include "Templates/SharedPointer.h"
 #include "Delegates/Delegate.h"
@@ -343,14 +344,14 @@ public:
 	 * @param  InFunctionName Member function name to bind.
 	 */
 	template< class SourceType >
-	static TAttribute< ObjectType > Create(SourceType* InUserObject, const FName& InFunctionName)
+	UE_NODISCARD static TAttribute< ObjectType > Create(SourceType* InUserObject, const FName& InFunctionName)
 	{
 		TAttribute< ObjectType > Attrib;
 		Attrib.BindUFunction<SourceType>(InUserObject, InFunctionName);
 		return Attrib;
 	}
 
-	static TAttribute< ObjectType > Create(TFunction<ObjectType(void)>&& InLambda)
+	UE_NODISCARD static TAttribute< ObjectType > Create(TFunction<ObjectType(void)>&& InLambda)
 	{
 		return Create(TAttribute< ObjectType >::FGetter::CreateLambda(MoveTemp(InLambda)));
 	}
@@ -366,15 +367,36 @@ public:
 	}
 
 	/**
-	* Gets the attribute's 'getter' which can be bound or unbound
-	*
-	* @return  The attribute's FGetter.
-	*/
-	const FGetter& GetBinding() const
+	 * Gets the attribute's 'getter' which can be bound or unbound
+	 *
+	 * @return  The attribute's FGetter.
+	 */
+	UE_NODISCARD const FGetter& GetBinding() const
 	{
 		return Getter;
 	}
 
+	/**
+	 * Move the attribute's 'getter' or the attribute's `Value` and reset the attribute. The attribute needs to be set.
+	 *
+	 * @return  The attribute's FGetter or `Value`
+	 */
+	UE_NODISCARD TVariant<ObjectType, FGetter> Steal()
+	{
+		checkf(IsSet(), TEXT("It is an error to call Steal() on an unset TAttribute. Check IsSet() before calling Steal()."));
+		bIsSet = false;
+		TVariant<ObjectType, FGetter> Temp;
+		if (IsBound())
+		{
+			Temp.template Set<FGetter>(MoveTemp(Getter));
+		}
+		else
+		{
+			Temp.template Set<ObjectType>(MoveTemp(Value));
+		}
+		return Temp;
+	}
+
 	/**
 	 * Is this attribute identical to another TAttribute.
 	 *
