Use env to work around the funny (and conflicting) single/double quote rules of shell vs .NET Core on Linux/Mac.

#rnx

diff --git a/Engine/Source/Runtime/Core/Private/Misc/MonitoredProcess.cpp b/Engine/Source/Runtime/Core/Private/Misc/MonitoredProcess.cpp
--- a/Engine/Source/Runtime/Core/Private/Misc/MonitoredProcess.cpp
+++ b/Engine/Source/Runtime/Core/Private/Misc/MonitoredProcess.cpp
@@ -204,8 +204,8 @@ FSerializedUATProcess::FSerializedUATProcess(const FString& RunUATCommandline)
 	URL = TEXT("cmd.exe");
 	Params = FString::Printf(TEXT("/c \"\"%s\" %s\""), *GetUATPath(), *RunUATCommandline);
 #elif PLATFORM_MAC || PLATFORM_LINUX
-	URL = TEXT("/bin/bash");
-	Params = FString::Printf(TEXT("-c '\"%s\" %s'"), *GetUATPath(), *RunUATCommandline);
+	URL = TEXT("/usr/bin/env");
+	Params = FString::Printf(TEXT(" -- \"%s\" %s"), *GetUATPath(), *RunUATCommandline);
 #endif
 }
 
@@ -243,7 +243,7 @@ bool FSerializedUATProcess::Launch()
 					Event->Trigger();
 				});
 
-			UE_LOG(LogMonitoredProcess, Log, TEXT("Running Serialized UAT: '%s %s'"), *URL, *Params);
+			UE_LOG(LogMonitoredProcess, Log, TEXT("Running Serialized UAT: [ %s %s ]"), *URL, *Params);
 
 			if (FMonitoredProcess::Launch() == false)
 			{
