StringBuilder: Added WriteToString<N>(Args...) to create and write to a temporary string

diff --git a/Engine/Source/Runtime/Core/Public/Misc/StringBuilder.h b/Engine/Source/Runtime/Core/Public/Misc/StringBuilder.h
--- a/Engine/Source/Runtime/Core/Public/Misc/StringBuilder.h
+++ b/Engine/Source/Runtime/Core/Public/Misc/StringBuilder.h
@@ -391,3 +391,27 @@ inline FAnsiStringBuilderBase&		operator<<(FAnsiStringBuilderBase& Builder, int1
 inline FAnsiStringBuilderBase&		operator<<(FAnsiStringBuilderBase& Builder, uint16 Value)							{ return Builder << uint32(Value); }
 inline FWideStringBuilderBase&		operator<<(FWideStringBuilderBase& Builder, int16 Value)							{ return Builder << int32(Value); }
 inline FWideStringBuilderBase&		operator<<(FWideStringBuilderBase& Builder, uint16 Value)							{ return Builder << uint32(Value); }
+
+/**
+ * A function-like type that creates a TStringBuilder by appending its arguments.
+ *
+ * Example Use Cases:
+ *
+ * For void Action(FStringView) -> Action(WriteToString<64>(Arg1, Arg2));
+ * For UE_LOG or checkf -> checkf(Condition, TEXT("%s"), *WriteToString<32>(Arg));
+ */
+template <int32 BufferSize>
+class WriteToString : public TStringBuilder<BufferSize>
+{
+public:
+	template <typename... ArgTypes>
+	explicit WriteToString(ArgTypes&&... Args)
+	{
+	#if PLATFORM_COMPILER_HAS_FOLD_EXPRESSIONS
+		(*this << ... << Forward<ArgTypes>(Args));
+	#else
+		using Fold = int[];
+		void(Fold{0, (void(*this << Forward<ArgTypes>(Args)), 0)...});
+	#endif
+	}
+};
