U5 - undoing accidental checked-in test code

diff --git a/Engine/Source/Runtime/Core/Private/Windows/WindowsPlatformMisc.cpp b/Engine/Source/Runtime/Core/Private/Windows/WindowsPlatformMisc.cpp
--- a/Engine/Source/Runtime/Core/Private/Windows/WindowsPlatformMisc.cpp
+++ b/Engine/Source/Runtime/Core/Private/Windows/WindowsPlatformMisc.cpp
@@ -1868,8 +1868,8 @@ const TCHAR* FWindowsPlatformMisc::GetPlatformFeaturesModuleName()
 }
 
 int32 FWindowsPlatformMisc::NumberOfWorkerThreadsToSpawn()
-{	return 26;
-/*	static int32 MaxServerWorkerThreads = 4;
+{	
+	static int32 MaxServerWorkerThreads = 4;
 
 	extern CORE_API int32 GUseNewTaskBackend;
 	int32 MaxWorkerThreads = GUseNewTaskBackend ? INT32_MAX : 26;
@@ -1889,7 +1889,7 @@ int32 FWindowsPlatformMisc::NumberOfWorkerThreadsToSpawn()
 
 	int32 MaxWorkerThreadsWanted = IsRunningDedicatedServer() ? MaxServerWorkerThreads : MaxWorkerThreads;
 	// need to spawn at least one worker thread (see FTaskGraphImplementation)
-	return FMath::Max(FMath::Min(NumberOfThreads, MaxWorkerThreadsWanted), 2);*/
+	return FMath::Max(FMath::Min(NumberOfThreads, MaxWorkerThreadsWanted), 2);
 }
 
 bool FWindowsPlatformMisc::OsExecute(const TCHAR* CommandType, const TCHAR* Command, const TCHAR* CommandLine)
