Add an overload to the /= operator for FString that is compatible with FStringView/TStringBuilder/FString and removed the FString specific /= operator.

- To maintain backwards compatibility an overload for non TCHAR raw string types has also been added which will be converted to the correct format via a temp FString first.
  - Fixed cases of Fstring /= char* found in AndroidPlatformFile.cpp and PythonScriptPluginPreload.cpp
- Added unit tests for the /= operator when passing in TCHAR*, ANSICHAR*, FString, FStringView and TStringBuilder in various forms.
- Changed GetNumfor TStringBuilderBase to return int32 (which is what TStringBuilderBase::Len() returns anyway) rather than SIZE_T to prevent compiler warnings about casting types.

diff --git a/Engine/Plugins/Experimental/PythonScriptPlugin/Source/PythonScriptPluginPreload/Private/PythonScriptPluginPreload.cpp b/Engine/Plugins/Experimental/PythonScriptPlugin/Source/PythonScriptPluginPreload/Private/PythonScriptPluginPreload.cpp
--- a/Engine/Plugins/Experimental/PythonScriptPlugin/Source/PythonScriptPluginPreload/Private/PythonScriptPluginPreload.cpp
+++ b/Engine/Plugins/Experimental/PythonScriptPlugin/Source/PythonScriptPluginPreload/Private/PythonScriptPluginPreload.cpp
@@ -85,7 +85,7 @@ private:
 			const FString PythonDSOWildcard = FString::Printf(TEXT("python%d*.dll"), PY_MAJOR_VERSION);
 #elif PLATFORM_LINUX
 			const FString PythonDSOWildcard = FString::Printf(TEXT("libpython%d*.so*"), PY_MAJOR_VERSION);
-			PythonDir /= "lib";
+			PythonDir /= TEXT("lib");
 #endif
 			LoadSharedDSO(PythonDSOWildcard, PythonDir);
 		}
diff --git a/Engine/Source/Runtime/Core/Private/Android/AndroidPlatformFile.cpp b/Engine/Source/Runtime/Core/Private/Android/AndroidPlatformFile.cpp
--- a/Engine/Source/Runtime/Core/Private/Android/AndroidPlatformFile.cpp
+++ b/Engine/Source/Runtime/Core/Private/Android/AndroidPlatformFile.cpp
@@ -607,12 +607,12 @@ public:
 		FEntryMap::TIterator Current;
 		FString Path;
 
-		Directory(FEntryMap & entries, const FString & dirpath)
-			: Current(entries.CreateIterator()), Path(dirpath)
+		Directory(FEntryMap& Entries, const FString& DirPath)
+			: Current(Entries.CreateIterator()), Path(DirPath)
 		{
 			if (!Path.IsEmpty())
 			{
-				Path /= "";
+				Path /= TEXT("");
 			}
 			// This would be much easier, and efficient, if TMap
 			// supported getting iterators to found entries in
diff --git a/Engine/Source/Runtime/Core/Private/Containers/StringTest.cpp b/Engine/Source/Runtime/Core/Private/Containers/StringTest.cpp
--- a/Engine/Source/Runtime/Core/Private/Containers/StringTest.cpp
+++ b/Engine/Source/Runtime/Core/Private/Containers/StringTest.cpp
@@ -3,6 +3,7 @@
 #include "CoreTypes.h"
 #include "Misc/AutomationTest.h"
 #include "Misc/AssertionMacros.h"
+#include "Misc/StringBuilder.h"
 #include "Containers/StringView.h"
 #include "Containers/UnrealString.h"
 #include "Serialization/MemoryReader.h"
@@ -518,4 +519,95 @@ bool FStringEqualityTest::RunTest(const FString& Parameters)
 	return true;	
 }
 
+IMPLEMENT_SIMPLE_AUTOMATION_TEST(FStringPathConcatCompoundOperatorTest, "System.Core.String.PathConcatCompoundOperator", EAutomationTestFlags::ApplicationContextMask | EAutomationTestFlags::SmokeFilter)
+bool FStringPathConcatCompoundOperatorTest::RunTest(const FString& Parameters)
+{
+	// No need to test a nullptr TCHAR* as an input parameter as this is expected to cause a crash
+	// No need to test self assignment, clang will catch that was a compiler error (-Wself-assign-overloaded)
+
+	const TCHAR* Path = TEXT("../Path");
+	const TCHAR* PathWithTrailingSlash = TEXT("../Path/");
+	const TCHAR* Filename = TEXT("File.txt");
+	const TCHAR* FilenameWithLeadingSlash = TEXT("/File.txt");
+	const TCHAR* CombinedPath = TEXT("../Path/File.txt");
+	const TCHAR* CombinedPathWithDoubleSeparator = TEXT("../Path//File.txt");
+
+	// Existing code supported ansi char so we need to test that to avoid potentially breaking license code
+	const ANSICHAR* AnsiFilename = "File.txt";
+	const ANSICHAR* AnsiFilenameWithLeadingSlash = "/File.txt";
+
+	// The TStringBuilders must be created up front as no easy constructor
+	TStringBuilder<128> EmptyStringBuilder;
+	TStringBuilder<128> FilenameStringBuilder; FilenameStringBuilder << Filename;
+	TStringBuilder<128> FilenameWithLeadingSlashStringBuilder; FilenameWithLeadingSlashStringBuilder << FilenameWithLeadingSlash;
+	
+#define TEST_EMPTYPATH_EMPTYFILE(Type, Input)						{ FString EmptyPathString; EmptyPathString /= Input; TestTrue(Type TEXT(": EmptyPath/EmptyFilename to be empty"), EmptyPathString.IsEmpty()); }
+#define TEST_VALIDPATH_EMPTYFILE(Type, Input)						{ FString Result(Path); Result /= Input; TestEqual(Type TEXT(": ValidPath/EmptyFilename result to be"), Result, PathWithTrailingSlash); } \
+																	{ FString Result(PathWithTrailingSlash); Result /= Input; TestEqual(Type TEXT(" (with extra /): ValidPath/EmptyFilename result to be"), Result, PathWithTrailingSlash); }	
+#define TEST_EMPTYPATH_VALIDFILE(Type, Input)						{ FString Result; Result /= Input; TestEqual(Type TEXT(": EmptyPath/ValidFilename"), Result, Filename); }
+#define TEST_VALIDPATH_VALIDFILE(Type, Path, File)					{ FString Result(Path); Result /= File; TestEqual(Type TEXT(": ValidPath/ValidFilename"), Result, CombinedPath); }
+#define TEST_VALIDPATH_VALIDFILE_DOUBLE_SEPARATOR(Type, Path, File)	{ FString Result(Path); Result /= File; TestEqual(Type TEXT(": ValidPath//ValidFilename"), Result, CombinedPathWithDoubleSeparator); }
+
+	// Test empty path /= empty file
+	TEST_EMPTYPATH_EMPTYFILE(TEXT("NullString"), FString());
+	TEST_EMPTYPATH_EMPTYFILE(TEXT("EmptyString"), FString(TEXT("")));
+	TEST_EMPTYPATH_EMPTYFILE(TEXT("EmptyAnsiLiteralString"), "");
+	TEST_EMPTYPATH_EMPTYFILE(TEXT("EmptyLiteralString"), TEXT(""));
+	TEST_EMPTYPATH_EMPTYFILE(TEXT("NullStringView"), FStringView());
+	TEST_EMPTYPATH_EMPTYFILE(TEXT("EmptyStringView"), FStringView(TEXT("")));
+	TEST_EMPTYPATH_EMPTYFILE(TEXT("EmptyStringBuilder"), EmptyStringBuilder);
+
+	// Test valid path /= empty file
+	TEST_VALIDPATH_EMPTYFILE(TEXT("NullString"), FString());
+	TEST_VALIDPATH_EMPTYFILE(TEXT("EmptyString"), FString(TEXT("")));
+	TEST_VALIDPATH_EMPTYFILE(TEXT("EmptyAnsiLiteralString"), "");
+	TEST_VALIDPATH_EMPTYFILE(TEXT("EmptyLiteralString"), TEXT(""));
+	TEST_VALIDPATH_EMPTYFILE(TEXT("NullStringView"), FStringView());
+	TEST_VALIDPATH_EMPTYFILE(TEXT("EmptyStringView"), FStringView(TEXT("")));
+	TEST_VALIDPATH_EMPTYFILE(TEXT("EmptyStringBuilder"), EmptyStringBuilder);
+	
+	// Test empty path /= valid file
+	TEST_EMPTYPATH_VALIDFILE(TEXT("String"), FString(Filename));
+	TEST_EMPTYPATH_VALIDFILE(TEXT("LiteralString"), Filename);
+	TEST_EMPTYPATH_VALIDFILE(TEXT("LiteralAnsiString"), AnsiFilename);
+	TEST_EMPTYPATH_VALIDFILE(TEXT("StringView"), FStringView(Filename));
+	
+	// Test valid path /= valid file
+	TEST_VALIDPATH_VALIDFILE(TEXT("String"), Path, FString(Filename));
+	TEST_VALIDPATH_VALIDFILE(TEXT("LiteralString"), Path, Filename);
+	TEST_VALIDPATH_VALIDFILE(TEXT("LiteralAnsiString"), Path, AnsiFilename);
+	TEST_VALIDPATH_VALIDFILE(TEXT("StringView"), Path, FStringView(Filename));
+	TEST_VALIDPATH_VALIDFILE(TEXT("StringBuilder"), Path, FilenameStringBuilder);
+
+	// Test valid path (ending in /) /= valid file
+	TEST_VALIDPATH_VALIDFILE(TEXT("String (path with extra /)"), PathWithTrailingSlash, FString(Filename));
+	TEST_VALIDPATH_VALIDFILE(TEXT("LiteralString (path with extra /)"), PathWithTrailingSlash, Filename);
+	TEST_VALIDPATH_VALIDFILE(TEXT("LiteralAnsiString (path with extra /)"), PathWithTrailingSlash, AnsiFilename);
+	TEST_VALIDPATH_VALIDFILE(TEXT("StringView (path with extra /)"), PathWithTrailingSlash, FStringView(Filename));
+	TEST_VALIDPATH_VALIDFILE(TEXT("StringBuilder (path with extra /)"), PathWithTrailingSlash, FilenameStringBuilder);
+	
+	// Test valid path / valid path + file (starting with /)
+	TEST_VALIDPATH_VALIDFILE(TEXT("String (filename with extra /)"), Path, FString(FilenameWithLeadingSlash));
+	TEST_VALIDPATH_VALIDFILE(TEXT("LiteralString (filename with extra /)"), Path, FilenameWithLeadingSlash);
+	TEST_VALIDPATH_VALIDFILE(TEXT("LiteralAnsiString (filename with extra /)"), Path, AnsiFilenameWithLeadingSlash);
+	TEST_VALIDPATH_VALIDFILE(TEXT("StringView (filename with extra /)"), Path, FStringView(FilenameWithLeadingSlash));
+	TEST_VALIDPATH_VALIDFILE(TEXT("StringBuilder (filename with extra /)"), Path, FilenameWithLeadingSlashStringBuilder);
+	
+	// Appending a file name that starts with a / to a directory that ends with a / will not remove the erroneous / and so 
+	// will end up with // in the path, these tests are to show this behavior
+	// For example "path/" /= "/file.txt" will result in "path//file.txt" not "path/file.txt"
+	TEST_VALIDPATH_VALIDFILE_DOUBLE_SEPARATOR(TEXT("String (path and filename with extra /)"), PathWithTrailingSlash, FString(FilenameWithLeadingSlash));
+	TEST_VALIDPATH_VALIDFILE_DOUBLE_SEPARATOR(TEXT("LiteralString (path and filename with extra /)"), PathWithTrailingSlash, FilenameWithLeadingSlash);
+	TEST_VALIDPATH_VALIDFILE_DOUBLE_SEPARATOR(TEXT("LiteralAnsiString (path and filename with extra /)"), PathWithTrailingSlash, AnsiFilenameWithLeadingSlash);
+	TEST_VALIDPATH_VALIDFILE_DOUBLE_SEPARATOR(TEXT("StringView (path and filename with extra /)"), PathWithTrailingSlash, FStringView(FilenameWithLeadingSlash));
+	TEST_VALIDPATH_VALIDFILE_DOUBLE_SEPARATOR(TEXT("StringBuilder (path and filename with extra /)"), PathWithTrailingSlash, FilenameWithLeadingSlashStringBuilder);
+
+#undef TEST_EMPTYPATH_EMPTYFILE
+#undef TEST_VALIDPATH_EMPTYFILE
+#undef TEST_VALIDPATH_EMPTYFILE
+#undef TEST_VALIDPATH_VALIDFILE
+
+	return true;
+}
+
 #endif // WITH_DEV_AUTOMATION_TESTS
diff --git a/Engine/Source/Runtime/Core/Public/Containers/UnrealString.h b/Engine/Source/Runtime/Core/Public/Containers/UnrealString.h
--- a/Engine/Source/Runtime/Core/Public/Containers/UnrealString.h
+++ b/Engine/Source/Runtime/Core/Public/Containers/UnrealString.h
@@ -34,7 +34,7 @@ typedef TArray<FStringFormatArg> FStringFormatOrderedArguments;
 
 TCHAR*       GetData(FString&);
 const TCHAR* GetData(const FString&);
-
+int32        GetNum(const FString& String);
 /**
  * A dynamically sizeable string.
  * @see https://docs.unrealengine.com/latest/INT/Programming/UnrealArchitecture/StringHandling/FString/
@@ -938,14 +938,29 @@ public:
 	}
 
 	/**
-	 * Concatenate this path with given path ensuring the / character is used between them
-	 * 
-	 * @param Str path FString to be concatenated onto the end of this
-	 * @return reference to path
-	 */
-	FORCEINLINE FString& operator/=( const FString& Str )
+	* Concatenate this path with given path ensuring the / character is used between them
+	* 
+	* @param Str path CharRangeType (FString/FStringView/TStringBuilder etc) to be concatenated onto the end of this
+	* @return reference to path
+	*/
+	template <typename CharRangeType, typename TEnableIf<TIsTCharRangeNotCArray <CharRangeType>::Value>::Type* = nullptr>
+	FORCEINLINE FString& operator/=(CharRangeType&& Str)
+	{
+		PathAppend(GetData(Str), GetNum(Str));
+		return *this;
+	}
+
+	/**
+	* Concatenate this path with given path ensuring the / character is used between them
+	* 
+	* @param Str path array of CharType (that needs converting) to be concatenated onto the end of this
+	* @return reference to path
+	*/
+	template <typename CharType,typename TEnableIf<TIsCharType<CharType>::Value>::Type* = nullptr>
+	FORCEINLINE FString& operator/=(const CharType* Str)
 	{
-		PathAppend(Str.Data.GetData(), Str.Len());
+		FString Temp = Str;
+		PathAppend(GetData(Temp), GetNum(Temp));
 		return *this;
 	}
 
diff --git a/Engine/Source/Runtime/Core/Public/Misc/StringBuilder.h b/Engine/Source/Runtime/Core/Public/Misc/StringBuilder.h
--- a/Engine/Source/Runtime/Core/Public/Misc/StringBuilder.h
+++ b/Engine/Source/Runtime/Core/Public/Misc/StringBuilder.h
@@ -313,7 +313,7 @@ protected:
 };
 
 template <typename CharType>
-constexpr inline SIZE_T GetNum(const TStringBuilderBase<CharType>& Builder)
+constexpr inline int32 GetNum(const TStringBuilderBase<CharType>& Builder)
 {
 	return Builder.Len();
 }
