Fixed dependencies on the DerivedDataCache module to be include-only when appropriate

diff --git a/Engine/Plugins/Editor/AssetSearch/Source/AssetSearch.Build.cs b/Engine/Plugins/Editor/AssetSearch/Source/AssetSearch.Build.cs
--- a/Engine/Plugins/Editor/AssetSearch/Source/AssetSearch.Build.cs
+++ b/Engine/Plugins/Editor/AssetSearch/Source/AssetSearch.Build.cs
@@ -27,7 +27,6 @@ public class AssetSearch : ModuleRules
 				"Slate",
 				"SlateCore",
 				"GameplayTags",
-				"DerivedDataCache",
 				"WorkspaceMenuStructure",
 				"EditorFramework",
 				"UnrealEd",
@@ -42,5 +41,11 @@ public class AssetSearch : ModuleRules
 				"DeveloperSettings"
 			}
 		);
+
+		PrivateIncludePathModuleNames.AddRange(
+			new string[] {
+				"DerivedDataCache",
+			}
+		);
 	}
 }
diff --git a/Engine/Plugins/FX/Niagara/Source/Niagara/Niagara.Build.cs b/Engine/Plugins/FX/Niagara/Source/Niagara/Niagara.Build.cs
--- a/Engine/Plugins/FX/Niagara/Source/Niagara/Niagara.Build.cs
+++ b/Engine/Plugins/FX/Niagara/Source/Niagara/Niagara.Build.cs
@@ -43,6 +43,10 @@ public class Niagara : ModuleRules
             }
         );
 
+		PrivateIncludePathModuleNames.AddRange(
+			new string[] {
+				"DerivedDataCache",
+			});
 
         PrivateIncludePaths.AddRange(
             new string[] {
@@ -50,16 +54,6 @@ public class Niagara : ModuleRules
             })
         ;
 
-        // If we're compiling with the engine, then add Core's engine dependencies
-        if (Target.bCompileAgainstEngine == true)
-        {
-            if (!Target.bBuildRequiresCookedData)
-            {
-                DynamicallyLoadedModuleNames.AddRange(new string[] { "DerivedDataCache" });
-            }
-        }
-
-
         if (Target.bBuildEditor == true)
         {
             PrivateDependencyModuleNames.AddRange(
diff --git a/Engine/Plugins/FX/Niagara/Source/NiagaraShader/NiagaraShader.Build.cs b/Engine/Plugins/FX/Niagara/Source/NiagaraShader/NiagaraShader.Build.cs
--- a/Engine/Plugins/FX/Niagara/Source/NiagaraShader/NiagaraShader.Build.cs
+++ b/Engine/Plugins/FX/Niagara/Source/NiagaraShader/NiagaraShader.Build.cs
@@ -31,9 +31,8 @@ public class NiagaraShader : ModuleRules
         {
             PublicDependencyModuleNames.AddRange(
                 new string[] {
-                "DerivedDataCache",
-                "TargetPlatform",
-            });
+                    "TargetPlatform",
+                });
 
             PrivateDependencyModuleNames.AddRange(
                 new string[] {
@@ -46,6 +45,7 @@ public class NiagaraShader : ModuleRules
 
         PrivateIncludePathModuleNames.AddRange(
             new string[] {
+                "DerivedDataCache",
                 "Niagara",
             });
     }
diff --git a/Engine/Plugins/Runtime/HairStrands/Source/HairStrandsCore/HairStrandsCore.Build.cs b/Engine/Plugins/Runtime/HairStrands/Source/HairStrandsCore/HairStrandsCore.Build.cs
--- a/Engine/Plugins/Runtime/HairStrands/Source/HairStrandsCore/HairStrandsCore.Build.cs
+++ b/Engine/Plugins/Runtime/HairStrands/Source/HairStrandsCore/HairStrandsCore.Build.cs
@@ -37,8 +37,13 @@ namespace UnrealBuildTool.Rules
 				PrivateDependencyModuleNames.AddRange(
 					new string[]
 					{
-					"DerivedDataCache",
-					"Eigen",
+						"Eigen",
+					});
+
+				PrivateIncludePathModuleNames.AddRange(
+					new string[]
+					{
+						"DerivedDataCache",
 					});
 			}
 		}
diff --git a/Engine/Source/Developer/PakFileUtilities/PakFileUtilities.Build.cs b/Engine/Source/Developer/PakFileUtilities/PakFileUtilities.Build.cs
--- a/Engine/Source/Developer/PakFileUtilities/PakFileUtilities.Build.cs
+++ b/Engine/Source/Developer/PakFileUtilities/PakFileUtilities.Build.cs
@@ -12,13 +12,14 @@ public class PakFileUtilities : ModuleRules
             "Json",
             "Projects",
             "RSA",
-            "DerivedDataCache"
         });
 
-        PrivateIncludePathModuleNames.AddRange(
-            new string[] {
-                "Json"
+        PrivateIncludePathModuleNames.AddRange(new string[] {
+            "DerivedDataCache",
+            "Json",
         });
+
+		DynamicallyLoadedModuleNames.Add("DerivedDataCache");
 		
 	}
 }
diff --git a/Engine/Source/Editor/Kismet/Kismet.Build.cs b/Engine/Source/Editor/Kismet/Kismet.Build.cs
--- a/Engine/Source/Editor/Kismet/Kismet.Build.cs
+++ b/Engine/Source/Editor/Kismet/Kismet.Build.cs
@@ -51,7 +51,6 @@ public class Kismet : ModuleRules
 				"EngineSettings",
                 "Projects",
                 "JsonUtilities",
-                "DerivedDataCache",
 				"DesktopPlatform",
 				"HotReload",
                 "BlueprintNativeCodeGen",
diff --git a/Engine/Source/Editor/UnrealEd/UnrealEd.Build.cs b/Engine/Source/Editor/UnrealEd/UnrealEd.Build.cs
--- a/Engine/Source/Editor/UnrealEd/UnrealEd.Build.cs
+++ b/Engine/Source/Editor/UnrealEd/UnrealEd.Build.cs
@@ -22,7 +22,6 @@ public class UnrealEd : ModuleRules
 				"Editor/UnrealEd/Private",
 				"Editor/UnrealEd/Private/Settings",
 				"Editor/PackagesDialog/Public",
-				"Developer/DerivedDataCache/Public",
 				"Developer/TargetPlatform/Public",
 			}
 		);
@@ -188,7 +187,6 @@ public class UnrealEd : ModuleRules
 				"PakFileUtilities",
 				"TimeManagement",
                 "LandscapeEditorUtilities",
-                "DerivedDataCache",
 				"ScriptDisassembler",
 				"ToolMenus",
 				"FreeImage",
diff --git a/Engine/Source/Runtime/Core/Core.Build.cs b/Engine/Source/Runtime/Core/Core.Build.cs
--- a/Engine/Source/Runtime/Core/Core.Build.cs
+++ b/Engine/Source/Runtime/Core/Core.Build.cs
@@ -26,7 +26,6 @@ public class Core : ModuleRules
 
 		PrivateIncludePaths.AddRange(
 			new string[] {
-				"Developer/DerivedDataCache/Public",
 				"Runtime/SynthBenchmark/Public",
 				"Runtime/Core/Private",
 				"Runtime/Core/Private/Misc",
@@ -157,12 +156,9 @@ public class Core : ModuleRules
 		PublicDefinitions.Add("UE_ENABLE_ICU=" + (Target.bCompileICU ? "1" : "0")); // Enable/disable (=1/=0) ICU usage in the codebase. NOTE: This flag is for use while integrating ICU and will be removed afterward.
 
 		// If we're compiling with the engine, then add Core's engine dependencies
-		if (Target.bCompileAgainstEngine == true)
+		if (Target.bCompileAgainstEngine && !Target.bBuildRequiresCookedData)
 		{
-			if (!Target.bBuildRequiresCookedData)
-			{
-				DynamicallyLoadedModuleNames.AddRange(new string[] { "DerivedDataCache" });
-			}
+			DynamicallyLoadedModuleNames.AddRange(new string[] { "DerivedDataCache" });
 		}
 
 		// On Windows platform, VSPerfExternalProfiler.cpp needs access to "VSPerf.h".  This header is included with Visual Studio, but it's not in a standard include path.
diff --git a/Engine/Source/Runtime/Engine/Engine.Build.cs b/Engine/Source/Runtime/Engine/Engine.Build.cs
--- a/Engine/Source/Runtime/Engine/Engine.Build.cs
+++ b/Engine/Source/Runtime/Engine/Engine.Build.cs
@@ -17,7 +17,6 @@ public class Engine : ModuleRules
 
 		PrivateIncludePaths.AddRange(
 			new string[] {
-				"Developer/DerivedDataCache/Public",
 				"Runtime/SynthBenchmark/Public",
 				"Runtime/Engine/Private",
 				"Runtime/Net/Core/Private/Net/Core/PushModel/Types"
@@ -26,6 +25,7 @@ public class Engine : ModuleRules
 
 		PrivateIncludePathModuleNames.AddRange(
 			new string[] {
+				"DerivedDataCache",
 				"TargetPlatform",
 				"ImageWrapper",
 				"ImageWriteQueue",
@@ -304,7 +304,6 @@ public class Engine : ModuleRules
 
 			DynamicallyLoadedModuleNames.AddRange(
 				new string[] {
-					"DerivedDataCache",
 					"DesktopPlatform",
 				}
 			);
diff --git a/Engine/Source/Runtime/Experimental/GeometryCollectionEngine/GeometryCollectionEngine.Build.cs b/Engine/Source/Runtime/Experimental/GeometryCollectionEngine/GeometryCollectionEngine.Build.cs
--- a/Engine/Source/Runtime/Experimental/GeometryCollectionEngine/GeometryCollectionEngine.Build.cs
+++ b/Engine/Source/Runtime/Experimental/GeometryCollectionEngine/GeometryCollectionEngine.Build.cs
@@ -33,11 +33,6 @@ namespace UnrealBuildTool.Rules
 				PublicDependencyModuleNames.Add("APEX");
 			}
 
-	        if (!Target.bBuildRequiresCookedData)
-			{
-				DynamicallyLoadedModuleNames.Add("DerivedDataCache");
-			}
-
 			PrivateIncludePathModuleNames.Add("DerivedDataCache");
 
 			if (Target.bBuildEditor)
diff --git a/Engine/Source/Runtime/Launch/Launch.Build.cs b/Engine/Source/Runtime/Launch/Launch.Build.cs
--- a/Engine/Source/Runtime/Launch/Launch.Build.cs
+++ b/Engine/Source/Runtime/Launch/Launch.Build.cs
@@ -137,7 +137,7 @@ public class Launch : ModuleRules
 			});
 
 			PublicDependencyModuleNames.Add("SessionServices");
-			PrivateIncludePaths.Add("Developer/DerivedDataCache/Public");
+			PrivateIncludePathModuleNames.Add("DerivedDataCache");
 
 			// LaunchEngineLoop.cpp will still attempt to load XMPP but not all projects require it so it will silently fail unless referenced by the project's build.cs file.
 			// DynamicallyLoadedModuleNames.Add("XMPP");
diff --git a/Engine/Source/Runtime/MRMesh/MRMesh.build.cs b/Engine/Source/Runtime/MRMesh/MRMesh.build.cs
--- a/Engine/Source/Runtime/MRMesh/MRMesh.build.cs
+++ b/Engine/Source/Runtime/MRMesh/MRMesh.build.cs
@@ -32,7 +32,8 @@ namespace UnrealBuildTool.Rules
 			// Used for including the private Chaos headers
 			string EnginePath = Path.GetFullPath(Target.RelativeEnginePath);
 			PrivateIncludePaths.Add(Path.Combine(EnginePath, "Source/Runtime/Engine/Private/PhysicsEngine"));
-			PrivateIncludePaths.Add(Path.Combine(EnginePath, "Source/Developer/DerivedDataCache/Public"));
+
+			PrivateIncludePathModuleNames.Add("DerivedDataCache");
 		}
 	}
 }
diff --git a/Engine/Source/Runtime/NavigationSystem/NavigationSystem.Build.cs b/Engine/Source/Runtime/NavigationSystem/NavigationSystem.Build.cs
--- a/Engine/Source/Runtime/NavigationSystem/NavigationSystem.Build.cs
+++ b/Engine/Source/Runtime/NavigationSystem/NavigationSystem.Build.cs
@@ -16,7 +16,6 @@ namespace UnrealBuildTool.Rules
                 new string[] {
                     "Runtime/NavigationSystem/Private",
                     "Runtime/Engine/Private",
-                    "Developer/DerivedDataCache/Public",
                 }
                 );
 
@@ -38,14 +37,10 @@ namespace UnrealBuildTool.Rules
 			PrivateIncludePathModuleNames.AddRange(
 				new string[]
 				{
+					"DerivedDataCache",
 					"TargetPlatform",
 				}
 				);
-            
-            if (!Target.bBuildRequiresCookedData && Target.bCompileAgainstEngine)
-            {
-                DynamicallyLoadedModuleNames.Add("DerivedDataCache");
-            }
 
             SetupModulePhysicsSupport(Target);
 
diff --git a/Engine/Source/Runtime/NetworkFile/NetworkFile.Build.cs b/Engine/Source/Runtime/NetworkFile/NetworkFile.Build.cs
--- a/Engine/Source/Runtime/NetworkFile/NetworkFile.Build.cs
+++ b/Engine/Source/Runtime/NetworkFile/NetworkFile.Build.cs
@@ -18,15 +18,6 @@ public class NetworkFile : ModuleRules
 		PublicIncludePaths.Add("Runtime/CoreUObject/Public/UObject");
 		PublicIncludePaths.Add("Runtime/CoreUObject/Public");
 
-		if (!Target.bBuildRequiresCookedData)
-		{
-			DynamicallyLoadedModuleNames.AddRange(
-				new string[]
-				{
-					"DerivedDataCache",
-				});
-		}
-
 		PublicDefinitions.Add("ENABLE_HTTP_FOR_NETWORK_FILE=0");
 	}
 }
