Update min macOS version to 10.15 and roll default MSL to v2.2.

#rnx

diff --git a/Engine/Config/BaseEngine.ini b/Engine/Config/BaseEngine.ini
--- a/Engine/Config/BaseEngine.ini
+++ b/Engine/Config/BaseEngine.ini
@@ -2390,7 +2390,7 @@ bTarget32Bit=false
 +TargetedRHIs=SF_VULKAN_SM5
 
 [/Script/MacTargetPlatform.MacTargetSettings]
-MaxShaderLanguageVersion=4
+MaxShaderLanguageVersion=5
 +TargetedRHIs=SF_METAL_SM5
 UseFastIntrinsics=False
 ForceFloats=False
diff --git a/Engine/Shaders/Public/Platform/Metal/MetalCommon.ush b/Engine/Shaders/Public/Platform/Metal/MetalCommon.ush
--- a/Engine/Shaders/Public/Platform/Metal/MetalCommon.ush
+++ b/Engine/Shaders/Public/Platform/Metal/MetalCommon.ush
@@ -5,7 +5,7 @@
 =============================================================================*/
 
 // Update this GUID to improve shader recompilation for Metal only shaders
-// GUID = 293FE8D0-534F-4A76-ACD0-2802BC36DEA7
+// GUID = A998F35E-2D63-47D6-BD71-B63DB703EE16
 
 #pragma once
 
diff --git a/Engine/Source/Developer/Apple/MetalShaderFormat/Private/MetalDerivedData.cpp b/Engine/Source/Developer/Apple/MetalShaderFormat/Private/MetalDerivedData.cpp
--- a/Engine/Source/Developer/Apple/MetalShaderFormat/Private/MetalDerivedData.cpp
+++ b/Engine/Source/Developer/Apple/MetalShaderFormat/Private/MetalDerivedData.cpp
@@ -1732,18 +1732,30 @@ bool FMetalShaderOutputCooker::Build(TArray<uint8>& OutData)
 				CCHeaderWriter.WriteTessellationIndexBuffer(HullIndexBuffer);
 			}
 			
+			if (VersionEnum >= 4)
+			{
+				if (IABTier >= 1)
+				{
+					TargetDesc.CompileFlags.SetDefine(TEXT("argument_buffers"), 1);
+					TargetDesc.CompileFlags.SetDefine(TEXT("argument_buffer_offset"), IABOffsetIndex);
+				}
+				TargetDesc.CompileFlags.SetDefine(TEXT("texture_buffer_native"), 1);
+			}
+			
 			switch (VersionEnum)
 			{
 				case 6:
+				{
+					TargetDesc.Version = 20300;
+					break;
+				}
 				case 5:
+				{
+					TargetDesc.Version = 20200;
+					break;
+				}
 				case 4:
 				{
-					if (IABTier >= 1)
-					{
-						TargetDesc.CompileFlags.SetDefine(TEXT("argument_buffers"), 1);
-						TargetDesc.CompileFlags.SetDefine(TEXT("argument_buffer_offset"), IABOffsetIndex);
-					}
-					TargetDesc.CompileFlags.SetDefine(TEXT("texture_buffer_native"), 1);
 					TargetDesc.Version = 20100;
 					break;
 				}
@@ -1769,7 +1781,6 @@ bool FMetalShaderOutputCooker::Build(TArray<uint8>& OutData)
 					break;
 				}
 			}
-			
 		}
 
 		// Convert SPIR-V binary to Metal source
diff --git a/Engine/Source/Developer/Mac/MacTargetPlatform/Classes/MacTargetSettings.h b/Engine/Source/Developer/Mac/MacTargetPlatform/Classes/MacTargetSettings.h
--- a/Engine/Source/Developer/Mac/MacTargetPlatform/Classes/MacTargetSettings.h
+++ b/Engine/Source/Developer/Mac/MacTargetPlatform/Classes/MacTargetSettings.h
@@ -13,11 +13,17 @@
 UENUM()
 enum class EMacMetalShaderStandard : uint8
 {
-    /** Metal Shaders, supporting multiple viewports, Compatible With macOS 10.13.0 or later (std=osx-metal2.0) */
+    /** Metal Shaders Compatible With macOS 10.13 or later (std=osx-metal2.0) */
     MacMetalSLStandard_2_0 = 3 UMETA(DisplayName="Metal v2.0 (10.13.0+)"),
     
-    /** Metal Shaders, supporting indirect command buffers, Compatible With macOS 10.14.0 or later (std=osx-metal2.1) */
+    /** Metal Shaders Compatible With macOS 10.14 or later (std=osx-metal2.1) */
     MacMetalSLStandard_2_1 = 4 UMETA(DisplayName="Metal v2.1 (10.14.0+)"),
+
+	/** Metal Shaders Compatible With macOS 10.15 or later (std=osx-metal2.2) */
+	MacMetalSLStandard_2_2 = 5 UMETA(DisplayName="Metal v2.2 (10.15.0+)"),
+
+	/** Metal Shaders Compatible With macOS 11.0 or later (std=osx-metal2.3) */
+	MacMetalSLStandard_2_3 = 6 UMETA(DisplayName="Metal v2.3 (11.0+)"),
 };
 
 UENUM()
diff --git a/Engine/Source/Developer/ShaderCompilerCommon/Private/ShaderConductorContext.cpp b/Engine/Source/Developer/ShaderCompilerCommon/Private/ShaderConductorContext.cpp
--- a/Engine/Source/Developer/ShaderCompilerCommon/Private/ShaderConductorContext.cpp
+++ b/Engine/Source/Developer/ShaderCompilerCommon/Private/ShaderConductorContext.cpp
@@ -259,6 +259,8 @@ namespace CrossCompiler
 	{
 		switch (Version)
 		{
+		case 20300: return "20300";
+		case 20200: return "20200";
 		case 20100: return "20100";
 		case 20000: return "20000";
 		case 10200: return "10200";
diff --git a/Engine/Source/Developer/ShaderCompilerCommon/Public/ShaderConductorContext.h b/Engine/Source/Developer/ShaderCompilerCommon/Public/ShaderConductorContext.h
--- a/Engine/Source/Developer/ShaderCompilerCommon/Public/ShaderConductorContext.h
+++ b/Engine/Source/Developer/ShaderCompilerCommon/Public/ShaderConductorContext.h
@@ -56,7 +56,7 @@ namespace CrossCompiler
 
 		/**
 		Target shader version.
-		Valid values for Metal family: 20100, 20000, 10200, 10100, 10000.
+		Valid values for Metal family: 20300, 20200, 20100, 20000, 10200, 10100, 10000.
 		Valid values for GLSL family: 310, 320, 330, 430.
 		*/
 		int32 Version = 0;
diff --git a/Engine/Source/Editor/DetailCustomizations/Private/MacTargetSettingsDetails.cpp b/Engine/Source/Editor/DetailCustomizations/Private/MacTargetSettingsDetails.cpp
--- a/Engine/Source/Editor/DetailCustomizations/Private/MacTargetSettingsDetails.cpp
+++ b/Engine/Source/Editor/DetailCustomizations/Private/MacTargetSettingsDetails.cpp
@@ -286,7 +286,9 @@ static uint32 GMacTargetSettingsMinOSVers[][3] = {
 	{10,11,6},
 	{10,12,6},
 	{10,13,0},
-	{10,14,0}
+	{10,14,0},
+	{10,15,0},
+	{11, 0,0},
 };
 
 TSharedRef<SWidget> FMacTargetSettingsDetails::OnGetShaderVersionContent()
diff --git a/Engine/Source/Editor/UnrealEd/Private/EditorEngine.cpp b/Engine/Source/Editor/UnrealEd/Private/EditorEngine.cpp
--- a/Engine/Source/Editor/UnrealEd/Private/EditorEngine.cpp
+++ b/Engine/Source/Editor/UnrealEd/Private/EditorEngine.cpp
@@ -6679,7 +6679,7 @@ void UEditorEngine::UpdateAutoLoadProject()
 #if PLATFORM_MAC
 	if ( !GIsBuildMachine )
 	{
-		if(FPlatformMisc::MacOSXVersionCompare(10,14,1) < 0)
+		if(FPlatformMisc::MacOSXVersionCompare(10,15,7) < 0)
 		{
 			if(FSlateApplication::IsInitialized())
 			{
diff --git a/Engine/Source/Programs/UnrealBuildTool/Platform/Mac/MacToolChain.cs b/Engine/Source/Programs/UnrealBuildTool/Platform/Mac/MacToolChain.cs
--- a/Engine/Source/Programs/UnrealBuildTool/Platform/Mac/MacToolChain.cs
+++ b/Engine/Source/Programs/UnrealBuildTool/Platform/Mac/MacToolChain.cs
@@ -28,12 +28,12 @@ namespace UnrealBuildTool
 		/// <summary>
 		/// Which version of the Mac OS X to allow at run time
 		/// </summary>
-		public string MacOSVersion = "10.14";
+		public string MacOSVersion = "10.15";
 
 		/// <summary>
 		/// Minimum version of Mac OS X to actually run on, running on earlier versions will display the system minimum version error dialog and exit.
 		/// </summary>
-		public string MinMacOSVersion = "10.14.6";
+		public string MinMacOSVersion = "10.15.7";
 
 		/// <summary>
 		/// Directory for the developer binaries
diff --git a/Engine/Source/Runtime/Apple/MetalRHI/Private/MetalCommandQueue.cpp b/Engine/Source/Runtime/Apple/MetalRHI/Private/MetalCommandQueue.cpp
--- a/Engine/Source/Runtime/Apple/MetalRHI/Private/MetalCommandQueue.cpp
+++ b/Engine/Source/Runtime/Apple/MetalRHI/Private/MetalCommandQueue.cpp
@@ -29,8 +29,8 @@ FMetalCommandQueue::FMetalCommandQueue(mtlpp::Device InDevice, uint32 const MaxN
 	int32 MaxShaderVersion = 0;
 	int32 IndirectArgumentTier = 0;
 #if PLATFORM_MAC
-	int32 DefaultMaxShaderVersion = 4;
-	int32 MinShaderVersion = 3;
+	int32 DefaultMaxShaderVersion = 5; // MSL v2.2
+	int32 MinShaderVersion = 4; // MSL v2.1
     const TCHAR* const Settings = TEXT("/Script/MacTargetPlatform.MacTargetSettings");
 #else
 	int32 DefaultMaxShaderVersion = 2;
diff --git a/Engine/Source/Runtime/Apple/MetalRHI/Private/MetalShaders.cpp b/Engine/Source/Runtime/Apple/MetalRHI/Private/MetalShaders.cpp
--- a/Engine/Source/Runtime/Apple/MetalRHI/Private/MetalShaders.cpp
+++ b/Engine/Source/Runtime/Apple/MetalRHI/Private/MetalShaders.cpp
@@ -53,13 +53,17 @@ mtlpp::LanguageVersion ValidateVersion(uint8 Version)
 		{10,12,6},
 		{10,13,0},
 		{10,14,0},
+		{10,15,0},
+		{11, 0,0},
 	};
 	static uint32 MetaliOSVersions[][3] = {
-		{8,0,0},
-		{9,0,0},
+		{ 8,0,0},
+		{ 9,0,0},
 		{10,0,0},
 		{11,0,0},
 		{12,0,0},
+		{13,0,0},
+		{14,0,0},
 	};
 	static TCHAR const* StandardNames[] =
 	{
@@ -68,11 +72,19 @@ mtlpp::LanguageVersion ValidateVersion(uint8 Version)
 		TEXT("Metal 1.2"),
 		TEXT("Metal 2.0"),
 		TEXT("Metal 2.1"),
+		TEXT("Metal 2.2"),
+		TEXT("Metal 2.3"),
 	};
 	
 	mtlpp::LanguageVersion Result = mtlpp::LanguageVersion::Version1_1;
 	switch(Version)
 	{
+		case 6:
+			Result = mtlpp::LanguageVersion::Version2_3;
+			break;
+		case 5:
+			Result = mtlpp::LanguageVersion::Version2_2;
+			break;
 		case 4:
 			Result = mtlpp::LanguageVersion::Version2_1;
 			break;
diff --git a/Engine/Source/Runtime/ApplicationCore/Private/Apple/AppleControllerInterface.cpp b/Engine/Source/Runtime/ApplicationCore/Private/Apple/AppleControllerInterface.cpp
--- a/Engine/Source/Runtime/ApplicationCore/Private/Apple/AppleControllerInterface.cpp
+++ b/Engine/Source/Runtime/ApplicationCore/Private/Apple/AppleControllerInterface.cpp
@@ -136,13 +136,6 @@ void FAppleControllerInterface::HandleConnection(GCController* Controller)
         Controllers[ControllerIndex].Controller = Controller;
         SetControllerType(ControllerIndex);
         
-		// Deprecated but buttonMenu behavior is unreliable in iOS/tvOS 14.0.1
-        Controllers[ControllerIndex].bPauseWasPressed = false;
-        Controller.controllerPausedHandler = ^(GCController* Cont)
-        {
-            Controllers[ControllerIndex].bPauseWasPressed = true;
-        };
-        
         bFoundSlot = true;
         
         UE_LOG(LogAppleController, Log, TEXT("New %s controller inserted, assigned to playerIndex %d"),
@@ -179,21 +172,6 @@ void FAppleControllerInterface::SendControllerEvents()
 {
     for(int32 i = 0; i < UE_ARRAY_COUNT(Controllers); ++i)
  	{
-        GCController* Cont = Controllers[i].Controller;
-        
-        GCExtendedGamepad* ExtendedGamepad = nil;
-
-        if (@available(iOS 13, tvOS 13, *))
-        {
-            ExtendedGamepad = [Cont capture].extendedGamepad;
-        }
-        else
-        {
-            ExtendedGamepad = [Cont.extendedGamepad saveSnapshot];
-        }
-
-		GCMotion* Motion = Cont.motion;
-		
 		// make sure the connection handler has run on this guy
 		if (Controllers[i].PlayerIndex == PlayerIndex::PlayerUnset)
 		{
@@ -201,14 +179,11 @@ void FAppleControllerInterface::SendControllerEvents()
 		}
 
 		FUserController& Controller = Controllers[i];
-		
-        if (Controller.bPauseWasPressed)
-        {
-            MessageHandler->OnControllerButtonPressed(FGamepadKeyNames::SpecialRight, Controllers[i].PlayerIndex, false);
-            MessageHandler->OnControllerButtonReleased(FGamepadKeyNames::SpecialRight, Controllers[i].PlayerIndex, false);
+		GCController* Cont = Controller.Controller;
 
-            Controller.bPauseWasPressed = false;
-        }
+		// Assumes iOS13, tvOS 13 & macOS 10.15
+        GCExtendedGamepad* ExtendedGamepad = [Cont capture].extendedGamepad;
+		GCMotion* Motion = Cont.motion;
         
 		if (ExtendedGamepad != nil)
 		{
diff --git a/Engine/Source/Runtime/ApplicationCore/Private/Apple/AppleControllerInterface.h b/Engine/Source/Runtime/ApplicationCore/Private/Apple/AppleControllerInterface.h
--- a/Engine/Source/Runtime/ApplicationCore/Private/Apple/AppleControllerInterface.h
+++ b/Engine/Source/Runtime/ApplicationCore/Private/Apple/AppleControllerInterface.h
@@ -105,9 +105,6 @@ protected:
 		bool bNeedsReferenceAttitude;
 		bool bHasReferenceAttitude;
 #endif
-        
-        // Deprecated but buttonMenu in iOS 14 is not working in current Beta (August 2020).
-        bool bPauseWasPressed;
 	};
 	
     // there is a hardcoded limit of 4 controllers in the API
diff --git a/Engine/Source/Runtime/Core/Private/Apple/ApplePlatformDebugEvents.cpp b/Engine/Source/Runtime/Core/Private/Apple/ApplePlatformDebugEvents.cpp
--- a/Engine/Source/Runtime/Core/Private/Apple/ApplePlatformDebugEvents.cpp
+++ b/Engine/Source/Runtime/Core/Private/Apple/ApplePlatformDebugEvents.cpp
@@ -60,7 +60,11 @@
 	#define PLATFORM_MAC 1
 	#define PLATFORM_IOS 0
 	#define PLATFORM_TVOS 0
-#elif PLATFORM_TVOS	
+
+	#define DEBUG_SIGNPOST_EMIT(Code, Arg1, Arg2, Arg3, Arg4)			os_signpost_event_emit(GetLog(), Code, "DebugSignPost", "%lu, %lu, %lu, %lu", Arg1, Arg2, Arg3, Arg4)
+	#define DEBUG_SIGNPOST_INTERVAL_BEGIN(Code, Arg1, Arg2, Arg3, Arg4)	os_signpost_interval_begin(GetLog(), Code, "DebugSignPost", "%lu, %lu, %lu, %lu", Arg1, Arg2, Arg3, Arg4)
+	#define DEBUG_SIGNPOST_INTERVAL_END(Code, Arg1, Arg2, Arg3, Arg4)	os_signpost_interval_end(GetLog(), Code, "DebugSignPost", "%lu, %lu, %lu, %lu", Arg1, Arg2, Arg3, Arg4)
+#elif PLATFORM_TVOS
 	#include <os/log.h>
 	#include <os/signpost.h>
 	
@@ -71,6 +75,10 @@
 	#define PLATFORM_MAC 0
 	#define PLATFORM_IOS 1
 	#define PLATFORM_TVOS 1
+
+	#define DEBUG_SIGNPOST_EMIT(Code, Arg1, Arg2, Arg3, Arg4)			os_signpost_event_emit(GetLog(), Code, "DebugSignPost", "%lu, %lu, %lu, %lu", Arg1, Arg2, Arg3, Arg4)
+	#define DEBUG_SIGNPOST_INTERVAL_BEGIN(Code, Arg1, Arg2, Arg3, Arg4)	os_signpost_interval_begin(GetLog(), Code, "DebugSignPost", "%lu, %lu, %lu, %lu", Arg1, Arg2, Arg3, Arg4)
+	#define DEBUG_SIGNPOST_INTERVAL_END(Code, Arg1, Arg2, Arg3, Arg4)	os_signpost_interval_end(GetLog(), Code, "DebugSignPost", "%lu, %lu, %lu, %lu", Arg1, Arg2, Arg3, Arg4)
 #else
 	#include <os/log.h>
 	#include <os/signpost.h>
@@ -82,8 +90,16 @@
 	#define PLATFORM_MAC 0
 	#define PLATFORM_IOS 1
 	#define PLATFORM_TVOS 0
+
+	#define DEBUG_SIGNPOST_EMIT(Code, Arg1, Arg2, Arg3, Arg4)			os_signpost_event_emit(GetLog(), Code, "DebugSignPost", "%lu, %lu, %lu, %lu", Arg1, Arg2, Arg3, Arg4)
+	#define DEBUG_SIGNPOST_INTERVAL_BEGIN(Code, Arg1, Arg2, Arg3, Arg4)	os_signpost_interval_begin(GetLog(), Code, "DebugSignPost", "%lu, %lu, %lu, %lu", Arg1, Arg2, Arg3, Arg4)
+	#define DEBUG_SIGNPOST_INTERVAL_END(Code, Arg1, Arg2, Arg3, Arg4)	os_signpost_interval_end(GetLog(), Code, "DebugSignPost", "%lu, %lu, %lu, %lu", Arg1, Arg2, Arg3, Arg4)
 #endif
-#endif
+#else // APPLE_PROFILING_SIGNPOST
+	#define DEBUG_SIGNPOST_EMIT(Code, Arg1, Arg2, Arg3, Arg4)			kdebug_signpost(Code, Arg1, Arg2, Arg3, Arg4)
+	#define DEBUG_SIGNPOST_INTERVAL_BEGIN(Code, Arg1, Arg2, Arg3, Arg4)	kdebug_signpost_start(Code, Arg1, Arg2, Arg3, Arg4)
+	#define DEBUG_SIGNPOST_INTERVAL_END(Code, Arg1, Arg2, Arg3, Arg4)	kdebug_signpost_end(Code, Arg1, Arg2, Arg3, Arg4)
+#endif // APPLE_PROFILING_SIGNPOST
 
 DEFINE_LOG_CATEGORY(LogInstruments)
 
@@ -163,7 +179,7 @@ void FApplePlatformDebugEvents::DebugSignPost(uint16 Code, uintptr_t Arg1, uintp
 		else
 #endif
 		{
-			kdebug_signpost(Code, Arg1, Arg2, Arg3, Arg4);
+			DEBUG_SIGNPOST_EMIT(Code, Arg1, Arg2, Arg3, Arg4);
 		}
 	}
 }
@@ -186,7 +202,7 @@ void FApplePlatformDebugEvents::DebugSignPostStart(uint16 Code, uintptr_t Arg1,
 		else
 #endif
 		{
-			kdebug_signpost_start(Code, Arg1, Arg2, Arg3, Arg4);
+			DEBUG_SIGNPOST_INTERVAL_BEGIN(Code, Arg1, Arg2, Arg3, Arg4);
 		}
 	}
 }
@@ -209,7 +225,7 @@ void FApplePlatformDebugEvents::DebugSignPostEnd(uint16 Code, uintptr_t Arg1, ui
 		else
 #endif
 		{
-			kdebug_signpost_end(Code, Arg1, Arg2, Arg3, Arg4);
+			DEBUG_SIGNPOST_INTERVAL_END(Code, Arg1, Arg2, Arg3, Arg4);
 		}
 	}
 }
diff --git a/Engine/Source/Runtime/IOS/IOSRuntimeSettings/Classes/IOSRuntimeSettings.h b/Engine/Source/Runtime/IOS/IOSRuntimeSettings/Classes/IOSRuntimeSettings.h
--- a/Engine/Source/Runtime/IOS/IOSRuntimeSettings/Classes/IOSRuntimeSettings.h
+++ b/Engine/Source/Runtime/IOS/IOSRuntimeSettings/Classes/IOSRuntimeSettings.h
@@ -51,6 +51,12 @@ enum class EIOSMetalShaderStandard : uint8
     
     /** Metal Shaders Compatible With iOS 12.0/tvOS 12.0 or later (std=ios-metal2.1) */
     IOSMetalSLStandard_2_1 = 4 UMETA(DisplayName="Metal v2.1 (iOS 12.0/tvOS 12.0)"),
+
+	/** Metal Shaders Compatible With iOS 13.0/tvOS 13.0  or later (std=ios-metal2.2) */
+	IOSMetalSLStandard_2_2 = 5 UMETA(DisplayName="Metal v2.2 (iOS 13.0/tvOS 13.0)"),
+
+	/** Metal Shaders Compatible With iOS 14.0/tvOS 14.0  or later (std=ios-metal2.3) */
+	IOSMetalSLStandard_2_3 = 6 UMETA(DisplayName="Metal v2.3 (iOS 14.0/tvOS 14.0)"),
 };
 
 UENUM()
diff --git a/Engine/Source/Runtime/Launch/Private/Mac/LaunchMac.cpp b/Engine/Source/Runtime/Launch/Private/Mac/LaunchMac.cpp
--- a/Engine/Source/Runtime/Launch/Private/Mac/LaunchMac.cpp
+++ b/Engine/Source/Runtime/Launch/Private/Mac/LaunchMac.cpp
@@ -259,8 +259,8 @@ static int32 MacOSVersionCompare(const NSOperatingSystemVersion& VersionA, const
 	NSString* MinimumSystemVersionString = (NSString*)InfoDictionary[@"LSMinimumSystemVersion"];
 	NSOperatingSystemVersion MinimumSystemVersion = { 0 };
 	NSOperatingSystemVersion CurrentSystemVersion = [[NSProcessInfo processInfo] operatingSystemVersion];
-	NSOperatingSystemVersion MinSupportedMacOSVersion = { 10, 14, 6 };
-	NSString* MinSupportedMacOSVersionString = @"10.14.6";
+	NSOperatingSystemVersion MinSupportedMacOSVersion = { 10, 15, 7 };
+	NSString* MinSupportedMacOSVersionString = @"10.15.7";
 
 	NSArray<NSString*>* VersionComponents = [MinimumSystemVersionString componentsSeparatedByString:@"."];
 	MinimumSystemVersion.majorVersion = [[VersionComponents objectAtIndex:0] integerValue];
diff --git a/Engine/Source/Runtime/RHI/Private/RHI.cpp b/Engine/Source/Runtime/RHI/Private/RHI.cpp
--- a/Engine/Source/Runtime/RHI/Private/RHI.cpp
+++ b/Engine/Source/Runtime/RHI/Private/RHI.cpp
@@ -1304,11 +1304,11 @@ RHI_API uint32 RHIGetShaderLanguageVersion(const FStaticShaderPlatform Platform)
 			static int32 MaxShaderVersion = -1;
 			if (MaxShaderVersion < 0)
 			{
-				MaxShaderVersion = 2;
-				int32 MinShaderVersion = 3;
+				MaxShaderVersion = 5; // MSL v2.2
+				int32 MinShaderVersion = 4; // MSL v2.1
 				if(!GConfig->GetInt(TEXT("/Script/MacTargetPlatform.MacTargetSettings"), TEXT("MaxShaderLanguageVersion"), MaxShaderVersion, GEngineIni))
 				{
-					MaxShaderVersion = 4;
+					MaxShaderVersion = 0;
 				}
 				MaxShaderVersion = FMath::Max(MinShaderVersion, MaxShaderVersion);
 			}
