Enable UnrealInsights GPU channel by default

diff --git a/Engine/Source/Runtime/Core/Private/ProfilingDebugging/TraceAuxiliary.cpp b/Engine/Source/Runtime/Core/Private/ProfilingDebugging/TraceAuxiliary.cpp
--- a/Engine/Source/Runtime/Core/Private/ProfilingDebugging/TraceAuxiliary.cpp
+++ b/Engine/Source/Runtime/Core/Private/ProfilingDebugging/TraceAuxiliary.cpp
@@ -163,7 +163,7 @@ bool FTraceAuxiliaryImpl::Connect(ETraceConnectType Type, const TCHAR* Parameter
 		FString Value;
 		if (!GConfig->GetString(TEXT("Trace.ChannelPresets"), TEXT("Default"), Value, GEngineIni))
 		{
-			Value = TEXT("cpu,frame,log,bookmark");
+			Value = TEXT("cpu,frame,log,bookmark,gpu");
 		}
 
 		AddChannels(*Value);
