Refactored DECLARE_STAT macro to use EStatFlags instead of series of bool parameters. Added a mechanism to disable a external profiler and trace events for a specific stat. Disabled events for CPU stall stats by default.

Use `-verbosenamedevents` command line argument to allow all events (restore old behavior; implies `-statnamedevents`).

diff --git a/Engine/Source/Runtime/Core/Private/Async/TaskGraph.cpp b/Engine/Source/Runtime/Core/Private/Async/TaskGraph.cpp
--- a/Engine/Source/Runtime/Core/Private/Async/TaskGraph.cpp
+++ b/Engine/Source/Runtime/Core/Private/Async/TaskGraph.cpp
@@ -703,7 +703,7 @@ public:
 				if (bAllowStall)
 				{
 					{
-						FScopeCycleCounter Scope(StallStatId);
+						FScopeCycleCounter Scope(StallStatId, EStatFlags::Verbose);
 						Queue(QueueIndex).StallRestartEvent->Wait(bIsRenderThreadAndPolling ? GRenderThreadPollPeriodMs : MAX_uint32, bCountAsStall);
 						if (Queue(QueueIndex).QuitForShutdown)
 						{
@@ -1055,7 +1055,7 @@ private:
 				const bool bIsMultithread = FTaskGraphInterface::IsMultithread();
 				if (bIsMultithread)
 				{
-					FScopeCycleCounter Scope(StallStatId);
+					FScopeCycleCounter Scope(StallStatId, EStatFlags::Verbose);
 					Queue.StallRestartEvent->Wait(MAX_uint32, bCountAsStall);
 					bDidStall = true;
 				}
diff --git a/Engine/Source/Runtime/Core/Private/HAL/ThreadingBase.cpp b/Engine/Source/Runtime/Core/Private/HAL/ThreadingBase.cpp
--- a/Engine/Source/Runtime/Core/Private/HAL/ThreadingBase.cpp
+++ b/Engine/Source/Runtime/Core/Private/HAL/ThreadingBase.cpp
@@ -1247,7 +1247,11 @@ FQueuedThread::Run()
 		{
 			while (bContinueWaiting)
 			{
-				DECLARE_SCOPE_CYCLE_COUNTER(TEXT("FQueuedThread::Run.WaitForWork"), STAT_FQueuedThread_Run_WaitForWork, STATGROUP_ThreadPoolAsyncTasks);
+				DECLARE_CYCLE_STAT_WITH_FLAGS(TEXT("FQueuedThread::Run.WaitForWork"),
+				STAT_FQueuedThread_Run_WaitForWork, STATGROUP_ThreadPoolAsyncTasks,
+					EStatFlags::ClearEveryFrame | EStatFlags::Verbose);
+
+				SCOPE_CYCLE_COUNTER(STAT_FQueuedThread_Run_WaitForWork);
 
 				// Wait for some work to do
 
diff --git a/Engine/Source/Runtime/Core/Private/Misc/CoreGlobals.cpp b/Engine/Source/Runtime/Core/Private/Misc/CoreGlobals.cpp
--- a/Engine/Source/Runtime/Core/Private/Misc/CoreGlobals.cpp
+++ b/Engine/Source/Runtime/Core/Private/Misc/CoreGlobals.cpp
@@ -393,6 +393,13 @@ bool				GHitchDetected = false;
 
 /** Whether stats should emit named events for e.g. PIX.													*/
 int32					GCycleStatsShouldEmitNamedEvents = 0;
+
+/** Whether verbose stats should be also generate external profiler named events.
+* Thread sleep/wait stats or extremely high frequency cycle counting stats are disabled by default.
+* Has no effect if GCycleStatsShouldEmitNamedEvents is 0.
+*/
+bool					GShouldEmitVerboseNamedEvents = false;
+
 /** Disables some warnings and minor features that would interrupt a demo presentation						*/
 bool					GIsDemoMode						= false;
 /** Whether or not a unit test is currently being run														*/
diff --git a/Engine/Source/Runtime/Core/Public/CoreGlobals.h b/Engine/Source/Runtime/Core/Public/CoreGlobals.h
--- a/Engine/Source/Runtime/Core/Public/CoreGlobals.h
+++ b/Engine/Source/Runtime/Core/Public/CoreGlobals.h
@@ -489,6 +489,12 @@ extern CORE_API bool GHitchDetected;
 /** Whether stats should emit named events for e.g. PIX. */
 extern CORE_API int32 GCycleStatsShouldEmitNamedEvents;
 
+/** Whether verbose stats should be also generate external profiler named events.
+* Thread sleep/wait stats or extremely high frequency cycle counting stats are disabled by default.
+* Has no effect if GCycleStatsShouldEmitNamedEvents is 0.
+*/
+extern CORE_API bool GShouldEmitVerboseNamedEvents;
+
 /** Disables some warnings and minor features that would interrupt a demo presentation*/
 extern CORE_API bool GIsDemoMode;
 
diff --git a/Engine/Source/Runtime/Core/Public/HAL/LowLevelMemStats.h b/Engine/Source/Runtime/Core/Public/HAL/LowLevelMemStats.h
--- a/Engine/Source/Runtime/Core/Public/HAL/LowLevelMemStats.h
+++ b/Engine/Source/Runtime/Core/Public/HAL/LowLevelMemStats.h
@@ -10,10 +10,10 @@
  * LLM Stat implementation macros; these macros are used publicly in LLM_SCOPED_SINGLE_PLATFORM_STAT_TAG and privately in LLM implementation
  */
 	#define DECLARE_LLM_MEMORY_STAT(CounterName,StatId,GroupId) \
-		DECLARE_STAT(CounterName,StatId,GroupId,EStatDataType::ST_int64, false, false, FPlatformMemory::MCR_PhysicalLLM); \
+		DECLARE_STAT(CounterName,StatId,GroupId,EStatDataType::ST_int64, EStatFlags::None, FPlatformMemory::MCR_PhysicalLLM); \
 		static DEFINE_STAT(StatId)
 	#define DECLARE_LLM_MEMORY_STAT_EXTERN(CounterName,StatId,GroupId, API) \
-		DECLARE_STAT(CounterName,StatId,GroupId,EStatDataType::ST_int64, false, false, FPlatformMemory::MCR_PhysicalLLM); \
+		DECLARE_STAT(CounterName,StatId,GroupId,EStatDataType::ST_int64, EStatFlags::None, FPlatformMemory::MCR_PhysicalLLM); \
 		extern API DEFINE_STAT(StatId);
 
 	DECLARE_STATS_GROUP(TEXT("LLM FULL"), STATGROUP_LLMFULL, STATCAT_Advanced);
diff --git a/Engine/Source/Runtime/Core/Public/Misc/CoreStats.h b/Engine/Source/Runtime/Core/Public/Misc/CoreStats.h
--- a/Engine/Source/Runtime/Core/Public/Misc/CoreStats.h
+++ b/Engine/Source/Runtime/Core/Public/Misc/CoreStats.h
@@ -51,9 +51,9 @@ DECLARE_FLOAT_COUNTER_STAT_EXTERN(TEXT("Percentage CPU utilization (relative to
 	CPU Stalls
 -----------------------------------------------------------------------------*/
 
-DECLARE_CYCLE_STAT_EXTERN(TEXT("CPU Stall - Sleep"), STAT_Sleep, STATGROUP_CPUStalls, CORE_API);
-DECLARE_CYCLE_STAT_EXTERN(TEXT("CPU Stall - Wait For Event" ), STAT_EventWait, STATGROUP_CPUStalls,);
+DECLARE_CYCLE_STAT_WITH_FLAGS_EXTERN(TEXT("CPU Stall - Sleep"), STAT_Sleep, STATGROUP_CPUStalls, EStatFlags::ClearEveryFrame | EStatFlags::Verbose, CORE_API);
+DECLARE_CYCLE_STAT_WITH_FLAGS_EXTERN(TEXT("CPU Stall - Wait For Event" ), STAT_EventWait, STATGROUP_CPUStalls, EStatFlags::ClearEveryFrame | EStatFlags::Verbose,);
 
 /** The id will be stored as uint64 in the stat message. */
-DECLARE_PTR_STAT_EXTERN( TEXT( "CPU Stall - Wait For Event with ID" ), STAT_EventWaitWithId, STATGROUP_CPUStalls, );
-DECLARE_PTR_STAT_EXTERN( TEXT( "CPU Stall - Trigger For Event with ID" ), STAT_EventTriggerWithId, STATGROUP_CPUStalls, );
+DECLARE_CYCLE_STAT_WITH_FLAGS_EXTERN( TEXT( "CPU Stall - Wait For Event with ID" ), STAT_EventWaitWithId, STATGROUP_CPUStalls, EStatFlags::ClearEveryFrame | EStatFlags::Verbose, );
+DECLARE_CYCLE_STAT_WITH_FLAGS_EXTERN( TEXT( "CPU Stall - Trigger For Event with ID" ), STAT_EventTriggerWithId, STATGROUP_CPUStalls, EStatFlags::ClearEveryFrame | EStatFlags::Verbose, );
diff --git a/Engine/Source/Runtime/Core/Public/Stats/Stats.h b/Engine/Source/Runtime/Core/Public/Stats/Stats.h
--- a/Engine/Source/Runtime/Core/Public/Stats/Stats.h
+++ b/Engine/Source/Runtime/Core/Public/Stats/Stats.h
@@ -17,6 +17,7 @@
  */
 
 struct TStatId;
+enum class EStatFlags : uint8;
 
 // used by the profiler
 enum EStatType
@@ -58,11 +59,15 @@ public:
 	 * Pushes the specified stat onto the hierarchy for this thread. Starts
 	 * the timing of the cycles used
 	 */
-	FORCEINLINE_STATS FScopeCycleCounter( TStatId StatId, bool bAlways = false )
+	FORCEINLINE_STATS FScopeCycleCounter( TStatId StatId, EStatFlags StatFlags, bool bAlways = false )
 	{
-		Start( StatId, bAlways );
+		Start( StatId, StatFlags, bAlways );
 	}
 
+	FORCEINLINE_STATS FScopeCycleCounter(TStatId StatId, bool bAlways = false)
+		: FScopeCycleCounter(StatId, EStatFlags::None, bAlways)
+	{}
+
 	/**
 	 * Updates the stat with the time spent
 	 */
@@ -144,7 +149,7 @@ public:
 class FScopeCycleCounter
 {
 public:
-	FORCEINLINE FScopeCycleCounter(TStatId InStatId, bool bAlways = false)
+	FORCEINLINE FScopeCycleCounter(TStatId InStatId, EStatFlags StatFlags, bool bAlways = false)
 		: 
 #if USE_LIGHTWEIGHT_STATS_FOR_HITCH_DETECTION && USE_HITCH_DETECTION
 		StatScope(InStatId.StatString),
@@ -159,6 +164,11 @@ public:
 		}
 	}
 
+	FORCEINLINE FScopeCycleCounter(TStatId InStatId, bool bAlways = false)
+		: FScopeCycleCounter(InStatId, EStatFlags::None, bAlways)
+	{
+	}
+
 	FORCEINLINE ~FScopeCycleCounter()
 	{
 		if (bPop)
@@ -179,6 +189,9 @@ struct TStatId {};
 class FScopeCycleCounter
 {
 public:
+	FORCEINLINE_STATS FScopeCycleCounter(TStatId, EStatFlags, bool bAlways = false)
+	{
+	}
 	FORCEINLINE_STATS FScopeCycleCounter(TStatId, bool bAlways = false)
 	{
 	}
@@ -273,6 +286,7 @@ public:
 #define DEFINE_STAT(Stat)
 #define QUICK_USE_CYCLE_STAT(StatId,GroupId) TStatId()
 #define DECLARE_CYCLE_STAT(CounterName,StatId,GroupId)
+#define DECLARE_CYCLE_STAT_WITH_FLAGS(CounterName,StatId,GroupId,StatFlags)
 #define DECLARE_FLOAT_COUNTER_STAT(CounterName,StatId,GroupId)
 #define DECLARE_DWORD_COUNTER_STAT(CounterName,StatId,GroupId)
 #define DECLARE_FLOAT_ACCUMULATOR_STAT(CounterName,StatId,GroupId)
@@ -282,6 +296,7 @@ public:
 #define DECLARE_MEMORY_STAT(CounterName,StatId,GroupId)
 #define DECLARE_MEMORY_STAT_POOL(CounterName,StatId,GroupId,Pool)
 #define DECLARE_CYCLE_STAT_EXTERN(CounterName,StatId,GroupId, API)
+#define DECLARE_CYCLE_STAT_WITH_FLAGS_EXTERN(CounterName,StatId,GroupId,StatFlags, API)
 #define DECLARE_FLOAT_COUNTER_STAT_EXTERN(CounterName,StatId,GroupId, API)
 #define DECLARE_DWORD_COUNTER_STAT_EXTERN(CounterName,StatId,GroupId, API)
 #define DECLARE_FLOAT_ACCUMULATOR_STAT_EXTERN(CounterName,StatId,GroupId, API)
diff --git a/Engine/Source/Runtime/Core/Public/Stats/Stats2.h b/Engine/Source/Runtime/Core/Public/Stats/Stats2.h
--- a/Engine/Source/Runtime/Core/Public/Stats/Stats2.h
+++ b/Engine/Source/Runtime/Core/Public/Stats/Stats2.h
@@ -22,6 +22,7 @@
 #include "ProfilingDebugging/CpuProfilerTrace.h"
 #include "ProfilingDebugging/MiscTrace.h"
 #include "StatsTrace.h"
+#include "Misc/EnumClassFlags.h"
 
 class FScopeCycleCounter;
 class FThreadStats;
@@ -111,6 +112,15 @@ struct CORE_API FStats
 	static TAtomic<int32> GameThreadStatsFrame;
 };
 
+enum class EStatFlags : uint8
+{
+	None            = 0,
+	ClearEveryFrame = 1 << 0,
+	CycleStat       = 1 << 1,
+	Verbose         = 1 << 2, // Profiling scopes for this stat will no generate a trace event by default. See GShouldEmitVerboseNamedEvents.
+};
+ENUM_CLASS_FLAGS(EStatFlags);
+
 #if STATS
 
 struct TStatIdData
@@ -1538,7 +1548,7 @@ public:
 	 * Pushes the specified stat onto the hierarchy for this thread. Starts
 	 * the timing of the cycles used
 	 */
-	FORCEINLINE_STATS void Start( TStatId InStatId, bool bAlways = false )
+	FORCEINLINE_STATS void Start( TStatId InStatId, EStatFlags InStatFlags, bool bAlways = false )
 	{
 		FMinimalName StatMinimalName = InStatId.GetMinimalName(EMemoryOrder::Relaxed);
 		if (StatMinimalName.IsNone())
@@ -1547,7 +1557,8 @@ public:
 		}
 
 		// Emit named event for active cycle stat.
-		if ( GCycleStatsShouldEmitNamedEvents > 0 )
+		if ( GCycleStatsShouldEmitNamedEvents > 0 
+			&& (GShouldEmitVerboseNamedEvents || !EnumHasAnyFlags(InStatFlags, EStatFlags::Verbose)))
 		{
 #if	PLATFORM_USES_ANSI_STRING_FOR_EXTERNAL_PROFILING
 			FPlatformMisc::BeginNamedEvent( FColor( 0 ), InStatId.GetStatDescriptionANSI() );
@@ -1558,7 +1569,8 @@ public:
 		}
 
 #if CPUPROFILERTRACE_ENABLED
-		if (GCycleStatsShouldEmitNamedEvents > 0 && UE_TRACE_CHANNELEXPR_IS_ENABLED(CpuChannel))
+		if (GCycleStatsShouldEmitNamedEvents > 0 && UE_TRACE_CHANNELEXPR_IS_ENABLED(CpuChannel)
+			&& (GShouldEmitVerboseNamedEvents || !EnumHasAnyFlags(InStatFlags, EStatFlags::Verbose)))
 		{
 			FCpuProfilerTrace::OutputBeginDynamicEvent(InStatId.GetStatDescriptionANSI()); //todo: Could we use FName index as event id?
 			EmittedEvent |= TraceEvent;
@@ -1574,6 +1586,11 @@ public:
 		}
 	}
 
+	FORCEINLINE_STATS void Start(TStatId InStatId, bool bAlways = false)
+	{
+		Start(InStatId, EStatFlags::None, bAlways);
+	}
+
 	/**
 	 * Stops the capturing and stores the result
 	 */
@@ -1799,7 +1816,7 @@ struct FStatGroup_##StatName\
 	} \
 };
 
-#define DECLARE_STAT(Description, StatName, GroupName, StatType, bShouldClearEveryFrame, bCycleStat, MemoryRegion) \
+#define DECLARE_STAT(Description, StatName, GroupName, StatType, StatFlags, MemoryRegion) \
 struct FStat_##StatName\
 { \
 	typedef FStatGroup_##GroupName TGroup; \
@@ -1817,11 +1834,15 @@ struct FStat_##StatName\
 	} \
 	static FORCEINLINE bool IsClearEveryFrame() \
 	{ \
-		return bShouldClearEveryFrame; \
+		return EnumHasAnyFlags(GetFlags(), EStatFlags::ClearEveryFrame); \
 	} \
 	static FORCEINLINE bool IsCycleStat() \
 	{ \
-		return bCycleStat; \
+		return EnumHasAnyFlags(GetFlags(), EStatFlags::CycleStat); \
+	} \
+	static FORCEINLINE EStatFlags GetFlags() \
+	{ \
+		return StatFlags; \
 	} \
 	static FORCEINLINE FPlatformMemory::EMemoryCounterRegion GetMemoryRegion() \
 	{ \
@@ -1833,6 +1854,7 @@ struct FStat_##StatName\
 #define GET_STATFNAME(Stat) (StatPtr_##Stat.GetStatFName())
 #define GET_STATDESCRIPTION(Stat) (FStat_##Stat::GetDescription())
 #define GET_STATISEVERYFRAME(Stat) (FStat_##Stat::IsClearEveryFrame())
+#define GET_STATFLAGS(Stat) (FStat_##Stat::GetFlags())
 
 #define STAT_GROUP_TO_FStatGroup(Group) FStatGroup_##Group
 
@@ -1844,44 +1866,47 @@ struct FStat_##StatName\
 	struct FThreadSafeStaticStat<FStat_##Stat> StatPtr_##Stat;
 
 #define RETURN_QUICK_DECLARE_CYCLE_STAT(StatId,GroupId) \
-	DECLARE_STAT(TEXT(#StatId),StatId,GroupId,EStatDataType::ST_int64, true, true, FPlatformMemory::MCR_Invalid); \
+	DECLARE_STAT(TEXT(#StatId),StatId,GroupId,EStatDataType::ST_int64, EStatFlags::ClearEveryFrame | EStatFlags::CycleStat, FPlatformMemory::MCR_Invalid); \
 	static DEFINE_STAT(StatId) \
 	return GET_STATID(StatId);
 
 #define QUICK_USE_CYCLE_STAT(StatId,GroupId) [](){ RETURN_QUICK_DECLARE_CYCLE_STAT(StatId, GroupId); }()
 
 #define DECLARE_CYCLE_STAT(CounterName,StatId,GroupId) \
-	DECLARE_STAT(CounterName,StatId,GroupId,EStatDataType::ST_int64, true, true, FPlatformMemory::MCR_Invalid); \
+	DECLARE_STAT(CounterName,StatId,GroupId,EStatDataType::ST_int64, EStatFlags::ClearEveryFrame | EStatFlags::CycleStat, FPlatformMemory::MCR_Invalid); \
+	static DEFINE_STAT(StatId)
+#define DECLARE_CYCLE_STAT_WITH_FLAGS(CounterName,StatId,GroupId,StatFlags) \
+	DECLARE_STAT(CounterName,StatId,GroupId,EStatDataType::ST_int64, (StatFlags) | EStatFlags::CycleStat, FPlatformMemory::MCR_Invalid); \
 	static DEFINE_STAT(StatId)
 #define DECLARE_FLOAT_COUNTER_STAT(CounterName,StatId,GroupId) \
-	DECLARE_STAT(CounterName,StatId,GroupId,EStatDataType::ST_double, true, false, FPlatformMemory::MCR_Invalid); \
+	DECLARE_STAT(CounterName,StatId,GroupId,EStatDataType::ST_double,EStatFlags::ClearEveryFrame, FPlatformMemory::MCR_Invalid); \
 	static DEFINE_STAT(StatId)
 #define DECLARE_DWORD_COUNTER_STAT(CounterName,StatId,GroupId) \
-	DECLARE_STAT(CounterName,StatId,GroupId,EStatDataType::ST_int64, true, false, FPlatformMemory::MCR_Invalid); \
+	DECLARE_STAT(CounterName,StatId,GroupId,EStatDataType::ST_int64,EStatFlags::ClearEveryFrame, FPlatformMemory::MCR_Invalid); \
 	static DEFINE_STAT(StatId)
 #define DECLARE_FLOAT_ACCUMULATOR_STAT(CounterName,StatId,GroupId) \
-	DECLARE_STAT(CounterName,StatId,GroupId,EStatDataType::ST_double, false, false, FPlatformMemory::MCR_Invalid); \
+	DECLARE_STAT(CounterName,StatId,GroupId,EStatDataType::ST_double, EStatFlags::None, FPlatformMemory::MCR_Invalid); \
 	static DEFINE_STAT(StatId)
 #define DECLARE_DWORD_ACCUMULATOR_STAT(CounterName,StatId,GroupId) \
-	DECLARE_STAT(CounterName,StatId,GroupId,EStatDataType::ST_int64, false, false, FPlatformMemory::MCR_Invalid); \
+	DECLARE_STAT(CounterName,StatId,GroupId,EStatDataType::ST_int64, EStatFlags::None, FPlatformMemory::MCR_Invalid); \
 	static DEFINE_STAT(StatId)
 
 /** FName stat that allows sending a string based data. */
 #define DECLARE_FNAME_STAT(CounterName,StatId,GroupId) \
-	DECLARE_STAT(CounterName,StatId,GroupId,EStatDataType::ST_FName, false, false, FPlatformMemory::MCR_Invalid); \
+	DECLARE_STAT(CounterName,StatId,GroupId,EStatDataType::ST_FName, EStatFlags::None, FPlatformMemory::MCR_Invalid); \
 	static DEFINE_STAT(StatId)
 
 /** This is a fake stat, mostly used to implement memory message or other custom stats that don't easily fit into the system. */
 #define DECLARE_PTR_STAT(CounterName,StatId,GroupId)\
-	DECLARE_STAT(CounterName,StatId,GroupId,EStatDataType::ST_Ptr, false, false, FPlatformMemory::MCR_Invalid); \
+	DECLARE_STAT(CounterName,StatId,GroupId,EStatDataType::ST_Ptr, EStatFlags::None, FPlatformMemory::MCR_Invalid); \
 	static DEFINE_STAT(StatId)
 
 #define DECLARE_MEMORY_STAT(CounterName,StatId,GroupId) \
-	DECLARE_STAT(CounterName,StatId,GroupId,EStatDataType::ST_int64, false, false, FPlatformMemory::MCR_Physical); \
+	DECLARE_STAT(CounterName,StatId,GroupId,EStatDataType::ST_int64, EStatFlags::None, FPlatformMemory::MCR_Physical); \
 	static DEFINE_STAT(StatId)
 
 #define DECLARE_MEMORY_STAT_POOL(CounterName,StatId,GroupId,Pool) \
-	DECLARE_STAT(CounterName,StatId,GroupId,EStatDataType::ST_int64, false, false, Pool); \
+	DECLARE_STAT(CounterName,StatId,GroupId,EStatDataType::ST_int64, EStatFlags::None, Pool); \
 	static DEFINE_STAT(StatId)
 
 /*-----------------------------------------------------------------------------
@@ -1889,37 +1914,40 @@ struct FStat_##StatName\
 -----------------------------------------------------------------------------*/
 
 #define DECLARE_CYCLE_STAT_EXTERN(CounterName,StatId,GroupId, APIX) \
-	DECLARE_STAT(CounterName,StatId,GroupId,EStatDataType::ST_int64, true, true, FPlatformMemory::MCR_Invalid); \
+	DECLARE_STAT(CounterName,StatId,GroupId,EStatDataType::ST_int64, EStatFlags::ClearEveryFrame | EStatFlags::CycleStat, FPlatformMemory::MCR_Invalid); \
+	extern APIX DEFINE_STAT(StatId);
+#define DECLARE_CYCLE_STAT_WITH_FLAGS_EXTERN(CounterName,StatId,GroupId,StatFlags, APIX) \
+	DECLARE_STAT(CounterName,StatId,GroupId,EStatDataType::ST_int64, (StatFlags) | EStatFlags::CycleStat, FPlatformMemory::MCR_Invalid); \
 	extern APIX DEFINE_STAT(StatId);
 #define DECLARE_FLOAT_COUNTER_STAT_EXTERN(CounterName,StatId,GroupId, API) \
-	DECLARE_STAT(CounterName,StatId,GroupId,EStatDataType::ST_double, true, false, FPlatformMemory::MCR_Invalid); \
+	DECLARE_STAT(CounterName,StatId,GroupId,EStatDataType::ST_double, EStatFlags::ClearEveryFrame, FPlatformMemory::MCR_Invalid); \
 	extern API DEFINE_STAT(StatId);
 #define DECLARE_DWORD_COUNTER_STAT_EXTERN(CounterName,StatId,GroupId, API) \
-	DECLARE_STAT(CounterName,StatId,GroupId,EStatDataType::ST_int64, true, false, FPlatformMemory::MCR_Invalid); \
+	DECLARE_STAT(CounterName,StatId,GroupId,EStatDataType::ST_int64, EStatFlags::ClearEveryFrame, FPlatformMemory::MCR_Invalid); \
 	extern API DEFINE_STAT(StatId);
 #define DECLARE_FLOAT_ACCUMULATOR_STAT_EXTERN(CounterName,StatId,GroupId, API) \
-	DECLARE_STAT(CounterName,StatId,GroupId,EStatDataType::ST_double, false, false, FPlatformMemory::MCR_Invalid); \
+	DECLARE_STAT(CounterName,StatId,GroupId,EStatDataType::ST_double, EStatFlags::None, FPlatformMemory::MCR_Invalid); \
 	extern API DEFINE_STAT(StatId);
 #define DECLARE_DWORD_ACCUMULATOR_STAT_EXTERN(CounterName,StatId,GroupId, API) \
-	DECLARE_STAT(CounterName,StatId,GroupId,EStatDataType::ST_int64, false, false, FPlatformMemory::MCR_Invalid); \
+	DECLARE_STAT(CounterName,StatId,GroupId,EStatDataType::ST_int64, EStatFlags::None, FPlatformMemory::MCR_Invalid); \
 	extern API DEFINE_STAT(StatId);
 
 /** FName stat that allows sending a string based data. */
 #define DECLARE_FNAME_STAT_EXTERN(CounterName,StatId,GroupId, API) \
-	DECLARE_STAT(CounterName,StatId,GroupId,EStatDataType::ST_FName, false, false, FPlatformMemory::MCR_Invalid); \
+	DECLARE_STAT(CounterName,StatId,GroupId,EStatDataType::ST_FName, EStatFlags::None, FPlatformMemory::MCR_Invalid); \
 	extern API DEFINE_STAT(StatId);
 
 /** This is a fake stat, mostly used to implement memory message or other custom stats that don't easily fit into the system. */
 #define DECLARE_PTR_STAT_EXTERN(CounterName,StatId,GroupId, API) \
-	DECLARE_STAT(CounterName,StatId,GroupId,EStatDataType::ST_Ptr, false, false, FPlatformMemory::MCR_Invalid); \
+	DECLARE_STAT(CounterName,StatId,GroupId,EStatDataType::ST_Ptr, EStatFlags::None, FPlatformMemory::MCR_Invalid); \
 	extern API DEFINE_STAT(StatId);
 
 #define DECLARE_MEMORY_STAT_EXTERN(CounterName,StatId,GroupId, API) \
-	DECLARE_STAT(CounterName,StatId,GroupId,EStatDataType::ST_int64, false, false, FPlatformMemory::MCR_Physical); \
+	DECLARE_STAT(CounterName,StatId,GroupId,EStatDataType::ST_int64, EStatFlags::None, FPlatformMemory::MCR_Physical); \
 	extern API DEFINE_STAT(StatId);
 
 #define DECLARE_MEMORY_STAT_POOL_EXTERN(CounterName,StatId,GroupId,Pool, API) \
-	DECLARE_STAT(CounterName,StatId,GroupId,EStatDataType::ST_int64, false, false, Pool); \
+	DECLARE_STAT(CounterName,StatId,GroupId,EStatDataType::ST_int64, EStatFlags::None, Pool); \
 	extern API DEFINE_STAT(StatId);
 
 /** Macro for declaring group factory instances */
@@ -1936,18 +1964,18 @@ struct FStat_##StatName\
 	DECLARE_STAT_GROUP(GroupDesc, GroupId, GroupCat, false, CompileIn, false);
 
 #define DECLARE_SCOPE_CYCLE_COUNTER(CounterName,Stat,GroupId) \
-	DECLARE_STAT(CounterName,Stat,GroupId,EStatDataType::ST_int64, true, true, FPlatformMemory::MCR_Invalid); \
+	DECLARE_STAT(CounterName,Stat,GroupId,EStatDataType::ST_int64, EStatFlags::ClearEveryFrame | EStatFlags::CycleStat, FPlatformMemory::MCR_Invalid); \
 	static DEFINE_STAT(Stat) \
-	FScopeCycleCounter CycleCount_##Stat(GET_STATID(Stat));
+	FScopeCycleCounter CycleCount_##Stat(GET_STATID(Stat), GET_STATFLAGS(Stat));
 
 #define QUICK_SCOPE_CYCLE_COUNTER(Stat) \
 	DECLARE_SCOPE_CYCLE_COUNTER(TEXT(#Stat),Stat,STATGROUP_Quick)
 
 #define SCOPE_CYCLE_COUNTER(Stat) \
-	FScopeCycleCounter CycleCount_##Stat(GET_STATID(Stat));
+	FScopeCycleCounter CycleCount_##Stat(GET_STATID(Stat), GET_STATFLAGS(Stat));
 
 #define CONDITIONAL_SCOPE_CYCLE_COUNTER(Stat,bCondition) \
-	FScopeCycleCounter CycleCount_##Stat(bCondition ? GET_STATID(Stat) : TStatId());
+	FScopeCycleCounter CycleCount_##Stat(bCondition ? GET_STATID(Stat) : TStatId(), GET_STATFLAGS(Stat));
 
 #define SCOPE_SECONDS_ACCUMULATOR(Stat) \
 	FSimpleScopeSecondsStat SecondsAccum_##Stat(GET_STATID(Stat));
diff --git a/Engine/Source/Runtime/Engine/Private/UnrealEngine.cpp b/Engine/Source/Runtime/Engine/Private/UnrealEngine.cpp
--- a/Engine/Source/Runtime/Engine/Private/UnrealEngine.cpp
+++ b/Engine/Source/Runtime/Engine/Private/UnrealEngine.cpp
@@ -1823,6 +1823,12 @@ void UEngine::Init(IEngineLoop* InEngineLoop)
 		GCycleStatsShouldEmitNamedEvents = 1;
 	}
 
+	if (FParse::Param(FCommandLine::Get(), TEXT("verbosenamedevents")))
+	{
+		GCycleStatsShouldEmitNamedEvents = 1;
+		GShouldEmitVerboseNamedEvents = 1;
+	}
+
 #if UE_NET_TRACE_ENABLED
 	uint32 NetTraceVerbosity;
 	if(FParse::Value(FCommandLine::Get(), TEXT("NetTrace="), NetTraceVerbosity))
