Tasks: pipe implementation incompatible with old TaskGraph API. Implemented directly over LowLevelTasks.

Pipe is a chain of tasks that are executed one after another and so can be used for synchronised access to a shared resource. Should be used instead of dedicated threads as much as possible.

diff --git a/Engine/Source/Runtime/Core/Private/Tasks/Pipe.cpp b/Engine/Source/Runtime/Core/Private/Tasks/Pipe.cpp
new file mode 100644
--- /dev/null
+++ b/Engine/Source/Runtime/Core/Private/Tasks/Pipe.cpp
@@ -0,0 +1,36 @@
+// Copyright Epic Games, Inc. All Rights Reserved.
+
+#include "Tasks/Pipe.h"
+#include "Misc/ScopeExit.h"
+#include "Tasks/Task.h"
+
+namespace UE { namespace Tasks
+{
+	bool FPipe::PushIntoPipe(Private::FTaskBase& Task)
+	{
+		// critical section: between exchanging tasks and adding the previous LastTask to prerequisites.
+		// `LastTask` can be in the process of destruction. Before this it will try to clear itself from the pipe
+		// we need to let it know we're still using it so destruction should wait until we finished
+		PushingThreadsNum.fetch_add(1, std::memory_order_relaxed);
+		ON_SCOPE_EXIT{ PushingThreadsNum.fetch_sub(1, std::memory_order_relaxed); };
+
+		Private::FTaskBase* LastTask_Local = LastTask.exchange(&Task, std::memory_order_release);
+		checkf(LastTask_Local != &Task, TEXT("Dependency cycle: adding itself as a prerequisite (or use after destruction)"));
+
+		return LastTask_Local == nullptr || !LastTask_Local->SetSubsequent(Task);
+	}
+
+	void FPipe::ClearTask(Private::FTaskBase& Task)
+	{
+		// clears `Task` from the pipe if it's still the `LastTask`. if it's not, another task already has set itself as the `LastTask` 
+		// and can be in the process of adding this `Task` as a prerequisite. We need to wait for this to complete before returning 
+		// from the call as the `Task` can be destroyed immediately after this
+		Private::FTaskBase* Task_Local = &Task;
+		if (!LastTask.compare_exchange_strong(Task_Local, nullptr, std::memory_order_acquire, std::memory_order_relaxed))
+		{
+			while (PushingThreadsNum.load(std::memory_order_relaxed) != 0)
+			{
+			}
+		}
+	}
+}}
\ No newline at end of file
diff --git a/Engine/Source/Runtime/Core/Private/Tasks/TaskPrivate.cpp b/Engine/Source/Runtime/Core/Private/Tasks/TaskPrivate.cpp
new file mode 100644
--- /dev/null
+++ b/Engine/Source/Runtime/Core/Private/Tasks/TaskPrivate.cpp
@@ -0,0 +1,21 @@
+// Copyright Epic Games, Inc. All Rights Reserved.
+
+#include "Tasks/TaskPrivate.h"
+#include "Tasks/Pipe.h"
+
+namespace UE { namespace Tasks { namespace Private
+{
+	bool FTaskBase::PushIntoPipe()
+	{
+		return GetPipe()->PushIntoPipe(*this);
+	}
+
+	void FTaskBase::UnblockPipe()
+	{
+		if (GetPipe() != nullptr)
+		{
+			check(bPushedIntoPipe);
+			GetPipe()->ClearTask(*this);
+		}
+	}
+}}}
\ No newline at end of file
diff --git a/Engine/Source/Runtime/Core/Private/Tests/Tasks/TasksTest.cpp b/Engine/Source/Runtime/Core/Private/Tests/Tasks/TasksTest.cpp
new file mode 100644
--- /dev/null
+++ b/Engine/Source/Runtime/Core/Private/Tests/Tasks/TasksTest.cpp
@@ -0,0 +1,462 @@
+// Copyright Epic Games, Inc. All Rights Reserved.
+
+#include "CoreTypes.h"
+#include "Misc/AutomationTest.h"
+#include "Tests/Benchmark.h"
+#include "Tasks/Pipe.h"
+#include "HAL/Thread.h"
+#include "Async/ParallelFor.h"
+#include "Experimental/Async/AwaitableTask.h"
+
+#include <atomic>
+
+#if WITH_DEV_AUTOMATION_TESTS
+
+namespace UE { namespace TasksTests
+{
+	using namespace Tasks;
+
+	void DummyFunc()
+	{
+	}
+
+	IMPLEMENT_SIMPLE_AUTOMATION_TEST(FTasksBasicTest, "System.Core.Tasks.Basic", EAutomationTestFlags::ApplicationContextMask | EAutomationTestFlags::EngineFilter);
+
+	void BasicStressTest()
+	{
+		constexpr uint32 SpawnerGroupsNum = 50;
+		TArray<TTask<void>> SpawnerGroups;
+		SpawnerGroups.Reserve(SpawnerGroupsNum);
+
+		constexpr uint32 SpawnersPerGroupNum = 100;
+		constexpr uint32 TasksNum = SpawnerGroupsNum * SpawnersPerGroupNum;
+		TArray<TTask<void>> Spawners;
+		Spawners.AddDefaulted(TasksNum);
+		TArray<TTask<void>> Tasks;
+		Tasks.AddDefaulted(TasksNum);
+
+		std::atomic<uint32> TasksExecutedNum{ 0 };
+
+		for (uint32 GroupIndex = 0; GroupIndex != SpawnerGroupsNum; ++GroupIndex)
+		{
+			SpawnerGroups.Add(Launch(UE_SOURCE_LOCATION,
+				[
+					Spawners = &Spawners[GroupIndex * SpawnersPerGroupNum],
+					Tasks = &Tasks[GroupIndex * SpawnersPerGroupNum],
+					&TasksExecutedNum,
+					SpawnersPerGroupNum
+				]
+				{
+					for (uint32 SpawnerIndex = 0; SpawnerIndex != SpawnersPerGroupNum; ++SpawnerIndex)
+					{
+						Spawners[SpawnerIndex] = Launch(UE_SOURCE_LOCATION,
+							[
+								Task = &Tasks[SpawnerIndex],
+								&TasksExecutedNum
+							]
+							{
+								*Task = Launch(UE_SOURCE_LOCATION,
+									[&TasksExecutedNum]
+									{
+										++TasksExecutedNum;
+									}
+								);
+							}
+						);
+					}
+				}
+			));
+		}
+
+		for (TTask<void>& SpawnerGroup : SpawnerGroups)
+		{
+			SpawnerGroup.Wait();
+		}
+
+		for (TTask<void>& Spawner : Spawners)
+		{
+			Spawner.Wait();
+		}
+
+		for (TTask<void>& Task : Tasks)
+		{
+			Task.Wait();
+		}
+
+		check(TasksExecutedNum == TasksNum);
+	}
+
+	bool FTasksBasicTest::RunTest(const FString& Parameters)
+	{
+		{	// basic example, fire and forget a high-pri task
+			Tasks::Launch(
+				UE_SOURCE_LOCATION, // debug name
+				[] {}, // task body
+				Tasks::ETaskPriority::High /* task priority, `Normal` by default */
+			);
+		}
+
+		{	// launch a task and wait till it's executed
+			Tasks::Launch(UE_SOURCE_LOCATION, [] {}).Wait();
+		}
+
+		{	// FTaskEvent
+			FTaskEvent Event;
+			check(!Event.IsTriggered());
+
+			// check that waiting blocks
+			TTask<void> Task = Launch(UE_SOURCE_LOCATION, [&Event] { Event.Wait(); });
+			FPlatformProcess::Sleep(0.1f);
+			check(!Task.IsCompleted());
+
+			Event.Trigger();
+			check(Event.IsTriggered());
+			verify(Event.Wait(FTimespan::Zero()));
+		}
+
+		{	// postpone execution so waiting kicks in first
+			std::atomic<int32> Counter{ 0 };
+			TTask<void> Task = Launch(UE_SOURCE_LOCATION, [&Counter] { ++Counter; FPlatformProcess::Sleep(0.1f); });
+
+			ensure(!Task.Wait(FTimespan::Zero()));
+			Task.Wait();
+			check(Counter == 1);
+		}
+
+		{	// same but using `FTaskEvent`
+			FTaskEvent Event;
+			TTask<void> Task = Launch(UE_SOURCE_LOCATION, [&Event] { Event.Wait(); });
+			ensure(!Task.Wait(FTimespan::FromMilliseconds(100)));
+			Event.Trigger();
+			Task.Wait();
+		}
+
+		{	// basic use-case, postpone waiting so the task is executed first
+			std::atomic<bool> Done{ false };
+			TTask<void> Task = Launch(UE_SOURCE_LOCATION, [&Done] { Done = true; });
+			while (!Task.IsCompleted())
+			{
+				FPlatformProcess::Yield();
+			}
+			Task.Wait();
+			check(Done);
+		}
+
+		{	// basic use-case with result, postpone execution so waiting kicks in first
+			TTask<int> Task = Launch(UE_SOURCE_LOCATION, [] { FPlatformProcess::Sleep(0.1f);  return 42; });
+			verify(Task.GetResult() == 42);
+		}
+
+		{	// basic use-case with result, postpone waiting so the task is executed first
+			TTask<int> Task = Launch(UE_SOURCE_LOCATION, [] { return 42; });
+			while (!Task.IsCompleted())
+			{
+				FPlatformProcess::Yield();
+			}
+			verify(Task.GetResult() == 42);
+		}
+
+		{	// check that movable-only result types are supported, that only single instance of result is created and that it's destroyed
+			static std::atomic<uint32> ConstructionsNum{ 0 }; // `static` to make it visible to `FMoveConstructable`
+			static std::atomic<uint32> DestructionsNum{ 0 };
+
+			struct FMoveConstructable
+			{
+				FORCENOINLINE FMoveConstructable()
+				{
+					++ConstructionsNum;
+				}
+
+				FMoveConstructable(FMoveConstructable&&)
+					: FMoveConstructable()
+				{
+				}
+
+				FMoveConstructable(const FMoveConstructable&) = delete;
+				FMoveConstructable& operator=(FMoveConstructable&&) = delete;
+				FMoveConstructable& operator=(const FMoveConstructable&) = delete;
+
+				FORCENOINLINE ~FMoveConstructable()
+				{
+					++DestructionsNum;
+				}
+			};
+
+			Launch(UE_SOURCE_LOCATION, [] { return FMoveConstructable{}; }).GetResult(); // consume the result
+				
+			checkf(ConstructionsNum == 1, TEXT("%d result instances were created but one was expected: the value stored in the task"));
+			checkf(ConstructionsNum == DestructionsNum, TEXT("Mismatched number of constructions (%d) and destructions (%d)"));
+		}
+
+		// fire and forget: launch a task w/o keeping its reference
+		if (LowLevelTasks::FScheduler::Get().GetNumWorkers() != 0)
+		{
+			std::atomic<bool> bDone{ false };
+			Launch(UE_SOURCE_LOCATION, [&bDone] { bDone = true; });
+			while (!bDone)
+			{
+				FPlatformProcess::Yield();
+			}
+		}
+
+		{	// mutable lambda, compilation check
+			Launch(UE_SOURCE_LOCATION, []() mutable {}).Wait();
+			Launch(UE_SOURCE_LOCATION, []() mutable { return false; }).GetResult();
+		}
+
+		{	// free memory occupied by a private task instance, can be required if task instance is held as a member var
+			TTask<void> Task = Launch(UE_SOURCE_LOCATION, [] {});
+			Task.Wait();
+			Task = {};
+		}
+
+		UE_BENCHMARK(5, BasicStressTest);
+
+		return true;
+	}
+
+	IMPLEMENT_SIMPLE_AUTOMATION_TEST(FTasksPipeTest, "System.Core.Tasks.Pipe", EAutomationTestFlags::ApplicationContextMask | EAutomationTestFlags::EngineFilter);
+
+	void PipeStressTest();
+
+	bool FTasksPipeTest::RunTest(const FString& Parameters)
+	{
+		{	// a basic usage example
+			Tasks::FPipe Pipe{ UE_SOURCE_LOCATION }; // a debug name, user-provided or 
+				// `UE_SOURCE_LOCATION` - source file name and line number
+			// launch two tasks in the pipe, they will be executed sequentially, but in parallel
+			// with other tasks (including TaskGraph's old API tasks)
+			Tasks::TTask<void> Task1 = Pipe.Launch(UE_SOURCE_LOCATION, [] {});
+			Tasks::TTask<void> Task2 = Pipe.Launch(UE_SOURCE_LOCATION, [] {});
+			Task2.Wait(); // wait for `Task2` completion
+		}
+
+		{	// an example of thread-safe async interface, kind of a primitive "actor"
+			class FAsyncClass
+			{
+			public:
+				TTask<bool> DoSomething()
+				{
+					return Pipe.Launch(TEXT("DoSomething()"), [this] { return DoSomethingImpl(); });
+				}
+
+				TTask<void> DoSomethingElse()
+				{
+					return Pipe.Launch(TEXT("DoSomethingElse()"), [this] { DoSomethingElseImpl(); });
+				}
+
+			private:
+				bool DoSomethingImpl() { return false; }
+				void DoSomethingElseImpl() {}
+
+			private:
+				FPipe Pipe{ UE_SOURCE_LOCATION };
+			};
+
+			// access the same instance from multiple threads
+			FAsyncClass AsyncInstance;
+			bool bRes = AsyncInstance.DoSomething().GetResult();
+			AsyncInstance.DoSomethingElse().Wait();
+		}
+
+		{	// basic
+			FPipe Pipe{ UE_SOURCE_LOCATION };
+			Pipe.Launch(UE_SOURCE_LOCATION, [] {});
+			Pipe.Launch(UE_SOURCE_LOCATION, [] {}).Wait();
+		}
+
+		{	// launching a piped task with pointer to a function
+			FPipe Pipe{ UE_SOURCE_LOCATION };
+			Pipe.Launch(UE_SOURCE_LOCATION, &DummyFunc).Wait();
+		}
+
+		{	// launching a piped task with a functor object
+			struct FFunctor
+			{
+				void operator()()
+				{
+				}
+			};
+
+			FPipe Pipe{ UE_SOURCE_LOCATION };
+			Pipe.Launch(UE_SOURCE_LOCATION, FFunctor{}).Wait();
+		}
+
+		{	// hold the first piped task execution until the next one is piped to test for non-concurrent execution
+			FPipe Pipe{ UE_SOURCE_LOCATION };
+			bool bTask1Done = false;
+			TTask<void> Task1 = Pipe.Launch(UE_SOURCE_LOCATION, 
+				[&bTask1Done] 
+				{ 
+					FPlatformProcess::Sleep(0.1f); 
+					bTask1Done = true; 
+				}
+			);
+			// we can't just check if `Task1` is completed because pipe gets unblocked and so the next piped task can start execution before the
+			// previous piped task's comlpetion flag is set
+			Pipe.Launch(UE_SOURCE_LOCATION, [&bTask1Done] { check(bTask1Done); }).Wait();
+		}
+
+		{	// piping another task after the previous one is completed and destroyed
+			FPipe Pipe{ UE_SOURCE_LOCATION };
+
+			Pipe.Launch(UE_SOURCE_LOCATION, [] {}).Wait();
+			Pipe.Launch(UE_SOURCE_LOCATION, [] {}).Wait();
+		}
+
+		{	// an example of blocking a pipe
+			FPipe Pipe{ UE_SOURCE_LOCATION };
+			std::atomic<bool> bBlocked;
+			FTaskEvent Event;
+			TTask<void> Task = Pipe.Launch(UE_SOURCE_LOCATION,
+				[&bBlocked, &Event]
+				{
+					bBlocked = true;
+					Event.Wait(); 
+				}
+			);
+			while (!bBlocked)
+			{
+			}
+			// now it's blocked
+			ensure(!Task.Wait(FTimespan::FromMilliseconds(100)));
+
+			Event.Trigger(); // unblock
+			Task.Wait();
+		}
+
+		UE_BENCHMARK(5, PipeStressTest);
+
+		return true;
+	}
+
+	// stress test for named thread tasks, checks spawning a large number of tasks from multiple threads and executing them
+	void PipeStressTest()
+	{
+		constexpr uint32 SpawnerGroupsNum = 50;
+		TArray<TTask<void>> SpawnerGroups;
+		SpawnerGroups.Reserve(SpawnerGroupsNum);
+
+		constexpr uint32 SpawnersPerGroupNum = 100;
+		constexpr uint32 TasksNum = SpawnerGroupsNum * SpawnersPerGroupNum;
+		TArray<TTask<void>> Spawners;
+		Spawners.AddDefaulted(TasksNum);
+		TArray<TTask<void>> Tasks;
+		Tasks.AddDefaulted(TasksNum);
+
+		std::atomic<bool> bExecuting{ false };
+		std::atomic<uint32> TasksExecutedNum{ 0 };
+
+		FPipe Pipe{ UE_SOURCE_LOCATION };
+
+		for (uint32 GroupIndex = 0; GroupIndex != SpawnerGroupsNum; ++GroupIndex)
+		{
+			SpawnerGroups.Add(Launch(UE_SOURCE_LOCATION,
+				[
+					Spawners = &Spawners[GroupIndex * SpawnersPerGroupNum],
+					Tasks = &Tasks[GroupIndex * SpawnersPerGroupNum],
+					&bExecuting,
+					&TasksExecutedNum,
+					SpawnersPerGroupNum,
+					&Pipe
+				]
+				{
+					for (uint32 SpawnerIndex = 0; SpawnerIndex != SpawnersPerGroupNum; ++SpawnerIndex)
+					{
+						Spawners[SpawnerIndex] = Launch(UE_SOURCE_LOCATION,
+							[
+								Task = &Tasks[SpawnerIndex],
+								&bExecuting,
+								&TasksExecutedNum,
+								&Pipe
+							]
+							{
+								*Task = Pipe.Launch(UE_SOURCE_LOCATION,
+									[&bExecuting, &TasksExecutedNum]
+									{
+										check(!bExecuting);
+										bExecuting = true;
+										++TasksExecutedNum;
+										bExecuting = false;
+									}
+								);
+							}
+						);
+					}
+				}
+				));
+		}
+
+		for (TTask<void>& SpawnerGroup : SpawnerGroups)
+		{
+			SpawnerGroup.Wait();
+		}
+
+		for (TTask<void>& Spawner : Spawners)
+		{
+			Spawner.Wait();
+		}
+
+		for (TTask<void>& Task : Tasks)
+		{
+			Task.Wait();
+		}
+
+		check(TasksExecutedNum == TasksNum);
+	}
+
+	struct FAutoTlsSlot
+	{
+		uint32 Slot;
+
+		FAutoTlsSlot()
+			: Slot(FPlatformTLS::AllocTlsSlot())
+		{
+		}
+
+		~FAutoTlsSlot()
+		{
+			FPlatformTLS::FreeTlsSlot(Slot);
+		}
+	};
+
+	template<uint64 Num>
+	void UeTlsStressTest()
+	{
+		static FAutoTlsSlot Slot;
+		double Dummy = 0;
+		for (uint64 i = 0; i != Num; ++i)
+		{
+			Dummy += (double)(uintptr_t)(FPlatformTLS::GetTlsValue(Slot.Slot));
+			double Now = FPlatformTime::Seconds();
+			FPlatformTLS::SetTlsValue(Slot.Slot, (void*)(uintptr_t)(Now));
+		}
+		FPlatformTLS::SetTlsValue(Slot.Slot, (void*)(uintptr_t)(Dummy));
+	}
+
+	template<uint64 Num>
+	void ThreadLocalStressTest()
+	{
+		static thread_local double TlsValue;
+		double Dummy = 0;
+		for (uint64 i = 0; i != Num; ++i)
+		{
+			Dummy += TlsValue;
+			double Now = FPlatformTime::Seconds();
+			TlsValue = Now;
+		}
+		TlsValue = Dummy;
+	}
+
+	IMPLEMENT_SIMPLE_AUTOMATION_TEST(FTlsTest, "System.Core.Tls", EAutomationTestFlags::ApplicationContextMask | EAutomationTestFlags::EngineFilter);
+
+	bool FTlsTest::RunTest(const FString& Parameters)
+	{
+		UE_BENCHMARK(5, UeTlsStressTest<10000000>);
+		UE_BENCHMARK(5, ThreadLocalStressTest<10000000>);
+
+		return true;
+	}
+}}
+
+#endif // WITH_DEV_AUTOMATION_TESTS
diff --git a/Engine/Source/Runtime/Core/Public/HAL/PreprocessorHelpers.h b/Engine/Source/Runtime/Core/Public/HAL/PreprocessorHelpers.h
--- a/Engine/Source/Runtime/Core/Public/HAL/PreprocessorHelpers.h
+++ b/Engine/Source/Runtime/Core/Public/HAL/PreprocessorHelpers.h
@@ -25,6 +25,8 @@
 // Expands to nothing - used as a placeholder
 #define PREPROCESSOR_NOTHING
 
+#define UE_SOURCE_LOCATION TEXT(__FILE__ "(" PREPROCESSOR_TO_STRING(__LINE__) ")")
+
 // Removes a single layer of parentheses from a macro argument if they are present - used to allow
 // brackets to be optionally added when the argument contains commas, e.g.:
 //
diff --git a/Engine/Source/Runtime/Core/Public/Misc/Timeout.h b/Engine/Source/Runtime/Core/Public/Misc/Timeout.h
new file mode 100644
--- /dev/null
+++ b/Engine/Source/Runtime/Core/Public/Misc/Timeout.h
@@ -0,0 +1,44 @@
+// Copyright Epic Games, Inc. All Rights Reserved.
+
+#pragma once
+
+#include "Misc/Timespan.h"
+#include "HAL/PlatformTime.h"
+#include "CoreTypes.h"
+
+namespace UE
+{
+	// utility class to handle timeouts
+	// usage:
+	// ------------------------------------
+	// FTimeout Timeout(FTimespan::FromMilliseconds(2));
+	// while (!Timeout) { ... }
+	// ------------------------------------
+	class FTimeout
+	{
+	public:
+		explicit FTimeout(FTimespan Value)
+			: Timeout(Value)
+		{
+		}
+
+		explicit operator bool() const
+		{
+			return GetElapsedTime() >= Timeout;
+		}
+
+		FTimespan GetElapsedTime() const
+		{
+			return FTimespan::FromSeconds(FPlatformTime::Seconds()) - Start;
+		}
+
+		FTimespan GetRemainingTime() const
+		{
+			return Timeout - GetElapsedTime();
+		}
+
+	private:
+		FTimespan Start{ FTimespan::FromSeconds(FPlatformTime::Seconds()) };
+		FTimespan Timeout;
+	};
+}
diff --git a/Engine/Source/Runtime/Core/Public/Tasks/Pipe.h b/Engine/Source/Runtime/Core/Public/Tasks/Pipe.h
new file mode 100644
--- /dev/null
+++ b/Engine/Source/Runtime/Core/Public/Tasks/Pipe.h
@@ -0,0 +1,82 @@
+// Copyright Epic Games, Inc. All Rights Reserved.
+
+#pragma once
+
+#include "Tasks/Task.h"
+#include "CoreTypes.h"
+
+#include <atomic>
+
+namespace UE { namespace Tasks
+{
+	// A chain of tasks that are executed one after another. Can be used to synchronise access to a shared resource as FPipe guarantees 
+	// non-concurrent tasks execution. FPipe is a replacement for named threads because it's lightweight and flexible -
+	// there can be a large dynamic number of pipes each controlling its own shared resource. Can be used as a replacement for
+	// dedicated threads.
+	// Execution order is not specified, only that tasks from the same pipe are not executed concurrently.
+	// A pipe must be alive until its last task is completed.
+	// See `FTasksPipeTest` for tests and examples.
+	class FPipe
+	{
+	public:
+		UE_NONCOPYABLE(FPipe);
+
+		// @param InDebugName helps to identify the pipe in debugger and profiler. `UE_SOURCE_LOCATION` can be used as an auto-generated
+		// unique name.
+		explicit FPipe(const TCHAR* InDebugName)
+			: DebugName(InDebugName)
+		{}
+
+		~FPipe()
+		{
+			check(!HasWork());
+		}
+
+		// returns `true` if the pipe has any not completed tasks
+		bool HasWork() const
+		{
+			return LastTask.load(std::memory_order_relaxed) != nullptr;
+		}
+
+		// launches a task in the pipe
+		// @param InDebugName helps to identify the task in debugger and profiler
+		// @param TaskBody a callable with no parameters whose return value is ignored, usually a lambda but can be also a functor object 
+		// or a pointer to a function. TaskBody can return results.
+		// @Priority - task priority, can affect task scheduling once it's passed the pipe
+		// @return Task instance that can be used to wait for task completion or to obtain the result of task execution
+		template<typename TaskBodyType>
+		TTask<TInvokeResult_T<TaskBodyType>> Launch(const TCHAR* InDebugName, TaskBodyType&& TaskBody, LowLevelTasks::ETaskPriority Priority = LowLevelTasks::ETaskPriority::Default)
+		{
+			using FResult = TInvokeResult_T<TaskBodyType>;
+			Private::TTaskWithResult<FResult>* Task = new Private::TTaskWithResult<FResult>;
+			Task->Init(InDebugName, Forward<TaskBodyType>(TaskBody), Priority);
+			Task->SetPipe(*this);
+			Task->TryLaunch();
+			return TTask<FResult>{ Task };
+		}
+
+	private:
+		friend class Private::FTaskBase;
+
+		// pushes given task into the pipe: adds the task as a subsequent to the last task if any and sets it as the new last task
+		// returns true if there wasn't any task in the pipe and the given task is free to be executed
+		bool PushIntoPipe(Private::FTaskBase& Task);
+		// pipe holds a "weak" reference to a task. the task must be cleared from the pipe when its execution finished before its completion, 
+		// otherwise the next piped task can try to add itself as a subsequent to already destroyed task
+		void ClearTask(Private::FTaskBase& Task);
+
+	private:
+		std::atomic<Private::FTaskBase*> LastTask{ nullptr }; // the last task that 
+		// a counter of how many threads trying to push a task into a pipe right now
+		std::atomic<int32> PushingThreadsNum{ 0 };
+
+	public:
+		FORCENOINLINE const TCHAR* GetDebugName() const
+		{
+			return DebugName;
+		}
+
+	private:
+		const TCHAR* const DebugName;
+	};
+}}
diff --git a/Engine/Source/Runtime/Core/Public/Tasks/Task.h b/Engine/Source/Runtime/Core/Public/Tasks/Task.h
new file mode 100644
--- /dev/null
+++ b/Engine/Source/Runtime/Core/Public/Tasks/Task.h
@@ -0,0 +1,186 @@
+// Copyright Epic Games, Inc. All Rights Reserved.
+
+#pragma once
+
+#include "Tasks/TaskPrivate.h"
+#include "Async/Fundamental/Task.h"
+#include "CoreTypes.h"
+
+namespace UE { namespace Tasks
+{
+	// a common part of the generic `TTask<ResultType>` and its `TTask<void>` specialisation
+	template<typename ResultType>
+	class TTaskBase
+	{
+	protected:
+		TTaskBase() = default;
+
+		explicit TTaskBase(Private::TTaskWithResult<ResultType>* Other)
+			: Pimpl(Other, /*bAddRef = */false)
+		{}
+
+	public:
+		bool IsValid() const
+		{
+			return Pimpl.IsValid();
+		}
+
+		// checks if task's execution is done
+		bool IsCompleted() const
+		{
+			return !IsValid() || Pimpl->IsCompleted();
+		}
+
+		// waits for task's completion
+		void Wait()
+		{
+			if (IsValid())
+			{
+				Pimpl->Wait();
+			}
+		}
+
+		// waits for task's completion at least specified amount of time.
+		// the call can return much later than the given timeout
+		// @return true if the task is completed
+		bool Wait(FTimespan Timeout)
+		{
+			return !IsValid() || Pimpl->Wait(Timeout);
+		}
+
+		// waits for task's completion or the given condition becomes true
+		// the call can return much later than the given condition became true
+		// @return true if the task is completed
+		template<typename ConditionType>
+		bool Wait(ConditionType&& Condition)
+		{
+			return !IsValid() || Pimpl->Wait(Forward<ConditionType>(Condition));
+		}
+
+	protected:
+		TRefCountPtr<Private::TTaskWithResult<ResultType>> Pimpl;
+	};
+
+	// a movable/copyable handle of `Private::TTaskWithResult*` with the API adopted for public usage.
+	// implements Pimpl idiom
+	template<typename ResultType>
+	class TTask : public TTaskBase<ResultType>
+	{
+		using FSuper = TTaskBase<ResultType>;
+
+	public:
+		TTask() = default;
+
+		// waits until the task is completed and returns task's result
+		ResultType& GetResult()
+		{
+			check(FSuper::IsValid());
+			return FSuper::Pimpl->GetResult();
+		}
+
+	private:
+		template<typename TaskBodyType>
+		friend TTask<TInvokeResult_T<TaskBodyType>> Launch(const TCHAR* DebugName, TaskBodyType&& TaskBody, LowLevelTasks::ETaskPriority Priority/* = LowLevelTasks::ETaskPriority::Normal*/);
+
+		friend FPipe;
+
+		// private constructor, valid instances can be created only by launching (see friends)
+		explicit TTask(Private::TTaskWithResult<ResultType>* Other)
+			: FSuper(Other)
+		{}
+	};
+
+	template<>
+	class TTask<void> : public TTaskBase<void>
+	{
+	public:
+		TTask() = default;
+	
+		void GetResult()
+		{
+			check(IsValid()); // to be consistent with a generic `TTask<ResultType>::GetResult()`
+			Wait();
+		}
+
+	private:
+		template<typename TaskBodyType>
+		friend TTask<TInvokeResult_T<TaskBodyType>> Launch(const TCHAR* DebugName, TaskBodyType&& TaskBody, LowLevelTasks::ETaskPriority Priority/* = LowLevelTasks::ETaskPriority::Normal*/);
+
+		friend FPipe;
+
+		// private constructor, valid instances can be created only by launching (see friends)
+		explicit TTask(Private::TTaskWithResult<void>* Other)
+			: TTaskBase<void>(Other)
+		{}
+	};
+
+	// a synchronisation primitive similar to `FEvent` that uses "busy waiting" - executing tasks while waiting
+	class FTaskEvent
+	{
+	public:
+		UE_NONCOPYABLE(FTaskEvent)
+
+		FTaskEvent() = default;
+
+		void Trigger()
+		{
+			bTriggered.store(true, std::memory_order_release);
+		}
+
+		bool IsTriggered() const
+		{
+			return bTriggered.load(std::memory_order_acquire);
+		}
+
+		// waits until the event is triggered
+		void Wait()
+		{
+			LowLevelTasks::BusyWaitUntil([this] { return bTriggered.load(std::memory_order_acquire); });
+		}
+
+		// waits until the event is triggered or the waiting timed out.
+		// the call can return much later than the given timeout
+		// @return true if the event was triggered
+		bool Wait(FTimespan InTimeout)
+		{
+			LowLevelTasks::BusyWaitUntil(
+				[this, Timeout = FTimeout{ InTimeout }] { return bTriggered.load(std::memory_order_acquire) || Timeout; }
+			);
+			return bTriggered.load(std::memory_order_relaxed);
+		}
+
+		// waits until the event is triggered or the given condition returns true.
+		// the call can return much later than the condition has become true
+		// @return true if the event was triggered
+		template<typename ConditionType>
+		bool Wait(ConditionType&& Condition)
+		{
+			LowLevelTasks::BusyWaitUntil(
+				[this, Condition = Forward<ConditionType>(Condition)] { return bTriggered.load(std::memory_order_acquire) || Condition(); }
+			);
+			return bTriggered.load(std::memory_order_relaxed);
+		}
+
+	private:
+		std::atomic<bool> bTriggered{ false };
+	};
+
+	// launches a task for asynchronous execution
+	// @param DebugName - a unique name for task identification in debugger and profiler, is compiled out in test/shipping builds
+	// @param TaskBody - a functor that will be executed asynchronously
+	// @param Priority - task priority that affects when the task will be executed
+	// @return a trivially relocatable instance that can be used to wait for task completion or to obtain task execution result
+	template<typename TaskBodyType>
+	TTask<TInvokeResult_T<TaskBodyType>> Launch(
+		const TCHAR* DebugName, 
+		TaskBodyType&& TaskBody, 
+		LowLevelTasks::ETaskPriority Priority/* = LowLevelTasks::ETaskPriority::Normal*/
+	)
+	{
+		using FResult = TInvokeResult_T<TaskBodyType>;
+		auto* Task{ new Private::TTaskWithResult<FResult> };
+		Task->Init(DebugName, Forward<TaskBodyType>(TaskBody), Priority);
+		Task->TryLaunch();
+		return TTask<FResult>{ Task };
+	}
+}}
\ No newline at end of file
diff --git a/Engine/Source/Runtime/Core/Public/Tasks/TaskPrivate.h b/Engine/Source/Runtime/Core/Public/Tasks/TaskPrivate.h
new file mode 100644
--- /dev/null
+++ b/Engine/Source/Runtime/Core/Public/Tasks/TaskPrivate.h
@@ -0,0 +1,374 @@
+// Copyright Epic Games, Inc. All Rights Reserved.
+
+#pragma once
+
+#include "Async/Fundamental/Task.h"
+#include "Templates/Invoke.h"
+#include "Templates/TypeCompatibleBytes.h"
+#include "Misc/Timeout.h"
+#include "CoreTypes.h"
+
+#include <atomic>
+#include <type_traits>
+
+namespace UE { namespace Tasks
+{
+	using LowLevelTasks::ETaskPriority;
+
+	// BEGIN: forward declarations
+
+	class FPipe;
+
+	template<typename T>
+	class TTask;
+
+	template<typename TaskBodyType>
+	TTask<TInvokeResult_T<TaskBodyType>> Launch(
+		const TCHAR* DebugName, 
+		TaskBodyType&& TaskBody, 
+		LowLevelTasks::ETaskPriority Priority = LowLevelTasks::ETaskPriority::Normal
+	);
+
+	// END: forward declarations
+
+	namespace Private
+	{
+		// Task class implementation with piping support.
+		class CORE_API FTaskBase
+		{
+			UE_NONCOPYABLE(FTaskBase);
+
+			// Can be used only as a base class.
+		protected:
+			FTaskBase()
+				: PipePtr(reinterpret_cast<uintptr_t>(nullptr))
+				, bPushedIntoPipe(false)
+			{
+			}
+
+			~FTaskBase()
+			{
+				check(IsCompleted());
+			}
+
+		public:
+			// initialises the task but doesn't launches it
+			template<typename TaskBodyType, typename DeleterType>
+			void Init(const TCHAR* DebugName, TaskBodyType&& TaskBody, DeleterType&& Deleter, LowLevelTasks::ETaskPriority Priority)
+			{
+				// to not store a copy of `TaskBody` that is stored inside `LowLevelTask` anyway, and to implement task retraction
+				// (see `Wait()` method), we have to use LowLevelTask's "Continuation" parameter instead of "Runnable". This way we can make
+				// the scheduler immediately execute retracted task.
+				LowLevelTask.Init(DebugName, Priority, 
+					[] {}, // Runnable
+					[
+						this, 
+						TaskBody = Forward<TaskBodyType>(TaskBody), 
+						// releasing scheduler's task reference can cause task's automatic destruction and so must be done after the task is 
+						// flagged as completed. The task is flagged as completed after the continuation is executed but before its destroyed.
+						// `Deleter` is captured by value and is destroyed along with the continuation, calling the given functor 
+						// on destruction
+						Deleter = Forward<DeleterType>(Deleter)
+					]() mutable
+					{
+						Invoke(TaskBody);
+						// finalisation
+						UnblockPipe();
+						FTaskBase* Subsequent = CloseAndReturnSubsequent();
+						if (Subsequent != nullptr)
+						{
+							check(TryPushIntoPipe());
+							LowLevelTasks::TryLaunch(Subsequent->LowLevelTask, LowLevelTasks::EQueuePreference::DefaultPreference, /*bWakeUpWorker =*/ false);
+						}
+					} // Continuation
+				);
+			}
+
+			// A piped task will be executed after the previous task from this pipe is completed. Tasks from the same pipe are not executed
+			// concurrently (so don't require synchronization), but not necessarily on the same thread.
+			// @See FPipe
+			void SetPipe(FPipe& Pipe)
+			{
+				PipePtr = reinterpret_cast<uintptr_t>(&Pipe);
+			}
+
+			FPipe* GetPipe() const
+			{
+				return reinterpret_cast<FPipe*>(PipePtr);
+			}
+
+			// allows the task to be scheduled and executed. Tasks w/o pending dependencies (not blocked by a pipe) are 
+			// scheduled immediately
+			bool TryLaunch()
+			{
+				if (TryPushIntoPipe())
+				{
+					// scheduler's reference was accounted on task creation
+					return LowLevelTasks::TryLaunch(LowLevelTask);
+				}
+				return false;
+			}
+
+			// sets a subsequent task that will be launched automatically when this task is done
+			// returns true only if the task is not closed yet ("closed" = is executed and has already checked its subsequent),
+			// if returned false the subsequent is not blocked by this task
+			bool SetSubsequent(FTaskBase& Task)
+			{
+#if DO_CHECK
+				FTaskBase* CurrentState = SubsequentAndState.load(std::memory_order_relaxed);
+				checkf(CurrentState == GetEmptyOpenState() || CurrentState == GetClosedState(), TEXT("only a single subsequent can be set"));
+#endif
+
+				FTaskBase* EmptyNotCompleted = GetEmptyOpenState();
+				FTaskBase* NotCompletedWithTask = &Task;
+				return SubsequentAndState.compare_exchange_strong(EmptyNotCompleted, NotCompletedWithTask, std::memory_order_release, std::memory_order_relaxed);
+			}
+
+			bool IsCompleted() const
+			{
+				return LowLevelTask.IsCompleted();
+			}
+
+			// waits until the task is completed
+			void Wait()
+			{
+				// we try to retract the task from the scheduler, if it's execution hasn't started yet, to execute it immediately. This makes
+				// the task skip waiting in the scheduler's queue.
+				// we can't retract it right here because it can be not scheduled yet (e.g. when it's blocked by a pipe), so we need to wait for this
+				LowLevelTasks::BusyWaitUntil([this] { return !LowLevelTask.IsReady(); }); // `!IsReady()` here means "the task is scheduled"
+				LowLevelTasks::TryCancelAndLaunchContinuation(LowLevelTask);
+				LowLevelTasks::BusyWaitUntil([this] { return LowLevelTask.IsCompleted(); });
+			}
+
+			// waits until the task is completed or waiting timed out
+			// @see `void Wait()`
+			bool Wait(FTimespan InTimeout)
+			{
+				FTimeout Timeout{ InTimeout };
+				LowLevelTasks::BusyWaitUntil([this, Timeout] { return !LowLevelTask.IsReady() || Timeout; });
+				if (LowLevelTask.IsReady()) // timeout
+				{
+					return false;
+				}
+
+				LowLevelTasks::TryCancelAndLaunchContinuation(LowLevelTask);
+				LowLevelTasks::BusyWaitUntil([this, Timeout] { return LowLevelTask.IsCompleted() || Timeout; });
+				return LowLevelTask.IsCompleted();
+			}
+
+			// waits until the task is completed or the condition returns true
+			// @see `void Wait()`
+			template<typename ConditionType>
+			bool Wait(ConditionType&& Condition)
+			{
+				LowLevelTasks::BusyWaitUntil(
+					[this, Condition = Forward<ConditionType>(Condition)] { return !LowLevelTask.IsReady() || Condition(); }
+				);
+
+				if (LowLevelTask.IsReady()) // `Condition` returned true
+				{
+					return false;
+				}
+
+				LowLevelTasks::TryCancelAndLaunchContinuation(LowLevelTask);
+				LowLevelTasks::BusyWaitUntil(
+					[this, Condition = Forward<ConditionType>(Condition)] { return LowLevelTask.IsCompleted() || Condition(); }
+				);
+				return LowLevelTask.IsCompleted(std::memory_order_relaxed);
+			}
+
+		private:
+			// checks if the task is ready to be launched by trying to push it into the pipe
+			bool TryPushIntoPipe()
+			{
+				if (GetPipe() == nullptr)
+				{
+					return true;
+				}
+
+				if (!bPushedIntoPipe)
+				{
+					bPushedIntoPipe = true;
+					return PushIntoPipe();
+				}
+
+				return true;
+			}
+
+			bool PushIntoPipe();
+			void UnblockPipe();
+
+			// after being executed the task checks if it has a subsequent to launch. it's the only time the task will do this check so
+			// it "closes" itself and setting a subsequent fill fail after that.
+			// atomic operation
+			FTaskBase* CloseAndReturnSubsequent()
+			{
+				FTaskBase* Subsequent_Local = SubsequentAndState.exchange(GetClosedState(), std::memory_order_acq_rel);
+				checkf(Subsequent_Local != GetClosedState(), TEXT("Task can be closed only once"));
+				return Subsequent_Local;
+			}
+
+			// no subsequent, still open
+			static FTaskBase* GetEmptyOpenState()
+			{
+				return nullptr;
+			}
+
+			static FTaskBase* GetClosedState()
+			{
+				// nullptr with the most significant bit set to indicate the closed state
+				return UE_LAUNDER(reinterpret_cast<FTaskBase*>(1ull << 63));
+			}
+
+		private:
+			LowLevelTasks::FTask LowLevelTask;
+
+			// packed pipe's ptr and the flag if the task is already pushed into pipe
+			uintptr_t PipePtr : 63; // actually `FPipe*`
+			uintptr_t bPushedIntoPipe : 1;	// bool, if the task is pushed into pipe
+
+			// a pointer to subsequent task with the most significant bit borrowed to be used as a close flag
+			// `LowLevelTask.IsCompleted()` flag can't be used as a "close" flag because we need to check if the task still can 
+			// accept a subsequent and to set the subsequent atomically.
+			// can be "empty and open" (`nullptr`), "have subsequent and open" (subsequent ptr) and closed (`nullptr` with the most significant
+			// bit set)
+			std::atomic<FTaskBase*> SubsequentAndState{ GetEmptyOpenState() };
+		};
+
+		// Extends FTaskBase by supporting execution result.
+		template<typename ResultType>
+		class TTaskWithResultBase : public FTaskBase
+		{
+			UE_NONCOPYABLE(TTaskWithResultBase);
+
+			using FSelf = TTaskWithResultBase<ResultType>;
+
+			// Can be used only as a base class
+		protected:
+			TTaskWithResultBase() = default;
+
+			~TTaskWithResultBase()
+			{
+				// Every task instance must be launched. This is quaranteed by the the high level task API (`UE::Tasks::TTask`)
+				checkf(bLaunched, TEXT("Every task instance must be launched"));
+				DestructItem(ResultStorage.GetTypedPtr());
+			}
+
+		public:
+			template<typename TaskBodyType, typename DeleterType>
+			void Init(const TCHAR* DebugName, TaskBodyType&& TaskBody, DeleterType&& Deleter, LowLevelTasks::ETaskPriority Priority)
+			{
+				FTaskBase::Init(
+					DebugName, 
+					[this, TaskBody = Forward<TaskBodyType>(TaskBody)]() mutable
+					{ 
+						new(&ResultStorage) ResultType{ Invoke(TaskBody) }; 
+					},
+					Forward<DeleterType>(Deleter),
+					Priority
+				);
+			}
+
+			bool TryLaunch()
+			{
+#if DO_CHECK
+				bLaunched = true;
+#endif
+				return FTaskBase::TryLaunch();
+			}
+
+			ResultType& GetResult()
+			{
+				check(bLaunched);
+				Wait();
+
+				return *ResultStorage.GetTypedPtr();
+			}
+
+		private:
+			TTypeCompatibleBytes<ResultType> ResultStorage;
+
+#if DO_CHECK
+			bool bLaunched = false;
+#endif
+		};
+
+		template<>
+		class CORE_API TTaskWithResultBase<void> : public FTaskBase
+		{
+		protected:
+			TTaskWithResultBase() = default; // instances can be created only on the heap by a ref-counting handler
+
+		public:
+			void GetResult()
+			{
+				Wait();
+			}
+		};
+
+		// intrusive atomic reference-counting base class.
+		// doesn't require derived classes to be polymorphic
+		template<typename ReferenceType>
+		class CORE_API TRefCountedBase
+		{
+		public:
+			explicit TRefCountedBase(uint32 InitRefCount = 0)
+				: RefCount{ InitRefCount }
+			{
+				static_assert(std::is_final<ReferenceType>::value, "ReferenceType should be `final` to guarantee no slicing during destruction");
+			}
+
+			void AddRef()
+			{
+				RefCount.fetch_add(1, std::memory_order_relaxed);
+			}
+
+			void Release()
+			{
+				uint32 LocalRefCount = RefCount.fetch_sub(1, std::memory_order_release) - 1;
+				if (LocalRefCount == 0)
+				{
+					std::atomic_thread_fence(std::memory_order_acquire);
+					delete static_cast<ReferenceType*>(this);
+				}
+			}
+
+		private:
+			std::atomic<uint32> RefCount;
+		};
+
+		// A final task class that must be used only with `TRefCountPtr`, thus only friends are allowed to construct instances.
+		// @see UE::Tasks::Task<T>
+		template<typename ResultType>
+		class CORE_API TTaskWithResult final : public TTaskWithResultBase<ResultType>, public TRefCountedBase<TTaskWithResult<ResultType>>
+		{
+			UE_NONCOPYABLE(TTaskWithResult);
+
+			using FSelf = TTaskWithResult<ResultType>;
+
+		private:
+			template<typename TaskBodyType>
+			friend TTask<TInvokeResult_T<TaskBodyType>> UE::Tasks::Launch(const TCHAR* DebugName, TaskBodyType&& TaskBody, LowLevelTasks::ETaskPriority Priority);
+
+			friend FPipe;
+
+			TTaskWithResult()
+				: TRefCountedBase<FSelf>(2) // one for the initial reference (we don't increment it when passing to `TRefCountPtr`) and 
+				// one for scheduler's reference which will be released on task completion
+			{
+			}
+
+		public:
+			template<typename TaskBodyType>
+			void Init(const TCHAR* DebugName, TaskBodyType&& TaskBody, LowLevelTasks::ETaskPriority Priority)
+			{
+				TTaskWithResultBase<ResultType>::Init(
+					DebugName,
+					Forward<TaskBodyType>(TaskBody),
+					LowLevelTasks::TDeleter<TRefCountedBase<FSelf>, &TRefCountedBase<FSelf>::Release>{ this },
+					Priority
+				);
+			}
+		};
+	}
+}}
\ No newline at end of file
