Fixed missing EStatFlags::ClearEveryFrame flag in DECLARE_CYCLE_STAT_WITH_FLAGS

diff --git a/Engine/Source/Runtime/Core/Public/Stats/Stats2.h b/Engine/Source/Runtime/Core/Public/Stats/Stats2.h
--- a/Engine/Source/Runtime/Core/Public/Stats/Stats2.h
+++ b/Engine/Source/Runtime/Core/Public/Stats/Stats2.h
@@ -1876,7 +1876,7 @@ struct FStat_##StatName\
 	DECLARE_STAT(CounterName,StatId,GroupId,EStatDataType::ST_int64, EStatFlags::ClearEveryFrame | EStatFlags::CycleStat, FPlatformMemory::MCR_Invalid); \
 	static DEFINE_STAT(StatId)
 #define DECLARE_CYCLE_STAT_WITH_FLAGS(CounterName,StatId,GroupId,StatFlags) \
-	DECLARE_STAT(CounterName,StatId,GroupId,EStatDataType::ST_int64, (StatFlags) | EStatFlags::CycleStat, FPlatformMemory::MCR_Invalid); \
+	DECLARE_STAT(CounterName,StatId,GroupId,EStatDataType::ST_int64, (StatFlags) | EStatFlags::ClearEveryFrame | EStatFlags::CycleStat, FPlatformMemory::MCR_Invalid); \
 	static DEFINE_STAT(StatId)
 #define DECLARE_FLOAT_COUNTER_STAT(CounterName,StatId,GroupId) \
 	DECLARE_STAT(CounterName,StatId,GroupId,EStatDataType::ST_double,EStatFlags::ClearEveryFrame, FPlatformMemory::MCR_Invalid); \
