compilation fix for NonUnity Compile UnrealGame Android

diff --git a/Engine/Source/Runtime/Core/Public/Tasks/TaskPrivate.h b/Engine/Source/Runtime/Core/Public/Tasks/TaskPrivate.h
--- a/Engine/Source/Runtime/Core/Public/Tasks/TaskPrivate.h
+++ b/Engine/Source/Runtime/Core/Public/Tasks/TaskPrivate.h
@@ -3,6 +3,7 @@
 #pragma once
 
 #include "Async/Fundamental/Task.h"
+#include "Async/Fundamental/Scheduler.h"
 #include "Templates/Invoke.h"
 #include "Templates/TypeCompatibleBytes.h"
 #include "Misc/Timeout.h"
