SharedBuffer: MakeOwned now returns a buffer instead of operating in-place

#rnx

diff --git a/Engine/Source/Editor/UnrealEd/Private/EditorDomain/EditorDomainArchive.cpp b/Engine/Source/Editor/UnrealEd/Private/EditorDomain/EditorDomainArchive.cpp
--- a/Engine/Source/Editor/UnrealEd/Private/EditorDomain/EditorDomainArchive.cpp
+++ b/Engine/Source/Editor/UnrealEd/Private/EditorDomain/EditorDomainArchive.cpp
@@ -188,7 +188,7 @@ void FEditorDomainReadArchive::OnCacheRequestComplete(UE::DerivedData::FCacheGet
 	{
 		const FCbObject& MetaData = Params.Record.GetMeta();
 		int64 FileSize = MetaData["FileSize"].AsInt64(-1);
-		FSharedBuffer LocalBytes = Params.Record.GetValue();
+		FSharedBuffer LocalBytes = MoveTemp(Params.Record).GetValue();
 		if (static_cast<int64>(LocalBytes.GetSize()) != FileSize)
 		{
 			UE_LOG(LogEditorDomain, Warning, TEXT("Package %s received invalid record from EditorDomainPackage table ")
@@ -200,8 +200,7 @@ void FEditorDomainReadArchive::OnCacheRequestComplete(UE::DerivedData::FCacheGet
 		{
 			AsyncSource = ESource::Bytes;
 			Size = FileSize;
-			LocalBytes.MakeOwned();
-			Bytes = MoveTemp(LocalBytes);
+			Bytes = MoveTemp(LocalBytes).MakeOwned();
 			PackageFormat = EPackageFormat::Binary;
 			PackageSource->Source = FEditorDomain::EPackageSource::Editor;
 		}
@@ -417,7 +416,7 @@ void FEditorDomainAsyncReadFileHandle::OnCacheRequestComplete(UE::DerivedData::F
 	{
 		const FCbObject& MetaData = Params.Record.GetMeta();
 		int64 FileSize = MetaData["FileSize"].AsInt64(-1);
-		FSharedBuffer LocalBytes = Params.Record.GetValue();
+		FSharedBuffer LocalBytes = MoveTemp(Params.Record).GetValue();
 		if (static_cast<int64>(LocalBytes.GetSize()) != FileSize)
 		{
 			UE_LOG(LogEditorDomain, Warning, TEXT("Package %s received invalid record from EditorDomainPackage ")
@@ -428,8 +427,7 @@ void FEditorDomainAsyncReadFileHandle::OnCacheRequestComplete(UE::DerivedData::F
 		else
 		{
 			AsyncSource = ESource::Bytes;
-			LocalBytes.MakeOwned();
-			Bytes = MoveTemp(LocalBytes);
+			Bytes = MoveTemp(LocalBytes).MakeOwned();
 			PackageSource->Source = FEditorDomain::EPackageSource::Editor;
 		}
 	}
diff --git a/Engine/Source/Runtime/Core/Private/Memory/SharedBuffer.cpp b/Engine/Source/Runtime/Core/Private/Memory/SharedBuffer.cpp
--- a/Engine/Source/Runtime/Core/Private/Memory/SharedBuffer.cpp
+++ b/Engine/Source/Runtime/Core/Private/Memory/SharedBuffer.cpp
@@ -139,12 +139,9 @@ void FUniqueBuffer::Reset()
 	Owner.Reset();
 }
 
-void FUniqueBuffer::MakeOwned()
+FUniqueBuffer FUniqueBuffer::MakeOwned() &&
 {
-	if (!IsOwned())
-	{
-		*this = Clone(GetView());
-	}
+	return IsOwned() ? MoveTemp(*this) : Clone(GetView());
 }
 
 void FUniqueBuffer::Materialize() const
@@ -211,12 +208,14 @@ void FSharedBuffer::Reset()
 	Owner.Reset();
 }
 
-void FSharedBuffer::MakeOwned()
+FSharedBuffer FSharedBuffer::MakeOwned() const &
 {
-	if (!IsOwned())
-	{
-		*this = FSharedBuffer::Clone(GetView());
-	}
+	return IsOwned() ? *this : Clone(GetView());
+}
+
+FSharedBuffer FSharedBuffer::MakeOwned() &&
+{
+	return IsOwned() ? MoveTemp(*this) : Clone(GetView());
 }
 
 void FSharedBuffer::Materialize() const
diff --git a/Engine/Source/Runtime/Core/Private/Serialization/CompactBinaryPackage.cpp b/Engine/Source/Runtime/Core/Private/Serialization/CompactBinaryPackage.cpp
--- a/Engine/Source/Runtime/Core/Private/Serialization/CompactBinaryPackage.cpp
+++ b/Engine/Source/Runtime/Core/Private/Serialization/CompactBinaryPackage.cpp
@@ -40,9 +40,8 @@ FCbAttachment::FCbAttachment(FCbFieldIterator InValue, const FIoHash* const InHa
 }
 
 FCbAttachment::FCbAttachment(FSharedBuffer InBuffer, const FIoHash* const InHash)
-	: Buffer(MoveTemp(InBuffer))
+	: Buffer(MoveTemp(InBuffer).MakeOwned())
 {
-	Buffer.MakeOwned();
 	if (InHash)
 	{
 		Hash = *InHash;
@@ -92,8 +91,7 @@ void FCbAttachment::Load(FCbFieldIterator& Fields)
 	const FMemoryView View = Fields.AsBinaryView();
 	if (View.GetSize() > 0)
 	{
-		Buffer = FSharedBuffer::MakeView(View, Fields.GetOuterBuffer());
-		Buffer.MakeOwned();
+		Buffer = FSharedBuffer::MakeView(View, Fields.GetOuterBuffer()).MakeOwned();
 		++Fields;
 		Hash = Fields.AsAttachment();
 		checkf(!Fields.HasError(), TEXT("Attachments must be a non-empty binary value with a content hash."));
@@ -119,8 +117,7 @@ void FCbAttachment::Load(FArchive& Ar, FCbBufferAllocator Allocator)
 	const FMemoryView View = BufferField.AsBinaryView();
 	if (View.GetSize() > 0)
 	{
-		Buffer = FSharedBuffer::MakeView(View, BufferField.GetOuterBuffer());
-		Buffer.MakeOwned();
+		Buffer = FSharedBuffer::MakeView(View, BufferField.GetOuterBuffer()).MakeOwned();
 		CompactBinary = FCbFieldViewIterator();
 
 		TArray<uint8, TInlineAllocator<64>> HashBuffer;
@@ -330,8 +327,7 @@ void FCbPackage::Load(FArchive& Ar, FCbBufferAllocator Allocator)
 			const FMemoryView View = ValueField.AsBinaryView();
 			if (View.GetSize() > 0)
 			{
-				FSharedBuffer Buffer = FSharedBuffer::MakeView(View, ValueField.GetOuterBuffer());
-				Buffer.MakeOwned();
+				FSharedBuffer Buffer = FSharedBuffer::MakeView(View, ValueField.GetOuterBuffer()).MakeOwned();
 				FCbField HashField = LoadCompactBinary(Ar, StackAllocator);
 				const FIoHash& Hash = HashField.AsAttachment();
 				checkf(!HashField.HasError(), TEXT("Attachments must be a non-empty binary value with a content hash."));
diff --git a/Engine/Source/Runtime/Core/Private/Tests/Memory/SharedBufferTest.cpp b/Engine/Source/Runtime/Core/Private/Tests/Memory/SharedBufferTest.cpp
--- a/Engine/Source/Runtime/Core/Private/Tests/Memory/SharedBufferTest.cpp
+++ b/Engine/Source/Runtime/Core/Private/Tests/Memory/SharedBufferTest.cpp
@@ -226,13 +226,12 @@ bool FUniqueBufferTest::RunTest(const FString& Parameters)
 	{
 		constexpr uint64 Size = 64;
 		uint8 Data[Size]{};
-		FUniqueBuffer Buffer = FUniqueBuffer::MakeView(Data, Size);
-		Buffer.MakeOwned();
+		FUniqueBuffer Buffer = FUniqueBuffer::MakeView(Data, Size).MakeOwned();
 		TestTrue(TEXT("FUniqueBuffer::MakeOwned(View).IsOwned()"), Buffer.IsOwned());
 		TestEqual(TEXT("FUniqueBuffer::MakeOwned(View).GetSize()"), Buffer.GetSize(), Size);
 		TestFalse(TEXT("FUniqueBuffer::MakeOwned(View).GetData()"), Buffer.GetData() == Data);
 		void* const OwnedData = Buffer.GetData();
-		Buffer.MakeOwned();
+		Buffer = MoveTemp(Buffer).MakeOwned();
 		TestTrue(TEXT("FUniqueBuffer::MakeOwned(Owned).IsOwned()"), Buffer.IsOwned());
 		TestEqual(TEXT("FUniqueBuffer::MakeOwned(Owned).GetSize()"), Buffer.GetSize(), Size);
 		TestEqual(TEXT("FUniqueBuffer::MakeOwned(Owned).GetData()"), Buffer.GetData(), OwnedData);
@@ -336,13 +335,11 @@ bool FSharedBufferTest::RunTest(const FString& Parameters)
 	{
 		constexpr uint64 Size = 64;
 		uint8 Data[Size]{};
-		FSharedBuffer Buffer = FSharedBuffer::MakeView(Data, Size);
-		Buffer.MakeOwned();
+		FSharedBuffer Buffer = FSharedBuffer::MakeView(Data, Size).MakeOwned();
 		TestTrue(TEXT("FSharedBuffer::MakeOwned(View).IsOwned()"), Buffer.IsOwned());
 		TestEqual(TEXT("FSharedBuffer::MakeOwned(View).GetSize()"), Buffer.GetSize(), Size);
 		TestFalse(TEXT("FSharedBuffer::MakeOwned(View).GetData()"), Buffer.GetData() == Data);
-		FSharedBuffer BufferCopy = Buffer;
-		BufferCopy.MakeOwned();
+		FSharedBuffer BufferCopy = Buffer.MakeOwned();
 		TestTrue(TEXT("FSharedBuffer::MakeOwned(Owned).IsOwned()"), BufferCopy.IsOwned());
 		TestEqual(TEXT("FSharedBuffer::MakeOwned(Owned).GetSize()"), BufferCopy.GetSize(), Size);
 		TestEqual(TEXT("FSharedBuffer::MakeOwned(Owned).GetData()"), BufferCopy.GetData(), Buffer.GetData());
@@ -396,10 +393,10 @@ bool FBufferOwnerTest::RunTest(const FString& Parameters)
 	TestFalse(TEXT("FSharedBuffer(BufferOwner)->!Freed"), bFreed);
 	TestFalse(TEXT("FSharedBuffer(BufferOwner)->!Deleted"), bDeleted);
 
-	Buffer.GetData();
+	(void)Buffer.GetData();
 	TestTrue(TEXT("FSharedBuffer(BufferOwner)->GetData()->Materialized"), bMaterialized);
 	bMaterialized = false;
-	Buffer.GetData();
+	(void)Buffer.GetData();
 	TestFalse(TEXT("FSharedBuffer(BufferOwner)->GetData()->!Materialized"), bMaterialized);
 	bMaterialized = false;
 
@@ -410,10 +407,10 @@ bool FBufferOwnerTest::RunTest(const FString& Parameters)
 	bFreed = false;
 	bDeleted = false;
 
-	Buffer.GetSize();
+	(void)Buffer.GetSize();
 	TestTrue(TEXT("FSharedBuffer(BufferOwner)->GetSize()->Materialized"), bMaterialized);
 	bMaterialized = false;
-	Buffer.GetSize();
+	(void)Buffer.GetSize();
 	TestFalse(TEXT("FSharedBuffer(BufferOwner)->GetSize()->!Materialized"), bMaterialized);
 	bMaterialized = false;
 
diff --git a/Engine/Source/Runtime/Core/Public/Memory/SharedBuffer.h b/Engine/Source/Runtime/Core/Public/Memory/SharedBuffer.h
--- a/Engine/Source/Runtime/Core/Public/Memory/SharedBuffer.h
+++ b/Engine/Source/Runtime/Core/Public/Memory/SharedBuffer.h
@@ -201,15 +201,15 @@ class FUniqueBuffer
 {
 public:
 	/** Make an uninitialized owned buffer of the specified size. */
-	CORE_API static FUniqueBuffer Alloc(uint64 Size);
+	UE_NODISCARD CORE_API static FUniqueBuffer Alloc(uint64 Size);
 
 	/** Make an owned clone of the input. */
-	CORE_API static FUniqueBuffer Clone(FMemoryView View);
-	CORE_API static FUniqueBuffer Clone(const void* Data, uint64 Size);
+	UE_NODISCARD CORE_API static FUniqueBuffer Clone(FMemoryView View);
+	UE_NODISCARD CORE_API static FUniqueBuffer Clone(const void* Data, uint64 Size);
 
 	/** Make a non-owned view of the input. */
-	CORE_API static FUniqueBuffer MakeView(FMutableMemoryView View);
-	CORE_API static FUniqueBuffer MakeView(void* Data, uint64 Size);
+	UE_NODISCARD CORE_API static FUniqueBuffer MakeView(FMutableMemoryView View);
+	UE_NODISCARD CORE_API static FUniqueBuffer MakeView(void* Data, uint64 Size);
 
 	/**
 	 * Make an owned buffer by taking ownership of the input.
@@ -218,7 +218,7 @@ public:
 	 */
 	template <typename DeleteFunctionType,
 		decltype(Invoke(std::declval<DeleteFunctionType>(), std::declval<void*>()))* = nullptr>
-	static inline FUniqueBuffer TakeOwnership(void* Data, uint64 Size, DeleteFunctionType&& DeleteFunction);
+	UE_NODISCARD static inline FUniqueBuffer TakeOwnership(void* Data, uint64 Size, DeleteFunctionType&& DeleteFunction);
 
 	/**
 	 * Make an owned buffer by taking ownership of the input.
@@ -227,7 +227,7 @@ public:
 	 */
 	template <typename DeleteFunctionType,
 		decltype(Invoke(std::declval<DeleteFunctionType>(), std::declval<void*>(), std::declval<uint64>()))* = nullptr>
-	static inline FUniqueBuffer TakeOwnership(void* Data, uint64 Size, DeleteFunctionType&& DeleteFunction);
+	UE_NODISCARD static inline FUniqueBuffer TakeOwnership(void* Data, uint64 Size, DeleteFunctionType&& DeleteFunction);
 
 	/**
 	 * Make a unique buffer from a shared buffer.
@@ -235,7 +235,7 @@ public:
 	 * Steals the buffer owner from the shared buffer if this is the last reference to it, otherwise
 	 * clones the shared buffer to guarantee unique ownership. An non-owned buffer is always cloned.
 	 */
-	CORE_API static FUniqueBuffer MakeUnique(FSharedBuffer Buffer);
+	UE_NODISCARD CORE_API static FUniqueBuffer MakeUnique(FSharedBuffer Buffer);
 
 	/** Construct a null unique buffer. */
 	FUniqueBuffer() = default;
@@ -253,36 +253,36 @@ public:
 	CORE_API void Reset();
 
 	/** Returns a pointer to the start of the buffer. */
-	inline void* GetData() { return Owner ? Owner->GetData() : nullptr; }
-	inline const void* GetData() const { return Owner ? Owner->GetData() : nullptr; }
+	UE_NODISCARD inline void* GetData() { return Owner ? Owner->GetData() : nullptr; }
+	UE_NODISCARD inline const void* GetData() const { return Owner ? Owner->GetData() : nullptr; }
 
 	/** Returns the size of the buffer in bytes. */
-	inline uint64 GetSize() const { return Owner ? Owner->GetSize() : 0; }
+	UE_NODISCARD inline uint64 GetSize() const { return Owner ? Owner->GetSize() : 0; }
 
 	/** Returns a view of the buffer. */
-	inline FMutableMemoryView GetView() { return FMutableMemoryView(GetData(), GetSize()); }
-	inline FMemoryView GetView() const { return FMemoryView(GetData(), GetSize()); }
-	inline operator FMutableMemoryView() { return GetView(); }
-	inline operator FMemoryView() const { return GetView(); }
+	UE_NODISCARD inline FMutableMemoryView GetView() { return FMutableMemoryView(GetData(), GetSize()); }
+	UE_NODISCARD inline FMemoryView GetView() const { return FMemoryView(GetData(), GetSize()); }
+	UE_NODISCARD inline operator FMutableMemoryView() { return GetView(); }
+	UE_NODISCARD inline operator FMemoryView() const { return GetView(); }
 
 	/** Returns true if this points to a buffer owner. */
-	inline explicit operator bool() const { return !IsNull(); }
+	UE_NODISCARD inline explicit operator bool() const { return !IsNull(); }
 
 	/**
 	 * Returns true if this does not point to a buffer owner.
 	 *
 	 * A null buffer is always owned, materialized, and empty.
 	 */
-	inline bool IsNull() const { return Owner.IsNull(); }
+	UE_NODISCARD inline bool IsNull() const { return Owner.IsNull(); }
 
 	/** Returns true if this keeps the referenced buffer alive. */
-	inline bool IsOwned() const { return !Owner || Owner->IsOwned(); }
+	UE_NODISCARD inline bool IsOwned() const { return !Owner || Owner->IsOwned(); }
 
-	/** Clone into a new buffer if the buffer is not owned. */
-	CORE_API void MakeOwned();
+	/** Returns a buffer that is owned, by cloning if not owned. */
+	CORE_API FUniqueBuffer MakeOwned() &&;
 
 	/** Returns true if the referenced buffer has been materialized. */
-	inline bool IsMaterialized() const { return !Owner || Owner->IsMaterialized(); }
+	UE_NODISCARD inline bool IsMaterialized() const { return !Owner || Owner->IsMaterialized(); }
 
 	/**
 	 * Materialize the buffer by making its data and size available.
@@ -313,16 +313,16 @@ class FSharedBuffer
 {
 public:
 	/** Make an owned clone of the input. */
-	CORE_API static FSharedBuffer Clone(FMemoryView View);
-	CORE_API static FSharedBuffer Clone(const void* Data, uint64 Size);
+	UE_NODISCARD CORE_API static FSharedBuffer Clone(FMemoryView View);
+	UE_NODISCARD CORE_API static FSharedBuffer Clone(const void* Data, uint64 Size);
 
 	/** Make a non-owned view of the input. */
-	CORE_API static FSharedBuffer MakeView(FMemoryView View);
-	CORE_API static FSharedBuffer MakeView(const void* Data, uint64 Size);
+	UE_NODISCARD CORE_API static FSharedBuffer MakeView(FMemoryView View);
+	UE_NODISCARD CORE_API static FSharedBuffer MakeView(const void* Data, uint64 Size);
 
 	/** Make a view of the input within its outer buffer. Ownership matches OuterBuffer. */
-	CORE_API static FSharedBuffer MakeView(FMemoryView View, FSharedBuffer OuterBuffer);
-	CORE_API static FSharedBuffer MakeView(const void* Data, uint64 Size, FSharedBuffer OuterBuffer);
+	UE_NODISCARD CORE_API static FSharedBuffer MakeView(FMemoryView View, FSharedBuffer OuterBuffer);
+	UE_NODISCARD CORE_API static FSharedBuffer MakeView(const void* Data, uint64 Size, FSharedBuffer OuterBuffer);
 
 	/**
 	 * Make an owned buffer by taking ownership of the input.
@@ -331,7 +331,7 @@ public:
 	 */
 	template <typename DeleteFunctionType,
 		decltype(Invoke(std::declval<DeleteFunctionType>(), std::declval<void*>()))* = nullptr>
-	static inline FSharedBuffer TakeOwnership(const void* Data, uint64 Size, DeleteFunctionType&& DeleteFunction);
+	UE_NODISCARD static inline FSharedBuffer TakeOwnership(const void* Data, uint64 Size, DeleteFunctionType&& DeleteFunction);
 
 	/**
 	 * Make an owned buffer by taking ownership of the input.
@@ -340,7 +340,7 @@ public:
 	 */
 	template <typename DeleteFunctionType,
 		decltype(Invoke(std::declval<DeleteFunctionType>(), std::declval<void*>(), std::declval<uint64>()))* = nullptr>
-	static inline FSharedBuffer TakeOwnership(const void* Data, uint64 Size, DeleteFunctionType&& DeleteFunction);
+	UE_NODISCARD static inline FSharedBuffer TakeOwnership(const void* Data, uint64 Size, DeleteFunctionType&& DeleteFunction);
 
 	/** Construct a null shared buffer. */
 	FSharedBuffer() = default;
@@ -365,33 +365,34 @@ public:
 	CORE_API void Reset();
 
 	/** Returns a pointer to the start of the buffer. */
-	inline const void* GetData() const { return Owner ? Owner->GetData() : nullptr; }
+	UE_NODISCARD inline const void* GetData() const { return Owner ? Owner->GetData() : nullptr; }
 
 	/** Returns the size of the buffer in bytes. */
-	inline uint64 GetSize() const { return Owner ? Owner->GetSize() : 0; }
+	UE_NODISCARD inline uint64 GetSize() const { return Owner ? Owner->GetSize() : 0; }
 
 	/** Returns a view of the buffer. */
-	inline FMemoryView GetView() const { return FMemoryView(GetData(), GetSize()); }
-	inline operator FMemoryView() const { return GetView(); }
+	UE_NODISCARD inline FMemoryView GetView() const { return FMemoryView(GetData(), GetSize()); }
+	UE_NODISCARD inline operator FMemoryView() const { return GetView(); }
 
 	/** Returns true if this points to a buffer owner. */
-	inline explicit operator bool() const { return !IsNull(); }
+	UE_NODISCARD inline explicit operator bool() const { return !IsNull(); }
 
 	/**
 	 * Returns true if this does not point to a buffer owner.
 	 *
 	 * A null buffer is always owned, materialized, and empty.
 	 */
-	inline bool IsNull() const { return Owner.IsNull(); }
+	UE_NODISCARD inline bool IsNull() const { return Owner.IsNull(); }
 
 	/** Returns true if this keeps the referenced buffer alive. */
-	inline bool IsOwned() const { return !Owner || Owner->IsOwned(); }
+	UE_NODISCARD inline bool IsOwned() const { return !Owner || Owner->IsOwned(); }
 
-	/** Clone into a new buffer if the buffer is not owned. */
-	CORE_API void MakeOwned();
+	/** Returns a buffer that is owned, by cloning if not owned. */
+	CORE_API FSharedBuffer MakeOwned() const &;
+	CORE_API FSharedBuffer MakeOwned() &&;
 
 	/** Returns true if the referenced buffer has been materialized. */
-	inline bool IsMaterialized() const { return !Owner || Owner->IsMaterialized(); }
+	UE_NODISCARD inline bool IsMaterialized() const { return !Owner || Owner->IsMaterialized(); }
 
 	/**
 	 * Materialize the buffer by making its data and size available.
@@ -436,7 +437,7 @@ public:
 	CORE_API void Reset();
 
 	/** Convert this to a shared buffer if it has any remaining shared references. */
-	CORE_API FSharedBuffer Pin() const;
+	UE_NODISCARD CORE_API FSharedBuffer Pin() const;
 
 private:
 	using OwnerPtrType = BufferOwnerPrivate::TBufferOwnerPtr<BufferOwnerPrivate::FWeakOps>;
@@ -837,7 +838,7 @@ private:
 
 /** Construct a shared buffer by taking ownership of an array. */
 template <typename T, typename Allocator>
-FSharedBuffer MakeSharedBufferFromArray(TArray<T, Allocator>&& Array)
+UE_NODISCARD FSharedBuffer MakeSharedBufferFromArray(TArray<T, Allocator>&& Array)
 {
 	return FSharedBuffer(new BufferOwnerPrivate::TBufferOwnerTArray<T, Allocator>(MoveTemp(Array)));
 }
diff --git a/Engine/Source/Runtime/CoreUObject/Private/Virtualization/VirtualizedBulkData.cpp b/Engine/Source/Runtime/CoreUObject/Private/Virtualization/VirtualizedBulkData.cpp
--- a/Engine/Source/Runtime/CoreUObject/Private/Virtualization/VirtualizedBulkData.cpp
+++ b/Engine/Source/Runtime/CoreUObject/Private/Virtualization/VirtualizedBulkData.cpp
@@ -411,8 +411,7 @@ void FVirtualizedUntypedBulkData::UpdatePayload(const FSharedBuffer& InPayload)
 	UnloadData();
 
 	// Make sure that we own the memory in the shared buffer
-	Payload = InPayload;
-	Payload.MakeOwned();
+	Payload = InPayload.MakeOwned();
 	PayloadLength = (int64)Payload.GetSize();
 
 	bIsVirtualized = false;
