U5 - use LIFO waking behaviour to impove readability of profiling tools and maybe gain some small benefits from warm caches.

diff --git a/Engine/Source/Runtime/Core/Public/Async/Fundamental/Scheduler.h b/Engine/Source/Runtime/Core/Public/Async/Fundamental/Scheduler.h
--- a/Engine/Source/Runtime/Core/Public/Async/Fundamental/Scheduler.h
+++ b/Engine/Source/Runtime/Core/Public/Async/Fundamental/Scheduler.h
@@ -97,12 +97,77 @@ namespace LowLevelTasks
 			Sleeping,
 		};
 
-		struct FSleepEvent
+		//the struct is naturally 8 bytes aligned, the extra alignment just 
+		//re-enforces this assumption and will error if it changes in the future
+		struct alignas(8) FSleepEvent
 		{
 			FEventRef SleepEvent;
 			std::atomic<ESleepState> SleepState { ESleepState::Running };
+			FSleepEvent* Next = nullptr;
+		};
+
+		//implementation of a treiber stack
+		//(https://en.wikipedia.org/wiki/Treiber_stack)
+		class FEventStack
+		{
+		private:
+			struct FTopNode
+			{
+				uintptr_t Address  : 45; //all CPUs we care about use less than 48 bits for their addressing, and the lower 3 bits are unused due to alignment
+				uintptr_t Revision : 19; //Tagging is used to avoid ABA (the wrap around is several minutes for this use-case (https://en.wikipedia.org/wiki/ABA_problem#Tagged_state_reference))
+			};
+			std::atomic<FTopNode> Top { FTopNode{0, 0} };
+
+		public:
+			FSleepEvent* Pop()
+			{
+				FTopNode LocalTop = Top.load(std::memory_order_relaxed);
+				while (true) 
+				{			
+					if (LocalTop.Address == 0)
+					{
+						return nullptr;
+					}
+#if DO_CHECK
+					int64 LastRevision = int64(LocalTop.Revision); 
+#endif
+					FSleepEvent* Item = reinterpret_cast<FSleepEvent*>(LocalTop.Address << 3);
+					if (Top.compare_exchange_weak(LocalTop, FTopNode { reinterpret_cast<uintptr_t>(Item->Next) >> 3, uintptr_t(LocalTop.Revision + 1) }, std::memory_order_acquire, std::memory_order_relaxed))
+					{
+						Item->Next = nullptr;
+						return Item;
+					}
+#if DO_CHECK
+					int64 NewRevision = int64(LocalTop.Revision) < LastRevision ? ((1ll << 19) + int64(LocalTop.Revision)) : int64(LocalTop.Revision);
+					ensureMsgf((NewRevision - LastRevision) < (1ll << 18), TEXT("Dangerously close to the wraparound: %d, %d"), LastRevision, NewRevision);
+#endif
+				}
+			}
+
+			void Push(FSleepEvent* Item)
+			{
+				checkSlow(reinterpret_cast<uintptr_t>(Item) < (1ull << 48));
+				checkSlow((reinterpret_cast<uintptr_t>(Item) & 0x7) == 0);
+				checkSlow(Item->Next == nullptr);
+				
+				FTopNode LocalTop = Top.load(std::memory_order_relaxed);
+				while (true) 
+				{
+#if DO_CHECK
+					int64 LastRevision = int64(LocalTop.Revision); 
+#endif
+					Item->Next = reinterpret_cast<FSleepEvent*>(LocalTop.Address << 3);
+					if (Top.compare_exchange_weak(LocalTop, FTopNode { reinterpret_cast<uintptr_t>(Item) >> 3, uintptr_t(LocalTop.Revision + 1) }, std::memory_order_release, std::memory_order_acquire))  
+					{
+						return;
+					}
+#if DO_CHECK
+					int64 NewRevision = int64(LocalTop.Revision) < LastRevision ? ((1ll << 19) + int64(LocalTop.Revision)) : int64(LocalTop.Revision);
+					ensureMsgf((NewRevision - LastRevision) < (1ll << 18), TEXT("Dangerously close to the wraparound: %d, %d"), LastRevision, NewRevision);
+#endif
+				}
+			}
 		};
-		using FEventQueueType = FAAArrayQueue<FSleepEvent>;
 
 	public:
 		FScheduler() = default;
@@ -120,7 +185,7 @@ namespace LowLevelTasks
 		bool TryExecuteTaskFrom(FLocalQueueType* Queue, FQueueRegistry::FOutOfWork& OutOfWork, bool bPermitBackgroundWork);
 
 	private:
-		FEventQueueType 			SleepEventQueue[2];
+		FEventStack 				SleepEventStack[2];
 		FQueueRegistry 				QueueRegistry;
 		FCriticalSection 			WorkerThreadsCS;
 		TArray<TUniquePtr<FThread>> WorkerThreads;
@@ -251,7 +316,7 @@ namespace LowLevelTasks
 		{
 			OutOfWork.Stop();
 			Drowsing = true;
-			SleepEventQueue[bBackgroundWorker].enqueue(WorkerEvent); // State one: (Running -> Drowsing)
+			SleepEventStack[bBackgroundWorker].Push(WorkerEvent); // State one: (Running -> Drowsing)
 		}
 		else
 		{
@@ -261,18 +326,17 @@ namespace LowLevelTasks
 
 	inline bool FScheduler::WakeUpWorker(bool bBackgroundWorker)
 	{
-		FSleepEvent* WorkerEvent = SleepEventQueue[bBackgroundWorker].dequeue();
+		FSleepEvent* WorkerEvent = SleepEventStack[bBackgroundWorker].Pop();
 		if (WorkerEvent)
 		{
-			// Solving State : (Running -> Drowsing) -> Running  OR ((Running -> Drowsing) -> Drowsing) -> Running OR (((Running -> Drowsing) -> Sleeping) -> Running)
 			ESleepState SleepState = WorkerEvent->SleepState.exchange(ESleepState::Running, std::memory_order_relaxed);
 			if (SleepState == ESleepState::Sleeping)
 			{
-				WorkerEvent->SleepEvent->Trigger(); // Solving State two: (((Running -> Drowsing) -> Sleeping) -> Running)
-				return true;
+				WorkerEvent->SleepEvent->Trigger(); 
+				return true; // Solving State two: (((Running -> Drowsing) -> Sleeping) -> Running)
 			}
 			checkf(SleepState == ESleepState::Drowsing, TEXT("Worker was not drowsing: %d"), WorkerEvent->SleepState.load(std::memory_order_relaxed));
-			return true;
+			return true; // Solving State one: (Running -> Drowsing) -> Running  OR ((Running -> Drowsing) -> Drowsing) -> Running
 		}
 		return false;
 	}
