U5 - scale FG workers with corecount to keep overall system more responsive on high core count machines

diff --git a/Engine/Source/Runtime/Core/Private/Async/TaskGraph.cpp b/Engine/Source/Runtime/Core/Private/Async/TaskGraph.cpp
--- a/Engine/Source/Runtime/Core/Private/Async/TaskGraph.cpp
+++ b/Engine/Source/Runtime/Core/Private/Async/TaskGraph.cpp
@@ -2204,6 +2204,16 @@ void FTaskGraphInterface::Startup(int32 NumThreads)
 
 	if (GUseNewTaskBackend)
 	{
+		//We want to reduce the number of overall threads that UE uses so that there is are some 
+		//free cores available for other things like the Browser or other Applications. 
+		//Therefore we increase the number of Foreground workers, which are mostly unused. 
+		//But when HighPrio work comes in the Foreground workers will be available and get the job done.
+		bool bIsCookCommandlet = FParse::Param(FCommandLine::Get(), TEXT("cookcommandlet")) || FParse::Param(FCommandLine::Get(), TEXT("run=cook"));
+		if (!bIsCookCommandlet)
+		{
+			GNumForegroundWorkers = FMath::Max(FMath::DivideAndRoundUp(NumThreads, 21), 2);
+		}
+
 		FParse::Value(FCommandLine::Get(), TEXT("-foregroundworkers="), GNumForegroundWorkers);
 		TaskGraphImplementationSingleton = new FTaskGraphCompatibilityImplementation(NumThreads);
 	}
