MemoryView: Added FMutableMemoryView::CopyFrom to copy from another view

#rnx

diff --git a/Engine/Source/Runtime/Core/Public/Memory/MemoryView.h b/Engine/Source/Runtime/Core/Public/Memory/MemoryView.h
--- a/Engine/Source/Runtime/Core/Public/Memory/MemoryView.h
+++ b/Engine/Source/Runtime/Core/Public/Memory/MemoryView.h
@@ -3,6 +3,7 @@
 #pragma once
 
 #include "CoreTypes.h"
+#include "HAL/PlatformString.h"
 #include "Math/NumericLimits.h"
 #include "Math/UnrealMathUtility.h"
 #include "Memory/MemoryFwd.h"
@@ -209,6 +210,15 @@ public:
 		return *this;
 	}
 
+	/** Copies bytes from the input view into this view, and returns the remainder of this view. */
+	inline TMemoryView CopyFrom(FMemoryView InView) const
+	{
+		checkf(InView.Size <= Size, TEXT("Failed to copy from a view of %" UINT64_FMT " bytes ")
+			TEXT("to a view of %" UINT64_FMT " bytes."), InView.Size, Size);
+		FMemory::Memcpy(Data, InView.Data, InView.Size);
+		return RightChop(InView.Size);
+	}
+
 private:
 	/** Returns the data pointer advanced by an offset in bytes. */
 	UE_NODISCARD inline DataType* GetDataAtOffsetNoCheck(uint64 InOffset) const
