Updated users of memory views to use CopyFrom instead of Memcpy where appropriate

#rnx

diff --git a/Engine/Source/Runtime/Core/Private/Memory/SharedBuffer.cpp b/Engine/Source/Runtime/Core/Private/Memory/SharedBuffer.cpp
--- a/Engine/Source/Runtime/Core/Private/Memory/SharedBuffer.cpp
+++ b/Engine/Source/Runtime/Core/Private/Memory/SharedBuffer.cpp
@@ -96,14 +96,14 @@ FUniqueBuffer FUniqueBuffer::Alloc(uint64 InSize)
 
 FUniqueBuffer FUniqueBuffer::Clone(FMemoryView View)
 {
-	return Clone(View.GetData(), View.GetSize());
+	FUniqueBuffer Buffer = Alloc(View.GetSize());
+	Buffer.GetView().CopyFrom(View);
+	return Buffer;
 }
 
 FUniqueBuffer FUniqueBuffer::Clone(const void* Data, uint64 Size)
 {
-	FUniqueBuffer Buffer = Alloc(Size);
-	FMemory::Memcpy(Buffer.GetData(), Data, Size);
-	return Buffer;
+	return Clone(MakeMemoryView(Data, Size));
 }
 
 FUniqueBuffer FUniqueBuffer::MakeView(FMutableMemoryView View)
diff --git a/Engine/Source/Runtime/Core/Private/Serialization/CompactBinary.cpp b/Engine/Source/Runtime/Core/Private/Serialization/CompactBinary.cpp
--- a/Engine/Source/Runtime/Core/Private/Serialization/CompactBinary.cpp
+++ b/Engine/Source/Runtime/Core/Private/Serialization/CompactBinary.cpp
@@ -572,11 +572,10 @@ void FCbFieldView::CopyTo(FMutableMemoryView Buffer) const
 {
 	const FMemoryView Source = GetViewNoType();
 	checkf(Buffer.GetSize() == sizeof(ECbFieldType) + Source.GetSize(),
-		TEXT("A buffer of %" UINT64_FMT " bytes was provided when %" UINT64_FMT " bytes are required"),
+		TEXT("Buffer is %" UINT64_FMT " bytes but %" UINT64_FMT " is required."),
 		Buffer.GetSize(), sizeof(ECbFieldType) + Source.GetSize());
-	*static_cast<ECbFieldType*>(Buffer.GetData()) = FCbFieldType::GetSerializedType(Type);
-	Buffer.RightChopInline(sizeof(ECbFieldType));
-	FMemory::Memcpy(Buffer.GetData(), Source.GetData(), Source.GetSize());
+	const ECbFieldType SerializedType = FCbFieldType::GetSerializedType(Type);
+	Buffer.CopyFrom(MakeMemoryView(&SerializedType, 1)).CopyFrom(Source);
 }
 
 void FCbFieldView::CopyTo(FArchive& Ar) const
@@ -679,15 +678,14 @@ bool FCbArrayView::Equals(const FCbArrayView& Other) const
 		GetPayloadView().EqualBytes(Other.GetPayloadView());
 }
 
-void FCbArrayView::CopyTo(FMutableMemoryView Buffer) const
+void FCbArrayView::CopyTo(const FMutableMemoryView Buffer) const
 {
 	const FMemoryView Source = GetPayloadView();
 	checkf(Buffer.GetSize() == sizeof(ECbFieldType) + Source.GetSize(),
 		TEXT("Buffer is %" UINT64_FMT " bytes but %" UINT64_FMT " is required."),
 		Buffer.GetSize(), sizeof(ECbFieldType) + Source.GetSize());
-	*static_cast<ECbFieldType*>(Buffer.GetData()) = FCbFieldType::GetType(GetType());
-	Buffer.RightChopInline(sizeof(ECbFieldType));
-	FMemory::Memcpy(Buffer.GetData(), Source.GetData(), Source.GetSize());
+	const ECbFieldType SerializedType = FCbFieldType::GetType(GetType());
+	Buffer.CopyFrom(MakeMemoryView(&SerializedType, 1)).CopyFrom(Source);
 }
 
 void FCbArrayView::CopyTo(FArchive& Ar) const
@@ -771,15 +769,14 @@ bool FCbObjectView::Equals(const FCbObjectView& Other) const
 		GetPayloadView().EqualBytes(Other.GetPayloadView());
 }
 
-void FCbObjectView::CopyTo(FMutableMemoryView Buffer) const
+void FCbObjectView::CopyTo(const FMutableMemoryView Buffer) const
 {
 	const FMemoryView Source = GetPayloadView();
 	checkf(Buffer.GetSize() == sizeof(ECbFieldType) + Source.GetSize(),
 		TEXT("Buffer is %" UINT64_FMT " bytes but %" UINT64_FMT " is required."),
 		Buffer.GetSize(), sizeof(ECbFieldType) + Source.GetSize());
-	*static_cast<ECbFieldType*>(Buffer.GetData()) = FCbFieldType::GetType(GetType());
-	Buffer.RightChopInline(sizeof(ECbFieldType));
-	FMemory::Memcpy(Buffer.GetData(), Source.GetData(), Source.GetSize());
+	const ECbFieldType SerializedType = FCbFieldType::GetType(GetType());
+	Buffer.CopyFrom(MakeMemoryView(&SerializedType, 1)).CopyFrom(Source);
 }
 
 void FCbObjectView::CopyTo(FArchive& Ar) const
@@ -846,7 +843,7 @@ void TCbFieldIterator<FieldType>::CopyRangeTo(FMutableMemoryView InBuffer) const
 		checkf(InBuffer.GetSize() == Source.GetSize(),
 			TEXT("Buffer is %" UINT64_FMT " bytes but %" UINT64_FMT " is required."),
 			InBuffer.GetSize(), Source.GetSize());
-		FMemory::Memcpy(InBuffer.GetData(), Source.GetData(), Source.GetSize());
+		InBuffer.CopyFrom(Source);
 	}
 	else
 	{
diff --git a/Engine/Source/Runtime/Core/Private/Serialization/CompactBinarySerialization.cpp b/Engine/Source/Runtime/Core/Private/Serialization/CompactBinarySerialization.cpp
--- a/Engine/Source/Runtime/Core/Private/Serialization/CompactBinarySerialization.cpp
+++ b/Engine/Source/Runtime/Core/Private/Serialization/CompactBinarySerialization.cpp
@@ -181,9 +181,7 @@ FCbField LoadCompactBinary(FArchive& Ar, FCbBufferAllocator Allocator)
 	checkf(Buffer.GetSize() == FieldSize,
 		TEXT("Allocator returned a buffer of size %" UINT64_FMT " bytes when %" UINT64_FMT " bytes were requested."),
 		Buffer.GetSize(), FieldSize);
-	FMutableMemoryView View = Buffer.GetView();
-	FMemory::Memcpy(View.GetData(), HeaderBytes.GetData(), HeaderBytes.Num());
-	View += HeaderBytes.Num();
+	FMutableMemoryView View = Buffer.GetView().CopyFrom(MakeMemoryView(HeaderBytes));
 	if (!View.IsEmpty())
 	{
 		Ar.Serialize(View.GetData(), static_cast<int64>(View.GetSize()));
diff --git a/Engine/Source/Runtime/Core/Private/Serialization/CompactBinaryWriter.cpp b/Engine/Source/Runtime/Core/Private/Serialization/CompactBinaryWriter.cpp
--- a/Engine/Source/Runtime/Core/Private/Serialization/CompactBinaryWriter.cpp
+++ b/Engine/Source/Runtime/Core/Private/Serialization/CompactBinaryWriter.cpp
@@ -59,7 +59,7 @@ static inline ECbFieldType AppendCompactBinary(const FCbFieldView& Value, TArray
 	const FCopy& ValueCopy = static_cast<const FCopy&>(Value);
 	const FMemoryView SourceView = ValueCopy.GetPayloadView();
 	const int64 TargetOffset = OutData.AddUninitialized(SourceView.GetSize());
-	FMemory::Memcpy(OutData.GetData() + TargetOffset, SourceView.GetData(), SourceView.GetSize());
+	MakeMemoryView(OutData).RightChop(TargetOffset).CopyFrom(SourceView);
 	return FCbFieldType::GetType(ValueCopy.GetType());
 }
 
@@ -100,7 +100,7 @@ FCbFieldViewIterator FCbWriter::Save(const FMutableMemoryView Buffer) const
 	checkf(Data.Num() > 0, TEXT("It is invalid to save when nothing has been written."));
 	checkf(Buffer.GetSize() == Data.Num(),
 		TEXT("Buffer is %" UINT64_FMT " bytes but %" INT64_FMT " is required."), Buffer.GetSize(), Data.Num());
-	FMemory::Memcpy(Buffer.GetData(), Data.GetData(), Data.Num());
+	Buffer.CopyFrom(MakeMemoryView(Data));
 	return FCbFieldViewIterator::MakeRange(Buffer);
 }
 
