Disable TasksTest for single-threaded execution as Tasks API doesn't support it until it replaces TaskGraph on the game thread. For this to happen we still lack support for the local queue.

diff --git a/Engine/Source/Runtime/Core/Private/Tests/Tasks/TasksTest.cpp b/Engine/Source/Runtime/Core/Private/Tests/Tasks/TasksTest.cpp
--- a/Engine/Source/Runtime/Core/Private/Tests/Tasks/TasksTest.cpp
+++ b/Engine/Source/Runtime/Core/Private/Tests/Tasks/TasksTest.cpp
@@ -88,6 +88,12 @@ namespace UE { namespace TasksTests
 
 	bool FTasksBasicTest::RunTest(const FString& Parameters)
 	{
+		if (!FPlatformProcess::SupportsMultithreading())
+		{
+			// the new API doesn't support single-threaded execution (`-nothreading`) until it's feature-compatible with the new API and completely replaces it
+			return true;
+		}
+
 		{	// basic example, fire and forget a high-pri task
 			Tasks::Launch(
 				UE_SOURCE_LOCATION, // debug name
@@ -221,6 +227,12 @@ namespace UE { namespace TasksTests
 
 	bool FTasksPipeTest::RunTest(const FString& Parameters)
 	{
+		if (!FPlatformProcess::SupportsMultithreading())
+		{
+			// the new API doesn't support single-threaded execution (`-nothreading`) until it's feature-compatible with the new API and completely replaces it
+			return true;
+		}
+
 		{	// a basic usage example
 			Tasks::FPipe Pipe{ UE_SOURCE_LOCATION }; // a debug name, user-provided or 
 				// `UE_SOURCE_LOCATION` - source file name and line number
