Updated FVSPerfExternalProfiler to be aware of VS 2017 and 2019, depending on the compiler version used.

#rnx

diff --git a/Engine/Source/Runtime/Core/Private/ProfilingDebugging/Windows/VSPerfExternalProfiler.cpp b/Engine/Source/Runtime/Core/Private/ProfilingDebugging/Windows/VSPerfExternalProfiler.cpp
--- a/Engine/Source/Runtime/Core/Private/ProfilingDebugging/Windows/VSPerfExternalProfiler.cpp
+++ b/Engine/Source/Runtime/Core/Private/ProfilingDebugging/Windows/VSPerfExternalProfiler.cpp
@@ -104,7 +104,11 @@ public:
 		// Try to load the VSPerf DLL
 		// NOTE: VSPerfXXX.dll is installed into /Windows/System32 when Visual Studio is installed.  The XXX is the version number of
 		// Visual Studio.  For example, for Visual Studio 2013, the file name is VSPerf120.dll.
-#if _MSC_VER >= 1900
+#if _MSC_VER >= 1920
+		DLLHandle = FPlatformProcess::GetDllHandle(TEXT("VSPerf160.dll"));	// Visual Studio 2019
+#elif _MSC_VER >= 1910
+		DLLHandle = FPlatformProcess::GetDllHandle(TEXT("VSPerf150.dll"));	// Visual Studio 2017
+#elif _MSC_VER >= 1900
 		DLLHandle = FPlatformProcess::GetDllHandle(TEXT("VSPerf140.dll"));	// Visual Studio 2015
 #elif _MSC_VER >= 1800
 		DLLHandle = FPlatformProcess::GetDllHandle(TEXT("VSPerf120.dll"));	// Visual Studio 2013
