made possible to override UE_TASK_TRACE_ENABLED in project config

diff --git a/Engine/Source/Runtime/Core/Public/Async/TaskTrace.h b/Engine/Source/Runtime/Core/Public/Async/TaskTrace.h
--- a/Engine/Source/Runtime/Core/Public/Async/TaskTrace.h
+++ b/Engine/Source/Runtime/Core/Public/Async/TaskTrace.h
@@ -5,11 +5,13 @@
 #include "Trace/Config.h"
 #include "Containers/Array.h"
 
+#if !defined(UE_TASK_TRACE_ENABLED)
 #if UE_TRACE_ENABLED && !IS_PROGRAM && !UE_BUILD_SHIPPING
 #define UE_TASK_TRACE_ENABLED 1
 #else
 #define UE_TASK_TRACE_ENABLED 0
 #endif
+#endif
 
 namespace ENamedThreads
 {
