fixed "Compile ClangEditor Win64"

diff --git a/Engine/Source/Runtime/Core/Public/Tasks/TaskPrivate.h b/Engine/Source/Runtime/Core/Public/Tasks/TaskPrivate.h
--- a/Engine/Source/Runtime/Core/Public/Tasks/TaskPrivate.h
+++ b/Engine/Source/Runtime/Core/Public/Tasks/TaskPrivate.h
@@ -35,7 +35,7 @@ namespace UE { namespace Tasks
 	namespace Private
 	{
 		// Task class implementation with piping support.
-		class CORE_API FTaskBase
+		class FTaskBase
 		{
 			UE_NONCOPYABLE(FTaskBase);
 
@@ -295,7 +295,7 @@ namespace UE { namespace Tasks
 		};
 
 		template<>
-		class CORE_API TTaskWithResultBase<void> : public FTaskBase
+		class TTaskWithResultBase<void> : public FTaskBase
 		{
 		protected:
 			TTaskWithResultBase() = default; // instances can be created only on the heap by a ref-counting handler
@@ -310,7 +310,7 @@ namespace UE { namespace Tasks
 		// intrusive atomic reference-counting base class.
 		// doesn't require derived classes to be polymorphic
 		template<typename ReferenceType>
-		class CORE_API TRefCountedBase
+		class TRefCountedBase
 		{
 		public:
 			explicit TRefCountedBase(uint32 InitRefCount = 0)
@@ -341,7 +341,7 @@ namespace UE { namespace Tasks
 		// A final task class that must be used only with `TRefCountPtr`, thus only friends are allowed to construct instances.
 		// @see UE::Tasks::Task<T>
 		template<typename ResultType>
-		class CORE_API TTaskWithResult final : public TTaskWithResultBase<ResultType>, public TRefCountedBase<TTaskWithResult<ResultType>>
+		class TTaskWithResult final : public TTaskWithResultBase<ResultType>, public TRefCountedBase<TTaskWithResult<ResultType>>
 		{
 			UE_NONCOPYABLE(TTaskWithResult);
 
