U5 - make static analysis happy

diff --git a/Engine/Source/Runtime/Core/Public/Async/Fundamental/Scheduler.h b/Engine/Source/Runtime/Core/Public/Async/Fundamental/Scheduler.h
--- a/Engine/Source/Runtime/Core/Public/Async/Fundamental/Scheduler.h
+++ b/Engine/Source/Runtime/Core/Public/Async/Fundamental/Scheduler.h
@@ -146,6 +146,7 @@ namespace LowLevelTasks
 
 			void Push(FSleepEvent* Item)
 			{
+				checkSlow(Item != nullptr);
 				checkSlow(reinterpret_cast<uintptr_t>(Item) < (1ull << 48));
 				checkSlow((reinterpret_cast<uintptr_t>(Item) & 0x7) == 0);
 				checkSlow(Item->Next == nullptr);
