Use FConfigCacheIni::ForPlatform in FShaderPlatformCachedIniValue implementation

diff --git a/Engine/Source/Runtime/Core/Public/Misc/ConfigCacheIni.h b/Engine/Source/Runtime/Core/Public/Misc/ConfigCacheIni.h
--- a/Engine/Source/Runtime/Core/Public/Misc/ConfigCacheIni.h
+++ b/Engine/Source/Runtime/Core/Public/Misc/ConfigCacheIni.h
@@ -886,6 +886,32 @@ public:
 		const FString&	Filename
 	);
 
+	/* Generic versions for use with templates */
+	bool GetValue(const TCHAR* Section, const TCHAR* Key, FString& Value, const FString& Filename)
+	{
+		return GetString(Section, Key, Value, Filename);
+	}
+	bool GetValue(const TCHAR* Section, const TCHAR* Key, FText& Value, const FString& Filename)
+	{
+		return GetText(Section, Key, Value, Filename);
+	}
+	bool GetValue(const TCHAR* Section, const TCHAR* Key, int32& Value, const FString& Filename)
+	{
+		return GetInt(Section, Key, Value, Filename);
+	}
+	bool GetValue(const TCHAR* Section, const TCHAR* Key, float& Value, const FString& Filename)
+	{
+		return GetFloat(Section, Key, Value, Filename);
+	}
+	bool GetValue(const TCHAR* Section, const TCHAR* Key, bool& Value, const FString& Filename)
+	{
+		return GetBool(Section, Key, Value, Filename);
+	}
+	int32 GetValue(const TCHAR* Section, const TCHAR* Key, TArray<FString>& Value, const FString& Filename)
+	{
+		return GetArray(Section, Key, Value, Filename);
+	}
+
 	void SetInt
 	(
 		const TCHAR*		Section,
diff --git a/Engine/Source/Runtime/RenderCore/Private/RenderUtils.cpp b/Engine/Source/Runtime/RenderCore/Private/RenderUtils.cpp
--- a/Engine/Source/Runtime/RenderCore/Private/RenderUtils.cpp
+++ b/Engine/Source/Runtime/RenderCore/Private/RenderUtils.cpp
@@ -1178,9 +1178,8 @@ Type FShaderPlatformCachedIniValue<Type>::Get(EShaderPlatform ShaderPlatform)
 	if (!bTestedIni)
 	{
 		bTestedIni = true;
-		FConfigFile PlatformEngineIni;
-		FConfigCacheIni::LoadLocalIniFile(PlatformEngineIni, TEXT("Engine"), true, *ShaderPlatformToPlatformName(ShaderPlatform).ToString());
-		if (PlatformEngineIni.GetValue(Section, Key, Value))
+		FConfigCacheIni* PlatformIni = FConfigCacheIni::ForPlatform(ShaderPlatformToPlatformName(ShaderPlatform));
+		if (PlatformIni != nullptr && PlatformIni->GetValue(Section, Key, Value, TEXT("Engine")))
 		{
 			CachedValues.Add(ShaderPlatform, Value);
 			return Value;
