ClosableMpscQueue: a multi producer single consumer lock-free queue can can be consumed only once

ported from Dev-TaskGraph stream

diff --git a/Engine/Source/Runtime/Core/Private/Tests/Containers/ConcurrentQueuesTest.cpp b/Engine/Source/Runtime/Core/Private/Tests/Containers/ConcurrentQueuesTest.cpp
--- a/Engine/Source/Runtime/Core/Private/Tests/Containers/ConcurrentQueuesTest.cpp
+++ b/Engine/Source/Runtime/Core/Private/Tests/Containers/ConcurrentQueuesTest.cpp
@@ -3,10 +3,12 @@
 #include "CoreTypes.h"
 #include "Misc/AutomationTest.h"
 #include "Async/TaskGraphInterfaces.h"
+#include "Containers/Array.h"
 #include "Containers/CircularQueue.h"
 #include "Containers/Queue.h"
 #include "Containers/SpscQueue.h"
 #include "Containers/MpscQueue.h"
+#include "Containers/ClosableMpscQueue.h"
 #include "Tests/Benchmark.h"
 
 #include <atomic>
@@ -426,4 +428,316 @@ namespace ConcurrentQueuesTests
 	}
 }
 
+namespace ClosableMpscQueueTests
+{
+	// straightforward implementation over MPSC TQueue for comparison with TClosableMpscQueue
+	template<typename T>
+	class TClosableMpscQueueOnTQueue
+	{
+	public:
+		bool Enqueue(T Value)
+		{
+			bool bRes = false;
+
+			// `acquire` to make it "happen-before" reading `bIsComplete` and `release` to "synchronise-with" reading from other threads
+			verify(SubscribingThreadsNum.fetch_add(1, std::memory_order_acq_rel) >= 0);
+			// `acquire` to "synchronise-with" the consumer
+			if (!bIsComplete.load(std::memory_order_acquire))
+			{
+				ensureAlways(Qeueue.Enqueue(MoveTemp(Value)));
+				bRes = true;
+			}
+			// `release` to "synchonise-with" other threads
+			verify(SubscribingThreadsNum.fetch_sub(1, std::memory_order_release) >= 1);
+
+			return bRes;
+		}
+
+		template<typename F>
+		void Close(F&& Consumer)
+		{
+			// `seq_cst` to make it "happen-before" reading `SubscribingThreadsNum`, and to "synchronise-with" reading from other threads
+			bIsComplete.store(true, std::memory_order_seq_cst);
+
+			// `acquire` to "synchronise-with" other threads and to make it "happen-before" dequeueing
+			while (SubscribingThreadsNum.load(std::memory_order_acquire) != 0)
+			{
+				FPlatformProcess::Yield();
+			}
+
+			T Value;
+			while (Qeueue.Dequeue(Value))
+			{
+				Consumer(MoveTemp(Value));
+			}
+		}
+
+	private:
+		TQueue<T, EQueueMode::Mpsc> Qeueue;
+		std::atomic<bool> bIsComplete{ false };
+		std::atomic<uint32> SubscribingThreadsNum{ 0 };
+	};
+
+	// straightforward implementation over TMpscQueue for comparison with TClosableMpscQueue
+	template<typename T>
+	class TClosableMpscQueueOnTMpscQueue
+	{
+	public:
+		bool Enqueue(T Value)
+		{
+			bool bRes = false;
+
+			// `acquire` to make it "happen-before" reading `bIsComplete` and `release` to "synchronise-with" reading from other threads
+			verify(SubscribingThreadsNum.fetch_add(1, std::memory_order_acq_rel) >= 0);
+			// `acquire` to "synchronise-with" the consumer
+			if (!bIsComplete.load(std::memory_order_acquire))
+			{
+				Qeueue.Enqueue(MoveTemp(Value));
+				bRes = true;
+			}
+			// `release` to "synchonise-with" other threads
+			verify(SubscribingThreadsNum.fetch_sub(1, std::memory_order_release) >= 1);
+
+			return bRes;
+		}
+
+		template<typename F>
+		void Close(F&& Consumer)
+		{
+			// `seq_cst` to make it "happen-before" reading `SubscribingThreadsNum`, and to "synchronise-with" reading from other threads
+			bIsComplete.store(true, std::memory_order_seq_cst);
+
+			// `acquire` to "synchronise-with" other threads and to make it "happen-before" dequeueing
+			while (SubscribingThreadsNum.load(std::memory_order_acquire) != 0)
+			{
+				FPlatformProcess::Yield();
+			}
+
+			while (TOptional<T> Value = Qeueue.Dequeue())
+			{
+				Consumer(MoveTemp(Value.GetValue()));
+			}
+		}
+
+	private:
+		TMpscQueue<T> Qeueue;
+		std::atomic<bool> bIsComplete{ false };
+		std::atomic<uint32> SubscribingThreadsNum{ 0 };
+	};
+
+	// wrapper over `TClosableLockFreePointerListUnorderedSingleConsumer`
+	template<typename T>
+	class TClosableLockFreePointerListUnorderedSingleConsumer_Adapter
+	{
+	public:
+		bool Enqueue(T* Value)
+		{
+			return Queue.PushIfNotClosed(MoveTemp(Value));
+		}
+
+		template<typename F>
+		void Close(F&& Consumer)
+		{
+			TArray<T*> Items;
+			Queue.PopAllAndClose(Items);
+			for (T* Item : Items)
+			{
+				Consumer(MoveTemp(Item));
+			}
+		}
+
+	private:
+		TClosableLockFreePointerListUnorderedSingleConsumer<T, 0 /* as was used in FGraphEvent */> Queue;
+	};
+
+	template<uint32 Num, typename QueueT>
+	void TestClosingEmptyQueue()
+	{
+		for (uint32 i = 0; i != Num; ++i)
+		{
+			QueueT{}.Close([](void*) {});
+		}
+	}
+
+	template<uint32 Num, typename QueueT>
+	void TestCorrectness()
+	{
+		for (uint32 Run = 0; Run != Num; ++Run)
+		{
+			QueueT Queue;
+
+			int32 ProducersNum = FPlatformMisc::NumberOfCoresIncludingHyperthreads() - 1;
+
+			FGraphEventArray Producers;
+			Producers.Reserve(ProducersNum);
+
+			TArray<uint32> NumsProduced;
+			NumsProduced.AddZeroed(ProducersNum);
+
+			for (int32 i = 0; i != ProducersNum; ++i)
+			{
+				Producers.Add(FFunctionGraphTask::CreateAndDispatchWhenReady(
+					[&Queue, NumProduced = &NumsProduced[i]]
+					{
+						uint32 i = 0;
+						while (Queue.Enqueue((void*)(intptr_t)(++i)))
+						{
+						}
+
+						*NumProduced = i - 1;
+					}
+					));
+			}
+
+			FPlatformProcess::Yield();
+
+			uint64 Consumed = 0;
+			Queue.Close([&Consumed](void* Value) { Consumed += (uint32)(intptr_t)Value; });
+
+			FTaskGraphInterface::Get().WaitUntilTasksComplete(Producers);
+
+			uint64 Produced = 0;
+			for (int32 i = 0; i != ProducersNum; ++i)
+			{
+				uint32 N = NumsProduced[i];
+				Produced += N * (N + 1) / 2;
+			}
+
+			checkf(Produced == Consumed, TEXT("%d - %d"), Produced, Consumed);
+
+			if (Produced == 0)
+			{
+				Run -= 1; // discard this run
+			}
+		}
+	}
+
+	template<uint32 Num, typename QueueType>
+	void TestPerf()
+	{
+		TArray<QueueType> Queues;
+		Queues.AddDefaulted(Num);
+		std::atomic<int> Index{ 0 };
+
+		const int32 NumProducers = FPlatformMisc::NumberOfWorkerThreadsToSpawn();
+		TArray<int32> ProducersResults;
+		ProducersResults.AddDefaulted(NumProducers);
+
+		FGraphEventArray Events;
+		for (int ProducerIndex = 0; ProducerIndex != NumProducers; ++ProducerIndex)
+		{
+			Events.Add(FFunctionGraphTask::CreateAndDispatchWhenReady(
+				[&Queues, &Index, &ProducersResults, ProducerIndex]
+				{
+					int NumProduced = 0;
+					int LocalIndex = Index.load(std::memory_order_acquire);
+
+					while (LocalIndex != Num && Queues[LocalIndex].Enqueue((void*)(intptr_t)(NumProduced + 1)))
+					{
+						++NumProduced;
+						LocalIndex = Index.load(std::memory_order_acquire);
+					}
+
+					ProducersResults[ProducerIndex] = NumProduced;
+				}
+				));
+			Events.Last()->SetDebugName(TEXT("ClosableMpscQueueTest.TestPerf.Producer"));
+		}
+
+		int32 NumConsumed = 0;
+		for (int LocalIndex = 0; LocalIndex != Num; ++LocalIndex)
+		{
+			Queues[LocalIndex].Close([&NumConsumed](void*) { ++NumConsumed; });
+			Index.store(LocalIndex + 1, std::memory_order_release);
+			FPlatformProcess::Yield();
+			FPlatformProcess::Yield();
+			FPlatformProcess::Yield();
+		}
+
+		FTaskGraphInterface::Get().WaitUntilTasksComplete(Events);
+
+		int32 NumProduced = 0;
+		for (int i = 0; i != NumProducers; ++i)
+		{
+			NumProduced += ProducersResults[i];
+		}
+
+		check(NumProduced == NumConsumed);
+		UE_LOG(LogTemp, Display, TEXT("items queued %d"), NumProduced);
+	}
+
+	IMPLEMENT_SIMPLE_AUTOMATION_TEST(FClosableMpscQueueTest, "System.Core.Async.ClosableMpscQueueTest", EAutomationTestFlags::ApplicationContextMask | EAutomationTestFlags::EngineFilter);
+
+	bool FClosableMpscQueueTest::RunTest(const FString& Parameters)
+	{
+		{
+			TClosableMpscQueue<int> Q;
+			Q.Close([](int) { check(false); });
+		}
+		{
+			TClosableMpscQueue<int> Q;
+			Q.Enqueue(1);
+			bool Done = false;
+			Q.Close([&Done](int Item) { check(!Done && Item == 1); Done = true; });
+		}
+
+		{
+			TClosableMpscQueue<int> Q;
+			Q.Enqueue(1);
+			Q.Enqueue(2);
+			int Step = 0;
+			Q.Close(
+				[&Step](int Item)
+				{
+					switch (Step)
+					{
+					case 0:
+						check(Item == 1);
+						++Step;
+						break;
+					case 1:
+						check(Item == 2);
+						++Step;
+						break;
+					default:
+						check(false);
+					}
+				}
+			);
+		}
+
+		//for (int i = 0; i != 1000000; ++i)
+		{
+			TClosableMpscQueue<int> Q;
+			Q.Close([](int) { check(false); });
+			Q.Close([](int) { check(false); });
+		}
+
+		//for (int i = 0; i != 1000000; ++i)
+		{
+			TClosableMpscQueue<int> Q;
+			check(!Q.IsClosed());
+			Q.Close([](int) { check(false); });
+			check(Q.IsClosed());
+		}
+
+		UE_BENCHMARK(5, TestClosingEmptyQueue<10'000'000, TClosableMpscQueue<void*>>);
+		UE_BENCHMARK(5, TestClosingEmptyQueue<10'000'000, TClosableMpscQueueOnTQueue<void*>>);
+		UE_BENCHMARK(5, TestClosingEmptyQueue<10'000'000, TClosableMpscQueueOnTMpscQueue<void*>>);
+		UE_BENCHMARK(5, TestClosingEmptyQueue<10'000'000, TClosableLockFreePointerListUnorderedSingleConsumer_Adapter<void>>);
+
+		UE_BENCHMARK(5, TestCorrectness<1000, TClosableMpscQueue<void*>>);
+		UE_BENCHMARK(5, TestCorrectness<1000, TClosableMpscQueueOnTQueue<void*>>);
+		UE_BENCHMARK(5, TestCorrectness<1000, TClosableMpscQueueOnTMpscQueue<void*>>);
+		UE_BENCHMARK(5, TestCorrectness<1000, TClosableLockFreePointerListUnorderedSingleConsumer_Adapter<void>>);
+
+		UE_BENCHMARK(5, TestPerf<1'000'000, TClosableMpscQueue<void*>>);
+		UE_BENCHMARK(5, TestPerf<1'000'000, TClosableMpscQueueOnTQueue<void*>>);
+		UE_BENCHMARK(5, TestPerf<1'000'000, TClosableMpscQueueOnTMpscQueue<void*>>);
+		UE_BENCHMARK(5, TestPerf<1'000'000, TClosableLockFreePointerListUnorderedSingleConsumer_Adapter<void>>);
+
+		return true;
+	}
+}
+
 #endif // WITH_DEV_AUTOMATION_TESTS
diff --git a/Engine/Source/Runtime/Core/Public/Containers/ClosableMpscQueue.h b/Engine/Source/Runtime/Core/Public/Containers/ClosableMpscQueue.h
new file mode 100644
--- /dev/null
+++ b/Engine/Source/Runtime/Core/Public/Containers/ClosableMpscQueue.h
@@ -0,0 +1,141 @@
+// Copyright Epic Games, Inc. All Rights Reserved.
+
+#pragma once
+
+#include "CoreTypes.h"
+#include "Templates/TypeCompatibleBytes.h"
+#include <atomic>
+
+/** 
+ * Multi-producer/single-consumer unbounded concurrent queue that can be consumed only once.
+ */
+template<typename T>
+class TClosableMpscQueue final
+{
+public:
+	UE_NONCOPYABLE(TClosableMpscQueue);
+
+	TClosableMpscQueue() = default;
+
+	~TClosableMpscQueue()
+	{
+		if (Head.load(std::memory_order_relaxed) == nullptr)
+		{
+			return; // closed
+		}
+
+		FNode* Tail = Sentinel.Next.load(std::memory_order_relaxed);
+		while (Tail != nullptr)
+		{
+			FNode* Next = Tail->Next.load(std::memory_order_relaxed);
+			DestructItem((T*)&Tail->Value);
+			delete Tail;
+			Tail = Next;
+		}
+	}
+
+	/**
+	 * Returns false if the queue is closed
+	 */
+	template <typename... ArgTypes>
+	bool Enqueue(ArgTypes&&... Args)
+	{
+		FNode* Prev = Head.load(std::memory_order_acquire);
+		if (Prev == nullptr)
+		{
+			return false; // already closed
+		}
+
+		FNode* New = new FNode;
+		new (&New->Value) T(Forward<ArgTypes>(Args)...);
+
+		while (!Head.compare_exchange_weak(Prev, New, std::memory_order_release) && Prev != nullptr) // linearisation point
+		{
+		}
+
+		if (Prev == nullptr)
+		{
+			DestructItem((T*)&Prev->Value);
+			delete New;
+			return false;
+		}
+
+		Prev->Next.store(New, std::memory_order_release);
+
+		return true;
+	}
+
+	/**
+	 * Closes the queue and consumes all items.
+	 * @param Consumer: a functor with signature `AnyReturnType (T Value)` that will receive all items in FIFO order
+	 */
+	template<typename F>
+	void Close(const F& Consumer)
+	{
+		FNode* Tail = &Sentinel;
+
+		// we need to take note of the head at the moment of nullifying it because it can be still unreacheable from the tail
+		FNode* const Head_Local = Head.exchange(nullptr, std::memory_order_acq_rel); // linearisation point
+
+		// the queue is closed at this point, and the user is free to destroy it
+		// no members should be accessed
+		Close_NonMember(Head_Local, Tail, Consumer);
+	}
+
+	bool IsClosed() const
+	{
+		return Head.load(std::memory_order_relaxed) == nullptr;
+	}
+
+private:
+	struct FNode
+	{
+		std::atomic<FNode*> Next{ nullptr };
+		TTypeCompatibleBytes<T> Value;
+	};
+
+	FNode Sentinel;
+	std::atomic<FNode*> Head{ &Sentinel };
+
+private:
+	template<typename F>
+	static void Close_NonMember(FNode* Head, FNode* Tail, const F& Consumer)
+	{
+		if (Head == Tail /* empty */ || Head == nullptr /* already closed */)
+		{
+			return;
+		}
+
+		auto GetNext = [](FNode* Node)
+		{
+			FNode* Next;
+			// producers can be still updating `Next`, we need to loop until we detect that the list is fully linked
+			do
+			{
+				Next = Node->Next.load(std::memory_order_relaxed);
+			} while (Next == nullptr);
+
+			return Next;
+		};
+
+		// skip sentinel, do it outside of the main loop to avoid unnecessary branching in it
+		Tail = GetNext(Tail);
+
+		auto Consume = [&Consumer](FNode* Node)
+		{
+			T* ValuePtr = (T*)&Node->Value;
+			Consumer(MoveTemp(*ValuePtr));
+			DestructItem(ValuePtr);
+			delete Node;
+		};
+
+		while (Tail != Head)
+		{
+			FNode* Next = GetNext(Tail);
+			Consume(Tail);
+			Tail = Next;
+		}
+
+		Consume(Head);
+	}
+};
