fixed non-unity build

diff --git a/Engine/Source/Runtime/Core/Public/Tasks/Task.h b/Engine/Source/Runtime/Core/Public/Tasks/Task.h
--- a/Engine/Source/Runtime/Core/Public/Tasks/Task.h
+++ b/Engine/Source/Runtime/Core/Public/Tasks/Task.h
@@ -4,6 +4,7 @@
 
 #include "Tasks/TaskPrivate.h"
 #include "Async/Fundamental/Task.h"
+#include "Templates/RefCounting.h"
 #include "CoreTypes.h"
 
 namespace UE { namespace Tasks
