fixed a bug in FClosableMpscQueue caught by static analysis

diff --git a/Engine/Source/Runtime/Core/Public/Containers/ClosableMpscQueue.h b/Engine/Source/Runtime/Core/Public/Containers/ClosableMpscQueue.h
--- a/Engine/Source/Runtime/Core/Public/Containers/ClosableMpscQueue.h
+++ b/Engine/Source/Runtime/Core/Public/Containers/ClosableMpscQueue.h
@@ -55,7 +55,7 @@ public:
 
 		if (Prev == nullptr)
 		{
-			DestructItem((T*)&Prev->Value);
+			DestructItem((T*)&New->Value);
 			delete New;
 			return false;
 		}
