remove some references to UE4

#rnx

diff --git a/Engine/Source/Runtime/Core/Private/FramePro/FrameProUE4.cpp b/Engine/Source/Runtime/Core/Private/FramePro/FrameProUE4.cpp
--- a/Engine/Source/Runtime/Core/Private/FramePro/FrameProUE4.cpp
+++ b/Engine/Source/Runtime/Core/Private/FramePro/FrameProUE4.cpp
@@ -144,8 +144,8 @@
 		//------------------------------------------------------------------------
 		Platform::Enum Platform::GetPlatformEnum()
 		{
-			#if defined(FRAMEPRO_UE4_PLATFORM) //@EPIC: allow external definition
-				return FRAMEPRO_UE4_PLATFORM;  //@EPIC end
+			#if defined(FRAMEPRO_UNREAL_PLATFORM) //@EPIC: allow external definition
+				return FRAMEPRO_UNREAL_PLATFORM;  //@EPIC end
 			#elif PLATFORM_WINDOWS
 				return Platform::Windows;
 			#elif PLATFORM_LINUX
@@ -164,7 +164,7 @@
 				return Platform::Switch;
 			#else
 				//@EPIC: begin - useful error
-				#error unknown platform or FRAMEPRO_UE4_PLATFORM not defined
+				#error unknown platform or FRAMEPRO_UNREAL_PLATFORM not defined
 				//@EPIC: end
 			#endif
 		}
diff --git a/Engine/Source/Runtime/Core/Public/HoloLens/MinHoloLensApi.h b/Engine/Source/Runtime/Core/Public/HoloLens/MinHoloLensApi.h
--- a/Engine/Source/Runtime/Core/Public/HoloLens/MinHoloLensApi.h
+++ b/Engine/Source/Runtime/Core/Public/HoloLens/MinHoloLensApi.h
@@ -8,7 +8,7 @@
 
 #include "CoreTypes.h"
 
-#if defined(_WINDOWS_) && !defined(UE4_MINIMAL_WINDOWS_INCLUDE)
+#if defined(_WINDOWS_) && !defined(UE_MINIMAL_WINDOWS_INCLUDE)
 	#pragma message ( " " )
 	#pragma message ( "You have included windows.h before MinWindows.h" )
 	#pragma message ( "All useless stuff from the windows headers won't be excluded !!!" )
@@ -19,7 +19,7 @@
 #error PLATFORM_HOLOLENS not defined
 #endif
 
-#define UE4_MINIMAL_WINDOWS_INCLUDE
+#define UE_MINIMAL_WINDOWS_INCLUDE
 
 // WIN32_LEAN_AND_MEAN excludes rarely-used services from windows headers.
 #define WIN32_LEAN_AND_MEAN
diff --git a/Engine/Source/Runtime/Core/Public/Windows/MinWindows.h b/Engine/Source/Runtime/Core/Public/Windows/MinWindows.h
--- a/Engine/Source/Runtime/Core/Public/Windows/MinWindows.h
+++ b/Engine/Source/Runtime/Core/Public/Windows/MinWindows.h
@@ -4,14 +4,14 @@
 
 #include "CoreTypes.h"
 
-#if defined(_WINDOWS_) && !defined(UE4_MINIMAL_WINDOWS_INCLUDE)
+#if defined(_WINDOWS_) && !defined(UE_MINIMAL_WINDOWS_INCLUDE)
 	#pragma message ( " " )
 	#pragma message ( "You have included windows.h before MinWindows.h" )
 	#pragma message ( "All useless stuff from the windows headers won't be excluded !!!" )
 	#pragma message ( " " )
 #endif // _WINDOWS_
 
-#define UE4_MINIMAL_WINDOWS_INCLUDE
+#define UE_MINIMAL_WINDOWS_INCLUDE
 
 // WIN32_LEAN_AND_MEAN excludes rarely-used services from windows headers.
 #define WIN32_LEAN_AND_MEAN
