U5 - only wake the emergency worker once during busy wait

diff --git a/Engine/Source/Runtime/Core/Private/Async/Fundamental/Scheduler.cpp b/Engine/Source/Runtime/Core/Private/Async/Fundamental/Scheduler.cpp
--- a/Engine/Source/Runtime/Core/Private/Async/Fundamental/Scheduler.cpp
+++ b/Engine/Source/Runtime/Core/Private/Async/Fundamental/Scheduler.cpp
@@ -300,6 +300,7 @@ namespace LowLevelTasks
 		FLocalQueueType* WorkerLocalQueue = LocalQueue;
 
 		uint32 WaitCount = 0;
+		bool HasWokenEmergencyWorker = false;
 		const bool bIsBackgroundWorker = WorkerType == EWorkerType::Background;
 		bool bPermitBackgroundWork = ActiveTask && ActiveTask->IsBackgroundTask();
 		FQueueRegistry::FOutOfWork OutOfWork = QueueRegistry.GetOutOfWorkScope(bIsBackgroundWorker);
@@ -343,7 +344,11 @@ namespace LowLevelTasks
 			}
 			else
 			{
-				WakeUpWorker(true);
+				if(!HasWokenEmergencyWorker)
+				{
+					WakeUpWorker(true);
+					HasWokenEmergencyWorker = true;
+				}
 				WaitCount = 0;
 			}			
 		}
