Used some of the unused bits in LLM's allocation map to store more bits for size. This extends the maximum supported size to 2^48 (or 256TB) at small ALU cost and no additional memory.

#rnx

diff --git a/Engine/Source/Runtime/Core/Private/HAL/LowLevelMemTracker.cpp b/Engine/Source/Runtime/Core/Private/HAL/LowLevelMemTracker.cpp
--- a/Engine/Source/Runtime/Core/Private/HAL/LowLevelMemTracker.cpp
+++ b/Engine/Source/Runtime/Core/Private/HAL/LowLevelMemTracker.cpp
@@ -2901,9 +2901,12 @@ namespace LLMPrivate
 #if LLM_ALLOW_ASSETS_TAGS
 			AllocInfo.SetAssetTag(AssetTagData, LLMRef);
 #endif
-			LLMCheck(Size <= 0xffffffffu);
+			LLMCheck(Size <= 0x0000'ffff'ffff'ffff);
+			uint32 SizeLow = uint32(Size);
+			uint16 SizeHigh = uint16(Size >> 32ull);
 			FScopeLock AllocationScopeLock(&AllocationMapLock);
-			AllocationMap.Add(Ptr, static_cast<uint32>(Size), AllocInfo);
+			PointerKey Key(Ptr, SizeHigh);
+			AllocationMap.Add(Key, SizeLow, AllocInfo);
 		}
 	}
 
@@ -2913,7 +2916,7 @@ namespace LLMPrivate
 		FLLMAllocMap::Values Values;
 		{
 			FScopeLock AllocationScopeLock(&AllocationMapLock);
-			if (!AllocationMap.Remove(Ptr, Values))
+			if (!AllocationMap.Remove(PointerKey(Ptr), Values))
 			{
 				return;
 			}
@@ -2925,7 +2928,9 @@ namespace LLMPrivate
 			return;
 		}
 
-		int64 Size = static_cast<int64>(Values.Value1);
+		int64 SizeLow = static_cast<int64>(Values.Value1);
+		int64 SizeHigh = Values.Key.GetExtraData();
+		int64 Size = (SizeHigh << 32ull) | SizeLow;
 		FLLMTracker::FLowLevelAllocInfo& AllocInfo = Values.Value2;
 
 		// track the total quickly
@@ -2947,12 +2952,13 @@ namespace LLMPrivate
 		FLLMAllocMap::Values Values;
 		{
 			FScopeLock AllocationScopeLock(&AllocationMapLock);
-			if (!AllocationMap.Remove(Source, Values))
+			if (!AllocationMap.Remove(PointerKey(Source), Values))
 			{
 				return;
 			}
 
-			AllocationMap.Add(Dest, Values.Value1, Values.Value2);
+			PointerKey Key(Dest, uint16(Values.Key.GetExtraData()));
+			AllocationMap.Add(Key, Values.Value1, Values.Value2);
 		}
 
 		if (IsPaused(AllocType))
@@ -2961,7 +2967,9 @@ namespace LLMPrivate
 			return;
 		}
 
-		int64 Size = static_cast<int64>(Values.Value1);
+		int64 SizeLow = static_cast<int64>(Values.Value1);
+		int64 SizeHigh = Values.Key.GetExtraData();
+		int64 Size = (SizeHigh << 32ull) | SizeLow;
 		const FLLMTracker::FLowLevelAllocInfo& AllocInfo = Values.Value2;
 		const FTagData* TagData = AllocInfo.GetTag(LLMRef);
 
@@ -3226,7 +3234,7 @@ namespace LLMPrivate
 			FScopeLock AllocationScopeLock(&AllocationMapLock);
 			uint32* Size;
 			FLowLevelAllocInfo* AllocInfoPtr;
-			AllocationMap.Find(Ptr, Size, AllocInfoPtr);
+			AllocationMap.Find(PointerKey(Ptr), Size, AllocInfoPtr);
 			if (!AllocInfoPtr)
 			{
 				return nullptr;
diff --git a/Engine/Source/Runtime/Core/Private/HAL/LowLevelMemoryUtils.h b/Engine/Source/Runtime/Core/Private/HAL/LowLevelMemoryUtils.h
--- a/Engine/Source/Runtime/Core/Private/HAL/LowLevelMemoryUtils.h
+++ b/Engine/Source/Runtime/Core/Private/HAL/LowLevelMemoryUtils.h
@@ -195,6 +195,7 @@ public:
 
 	struct Values
 	{
+		TKey Key;
 		TValue1 Value1;
 		TValue2 Value2;
 	};
@@ -308,14 +309,16 @@ public:
 	{
 		TValue1* Value1;
 		TValue2* Value2;
-		LLMEnsure(Find(Key, Value1, Value2));
+		Key = Find(Key, Value1, Value2);
+		LLMEnsure(Key);
 		Values RetValues;
+		RetValues.Key = Key;
 		RetValues.Value1 = *Value1;
 		RetValues.Value2 = *Value2;
 		return RetValues;
 	}
 
-	bool Find(const TKey& Key, TValue1*& OutValue1, TValue2*& OutValue2)
+	TKey Find(const TKey& Key, TValue1*& OutValue1, TValue2*& OutValue2)
 	{
 		SizeType KeyHash = Key.GetHashCode();
 
@@ -326,13 +329,13 @@ public:
 		{
 			OutValue1 = nullptr;
 			OutValue2 = nullptr;
-			return false;
+			return TKey();
 		}
 
 		OutValue1 = &Values1[KeyIndex];
 		OutValue2 = &Values2[KeyIndex];
 
-		return true;
+		return Keys[KeyIndex];
 	}
 
 	bool Remove(const TKey& Key, Values& OutValues)
@@ -349,6 +352,7 @@ public:
 
 		SizeType KeyIndex = Map[MapIndex];
 
+		OutValues.Key = Keys[KeyIndex];
 		OutValues.Value1 = Values1[KeyIndex];
 		OutValues.Value2 = Values2[KeyIndex];
 
@@ -697,28 +701,37 @@ private:
 // Pointer key for hash map
 struct PointerKey
 {
-	PointerKey() : Pointer(NULL) {}
-	PointerKey(const void* InPointer) : Pointer(InPointer) {}
+	explicit PointerKey() = default;
+	explicit PointerKey(const void* InPointer, uint16 ExtraData=0)
+	: Pointer((UPTRINT(InPointer) << 16) | UPTRINT(ExtraData))
+	{
+		static_assert(sizeof(InPointer) == 8, "LLM only supports 64 bit platforms");
+		LLMEnsure(UPTRINT(InPointer) && UPTRINT(InPointer) <= 0x0000'ffff'ffff'ffff);
+	}
 	LLMNumAllocsType GetHashCode() const
 	{
 		return GetHashCodeImpl<sizeof(LLMNumAllocsType), sizeof(void*)>();
 	}
-	bool operator==(const PointerKey& other) const { return Pointer == other.Pointer; }
-	const void* Pointer;
+	bool operator==(const PointerKey& other) const { return GetPointer() == other.GetPointer(); }
+	explicit operator bool () const { return !!Pointer; }
+	void* GetPointer() const { return (void*)(Pointer >> 16); };
+	uint64 GetExtraData() const { return uint64(Pointer & 0xffff); }
 
 private:
+	UPTRINT Pointer = 0;
+
 	template <int HashSize, int PointerSize>
 	LLMNumAllocsType GetHashCodeImpl() const
 	{
 		static_assert(HashSize == 0 && PointerSize == 0, "Converting void* to a LLMNumAllocsType - sized hashkey is not implemented for the current sizes.");
-		return (LLMNumAllocsType)(intptr_t)Pointer;
+		return (LLMNumAllocsType)(UPTRINT)GetPointer();
 	}
 
 	template <>
 	LLMNumAllocsType GetHashCodeImpl<8,8>() const
 	{
 		// 64 bit pointer to 64 bit hash
-		uint64 Key = (uint64)Pointer;
+		uint64 Key = (uint64)GetPointer();
 		Key = (~Key) + (Key << 21);
 		Key = Key ^ (Key >> 24);
 		Key = Key * 265;
@@ -733,7 +746,7 @@ private:
 	LLMNumAllocsType GetHashCodeImpl<4, 8>() const
 	{
 		// 64 bit pointer to 32 bit hash
-		uint64 Key = (uint64)Pointer;
+		uint64 Key = (uint64)GetPointer();
 		Key = (~Key) + (Key << 21);
 		Key = Key ^ (Key >> 24);
 		Key = Key * 265;
@@ -748,7 +761,7 @@ private:
 	LLMNumAllocsType GetHashCodeImpl<4, 4>() const
 	{
 		// 32 bit pointer to 32 bit Hash
-		uint64 Key = (uint64)Pointer;
+		uint64 Key = (uint64)GetPointer();
 		Key = (~Key) + (Key << 18);
 		Key = Key ^ (Key >> 31);
 		Key = Key * 21;
