[Insights] Fixed size of the name buffer for Diagnostics.ModuleLoad trace event and string lengths for Diagnostics.Session2 trace event.

diff --git a/Engine/Source/Runtime/Core/Private/ProfilingDebugging/TraceAuxiliary.cpp b/Engine/Source/Runtime/Core/Private/ProfilingDebugging/TraceAuxiliary.cpp
--- a/Engine/Source/Runtime/Core/Private/ProfilingDebugging/TraceAuxiliary.cpp
+++ b/Engine/Source/Runtime/Core/Private/ProfilingDebugging/TraceAuxiliary.cpp
@@ -363,8 +363,8 @@ static void TraceAuxiliaryStart(const TArray<FString>& Args)
 
 		Channels += Channel;
 	});
-	UE_LOG(LogConsoleResponse, Log, TEXT("Tracing to; %s"), TraceDest);
-	UE_LOG(LogConsoleResponse, Log, TEXT("Trace channels; %s"), *Channels);
+	UE_LOG(LogConsoleResponse, Log, TEXT("Tracing to: %s"), TraceDest);
+	UE_LOG(LogConsoleResponse, Log, TEXT("Trace channels: %s"), *Channels);
 }
 
 ////////////////////////////////////////////////////////////////////////////////
@@ -415,18 +415,23 @@ void FTraceAuxiliary::Initialize(const TCHAR* CommandLine)
 	// so that it is always sent (all channels are enabled prior to initialisation)
 	const TCHAR* BranchName = BuildSettings::GetBranchName();
 	const TCHAR* BuildVersion = BuildSettings::GetBuildVersion();
-	uint32 DataSize = 
-		(UE_ARRAY_COUNT(PREPROCESSOR_TO_STRING(UBT_COMPILED_PLATFORM)) * sizeof(ANSICHAR)) +
-		(UE_ARRAY_COUNT(UE_APP_NAME) * sizeof(ANSICHAR)) +
-		(FCString::Strlen(CommandLine) * sizeof(TCHAR)) +
-		(FCString::Strlen(BranchName) * sizeof(TCHAR)) +
-		(FCString::Strlen(BuildVersion) * sizeof(TCHAR));
+	constexpr uint32 PlatformLen = UE_ARRAY_COUNT(PREPROCESSOR_TO_STRING(UBT_COMPILED_PLATFORM)) - 1;
+	constexpr uint32 AppNameLen = UE_ARRAY_COUNT(UE_APP_NAME) - 1;
+	const uint32 CommandLineLen = FCString::Strlen(CommandLine);
+	const uint32 BranchNameLen = FCString::Strlen(BranchName);
+	const uint32 BuildVersionLen = FCString::Strlen(BuildVersion);
+	uint32 DataSize =
+		(PlatformLen * sizeof(ANSICHAR)) +
+		(AppNameLen * sizeof(ANSICHAR)) +
+		(CommandLineLen * sizeof(TCHAR)) +
+		(BranchNameLen * sizeof(TCHAR)) +
+		(BuildVersionLen * sizeof(TCHAR));
 	UE_TRACE_LOG(Diagnostics, Session2, UE::Trace::TraceLogChannel, DataSize)
-		<< Session2.Platform(PREPROCESSOR_TO_STRING(UBT_COMPILED_PLATFORM))
-		<< Session2.AppName(UE_APP_NAME)
-		<< Session2.CommandLine(CommandLine)
-		<< Session2.Branch(BranchName)
-		<< Session2.BuildVersion(BuildVersion)
+		<< Session2.Platform(PREPROCESSOR_TO_STRING(UBT_COMPILED_PLATFORM), PlatformLen)
+		<< Session2.AppName(UE_APP_NAME, AppNameLen)
+		<< Session2.CommandLine(CommandLine, CommandLineLen)
+		<< Session2.Branch(BranchName, BranchNameLen)
+		<< Session2.BuildVersion(BuildVersion, BuildVersionLen)
 		<< Session2.Changelist(BuildSettings::GetCurrentChangelist())
 		<< Session2.ConfigurationType(uint8(FApp::GetBuildConfiguration()))
 		<< Session2.TargetType(uint8(FApp::GetBuildTargetType()));
diff --git a/Engine/Source/Runtime/Core/Private/ProfilingDebugging/Windows/ModuleDiagnostics.cpp b/Engine/Source/Runtime/Core/Private/ProfilingDebugging/Windows/ModuleDiagnostics.cpp
--- a/Engine/Source/Runtime/Core/Private/ProfilingDebugging/Windows/ModuleDiagnostics.cpp
+++ b/Engine/Source/Runtime/Core/Private/ProfilingDebugging/Windows/ModuleDiagnostics.cpp
@@ -171,7 +171,8 @@ void FModuleTrace::OnDllLoaded(const UNICODE_STRING& Name, UPTRINT Base)
 	const auto* NtHeaders = (IMAGE_NT_HEADERS*)(Base + DosHeader->e_lfanew);
 	const IMAGE_OPTIONAL_HEADER& OptionalHeader = NtHeaders->OptionalHeader;
 
-	UE_TRACE_LOG(Diagnostics, ModuleLoad, ModuleChannel, Name.Length * sizeof(TCHAR))
+	// Note: UNICODE_STRING.Length is the size in bytes of the string buffer.
+	UE_TRACE_LOG(Diagnostics, ModuleLoad, ModuleChannel, Name.Length)
 		<< ModuleLoad.Name(Name.Buffer, Name.Length / 2)
 		<< ModuleLoad.Base(uint32(Base >> 16)) // Windows' DLLs are on 64K page boundaries
 		<< ModuleLoad.Size(OptionalHeader.SizeOfImage);
