UE5EA - the event that blocks thread creation is outside of a thread tag scope, therefore revert this code to the previous version that did no care about tags.

diff --git a/Engine/Source/Runtime/Core/Public/Stats/Stats2.h b/Engine/Source/Runtime/Core/Public/Stats/Stats2.h
--- a/Engine/Source/Runtime/Core/Public/Stats/Stats2.h
+++ b/Engine/Source/Runtime/Core/Public/Stats/Stats2.h
@@ -1133,11 +1133,12 @@ struct FStatPacket
 	void SetThreadProperties()
 	{
 		ThreadId = FPlatformTLS::GetCurrentThreadId();
-		if (IsInGameThread())
+		PRAGMA_DISABLE_DEPRECATION_WARNINGS
+		if (ThreadId == GGameThreadId)
 		{
 			ThreadType = EThreadType::Game;
 		}
-		else if (IsInActualRenderingThread())
+		else if (ThreadId == GRenderThreadId)
 		{
 			ThreadType = EThreadType::Renderer;
 		}
@@ -1145,6 +1146,7 @@ struct FStatPacket
 		{
 			ThreadType = EThreadType::Other;
 		}
+		PRAGMA_ENABLE_DEPRECATION_WARNINGS
 	}
 
 
