Minor optimizations to the memory image writing (UE-110824)

- Flatten memory image sections in parallel and preallocate the array for writing.
- Memory image writing is currently happening on the game thread and either moving it out or optimizing is critical.
- Makes flattening about 3x faster on the test asset (~90ms -> ~30ms)

diff --git a/Engine/Source/Runtime/Core/Private/Serialization/MemoryImage.cpp b/Engine/Source/Runtime/Core/Private/Serialization/MemoryImage.cpp
--- a/Engine/Source/Runtime/Core/Private/Serialization/MemoryImage.cpp
+++ b/Engine/Source/Runtime/Core/Private/Serialization/MemoryImage.cpp
@@ -13,6 +13,7 @@
 #include "Misc/DataDrivenPlatformInfoRegistry.h"
 #include "Serialization/Archive.h"
 #include "Misc/MemStack.h"
+#include "Async/ParallelFor.h"
 
 DEFINE_LOG_CATEGORY_STATIC(LogMemoryImage, Log, All);
 
@@ -1751,11 +1752,13 @@ void FMemoryImage::Flatten(FMemoryImageResult& OutResult, bool bMergeDuplicateSe
 	if(bMergeDuplicateSections)
 	{
 		// Find unique sections
+		// precompute hashes in parallel
+		ParallelFor(Sections.Num(), [this](int32 Index) { Sections[Index]->ComputeHash(); });
+
 		TMap<FSHAHash, int32> HashToSectionIndex;
 		for (int32 SectionIndex = 0; SectionIndex < Sections.Num(); ++SectionIndex)
 		{
 			FMemoryImageSection* Section = Sections[SectionIndex];
-			Section->ComputeHash();
 			int32* FoundIndex = HashToSectionIndex.Find(Section->Hash);
 			int32 NewIndex = INDEX_NONE;
 			if (!FoundIndex)
diff --git a/Engine/Source/Runtime/Core/Public/Serialization/MemoryImage.h b/Engine/Source/Runtime/Core/Public/Serialization/MemoryImage.h
--- a/Engine/Source/Runtime/Core/Public/Serialization/MemoryImage.h
+++ b/Engine/Source/Runtime/Core/Public/Serialization/MemoryImage.h
@@ -270,6 +270,9 @@ public:
 	FMemoryImageSection* AllocateSection()
 	{
 		FMemoryImageSection* Section = new FMemoryImageSection(this);
+#if WITH_EDITOR
+		Section->Bytes.Reserve(512 * 1024);	// reserve at least 512KB per section to avoid resizing
+#endif
 		Sections.Add(Section);
 		return Section;
 	}
