Add FStringView-compatible FAsciiSet algorithms

diff --git a/Engine/Source/Runtime/Core/Private/Tests/Misc/AsciiSetTest.cpp b/Engine/Source/Runtime/Core/Private/Tests/Misc/AsciiSetTest.cpp
--- a/Engine/Source/Runtime/Core/Private/Tests/Misc/AsciiSetTest.cpp
+++ b/Engine/Source/Runtime/Core/Private/Tests/Misc/AsciiSetTest.cpp
@@ -1,5 +1,7 @@
 // Copyright Epic Games, Inc. All Rights Reserved.
 
+
+#include "Containers/StringView.h"
 #include "Misc/AsciiSet.h"
 #include "Misc/AutomationTest.h"
 
@@ -44,22 +46,44 @@ bool TAsciiSetTest::RunTest(const FString& Parameters)
 	TestEqual(TEXT("AdvanceToFirst"),	*FAsciiSet::FindFirstOrEnd("NonWhitespaceNonWhitespace", Whitespaces), '\0');
 	TestEqual(TEXT("AdvanceToLast"),	*FAsciiSet::FindLastOrEnd("NonWhitespaceNonWhitespace", Whitespaces), '\0');
 
-	constexpr FAsciiSet XmlEscapeChars("&<>\"'");
-	TestTrue(TEXT("None"), FAsciiSet::HasNone("No escape chars", XmlEscapeChars));
-	TestFalse(TEXT("Any"), FAsciiSet::HasAny("No escape chars", XmlEscapeChars));
-	TestFalse(TEXT("Only"), FAsciiSet::HasOnly("No escape chars", XmlEscapeChars));
+	constexpr FAsciiSet Lowercase("abcdefghijklmnopqrstuvwxyz");
+	TestEqual(TEXT("TrimPrefixWithout"),		FAsciiSet::TrimPrefixWithout("ABcdEF"_ASV, Lowercase), "cdEF"_ASV);
+	TestEqual(TEXT("FindPrefixWithout"),		FAsciiSet::FindPrefixWithout("ABcdEF"_ASV, Lowercase), "AB"_ASV);
+	TestEqual(TEXT("TrimSuffixWithout"),		FAsciiSet::TrimSuffixWithout("ABcdEF"_ASV, Lowercase), "ABcd"_ASV);
+	TestEqual(TEXT("FindSuffixWithout"),		FAsciiSet::FindSuffixWithout("ABcdEF"_ASV, Lowercase), "EF"_ASV);
+	TestEqual(TEXT("TrimPrefixWithout none"),	FAsciiSet::TrimPrefixWithout("same"_ASV, Lowercase), "same"_ASV);
+	TestEqual(TEXT("FindPrefixWithout none"),	FAsciiSet::FindPrefixWithout("same"_ASV, Lowercase), ""_ASV);
+	TestEqual(TEXT("TrimSuffixWithout none"),	FAsciiSet::TrimSuffixWithout("same"_ASV, Lowercase), "same"_ASV);
+	TestEqual(TEXT("FindSuffixWithout none"),	FAsciiSet::FindSuffixWithout("same"_ASV, Lowercase), ""_ASV);
+	TestEqual(TEXT("TrimPrefixWithout empty"),	FAsciiSet::TrimPrefixWithout(""_ASV, Lowercase), ""_ASV);
+	TestEqual(TEXT("FindPrefixWithout empty"),	FAsciiSet::FindPrefixWithout(""_ASV, Lowercase), ""_ASV);
+	TestEqual(TEXT("TrimSuffixWithout empty"),	FAsciiSet::TrimSuffixWithout(""_ASV, Lowercase), ""_ASV);
+	TestEqual(TEXT("FindSuffixWithout empty"),	FAsciiSet::FindSuffixWithout(""_ASV, Lowercase), ""_ASV);
+
+
+	auto TestHasFunctions = [&](auto MakeString)
+	{
+		constexpr FAsciiSet XmlEscapeChars("&<>\"'");
+		TestTrue(TEXT("None"),	FAsciiSet::HasNone(MakeString("No escape chars"), XmlEscapeChars));
+		TestFalse(TEXT("Any"),	FAsciiSet::HasAny(MakeString("No escape chars"), XmlEscapeChars));
+		TestFalse(TEXT("Only"), FAsciiSet::HasOnly(MakeString("No escape chars"), XmlEscapeChars));
+
+		TestTrue(TEXT("None"),	FAsciiSet::HasNone(MakeString(""), XmlEscapeChars));
+		TestFalse(TEXT("Any"),	FAsciiSet::HasAny(MakeString(""), XmlEscapeChars));
+		TestTrue(TEXT("Only"),	FAsciiSet::HasOnly(MakeString(""), XmlEscapeChars));
 
-	TestTrue(TEXT("None"), FAsciiSet::HasNone("", XmlEscapeChars));
-	TestFalse(TEXT("Any"), FAsciiSet::HasAny("", XmlEscapeChars));
-	TestTrue(TEXT("Only"), FAsciiSet::HasOnly("", XmlEscapeChars));
+		TestFalse(TEXT("None"), FAsciiSet::HasNone(MakeString("&<>\"'"), XmlEscapeChars));
+		TestTrue(TEXT("Any"),	FAsciiSet::HasAny(MakeString("&<>\"'"), XmlEscapeChars));
+		TestTrue(TEXT("Only"),	FAsciiSet::HasOnly(MakeString("&<>\"'"), XmlEscapeChars));
 
-	TestFalse(TEXT("None"), FAsciiSet::HasNone("&<>\"'", XmlEscapeChars));
-	TestTrue(TEXT("Any"), FAsciiSet::HasAny("&<>\"'", XmlEscapeChars));
-	TestTrue(TEXT("Only"), FAsciiSet::HasOnly("&<>\"'", XmlEscapeChars));
+		TestFalse(TEXT("None"), FAsciiSet::HasNone(MakeString("&<>\"' and more"), XmlEscapeChars));
+		TestTrue(TEXT("Any"),	FAsciiSet::HasAny(MakeString("&<>\"' and more"), XmlEscapeChars));
+		TestFalse(TEXT("Only"), FAsciiSet::HasOnly(MakeString("&<>\"' and more"), XmlEscapeChars));
+	};
+	TestHasFunctions([] (const char* Str) { return Str; });
+	TestHasFunctions([] (const char* Str) { return FAnsiStringView(Str); });
+	TestHasFunctions([] (const char* Str) { return FString(Str); });
 
-	TestFalse(TEXT("None"), FAsciiSet::HasNone("&<>\"' and more", XmlEscapeChars));
-	TestTrue(TEXT("Any"), FAsciiSet::HasAny("&<>\"' and more", XmlEscapeChars));
-	TestFalse(TEXT("Only"), FAsciiSet::HasOnly("&<>\"' and more", XmlEscapeChars));
 
 	constexpr FAsciiSet Abc("abc");
 	constexpr FAsciiSet Abcd = Abc + 'd';
diff --git a/Engine/Source/Runtime/Core/Public/Misc/AsciiSet.h b/Engine/Source/Runtime/Core/Public/Misc/AsciiSet.h
--- a/Engine/Source/Runtime/Core/Public/Misc/AsciiSet.h
+++ b/Engine/Source/Runtime/Core/Public/Misc/AsciiSet.h
@@ -62,7 +62,7 @@ public:
 		return FAsciiSet(~LoMask, ~HiMask);
 	}
 
-	// C string util functions
+	////////// Algorithms for C strings //////////
 
 	/** Find first character of string inside set or end pointer. Never returns null. */
 	template<class CharType>
@@ -120,7 +120,122 @@ public:
 		return *Skip(Str, Set) == '\0';
 	}
 
+	////////// Algorithms for string types like FStringView and FString //////////
+
+	/** Get initial substring with all characters in set */
+	template<class StringType>
+	static constexpr StringType FindPrefixWith(const StringType& Str, FAsciiSet Set)
+	{
+		return Scan<EDir::Forward, ESkip::Members, EKeep::Head>(Str, Set);
+	}
+
+	/** Get initial substring with no characters in set */
+	template<class StringType>
+	static constexpr StringType FindPrefixWithout(const StringType& Str, FAsciiSet Set)
+	{
+		return Scan<EDir::Forward, ESkip::NonMembers, EKeep::Head>(Str, Set);
+	}
+
+	/** Trim initial characters in set */
+	template<class StringType>
+	static constexpr StringType TrimPrefixWith(const StringType& Str, FAsciiSet Set)
+	{
+		return Scan<EDir::Forward, ESkip::Members, EKeep::Tail>(Str, Set);
+	}
+
+	/** Trim initial characters not in set */
+	template<class StringType>
+	static constexpr StringType TrimPrefixWithout(const StringType& Str, FAsciiSet Set)
+	{
+		return Scan<EDir::Forward, ESkip::NonMembers, EKeep::Tail>(Str, Set);
+	}
+
+	/** Get trailing substring with all characters in set */
+	template<class StringType>
+	static constexpr StringType FindSuffixWith(const StringType& Str, FAsciiSet Set)
+	{
+		return Scan<EDir::Reverse, ESkip::Members, EKeep::Tail>(Str, Set);
+	}
+
+	/** Get trailing substring with no characters in set */
+	template<class StringType>
+	static constexpr StringType FindSuffixWithout(const StringType& Str, FAsciiSet Set)
+	{
+		return Scan<EDir::Reverse, ESkip::NonMembers, EKeep::Tail>(Str, Set);
+	}
+
+	/** Trim trailing characters in set */
+	template<class StringType>
+	static constexpr StringType TrimSuffixWith(const StringType& Str, FAsciiSet Set)
+	{
+		return Scan<EDir::Reverse, ESkip::Members, EKeep::Head>(Str, Set);
+	}
+
+	/** Trim trailing characters not in set */
+	template<class StringType>
+	static constexpr StringType TrimSuffixWithout(const StringType& Str, FAsciiSet Set)
+	{
+		return Scan<EDir::Reverse, ESkip::NonMembers, EKeep::Head>(Str, Set);
+	}
+	
+	/** Test if string contains any character in set */
+	template<class StringType>
+	static constexpr bool HasAny(const StringType& Str, FAsciiSet Set)
+	{
+		return !HasNone(Str, Set);
+	}
+
+	/** Test if string contains no character in set */
+	template<class StringType>
+	static constexpr bool HasNone(const StringType& Str, FAsciiSet Set)
+	{
+		uint64 Match = 0;
+		for (auto Char : Str)
+		{
+			Match |= Set.Test(Char);
+		}
+		return Match == 0;
+	}
+
+	/** Test if string contains any character outside of set */
+	template<class StringType>
+	static constexpr bool HasOnly(const StringType& Str, FAsciiSet Set)
+	{
+		auto End = GetData(Str) + GetNum(Str);
+		return FindFirst<ESkip::Members>(Set, GetData(Str), End) == End;
+	}
+
 private:
+	enum class EDir {Forward, Reverse};
+	enum class ESkip {Members, NonMembers};
+	enum class EKeep {Head, Tail};
+
+	template<ESkip Skip, typename CharType>
+	static constexpr const CharType* FindFirst(FAsciiSet Set, const CharType* It, const CharType* End)
+	{
+		for (; It != End && (Skip == ESkip::Members) == !!Set.Test(*It); ++It);
+		return It;
+	}
+
+	template<ESkip Skip, typename CharType>
+	static constexpr const CharType* FindLast(FAsciiSet Set, const CharType* It, const CharType* End)
+	{
+		for (; It != End && (Skip == ESkip::Members) == !!Set.Test(*It); --It);
+		return It;
+	}
+
+	template<EDir Dir, ESkip Skip, EKeep Keep, class StringType>
+	static constexpr StringType Scan(const StringType& Str, FAsciiSet Set)
+	{
+		auto Begin = GetData(Str);
+		auto End = Begin + GetNum(Str);
+		auto It = Dir == EDir::Forward	? FindFirst<Skip>(Set, Begin, End)
+										: FindLast<Skip>(Set, End - 1, Begin - 1) + 1;
+
+		return Keep == EKeep::Head	? StringType(Begin, static_cast<int32>(It - Begin))
+									: StringType(It, static_cast<int32>(End - It));
+	}
+
 	// Work-around for constexpr limitations
 	struct InitData { uint64 Lo, Hi; };
 	static constexpr uint64 NilMask = uint64(1) << '\0';
