added traces for UE_BENCHMARK utility

diff --git a/Engine/Source/Runtime/Core/Public/Tests/Benchmark.h b/Engine/Source/Runtime/Core/Public/Tests/Benchmark.h
--- a/Engine/Source/Runtime/Core/Public/Tests/Benchmark.h
+++ b/Engine/Source/Runtime/Core/Public/Tests/Benchmark.h
@@ -30,9 +30,9 @@ void Benchmark(const TCHAR* TestName, TestT&& TestBody)
 	UE_LOG(LogTemp, Display, TEXT("min: %f secs, avg: %f secs\n-------------------------------\n"), MinTime, TotalTime / NumRuns);
 
 #if NO_LOGGING
-	printf("min: %f secs, avg: %f secs\n-------------------------------\n\n", MinTime, TotalTime / NumRuns);
+	printf("%s\nmin: %f secs, avg: %f secs\n-------------------------------\n\n", TCHAR_TO_ANSI(TestName), MinTime, TotalTime / NumRuns);
 #endif
 }
 
-#define UE_BENCHMARK(NumRuns, ...) Benchmark<NumRuns>(TEXT(#__VA_ARGS__), __VA_ARGS__)
+#define UE_BENCHMARK(NumRuns, ...) { TRACE_CPUPROFILER_EVENT_SCOPE(TEXT(#__VA_ARGS__)); Benchmark<NumRuns>(TEXT(#__VA_ARGS__), __VA_ARGS__); }
 
