Add StringBuilder ReplaceAt, InsertAt, RemoveAt and Prepend

diff --git a/Engine/Source/Runtime/Core/Private/Tests/Misc/StringBuilderTest.cpp b/Engine/Source/Runtime/Core/Private/Tests/Misc/StringBuilderTest.cpp
--- a/Engine/Source/Runtime/Core/Private/Tests/Misc/StringBuilderTest.cpp
+++ b/Engine/Source/Runtime/Core/Private/Tests/Misc/StringBuilderTest.cpp
@@ -1,9 +1,9 @@
 // Copyright Epic Games, Inc. All Rights Reserved.
 
 #include "Misc/StringBuilder.h"
-
-#include "Containers/UnrealString.h"
 #include "Misc/AutomationTest.h"
+#include "Containers/UnrealString.h"
+#include "String/Find.h"
 
 static_assert(TIsSame<typename FStringBuilderBase::ElementType, TCHAR>::Value, "FStringBuilderBase must use TCHAR.");
 static_assert(TIsSame<typename FAnsiStringBuilderBase::ElementType, ANSICHAR>::Value, "FAnsiStringBuilderBase must use ANSICHAR.");
@@ -23,7 +23,7 @@ static_assert(TIsContiguousContainer<TWideStringBuilder<128>>::Value, "TWideStri
 
 #if WITH_DEV_AUTOMATION_TESTS
 
-IMPLEMENT_SIMPLE_AUTOMATION_TEST(FStringBuilderTestAppendString, "System.Core.StringBuilder.AppendString", EAutomationTestFlags::ApplicationContextMask | EAutomationTestFlags::EngineFilter)
+IMPLEMENT_SIMPLE_AUTOMATION_TEST(FStringBuilderTestAppendString, "System.Core.StringBuilder.AppendString", EAutomationTestFlags::ApplicationContextMask | EAutomationTestFlags::SmokeFilter)
 bool FStringBuilderTestAppendString::RunTest(const FString& Parameters)
 {
 	// Append Char
@@ -97,6 +97,128 @@ bool FStringBuilderTestAppendString::RunTest(const FString& Parameters)
 		TestEqual(TEXT("Append Char Array"), FAnsiStringView(AnsiBuilder), "ABC"_ASV);
 	}
 
+
+	// Simple ReplaceAt
+	{
+		TAnsiStringBuilder<4> Builder;
+		Builder.ReplaceAt(0, 0, FAnsiStringView());
+		TestEqual(TEXT("Replace nothing empty"), Builder.ToString(), "");
+
+		Builder << 'a';
+		
+		Builder.ReplaceAt(0, 0, FAnsiStringView());
+		TestEqual(TEXT("Replace nothing non-empty"), Builder.ToString(), "a");
+
+		Builder.ReplaceAt(0, 1, FAnsiStringView("b"));
+		TestEqual(TEXT("Replace single char"), Builder.ToString(), "b");
+	}
+
+	// Advanced ReplaceAt
+	auto TestReplace = [&](FAnsiStringView Original, FAnsiStringView SearchFor, FAnsiStringView ReplaceWith, FAnsiStringView Expected)
+	{
+		int32 ReplacePos = UE::String::FindFirst(Original, SearchFor);
+		check(ReplacePos != INDEX_NONE);
+
+		TAnsiStringBuilder<4> Builder;
+		Builder << Original;
+		Builder.ReplaceAt(ReplacePos, SearchFor.Len(), ReplaceWith);
+
+		TestEqual(TEXT("Replace"), FAnsiStringView(Builder), Expected);
+	};
+
+	// Test single character erase
+	TestReplace(".foo", ".", "", "foo");
+	TestReplace("f.oo", ".", "", "foo");
+	TestReplace("foo.", ".", "", "foo");
+		
+	// Test multi character erase
+	TestReplace("FooBar", "Bar", "", "Foo");
+	TestReplace("FooBar", "Foo", "", "Bar");
+	TestReplace("FooBar", "Foo", "fOOO", "fOOOBar");
+
+	// Test replace everything
+	TestReplace("Foo", "Foo", "", "");
+	TestReplace("Foo", "Foo", "Bar", "Bar");
+	TestReplace("Foo", "Foo", "0123456789", "0123456789");
+
+	// Test expanding replace
+	TestReplace(".foo", ".", "<dot>", "<dot>foo");
+	TestReplace("foo.", ".", "<dot>", "foo<dot>");
+	TestReplace("f.oo", ".", "<dot>", "f<dot>oo");
+
+	// Test shrinking replace
+	TestReplace("aabbcc", "aa", "A", "Abbcc");
+	TestReplace("aabbcc", "bb", "B", "aaBcc");
+	TestReplace("aabbcc", "cc", "C", "aabbC");
+
+	// Prepend
+	{
+		TAnsiStringBuilder<4> Builder;
+		Builder.Prepend("");
+		TestEqual(TEXT("Prepend nothing to empty"), Builder.Len(), 0);
+		
+		Builder.Prepend("e");
+		TestEqual(TEXT("Prepend single characer"), Builder.ToString(), "e");
+		
+		Builder.Prepend("abcd");
+		TestEqual(TEXT("Prepend substring"), Builder.ToString(), "abcde");
+
+		Builder.Prepend("");
+		TestEqual(TEXT("Prepend nothing to non-empty"), Builder.ToString(), "abcde");
+	}
+	
+	// InsertAt
+	{
+		TAnsiStringBuilder<4> Builder;
+		Builder.InsertAt(0, "");
+		TestEqual(TEXT("Insert nothing to empty"), Builder.Len(), 0);
+		
+		Builder.InsertAt(0, "d");
+		TestEqual(TEXT("Insert first char"), Builder.ToString(), "d");
+		
+		Builder.InsertAt(0, "c");
+		Builder.InsertAt(0, "a");
+		Builder.InsertAt(1, "b");
+		Builder.InsertAt(4, "e");
+		TestEqual(TEXT("Insert single char"), Builder.ToString(), "abcde");
+		
+		Builder.InsertAt(3, "__");
+		Builder.InsertAt(0, "__");
+		Builder.InsertAt(Builder.Len(), "__");
+		TestEqual(TEXT("Insert substrings"), Builder.ToString(), "__abc__de__");
+
+		Builder.InsertAt(Builder.Len(), "");
+		TestEqual(TEXT("Insert nothing"), Builder.ToString(), "__abc__de__");
+	}
+
+	// RemoveAt
+	{
+		TAnsiStringBuilder<4> Builder;
+		Builder << "0123456789";
+		Builder.RemoveAt(0, 0);
+		Builder.RemoveAt(Builder.Len(), 0);
+		Builder.RemoveAt(Builder.Len() / 2, 0);
+		TestEqual(TEXT("Remove nothing"), Builder.ToString(), "0123456789");
+		
+		Builder.RemoveAt(Builder.Len() - 1, 1);
+		TestEqual(TEXT("Remove last char"), Builder.ToString(), "012345678");
+		
+		Builder.RemoveAt(0, 1);
+		TestEqual(TEXT("Remove first char"), Builder.ToString(), "12345678");
+		
+		Builder.RemoveAt(4, 2);
+		TestEqual(TEXT("Remove middle"), Builder.ToString(), "123478");
+
+		Builder.RemoveAt(4, 2);
+		TestEqual(TEXT("Remove end"), Builder.ToString(), "1234");
+		
+		Builder.RemoveAt(0, 2);
+		TestEqual(TEXT("Remove start"), Builder.ToString(), "34");
+
+		Builder.RemoveAt(0, 2);
+		TestEqual(TEXT("Remove start"), Builder.ToString(), "");
+	}
+
 	return true;
 }
 
diff --git a/Engine/Source/Runtime/Core/Public/Misc/StringBuilder.h b/Engine/Source/Runtime/Core/Public/Misc/StringBuilder.h
--- a/Engine/Source/Runtime/Core/Public/Misc/StringBuilder.h
+++ b/Engine/Source/Runtime/Core/Public/Misc/StringBuilder.h
@@ -193,6 +193,55 @@ public:
 		return *this;
 	}
 
+	/** Replace characters at given position and length with substring */
+	void ReplaceAt(int32 Pos, int32 RemoveLen, ViewType Str)
+	{
+		check(Pos >= 0);
+		check(RemoveLen >= 0);
+		check(Pos + RemoveLen <= Len());
+
+		const int DeltaLen = Str.Len() - RemoveLen;		
+		if (DeltaLen < 0)
+		{
+			CurPos += DeltaLen;
+
+			for (CharType* It = Base + Pos, *NewEnd = CurPos; It != NewEnd; ++It)
+			{
+				*It = *(It - DeltaLen);
+			}
+		}
+		else if (DeltaLen > 0)
+		{
+			EnsureCapacity(Len() + DeltaLen);
+			CurPos += DeltaLen;
+
+			for (CharType* It = CurPos - 1, *StopIt = Base + Pos + Str.Len() - 1; It != StopIt; --It)
+			{
+				*It = *(It - DeltaLen);
+			}
+		}
+		
+		FMemory::Memcpy(Base + Pos, Str.GetData(), Str.Len() * sizeof(CharType));
+	}
+
+	/** Insert substring at given position */
+	void InsertAt(int32 Pos, ViewType Str)
+	{
+		ReplaceAt(Pos, 0, Str);
+	}
+
+	/** Remove characters at given position */
+	void RemoveAt(int32 Pos, int32 RemoveLen)
+	{
+		ReplaceAt(Pos, RemoveLen, ViewType());
+	}
+
+	/** Insert prefix */
+	void Prepend(ViewType Str)
+	{
+		ReplaceAt(0, 0, Str);
+	}
+
 	/**
 	 * Append every element of the range to the builder, separating the elements by the delimiter.
 	 *
