StatsThread replaced by a Pipe:
* the thread was removed and its surrounding was cleaned up a bit
* ENamedThreads::StatsThread (and relatives) deprecated
* support for single-threaded execution looks a bit messy because it still have to be based on the old TaskGraph API until we convert the game thread to the new API (which will be one of the very last things)
* despite the direct usage of ENamedThreads::StatsThread from an external code is highly unlikely, to preserve backward compatiblity internally TaskGraph redirects all remaining stats tasks to StatsPipe

diff --git a/Engine/Plugins/Messaging/MessagingDebugger/Source/MessagingDebugger/Private/Widgets/MessageDetails/SMessagingDispatchStateTableRow.h b/Engine/Plugins/Messaging/MessagingDebugger/Source/MessagingDebugger/Private/Widgets/MessageDetails/SMessagingDispatchStateTableRow.h
--- a/Engine/Plugins/Messaging/MessagingDebugger/Source/MessagingDebugger/Private/Widgets/MessageDetails/SMessagingDispatchStateTableRow.h
+++ b/Engine/Plugins/Messaging/MessagingDebugger/Source/MessagingDebugger/Private/Widgets/MessageDetails/SMessagingDispatchStateTableRow.h
@@ -189,12 +189,6 @@ protected:
 			return LOCTEXT("ActualRenderingThread", "ActualRenderingThread");
 			break;
 
-#if STATS
-		case ENamedThreads::StatsThread:
-			return LOCTEXT("StatsThread", "StatsThread");
-			break;
-#endif
-
 		default:
 			return LOCTEXT("UnknownThread", "Unknown");
 		}
diff --git a/Engine/Plugins/Messaging/MessagingDebugger/Source/MessagingDebugger/Private/Widgets/MessageDetails/SMessagingMessageDetails.cpp b/Engine/Plugins/Messaging/MessagingDebugger/Source/MessagingDebugger/Private/Widgets/MessageDetails/SMessagingMessageDetails.cpp
--- a/Engine/Plugins/Messaging/MessagingDebugger/Source/MessagingDebugger/Private/Widgets/MessageDetails/SMessagingMessageDetails.cpp
+++ b/Engine/Plugins/Messaging/MessagingDebugger/Source/MessagingDebugger/Private/Widgets/MessageDetails/SMessagingMessageDetails.cpp
@@ -233,16 +233,6 @@ FText SMessagingMessageDetails::HandleSenderThreadText() const
 			return LOCTEXT("ActualRenderingThread_Local", "ActualRenderingThread_Local");
 			break;
 
-#if STATS
-		case ENamedThreads::StatsThread:
-			return LOCTEXT("StatsThread", "StatsThread");
-			break;
-
-		case ENamedThreads::StatsThread_Local:
-			return LOCTEXT("StatsThread_Local", "StatsThread_Local");
-			break;
-#endif
-
 		default:
 			return LOCTEXT("UnknownThread", "Unknown");
 		}
diff --git a/Engine/Source/Developer/ProfilerService/Private/ProfilerServiceManager.cpp b/Engine/Source/Developer/ProfilerService/Private/ProfilerServiceManager.cpp
--- a/Engine/Source/Developer/ProfilerService/Private/ProfilerServiceManager.cpp
+++ b/Engine/Source/Developer/ProfilerService/Private/ProfilerServiceManager.cpp
@@ -7,7 +7,7 @@
 #include "Serialization/MemoryReader.h"
 #include "Stats/StatsData.h"
 #include "Stats/StatsFile.h"
-
+#include "Tasks/Pipe.h"
 
 DEFINE_LOG_CATEGORY(LogProfilerService);
 
@@ -126,7 +126,7 @@ TSharedPtr<IProfilerServiceManager> FProfilerServiceManager::CreateSharedService
 }
 
 
-void FProfilerServiceManager::AddNewFrameHandleStatsThread()
+void FProfilerServiceManager::AddNewFrameHandleStatsPipe()
 {
 #if	STATS
 	const FStatsThreadState& Stats = FStatsThreadState::GetLocalState();
@@ -137,7 +137,7 @@ void FProfilerServiceManager::AddNewFrameHandleStatsThread()
 }
 
 
-void FProfilerServiceManager::RemoveNewFrameHandleStatsThread()
+void FProfilerServiceManager::RemoveNewFrameHandleStatsPipe()
 {
 #if	STATS
 	const FStatsThreadState& Stats = FStatsThreadState::GetLocalState();
@@ -147,6 +147,7 @@ void FProfilerServiceManager::RemoveNewFrameHandleStatsThread()
 #endif //STATS
 }
 
+extern CORE_API UE::Tasks::FPipe GStatsPipe;
 
 void FProfilerServiceManager::SetPreviewState( const FMessageAddress& ClientAddress, const bool bRequestedPreviewState )
 {
@@ -163,12 +164,18 @@ void FProfilerServiceManager::SetPreviewState( const FMessageAddress& ClientAddr
 				// Enable stat capture.
 				if (PreviewClients.Num() == 0)
 				{
-					FSimpleDelegateGraphTask::CreateAndDispatchWhenReady
-					(
-						FSimpleDelegateGraphTask::FDelegate::CreateRaw( this, &FProfilerServiceManager::AddNewFrameHandleStatsThread ),
-						TStatId(), nullptr,
-						FPlatformProcess::SupportsMultithreading() ? ENamedThreads::StatsThread : ENamedThreads::GameThread
-					);
+					if (FPlatformProcess::SupportsMultithreading())
+					{
+						GStatsPipe.Launch(UE_SOURCE_LOCATION, [this] { AddNewFrameHandleStatsPipe(); });
+					}
+					else
+					{
+						FSimpleDelegateGraphTask::CreateAndDispatchWhenReady
+						(
+							FSimpleDelegateGraphTask::FDelegate::CreateRaw(this, &FProfilerServiceManager::AddNewFrameHandleStatsPipe),
+							TStatId(), nullptr, ENamedThreads::GameThread
+						);
+					}
 				}
 				PreviewClients.Add(ClientAddress);
 				Client->Preview = true;
@@ -183,13 +190,18 @@ void FProfilerServiceManager::SetPreviewState( const FMessageAddress& ClientAddr
 				// Disable stat capture.
 				if (PreviewClients.Num() == 0)
 				{
-					FSimpleDelegateGraphTask::CreateAndDispatchWhenReady
-					(
-						FSimpleDelegateGraphTask::FDelegate::CreateRaw( this, &FProfilerServiceManager::RemoveNewFrameHandleStatsThread ),
-						TStatId(), nullptr,
-						FPlatformProcess::SupportsMultithreading() ? ENamedThreads::StatsThread : ENamedThreads::GameThread
-					);
-					
+					if (FPlatformProcess::SupportsMultithreading())
+					{
+						GStatsPipe.Launch(UE_SOURCE_LOCATION, [this] { RemoveNewFrameHandleStatsPipe(); });
+					}
+					else
+					{
+						FSimpleDelegateGraphTask::CreateAndDispatchWhenReady
+						(
+							FSimpleDelegateGraphTask::FDelegate::CreateRaw(this, &FProfilerServiceManager::RemoveNewFrameHandleStatsPipe),
+							TStatId(), nullptr, ENamedThreads::GameThread
+						);
+					}
 				}	
 			}
 		}
diff --git a/Engine/Source/Developer/ProfilerService/Private/ProfilerServiceManager.h b/Engine/Source/Developer/ProfilerService/Private/ProfilerServiceManager.h
--- a/Engine/Source/Developer/ProfilerService/Private/ProfilerServiceManager.h
+++ b/Engine/Source/Developer/ProfilerService/Private/ProfilerServiceManager.h
@@ -205,9 +205,9 @@ private:
 
 #endif //STATS
 
-	void AddNewFrameHandleStatsThread();
+	void AddNewFrameHandleStatsPipe();
 
-	void RemoveNewFrameHandleStatsThread();
+	void RemoveNewFrameHandleStatsPipe();
 
 	/** Holds the messaging endpoint. */
 	TSharedPtr<FMessageEndpoint, ESPMode::ThreadSafe> MessageEndpoint;
diff --git a/Engine/Source/Runtime/Core/Private/Async/TaskGraph.cpp b/Engine/Source/Runtime/Core/Private/Async/TaskGraph.cpp
--- a/Engine/Source/Runtime/Core/Private/Async/TaskGraph.cpp
+++ b/Engine/Source/Runtime/Core/Private/Async/TaskGraph.cpp
@@ -34,6 +34,7 @@
 #include "Misc/ConfigCacheIni.h"
 
 #include "Async/Fundamental/Scheduler.h"
+#include "Tasks/Pipe.h"
 
 DEFINE_LOG_CATEGORY_STATIC(LogTaskGraph, Log, All);
 
@@ -671,7 +672,9 @@ public:
 			}
 			// else StatName = none, we need to let the scope empty so that the render thread submits tasks in a timely manner. 
 		}
+PRAGMA_DISABLE_DEPRECATION_WARNINGS
 		else if (ThreadId != ENamedThreads::StatsThread)
+PRAGMA_ENABLE_DEPRECATION_WARNINGS
 		{
 			StatName = GET_STATID(STAT_TaskGraph_OtherTasks);
 			StallStatId = GET_STATID(STAT_TaskGraph_OtherStalls);
@@ -826,7 +829,9 @@ private:
 			return TEXT("Audio Thread");
 		}
 #if STATS
+PRAGMA_DISABLE_DEPRECATION_WARNINGS
 		else if (ThreadId == ENamedThreads::StatsThread)
+PRAGMA_ENABLE_DEPRECATION_WARNINGS
 		{
 			return TEXT("Stats Thread");
 		}
@@ -1377,6 +1382,30 @@ public:
 		FPlatformTLS::FreeTlsSlot(PerThreadIDTLSSlot);
 	}
 
+	// StatsThread was removed and replaced by a pipe. While all engine code was converted where's a tiny possibilty that 
+	// an external code is using it directly. Redirect all stats tasks to the pipe
+	static void RedirectStatsTasksToPipe(FBaseGraphTask* Task)
+	{
+#if STATS
+PRAGMA_DISABLE_DEPRECATION_WARNINGS
+		if (ENamedThreads::GetThreadIndex(Task->ThreadToExecuteOn) == ENamedThreads::StatsThread)
+PRAGMA_ENABLE_DEPRECATION_WARNINGS
+		{
+			checkf(ENamedThreads::GetQueueIndex(Task->ThreadToExecuteOn) == ENamedThreads::MainQueue, TEXT("`StatsThread_Local` is not supported and unlikely to be used intentionally. Check that `Task->ThreadToExecuteOn` value is not rubbish: %d"), Task->ThreadToExecuteOn);
+
+			extern CORE_API UE::Tasks::FPipe GStatsPipe;
+			GStatsPipe.Launch(UE_SOURCE_LOCATION,
+				[Task]
+				{
+					TArray<FBaseGraphTask*> Dummy;
+					Task->Execute(Dummy, ENamedThreads::AnyThread, true);
+				}
+			);
+			return;
+		}
+#endif
+	}
+
 	// API inherited from FTaskGraphInterface
 
 	/** 
@@ -1389,6 +1418,8 @@ public:
 	{
 		TASKGRAPH_SCOPE_CYCLE_COUNTER(2, STAT_TaskGraph_QueueTask);
 
+		RedirectStatsTasksToPipe(Task);
+
 		if (ENamedThreads::GetThreadIndex(ThreadToExecuteOn) == ENamedThreads::AnyThread)
 		{
 			TASKGRAPH_SCOPE_CYCLE_COUNTER(3, STAT_TaskGraph_QueueTask_AnyThread);
@@ -1889,6 +1920,8 @@ public:
 private:
 	void QueueTask(class FBaseGraphTask* Task, bool bWakeUpWorker, ENamedThreads::Type InThreadToExecuteOn, ENamedThreads::Type InCurrentThreadIfKnown) override
 	{
+		FTaskGraphImplementation::RedirectStatsTasksToPipe(Task);
+
 		if (ENamedThreads::GetThreadIndex(InThreadToExecuteOn) == ENamedThreads::AnyThread)
 		{
 			uint32 ThreadPriority = GetThreadPriorityIndex(InThreadToExecuteOn);
@@ -2520,10 +2553,12 @@ void FTaskGraphInterface::BroadcastSlow_OnlyUseForSpecialPurposes(bool bDoTaskTh
 
 	FGraphEventArray Tasks;
 #if STATS
+PRAGMA_DISABLE_DEPRECATION_WARNINGS
 	if (FTaskGraphInterface::Get().IsThreadProcessingTasks(ENamedThreads::StatsThread))
 	{
 		Tasks.Add(TGraphTask<FBroadcastTask>::CreateTask().ConstructAndDispatchWhenReady(Callback, StartTime, TEXT("Stats"), ENamedThreads::SetTaskPriority(ENamedThreads::StatsThread, ENamedThreads::HighTaskPriority), nullptr, nullptr, nullptr));
 	}
+PRAGMA_ENABLE_DEPRECATION_WARNINGS
 #endif
 	if (IsRHIThreadRunning())
 	{
diff --git a/Engine/Source/Runtime/Core/Private/Stats/Stats2.cpp b/Engine/Source/Runtime/Core/Private/Stats/Stats2.cpp
--- a/Engine/Source/Runtime/Core/Private/Stats/Stats2.cpp
+++ b/Engine/Source/Runtime/Core/Private/Stats/Stats2.cpp
@@ -24,6 +24,7 @@
 #include "HAL/IConsoleManager.h"
 #include "Async/TaskGraphInterfaces.h"
 #include "HAL/LowLevelMemTracker.h"
+#include "Tasks/Pipe.h"
 
 /*-----------------------------------------------------------------------------
 	Global
@@ -40,12 +41,15 @@ struct FStats2Globals
 	}
 };
 
+#if STATS
+CORE_API UE::Tasks::FPipe GStatsPipe{ UE_SOURCE_LOCATION };
+#endif
+
 static struct FForceInitAtBootFStats2 : public TForceInitAtBoot<FStats2Globals>
 {} FForceInitAtBootFStats2;
 
 DECLARE_DWORD_COUNTER_STAT( TEXT("Frame Packets Received"),STAT_StatFramePacketsRecv,STATGROUP_StatSystem);
 
-DECLARE_CYCLE_STAT(TEXT("WaitForStats"),STAT_WaitForStats,STATGROUP_Engine);
 DECLARE_CYCLE_STAT(TEXT("StatsNew Tick"),STAT_StatsNewTick,STATGROUP_StatSystem);
 DECLARE_CYCLE_STAT(TEXT("Parse Meta"),STAT_StatsNewParseMeta,STATGROUP_StatSystem);
 DECLARE_CYCLE_STAT(TEXT("Scan For Advance"),STAT_ScanForAdvance,STATGROUP_StatSystem);
@@ -861,14 +865,14 @@ bool FStatNameAndInfo::GetSortByNameFrom(FName InLongName)
 }
 
 /*-----------------------------------------------------------------------------
-	FStatsThread
+	FStatsProcessor
 -----------------------------------------------------------------------------*/
 
 static TAutoConsoleVariable<int32> CVarDumpStatPackets(	TEXT("DumpStatPackets"),0,	TEXT("If true, dump stat packets."));
 
 
 /** The rendering thread runnable object. */
-class FStatsThread : public FRunnable, FSingleThreadRunnable
+class FStatsProcessor
 {
 	/** Array of stat packets, queued data to be processed on this thread. */
 	FStatPacketArray IncomingData;
@@ -878,52 +882,53 @@ class FStatsThread : public FRunnable, FSingleThreadRunnable
 
 	/** Whether we are ready to process the packets, sets by game or render packets. */
 	bool bReadyToProcess;
+
 public:
+	static FStatsProcessor& Get()
+	{
+		static FStatsProcessor Singleton;
+		return Singleton;
+	}
 
 	/** Default constructor. */
-	FStatsThread()
+	FStatsProcessor()
 		: State(FStatsThreadState::GetLocalState())
 		, bReadyToProcess(false)
 	{
-		check(IsInGameThread());
 	}
 
-	/**
-	 * Returns a pointer to the single threaded interface when multithreading is disabled.
-	 */
-	virtual FSingleThreadRunnable* GetSingleThreadInterface() override
+	/** Received a stat packet from other thread and add to the processing queue. */
+	void StatMessage(FStatPacket* Packet)
 	{
-		return this;
-	}
+		LLM_SCOPE(ELLMTag::Stats);
 
-	/** Attaches to the task graph stats thread, all processing will be handled by the task graph. */
-	virtual uint32 Run() override
-	{
-		FThreadStats::GetThreadStats()->bIsStatsThread = true;
-		FMemory::SetupTLSCachesOnCurrentThread();
-		FTaskGraphInterface::Get().AttachToThread(ENamedThreads::StatsThread);
-		FTaskGraphInterface::Get().ProcessThreadUntilRequestReturn(ENamedThreads::StatsThread);
-		FMemory::ClearAndDisableTLSCachesOnCurrentThread();
-		return 0;
+		if (CVarDumpStatPackets.GetValueOnAnyThread())
+		{
+			UE_LOG(LogStats, Log, TEXT("Packet from %x with %d messages"), Packet->ThreadId, Packet->StatMessages.Num());
+		}
+
+		bReadyToProcess = Packet->ThreadType != EThreadType::Other;
+		IncomingData.Packets.Add(Packet);
+		State.NumStatMessages.Add(Packet->StatMessages.Num());
+
+		Process();
 	}
 
-	/** Tick function. */
-	virtual void Tick() override
+private:
+	void Process()
 	{
-		LLM_SCOPE(ELLMTag::Stats);
-
 		static double LastTime = -1.0;
 		bool bShouldProcess = false;
 
 		const int32 MaxIncomingPackets = 16;
-		if( FThreadStats::bIsRawStatsActive )
+		if (FThreadStats::bIsRawStatsActive)
 		{
 			// For raw stats we process every 24MB of packet data to minimize the stats messages memory usage.
 			//const bool bShouldProcessRawStats = IncomingData.Packets.Num() > 10;
-			const int32 MaxIncomingMessages = 24*1024*1024/sizeof(FStatMessage);
+			const int32 MaxIncomingMessages = 24 * 1024 * 1024 / sizeof(FStatMessage);
 
 			int32 IncomingDataMessages = 0;
-			for( FStatPacket* Packet : IncomingData.Packets )
+			for (FStatPacket* Packet : IncomingData.Packets)
 			{
 				IncomingDataMessages += Packet->StatMessages.Num();
 			}
@@ -937,7 +942,7 @@ public:
 			bShouldProcess = bReadyToProcess && (FPlatformTime::Seconds() - LastTime > 0.005f || IncomingData.Packets.Num() > MaxIncomingPackets || FStats::EnabledForCommandlet());
 		}
 
-		if( bShouldProcess )
+		if (bShouldProcess)
 		{
 			SCOPE_CYCLE_COUNTER(STAT_StatsNewTick);
 
@@ -965,7 +970,7 @@ public:
 				State.ScanForAdvance(NowData);
 			}
 
-			if( FThreadStats::bIsRawStatsActive )
+			if (FThreadStats::bIsRawStatsActive)
 			{
 				// Process raw stats.
 				State.ProcessRawStats(NowData);
@@ -982,88 +987,6 @@ public:
 			LastTime = FPlatformTime::Seconds();
 		}
 	}
-
-	/** Accesses singleton. */
-	static FStatsThread& Get()
-	{
-		static FStatsThread Singleton;
-		return Singleton;
-	}
-
-	/** Received a stat packet from other thread and add to the processing queue. */
-	void StatMessage(FStatPacket* Packet)
-	{
-		LLM_SCOPE(ELLMTag::Stats);
-
-		if (CVarDumpStatPackets.GetValueOnAnyThread())
-		{
-			UE_LOG(LogStats, Log, TEXT("Packet from %x with %d messages"), Packet->ThreadId, Packet->StatMessages.Num());
-		}
-
-		bReadyToProcess = Packet->ThreadType != EThreadType::Other;
-		IncomingData.Packets.Add(Packet);
-		State.NumStatMessages.Add(Packet->StatMessages.Num());
-
-		Tick();
-	}
-
-	void SelfStatMessage(FStatPacket* Packet)
-	{
-		if (CVarDumpStatPackets.GetValueOnAnyThread())
-		{
-			UE_LOG(LogStats, Log, TEXT("Self Packet from %x with %d messages"), Packet->ThreadId, Packet->StatMessages.Num());
-		}
-
-		IncomingData.Packets.Add(Packet);
-		State.NumStatMessages.Add(Packet->StatMessages.Num());
-	}
-
-	/** Start a stats runnable thread. */
-	void Start()
-	{
-		Thread = FRunnableThread::Create(this, TEXT("StatsThread"), 512 * 1024, TPri_BelowNormal, FPlatformAffinity::GetStatsThreadMask());
-		check(Thread);
-	}
-
-	/** Ends the stats runnable thread. */
-	void End()
-	{
-		delete Thread;
-		Thread = nullptr;
-	}
-
-private:
-	FRunnableThread* Thread;
-};
-
-/*-----------------------------------------------------------------------------
-	FStatMessagesTask
------------------------------------------------------------------------------*/
-
-// not using a delegate here to allow higher performance since we may end up sending a lot of small message arrays to the thread.
-class FStatMessagesTask
-{
-	FStatPacket* Packet;
-public:
-	FStatMessagesTask(FStatPacket* InPacket)
-		: Packet(InPacket)
-	{
-	}
-	FORCEINLINE TStatId GetStatId() const
-	{
-		return TStatId(); // we don't want to record this or it spams the stat system; we cover this time when we tick the stats system
-	}
-	static ENamedThreads::Type GetDesiredThread()
-	{
-		return FPlatformProcess::SupportsMultithreading() ? ENamedThreads::StatsThread : ENamedThreads::GameThread;
-	}
-	static ESubsequentsMode::Type GetSubsequentsMode() { return ESubsequentsMode::TrackSubsequents; }
-
-	void DoTask(ENamedThreads::Type CurrentThread, const FGraphEventRef& MyCompletionGraphEvent)
-	{
-		FStatsThread::Get().StatMessage(Packet);
-		Packet = NULL;
-	}
 };
 
 /*-----------------------------------------------------------------------------
@@ -1117,8 +1040,7 @@ FThreadStats::FThreadStats():
 	bWaitForExplicitFlush(0),
 	MemoryMessageScope(0),
 	bReentranceGuard(false),
-	bSawExplicitFlush(false),
-	bIsStatsThread(false)
+	bSawExplicitFlush(false)
 {
 	Packet.SetThreadProperties();
 
@@ -1132,8 +1054,7 @@ FThreadStats::FThreadStats( EConstructor ):
 	bWaitForExplicitFlush(0),
 	MemoryMessageScope(0),
 	bReentranceGuard(false),
-	bSawExplicitFlush(false),
-	bIsStatsThread(false)
+	bSawExplicitFlush(false)
 {}
 
 void FThreadStats::CheckEnable()
@@ -1166,6 +1087,18 @@ void FThreadStats::Flush( bool bHasBrokenCallstacks /*= false*/, bool bForceFlus
 	}
 }
 
+void FThreadStats::SendMessage_Async(FStatPacket* ToSend)
+{
+	if (FPlatformProcess::SupportsMultithreading())
+	{
+		GStatsPipe.Launch(UE_SOURCE_LOCATION, [ToSend] { FStatsProcessor::Get().StatMessage(ToSend); });
+	}
+	else
+	{
+		FFunctionGraphTask::CreateAndDispatchWhenReady([ToSend] { FStatsProcessor::Get().StatMessage(ToSend); }, TStatId{}, nullptr, ENamedThreads::GameThread);
+	}
+}
+
 void FThreadStats::FlushRegularStats( bool bHasBrokenCallstacks, bool bForceFlush )
 {
 	LLM_SCOPE(ELLMTag::Stats);
@@ -1225,14 +1158,7 @@ void FThreadStats::FlushRegularStats( bool bHasBrokenCallstacks, bool bForceFlus
 			}
 			Packet.StatMessages.Empty(MaxPresize);
 		}
-		if (bIsStatsThread)
-		{
-			FStatsThread::Get().SelfStatMessage(ToSend);
-		}
-		else
-		{
-			TGraphTask<FStatMessagesTask>::CreateTask().ConstructAndDispatchWhenReady(ToSend);
-		}
+		SendMessage_Async(ToSend);
 		UpdateExplicitFlush();
 	}
 }
@@ -1270,14 +1196,7 @@ void FThreadStats::FlushRawStats( bool bHasBrokenCallstacks /*= false*/, bool bF
 
 		check(!Packet.StatMessages.Num());
 
-		if (bIsStatsThread)
-		{
-			FStatsThread::Get().SelfStatMessage(ToSend);
-		}
-		else
-		{
-			TGraphTask<FStatMessagesTask>::CreateTask().ConstructAndDispatchWhenReady(ToSend);
-		}
+		SendMessage_Async(ToSend);
 		UpdateExplicitFlush();
 
 		const float NumMessagesAsMB = float(NumMessages * sizeof(FStatMessage)) / 1024.0f / 1024.0f;
@@ -1386,8 +1305,6 @@ void FThreadStats::StartThread()
 		TlsSlot = FPlatformTLS::AllocTlsSlot();
 		check(TlsSlot);
 	}
-	FStatsThread::Get();
-	FStatsThread::Get().Start();
 
 	check(IsThreadingReady());
 	CheckEnable();
@@ -1403,7 +1320,6 @@ void FThreadStats::StartThread()
 	UE_LOG( LogStats, Log, TEXT( "Stats thread started at %f" ), FPlatformTime::Seconds() - GStartTime );
 }
 
-static FGraphEventRef LastFramesEvents[MAX_STAT_LAG];
 static int32 CurrentEventIndex = 0;
 
 void FThreadStats::StopThread()
@@ -1423,14 +1339,12 @@ void FThreadStats::StopThread()
 		FThreadStats::MasterDisableForever();
 
 		WaitForStats();
-		for (int32 Index = 0; Index < MAX_STAT_LAG; Index++)
-		{
-			LastFramesEvents[Index] = NULL;
-		}
-		FGraphEventRef QuitTask = TGraphTask<FReturnGraphTask>::CreateTask(NULL, ENamedThreads::GameThread).ConstructAndDispatchWhenReady(FPlatformProcess::SupportsMultithreading() ? ENamedThreads::StatsThread : ENamedThreads::GameThread);
-		FTaskGraphInterface::Get().WaitUntilTaskCompletes(QuitTask, ENamedThreads::GameThread_Local);
+	}
 
-		FStatsThread::Get().End();
+	// wait for the pipe to complete all currently piped tasks
+	while (GStatsPipe.HasWork())
+	{
+		FPlatformProcess::Yield();
 	}
 }
 
@@ -1444,19 +1358,37 @@ void FThreadStats::WaitForStats()
 	check(IsInGameThread());
 	if (IsThreadingReady() && !bMasterDisableForever)
 	{
+		DECLARE_CYCLE_STAT(TEXT("WaitForStats"), STAT_WaitForStats, STATGROUP_Engine);
+
+		int32 EventIndex = (CurrentEventIndex + MAX_STAT_LAG - 1) % MAX_STAT_LAG;
+
+		if (FPlatformProcess::SupportsMultithreading())
 		{
-			SCOPE_CYCLE_COUNTER(STAT_WaitForStats);
-			if (LastFramesEvents[(CurrentEventIndex + MAX_STAT_LAG - 1) % MAX_STAT_LAG].GetReference())
+			static UE::Tasks::TTask<void> LastFramesEvents[MAX_STAT_LAG];
+		
 			{
-				FTaskGraphInterface::Get().WaitUntilTaskCompletes(LastFramesEvents[(CurrentEventIndex + MAX_STAT_LAG - 1) % MAX_STAT_LAG], ENamedThreads::GameThread_Local);
+				SCOPE_CYCLE_COUNTER(STAT_WaitForStats);
+				LastFramesEvents[EventIndex].Wait();
 			}
+
+			LastFramesEvents[EventIndex] = GStatsPipe.Launch(UE_SOURCE_LOCATION, [] {});
 		}
+		else
+		{
+			static FGraphEventRef LastFramesEvents[MAX_STAT_LAG];
 
-		DECLARE_CYCLE_STAT(TEXT("FNullGraphTask.StatWaitFence"),
-			STAT_FNullGraphTask_StatWaitFence,
-			STATGROUP_TaskGraphTasks);
+			{
+				SCOPE_CYCLE_COUNTER(STAT_WaitForStats);
+				if (LastFramesEvents[EventIndex].GetReference())
+				{
+					FTaskGraphInterface::Get().WaitUntilTaskCompletes(LastFramesEvents[EventIndex], ENamedThreads::GameThread_Local);
+				}
+			}
+
+			LastFramesEvents[EventIndex] = TGraphTask<FNullGraphTask>::CreateTask(NULL, ENamedThreads::GameThread).
+				ConstructAndDispatchWhenReady(TStatId{}, ENamedThreads::GameThread);
+		}
 
-		LastFramesEvents[(CurrentEventIndex + MAX_STAT_LAG - 1) % MAX_STAT_LAG] = TGraphTask<FNullGraphTask>::CreateTask(NULL, ENamedThreads::GameThread).ConstructAndDispatchWhenReady(GET_STATID(STAT_FNullGraphTask_StatWaitFence), FPlatformProcess::SupportsMultithreading() ? ENamedThreads::StatsThread : ENamedThreads::GameThread);
 		CurrentEventIndex++;
 
 #if	!UE_BUILD_SHIPPING
diff --git a/Engine/Source/Runtime/Core/Private/Stats/StatsCommand.cpp b/Engine/Source/Runtime/Core/Private/Stats/StatsCommand.cpp
--- a/Engine/Source/Runtime/Core/Private/Stats/StatsCommand.cpp
+++ b/Engine/Source/Runtime/Core/Private/Stats/StatsCommand.cpp
@@ -24,6 +24,7 @@
 #include "Misc/CoreDelegates.h"
 #include "Async/TaskGraphInterfaces.h"
 #include "Misc/DefaultValueHelper.h"
+#include "Tasks/Pipe.h"
 
 #if STATS
 
@@ -42,6 +43,8 @@ static float DumpCull = 1.0f;
 //Whether or not we render stats in certain modes
 bool GRenderStats = true;
 
+extern CORE_API UE::Tasks::FPipe GStatsPipe;
+
 static TAutoConsoleVariable<int32> GCVarDumpHitchesAllThreads(
 	TEXT("t.DumpHitches.AllThreads"),
 	0,
@@ -2340,12 +2343,6 @@ bool DirectStatsCommand(const TCHAR* Cmd, bool bBlockForCompletion /*= false*/,
 		{
 			const FString FullCmd = FString(Cmd) + AddArgs;
 #if STATS
-			ENamedThreads::Type ThreadType = ENamedThreads::GameThread;
-			if (FPlatformProcess::SupportsMultithreading())
-			{
-				ThreadType = ENamedThreads::StatsThread;
-			}
-
 			// make sure these are initialized on the game thread
 			FLatestGameThreadStatsData::Get();
 			FStatGroupGameThreadNotifier::Get();
@@ -2354,13 +2351,10 @@ bool DirectStatsCommand(const TCHAR* Cmd, bool bBlockForCompletion /*= false*/,
 				STAT_FSimpleDelegateGraphTask_StatCmd,
 				STATGROUP_TaskGraphTasks);
 
-			FGraphEventRef CompleteHandle = FSimpleDelegateGraphTask::CreateAndDispatchWhenReady(
-				FSimpleDelegateGraphTask::FDelegate::CreateStatic(&StatCmd, FullCmd, bStatCommand, Ar),
-				GET_STATID(STAT_FSimpleDelegateGraphTask_StatCmd), NULL, ThreadType
-			);
-			if (bBlockForCompletion && FPlatformProcess::SupportsMultithreading())
+			UE::Tasks::TTask<void> Task = GStatsPipe.Launch(UE_SOURCE_LOCATION, [FullCmd, bStatCommand, Ar] { StatCmd(FullCmd, bStatCommand, Ar); });
+			if (bBlockForCompletion)
 			{
-				FTaskGraphInterface::Get().WaitUntilTaskCompletes(CompleteHandle);
+				Task.Wait();
 				GLog->FlushThreadedLogs();
 			}
 #else
@@ -2374,8 +2368,13 @@ bool DirectStatsCommand(const TCHAR* Cmd, bool bBlockForCompletion /*= false*/,
 
 #if STATS
 
-static void GetPermanentStats_StatsThread(TArray<FStatMessage>* OutStats)
+static void GetPermanentStats_StatsPipe(TArray<FStatMessage>* OutStats)
 {
+	DECLARE_CYCLE_STAT(TEXT("FSimpleDelegateGraphTask.GetPermanentStatsString_StatsPipe"),
+		STAT_FSimpleDelegateGraphTask_GetPermanentStatsString_StatsPipe,
+		STATGROUP_TaskGraphTasks);
+	SCOPE_CYCLE_COUNTER(STAT_FSimpleDelegateGraphTask_GetPermanentStatsString_StatsPipe);
+
 	FStatsThreadState& StatsData = FStatsThreadState::GetLocalState();
 	TArray<FStatMessage>& Stats = *OutStats;
 	for (auto It = StatsData.NotClearedEveryFrame.CreateConstIterator(); It; ++It)
@@ -2387,18 +2386,10 @@ static void GetPermanentStats_StatsThread(TArray<FStatMessage>* OutStats)
 
 void GetPermanentStats(TArray<FStatMessage>& OutStats)
 {
-	DECLARE_CYCLE_STAT(TEXT("FSimpleDelegateGraphTask.GetPermanentStatsString_StatsThread"),
-		STAT_FSimpleDelegateGraphTask_GetPermanentStatsString_StatsThread,
-		STATGROUP_TaskGraphTasks);
-
-	FGraphEventRef CompleteHandle = FSimpleDelegateGraphTask::CreateAndDispatchWhenReady(
-		FSimpleDelegateGraphTask::FDelegate::CreateStatic(&GetPermanentStats_StatsThread, &OutStats),
-		GET_STATID(STAT_FSimpleDelegateGraphTask_GetPermanentStatsString_StatsThread), NULL,
-		FPlatformProcess::SupportsMultithreading() ? ENamedThreads::StatsThread : ENamedThreads::GameThread
-	);
-	FTaskGraphInterface::Get().WaitUntilTaskCompletes(CompleteHandle);
+	GStatsPipe
+		.Launch(UE_SOURCE_LOCATION, [&OutStats] { GetPermanentStats_StatsPipe(&OutStats); })
+		.Wait();
 }
 
-
 #endif
 
diff --git a/Engine/Source/Runtime/Core/Private/Tests/Async/TaskGraphTest.cpp b/Engine/Source/Runtime/Core/Private/Tests/Async/TaskGraphTest.cpp
--- a/Engine/Source/Runtime/Core/Private/Tests/Async/TaskGraphTest.cpp
+++ b/Engine/Source/Runtime/Core/Private/Tests/Async/TaskGraphTest.cpp
@@ -896,7 +896,9 @@ namespace TaskGraphTests
 #if STATS
 		{	// StatsThread
 			bool bExecuted = false;
+PRAGMA_DISABLE_DEPRECATION_WARNINGS
 			FFunctionGraphTask::CreateAndDispatchWhenReady([&bExecuted] { bExecuted = true; }, TStatId{}, nullptr, ENamedThreads::StatsThread)->Wait();
+PRAGMA_ENABLE_DEPRECATION_WARNINGS
 			check(bExecuted);
 		}
 #endif
@@ -1297,6 +1299,8 @@ PRAGMA_ENABLE_DEPRECATION_WARNINGS
 
 	bool FPerfTest::RunTest(const FString& Parameters)
 	{
+		TRACE_CPUPROFILER_EVENT_SCOPE(TaskGraphTests_PerfTest);
+
 		//UE_BENCHMARK(1, TestSpawning<100000>);
 		//return true;
 
@@ -1318,8 +1322,27 @@ PRAGMA_ENABLE_DEPRECATION_WARNINGS
 	void PipeStressTest()
 	{
 #if STATS
-		constexpr ENamedThreads::Type Pipe = ENamedThreads::StatsThread;
-		constexpr ENamedThreads::Type ManualPipe = ENamedThreads::Type(ENamedThreads::StatsThread | ENamedThreads::LocalQueue);
+		FThread StatsThread;
+		if (FPlatformProcess::SupportsMultithreading())
+		{
+			// we simulate StatsThread for this test, as it was removed but the interface is still present (though it's deprecated)
+			StatsThread = FThread{ TEXT("Test StatsThread"),
+				[]
+				{
+					FMemory::SetupTLSCachesOnCurrentThread();
+PRAGMA_DISABLE_DEPRECATION_WARNINGS
+					FTaskGraphInterface::Get().AttachToThread(ENamedThreads::StatsThread);
+					FTaskGraphInterface::Get().ProcessThreadUntilRequestReturn(ENamedThreads::StatsThread);
+PRAGMA_ENABLE_DEPRECATION_WARNINGS
+					FMemory::ClearAndDisableTLSCachesOnCurrentThread();
+				}
+			};
+		}
+
+PRAGMA_DISABLE_DEPRECATION_WARNINGS
+		const ENamedThreads::Type Pipe = FPlatformProcess::SupportsMultithreading() ? ENamedThreads::StatsThread : ENamedThreads::GameThread;
+		const ENamedThreads::Type ManualPipe = FPlatformProcess::SupportsMultithreading() ? ENamedThreads::StatsThread_Local : ENamedThreads::GameThread_Local;
+PRAGMA_ENABLE_DEPRECATION_WARNINGS
 
 		{	// check that piped task is executed
 			std::atomic<bool> bExecuted{ false };
@@ -1362,7 +1385,7 @@ PRAGMA_ENABLE_DEPRECATION_WARNINGS
 					T1Launcher->Wait();
 					bT1Executed = true;
 					bExecutingPipe = false;
-				}, 
+				},
 				TStatId{}, nullptr, Pipe
 			);
 
@@ -1376,7 +1399,7 @@ PRAGMA_ENABLE_DEPRECATION_WARNINGS
 					T2Launcher->Wait();
 					bT2Executed = true;
 					bExecutingPipe = false;
-				}, 
+				},
 				TStatId{}, nullptr, Pipe
 			);
 
@@ -1390,7 +1413,7 @@ PRAGMA_ENABLE_DEPRECATION_WARNINGS
 					T3Launcher->Wait();
 					bT3Executed = true;
 					bExecutingPipe = false;
-				}, 
+				},
 				TStatId{}, nullptr, Pipe
 			);
 
@@ -1433,7 +1456,7 @@ PRAGMA_ENABLE_DEPRECATION_WARNINGS
 					FTaskGraphInterface::Get().ProcessThreadUntilIdle(ManualPipe);
 					check(bExecuted);
 					check(ManualPipeTask->IsComplete());
-				}, 
+				},
 				TStatId{}, nullptr, Pipe
 			)->Wait();
 			check(bExecuted);
@@ -1477,10 +1500,10 @@ PRAGMA_ENABLE_DEPRECATION_WARNINGS
 		{	// test GetCurrentThread() from inside manual pipe
 			FFunctionGraphTask::CreateAndDispatchWhenReady(
 				[ManualPipe]
-				{ 
+				{
 					checkf(FTaskGraphInterface::Get().GetCurrentThreadIfKnown(true) == ManualPipe, TEXT("%d - %d"), FTaskGraphInterface::Get().GetCurrentThreadIfKnown(true), ManualPipe);
 				},
-			TStatId{}, nullptr, ManualPipe);
+				TStatId{}, nullptr, ManualPipe);
 
 			FFunctionGraphTask::CreateAndDispatchWhenReady([ManualPipe] { FTaskGraphInterface::Get().ProcessThreadUntilIdle(ManualPipe); }, TStatId{}, nullptr, Pipe)
 				->Wait();
@@ -1492,15 +1515,15 @@ PRAGMA_ENABLE_DEPRECATION_WARNINGS
 				{
 					check(FTaskGraphInterface::Get().IsThreadProcessingTasks(ManualPipe));
 				},
-			TStatId{}, nullptr, ManualPipe);
+				TStatId{}, nullptr, ManualPipe);
 
 			FFunctionGraphTask::CreateAndDispatchWhenReady(
 				[ManualPipe]
 				{
 					check(!FTaskGraphInterface::Get().IsThreadProcessingTasks(ManualPipe));
-					FTaskGraphInterface::Get().ProcessThreadUntilIdle(ManualPipe); 
-				}, 
-			TStatId{}, nullptr, Pipe)->Wait();
+					FTaskGraphInterface::Get().ProcessThreadUntilIdle(ManualPipe);
+				},
+				TStatId{}, nullptr, Pipe)->Wait();
 		}
 
 		{	// check that independent task can bypass a task with pending prereq
@@ -1575,6 +1598,16 @@ PRAGMA_ENABLE_DEPRECATION_WARNINGS
 			//UE_LOG(LogTemp, Display, TEXT("Tasks completed"));
 			check(TasksExecutedNum == TasksNum);
 		}
+
+		if (FPlatformProcess::SupportsMultithreading())
+		{
+			// stop StatsThread
+PRAGMA_DISABLE_DEPRECATION_WARNINGS
+			FGraphEventRef QuitTask = TGraphTask<FReturnGraphTask>::CreateTask(NULL, ENamedThreads::GameThread).ConstructAndDispatchWhenReady(ENamedThreads::StatsThread);
+PRAGMA_ENABLE_DEPRECATION_WARNINGS
+			FTaskGraphInterface::Get().WaitUntilTaskCompletes(QuitTask, ENamedThreads::GameThread_Local);
+			StatsThread.Join();
+		}
 #endif
 	}
 
@@ -1582,10 +1615,35 @@ PRAGMA_ENABLE_DEPRECATION_WARNINGS
 
 	bool FPipeTest::RunTest(const FString& Parameters)
 	{
+		TRACE_CPUPROFILER_EVENT_SCOPE(TaskGraphTests_PipeTest);
+
 		UE_BENCHMARK(5, PipeStressTest);
 
 		return true;
 	}
+
+	IMPLEMENT_SIMPLE_AUTOMATION_TEST(FTaskGraphStatsThreadRedirectionTest, "System.Core.Async.TaskGraph.StatsThreadRedirection", EAutomationTestFlags::ApplicationContextMask | EAutomationTestFlags::EngineFilter);
+
+	bool FTaskGraphStatsThreadRedirectionTest::RunTest(const FString& Parameters)
+	{
+		if (!FPlatformProcess::SupportsMultithreading())
+		{
+			// no StatsThread around
+			return true;
+		}
+
+#if STATS
+
+		bool bExecuted = false;
+PRAGMA_DISABLE_DEPRECATION_WARNINGS
+		FFunctionGraphTask::CreateAndDispatchWhenReady([&bExecuted] { bExecuted = true; }, TStatId{}, nullptr, ENamedThreads::StatsThread)->Wait();
+PRAGMA_ENABLE_DEPRECATION_WARNINGS
+		check(bExecuted);
+
+#endif
+
+		return true;
+	}
 }
 
 #endif //WITH_DEV_AUTOMATION_TESTS
diff --git a/Engine/Source/Runtime/Core/Public/Async/TaskGraphInterfaces.h b/Engine/Source/Runtime/Core/Public/Async/TaskGraphInterfaces.h
--- a/Engine/Source/Runtime/Core/Public/Async/TaskGraphInterfaces.h
+++ b/Engine/Source/Runtime/Core/Public/Async/TaskGraphInterfaces.h
@@ -57,7 +57,7 @@ namespace ENamedThreads
 		UnusedAnchor = -1,
 		/** The always-present, named threads are listed next **/
 #if STATS
-		StatsThread, 
+		StatsThread UE_DEPRECATED(5.0, "`StatsThread` has been removed. Stats system should be used by its public API in `Stats2.h`"),
 #endif
 		RHIThread,
 		AudioThread,
@@ -98,7 +98,9 @@ namespace ENamedThreads
 
 		/** Combinations **/
 #if STATS
-		StatsThread_Local = StatsThread | LocalQueue,
+PRAGMA_DISABLE_DEPRECATION_WARNINGS
+		StatsThread_Local UE_DEPRECATED(5.0, "`StatsThread_Local` has been removed (and has never been supported). Stats system should be used by its public API in `Stats2.h`") = StatsThread | LocalQueue,
+PRAGMA_ENABLE_DEPRECATION_WARNINGS
 #endif
 		GameThread_Local = GameThread | LocalQueue,
 		ActualRenderingThread_Local = ActualRenderingThread | LocalQueue,
diff --git a/Engine/Source/Runtime/Core/Public/GenericPlatform/GenericPlatformAffinity.h b/Engine/Source/Runtime/Core/Public/GenericPlatform/GenericPlatformAffinity.h
--- a/Engine/Source/Runtime/Core/Public/GenericPlatform/GenericPlatformAffinity.h
+++ b/Engine/Source/Runtime/Core/Public/GenericPlatform/GenericPlatformAffinity.h
@@ -80,6 +80,7 @@ public:
 		return 0xFFFFFFFFFFFFFFFF;
 	}
 
+	UE_DEPRECATED(5.0, "This parameter is obsolete as StatsThread was removed and the affinity mask has no effect anymore. Please remove the call.")
 	static const CORE_API uint64 GetStatsThreadMask()
 	{
 		return 0xFFFFFFFFFFFFFFFF;
diff --git a/Engine/Source/Runtime/Core/Public/Stats/Stats2.h b/Engine/Source/Runtime/Core/Public/Stats/Stats2.h
--- a/Engine/Source/Runtime/Core/Public/Stats/Stats2.h
+++ b/Engine/Source/Runtime/Core/Public/Stats/Stats2.h
@@ -1222,7 +1222,7 @@ class FThreadStats : FNoncopyable
 {
 	friend class FStatsMallocProfilerProxy;
 	friend class FStatsThreadState;
-	friend class FStatsThread;
+	friend class FStatsProcessor;
 	friend struct FThreadStatsPool;
 
 	/** Used to control when we are collecting stats. User of the stats system increment and decrement this counter as they need data. **/
@@ -1265,15 +1265,14 @@ class FThreadStats : FNoncopyable
 	/** Tracks current stack depth for cycle counters. **/
 	bool bSawExplicitFlush;
 
-	/** True if this is the stats thread, which needs special handling. **/
-	bool bIsStatsThread;
-
 	/** Gathers information about the current thread and sets up the TLS value. **/
 	CORE_API FThreadStats();
 
 	/** Constructor used for the pool. */
 	CORE_API FThreadStats(EConstructor);
 
+	void SendMessage_Async(FStatPacket* ToSend);
+
 public:
 	/** Checks the TLS for a thread packet and if it isn't found, it makes a new one. **/
 	static FORCEINLINE_STATS FThreadStats* GetThreadStats()
@@ -1656,7 +1655,7 @@ private:
 /** Manages startup messages, usually to update the metadata. */
 class FStartupMessages
 {
-	friend class FStatsThread;
+	friend class FStatsProcessor;
 
 	TArray<FStatMessage> DelayedMessages;
 	FCriticalSection CriticalSection;
diff --git a/Engine/Source/Runtime/Core/Public/Stats/StatsData.h b/Engine/Source/Runtime/Core/Public/Stats/StatsData.h
--- a/Engine/Source/Runtime/Core/Public/Stats/StatsData.h
+++ b/Engine/Source/Runtime/Core/Public/Stats/StatsData.h
@@ -488,7 +488,7 @@ struct IItemFilter
  */
 class CORE_API FStatsThreadState
 {
-	friend class FStatsThread;
+	friend class FStatsProcessor;
 	friend struct FStatPacketArray;
 	friend struct FStatsReadFile;
 
diff --git a/Engine/Source/Runtime/Core/Public/Tasks/TaskPrivate.h b/Engine/Source/Runtime/Core/Public/Tasks/TaskPrivate.h
--- a/Engine/Source/Runtime/Core/Public/Tasks/TaskPrivate.h
+++ b/Engine/Source/Runtime/Core/Public/Tasks/TaskPrivate.h
@@ -196,8 +196,8 @@ namespace UE { namespace Tasks
 				return true;
 			}
 
-			bool PushIntoPipe();
-			void UnblockPipe();
+			CORE_API bool PushIntoPipe();
+			CORE_API void UnblockPipe();
 
 			// after being executed the task checks if it has a subsequent to launch. it's the only time the task will do this check so
 			// it "closes" itself and setting a subsequent fill fail after that.
