Add macro variant GET_FUNCTION_NAME_CHECKED_PROTOTYPE. The basic GET_FUNCTION_NAME_CHECKED does not work with overloaded functions.
#rnx

diff --git a/Engine/Source/Runtime/Core/Public/Misc/AssertionMacros.h b/Engine/Source/Runtime/Core/Public/Misc/AssertionMacros.h
--- a/Engine/Source/Runtime/Core/Public/Misc/AssertionMacros.h
+++ b/Engine/Source/Runtime/Core/Public/Misc/AssertionMacros.h
@@ -5,6 +5,7 @@
 #include "CoreTypes.h"
 #include "HAL/Platform.h"
 #include "HAL/PlatformMisc.h"
+#include "HAL/PreprocessorHelpers.h"
 #include "Templates/AndOrNot.h"
 #include "Templates/EnableIf.h"
 #include "Templates/IsArrayOrRefOfType.h"
@@ -393,6 +394,14 @@ namespace UE4Asserts_Private
 #define GET_FUNCTION_NAME_STRING_CHECKED(ClassName, FunctionName) \
 	((void)sizeof(&ClassName::FunctionName), TEXT(#FunctionName))
 
+// Returns FName(TEXT("FunctionName")), while statically verifying that the function exists in ClassName
+// Handles overloaded functions by specifying the argument of the overload to use
+#define GET_FUNCTION_NAME_CHECKED_OneParam(ClassName, FunctionName, ArgType) \
+	((void)sizeof(std::declval<ClassName&>().FunctionName(std::declval<PREPROCESSOR_REMOVE_OPTIONAL_PARENS(ArgType)>()), (int)0), FName(TEXT(#FunctionName)))
+
+#define GET_FUNCTION_NAME_STRING_CHECKED_OneParam(ClassName, FunctionName, ArgType) \
+	((void)sizeof(std::declval<ClassName&>().FunctionName(std::declval<PREPROCESSOR_REMOVE_OPTIONAL_PARENS(ArgType)>()), (int)0), TEXT(#FunctionName))
+
 /*----------------------------------------------------------------------------
 	Low level error macros
 ----------------------------------------------------------------------------*/
