Deprecate non-const binding function in TAttribute. TAttribute::Steal now returns an inplace TVariant, before an ObjectType was created by the default constructor.

diff --git a/Engine/Plugins/Developer/Concert/MultiUserClient/Source/MultiUserClient/Private/ConcertWorkspaceUI.cpp b/Engine/Plugins/Developer/Concert/MultiUserClient/Source/MultiUserClient/Private/ConcertWorkspaceUI.cpp
--- a/Engine/Plugins/Developer/Concert/MultiUserClient/Source/MultiUserClient/Private/ConcertWorkspaceUI.cpp
+++ b/Engine/Plugins/Developer/Concert/MultiUserClient/Source/MultiUserClient/Private/ConcertWorkspaceUI.cpp
@@ -105,7 +105,7 @@ public:
 	}
 
 private:
-	EVisibility GetVisibility()
+	EVisibility GetVisibility() const
 	{
 		// If the asset is locked, make the icon visible, collapsed/hidden otherwise.
 		return WorkspaceFrontend->GetResourceLockId(AssetPath).IsValid() ? EVisibility::Visible : EVisibility::Collapsed;
@@ -265,7 +265,7 @@ public:
 	}
 
 private:
-	EVisibility GetVisibility()
+	EVisibility GetVisibility() const
 	{
 		return WorkspaceFrontend->IsAssetModifiedByOtherClients(AssetPath) ? EVisibility::Visible : EVisibility::Collapsed;
 	}
diff --git a/Engine/Plugins/Media/MediaFrameworkUtilities/Source/MediaFrameworkUtilitiesEditor/Private/Profile/SMediaProfileSettingsOptionsWindow.cpp b/Engine/Plugins/Media/MediaFrameworkUtilities/Source/MediaFrameworkUtilitiesEditor/Private/Profile/SMediaProfileSettingsOptionsWindow.cpp
--- a/Engine/Plugins/Media/MediaFrameworkUtilities/Source/MediaFrameworkUtilitiesEditor/Private/Profile/SMediaProfileSettingsOptionsWindow.cpp
+++ b/Engine/Plugins/Media/MediaFrameworkUtilities/Source/MediaFrameworkUtilitiesEditor/Private/Profile/SMediaProfileSettingsOptionsWindow.cpp
@@ -143,7 +143,7 @@ FReply SMediaProfileSettingsOptionsWindow::OnCancel()
 }
 
 
-bool SMediaProfileSettingsOptionsWindow::CanConfigure()
+bool SMediaProfileSettingsOptionsWindow::CanConfigure() const
 {
 	return Options.IsValid();
 }
diff --git a/Engine/Plugins/Media/MediaFrameworkUtilities/Source/MediaFrameworkUtilitiesEditor/Private/Profile/SMediaProfileSettingsOptionsWindow.h b/Engine/Plugins/Media/MediaFrameworkUtilities/Source/MediaFrameworkUtilitiesEditor/Private/Profile/SMediaProfileSettingsOptionsWindow.h
--- a/Engine/Plugins/Media/MediaFrameworkUtilities/Source/MediaFrameworkUtilitiesEditor/Private/Profile/SMediaProfileSettingsOptionsWindow.h
+++ b/Engine/Plugins/Media/MediaFrameworkUtilities/Source/MediaFrameworkUtilitiesEditor/Private/Profile/SMediaProfileSettingsOptionsWindow.h
@@ -36,7 +36,7 @@ public:
 private:
 	FReply OnConfigure();
 	FReply OnCancel();
-	bool CanConfigure();
+	bool CanConfigure() const;
 
 private:
 	TWeakPtr<SWindow> Window;
diff --git a/Engine/Source/Runtime/Core/Public/Misc/Attribute.h b/Engine/Source/Runtime/Core/Public/Misc/Attribute.h
--- a/Engine/Source/Runtime/Core/Public/Misc/Attribute.h
+++ b/Engine/Source/Runtime/Core/Public/Misc/Attribute.h
@@ -92,7 +92,7 @@ public:
 	 *
 	 * @param  InGetter		Delegate to bind
 	 */
-	static TAttribute Create( const FGetter& InGetter )
+	UE_NODISCARD static TAttribute Create( const FGetter& InGetter )
 	{
 		const bool bExplicitConstructor = true;
 		return TAttribute( InGetter, bExplicitConstructor );
@@ -105,17 +105,18 @@ public:
 	 *
 	 * @param  InFuncPtr Member function to bind.  The function's structure (return value, arguments, etc) must match IBoundAttributeDelegate's definition.
 	 */
-	static TAttribute Create( typename FGetter::FStaticDelegate::FFuncPtr InFuncPtr )
+	UE_NODISCARD static TAttribute Create( typename FGetter::FStaticDelegate::FFuncPtr InFuncPtr )
 	{
 		const bool bExplicitConstructor = true;
 		return TAttribute( FGetter::CreateStatic( InFuncPtr ), bExplicitConstructor );
 	}
 
 	/**
-	 * Helper function for creating TAttributes from a non-const member function pointer, accessed through a raw pointer
+	 * Helper function for creating TAttributes from a function pointer, accessed through a raw pointer
 	 */
 	template<typename SourceType, typename SourceTypeOrBase, typename... PayloadTypes>
-	FORCEINLINE static TAttribute CreateRaw(SourceType* InObject, ObjectType (SourceTypeOrBase::*InMethod)(PayloadTypes...), typename TDecay<PayloadTypes>::Type... InputPayload)
+	UE_DEPRECATED(5.0, "Attribute's Getter should be const.")
+	UE_NODISCARD FORCEINLINE static TAttribute CreateRaw(SourceType* InObject, ObjectType (SourceTypeOrBase::*InMethod)(PayloadTypes...), typename TDecay<PayloadTypes>::Type... InputPayload)
 	{
 		return Create(FGetter::CreateRaw(InObject, InMethod, MoveTemp(InputPayload)...));
 	}
@@ -124,7 +125,7 @@ public:
 	 * Helper function for creating TAttributes from a const member function pointer, accessed through a raw pointer
 	 */
 	template<typename SourceType, typename SourceTypeOrBase, typename... PayloadTypes>
-	FORCEINLINE static TAttribute CreateRaw(const SourceType* InObject, ObjectType (SourceTypeOrBase::*InMethod)(PayloadTypes...) const, typename TDecay<PayloadTypes>::Type... InputPayload)
+	UE_NODISCARD FORCEINLINE static TAttribute CreateRaw(const SourceType* InObject, ObjectType (SourceTypeOrBase::*InMethod)(PayloadTypes...) const, typename TDecay<PayloadTypes>::Type... InputPayload)
 	{
 		return Create(FGetter::CreateRaw(InObject, InMethod, MoveTemp(InputPayload)...));
 	}
@@ -133,7 +134,8 @@ public:
 	 * Helper function for creating TAttributes from a non-const member function pointer, accessed through a weak pointer to the shared object
 	 */
 	template<typename SourceType, typename SourceTypeOrBase, typename... PayloadTypes>
-	FORCEINLINE static TAttribute CreateSP(SourceType* InObject, ObjectType (SourceTypeOrBase::*InMethod)(PayloadTypes...), typename TDecay<PayloadTypes>::Type... InputPayload)
+	UE_DEPRECATED(5.0, "Attribute's Getter should be const.")
+	UE_NODISCARD FORCEINLINE static TAttribute CreateSP(SourceType* InObject, ObjectType (SourceTypeOrBase::*InMethod)(PayloadTypes...), typename TDecay<PayloadTypes>::Type... InputPayload)
 	{
 		return Create(FGetter::CreateSP(InObject, InMethod, MoveTemp(InputPayload)...));
 	}
@@ -142,7 +144,7 @@ public:
 	 * Helper function for creating TAttributes from a const member function pointer, accessed through a weak pointer to the shared object
 	 */
 	template<typename SourceType, typename SourceTypeOrBase, typename... PayloadTypes>
-	FORCEINLINE static TAttribute CreateSP(const SourceType* InObject, ObjectType (SourceTypeOrBase::*InMethod)(PayloadTypes...) const, typename TDecay<PayloadTypes>::Type... InputPayload)
+	UE_NODISCARD FORCEINLINE static TAttribute CreateSP(const SourceType* InObject, ObjectType (SourceTypeOrBase::*InMethod)(PayloadTypes...) const, typename TDecay<PayloadTypes>::Type... InputPayload)
 	{
 		return Create(FGetter::CreateSP(InObject, InMethod, MoveTemp(InputPayload)...));
 	}
@@ -152,11 +154,32 @@ public:
 	 * TAttribute<float> FloatAttribute = TAttribute<float>::CreateLambda([]{ return 10.f; });
 	 */
 	template<typename LambdaType, typename... PayloadTypes>
-	FORCEINLINE static TAttribute CreateLambda(LambdaType&& InCallable, PayloadTypes&&... InputPayload)
+	UE_NODISCARD FORCEINLINE static TAttribute CreateLambda(LambdaType&& InCallable, PayloadTypes&&... InputPayload)
 	{
 		return Create(FGetter::CreateLambda(InCallable, Forward<PayloadTypes>(InputPayload)...));
 	}
 
+	/**
+	 * Creates an attribute by binding an arbitrary function that will be called to generate this attribute's value on demand.
+	 * After binding, the attribute will no longer have a value that can be accessed directly, and instead the bound
+	 * function will always be called to generate the value.
+	 *
+	 * @param  InUserObject  Instance of the class that contains the member function you want to bind.
+	 * @param  InFunctionName Member function name to bind.
+	 */
+	template< class SourceType >
+	UE_NODISCARD static TAttribute< ObjectType > Create(SourceType* InUserObject, const FName& InFunctionName)
+	{
+		TAttribute< ObjectType > Attrib;
+		Attrib.BindUFunction<SourceType>(InUserObject, InFunctionName);
+		return Attrib;
+	}
+
+	UE_NODISCARD FORCEINLINE static TAttribute< ObjectType > Create(TFunction<ObjectType(void)>&& InLambda)
+	{
+		return Create(TAttribute< ObjectType >::FGetter::CreateLambda(MoveTemp(InLambda)));
+	}
+
 	/**
 	 * Sets the attribute's value
 	 *
@@ -335,27 +358,6 @@ public:
 		Getter.BindUFunction(InUserObject, InFunctionName);
 	}
 
-	/**
-	 * Creates an attribute by binding an arbitrary function that will be called to generate this attribute's value on demand.
-	 * After binding, the attribute will no longer have a value that can be accessed directly, and instead the bound
-	 * function will always be called to generate the value.
-	 *
-	 * @param  InUserObject  Instance of the class that contains the member function you want to bind.
-	 * @param  InFunctionName Member function name to bind.
-	 */
-	template< class SourceType >
-	UE_NODISCARD static TAttribute< ObjectType > Create(SourceType* InUserObject, const FName& InFunctionName)
-	{
-		TAttribute< ObjectType > Attrib;
-		Attrib.BindUFunction<SourceType>(InUserObject, InFunctionName);
-		return Attrib;
-	}
-
-	UE_NODISCARD static TAttribute< ObjectType > Create(TFunction<ObjectType(void)>&& InLambda)
-	{
-		return Create(TAttribute< ObjectType >::FGetter::CreateLambda(MoveTemp(InLambda)));
-	}
-
 	/**
 	 * Checks to see if this attribute has a 'getter' function bound
 	 *
@@ -385,16 +387,7 @@ public:
 	{
 		checkf(IsSet(), TEXT("It is an error to call Steal() on an unset TAttribute. Check IsSet() before calling Steal()."));
 		bIsSet = false;
-		TVariant<ObjectType, FGetter> Temp;
-		if (IsBound())
-		{
-			Temp.template Set<FGetter>(MoveTemp(Getter));
-		}
-		else
-		{
-			Temp.template Set<ObjectType>(MoveTemp(Value));
-		}
-		return Temp;
+		return IsBound() ? TVariant<ObjectType, FGetter>(TInPlaceType<FGetter>(), MoveTemp(Getter)) : TVariant<ObjectType, FGetter>(TInPlaceType<ObjectType>(), MoveTemp(Value));
 	}
 
 	/**
@@ -450,7 +443,8 @@ private:
  * Helper function for creating TAttributes from a non-const member function pointer, accessed through a raw pointer
  */
 template<typename T, typename SourceType, typename SourceTypeOrBase, typename... PayloadTypes>
-FORCEINLINE TAttribute<T> MakeAttributeRaw(SourceType* InObject, T (SourceTypeOrBase::*InMethod)(PayloadTypes...), typename TDecay<PayloadTypes>::Type... InputPayload)
+UE_DEPRECATED(5.0, "Attribute's Getter should be const.")
+UE_NODISCARD FORCEINLINE TAttribute<T> MakeAttributeRaw(SourceType* InObject, T (SourceTypeOrBase::*InMethod)(PayloadTypes...), typename TDecay<PayloadTypes>::Type... InputPayload)
 {
 	return TAttribute<T>::Create(TAttribute<T>::FGetter::CreateRaw(InObject, InMethod, MoveTemp(InputPayload)...));
 }
@@ -459,7 +453,7 @@ FORCEINLINE TAttribute<T> MakeAttributeRaw(SourceType* InObject, T (SourceTypeOr
  * Helper function for creating TAttributes from a const member function pointer, accessed through a raw pointer
  */
 template<typename T, typename SourceType, typename SourceTypeOrBase, typename... PayloadTypes>
-FORCEINLINE TAttribute<T> MakeAttributeRaw(const SourceType* InObject, T (SourceTypeOrBase::*InMethod)(PayloadTypes...) const, typename TDecay<PayloadTypes>::Type...  InputPayload)
+UE_NODISCARD FORCEINLINE TAttribute<T> MakeAttributeRaw(const SourceType* InObject, T (SourceTypeOrBase::*InMethod)(PayloadTypes...) const, typename TDecay<PayloadTypes>::Type... InputPayload)
 {
 	return TAttribute<T>::Create(TAttribute<T>::FGetter::CreateRaw(InObject, InMethod, MoveTemp(InputPayload)...));
 }
@@ -468,7 +462,8 @@ FORCEINLINE TAttribute<T> MakeAttributeRaw(const SourceType* InObject, T (Source
  * Helper function for creating TAttributes from a non-const member function pointer, accessed through a weak pointer to the shared object
  */
 template<typename T, typename SourceType, typename SourceTypeOrBase, typename... PayloadTypes>
-FORCEINLINE TAttribute<T> MakeAttributeSP(SourceType* InObject, T (SourceTypeOrBase::*InMethod)(PayloadTypes...), typename TDecay<PayloadTypes>::Type...  InputPayload)
+UE_DEPRECATED(5.0, "Attribute's Getter should be const.")
+UE_NODISCARD FORCEINLINE TAttribute<T> MakeAttributeSP(SourceType* InObject, T (SourceTypeOrBase::*InMethod)(PayloadTypes...), typename TDecay<PayloadTypes>::Type...  InputPayload)
 {
 	return TAttribute<T>::Create(TAttribute<T>::FGetter::CreateSP(InObject, InMethod, MoveTemp(InputPayload)...));
 }
@@ -477,7 +472,7 @@ FORCEINLINE TAttribute<T> MakeAttributeSP(SourceType* InObject, T (SourceTypeOrB
  * Helper function for creating TAttributes from a const member function pointer, accessed through a weak pointer to the shared object
  */
 template<typename T, typename SourceType, typename SourceTypeOrBase, typename... PayloadTypes>
-FORCEINLINE TAttribute<T> MakeAttributeSP(const SourceType* InObject, T (SourceTypeOrBase::*InMethod)(PayloadTypes...) const, typename TDecay<PayloadTypes>::Type...  InputPayload)
+UE_NODISCARD FORCEINLINE TAttribute<T> MakeAttributeSP(const SourceType* InObject, T (SourceTypeOrBase::*InMethod)(PayloadTypes...) const, typename TDecay<PayloadTypes>::Type...  InputPayload)
 {
 	return TAttribute<T>::Create(TAttribute<T>::FGetter::CreateSP(InObject, InMethod, MoveTemp(InputPayload)...));
 }
@@ -487,7 +482,7 @@ FORCEINLINE TAttribute<T> MakeAttributeSP(const SourceType* InObject, T (SourceT
  * TAttribute<float> FloatAttribute = MakeAttributeLambda([]{ return 10.f; });
  */
 template<typename LambdaType, typename... PayloadTypes>
-decltype(auto) MakeAttributeLambda(LambdaType&& InCallable, PayloadTypes&&... InputPayload)
+UE_NODISCARD decltype(auto) MakeAttributeLambda(LambdaType&& InCallable, PayloadTypes&&... InputPayload)
 {
 	typedef decltype(InCallable(DeclVal<PayloadTypes>()...)) T;
 
