U5 - Make static analisis happy

diff --git a/Engine/Source/Runtime/Core/Public/Async/Fundamental/Scheduler.h b/Engine/Source/Runtime/Core/Public/Async/Fundamental/Scheduler.h
--- a/Engine/Source/Runtime/Core/Public/Async/Fundamental/Scheduler.h
+++ b/Engine/Source/Runtime/Core/Public/Async/Fundamental/Scheduler.h
@@ -144,7 +144,7 @@ namespace LowLevelTasks
 				}
 			}
 
-			void Push(FSleepEvent* Item)
+			void Push(CA_VALID_POINTER FSleepEvent* Item)
 			{
 				check(Item != nullptr);
 				checkSlow(reinterpret_cast<uintptr_t>(Item) < (1ull << 48));
