Add TStringBuilder::ToView()

diff --git a/Engine/Source/Runtime/Core/Public/Misc/StringBuilder.h b/Engine/Source/Runtime/Core/Public/Misc/StringBuilder.h
--- a/Engine/Source/Runtime/Core/Public/Misc/StringBuilder.h
+++ b/Engine/Source/Runtime/Core/Public/Misc/StringBuilder.h
@@ -81,6 +81,7 @@ public:
 	inline CharType* GetData()					{ return Base; }
 	inline const CharType* GetData() const		{ return Base; }
 	inline const CharType* ToString() const		{ EnsureNulTerminated(); return Base; }
+	inline ViewType ToView() const				{ return ViewType(Base, Len()); }
 	inline const CharType* operator*() const	{ EnsureNulTerminated(); return Base; }
 
 	inline const CharType	LastChar() const	{ return *(CurPos - 1); }
