U5 - This static analizer is really a pain.

diff --git a/Engine/Source/Runtime/Core/Public/Async/Fundamental/Scheduler.h b/Engine/Source/Runtime/Core/Public/Async/Fundamental/Scheduler.h
--- a/Engine/Source/Runtime/Core/Public/Async/Fundamental/Scheduler.h
+++ b/Engine/Source/Runtime/Core/Public/Async/Fundamental/Scheduler.h
@@ -144,28 +144,33 @@ namespace LowLevelTasks
 				}
 			}
 
-			void Push(CA_VALID_POINTER FSleepEvent* Item)
+			void Push(FSleepEvent* Item)
 			{
-				check(Item != nullptr);
-				checkSlow(reinterpret_cast<uintptr_t>(Item) < (1ull << 48));
-				checkSlow((reinterpret_cast<uintptr_t>(Item) & 0x7) == 0);
-				checkSlow(Item->Next == nullptr);
-				
-				FTopNode LocalTop = Top.load(std::memory_order_relaxed);
-				while (true) 
+#if USING_CODE_ANALYSIS
+				if(Item != nullptr)
+#endif
 				{
+					check(Item != nullptr);
+					checkSlow(reinterpret_cast<uintptr_t>(Item) < (1ull << 48));
+					checkSlow((reinterpret_cast<uintptr_t>(Item) & 0x7) == 0);
+					checkSlow(Item->Next == nullptr);
+				
+					FTopNode LocalTop = Top.load(std::memory_order_relaxed);
+					while (true) 
+					{
 #if DO_CHECK
-					int64 LastRevision = int64(LocalTop.Revision); 
+						int64 LastRevision = int64(LocalTop.Revision); 
 #endif
-					Item->Next = reinterpret_cast<FSleepEvent*>(LocalTop.Address << 3);
-					if (Top.compare_exchange_weak(LocalTop, FTopNode { reinterpret_cast<uintptr_t>(Item) >> 3, uintptr_t(LocalTop.Revision + 1) }, std::memory_order_release, std::memory_order_acquire))  
-					{
-						return;
-					}
+						Item->Next = reinterpret_cast<FSleepEvent*>(LocalTop.Address << 3);
+						if (Top.compare_exchange_weak(LocalTop, FTopNode { reinterpret_cast<uintptr_t>(Item) >> 3, uintptr_t(LocalTop.Revision + 1) }, std::memory_order_release, std::memory_order_acquire))  
+						{
+							return;
+						}
 #if DO_CHECK
-					int64 NewRevision = int64(LocalTop.Revision) < LastRevision ? ((1ll << 19) + int64(LocalTop.Revision)) : int64(LocalTop.Revision);
-					ensureMsgf((NewRevision - LastRevision) < (1ll << 18), TEXT("Dangerously close to the wraparound: %d, %d"), LastRevision, NewRevision);
+						int64 NewRevision = int64(LocalTop.Revision) < LastRevision ? ((1ll << 19) + int64(LocalTop.Revision)) : int64(LocalTop.Revision);
+						ensureMsgf((NewRevision - LastRevision) < (1ll << 18), TEXT("Dangerously close to the wraparound: %d, %d"), LastRevision, NewRevision);
 #endif
+					}
 				}
 			}
 		};
