- Rewrote SerializedUATProcess to simplify the threading, and allow shutting down to work better while operations are running (especially on Linux and maybe Mac)

diff --git a/Engine/Source/Runtime/Core/Private/Misc/MonitoredProcess.cpp b/Engine/Source/Runtime/Core/Private/Misc/MonitoredProcess.cpp
--- a/Engine/Source/Runtime/Core/Private/Misc/MonitoredProcess.cpp
+++ b/Engine/Source/Runtime/Core/Private/Misc/MonitoredProcess.cpp
@@ -5,6 +5,7 @@
 #include "Misc/Paths.h"
 #include "Logging/LogMacros.h"
 #include "Async/Async.h"
+#include "Misc/CoreDelegates.h"
 #include <atomic>
 
 DEFINE_LOG_CATEGORY_STATIC(LogMonitoredProcess, Log, All);
@@ -125,6 +126,7 @@ void FMonitoredProcess::TickInternal()
 	// monitor the process
 	ProcessOutput(FPlatformProcess::ReadPipe(ReadPipe));
 
+	// kill child process if we choose to cancel or it the engine is shutting down
 	if (Canceling)
 	{
 		FPlatformProcess::TerminateProc(ProcessHandle, KillTree);
@@ -192,8 +194,33 @@ void FMonitoredProcess::Tick()
 *****************************************************************************/
 
 FCriticalSection FSerializedUATProcess::Serializer;
+FSerializedUATProcess* FSerializedUATProcess::HeadProcess = nullptr;
 bool FSerializedUATProcess::bHasSucceededOnce = false;
 
+void FSerializedUATProcess::CancelQueue()
+{
+	FScopeLock Lock(&Serializer);
+	if (HeadProcess)
+	{
+		// delete anything after Head
+		FSerializedUATProcess* Travel = HeadProcess->NextProcessToRun;
+		while (Travel != nullptr)
+		{
+			FSerializedUATProcess* Delete = Travel;
+			Travel = Travel->NextProcessToRun;
+			delete Travel;
+		}
+
+
+		Travel = HeadProcess;
+		HeadProcess = nullptr;
+		Travel->NextProcessToRun = nullptr;
+		Travel->Cancel(true);
+
+		// can/should we block here?
+	}
+}
+
 FSerializedUATProcess::FSerializedUATProcess(const FString& RunUATCommandline)
 	// we will modify URL and Params in this constructor, so there's no need to pass anything up to base
 	: FMonitoredProcess("", "", true, true)
@@ -207,53 +234,101 @@ FSerializedUATProcess::FSerializedUATProcess(const FString& RunUATCommandline)
 	URL = TEXT("/usr/bin/env");
 	Params = FString::Printf(TEXT(" -- \"%s\" %s"), *GetUATPath(), *RunUATCommandline);
 #endif
+
+	static bool bHasSetupDelegate = false;
+	if (!bHasSetupDelegate)
+	{
+		bHasSetupDelegate = true;
+		FCoreDelegates::OnShutdownAfterError.AddStatic(&CancelQueue);
+		FCoreDelegates::OnExit.AddStatic(&CancelQueue);
+	}
 }
 
 bool FSerializedUATProcess::Launch()
 {
-	Async(EAsyncExecution::Thread, [this]()
-		{
-			// don't let this do anything if another process is running
-			FScopeLock Lock(&Serializer);
+	FScopeLock Lock(&Serializer);
 
-			FEventRef Event;
+	// are we first in line?
+	if (HeadProcess == nullptr)
+	{
+		HeadProcess = this;
 
-			if (bHasSucceededOnce)
+		return LaunchInternal();
+	}
+	else
+	{
+		// put us last in line
+		for (FSerializedUATProcess* Travel = HeadProcess; Travel; Travel = Travel->NextProcessToRun)
+		{
+			if (Travel->NextProcessToRun == nullptr)
 			{
-				Params += TEXT(" -nocompile");
+				Travel->NextProcessToRun = this;
+				break;
 			}
+		}
 
-			FOnMonitoredProcessCompleted OriginalCompletedDelegate = CompletedDelegate;
-			FSimpleDelegate OriginalCanceledDelegate = CanceledDelegate;
+		// and return immediately
+		return true;
+	}
+}
 
-			CompletedDelegate.BindLambda([this, &Event, OriginalCompletedDelegate](int32 ExitCode)
-				{
-					if (ExitCode == 0 || ExitCode == 10)
-					{
-						bHasSucceededOnce = true;
-					}
-					OriginalCompletedDelegate.ExecuteIfBound(ExitCode);
+bool FSerializedUATProcess::LaunchNext()
+{
+	FScopeLock Lock(&Serializer);
+	if (HeadProcess == nullptr)
+	{
+		// can happen during shutdown
+		return false;
+	}
+	
+	check(this == HeadProcess);
 
-					Event->Trigger();
-				});
-			CanceledDelegate.BindLambda([this, &Event, OriginalCanceledDelegate]()
-				{
-					OriginalCanceledDelegate.ExecuteIfBound();
+	FSerializedUATProcess* FinishedTask = HeadProcess;
+	HeadProcess = HeadProcess->NextProcessToRun;
+//	delete FinishedTask;
 
-					Event->Trigger();
-				});
+	if (HeadProcess != nullptr)
+	{
+		return HeadProcess->LaunchInternal();
+	}
+	return false;
+}
 
-			UE_LOG(LogMonitoredProcess, Log, TEXT("Running Serialized UAT: [ %s %s ]"), *URL, *Params);
+bool FSerializedUATProcess::LaunchInternal()
+{
+	if (bHasSucceededOnce)
+	{
+//		Params += TEXT(" -nocompile");
+	}
 
-			if (FMonitoredProcess::Launch() == false)
-			{
-				LaunchFailedDelegate.ExecuteIfBound();
+	FOnMonitoredProcessCompleted OriginalCompletedDelegate = CompletedDelegate;
+	FSimpleDelegate OriginalCanceledDelegate = CanceledDelegate;
 
-				Event->Trigger();
+	CompletedDelegate.BindLambda([this, OriginalCompletedDelegate](int32 ExitCode)
+		{
+			if (ExitCode == 0 || ExitCode == 10)
+			{
+				bHasSucceededOnce = true;
 			}
+			OriginalCompletedDelegate.ExecuteIfBound(ExitCode);
 
-			Event->Wait();
+			LaunchNext();
 		});
+	CanceledDelegate.BindLambda([this, OriginalCanceledDelegate]()
+		{
+			OriginalCanceledDelegate.ExecuteIfBound();
+
+			LaunchNext();
+		});
+
+	UE_LOG(LogMonitoredProcess, Log, TEXT("Running Serialized UAT: [ %s %s ]"), *URL, *Params);
+
+	if (FMonitoredProcess::Launch() == false)
+	{
+		LaunchFailedDelegate.ExecuteIfBound();
+
+		LaunchNext();
+	}
 
 	return true;
 }
@@ -270,3 +345,4 @@ FString FSerializedUATProcess::GetUATPath()
 
 	return FPaths::ConvertRelativePathToFull(FPaths::EngineDir() / TEXT("Build/BatchFiles") / RunUATScriptName);
 }
+
diff --git a/Engine/Source/Runtime/Core/Public/Misc/MonitoredProcess.h b/Engine/Source/Runtime/Core/Public/Misc/MonitoredProcess.h
--- a/Engine/Source/Runtime/Core/Public/Misc/MonitoredProcess.h
+++ b/Engine/Source/Runtime/Core/Public/Misc/MonitoredProcess.h
@@ -285,10 +285,18 @@ public:
 
 private:
 
+	bool LaunchNext();
+	bool LaunchInternal();
+	static void CancelQueue();
+
+	// When this one completes, run the next in line
+	FSerializedUATProcess* NextProcessToRun = nullptr;
+
 	// Holds a delegate that is executed when the process fails to launch (delayed, in a thread). Used in place of the return value
 	// of Launch in the parent class, since it's async
 	FSimpleDelegate LaunchFailedDelegate;
 
 	static FCriticalSection Serializer;
 	static bool bHasSucceededOnce;
+	static FSerializedUATProcess* HeadProcess;
 };
