- Uncommented code that was mistakenly left commented out

diff --git a/Engine/Source/Runtime/Core/Private/Misc/MonitoredProcess.cpp b/Engine/Source/Runtime/Core/Private/Misc/MonitoredProcess.cpp
--- a/Engine/Source/Runtime/Core/Private/Misc/MonitoredProcess.cpp
+++ b/Engine/Source/Runtime/Core/Private/Misc/MonitoredProcess.cpp
@@ -298,7 +298,7 @@ bool FSerializedUATProcess::LaunchInternal()
 {
 	if (bHasSucceededOnce)
 	{
-//		Params += TEXT(" -nocompile");
+		Params += TEXT(" -nocompile");
 	}
 
 	FOnMonitoredProcessCompleted OriginalCompletedDelegate = CompletedDelegate;
