Rename UE4 code mentions to Unreal
#rnx

diff --git a/Engine/Source/Runtime/Core/Public/Android/AndroidSystemIncludes.h b/Engine/Source/Runtime/Core/Public/Android/AndroidSystemIncludes.h
--- a/Engine/Source/Runtime/Core/Public/Android/AndroidSystemIncludes.h
+++ b/Engine/Source/Runtime/Core/Public/Android/AndroidSystemIncludes.h
@@ -27,7 +27,7 @@ int vswprintf(TCHAR *dst, int count, const TCHAR *fmt, va_list arg);
 // Define TEXT early	//@todo android:
 #define TEXT(x) L##x
 
-// map the Windows functions (that UE4 unfortunately uses be default) to normal functions
+// map the Windows functions (that Unreal unfortunately uses by default) to normal functions
 #define _alloca alloca
 */
 
