Added "stats VerboseNamedEvents" command to toggle GShouldEmitVerboseNamedEvents. This is in addition to exisiting "stat NamedEvents" command that toggles GCycleStatsShouldEmitNamedEvents. When verbose is toggled on, it will also turn on the NamedEvents, if not alreay enabled.
Also made a minor optimization re test for "verbose" flag in FCycleCounter::Start.

diff --git a/Engine/Config/BaseInput.ini b/Engine/Config/BaseInput.ini
--- a/Engine/Config/BaseInput.ini
+++ b/Engine/Config/BaseInput.ini
@@ -99,6 +99,7 @@ MaxScrollbackSize=1024
 +ManualAutoCompleteList=(Command="Stat DrawCount",Desc="Shows draw counts broken down by category")
 +ManualAutoCompleteList=(Command="Stat UnitGraph",Desc="Draws simple unit time graph")
 +ManualAutoCompleteList=(Command="Stat NamedEvents",Desc="Stat NamedEvents (Enables named events for external profilers)")
++ManualAutoCompleteList=(Command="Stat VerboseNamedEvents",Desc="Stat VerboseNamedEvents (Enables verbose named events for external profilers)")
 +ManualAutoCompleteList=(Command="Stat StartFile",Desc="Stat StartFile (starts a stats capture, creating a new file in the Profiling directory; stop with stat StopFile to close the file)")
 +ManualAutoCompleteList=(Command="Stat StopFile",Desc="Stat StopFile (finishes a stats capture started by stat StartFile)")
 +ManualAutoCompleteList=(Command="Stat CPULoad",Desc="Stat CPULoad (Shows CPU Utilization)")
diff --git a/Engine/Source/Runtime/Core/Public/Stats/Stats2.h b/Engine/Source/Runtime/Core/Public/Stats/Stats2.h
--- a/Engine/Source/Runtime/Core/Public/Stats/Stats2.h
+++ b/Engine/Source/Runtime/Core/Public/Stats/Stats2.h
@@ -1549,7 +1549,7 @@ public:
 	 * Pushes the specified stat onto the hierarchy for this thread. Starts
 	 * the timing of the cycles used
 	 */
-	FORCEINLINE_STATS void Start( TStatId InStatId, EStatFlags InStatFlags, bool bAlways = false )
+	FORCEINLINE_STATS void Start(TStatId InStatId, EStatFlags InStatFlags, bool bAlways = false)
 	{
 		FMinimalName StatMinimalName = InStatId.GetMinimalName(EMemoryOrder::Relaxed);
 		if (StatMinimalName.IsNone())
@@ -1558,31 +1558,30 @@ public:
 		}
 
 		// Emit named event for active cycle stat.
-		if ( GCycleStatsShouldEmitNamedEvents > 0 
+		if (GCycleStatsShouldEmitNamedEvents > 0
 			&& (GShouldEmitVerboseNamedEvents || !EnumHasAnyFlags(InStatFlags, EStatFlags::Verbose)))
 		{
-#if	PLATFORM_USES_ANSI_STRING_FOR_EXTERNAL_PROFILING
-			FPlatformMisc::BeginNamedEvent( FColor( 0 ), InStatId.GetStatDescriptionANSI() );
+#if PLATFORM_USES_ANSI_STRING_FOR_EXTERNAL_PROFILING
+			FPlatformMisc::BeginNamedEvent(FColor(0), InStatId.GetStatDescriptionANSI());
 #else
-			FPlatformMisc::BeginNamedEvent( FColor( 0 ), InStatId.GetStatDescriptionWIDE() );
-#endif // PLATFORM_USES_ANSI_STRING_FOR_EXTERNAL_PROFILING
+			FPlatformMisc::BeginNamedEvent(FColor(0), InStatId.GetStatDescriptionWIDE());
+#endif
 			EmittedEvent |= NamedEvent;
-		}
 
 #if CPUPROFILERTRACE_ENABLED
-		if (GCycleStatsShouldEmitNamedEvents > 0 && UE_TRACE_CHANNELEXPR_IS_ENABLED(CpuChannel)
-			&& (GShouldEmitVerboseNamedEvents || !EnumHasAnyFlags(InStatFlags, EStatFlags::Verbose)))
-		{
-			FCpuProfilerTrace::OutputBeginDynamicEvent(InStatId.GetStatDescriptionANSI()); //todo: Could we use FName index as event id?
-			EmittedEvent |= TraceEvent;
-		}
+			if (UE_TRACE_CHANNELEXPR_IS_ENABLED(CpuChannel))
+			{
+				FCpuProfilerTrace::OutputBeginDynamicEvent(InStatId.GetStatDescriptionANSI()); //todo: Could we use FName index as event id?
+				EmittedEvent |= TraceEvent;
+			}
 #endif
+		}
 
-		if( (bAlways && FThreadStats::WillEverCollectData()) || FThreadStats::IsCollectingData() )
+		if ((bAlways && FThreadStats::WillEverCollectData()) || FThreadStats::IsCollectingData())
 		{
 			FName StatName = MinimalNameToName(StatMinimalName);
 			StatId = StatName;
-			FThreadStats::AddMessage( StatName, EStatOperation::CycleScopeStart );
+			FThreadStats::AddMessage(StatName, EStatOperation::CycleScopeStart);
 			EmittedEvent |= ThreadStatsEvent;
 		}
 	}
@@ -1597,7 +1596,7 @@ public:
 	 */
 	FORCEINLINE_STATS void Stop()
 	{
-		if ( EmittedEvent & NamedEvent )
+		if (EmittedEvent & NamedEvent)
 		{
 			FPlatformMisc::EndNamedEvent();
 		}
@@ -1609,7 +1608,7 @@ public:
 		}
 #endif
 
-		if(EmittedEvent & ThreadStatsEvent)
+		if (EmittedEvent & ThreadStatsEvent)
 		{
 			FThreadStats::AddMessage(StatId, EStatOperation::CycleScopeEnd);
 		}
diff --git a/Engine/Source/Runtime/Engine/Classes/Engine/Engine.h b/Engine/Source/Runtime/Engine/Classes/Engine/Engine.h
--- a/Engine/Source/Runtime/Engine/Classes/Engine/Engine.h
+++ b/Engine/Source/Runtime/Engine/Classes/Engine/Engine.h
@@ -3399,6 +3399,7 @@ private:
 	bool ToggleStatDetailed(UWorld* World, FCommonViewportClient* ViewportClient, const TCHAR* Stream = nullptr);
 	bool ToggleStatHitches(UWorld* World, FCommonViewportClient* ViewportClient, const TCHAR* Stream = nullptr);
 	bool ToggleStatNamedEvents(UWorld* World, FCommonViewportClient* ViewportClient, const TCHAR* Stream = nullptr);
+	bool ToggleStatVerboseNamedEvents(UWorld* World, FCommonViewportClient* ViewportClient, const TCHAR* Stream = nullptr);
 	bool ToggleStatUnit(UWorld* World, FCommonViewportClient* ViewportClient, const TCHAR* Stream = nullptr);
 #if !UE_BUILD_SHIPPING
 	bool PostStatSoundModulatorHelp(UWorld* World, FCommonViewportClient* ViewportClient, const TCHAR* Stream = nullptr);
diff --git a/Engine/Source/Runtime/Engine/Private/UnrealEngine.cpp b/Engine/Source/Runtime/Engine/Private/UnrealEngine.cpp
--- a/Engine/Source/Runtime/Engine/Private/UnrealEngine.cpp
+++ b/Engine/Source/Runtime/Engine/Private/UnrealEngine.cpp
@@ -1773,6 +1773,7 @@ void UEngine::Init(IEngineLoop* InEngineLoop)
 	EngineStats.Add(FEngineStatFuncs(TEXT("STAT_Version"), TEXT("STATCAT_Engine"), FText::GetEmpty(), FEngineStatRender::CreateUObject(this, &UEngine::RenderStatVersion), FEngineStatToggle(), bIsRHS));
 #endif // !UE_BUILD_SHIPPING
 	EngineStats.Add(FEngineStatFuncs(TEXT("STAT_NamedEvents"), TEXT("STATCAT_Engine"), FText::GetEmpty(), FEngineStatRender::CreateUObject(this, &UEngine::RenderStatNamedEvents), FEngineStatToggle::CreateUObject(this, &UEngine::ToggleStatNamedEvents), bIsRHS));
+	EngineStats.Add(FEngineStatFuncs(TEXT("STAT_VerboseNamedEvents"), TEXT("STATCAT_Engine"), FText::GetEmpty(), FEngineStatRender(), FEngineStatToggle::CreateUObject(this, &UEngine::ToggleStatVerboseNamedEvents), bIsRHS));
 	EngineStats.Add(FEngineStatFuncs(TEXT("STAT_FPS"), TEXT("STATCAT_Engine"), FText::GetEmpty(), FEngineStatRender::CreateUObject(this, &UEngine::RenderStatFPS), FEngineStatToggle::CreateUObject(this, &UEngine::ToggleStatFPS), bIsRHS));
 	EngineStats.Add(FEngineStatFuncs(TEXT("STAT_Summary"), TEXT("STATCAT_Engine"), FText::GetEmpty(), FEngineStatRender::CreateUObject(this, &UEngine::RenderStatSummary), FEngineStatToggle(), bIsRHS));
 	EngineStats.Add(FEngineStatFuncs(TEXT("STAT_Unit"), TEXT("STATCAT_Engine"), FText::GetEmpty(), FEngineStatRender::CreateUObject(this, &UEngine::RenderStatUnit), FEngineStatToggle::CreateUObject(this, &UEngine::ToggleStatUnit), bIsRHS));
@@ -15251,9 +15252,26 @@ bool UEngine::ToggleStatNamedEvents(UWorld* World, FCommonViewportClient* Viewpo
 	return false;
 }
 
+// VERBOSE NAMED EVENTS
+bool UEngine::ToggleStatVerboseNamedEvents(UWorld* World, FCommonViewportClient* ViewportClient, const TCHAR* Stream)
+{
+	GShouldEmitVerboseNamedEvents = !GShouldEmitVerboseNamedEvents;
+
+	if (GShouldEmitVerboseNamedEvents && ViewportClient && !ViewportClient->IsStatEnabled(TEXT("NamedEvents")))
+	{
+		SetEngineStat(World, ViewportClient, TEXT("NamedEvents"), true);
+	}
+
+	return false;
+}
+
 int32 UEngine::RenderStatNamedEvents(UWorld* World, FViewport* Viewport, FCanvas* Canvas, int32 X, int32 Y, const FVector* ViewLocation, const FRotator* ViewRotation)
 {
-	FCanvasTextItem TextItem(FVector2D(X - 40, Y), LOCTEXT("NAMEDEVENTSENABLED", "NAMED EVENTS ENABLED"), GetSmallFont(), FLinearColor::Blue);
+	const FText Text = GShouldEmitVerboseNamedEvents ?
+		LOCTEXT("VERBOSENAMEDEVENTSENABLED", "VERBOSE NAMED EVENTS ENABLED") :
+		LOCTEXT("NAMEDEVENTSENABLED", "NAMED EVENTS ENABLED");
+	X -= GShouldEmitVerboseNamedEvents ? 100 : 42;
+	FCanvasTextItem TextItem(FVector2D(X, Y), Text, GetSmallFont(), FLinearColor::Blue);
 	TextItem.EnableShadow(FLinearColor::Black);
 	Canvas->DrawItem(TextItem);
 	Y += TextItem.DrawnSize.Y;
