non-unity compilation fix

diff --git a/Engine/Source/Runtime/Core/Private/Tests/Async/TaskGraphTest.cpp b/Engine/Source/Runtime/Core/Private/Tests/Async/TaskGraphTest.cpp
--- a/Engine/Source/Runtime/Core/Private/Tests/Async/TaskGraphTest.cpp
+++ b/Engine/Source/Runtime/Core/Private/Tests/Async/TaskGraphTest.cpp
@@ -14,6 +14,7 @@
 #include "Containers/CircularQueue.h"
 #include "Containers/Queue.h"
 #include "Tests/Benchmark.h"
+#include "HAL/Thread.h"
 
 #include <atomic>
 
