Array: Use Invoke for predicates in TArray and TArrayView

diff --git a/Engine/Source/Runtime/Core/Public/Containers/Array.h b/Engine/Source/Runtime/Core/Public/Containers/Array.h
--- a/Engine/Source/Runtime/Core/Public/Containers/Array.h
+++ b/Engine/Source/Runtime/Core/Public/Containers/Array.h
@@ -20,6 +20,7 @@
 #include "Algo/Impl/BinaryHeap.h"
 #include "Templates/AndOrNot.h"
 #include "Templates/IdentityFunctor.h"
+#include "Templates/Invoke.h"
 #include "Templates/Less.h"
 #include "Templates/ChooseClass.h"
 #include "Templates/Sorting.h"
@@ -959,7 +960,7 @@ public:
 		for (const ElementType* RESTRICT Start = GetData(), *RESTRICT Data = Start + Count; Data != Start; )
 		{
 			--Data;
-			if (Pred(*Data))
+			if (::Invoke(Pred, *Data))
 			{
 				return static_cast<SizeType>(Data - Start);
 			}
@@ -1012,7 +1013,7 @@ public:
 		const ElementType* RESTRICT Start = GetData();
 		for (const ElementType* RESTRICT Data = Start, *RESTRICT DataEnd = Start + ArrayNum; Data != DataEnd; ++Data)
 		{
-			if (Pred(*Data))
+			if (::Invoke(Pred, *Data))
 			{
 				return static_cast<SizeType>(Data - Start);
 			}
@@ -1080,7 +1081,7 @@ public:
 	{
 		for (ElementType* RESTRICT Data = GetData(), *RESTRICT DataEnd = Data + ArrayNum; Data != DataEnd; ++Data)
 		{
-			if (Pred(*Data))
+			if (::Invoke(Pred, *Data))
 			{
 				return Data;
 			}
@@ -1103,7 +1104,7 @@ public:
 		TArray<ElementType> FilterResults;
 		for (const ElementType* RESTRICT Data = GetData(), *RESTRICT DataEnd = Data + ArrayNum; Data != DataEnd; ++Data)
 		{
-			if (Pred(*Data))
+			if (::Invoke(Pred, *Data))
 			{
 				FilterResults.Add(*Data);
 			}
@@ -2415,11 +2416,11 @@ public:
 
 		SizeType WriteIndex = 0;
 		SizeType ReadIndex = 0;
-		bool NotMatch = !Predicate(GetData()[ReadIndex]); // use a ! to guarantee it can't be anything other than zero or one
+		bool NotMatch = !::Invoke(Predicate, GetData()[ReadIndex]); // use a ! to guarantee it can't be anything other than zero or one
 		do
 		{
 			SizeType RunStartIndex = ReadIndex++;
-			while (ReadIndex < OriginalNum && NotMatch == !Predicate(GetData()[ReadIndex]))
+			while (ReadIndex < OriginalNum && NotMatch == !::Invoke(Predicate, GetData()[ReadIndex]))
 			{
 				ReadIndex++;
 			}
@@ -2460,7 +2461,7 @@ public:
 		const SizeType OriginalNum = ArrayNum;
 		for (SizeType ItemIndex = 0; ItemIndex < Num();)
 		{
-			if (Predicate((*this)[ItemIndex]))
+			if (::Invoke(Predicate, (*this)[ItemIndex]))
 			{
 				bRemoved = true;
 				RemoveAtSwap(ItemIndex, 1, false);
diff --git a/Engine/Source/Runtime/Core/Public/Containers/ArrayView.h b/Engine/Source/Runtime/Core/Public/Containers/ArrayView.h
--- a/Engine/Source/Runtime/Core/Public/Containers/ArrayView.h
+++ b/Engine/Source/Runtime/Core/Public/Containers/ArrayView.h
@@ -7,6 +7,7 @@
 #include "Templates/IsSigned.h"
 #include "Templates/PointerIsConvertibleFromTo.h"
 #include "Misc/AssertionMacros.h"
+#include "Templates/Invoke.h"
 #include "Templates/UnrealTypeTraits.h"
 #include "Traits/ElementType.h"
 #include "Containers/Array.h"
@@ -400,7 +401,7 @@ public:
 		for (const ElementType* RESTRICT Start = GetData(), *RESTRICT Data = Start + StartIndex; Data != Start; )
 		{
 			--Data;
-			if (Pred(*Data))
+			if (::Invoke(Pred, *Data))
 			{
 				return static_cast<SizeType>(Data - Start);
 			}
@@ -458,7 +459,7 @@ public:
 		const ElementType* RESTRICT Start = GetData();
 		for (const ElementType* RESTRICT Data = Start, *RESTRICT DataEnd = Start + ArrayNum; Data != DataEnd; ++Data)
 		{
-			if (Pred(*Data))
+			if (::Invoke(Pred, *Data))
 			{
 				return static_cast<SizeType>(Data - Start);
 			}
@@ -501,7 +502,7 @@ public:
 	{
 		for (ElementType* RESTRICT Data = GetData(), *RESTRICT DataEnd = Data + ArrayNum; Data != DataEnd; ++Data)
 		{
-			if (Pred(*Data))
+			if (::Invoke(Pred, *Data))
 			{
 				return Data;
 			}
@@ -524,7 +525,7 @@ public:
 		TArray<typename TRemoveConst<ElementType>::Type> FilterResults;
 		for (const ElementType* RESTRICT Data = GetData(), *RESTRICT DataEnd = Data + ArrayNum; Data != DataEnd; ++Data)
 		{
-			if (Pred(*Data))
+			if (::Invoke(Pred, *Data))
 			{
 				FilterResults.Add(*Data);
 			}
