Improve Blueprint Spring Interp behaviour. Note that this changes FloatSpringInterp and VectorSpringInterp to make them more accurate (as well as more useful), which could affect backwards compatibility of existing assets - though the difference are likely to be tiny.

diff --git a/Engine/Plugins/Experimental/ControlRig/Source/ControlRig/Private/Units/Simulation/RigUnit_SpringInterp.cpp b/Engine/Plugins/Experimental/ControlRig/Source/ControlRig/Private/Units/Simulation/RigUnit_SpringInterp.cpp
--- a/Engine/Plugins/Experimental/ControlRig/Source/ControlRig/Private/Units/Simulation/RigUnit_SpringInterp.cpp
+++ b/Engine/Plugins/Experimental/ControlRig/Source/ControlRig/Private/Units/Simulation/RigUnit_SpringInterp.cpp
@@ -7,6 +7,7 @@ namespace RigUnitSpringInterpConstants
 {
 	static const float FixedTimeStep = 1.0f / 60.0f;
 	static const float MaxTimeStep = 0.1f;
+	static const float TargetVelocityAmount = 1.0f;
 }
 
 FRigUnit_SpringInterp_Execute()
@@ -25,11 +26,14 @@ FRigUnit_SpringInterp_Execute()
 		Result = Current;
 		while (RemainingTime >= RigUnitSpringInterpConstants::FixedTimeStep)
 		{
-			Result = UKismetMathLibrary::FloatSpringInterp(Result, Target, SpringState, Stiffness, CriticalDamping, RigUnitSpringInterpConstants::FixedTimeStep, Mass);
+			Result = UKismetMathLibrary::FloatSpringInterp(Result, Target, SpringState, Stiffness, CriticalDamping,
+				RigUnitSpringInterpConstants::FixedTimeStep, Mass, RigUnitSpringInterpConstants::TargetVelocityAmount,
+				false, 0.0f, 0.0f, true);
 			RemainingTime -= RigUnitSpringInterpConstants::FixedTimeStep;
 		}
 
-		Result = UKismetMathLibrary::FloatSpringInterp(Result, Target, SpringState, Stiffness, CriticalDamping, RemainingTime, Mass);
+		Result = UKismetMathLibrary::FloatSpringInterp(Result, Target, SpringState, Stiffness, CriticalDamping,
+			RemainingTime, Mass, RigUnitSpringInterpConstants::TargetVelocityAmount, false, 0.0f, 0.0f, true);
 	}
 }
 
@@ -49,10 +53,13 @@ FRigUnit_SpringInterpVector_Execute()
 		Result = Current;
 		while (RemainingTime >= RigUnitSpringInterpConstants::FixedTimeStep)
 		{
-			Result = UKismetMathLibrary::VectorSpringInterp(Result, Target, SpringState, Stiffness, CriticalDamping, RigUnitSpringInterpConstants::FixedTimeStep, Mass);
+			Result = UKismetMathLibrary::VectorSpringInterp(Result, Target, SpringState, Stiffness, CriticalDamping,
+				RigUnitSpringInterpConstants::FixedTimeStep, Mass, RigUnitSpringInterpConstants::TargetVelocityAmount,
+				false, FVector(), FVector(), true);
 			RemainingTime -= RigUnitSpringInterpConstants::FixedTimeStep;
 		}
 
-		Result = UKismetMathLibrary::VectorSpringInterp(Result, Target, SpringState, Stiffness, CriticalDamping, RemainingTime, Mass);
+		Result = UKismetMathLibrary::VectorSpringInterp(Result, Target, SpringState, Stiffness, CriticalDamping,
+			RemainingTime, Mass, RigUnitSpringInterpConstants::TargetVelocityAmount, false, FVector(), FVector(), true);
 	}
 }
\ No newline at end of file
diff --git a/Engine/Source/Runtime/Core/Public/Math/UnrealMathUtility.h b/Engine/Source/Runtime/Core/Public/Math/UnrealMathUtility.h
--- a/Engine/Source/Runtime/Core/Public/Math/UnrealMathUtility.h
+++ b/Engine/Source/Runtime/Core/Public/Math/UnrealMathUtility.h
@@ -1142,6 +1142,50 @@ public:
 	/** Interpolate quaternion from Current to Target. Scaled by angle to Target, so it has a strong start speed and ease out. */
 	static CORE_API FQuat QInterpTo(const FQuat& Current, const FQuat& Target, float DeltaTime, float InterpSpeed);
 
+	/**
+	* Returns an approximation of Exp(-X) based on a Taylor expansion that has had the coefficients adjusted (using
+	* optimisation) to minimise the error in the range 0 < X < 1, which is below 0.1%. Note that it returns exactly 1
+	* when X is 0, and the return value is greater than the real value for values of X > 1 (but it still tends
+	* to zero for large X). 
+	*/
+	template<class T>
+	static T InvExpApprox(T X)
+	{
+		const T A(1.00746054); // 1 / 1! in Taylor series
+		const T B(0.45053901); // 1 / 2! in Taylor series
+		const T C(0.25724632); // 1 / 3! in Taylor series
+		return 1 / (1 + A * X + B * X * X + C * X * X * X);		
+	}
+	
+	/**
+	* Smooths a value using exponential damping towards a target. Works for any type that supports basic arithmetic operations.
+	* 
+	* An approximation is used that is accurate so long as InDeltaTime < 0.5 * InSmoothingTime
+	* 
+	* @param  InOutValue      The value to be smoothed
+	* @param  InTargetValue   The target to smooth towards
+	* @param  InDeltaTime     Time interval
+	* @param  InSmoothingTime Timescale over which to smooth. Larger values result in more smoothed behaviour. Can be zero.
+	*/
+	template< class T >
+    static void ExponentialSmoothingApprox(
+        T&          InOutValue,
+        const T&    InTargetValue,
+        const float InDeltaTime,
+        const float InSmoothingTime)
+	{
+		if (InSmoothingTime > KINDA_SMALL_NUMBER)
+		{
+			const float A = InDeltaTime / InSmoothingTime;
+			float Exp = InvExpApprox(A);
+			InOutValue = InTargetValue + (InOutValue - InTargetValue) * Exp;
+		}
+		else
+		{
+			InOutValue = InTargetValue;
+		}
+	}
+
 	/**
 	 * Smooths a value using a critically damped spring. Works for any type that supports basic arithmetic operations.
 	 * 
@@ -1152,21 +1196,23 @@ public:
 	 * 
 	 * When starting from zero velocity, the maximum velocity is reached after time InSmoothingTime * 0.5
 	 *
-	 * @param  InOutValue      The value to be smoothed
-	 * @param  InOutValueRate  The rate of change of the value
-	 * @param  InTargetValue   The target to smooth towards
-	 * @param  InDeltaTime     Time interval
-	 * @param  InSmoothingTime Timescale over which to smooth. Larger values result in more smoothed behaviour. Can be zero.
+	 * @param  InOutValue        The value to be smoothed
+	 * @param  InOutValueRate    The rate of change of the value
+	 * @param  InTargetValue     The target to smooth towards
+	 * @param  InTargetValueRate The target rate of change smooth towards. Note that if this is discontinuous, then the output will have discontinuous velocity too.
+	 * @param  InDeltaTime       Time interval
+	 * @param  InSmoothingTime   Timescale over which to smooth. Larger values result in more smoothed behaviour. Can be zero.
 	 */
 	template< class T >
 	static void CriticallyDampedSmoothing(
 		T&          InOutValue,
 		T&          InOutValueRate,
 		const T&    InTargetValue,
+	    const T&    InTargetValueRate,
 		const float InDeltaTime,
 		const float InSmoothingTime)
 	{
-		if (InSmoothingTime > 0.0f)
+		if (InSmoothingTime > KINDA_SMALL_NUMBER)
 		{
 			// The closed form solution for critically damped motion towards zero is:
 			//
@@ -1179,15 +1225,19 @@ public:
 			// v = ( v0 - w t (v0 + w x0) ) exp(-w t)
 			//
 			// We're damping towards a target, so convert - i.e. x = value - target
-			// 
-			// SmoothingTime is 
-			float W = 2.0f / InSmoothingTime;
-			float A = W * InDeltaTime;
-			// Taylor expansion for Exp(-W * InDeltaTime), adjusted to be more accurate for the expected delta time
-			float Exp = 1.0f / (1.0f + A + 0.48f * A * A + 0.235f * A * A * A);
-			T X0 = InOutValue - InTargetValue;
-			T B = (InOutValueRate + X0 * W) * InDeltaTime;
-			InOutValue = InTargetValue + (X0 + B) * Exp;
+			//
+			// Target velocity turns into an offset to the position
+			T AdjustedTarget = InTargetValue + InTargetValueRate * InSmoothingTime;
+
+			const float W = 2.0f / InSmoothingTime;
+			const float A = W * InDeltaTime;
+			// Approximation requires A < 1, so DeltaTime < SmoothingTime / 2
+			const float Exp = InvExpApprox(A);
+			T X0 = InOutValue - AdjustedTarget;
+			// Target velocity turns into an offset to the position
+			X0 -= InTargetValueRate * InSmoothingTime;		
+			const T B = (InOutValueRate + X0 * W) * InDeltaTime;
+			InOutValue = AdjustedTarget + (X0 + B) * Exp;
 			InOutValueRate = (InOutValueRate - B * W) * Exp;
 		}
 		else if (InDeltaTime > 0.0f)
@@ -1203,51 +1253,103 @@ public:
 	}
 
 	/**
-	* Smooths a value using exponential damping towards a target. Works for any type that supports basic arithmetic operations.
+	* Smooths a value using a spring damper towards a target.
 	* 
-	 * An approximation is used that is accurate so long as InDeltaTime < 0.5 * InSmoothingTime
+	* The implementation is exact, so stable and accurate for all values of DeltaTime etc
 	* 
-	* @param  InOutValue      The value to be smoothed
-	* @param  InTargetValue   The target to smooth towards
-	* @param  InDeltaTime     Time interval
-	* @param  InSmoothingTime Timescale over which to smooth. Larger values result in more smoothed behaviour. Can be zero.
+	* @param  InOutValue        The value to be smoothed
+	* @param  InOutValueRate    The rate of change of the value
+	* @param  InTargetValue     The target to smooth towards
+	* @param  InTargetValueRate The target rate of change smooth towards. Note that if this is discontinuous, then the output will have discontinuous velocity too.
+	* @param  InDeltaTime       Time interval
+	* @param  InSmoothingTime   Timescale over which to smooth. Larger values result in more smoothed behaviour. Can be zero.
+	* @param  InDampingRatio    1 is critical damping. <1 results in under-damped motion (i.e. with overshoot), and >1 results in over-damped motion. 
 	*/
 	template< class T >
-	static void ExponentialSmoothing(
+	static void SpringDamperSmoothing(
 	    T&          InOutValue,
+	    T&          InOutValueRate,
 	    const T&    InTargetValue,
+	    const T&    InTargetValueRate,
 	    const float InDeltaTime,
-	    const float InSmoothingTime)
+	    const float InSmoothingTime,
+	    const float InDampingRatio)
 	{
-		if (InSmoothingTime > 0.0f)
+		if (InDeltaTime <= 0.0f)
 		{
-			const float A = InDeltaTime / InSmoothingTime;
-			// Taylor expansion for Exp(-W * InDeltaTime), adjusted to be more accurate for the expected delta time
-			float Exp = 1.0f / (1.0f + A + 0.48f * A * A + 0.235f * A * A * A);
-			InOutValue = InTargetValue + (InOutValue - InTargetValue) * Exp;
+			return;
 		}
-		else
+		if (InSmoothingTime < KINDA_SMALL_NUMBER)
 		{
+			InOutValueRate = (InTargetValue - InOutValue) / InDeltaTime;
 			InOutValue = InTargetValue;
+			return;
+		}
+		// Target velocity turns into an offset to the position
+		T AdjustedTarget = InTargetValue + InTargetValueRate * (InDampingRatio * InSmoothingTime);
+
+		float W = 2.0f / InSmoothingTime;
+		T Err = InOutValue - AdjustedTarget;
+
+		// Handle the three cases separately
+		if (InDampingRatio > 1.0f) // Overdamped
+		{
+			const float WD = W * FMath::Sqrt(FMath::Square(InDampingRatio) - 1.0f);
+			const T C2 = -(InOutValueRate + (W * InDampingRatio - WD) * Err) / (2.0f * WD);
+			const T C1 = Err - C2;
+			const float A1 = (WD - InDampingRatio * W);
+			const float A2 = -(WD + InDampingRatio * W);
+			const float E1 = FMath::Exp(A1 * InDeltaTime);
+			const float E2 = FMath::Exp(A2 * InDeltaTime);
+			InOutValue = E1 * C1 + E2 * C2;
+			InOutValueRate = E1 * C1 * A1 + E2 * C2 * A2;
+			InOutValue += AdjustedTarget;
+		}
+		else if (InDampingRatio < 1.0f) // Underdamped
+		{
+			const float WD = W * FMath::Sqrt(1.0f - FMath::Square(InDampingRatio));
+			const T A = Err;
+			const T B = (InOutValueRate + Err * (InDampingRatio * W)) / WD;
+			const float C = FMath::Cos(WD * InDeltaTime);
+			const float S = FMath::Sin(WD * InDeltaTime);
+			const float E = FMath::Exp(-InDampingRatio * W * InDeltaTime);
+			InOutValue = E * (A * C + B * S);
+			InOutValueRate = -InOutValue * InDampingRatio * W;
+			InOutValueRate += E * (B * (WD * C) - A * (WD * S));
+			InOutValue += AdjustedTarget;
+		}
+		else // Critical damping
+		{
+			const T& C1 = Err;
+			T C2 = InOutValueRate + Err * W;
+			float E = FMath::Exp(-W * InDeltaTime);
+			InOutValue = (C1 + C2 * InDeltaTime) * E;
+			InOutValueRate = (C2 - C1 * W - C2 * (W * InDeltaTime)) * E;
+			InOutValue += AdjustedTarget;
 		}
 	}
 
 	/**
-	* Smooths a value using a spring damper towards a target. Works for any type that supports basic arithmetic operations.
+	* Smooths a value using a spring damper towards a target.
 	* 
-	* The implementation is exact, so with approximations for Exp/Sin/Cos
+	* The implementation uses approximations for Exp/Sin/Cos. These are good for most sensible values of DampingRatio so
+	* long as InDeltaTime < 0.5 * InSmoothingTime. If you are passing in a very large delta time, or want accurate
+	* results for highly overdamped motion, you might need to use the exact version of this function. 
 	* 
-	* @param  InOutValue      The value to be smoothed
-	* @param  InTargetValue   The target to smooth towards
-	* @param  InDeltaTime     Time interval
-	* @param  InSmoothingTime Timescale over which to smooth. Larger values result in more smoothed behaviour. Can be zero.
-	* @param  InDampingRatio  1 is critical damping. <1 results in under-damped motion (i.e. with overshoot), and >1 results in over-damped motion. 
+	* @param  InOutValue        The value to be smoothed
+	* @param  InOutValueRate    The rate of change of the value
+	* @param  InTargetValue     The target to smooth towards
+	* @param  InTargetValueRate The target rate of change smooth towards. Note that if this is discontinuous, then the output will have discontinuous velocity too.
+	* @param  InDeltaTime       Time interval
+	* @param  InSmoothingTime   Timescale over which to smooth. Larger values result in more smoothed behaviour. Can be zero.
+	* @param  InDampingRatio    1 is critical damping. <1 results in under-damped motion (i.e. with overshoot), and >1 results in over-damped motion. 
 	*/
 	template< class T >
-	static void SpringDamperSmoothing(
+	static void SpringDamperSmoothingApprox(
 	    T&          InOutValue,
 	    T&          InOutValueRate,
 	    const T&    InTargetValue,
+	    const T&    InTargetValueRate,
 	    const float InDeltaTime,
 	    const float InSmoothingTime,
 	    const float InDampingRatio)
@@ -1256,14 +1358,18 @@ public:
 		{
 			return;
 		}
-		if (InSmoothingTime <= 0.0f)
+		if (InSmoothingTime < KINDA_SMALL_NUMBER)
 		{
 			InOutValueRate = (InTargetValue - InOutValue) / InDeltaTime;
 			InOutValue = InTargetValue;
 			return;
 		}
+		// Target velocity turns into an offset to the position
+		T AdjustedTarget = InTargetValue + InTargetValueRate * (InDampingRatio * InSmoothingTime);
+
 		float W = 2.0f / InSmoothingTime;
-		T Err = InOutValue - InTargetValue;
+		T Err = InOutValue - AdjustedTarget;
+
 		// Handle the three cases separately
 		if (InDampingRatio > 1.0f) // Overdamped
 		{
@@ -1272,36 +1378,49 @@ public:
 			const T C1 = Err - C2;
 			const float A1 = (WD - InDampingRatio * W);
 			const float A2 = -(WD + InDampingRatio * W);
-			const float E1 = FMath::Exp(A1 * InDeltaTime);
-			const float E2 = FMath::Exp(A2 * InDeltaTime);
+			// Note that A1 and A2 will always be negative. We will use an approximation for 1/Exp(-A * DeltaTime).
+			const float A1_DT = -A1 * InDeltaTime;
+			const float A2_DT = -A2 * InDeltaTime;
+			// This approximation in practice will be good for all DampingRatios
+			const float E1 = InvExpApprox(A1_DT);
+			// As DampingRatio gets big, this approximation gets worse, but mere inaccuracy for overdamped motion is
+			// not likely to be important, since we end up with 1 / BigNumber
+			const float E2 = InvExpApprox(A2_DT);
 			InOutValue = E1 * C1 + E2 * C2;
 			InOutValueRate = E1 * C1 * A1 + E2 * C2 * A2;
-			InOutValue += InTargetValue;
+			InOutValue += AdjustedTarget;
 		}
 		else if (InDampingRatio < 1.0f) // Underdamped
 		{
 			const float WD = W * FMath::Sqrt(1.0f - FMath::Square(InDampingRatio));
 			const T A = Err;
 			const T B = (InOutValueRate + Err * (InDampingRatio * W)) / WD;
-			const float C = FMath::Cos(WD * InDeltaTime);
-			const float S = FMath::Sin(WD * InDeltaTime);
-			const float E = FMath::Exp(-InDampingRatio * W * InDeltaTime);
+			// Bhaskara approximation for Sin(Angle) and Cos.
+			// Needs Angle < 1 so DeltaTime < SmoothingTime / (2 * DampingRatio * Sqrt(1 - DampingRatio^2)) 
+			const float Angle = WD * InDeltaTime;
+			const float S = 16.f * Angle * (PI - Angle) / (5.f * FMath::Square(PI) - 4.f * Angle * (PI - Angle));
+			const float C = (FMath::Square(PI) - 4.0f * FMath::Square(Angle)) / (FMath::Square(PI) + FMath::Square(Angle));
+			const float E0 = InDampingRatio * WD * InDeltaTime;
+			// Needs E0 < 1 so DeltaTime < SmoothingTime / (2 * DampingRatio * Sqrt(1 - DampingRatio^2))
+			const float E = InvExpApprox(E0);
 			InOutValue = E * (A * C + B * S);
 			InOutValueRate = -InOutValue * InDampingRatio * W;
 			InOutValueRate += E * (B * (WD * C) - A * (WD * S));
-			InOutValue += InTargetValue;
+			InOutValue += AdjustedTarget;
 		}
 		else // Critical damping
 		{
 			const T& C1 = Err;
 			T C2 = InOutValueRate + Err * W;
-			float E = FMath::Exp(-W * InDeltaTime);
+			const float E0 = W * InDeltaTime;
+			// Needs E0 < 1 so InDeltaTime < InSmoothingTime / 2 
+			float E = InvExpApprox(E0);
 			InOutValue = (C1 + C2 * InDeltaTime) * E;
 			InOutValueRate = (C2 - C1 * W - C2 * (W * InDeltaTime)) * E;
-			InOutValue += InTargetValue;
+			InOutValue += AdjustedTarget;
 		}
 	}
-	
+
 	/**
 	 * Simple function to create a pulsating scalar value
 	 *
diff --git a/Engine/Source/Runtime/Engine/Classes/Animation/BlendSpace.h b/Engine/Source/Runtime/Engine/Classes/Animation/BlendSpace.h
--- a/Engine/Source/Runtime/Engine/Classes/Animation/BlendSpace.h
+++ b/Engine/Source/Runtime/Engine/Classes/Animation/BlendSpace.h
@@ -48,8 +48,8 @@ struct FInterpolationParameter
 
 	/**
 	 * Damping ratio - only used when the type is set to SpringDamper. A value of 1 will move quickly and
-	 * smoothly to the target, without overshooting. Values below 1 can be used to encourage some overshoot,
-	 * which can make pose transitions look more natural.
+	 * smoothly to the target, without overshooting. Values as low as 0 can be used to encourage some overshoot,
+	 * and values around 0.7 can make pose transitions look more natural.
 	 */
 	UPROPERTY(EditAnywhere, Category=Parameter, meta = (ClampMin = "0", EditCondition = "InterpolationType == EFilterInterpolationType::BSIT_SpringDamper && InterpolationTime > 0"))
 	float DampingRatio = 1.f;
diff --git a/Engine/Source/Runtime/Engine/Classes/Kismet/KismetMathLibrary.h b/Engine/Source/Runtime/Engine/Classes/Kismet/KismetMathLibrary.h
--- a/Engine/Source/Runtime/Engine/Classes/Kismet/KismetMathLibrary.h
+++ b/Engine/Source/Runtime/Engine/Classes/Kismet/KismetMathLibrary.h
@@ -111,19 +111,21 @@ struct ENGINE_API FFloatSpringState
 {
 	GENERATED_BODY()
 
-	float PrevError;
+	float PrevTarget;
 	float Velocity;
+	bool  bPrevTargetValid;
 
 	FFloatSpringState()
-	: PrevError(0.f)
+	: PrevTarget(0.f)
 	, Velocity(0.f)
+	, bPrevTargetValid(false)
 	{
-
 	}
 
 	void Reset()
 	{
-		PrevError = Velocity = 0.f;
+		PrevTarget = Velocity = 0.f;
+		bPrevTargetValid = false;
 	}
 };
 
@@ -132,19 +134,21 @@ struct ENGINE_API FVectorSpringState
 {
 	GENERATED_BODY()
 
-	FVector PrevError;
+	FVector PrevTarget;
 	FVector Velocity;
+	bool    bPrevTargetValid;
 
 	FVectorSpringState()
-	: PrevError(FVector::ZeroVector)
+	: PrevTarget(FVector::ZeroVector)
 	, Velocity(FVector::ZeroVector)
+	, bPrevTargetValid(false)
 	{
-
 	}
 
 	void Reset()
 	{
-		PrevError = Velocity = FVector::ZeroVector;
+		PrevTarget = Velocity = FVector::ZeroVector;
+		bPrevTargetValid = false;
 	}
 };
 
@@ -1735,17 +1739,27 @@ class ENGINE_API UKismetMathLibrary : public UBlueprintFunctionLibrary
 	static FVector VInterpTo_Constant(FVector Current, FVector Target, float DeltaTime, float InterpSpeed);
 
 	/**
-	* Uses a simple spring model to interpolate a vector from Current to Target.
-	*
-	* @param Current				Current value
-	* @param Target					Target value
-	* @param SpringState			Data related to spring model (velocity, error, etc..) - Create a unique variable per spring
-	* @param Stiffness				How stiff the spring model is (more stiffness means more oscillation around the target value)
-	* @param CriticalDampingFactor	How much damping to apply to the spring (0 means no damping, 1 means critically damped which means no oscillation)
-	* @param Mass					Multiplier that acts like mass on a spring
-	*/
-	UFUNCTION(BlueprintCallable,  meta = (ScriptMethod = "InterpSpringTo", Keywords = "position"), Category = "Math|Interpolation")
-	static FVector VectorSpringInterp(FVector Current, FVector Target, UPARAM(ref) FVectorSpringState& SpringState, float Stiffness, float CriticalDampingFactor, float DeltaTime, float Mass = 1.f);
+	 * Uses a simple spring model to interpolate a vector from Current to Target.
+	 *
+	 * @param Current               Current value
+	 * @param Target                Target value
+	 * @param SpringState           Data related to spring model (velocity, error, etc..) - Create a unique variable per spring
+	 * @param Stiffness             How stiff the spring model is (more stiffness means more oscillation around the target value)
+	 * @param CriticalDampingFactor How much damping to apply to the spring (0 means no damping, 1 means critically damped which means no oscillation)
+	 * @param DeltaTime             Time difference since the last update
+	 * @param Mass                  Multiplier that acts like mass on a spring
+	 * @param TargetVelocityAmount  If 1 then the target velocity will be calculated and used, which results following the target more closely/without lag. Values down to zero (recommended when using this to smooth data) will progressively disable this effect.
+	 * @param bClamp                Whether to use the Min/Max values to clamp the motion
+	 * @param MinValue              Clamps the minimum output value and cancels the velocity if it reaches this limit
+	 * @param MaxValue              Clamps the maximum output value and cancels the velocity if it reaches this limit
+	 * @param bApproximate          Use an approximate/fast implementation. Will be suitable for all cases unless the delta time is very large or high accuracy is needed.
+	 */
+	UFUNCTION(BlueprintCallable,  meta = (ScriptMethod = "InterpSpringTo", Keywords = "position", AdvancedDisplay = "8"), Category = "Math|Interpolation")
+	static FVector VectorSpringInterp(FVector Current, FVector Target, UPARAM(ref) FVectorSpringState& SpringState,
+	                                  float Stiffness, float CriticalDampingFactor, float DeltaTime,
+	                                  float Mass = 1.f, float TargetVelocityAmount = 1.f, 
+	                                  bool bClamp = false, FVector MinValue = FVector(-1.f), FVector MaxValue = FVector(1.f),
+	                                  bool bApproximate = true);
 
 	/**
 	 * Gets the reciprocal of this vector, avoiding division by zero.
@@ -3827,15 +3841,25 @@ class ENGINE_API UKismetMathLibrary : public UBlueprintFunctionLibrary
 	/**
 	 * Uses a simple spring model to interpolate a float from Current to Target.
 	 *
-	 * @param Current				Current value
-	 * @param Target				Target value
-	 * @param SpringState			Data related to spring model (velocity, error, etc..) - Create a unique variable per spring
-	 * @param Stiffness				How stiff the spring model is (more stiffness means more oscillation around the target value)
-	 * @param CriticalDampingFactor	How much damping to apply to the spring (0 means no damping, 1 means critically damped which means no oscillation)
-	 * @param Mass					Multiplier that acts like mass on a spring
-	 */
-	UFUNCTION(BlueprintCallable, Category = "Math|Interpolation")
-	static float FloatSpringInterp(float Current, float Target, UPARAM(ref) FFloatSpringState& SpringState, float Stiffness, float CriticalDampingFactor, float DeltaTime, float Mass = 1.f);
+	 * @param Current               Current value
+	 * @param Target                Target value
+	 * @param SpringState           Data related to spring model (velocity, error, etc..) - Create a unique variable per spring
+	 * @param Stiffness             How stiff the spring model is (more stiffness means more oscillation around the target value)
+	 * @param CriticalDampingFactor How much damping to apply to the spring (0 means no damping, 1 means critically damped which means no oscillation)
+	 * @param DeltaTime             Time difference since the last update
+	 * @param Mass                  Multiplier that acts like mass on a spring
+	 * @param TargetVelocityAmount  If 1 then the target velocity will be calculated and used, which results following the target more closely/without lag. Values down to zero (recommended when using this to smooth data) will progressively disable this effect.
+	 * @param bClamp                Whether to use the Min/Max values to clamp the motion
+	 * @param MinValue              Clamps the minimum output value and cancels the velocity if it reaches this limit
+	 * @param MaxValue              Clamps the maximum output value and cancels the velocity if it reaches this limit
+	 * @param bApproximate          Use an approximate/fast implementation. Will be suitable for all cases unless the delta time is very large or high accuracy is needed.
+	 */
+	UFUNCTION(BlueprintCallable, Category = "Math|Interpolation", meta=(AdvancedDisplay = "8"))
+	static float FloatSpringInterp(float Current, float Target, UPARAM(ref) FFloatSpringState& SpringState,
+	                               float Stiffness, float CriticalDampingFactor, float DeltaTime,
+	                               float Mass = 1.f, float TargetVelocityAmount = 1.f, 
+	                               bool bClamp = false, float MinValue = -1.f, float MaxValue = 1.f,
+	                               bool bApproximate = true);
 
 	/** Resets the state of a given spring */
 	UFUNCTION(BlueprintCallable, Category = "Math|Interpolation")
diff --git a/Engine/Source/Runtime/Engine/Private/Animation/AnimInterpFilter.cpp b/Engine/Source/Runtime/Engine/Private/Animation/AnimInterpFilter.cpp
--- a/Engine/Source/Runtime/Engine/Private/Animation/AnimInterpFilter.cpp
+++ b/Engine/Source/Runtime/Engine/Private/Animation/AnimInterpFilter.cpp
@@ -159,7 +159,7 @@ float FFIRFilterTimeBased::UpdateAndGetFilteredData(float Input, float DeltaTime
 					FilterWindow.Push(FFilterData(Input, 0.0f));
 				}
 				const float OrigValue = FilterWindow[0].Input;
-				FMath::ExponentialSmoothing(FilterWindow[0].Input, Input, DeltaTime, WindowDuration / EULERS_NUMBER);
+				FMath::ExponentialSmoothingApprox(FilterWindow[0].Input, Input, DeltaTime, WindowDuration / EULERS_NUMBER);
 				if (MaxSpeed > 0.0f)
 				{
 					// Clamp the speed
@@ -179,8 +179,8 @@ float FFIRFilterTimeBased::UpdateAndGetFilteredData(float Input, float DeltaTime
 					FilterWindow.Push(FFilterData(0.0f, 0.0f));
 				}
 				const float OrigValue = FilterWindow[0].Input;
-				FMath::SpringDamperSmoothing(FilterWindow[0].Input, FilterWindow[1].Input, Input, DeltaTime,
-				                      WindowDuration / EULERS_NUMBER, DampingRatio);
+				FMath::SpringDamperSmoothingApprox(FilterWindow[0].Input, FilterWindow[1].Input, Input, 0.0f, DeltaTime,
+				                                   WindowDuration / EULERS_NUMBER, DampingRatio);
 				if (MaxSpeed > 0.0f)
 				{
 					// Clamp the speed
diff --git a/Engine/Source/Runtime/Engine/Private/Animation/BlendSpace.cpp b/Engine/Source/Runtime/Engine/Private/Animation/BlendSpace.cpp
--- a/Engine/Source/Runtime/Engine/Private/Animation/BlendSpace.cpp
+++ b/Engine/Source/Runtime/Engine/Private/Animation/BlendSpace.cpp
@@ -1518,7 +1518,7 @@ static void SmoothWeight(float& Output, float& OutputRate, float Input, float In
 	{
 		Output = Input;
 		OutputRate = InputRate;
-		FMath::CriticallyDampedSmoothing(Output, OutputRate, Target, DeltaTime, SmoothingTimeFromSpeed(Speed));
+		FMath::CriticallyDampedSmoothing(Output, OutputRate, Target, 0.0f, DeltaTime, SmoothingTimeFromSpeed(Speed));
 	}
 	else
 	{
diff --git a/Engine/Source/Runtime/Engine/Private/KismetMathLibrary.cpp b/Engine/Source/Runtime/Engine/Private/KismetMathLibrary.cpp
--- a/Engine/Source/Runtime/Engine/Private/KismetMathLibrary.cpp
+++ b/Engine/Source/Runtime/Engine/Private/KismetMathLibrary.cpp
@@ -345,39 +345,62 @@ float UKismetMathLibrary::Ease(float A, float B, float Alpha, TEnumAsByte<EEasin
 	return Lerp(A, B, EaseAlpha(Alpha, EasingFunc, BlendExp, Steps));
 }
 
-float ComputeDamping(float Mass, float Stiffness, float CriticalDampingFactor)
-{
-	return 2 * FMath::Sqrt(Mass * Stiffness) * CriticalDampingFactor;
-}
-
 template <typename T>
-T GenericSpringInterp(T Current, T Target, T& PrevError, T& Velocity, float Stiffness, float CriticalDamping, float DeltaTime, float Mass)
+void GenericSpringInterp(T& Current, const T Target, T& PrevTarget, bool& bPrevTargetValid, T& Velocity,
+                         float Stiffness, float CriticalDamping, float DeltaTime,
+                         float TargetVelocityAmount, float Mass, bool bApproximate)
 {
 	if (DeltaTime > SMALL_NUMBER)
 	{
 		if (!FMath::IsNearlyZero(Mass))
 		{
-			const float Damping = ComputeDamping(Mass, Stiffness, CriticalDamping);
-			const T Error = Target - Current;
-			const T ErrorDeriv = (Error - PrevError);	//ignore divide by delta time since we multiply later anyway
-			Velocity += (Error * Stiffness * DeltaTime + ErrorDeriv * Damping) / Mass;
-			PrevError = Error;
-
-			const T NewValue = Current + Velocity * DeltaTime;
-			return NewValue;
-		}
-		else
-		{
-			return Target;
+			// Note that old target was stored in PrevTarget
+			T TargetVel = bPrevTargetValid ? (Target - PrevTarget) * (TargetVelocityAmount / DeltaTime) : T(0.f);
+			const float Omega = FMath::Sqrt(Stiffness / Mass);
+			const float SmoothingTime = 2.0f / Omega;
+			T NewValue = Current;
+			if (bApproximate)
+			{
+				FMath::SpringDamperSmoothingApprox(Current, Velocity, Target, TargetVel, DeltaTime, SmoothingTime, CriticalDamping);
+			}
+			else
+			{
+				FMath::SpringDamperSmoothing(Current, Velocity, Target, TargetVel, DeltaTime, SmoothingTime, CriticalDamping);
+			}
+			PrevTarget = Target;
+			bPrevTargetValid = true;
 		}
 	}
-
-	return Current;
 }
 
-float UKismetMathLibrary::FloatSpringInterp(float Current, float Target, FFloatSpringState& SpringState, float Stiffness, float CriticalDamping, float DeltaTime, float Mass)
+float UKismetMathLibrary::FloatSpringInterp(float Current, float Target, FFloatSpringState& SpringState,
+                                            float Stiffness, float CriticalDamping, float DeltaTime,
+                                            float Mass, float TargetVelocityAmount, 
+                                            bool bClamp, float MinValue, float MaxValue,
+                                            bool bApproximate)
 {
-	return GenericSpringInterp(Current, Target, SpringState.PrevError, SpringState.Velocity, Stiffness, CriticalDamping, DeltaTime, Mass);
+	GenericSpringInterp(Current, Target, SpringState.PrevTarget, SpringState.bPrevTargetValid, SpringState.Velocity, Stiffness, CriticalDamping,
+	                    DeltaTime, TargetVelocityAmount, Mass, bApproximate);
+	if (bClamp)
+	{
+		if (Current < MinValue)
+		{
+			Current = MinValue;
+			if (SpringState.Velocity < 0.0f)
+			{
+				SpringState.Velocity = 0.0f;
+			}
+		}
+		else if (Current > MaxValue)
+		{
+			Current = MaxValue;
+			if (SpringState.Velocity > 0.0f)
+			{
+				SpringState.Velocity = 0.0f;
+			}
+		}
+	}
+	return Current;
 }
 
 FVector  UKismetMathLibrary::RotateAngleAxis(FVector InVect, float AngleDeg, FVector Axis)
@@ -390,9 +413,37 @@ FVector UKismetMathLibrary::VEase(FVector A, FVector B, float Alpha, TEnumAsByte
 	return VLerp(A, B, EaseAlpha(Alpha, EasingFunc, BlendExp, Steps));
 }
 
-FVector UKismetMathLibrary::VectorSpringInterp(FVector Current, FVector Target, FVectorSpringState& SpringState, float Stiffness, float CriticalDamping, float DeltaTime, float Mass)
+FVector UKismetMathLibrary::VectorSpringInterp(FVector Current, FVector Target, FVectorSpringState& SpringState,
+                                               float Stiffness, float CriticalDamping, float DeltaTime,
+                                               float Mass, float TargetVelocityAmount, 
+                                               bool bClamp, FVector MinValue, FVector MaxValue,
+                                               bool bApproximate)
 {
-	return GenericSpringInterp(Current, Target, SpringState.PrevError, SpringState.Velocity, Stiffness, CriticalDamping, DeltaTime, Mass);
+	GenericSpringInterp(Current, Target, SpringState.PrevTarget, SpringState.bPrevTargetValid, SpringState.Velocity, Stiffness,
+	                    CriticalDamping, DeltaTime, TargetVelocityAmount, Mass, bApproximate);
+	if (bClamp)
+	{
+		for (int Index = 0 ; Index != 3 ; ++Index)
+		{
+			if (Current[Index] < MinValue[Index])
+			{
+				Current[Index] = MinValue[Index];
+				if (SpringState.Velocity[Index] < 0.0f)
+				{
+					SpringState.Velocity[Index] = 0.0f;
+				}
+			}
+			else if (Current[Index] > MaxValue[Index])
+			{
+				Current[Index] = MaxValue[Index];
+				if (SpringState.Velocity[Index] > 0.0f)
+				{
+					SpringState.Velocity[Index] = 0.0f;
+				}
+			}
+		}
+	}
+	return Current;
 }
 
 void UKismetMathLibrary::ResetFloatSpringState(FFloatSpringState& SpringState)
