Added LexFromString and LexToString to FBlake3Hash and FIoHash

This follows the addition of LexToString to FBlake3Hash yesterday.

#rnx

diff --git a/Engine/Source/Runtime/Core/Private/Hash/Blake3.cpp b/Engine/Source/Runtime/Core/Private/Hash/Blake3.cpp
--- a/Engine/Source/Runtime/Core/Private/Hash/Blake3.cpp
+++ b/Engine/Source/Runtime/Core/Private/Hash/Blake3.cpp
@@ -40,7 +40,12 @@ FBlake3Hash FBlake3::HashBuffer(const void* Data, uint64 Size)
 
 FString LexToString(const FBlake3Hash& Hash)
 {
-	return BytesToHex((const uint8*)Hash.GetBytes(), sizeof(decltype(Hash.GetBytes())));
+	FString Output;
+	TArray<TCHAR>& CharArray = Output.GetCharArray();
+	CharArray.AddUninitialized(sizeof(FBlake3Hash::ByteArray) * 2 + 1);
+	UE::String::BytesToHexLower(Hash.GetBytes(), CharArray.GetData());
+	CharArray.Last() = 0;
+	return Output;
 }
 
 ///////////////////////////////////////////////////////////////////////////////////////////////////////////////////////
diff --git a/Engine/Source/Runtime/Core/Private/IO/IoHash.cpp b/Engine/Source/Runtime/Core/Private/IO/IoHash.cpp
--- a/Engine/Source/Runtime/Core/Private/IO/IoHash.cpp
+++ b/Engine/Source/Runtime/Core/Private/IO/IoHash.cpp
@@ -4,12 +4,12 @@
 
 #include "Containers/UnrealString.h"
 
-FString FIoHash::ToString() const
+FString LexToString(const FIoHash& Hash)
 {
 	FString Output;
 	TArray<TCHAR>& CharArray = Output.GetCharArray();
-	CharArray.AddUninitialized(sizeof(ByteArray) * 2 + 1);
-	UE::String::BytesToHexLower(Hash, CharArray.GetData());
+	CharArray.AddUninitialized(sizeof(FIoHash::ByteArray) * 2 + 1);
+	UE::String::BytesToHexLower(Hash.GetBytes(), CharArray.GetData());
 	CharArray.Last() = 0;
 	return Output;
 }
diff --git a/Engine/Source/Runtime/Core/Public/Hash/Blake3.h b/Engine/Source/Runtime/Core/Public/Hash/Blake3.h
--- a/Engine/Source/Runtime/Core/Public/Hash/Blake3.h
+++ b/Engine/Source/Runtime/Core/Public/Hash/Blake3.h
@@ -137,9 +137,6 @@ inline FArchive& operator<<(FArchive& Ar, FBlake3Hash& Hash)
 	return Ar;
 }
 
-/** Stringifies the hash for textual output (e.g. to log file) */
-CORE_API FString LexToString(const FBlake3Hash& Hash);
-
 inline uint32 GetTypeHash(const FBlake3Hash& Hash)
 {
 	return *reinterpret_cast<const uint32*>(Hash.GetBytes());
@@ -156,3 +153,12 @@ inline FWideStringBuilderBase& operator<<(FWideStringBuilderBase& Builder, const
 	UE::String::BytesToHexLower(Hash.GetBytes(), Builder);
 	return Builder;
 }
+
+/** Construct a hash from a 64-character hex string. */
+inline void LexFromString(FBlake3Hash& OutHash, const TCHAR* Buffer)
+{
+	OutHash = FBlake3Hash(Buffer);
+}
+
+/** Convert a hash to a 64-character hex string. */
+UE_NODISCARD CORE_API FString LexToString(const FBlake3Hash& Hash);
diff --git a/Engine/Source/Runtime/Core/Public/IO/IoHash.h b/Engine/Source/Runtime/Core/Public/IO/IoHash.h
--- a/Engine/Source/Runtime/Core/Public/IO/IoHash.h
+++ b/Engine/Source/Runtime/Core/Public/IO/IoHash.h
@@ -49,9 +49,6 @@ public:
 	inline ByteArray& GetBytes() { return Hash; }
 	inline const ByteArray& GetBytes() const { return Hash; }
 
-	/** Returns the hash as a 40-character hex string. */
-	CORE_API FString ToString() const;
-
 	/** Calculate the hash of the input data. */
 	static inline FIoHash HashBuffer(const void* Data, uint64 Size);
 
@@ -146,3 +143,12 @@ inline FWideStringBuilderBase& operator<<(FWideStringBuilderBase& Builder, const
 	UE::String::BytesToHexLower(Hash.GetBytes(), Builder);
 	return Builder;
 }
+
+/** Construct a hash from a 40-character hex string. */
+inline void LexFromString(FIoHash& OutHash, const TCHAR* Buffer)
+{
+	OutHash = FIoHash(Buffer);
+}
+
+/** Convert a hash to a 40-character hex string. */
+UE_NODISCARD CORE_API FString LexToString(const FIoHash& Hash);
