SharedBuffer: Require an rvalue reference for MakeUnique and take references in MakeView

#rnx

diff --git a/Engine/Source/Runtime/Core/Private/Memory/SharedBuffer.cpp b/Engine/Source/Runtime/Core/Private/Memory/SharedBuffer.cpp
--- a/Engine/Source/Runtime/Core/Private/Memory/SharedBuffer.cpp
+++ b/Engine/Source/Runtime/Core/Private/Memory/SharedBuffer.cpp
@@ -63,16 +63,18 @@ protected:
 class FBufferOwnerOuterView final : public FBufferOwner
 {
 public:
-	inline FBufferOwnerOuterView(void* Data, uint64 Size, FSharedBuffer InOuterBuffer)
+	inline FBufferOwnerOuterView(void* Data, uint64 Size, FSharedBuffer&& InOuterBuffer)
 		: FBufferOwner(Data, Size)
 		, OuterBuffer(MoveTemp(InOuterBuffer))
 	{
-		check(OuterBuffer.GetView().Contains(MakeMemoryView(Data, Size)));
-		SetIsMaterialized();
-		if (OuterBuffer.IsOwned())
-		{
-			SetIsOwned();
-		}
+		Initialize();
+	}
+
+	inline FBufferOwnerOuterView(void* Data, uint64 Size, const FSharedBuffer& InOuterBuffer)
+		: FBufferOwner(Data, Size)
+		, OuterBuffer(InOuterBuffer)
+	{
+		Initialize();
 	}
 
 protected:
@@ -82,6 +84,16 @@ protected:
 	}
 
 private:
+	inline void Initialize()
+	{
+		check(OuterBuffer.GetView().Contains(MakeMemoryView(GetData(), GetSize())));
+		SetIsMaterialized();
+		if (OuterBuffer.IsOwned())
+		{
+			SetIsOwned();
+		}
+	}
+
 	FSharedBuffer OuterBuffer;
 };
 
@@ -116,7 +128,7 @@ FUniqueBuffer FUniqueBuffer::MakeView(void* Data, uint64 Size)
 	return FUniqueBuffer(new BufferOwnerPrivate::FBufferOwnerView(Data, Size));
 }
 
-FUniqueBuffer FUniqueBuffer::MakeUnique(FSharedBuffer Buffer)
+FUniqueBuffer FUniqueBuffer::MakeUnique(FSharedBuffer&& Buffer)
 {
 	OwnerPtrType ExistingOwner = ToPrivateOwnerPtr(MoveTemp(Buffer));
 	if (!ExistingOwner || (ExistingOwner->IsOwned() && ExistingOwner->GetTotalRefCount() == 1))
@@ -169,7 +181,7 @@ FSharedBuffer FSharedBuffer::MakeView(FMemoryView View)
 	return MakeView(View.GetData(), View.GetSize());
 }
 
-FSharedBuffer FSharedBuffer::MakeView(FMemoryView View, FSharedBuffer OuterBuffer)
+FSharedBuffer FSharedBuffer::MakeView(FMemoryView View, FSharedBuffer&& OuterBuffer)
 {
 	if (OuterBuffer.IsNull())
 	{
@@ -183,16 +195,35 @@ FSharedBuffer FSharedBuffer::MakeView(FMemoryView View, FSharedBuffer OuterBuffe
 		const_cast<void*>(View.GetData()), View.GetSize(), MoveTemp(OuterBuffer)));
 }
 
+FSharedBuffer FSharedBuffer::MakeView(FMemoryView View, const FSharedBuffer& OuterBuffer)
+{
+	if (OuterBuffer.IsNull())
+	{
+		return MakeView(View);
+	}
+	if (View == OuterBuffer.GetView())
+	{
+		return OuterBuffer;
+	}
+	return FSharedBuffer(new BufferOwnerPrivate::FBufferOwnerOuterView(
+		const_cast<void*>(View.GetData()), View.GetSize(), OuterBuffer));
+}
+
 FSharedBuffer FSharedBuffer::MakeView(const void* Data, uint64 Size)
 {
 	return FSharedBuffer(new BufferOwnerPrivate::FBufferOwnerView(const_cast<void*>(Data), Size));
 }
 
-FSharedBuffer FSharedBuffer::MakeView(const void* Data, uint64 Size, FSharedBuffer OuterBuffer)
+FSharedBuffer FSharedBuffer::MakeView(const void* Data, uint64 Size, FSharedBuffer&& OuterBuffer)
 {
 	return MakeView(MakeMemoryView(Data, Size), MoveTemp(OuterBuffer));
 }
 
+FSharedBuffer FSharedBuffer::MakeView(const void* Data, uint64 Size, const FSharedBuffer& OuterBuffer)
+{
+	return MakeView(MakeMemoryView(Data, Size), OuterBuffer);
+}
+
 FSharedBuffer::FSharedBuffer(FBufferOwner* InOwner)
 	: Owner(InOwner)
 {
diff --git a/Engine/Source/Runtime/Core/Private/Tests/Memory/SharedBufferTest.cpp b/Engine/Source/Runtime/Core/Private/Tests/Memory/SharedBufferTest.cpp
--- a/Engine/Source/Runtime/Core/Private/Tests/Memory/SharedBufferTest.cpp
+++ b/Engine/Source/Runtime/Core/Private/Tests/Memory/SharedBufferTest.cpp
@@ -191,7 +191,7 @@ bool FUniqueBufferTest::RunTest(const FString& Parameters)
 	}
 	{
 		FSharedBuffer Shared = FSharedBuffer::Clone(MakeMemoryView<uint8>({1, 2, 3, 4}));
-		FUniqueBuffer Buffer = FUniqueBuffer::MakeUnique(Shared);
+		FUniqueBuffer Buffer = FUniqueBuffer::MakeUnique(FSharedBuffer(Shared));
 		TestTrue(TEXT("FUniqueBuffer::MakeUnique(Shared).IsOwned()"), Buffer.IsOwned());
 		TestEqual(TEXT("FUniqueBuffer::MakeUnique(Shared).GetSize()"), Buffer.GetSize(), Shared.GetSize());
 		TestNotEqual(TEXT("FUniqueBuffer::MakeUnique(Shared).GetData()"), const_cast<const void*>(Buffer.GetData()), Shared.GetData());
diff --git a/Engine/Source/Runtime/Core/Public/Memory/SharedBuffer.h b/Engine/Source/Runtime/Core/Public/Memory/SharedBuffer.h
--- a/Engine/Source/Runtime/Core/Public/Memory/SharedBuffer.h
+++ b/Engine/Source/Runtime/Core/Public/Memory/SharedBuffer.h
@@ -235,7 +235,7 @@ public:
 	 * Steals the buffer owner from the shared buffer if this is the last reference to it, otherwise
 	 * clones the shared buffer to guarantee unique ownership. An non-owned buffer is always cloned.
 	 */
-	UE_NODISCARD CORE_API static FUniqueBuffer MakeUnique(FSharedBuffer Buffer);
+	UE_NODISCARD CORE_API static FUniqueBuffer MakeUnique(FSharedBuffer&& Buffer);
 
 	/** Construct a null unique buffer. */
 	FUniqueBuffer() = default;
@@ -279,7 +279,7 @@ public:
 	UE_NODISCARD inline bool IsOwned() const { return !Owner || Owner->IsOwned(); }
 
 	/** Returns a buffer that is owned, by cloning if not owned. */
-	CORE_API FUniqueBuffer MakeOwned() &&;
+	UE_NODISCARD CORE_API FUniqueBuffer MakeOwned() &&;
 
 	/** Returns true if the referenced buffer has been materialized. */
 	UE_NODISCARD inline bool IsMaterialized() const { return !Owner || Owner->IsMaterialized(); }
@@ -321,8 +321,10 @@ public:
 	UE_NODISCARD CORE_API static FSharedBuffer MakeView(const void* Data, uint64 Size);
 
 	/** Make a view of the input within its outer buffer. Ownership matches OuterBuffer. */
-	UE_NODISCARD CORE_API static FSharedBuffer MakeView(FMemoryView View, FSharedBuffer OuterBuffer);
-	UE_NODISCARD CORE_API static FSharedBuffer MakeView(const void* Data, uint64 Size, FSharedBuffer OuterBuffer);
+	UE_NODISCARD CORE_API static FSharedBuffer MakeView(FMemoryView View, FSharedBuffer&& OuterBuffer);
+	UE_NODISCARD CORE_API static FSharedBuffer MakeView(FMemoryView View, const FSharedBuffer& OuterBuffer);
+	UE_NODISCARD CORE_API static FSharedBuffer MakeView(const void* Data, uint64 Size, FSharedBuffer&& OuterBuffer);
+	UE_NODISCARD CORE_API static FSharedBuffer MakeView(const void* Data, uint64 Size, const FSharedBuffer& OuterBuffer);
 
 	/**
 	 * Make an owned buffer by taking ownership of the input.
@@ -388,8 +390,8 @@ public:
 	UE_NODISCARD inline bool IsOwned() const { return !Owner || Owner->IsOwned(); }
 
 	/** Returns a buffer that is owned, by cloning if not owned. */
-	CORE_API FSharedBuffer MakeOwned() const &;
-	CORE_API FSharedBuffer MakeOwned() &&;
+	UE_NODISCARD CORE_API FSharedBuffer MakeOwned() const &;
+	UE_NODISCARD CORE_API FSharedBuffer MakeOwned() &&;
 
 	/** Returns true if the referenced buffer has been materialized. */
 	UE_NODISCARD inline bool IsMaterialized() const { return !Owner || Owner->IsMaterialized(); }
@@ -838,7 +840,7 @@ private:
 
 /** Construct a shared buffer by taking ownership of an array. */
 template <typename T, typename Allocator>
-UE_NODISCARD FSharedBuffer MakeSharedBufferFromArray(TArray<T, Allocator>&& Array)
+UE_NODISCARD inline FSharedBuffer MakeSharedBufferFromArray(TArray<T, Allocator>&& Array)
 {
 	return FSharedBuffer(new BufferOwnerPrivate::TBufferOwnerTArray<T, Allocator>(MoveTemp(Array)));
 }
