[TraceLog] Fix channel presets early initialization

In a previous change TraceAux::Initialize was moved to a point very early in process lifetime. At this point the config system is not ready, so resolving data driven channel presets does not work. To solve this we make two changes, define two presets in code (default and memory) which require very early application, and add a second pass of channel initialization which is executed directly after config system is up.

diff --git a/Engine/Config/BaseEngine.ini b/Engine/Config/BaseEngine.ini
--- a/Engine/Config/BaseEngine.ini
+++ b/Engine/Config/BaseEngine.ini
@@ -2047,9 +2047,7 @@ DrawColor=(R=38,G=38,B=38,A=64)
 DrawColor=(R=140,G=255,B=0,A=164)
 
 [Trace.ChannelPresets]
-Default=gpu,cpu,frame,log,bookmark
 Rendering=gpu,cpu,frame,log,bookmark
-Memory=memtag,memalloc,callstack,module
 
 [RemoteConfiguration]
 Enabled=false
diff --git a/Engine/Source/Runtime/Core/Private/ProfilingDebugging/TraceAuxiliary.cpp b/Engine/Source/Runtime/Core/Private/ProfilingDebugging/TraceAuxiliary.cpp
--- a/Engine/Source/Runtime/Core/Private/ProfilingDebugging/TraceAuxiliary.cpp
+++ b/Engine/Source/Runtime/Core/Private/ProfilingDebugging/TraceAuxiliary.cpp
@@ -34,6 +34,10 @@
 #include "Windows/HideWindowsPlatformTypes.h"
 #endif
 
+////////////////////////////////////////////////////////////////////////////////
+const TCHAR* GDefaultChannels = TEXT("cpu,gpu,frame,log,bookmark");
+const TCHAR* GMemoryChannels = TEXT("memtag,memalloc,callstack,module");
+
 ////////////////////////////////////////////////////////////////////////////////
 enum class ETraceConnectType
 {
@@ -99,7 +103,17 @@ void FTraceAuxiliaryImpl::AddChannels(const TCHAR* ChannelList, bool bResolvePre
 		if (bResolvePresets)
 		{
 			FString Value;
-			if (GConfig->GetString(TEXT("Trace.ChannelPresets"), Name, Value, GEngineIni))
+			// Check against hard coded presets
+			if(FCString::Stricmp(Name, TEXT("default")) == 0)
+			{
+				AddChannels(GDefaultChannels, false);
+			}
+			else if(FCString::Stricmp(Name, TEXT("memory"))== 0)
+			{
+				AddChannels(GMemoryChannels, false);
+			}
+			// Check against data driven presets (if available)
+			else if (GConfig && GConfig->GetString(TEXT("Trace.ChannelPresets"), Name, Value, GEngineIni))
 			{
 				AddChannels(*Value, false);
 				return;
@@ -161,13 +175,7 @@ bool FTraceAuxiliaryImpl::Connect(ETraceConnectType Type, const TCHAR* Parameter
 	// some defaults for the user. Less futzing.
 	if (!Channels.Num())
 	{
-		FString Value;
-		if (!GConfig->GetString(TEXT("Trace.ChannelPresets"), TEXT("Default"), Value, GEngineIni))
-		{
-			Value = TEXT("cpu,frame,log,bookmark,gpu");
-		}
-
-		AddChannels(*Value);
+		AddChannels(GDefaultChannels);
 	}
 
 	EnableChannels();
@@ -477,6 +485,21 @@ void FTraceAuxiliary::Initialize(const TCHAR* CommandLine)
 #endif
 }
 
+////////////////////////////////////////////////////////////////////////////////
+void FTraceAuxiliary::InitializePresets(const TCHAR* CommandLine)
+{
+#if UE_TRACE_ENABLED
+	// Second pass over trace arguments, this time to allow config defined presets
+	// to be applied.
+	FString Parameter;
+	if (FParse::Value(CommandLine, TEXT("-trace="), Parameter, false))
+	{
+		GTraceAuxiliary.AddChannels(*Parameter);
+		GTraceAuxiliary.EnableChannels();
+	}
+#endif 
+}
+
 ////////////////////////////////////////////////////////////////////////////////
 void FTraceAuxiliary::EnableChannels()
 {
diff --git a/Engine/Source/Runtime/Core/Public/ProfilingDebugging/TraceAuxiliary.h b/Engine/Source/Runtime/Core/Public/ProfilingDebugging/TraceAuxiliary.h
--- a/Engine/Source/Runtime/Core/Public/ProfilingDebugging/TraceAuxiliary.h
+++ b/Engine/Source/Runtime/Core/Public/ProfilingDebugging/TraceAuxiliary.h
@@ -9,6 +9,7 @@ class CORE_API FTraceAuxiliary
 {
 public:
 	static void Initialize(const TCHAR* CommandLine);
+	static void InitializePresets(const TCHAR* CommandLine);
 	static void TryAutoConnect();
 	static void EnableChannels();
 };
diff --git a/Engine/Source/Runtime/Launch/Private/LaunchEngineLoop.cpp b/Engine/Source/Runtime/Launch/Private/LaunchEngineLoop.cpp
--- a/Engine/Source/Runtime/Launch/Private/LaunchEngineLoop.cpp
+++ b/Engine/Source/Runtime/Launch/Private/LaunchEngineLoop.cpp
@@ -5576,6 +5576,9 @@ bool FEngineLoop::AppInit( )
 		FConfigCacheIni::InitializeConfigSystem();
 	}
 
+	// Apply config driven presets
+	FTraceAuxiliary::InitializePresets(FCommandLine::Get());
+
 	FDelayedAutoRegisterHelper::RunAndClearDelayedAutoRegisterDelegates(EDelayedRegisterRunPhase::IniSystemReady);
 
 	// Load "asap" plugin modules
