changed `FUnixPlatformMisc::NumberOfWorkerThreadsToSpawn` to match the logic for Windows, as it was too pessimistic (outdated) for modern hw with high number of cores. This affects Linux, LinuxAArch64 and platforms derived from them

diff --git a/Engine/Source/Runtime/Core/Private/Unix/UnixPlatformMisc.cpp b/Engine/Source/Runtime/Core/Private/Unix/UnixPlatformMisc.cpp
--- a/Engine/Source/Runtime/Core/Private/Unix/UnixPlatformMisc.cpp
+++ b/Engine/Source/Runtime/Core/Private/Unix/UnixPlatformMisc.cpp
@@ -1114,3 +1114,28 @@ bool FUnixPlatformMisc::SetStoredValues(const FString& InStoreId, const FString&
 	ConfigFile.Dirty = true;
 	return ConfigFile.Write(ConfigPath);
 }
+
+int32 FUnixPlatformMisc::NumberOfWorkerThreadsToSpawn()
+{
+	static int32 MaxServerWorkerThreads = 4;
+
+	extern CORE_API int32 GUseNewTaskBackend;
+	int32 MaxWorkerThreads = GUseNewTaskBackend ? INT32_MAX : 26;
+
+	int32 NumberOfCores = FPlatformMisc::NumberOfCores();
+	int32 NumberOfCoresIncludingHyperthreads = FPlatformMisc::NumberOfCoresIncludingHyperthreads();
+	int32 NumberOfThreads = 0;
+
+	if (NumberOfCoresIncludingHyperthreads > NumberOfCores)
+	{
+		NumberOfThreads = NumberOfCoresIncludingHyperthreads - 2;
+	}
+	else
+	{
+		NumberOfThreads = NumberOfCores - 1;
+	}
+
+	int32 MaxWorkerThreadsWanted = IsRunningDedicatedServer() ? MaxServerWorkerThreads : MaxWorkerThreads;
+	// need to spawn at least one worker thread (see FTaskGraphImplementation)
+	return FMath::Max(FMath::Min(NumberOfThreads, MaxWorkerThreadsWanted), 2);
+}
\ No newline at end of file
diff --git a/Engine/Source/Runtime/Core/Public/Unix/UnixPlatformMisc.h b/Engine/Source/Runtime/Core/Public/Unix/UnixPlatformMisc.h
--- a/Engine/Source/Runtime/Core/Public/Unix/UnixPlatformMisc.h
+++ b/Engine/Source/Runtime/Core/Public/Unix/UnixPlatformMisc.h
@@ -139,6 +139,8 @@ struct CORE_API FUnixPlatformMisc : public FGenericPlatformMisc
 
 	static bool SetStoredValues(const FString& InStoreId, const FString& InSectionName, const TMap<FString, FString>& InKeyValues);
 
+	static int32 NumberOfWorkerThreadsToSpawn();
+
 #if STATS || ENABLE_STATNAMEDEVENTS
 	static void BeginNamedEventFrame();
 	static void BeginNamedEvent(const struct FColor& Color, const TCHAR* Text);
