disabled the old algo of tasks execution context checks to handle busy waiting on named threads

diff --git a/Engine/Source/Runtime/Core/Private/Async/Fundamental/Scheduler.cpp b/Engine/Source/Runtime/Core/Private/Async/Fundamental/Scheduler.cpp
--- a/Engine/Source/Runtime/Core/Private/Async/Fundamental/Scheduler.cpp
+++ b/Engine/Source/Runtime/Core/Private/Async/Fundamental/Scheduler.cpp
@@ -14,6 +14,7 @@ namespace LowLevelTasks
 	thread_local FTask* FScheduler::ActiveTask = nullptr;
 	thread_local FScheduler* FScheduler::ActiveScheduler = nullptr;
 	thread_local FScheduler::EWorkerType FScheduler::WorkerType = FScheduler::EWorkerType::None;
+	thread_local uint32 FScheduler::BusyWaitingDepth = 0;
 
 	FScheduler FScheduler::Singleton;
 
@@ -295,6 +296,9 @@ namespace LowLevelTasks
 		TRACE_CPUPROFILER_EVENT_SCOPE(FScheduler::BusyWaitInternal);
 		FTaskTagScope WorkerScope(ETaskTag::EWorkerThread);
 
+		++BusyWaitingDepth;
+		ON_SCOPE_EXIT{ --BusyWaitingDepth; };
+
 		checkSlow(LocalQueue != nullptr);
 		check(ActiveWorkers.load(std::memory_order_relaxed));
 		FLocalQueueType* WorkerLocalQueue = LocalQueue;
diff --git a/Engine/Source/Runtime/Core/Private/HAL/ThreadingBase.cpp b/Engine/Source/Runtime/Core/Private/HAL/ThreadingBase.cpp
--- a/Engine/Source/Runtime/Core/Private/HAL/ThreadingBase.cpp
+++ b/Engine/Source/Runtime/Core/Private/HAL/ThreadingBase.cpp
@@ -12,6 +12,7 @@
 #include "HAL/IConsoleManager.h"
 #include "HAL/PlatformStackWalk.h"
 #include "ProfilingDebugging/MiscTrace.h"
+#include "Async/Fundamental/Scheduler.h"
 
 #include <atomic>
 
@@ -158,10 +159,13 @@ CORE_API bool IsInGameThread()
 	{
 		bool newValue = FTaskTagScope::IsCurrentTag(ETaskTag::EGameThread) || FTaskTagScope::IsRunningDuringStaticInit();
 #if !UE_BUILD_SHIPPING && !UE_BUILD_TEST
-		const uint32 CurrentThreadId = FPlatformTLS::GetCurrentThreadId();
-		bool oldValue = CurrentThreadId == GGameThreadId;
-		ensureMsgf(oldValue == newValue, TEXT("oldValue(%i) newValue(%i) If this check fails make sure that there is a FTaskTagScope(ETaskTag::EGameThread) as deep as possible on the current callstack, you can see the current value in ActiveNamedThreads(%x)"), oldValue, newValue, FTaskTagScope::GetCurrentTag());
-		newValue = oldValue;
+		if (!LowLevelTasks::FScheduler::IsBusyWaiting())
+		{
+			const uint32 CurrentThreadId = FPlatformTLS::GetCurrentThreadId();
+			bool oldValue = CurrentThreadId == GGameThreadId;
+			ensureMsgf(oldValue == newValue, TEXT("oldValue(%i) newValue(%i) If this check fails make sure that there is a FTaskTagScope(ETaskTag::EGameThread) as deep as possible on the current callstack, you can see the current value in ActiveNamedThreads(%x)"), oldValue, newValue, FTaskTagScope::GetCurrentTag());
+			newValue = oldValue;
+		}
 #endif
 		return newValue;
 	}
@@ -179,9 +183,12 @@ CORE_API bool IsInSlateThread()
 	// If this explicitly is a slate thread, not just the main thread running slate
 	bool newValue = FTaskTagScope::IsCurrentTag(ETaskTag::ESlateThread);
 #if !UE_BUILD_SHIPPING && !UE_BUILD_TEST
-	bool oldValue = GSlateLoadingThreadId != 0 && FPlatformTLS::GetCurrentThreadId() == GSlateLoadingThreadId;
-	ensureMsgf(oldValue == newValue, TEXT("oldValue(%i) newValue(%i) If this check fails make sure that there is a FTaskTagScope(ETaskTag::ESlateThread) as as deep as possible on the current callstack, you can see the current value in ActiveNamedThreads(%x)"), oldValue, newValue, FTaskTagScope::GetCurrentTag());
-	newValue = oldValue;
+	if (!LowLevelTasks::FScheduler::IsBusyWaiting())
+	{
+		bool oldValue = GSlateLoadingThreadId != 0 && FPlatformTLS::GetCurrentThreadId() == GSlateLoadingThreadId;
+		ensureMsgf(oldValue == newValue, TEXT("oldValue(%i) newValue(%i) If this check fails make sure that there is a FTaskTagScope(ETaskTag::ESlateThread) as as deep as possible on the current callstack, you can see the current value in ActiveNamedThreads(%x)"), oldValue, newValue, FTaskTagScope::GetCurrentTag());
+		newValue = oldValue;
+	}
 #endif
 	return newValue;
 }
@@ -205,9 +212,12 @@ PRAGMA_DISABLE_DEPRECATION_WARNINGS
 		? FTaskTagScope::IsCurrentTag(ETaskTag::EGameThread)
 		: FTaskTagScope::IsCurrentTag(ETaskTag::EAudioThread);
 #if !UE_BUILD_SHIPPING && !UE_BUILD_TEST
-	bool oldValue = FPlatformTLS::GetCurrentThreadId() == ((nullptr == GAudioThread || GIsAudioThreadSuspended.Load(EMemoryOrder::Relaxed)) ? GGameThreadId : GAudioThread->GetThreadID());
-	ensureMsgf(oldValue == newValue, TEXT("oldValue(%i) newValue(%i) If this check fails make sure that there is a FTaskTagScope(ETaskTag::EAudioThread) as deep as possible on the current callstack, you can see the current value in ActiveNamedThreads(%x)"), oldValue, newValue, FTaskTagScope::GetCurrentTag());
-	newValue = oldValue;
+	if (!LowLevelTasks::FScheduler::IsBusyWaiting())
+	{
+		bool oldValue = FPlatformTLS::GetCurrentThreadId() == ((nullptr == GAudioThread || GIsAudioThreadSuspended.Load(EMemoryOrder::Relaxed)) ? GGameThreadId : GAudioThread->GetThreadID());
+		ensureMsgf(oldValue == newValue, TEXT("oldValue(%i) newValue(%i) If this check fails make sure that there is a FTaskTagScope(ETaskTag::EAudioThread) as deep as possible on the current callstack, you can see the current value in ActiveNamedThreads(%x)"), oldValue, newValue, FTaskTagScope::GetCurrentTag());
+		newValue = oldValue;
+	}
 #endif
 	return newValue;
 PRAGMA_ENABLE_DEPRECATION_WARNINGS
@@ -222,9 +232,12 @@ CORE_API bool IsInActualRenderingThread()
 PRAGMA_DISABLE_DEPRECATION_WARNINGS
 	bool newValue = FTaskTagScope::IsCurrentTag(ETaskTag::ERenderingThread);
 #if !UE_BUILD_SHIPPING && !UE_BUILD_TEST
-	bool oldValue = FPlatformTLS::GetCurrentThreadId() == GRenderThreadId;
-	ensureMsgf(oldValue == newValue, TEXT("oldValue(%i) newValue(%i) If this check fails make sure that there is a FTaskTagScope(ETaskTag::ERenderingThread) as deep as possible on the current callstack, you can see the current value in ActiveNamedThreads(%x)"), oldValue, newValue, FTaskTagScope::GetCurrentTag());
-	newValue = oldValue;
+	if (!LowLevelTasks::FScheduler::IsBusyWaiting())
+	{
+		bool oldValue = FPlatformTLS::GetCurrentThreadId() == GRenderThreadId;
+		ensureMsgf(oldValue == newValue, TEXT("oldValue(%i) newValue(%i) If this check fails make sure that there is a FTaskTagScope(ETaskTag::ERenderingThread) as deep as possible on the current callstack, you can see the current value in ActiveNamedThreads(%x)"), oldValue, newValue, FTaskTagScope::GetCurrentTag());
+		newValue = oldValue;
+	}
 #endif
 	return newValue;
 PRAGMA_ENABLE_DEPRECATION_WARNINGS
@@ -240,9 +253,12 @@ PRAGMA_DISABLE_DEPRECATION_WARNINGS
 		: FTaskTagScope::IsCurrentTag(ETaskTag::ERenderingThread);
 
 #if !UE_BUILD_SHIPPING && !UE_BUILD_TEST
-	bool oldValue = !GRenderingThread || bLocalIsLoadingThreadSuspended || (FPlatformTLS::GetCurrentThreadId() == GRenderingThread->GetThreadID());
-	ensureMsgf(oldValue == newValue, TEXT("oldValue(%i) newValue(%i) If this check fails make sure that there is a FTaskTagScope(ETaskTag::ERenderingThread) as deep as possible on the current callstack, you can see the current value in ActiveNamedThreads(%x), GRenderingThread(%x), GIsRenderingThreadSuspended(%d)"), oldValue, newValue, FTaskTagScope::GetCurrentTag(), GRenderingThread, bLocalIsLoadingThreadSuspended);
-	newValue = oldValue;
+	if (!LowLevelTasks::FScheduler::IsBusyWaiting())
+	{
+		bool oldValue = !GRenderingThread || bLocalIsLoadingThreadSuspended || (FPlatformTLS::GetCurrentThreadId() == GRenderingThread->GetThreadID());
+		ensureMsgf(oldValue == newValue, TEXT("oldValue(%i) newValue(%i) If this check fails make sure that there is a FTaskTagScope(ETaskTag::ERenderingThread) as deep as possible on the current callstack, you can see the current value in ActiveNamedThreads(%x), GRenderingThread(%x), GIsRenderingThreadSuspended(%d)"), oldValue, newValue, FTaskTagScope::GetCurrentTag(), GRenderingThread, bLocalIsLoadingThreadSuspended);
+		newValue = oldValue;
+	}
 #endif
 	return newValue;
 PRAGMA_ENABLE_DEPRECATION_WARNINGS
@@ -267,17 +283,20 @@ PRAGMA_DISABLE_DEPRECATION_WARNINGS
 	}
 
 #if !UE_BUILD_SHIPPING && !UE_BUILD_TEST
-	bool oldValue = false;
-	if (!GRenderingThread || bLocalIsLoadingThreadSuspended)
-	{
-		oldValue = true;
-	}
-	else
+	if (!LowLevelTasks::FScheduler::IsBusyWaiting())
 	{
-		oldValue = FPlatformTLS::GetCurrentThreadId() != GGameThreadId;
+		bool oldValue = false;
+		if (!GRenderingThread || bLocalIsLoadingThreadSuspended)
+		{
+			oldValue = true;
+		}
+		else
+		{
+			oldValue = FPlatformTLS::GetCurrentThreadId() != GGameThreadId;
+		}
+		ensureMsgf(oldValue == newValue, TEXT("oldValue(%i) newValue(%i) If this check fails make sure that there is a FTaskTagScope(ETaskTag::EParallelRenderingThread) as deep as possible on the current callstack, you can see the current value in ActiveNamedThreads(%x), GRenderingThread(%x), GIsRenderingThreadSuspended(%d)"), oldValue, newValue, FTaskTagScope::GetCurrentTag(), GRenderingThread, bLocalIsLoadingThreadSuspended);
+		newValue = oldValue;
 	}
-	ensureMsgf(oldValue == newValue, TEXT("oldValue(%i) newValue(%i) If this check fails make sure that there is a FTaskTagScope(ETaskTag::EParallelRenderingThread) as deep as possible on the current callstack, you can see the current value in ActiveNamedThreads(%x), GRenderingThread(%x), GIsRenderingThreadSuspended(%d)"), oldValue, newValue, FTaskTagScope::GetCurrentTag(), GRenderingThread, bLocalIsLoadingThreadSuspended);
-	newValue = oldValue;
 #endif
 	return newValue;
 PRAGMA_ENABLE_DEPRECATION_WARNINGS
diff --git a/Engine/Source/Runtime/Core/Public/Async/Fundamental/Scheduler.h b/Engine/Source/Runtime/Core/Public/Async/Fundamental/Scheduler.h
--- a/Engine/Source/Runtime/Core/Public/Async/Fundamental/Scheduler.h
+++ b/Engine/Source/Runtime/Core/Public/Async/Fundamental/Scheduler.h
@@ -41,6 +41,9 @@ namespace LowLevelTasks
 		static thread_local EWorkerType WorkerType;
 		static CORE_API FScheduler Singleton;
 
+		// number of busy-waiting calls in the call-stack
+		static thread_local uint32 BusyWaitingDepth;
+
 		// using 16 bytes here because it fits the vtable and one additional pointer
 		using FConditional = TTaskDelegate<16, bool>;
 
@@ -81,6 +84,9 @@ namespace LowLevelTasks
 		template<typename TaskType>
 		inline void BusyWait(const TArrayView<const TaskType>& Tasks);
 
+		// returns true if the current thread execution is in the context of busy-waiting
+		inline static bool IsBusyWaiting();
+
 		//number of instantiated workers
 		inline uint32 GetNumWorkers() const;
 
@@ -302,6 +308,11 @@ namespace LowLevelTasks
 		}
 	}
 
+	inline bool FScheduler::IsBusyWaiting()
+	{
+		return BusyWaitingDepth != 0;
+	}
+
 	inline void FScheduler::TrySleeping(FSleepEvent* WorkerEvent, FQueueRegistry::FOutOfWork& OutOfWork, bool& Drowsing, bool bBackgroundWorker)
 	{
 		ESleepState DrowsingState1 = ESleepState::Drowsing;
