Fixed UE4 references in ConsoleCommandLibrary_DumpLibraryHTML()

diff --git a/Engine/Source/Runtime/Core/Private/Misc/Parse.cpp b/Engine/Source/Runtime/Core/Private/Misc/Parse.cpp
--- a/Engine/Source/Runtime/Core/Private/Misc/Parse.cpp
+++ b/Engine/Source/Runtime/Core/Private/Misc/Parse.cpp
@@ -141,16 +141,16 @@ void ConsoleCommandLibrary_DumpLibraryHTML(UWorld* InWorld, FExec& SubSystem, co
 		FArchive* File = nullptr;
 #endif
 
-		if(File)
+		if (File)
 		{
 			FLazyPrintf LazyPrintf(*TemplateFile);
 
 			// title
-			LazyPrintf.PushParam(TEXT("UE4 Console Variables and Commands"));
+			LazyPrintf.PushParam(TEXT("UE5 Console Variables and Commands"));
 			// headline
-			LazyPrintf.PushParam(TEXT("Unreal Engine 4 Console Variables and Commands"));
+			LazyPrintf.PushParam(TEXT("Unreal Engine 5 Console Variables and Commands"));
 			// generated by
-			LazyPrintf.PushParam(TEXT("Unreal Engine 4 console command 'Help'"));
+			LazyPrintf.PushParam(TEXT("Unreal Engine 5 console command 'Help'"));
 			// version
 			LazyPrintf.PushParam(TEXT("0.95"));
 			// date
@@ -158,7 +158,7 @@ void ConsoleCommandLibrary_DumpLibraryHTML(UWorld* InWorld, FExec& SubSystem, co
 
 			FString AllData;
 
-			for(TSet<FString>::TConstIterator It(LocalConsoleCommandLibrary.KnownNames); It; ++It)
+			for (TSet<FString>::TConstIterator It(LocalConsoleCommandLibrary.KnownNames); It; ++It)
 			{
 				const FString& Name = *It;
 
