[TraceLog] Trace module image id

In order to check for matching debug files we need to trace the module id (on Windows this means Guid and age).

diff --git a/Engine/Source/Runtime/Core/Private/ProfilingDebugging/Windows/ModuleDiagnostics.cpp b/Engine/Source/Runtime/Core/Private/ProfilingDebugging/Windows/ModuleDiagnostics.cpp
--- a/Engine/Source/Runtime/Core/Private/ProfilingDebugging/Windows/ModuleDiagnostics.cpp
+++ b/Engine/Source/Runtime/Core/Private/ProfilingDebugging/Windows/ModuleDiagnostics.cpp
@@ -42,6 +42,7 @@ UE_TRACE_EVENT_BEGIN(Diagnostics, ModuleLoad, NoSync|Important)
 	UE_TRACE_EVENT_FIELD(UE::Trace::WideString, Name)
 	UE_TRACE_EVENT_FIELD(uint32, Base)
 	UE_TRACE_EVENT_FIELD(uint32, Size)
+	UE_TRACE_EVENT_FIELD(uint8[], ImageId) // Platform specific id for this image, used to match debug files were available
 UE_TRACE_EVENT_END()
 
 UE_TRACE_EVENT_BEGIN(Diagnostics, ModuleUnload, NoSync|Important)
@@ -170,12 +171,46 @@ void FModuleTrace::OnDllLoaded(const UNICODE_STRING& Name, UPTRINT Base)
 	const auto* DosHeader = (IMAGE_DOS_HEADER*)Base;
 	const auto* NtHeaders = (IMAGE_NT_HEADERS*)(Base + DosHeader->e_lfanew);
 	const IMAGE_OPTIONAL_HEADER& OptionalHeader = NtHeaders->OptionalHeader;
+	TArray<uint8, TFixedAllocator<20>> ImageId;
+
+	// Find the guid and age of the binary, used to match debug files
+	const IMAGE_DATA_DIRECTORY& DebugInfoEntry = OptionalHeader.DataDirectory[IMAGE_DIRECTORY_ENTRY_DEBUG];
+	const auto* DebugEntries = (IMAGE_DEBUG_DIRECTORY*)(Base + DebugInfoEntry.VirtualAddress);
+	for (uint32 i = 0, n = DebugInfoEntry.Size / sizeof(DebugEntries[0]); i < n; ++i)
+	{
+		const IMAGE_DEBUG_DIRECTORY& Entry = DebugEntries[i];
+		if (Entry.Type == IMAGE_DEBUG_TYPE_CODEVIEW)
+		{
+			struct FCodeView7
+			{
+				uint32  Signature;
+				uint32  Guid[4];
+				uint32  Age;
+			};
+
+			if (Entry.SizeOfData < sizeof(FCodeView7))
+			{
+				continue;
+			}
+
+			const auto* CodeView7 = (FCodeView7*)(Base + Entry.AddressOfRawData);
+			if (CodeView7->Signature != 'SDSR')
+			{
+				continue;
+			}
+
+			ImageId.Append((uint8*)&CodeView7->Guid, sizeof(uint32) * 4);
+			ImageId.Append((uint8*)&CodeView7->Age, sizeof(uint32));
+			break;
+		}
+	}
 
 	// Note: UNICODE_STRING.Length is the size in bytes of the string buffer.
-	UE_TRACE_LOG(Diagnostics, ModuleLoad, ModuleChannel, Name.Length)
+	UE_TRACE_LOG(Diagnostics, ModuleLoad, ModuleChannel, Name.Length + ImageId.Num())
 		<< ModuleLoad.Name(Name.Buffer, Name.Length / 2)
 		<< ModuleLoad.Base(uint32(Base >> 16)) // Windows' DLLs are on 64K page boundaries
-		<< ModuleLoad.Size(OptionalHeader.SizeOfImage);
+		<< ModuleLoad.Size(OptionalHeader.SizeOfImage)
+		<< ModuleLoad.ImageId(ImageId.GetData(), ImageId.Num());
 
 	for (SubscribeFunc Subscriber : Subscribers)
 	{
