added Tasks::FPipe::IsInContext() to check if the current call is in pipe's execution context
added Tasks::FPipe::WaitUntilEmpty() to wait for the last piped task completion

diff --git a/Engine/Source/Runtime/Core/Private/Tasks/Pipe.cpp b/Engine/Source/Runtime/Core/Private/Tasks/Pipe.cpp
--- a/Engine/Source/Runtime/Core/Private/Tasks/Pipe.cpp
+++ b/Engine/Source/Runtime/Core/Private/Tasks/Pipe.cpp
@@ -33,4 +33,48 @@ namespace UE { namespace Tasks
 			}
 		}
 	}
+
+	// Maintains pipe callstack. Due to busy waiting tasks from multiple pipes can be executed nested.
+	class FPipeCallStack
+	{
+	public:
+		static void Push(const FPipe& Pipe)
+		{
+			CallStack.Add(&Pipe);
+		}
+
+		static void Pop(const FPipe& Pipe)
+		{
+			check(CallStack.Last() == &Pipe);
+			CallStack.Pop(/*bAllowShrinking=*/ false);
+		}
+
+		// returns true if a task from the given pipe is being being executed on the top of the stack.
+		// the method deliberately doesn't look deeper because even if the pipe is there and technically it's safe to assume
+		// accessing a resource protected by a pipe is thread-safe, logically it's a bug because it's an accidental condition
+		static bool IsOnTop(const FPipe& Pipe)
+		{
+			return CallStack.Num() != 0 && CallStack.Last() == &Pipe;
+		}
+
+	private:
+		static thread_local TArray<const FPipe*> CallStack;
+	};
+
+	thread_local TArray<const FPipe*> FPipeCallStack::CallStack;
+
+	void FPipe::ExecutionStarted()
+	{
+		FPipeCallStack::Push(*this);
+	}
+
+	void FPipe::ExecutionFinished()
+	{
+		FPipeCallStack::Pop(*this);
+	}
+
+	bool FPipe::IsInContext() const
+	{
+		return FPipeCallStack::IsOnTop(*this);
+	}
 }}
\ No newline at end of file
diff --git a/Engine/Source/Runtime/Core/Private/Tasks/TaskPrivate.cpp b/Engine/Source/Runtime/Core/Private/Tasks/TaskPrivate.cpp
--- a/Engine/Source/Runtime/Core/Private/Tasks/TaskPrivate.cpp
+++ b/Engine/Source/Runtime/Core/Private/Tasks/TaskPrivate.cpp
@@ -10,11 +10,20 @@ namespace UE { namespace Tasks { namespace Private
 		return GetPipe()->PushIntoPipe(*this);
 	}
 
-	void FTaskBase::UnblockPipe()
+	void FTaskBase::StartPipeExecution()
+	{
+		if (GetPipe() != nullptr)
+		{
+			GetPipe()->ExecutionStarted();
+		}
+	}
+
+	void FTaskBase::FinishPipeExecution()
 	{
 		if (GetPipe() != nullptr)
 		{
 			check(bPushedIntoPipe);
+			GetPipe()->ExecutionFinished();
 			GetPipe()->ClearTask(*this);
 		}
 	}
diff --git a/Engine/Source/Runtime/Core/Public/Tasks/Pipe.h b/Engine/Source/Runtime/Core/Public/Tasks/Pipe.h
--- a/Engine/Source/Runtime/Core/Public/Tasks/Pipe.h
+++ b/Engine/Source/Runtime/Core/Public/Tasks/Pipe.h
@@ -38,6 +38,13 @@ namespace UE { namespace Tasks
 			return LastTask.load(std::memory_order_relaxed) != nullptr;
 		}
 
+		// waits until the pipe is empty (its last task is executed)
+		// should be used only after no more tasks are launched in the pipe, e.g. preparing for the pipe destruction
+		void WaitUntilEmpty()
+		{
+			LowLevelTasks::FScheduler::Get().BusyWaitUntil([this] { return !HasWork(); });
+		}
+
 		// launches a task in the pipe
 		// @param InDebugName helps to identify the task in debugger and profiler
 		// @param TaskBody a callable with no parameters whose return value is ignored, usually a lambda but can be also a functor object 
@@ -55,6 +62,10 @@ namespace UE { namespace Tasks
 			return TTask<FResult>{ Task };
 		}
 
+		// checks if pipe's task is being executed by the current thread. Allows to check if accessing a resource protected by a pipe
+		// is thread-safe
+		bool IsInContext() const;
+
 	private:
 		friend class Private::FTaskBase;
 
@@ -65,6 +76,10 @@ namespace UE { namespace Tasks
 		// otherwise the next piped task can try to add itself as a subsequent to already destroyed task
 		void ClearTask(Private::FTaskBase& Task);
 
+		// notifications about pipe's task execution
+		void ExecutionStarted();
+		void ExecutionFinished();
+
 	private:
 		std::atomic<Private::FTaskBase*> LastTask{ nullptr }; // the last task that 
 		// a counter of how many threads trying to push a task into a pipe right now
diff --git a/Engine/Source/Runtime/Core/Public/Tasks/TaskPrivate.h b/Engine/Source/Runtime/Core/Public/Tasks/TaskPrivate.h
--- a/Engine/Source/Runtime/Core/Public/Tasks/TaskPrivate.h
+++ b/Engine/Source/Runtime/Core/Public/Tasks/TaskPrivate.h
@@ -72,9 +72,9 @@ namespace UE { namespace Tasks
 						Deleter = Forward<DeleterType>(Deleter)
 					]() mutable
 					{
+						StartPipeExecution();
 						Invoke(TaskBody);
-						// finalisation
-						UnblockPipe();
+						FinishPipeExecution();
 						FTaskBase* Subsequent = CloseAndReturnSubsequent();
 						if (Subsequent != nullptr)
 						{
@@ -197,7 +197,8 @@ namespace UE { namespace Tasks
 			}
 
 			CORE_API bool PushIntoPipe();
-			CORE_API void UnblockPipe();
+			CORE_API void StartPipeExecution();
+			CORE_API void FinishPipeExecution();
 
 			// after being executed the task checks if it has a subsequent to launch. it's the only time the task will do this check so
 			// it "closes" itself and setting a subsequent fill fail after that.
