fixed unresolved symbol (`Tasks::FPipe::IsInContext()`) detected by local changes

diff --git a/Engine/Source/Runtime/Core/Public/Tasks/Pipe.h b/Engine/Source/Runtime/Core/Public/Tasks/Pipe.h
--- a/Engine/Source/Runtime/Core/Public/Tasks/Pipe.h
+++ b/Engine/Source/Runtime/Core/Public/Tasks/Pipe.h
@@ -64,7 +64,7 @@ namespace UE { namespace Tasks
 
 		// checks if pipe's task is being executed by the current thread. Allows to check if accessing a resource protected by a pipe
 		// is thread-safe
-		bool IsInContext() const;
+		CORE_API bool IsInContext() const;
 
 	private:
 		friend class Private::FTaskBase;
