removed an outdated test that was relying on StatsThread which doesn't exist anymore

diff --git a/Engine/Source/Runtime/Core/Private/Tests/Async/TaskGraphTest.cpp b/Engine/Source/Runtime/Core/Private/Tests/Async/TaskGraphTest.cpp
--- a/Engine/Source/Runtime/Core/Private/Tests/Async/TaskGraphTest.cpp
+++ b/Engine/Source/Runtime/Core/Private/Tests/Async/TaskGraphTest.cpp
@@ -1320,309 +1320,6 @@ PRAGMA_ENABLE_DEPRECATION_WARNINGS
 		return true;
 	}
 
-	void PipeStressTest()
-	{
-#if STATS
-		FThread StatsThread;
-		if (FPlatformProcess::SupportsMultithreading())
-		{
-			// we simulate StatsThread for this test, as it was removed but the interface is still present (though it's deprecated)
-			StatsThread = FThread{ TEXT("Test StatsThread"),
-				[]
-				{
-					FMemory::SetupTLSCachesOnCurrentThread();
-PRAGMA_DISABLE_DEPRECATION_WARNINGS
-					FTaskGraphInterface::Get().AttachToThread(ENamedThreads::StatsThread);
-					FTaskGraphInterface::Get().ProcessThreadUntilRequestReturn(ENamedThreads::StatsThread);
-PRAGMA_ENABLE_DEPRECATION_WARNINGS
-					FMemory::ClearAndDisableTLSCachesOnCurrentThread();
-				}
-			};
-		}
-
-PRAGMA_DISABLE_DEPRECATION_WARNINGS
-		const ENamedThreads::Type Pipe = FPlatformProcess::SupportsMultithreading() ? ENamedThreads::StatsThread : ENamedThreads::GameThread;
-		const ENamedThreads::Type ManualPipe = FPlatformProcess::SupportsMultithreading() ? ENamedThreads::StatsThread_Local : ENamedThreads::GameThread_Local;
-PRAGMA_ENABLE_DEPRECATION_WARNINGS
-
-		{	// check that piped task is executed
-			std::atomic<bool> bExecuted{ false };
-			FFunctionGraphTask::CreateAndDispatchWhenReady([&bExecuted] { bExecuted = true; }, TStatId{}, nullptr, Pipe)->Wait();
-			check(bExecuted);
-		}
-
-		{	// check that waiting inside a piped task works
-			FGraphEventRef ContinueTask = FGraphEvent::CreateGraphEvent();
-
-			FGraphEventRef Task = FFunctionGraphTask::CreateAndDispatchWhenReady(
-				[&ContinueTask]
-				{
-					check(!ContinueTask->IsComplete() || FTaskGraphInterface::Get().GetNumWorkerThreads() == 0);
-					ContinueTask->Wait();
-				},
-				TStatId{}, nullptr, Pipe
-			);
-
-			FPlatformProcess::Sleep(0.1f);
-			ContinueTask->DispatchSubsequents();
-			Task->Wait();
-		}
-
-		// DEADLOCK:
-		// `T1Launcher->DispatchSubsequents()` -> resolves T1 which is executed and unlocks T2. Then `T1->Wait()` (which is currently implemented
-		// as busy-waiting) starts executing T2 and blocks waiting for T2Launcher, but T2Launcher is triggered only the next line. BOOM
-		// Conclusion: mixing blocking waiting and busy waiting is explosive
-		if (false)
-		{	// check that piped tasks don't execute in parallel
-			std::atomic<bool> bExecutingPipe{ false };
-
-			FGraphEventRef T1Launcher = FGraphEvent::CreateGraphEvent();
-			std::atomic<bool> bT1Executed{ false };
-			FGraphEventRef T1 = FFunctionGraphTask::CreateAndDispatchWhenReady(
-				[&bExecutingPipe, T1Launcher, &bT1Executed]
-				{
-					check(!bExecutingPipe);
-					bExecutingPipe = true;
-					T1Launcher->Wait();
-					bT1Executed = true;
-					bExecutingPipe = false;
-				},
-				TStatId{}, nullptr, Pipe
-			);
-
-			FGraphEventRef T2Launcher = FGraphEvent::CreateGraphEvent();
-			std::atomic<bool> bT2Executed{ false };
-			FGraphEventRef T2 = FFunctionGraphTask::CreateAndDispatchWhenReady(
-				[&bExecutingPipe, T2Launcher, &bT2Executed]
-				{
-					check(!bExecutingPipe);
-					bExecutingPipe = true;
-					T2Launcher->Wait();
-					bT2Executed = true;
-					bExecutingPipe = false;
-				},
-				TStatId{}, nullptr, Pipe
-			);
-
-			FGraphEventRef T3Launcher = FGraphEvent::CreateGraphEvent();
-			std::atomic<bool> bT3Executed{ false };
-			FGraphEventRef T3 = FFunctionGraphTask::CreateAndDispatchWhenReady(
-				[&bExecutingPipe, T3Launcher, &bT3Executed]
-				{
-					check(!bExecutingPipe);
-					bExecutingPipe = true;
-					T3Launcher->Wait();
-					bT3Executed = true;
-					bExecutingPipe = false;
-				},
-				TStatId{}, nullptr, Pipe
-			);
-
-			check(!T1->IsComplete());
-			check(!T2->IsComplete());
-			check(!T3->IsComplete());
-			T1Launcher->DispatchSubsequents();
-			T1->Wait();
-			check(!T2->IsComplete());
-			check(!T3->IsComplete());
-			T2Launcher->DispatchSubsequents();
-			T2->Wait();
-			check(!T3->IsComplete());
-			T3Launcher->DispatchSubsequents();
-			T3->Wait();
-		}
-
-		{	// test that manual pipe tasks are executed only manually
-			std::atomic<bool> bExecuted{ false };
-			FGraphEventRef ManualPipeTask = FFunctionGraphTask::CreateAndDispatchWhenReady([&bExecuted] { bExecuted = true; }, TStatId{}, nullptr, ManualPipe);
-			FPlatformProcess::Sleep(0.1f);
-			check(!bExecuted);
-			check(!ManualPipeTask->IsComplete());
-
-			// manual pipes must be executed from corresponding pipe
-			FFunctionGraphTask::CreateAndDispatchWhenReady([ManualPipe] { FTaskGraphInterface::Get().ProcessThreadUntilIdle(ManualPipe); }, TStatId{}, nullptr, Pipe)
-				->Wait();
-
-			check(bExecuted);
-			check(ManualPipeTask->IsComplete());
-		}
-
-		{	// test ProcessThreadUntilIdle
-			std::atomic<bool> bExecuted{ false };
-
-			FFunctionGraphTask::CreateAndDispatchWhenReady(
-				[ManualPipe, &bExecuted]
-				{
-					FGraphEventRef ManualPipeTask = FFunctionGraphTask::CreateAndDispatchWhenReady([&bExecuted] { bExecuted = true; }, TStatId{}, nullptr, ManualPipe);
-					FTaskGraphInterface::Get().ProcessThreadUntilIdle(ManualPipe);
-					check(bExecuted);
-					check(ManualPipeTask->IsComplete());
-				},
-				TStatId{}, nullptr, Pipe
-			)->Wait();
-			check(bExecuted);
-		}
-
-		{	// test Wait
-			std::atomic<bool> bExecuted{ false };
-
-			FFunctionGraphTask::CreateAndDispatchWhenReady(
-				[ManualPipe, &bExecuted]
-				{
-					FFunctionGraphTask::CreateAndDispatchWhenReady([&bExecuted] { bExecuted = true; }, TStatId{}, nullptr, ManualPipe)
-						->Wait(ManualPipe);
-					check(bExecuted);
-				},
-				TStatId{}, nullptr, Pipe
-			)->Wait();
-			check(bExecuted);
-		}
-
-		if (FTaskGraphInterface::Get().GetNumWorkerThreads() != 0)
-		{	// auto-pipe with fire-and-forget task
-			std::atomic<bool> bExecuted{ false };
-
-			TGraphTask<TFunctionGraphTaskImpl<void(), ESubsequentsMode::FireAndForget>>::CreateTask().ConstructAndDispatchWhenReady([&bExecuted] { bExecuted = true; }, TStatId{}, Pipe);
-
-			FPlatformProcess::Sleep(0.1f);
-			check(bExecuted);
-		}
-
-		{	// manual pipe with fire-and-forget task
-			std::atomic<bool> bExecuted{ false };
-
-			TGraphTask<TFunctionGraphTaskImpl<void(), ESubsequentsMode::FireAndForget>>::CreateTask().ConstructAndDispatchWhenReady([&bExecuted] { bExecuted = true; }, TStatId{}, ManualPipe);
-
-			FFunctionGraphTask::CreateAndDispatchWhenReady([ManualPipe] { FTaskGraphInterface::Get().ProcessThreadUntilIdle(ManualPipe); }, TStatId{}, nullptr, Pipe)
-				->Wait();
-			check(bExecuted);
-		}
-
-		{	// test GetCurrentThread() from inside manual pipe
-			FFunctionGraphTask::CreateAndDispatchWhenReady(
-				[ManualPipe]
-				{
-					checkf(FTaskGraphInterface::Get().GetCurrentThreadIfKnown(true) == ManualPipe, TEXT("%d - %d"), FTaskGraphInterface::Get().GetCurrentThreadIfKnown(true), ManualPipe);
-				},
-				TStatId{}, nullptr, ManualPipe);
-
-			FFunctionGraphTask::CreateAndDispatchWhenReady([ManualPipe] { FTaskGraphInterface::Get().ProcessThreadUntilIdle(ManualPipe); }, TStatId{}, nullptr, Pipe)
-				->Wait();
-		}
-
-		{	// test IsThreadProcessingTasks() from inside manual pipe
-			FFunctionGraphTask::CreateAndDispatchWhenReady(
-				[ManualPipe]
-				{
-					check(FTaskGraphInterface::Get().IsThreadProcessingTasks(ManualPipe));
-				},
-				TStatId{}, nullptr, ManualPipe);
-
-			FFunctionGraphTask::CreateAndDispatchWhenReady(
-				[ManualPipe]
-				{
-					check(!FTaskGraphInterface::Get().IsThreadProcessingTasks(ManualPipe));
-					FTaskGraphInterface::Get().ProcessThreadUntilIdle(ManualPipe);
-				},
-				TStatId{}, nullptr, Pipe)->Wait();
-		}
-
-		{	// check that independent task can bypass a task with pending prereq
-			FGraphEventRef Prereq = FGraphEvent::CreateGraphEvent();
-			FGraphEventRef Dependent = FFunctionGraphTask::CreateAndDispatchWhenReady([] {}, TStatId{}, Prereq, Pipe);
-			FGraphEventRef Independent = FFunctionGraphTask::CreateAndDispatchWhenReady([] {}, TStatId{}, nullptr, Pipe);
-			Independent->Wait();
-			check(Independent->IsComplete() && !Dependent->IsComplete());
-			Prereq->DispatchSubsequents();
-			Dependent->Wait();
-		}
-
-		{	// stress test for piped tasks, checks spawning a large number of piped tasks from multiple threads and executing them
-			//UE_LOG(LogTemp, Display, TEXT("Pipe stress test"));
-
-			constexpr uint32 SpawnerGroupsNum = 50;
-			FGraphEventArray SpawnerGroups;
-			SpawnerGroups.Reserve(SpawnerGroupsNum);
-
-			constexpr uint32 SpawnersPerGroupNum = 50;
-			constexpr uint32 TasksNum = SpawnerGroupsNum * SpawnersPerGroupNum;
-			FGraphEventArray Spawners;
-			Spawners.AddDefaulted(TasksNum);
-			FGraphEventArray Tasks;
-			Tasks.AddDefaulted(TasksNum);
-
-			std::atomic<bool> bExecuting{ false };
-			std::atomic<uint32> TasksExecutedNum{ 0 };
-
-			for (uint32 GroupIndex = 0; GroupIndex != SpawnerGroupsNum; ++GroupIndex)
-			{
-				SpawnerGroups.Add(FFunctionGraphTask::CreateAndDispatchWhenReady(
-					[Spawners = &Spawners[GroupIndex * SpawnersPerGroupNum],
-					Tasks = &Tasks[GroupIndex * SpawnersPerGroupNum],
-					&bExecuting,
-					&TasksExecutedNum,
-					SpawnersPerGroupNum,
-					Pipe]
-					{
-						for (uint32 SpawnerIndex = 0; SpawnerIndex != SpawnersPerGroupNum; ++SpawnerIndex)
-						{
-							Spawners[SpawnerIndex] = FFunctionGraphTask::CreateAndDispatchWhenReady(
-								[Task = &Tasks[SpawnerIndex], &bExecuting, &TasksExecutedNum, Pipe]
-								{
-									*Task = FFunctionGraphTask::CreateAndDispatchWhenReady(
-										[&bExecuting, &TasksExecutedNum, Pipe]
-										{
-											check(!bExecuting);
-											bExecuting = true;
-											check(FTaskGraphInterface::Get().GetCurrentThreadIfKnown() == Pipe);
-											FPlatformProcess::SleepNoStats(0.0f);
-											++TasksExecutedNum;
-											bExecuting = false;
-										},
-										TStatId{}, nullptr, Pipe
-									);
-								}
-							);
-						}
-					}
-				));
-			}
-
-			//UE_LOG(LogTemp, Display, TEXT("Spawners groups spawned"));
-			FTaskGraphInterface::Get().WaitUntilTasksComplete(SpawnerGroups);
-			//UE_LOG(LogTemp, Display, TEXT("Spawners groups completed"));
-
-			FTaskGraphInterface::Get().WaitUntilTasksComplete(Spawners);
-			//UE_LOG(LogTemp, Display, TEXT("Spawners completed"));
-
-			FTaskGraphInterface::Get().WaitUntilTasksComplete(Tasks);
-			//UE_LOG(LogTemp, Display, TEXT("Tasks completed"));
-			check(TasksExecutedNum == TasksNum);
-		}
-
-		if (FPlatformProcess::SupportsMultithreading())
-		{
-			// stop StatsThread
-PRAGMA_DISABLE_DEPRECATION_WARNINGS
-			FGraphEventRef QuitTask = TGraphTask<FReturnGraphTask>::CreateTask(NULL, ENamedThreads::GameThread).ConstructAndDispatchWhenReady(ENamedThreads::StatsThread);
-PRAGMA_ENABLE_DEPRECATION_WARNINGS
-			FTaskGraphInterface::Get().WaitUntilTaskCompletes(QuitTask, ENamedThreads::GameThread_Local);
-			StatsThread.Join();
-		}
-#endif
-	}
-
-	IMPLEMENT_SIMPLE_AUTOMATION_TEST(FPipeTest, "System.Core.Async.TaskGraph.Pipe", EAutomationTestFlags::ApplicationContextMask | EAutomationTestFlags::EngineFilter);
-
-	bool FPipeTest::RunTest(const FString& Parameters)
-	{
-		TRACE_CPUPROFILER_EVENT_SCOPE(TaskGraphTests_PipeTest);
-
-		UE_BENCHMARK(5, PipeStressTest);
-
-		return true;
-	}
-
 	IMPLEMENT_SIMPLE_AUTOMATION_TEST(FTaskGraphStatsThreadRedirectionTest, "System.Core.Async.TaskGraph.StatsThreadRedirection", EAutomationTestFlags::ApplicationContextMask | EAutomationTestFlags::EngineFilter);
 
 	bool FTaskGraphStatsThreadRedirectionTest::RunTest(const FString& Parameters)
