Fix test that cause that newer models are considered as less powerful than their predecessor.

diff --git a/Engine/Source/Runtime/Core/Private/Mac/MacPlatformMisc.cpp b/Engine/Source/Runtime/Core/Private/Mac/MacPlatformMisc.cpp
--- a/Engine/Source/Runtime/Core/Private/Mac/MacPlatformMisc.cpp
+++ b/Engine/Source/Runtime/Core/Private/Mac/MacPlatformMisc.cpp
@@ -1814,7 +1814,7 @@ bool FMacPlatformMisc::IsRunningOnRecommendedMinSpecHardware()
 			const FMacModel& Model = MinSupportedMacModels[Index];
 
 			if ( (SystemModelMajor <  Model.Major) ||
-				((SystemModelMajor >= Model.Major) && (SystemModelMinor < Model.Minor)))
+				((SystemModelMajor == Model.Major) && (SystemModelMinor < Model.Minor)))
 			{
 				bSupported = false;
 			}
