Add FPlatformMisc::NormalizePath(FStringBuilderBase&) overload

diff --git a/Engine/Source/Runtime/Core/Private/Mac/MacPlatformMisc.cpp b/Engine/Source/Runtime/Core/Private/Mac/MacPlatformMisc.cpp
--- a/Engine/Source/Runtime/Core/Private/Mac/MacPlatformMisc.cpp
+++ b/Engine/Source/Runtime/Core/Private/Mac/MacPlatformMisc.cpp
@@ -23,6 +23,7 @@
 #include "HAL/ThreadManager.h"
 #include "Misc/OutputDeviceError.h"
 #include "Misc/OutputDeviceRedirector.h"
+#include "Misc/StringBuilder.h"
 #include "Misc/FeedbackContext.h"
 #include "Misc/CoreDelegates.h"
 #include "Internationalization/Internationalization.h"
@@ -783,6 +784,15 @@ void FMacPlatformMisc::NormalizePath(FString& InPath)
 	}
 }
 
+void FMacPlatformMisc::NormalizePath(FStringBuilderBase& InPath)
+{
+	// only expand if path starts with ~, e.g. ~/ should be expanded, /~ should not
+	if (FStringView(InPath).StartsWith('~'))
+	{
+		InPath.ReplaceAt(0, 1, FPlatformProcess::UserHomeDir());
+	}
+}
+
 template<typename T>
 FMacPlatformMisc::FGPUDescriptorCommon<T>::~FGPUDescriptorCommon()
 {
diff --git a/Engine/Source/Runtime/Core/Private/Unix/UnixPlatformMisc.cpp b/Engine/Source/Runtime/Core/Private/Unix/UnixPlatformMisc.cpp
--- a/Engine/Source/Runtime/Core/Private/Unix/UnixPlatformMisc.cpp
+++ b/Engine/Source/Runtime/Core/Private/Unix/UnixPlatformMisc.cpp
@@ -13,6 +13,7 @@
 #include "Containers/Array.h"
 #include "Containers/UnrealString.h"
 #include "Misc/DateTime.h"
+#include "Misc/StringBuilder.h"
 #include "HAL/PlatformTime.h"
 #include "Containers/StringConv.h"
 #include "Logging/LogMacros.h"
@@ -86,6 +87,15 @@ void FUnixPlatformMisc::NormalizePath(FString& InPath)
 	}
 }
 
+void FUnixPlatformMisc::NormalizePath(FStringBuilderBase& InPath)
+{
+	// only expand if path starts with ~, e.g. ~/ should be expanded, /~ should not
+	if (FStringView(InPath).StartsWith('~'))
+	{
+		InPath.ReplaceAt(0, 1, FPlatformProcess::UserHomeDir());
+	}
+}
+
 size_t CORE_API GCacheLineSize = PLATFORM_CACHE_LINE_SIZE;
 
 void UnixPlatform_UpdateCacheLineSize()
diff --git a/Engine/Source/Runtime/Core/Public/GenericPlatform/GenericPlatformMisc.h b/Engine/Source/Runtime/Core/Public/GenericPlatform/GenericPlatformMisc.h
--- a/Engine/Source/Runtime/Core/Public/GenericPlatform/GenericPlatformMisc.h
+++ b/Engine/Source/Runtime/Core/Public/GenericPlatform/GenericPlatformMisc.h
@@ -4,6 +4,7 @@
 
 #include "CoreTypes.h"
 #include "CoreFwd.h"
+#include "Containers/StringFwd.h"
 #include "HAL/PlatformCrt.h"
 #include "Misc/CompressionFlags.h"
 #include "Misc/EnumClassFlags.h"
@@ -849,6 +850,14 @@ public:
 	{
 	}
 
+	/**
+	 * Platform-specific normalization of path
+	 * E.g. on Linux/Unix platforms, replaces ~ with user home directory, so ~/.config becomes [REDACTED_USER_HOME]/.config (or [REDACTED_USER_HOME]/.config)
+	 */
+	static void NormalizePath(FStringBuilderBase& InPath)
+	{
+	}
+
 	/**
 	 * @return platform specific path separator.
 	 */
diff --git a/Engine/Source/Runtime/Core/Public/Mac/MacPlatformMisc.h b/Engine/Source/Runtime/Core/Public/Mac/MacPlatformMisc.h
--- a/Engine/Source/Runtime/Core/Public/Mac/MacPlatformMisc.h
+++ b/Engine/Source/Runtime/Core/Public/Mac/MacPlatformMisc.h
@@ -37,6 +37,7 @@ struct CORE_API FMacPlatformMisc : public FApplePlatformMisc
 	static int32 NumberOfCores();
 	static int32 NumberOfCoresIncludingHyperthreads();
 	static void NormalizePath(FString& InPath);
+	static void NormalizePath(FStringBuilderBase& InPath);
 	static FString GetPrimaryGPUBrand();
 	static struct FGPUDriverInfo GetGPUDriverInfo(const FString& DeviceDescription);
 	static void GetOSVersions( FString& out_OSVersionLabel, FString& out_OSSubVersionLabel );
diff --git a/Engine/Source/Runtime/Core/Public/Unix/UnixPlatformMisc.h b/Engine/Source/Runtime/Core/Public/Unix/UnixPlatformMisc.h
--- a/Engine/Source/Runtime/Core/Public/Unix/UnixPlatformMisc.h
+++ b/Engine/Source/Runtime/Core/Public/Unix/UnixPlatformMisc.h
@@ -46,6 +46,7 @@ struct CORE_API FUnixPlatformMisc : public FGenericPlatformMisc
 	static const TCHAR* GetSystemErrorMessage(TCHAR* OutBuffer, int32 BufferCount, int32 Error);
 
 	static void NormalizePath(FString& InPath);
+	static void NormalizePath(FStringBuilderBase& InPath);
 
 	static const TCHAR* GetPathVarDelimiter()
 	{
