Silence bogus static analysis warning due to an enum class value sharing name with a member function

#rnx

diff --git a/Engine/Source/Runtime/Core/Public/Misc/AsciiSet.h b/Engine/Source/Runtime/Core/Public/Misc/AsciiSet.h
--- a/Engine/Source/Runtime/Core/Public/Misc/AsciiSet.h
+++ b/Engine/Source/Runtime/Core/Public/Misc/AsciiSet.h
@@ -126,56 +126,56 @@ public:
 	template<class StringType>
 	static constexpr StringType FindPrefixWith(const StringType& Str, FAsciiSet Set)
 	{
-		return Scan<EDir::Forward, ESkip::Members, EKeep::Head>(Str, Set);
+		return Scan<EDir::Forward, EInclude::Members, EKeep::Head>(Str, Set);
 	}
 
 	/** Get initial substring with no characters in set */
 	template<class StringType>
 	static constexpr StringType FindPrefixWithout(const StringType& Str, FAsciiSet Set)
 	{
-		return Scan<EDir::Forward, ESkip::NonMembers, EKeep::Head>(Str, Set);
+		return Scan<EDir::Forward, EInclude::NonMembers, EKeep::Head>(Str, Set);
 	}
 
 	/** Trim initial characters in set */
 	template<class StringType>
 	static constexpr StringType TrimPrefixWith(const StringType& Str, FAsciiSet Set)
 	{
-		return Scan<EDir::Forward, ESkip::Members, EKeep::Tail>(Str, Set);
+		return Scan<EDir::Forward, EInclude::Members, EKeep::Tail>(Str, Set);
 	}
 
 	/** Trim initial characters not in set */
 	template<class StringType>
 	static constexpr StringType TrimPrefixWithout(const StringType& Str, FAsciiSet Set)
 	{
-		return Scan<EDir::Forward, ESkip::NonMembers, EKeep::Tail>(Str, Set);
+		return Scan<EDir::Forward, EInclude::NonMembers, EKeep::Tail>(Str, Set);
 	}
 
 	/** Get trailing substring with all characters in set */
 	template<class StringType>
 	static constexpr StringType FindSuffixWith(const StringType& Str, FAsciiSet Set)
 	{
-		return Scan<EDir::Reverse, ESkip::Members, EKeep::Tail>(Str, Set);
+		return Scan<EDir::Reverse, EInclude::Members, EKeep::Tail>(Str, Set);
 	}
 
 	/** Get trailing substring with no characters in set */
 	template<class StringType>
 	static constexpr StringType FindSuffixWithout(const StringType& Str, FAsciiSet Set)
 	{
-		return Scan<EDir::Reverse, ESkip::NonMembers, EKeep::Tail>(Str, Set);
+		return Scan<EDir::Reverse, EInclude::NonMembers, EKeep::Tail>(Str, Set);
 	}
 
 	/** Trim trailing characters in set */
 	template<class StringType>
 	static constexpr StringType TrimSuffixWith(const StringType& Str, FAsciiSet Set)
 	{
-		return Scan<EDir::Reverse, ESkip::Members, EKeep::Head>(Str, Set);
+		return Scan<EDir::Reverse, EInclude::Members, EKeep::Head>(Str, Set);
 	}
 
 	/** Trim trailing characters not in set */
 	template<class StringType>
 	static constexpr StringType TrimSuffixWithout(const StringType& Str, FAsciiSet Set)
 	{
-		return Scan<EDir::Reverse, ESkip::NonMembers, EKeep::Head>(Str, Set);
+		return Scan<EDir::Reverse, EInclude::NonMembers, EKeep::Head>(Str, Set);
 	}
 	
 	/** Test if string contains any character in set */
@@ -202,35 +202,35 @@ public:
 	static constexpr bool HasOnly(const StringType& Str, FAsciiSet Set)
 	{
 		auto End = GetData(Str) + GetNum(Str);
-		return FindFirst<ESkip::Members>(Set, GetData(Str), End) == End;
+		return FindFirst<EInclude::Members>(Set, GetData(Str), End) == End;
 	}
 
 private:
 	enum class EDir {Forward, Reverse};
-	enum class ESkip {Members, NonMembers};
+	enum class EInclude {Members, NonMembers};
 	enum class EKeep {Head, Tail};
 
-	template<ESkip Skip, typename CharType>
+	template<EInclude Include, typename CharType>
 	static constexpr const CharType* FindFirst(FAsciiSet Set, const CharType* It, const CharType* End)
 	{
-		for (; It != End && (Skip == ESkip::Members) == !!Set.Test(*It); ++It);
+		for (; It != End && (Include == EInclude::Members) == !!Set.Test(*It); ++It);
 		return It;
 	}
 
-	template<ESkip Skip, typename CharType>
+	template<EInclude Include, typename CharType>
 	static constexpr const CharType* FindLast(FAsciiSet Set, const CharType* It, const CharType* End)
 	{
-		for (; It != End && (Skip == ESkip::Members) == !!Set.Test(*It); --It);
+		for (; It != End && (Include == EInclude::Members) == !!Set.Test(*It); --It);
 		return It;
 	}
 
-	template<EDir Dir, ESkip Skip, EKeep Keep, class StringType>
+	template<EDir Dir, EInclude Include, EKeep Keep, class StringType>
 	static constexpr StringType Scan(const StringType& Str, FAsciiSet Set)
 	{
 		auto Begin = GetData(Str);
 		auto End = Begin + GetNum(Str);
-		auto It = Dir == EDir::Forward	? FindFirst<Skip>(Set, Begin, End)
-										: FindLast<Skip>(Set, End - 1, Begin - 1) + 1;
+		auto It = Dir == EDir::Forward	? FindFirst<Include>(Set, Begin, End)
+										: FindLast<Include>(Set, End - 1, Begin - 1) + 1;
 
 		return Keep == EKeep::Head	? StringType(Begin, static_cast<int32>(It - Begin))
 									: StringType(It, static_cast<int32>(End - It));
