std::tuple_size and std::tuple_element are defined as `struct`s by the C++17 standard.

#rnx

diff --git a/Engine/Source/Runtime/Core/Public/Templates/Tuple.h b/Engine/Source/Runtime/Core/Public/Templates/Tuple.h
--- a/Engine/Source/Runtime/Core/Public/Templates/Tuple.h
+++ b/Engine/Source/Runtime/Core/Public/Templates/Tuple.h
@@ -890,12 +890,12 @@ FORCEINLINE TTuple<Types&...> Tie(Types&... Args)
 #if UE_TUPLE_STRUCTURED_BINDING_SUPPORT
 // TTuple support for structured bindings
 template <typename... ArgTypes>
-class std::tuple_size<TTuple<ArgTypes...>>
-	: public std::integral_constant<std::size_t, sizeof...(ArgTypes)>
+struct std::tuple_size<TTuple<ArgTypes...>>
+	: std::integral_constant<std::size_t, sizeof...(ArgTypes)>
 {
 };
 template <std::size_t N, typename... ArgTypes>
-class std::tuple_element<N, TTuple<ArgTypes...>>
+struct std::tuple_element<N, TTuple<ArgTypes...>>
 {
 public:
 	using type = typename TTupleElement<N, TTuple<ArgTypes...>>::Type;
