Lumin's libc++ 8.0 has a broken definition of std::tuple_element, using 'class' instead of 'struct'. This is fixed in libc++9.0 and a macro guard reflects that.

#trivial
#rnx

diff --git a/Engine/Source/Runtime/Core/Public/Templates/Tuple.h b/Engine/Source/Runtime/Core/Public/Templates/Tuple.h
--- a/Engine/Source/Runtime/Core/Public/Templates/Tuple.h
+++ b/Engine/Source/Runtime/Core/Public/Templates/Tuple.h
@@ -895,7 +895,13 @@ struct std::tuple_size<TTuple<ArgTypes...>>
 {
 };
 template <std::size_t N, typename... ArgTypes>
+#if defined(_LIBCPP_VERSION) && _LIBCPP_VERSION < 9000
+// libc++ 8.0 and older incorrectly defined std::tuple_element as class, whereas the standard
+// says it should be a struct.
+class std::tuple_element<N, TTuple<ArgTypes...>>
+#else
 struct std::tuple_element<N, TTuple<ArgTypes...>>
+#endif
 {
 public:
 	using type = typename TTupleElement<N, TTuple<ArgTypes...>>::Type;
