Remove deprecated GetStatsThreadMask from FIOSPlatformAffinity

diff --git a/Engine/Source/Runtime/Core/Private/IOS/IOSPlatformProcess.cpp b/Engine/Source/Runtime/Core/Private/IOS/IOSPlatformProcess.cpp
--- a/Engine/Source/Runtime/Core/Private/IOS/IOSPlatformProcess.cpp
+++ b/Engine/Source/Runtime/Core/Private/IOS/IOSPlatformProcess.cpp
@@ -246,28 +246,6 @@ const uint64 FIOSPlatformAffinity::GetTaskGraphThreadMask()
 	return Mask;
 }
 
-const uint64 FIOSPlatformAffinity::GetStatsThreadMask()
-{
-	static int Mask = 0;
-	if (Mask == 0)
-	{
-		switch (FPlatformMisc::NumberOfCores())
-		{
-		case 2:
-			Mask = MAKEAFFINITYMASK1(0);
-			break;
-		case 3:
-			Mask = MAKEAFFINITYMASK1(2);
-			break;
-		default:
-			Mask = FGenericPlatformAffinity::GetStatsThreadMask();
-			break;
-		}
-	}
-
-	return Mask;
-}
-
 const uint64 FIOSPlatformAffinity::GetNoAffinityMask()
 {
 	static int Mask = 0;
diff --git a/Engine/Source/Runtime/Core/Public/IOS/IOSPlatformAffinity.h b/Engine/Source/Runtime/Core/Public/IOS/IOSPlatformAffinity.h
--- a/Engine/Source/Runtime/Core/Public/IOS/IOSPlatformAffinity.h
+++ b/Engine/Source/Runtime/Core/Public/IOS/IOSPlatformAffinity.h
@@ -12,7 +12,6 @@ public:
 	static const uint64 GetRTHeartBeatMask();
 	static const uint64 GetPoolThreadMask();
 	static const uint64 GetTaskGraphThreadMask();
-	static const uint64 GetStatsThreadMask();
 	static const uint64 GetNoAffinityMask();
 };
 
