Add PathViews ToAbsolutePath[Inline] / NormalizeFilename / NormalizeDirectoryName / CollapseRelativeDirectories / RemoveDuplicateSlashes

diff --git a/Engine/Source/Runtime/Core/Private/Misc/PathViews.cpp b/Engine/Source/Runtime/Core/Private/Misc/PathViews.cpp
--- a/Engine/Source/Runtime/Core/Private/Misc/PathViews.cpp
+++ b/Engine/Source/Runtime/Core/Private/Misc/PathViews.cpp
@@ -4,10 +4,13 @@
 
 #include "Algo/Find.h"
 #include "Algo/FindLast.h"
+#include "Algo/Replace.h"
 #include "Containers/UnrealString.h"
 #include "Containers/StringView.h"
+#include "HAL/PlatformProcess.h"
 #include "Misc/Char.h"
 #include "Misc/StringBuilder.h"
+#include "String/Find.h"
 #include "String/ParseTokens.h"
 #include "Templates/Function.h"
 
@@ -408,6 +411,190 @@ bool FPathViews::IsRelativePath(FStringView InPath)
 	}
 }
 
+static void NormalizeAndCollapseDirectories(FStringBuilderBase& InOutPath)
+{
+	FPathViews::NormalizeFilename(InOutPath);
+	FPathViews::CollapseRelativeDirectories(InOutPath);
+
+	if (InOutPath.Len() == 0)
+	{
+		// Empty path is not absolute, and '/' is the best guess across all the platforms.
+		// This substituion is not valid for Windows of course; however CollapseRelativeDirectories() will not produce an empty
+		// absolute path on Windows as it takes care not to remove the drive letter.
+		InOutPath << '/';
+	}
+}
+
+void FPathViews::ToAbsolutePathInline(FStringView BasePath, FStringBuilderBase& InOutPath)
+{
+	check(BasePath.Len() > 0);
+
+	if (IsRelativePath(InOutPath))
+	{
+		if (BasePath.EndsWith('/') || BasePath.EndsWith('\\'))
+		{
+			InOutPath.Prepend(BasePath);
+		}
+		else
+		{
+			TStringBuilder<128> BaseWithSlash;
+			BaseWithSlash << BasePath << '/';
+			InOutPath.Prepend(BaseWithSlash);
+		}
+	}
+
+	NormalizeAndCollapseDirectories(InOutPath);
+}
+
+void FPathViews::ToAbsolutePathInline(FStringBuilderBase& InOutPath)
+{
+	ToAbsolutePathInline(FStringView(FPlatformProcess::BaseDir()), InOutPath);
+}
+
+void FPathViews::ToAbsolutePath(FStringView BasePath, FStringView InPath, FStringBuilderBase& OutPath)
+{
+	if (int32 OriginalLen = OutPath.Len())
+	{
+		// Use temporary to avoid normalizing existing string builder contents
+		TStringBuilder<256> AbsolutePath;
+		ToAbsolutePath(BasePath, InPath, AbsolutePath);
+		OutPath.Append(AbsolutePath);
+	}
+	else
+	{
+		if (IsRelativePath(InPath))
+		{
+			OutPath << BasePath;
+		}
+
+		FPathViews::Append(OutPath, InPath);
+
+		NormalizeAndCollapseDirectories(OutPath);
+	}
+}
+
+void FPathViews::ToAbsolutePath(FStringView InPath, FStringBuilderBase& OutPath)
+{
+	ToAbsolutePath(FStringView(FPlatformProcess::BaseDir()), InPath, OutPath);
+}
+
+void FPathViews::NormalizeFilename(FStringBuilderBase& InOutPath)
+{
+	Algo::Replace(MakeArrayView(InOutPath), '\\', '/');
+	FPlatformMisc::NormalizePath(InOutPath);
+}
+
+static bool ShouldRemoveTrailingSlash(FStringView Dir)
+{
+	return Dir.EndsWith('/') && !Dir.EndsWith(TEXT("//"_SV), ESearchCase::CaseSensitive) && !Dir.EndsWith(TEXT(":/"_SV), ESearchCase::CaseSensitive);
+}
+
+void FPathViews::NormalizeDirectoryName(FStringBuilderBase& InOutPath)
+{
+	Algo::Replace(MakeArrayView(InOutPath), '\\', '/');
+	InOutPath.RemoveSuffix(ShouldRemoveTrailingSlash(InOutPath.ToView()) ? 1 : 0);
+	FPlatformMisc::NormalizePath(InOutPath);
+}
+
+static int32 FindNext(FStringView View, int32 FromIdx, FStringView Search)
+{
+	int32 RelativeIdx = UE::String::FindFirst(View.RightChop(FromIdx), Search);
+	return RelativeIdx == INDEX_NONE ? INDEX_NONE : RelativeIdx + FromIdx;
+}
+
+static bool Contains(FStringView Str, TCHAR Char)
+{
+	int32 Dummy;
+	return Str.FindChar(Char, Dummy);
+}
+
+static void RemoveAll(FStringBuilderBase& String, FStringView Substring)
+{
+	int32 Pos = String.Len();
+	while (true)
+	{
+		Pos = UE::String::FindLast(String.ToView().Left(Pos), Substring);
+		
+		if (Pos == INDEX_NONE)
+		{
+			break;
+		}
+
+		String.RemoveAt(Pos, Substring.Len());
+	}
+}
+
+bool FPathViews::CollapseRelativeDirectories(FStringBuilderBase& InOutPath)
+{
+	while (InOutPath.Len())
+	{
+		FStringView Path = InOutPath.ToView();
+
+		// Consider paths which start with .. or /.. as invalid
+		constexpr FStringView ParentDir = TEXT("/.."_SV);
+		if (Path.StartsWith(TEXT(".."_SV)) || Path.StartsWith(ParentDir))
+		{
+			return false;
+		}
+	
+		int32 Index = UE::String::FindFirst(Path, ParentDir);
+
+		// Ignore folders beginning with dots
+		while (Index != INDEX_NONE && Index + ParentDir.Len() < Path.Len() && Path[Index + ParentDir.Len()] != '/')
+		{
+			Index = FindNext(Path, Index + ParentDir.Len(), ParentDir);
+		}
+
+		// If there are no "/.."s left then we're done
+		if (Index == INDEX_NONE)
+		{
+			break;
+		}
+
+		// Find previous directory, stop if we've found a directory that isn't "/./"
+		int32 PreviousSeparatorIndex = Index;
+		while (Path.Left(PreviousSeparatorIndex).FindLastChar('/', /* Out */ PreviousSeparatorIndex) &&
+				Path.Mid(PreviousSeparatorIndex + 1, 2) == TEXT("./"_SV));
+
+		PreviousSeparatorIndex = FMath::Max(0, PreviousSeparatorIndex);
+		
+		// If we're attempting to remove the drive letter, that's illegal
+		int32 RemoveLen = Index - PreviousSeparatorIndex + ParentDir.Len();
+		if (Contains(Path.Mid(PreviousSeparatorIndex, RemoveLen), ':'))
+		{
+			return false;
+		}
+
+		InOutPath.RemoveAt(PreviousSeparatorIndex, RemoveLen);
+	}
+
+	RemoveAll(InOutPath, TEXT("./"_SV));
+
+	return true;
+}
+
+void FPathViews::RemoveDuplicateSlashes(FStringBuilderBase& InOutPath)
+{
+	int32 DoubleSlashIdx = UE::String::FindFirst(InOutPath.ToView(), TEXT("//"_SV));
+	if (DoubleSlashIdx == INDEX_NONE)
+	{
+		return;
+	}
+
+	TCHAR* WriteIt = InOutPath.GetData() + DoubleSlashIdx + 1;
+	TCHAR* ReadIt = WriteIt + 1;
+	for (TCHAR* End = InOutPath.GetData() + InOutPath.Len(), LastChar = '/'; ReadIt != End; LastChar = *ReadIt++)
+	{
+		if ((*ReadIt != '/') | (LastChar != '/'))
+		{
+			*WriteIt++ = *ReadIt;
+		}
+	}
+
+	*WriteIt = '\0';
+	InOutPath.RemoveSuffix(int32(ReadIt - WriteIt));
+}
+
 void FPathViews::SplitFirstComponent(FStringView InPath, FStringView& OutFirstComponent, FStringView& OutRemainder)
 {
 	using namespace UE4PathViews_Private;
diff --git a/Engine/Source/Runtime/Core/Private/Tests/Misc/PathTests.h b/Engine/Source/Runtime/Core/Private/Tests/Misc/PathTests.h
new file mode 100644
--- /dev/null
+++ b/Engine/Source/Runtime/Core/Private/Tests/Misc/PathTests.h
@@ -0,0 +1,147 @@
+// Copyright Epic Games, Inc. All Rights Reserved.
+
+#pragma once
+
+#include "Containers/StringView.h"
+#include "Misc/AutomationTest.h"
+
+#if WITH_DEV_AUTOMATION_TESTS
+
+template<class PathType, class StringType>
+void TestCollapseRelativeDirectories(FAutomationTestBase& Test)
+{
+	auto Run = [&](const TCHAR* Path, const TCHAR* Expected)
+	{
+		// Run test
+		StringType Actual;
+		Actual += Path;
+		const bool bValid = PathType::CollapseRelativeDirectories(Actual);
+
+		if (Expected)
+		{
+			// If we're looking for a result, make sure it was returned correctly
+			if (!bValid || FCString::Strcmp(*Actual, Expected) != 0)
+			{
+				Test.AddError(FString::Printf(TEXT("Path '%s' failed to collapse correctly (got '%s', expected '%s')."), Path, *Actual, Expected));
+			}
+		}
+		else
+		{
+			// Otherwise, make sure it failed
+			if (bValid)
+			{
+				Test.AddError(FString::Printf(TEXT("Path '%s' collapsed unexpectedly."), Path));
+			}
+		}
+	};
+
+	Run(TEXT(".."),                                                   nullptr);
+	Run(TEXT("/.."),                                                  nullptr);
+	Run(TEXT("./"),                                                   TEXT(""));
+	Run(TEXT("./file.txt"),                                           TEXT("file.txt"));
+	Run(TEXT("/."),                                                   TEXT("/."));
+	Run(TEXT("Folder"),                                               TEXT("Folder"));
+	Run(TEXT("/Folder"),                                              TEXT("/Folder"));
+	Run(TEXT("C:/Folder"),                                            TEXT("C:/Folder"));
+	Run(TEXT("C:/Folder/.."),                                         TEXT("C:"));
+	Run(TEXT("C:/Folder/../"),                                        TEXT("C:/"));
+	Run(TEXT("C:/Folder/../file.txt"),                                TEXT("C:/file.txt"));
+	Run(TEXT("Folder/.."),                                            TEXT(""));
+	Run(TEXT("Folder/../"),                                           TEXT("/"));
+	Run(TEXT("Folder/../file.txt"),                                   TEXT("/file.txt"));
+	Run(TEXT("/Folder/.."),                                           TEXT(""));
+	Run(TEXT("/Folder/../"),                                          TEXT("/"));
+	Run(TEXT("/Folder/../file.txt"),                                  TEXT("/file.txt"));
+	Run(TEXT("Folder/../.."),                                         nullptr);
+	Run(TEXT("Folder/../../"),                                        nullptr);
+	Run(TEXT("Folder/../../file.txt"),                                nullptr);
+	Run(TEXT("C:/.."),                                                nullptr);
+	Run(TEXT("C:/."),                                                 TEXT("C:/."));
+	Run(TEXT("C:/./"),                                                TEXT("C:/"));
+	Run(TEXT("C:/./file.txt"),                                        TEXT("C:/file.txt"));
+	Run(TEXT("C:/Folder1/../Folder2"),                                TEXT("C:/Folder2"));
+	Run(TEXT("C:/Folder1/../Folder2/"),                               TEXT("C:/Folder2/"));
+	Run(TEXT("C:/Folder1/../Folder2/file.txt"),                       TEXT("C:/Folder2/file.txt"));
+	Run(TEXT("C:/Folder1/../Folder2/../.."),                          nullptr);
+	Run(TEXT("C:/Folder1/../Folder2/../Folder3"),                     TEXT("C:/Folder3"));
+	Run(TEXT("C:/Folder1/../Folder2/../Folder3/"),                    TEXT("C:/Folder3/"));
+	Run(TEXT("C:/Folder1/../Folder2/../Folder3/file.txt"),            TEXT("C:/Folder3/file.txt"));
+	Run(TEXT("C:/Folder1/Folder2/../../Folder3"),                     TEXT("C:/Folder3"));
+	Run(TEXT("C:/Folder1/Folder2/../../Folder3/"),                    TEXT("C:/Folder3/"));
+	Run(TEXT("C:/Folder1/Folder2/../../Folder3/file.txt"),            TEXT("C:/Folder3/file.txt"));
+	Run(TEXT("C:/Folder1/Folder2/../../Folder3/../Folder4"),          TEXT("C:/Folder4"));
+	Run(TEXT("C:/Folder1/Folder2/../../Folder3/../Folder4/"),         TEXT("C:/Folder4/"));
+	Run(TEXT("C:/Folder1/Folder2/../../Folder3/../Folder4/file.txt"), TEXT("C:/Folder4/file.txt"));
+	Run(TEXT("C:/Folder1/Folder2/../Folder3/../../Folder4"),          TEXT("C:/Folder4"));
+	Run(TEXT("C:/Folder1/Folder2/../Folder3/../../Folder4/"),         TEXT("C:/Folder4/"));
+	Run(TEXT("C:/Folder1/Folder2/../Folder3/../../Folder4/file.txt"), TEXT("C:/Folder4/file.txt"));
+	Run(TEXT("C:/Folder1/Folder2/.././../Folder4"),                   TEXT("C:/Folder4"));
+	Run(TEXT("C:/Folder1/Folder2/.././../Folder4/"),                  TEXT("C:/Folder4/"));
+	Run(TEXT("C:/Folder1/Folder2/.././../Folder4/file.txt"),          TEXT("C:/Folder4/file.txt"));
+	Run(TEXT("C:/A/B/.././../C"),                                     TEXT("C:/C"));
+	Run(TEXT("C:/A/B/.././../C/"),                                    TEXT("C:/C/"));
+	Run(TEXT("C:/A/B/.././../C/file.txt"),                            TEXT("C:/C/file.txt"));
+	Run(TEXT(".svn"),                                                 TEXT(".svn"));
+	Run(TEXT("/.svn"),                                                TEXT("/.svn"));
+	Run(TEXT("./Folder/.svn"),                                        TEXT("Folder/.svn"));
+	Run(TEXT("./.svn/../.svn"),                                       TEXT(".svn"));
+	Run(TEXT(".svn/./.svn/.././../.svn"),                             TEXT("/.svn"));
+	Run(TEXT("C:/Folder1/./Folder2/..Folder3"),						 TEXT("C:/Folder1/Folder2/..Folder3"));
+	Run(TEXT("C:/Folder1/./Folder2/..Folder3/Folder4"),				 TEXT("C:/Folder1/Folder2/..Folder3/Folder4"));
+	Run(TEXT("C:/Folder1/./Folder2/..Folder3/..Folder4"),			 TEXT("C:/Folder1/Folder2/..Folder3/..Folder4"));
+	Run(TEXT("C:/Folder1/./Folder2/..Folder3/Folder4/../Folder5"),	 TEXT("C:/Folder1/Folder2/..Folder3/Folder5"));
+	Run(TEXT("C:/Folder1/..Folder2/Folder3/..Folder4/../Folder5"),	 TEXT("C:/Folder1/..Folder2/Folder3/Folder5"));
+}
+
+template<class PathType, class StringType>
+void TestRemoveDuplicateSlashes(FAutomationTestBase& Test)
+{
+	auto Run = [&](const TCHAR* Path, const TCHAR* Expected)
+	{
+		StringType Actual;
+		Actual += Path;	
+		PathType::RemoveDuplicateSlashes(Actual);
+		Test.TestEqual(TEXT("RemoveDuplicateSlashes"), *Actual, Expected);
+	};
+
+	Run(TEXT(""), TEXT(""));
+	Run(TEXT("C:/Folder/File.txt"), TEXT("C:/Folder/File.txt"));
+	Run(TEXT("C:/Folder/File/"), TEXT("C:/Folder/File/"));
+	Run(TEXT("/"), TEXT("/"));
+	Run(TEXT("//"), TEXT("/"));
+	Run(TEXT("////"), TEXT("/"));
+	Run(TEXT("/Folder/File"), TEXT("/Folder/File"));
+	Run(TEXT("//Folder/File"), TEXT("/Folder/File")); // Don't use on //UNC paths; it will be stripped!
+	Run(TEXT("/////Folder//////File/////"), TEXT("/Folder/File/"));
+	Run(TEXT("\\\\Folder\\\\File\\\\"), TEXT("\\\\Folder\\\\File\\\\")); // It doesn't strip backslash, and we rely on that in some places
+	Run(TEXT("//\\\\//Folder//\\\\//File//\\\\//"), TEXT("/\\\\/Folder/\\\\/File/\\\\/"));
+}
+
+namespace PathTest
+{
+
+struct FTestPair
+{
+	FStringView Input;
+	FStringView Expected;
+};
+
+constexpr FStringView BaseDir = TEXT("/root"_SV);
+
+constexpr FTestPair ExpectedRelativeToAbsolutePaths[] =
+{
+	{ TEXT(""_SV),					TEXT("/root/"_SV) },
+	{ TEXT("dir"_SV),				TEXT("/root/dir"_SV) },
+	{ TEXT("/groot"_SV),			TEXT("/groot"_SV) },
+	{ TEXT("/groot/"_SV),			TEXT("/groot/"_SV) },
+	{ TEXT("/r/dir"_SV),			TEXT("/r/dir"_SV) },
+	{ TEXT("/r/dir"_SV),			TEXT("/r/dir"_SV) },
+	{ TEXT("C:\\"_SV),				TEXT("C:/"_SV) },
+	{ TEXT("C:\\A\\B"_SV),			TEXT("C:/A/B"_SV) },
+	{ TEXT("a/b/../c"_SV),			TEXT("/root/a/c"_SV) },
+	{ TEXT("/a/b/../c"_SV),			TEXT("/a/c"_SV) },
+};
+
+}
+
+#endif
\ No newline at end of file
diff --git a/Engine/Source/Runtime/Core/Private/Tests/Misc/PathViewsTest.cpp b/Engine/Source/Runtime/Core/Private/Tests/Misc/PathViewsTest.cpp
--- a/Engine/Source/Runtime/Core/Private/Tests/Misc/PathViewsTest.cpp
+++ b/Engine/Source/Runtime/Core/Private/Tests/Misc/PathViewsTest.cpp
@@ -1,8 +1,8 @@
 // Copyright Epic Games, Inc. All Rights Reserved.
 
-#include "Misc/PathViews.h"
-
+#include "PathTests.h"
 #include "Containers/StringView.h"
+#include "Misc/PathViews.h"
 #include "Misc/AutomationTest.h"
 #include "Misc/StringBuilder.h"
 
@@ -29,6 +29,20 @@ protected:
 	}
 };
 
+IMPLEMENT_CUSTOM_SIMPLE_AUTOMATION_TEST(FPathViewsCollapseDirectoriesTest, FPathViewsTest, "System.Core.Misc.PathViews.CollapseDirectories", EAutomationTestFlags::ApplicationContextMask | EAutomationTestFlags::SmokeFilter)
+bool FPathViewsCollapseDirectoriesTest::RunTest(const FString& InParameters)
+{
+	TestCollapseRelativeDirectories<FPathViews, TStringBuilder<64>>(*this);
+	return true;
+}
+
+IMPLEMENT_CUSTOM_SIMPLE_AUTOMATION_TEST(FPathViewsRemoveDuplicateSlashesTest, FPathViewsTest, "System.Core.Misc.PathViews.RemoveDuplicateSlashes", EAutomationTestFlags::ApplicationContextMask | EAutomationTestFlags::SmokeFilter)
+bool FPathViewsRemoveDuplicateSlashesTest::RunTest(const FString& InParameters)
+{
+	TestRemoveDuplicateSlashes<FPathViews, TStringBuilder<64>>(*this);
+	return true;
+}
+
 IMPLEMENT_CUSTOM_SIMPLE_AUTOMATION_TEST(FPathViewsGetCleanFilenameTest, FPathViewsTest, "System.Core.Misc.PathViews.GetCleanFilename", EAutomationTestFlags::ApplicationContextMask | EAutomationTestFlags::SmokeFilter)
 bool FPathViewsGetCleanFilenameTest::RunTest(const FString& InParameters)
 {
@@ -650,4 +664,33 @@ bool FPathViewsAppendPathTest::RunTest(const FString& InParameters)
 	return true;
 }
 
+IMPLEMENT_CUSTOM_SIMPLE_AUTOMATION_TEST(FPathViewsToAbsoluteTest, FPathViewsTest, "System.Core.Misc.PathViews.ToAbsolute", EAutomationTestFlags::ApplicationContextMask | EAutomationTestFlags::SmokeFilter)
+bool FPathViewsToAbsoluteTest::RunTest(const FString& InParameters)
+{
+	using namespace PathTest;
+
+	for (FTestPair Pair : ExpectedRelativeToAbsolutePaths)
+	{
+		TStringBuilder<64> ActualAppend;
+		FPathViews::ToAbsolutePath(BaseDir, Pair.Input, ActualAppend);
+		TestEqual(TEXT("ToAbsolutePath"), ActualAppend.ToView(), Pair.Expected);
+
+		TStringBuilder<64> ActualInline;
+		ActualInline << Pair.Input;
+		FPathViews::ToAbsolutePathInline(BaseDir, ActualInline);	
+		TestEqual(TEXT("ToAbsolutePathInline"), ActualInline.ToView(), Pair.Expected);
+		
+		constexpr FStringView Original = TEXT("\\\\la/./.././la////"_SV);
+		TStringBuilder<64> ActualNondestructive;
+		ActualNondestructive << Original;
+		FPathViews::ToAbsolutePath(BaseDir, Pair.Input, ActualNondestructive);
+		TestEqual(TEXT("ToAbsolutePath non-destructive append"), ActualNondestructive.ToView().Left(Original.Len()), Original);
+		TestEqual(TEXT("ToAbsolutePath non-destructive append"), ActualNondestructive.ToView().RightChop(Original.Len()), Pair.Expected);
+	}
+
+
+
+	return true;
+}
+
 #endif
diff --git a/Engine/Source/Runtime/Core/Private/Tests/Misc/PathsTest.cpp b/Engine/Source/Runtime/Core/Private/Tests/Misc/PathsTest.cpp
--- a/Engine/Source/Runtime/Core/Private/Tests/Misc/PathsTest.cpp
+++ b/Engine/Source/Runtime/Core/Private/Tests/Misc/PathsTest.cpp
@@ -1,5 +1,6 @@
 // Copyright Epic Games, Inc. All Rights Reserved.
 
+#include "PathTests.h"
 #include "CoreTypes.h"
 #include "Misc/AssertionMacros.h"
 #include "Containers/UnrealString.h"
@@ -12,89 +13,7 @@ IMPLEMENT_SIMPLE_AUTOMATION_TEST(FPathTests, "System.Core.Misc.Paths", EAutomati
 
 bool FPathTests::RunTest( const FString& Parameters )
 {
-	// Directory collapsing
-	{
-		auto RunCollapseRelativeDirectoriesTest = [this](const TCHAR* InPath, const TCHAR* InResult)
-		{
-			// Run test
-			FString CollapsedPath = InPath;
-			const bool bValid = FPaths::CollapseRelativeDirectories(CollapsedPath);
-
-			if (InResult)
-			{
-				// If we're looking for a result, make sure it was returned correctly
-				if (!bValid || CollapsedPath != InResult)
-				{
-					AddError(FString::Printf(TEXT("Path '%s' failed to collapse correctly (got '%s', expected '%s')."), InPath, *CollapsedPath, InResult));
-				}
-			}
-			else
-			{
-				// Otherwise, make sure it failed
-				if (bValid)
-				{
-					AddError(FString::Printf(TEXT("Path '%s' collapsed unexpectedly."), InPath));
-				}
-			}
-		};
-
-		RunCollapseRelativeDirectoriesTest(TEXT(".."),                                                   NULL);
-		RunCollapseRelativeDirectoriesTest(TEXT("/.."),                                                  NULL);
-		RunCollapseRelativeDirectoriesTest(TEXT("./"),                                                   TEXT(""));
-		RunCollapseRelativeDirectoriesTest(TEXT("./file.txt"),                                           TEXT("file.txt"));
-		RunCollapseRelativeDirectoriesTest(TEXT("/."),                                                   TEXT("/."));
-		RunCollapseRelativeDirectoriesTest(TEXT("Folder"),                                               TEXT("Folder"));
-		RunCollapseRelativeDirectoriesTest(TEXT("/Folder"),                                              TEXT("/Folder"));
-		RunCollapseRelativeDirectoriesTest(TEXT("C:/Folder"),                                            TEXT("C:/Folder"));
-		RunCollapseRelativeDirectoriesTest(TEXT("C:/Folder/.."),                                         TEXT("C:"));
-		RunCollapseRelativeDirectoriesTest(TEXT("C:/Folder/../"),                                        TEXT("C:/"));
-		RunCollapseRelativeDirectoriesTest(TEXT("C:/Folder/../file.txt"),                                TEXT("C:/file.txt"));
-		RunCollapseRelativeDirectoriesTest(TEXT("Folder/.."),                                            TEXT(""));
-		RunCollapseRelativeDirectoriesTest(TEXT("Folder/../"),                                           TEXT("/"));
-		RunCollapseRelativeDirectoriesTest(TEXT("Folder/../file.txt"),                                   TEXT("/file.txt"));
-		RunCollapseRelativeDirectoriesTest(TEXT("/Folder/.."),                                           TEXT(""));
-		RunCollapseRelativeDirectoriesTest(TEXT("/Folder/../"),                                          TEXT("/"));
-		RunCollapseRelativeDirectoriesTest(TEXT("/Folder/../file.txt"),                                  TEXT("/file.txt"));
-		RunCollapseRelativeDirectoriesTest(TEXT("Folder/../.."),                                         NULL);
-		RunCollapseRelativeDirectoriesTest(TEXT("Folder/../../"),                                        NULL);
-		RunCollapseRelativeDirectoriesTest(TEXT("Folder/../../file.txt"),                                NULL);
-		RunCollapseRelativeDirectoriesTest(TEXT("C:/.."),                                                NULL);
-		RunCollapseRelativeDirectoriesTest(TEXT("C:/."),                                                 TEXT("C:/."));
-		RunCollapseRelativeDirectoriesTest(TEXT("C:/./"),                                                TEXT("C:/"));
-		RunCollapseRelativeDirectoriesTest(TEXT("C:/./file.txt"),                                        TEXT("C:/file.txt"));
-		RunCollapseRelativeDirectoriesTest(TEXT("C:/Folder1/../Folder2"),                                TEXT("C:/Folder2"));
-		RunCollapseRelativeDirectoriesTest(TEXT("C:/Folder1/../Folder2/"),                               TEXT("C:/Folder2/"));
-		RunCollapseRelativeDirectoriesTest(TEXT("C:/Folder1/../Folder2/file.txt"),                       TEXT("C:/Folder2/file.txt"));
-		RunCollapseRelativeDirectoriesTest(TEXT("C:/Folder1/../Folder2/../.."),                          NULL);
-		RunCollapseRelativeDirectoriesTest(TEXT("C:/Folder1/../Folder2/../Folder3"),                     TEXT("C:/Folder3"));
-		RunCollapseRelativeDirectoriesTest(TEXT("C:/Folder1/../Folder2/../Folder3/"),                    TEXT("C:/Folder3/"));
-		RunCollapseRelativeDirectoriesTest(TEXT("C:/Folder1/../Folder2/../Folder3/file.txt"),            TEXT("C:/Folder3/file.txt"));
-		RunCollapseRelativeDirectoriesTest(TEXT("C:/Folder1/Folder2/../../Folder3"),                     TEXT("C:/Folder3"));
-		RunCollapseRelativeDirectoriesTest(TEXT("C:/Folder1/Folder2/../../Folder3/"),                    TEXT("C:/Folder3/"));
-		RunCollapseRelativeDirectoriesTest(TEXT("C:/Folder1/Folder2/../../Folder3/file.txt"),            TEXT("C:/Folder3/file.txt"));
-		RunCollapseRelativeDirectoriesTest(TEXT("C:/Folder1/Folder2/../../Folder3/../Folder4"),          TEXT("C:/Folder4"));
-		RunCollapseRelativeDirectoriesTest(TEXT("C:/Folder1/Folder2/../../Folder3/../Folder4/"),         TEXT("C:/Folder4/"));
-		RunCollapseRelativeDirectoriesTest(TEXT("C:/Folder1/Folder2/../../Folder3/../Folder4/file.txt"), TEXT("C:/Folder4/file.txt"));
-		RunCollapseRelativeDirectoriesTest(TEXT("C:/Folder1/Folder2/../Folder3/../../Folder4"),          TEXT("C:/Folder4"));
-		RunCollapseRelativeDirectoriesTest(TEXT("C:/Folder1/Folder2/../Folder3/../../Folder4/"),         TEXT("C:/Folder4/"));
-		RunCollapseRelativeDirectoriesTest(TEXT("C:/Folder1/Folder2/../Folder3/../../Folder4/file.txt"), TEXT("C:/Folder4/file.txt"));
-		RunCollapseRelativeDirectoriesTest(TEXT("C:/Folder1/Folder2/.././../Folder4"),                   TEXT("C:/Folder4"));
-		RunCollapseRelativeDirectoriesTest(TEXT("C:/Folder1/Folder2/.././../Folder4/"),                  TEXT("C:/Folder4/"));
-		RunCollapseRelativeDirectoriesTest(TEXT("C:/Folder1/Folder2/.././../Folder4/file.txt"),          TEXT("C:/Folder4/file.txt"));
-		RunCollapseRelativeDirectoriesTest(TEXT("C:/A/B/.././../C"),                                     TEXT("C:/C"));
-		RunCollapseRelativeDirectoriesTest(TEXT("C:/A/B/.././../C/"),                                    TEXT("C:/C/"));
-		RunCollapseRelativeDirectoriesTest(TEXT("C:/A/B/.././../C/file.txt"),                            TEXT("C:/C/file.txt"));
-		RunCollapseRelativeDirectoriesTest(TEXT(".svn"),                                                 TEXT(".svn"));
-		RunCollapseRelativeDirectoriesTest(TEXT("/.svn"),                                                TEXT("/.svn"));
-		RunCollapseRelativeDirectoriesTest(TEXT("./Folder/.svn"),                                        TEXT("Folder/.svn"));
-		RunCollapseRelativeDirectoriesTest(TEXT("./.svn/../.svn"),                                       TEXT(".svn"));
-		RunCollapseRelativeDirectoriesTest(TEXT(".svn/./.svn/.././../.svn"),                             TEXT("/.svn"));
-		RunCollapseRelativeDirectoriesTest(TEXT("C:/Folder1/./Folder2/..Folder3"),						 TEXT("C:/Folder1/Folder2/..Folder3"));
-		RunCollapseRelativeDirectoriesTest(TEXT("C:/Folder1/./Folder2/..Folder3/Folder4"),				 TEXT("C:/Folder1/Folder2/..Folder3/Folder4"));
-		RunCollapseRelativeDirectoriesTest(TEXT("C:/Folder1/./Folder2/..Folder3/..Folder4"),			 TEXT("C:/Folder1/Folder2/..Folder3/..Folder4"));
-		RunCollapseRelativeDirectoriesTest(TEXT("C:/Folder1/./Folder2/..Folder3/Folder4/../Folder5"),	 TEXT("C:/Folder1/Folder2/..Folder3/Folder5"));
-		RunCollapseRelativeDirectoriesTest(TEXT("C:/Folder1/..Folder2/Folder3/..Folder4/../Folder5"),	 TEXT("C:/Folder1/..Folder2/Folder3/Folder5"));
-	}
+	TestCollapseRelativeDirectories<FPaths, FString>(*this);
 
 	// Extension texts
 	{
@@ -184,34 +103,17 @@ bool FPathTests::RunTest( const FString& Parameters )
 		RunIsUnderDirectoryTest(TEXT("C:/Folder/Subdir/"),	TEXT("C:/Folder/"), true);
 	}
 
-	// RemoveDuplicateSlashes
+	TestRemoveDuplicateSlashes<FPaths, FString>(*this);
+
+	// ConvertRelativePathToFull
 	{
-		auto RunRemoveDuplicateSlashesTest = [this](const TCHAR* InPath, const TCHAR* InExpectedResult)
-		{
-			FString ReplacementPath(InPath);
-			const FString ExpectedResult(InExpectedResult);
-			FPaths::RemoveDuplicateSlashes(ReplacementPath);
-			if (!ReplacementPath.Equals(ExpectedResult, ESearchCase::CaseSensitive))
-			{
-				AddError(FString::Printf(TEXT("FPaths::RemoveDuplicateSlashes('%s') != '%s'."), InPath, InExpectedResult));
-			}
-			else if (ReplacementPath.GetCharArray().Num() != 0 && ReplacementPath.GetCharArray().Num() != ExpectedResult.Len() + 1)
-			{
-				AddError(FString::Printf(TEXT("FPaths::RemoveDuplicateSlashes('%s') returned a result with extra space still allocated after the null terminator."), InPath));
-			}
-		};
+		using namespace PathTest;
 
-		RunRemoveDuplicateSlashesTest(TEXT(""), TEXT(""));
-		RunRemoveDuplicateSlashesTest(TEXT("C:/Folder/File.txt"), TEXT("C:/Folder/File.txt"));
-		RunRemoveDuplicateSlashesTest(TEXT("C:/Folder/File/"), TEXT("C:/Folder/File/"));
-		RunRemoveDuplicateSlashesTest(TEXT("/"), TEXT("/"));
-		RunRemoveDuplicateSlashesTest(TEXT("//"), TEXT("/"));
-		RunRemoveDuplicateSlashesTest(TEXT("////"), TEXT("/"));
-		RunRemoveDuplicateSlashesTest(TEXT("/Folder/File"), TEXT("/Folder/File"));
-		RunRemoveDuplicateSlashesTest(TEXT("//Folder/File"), TEXT("/Folder/File")); // Don't use on //UNC paths; it will be stripped!
-		RunRemoveDuplicateSlashesTest(TEXT("/////Folder//////File/////"), TEXT("/Folder/File/"));
-		RunRemoveDuplicateSlashesTest(TEXT("\\\\Folder\\\\File\\\\"), TEXT("\\\\Folder\\\\File\\\\")); // It doesn't strip backslash, and we rely on that in some places
-		RunRemoveDuplicateSlashesTest(TEXT("//\\\\//Folder//\\\\//File//\\\\//"), TEXT("/\\\\/Folder/\\\\/File/\\\\/"));
+		for (FTestPair Pair : ExpectedRelativeToAbsolutePaths)
+		{
+			FString Actual = FPaths::ConvertRelativePathToFull(FString(BaseDir), FString(Pair.Input));
+			TestEqual(TEXT("ConvertRelativePathToFull"), FStringView(Actual), Pair.Expected);
+		}
 	}
 
 	return true;
diff --git a/Engine/Source/Runtime/Core/Public/Misc/PathViews.h b/Engine/Source/Runtime/Core/Public/Misc/PathViews.h
--- a/Engine/Source/Runtime/Core/Public/Misc/PathViews.h
+++ b/Engine/Source/Runtime/Core/Public/Misc/PathViews.h
@@ -280,6 +280,30 @@ public:
 	 */
 	static bool IsRelativePath(FStringView InPath);
 
+	/** Convert to absolute using process BaseDir(), normalize and append. FPaths::ConvertRelativePathToFull() equivalent. */
+	static void ToAbsolutePath(FStringView InPath, FStringBuilderBase& OutPath);
+	
+	/** Convert to absolute using explicit BasePath, normalize and append. FPaths::ConvertRelativePathToFull() equivalent. */
+	static void ToAbsolutePath(FStringView BasePath, FStringView InPath, FStringBuilderBase& OutPath);
+	
+	/** Convert to absolute using process BaseDir() and normalize inlined. FPaths::ConvertRelativePathToFull() equivalent. */
+	static void ToAbsolutePathInline(FStringBuilderBase& InOutPath);
+	
+	/** Convert to absolute using explicit BasePath and normalize inlined. FPaths::ConvertRelativePathToFull() equivalent. */
+	static void ToAbsolutePathInline(FStringView BasePath, FStringBuilderBase& InOutPath);
+
+	/** Convert \\ to / and do platform-specific normalization */
+	static void NormalizeFilename(FStringBuilderBase& InOutPath);
+
+	/** Normalize and remove trailing slash unless the preceding character is '/' or ':' */
+	static void NormalizeDirectoryName(FStringBuilderBase& InOutPath);
+
+	/** Collapses redundant paths like "/./" and "SkipDir/..". FPaths::CollapseRelativeDirectories() equivalent. */
+	static bool CollapseRelativeDirectories(FStringBuilderBase& InOutPath);
+
+	/** Removes duplicate forward slashes, e.g. "a/b//c////f.e" -> "a/b/c/f.e" */
+	static void RemoveDuplicateSlashes(FStringBuilderBase& InOutPath);
+
 	/**
 	 * Split the given absolute or relative path into its topmost directory and the relative path from that directory.
 	 * Directory elements of . and .. are currently not interpreted and are treated as literal characters.
