Removed "overview" memory tracing. It is disabled and interleaves needless noise in the full-fat memory tracing.

diff --git a/Engine/Source/Runtime/Core/Private/ProfilingDebugging/Windows/MemoryTrace.cpp b/Engine/Source/Runtime/Core/Private/ProfilingDebugging/Windows/MemoryTrace.cpp
--- a/Engine/Source/Runtime/Core/Private/ProfilingDebugging/Windows/MemoryTrace.cpp
+++ b/Engine/Source/Runtime/Core/Private/ProfilingDebugging/Windows/MemoryTrace.cpp
@@ -19,8 +19,6 @@
 #include <atomic>
 #include <intrin.h>
 
-#define USE_OVERVIEW_TRACE 0 // Disabled for now, as it occasionally conflicts with MallocBinned2's canary poison logic.
-
 ////////////////////////////////////////////////////////////////////////////////
 namespace UE {
 namespace Trace {
@@ -86,22 +84,9 @@ UE_TRACE_EVENT_BEGIN(Memory, Init, NoSync|Important)
 	UE_TRACE_EVENT_FIELD(uint32, MarkerPeriod)
 	UE_TRACE_EVENT_FIELD(uint8, MinAlignment)
 	UE_TRACE_EVENT_FIELD(uint8, SizeShift)
-#if USE_OVERVIEW_TRACE
-	UE_TRACE_EVENT_FIELD(uint8, SummarySizeShift)
-#endif
 	UE_TRACE_EVENT_FIELD(uint8, Mode)
 UE_TRACE_EVENT_END()
 
-#if USE_OVERVIEW_TRACE
-UE_TRACE_EVENT_BEGIN(Memory, Summary)
-	UE_TRACE_EVENT_FIELD(uint32, Bytes)
-	UE_TRACE_EVENT_FIELD(uint32, ActiveAllocs)
-	UE_TRACE_EVENT_FIELD(uint32, TotalAllocs)
-	UE_TRACE_EVENT_FIELD(uint32, TotalReallocs)
-	UE_TRACE_EVENT_FIELD(uint32, TotalFrees)
-UE_TRACE_EVENT_END()
-#endif
-
 UE_TRACE_EVENT_BEGIN(Memory, Marker)
 	UE_TRACE_EVENT_FIELD(uint64, Cycle)
 UE_TRACE_EVENT_END()
@@ -142,139 +127,6 @@ UE_TRACE_EVENT_END()
 
 
 
-#if USE_OVERVIEW_TRACE
-
-////////////////////////////////////////////////////////////////////////////////
-class FSummaryTrace
-{
-public:
-	static const uint32			SizeShift = 4;
-	void						Initialize();
-	uint32						SetTag(uint32 Tag);
-	uint32						Alloc(SIZE_T Size, bool IsRealloc);
-	void						Free(uint32 Tag, SIZE_T Size);
-	void						Flush();
-
-private:
-	template <typename T>
-	struct FItem
-	{
-		std::atomic<T>			Current = 0;
-		T						Max = 0;
-		T						Min = 0;
-		void					Add(T Operand);
-	};
-
-	struct alignas(PLATFORM_CACHE_LINE_SIZE) FInfo
-	{
-		uint32					Tag;
-		FItem<int64>			Bytes;
-		FItem<int32>			ActiveAllocs;
-		std::atomic<uint32>		TotalAllocs = 0;
-		std::atomic<uint32>		TotalReallocs = 0;
-		std::atomic<uint32>		TotalFrees = 0;
-	};
-
-	FInfo&						GetInfo();
-	static thread_local FInfo*	InfoPtr;
-	std::atomic<uint32>			InfoIndex;
-	FInfo						InfoPool[1];
-};
-
-////////////////////////////////////////////////////////////////////////////////
-static FUndestructed<FSummaryTrace>	GSummaryTrace;
-thread_local FSummaryTrace::FInfo*	FSummaryTrace::InfoPtr = nullptr;
-
-////////////////////////////////////////////////////////////////////////////////
-template <typename T>
-void FSummaryTrace::FItem<T>::Add(T Size)
-{
-	T Value = Current.fetch_add(T(Size), std::memory_order_relaxed);
-	Value += Size;
-	Max = (Current > Max) ? Value : Max;
-	Min = (Current < Min) ? Value : Min;
-}
-
-////////////////////////////////////////////////////////////////////////////////
-void FSummaryTrace::Initialize()
-{
-	FCoreDelegates::OnEndFrame.AddStatic([] () { GSummaryTrace->Flush(); });
-}
-
-////////////////////////////////////////////////////////////////////////////////
-uint32 FSummaryTrace::SetTag(uint32 Tag)
-{
-	FInfo& Info = GetInfo();
-	uint32 PrevTag = Info.Tag;
-	Info.Tag = Tag;
-	return PrevTag;
-}
-
-////////////////////////////////////////////////////////////////////////////////
-uint32 FSummaryTrace::Alloc(SIZE_T Size, bool IsRealloc)
-{
-	FInfo& Info = GetInfo();
-	Info.Bytes.Add(Size);
-	Info.ActiveAllocs.Add(1);
-	Info.TotalAllocs.fetch_add(1, std::memory_order_relaxed);
-	if (IsRealloc)
-	{
-		Info.TotalReallocs.fetch_add(1, std::memory_order_relaxed);
-	}
-	return Info.Tag;
-}
-
-////////////////////////////////////////////////////////////////////////////////
-void FSummaryTrace::Free(uint32 Tag, SIZE_T Size)
-{
-	FInfo& Info = GetInfo();
-	Info.Bytes.Add(-int64(Size));
-	Info.ActiveAllocs.Add(-1);
-	Info.TotalFrees.fetch_add(1, std::memory_order_relaxed);
-}
-
-////////////////////////////////////////////////////////////////////////////////
-FSummaryTrace::FInfo& FSummaryTrace::GetInfo()
-{
-	FInfo* Info = InfoPtr;
-	if (Info == nullptr)
-	{
-		uint32 Index = InfoIndex.fetch_add(1, std::memory_order_relaxed);
-		Index &= UE_ARRAY_COUNT(InfoPool) - 1;
-		Info = InfoPtr = InfoPool + Index;
-	}
-	return *Info;
-}
-
-////////////////////////////////////////////////////////////////////////////////
-void FSummaryTrace::Flush()
-{
-	int64 Bytes = 0;
-	int32 Active = 0;
-	int32 TotalAllocs = 0;
-	int32 TotalReallocs = 0;
-	int32 TotalFrees = 0;
-	for (const FInfo& Info : InfoPool)
-	{
-		Bytes += Info.Bytes.Current.load(std::memory_order_relaxed);
-		Active += Info.ActiveAllocs.Current.load(std::memory_order_relaxed);
-		TotalAllocs += Info.TotalAllocs.load(std::memory_order_relaxed);
-		TotalReallocs += Info.TotalReallocs.load(std::memory_order_relaxed);
-		TotalFrees += Info.TotalFrees.load(std::memory_order_relaxed);
-	}
-	
-	UE_TRACE_LOG(Memory, Summary, MemSummaryChannel)
-		<< Summary.Bytes(uint32(Bytes >> SizeShift))
-		<< Summary.ActiveAllocs(Active)
-		<< Summary.TotalAllocs(TotalAllocs)
-		<< Summary.TotalReallocs(TotalReallocs)
-		<< Summary.TotalFrees(TotalFrees);
-}
-
-#endif // USE_OVERVIEW_TRACE
-
-
-
 ////////////////////////////////////////////////////////////////////////////////
 class FAllocationTrace
 {
@@ -306,9 +158,6 @@ void FAllocationTrace::Initialize()
 	UE_TRACE_LOG(Memory, Init, MemAllocChannel)
 		<< Init.MarkerPeriod(MarkerSamplePeriod + 1)
 		<< Init.MinAlignment(uint8(MIN_ALIGNMENT))
-#if USE_OVERVIEW_TRACE
-		<< Init.SummarySizeShift(uint8(FSummaryTrace::SizeShift))
-#endif
 		<< Init.SizeShift(uint8(SizeShift));
 
 	static_assert((1 << SizeShift) - 1 <= MIN_ALIGNMENT, "Not enough bits to pack size fields");
@@ -459,22 +308,7 @@ void* FMallocWrapper::Malloc(SIZE_T Size, uint32 Alignment)
 	}
 
 	uint32 ActualAlignment = GetActualAlignment(Size, Alignment);
-
-	SIZE_T InnerSize = Size;
-#if USE_OVERVIEW_TRACE
-	static_assert(sizeof(FCookie) <= 8, "Unexpected FCookie size");
-	InnerSize += ActualAlignment;
-#endif 
-
-	void* Address = InnerMalloc->Malloc(InnerSize, Alignment);
-
-#if USE_OVERVIEW_TRACE
-	Address = (uint8*)Address + ActualAlignment;
-	FCookie* Cookie = (FCookie*)Address - 1;
-	Cookie->Size = Size;
-	Cookie->Tag = GSummaryTrace->Alloc(Size, false);
-	Cookie->Bias = ActualAlignment >> 3;
-#endif
+	void* Address = InnerMalloc->Malloc(Size, Alignment);
 
 	void* Owner = GetOwner(bLight);
 	GAllocationTrace->Alloc(Address, Size, ActualAlignment, Owner);
@@ -497,29 +331,9 @@ void* FMallocWrapper::Realloc(void* PrevAddress, SIZE_T NewSize, uint32 Alignmen
 		return nullptr;
 	}
 
-	// Track the block that will (or might) get freed
-	uint32 HeaderSize = 0;
-	void* InnerAddress = PrevAddress;
-#if USE_OVERVIEW_TRACE
-	FCookie* Cookie = (FCookie*)PrevAddress - 1;
-	HeaderSize = uint32(Cookie->Bias << 3);
-	InnerAddress = (uint8*)PrevAddress - HeaderSize;
-	GSummaryTrace->Free(Cookie->Tag, Cookie->Size);
-#endif
-
 	GAllocationTrace->ReallocFree(PrevAddress);
 
-	// Do the actual malloc
-	SIZE_T InnerSize = NewSize + HeaderSize;
-	void* RetAddress = InnerMalloc->Realloc(InnerAddress, InnerSize, Alignment);
-
-	// Track the block that was allocated
-#if USE_OVERVIEW_TRACE
-	RetAddress = (uint8*)RetAddress + HeaderSize;
-	Cookie = (FCookie*)RetAddress - 1;
-	Cookie->Size = NewSize;
-	GSummaryTrace->Alloc(NewSize, true);
-#endif
+	void* RetAddress = InnerMalloc->Realloc(PrevAddress, NewSize, Alignment);
 
 	void* Owner = GetOwner(bLight);
 	Alignment = GetActualAlignment(NewSize, Alignment);
@@ -539,11 +353,6 @@ void FMallocWrapper::Free(void* Address)
 	GAllocationTrace->Free(Address);
 
 	void* InnerAddress = Address;
-#if USE_OVERVIEW_TRACE
-	const FCookie* Cookie = (const FCookie*)Address - 1;
-	InnerAddress = (uint8*)Address - (Cookie->Bias << 3);
-	GSummaryTrace->Free(Cookie->Tag, Cookie->Size);
-#endif
 
 	return InnerMalloc->Free(InnerAddress);
 }
@@ -811,10 +620,6 @@ FMalloc* MemoryTrace_Create(FMalloc* InMalloc)
 		GAllocationTrace->Initialize();
 
 		MemoryTrace_InitTags(InMalloc);
-#if USE_OVERVIEW_TRACE
-		GSummaryTrace.Construct();
-		GSummaryTrace->Initialize();
-#endif
 
 		Backtracer_Create(InMalloc);
 
@@ -853,8 +658,6 @@ void MemoryTrace_CoreRemove(void* Base, SIZE_T Size)
 }
 #endif // 0
 
-#undef USE_OVERVIEW_TRACE
-
 #include "Windows/HideWindowsPlatformTypes.h"
 
 #endif // UE_MEMORY_TRACE_ENABLED
