Enable Memory tracing for Memory channel preset

diff --git a/Engine/Source/Runtime/Core/Private/ProfilingDebugging/Windows/MemoryTrace.cpp b/Engine/Source/Runtime/Core/Private/ProfilingDebugging/Windows/MemoryTrace.cpp
--- a/Engine/Source/Runtime/Core/Private/ProfilingDebugging/Windows/MemoryTrace.cpp
+++ b/Engine/Source/Runtime/Core/Private/ProfilingDebugging/Windows/MemoryTrace.cpp
@@ -594,7 +594,7 @@ FMalloc* MemoryTrace_Create(FMalloc* InMalloc)
 			if (*c == ' ' || *c == '\0' || *c == ',')
 			{
 				FStringView View(Start, uint32(c - Start));
-				if (View.Equals(TEXT("memalloc")))
+				if (View.Equals(TEXT("memalloc"), ESearchCase::IgnoreCase) || View.Equals(TEXT("memory"), ESearchCase::IgnoreCase))
 				{
 					Mode = 2;
 					break;
