Fixed up "UE4" references in comments under Engine/Runtime/Core/.../Misc.

diff --git a/Engine/Source/Runtime/Core/Private/Misc/CoreMisc.cpp b/Engine/Source/Runtime/Core/Private/Misc/CoreMisc.cpp
--- a/Engine/Source/Runtime/Core/Private/Misc/CoreMisc.cpp
+++ b/Engine/Source/Runtime/Core/Private/Misc/CoreMisc.cpp
@@ -82,7 +82,7 @@ bool FStaticSelfRegisteringExec::Exec( UWorld* InWorld, const TCHAR* Cmd, FOutpu
 }
 
 
-			// Remove old UE4 crash contexts
+// Remove old crash contexts
 
 
 /*-----------------------------------------------------------------------------
diff --git a/Engine/Source/Runtime/Core/Private/Misc/FileHelper.cpp b/Engine/Source/Runtime/Core/Private/Misc/FileHelper.cpp
--- a/Engine/Source/Runtime/Core/Private/Misc/FileHelper.cpp
+++ b/Engine/Source/Runtime/Core/Private/Misc/FileHelper.cpp
@@ -959,7 +959,7 @@ void FMaintenance::DeleteOldLogs()
 		}
 	}
 
-	// Remove all legacy UE4 crash contexts (regardless of age and purge settings, these are deprecated)
+	// Remove all legacy crash contexts (regardless of age and purge settings, these are deprecated)
 	TArray<FString> Directories;
 	IFileManager::Get().FindFiles(Directories, *FString::Printf(TEXT("%s/UE4CC*"), *FPaths::ProjectLogDir()), false, true);
 
diff --git a/Engine/Source/Runtime/Core/Private/Misc/Guid.cpp b/Engine/Source/Runtime/Core/Private/Misc/Guid.cpp
--- a/Engine/Source/Runtime/Core/Private/Misc/Guid.cpp
+++ b/Engine/Source/Runtime/Core/Private/Misc/Guid.cpp
@@ -186,7 +186,7 @@ bool FGuid::ParseExact(const FString& GuidString, EGuidFormats Format, FGuid& Ou
 			return false;
 		}
 
-		// Replace the characters we replaced going out (but not the padding, UE4 discards it immediately)
+		// Replace the characters we replaced going out (but not the padding, UE discards it immediately)
 		FString GuidCopy = GuidString;
 		GuidCopy.ReplaceCharInline(TEXT('-'), TEXT('+'), ESearchCase::CaseSensitive);
 		GuidCopy.ReplaceCharInline(TEXT('_'), TEXT('/'), ESearchCase::CaseSensitive);
diff --git a/Engine/Source/Runtime/Core/Private/Misc/SecureHash.cpp b/Engine/Source/Runtime/Core/Private/Misc/SecureHash.cpp
--- a/Engine/Source/Runtime/Core/Private/Misc/SecureHash.cpp
+++ b/Engine/Source/Runtime/Core/Private/Misc/SecureHash.cpp
@@ -984,7 +984,7 @@ void FAsyncSHAVerify::DoWork()
  */
 
 /* *** NEVER CHECK THE BELOW IN SET TO 1!!! *** */
-#define DISABLE_AUTHENTICATION_FOR_DEV 1    // UE4 - for now, we _will_ check this in with 1 because we are not generating shipping images at the moment.
+#define DISABLE_AUTHENTICATION_FOR_DEV 1    // UE - for now, we _will_ check this in with 1 because we are not generating shipping images at the moment.
 /* *** NEVER CHECK THE ABOVE IN SET TO 1!!! *** */
 
 void appOnFailSHAVerification(const TCHAR* FailedPathname, bool bFailedDueToMissingHash)
diff --git a/Engine/Source/Runtime/Core/Public/Misc/App.h b/Engine/Source/Runtime/Core/Public/Misc/App.h
--- a/Engine/Source/Runtime/Core/Public/Misc/App.h
+++ b/Engine/Source/Runtime/Core/Public/Misc/App.h
@@ -102,7 +102,7 @@ public:
 	}
 
 	/**
-	 * Gets the name of the application, i.e. "UE4" or "Rocket".
+	 * Gets the name of the application, i.e. "UE" or "Rocket".
 	 *
 	 * @todo need better application name discovery. this is quite horrible and may not work on future platforms.
 	 * @return Application name string.
@@ -405,7 +405,7 @@ public:
 	/**
 	 * Checks whether the engine components of this application have been installed.
 	 *
-	 * In binary UE4 releases, the engine can be installed while the game is not. The game IsInstalled()
+	 * In binary Unreal Engine releases, the engine can be installed while the game is not. The game IsInstalled()
 	 * setting will take precedence over this flag.
 	 *
 	 * To override, pass -engineinstalled or -enginenotinstalled on the command line.
diff --git a/Engine/Source/Runtime/Core/Public/Misc/CompilationResult.h b/Engine/Source/Runtime/Core/Public/Misc/CompilationResult.h
--- a/Engine/Source/Runtime/Core/Public/Misc/CompilationResult.h
+++ b/Engine/Source/Runtime/Core/Public/Misc/CompilationResult.h
@@ -8,7 +8,7 @@
  * Enumerates possible results of a compilation operation.
  *
  * This enum has to be compatible with the one defined in the
- * UE4\Engine\Source\Programs\UnrealBuildTool\System\ExternalExecution.cs file
+ * Engine\Source\Programs\UnrealBuildTool\System\ExternalExecution.cs file
  * to keep communication between UHT, UBT and Editor compiling processes valid.
  */
 namespace ECompilationResult
diff --git a/Engine/Source/Runtime/Core/Public/Misc/ConfigCacheIni.h b/Engine/Source/Runtime/Core/Public/Misc/ConfigCacheIni.h
--- a/Engine/Source/Runtime/Core/Public/Misc/ConfigCacheIni.h
+++ b/Engine/Source/Runtime/Core/Public/Misc/ConfigCacheIni.h
@@ -429,7 +429,7 @@ public:
 #endif // ALLOW_INI_OVERRIDE_FROM_COMMANDLINE
 	
 	CORE_API FConfigFile();
-	FConfigFile( int32 ) {}	// @todo UE4 DLL: Workaround for instantiated TMap template during DLLExport (TMap::FindRef)
+	FConfigFile( int32 ) {}	// @todo UE-DLL: Workaround for instantiated TMap template during DLLExport (TMap::FindRef)
 	CORE_API ~FConfigFile();
 	
 	// looks for a section by name, and creates an empty one if it can't be found
diff --git a/Engine/Source/Runtime/Core/Public/Misc/CoreDelegates.h b/Engine/Source/Runtime/Core/Public/Misc/CoreDelegates.h
--- a/Engine/Source/Runtime/Core/Public/Misc/CoreDelegates.h
+++ b/Engine/Source/Runtime/Core/Public/Misc/CoreDelegates.h
@@ -258,7 +258,7 @@ public:
 	// Called before the engine exits. Separate from OnPreExit as OnEnginePreExit occurs before shutting down any core modules.
 	static FSimpleMulticastDelegate OnEnginePreExit;
 
-	/** Delegate for gathering up additional localization paths that are unknown to the UE4 core (such as plugins) */
+	/** Delegate for gathering up additional localization paths that are unknown to the engine core (such as plugins) */
 	DECLARE_MULTICAST_DELEGATE_OneParam(FGatherAdditionalLocResPathsDelegate, TArray<FString>&);
 	static FGatherAdditionalLocResPathsDelegate GatherAdditionalLocResPathsCallback;
 
diff --git a/Engine/Source/Runtime/Core/Public/Misc/EmbeddedCommunication.h b/Engine/Source/Runtime/Core/Public/Misc/EmbeddedCommunication.h
--- a/Engine/Source/Runtime/Core/Public/Misc/EmbeddedCommunication.h
+++ b/Engine/Source/Runtime/Core/Public/Misc/EmbeddedCommunication.h
@@ -28,7 +28,7 @@ class CORE_API FEmbeddedDelegates
 {
 public:
 
-	// delegate for calling between native wrapper app and embedded ue4
+	// delegate for calling between native wrapper app and embedded ue
 	DECLARE_MULTICAST_DELEGATE_OneParam(FEmbeddedCommunicationParamsDelegate, const FEmbeddedCallParamsHelper&);
 
 	// calling in from native wrapper to engine
@@ -66,7 +66,7 @@ class CORE_API FEmbeddedCommunication
 {
 public:
 
-	// called early in UE4 lifecycle - RunOnGameThread can be called before this is called
+	// called early in UE lifecycle - RunOnGameThread can be called before this is called
 	static void Init();
 
 	// force some ticking to happen - used to process messages during otherwise blocking operations like boot
@@ -81,7 +81,7 @@ public:
 	// called from game thread to pull off
 	static bool TickGameThread(float DeltaTime);
 	
-	// tell UE4 to stay awake (or allow it to sleep when nothing to do). Repeated calls with the same Requester are
+	// tell UE to stay awake (or allow it to sleep when nothing to do). Repeated calls with the same Requester are
 	// allowed, but bNeedsRendering must match
 	static void KeepAwake(FName Requester, bool bNeedsRendering);
 	static void AllowSleep(FName Requester);
@@ -105,7 +105,7 @@ public:
 class FEmbeddedKeepAwake
 {
 public:
-	// tell UE4 to stay awake (or allow it to sleep when nothing to do). Repeated calls with the same Requester are
+	// tell UE to stay awake (or allow it to sleep when nothing to do). Repeated calls with the same Requester are
 	// allowed, but bNeedsRendering must match
 	FEmbeddedKeepAwake(FName InRequester, bool bInNeedsRendering)
 		: Requester(InRequester)
diff --git a/Engine/Source/Runtime/Core/Public/Misc/EngineVersionComparison.h b/Engine/Source/Runtime/Core/Public/Misc/EngineVersionComparison.h
--- a/Engine/Source/Runtime/Core/Public/Misc/EngineVersionComparison.h
+++ b/Engine/Source/Runtime/Core/Public/Misc/EngineVersionComparison.h
@@ -8,12 +8,12 @@
 #define UE_GREATER_SORT(Value, ValueToBeGreaterThan, TieBreaker) \
 	(((Value) > (ValueToBeGreaterThan)) || (((Value) == (ValueToBeGreaterThan)) && (TieBreaker)))
 
-// Version comparison macro that is defined to true if the UE4 version is newer than MajorVer.MinorVer.PatchVer and false otherwise
+// Version comparison macro that is defined to true if the UE version is newer than MajorVer.MinorVer.PatchVer and false otherwise
 // (a typical use is to alert integrators to revisit this code when upgrading to a new engine version)
 #define UE_VERSION_NEWER_THAN(MajorVersion, MinorVersion, PatchVersion) \
 	UE_GREATER_SORT(ENGINE_MAJOR_VERSION, MajorVersion, UE_GREATER_SORT(ENGINE_MINOR_VERSION, MinorVersion, UE_GREATER_SORT(ENGINE_PATCH_VERSION, PatchVersion, false)))
 
-// Version comparison macro that is defined to true if the UE4 version is older than MajorVer.MinorVer.PatchVer and false otherwise
+// Version comparison macro that is defined to true if the UE version is older than MajorVer.MinorVer.PatchVer and false otherwise
 // (use when making code that needs to be compatible with older engine versions)
 #define UE_VERSION_OLDER_THAN(MajorVersion, MinorVersion, PatchVersion) \
 	UE_GREATER_SORT(MajorVersion, ENGINE_MAJOR_VERSION, UE_GREATER_SORT(MinorVersion, ENGINE_MINOR_VERSION, UE_GREATER_SORT(PatchVersion, ENGINE_PATCH_VERSION, false)))
diff --git a/Engine/Source/Runtime/Core/Public/Misc/TypeContainer.h b/Engine/Source/Runtime/Core/Public/Misc/TypeContainer.h
--- a/Engine/Source/Runtime/Core/Public/Misc/TypeContainer.h
+++ b/Engine/Source/Runtime/Core/Public/Misc/TypeContainer.h
@@ -38,7 +38,7 @@ enum class ETypeContainerScope
  * being locked to a particular type of class. When passed into class constructors or
  * methods, type containers can facilitate the Inversion of Control (IoC) pattern.
  *
- * Since UE4 neither uses run-time type information nor pre-processes plain old C++ classes,
+ * Since UE neither uses run-time type information nor pre-processes plain old C++ classes,
  * type names need to be exposed using the Expose_TNameOf macro in order to make them
  * registrable with type containers, i.e. Expose_TNameOf(FMyClass).
  *
