Add "uninitialized const reference" warning disable flag for third party includes.

- Add PRAGMA_DISABLE_UNINITIALIZED_CONST_REFERENCE_WARNINGS and PRAGMA_ENABLE_UNINITIALIZED_CONST_REFERENCE_WARNINGS.

#rnx

diff --git a/Engine/Source/Runtime/Core/Public/Android/AndroidPlatformCompilerPreSetup.h b/Engine/Source/Runtime/Core/Public/Android/AndroidPlatformCompilerPreSetup.h
--- a/Engine/Source/Runtime/Core/Public/Android/AndroidPlatformCompilerPreSetup.h
+++ b/Engine/Source/Runtime/Core/Public/Android/AndroidPlatformCompilerPreSetup.h
@@ -11,11 +11,13 @@
 		PRAGMA_DISABLE_SHADOW_VARIABLE_WARNINGS \
 		PRAGMA_DISABLE_UNDEFINED_IDENTIFIER_WARNINGS \
 		PRAGMA_DISABLE_MISSING_VIRTUAL_DESTRUCTOR_WARNINGS \
-		PRAGMA_DISABLE_DEPRECATION_WARNINGS
+		PRAGMA_DISABLE_DEPRECATION_WARNINGS \
+        PRAGMA_DISABLE_UNINITIALIZED_CONST_REFERENCE_WARNINGS
 #endif
 
 #ifndef THIRD_PARTY_INCLUDES_END
 	#define THIRD_PARTY_INCLUDES_END \
+        PRAGMA_ENABLE_UNINITIALIZED_CONST_REFERENCE_WARNINGS \
 		PRAGMA_ENABLE_DEPRECATION_WARNINGS \
 		PRAGMA_ENABLE_MISSING_VIRTUAL_DESTRUCTOR_WARNINGS \
 		PRAGMA_ENABLE_UNDEFINED_IDENTIFIER_WARNINGS \
diff --git a/Engine/Source/Runtime/Core/Public/Apple/ApplePlatformCompilerPreSetup.h b/Engine/Source/Runtime/Core/Public/Apple/ApplePlatformCompilerPreSetup.h
--- a/Engine/Source/Runtime/Core/Public/Apple/ApplePlatformCompilerPreSetup.h
+++ b/Engine/Source/Runtime/Core/Public/Apple/ApplePlatformCompilerPreSetup.h
@@ -11,11 +11,13 @@
 		PRAGMA_DISABLE_SHADOW_VARIABLE_WARNINGS \
 		PRAGMA_DISABLE_UNDEFINED_IDENTIFIER_WARNINGS \
 		PRAGMA_DISABLE_MISSING_VIRTUAL_DESTRUCTOR_WARNINGS \
-		PRAGMA_DISABLE_DEPRECATION_WARNINGS
+		PRAGMA_DISABLE_DEPRECATION_WARNINGS \
+        PRAGMA_DISABLE_UNINITIALIZED_CONST_REFERENCE_WARNINGS
 #endif
 
 #ifndef THIRD_PARTY_INCLUDES_END
 	#define THIRD_PARTY_INCLUDES_END \
+        PRAGMA_ENABLE_UNINITIALIZED_CONST_REFERENCE_WARNINGS \
 		PRAGMA_ENABLE_DEPRECATION_WARNINGS \
 		PRAGMA_ENABLE_MISSING_VIRTUAL_DESTRUCTOR_WARNINGS \
 		PRAGMA_ENABLE_UNDEFINED_IDENTIFIER_WARNINGS \
diff --git a/Engine/Source/Runtime/Core/Public/Clang/ClangPlatformCompilerPreSetup.h b/Engine/Source/Runtime/Core/Public/Clang/ClangPlatformCompilerPreSetup.h
--- a/Engine/Source/Runtime/Core/Public/Clang/ClangPlatformCompilerPreSetup.h
+++ b/Engine/Source/Runtime/Core/Public/Clang/ClangPlatformCompilerPreSetup.h
@@ -165,6 +165,23 @@
 		_Pragma("clang diagnostic pop")
 #endif // PRAGMA_ENABLE_MACRO_REDEFINED_WARNINGS
 
+#if __has_warning("-Wuninitialized-const-reference")
+#define DISABLE_UNINITIALIZED_CONST_REFERENCE _Pragma("clang diagnostic ignored \"-Wuninitialized-const-reference\"")
+#else
+#define DISABLE_UNINITIALIZED_CONST_REFERENCE
+#endif
+
+#ifndef PRAGMA_DISABLE_UNINITIALIZED_CONST_REFERENCE_WARNINGS
+	#define PRAGMA_DISABLE_UNINITIALIZED_CONST_REFERENCE_WARNINGS \
+		_Pragma("clang diagnostic push") \
+        DISABLE_UNINITIALIZED_CONST_REFERENCE
+#endif // PRAGMA_DISABLE_UNINITIALIZED_CONST_REFERENCE_WARNINGS
+
+#ifndef PRAGMA_ENABLE_UNINITIALIZED_CONST_REFERENCE_WARNINGS
+	#define PRAGMA_ENABLE_UNINITIALIZED_CONST_REFERENCE_WARNINGS \
+		_Pragma("clang diagnostic pop")
+#endif // PRAGMA_ENABLE_UNINITIALIZED_CONST_REFERENCE_WARNINGS
+
 #ifndef PRAGMA_POP
 	#define PRAGMA_POP \
 		_Pragma("clang diagnostic pop")
diff --git a/Engine/Source/Runtime/Core/Public/HoloLens/HoloLensPlatformCompilerPreSetup.h b/Engine/Source/Runtime/Core/Public/HoloLens/HoloLensPlatformCompilerPreSetup.h
--- a/Engine/Source/Runtime/Core/Public/HoloLens/HoloLensPlatformCompilerPreSetup.h
+++ b/Engine/Source/Runtime/Core/Public/HoloLens/HoloLensPlatformCompilerPreSetup.h
@@ -94,6 +94,16 @@
 		// MSVC doesn't seem to have a warning similar to -Wunused-private-field
 #endif // PRAGMA_ENABLE_UNUSED_PRIVATE_FIELDS_WARNINGS
 
+#ifndef PRAGMA_DISABLE_UNINITIALIZED_CONST_REFERENCE_WARNINGS
+#define PRAGMA_DISABLE_UNINITIALIZED_CONST_REFERENCE_WARNINGS
+	// MSVC doesn't seem to have a warning similar to -Wuninitialized-const-reference
+#endif // PRAGMA_DISABLE_UNINITIALIZED_CONST_REFERENCE_WARNINGS
+
+#ifndef PRAGMA_ENABLE_UNINITIALIZED_CONST_REFERENCE_WARNINGS
+#define PRAGMA_ENABLE_UNINITIALIZED_CONST_REFERENCE_WARNINGS
+	// MSVC doesn't seem to have a warning similar to -Wuninitialized-const-reference
+#endif // PRAGMA_ENABLE_UNINITIALIZED_CONST_REFERENCE_WARNINGS
+
 #ifndef PRAGMA_POP
 	#define PRAGMA_POP \
 		__pragma(warning(pop))
@@ -145,11 +155,13 @@
 		PRAGMA_DISABLE_UNDEFINED_IDENTIFIER_WARNINGS \
 		PRAGMA_DISABLE_SHADOW_VARIABLE_WARNINGS \
 		PRAGMA_DISABLE_MISSING_VIRTUAL_DESTRUCTOR_WARNINGS \
-		PRAGMA_DISABLE_DEPRECATION_WARNINGS
+		PRAGMA_DISABLE_DEPRECATION_WARNINGS \
+        PRAGMA_DISABLE_UNINITIALIZED_CONST_REFERENCE_WARNINGS
 #endif // THIRD_PARTY_INCLUDES_START
 
 #ifndef THIRD_PARTY_INCLUDES_END
 	#define THIRD_PARTY_INCLUDES_END \
+        PRAGMA_ENABLE_UNINITIALIZED_CONST_REFERENCE_WARNINGS \
 		PRAGMA_ENABLE_DEPRECATION_WARNINGS \
 		PRAGMA_ENABLE_MISSING_VIRTUAL_DESTRUCTOR_WARNINGS \
 		PRAGMA_ENABLE_SHADOW_VARIABLE_WARNINGS \
diff --git a/Engine/Source/Runtime/Core/Public/Unix/UnixPlatformCompilerPreSetup.h b/Engine/Source/Runtime/Core/Public/Unix/UnixPlatformCompilerPreSetup.h
--- a/Engine/Source/Runtime/Core/Public/Unix/UnixPlatformCompilerPreSetup.h
+++ b/Engine/Source/Runtime/Core/Public/Unix/UnixPlatformCompilerPreSetup.h
@@ -41,11 +41,13 @@
 		PRAGMA_DISABLE_MISSING_VIRTUAL_DESTRUCTOR_WARNINGS \
 		PRAGMA_DISABLE_DEPRECATED_REGISTER_WARNINGS \
 		PRAGMA_DISABLE_DEPRECATION_WARNINGS \
-		PRAGMA_DISABLE_MACRO_REDEFINED_WARNINGS
+		PRAGMA_DISABLE_MACRO_REDEFINED_WARNINGS \
+        PRAGMA_DISABLE_UNINITIALIZED_CONST_REFERENCE_WARNINGS
 #endif
 
 #ifndef THIRD_PARTY_INCLUDES_END
 	#define THIRD_PARTY_INCLUDES_END \
+        PRAGMA_ENABLE_UNINITIALIZED_CONST_REFERENCE_WARNINGS \
 		PRAGMA_ENABLE_MACRO_REDEFINED_WARNINGS \
 		PRAGMA_ENABLE_DEPRECATION_WARNINGS \
 		PRAGMA_ENABLE_DEPRECATED_REGISTER_WARNINGS \
diff --git a/Engine/Source/Runtime/Core/Public/Windows/WindowsPlatformCompilerPreSetup.h b/Engine/Source/Runtime/Core/Public/Windows/WindowsPlatformCompilerPreSetup.h
--- a/Engine/Source/Runtime/Core/Public/Windows/WindowsPlatformCompilerPreSetup.h
+++ b/Engine/Source/Runtime/Core/Public/Windows/WindowsPlatformCompilerPreSetup.h
@@ -34,13 +34,15 @@
 			PRAGMA_DISABLE_UNDEFINED_IDENTIFIER_WARNINGS \
 			PRAGMA_DISABLE_DEPRECATION_WARNINGS \
 			PRAGMA_DISABLE_OVERLOADED_VIRTUAL_WARNINGS \
-			PRAGMA_DISABLE_MISSING_BRACES_WARNINGS
+			PRAGMA_DISABLE_MISSING_BRACES_WARNINGS \
+	        PRAGMA_DISABLE_UNINITIALIZED_CONST_REFERENCE_WARNINGS
 	#endif // THIRD_PARTY_INCLUDES_START
 
 	#ifndef THIRD_PARTY_INCLUDES_END
 		#define THIRD_PARTY_INCLUDES_END \
-			PRAGMA_DISABLE_MISSING_BRACES_WARNINGS \
-			PRAGMA_DISABLE_OVERLOADED_VIRTUAL_WARNINGS \
+	        PRAGMA_ENABLE_UNINITIALIZED_CONST_REFERENCE_WARNINGS \
+			PRAGMA_ENABLE_MISSING_BRACES_WARNINGS \
+			PRAGMA_ENABLE_OVERLOADED_VIRTUAL_WARNINGS \
 			PRAGMA_ENABLE_DEPRECATION_WARNINGS \
 			PRAGMA_ENABLE_UNDEFINED_IDENTIFIER_WARNINGS \
 			PRAGMA_ENABLE_SHADOW_VARIABLE_WARNINGS \
@@ -139,6 +141,16 @@
 			// MSVC doesn't seem to have a warning similar to -Wunused-private-field
 	#endif // PRAGMA_ENABLE_UNUSED_PRIVATE_FIELDS_WARNINGS
 
+	#ifndef PRAGMA_DISABLE_UNINITIALIZED_CONST_REFERENCE_WARNINGS
+	#define PRAGMA_DISABLE_UNINITIALIZED_CONST_REFERENCE_WARNINGS
+		// MSVC doesn't seem to have a warning similar to -Wuninitialized-const-reference
+	#endif // PRAGMA_DISABLE_UNINITIALIZED_CONST_REFERENCE_WARNINGS
+
+	#ifndef PRAGMA_ENABLE_UNINITIALIZED_CONST_REFERENCE_WARNINGS
+	#define PRAGMA_ENABLE_UNINITIALIZED_CONST_REFERENCE_WARNINGS
+		// MSVC doesn't seem to have a warning similar to -Wuninitialized-const-reference
+	#endif // PRAGMA_ENABLE_UNINITIALIZED_CONST_REFERENCE_WARNINGS
+
 	#ifndef PRAGMA_POP
 		#define PRAGMA_POP \
 			__pragma(warning(pop))
@@ -191,11 +203,13 @@
 			PRAGMA_DISABLE_UNDEFINED_IDENTIFIER_WARNINGS \
 			PRAGMA_DISABLE_SHADOW_VARIABLE_WARNINGS \
 			PRAGMA_DISABLE_MISSING_VIRTUAL_DESTRUCTOR_WARNINGS \
-			PRAGMA_DISABLE_DEPRECATION_WARNINGS
+			PRAGMA_DISABLE_DEPRECATION_WARNINGS \
+			PRAGMA_DISABLE_UNINITIALIZED_CONST_REFERENCE_WARNINGS
 	#endif // THIRD_PARTY_INCLUDES_START
 
 	#ifndef THIRD_PARTY_INCLUDES_END
 		#define THIRD_PARTY_INCLUDES_END \
+			PRAGMA_ENABLE_UNINITIALIZED_CONST_REFERENCE_WARNINGS \
 			PRAGMA_ENABLE_DEPRECATION_WARNINGS \
 			PRAGMA_ENABLE_MISSING_VIRTUAL_DESTRUCTOR_WARNINGS \
 			PRAGMA_ENABLE_SHADOW_VARIABLE_WARNINGS \
