Add Concurrency Viewer as an external profiler and compile only when extension is installed in VS
Avoid activating any external profiler by default as this could cause slowdown depending on which profiler gets randomly activated

diff --git a/Engine/Source/Runtime/Core/Core.Build.cs b/Engine/Source/Runtime/Core/Core.Build.cs
--- a/Engine/Source/Runtime/Core/Core.Build.cs
+++ b/Engine/Source/Runtime/Core/Core.Build.cs
@@ -63,8 +63,7 @@ public class Core : ModuleRules
 				"IntelVTune"
 				);
 
-			AddEngineThirdPartyPrivateStaticDependencies(Target,
-				"mimalloc");
+			PrivateDependencyModuleNames.Add("mimalloc");
 
 			if (Target.WindowsPlatform.bUseBundledDbgHelp)
 			{
@@ -217,6 +216,31 @@ public class Core : ModuleRules
 			}
 		}
 
+		// Detect if the Concurrency Viewer Extension is installed
+		if (Target.Platform.IsInGroup(UnrealPlatformGroup.Windows))
+		{
+			bool VSCVDefined = false;
+			string VisualStudioInstallation = Target.WindowsPlatform.IDEDir;
+			if (VisualStudioInstallation != null && VisualStudioInstallation != string.Empty && Directory.Exists(VisualStudioInstallation))
+			{
+				string SubFolderName = @"Common7\IDE\Extensions\jcn3iwfw.vp2\SDK\Native\Inc";
+				string IncludeDirectory = Path.Combine(VisualStudioInstallation, SubFolderName);
+
+				if (File.Exists(Path.Combine(IncludeDirectory, "cvmarkers.h"))
+					&& Target.Configuration != UnrealTargetConfiguration.Shipping)
+				{
+					PrivateIncludePaths.Add(IncludeDirectory);
+					PublicDefinitions.Add("WITH_CONCURRENCYVIEWER_PROFILER=1");
+					VSCVDefined = true;
+				}
+			}
+
+			if (!VSCVDefined)
+			{
+				PublicDefinitions.Add("WITH_CONCURRENCYVIEWER_PROFILER=0");
+			}
+		}
+
 		if (Target.bWithDirectXMath && Target.Platform == UnrealTargetPlatform.Win64)
 		{
 			PublicDefinitions.Add("WITH_DIRECTXMATH=1");
diff --git a/Engine/Source/Runtime/Core/Private/GenericPlatform/GenericPlatformMisc.cpp b/Engine/Source/Runtime/Core/Private/GenericPlatform/GenericPlatformMisc.cpp
--- a/Engine/Source/Runtime/Core/Private/GenericPlatform/GenericPlatformMisc.cpp
+++ b/Engine/Source/Runtime/Core/Private/GenericPlatform/GenericPlatformMisc.cpp
@@ -498,38 +498,21 @@ void FGenericPlatformMisc::RaiseException(uint32 ExceptionCode)
 void FGenericPlatformMisc::BeginNamedEvent(const struct FColor& Color, const ANSICHAR* Text)
 {
 #if UE_EXTERNAL_PROFILING_ENABLED
-	//If there's an external profiler attached, trigger its scoped event.
-	FExternalProfiler* CurrentProfiler = FActiveExternalProfilerBase::GetActiveProfiler();
-	if (CurrentProfiler != NULL)
-	{
-		CurrentProfiler->StartScopedEvent(Color, ANSI_TO_TCHAR(Text));
-	}
+	FExternalProfilerTrace::StartScopedEvent(Color, Text);
 #endif
 }
 
 void FGenericPlatformMisc::BeginNamedEvent(const struct FColor& Color, const TCHAR* Text)
 {
 #if UE_EXTERNAL_PROFILING_ENABLED
-	//If there's an external profiler attached, trigger its scoped event.
-	FExternalProfiler* CurrentProfiler = FActiveExternalProfilerBase::GetActiveProfiler();
-
-	if (CurrentProfiler != NULL)
-	{
-		CurrentProfiler->StartScopedEvent(Color, Text);
-	}
+	FExternalProfilerTrace::StartScopedEvent(Color, Text);
 #endif
 }
 
 void FGenericPlatformMisc::EndNamedEvent()
 {
 #if UE_EXTERNAL_PROFILING_ENABLED
-	//If there's an external profiler attached, trigger its scoped event.
-	FExternalProfiler* CurrentProfiler = FActiveExternalProfilerBase::GetActiveProfiler();
-
-	if (CurrentProfiler != NULL)
-	{
-		CurrentProfiler->EndScopedEvent();
-	}
+	FExternalProfilerTrace::EndScopedEvent();
 #endif
 }
 
diff --git a/Engine/Source/Runtime/Core/Private/ProfilingDebugging/ExternalProfiler.cpp b/Engine/Source/Runtime/Core/Private/ProfilingDebugging/ExternalProfiler.cpp
--- a/Engine/Source/Runtime/Core/Private/ProfilingDebugging/ExternalProfiler.cpp
+++ b/Engine/Source/Runtime/Core/Private/ProfilingDebugging/ExternalProfiler.cpp
@@ -43,17 +43,17 @@ FName FExternalProfiler::GetFeatureName()
 
 bool FActiveExternalProfilerBase::bDidInitialize = false;
 
-FExternalProfiler* FActiveExternalProfilerBase::ActiveProfiler = NULL;
+FExternalProfiler* FActiveExternalProfilerBase::ActiveProfiler = nullptr;
 
 FExternalProfiler* FActiveExternalProfilerBase::InitActiveProfiler()
 {
 	// Create profiler on demand.
-	if (ActiveProfiler == NULL && !bDidInitialize)
+	if (ActiveProfiler == nullptr && !bDidInitialize && FCommandLine::IsInitialized())
 	{
 		const FName FeatureName = FExternalProfiler::GetFeatureName();
-		TArray<FExternalProfiler*> AvailbleProfilers = IModularFeatures::Get().GetModularFeatureImplementations<FExternalProfiler>(FeatureName);
+		TArray<FExternalProfiler*> AvailableProfilers = IModularFeatures::Get().GetModularFeatureImplementations<FExternalProfiler>(FeatureName);
 
-		for (FExternalProfiler* CurProfiler : AvailbleProfilers)
+		for (FExternalProfiler* CurProfiler : AvailableProfilers)
 		{
 			check(CurProfiler != nullptr);
 
@@ -62,12 +62,6 @@ FExternalProfiler* FActiveExternalProfilerBase::InitActiveProfiler()
 			UE_LOG(LogExternalProfiler, Log, TEXT("Found external profiler: %s"), CurProfiler->GetProfilerName());
 #endif
 
-			// Default to the first profiler we have if none were specified on the command-line
-			if (ActiveProfiler == NULL)
-			{
-				ActiveProfiler = CurProfiler;
-			}
-
 			// Check to see if the profiler was specified on the command-line (e.g., "-VTune")
 			if (FParse::Param(FCommandLine::Get(), CurProfiler->GetProfilerName()))
 			{
@@ -77,7 +71,7 @@ FExternalProfiler* FActiveExternalProfilerBase::InitActiveProfiler()
 
 #if 0
 		// Logging disabled here as it can cause a stack overflow whilst flushing logs during EnginePreInit
-		if (ActiveProfiler != NULL)
+		if (ActiveProfiler != nullptr)
 		{
 			UE_LOG(LogExternalProfiler, Log, TEXT("Using external profiler: %s"), ActiveProfiler->GetProfilerName());
 		}
diff --git a/Engine/Source/Runtime/Core/Private/ProfilingDebugging/Windows/ConcurrencyViewerExternalProfiler.cpp b/Engine/Source/Runtime/Core/Private/ProfilingDebugging/Windows/ConcurrencyViewerExternalProfiler.cpp
new file mode 100644
--- /dev/null
+++ b/Engine/Source/Runtime/Core/Private/ProfilingDebugging/Windows/ConcurrencyViewerExternalProfiler.cpp
@@ -0,0 +1,155 @@
+// Copyright Epic Games, Inc. All Rights Reserved.
+
+#include "CoreTypes.h"
+#include "HAL/PlatformProcess.h"
+#include "Misc/AssertionMacros.h"
+#include "ProfilingDebugging/ExternalProfiler.h"
+#include "Features/IModularFeatures.h"
+#include "Templates/UniquePtr.h"
+#include "Containers/Map.h"
+#include "HAL/ThreadSingleton.h"
+
+#if UE_EXTERNAL_PROFILING_ENABLED && WITH_CONCURRENCYVIEWER_PROFILER
+
+#include "Windows/PreWindowsApi.h"
+THIRD_PARTY_INCLUDES_START
+#include <cvmarkers.h>
+THIRD_PARTY_INCLUDES_END
+#include "Windows/PostWindowsApi.h"
+
+/**
+ * ConcurrencyViewer implementation of FExternalProfiler
+ */
+class FConcurrencyViewerExternalProfiler : public FExternalProfiler
+{
+public:
+
+	/** Constructor */
+	FConcurrencyViewerExternalProfiler()
+	{
+		// Register as a modular feature
+		IModularFeatures::Get().RegisterModularFeature(FExternalProfiler::GetFeatureName(), this);
+	}
+
+
+	/** Destructor */
+	virtual ~FConcurrencyViewerExternalProfiler()
+	{
+		IModularFeatures::Get().UnregisterModularFeature(FExternalProfiler::GetFeatureName(), this);
+	}
+
+	virtual void FrameSync() override
+	{
+	}
+
+	/** Gets the name of this profiler as a string.  This is used to allow the user to select this profiler in a system configuration file or on the command-line */
+	virtual const TCHAR* GetProfilerName() const override
+	{
+		return TEXT("ConcurrencyViewer");
+	}
+
+
+	/** Pauses profiling. */
+	virtual void ProfilerPauseFunction() override
+	{
+	}
+
+	/** Resumes profiling. */
+	virtual void ProfilerResumeFunction() override
+	{
+	}
+
+	void StartScopedEvent(const struct FColor& Color, const TCHAR* Text) override
+	{
+		SpanStack.AddZeroed();
+		if (MarkerSeries.Num() > 0 && SpanStack.Num() <= MaxDepth)
+		{
+			CvEnterSpanW(MarkerSeries[SpanStack.Num() - 1], &SpanStack.Last(), Text);
+		}
+	}
+
+	void StartScopedEvent(const struct FColor& Color, const ANSICHAR* Text) override
+	{
+		SpanStack.AddZeroed();
+		if (MarkerSeries.Num() > 0 && SpanStack.Num() <= MaxDepth)
+		{
+			CvEnterSpanA(MarkerSeries[SpanStack.Num() - 1], &SpanStack.Last(), Text);
+		}
+	}
+
+	void EndScopedEvent() override
+	{
+		if (SpanStack.Num())
+		{
+			if (SpanStack.Last())
+			{
+				CvLeaveSpan(SpanStack.Last());
+			}
+			SpanStack.Pop(false);
+		}
+	}
+
+	virtual void SetThreadName(const TCHAR* Name) override
+	{
+		
+	}
+
+	/**
+	 * Initializes profiler hooks. It is not valid to call pause/ resume on an uninitialized
+	 * profiler and the profiler implementation is free to assert or have other undefined
+	 * behavior.
+	 *
+	 * @return true if successful, false otherwise.
+	 */
+	bool Initialize()
+	{
+		if (FAILED(CvInitProvider(&CvDefaultProviderGuid, &Provider)))
+		{
+			return false;
+		}
+		
+		MarkerSeries.SetNum(MaxDepth);
+		for (int32 Depth = 0; Depth < MaxDepth; ++Depth)
+		{
+			// We use 0 left padding so we get proper track sorting in Concurrency Viewer
+			if (FAILED(CvCreateMarkerSeries(Provider, *FString::Printf(TEXT("%02d"), Depth), &MarkerSeries[Depth])))
+			{
+				return false;
+			}
+		}
+
+		return true;
+	}
+
+
+private:
+	const int32 MaxDepth = 99;
+	// Those are immutable once initialized
+	PCV_PROVIDER Provider = nullptr;
+	TArray<PCV_MARKERSERIES> MarkerSeries;
+
+	// Modified during usage
+	static thread_local TArray<PCV_SPAN> SpanStack;
+};
+
+thread_local TArray<PCV_SPAN> FConcurrencyViewerExternalProfiler::SpanStack;
+
+namespace ConcurrencyViewerProfiler
+{
+	struct FAtModuleInit
+	{
+		FAtModuleInit()
+		{
+			static TUniquePtr<FConcurrencyViewerExternalProfiler> Profiler = MakeUnique<FConcurrencyViewerExternalProfiler>();
+			if (!Profiler->Initialize())
+			{
+				Profiler.Reset();
+			}
+		}
+	};
+
+	static FAtModuleInit AtModuleInit;
+}
+
+
+#endif	// UE_EXTERNAL_PROFILING_ENABLED && WITH_CONCURRENCYVIEWER_PROFILER
diff --git a/Engine/Source/Runtime/Core/Private/Windows/WindowsPlatformMisc.cpp b/Engine/Source/Runtime/Core/Private/Windows/WindowsPlatformMisc.cpp
--- a/Engine/Source/Runtime/Core/Private/Windows/WindowsPlatformMisc.cpp
+++ b/Engine/Source/Runtime/Core/Private/Windows/WindowsPlatformMisc.cpp
@@ -828,11 +828,7 @@ void FWindowsPlatformMisc::BeginNamedEvent(const struct FColor& Color, const TCH
 #if FRAMEPRO_ENABLED
 	FFrameProProfiler::PushEvent(Text);
 #elif UE_EXTERNAL_PROFILING_ENABLED
-	FExternalProfiler* Profiler = FActiveExternalProfilerBase::GetActiveProfiler();
-	if (Profiler)
-	{
-		Profiler->StartScopedEvent(Color, Text);
-	}
+	FExternalProfilerTrace::StartScopedEvent(Color, Text);
 #endif
 }
 
@@ -841,11 +837,7 @@ void FWindowsPlatformMisc::BeginNamedEvent(const struct FColor& Color, const ANS
 #if FRAMEPRO_ENABLED
 	FFrameProProfiler::PushEvent(Text);
 #elif UE_EXTERNAL_PROFILING_ENABLED
-	FExternalProfiler* Profiler = FActiveExternalProfilerBase::GetActiveProfiler();
-	if (Profiler)
-	{
-		Profiler->StartScopedEvent(Color, ANSI_TO_TCHAR(Text));
-	}
+	FExternalProfilerTrace::StartScopedEvent(Color, Text);
 #endif
 }
 
@@ -854,11 +846,7 @@ void FWindowsPlatformMisc::EndNamedEvent()
 #if FRAMEPRO_ENABLED
 	FFrameProProfiler::PopEvent();
 #elif UE_EXTERNAL_PROFILING_ENABLED
-	FExternalProfiler* Profiler = FActiveExternalProfilerBase::GetActiveProfiler();
-	if (Profiler)
-	{
-		Profiler->EndScopedEvent();
-	}
+	FExternalProfilerTrace::EndScopedEvent();
 #endif
 }
 #endif // STATS || ENABLE_STATNAMEDEVENTS
diff --git a/Engine/Source/Runtime/Core/Public/ProfilingDebugging/ExternalProfiler.h b/Engine/Source/Runtime/Core/Public/ProfilingDebugging/ExternalProfiler.h
--- a/Engine/Source/Runtime/Core/Public/ProfilingDebugging/ExternalProfiler.h
+++ b/Engine/Source/Runtime/Core/Public/ProfilingDebugging/ExternalProfiler.h
@@ -62,6 +62,9 @@ public:
 	/** Starts a scoped event specific to the profiler. */
 	virtual void StartScopedEvent(const struct FColor& Color, const TCHAR* Text) {};
 
+	/** Starts a scoped event specific to the profiler. Default implementation for backward compabitility. */
+	virtual void StartScopedEvent(const struct FColor& Color, const ANSICHAR* Text) { StartScopedEvent(Color, ANSI_TO_TCHAR(Text)); };
+
 	/** Ends a scoped event specific to the profiler. */
 	virtual void EndScopedEvent() {};
 
@@ -167,6 +170,40 @@ public:
 
 };
 
+class FExternalProfilerTrace
+{
+public:
+	/** Starts a scoped event specific to the profiler. */
+	FORCEINLINE static void StartScopedEvent(const struct FColor& Color, const TCHAR* Text)
+	{
+		FExternalProfiler* Profiler = FActiveExternalProfilerBase::GetActiveProfiler();
+		if (Profiler)
+		{
+			Profiler->StartScopedEvent(Color, Text);
+		}
+	}
+
+	/** Starts a scoped event specific to the profiler. */
+	FORCEINLINE static void StartScopedEvent(const struct FColor& Color, const ANSICHAR* Text)
+	{
+		FExternalProfiler* Profiler = FActiveExternalProfilerBase::GetActiveProfiler();
+		if (Profiler)
+		{
+			Profiler->StartScopedEvent(Color, Text);
+		}
+	}
+
+	/** Ends a scoped event specific to the profiler. */
+	FORCEINLINE static void EndScopedEvent()
+	{
+		FExternalProfiler* Profiler = FActiveExternalProfilerBase::GetActiveProfiler();
+		if (Profiler)
+		{
+			Profiler->EndScopedEvent();
+		}
+	}
+};
+
 #define SCOPE_PROFILER_INCLUDER(X) FExternalProfilerIncluder ExternalProfilerIncluder_##X;
 #define SCOPE_PROFILER_EXCLUDER(X) FExternalProfilerExcluder ExternalProfilerExcluder_##X;
 
