Implement garbage collection for mimalloc

diff --git a/Engine/Source/Runtime/Core/Private/HAL/MallocMimalloc.cpp b/Engine/Source/Runtime/Core/Private/HAL/MallocMimalloc.cpp
--- a/Engine/Source/Runtime/Core/Private/HAL/MallocMimalloc.cpp
+++ b/Engine/Source/Runtime/Core/Private/HAL/MallocMimalloc.cpp
@@ -159,6 +159,7 @@ bool FMallocMimalloc::GetAllocationSize(void *Original, SIZE_T &SizeOut)
 
 void FMallocMimalloc::Trim(bool bTrimThreadCaches)
 {
+	mi_collect(bTrimThreadCaches);
 }
 
 #endif
