FName: Respect FindType on character range constructors

Reported by [REDACTED_PERSON] at 2K Cloud Chamber

diff --git a/Engine/Source/Runtime/Core/Public/UObject/NameTypes.h b/Engine/Source/Runtime/Core/Public/UObject/NameTypes.h
--- a/Engine/Source/Runtime/Core/Public/UObject/NameTypes.h
+++ b/Engine/Source/Runtime/Core/Public/UObject/NameTypes.h
@@ -791,7 +791,7 @@ public:
 		: FName(NoInit)
 	{
 		TStringView<CharType> View = Forward<CharRangeType>(Name);
-		*this = FName(View.Len(), View.GetData());
+		*this = FName(View.Len(), View.GetData(), FindType);
 	}
 
 	/**
@@ -815,7 +815,7 @@ public:
 		: FName(NoInit)
 	{
 		TStringView<CharType> View = Forward<CharRangeType>(Name);
-		*this = FName(View.Len(), View.GetData(), InNumber);
+		*this = FName(View.Len(), View.GetData(), InNumber, FindType);
 	}
 
 	/**
