Fix unity file compilation when concurrency viewer extension is installed

diff --git a/Engine/Source/Runtime/Core/Private/ProfilingDebugging/Windows/ConcurrencyViewerExternalProfiler.cpp b/Engine/Source/Runtime/Core/Private/ProfilingDebugging/Windows/ConcurrencyViewerExternalProfiler.cpp
--- a/Engine/Source/Runtime/Core/Private/ProfilingDebugging/Windows/ConcurrencyViewerExternalProfiler.cpp
+++ b/Engine/Source/Runtime/Core/Private/ProfilingDebugging/Windows/ConcurrencyViewerExternalProfiler.cpp
@@ -11,11 +11,13 @@
 
 #if UE_EXTERNAL_PROFILING_ENABLED && WITH_CONCURRENCYVIEWER_PROFILER
 
+#include "Windows/AllowWindowsPlatformTypes.h"
+#include "Windows/AllowWindowsPlatformAtomics.h"
 #include "Windows/PreWindowsApi.h"
-THIRD_PARTY_INCLUDES_START
 #include <cvmarkers.h>
-THIRD_PARTY_INCLUDES_END
 #include "Windows/PostWindowsApi.h"
+#include "Windows/HideWindowsPlatformAtomics.h"
+#include "Windows/HideWindowsPlatformTypes.h"
 
 /**
  * ConcurrencyViewer implementation of FExternalProfiler
