UE5EA - Fix static analysis

diff --git a/Engine/Source/Runtime/Core/Private/Misc/MonitoredProcess.cpp b/Engine/Source/Runtime/Core/Private/Misc/MonitoredProcess.cpp
--- a/Engine/Source/Runtime/Core/Private/Misc/MonitoredProcess.cpp
+++ b/Engine/Source/Runtime/Core/Private/Misc/MonitoredProcess.cpp
@@ -208,7 +208,7 @@ void FSerializedUATProcess::CancelQueue()
 		{
 			FSerializedUATProcess* Delete = Travel;
 			Travel = Travel->NextProcessToRun;
-			delete Travel;
+			delete Delete;
 		}
 
 
