Initial (prototype) version of templated core TVector types, to allow for conversion of GeometryProcessing FVector3f/FVector3d
#rnx

diff --git a/Engine/Source/Runtime/Core/Private/Math/VectorLWC.cpp b/Engine/Source/Runtime/Core/Private/Math/VectorLWC.cpp
new file mode 100644
--- /dev/null
+++ b/Engine/Source/Runtime/Core/Private/Math/VectorLWC.cpp
@@ -0,0 +1,151 @@
+// Copyright Epic Games, Inc. All Rights Reserved.
+
+// NOTE: This is a temporary place holder representing the large world coordinate version of FVector, and will be replaced shortly. You SHOULD NOT be using these types in your code!
+
+#include "Math/VectorLWC.h"
+#include "CoreMinimal.h"
+
+CORE_API const FVector3f FVector3f::ZeroVector(0, 0, 0);
+CORE_API const FVector3f FVector3f::OneVector(1, 1, 1);
+CORE_API const FVector3f FVector3f::UpVector(0, 0, 1);
+CORE_API const FVector3f FVector3f::DownVector(0, 0, -1);
+CORE_API const FVector3f FVector3f::ForwardVector(1, 0, 0);
+CORE_API const FVector3f FVector3f::BackwardVector(-1, 0, 0);
+CORE_API const FVector3f FVector3f::RightVector(0, 1, 0);
+CORE_API const FVector3f FVector3f::LeftVector(0, -1, 0);
+CORE_API const FVector3f FVector3f::XAxisVector(1, 0, 0);
+CORE_API const FVector3f FVector3f::YAxisVector(0, 1, 0);
+CORE_API const FVector3f FVector3f::ZAxisVector(0, 0, 1);
+
+CORE_API const FVector3d FVector3d::ZeroVector(0, 0, 0);
+CORE_API const FVector3d FVector3d::OneVector(1, 1, 1);
+CORE_API const FVector3d FVector3d::UpVector(0, 0, 1);
+CORE_API const FVector3d FVector3d::DownVector(0, 0, -1);
+CORE_API const FVector3d FVector3d::ForwardVector(1, 0, 0);
+CORE_API const FVector3d FVector3d::BackwardVector(-1, 0, 0);
+CORE_API const FVector3d FVector3d::RightVector(0, 1, 0);
+CORE_API const FVector3d FVector3d::LeftVector(0, -1, 0);
+CORE_API const FVector3d FVector3d::XAxisVector(1, 0, 0);
+CORE_API const FVector3d FVector3d::YAxisVector(0, 1, 0);
+CORE_API const FVector3d FVector3d::ZAxisVector(0, 0, 1);
+
+template<>
+UE::Core::TVector<float>::TVector(const FVector4& V)
+	: X(V.X), Y(V.Y), Z(V.Z)
+{
+	DiagnosticCheckNaN();
+}
+
+template<>
+UE::Core::TVector<double>::TVector(const FVector4& V)
+	: X(V.X), Y(V.Y), Z(V.Z)
+{
+	DiagnosticCheckNaN();
+}
+
+template<>
+FQuat UE::Core::TVector<float>::ToOrientationQuat() const
+{
+	// Essentially an optimized Vector->Rotator->Quat made possible by knowing Roll == 0, and avoiding radians->degrees->radians.
+	// This is done to avoid adding any roll (which our API states as a constraint).
+	const float YawRad = FMath::Atan2(Y, X);
+	const float PitchRad = FMath::Atan2(Z, FMath::Sqrt(X * X + Y * Y));
+
+	const float DIVIDE_BY_2 = 0.5f;
+	float SP, SY;
+	float CP, CY;
+
+	FMath::SinCos(&SP, &CP, PitchRad * DIVIDE_BY_2);
+	FMath::SinCos(&SY, &CY, YawRad * DIVIDE_BY_2);
+
+	FQuat RotationQuat;
+	RotationQuat.X = SP * SY;
+	RotationQuat.Y = -SP * CY;
+	RotationQuat.Z = CP * SY;
+	RotationQuat.W = CP * CY;
+	return RotationQuat;
+}
+
+template<>
+FQuat UE::Core::TVector<double>::ToOrientationQuat() const
+{
+	// NOTE: Precision issues here. Identical to the float version, due to missing double support for FMath::SinCos.
+	const float YawRad = (float)FMath::Atan2(Y, X);
+	const float PitchRad = (float)FMath::Atan2(Z, FMath::Sqrt(X * X + Y * Y));
+
+	const float DIVIDE_BY_2 = 0.5f;
+	float SP, SY;
+	float CP, CY;
+
+	FMath::SinCos(&SP, &CP, PitchRad * DIVIDE_BY_2);
+	FMath::SinCos(&SY, &CY, YawRad * DIVIDE_BY_2);
+
+	FQuat RotationQuat;
+	RotationQuat.X = SP * SY;
+	RotationQuat.Y = -SP * CY;
+	RotationQuat.Z = CP * SY;
+	RotationQuat.W = CP * CY;
+	return RotationQuat;
+}
+
+template<>
+FRotator UE::Core::TVector<float>::ToOrientationRotator() const
+{
+	FRotator R;
+
+	// Find yaw.
+	R.Yaw = FMath::Atan2(Y, X) * (180.f / PI);
+
+	// Find pitch.
+	R.Pitch = FMath::Atan2(Z, FMath::Sqrt(X * X + Y * Y)) * (180.f / PI);
+
+	// Find roll.
+	R.Roll = 0;
+
+#if ENABLE_NAN_DIAGNOSTIC || (DO_CHECK && !UE_BUILD_SHIPPING)
+	if (R.ContainsNaN())
+	{
+		logOrEnsureNanError(TEXT("FVector::Rotation(): Rotator result %s contains NaN! Input FVector = %s"), *R.ToString(), *this->ToString());
+		R = FRotator::ZeroRotator;
+	}
+#endif
+
+	return R;
+}
+
+template<>
+FRotator UE::Core::TVector<double>::ToOrientationRotator() const
+{
+	FRotator R;
+
+	// Find yaw.
+	R.Yaw = static_cast<float>(FMath::Atan2(Y, X) * (180.0 / PI));
+
+	// Find pitch.
+	R.Pitch = static_cast<float>(FMath::Atan2(Z, FMath::Sqrt(X * X + Y * Y)) * (180.0 / PI));
+
+	// Find roll.
+	R.Roll = 0;
+
+#if ENABLE_NAN_DIAGNOSTIC || (DO_CHECK && !UE_BUILD_SHIPPING)
+	if (R.ContainsNaN())
+	{
+		logOrEnsureNanError(TEXT("FVector::Rotation(): Rotator result %s contains NaN! Input FVector = %s"), *R.ToString(), *this->ToString());
+		R = FRotator::ZeroRotator;
+	}
+#endif
+
+	return R;
+}
+
+template<>
+FRotator UE::Core::TVector<float>::Rotation() const
+{
+	return ToOrientationRotator();
+}
+
+template<>
+FRotator UE::Core::TVector<double>::Rotation() const
+{
+	return ToOrientationRotator();
+}
diff --git a/Engine/Source/Runtime/Core/Public/Math/VectorLWC.h b/Engine/Source/Runtime/Core/Public/Math/VectorLWC.h
new file mode 100644
--- /dev/null
+++ b/Engine/Source/Runtime/Core/Public/Math/VectorLWC.h
@@ -0,0 +1,2374 @@
+// Copyright Epic Games, Inc. All Rights Reserved.
+
+/**
+ * DO NOT USE!
+ * DO NOT USE!
+ * DO NOT USE!
+ *
+ * This file is intended as a placeholder for the Large World Coordinate vector refactors, and may be subject to significant refactoring.
+ *
+ * DO NOT USE!
+ * DO NOT USE!
+ * DO NOT USE!
+ */
+
+#pragma once
+
+#include "CoreTypes.h"
+#include "Misc/AssertionMacros.h"
+#include "Math/NumericLimits.h"
+#include "Misc/Crc.h"
+#include "Math/UnrealMathUtility.h"
+#include "Containers/UnrealString.h"
+#include "Misc/Parse.h"
+#include "Math/Color.h"
+#include "Math/IntPoint.h"
+#include "Logging/LogMacros.h"
+#include "Math/Vector2D.h"
+#include "Misc/ByteSwap.h"
+#include "Internationalization/Text.h"
+#include "Internationalization/Internationalization.h"
+#include "Math/IntVector.h"
+#include "Math/Axis.h"
+#include "Serialization/MemoryLayout.h"
+#include "UObject/ObjectVersion.h"
+#include <type_traits>
+
+#ifdef _MSC_VER
+#pragma warning (push)
+// Ensure template functions don't generate shadowing warnings against global variables at the point of instantiation.
+#pragma warning (disable : 4459)
+#endif
+
+/**
+ * A vector in 3-D space composed of components (X, Y, Z) with type defined precision.
+ */
+
+ // Move out of global namespace to avoid collisions with Chaos::TVector within the physics code.
+namespace UE
+{
+namespace Core
+{
+
+template<typename TReal>
+struct TVector
+{
+public:
+	using TComponent = TReal;
+
+	/** Vector's X component. */
+	TReal X;
+
+	/** Vector's Y component. */
+	TReal Y;
+
+	/** Vector's Z component. */
+	TReal Z;
+
+public:
+
+#if ENABLE_NAN_DIAGNOSTIC
+	FORCEINLINE void DiagnosticCheckNaN() const
+	{
+		if (ContainsNaN())
+		{
+			logOrEnsureNanError(TEXT("FVector contains NaN: %s"), *ToString());
+			*const_cast<TVector<TReal>*>(this) = TVector<TReal>(0, 0, 0);
+		}
+	}
+
+	FORCEINLINE void DiagnosticCheckNaN(const TCHAR* Message) const
+	{
+		if (ContainsNaN())
+		{
+			logOrEnsureNanError(TEXT("%s: FVector contains NaN: %s"), Message, *ToString());
+			*const_cast<TVector<TReal>*>(this) = TVector<TReal>(0, 0, 0);
+		}
+	}
+#else
+	FORCEINLINE void DiagnosticCheckNaN() const {}
+	FORCEINLINE void DiagnosticCheckNaN(const TCHAR* Message) const {}
+#endif
+
+	/** Default constructor (no initialization). */
+	FORCEINLINE TVector();
+
+	/**
+	 * Constructor initializing all components to a single TReal value.
+	 *
+	 * @param InF Value to set all components to.
+	 */
+	explicit FORCEINLINE TVector(TReal InF);
+
+	/**
+	 * Constructor using initial values for each component.
+	 *
+	 * @param InX X Coordinate.
+	 * @param InY Y Coordinate.
+	 * @param InZ Z Coordinate.
+	 */
+	FORCEINLINE TVector(TReal InX, TReal InY, TReal InZ);
+
+	/**
+	 * Constructs a vector from an FVector2D and Z value.
+	 *
+	 * @param V Vector to copy from.
+	 * @param InZ Z Coordinate.
+	 */
+	explicit FORCEINLINE TVector(const FVector2D V, TReal InZ);
+
+	/**
+	 * Constructor using the XYZ components from a 4D vector.
+	 *
+	 * @param V 4D Vector to copy from.
+	 */
+	CORE_API TVector(const FVector4& V);
+
+	/**
+	 * Constructs a vector from an FLinearColor.
+	 *
+	 * @param InColor Color to copy from.
+	 */
+	explicit TVector(const FLinearColor& InColor);
+
+	/**
+	 * Constructs a vector from an FIntVector.
+	 *
+	 * @param InVector FIntVector to copy from.
+	 */
+	explicit TVector(FIntVector InVector);
+
+	/**
+	 * Constructs a vector from an FIntPoint.
+	 *
+	 * @param A Int Point used to set X and Y coordinates, Z is set to zero.
+	 */
+	explicit TVector(FIntPoint A);
+
+	/**
+	 * Constructor which initializes all components to zero.
+	 *
+	 * @param EForceInit Force init enum
+	 */
+	explicit FORCEINLINE TVector(EForceInit);
+
+#ifdef IMPLEMENT_ASSIGNMENT_OPERATOR_MANUALLY
+	/**
+	* Copy another TVector<TReal> into this one
+	*
+	* @param Other The other vector.
+	* @return Reference to vector after copy.
+	*/
+	FORCEINLINE TVector<TReal>& operator=(const TVector<TReal>& Other);
+#endif
+
+	/**
+	 * Calculate cross product between this and another vector.
+	 *
+	 * @param V The other vector.
+	 * @return The cross product.
+	 */
+	FORCEINLINE TVector<TReal> operator^(const TVector<TReal>& V) const;
+
+	/**
+	 * Calculate the cross product of two vectors.
+	 *
+	 * @param A The first vector.
+	 * @param B The second vector.
+	 * @return The cross product.
+	 */
+	FORCEINLINE static TVector<TReal> CrossProduct(const TVector<TReal>& A, const TVector<TReal>& B);
+
+	/**
+	 * Calculate the dot product between this and another vector.
+	 *
+	 * @param V The other vector.
+	 * @return The dot product.
+	 */
+	FORCEINLINE TReal operator|(const TVector<TReal>& V) const;
+
+	/**
+	 * Calculate the dot product of two vectors.
+	 *
+	 * @param A The first vector.
+	 * @param B The second vector.
+	 * @return The dot product.
+	 */
+	FORCEINLINE static TReal DotProduct(const TVector<TReal>& A, const TVector<TReal>& B);
+
+	/**
+	 * Gets the result of component-wise addition of this and another vector.
+	 *
+	 * @param V The vector to add to this.
+	 * @return The result of vector addition.
+	 */
+	FORCEINLINE TVector<TReal> operator+(const TVector<TReal>& V) const;
+
+	/**
+	 * Gets the result of component-wise subtraction of this by another vector.
+	 *
+	 * @param V The vector to subtract from this.
+	 * @return The result of vector subtraction.
+	 */
+	FORCEINLINE TVector<TReal> operator-(const TVector<TReal>& V) const;
+
+	/**
+	 * Gets the result of subtracting from each component of the vector.
+	 *
+	 * @param Bias How much to subtract from each component.
+	 * @return The result of subtraction.
+	 */
+	template<typename TReal2, TEMPLATE_REQUIRES(std::is_convertible<TReal2, float>::value)>
+	FORCEINLINE TVector<TReal> operator-(TReal2 Bias) const
+	{
+		return TVector<TReal>(X - (TReal)Bias, Y - (TReal)Bias, Z - (TReal)Bias);
+	}
+
+	/**
+	 * Gets the result of adding to each component of the vector.
+	 *
+	 * @param Bias How much to add to each component.
+	 * @return The result of addition.
+	 */
+	template<typename TReal2, TEMPLATE_REQUIRES(std::is_convertible<TReal2, float>::value)>
+	FORCEINLINE TVector<TReal> operator+(TReal2 Bias) const
+	{
+		return TVector<TReal>(X + (TReal)Bias, Y + (TReal)Bias, Z + (TReal)Bias);
+	}
+
+	/**
+	 * Gets the result of scaling the vector (multiplying each component by a value).
+	 *
+	 * @param Scale What to multiply each component by.
+	 * @return The result of multiplication.
+	 */
+	template<typename TReal2, TEMPLATE_REQUIRES(std::is_convertible<TReal2, float>::value)>
+	FORCEINLINE TVector<TReal> operator*(TReal2 Scale) const
+	{
+		return TVector<TReal>(X * (TReal)Scale, Y * (TReal)Scale, Z * (TReal)Scale);
+	}
+
+	/**
+	 * Gets the result of dividing each component of the vector by a value.
+	 *
+	 * @param Scale What to divide each component by.
+	 * @return The result of division.
+	 */
+	template<typename TReal2, TEMPLATE_REQUIRES(std::is_convertible<TReal2, float>::value)>
+	TVector<TReal> operator/(TReal2 Scale) const
+	{
+		const TReal RScale = (TReal)1 / (TReal)Scale;
+		return TVector<TReal>(X * RScale, Y * RScale, Z * RScale);
+	}
+
+	/**
+	 * Gets the result of component-wise multiplication of this vector by another.
+	 *
+	 * @param V The vector to multiply with.
+	 * @return The result of multiplication.
+	 */
+	FORCEINLINE TVector<TReal> operator*(const TVector<TReal>& V) const;
+
+	/**
+	 * Gets the result of component-wise division of this vector by another.
+	 *
+	 * @param V The vector to divide by.
+	 * @return The result of division.
+	 */
+	FORCEINLINE TVector<TReal> operator/(const TVector<TReal>& V) const;
+
+	// Binary comparison operators.
+
+	/**
+	 * Check against another vector for equality.
+	 *
+	 * @param V The vector to check against.
+	 * @return true if the vectors are equal, false otherwise.
+	 */
+	bool operator==(const TVector<TReal>& V) const;
+
+	/**
+	 * Check against another vector for inequality.
+	 *
+	 * @param V The vector to check against.
+	 * @return true if the vectors are not equal, false otherwise.
+	 */
+	bool operator!=(const TVector<TReal>& V) const;
+
+	/**
+	 * Check against another vector for equality, within specified error limits.
+	 *
+	 * @param V The vector to check against.
+	 * @param Tolerance Error tolerance.
+	 * @return true if the vectors are equal within tolerance limits, false otherwise.
+	 */
+	bool Equals(const TVector<TReal>& V, TReal Tolerance = KINDA_SMALL_NUMBER) const;
+
+	/**
+	 * Checks whether all components of this vector are the same, within a tolerance.
+	 *
+	 * @param Tolerance Error tolerance.
+	 * @return true if the vectors are equal within tolerance limits, false otherwise.
+	 */
+	bool AllComponentsEqual(TReal Tolerance = KINDA_SMALL_NUMBER) const;
+
+	/**
+	 * Get a negated copy of the vector.
+	 *
+	 * @return A negated copy of the vector.
+	 */
+	FORCEINLINE TVector<TReal> operator-() const;
+
+	/**
+	 * Adds another vector to this.
+	 * Uses component-wise addition.
+	 *
+	 * @param V Vector to add to this.
+	 * @return Copy of the vector after addition.
+	 */
+	FORCEINLINE TVector<TReal> operator+=(const TVector<TReal>& V);
+
+	/**
+	 * Subtracts another vector from this.
+	 * Uses component-wise subtraction.
+	 *
+	 * @param V Vector to subtract from this.
+	 * @return Copy of the vector after subtraction.
+	 */
+	FORCEINLINE TVector<TReal> operator-=(const TVector<TReal>& V);
+
+	/**
+	 * Scales the vector.
+	 *
+	 * @param Scale Amount to scale this vector by.
+	 * @return Copy of the vector after scaling.
+	 */
+	template<typename TReal2, TEMPLATE_REQUIRES(std::is_convertible<TReal2, float>::value)>
+	FORCEINLINE TVector<TReal> operator*=(TReal2 Scale)
+	{
+		X *= (TReal)Scale; Y *= (TReal)Scale; Z *= (TReal)Scale;
+		DiagnosticCheckNaN();
+		return *this;
+	}
+
+	/**
+	 * Divides the vector by a number.
+	 *
+	 * @param V What to divide this vector by.
+	 * @return Copy of the vector after division.
+	 */
+	template<typename TReal2, TEMPLATE_REQUIRES(std::is_convertible<TReal2, float>::value)>
+	TVector<TReal> operator/=(TReal2 Scale)
+	{
+		const TReal RV = (TReal)1 / (TReal)Scale;
+		X *= RV; Y *= RV; Z *= RV;
+		DiagnosticCheckNaN();
+		return *this;
+	}
+
+	/**
+	 * Multiplies the vector with another vector, using component-wise multiplication.
+	 *
+	 * @param V What to multiply this vector with.
+	 * @return Copy of the vector after multiplication.
+	 */
+	TVector<TReal> operator*=(const TVector<TReal>& V);
+
+	/**
+	 * Divides the vector by another vector, using component-wise division.
+	 *
+	 * @param V What to divide vector by.
+	 * @return Copy of the vector after division.
+	 */
+	TVector<TReal> operator/=(const TVector<TReal>& V);
+
+	/**
+	 * Gets specific component of the vector.
+	 *
+	 * @param Index the index of vector component
+	 * @return reference to component.
+	 */
+	TReal& operator[](int32 Index);
+
+	/**
+	 * Gets specific component of the vector.
+	 *
+	 * @param Index the index of vector component
+	 * @return Copy of the component.
+	 */
+	TReal operator[](int32 Index)const;
+
+	/**
+	* Gets a specific component of the vector.
+	*
+	* @param Index The index of the component required.
+	*
+	* @return Reference to the specified component.
+	*/
+	TReal& Component(int32 Index);
+
+	/**
+	* Gets a specific component of the vector.
+	*
+	* @param Index The index of the component required.
+	* @return Copy of the specified component.
+	*/
+	TReal Component(int32 Index) const;
+
+
+	/** Get a specific component of the vector, given a specific axis by enum */
+	TReal GetComponentForAxis(EAxis::Type Axis) const;
+
+	/** Set a specified componet of the vector, given a specific axis by enum */
+	void SetComponentForAxis(EAxis::Type Axis, TReal Component);
+
+public:
+
+	// Simple functions.
+
+	/**
+	 * Set the values of the vector directly.
+	 *
+	 * @param InX New X coordinate.
+	 * @param InY New Y coordinate.
+	 * @param InZ New Z coordinate.
+	 */
+	void Set(TReal InX, TReal InY, TReal InZ);
+
+	/**
+	 * Get the maximum value of the vector's components.
+	 *
+	 * @return The maximum value of the vector's components.
+	 */
+	TReal GetMax() const;
+
+	/**
+	 * Get the maximum absolute value of the vector's components.
+	 *
+	 * @return The maximum absolute value of the vector's components.
+	 */
+	TReal GetAbsMax() const;
+
+	/**
+	 * Get the minimum value of the vector's components.
+	 *
+	 * @return The minimum value of the vector's components.
+	 */
+	TReal GetMin() const;
+
+	/**
+	 * Get the minimum absolute value of the vector's components.
+	 *
+	 * @return The minimum absolute value of the vector's components.
+	 */
+	TReal GetAbsMin() const;
+
+	/** Gets the component-wise min of two vectors. */
+	TVector<TReal> ComponentMin(const TVector<TReal>& Other) const;
+
+	/** Gets the component-wise max of two vectors. */
+	TVector<TReal> ComponentMax(const TVector<TReal>& Other) const;
+
+	/**
+	 * Get a copy of this vector with absolute value of each component.
+	 *
+	 * @return A copy of this vector with absolute value of each component.
+	 */
+	TVector<TReal> GetAbs() const;
+
+	/**
+	 * Get the length (magnitude) of this vector.
+	 *
+	 * @return The length of this vector.
+	 */
+	TReal Size() const;
+
+	/**
+	 * Get the squared length of this vector.
+	 *
+	 * @return The squared length of this vector.
+	 */
+	TReal SizeSquared() const;
+
+	/**
+	 * Get the length of the 2D components of this vector.
+	 *
+	 * @return The 2D length of this vector.
+	 */
+	TReal Size2D() const;
+
+	/**
+	 * Get the squared length of the 2D components of this vector.
+	 *
+	 * @return The squared 2D length of this vector.
+	 */
+	TReal SizeSquared2D() const;
+
+	/**
+	 * Checks whether vector is near to zero within a specified tolerance.
+	 *
+	 * @param Tolerance Error tolerance.
+	 * @return true if the vector is near to zero, false otherwise.
+	 */
+	bool IsNearlyZero(TReal Tolerance = KINDA_SMALL_NUMBER) const;
+
+	/**
+	 * Checks whether all components of the vector are exactly zero.
+	 *
+	 * @return true if the vector is exactly zero, false otherwise.
+	 */
+	bool IsZero() const;
+
+	/**
+	 * Check if the vector is of unit length, with specified tolerance.
+	 *
+	 * @param LengthSquaredTolerance Tolerance against squared length.
+	 * @return true if the vector is a unit vector within the specified tolerance.
+	 */
+	FORCEINLINE bool IsUnit(TReal LengthSquaredTolerance = KINDA_SMALL_NUMBER) const;
+
+	/**
+	 * Checks whether vector is normalized.
+	 *
+	 * @return true if normalized, false otherwise.
+	 */
+	bool IsNormalized() const;
+
+	/**
+	 * Normalize this vector in-place if it is larger than a given tolerance. Leaves it unchanged if not.
+	 *
+	 * @param Tolerance Minimum squared length of vector for normalization.
+	 * @return true if the vector was normalized correctly, false otherwise.
+	 */
+	bool Normalize(TReal Tolerance = SMALL_NUMBER);
+
+	/**
+	 * Calculates normalized version of vector without checking for zero length.
+	 *
+	 * @return Normalized version of vector.
+	 * @see GetSafeNormal()
+	 */
+	FORCEINLINE TVector<TReal> GetUnsafeNormal() const;
+
+	/**
+	 * Gets a normalized copy of the vector, checking it is safe to do so based on the length.
+	 * Returns zero vector if vector length is too small to safely normalize.
+	 *
+	 * @param Tolerance Minimum squared vector length.
+	 * @return A normalized copy if safe, (0,0,0) otherwise.
+	 */
+	TVector<TReal> GetSafeNormal(TReal Tolerance = SMALL_NUMBER) const;
+
+	/**
+	 * Gets a normalized copy of the 2D components of the vector, checking it is safe to do so. Z is set to zero.
+	 * Returns zero vector if vector length is too small to normalize.
+	 *
+	 * @param Tolerance Minimum squared vector length.
+	 * @return Normalized copy if safe, otherwise returns zero vector.
+	 */
+	TVector<TReal> GetSafeNormal2D(TReal Tolerance = SMALL_NUMBER) const;
+
+	/**
+	 * Util to convert this vector into a unit direction vector and its original length.
+	 *
+	 * @param OutDir Reference passed in to store unit direction vector.
+	 * @param OutLength Reference passed in to store length of the vector.
+	 */
+	void ToDirectionAndLength(TVector<TReal>& OutDir, double& OutLength) const;
+	void ToDirectionAndLength(TVector<TReal>& OutDir, float& OutLength) const;
+
+	/**
+	 * Get a copy of the vector as sign only.
+	 * Each component is set to +1 or -1, with the sign of zero treated as +1.
+	 *
+	 * @param A copy of the vector with each component set to +1 or -1
+	 */
+	FORCEINLINE TVector<TReal> GetSignVector() const;
+
+	/**
+	 * Projects 2D components of vector based on Z.
+	 *
+	 * @return Projected version of vector based on Z.
+	 */
+	TVector<TReal> Projection() const;
+
+	/**
+	* Calculates normalized 2D version of vector without checking for zero length.
+	*
+	* @return Normalized version of vector.
+	* @see GetSafeNormal2D()
+	*/
+	FORCEINLINE TVector<TReal> GetUnsafeNormal2D() const;
+
+	/**
+	 * Gets a copy of this vector snapped to a grid.
+	 *
+	 * @param GridSz Grid dimension.
+	 * @return A copy of this vector snapped to a grid.
+	 * @see FMath::GridSnap()
+	 */
+	TVector<TReal> GridSnap(const TReal& GridSz) const;
+
+	/**
+	 * Get a copy of this vector, clamped inside of a cube.
+	 *
+	 * @param Radius Half size of the cube.
+	 * @return A copy of this vector, bound by cube.
+	 */
+	TVector<TReal> BoundToCube(TReal Radius) const;
+
+	/** Get a copy of this vector, clamped inside of a cube. */
+	TVector<TReal> BoundToBox(const TVector<TReal>& Min, const TVector<TReal> Max) const;
+
+	/** Create a copy of this vector, with its magnitude clamped between Min and Max. */
+	TVector<TReal> GetClampedToSize(TReal Min, TReal Max) const;
+
+	/** Create a copy of this vector, with the 2D magnitude clamped between Min and Max. Z is unchanged. */
+	TVector<TReal> GetClampedToSize2D(TReal Min, TReal Max) const;
+
+	/** Create a copy of this vector, with its maximum magnitude clamped to MaxSize. */
+	TVector<TReal> GetClampedToMaxSize(TReal MaxSize) const;
+
+	/** Create a copy of this vector, with the maximum 2D magnitude clamped to MaxSize. Z is unchanged. */
+	TVector<TReal> GetClampedToMaxSize2D(TReal MaxSize) const;
+
+	/**
+	 * Add a vector to this and clamp the result in a cube.
+	 *
+	 * @param V Vector to add.
+	 * @param Radius Half size of the cube.
+	 */
+	void AddBounded(const TVector<TReal>& V, TReal Radius = MAX_int16);
+
+	/**
+	 * Gets the reciprocal of this vector, avoiding division by zero.
+	 * Zero components are set to BIG_NUMBER.
+	 *
+	 * @return Reciprocal of this vector.
+	 */
+	TVector<TReal> Reciprocal() const;
+
+	/**
+	 * Check whether X, Y and Z are nearly equal.
+	 *
+	 * @param Tolerance Specified Tolerance.
+	 * @return true if X == Y == Z within the specified tolerance.
+	 */
+	bool IsUniform(TReal Tolerance = KINDA_SMALL_NUMBER) const;
+
+	/**
+	 * Mirror a vector about a normal vector.
+	 *
+	 * @param MirrorNormal Normal vector to mirror about.
+	 * @return Mirrored vector.
+	 */
+	TVector<TReal> MirrorByVector(const TVector<TReal>& MirrorNormal) const;
+
+	/**
+	 * Rotates around Axis (assumes Axis.Size() == 1).
+	 *
+	 * @param Angle Angle to rotate (in degrees).
+	 * @param Axis Axis to rotate around.
+	 * @return Rotated Vector.
+	 */
+	TVector<TReal> RotateAngleAxis(const TReal AngleDeg, const TVector<TReal>& Axis) const;
+
+	/**
+	 * Returns the cosine of the angle between this vector and another projected onto the XY plane (no Z).
+	 *
+	 * @param B the other vector to find the 2D cosine of the angle with.
+	 * @return The cosine.
+	 */
+	FORCEINLINE TReal CosineAngle2D(TVector<TReal> B) const;
+
+	/**
+	 * Gets a copy of this vector projected onto the input vector.
+	 *
+	 * @param A	Vector to project onto, does not assume it is normalized.
+	 * @return Projected vector.
+	 */
+	FORCEINLINE TVector<TReal> ProjectOnTo(const TVector<TReal>& A) const;
+
+	/**
+	 * Gets a copy of this vector projected onto the input vector, which is assumed to be unit length.
+	 *
+	 * @param  Normal Vector to project onto (assumed to be unit length).
+	 * @return Projected vector.
+	 */
+	FORCEINLINE TVector<TReal> ProjectOnToNormal(const TVector<TReal>& Normal) const;
+
+	/**
+	 * Return the FRotator orientation corresponding to the direction in which the vector points.
+	 * Sets Yaw and Pitch to the proper numbers, and sets Roll to zero because the roll can'TComp be determined from a vector.
+	 *
+	 * @return FRotator from the Vector's direction, without any roll.
+	 * @see ToOrientationQuat()
+	 */
+	CORE_API FRotator ToOrientationRotator() const;
+
+	/**
+	 * Return the Quaternion orientation corresponding to the direction in which the vector points.
+	 * Similar to the FRotator version, returns a result without roll such that it preserves the up vector.
+	 *
+	 * @note If you don'TComp care about preserving the up vector and just want the most direct rotation, you can use the faster
+	 * 'FQuat::FindBetweenVectors(FVector::ForwardVector, YourVector)' or 'FQuat::FindBetweenNormals(...)' if you know the vector is of unit length.
+	 *
+	 * @return Quaternion from the Vector's direction, without any roll.
+	 * @see ToOrientationRotator(), FQuat::FindBetweenVectors()
+	 */
+	CORE_API FQuat ToOrientationQuat() const;
+
+	/**
+	 * Return the FRotator orientation corresponding to the direction in which the vector points.
+	 * Sets Yaw and Pitch to the proper numbers, and sets Roll to zero because the roll can'TComp be determined from a vector.
+	 * @note Identical to 'ToOrientationRotator()' and preserved for legacy reasons.
+	 * @return FRotator from the Vector's direction.
+	 * @see ToOrientationRotator(), ToOrientationQuat()
+	 */
+	CORE_API FRotator Rotation() const;
+
+	/**
+	 * Find good arbitrary axis vectors to represent U and V axes of a plane,
+	 * using this vector as the normal of the plane.
+	 *
+	 * @param Axis1 Reference to first axis.
+	 * @param Axis2 Reference to second axis.
+	 */
+	void FindBestAxisVectors(TVector<TReal>& Axis1, TVector<TReal>& Axis2) const;
+
+	/** When this vector contains Euler angles (degrees), ensure that angles are between +/-180 */
+	void UnwindEuler();
+
+	/**
+	 * Utility to check if there are any non-finite values (NaN or Inf) in this vector.
+	 *
+	 * @return true if there are any non-finite values in this vector, false otherwise.
+	 */
+	bool ContainsNaN() const;
+
+	/**
+	 * Get a textual representation of this vector.
+	 *
+	 * @return A string describing the vector.
+	 */
+	FString ToString() const;
+
+	/**
+	* Get a locale aware textual representation of this vector.
+	*
+	* @return A string describing the vector.
+	*/
+	FText ToText() const;
+
+	/** Get a short textural representation of this vector, for compact readable logging. */
+	FString ToCompactString() const;
+
+	/** Get a short locale aware textural representation of this vector, for compact readable logging. */
+	FText ToCompactText() const;
+
+	/**
+	 * Initialize this Vector based on an FString. The String is expected to contain X=, Y=, Z=.
+	 * The TVector<TReal> will be bogus when InitFromString returns false.
+	 *
+	 * @param	InSourceString	FString containing the vector values.
+	 * @return true if the X,Y,Z values were read successfully; false otherwise.
+	 */
+	bool InitFromString(const FString& InSourceString);
+
+	/**
+	 * Converts a Cartesian unit vector into spherical coordinates on the unit sphere.
+	 * @return Output Theta will be in the range [0, PI], and output Phi will be in the range [-PI, PI].
+	 */
+	FVector2D UnitCartesianToSpherical() const;
+
+	/**
+	 * Convert a direction vector into a 'heading' angle.
+	 *
+	 * @return 'Heading' angle between +/-PI. 0 is pointing down +X.
+	 */
+	TReal HeadingAngle() const;
+
+	/**
+	 * Create an orthonormal basis from a basis with at least two orthogonal vectors.
+	 * It may change the directions of the X and Y axes to make the basis orthogonal,
+	 * but it won'TReal change the direction of the Z axis.
+	 * All axes will be normalized.
+	 *
+	 * @param XAxis The input basis' XAxis, and upon return the orthonormal basis' XAxis.
+	 * @param YAxis The input basis' YAxis, and upon return the orthonormal basis' YAxis.
+	 * @param ZAxis The input basis' ZAxis, and upon return the orthonormal basis' ZAxis.
+	 */
+	static void CreateOrthonormalBasis(TVector<TReal>& XAxis, TVector<TReal>& YAxis, TVector<TReal>& ZAxis);
+
+	/**
+	 * Compare two points and see if they're the same, using a threshold.
+	 *
+	 * @param P First vector.
+	 * @param Q Second vector.
+	 * @return Whether points are the same within a threshold. Uses fast distance approximation (linear per-component distance).
+	 */
+	static bool PointsAreSame(const TVector<TReal>& P, const TVector<TReal>& Q);
+
+	/**
+	 * Compare two points and see if they're within specified distance.
+	 *
+	 * @param Point1 First vector.
+	 * @param Point2 Second vector.
+	 * @param Dist Specified distance.
+	 * @return Whether two points are within the specified distance. Uses fast distance approximation (linear per-component distance).
+	 */
+	static bool PointsAreNear(const TVector<TReal>& Point1, const TVector<TReal>& Point2, TReal Dist);
+
+	/**
+	 * Calculate the signed distance (in the direction of the normal) between a point and a plane.
+	 *
+	 * @param Point The Point we are checking.
+	 * @param PlaneBase The Base Point in the plane.
+	 * @param PlaneNormal The Normal of the plane (assumed to be unit length).
+	 * @return Signed distance between point and plane.
+	 */
+	static TReal PointPlaneDist(const TVector<TReal>& Point, const TVector<TReal>& PlaneBase, const TVector<TReal>& PlaneNormal);
+
+	/**
+	 * Calculate the projection of a point on the plane defined by counter-clockwise (CCW) points A,B,C.
+	 *
+	 * @param Point The point to project onto the plane
+	 * @param A 1st of three points in CCW order defining the plane
+	 * @param B 2nd of three points in CCW order defining the plane
+	 * @param C 3rd of three points in CCW order defining the plane
+	 * @return Projection of Point onto plane ABC
+	 */
+	static TVector<TReal> PointPlaneProject(const TVector<TReal>& Point, const TVector<TReal>& A, const TVector<TReal>& B, const TVector<TReal>& C);
+
+	/**
+	* Calculate the projection of a point on the plane defined by PlaneBase and PlaneNormal.
+	*
+	* @param Point The point to project onto the plane
+	* @param PlaneBase Point on the plane
+	* @param PlaneNorm Normal of the plane (assumed to be unit length).
+	* @return Projection of Point onto plane
+	*/
+	static TVector<TReal> PointPlaneProject(const TVector<TReal>& Point, const TVector<TReal>& PlaneBase, const TVector<TReal>& PlaneNormal);
+
+	/**
+	 * Calculate the projection of a vector on the plane defined by PlaneNormal.
+	 *
+	 * @param  V The vector to project onto the plane.
+	 * @param  PlaneNormal Normal of the plane (assumed to be unit length).
+	 * @return Projection of V onto plane.
+	 */
+	static TVector<TReal> VectorPlaneProject(const TVector<TReal>& V, const TVector<TReal>& PlaneNormal);
+
+	/**
+	 * Euclidean distance between two points.
+	 *
+	 * @param V1 The first point.
+	 * @param V2 The second point.
+	 * @return The distance between two points.
+	 */
+	static FORCEINLINE TReal Dist(const TVector<TReal>& V1, const TVector<TReal>& V2);
+	static FORCEINLINE TReal Distance(const TVector<TReal>& V1, const TVector<TReal>& V2) {
+		return Dist(V1, V2);
+	}
+
+	/**
+	* Euclidean distance between two points in the XY plane (ignoring Z).
+	*
+	* @param V1 The first point.
+	* @param V2 The second point.
+	* @return The distance between two points in the XY plane.
+	*/
+	static FORCEINLINE TReal DistXY(const TVector<TReal>& V1, const TVector<TReal>& V2);
+	static FORCEINLINE TReal Dist2D(const TVector<TReal>& V1, const TVector<TReal>& V2) {
+		return DistXY(V1, V2);
+	}
+
+	/**
+	 * Squared distance between two points.
+	 *
+	 * @param V1 The first point.
+	 * @param V2 The second point.
+	 * @return The squared distance between two points.
+	 */
+	static FORCEINLINE TReal DistSquared(const TVector<TReal>& V1, const TVector<TReal>& V2);
+
+	/**
+	 * Squared distance between two points in the XY plane only.
+	 *
+	 * @param V1 The first point.
+	 * @param V2 The second point.
+	 * @return The squared distance between two points in the XY plane
+	 */
+	static FORCEINLINE TReal DistSquaredXY(const TVector<TReal>& V1, const TVector<TReal>& V2);
+	static FORCEINLINE TReal DistSquared2D(const TVector<TReal>& V1, const TVector<TReal>& V2) {
+		return DistSquaredXY(V1, V2);
+	}
+
+	/**
+	 * Compute pushout of a box from a plane.
+	 *
+	 * @param Normal The plane normal.
+	 * @param Size The size of the box.
+	 * @return Pushout required.
+	 */
+	static FORCEINLINE TReal BoxPushOut(const TVector<TReal>& Normal, const TVector<TReal>& Size);
+
+	/**
+	 * Min, Max, Min3, Max3 like FMath
+	 */
+	static FORCEINLINE TVector<TReal> Min(const TVector<TReal>& A, const TVector<TReal>& B);
+	static FORCEINLINE TVector<TReal> Max(const TVector<TReal>& A, const TVector<TReal>& B);
+
+	static FORCEINLINE TVector<TReal> Min3(const TVector<TReal>& A, const TVector<TReal>& B, const TVector<TReal>& C);
+	static FORCEINLINE TVector<TReal> Max3(const TVector<TReal>& A, const TVector<TReal>& B, const TVector<TReal>& C);
+
+	/**
+	 * See if two normal vectors are nearly parallel, meaning the angle between them is close to 0 degrees.
+	 *
+	 * @param  Normal1 First normalized vector.
+	 * @param  Normal1 Second normalized vector.
+	 * @param  ParallelCosineThreshold Normals are parallel if absolute value of dot product (cosine of angle between them) is greater than or equal to this. For example: cos(1.0 degrees).
+	 * @return true if vectors are nearly parallel, false otherwise.
+	 */
+	static bool Parallel(const TVector<TReal>& Normal1, const TVector<TReal>& Normal2, TReal ParallelCosineThreshold = THRESH_NORMALS_ARE_PARALLEL);
+
+	/**
+	 * See if two normal vectors are coincident (nearly parallel and point in the same direction).
+	 *
+	 * @param  Normal1 First normalized vector.
+	 * @param  Normal2 Second normalized vector.
+	 * @param  ParallelCosineThreshold Normals are coincident if dot product (cosine of angle between them) is greater than or equal to this. For example: cos(1.0 degrees).
+	 * @return true if vectors are coincident (nearly parallel and point in the same direction), false otherwise.
+	 */
+	static bool Coincident(const TVector<TReal>& Normal1, const TVector<TReal>& Normal2, TReal ParallelCosineThreshold = THRESH_NORMALS_ARE_PARALLEL);
+
+	/**
+	 * See if two normal vectors are nearly orthogonal (perpendicular), meaning the angle between them is close to 90 degrees.
+	 *
+	 * @param  Normal1 First normalized vector.
+	 * @param  Normal2 Second normalized vector.
+	 * @param  OrthogonalCosineThreshold Normals are orthogonal if absolute value of dot product (cosine of angle between them) is less than or equal to this. For example: cos(89.0 degrees).
+	 * @return true if vectors are orthogonal (perpendicular), false otherwise.
+	 */
+	static bool Orthogonal(const TVector<TReal>& Normal1, const TVector<TReal>& Normal2, TReal OrthogonalCosineThreshold = THRESH_NORMALS_ARE_ORTHOGONAL);
+
+	/**
+	 * See if two planes are coplanar. They are coplanar if the normals are nearly parallel and the planes include the same set of points.
+	 *
+	 * @param Base1 The base point in the first plane.
+	 * @param Normal1 The normal of the first plane.
+	 * @param Base2 The base point in the second plane.
+	 * @param Normal2 The normal of the second plane.
+	 * @param ParallelCosineThreshold Normals are parallel if absolute value of dot product is greater than or equal to this.
+	 * @return true if the planes are coplanar, false otherwise.
+	 */
+	static bool Coplanar(const TVector<TReal>& Base1, const TVector<TReal>& Normal1, const TVector<TReal>& Base2, const TVector<TReal>& Normal2, TReal ParallelCosineThreshold = THRESH_NORMALS_ARE_PARALLEL);
+
+	/**
+	 * Triple product of three vectors: X dot (Y cross Z).
+	 *
+	 * @param X The first vector.
+	 * @param Y The second vector.
+	 * @param Z The third vector.
+	 * @return The triple product: X dot (Y cross Z).
+	 */
+	static TReal Triple(const TVector<TReal>& X, const TVector<TReal>& Y, const TVector<TReal>& Z);
+
+	/**
+	 * Generates a list of sample points on a Bezier curve defined by 2 points.
+	 *
+	 * @param ControlPoints	Array of 4 FVectors (vert1, controlpoint1, controlpoint2, vert2).
+	 * @param NumPoints Number of samples.
+	 * @param OutPoints Receives the output samples.
+	 * @return The path length.
+	 */
+	static TReal EvaluateBezier(const TVector<TReal>* ControlPoints, int32 NumPoints, TArray<TVector<TReal>>& OutPoints);
+
+	/**
+	 * Converts a vector containing radian values to a vector containing degree values.
+	 *
+	 * @param RadVector	Vector containing radian values
+	 * @return Vector  containing degree values
+	 */
+	static TVector<TReal> RadiansToDegrees(const TVector<TReal>& RadVector);
+
+	/**
+	 * Converts a vector containing degree values to a vector containing radian values.
+	 *
+	 * @param DegVector	Vector containing degree values
+	 * @return Vector containing radian values
+	 */
+	static TVector<TReal> DegreesToRadians(const TVector<TReal>& DegVector);
+
+	/**
+	 * Given a current set of cluster centers, a set of points, iterate N times to move clusters to be central.
+	 *
+	 * @param Clusters Reference to array of Clusters.
+	 * @param Points Set of points.
+	 * @param NumIterations Number of iterations.
+	 * @param NumConnectionsToBeValid Sometimes you will have long strings that come off the mass of points
+	 * which happen to have been chosen as Cluster starting points.  You want to be able to disregard those.
+	 */
+	static void GenerateClusterCenters(TArray<TVector<TReal>>& Clusters, const TArray<TVector<TReal>>& Points, int32 NumIterations, int32 NumConnectionsToBeValid);
+
+	/**
+	 * Structured archive slot serializer.
+	 *
+	 * @param Slot Structured archive slot.
+	 * @param V Vector to serialize.
+	 */
+	FORCENOINLINE friend void operator<<(FStructuredArchive::FSlot Slot, TVector<TReal>& V)
+	{
+		// @warning BulkSerialize: TVector<TReal> is serialized as memory dump
+		// See TArray::BulkSerialize for detailed description of implied limitations.
+		FStructuredArchive::FRecord Record = Slot.EnterRecord();
+		//if (Record.GetUnderlyingArchive().UE4Ver() >= VER_UE5_LARGE_WORLD_COORDINATES)
+		//{
+		//	// Stored as doubles, so serialize double and copy.
+		//	double SX = V.X, SY = V.Y, SZ = V.Z;
+		//	Record << SA_VALUE(TEXT("X"), SX);
+		//	Record << SA_VALUE(TEXT("Y"), SY);
+		//	Record << SA_VALUE(TEXT("Z"), SZ);
+		//	V = TVector<TReal>(SX, SY, SZ);
+		//}
+		//else
+		{
+			// Stored as floats, so serialize float and copy.
+			float SX = V.X, SY = V.Y, SZ = V.Z;
+			Record << SA_VALUE(TEXT("X"), SX);
+			Record << SA_VALUE(TEXT("Y"), SY);
+			Record << SA_VALUE(TEXT("Z"), SZ);
+			V = TVector<TReal>(SX, SY, SZ);
+		}
+	}
+
+	bool Serialize(FStructuredArchive::FSlot Slot)
+	{
+		Slot << (TVector<TReal>&)*this;
+		return true;
+	}
+
+	/**
+	 * Network serialization function.
+	 * FVectors NetSerialize without quantization (ie exact values are serialized). se the FVectors_NetQuantize etc (NetSerialization.h) instead.
+	 *
+	 * @see FVector_NetQuantize, FVector_NetQuantize10, FVector_NetQuantize100, FVector_NetQuantizeNormal
+	 */
+	bool NetSerialize(FArchive& Ar, class UPackageMap* Map, bool& bOutSuccess)
+	{
+		Ar << X;
+		Ar << Y;
+		Ar << Z;
+		return true;
+	}
+};
+
+/* FVector inline functions
+ *****************************************************************************/
+
+template<typename TReal>
+FORCEINLINE TVector<TReal>::TVector(const FVector2D V, TReal InZ)
+	: X(V.X), Y(V.Y), Z(InZ)
+{
+	DiagnosticCheckNaN();
+}
+
+
+
+template<typename TReal>
+inline TVector<TReal> TVector<TReal>::RotateAngleAxis(const TReal AngleDeg, const TVector<TReal>& Axis) const
+{
+	TReal S, C;
+	FMath::SinCos(&S, &C, FMath::DegreesToRadians(AngleDeg));
+
+	const TReal XX = Axis.X * Axis.X;
+	const TReal YY = Axis.Y * Axis.Y;
+	const TReal ZZ = Axis.Z * Axis.Z;
+
+	const TReal XY = Axis.X * Axis.Y;
+	const TReal YZ = Axis.Y * Axis.Z;
+	const TReal ZX = Axis.Z * Axis.X;
+
+	const TReal XS = Axis.X * S;
+	const TReal YS = Axis.Y * S;
+	const TReal ZS = Axis.Z * S;
+
+	const TReal OMC = 1.f - C;
+
+	return TVector<TReal>(
+		(OMC * XX + C) * X + (OMC * XY - ZS) * Y + (OMC * ZX + YS) * Z,
+		(OMC * XY + ZS) * X + (OMC * YY + C) * Y + (OMC * YZ - XS) * Z,
+		(OMC * ZX - YS) * X + (OMC * YZ + XS) * Y + (OMC * ZZ + C) * Z
+	);
+}
+
+template<typename TReal>
+inline bool TVector<TReal>::PointsAreSame(const TVector<TReal>& P, const TVector<TReal>& Q)
+{
+	TReal Temp;
+	Temp = P.X - Q.X;
+	if ((Temp > -THRESH_POINTS_ARE_SAME) && (Temp < THRESH_POINTS_ARE_SAME))
+	{
+		Temp = P.Y - Q.Y;
+		if ((Temp > -THRESH_POINTS_ARE_SAME) && (Temp < THRESH_POINTS_ARE_SAME))
+		{
+			Temp = P.Z - Q.Z;
+			if ((Temp > -THRESH_POINTS_ARE_SAME) && (Temp < THRESH_POINTS_ARE_SAME))
+			{
+				return true;
+			}
+		}
+	}
+	return false;
+}
+
+template<typename TReal>
+inline bool TVector<TReal>::PointsAreNear(const TVector<TReal>& Point1, const TVector<TReal>& Point2, TReal Dist)
+{
+	TReal Temp;
+	Temp = (Point1.X - Point2.X); if (FMath::Abs(Temp) >= Dist) return false;
+	Temp = (Point1.Y - Point2.Y); if (FMath::Abs(Temp) >= Dist) return false;
+	Temp = (Point1.Z - Point2.Z); if (FMath::Abs(Temp) >= Dist) return false;
+	return true;
+}
+
+template<typename TReal>
+inline TReal TVector<TReal>::PointPlaneDist
+(
+	const TVector<TReal>& Point,
+	const TVector<TReal>& PlaneBase,
+	const TVector<TReal>& PlaneNormal
+)
+{
+	return (Point - PlaneBase) | PlaneNormal;
+}
+
+
+template<typename TReal>
+inline TVector<TReal> TVector<TReal>::PointPlaneProject(const TVector<TReal>& Point, const TVector<TReal>& PlaneBase, const TVector<TReal>& PlaneNorm)
+{
+	//Find the distance of X from the plane
+	//Add the distance back along the normal from the point
+	return Point - TVector<TReal>::PointPlaneDist(Point, PlaneBase, PlaneNorm) * PlaneNorm;
+}
+
+template<typename TReal>
+inline TVector<TReal> TVector<TReal>::VectorPlaneProject(const TVector<TReal>& V, const TVector<TReal>& PlaneNormal)
+{
+	return V - V.ProjectOnToNormal(PlaneNormal);
+}
+
+template<typename TReal>
+inline bool TVector<TReal>::Parallel(const TVector<TReal>& Normal1, const TVector<TReal>& Normal2, TReal ParallelCosineThreshold)
+{
+	const TReal NormalDot = Normal1 | Normal2;
+	return FMath::Abs(NormalDot) >= ParallelCosineThreshold;
+}
+
+template<typename TReal>
+inline bool TVector<TReal>::Coincident(const TVector<TReal>& Normal1, const TVector<TReal>& Normal2, TReal ParallelCosineThreshold)
+{
+	const TReal NormalDot = Normal1 | Normal2;
+	return NormalDot >= ParallelCosineThreshold;
+}
+
+template<typename TReal>
+inline bool TVector<TReal>::Orthogonal(const TVector<TReal>& Normal1, const TVector<TReal>& Normal2, TReal OrthogonalCosineThreshold)
+{
+	const TReal NormalDot = Normal1 | Normal2;
+	return FMath::Abs(NormalDot) <= OrthogonalCosineThreshold;
+}
+
+template<typename TReal>
+inline bool TVector<TReal>::Coplanar(const TVector<TReal>& Base1, const TVector<TReal>& Normal1, const TVector<TReal>& Base2, const TVector<TReal>& Normal2, TReal ParallelCosineThreshold)
+{
+	if (!TVector<TReal>::Parallel(Normal1, Normal2, ParallelCosineThreshold)) return false;
+	else if (FMath::Abs(TVector<TReal>::PointPlaneDist(Base2, Base1, Normal1)) > THRESH_POINT_ON_PLANE) return false;
+	else return true;
+}
+
+template<typename TReal>
+inline TReal TVector<TReal>::Triple(const TVector<TReal>& X, const TVector<TReal>& Y, const TVector<TReal>& Z)
+{
+	return
+		((X.X * (Y.Y * Z.Z - Y.Z * Z.Y))
+			+ (X.Y * (Y.Z * Z.X - Y.X * Z.Z))
+			+ (X.Z * (Y.X * Z.Y - Y.Y * Z.X)));
+}
+
+template<typename TReal>
+inline TVector<TReal> TVector<TReal>::RadiansToDegrees(const TVector<TReal>& RadVector)
+{
+	return RadVector * (180.f / PI);
+}
+
+template<typename TReal>
+inline TVector<TReal> TVector<TReal>::DegreesToRadians(const TVector<TReal>& DegVector)
+{
+	return DegVector * (PI / 180.f);
+}
+
+template<typename TReal>
+FORCEINLINE TVector<TReal>::TVector()
+{}
+
+template<typename TReal>
+FORCEINLINE TVector<TReal>::TVector(TReal InF)
+	: X(InF), Y(InF), Z(InF)
+{
+	DiagnosticCheckNaN();
+}
+
+template<typename TReal>
+FORCEINLINE TVector<TReal>::TVector(TReal InX, TReal InY, TReal InZ)
+	: X(InX), Y(InY), Z(InZ)
+{
+	DiagnosticCheckNaN();
+}
+
+template<typename TReal>
+FORCEINLINE TVector<TReal>::TVector(const FLinearColor& InColor)
+	: X(InColor.R), Y(InColor.G), Z(InColor.B)
+{
+	DiagnosticCheckNaN();
+}
+
+template<typename TReal>
+FORCEINLINE TVector<TReal>::TVector(FIntVector InVector)
+	: X((TReal)InVector.X), Y((TReal)InVector.Y), Z((TReal)InVector.Z)
+{
+	DiagnosticCheckNaN();
+}
+
+template<typename TReal>
+FORCEINLINE TVector<TReal>::TVector(FIntPoint A)
+	: X((TReal)A.X), Y((TReal)A.Y), Z(0.f)
+{
+	DiagnosticCheckNaN();
+}
+
+template<typename TReal>
+FORCEINLINE TVector<TReal>::TVector(EForceInit)
+	: X(0.0f), Y(0.0f), Z(0.0f)
+{
+	DiagnosticCheckNaN();
+}
+
+#ifdef IMPLEMENT_ASSIGNMENT_OPERATOR_MANUALLY
+template<typename TReal>
+FORCEINLINE TVector<TReal>& TVector<TReal>::operator=(const TVector<TReal>& Other)
+{
+	this->X = Other.X;
+	this->Y = Other.Y;
+	this->Z = Other.Z;
+
+	DiagnosticCheckNaN();
+
+	return *this;
+}
+#endif
+
+template<typename TReal>
+FORCEINLINE TVector<TReal> TVector<TReal>::operator^(const TVector<TReal>& V) const
+{
+	return TVector<TReal>
+	(
+		Y * V.Z - Z * V.Y,
+		Z * V.X - X * V.Z,
+		X * V.Y - Y * V.X
+	);
+}
+
+template<typename TReal>
+FORCEINLINE TVector<TReal> TVector<TReal>::CrossProduct(const TVector<TReal>& A, const TVector<TReal>& B)
+{
+	return A ^ B;
+}
+
+template<typename TReal>
+FORCEINLINE TReal TVector<TReal>::operator|(const TVector<TReal>& V) const
+{
+	return X * V.X + Y * V.Y + Z * V.Z;
+}
+
+template<typename TReal>
+FORCEINLINE TReal TVector<TReal>::DotProduct(const TVector<TReal>& A, const TVector<TReal>& B)
+{
+	return A | B;
+}
+
+template<typename TReal>
+FORCEINLINE TVector<TReal> TVector<TReal>::operator+(const TVector<TReal>& V) const
+{
+	return TVector<TReal>(X + V.X, Y + V.Y, Z + V.Z);
+}
+
+template<typename TReal>
+FORCEINLINE TVector<TReal> TVector<TReal>::operator-(const TVector<TReal>& V) const
+{
+	return TVector<TReal>(X - V.X, Y - V.Y, Z - V.Z);
+}
+
+template<typename TReal>
+FORCEINLINE TVector<TReal> TVector<TReal>::operator*(const TVector<TReal>& V) const
+{
+	return TVector<TReal>(X * V.X, Y * V.Y, Z * V.Z);
+}
+
+template<typename TReal>
+FORCEINLINE TVector<TReal> TVector<TReal>::operator/(const TVector<TReal>& V) const
+{
+	return TVector<TReal>(X / V.X, Y / V.Y, Z / V.Z);
+}
+
+template<typename TReal>
+FORCEINLINE bool TVector<TReal>::operator==(const TVector<TReal>& V) const
+{
+	return X == V.X && Y == V.Y && Z == V.Z;
+}
+
+template<typename TReal>
+FORCEINLINE bool TVector<TReal>::operator!=(const TVector<TReal>& V) const
+{
+	return X != V.X || Y != V.Y || Z != V.Z;
+}
+
+template<typename TReal>
+FORCEINLINE bool TVector<TReal>::Equals(const TVector<TReal>& V, TReal Tolerance) const
+{
+	return FMath::Abs(X - V.X) <= Tolerance && FMath::Abs(Y - V.Y) <= Tolerance && FMath::Abs(Z - V.Z) <= Tolerance;
+}
+
+template<typename TReal>
+FORCEINLINE bool TVector<TReal>::AllComponentsEqual(TReal Tolerance) const
+{
+	return FMath::Abs(X - Y) <= Tolerance && FMath::Abs(X - Z) <= Tolerance && FMath::Abs(Y - Z) <= Tolerance;
+}
+
+
+template<typename TReal>
+FORCEINLINE TVector<TReal> TVector<TReal>::operator-() const
+{
+	return TVector<TReal>(-X, -Y, -Z);
+}
+
+
+template<typename TReal>
+FORCEINLINE TVector<TReal> TVector<TReal>::operator+=(const TVector<TReal>& V)
+{
+	X += V.X; Y += V.Y; Z += V.Z;
+	DiagnosticCheckNaN();
+	return *this;
+}
+
+template<typename TReal>
+FORCEINLINE TVector<TReal> TVector<TReal>::operator-=(const TVector<TReal>& V)
+{
+	X -= V.X; Y -= V.Y; Z -= V.Z;
+	DiagnosticCheckNaN();
+	return *this;
+}
+
+template<typename TReal>
+FORCEINLINE TVector<TReal> TVector<TReal>::operator*=(const TVector<TReal>& V)
+{
+	X *= V.X; Y *= V.Y; Z *= V.Z;
+	DiagnosticCheckNaN();
+	return *this;
+}
+
+template<typename TReal>
+FORCEINLINE TVector<TReal> TVector<TReal>::operator/=(const TVector<TReal>& V)
+{
+	X /= V.X; Y /= V.Y; Z /= V.Z;
+	DiagnosticCheckNaN();
+	return *this;
+}
+
+template<typename TReal>
+FORCEINLINE TReal& TVector<TReal>::operator[](int32 Index)
+{
+	checkSlow(Index >= 0 && Index < 3);
+	return (&X)[Index];
+}
+
+template<typename TReal>
+FORCEINLINE TReal TVector<TReal>::operator[](int32 Index)const
+{
+	checkSlow(Index >= 0 && Index < 3);
+	return (&X)[Index];
+}
+
+template<typename TReal>
+FORCEINLINE void TVector<TReal>::Set(TReal InX, TReal InY, TReal InZ)
+{
+	X = InX;
+	Y = InY;
+	Z = InZ;
+	DiagnosticCheckNaN();
+}
+
+template<typename TReal>
+FORCEINLINE TReal TVector<TReal>::GetMax() const
+{
+	return FMath::Max(FMath::Max(X, Y), Z);
+}
+
+template<typename TReal>
+FORCEINLINE TReal TVector<TReal>::GetAbsMax() const
+{
+	return FMath::Max(FMath::Max(FMath::Abs(X), FMath::Abs(Y)), FMath::Abs(Z));
+}
+
+template<typename TReal>
+FORCEINLINE TReal TVector<TReal>::GetMin() const
+{
+	return FMath::Min(FMath::Min(X, Y), Z);
+}
+
+template<typename TReal>
+FORCEINLINE TReal TVector<TReal>::GetAbsMin() const
+{
+	return FMath::Min(FMath::Min(FMath::Abs(X), FMath::Abs(Y)), FMath::Abs(Z));
+}
+
+template<typename TReal>
+FORCEINLINE TVector<TReal> TVector<TReal>::ComponentMin(const TVector<TReal>& Other) const
+{
+	return TVector<TReal>(FMath::Min(X, Other.X), FMath::Min(Y, Other.Y), FMath::Min(Z, Other.Z));
+}
+
+template<typename TReal>
+FORCEINLINE TVector<TReal> TVector<TReal>::ComponentMax(const TVector<TReal>& Other) const
+{
+	return TVector<TReal>(FMath::Max(X, Other.X), FMath::Max(Y, Other.Y), FMath::Max(Z, Other.Z));
+}
+
+template<typename TReal>
+FORCEINLINE TVector<TReal> TVector<TReal>::GetAbs() const
+{
+	return TVector<TReal>(FMath::Abs(X), FMath::Abs(Y), FMath::Abs(Z));
+}
+
+template<typename TReal>
+FORCEINLINE TReal TVector<TReal>::Size() const
+{
+	return FMath::Sqrt(X * X + Y * Y + Z * Z);
+}
+
+template<typename TReal>
+FORCEINLINE TReal TVector<TReal>::SizeSquared() const
+{
+	return X * X + Y * Y + Z * Z;
+}
+
+template<typename TReal>
+FORCEINLINE TReal TVector<TReal>::Size2D() const
+{
+	return FMath::Sqrt(X * X + Y * Y);
+}
+
+template<typename TReal>
+FORCEINLINE TReal TVector<TReal>::SizeSquared2D() const
+{
+	return X * X + Y * Y;
+}
+
+template<typename TReal>
+FORCEINLINE bool TVector<TReal>::IsNearlyZero(TReal Tolerance) const
+{
+	return
+		FMath::Abs(X) <= Tolerance
+		&& FMath::Abs(Y) <= Tolerance
+		&& FMath::Abs(Z) <= Tolerance;
+}
+
+template<typename TReal>
+FORCEINLINE bool TVector<TReal>::IsZero() const
+{
+	return X == 0.f && Y == 0.f && Z == 0.f;
+}
+
+template<typename TReal>
+FORCEINLINE bool TVector<TReal>::Normalize(TReal Tolerance)
+{
+	const TReal SquareSum = X * X + Y * Y + Z * Z;
+	if (SquareSum > Tolerance)
+	{
+		const TReal Scale = FMath::InvSqrt(SquareSum);
+		X *= Scale; Y *= Scale; Z *= Scale;
+		return true;
+	}
+	return false;
+}
+
+template<typename TReal>
+FORCEINLINE bool TVector<TReal>::IsUnit(TReal LengthSquaredTolerance) const
+{
+	return FMath::Abs(1.0f - SizeSquared()) < LengthSquaredTolerance;
+}
+
+template<typename TReal>
+FORCEINLINE bool TVector<TReal>::IsNormalized() const
+{
+	return (FMath::Abs(1.f - SizeSquared()) < THRESH_VECTOR_NORMALIZED);
+}
+
+template<typename TReal>
+FORCEINLINE void TVector<TReal>::ToDirectionAndLength(TVector<TReal>& OutDir, double& OutLength) const
+{
+	OutLength = Size();
+	if (OutLength > SMALL_NUMBER)
+	{
+		double OneOverLength = 1.0f / OutLength;
+		OutDir = TVector<TReal>(X * OneOverLength, Y * OneOverLength,
+			Z * OneOverLength);
+	}
+	else
+	{
+		OutDir = TVector<TReal>(0, 0, 0);
+	}
+}
+
+template<typename TReal>
+FORCEINLINE void TVector<TReal>::ToDirectionAndLength(TVector<TReal>& OutDir, float& OutLength) const
+{
+	OutLength = Size();
+	if (OutLength > SMALL_NUMBER)
+	{
+		float OneOverLength = 1.0f / OutLength;
+		OutDir = TVector<TReal>(X * OneOverLength, Y * OneOverLength,
+			Z * OneOverLength);
+	}
+	else
+	{
+		OutDir = TVector<TReal>(0, 0, 0);
+	}
+}
+
+
+template<typename TReal>
+FORCEINLINE TVector<TReal> TVector<TReal>::GetSignVector() const
+{
+	return TVector<TReal>
+	(
+		FMath::FloatSelect(X, (TReal)1, (TReal)-1),		// LWC_TODO: Templatize FMath functionality
+		FMath::FloatSelect(Y, (TReal)1, (TReal)-1),
+		FMath::FloatSelect(Z, (TReal)1, (TReal)-1)
+	);
+}
+
+template<typename TReal>
+FORCEINLINE TVector<TReal> TVector<TReal>::Projection() const
+{
+	const TReal RZ = 1.f / Z;
+	return TVector<TReal>(X * RZ, Y * RZ, 1);
+}
+
+template<typename TReal>
+FORCEINLINE TVector<TReal> TVector<TReal>::GetUnsafeNormal() const
+{
+	const TReal Scale = FMath::InvSqrt(X * X + Y * Y + Z * Z);
+	return TVector<TReal>(X * Scale, Y * Scale, Z * Scale);
+}
+
+template<typename TReal>
+FORCEINLINE TVector<TReal> TVector<TReal>::GetUnsafeNormal2D() const
+{
+	const TReal Scale = FMath::InvSqrt(X * X + Y * Y);
+	return TVector<TReal>(X * Scale, Y * Scale, 0);
+}
+
+template<typename TReal>
+FORCEINLINE TVector<TReal> TVector<TReal>::GridSnap(const TReal& GridSz) const
+{
+	return TVector<TReal>(FMath::GridSnap(X, GridSz), FMath::GridSnap(Y, GridSz), FMath::GridSnap(Z, GridSz));
+}
+
+template<typename TReal>
+FORCEINLINE TVector<TReal> TVector<TReal>::BoundToCube(TReal Radius) const
+{
+	return TVector<TReal>
+	(
+		FMath::Clamp(X, -Radius, Radius),
+		FMath::Clamp(Y, -Radius, Radius),
+		FMath::Clamp(Z, -Radius, Radius)
+	);
+}
+
+template<typename TReal>
+FORCEINLINE TVector<TReal> TVector<TReal>::BoundToBox(const TVector<TReal>& Min, const TVector<TReal> Max) const
+{
+	return TVector<TReal>
+	(
+		FMath::Clamp(X, Min.X, Max.X),
+		FMath::Clamp(Y, Min.Y, Max.Y),
+		FMath::Clamp(Z, Min.Z, Max.Z)
+	);
+}
+
+
+template<typename TReal>
+FORCEINLINE TVector<TReal> TVector<TReal>::GetClampedToSize(TReal Min, TReal Max) const
+{
+	TReal VecSize = Size();
+	const TVector<TReal> VecDir = (VecSize > SMALL_NUMBER) ? (*this / VecSize) : TVector<TReal>(0, 0, 0);
+
+	VecSize = FMath::Clamp(VecSize, Min, Max);
+
+	return VecSize * VecDir;
+}
+
+template<typename TReal>
+FORCEINLINE TVector<TReal> TVector<TReal>::GetClampedToSize2D(TReal Min, TReal Max) const
+{
+	TReal VecSize2D = Size2D();
+	const TVector<TReal> VecDir = (VecSize2D > SMALL_NUMBER) ? (*this / VecSize2D) : TVector<TReal>(0, 0, 0);
+
+	VecSize2D = FMath::Clamp(VecSize2D, Min, Max);
+
+	return TVector<TReal>(VecSize2D * VecDir.X, VecSize2D * VecDir.Y, Z);
+}
+
+template<typename TReal>
+FORCEINLINE TVector<TReal> TVector<TReal>::GetClampedToMaxSize(TReal MaxSize) const
+{
+	if (MaxSize < KINDA_SMALL_NUMBER)
+	{
+		return TVector<TReal>(0, 0, 0);
+	}
+
+	const TReal VSq = SizeSquared();
+	if (VSq > FMath::Square(MaxSize))
+	{
+		const TReal Scale = MaxSize * FMath::InvSqrt(VSq);
+		return TVector<TReal>(X * Scale, Y * Scale, Z * Scale);
+	}
+	else
+	{
+		return *this;
+	}
+}
+
+template<typename TReal>
+FORCEINLINE TVector<TReal> TVector<TReal>::GetClampedToMaxSize2D(TReal MaxSize) const
+{
+	if (MaxSize < KINDA_SMALL_NUMBER)
+	{
+		return TVector<TReal>(0.f, 0.f, Z);
+	}
+
+	const TReal VSq2D = SizeSquared2D();
+	if (VSq2D > FMath::Square(MaxSize))
+	{
+		const TReal Scale = MaxSize * FMath::InvSqrt(VSq2D);
+		return TVector<TReal>(X * Scale, Y * Scale, Z);
+	}
+	else
+	{
+		return *this;
+	}
+}
+
+template<typename TReal>
+FORCEINLINE void TVector<TReal>::AddBounded(const TVector<TReal>& V, TReal Radius)
+{
+	*this = (*this + V).BoundToCube(Radius);
+}
+
+template<typename TReal>
+FORCEINLINE TReal& TVector<TReal>::Component(int32 Index)
+{
+	return (&X)[Index];
+}
+
+template<typename TReal>
+FORCEINLINE TReal TVector<TReal>::Component(int32 Index) const
+{
+	return (&X)[Index];
+}
+
+template<typename TReal>
+FORCEINLINE TReal TVector<TReal>::GetComponentForAxis(EAxis::Type Axis) const
+{
+	switch (Axis)
+	{
+	case EAxis::X:
+		return X;
+	case EAxis::Y:
+		return Y;
+	case EAxis::Z:
+		return Z;
+	default:
+		return 0.f;
+	}
+}
+
+template<typename TReal>
+FORCEINLINE void TVector<TReal>::SetComponentForAxis(EAxis::Type Axis, TReal Component)
+{
+	switch (Axis)
+	{
+	case EAxis::X:
+		X = Component;
+		break;
+	case EAxis::Y:
+		Y = Component;
+		break;
+	case EAxis::Z:
+		Z = Component;
+		break;
+	}
+}
+
+template<typename TReal>
+FORCEINLINE TVector<TReal> TVector<TReal>::Reciprocal() const
+{
+	TVector<TReal> RecVector;
+	if (X != 0.f)
+	{
+		RecVector.X = 1.f / X;
+	}
+	else
+	{
+		RecVector.X = BIG_NUMBER;
+	}
+	if (Y != 0.f)
+	{
+		RecVector.Y = 1.f / Y;
+	}
+	else
+	{
+		RecVector.Y = BIG_NUMBER;
+	}
+	if (Z != 0.f)
+	{
+		RecVector.Z = 1.f / Z;
+	}
+	else
+	{
+		RecVector.Z = BIG_NUMBER;
+	}
+
+	return RecVector;
+}
+
+
+
+
+template<typename TReal>
+FORCEINLINE bool TVector<TReal>::IsUniform(TReal Tolerance) const
+{
+	return AllComponentsEqual(Tolerance);
+}
+
+template<typename TReal>
+FORCEINLINE TVector<TReal> TVector<TReal>::MirrorByVector(const TVector<TReal>& MirrorNormal) const
+{
+	return *this - MirrorNormal * (2.f * (*this | MirrorNormal));
+}
+
+template<typename TReal>
+FORCEINLINE TVector<TReal> TVector<TReal>::GetSafeNormal(TReal Tolerance) const
+{
+	const TReal SquareSum = X * X + Y * Y + Z * Z;
+
+	// Not sure if it's safe to add tolerance in there. Might introduce too many errors
+	if (SquareSum == 1.f)
+	{
+		return *this;
+	}
+	else if (SquareSum < Tolerance)
+	{
+		return TVector<TReal>(0, 0, 0);
+	}
+	const TReal Scale = (TReal)FMath::InvSqrt(SquareSum);
+	return TVector<TReal>(X * Scale, Y * Scale, Z * Scale);
+}
+
+template<typename TReal>
+FORCEINLINE TVector<TReal> TVector<TReal>::GetSafeNormal2D(TReal Tolerance) const
+{
+	const TReal SquareSum = X * X + Y * Y;
+
+	// Not sure if it's safe to add tolerance in there. Might introduce too many errors
+	if (SquareSum == 1.f)
+	{
+		if (Z == 0.f)
+		{
+			return *this;
+		}
+		else
+		{
+			return TVector<TReal>(X, Y, 0.f);
+		}
+	}
+	else if (SquareSum < Tolerance)
+	{
+		return TVector<TReal>(0, 0, 0);
+	}
+
+	const TReal Scale = FMath::InvSqrt(SquareSum);
+	return TVector<TReal>(X * Scale, Y * Scale, 0.f);
+}
+
+template<typename TReal>
+FORCEINLINE TReal TVector<TReal>::CosineAngle2D(TVector<TReal> B) const
+{
+	TVector<TReal> A(*this);
+	A.Z = 0.0f;
+	B.Z = 0.0f;
+	A.Normalize();
+	B.Normalize();
+	return A | B;
+}
+
+template<typename TReal>
+FORCEINLINE TVector<TReal> TVector<TReal>::ProjectOnTo(const TVector<TReal>& A) const
+{
+	return (A * ((*this | A) / (A | A)));
+}
+
+template<typename TReal>
+FORCEINLINE TVector<TReal> TVector<TReal>::ProjectOnToNormal(const TVector<TReal>& Normal) const
+{
+	return (Normal * (*this | Normal));
+}
+
+
+template<typename TReal>
+void TVector<TReal>::GenerateClusterCenters(TArray<TVector<TReal>>& Clusters, const TArray<TVector<TReal>>& Points, int32 NumIterations, int32 NumConnectionsToBeValid)
+{
+	struct FClusterMovedHereToMakeCompile
+	{
+		TVector<TReal> ClusterPosAccum;
+		int32 ClusterSize;
+	};
+
+	// Check we have >0 points and clusters
+	if (Points.Num() == 0 || Clusters.Num() == 0)
+	{
+		return;
+	}
+
+	// Temp storage for each cluster that mirrors the order of the passed in Clusters array
+	TArray<FClusterMovedHereToMakeCompile> ClusterData;
+	ClusterData.AddZeroed(Clusters.Num());
+
+	// Then iterate
+	for (int32 ItCount = 0; ItCount < NumIterations; ItCount++)
+	{
+		// Classify each point - find closest cluster center
+		for (int32 i = 0; i < Points.Num(); i++)
+		{
+			const TVector<TReal>& Pos = Points[i];
+
+			// Iterate over all clusters to find closes one
+			int32 NearestClusterIndex = INDEX_NONE;
+			TReal NearestClusterDistSqr = BIG_NUMBER;
+			for (int32 j = 0; j < Clusters.Num(); j++)
+			{
+				const TReal DistSqr = (Pos - Clusters[j]).SizeSquared();
+				if (DistSqr < NearestClusterDistSqr)
+				{
+					NearestClusterDistSqr = DistSqr;
+					NearestClusterIndex = j;
+				}
+			}
+			// Update its info with this point
+			if (NearestClusterIndex != INDEX_NONE)
+			{
+				ClusterData[NearestClusterIndex].ClusterPosAccum += Pos;
+				ClusterData[NearestClusterIndex].ClusterSize++;
+			}
+		}
+
+		// All points classified - update cluster center as average of membership
+		for (int32 i = 0; i < Clusters.Num(); i++)
+		{
+			if (ClusterData[i].ClusterSize > 0)
+			{
+				Clusters[i] = ClusterData[i].ClusterPosAccum / (TReal)ClusterData[i].ClusterSize;
+			}
+		}
+	}
+
+	// so now after we have possible cluster centers we want to remove the ones that are outliers and not part of the main cluster
+	for (int32 i = 0; i < ClusterData.Num(); i++)
+	{
+		if (ClusterData[i].ClusterSize < NumConnectionsToBeValid)
+		{
+			Clusters.RemoveAt(i);
+		}
+	}
+}
+
+template<typename TReal>
+TReal TVector<TReal>::EvaluateBezier(const TVector<TReal>* ControlPoints, int32 NumPoints, TArray<TVector<TReal>>& OutPoints)
+{
+	check(ControlPoints);
+	check(NumPoints >= 2);
+
+	// var q is the change in t between successive evaluations.
+	const TReal q = 1.f / (TReal)(NumPoints - 1); // q is dependent on the number of GAPS = POINTS-1
+
+	// recreate the names used in the derivation
+	const TVector<TReal>& P0 = ControlPoints[0];
+	const TVector<TReal>& P1 = ControlPoints[1];
+	const TVector<TReal>& P2 = ControlPoints[2];
+	const TVector<TReal>& P3 = ControlPoints[3];
+
+	// coefficients of the cubic polynomial that we're FDing -
+	const TVector<TReal> a = P0;
+	const TVector<TReal> b = 3 * (P1 - P0);
+	const TVector<TReal> c = 3 * (P2 - 2 * P1 + P0);
+	const TVector<TReal> d = P3 - 3 * P2 + 3 * P1 - P0;
+
+	// initial values of the poly and the 3 diffs -
+	TVector<TReal> S = a;						// the poly value
+	TVector<TReal> U = b * q + c * q * q + d * q * q * q;	// 1st order diff (quadratic)
+	TVector<TReal> V = 2 * c * q * q + 6 * d * q * q * q;	// 2nd order diff (linear)
+	TVector<TReal> W = 6 * d * q * q * q;				// 3rd order diff (constant)
+
+	// Path length.
+	TReal Length = 0;
+
+	TVector<TReal> OldPos = P0;
+	OutPoints.Add(P0);	// first point on the curve is always P0.
+
+	for (int32 i = 1; i < NumPoints; ++i)
+	{
+		// calculate the next value and update the deltas
+		S += U;			// update poly value
+		U += V;			// update 1st order diff value
+		V += W;			// update 2st order diff value
+		// 3rd order diff is constant => no update needed.
+
+		// Update Length.
+		Length += TVector<TReal>::Dist(S, OldPos);
+		OldPos = S;
+
+		OutPoints.Add(S);
+	}
+
+	// Return path length as experienced in sequence (linear interpolation between points).
+	return Length;
+}
+
+template<typename TReal>
+void TVector<TReal>::CreateOrthonormalBasis(TVector<TReal>& XAxis, TVector<TReal>& YAxis, TVector<TReal>& ZAxis)
+{
+	// Project the X and Y axes onto the plane perpendicular to the Z axis.
+	XAxis -= (XAxis | ZAxis) / (ZAxis | ZAxis) * ZAxis;
+	YAxis -= (YAxis | ZAxis) / (ZAxis | ZAxis) * ZAxis;
+
+	// If the X axis was parallel to the Z axis, choose a vector which is orthogonal to the Y and Z axes.
+	if (XAxis.SizeSquared() < DELTA * DELTA)
+	{
+		XAxis = YAxis ^ ZAxis;
+	}
+
+	// If the Y axis was parallel to the Z axis, choose a vector which is orthogonal to the X and Z axes.
+	if (YAxis.SizeSquared() < DELTA * DELTA)
+	{
+		YAxis = XAxis ^ ZAxis;
+	}
+
+	// Normalize the basis vectors.
+	XAxis.Normalize();
+	YAxis.Normalize();
+	ZAxis.Normalize();
+}
+
+template<typename TReal>
+void TVector<TReal>::UnwindEuler()
+{
+	X = FMath::UnwindDegrees(X);
+	Y = FMath::UnwindDegrees(Y);
+	Z = FMath::UnwindDegrees(Z);
+}
+
+template<typename TReal>
+void TVector<TReal>::FindBestAxisVectors(TVector<TReal>& Axis1, TVector<TReal>& Axis2) const
+{
+	const TReal NX = FMath::Abs(X);
+	const TReal NY = FMath::Abs(Y);
+	const TReal NZ = FMath::Abs(Z);
+
+	// Find best basis vectors.
+	if (NZ > NX && NZ > NY)	Axis1 = TVector<TReal>(1, 0, 0);
+	else					Axis1 = TVector<TReal>(0, 0, 1);
+
+	TVector<TReal> Tmp = Axis1 - *this * (Axis1 | *this);
+	Axis1 = Tmp.GetSafeNormal();
+	Axis2 = Axis1 ^ *this;
+}
+
+template<typename TReal>
+FORCEINLINE bool TVector<TReal>::ContainsNaN() const
+{
+	return (!FMath::IsFinite(X) ||
+		!FMath::IsFinite(Y) ||
+		!FMath::IsFinite(Z));
+}
+
+template<typename TReal>
+FORCEINLINE FString TVector<TReal>::ToString() const
+{
+	return FString::Printf(TEXT("X=%3.3f Y=%3.3f Z=%3.3f"), X, Y, Z);
+}
+
+template<typename TReal>
+FORCEINLINE FText TVector<TReal>::ToText() const
+{
+	FFormatNamedArguments Args;
+	Args.Add(TEXT("X"), X);
+	Args.Add(TEXT("Y"), Y);
+	Args.Add(TEXT("Z"), Z);
+
+	return FText::Format(NSLOCTEXT("Core", "Vector3", "X={X} Y={Y} Z={Z}"), Args);
+}
+
+template<typename TReal>
+FORCEINLINE FText TVector<TReal>::ToCompactText() const
+{
+	if (IsNearlyZero())
+	{
+		return NSLOCTEXT("Core", "Vector3_CompactZeroVector", "V(0)");
+	}
+
+	const bool XIsNotZero = !FMath::IsNearlyZero(X);
+	const bool YIsNotZero = !FMath::IsNearlyZero(Y);
+	const bool ZIsNotZero = !FMath::IsNearlyZero(Z);
+
+	FNumberFormattingOptions FormatRules;
+	FormatRules.MinimumFractionalDigits = 2;
+	FormatRules.MinimumIntegralDigits = 0;
+
+	FFormatNamedArguments Args;
+	Args.Add(TEXT("X"), FText::AsNumber(X, &FormatRules));
+	Args.Add(TEXT("Y"), FText::AsNumber(Y, &FormatRules));
+	Args.Add(TEXT("Z"), FText::AsNumber(Z, &FormatRules));
+
+	if (XIsNotZero && YIsNotZero && ZIsNotZero)
+	{
+		return FText::Format(NSLOCTEXT("Core", "Vector3_CompactXYZ", "V(X={X}, Y={Y}, Z={Z})"), Args);
+	}
+	else if (!XIsNotZero && YIsNotZero && ZIsNotZero)
+	{
+		return FText::Format(NSLOCTEXT("Core", "Vector3_CompactYZ", "V(Y={Y}, Z={Z})"), Args);
+	}
+	else if (XIsNotZero && !YIsNotZero && ZIsNotZero)
+	{
+		return FText::Format(NSLOCTEXT("Core", "Vector3_CompactXZ", "V(X={X}, Z={Z})"), Args);
+	}
+	else if (XIsNotZero && YIsNotZero && !ZIsNotZero)
+	{
+		return FText::Format(NSLOCTEXT("Core", "Vector3_CompactXY", "V(X={X}, Y={Y})"), Args);
+	}
+	else if (!XIsNotZero && !YIsNotZero && ZIsNotZero)
+	{
+		return FText::Format(NSLOCTEXT("Core", "Vector3_CompactZ", "V(Z={Z})"), Args);
+	}
+	else if (XIsNotZero && !YIsNotZero && !ZIsNotZero)
+	{
+		return FText::Format(NSLOCTEXT("Core", "Vector3_CompactX", "V(X={X})"), Args);
+	}
+	else if (!XIsNotZero && YIsNotZero && !ZIsNotZero)
+	{
+		return FText::Format(NSLOCTEXT("Core", "Vector3_CompactY", "V(Y={Y})"), Args);
+	}
+
+	return NSLOCTEXT("Core", "Vector3_CompactZeroVector", "V(0)");
+}
+
+template<typename TReal>
+FORCEINLINE FString TVector<TReal>::ToCompactString() const
+{
+	if (IsNearlyZero())
+	{
+		return FString::Printf(TEXT("V(0)"));
+	}
+
+	FString ReturnString(TEXT("V("));
+	bool bIsEmptyString = true;
+	if (!FMath::IsNearlyZero(X))
+	{
+		ReturnString += FString::Printf(TEXT("X=%.2f"), X);
+		bIsEmptyString = false;
+	}
+	if (!FMath::IsNearlyZero(Y))
+	{
+		if (!bIsEmptyString)
+		{
+			ReturnString += FString(TEXT(", "));
+		}
+		ReturnString += FString::Printf(TEXT("Y=%.2f"), Y);
+		bIsEmptyString = false;
+	}
+	if (!FMath::IsNearlyZero(Z))
+	{
+		if (!bIsEmptyString)
+		{
+			ReturnString += FString(TEXT(", "));
+		}
+		ReturnString += FString::Printf(TEXT("Z=%.2f"), Z);
+		bIsEmptyString = false;
+	}
+	ReturnString += FString(TEXT(")"));
+	return ReturnString;
+}
+
+template<typename TReal>
+FORCEINLINE bool TVector<TReal>::InitFromString(const FString& InSourceString)
+{
+	X = Y = Z = 0;
+
+	// The initialization is only successful if the X, Y, and Z values can all be parsed from the string
+	const bool bSuccessful = FParse::Value(*InSourceString, TEXT("X="), X) && FParse::Value(*InSourceString, TEXT("Y="), Y) && FParse::Value(*InSourceString, TEXT("Z="), Z);
+
+	return bSuccessful;
+}
+
+template<typename TReal>
+FORCEINLINE FVector2D TVector<TReal>::UnitCartesianToSpherical() const
+{
+	checkSlow(IsUnit());
+	const TReal Theta = FMath::Acos(Z / Size());
+	const TReal Phi = FMath::Atan2(Y, X);
+	return FVector2D(Theta, Phi);
+}
+
+template<typename TReal>
+FORCEINLINE TReal TVector<TReal>::HeadingAngle() const
+{
+	// Project Dir into Z plane.
+	TVector<TReal> PlaneDir = *this;
+	PlaneDir.Z = 0.f;
+	PlaneDir = PlaneDir.GetSafeNormal();
+
+	TReal Angle = FMath::Acos(PlaneDir.X);
+
+	if (PlaneDir.Y < 0.0f)
+	{
+		Angle *= -1.0f;
+	}
+
+	return Angle;
+}
+
+
+
+template<typename TReal>
+FORCEINLINE TReal TVector<TReal>::Dist(const TVector<TReal>& V1, const TVector<TReal>& V2)
+{
+	return FMath::Sqrt(TVector<TReal>::DistSquared(V1, V2));
+}
+
+template<typename TReal>
+FORCEINLINE TReal TVector<TReal>::DistXY(const TVector<TReal>& V1, const TVector<TReal>& V2)
+{
+	return FMath::Sqrt(TVector<TReal>::DistSquaredXY(V1, V2));
+}
+
+template<typename TReal>
+FORCEINLINE TReal TVector<TReal>::DistSquared(const TVector<TReal>& V1, const TVector<TReal>& V2)
+{
+	return FMath::Square(V2.X - V1.X) + FMath::Square(V2.Y - V1.Y) + FMath::Square(V2.Z - V1.Z);
+}
+
+template<typename TReal>
+FORCEINLINE TReal TVector<TReal>::DistSquaredXY(const TVector<TReal>& V1, const TVector<TReal>& V2)
+{
+	return FMath::Square(V2.X - V1.X) + FMath::Square(V2.Y - V1.Y);
+}
+
+template<typename TReal>
+FORCEINLINE TReal TVector<TReal>::BoxPushOut(const TVector<TReal>& Normal, const TVector<TReal>& Size)
+{
+	return FMath::Abs(Normal.X * Size.X) + FMath::Abs(Normal.Y * Size.Y) + FMath::Abs(Normal.Z * Size.Z);
+}
+
+template<typename TReal>
+FORCEINLINE TVector<TReal> TVector<TReal>::Min(const TVector<TReal>& A, const TVector<TReal>& B)
+{
+	return TVector<TReal>(
+		FMath::Min(A.X, B.X),
+		FMath::Min(A.Y, B.Y),
+		FMath::Min(A.Z, B.Z)
+	);
+}
+
+template<typename TReal>
+FORCEINLINE TVector<TReal> TVector<TReal>::Max(const TVector<TReal>& A, const TVector<TReal>& B)
+{
+	return TVector<TReal>(
+		FMath::Max(A.X, B.X),
+		FMath::Max(A.Y, B.Y),
+		FMath::Max(A.Z, B.Z)
+	);
+}
+
+template<typename TReal>
+FORCEINLINE TVector<TReal> TVector<TReal>::Min3(const TVector<TReal>& A, const TVector<TReal>& B, const TVector<TReal>& C)
+{
+	return TVector<TReal>(
+		FMath::Min3(A.X, B.X, C.X),
+		FMath::Min3(A.Y, B.Y, C.Y),
+		FMath::Min3(A.Z, B.Z, C.Z)
+	);
+}
+
+template<typename TReal>
+FORCEINLINE TVector<TReal> TVector<TReal>::Max3(const TVector<TReal>& A, const TVector<TReal>& B, const TVector<TReal>& C)
+{
+	return TVector<TReal>(
+		FMath::Max3(A.X, B.X, C.X),
+		FMath::Max3(A.Y, B.Y, C.Y),
+		FMath::Max3(A.Z, B.Z, C.Z)
+	);
+}
+
+} // namespace UE::Core
+} // namespace UE
+
+
+/* Typed declarations
+ *****************************************************************************/
+
+// DO NOT USE! Large World Coordinate placeholder only.
+struct FVector3f : public UE::Core::TVector<float>
+{
+	using TVector::TVector;
+	FVector3f() : UE::Core::TVector<float>() {}
+	FVector3f(const UE::Core::TVector<float>& Vec) : UE::Core::TVector<float>(Vec.X, Vec.Y, Vec.Z) {}
+
+
+	/** A zero vector (0,0,0) */
+	static CORE_API const FVector3f ZeroVector;
+
+	/** One vector (1,1,1) */
+	static CORE_API const FVector3f OneVector;
+
+	/** Unreal up vector (0,0,1) */
+	static CORE_API const FVector3f UpVector;
+
+	/** Unreal down vector (0,0,-1) */
+	static CORE_API const FVector3f DownVector;
+
+	/** Unreal forward vector (1,0,0) */
+	static CORE_API const FVector3f ForwardVector;
+
+	/** Unreal backward vector (-1,0,0) */
+	static CORE_API const FVector3f BackwardVector;
+
+	/** Unreal right vector (0,1,0) */
+	static CORE_API const FVector3f RightVector;
+
+	/** Unreal left vector (0,-1,0) */
+	static CORE_API const FVector3f LeftVector;
+
+	/** Unit X axis vector (1,0,0) */
+	static CORE_API const FVector3f XAxisVector;
+
+	/** Unit Y axis vector (0,1,0) */
+	static CORE_API const FVector3f YAxisVector;
+
+	/** Unit Z axis vector (0,0,1) */
+	static CORE_API const FVector3f ZAxisVector;
+};
+//using FVector3f = UE::Core::TVector<float>;
+template<> struct TCanBulkSerialize<FVector3f> { enum { Value = false }; };
+template<> struct TIsPODType<FVector3f> { enum { Value = true }; };
+
+// DO NOT USE! Large World Coordinate placeholder only.
+struct FVector3d : public UE::Core::TVector<double>
+{
+	using TVector::TVector;
+	FVector3d() : UE::Core::TVector<double>() {}
+	FVector3d(const UE::Core::TVector<double>& Vec) : UE::Core::TVector<double>(Vec.X, Vec.Y, Vec.Z) {}
+
+
+	/** A zero vector (0,0,0) */
+	static CORE_API const FVector3d ZeroVector;
+
+	/** One vector (1,1,1) */
+	static CORE_API const FVector3d OneVector;
+
+	/** Unreal up vector (0,0,1) */
+	static CORE_API const FVector3d UpVector;
+
+	/** Unreal down vector (0,0,-1) */
+	static CORE_API const FVector3d DownVector;
+
+	/** Unreal forward vector (1,0,0) */
+	static CORE_API const FVector3d ForwardVector;
+
+	/** Unreal backward vector (-1,0,0) */
+	static CORE_API const FVector3d BackwardVector;
+
+	/** Unreal right vector (0,1,0) */
+	static CORE_API const FVector3d RightVector;
+
+	/** Unreal left vector (0,-1,0) */
+	static CORE_API const FVector3d LeftVector;
+
+	/** Unit X axis vector (1,0,0) */
+	static CORE_API const FVector3d XAxisVector;
+
+	/** Unit Y axis vector (0,1,0) */
+	static CORE_API const FVector3d YAxisVector;
+
+	/** Unit Z axis vector (0,0,1) */
+	static CORE_API const FVector3d ZAxisVector;
+};
+//using FVector3d = UE::Core::TVector<double>;
+template<> struct TCanBulkSerialize<FVector3d> { enum { Value = false }; };
+template<> struct TIsPODType<FVector3d> { enum { Value = true }; };
+
+// Forward declare all explicit specializations (defined in VectorLWC.cpp)
+template<> CORE_API UE::Core::TVector<float>::TVector(const FVector4&);
+template<> CORE_API UE::Core::TVector<double>::TVector(const FVector4&);
+template<> CORE_API FRotator UE::Core::TVector<float>::ToOrientationRotator() const;
+template<> CORE_API FRotator UE::Core::TVector<double>::ToOrientationRotator() const;
+template<> CORE_API FQuat UE::Core::TVector<float>::ToOrientationQuat() const;
+template<> CORE_API FQuat UE::Core::TVector<double>::ToOrientationQuat() const;
+template<> CORE_API FRotator UE::Core::TVector<float>::Rotation() const;
+template<> CORE_API FRotator UE::Core::TVector<double>::Rotation() const;
+
+/**
+ * Multiplies a vector by a scaling factor.
+ *
+ * @param Scale Scaling factor.
+ * @param V Vector to scale.
+ * @return Result of multiplication.
+ */
+FORCEINLINE FVector3f operator*(float Scale, const FVector3f& V)
+{
+	return V.operator*(Scale);
+}
+FORCEINLINE FVector3d operator*(double Scale, const FVector3d& V)
+{
+	return V.operator*(Scale);
+}
+
+/**
+ * Creates a hash value from an FVector.
+ *
+ * @param Vector the vector to create a hash value for
+ * @return The hash value from the components
+ */
+FORCEINLINE uint32 GetTypeHash(const FVector3f& Vector)
+{
+	// Note: this assumes there's no padding in Vector that could contain uncompared data.
+	return FCrc::MemCrc_DEPRECATED(&Vector, sizeof(Vector));
+}
+FORCEINLINE uint32 GetTypeHash(const FVector3d& Vector)
+{
+	// Note: this assumes there's no padding in Vector that could contain uncompared data.
+	return FCrc::MemCrc_DEPRECATED(&Vector, sizeof(Vector));
+}
+
+#ifdef _MSC_VER
+#pragma warning (pop)
+#endif
\ No newline at end of file
