Extend TVector with functions necessary for GeometryProcessing FVector3d/3f replacement.
- add casts to/from FVector, and between FVector3d and FVector3f
- add .Dot(), .Cross(), .Length(), .SquaredLength()
- add operator*(Real, TVector<Real>)
- add Zero(), One(), UnitX(), UnitY(), UnitZ() to FVector3f/3d
#rnx

diff --git a/Engine/Source/Runtime/Core/Public/Math/VectorLWC.h b/Engine/Source/Runtime/Core/Public/Math/VectorLWC.h
--- a/Engine/Source/Runtime/Core/Public/Math/VectorLWC.h
+++ b/Engine/Source/Runtime/Core/Public/Math/VectorLWC.h
@@ -24,6 +24,7 @@
 #include "Math/Color.h"
 #include "Math/IntPoint.h"
 #include "Logging/LogMacros.h"
+#include "Math/Vector.h"
 #include "Math/Vector2D.h"
 #include "Misc/ByteSwap.h"
 #include "Internationalization/Text.h"
@@ -162,6 +163,12 @@ public:
 	FORCEINLINE TVector<TReal>& operator=(const TVector<TReal>& Other);
 #endif
 
+	/**
+	 * Cast to FVector, however that type is defined.
+	 * TODO: FVector is currently a float but this should be changed to whatever the compiled FVector type is
+	 */
+	explicit FORCEINLINE operator FVector() const;
+
 	/**
 	 * Calculate cross product between this and another vector.
 	 *
@@ -170,6 +177,15 @@ public:
 	 */
 	FORCEINLINE TVector<TReal> operator^(const TVector<TReal>& V) const;
 
+
+	/**
+	 * Calculate cross product between this and another vector.
+	 *
+	 * @param V The other vector.
+	 * @return The cross product.
+	 */
+	FORCEINLINE TVector<TReal> Cross(const TVector<TReal>& V2) const;
+
 	/**
 	 * Calculate the cross product of two vectors.
 	 *
@@ -187,6 +203,14 @@ public:
 	 */
 	FORCEINLINE TReal operator|(const TVector<TReal>& V) const;
 
+	/**
+	 * Calculate the dot product between this and another vector.
+	 *
+	 * @param V The other vector.
+	 * @return The dot product.
+	 */
+	FORCEINLINE TReal Dot(const TVector<TReal>& V) const;
+
 	/**
 	 * Calculate the dot product of two vectors.
 	 *
@@ -483,6 +507,13 @@ public:
 	 */
 	TReal Size() const;
 
+	/**
+	 * Get the length (magnitude) of this vector.
+	 *
+	 * @return The length of this vector.
+	 */
+	TReal Length() const;
+
 	/**
 	 * Get the squared length of this vector.
 	 *
@@ -490,6 +521,13 @@ public:
 	 */
 	TReal SizeSquared() const;
 
+	/**
+	 * Get the squared length of this vector.
+	 *
+	 * @return The squared length of this vector.
+	 */
+	TReal SquaredLength() const;
+
 	/**
 	 * Get the length of the 2D components of this vector.
 	 *
@@ -1268,6 +1306,12 @@ FORCEINLINE TVector<TReal>& TVector<TReal>::operator=(const TVector<TReal>& Othe
 }
 #endif
 
+template<typename TReal>
+FORCEINLINE TVector<TReal>::operator FVector() const
+{
+	return FVector((float)X, (float)Y, (float)Z);
+}
+
 template<typename TReal>
 FORCEINLINE TVector<TReal> TVector<TReal>::operator^(const TVector<TReal>& V) const
 {
@@ -1279,6 +1323,12 @@ FORCEINLINE TVector<TReal> TVector<TReal>::operator^(const TVector<TReal>& V) co
 	);
 }
 
+template<typename TReal>
+FORCEINLINE TVector<TReal> TVector<TReal>::Cross(const TVector<TReal>& V) const
+{
+	return *this ^ V;
+}
+
 template<typename TReal>
 FORCEINLINE TVector<TReal> TVector<TReal>::CrossProduct(const TVector<TReal>& A, const TVector<TReal>& B)
 {
@@ -1291,6 +1341,12 @@ FORCEINLINE TReal TVector<TReal>::operator|(const TVector<TReal>& V) const
 	return X * V.X + Y * V.Y + Z * V.Z;
 }
 
+template<typename TReal>
+FORCEINLINE TReal TVector<TReal>::Dot(const TVector<TReal>& V) const
+{
+	return *this | V;
+}
+
 template<typename TReal>
 FORCEINLINE TReal TVector<TReal>::DotProduct(const TVector<TReal>& A, const TVector<TReal>& B)
 {
@@ -1456,12 +1512,24 @@ FORCEINLINE TReal TVector<TReal>::Size() const
 	return FMath::Sqrt(X * X + Y * Y + Z * Z);
 }
 
+template<typename TReal>
+FORCEINLINE TReal TVector<TReal>::Length() const
+{
+	return Size();
+}
+
 template<typename TReal>
 FORCEINLINE TReal TVector<TReal>::SizeSquared() const
 {
 	return X * X + Y * Y + Z * Z;
 }
 
+template<typename TReal>
+FORCEINLINE TReal TVector<TReal>::SquaredLength() const
+{
+	return SizeSquared();
+}
+
 template<typename TReal>
 FORCEINLINE TReal TVector<TReal>::Size2D() const
 {
@@ -2229,9 +2297,17 @@ FORCEINLINE TVector<TReal> TVector<TReal>::Max3(const TVector<TReal>& A, const T
 	);
 }
 
+template<typename TReal>
+FORCEINLINE TVector<TReal> operator*(TReal Scale, const TVector<TReal>& V)
+{
+	return V.operator*(Scale);
+}
+
 } // namespace UE::Core
 } // namespace UE
 
+// predeclare FVector3d, necessary for conversion operators in FVector3f
+struct FVector3d;
 
 /* Typed declarations
  *****************************************************************************/
@@ -2243,6 +2319,14 @@ struct FVector3f : public UE::Core::TVector<float>
 	FVector3f() : UE::Core::TVector<float>() {}
 	FVector3f(const UE::Core::TVector<float>& Vec) : UE::Core::TVector<float>(Vec.X, Vec.Y, Vec.Z) {}
 
+	/** Construct from double vector */
+	explicit FVector3f(const UE::Core::TVector<double>& Vec) : UE::Core::TVector<float>((float)Vec.X, (float)Vec.Y, (float)Vec.Z) {}
+
+	/** Construct from FVector, regardless of what type it is defined as */
+	explicit FVector3f(const FVector& Vec) : UE::Core::TVector<float>(Vec.X, Vec.Y, Vec.Z) {}
+
+	/** @return cast to double-precision FVector3d */
+	explicit operator FVector3d() const;
 
 	/** A zero vector (0,0,0) */
 	static CORE_API const FVector3f ZeroVector;
@@ -2276,6 +2360,21 @@ struct FVector3f : public UE::Core::TVector<float>
 
 	/** Unit Z axis vector (0,0,1) */
 	static CORE_API const FVector3f ZAxisVector;
+
+	/** @return Zero Vector (0,0,0) */
+	static FVector3f Zero() { return ZeroVector; }
+
+	/** @return One Vector (1,1,1) */
+	static FVector3f One() { return OneVector; }
+
+	/** @return Unit X Vector (1,0,0)  */
+	static FVector3f UnitX() { return XAxisVector; }
+
+	/** @return Unit Y Vector (0,1,0)  */
+	static FVector3f UnitY() { return YAxisVector; }
+
+	/** @return Unit Z Vector (0,0,1)  */
+	static FVector3f UnitZ() { return ZAxisVector; }
 };
 //using FVector3f = UE::Core::TVector<float>;
 template<> struct TCanBulkSerialize<FVector3f> { enum { Value = false }; };
@@ -2288,6 +2387,17 @@ struct FVector3d : public UE::Core::TVector<double>
 	FVector3d() : UE::Core::TVector<double>() {}
 	FVector3d(const UE::Core::TVector<double>& Vec) : UE::Core::TVector<double>(Vec.X, Vec.Y, Vec.Z) {}
 
+	/** Construct from float vector */
+	explicit FVector3d(const UE::Core::TVector<float>& Vec) : UE::Core::TVector<double>((double)Vec.X, (double)Vec.Y, (double)Vec.Z) {}
+
+	/** Construct from FVector, regardless of what type it is defined as */
+	explicit FVector3d(const FVector& Vec) : UE::Core::TVector<double>((double)Vec.X, (double)Vec.Y, (double)Vec.Z) {}
+
+	/** @return cast to single-precision FVector3f */
+	explicit operator FVector3f() const
+	{
+		return FVector3f((float)X, (float)Y, (float)Z);
+	}
 
 	/** A zero vector (0,0,0) */
 	static CORE_API const FVector3d ZeroVector;
@@ -2321,6 +2431,21 @@ struct FVector3d : public UE::Core::TVector<double>
 
 	/** Unit Z axis vector (0,0,1) */
 	static CORE_API const FVector3d ZAxisVector;
+
+	/** @return Zero Vector (0,0,0) */
+	static FVector3d Zero() { return ZeroVector; }
+
+	/** @return One Vector (1,1,1) */
+	static FVector3d One() { return OneVector; }
+
+	/** @return Unit X Vector (1,0,0)  */
+	static FVector3d UnitX() { return XAxisVector; }
+
+	/** @return Unit Y Vector (0,1,0)  */
+	static FVector3d UnitY() { return YAxisVector; }
+
+	/** @return Unit Z Vector (0,0,1)  */
+	static FVector3d UnitZ() { return ZAxisVector; }
 };
 //using FVector3d = UE::Core::TVector<double>;
 template<> struct TCanBulkSerialize<FVector3d> { enum { Value = false }; };
@@ -2336,6 +2461,13 @@ template<> CORE_API FQuat UE::Core::TVector<double>::ToOrientationQuat() const;
 template<> CORE_API FRotator UE::Core::TVector<float>::Rotation() const;
 template<> CORE_API FRotator UE::Core::TVector<double>::Rotation() const;
 
+
+inline FVector3f::operator FVector3d() const
+{ 
+	return FVector3d((double)X, (double)Y, (double)Z);
+}
+
+
 /**
  * Multiplies a vector by a scaling factor.
  *
