When checking for a TargetPlatform in HasCompiledSupportForPlatform() don't require compiled support if the TargetPlatform base modules doesn't exist either.  This prevents tools like UnrealHeaderTool, which doesn't actually need to load these platform modules but still must do some processing based on their existence from failing.
#rnx

diff --git a/Engine/Source/Runtime/Core/Private/Misc/DataDrivenPlatformInfoRegistry.cpp b/Engine/Source/Runtime/Core/Private/Misc/DataDrivenPlatformInfoRegistry.cpp
--- a/Engine/Source/Runtime/Core/Private/Misc/DataDrivenPlatformInfoRegistry.cpp
+++ b/Engine/Source/Runtime/Core/Private/Misc/DataDrivenPlatformInfoRegistry.cpp
@@ -506,7 +506,10 @@ bool FDataDrivenPlatformInfoRegistry::HasCompiledSupportForPlatform(FName Platfo
 	else if (PlatformNameType == EPlatformNameType::TargetPlatform)
 	{
 		// was this TP compiled, or a shaderformat (useful for SCW if it ever calls this)
-		return 
+		// if this is a program that doesn't use TargetPlatform, then we don't need to know if one particular 
+		// TargetPlatform is missing support compiled in, because no platforms are "compiled in"
+		return
+			!FModuleManager::Get().ModuleExists(TEXT("TargetPlatform")) ||
 			FModuleManager::Get().ModuleExists(*FString::Printf(TEXT("%sTargetPlatform"), *PlatformName.ToString())) || 
 			FModuleManager::Get().ModuleExists(*FString::Printf(TEXT("%sShaderFormat"), *PlatformName.ToString()));
 	}
