Deleted unused include

#rnx

diff --git a/Engine/Source/Runtime/Core/Private/ProfilingDebugging/Windows/CallstackTrace.cpp b/Engine/Source/Runtime/Core/Private/ProfilingDebugging/Windows/CallstackTrace.cpp
--- a/Engine/Source/Runtime/Core/Private/ProfilingDebugging/Windows/CallstackTrace.cpp
+++ b/Engine/Source/Runtime/Core/Private/ProfilingDebugging/Windows/CallstackTrace.cpp
@@ -7,7 +7,6 @@
 #include "HAL/CriticalSection.h"
 #include "HAL/MemoryBase.h"
 #include "Misc/ScopeRWLock.h"
-#include "Containers/CircularQueue.h"
 #include "HAL/RunnableThread.h"
 #include "HAL/LowLevelMemTracker.h"
 #include "ProfilingDebugging/MemoryTrace.h"
