Reduced scope of "using namespace" so it doesn't leak into other unity files.

#rnx

diff --git a/Engine/Source/Runtime/Core/Private/ProfilingDebugging/Windows/CallstackTrace.cpp b/Engine/Source/Runtime/Core/Private/ProfilingDebugging/Windows/CallstackTrace.cpp
--- a/Engine/Source/Runtime/Core/Private/ProfilingDebugging/Windows/CallstackTrace.cpp
+++ b/Engine/Source/Runtime/Core/Private/ProfilingDebugging/Windows/CallstackTrace.cpp
@@ -102,9 +102,6 @@ enum
 ////////////////////////////////////////////////////////////////////////////////
 namespace {
 
-	using namespace Experimental;
-	using namespace atomic_queue;
-
 	class FCallstackProcWorker : public FRunnable 
 	{
 		public:
@@ -127,9 +124,9 @@ namespace {
 		void AddCallstack(const FBacktraceEntry& Entry);
 		void AddWork(const FBacktraceEntry& Entry);
 
-		TSherwoodSet<uint64> 				KnownSet;
+		Experimental::TSherwoodSet<uint64> 	KnownSet;
 		FCriticalSection					ProducerCs;
-		AtomicQueue2<FBacktraceEntry, 256>	Queue;
+		atomic_queue::AtomicQueue2<FBacktraceEntry, 256>	Queue;
 		TAtomic<bool>						bRun;
 		TAtomic<bool>						bStarted;
 		static FRunnableThread*				Thread;
