Swapped TAtomic out for std::atomic

#rnx

diff --git a/Engine/Source/Runtime/Core/Private/ProfilingDebugging/Windows/CallstackTrace.cpp b/Engine/Source/Runtime/Core/Private/ProfilingDebugging/Windows/CallstackTrace.cpp
--- a/Engine/Source/Runtime/Core/Private/ProfilingDebugging/Windows/CallstackTrace.cpp
+++ b/Engine/Source/Runtime/Core/Private/ProfilingDebugging/Windows/CallstackTrace.cpp
@@ -22,6 +22,8 @@
 #include "HAL/PlatformProcess.h"
 #include "Misc/ScopeLock.h"
 
+#include <atomic>
+
 // 0=off, 1=stats, 2=validation, 3=truth_compare
 #define BACKTRACE_DBGLVL 0
 
@@ -127,8 +129,8 @@ namespace {
 		Experimental::TSherwoodSet<uint64> 	KnownSet;
 		FCriticalSection					ProducerCs;
 		atomic_queue::AtomicQueue2<FBacktraceEntry, 256>	Queue;
-		TAtomic<bool>						bRun;
-		TAtomic<bool>						bStarted;
+		std::atomic<bool>					bRun;
+		std::atomic<bool>					bStarted;
 		static FRunnableThread*				Thread;
 	};
 }
