Correct determination of the number of functions in a .pdata segment. There is no null-terminator to the list, instead there's null _padding_. If the end of the last RUNTIME_FUNCTION object happened to align to 2^8 then the old function-count loop would read undefined data. This also happens to be much faster which is a bonus.

#rnx

diff --git a/Engine/Source/Runtime/Core/Private/ProfilingDebugging/Windows/CallstackTrace.cpp b/Engine/Source/Runtime/Core/Private/ProfilingDebugging/Windows/CallstackTrace.cpp
--- a/Engine/Source/Runtime/Core/Private/ProfilingDebugging/Windows/CallstackTrace.cpp
+++ b/Engine/Source/Runtime/Core/Private/ProfilingDebugging/Windows/CallstackTrace.cpp
@@ -294,12 +294,14 @@ void FBacktracer::AddModule(UPTRINT ModuleBase, const TCHAR* Name)
 
 	// Find ".pdata" section
 	UPTRINT PdataBase = 0;
+	UPTRINT PdataEnd = 0;
 	for (uint32 i = 0; i < NumSections; ++i)
 	{
 		const IMAGE_SECTION_HEADER* Section = Sections + i;
 		if (*(uint64*)(Section->Name) == 0x61'74'61'64'70'2eull) // Sections names are eight bytes and zero padded. This constant is '.pdata'
 		{
 			PdataBase = ModuleBase + Section->VirtualAddress;
+			PdataEnd = PdataBase + Section->SizeOfRawData;
 			break;
 		}
 	}
@@ -309,15 +311,21 @@ void FBacktracer::AddModule(UPTRINT ModuleBase, const TCHAR* Name)
 		return;
 	}
 
-	// Count the number of functions.
+	// Count the number of functions. The assumption here is that if we have got this far then there is at least one function
+	uint32 NumFunctions = uint32(PdataEnd - PdataBase) / sizeof(RUNTIME_FUNCTION);
 	const auto* FunctionTables = (RUNTIME_FUNCTION*)PdataBase;
-	uint32 NumFunctions = 0;
 	do
 	{
-		++NumFunctions;
-	}
-	while (FunctionTables[NumFunctions].BeginAddress);
+		const RUNTIME_FUNCTION* Function = FunctionTables + NumFunctions - 1;
+		if (uint32(Function->BeginAddress) < uint32(Function->EndAddress))
+		{
+			break;
+		}
 
+		--NumFunctions;
+	}
+	while (NumFunctions != 0);
+		
 	// Allocate some space for the module's function-to-frame-size table
 	auto* OutTable = (FFunction*)Malloc->Malloc(sizeof(FFunction) * NumFunctions);
 	FFunction* OutTableCursor = OutTable;
