Add a double version of FInterpTo.

diff --git a/Engine/Source/Runtime/Core/Private/Math/UnrealMath.cpp b/Engine/Source/Runtime/Core/Private/Math/UnrealMath.cpp
--- a/Engine/Source/Runtime/Core/Private/Math/UnrealMath.cpp
+++ b/Engine/Source/Runtime/Core/Private/Math/UnrealMath.cpp
@@ -2488,6 +2488,29 @@ CORE_API float FMath::FInterpTo( float Current, float Target, float DeltaTime, f
 	return Current + DeltaMove;
 }
 
+CORE_API double FMath::FInterpTo( double Current, double Target, double DeltaTime, double InterpSpeed )
+{
+	// If no interp speed, jump to target value
+	if( InterpSpeed <= 0.0 )
+	{
+		return Target;
+	}
+
+	// Distance to reach
+	const double Dist = Target - Current;
+
+	// If distance is too small, just set the desired location
+	if( FMath::Square(Dist) < SMALL_NUMBER )
+	{
+		return Target;
+	}
+
+	// Delta Move, Clamp so we do not over shoot.
+	const double DeltaMove = Dist * FMath::Clamp<double>(DeltaTime * InterpSpeed, 0.0, 1.0);
+
+	return Current + DeltaMove;
+}
+
 CORE_API float FMath::FInterpConstantTo( float Current, float Target, float DeltaTime, float InterpSpeed )
 {
 	const float Dist = Target - Current;
diff --git a/Engine/Source/Runtime/Core/Public/Math/UnrealMathUtility.h b/Engine/Source/Runtime/Core/Public/Math/UnrealMathUtility.h
--- a/Engine/Source/Runtime/Core/Public/Math/UnrealMathUtility.h
+++ b/Engine/Source/Runtime/Core/Public/Math/UnrealMathUtility.h
@@ -1133,6 +1133,9 @@ public:
 	/** Interpolate float from Current to Target. Scaled by distance to Target, so it has a strong start speed and ease out. */
 	static CORE_API float FInterpTo( float Current, float Target, float DeltaTime, float InterpSpeed );
 
+	/** Interpolate double from Current to Target. Scaled by distance to Target, so it has a strong start speed and ease out. */
+	static CORE_API double FInterpTo(double Current, double Target, double DeltaTime, double InterpSpeed);
+
 	/** Interpolate Linear Color from Current to Target. Scaled by distance to Target, so it has a strong start speed and ease out. */
 	static CORE_API FLinearColor CInterpTo(const FLinearColor& Current, const FLinearColor& Target, float DeltaTime, float InterpSpeed);
 
diff --git a/Engine/Source/Runtime/Slate/Private/Widgets/Views/STableViewBase.cpp b/Engine/Source/Runtime/Slate/Private/Widgets/Views/STableViewBase.cpp
--- a/Engine/Source/Runtime/Slate/Private/Widgets/Views/STableViewBase.cpp
+++ b/Engine/Source/Runtime/Slate/Private/Widgets/Views/STableViewBase.cpp
@@ -257,7 +257,7 @@ void STableViewBase::Tick( const FGeometry& AllottedGeometry, const double InCur
 
 			if (bEnableAnimatedScrolling)
 			{
-				CurrentScrollOffset = FMath::FInterpTo(CurrentScrollOffset, TargetScrollOffset, InDeltaTime, 12.f);
+				CurrentScrollOffset = FMath::FInterpTo(CurrentScrollOffset, TargetScrollOffset, (double)InDeltaTime, 12.0);
 				if (FMath::IsNearlyEqual(CurrentScrollOffset, TargetScrollOffset, 0.01))
 				{
 					CurrentScrollOffset = TargetScrollOffset;
