Roll back memory reservation when writing memory images.

- Spikes of up to 8.9GB have been reported, and the CPU savings aren't worth that.

diff --git a/Engine/Source/Runtime/Core/Public/Serialization/MemoryImage.h b/Engine/Source/Runtime/Core/Public/Serialization/MemoryImage.h
--- a/Engine/Source/Runtime/Core/Public/Serialization/MemoryImage.h
+++ b/Engine/Source/Runtime/Core/Public/Serialization/MemoryImage.h
@@ -270,9 +270,7 @@ public:
 	FMemoryImageSection* AllocateSection()
 	{
 		FMemoryImageSection* Section = new FMemoryImageSection(this);
-#if WITH_EDITOR
-		Section->Bytes.Reserve(512 * 1024);	// reserve at least 512KB per section to avoid resizing
-#endif
+		// reserving memory here could reduce the reallocations, but leads to huge spikes for images with many sections. TODO: try chunked array or a better heuristic value for reservation
 		Sections.Add(Section);
 		return Section;
 	}
