Added string builder assignment operator and overloads for FSimpleParse and FParse::LineExtended

diff --git a/Engine/Source/Runtime/Core/Private/Misc/Parse.cpp b/Engine/Source/Runtime/Core/Private/Misc/Parse.cpp
--- a/Engine/Source/Runtime/Core/Private/Misc/Parse.cpp
+++ b/Engine/Source/Runtime/Core/Private/Misc/Parse.cpp
@@ -1079,7 +1079,8 @@ bool FParse::Line(const TCHAR** Stream, FString& Result, bool bExact)
 	return **Stream != TEXT('\0') || bGotStream;
 }
 
-bool FParse::LineExtended(const TCHAR** Stream, FString& Result, int32& LinesConsumed, bool bExact)
+template<class T>
+bool ParseLineExtended(const TCHAR** Stream, T& Result, int32& LinesConsumed, bool bExact)
 {
 	bool bGotStream = false;
 	bool bIsQuoted = false;
@@ -1110,7 +1111,7 @@ bool FParse::LineExtended(const TCHAR** Stream, FString& Result, int32& LinesCon
 		{
 			checkSlow(BracketDepth > 0);
 
-			Result.AppendChar(TEXT(' '));
+			Result += TEXT(' ');
 			LinesConsumed++;
 			(*Stream)++;
 			if (**Stream == TEXT('\n') || **Stream == TEXT('\r'))
@@ -1121,7 +1122,7 @@ bool FParse::LineExtended(const TCHAR** Stream, FString& Result, int32& LinesCon
 		// allow line break if the end of the line is a backslash
 		else if (!bIsQuoted && (*Stream)[0] == TEXT('\\') && ((*Stream)[1] == TEXT('\n') || (*Stream)[1] == TEXT('\r')))
 		{
-			Result.AppendChar(TEXT(' '));
+			Result += TEXT(' ');
 			LinesConsumed++;
 			(*Stream) += 2;
 			if (**Stream == TEXT('\n') || **Stream == TEXT('\r'))
@@ -1143,7 +1144,7 @@ bool FParse::LineExtended(const TCHAR** Stream, FString& Result, int32& LinesCon
 		// specifically consume escaped backslashes and quotes within quoted strings
 		else if (bIsQuoted && !bIgnore && (*Stream)[0] == TEXT('\\') && ( (*Stream)[1] == TEXT('\"') || (*Stream)[1] == TEXT('\\') ))
 		{
-			Result.AppendChars(*Stream, 2);
+			Result += FStringView(*Stream, 2);
 			(*Stream) += 2;
 		}
 		else
@@ -1153,7 +1154,7 @@ bool FParse::LineExtended(const TCHAR** Stream, FString& Result, int32& LinesCon
 			// Got stuff.
 			if (!bIgnore)
 			{
-				Result.AppendChar(*((*Stream)++));
+				Result += *((*Stream)++);
 			}
 			else
 			{
@@ -1204,6 +1205,16 @@ bool FParse::LineExtended(const TCHAR** Stream, FString& Result, int32& LinesCon
 	return **Stream != TEXT('\0') || bGotStream;
 }
 
+bool FParse::LineExtended(const TCHAR** Stream, FString& Result, int32& LinesConsumed, bool bExact)
+{
+	return ParseLineExtended(Stream, Result, LinesConsumed, bExact);
+}
+
+bool FParse::LineExtended(const TCHAR** Stream, FStringBuilderBase& Result, int32& LinesConsumed, bool bExact)
+{
+	return ParseLineExtended(Stream, Result, LinesConsumed, bExact);
+}
+
 uint32 FParse::HexNumber(const TCHAR* HexString)
 {
 	uint32 Ret = 0;
diff --git a/Engine/Source/Runtime/Core/Private/Modules/SimpleParse.cpp b/Engine/Source/Runtime/Core/Private/Modules/SimpleParse.cpp
--- a/Engine/Source/Runtime/Core/Private/Modules/SimpleParse.cpp
+++ b/Engine/Source/Runtime/Core/Private/Modules/SimpleParse.cpp
@@ -1,6 +1,7 @@
 // Copyright Epic Games, Inc. All Rights Reserved.
 
 #include "Modules/SimpleParse.h"
+#include "Misc/StringBuilder.h"
 
 bool FSimpleParse::MatchZeroOrMoreWhitespace(const TCHAR*& InOutPtr)
 {
@@ -25,7 +26,8 @@ bool FSimpleParse::MatchChar(const TCHAR*& InOutPtr, TCHAR Ch)
 	return true;
 }
 
-bool FSimpleParse::ParseString(const TCHAR*& InOutPtr, FString& OutStr)
+template<class T>
+bool SimpleParseString(const TCHAR*& InOutPtr, T& OutStr)
 {
 	const TCHAR* Ptr = InOutPtr;
 	if (*Ptr != '"')
@@ -78,6 +80,16 @@ bool FSimpleParse::ParseString(const TCHAR*& InOutPtr, FString& OutStr)
 	return true;
 }
 
+bool FSimpleParse::ParseString(const TCHAR*& InOutPtr, FString& OutStr)
+{
+	return SimpleParseString(InOutPtr, OutStr);
+}
+
+bool FSimpleParse::ParseString(const TCHAR*& InOutPtr, FStringBuilderBase& OutStr)
+{
+	return SimpleParseString(InOutPtr, OutStr);
+}
+
 bool FSimpleParse::ParseUnsignedNumber(const TCHAR*& InOutPtr, int32& OutNumber)
 {
 	const TCHAR* Ptr = InOutPtr;
diff --git a/Engine/Source/Runtime/Core/Private/Modules/SimpleParse.h b/Engine/Source/Runtime/Core/Private/Modules/SimpleParse.h
--- a/Engine/Source/Runtime/Core/Private/Modules/SimpleParse.h
+++ b/Engine/Source/Runtime/Core/Private/Modules/SimpleParse.h
@@ -10,5 +10,6 @@ struct FSimpleParse
 	static bool MatchZeroOrMoreWhitespace(const TCHAR*& InOutPtr);
 	static bool MatchChar(const TCHAR*& InOutPtr, TCHAR Ch);
 	static bool ParseString(const TCHAR*& InOutPtr, FString& OutStr);
+	static bool ParseString(const TCHAR*& InOutPtr, FStringBuilderBase& OutStr);
 	static bool ParseUnsignedNumber(const TCHAR*& InOutPtr, int32& OutNumber);
 };
diff --git a/Engine/Source/Runtime/Core/Public/Misc/Parse.h b/Engine/Source/Runtime/Core/Public/Misc/Parse.h
--- a/Engine/Source/Runtime/Core/Public/Misc/Parse.h
+++ b/Engine/Source/Runtime/Core/Public/Misc/Parse.h
@@ -57,6 +57,7 @@ struct CORE_API FParse
 	 * the out character array will not include the ignored endlines
 	 */
 	static bool LineExtended(const TCHAR** Stream, FString& Result, int32& LinesConsumed, bool Exact = 0);
+	static bool LineExtended(const TCHAR** Stream, FStringBuilderBase& Result, int32& LinesConsumed, bool Exact = 0);
 	/** Grabs the next space-delimited string from the input stream. If quoted, gets entire quoted string. */
 	static bool Token( const TCHAR*& Str, TCHAR* Result, int32 MaxLen, bool UseEscape );
 	/** Grabs the next space-delimited string from the input stream. If quoted, gets entire quoted string. */
diff --git a/Engine/Source/Runtime/Core/Public/Misc/StringBuilder.h b/Engine/Source/Runtime/Core/Public/Misc/StringBuilder.h
--- a/Engine/Source/Runtime/Core/Public/Misc/StringBuilder.h
+++ b/Engine/Source/Runtime/Core/Public/Misc/StringBuilder.h
@@ -72,6 +72,17 @@ public:
 	TStringBuilderBase& operator=(const TStringBuilderBase&) = delete;
 	TStringBuilderBase& operator=(TStringBuilderBase&&) = delete;
 
+	TStringBuilderBase& operator=(ViewType Str)
+	{
+		Reset();
+		return Append(Str);
+	}
+	
+	TStringBuilderBase& operator=(const CharType* Str)
+	{
+		return *this = ViewType(Str);
+	}
+
 	inline TStringBuilderBase(CharType* BufferPointer, int32 BufferCapacity)
 	{
 		Initialize(BufferPointer, BufferCapacity);
@@ -384,6 +395,8 @@ public:
 	{
 	}
 
+	using TStringBuilderBase<CharType>::operator=;
+
 private:
 	CharType StringBuffer[BufferSize];
 };
