Added FPathViews::Append overloads for FName and C strings

diff --git a/Engine/Source/Runtime/Core/Private/Misc/PathViews.cpp b/Engine/Source/Runtime/Core/Private/Misc/PathViews.cpp
--- a/Engine/Source/Runtime/Core/Private/Misc/PathViews.cpp
+++ b/Engine/Source/Runtime/Core/Private/Misc/PathViews.cpp
@@ -13,6 +13,7 @@
 #include "String/Find.h"
 #include "String/ParseTokens.h"
 #include "Templates/Function.h"
+#include "UObject/NameTypes.h"
 
 namespace UE4PathViews_Private
 {
@@ -224,15 +225,31 @@ void FPathViews::Split(const FStringView& InPath, FStringView& OutPath, FStringV
 	OutExt = CleanName.RightChop(NameLen + 1);
 }
 
-void FPathViews::Append(FStringBuilderBase& Builder, const FStringView& Suffix)
+void FPathViews::Append(FStringBuilderBase& Builder, FStringView Suffix)
 {
 	if (Builder.Len() > 0 && !UE4PathViews_Private::IsSlashOrBackslash(Builder.LastChar()))
 	{
-		Builder.Append(TEXT('/'));
+		Builder.Append('/');
 	}
+
 	Builder.Append(Suffix);
 }
 
+void FPathViews::Append(FStringBuilderBase& Builder, FName Suffix)
+{
+	if (Builder.Len() > 0 && !UE4PathViews_Private::IsSlashOrBackslash(Builder.LastChar()))
+	{
+		Builder.Append('/');
+	}
+
+	Suffix.AppendString(Builder);
+}
+
+void FPathViews::Append(FStringBuilderBase& Builder, const TCHAR* Suffix)
+{
+	Append(Builder, FStringView(Suffix));
+}
+
 FString FPathViews::ChangeExtension(const FStringView& InPath, const FStringView& InNewExtension)
 {
 	// Make sure the period we found was actually for a file extension and not part of the file path.
diff --git a/Engine/Source/Runtime/Core/Public/Misc/PathViews.h b/Engine/Source/Runtime/Core/Public/Misc/PathViews.h
--- a/Engine/Source/Runtime/Core/Public/Misc/PathViews.h
+++ b/Engine/Source/Runtime/Core/Public/Misc/PathViews.h
@@ -5,6 +5,7 @@
 #include "CoreTypes.h"
 #include "Containers/StringFwd.h"
 
+class FName;
 class FString;
 template <typename FuncType> class TFunctionRef;
 
@@ -160,7 +161,9 @@ public:
 	 * @param Builder A possibly-empty path that may end in a separator.
 	 * @param Suffix A possibly-empty suffix that does not start with a separator.
 	 */
-	static void Append(FStringBuilderBase& Builder, const FStringView& Suffix);
+	static void Append(FStringBuilderBase& Builder, FStringView Suffix);
+	static void Append(FStringBuilderBase& Builder, const TCHAR* Suffix);
+	static void Append(FStringBuilderBase& Builder, FName Suffix);
 
 	/**
 	 * Replaces the pre-existing file extension of a filename. 
