Replaced the UE4Variant_Detail namespace with UE::Core::Private.

diff --git a/Engine/Source/Runtime/Core/Public/Misc/TVariant.h b/Engine/Source/Runtime/Core/Public/Misc/TVariant.h
--- a/Engine/Source/Runtime/Core/Public/Misc/TVariant.h
+++ b/Engine/Source/Runtime/Core/Public/Misc/TVariant.h
@@ -23,17 +23,17 @@ struct FEmptyVariantState {};
  * Attempting to use the value of a Get() when the underlying type is different leads to undefined behavior.
  */
 template <typename T, typename... Ts>
-class TVariant final : private UE4Variant_Details::TVariantStorage<T, Ts...>
+class TVariant final : private UE::Core::Private::TVariantStorage<T, Ts...>
 {
-	static_assert(!UE4Variant_Details::TTypePackContainsDuplicates<T, Ts...>::Value, "All the types used in TVariant should be unique");
-	static_assert(!UE4Variant_Details::TContainsReferenceType<T, Ts...>::Value, "TVariant cannot hold reference types");
+	static_assert(!UE::Core::Private::TTypePackContainsDuplicates<T, Ts...>::Value, "All the types used in TVariant should be unique");
+	static_assert(!UE::Core::Private::TContainsReferenceType<T, Ts...>::Value, "TVariant cannot hold reference types");
 
 public:
 	/** Default initialize the TVariant to the first type in the parameter pack */
 	TVariant()
 	{
 		static_assert(TIsConstructible<T>::Value, "To default-initialize a TVariant, the first type in the parameter pack must be default constructible. Use FEmptyVariantState as the first type if none of the other types can be listed first.");
-		new(&UE4Variant_Details::CastToStorage(*this).Storage) T();
+		new(&UE::Core::Private::CastToStorage(*this).Storage) T();
 		TypeIndex = 0;
 	}
 
@@ -41,10 +41,10 @@ public:
 	template <typename U, typename... TArgs>
 	explicit TVariant(TInPlaceType<U>&&, TArgs&&... Args)
 	{
-		constexpr SIZE_T Index = UE4Variant_Details::TParameterPackTypeIndex<U, T, Ts...>::Value;
+		constexpr SIZE_T Index = UE::Core::Private::TParameterPackTypeIndex<U, T, Ts...>::Value;
 		static_assert(Index != (SIZE_T)-1, "The TVariant is not declared to hold the type being constructed");
 
-		new(&UE4Variant_Details::CastToStorage(*this).Storage) U(Forward<TArgs>(Args)...);
+		new(&UE::Core::Private::CastToStorage(*this).Storage) U(Forward<TArgs>(Args)...);
 		TypeIndex = Index;
 	}
 
@@ -52,14 +52,14 @@ public:
 	TVariant(const TVariant& Other)
 		: TypeIndex(Other.TypeIndex)
 	{
-		UE4Variant_Details::TCopyConstructorLookup<T, Ts...>::Construct(TypeIndex, &UE4Variant_Details::CastToStorage(*this).Storage, &UE4Variant_Details::CastToStorage(Other).Storage);
+		UE::Core::Private::TCopyConstructorLookup<T, Ts...>::Construct(TypeIndex, &UE::Core::Private::CastToStorage(*this).Storage, &UE::Core::Private::CastToStorage(Other).Storage);
 	}
 
 	/** Move construct the variant from another variant of the same type */
 	TVariant(TVariant&& Other)
 		: TypeIndex(Other.TypeIndex)
 	{
-		UE4Variant_Details::TMoveConstructorLookup<T, Ts...>::Construct(TypeIndex, &UE4Variant_Details::CastToStorage(*this).Storage, &UE4Variant_Details::CastToStorage(Other).Storage);
+		UE::Core::Private::TMoveConstructorLookup<T, Ts...>::Construct(TypeIndex, &UE::Core::Private::CastToStorage(*this).Storage, &UE::Core::Private::CastToStorage(Other).Storage);
 	}
 
 	/** Copy assign a variant from another variant of the same type */
@@ -87,25 +87,25 @@ public:
 	/** Destruct the underlying type (if appropriate) */
 	~TVariant()
 	{
-		UE4Variant_Details::TDestructorLookup<T, Ts...>::Destruct(TypeIndex, &UE4Variant_Details::CastToStorage(*this).Storage);
+		UE::Core::Private::TDestructorLookup<T, Ts...>::Destruct(TypeIndex, &UE::Core::Private::CastToStorage(*this).Storage);
 	}
 
 	/** Determine if the variant holds the specific type */
 	template <typename U>
 	bool IsType() const
 	{
-		return UE4Variant_Details::TIsType<U, T, Ts...>::IsSame(TypeIndex);
+		return UE::Core::Private::TIsType<U, T, Ts...>::IsSame(TypeIndex);
 	}
 
 	/** Get a reference to the held value. Bad things can happen if this is called on a variant that does not hold the type asked for */
 	template <typename U>
 	U& Get()
 	{
-		constexpr SIZE_T Index = UE4Variant_Details::TParameterPackTypeIndex<U, T, Ts...>::Value;
+		constexpr SIZE_T Index = UE::Core::Private::TParameterPackTypeIndex<U, T, Ts...>::Value;
 		static_assert(Index != (SIZE_T)-1, "The TVariant is not declared to hold the type passed to Get<>");
 
 		check(Index == TypeIndex);
-		return *reinterpret_cast<U*>(&UE4Variant_Details::CastToStorage(*this).Storage);
+		return *reinterpret_cast<U*>(&UE::Core::Private::CastToStorage(*this).Storage);
 	}
 
 	/** Get a reference to the held value. Bad things can happen if this is called on a variant that does not hold the type asked for */
@@ -120,9 +120,9 @@ public:
 	template <typename U>
 	U* TryGet()
 	{
-		constexpr SIZE_T Index = UE4Variant_Details::TParameterPackTypeIndex<U, T, Ts...>::Value;
+		constexpr SIZE_T Index = UE::Core::Private::TParameterPackTypeIndex<U, T, Ts...>::Value;
 		static_assert(Index != (SIZE_T)-1, "The TVariant is not declared to hold the type passed to TryGet<>");
-		return Index == TypeIndex ? reinterpret_cast<U*>(&UE4Variant_Details::CastToStorage(*this).Storage) : nullptr;
+		return Index == TypeIndex ? reinterpret_cast<U*>(&UE::Core::Private::CastToStorage(*this).Storage) : nullptr;
 	}
 
 	/** Get a pointer to the held value if the held type is the same as the one specified */
@@ -151,11 +151,11 @@ public:
 	template <typename U, typename... TArgs>
 	void Emplace(TArgs&&... Args)
 	{
-		constexpr SIZE_T Index = UE4Variant_Details::TParameterPackTypeIndex<U, T, Ts...>::Value;
+		constexpr SIZE_T Index = UE::Core::Private::TParameterPackTypeIndex<U, T, Ts...>::Value;
 		static_assert(Index != (SIZE_T)-1, "The TVariant is not declared to hold the type passed to Emplace<>");
 
-		UE4Variant_Details::TDestructorLookup<T, Ts...>::Destruct(TypeIndex, &UE4Variant_Details::CastToStorage(*this).Storage);
-		new(&UE4Variant_Details::CastToStorage(*this).Storage) U(Forward<TArgs>(Args)...);
+		UE::Core::Private::TDestructorLookup<T, Ts...>::Destruct(TypeIndex, &UE::Core::Private::CastToStorage(*this).Storage);
+		new(&UE::Core::Private::CastToStorage(*this).Storage) U(Forward<TArgs>(Args)...);
 		TypeIndex = Index;
 	}
 
@@ -163,7 +163,7 @@ public:
 	template <typename U>
 	static constexpr SIZE_T IndexOfType()
 	{
-		constexpr SIZE_T Index = UE4Variant_Details::TParameterPackTypeIndex<U, T, Ts...>::Value;
+		constexpr SIZE_T Index = UE::Core::Private::TParameterPackTypeIndex<U, T, Ts...>::Value;
 		static_assert(Index != (SIZE_T)-1, "The TVariant is not declared to hold the type passed to IndexOfType<>");
 		return Index;
 	}
@@ -213,7 +213,7 @@ template <typename T> struct TVariantSize<const T> : public TVariantSize<T> {};
 template <
 	typename Func,
 	typename... Variants,
-	typename = typename TEnableIf<UE4Variant_Details::TIsAllVariant<typename TDecay<Variants>::Type...>::Value>::Type
+	typename = typename TEnableIf<UE::Core::Private::TIsAllVariant<typename TDecay<Variants>::Type...>::Value>::Type
 >
 decltype(auto) Visit(Func&& Callable, Variants&&... Args)
 {
@@ -221,11 +221,11 @@ decltype(auto) Visit(Func&& Callable, Variants&&... Args)
 	constexpr SIZE_T NumPermutations = (1 * ... * (TVariantSize<Variants>::Value));
 #else
 	constexpr SIZE_T VariantSizes[] = { TVariantSize<Variants>::Value... };
-	constexpr SIZE_T NumPermutations = UE4Variant_Details::Multiply(VariantSizes, sizeof...(Variants));
+	constexpr SIZE_T NumPermutations = UE::Core::Private::Multiply(VariantSizes, sizeof...(Variants));
 #endif
 
-	return UE4Variant_Details::VisitImpl(
-		UE4Variant_Details::EncodeIndices(Args...),
+	return UE::Core::Private::VisitImpl(
+		UE::Core::Private::EncodeIndices(Args...),
 		Forward<Func>(Callable),
 		TMakeIntegerSequence<SIZE_T, NumPermutations>{},
 		TMakeIntegerSequence<SIZE_T, sizeof...(Variants)>{},
diff --git a/Engine/Source/Runtime/Core/Public/Misc/TVariantMeta.h b/Engine/Source/Runtime/Core/Public/Misc/TVariantMeta.h
--- a/Engine/Source/Runtime/Core/Public/Misc/TVariantMeta.h
+++ b/Engine/Source/Runtime/Core/Public/Misc/TVariantMeta.h
@@ -29,7 +29,11 @@ struct TIsVariant;
 template <typename T>
 struct TVariantSize;
 
-namespace UE4Variant_Details
+namespace UE
+{
+namespace Core
+{
+namespace Private
 {
 	/** Determine if all the types in a template parameter pack has duplicate types */
 	template <typename...>
@@ -381,6 +385,8 @@ namespace UE4Variant_Details
 		static constexpr InvokeFn Invokers[] = { WrapperType::template FuncPtr<EncodedIndices>... };
 		return Invokers[EncodedIndex](Forward<Func>(Callable), Forward<Variants>(Args)...);
 	}
-}
+} // namespace Private
+} // namespace Core
+} // namespace UE
 
 #undef TVARIANT_STORAGE_USE_RECURSIVE_TEMPLATE
