UE4Invoke_Private renamed to UE::Core::Private.

diff --git a/Engine/Source/Runtime/Core/Public/Templates/Invoke.h b/Engine/Source/Runtime/Core/Public/Templates/Invoke.h
--- a/Engine/Source/Runtime/Core/Public/Templates/Invoke.h
+++ b/Engine/Source/Runtime/Core/Public/Templates/Invoke.h
@@ -8,7 +8,7 @@
 #include "Templates/Decay.h"
 
 
-namespace UE4Invoke_Private
+namespace UE::Core::Private
 {
 	template <
 		typename BaseType,
@@ -53,23 +53,23 @@ FORCEINLINE auto Invoke(FuncType&& Func, ArgTypes&&... Args)
 
 template <typename ReturnType, typename ObjType, typename CallableType>
 FORCEINLINE auto Invoke(ReturnType ObjType::*pdm, CallableType&& Callable)
-	-> decltype(UE4Invoke_Private::DereferenceIfNecessary<ObjType>(Forward<CallableType>(Callable)).*pdm)
+	-> decltype(UE::Core::Private::DereferenceIfNecessary<ObjType>(Forward<CallableType>(Callable)).*pdm)
 {
-	return UE4Invoke_Private::DereferenceIfNecessary<ObjType>(Forward<CallableType>(Callable)).*pdm;
+	return UE::Core::Private::DereferenceIfNecessary<ObjType>(Forward<CallableType>(Callable)).*pdm;
 }
 
 template <typename ReturnType, typename ObjType, typename... PMFArgTypes, typename CallableType, typename... ArgTypes>
 FORCEINLINE auto Invoke(ReturnType (ObjType::*PtrMemFun)(PMFArgTypes...), CallableType&& Callable, ArgTypes&&... Args)
-	-> decltype((UE4Invoke_Private::DereferenceIfNecessary<ObjType>(Forward<CallableType>(Callable)).*PtrMemFun)(Forward<ArgTypes>(Args)...))
+	-> decltype((UE::Core::Private::DereferenceIfNecessary<ObjType>(Forward<CallableType>(Callable)).*PtrMemFun)(Forward<ArgTypes>(Args)...))
 {
-	return (UE4Invoke_Private::DereferenceIfNecessary<ObjType>(Forward<CallableType>(Callable)).*PtrMemFun)(Forward<ArgTypes>(Args)...);
+	return (UE::Core::Private::DereferenceIfNecessary<ObjType>(Forward<CallableType>(Callable)).*PtrMemFun)(Forward<ArgTypes>(Args)...);
 }
 
 template <typename ReturnType, typename ObjType, typename... PMFArgTypes, typename CallableType, typename... ArgTypes>
 FORCEINLINE auto Invoke(ReturnType (ObjType::*PtrMemFun)(PMFArgTypes...) const, CallableType&& Callable, ArgTypes&&... Args)
-	-> decltype((UE4Invoke_Private::DereferenceIfNecessary<ObjType>(Forward<CallableType>(Callable)).*PtrMemFun)(Forward<ArgTypes>(Args)...))
+	-> decltype((UE::Core::Private::DereferenceIfNecessary<ObjType>(Forward<CallableType>(Callable)).*PtrMemFun)(Forward<ArgTypes>(Args)...))
 {
-	return (UE4Invoke_Private::DereferenceIfNecessary<ObjType>(Forward<CallableType>(Callable)).*PtrMemFun)(Forward<ArgTypes>(Args)...);
+	return (UE::Core::Private::DereferenceIfNecessary<ObjType>(Forward<CallableType>(Callable)).*PtrMemFun)(Forward<ArgTypes>(Args)...);
 }
 
 
@@ -111,13 +111,13 @@ FORCEINLINE auto Invoke(ReturnType (ObjType::*PtrMemFun)(PMFArgTypes...) const,
 #define UE_PROJECTION_MEMBER(Type, FuncName) \
 	[](auto&& Obj, auto&&... Args) -> decltype(auto) \
 	{ \
-		return UE4Invoke_Private::DereferenceIfNecessary<Type>(Forward<decltype(Obj)>(Obj)).FuncName(Forward<decltype(Args)>(Args)...); \
+		return UE::Core::Private::DereferenceIfNecessary<Type>(Forward<decltype(Obj)>(Obj)).FuncName(Forward<decltype(Args)>(Args)...); \
 	}
 
 #define PROJECTION(FuncName)              DEPRECATED_MACRO(4.26, "The PROJECTION macro is deprecated, please use UE_PROJECTION instead.") UE_PROJECTION(FuncName)
 #define PROJECTION_MEMBER(Type, FuncName) DEPRECATED_MACRO(4.26, "The PROJECTION_MEMBER macro is deprecated, please use UE_PROJECTION_MEMBER instead.") UE_PROJECTION_MEMBER(Type, FuncName)
 
-namespace UE4Invoke_Private
+namespace UE::Core::Private
 {
 	template <typename, typename FuncType, typename... ArgTypes>
 	struct TInvokeResult_Impl
@@ -136,7 +136,7 @@ namespace UE4Invoke_Private
  * Not defined (as thus usable in SFINAE contexts) when the callable cannot be invoked with the given argument types.
  */
 template <typename FuncType, typename... ArgTypes>
-struct TInvokeResult : UE4Invoke_Private::TInvokeResult_Impl<void, FuncType, ArgTypes...>
+struct TInvokeResult : UE::Core::Private::TInvokeResult_Impl<void, FuncType, ArgTypes...>
 {
 };
 
