CompositeBuffer: Added a new immutable non-contiguous buffer type that is a collection of FSharedBuffer

diff --git a/Engine/Source/Runtime/Core/Private/Hash/Blake3.cpp b/Engine/Source/Runtime/Core/Private/Hash/Blake3.cpp
--- a/Engine/Source/Runtime/Core/Private/Hash/Blake3.cpp
+++ b/Engine/Source/Runtime/Core/Private/Hash/Blake3.cpp
@@ -1,7 +1,9 @@
 // Copyright Epic Games, Inc. All Rights Reserved.
 
 #include "Hash/Blake3.h"
+
 #include "Containers/UnrealString.h"
+#include "Memory/CompositeBuffer.h"
 
 #include "blake3.h"
 
@@ -15,12 +17,27 @@ void FBlake3::Reset()
 	blake3_hasher_init(&Hasher);
 }
 
+void FBlake3::Update(FMemoryView View)
+{
+	blake3_hasher& Hasher = reinterpret_cast<blake3_hasher&>(HasherBytes);
+	blake3_hasher_update(&Hasher, View.GetData(), View.GetSize());
+}
+
 void FBlake3::Update(const void* Data, uint64 Size)
 {
 	blake3_hasher& Hasher = reinterpret_cast<blake3_hasher&>(HasherBytes);
 	blake3_hasher_update(&Hasher, Data, Size);
 }
 
+void FBlake3::Update(const FCompositeBuffer& Buffer)
+{
+	blake3_hasher& Hasher = reinterpret_cast<blake3_hasher&>(HasherBytes);
+	for (const FSharedBuffer& Segment : Buffer.GetSegments())
+	{
+		blake3_hasher_update(&Hasher, Segment.GetData(), Segment.GetSize());
+	}
+}
+
 FBlake3Hash FBlake3::Finalize() const
 {
 	FBlake3Hash Hash;
@@ -31,6 +48,13 @@ FBlake3Hash FBlake3::Finalize() const
 	return Hash;
 }
 
+FBlake3Hash FBlake3::HashBuffer(FMemoryView View)
+{
+	FBlake3 Hash;
+	Hash.Update(View.GetData(), View.GetSize());
+	return Hash.Finalize();
+}
+
 FBlake3Hash FBlake3::HashBuffer(const void* Data, uint64 Size)
 {
 	FBlake3 Hash;
@@ -38,6 +62,15 @@ FBlake3Hash FBlake3::HashBuffer(const void* Data, uint64 Size)
 	return Hash.Finalize();
 }
 
+FBlake3Hash FBlake3::HashBuffer(const FCompositeBuffer& Buffer)
+{
+	FBlake3 Hash;
+	Hash.Update(Buffer);
+	return Hash.Finalize();
+}
+
+///////////////////////////////////////////////////////////////////////////////////////////////////////////////////////
+
 FString LexToString(const FBlake3Hash& Hash)
 {
 	FString Output;
diff --git a/Engine/Source/Runtime/Core/Private/Memory/CompositeBuffer.cpp b/Engine/Source/Runtime/Core/Private/Memory/CompositeBuffer.cpp
new file mode 100644
--- /dev/null
+++ b/Engine/Source/Runtime/Core/Private/Memory/CompositeBuffer.cpp
@@ -0,0 +1,129 @@
+// Copyright Epic Games, Inc. All Rights Reserved.
+
+#include "Memory/CompositeBuffer.h"
+
+#include "Algo/Accumulate.h"
+#include "Algo/AllOf.h"
+#include "Templates/Function.h"
+
+void FCompositeBuffer::Reset()
+{
+	Segments.Empty();
+}
+
+uint64 FCompositeBuffer::GetSize() const
+{
+	return Algo::TransformAccumulate(Segments, &FSharedBuffer::GetSize, uint64(0));
+}
+
+bool FCompositeBuffer::IsOwned() const
+{
+	return Algo::AllOf(Segments, &FSharedBuffer::IsOwned);
+}
+
+FCompositeBuffer FCompositeBuffer::MakeOwned() const &
+{
+	return FCompositeBuffer(*this).MakeOwned();
+}
+
+FCompositeBuffer FCompositeBuffer::MakeOwned() &&
+{
+	for (FSharedBuffer& Segment : Segments)
+	{
+		Segment = MoveTemp(Segment).MakeOwned();
+	}
+	return MoveTemp(*this);
+}
+
+FSharedBuffer FCompositeBuffer::Flatten() const &
+{
+	switch (Segments.Num())
+	{
+	case 0:
+		return FSharedBuffer();
+	case 1:
+		return Segments[0];
+	default:
+		FUniqueBuffer Buffer = FUniqueBuffer::Alloc(GetSize());
+		Algo::TransformAccumulate(Segments, &FSharedBuffer::GetView, Buffer.GetView(), &FMutableMemoryView::CopyFrom);
+		return FSharedBuffer(MoveTemp(Buffer));
+	}
+}
+
+FSharedBuffer FCompositeBuffer::Flatten() &&
+{
+	return Segments.Num() == 1 ? MoveTemp(Segments[0]) : AsConst(*this).Flatten();
+}
+
+FCompositeBuffer FCompositeBuffer::Mid(uint64 Offset, uint64 Size) const
+{
+	FCompositeBuffer Buffer;
+	IterateRange(Offset, Size, [&Buffer](FMemoryView View, const FSharedBuffer& ViewOuter)
+		{
+			Buffer.Segments.Add(FSharedBuffer::MakeView(View, ViewOuter));
+		});
+	return Buffer;
+}
+
+FMemoryView FCompositeBuffer::ViewOrCopyRange(uint64 Offset, uint64 Size, FUniqueBuffer& CopyBuffer) const
+{
+	FMemoryView View;
+	IterateRange(Offset, Size, [Size, &View, &CopyBuffer, WriteView = FMutableMemoryView()](FMemoryView Segment) mutable
+		{
+			if (Size == Segment.GetSize())
+			{
+				View = Segment;
+			}
+			else
+			{
+				if (WriteView.IsEmpty())
+				{
+					if (CopyBuffer.GetSize() < Size)
+					{
+						CopyBuffer = FUniqueBuffer::Alloc(Size);
+					}
+					View = WriteView = CopyBuffer;
+				}
+				WriteView = WriteView.CopyFrom(Segment);
+			}
+		});
+	return View;
+}
+
+void FCompositeBuffer::CopyTo(FMutableMemoryView Target, uint64 Offset) const
+{
+	IterateRange(Offset, Target.GetSize(), [Target](FMemoryView View, const FSharedBuffer& ViewOuter) mutable
+		{
+			Target = Target.CopyFrom(View);
+		});
+}
+
+void FCompositeBuffer::IterateRange(uint64 Offset, uint64 Size, TFunctionRef<void (FMemoryView View)> Visitor) const
+{
+	IterateRange(Offset, Size, [Visitor](FMemoryView View, const FSharedBuffer& ViewOuter) { Visitor(View); });
+}
+
+void FCompositeBuffer::IterateRange(uint64 Offset, uint64 Size,
+	TFunctionRef<void (FMemoryView View, const FSharedBuffer& ViewOuter)> Visitor) const
+{
+	checkf(Offset + Size <= GetSize(), TEXT("Failed to access %" UINT64_FMT " bytes at offset %" UINT64_FMT)
+		TEXT(" of a composite buffer containing %" UINT64_FMT " bytes."), Size, Offset, GetSize());
+	for (const FSharedBuffer& Segment : Segments)
+	{
+		if (Size == 0)
+		{
+			break;
+		}
+		const FMemoryView View = Segment.GetView().Mid(Offset, Size);
+		if (const uint64 ViewSize = View.GetSize())
+		{
+			Visitor(View, Segment);
+			Offset = 0;
+			Size -= ViewSize;
+		}
+		else
+		{
+			Offset -= Segment.GetSize();
+		}
+	}
+}
diff --git a/Engine/Source/Runtime/Core/Public/Hash/Blake3.h b/Engine/Source/Runtime/Core/Public/Hash/Blake3.h
--- a/Engine/Source/Runtime/Core/Public/Hash/Blake3.h
+++ b/Engine/Source/Runtime/Core/Public/Hash/Blake3.h
@@ -5,6 +5,7 @@
 #include "Containers/StringFwd.h"
 #include "Containers/StringView.h"
 #include "HAL/UnrealMemory.h"
+#include "Memory/MemoryFwd.h"
 #include "Memory/MemoryView.h"
 #include "Serialization/Archive.h"
 #include "String/BytesToHex.h"
@@ -58,26 +59,22 @@ public:
 	/** Reset to the default state in which no input has been written. */
 	CORE_API void Reset();
 
-	/** Add the data as input to the hash. May be called any number of times. */
+	/** Add the buffer as input to the hash. May be called any number of times. */
+	CORE_API void Update(FMemoryView View);
 	CORE_API void Update(const void* Data, uint64 Size);
+	CORE_API void Update(const FCompositeBuffer& Buffer);
 
-	/** Add the view as input to the hash. May be called any number of times. */
-	inline void Update(FMemoryView View)
-	{
-		Update(View.GetData(), View.GetSize());
-	}
-
-	/** Finalize the hash of the input data. May be called any number of times, and more input may be added after. */
-	CORE_API FBlake3Hash Finalize() const;
+	/**
+	 * Finalize the hash of the input data.
+	 *
+	 * May be called any number of times, and more input may be added after.
+	 */
+	[[nodiscard]] CORE_API FBlake3Hash Finalize() const;
 
-	/** Calculate the hash of the input data. */
-	CORE_API static FBlake3Hash HashBuffer(const void* Data, uint64 Size);
-
-	/** Calculate the hash of the input view. */
-	static inline FBlake3Hash HashBuffer(FMemoryView View)
-	{
-		return HashBuffer(View.GetData(), View.GetSize());
-	}
+	/** Calculate the hash of the buffer. */
+	[[nodiscard]] CORE_API static FBlake3Hash HashBuffer(FMemoryView View);
+	[[nodiscard]] CORE_API static FBlake3Hash HashBuffer(const void* Data, uint64 Size);
+	[[nodiscard]] CORE_API static FBlake3Hash HashBuffer(const FCompositeBuffer& Buffer);
 
 private:
 	TAlignedBytes<1912, 8> HasherBytes;
@@ -161,4 +158,4 @@ inline void LexFromString(FBlake3Hash& OutHash, const TCHAR* Buffer)
 }
 
 /** Convert a hash to a 64-character hex string. */
-UE_NODISCARD CORE_API FString LexToString(const FBlake3Hash& Hash);
+[[nodiscard]] CORE_API FString LexToString(const FBlake3Hash& Hash);
diff --git a/Engine/Source/Runtime/Core/Public/IO/IoHash.h b/Engine/Source/Runtime/Core/Public/IO/IoHash.h
--- a/Engine/Source/Runtime/Core/Public/IO/IoHash.h
+++ b/Engine/Source/Runtime/Core/Public/IO/IoHash.h
@@ -6,6 +6,7 @@
 #include "Containers/StringView.h"
 #include "HAL/UnrealMemory.h"
 #include "Hash/Blake3.h"
+#include "Memory/MemoryFwd.h"
 #include "Memory/MemoryView.h"
 #include "Serialization/Archive.h"
 #include "String/BytesToHex.h"
@@ -49,11 +50,10 @@ public:
 	inline ByteArray& GetBytes() { return Hash; }
 	inline const ByteArray& GetBytes() const { return Hash; }
 
-	/** Calculate the hash of the input data. */
-	static inline FIoHash HashBuffer(const void* Data, uint64 Size);
-
-	/** Calculate the hash of the input view. */
-	static inline FIoHash HashBuffer(FMemoryView View);
+	/** Calculate the hash of the buffer. */
+	[[nodiscard]] static inline FIoHash HashBuffer(FMemoryView View);
+	[[nodiscard]] static inline FIoHash HashBuffer(const void* Data, uint64 Size);
+	[[nodiscard]] static inline FIoHash HashBuffer(const FCompositeBuffer& Buffer);
 
 private:
 	alignas(uint32) ByteArray Hash{};
@@ -96,14 +96,19 @@ inline bool FIoHash::IsZero() const
 	return true;
 }
 
+inline FIoHash FIoHash::HashBuffer(FMemoryView View)
+{
+	return FIoHash(FBlake3::HashBuffer(View));
+}
+
 inline FIoHash FIoHash::HashBuffer(const void* Data, uint64 Size)
 {
 	return FIoHash(FBlake3::HashBuffer(Data, Size));
 }
 
-inline FIoHash FIoHash::HashBuffer(FMemoryView View)
+inline FIoHash FIoHash::HashBuffer(const FCompositeBuffer& Buffer)
 {
-	return FIoHash(FBlake3::HashBuffer(View));
+	return FIoHash(FBlake3::HashBuffer(Buffer));
 }
 
 inline bool operator==(const FIoHash& A, const FIoHash& B)
@@ -151,4 +156,4 @@ inline void LexFromString(FIoHash& OutHash, const TCHAR* Buffer)
 }
 
 /** Convert a hash to a 40-character hex string. */
-UE_NODISCARD CORE_API FString LexToString(const FIoHash& Hash);
+[[nodiscard]] CORE_API FString LexToString(const FIoHash& Hash);
diff --git a/Engine/Source/Runtime/Core/Public/Memory/CompositeBuffer.h b/Engine/Source/Runtime/Core/Public/Memory/CompositeBuffer.h
new file mode 100644
--- /dev/null
+++ b/Engine/Source/Runtime/Core/Public/Memory/CompositeBuffer.h
@@ -0,0 +1,115 @@
+// Copyright Epic Games, Inc. All Rights Reserved.
+
+#pragma once
+
+#include "Containers/Array.h"
+#include "Containers/ArrayView.h"
+#include "Memory/MemoryView.h"
+#include "Memory/SharedBuffer.h"
+
+template <typename FuncType> class TFunctionRef;
+
+/**
+ * FCompositeBuffer is a non-contiguous buffer composed of zero or more immutable shared buffers.
+ *
+ * A composite buffer is most efficient when its segments are consumed as they are, but it can be
+ * flattened into a contiguous buffer, when necessary, by calling Flatten(). Ownership of segment
+ * buffers is not changed on construction, but if ownership of segments is required then that can
+ * be guaranteed by calling MakeOwned().
+ */
+class FCompositeBuffer
+{
+public:
+	/**
+	 * Construct a composite buffer by concatenating the buffers. Does not enforce ownership.
+	 *
+	 * Buffer parameters may be FSharedBuffer, FCompositeBuffer, or TArray<FSharedBuffer>.
+	 */
+	template <typename... BufferTypes>
+	inline explicit FCompositeBuffer(BufferTypes&&... Buffers)
+	{
+		if constexpr (sizeof...(Buffers) > 0)
+		{
+			Segments.Reserve((GetBufferCount(Forward<BufferTypes>(Buffers)) + ...));
+			(AppendBuffers(Forward<BufferTypes>(Buffers)), ...);
+			Segments.RemoveAll(&FSharedBuffer::IsNull);
+		}
+	}
+
+	/** Reset this to empty. */
+	CORE_API void Reset();
+
+	/** Returns the total size of the composite buffer in bytes. */
+	[[nodiscard]] CORE_API uint64 GetSize() const;
+
+	/** Returns the segments that the buffer is composed from. */
+	[[nodiscard]] inline TConstArrayView<FSharedBuffer> GetSegments() const { return Segments; }
+
+	/** Returns true if the composite buffer has a total size greater than zero. */
+	[[nodiscard]] inline explicit operator bool() const { return !IsEmpty(); }
+
+	/** Returns true if the composite buffer has a total size of zero. */
+	[[nodiscard]] inline bool IsEmpty() const { return Segments.IsEmpty() || GetSize() == 0; }
+
+	/** Returns true if every segment in the composite buffer is owned. */
+	[[nodiscard]] CORE_API bool IsOwned() const;
+
+	/** Returns a copy of the buffer where every segment is owned. */
+	[[nodiscard]] CORE_API FCompositeBuffer MakeOwned() const &;
+	[[nodiscard]] CORE_API FCompositeBuffer MakeOwned() &&;
+
+	/** Returns the concatenation of the segments into a contiguous buffer. */
+	[[nodiscard]] CORE_API FSharedBuffer Flatten() const &;
+	[[nodiscard]] CORE_API FSharedBuffer Flatten() &&;
+
+	/** Returns the middle part of the buffer by taking the size starting at the offset. */
+	[[nodiscard]] CORE_API FCompositeBuffer Mid(uint64 Offset, uint64 Size = MAX_uint64) const;
+
+	/**
+	 * Returns a view of the range if contained by one segment, otherwise a view of a copy of the range.
+	 *
+	 * @note CopyBuffer is reused if large enough, and otherwise allocated when needed.
+	 *
+	 * @param Offset       The byte offset in this buffer that the range starts at.
+	 * @param Size         The number of bytes in the range to view or copy.
+	 * @param CopyBuffer   The buffer to write the copy into if a copy is required.
+	 */
+	[[nodiscard]] CORE_API FMemoryView ViewOrCopyRange(uint64 Offset, uint64 Size, FUniqueBuffer& CopyBuffer) const;
+
+	/**
+	 * Copies a range of the buffer to a contiguous region of memory.
+	 *
+	 * @param Target   The view to copy to. Must be no larger than the data available at the offset.
+	 * @param Offset   The byte offset in this buffer to start copying from.
+	 */
+	CORE_API void CopyTo(FMutableMemoryView Target, uint64 Offset = 0) const;
+
+	/**
+	 * Invokes a visitor with a view of each segment that intersects with a range.
+	 *
+	 * @param Offset    The byte offset in this buffer to start visiting from.
+	 * @param Size      The number of bytes in the range to visit.
+	 * @param Visitor   The visitor to invoke from zero to GetSegments().Num() times.
+	 */
+	CORE_API void IterateRange(uint64 Offset, uint64 Size,
+		TFunctionRef<void (FMemoryView View)> Visitor) const;
+	CORE_API void IterateRange(uint64 Offset, uint64 Size,
+		TFunctionRef<void (FMemoryView View, const FSharedBuffer& ViewOuter)> Visitor) const;
+
+private:
+	static inline int32 GetBufferCount(const FCompositeBuffer& Buffer) { return Buffer.Segments.Num(); }
+	inline void AppendBuffers(const FCompositeBuffer& Buffer) { Segments.Append(Buffer.Segments); }
+	inline void AppendBuffers(FCompositeBuffer&& Buffer) { Segments.Append(MoveTemp(Buffer.Segments)); }
+
+	static inline int32 GetBufferCount(const FSharedBuffer& Buffer) { return 1; }
+	inline void AppendBuffers(const FSharedBuffer& Buffer) { Segments.Add(Buffer); }
+	inline void AppendBuffers(FSharedBuffer&& Buffer) { Segments.Add(MoveTemp(Buffer)); }
+
+	template <typename BufferType, decltype(std::declval<TArray<FSharedBuffer>>().Append(std::declval<BufferType>()))* = nullptr>
+	static inline int32 GetBufferCount(BufferType&& Buffer) { return GetNum(Buffer); }
+	template <typename BufferType, decltype(std::declval<TArray<FSharedBuffer>>().Append(std::declval<BufferType>()))* = nullptr>
+	inline void AppendBuffers(BufferType&& Buffer) { Segments.Append(Forward<BufferType>(Buffer)); }
+
+private:
+	TArray<FSharedBuffer, TInlineAllocator<1>> Segments;
+};
diff --git a/Engine/Source/Runtime/Core/Public/Memory/MemoryFwd.h b/Engine/Source/Runtime/Core/Public/Memory/MemoryFwd.h
--- a/Engine/Source/Runtime/Core/Public/Memory/MemoryFwd.h
+++ b/Engine/Source/Runtime/Core/Public/Memory/MemoryFwd.h
@@ -14,7 +14,10 @@ using FMemoryView = TMemoryView<const void>;
 using FMutableMemoryView = TMemoryView<void>;
 
 // SharedBuffer
-
+class FBufferOwner;
 class FUniqueBuffer;
 class FSharedBuffer;
 class FWeakSharedBuffer;
+
+// CompositeBuffer
+class FCompositeBuffer;
