Remove old clang workaround that creates symbol resolution linker errors

diff --git a/Engine/Source/Runtime/Core/Private/Async/Async.cpp b/Engine/Source/Runtime/Core/Private/Async/Async.cpp
--- a/Engine/Source/Runtime/Core/Private/Async/Async.cpp
+++ b/Engine/Source/Runtime/Core/Private/Async/Async.cpp
@@ -43,13 +43,11 @@ private:
 /* Global functions
  *****************************************************************************/
 
-#if PLATFORM_WINDOWS && defined(__clang__)
-CORE_API int32 FAsyncThreadIndex::GetNext() // @todo clang: Workaround for missing symbol export
+CORE_API int32 FAsyncThreadIndex::GetNext()
 {
 	static FThreadSafeCounter ThreadIndex;
 	return ThreadIndex.Add(1);
 }
-#endif
 
 void AsyncTask(ENamedThreads::Type Thread, TUniqueFunction<void()> Function)
 {
diff --git a/Engine/Source/Runtime/Core/Private/Internationalization/Text.cpp b/Engine/Source/Runtime/Core/Private/Internationalization/Text.cpp
--- a/Engine/Source/Runtime/Core/Private/Internationalization/Text.cpp
+++ b/Engine/Source/Runtime/Core/Private/Internationalization/Text.cpp
@@ -241,13 +241,11 @@ FText::FText( EInitToEmptyString )
 {
 }
 
-#if PLATFORM_WINDOWS && defined(__clang__)
-CORE_API const FText& FText::GetEmpty() // @todo clang: Workaround for missing symbol export
+CORE_API const FText& FText::GetEmpty()
 {
 	static const FText StaticEmptyText = FText(FText::EInitToEmptyString::Value);
 	return StaticEmptyText;
 }
-#endif
 
 FText::FText( TSharedRef<ITextData, ESPMode::ThreadSafe> InTextData )
 	: TextData(MoveTemp(InTextData))
diff --git a/Engine/Source/Runtime/Core/Private/Misc/App.cpp b/Engine/Source/Runtime/Core/Private/Misc/App.cpp
--- a/Engine/Source/Runtime/Core/Private/Misc/App.cpp
+++ b/Engine/Source/Runtime/Core/Private/Misc/App.cpp
@@ -228,13 +228,13 @@ bool FApp::IsEngineInstalled()
 	return EngineInstalledState == 1;
 }
 
-#if PLATFORM_WINDOWS && defined(__clang__)
-bool FApp::IsUnattended() // @todo clang: Workaround for missing symbol export
+bool FApp::IsUnattended()
 {
+	// FCommandLine::Get() will assert that the command line has been set.
+	// This function may not be used before FCommandLine::Set() is called.
 	static bool bIsUnattended = FParse::Param(FCommandLine::Get(), TEXT("UNATTENDED"));
 	return bIsUnattended || GIsAutomationTesting;
 }
-#endif
 
 bool FApp::ShouldUseThreadingForPerformance()
 {
diff --git a/Engine/Source/Runtime/Core/Public/Async/Async.h b/Engine/Source/Runtime/Core/Public/Async/Async.h
--- a/Engine/Source/Runtime/Core/Public/Async/Async.h
+++ b/Engine/Source/Runtime/Core/Public/Async/Async.h
@@ -244,15 +244,7 @@ private:
  */
 struct FAsyncThreadIndex
 {
-#if ( !PLATFORM_WINDOWS ) || ( !defined(__clang__) )
-	static CORE_API int32 GetNext()
-	{
-		static FThreadSafeCounter ThreadIndex;
-		return ThreadIndex.Add(1);
-	}
-#else
-	static CORE_API int32 GetNext(); // @todo clang: Workaround for missing symbol export
-#endif
+	static CORE_API int32 GetNext();
 };
 
 
diff --git a/Engine/Source/Runtime/Core/Public/Internationalization/Text.h b/Engine/Source/Runtime/Core/Public/Internationalization/Text.h
--- a/Engine/Source/Runtime/Core/Public/Internationalization/Text.h
+++ b/Engine/Source/Runtime/Core/Public/Internationalization/Text.h
@@ -367,17 +367,7 @@ class CORE_API FText
 {
 public:
 
-#if ( !PLATFORM_WINDOWS ) || ( !defined(__clang__) )
-	static const FText& GetEmpty()
-	{
-		// This is initialized inside this function as we need to be able to control the initialization order of the empty FText instance
-		// If this were a file-scope static, we can end up with other statics trying to construct an empty FText before our empty FText has itself been constructed
-		static const FText StaticEmptyText = FText(FText::EInitToEmptyString::Value);
-		return StaticEmptyText;
-	}
-#else
-	static const FText& GetEmpty(); // @todo clang: Workaround for missing symbol export
-#endif
+	static const FText& GetEmpty();
 
 public:
 
diff --git a/Engine/Source/Runtime/Core/Public/Misc/App.h b/Engine/Source/Runtime/Core/Public/Misc/App.h
--- a/Engine/Source/Runtime/Core/Public/Misc/App.h
+++ b/Engine/Source/Runtime/Core/Public/Misc/App.h
@@ -422,17 +422,7 @@ public:
 	 *
 	 * @return true if the application runs unattended, false otherwise.
 	 */
-#if ( !PLATFORM_WINDOWS ) || ( !defined(__clang__) )
-	static bool IsUnattended()
-	{
-		// FCommandLine::Get() will assert that the command line has been set.
-		// This function may not be used before FCommandLine::Set() is called.
-		static bool bIsUnattended = FParse::Param(FCommandLine::Get(), TEXT("UNATTENDED"));
-		return bIsUnattended || GIsAutomationTesting;
-	}
-#else
-	static bool IsUnattended(); // @todo clang: Workaround for missing symbol export
-#endif
+	static bool IsUnattended();
 
 	/**
 	 * Checks whether the application should run multi-threaded for performance critical features.
