Switch default allocator on Windows from tbb to mimalloc because of both performance and memory usage improvements
   - See CL 15887498 description for benchmarks and motivation

diff --git a/Engine/Source/Runtime/Core/Private/Windows/WindowsPlatformMemory.cpp b/Engine/Source/Runtime/Core/Private/Windows/WindowsPlatformMemory.cpp
--- a/Engine/Source/Runtime/Core/Private/Windows/WindowsPlatformMemory.cpp
+++ b/Engine/Source/Runtime/Core/Private/Windows/WindowsPlatformMemory.cpp
@@ -96,15 +96,20 @@ FMalloc* FWindowsPlatformMemory::BaseAllocator()
 	{
 		AllocatorToUse = EMemoryAllocatorToUse::Ansi;
 	}
-	else if ((WITH_EDITORONLY_DATA || IS_PROGRAM) && TBB_ALLOCATOR_ALLOWED) //-V517
-	{
-		AllocatorToUse = EMemoryAllocatorToUse::TBB;
-	}
 #if PLATFORM_64BITS
+	// Mimalloc is now the default allocator for editor and programs because it has shown
+	// both great performance and as much as half the memory usage of TBB after
+	// heavy editor workloads. See CL 15887498 description for benchmarks.
 	else if ((WITH_EDITORONLY_DATA || IS_PROGRAM) && MIMALLOC_ALLOCATOR_ALLOWED) //-V517
 	{
 		AllocatorToUse = EMemoryAllocatorToUse::Mimalloc;
 	}
+#endif
+	else if ((WITH_EDITORONLY_DATA || IS_PROGRAM) && TBB_ALLOCATOR_ALLOWED) //-V517
+	{
+		AllocatorToUse = EMemoryAllocatorToUse::TBB;
+	}
+#if PLATFORM_64BITS
 	else if (USE_MALLOC_BINNED3)
 	{
 		AllocatorToUse = EMemoryAllocatorToUse::Binned3;
