Tidying up of some metaprogramming using C++17 features.

diff --git a/Engine/Source/Runtime/Core/Public/Templates/Invoke.h b/Engine/Source/Runtime/Core/Public/Templates/Invoke.h
--- a/Engine/Source/Runtime/Core/Public/Templates/Invoke.h
+++ b/Engine/Source/Runtime/Core/Public/Templates/Invoke.h
@@ -3,31 +3,22 @@
 #pragma once
 
 #include "CoreTypes.h"
-#include "Templates/PointerIsConvertibleFromTo.h"
-#include "Templates/UnrealTemplate.h"
-#include "Templates/Decay.h"
+#include <type_traits>
 
 
 namespace UE::Core::Private
 {
-	template <
-		typename BaseType,
-		typename CallableType,
-		typename TEnableIf<TPointerIsConvertibleFromTo<typename TDecay<CallableType>::Type, typename TDecay<BaseType>::Type>::Value>::Type* = nullptr
-	>
-	FORCEINLINE auto DereferenceIfNecessary(CallableType&& Callable) -> decltype((CallableType&&)Callable)
+	template <typename BaseType, typename CallableType>
+	FORCEINLINE decltype(auto) DereferenceIfNecessary(CallableType&& Callable)
 	{
-		return (CallableType&&)Callable;
-	}
-
-	template <
-		typename BaseType,
-		typename CallableType,
-		typename TEnableIf<!TPointerIsConvertibleFromTo<typename TDecay<CallableType>::Type, typename TDecay<BaseType>::Type>::Value>::Type* = nullptr
-	>
-	FORCEINLINE auto DereferenceIfNecessary(CallableType&& Callable) -> decltype(*(CallableType&&)Callable)
-	{
-		return *(CallableType&&)Callable;
+		if constexpr (std::is_convertible_v<std::decay_t<CallableType>*, std::decay_t<BaseType>*>)
+		{
+			return (CallableType&&)Callable;
+		}
+		else
+		{
+			return *(CallableType&&)Callable;
+		}
 	}
 }
 
