Invoke() support for member functions with ref-qualifiers.

diff --git a/Engine/Source/Runtime/Core/Public/Templates/Invoke.h b/Engine/Source/Runtime/Core/Public/Templates/Invoke.h
--- a/Engine/Source/Runtime/Core/Public/Templates/Invoke.h
+++ b/Engine/Source/Runtime/Core/Public/Templates/Invoke.h
@@ -3,6 +3,7 @@
 #pragma once
 
 #include "CoreTypes.h"
+#include "Traits/MemberFunctionPtrOuter.h"
 #include <type_traits>
 
 
@@ -49,15 +50,13 @@ FORCEINLINE auto Invoke(ReturnType ObjType::*pdm, CallableType&& Callable)
 	return UE::Core::Private::DereferenceIfNecessary<ObjType>(Forward<CallableType>(Callable)).*pdm;
 }
 
-template <typename ReturnType, typename ObjType, typename... PMFArgTypes, typename CallableType, typename... ArgTypes>
-FORCEINLINE auto Invoke(ReturnType (ObjType::*PtrMemFun)(PMFArgTypes...), CallableType&& Callable, ArgTypes&&... Args)
-	-> decltype((UE::Core::Private::DereferenceIfNecessary<ObjType>(Forward<CallableType>(Callable)).*PtrMemFun)(Forward<ArgTypes>(Args)...))
-{
-	return (UE::Core::Private::DereferenceIfNecessary<ObjType>(Forward<CallableType>(Callable)).*PtrMemFun)(Forward<ArgTypes>(Args)...);
-}
-
-template <typename ReturnType, typename ObjType, typename... PMFArgTypes, typename CallableType, typename... ArgTypes>
-FORCEINLINE auto Invoke(ReturnType (ObjType::*PtrMemFun)(PMFArgTypes...) const, CallableType&& Callable, ArgTypes&&... Args)
+template <
+	typename    PtrMemFunType,
+	typename    CallableType,
+	typename... ArgTypes,
+	typename    ObjType = TMemberFunctionPtrOuter_T<PtrMemFunType>
+>
+FORCEINLINE auto Invoke(PtrMemFunType PtrMemFun, CallableType&& Callable, ArgTypes&&... Args)
 	-> decltype((UE::Core::Private::DereferenceIfNecessary<ObjType>(Forward<CallableType>(Callable)).*PtrMemFun)(Forward<ArgTypes>(Args)...))
 {
 	return (UE::Core::Private::DereferenceIfNecessary<ObjType>(Forward<CallableType>(Callable)).*PtrMemFun)(Forward<ArgTypes>(Args)...);
diff --git a/Engine/Source/Runtime/Core/Public/Traits/MemberFunctionPtrOuter.h b/Engine/Source/Runtime/Core/Public/Traits/MemberFunctionPtrOuter.h
new file mode 100644
--- /dev/null
+++ b/Engine/Source/Runtime/Core/Public/Traits/MemberFunctionPtrOuter.h
@@ -0,0 +1,27 @@
+// Copyright Epic Games, Inc. All Rights Reserved.
+
+#pragma once
+
+/**
+ * Type trait which yields the type of the class given a pointer to a member function of that class, e.g.:
+ *
+ * TMemberFunctionPtrOuter_T<decltype(&FVector::DotProduct)>::Type == FVector
+ */
+template <typename T>
+struct TMemberFunctionPtrOuter;
+
+template <typename ReturnType, typename ObjectType, typename... ArgTypes> struct TMemberFunctionPtrOuter<ReturnType (ObjectType::*)(ArgTypes...)                 > { using Type = ObjectType; };
+template <typename ReturnType, typename ObjectType, typename... ArgTypes> struct TMemberFunctionPtrOuter<ReturnType (ObjectType::*)(ArgTypes...)               & > { using Type = ObjectType; };
+template <typename ReturnType, typename ObjectType, typename... ArgTypes> struct TMemberFunctionPtrOuter<ReturnType (ObjectType::*)(ArgTypes...)               &&> { using Type = ObjectType; };
+template <typename ReturnType, typename ObjectType, typename... ArgTypes> struct TMemberFunctionPtrOuter<ReturnType (ObjectType::*)(ArgTypes...) const           > { using Type = ObjectType; };
+template <typename ReturnType, typename ObjectType, typename... ArgTypes> struct TMemberFunctionPtrOuter<ReturnType (ObjectType::*)(ArgTypes...) const         & > { using Type = ObjectType; };
+template <typename ReturnType, typename ObjectType, typename... ArgTypes> struct TMemberFunctionPtrOuter<ReturnType (ObjectType::*)(ArgTypes...) const         &&> { using Type = ObjectType; };
+template <typename ReturnType, typename ObjectType, typename... ArgTypes> struct TMemberFunctionPtrOuter<ReturnType (ObjectType::*)(ArgTypes...)       volatile  > { using Type = ObjectType; };
+template <typename ReturnType, typename ObjectType, typename... ArgTypes> struct TMemberFunctionPtrOuter<ReturnType (ObjectType::*)(ArgTypes...)       volatile& > { using Type = ObjectType; };
+template <typename ReturnType, typename ObjectType, typename... ArgTypes> struct TMemberFunctionPtrOuter<ReturnType (ObjectType::*)(ArgTypes...)       volatile&&> { using Type = ObjectType; };
+template <typename ReturnType, typename ObjectType, typename... ArgTypes> struct TMemberFunctionPtrOuter<ReturnType (ObjectType::*)(ArgTypes...) const volatile  > { using Type = ObjectType; };
+template <typename ReturnType, typename ObjectType, typename... ArgTypes> struct TMemberFunctionPtrOuter<ReturnType (ObjectType::*)(ArgTypes...) const volatile& > { using Type = ObjectType; };
+template <typename ReturnType, typename ObjectType, typename... ArgTypes> struct TMemberFunctionPtrOuter<ReturnType (ObjectType::*)(ArgTypes...) const volatile&&> { using Type = ObjectType; };
+
+template <typename T>
+using TMemberFunctionPtrOuter_T = typename TMemberFunctionPtrOuter<T>::Type;
