CompressedBuffer: Implemented a stable format for compressed buffers with support for no compression and LZ4

diff --git a/Engine/Source/Developer/DerivedDataCache/Private/DerivedDataCache.cpp b/Engine/Source/Developer/DerivedDataCache/Private/DerivedDataCache.cpp
--- a/Engine/Source/Developer/DerivedDataCache/Private/DerivedDataCache.cpp
+++ b/Engine/Source/Developer/DerivedDataCache/Private/DerivedDataCache.cpp
@@ -834,7 +834,7 @@ void FCache::Put(
 
 	const auto PutPayload = [&Key, Context](const FPayload& Payload)
 	{
-		const FSharedBuffer Buffer = Payload.GetBuffer().GetCompressed();
+		const FSharedBuffer Buffer = Payload.GetBuffer().GetCompressed().Flatten();
 		check(Buffer && Buffer.GetSize() <= MAX_int32);
 		TConstArrayView<uint8> BufferView = MakeArrayView(
 			static_cast<const uint8*>(Buffer.GetData()),
@@ -974,7 +974,7 @@ void FCache::Get(
 		{
 			if (FIoHash::HashBuffer(MakeMemoryView(PayloadData)) == PayloadHash)
 			{
-				return FPayload(PayloadId, PayloadHash, FCompressedBuffer(MakeSharedBufferFromArray(MoveTemp(PayloadData))));
+				return FPayload(PayloadId, PayloadHash, FCompressedBuffer::FromCompressed(MakeSharedBufferFromArray(MoveTemp(PayloadData))));
 			}
 			else
 			{
@@ -1165,7 +1165,7 @@ void FCache::GetPayload(
 		return;
 	}
 
-	Payload = FPayload(Key.Id, FCompressedBuffer(MakeSharedBufferFromArray(MoveTemp(Data))));
+	Payload = FPayload(Key.Id, FCompressedBuffer::FromCompressed(MakeSharedBufferFromArray(MoveTemp(Data))));
 
 	UE_LOG(LogDerivedDataCache, Verbose, TEXT("Cache: GetPayload cache hit for %s from '%.*s'"),
 		*WriteToString<96>(Key), Context.Len(), Context.GetData());
diff --git a/Engine/Source/Developer/DerivedDataCache/Private/DerivedDataCacheRecord.cpp b/Engine/Source/Developer/DerivedDataCache/Private/DerivedDataCacheRecord.cpp
--- a/Engine/Source/Developer/DerivedDataCache/Private/DerivedDataCacheRecord.cpp
+++ b/Engine/Source/Developer/DerivedDataCache/Private/DerivedDataCacheRecord.cpp
@@ -6,6 +6,7 @@
 #include "DerivedDataCacheKey.h"
 #include "Misc/StringBuilder.h"
 #include "Serialization/CompactBinary.h"
+#include "UObject/NameTypes.h"
 
 namespace UE
 {
@@ -171,7 +172,7 @@ void FCacheRecordBuilderInternal::SetMeta(FCbObject&& InMeta)
 FPayloadId FCacheRecordBuilderInternal::SetValue(const FSharedBuffer& Buffer, const FPayloadId& Id)
 {
 	const FPayloadId ValueId = GetOrCreatePayloadId(Id, Buffer);
-	return SetValue(FPayload(ValueId, FCompressedBuffer::Compress(Buffer)));
+	return SetValue(FPayload(ValueId, FCompressedBuffer::Compress(NAME_LZ4, Buffer)));
 }
 
 FPayloadId FCacheRecordBuilderInternal::SetValue(FPayload&& Payload)
@@ -185,7 +186,7 @@ FPayloadId FCacheRecordBuilderInternal::SetValue(FPayload&& Payload)
 FPayloadId FCacheRecordBuilderInternal::AddAttachment(const FSharedBuffer& Buffer, const FPayloadId& Id)
 {
 	const FPayloadId AttachmentId = GetOrCreatePayloadId(Id, Buffer);
-	return AddAttachment(FPayload(AttachmentId, FCompressedBuffer::Compress(Buffer)));
+	return AddAttachment(FPayload(AttachmentId, FCompressedBuffer::Compress(NAME_LZ4, Buffer)));
 }
 
 FPayloadId FCacheRecordBuilderInternal::AddAttachment(FPayload&& Payload)
diff --git a/Engine/Source/Editor/UnrealEd/Private/EditorDomain/EditorDomainSave.cpp b/Engine/Source/Editor/UnrealEd/Private/EditorDomain/EditorDomainSave.cpp
--- a/Engine/Source/Editor/UnrealEd/Private/EditorDomain/EditorDomainSave.cpp
+++ b/Engine/Source/Editor/UnrealEd/Private/EditorDomain/EditorDomainSave.cpp
@@ -752,24 +752,28 @@ bool FEditorDomainSaveServer::TrySavePackage(const FPackagePath& PackagePath, FS
 		return false;
 	}
 
-	TArray64<uint8> TempBytes;
-	bool bBytesRead = FFileHelper::LoadFileToArray(TempBytes, *TempFilename);
-	IFileManager::Get().Delete(*TempFilename);
-	if (!bBytesRead)
+	FSharedBuffer PackageBuffer;
 	{
-		OutErrorMessage = FString::Printf(TEXT("Package %s could not be read from temporary file %s."),
-			*PackagePath.GetDebugName(), *TempFilename, Result.Result);
-		return false;
+		TArray64<uint8> TempBytes;
+		bool bBytesRead = FFileHelper::LoadFileToArray(TempBytes, *TempFilename);
+		IFileManager::Get().Delete(*TempFilename);
+		if (!bBytesRead)
+		{
+			OutErrorMessage = FString::Printf(TEXT("Package %s could not be read from temporary file %s."),
+				*PackagePath.GetDebugName(), *TempFilename, Result.Result);
+			return false;
+		}
+		PackageBuffer = MakeSharedBufferFromArray(MoveTemp(TempBytes));
 	}
 
 	UE::DerivedData::ICache& Cache = GetDerivedDataCacheRef().GetCache();
 	UE::DerivedData::FCacheRecordBuilder RecordBuilder = Cache.CreateRecord(GetEditorDomainPackageKey(PackageDigest));
 	TCbWriter<256> MetaData;
 	MetaData.BeginObject();
-	MetaData << "FileSize" << TempBytes.Num();
+	MetaData << "FileSize" << PackageBuffer.GetSize();
 	MetaData.EndObject();
 	RecordBuilder.SetMeta(MetaData.Save().AsObject());
-	RecordBuilder.SetValue(FSharedBuffer::MakeView(TempBytes.GetData(), TempBytes.Num()));
+	RecordBuilder.SetValue(PackageBuffer);
 	UE::DerivedData::FCacheRecord Record = RecordBuilder.Build();
 	Cache.Put(MakeArrayView(&Record, 1), PackagePath.GetDebugName());
 	return true;
diff --git a/Engine/Source/Editor/UnrealEd/Private/EditorDomain/EditorDomainUtils.cpp b/Engine/Source/Editor/UnrealEd/Private/EditorDomain/EditorDomainUtils.cpp
--- a/Engine/Source/Editor/UnrealEd/Private/EditorDomain/EditorDomainUtils.cpp
+++ b/Engine/Source/Editor/UnrealEd/Private/EditorDomain/EditorDomainUtils.cpp
@@ -19,7 +19,7 @@ namespace EditorDomain
 {
 
 // Change to a new guid when EditorDomain needs to be invalidated
-const TCHAR* EditorDomainVersion = TEXT("34A0B57B-32D5-4229-A6A3-DA0FB7BE0440");
+const TCHAR* EditorDomainVersion = TEXT("BF57FE551C4F4FAFAFFF2EDC17132DC9");
 // Identifier of the CacheBucket for EditorDomainPackages
 const TCHAR* EditorDomainPackageBucketName = TEXT("EditorDomainPackage");
 
diff --git a/Engine/Source/Runtime/Core/Private/Compression/CompressedBuffer.cpp b/Engine/Source/Runtime/Core/Private/Compression/CompressedBuffer.cpp
--- a/Engine/Source/Runtime/Core/Private/Compression/CompressedBuffer.cpp
+++ b/Engine/Source/Runtime/Core/Private/Compression/CompressedBuffer.cpp
@@ -2,22 +2,643 @@
 
 #include "Compression/CompressedBuffer.h"
 
-FCompressedBuffer FCompressedBuffer::Compress(const FSharedBuffer& Buffer)
+#include "Algo/ForEach.h"
+#include "Math/UnrealMathUtility.h"
+#include "Misc/ByteSwap.h"
+#include "Misc/Crc.h"
+#include "Templates/Function.h"
+#include "UObject/NameTypes.h"
+
+THIRD_PARTY_INCLUDES_START
+#include "Compression/lz4.h"
+THIRD_PARTY_INCLUDES_END
+
+namespace UE::CompressedBuffer
+{
+
+///////////////////////////////////////////////////////////////////////////////////////////////////////////////////////
+
+static constexpr uint64 DefaultBlockSize = 256 * 1024;
+
+///////////////////////////////////////////////////////////////////////////////////////////////////////////////////////
+
+/** Method used to compress the data in a compressed buffer. */
+enum class EMethod : uint8
+{
+	/** Header is followed by one uncompressed block. */
+	None = 0,
+	/** Header is followed by an array of compressed block sizes then the compressed blocks. */
+	LZ4 = 4,
+};
+
+///////////////////////////////////////////////////////////////////////////////////////////////////////////////////////
+
+/** Header used on every compressed buffer. Always stored in big-endian format. */
+struct FHeader
+{
+	static constexpr uint32 ExpectedMagic = 0xb7756362; // <dot>ucb
+
+	/** A magic number to identify a compressed buffer. Always 0xb7756362. */
+	uint32 Magic = ExpectedMagic;
+	/** A CRC-32 used to check integrity of the buffer. Uses the polynomial 0x04c11db7. */
+	uint32 Crc32 = 0;
+	/** The method used to compress the buffer. Affects layout of data following the header. */
+	EMethod Method = EMethod::None;
+	/** The reserved bytes must be initialized to zero. */
+	uint8 Reserved[2]{};
+	/** The power of two size of every uncompressed block except the last. Size is 1 << BlockSizeExponent. */
+	uint8 BlockSizeExponent = 0;
+	/** The number of blocks that follow the header. */
+	uint32 BlockCount = 0;
+	/** The total size of the uncompressed data. */
+	uint64 TotalRawSize = 0;
+
+	/** Checks validity of the buffer based on the magic number, method, and CRC-32. */
+	static bool IsValid(const FCompositeBuffer& CompressedData);
+	static bool IsValid(const FSharedBuffer& CompressedData) { return IsValid(FCompositeBuffer(CompressedData)); }
+
+	/** Read a header from a buffer that is at least sizeof(FHeader) without any validation. */
+	static FHeader Read(const FCompositeBuffer& CompressedData)
+	{
+		FHeader Header;
+		if (sizeof(FHeader) <= CompressedData.GetSize())
+		{
+			CompressedData.CopyTo(MakeMemoryView(&Header, &Header + 1));
+			Header.ByteSwap();
+		}
+		return Header;
+	}
+
+	/** Write a header to a memory view that is at least sizeof(FHeader). */
+	void Write(FMutableMemoryView CompressedData) const
+	{
+		FHeader Header = *this;
+		Header.ByteSwap();
+		CompressedData.CopyFrom(MakeMemoryView(&Header, &Header + 1));
+	}
+
+private:
+	void ByteSwap()
+	{
+		Magic = NETWORK_ORDER32(Magic);
+		Crc32 = NETWORK_ORDER32(Crc32);
+		BlockCount = NETWORK_ORDER32(BlockCount);
+		TotalRawSize = NETWORK_ORDER64(TotalRawSize);
+	}
+};
+
+static_assert(sizeof(FHeader) == 24, "FHeader is the wrong size.");
+
+///////////////////////////////////////////////////////////////////////////////////////////////////////////////////////
+
+class FMethod
+{
+public:
+	virtual FCompositeBuffer Compress(const FCompositeBuffer& RawData, uint64 BlockSize = DefaultBlockSize) const = 0;
+	virtual FCompositeBuffer Decompress(const FHeader& Header, const FCompositeBuffer& CompressedData) const = 0;
+	virtual bool TryDecompressTo(const FHeader& Header, const FCompositeBuffer& CompressedData, FMutableMemoryView RawView) const = 0;
+	virtual uint32 CalculateCrc32(const FHeader& Header, const FCompositeBuffer& CompressedData) const;
+
+protected:
+	static uint32 CalculateRangeCrc32(const FCompositeBuffer& CompressedData, uint64 Offset, uint64 Size, uint32 Crc32 = 0);
+};
+
+uint32 FMethod::CalculateCrc32(const FHeader& Header, const FCompositeBuffer& CompressedData) const
+{
+	if (sizeof(FHeader) <= CompressedData.GetSize())
+	{
+		constexpr uint64 MethodOffset = STRUCT_OFFSET(FHeader, Method);
+		return CalculateRangeCrc32(CompressedData, MethodOffset, sizeof(FHeader) - MethodOffset);
+	}
+	return 0;
+}
+
+uint32 FMethod::CalculateRangeCrc32(const FCompositeBuffer& CompressedData, uint64 Offset, uint64 Size, uint32 Crc32)
+{
+	CompressedData.IterateRange(Offset, Size, [&Crc32](FMemoryView Crc32View)
+	{
+		while (const uint64 ViewSize = Crc32View.GetSize())
+		{
+			const int32 Crc32Size = static_cast<int32>(FMath::Min<uint64>(ViewSize, MAX_int32));
+			Crc32 = FCrc::MemCrc32(Crc32View.GetData(), Crc32Size, Crc32);
+			Crc32View += Crc32Size;
+		}
+	});
+	return Crc32;
+}
+
+///////////////////////////////////////////////////////////////////////////////////////////////////////////////////////
+
+class FMethodNone final : public FMethod
+{
+public:
+	FCompositeBuffer Compress(const FCompositeBuffer& RawData, uint64 BlockSize = DefaultBlockSize) const final
+	{
+		FHeader Header;
+		Header.Method = EMethod::None;
+		Header.BlockCount = 1;
+		Header.TotalRawSize = RawData.GetSize();
+
+		FUniqueBuffer HeaderData = FUniqueBuffer::Alloc(sizeof(FHeader));
+		const FMutableMemoryView HeaderDataView = HeaderData;
+
+		FCompositeBuffer CompressedData(FSharedBuffer(MoveTemp(HeaderData)), RawData.MakeOwned());
+		Header.Write(HeaderDataView);
+		Header.Crc32 = CalculateCrc32(Header, CompressedData);
+		Header.Write(HeaderDataView);
+		return CompressedData;
+	}
+
+	FCompositeBuffer Decompress(const FHeader& Header, const FCompositeBuffer& CompressedData) const final
+	{
+		if (Header.Method == EMethod::None &&
+			sizeof(FHeader) + Header.TotalRawSize == CompressedData.GetSize())
+		{
+			return CompressedData.Mid(sizeof(FHeader), Header.TotalRawSize).MakeOwned();
+		}
+		return FCompositeBuffer();
+	}
+
+	bool TryDecompressTo(const FHeader& Header, const FCompositeBuffer& CompressedData, FMutableMemoryView RawView) const final
+	{
+		if (Header.Method == EMethod::None &&
+			Header.TotalRawSize == RawView.GetSize() &&
+			sizeof(FHeader) + Header.TotalRawSize == CompressedData.GetSize())
+		{
+			CompressedData.CopyTo(RawView, sizeof(FHeader));
+			return true;
+		}
+		return false;
+	}
+};
+
+///////////////////////////////////////////////////////////////////////////////////////////////////////////////////////
+
+class FMethodBlock : public FMethod
+{
+public:
+	FCompositeBuffer Compress(const FCompositeBuffer& RawData, uint64 BlockSize = DefaultBlockSize) const final;
+	FCompositeBuffer Decompress(const FHeader& Header, const FCompositeBuffer& CompressedData) const final;
+	bool TryDecompressTo(const FHeader& Header, const FCompositeBuffer& CompressedData, FMutableMemoryView RawView) const final;
+
+protected:
+	uint32 CalculateCrc32(const FHeader& Header, const FCompositeBuffer& CompressedData) const final
+	{
+		const uint64 MetaSize = uint64(Header.BlockCount) * sizeof(uint32);
+		if (sizeof(FHeader) + MetaSize <= CompressedData.GetSize())
+		{
+			const uint32 BaseCrc32 = FMethod::CalculateCrc32(Header, CompressedData);
+			return CalculateRangeCrc32(CompressedData, sizeof(FHeader), MetaSize, BaseCrc32);
+		}
+		return 0;
+	}
+
+	virtual EMethod GetMethod() const = 0;
+	virtual uint64 CompressBlockBound(uint64 RawSize) const = 0;
+	virtual bool CompressBlock(FMutableMemoryView& CompressedData, FMemoryView RawData) const = 0;
+	virtual bool DecompressBlock(FMutableMemoryView RawData, FMemoryView CompressedData) const = 0;
+};
+
+FCompositeBuffer FMethodBlock::Compress(const FCompositeBuffer& RawData, const uint64 BlockSize) const
 {
-	// DDC-TODO: Switch to the new compression interface when it is ready. For now, match the "no compression" method.
-	FUniqueBuffer CompressedBuffer = FUniqueBuffer::Alloc(Buffer.GetSize() + 1);
-	FMemory::Memset(CompressedBuffer.GetData(), 0, 1);
-	FMemory::Memcpy(CompressedBuffer.GetView().RightChop(1).GetData(), Buffer.GetData(), Buffer.GetSize());
-	return FCompressedBuffer(FSharedBuffer(MoveTemp(CompressedBuffer)));
+	checkf(FMath::IsPowerOfTwo(BlockSize) && BlockSize <= MAX_uint32,
+		TEXT("BlockSize must be a 32-bit power of two but was %" UINT64_FMT "."), BlockSize);
+	const uint64 RawSize = RawData.GetSize();
+
+	const uint64 BlockCount = FMath::DivideAndRoundUp(RawSize, BlockSize);
+	checkf(BlockCount <= MAX_uint32, TEXT("Raw data of size %" UINT64_FMT " with block size %" UINT64_FMT " requires ")
+		TEXT("%" UINT64_FMT " blocks, but the limit is %u."), RawSize, BlockSize, BlockCount, MAX_uint32);
+
+	// Allocate the buffer for the header, metadata, and compressed blocks.
+	const uint64 MetaSize = BlockCount * sizeof(uint32);
+	const uint64 CompressedDataSize = sizeof(FHeader) + MetaSize +
+		BlockCount == 0 ? 0 :
+		BlockCount == 1 ? CompressBlockBound(RawSize) : CompressBlockBound(BlockSize) - BlockSize + RawSize;
+	FUniqueBuffer CompressedData = FUniqueBuffer::Alloc(CompressedDataSize);
+	const FMutableMemoryView CompressedDataView = CompressedData.GetView();
+
+	// Compress the raw data in blocks and store the raw data for incompressible blocks.
+	TArray64<uint32> CompressedBlockSizes;
+	CompressedBlockSizes.Reserve(static_cast<uint32>(BlockCount));
+	uint64 CompressedSize = 0;
+	{
+		FUniqueBuffer RawBlockCopy;
+		FMutableMemoryView CompressedBlocksView = CompressedDataView + sizeof(FHeader) + MetaSize;
+		for (uint64 RawOffset = 0; RawOffset < RawSize;)
+		{
+			const uint64 RawBlockSize = FMath::Min(RawSize - RawOffset, BlockSize);
+			const FMemoryView RawBlock = RawData.ViewOrCopyRange(RawOffset, RawBlockSize, RawBlockCopy);
+
+			FMutableMemoryView CompressedBlock = CompressedBlocksView;
+			if (!CompressBlock(CompressedBlock, RawBlock))
+			{
+				return FCompositeBuffer();
+			}
+
+			uint64 CompressedBlockSize = CompressedBlock.GetSize();
+			if (RawBlockSize <= CompressedBlockSize)
+			{
+				CompressedBlockSize = RawBlockSize;
+				CompressedBlocksView = CompressedBlocksView.CopyFrom(RawBlock);
+			}
+			else
+			{
+				CompressedBlocksView += CompressedBlockSize;
+			}
+
+			CompressedBlockSizes.Add(static_cast<uint32>(CompressedBlockSize));
+			CompressedSize += CompressedBlockSize;
+			RawOffset += RawBlockSize;
+		}
+	}
+
+	// Return an uncompressed buffer if the compressed data is larger than the raw data.
+	if (RawSize <= MetaSize + CompressedSize)
+	{
+		CompressedData.Reset();
+		return FMethodNone().Compress(RawData, BlockSize);
+	}
+
+	// Write the header and calculate the CRC-32.
+	Algo::ForEach(CompressedBlockSizes, [](uint32& Size) { Size = NETWORK_ORDER32(Size); });
+	CompressedDataView.Mid(sizeof(FHeader), MetaSize).CopyFrom(MakeMemoryView(CompressedBlockSizes));
+
+	const FMemoryView CompositeView = CompressedDataView.Left(sizeof(FHeader) + MetaSize + CompressedSize);
+	const FCompositeBuffer Composite(FSharedBuffer::MakeView(CompositeView, FSharedBuffer(MoveTemp(CompressedData))));
+
+	FHeader Header;
+	Header.Method = GetMethod();
+	Header.BlockSizeExponent = static_cast<uint8>(FMath::FloorLog2_64(BlockSize));
+	Header.BlockCount = static_cast<uint32>(BlockCount);
+	Header.TotalRawSize = RawSize;
+	Header.Write(CompressedDataView);
+	Header.Crc32 = CalculateCrc32(Header, Composite);
+	Header.Write(CompressedDataView);
+
+	return Composite;
 }
 
-FSharedBuffer FCompressedBuffer::Decompress(const FCompressedBuffer& Buffer)
+FCompositeBuffer FMethodBlock::Decompress(const FHeader& Header, const FCompositeBuffer& CompressedData) const
 {
-	// DDC-TODO: Switch to the new compression interface when it is ready. For now, match the "no compression" method.
-	if (FSharedBuffer Compressed = Buffer.GetCompressed())
+	if (Header.BlockCount == 0)
+	{
+		return FCompositeBuffer();
+	}
+
+	// The raw data cannot reference the compressed data unless it is owned.
+	if (!CompressedData.IsOwned())
+	{
+		FUniqueBuffer Buffer = FUniqueBuffer::Alloc(Header.TotalRawSize);
+		return TryDecompressTo(Header, CompressedData, Buffer) ? FCompositeBuffer(FSharedBuffer(MoveTemp(Buffer))) : FCompositeBuffer();
+	}
+
+	TArray64<uint32> CompressedBlockSizes;
+	CompressedBlockSizes.AddUninitialized(Header.BlockCount);
+	CompressedData.CopyTo(MakeMemoryView(CompressedBlockSizes), sizeof(FHeader));
+	Algo::ForEach(CompressedBlockSizes, [](uint32& Size) { Size = NETWORK_ORDER32(Size); });
+
+	// Allocate the buffer for the raw blocks that were compressed.
+	FSharedBuffer RawData;
+	FMutableMemoryView RawDataView;
+	const uint64 BlockSize = uint64(1) << Header.BlockSizeExponent;
 	{
-		const FMemoryView View = Compressed.GetView() + 1;
-		return FSharedBuffer::MakeView(View, MoveTemp(Compressed));
+		uint64 RawDataSize = 0;
+		uint64 RemainingRawSize = Header.TotalRawSize;
+		for (const uint32 CompressedBlockSize : CompressedBlockSizes)
+		{
+			const uint64 RawBlockSize = FMath::Min(RemainingRawSize, BlockSize);
+			if (CompressedBlockSize < BlockSize)
+			{
+				RawDataSize += RawBlockSize;
+			}
+			RemainingRawSize -= RawBlockSize;
+		}
+		FUniqueBuffer RawDataBuffer = FUniqueBuffer::Alloc(RawDataSize);
+		RawDataView = RawDataBuffer;
+		RawData = FSharedBuffer(MoveTemp(RawDataBuffer));
+	}
+
+	// Decompress the compressed data in blocks and reference the uncompressed blocks.
+	uint64 PendingCompressedSegmentOffset = sizeof(FHeader) + uint64(Header.BlockCount) * sizeof(uint32);
+	uint64 PendingCompressedSegmentSize = 0;
+	uint64 PendingRawSegmentOffset = 0;
+	uint64 PendingRawSegmentSize = 0;
+	TArray<FSharedBuffer, TInlineAllocator<1>> Segments;
+
+	const auto CommitPendingCompressedSegment = [&PendingCompressedSegmentOffset, &PendingCompressedSegmentSize, &CompressedData, &Segments]
+	{
+		if (PendingCompressedSegmentSize)
+		{
+			CompressedData.IterateRange(PendingCompressedSegmentOffset, PendingCompressedSegmentSize,
+				[&Segments](FMemoryView View, const FSharedBuffer& ViewOuter)
+				{
+					Segments.Add(FSharedBuffer::MakeView(View, ViewOuter));
+				});
+			PendingCompressedSegmentOffset += PendingCompressedSegmentSize;
+			PendingCompressedSegmentSize = 0;
+		}
+	};
+
+	const auto CommitPendingRawSegment = [&PendingRawSegmentOffset, &PendingRawSegmentSize, &RawData, &Segments]
+	{
+		if (PendingRawSegmentSize)
+		{
+			const FMemoryView PendingSegment = RawData.GetView().Mid(PendingRawSegmentOffset, PendingRawSegmentSize);
+			Segments.Add(FSharedBuffer::MakeView(PendingSegment, RawData));
+			PendingRawSegmentOffset += PendingRawSegmentSize;
+			PendingRawSegmentSize = 0;
+		}
+	};
+
+	FUniqueBuffer CompressedBlockCopy;
+	uint64 RemainingRawSize = Header.TotalRawSize;
+	uint64 RemainingCompressedSize = CompressedData.GetSize();
+	for (const uint32 CompressedBlockSize : CompressedBlockSizes)
+	{
+		if (RemainingCompressedSize < CompressedBlockSize)
+		{
+			return FCompositeBuffer();
+		}
+
+		const uint64 RawBlockSize = FMath::Min(RemainingRawSize, BlockSize);
+		if (RawBlockSize == CompressedBlockSize)
+		{
+			CommitPendingRawSegment();
+			PendingCompressedSegmentSize += RawBlockSize;
+		}
+		else
+		{
+			CommitPendingCompressedSegment();
+			const FMemoryView CompressedBlock = CompressedData.ViewOrCopyRange(PendingCompressedSegmentOffset, CompressedBlockSize, CompressedBlockCopy);
+			if (!DecompressBlock(RawDataView.Left(RawBlockSize), CompressedBlock))
+			{
+				return FCompositeBuffer();
+			}
+			PendingCompressedSegmentOffset += CompressedBlockSize;
+			PendingRawSegmentSize += RawBlockSize;
+			RawDataView += RawBlockSize;
+		}
+
+		RemainingCompressedSize -= CompressedBlockSize;
+		RemainingRawSize -= RawBlockSize;
+	}
+
+	CommitPendingCompressedSegment();
+	CommitPendingRawSegment();
+
+	return FCompositeBuffer(MoveTemp(Segments));
+}
+
+bool FMethodBlock::TryDecompressTo(const FHeader& Header, const FCompositeBuffer& CompressedData, FMutableMemoryView RawView) const
+{
+	if (Header.TotalRawSize != RawView.GetSize())
+	{
+		return false;
+	}
+
+	TArray64<uint32> CompressedBlockSizes;
+	CompressedBlockSizes.AddUninitialized(Header.BlockCount);
+	CompressedData.CopyTo(MakeMemoryView(CompressedBlockSizes), sizeof(FHeader));
+	Algo::ForEach(CompressedBlockSizes, [](uint32& Size) { Size = NETWORK_ORDER32(Size); });
+
+	FUniqueBuffer CompressedBlockCopy;
+	const uint64 BlockSize = uint64(1) << Header.BlockSizeExponent;
+	uint64 CompressedOffset = sizeof(FHeader) + uint64(Header.BlockCount) * sizeof(uint32);
+	uint64 RemainingRawSize = Header.TotalRawSize;
+	uint64 RemainingCompressedSize = CompressedData.GetSize();
+	for (uint32 CompressedBlockSize : CompressedBlockSizes)
+	{
+		if (RemainingCompressedSize < CompressedBlockSize)
+		{
+			return false;
+		}
+
+		const uint64 RawBlockSize = FMath::Min(RemainingRawSize, BlockSize);
+		if (RawBlockSize == CompressedBlockSize)
+		{
+			CompressedData.CopyTo(RawView.Left(RawBlockSize), CompressedOffset);
+		}
+		else
+		{
+			const FMemoryView CompressedBlock = CompressedData.ViewOrCopyRange(CompressedOffset, CompressedBlockSize, CompressedBlockCopy);
+			if (!DecompressBlock(RawView.Left(RawBlockSize), CompressedBlock))
+			{
+				return false;
+			}
+		}
+
+		RemainingCompressedSize -= CompressedBlockSize;
+		RemainingRawSize -= RawBlockSize;
+		CompressedOffset += CompressedBlockSize;
+		RawView += RawBlockSize;
+	}
+
+	return RemainingRawSize == 0;
+}
+
+///////////////////////////////////////////////////////////////////////////////////////////////////////////////////////
+
+class FMethodLZ4 final : public FMethodBlock
+{
+protected:
+	EMethod GetMethod() const final
+	{
+		return EMethod::LZ4;
+	}
+
+	uint64 CompressBlockBound(uint64 RawSize) const final
+	{
+		if (RawSize <= LZ4_MAX_INPUT_SIZE)
+		{
+			return static_cast<uint64>(LZ4_compressBound(static_cast<int>(RawSize)));
+		}
+		return 0;
+	}
+
+	bool CompressBlock(FMutableMemoryView& CompressedData, FMemoryView RawData) const final
+	{
+		if (RawData.GetSize() <= LZ4_MAX_INPUT_SIZE)
+		{
+			const int Size = LZ4_compress_default(
+				static_cast<const char*>(RawData.GetData()), static_cast<char*>(CompressedData.GetData()),
+				static_cast<int>(RawData.GetSize()), static_cast<int>(FMath::Min<uint64>(CompressedData.GetSize(), MAX_int32)));
+			CompressedData.LeftInline(static_cast<uint64>(Size));
+			return Size > 0;
+		}
+		return false;
+	}
+
+	bool DecompressBlock(FMutableMemoryView RawData, FMemoryView CompressedData) const final
+	{
+		if (CompressedData.GetSize() <= MAX_int32)
+		{
+			const int Size = LZ4_decompress_safe(
+				static_cast<const char*>(CompressedData.GetData()),
+				static_cast<char*>(RawData.GetData()),
+				static_cast<int>(CompressedData.GetSize()),
+				static_cast<int>(FMath::Min<uint64>(RawData.GetSize(), LZ4_MAX_INPUT_SIZE)));
+			return static_cast<uint64>(Size) == RawData.GetSize();
+		}
+		return false;
+	}
+};
+
+///////////////////////////////////////////////////////////////////////////////////////////////////////////////////////
+
+static const FMethod* GetMethod(EMethod Method)
+{
+	static FMethodNone MethodNone;
+	static FMethodLZ4 MethodLZ4;
+	switch (Method)
+	{
+	default:
+		return nullptr;
+	case EMethod::None:
+		return &MethodNone;
+	case EMethod::LZ4:
+		return &MethodLZ4;
+	}
+}
+
+static const FMethod* GetMethod(FName FormatName)
+{
+	if (FormatName.IsNone())
+	{
+		return GetMethod(EMethod::None);
+	}
+	if (FormatName == NAME_LZ4)
+	{
+		return GetMethod(EMethod::LZ4);
+	}
+	return nullptr;
+}
+
+template <typename BufferType>
+inline FCompositeBuffer ValidBufferOrEmpty(BufferType&& CompressedData)
+{
+	return FHeader::IsValid(CompressedData) ? FCompositeBuffer(Forward<BufferType>(CompressedData)) : FCompositeBuffer();
+}
+
+///////////////////////////////////////////////////////////////////////////////////////////////////////////////////////
+
+bool FHeader::IsValid(const FCompositeBuffer& CompressedData)
+{
+	if (sizeof(FHeader) <= CompressedData.GetSize())
+	{
+		const FHeader Header = Read(CompressedData);
+		if (Header.Magic == FHeader::ExpectedMagic)
+		{
+			if (const FMethod* Method = GetMethod(Header.Method))
+			{
+				if (Header.Crc32 == Method->CalculateCrc32(Header, CompressedData))
+				{
+					return true;
+				}
+			}
+		}
+	}
+	return false;
+}
+
+///////////////////////////////////////////////////////////////////////////////////////////////////////////////////////
+
+} // UE::CompressedBuffer
+
+FCompressedBuffer FCompressedBuffer::Compress(FName FormatName, const FCompositeBuffer& RawData)
+{
+	using namespace UE::CompressedBuffer;
+	FCompressedBuffer Local;
+	if (const FMethod* Method = GetMethod(FormatName))
+	{
+		Local.CompressedData = Method->Compress(RawData);
+	}
+	return Local;
+}
+
+FCompressedBuffer FCompressedBuffer::Compress(FName FormatName, const FSharedBuffer& RawData)
+{
+	return Compress(FormatName, FCompositeBuffer(RawData));
+}
+
+FCompressedBuffer FCompressedBuffer::FromCompressed(const FCompositeBuffer& InCompressedData)
+{
+	FCompressedBuffer Local;
+	Local.CompressedData = UE::CompressedBuffer::ValidBufferOrEmpty(InCompressedData);
+	return Local;
+}
+
+FCompressedBuffer FCompressedBuffer::FromCompressed(FCompositeBuffer&& InCompressedData)
+{
+	FCompressedBuffer Local;
+	Local.CompressedData = UE::CompressedBuffer::ValidBufferOrEmpty(MoveTemp(InCompressedData));
+	return Local;
+}
+
+FCompressedBuffer FCompressedBuffer::FromCompressed(const FSharedBuffer& InCompressedData)
+{
+	FCompressedBuffer Local;
+	Local.CompressedData = UE::CompressedBuffer::ValidBufferOrEmpty(InCompressedData);
+	return Local;
+}
+
+FCompressedBuffer FCompressedBuffer::FromCompressed(FSharedBuffer&& InCompressedData)
+{
+	FCompressedBuffer Local;
+	Local.CompressedData = UE::CompressedBuffer::ValidBufferOrEmpty(MoveTemp(InCompressedData));
+	return Local;
+}
+
+uint64 FCompressedBuffer::GetRawSize() const
+{
+	return CompressedData ? UE::CompressedBuffer::FHeader::Read(CompressedData).TotalRawSize : 0;
+}
+
+bool FCompressedBuffer::TryDecompressTo(FMutableMemoryView RawView) const
+{
+	using namespace UE::CompressedBuffer;
+	if (CompressedData)
+	{
+		const FHeader Header = FHeader::Read(CompressedData);
+		if (const FMethod* Method = GetMethod(Header.Method))
+		{
+			return Method->TryDecompressTo(Header, CompressedData, RawView);
+		}
+	}
+	return false;
+}
+
+FSharedBuffer FCompressedBuffer::Decompress() const
+{
+	using namespace UE::CompressedBuffer;
+	if (CompressedData)
+	{
+		const FHeader Header = FHeader::Read(CompressedData);
+		if (const FMethod* Method = GetMethod(Header.Method))
+		{
+			if (Header.Method == EMethod::None)
+			{
+				return Method->Decompress(Header, CompressedData).Flatten();
+			}
+			FUniqueBuffer RawData = FUniqueBuffer::Alloc(Header.TotalRawSize);
+			if (Method->TryDecompressTo(Header, CompressedData, RawData))
+			{
+				return FSharedBuffer(MoveTemp(RawData));
+			}
+		}
 	}
 	return FSharedBuffer();
 }
+
+FCompositeBuffer FCompressedBuffer::DecompressToComposite() const
+{
+	using namespace UE::CompressedBuffer;
+	if (CompressedData)
+	{
+		const FHeader Header = FHeader::Read(CompressedData);
+		if (const FMethod* Method = GetMethod(Header.Method))
+		{
+			return Method->Decompress(Header, CompressedData);
+		}
+	}
+	return FCompositeBuffer();
+}
diff --git a/Engine/Source/Runtime/Core/Public/Compression/CompressedBuffer.h b/Engine/Source/Runtime/Core/Public/Compression/CompressedBuffer.h
--- a/Engine/Source/Runtime/Core/Public/Compression/CompressedBuffer.h
+++ b/Engine/Source/Runtime/Core/Public/Compression/CompressedBuffer.h
@@ -2,48 +2,82 @@
 
 #pragma once
 
-#include "Memory/SharedBuffer.h"
+#include "Memory/CompositeBuffer.h"
+#include "Memory/MemoryFwd.h"
+
+class FName;
 
 /**
- * A compressed buffer stores data in a self-contained format.
+ * A compressed buffer stores compressed data in a self-contained format.
  *
  * A buffer is self-contained in the sense that it can be decompressed without external knowledge
- * of the compression format or the size of the uncompressed data.
+ * of the compression format or the size of the raw data.
  */
 class FCompressedBuffer
 {
 public:
-	/** Construct a null compressed buffer. */
-	FCompressedBuffer() = default;
+	/**
+	 * Compress the buffer using the requested compression format.
+	 *
+	 * @param FormatName   One of NAME_None, NAME_LZ4.
+	 * @param RawData      Raw data to compress. NAME_None will reference owned raw data.
+	 * @return An owned compressed buffer, or empty on error.
+	 */
+	[[nodiscard]] CORE_API static FCompressedBuffer Compress(FName FormatName, const FCompositeBuffer& RawData);
+	[[nodiscard]] CORE_API static FCompressedBuffer Compress(FName FormatName, const FSharedBuffer& RawData);
+
+	/**
+	 * Construct from a compressed buffer previously created by Compress().
+	 *
+	 * @return A compressed buffer, or empty on error, such as an invalid format or corrupt header.
+	 */
+	[[nodiscard]] CORE_API static FCompressedBuffer FromCompressed(const FCompositeBuffer& CompressedData);
+	[[nodiscard]] CORE_API static FCompressedBuffer FromCompressed(FCompositeBuffer&& CompressedData);
+	[[nodiscard]] CORE_API static FCompressedBuffer FromCompressed(const FSharedBuffer& CompressedData);
+	[[nodiscard]] CORE_API static FCompressedBuffer FromCompressed(FSharedBuffer&& CompressedData);
+
+	/** Returns true if the compressed buffer is not empty. */
+	[[nodiscard]] inline explicit operator bool() const { return !IsEmpty(); }
+
+	/** Returns true if the compressed buffer is empty. */
+	[[nodiscard]] inline bool IsEmpty() const { return CompressedData.IsEmpty(); }
 
-	/** Construct from a compressed shared buffer. */
-	inline explicit FCompressedBuffer(FSharedBuffer InCompressed)
-		: Compressed(MoveTemp(InCompressed))
-	{
-	}
+	/** Returns true if the composite buffer is owned. */
+	[[nodiscard]] inline bool IsOwned() const { return CompressedData.IsOwned(); }
 
-	/** Returns true if the compressed buffer is not null. */
-	inline explicit operator bool() const { return !IsNull(); }
+	/** Returns a copy of the compressed buffer that owns its underlying memory. */
+	[[nodiscard]] inline FCompressedBuffer MakeOwned() const & { return FromCompressed(CompressedData.MakeOwned()); }
+	[[nodiscard]] inline FCompressedBuffer MakeOwned() && { return FromCompressed(MoveTemp(CompressedData).MakeOwned()); }
 
-	/** Returns true if the compressed buffer is null. */
-	inline bool IsNull() const { return Compressed.IsNull(); }
+	/** Returns a composite buffer containing the compressed data. May be empty. May not be owned. */
+	[[nodiscard]] inline const FCompositeBuffer& GetCompressed() const & { return CompressedData; }
+	[[nodiscard]] inline FCompositeBuffer GetCompressed() && { return MoveTemp(CompressedData); }
 
-	/** Returns the shared buffer containing the compressed data. May be null. */
-	inline FSharedBuffer GetCompressed() const & { return Compressed; }
-	inline FSharedBuffer GetCompressed() && { return MoveTemp(Compressed); }
+	/** Returns the size of the compressed data. Zero if this is empty. */
+	[[nodiscard]] inline uint64 GetCompressedSize() const { return CompressedData.GetSize(); }
 
-	/** Returns the size of the compressed data. */
-	inline uint64 GetCompressedSize() const { return Compressed.GetSize(); }
+	/** Returns the size of the raw data. Zero on error or if this is empty. */
+	[[nodiscard]] CORE_API uint64 GetRawSize() const;
 
-	/** Returns an owned shared buffer containing the uncompresesd data. May be null. */
-	inline FSharedBuffer Decompress() const { return Decompress(*this); }
+	/**
+	 * Decompress into a memory view that is exactly GetRawSize() bytes.
+	 */
+	[[nodiscard]] CORE_API bool TryDecompressTo(FMutableMemoryView RawView) const;
 
-	/** Compress the buffer using default compression settings. Returns null if input is null. */
-	CORE_API static FCompressedBuffer Compress(const FSharedBuffer& Buffer);
+	/**
+	 * Decompress into an owned buffer.
+	 *
+	 * @return An owned buffer containing the raw data, or empty on error or if this is empty.
+	 */
+	[[nodiscard]] CORE_API FSharedBuffer Decompress() const;
 
-	/** Decompress the compressed buffer. Returns null if input is null. */
-	CORE_API static FSharedBuffer Decompress(const FCompressedBuffer& Buffer);
+	/**
+	 * Decompress into an owned composite buffer.
+	 *
+	 * @return An owned buffer containing the raw data, or empty on error or if this is empty.
+	 */
+	[[nodiscard]] CORE_API FCompositeBuffer DecompressToComposite() const;
 
 private:
-	FSharedBuffer Compressed;
+	FCompositeBuffer CompressedData;
 };
