Non-unity build fix.

#rnx

diff --git a/Engine/Source/Runtime/Core/Public/Templates/Invoke.h b/Engine/Source/Runtime/Core/Public/Templates/Invoke.h
--- a/Engine/Source/Runtime/Core/Public/Templates/Invoke.h
+++ b/Engine/Source/Runtime/Core/Public/Templates/Invoke.h
@@ -4,6 +4,7 @@
 
 #include "CoreTypes.h"
 #include "Traits/MemberFunctionPtrOuter.h"
+#include "Templates/UnrealTemplate.h" // For Forward<>()
 #include <type_traits>
 
 
