SharedBuffer: Added MoveToShared() and MoveToUnique() to convert between buffer types

#rnx

diff --git a/Engine/Source/Runtime/Core/Private/Compression/CompressedBuffer.cpp b/Engine/Source/Runtime/Core/Private/Compression/CompressedBuffer.cpp
--- a/Engine/Source/Runtime/Core/Private/Compression/CompressedBuffer.cpp
+++ b/Engine/Source/Runtime/Core/Private/Compression/CompressedBuffer.cpp
@@ -142,7 +142,7 @@ public:
 		FUniqueBuffer HeaderData = FUniqueBuffer::Alloc(sizeof(FHeader));
 		const FMutableMemoryView HeaderDataView = HeaderData;
 
-		FCompositeBuffer CompressedData(FSharedBuffer(MoveTemp(HeaderData)), RawData.MakeOwned());
+		FCompositeBuffer CompressedData(HeaderData.MoveToShared(), RawData.MakeOwned());
 		Header.Write(HeaderDataView);
 		Header.Crc32 = CalculateCrc32(Header, CompressedData);
 		Header.Write(HeaderDataView);
@@ -264,7 +264,7 @@ FCompositeBuffer FMethodBlock::Compress(const FCompositeBuffer& RawData, const u
 	CompressedDataView.Mid(sizeof(FHeader), MetaSize).CopyFrom(MakeMemoryView(CompressedBlockSizes));
 
 	const FMemoryView CompositeView = CompressedDataView.Left(sizeof(FHeader) + MetaSize + CompressedSize);
-	const FCompositeBuffer Composite(FSharedBuffer::MakeView(CompositeView, FSharedBuffer(MoveTemp(CompressedData))));
+	const FCompositeBuffer Composite(FSharedBuffer::MakeView(CompositeView, CompressedData.MoveToShared()));
 
 	FHeader Header;
 	Header.Method = GetMethod();
@@ -289,7 +289,7 @@ FCompositeBuffer FMethodBlock::Decompress(const FHeader& Header, const FComposit
 	if (!CompressedData.IsOwned())
 	{
 		FUniqueBuffer Buffer = FUniqueBuffer::Alloc(Header.TotalRawSize);
-		return TryDecompressTo(Header, CompressedData, Buffer) ? FCompositeBuffer(FSharedBuffer(MoveTemp(Buffer))) : FCompositeBuffer();
+		return TryDecompressTo(Header, CompressedData, Buffer) ? FCompositeBuffer(Buffer.MoveToShared()) : FCompositeBuffer();
 	}
 
 	TArray64<uint32> CompressedBlockSizes;
@@ -315,7 +315,7 @@ FCompositeBuffer FMethodBlock::Decompress(const FHeader& Header, const FComposit
 		}
 		FUniqueBuffer RawDataBuffer = FUniqueBuffer::Alloc(RawDataSize);
 		RawDataView = RawDataBuffer;
-		RawData = FSharedBuffer(MoveTemp(RawDataBuffer));
+		RawData = RawDataBuffer.MoveToShared();
 	}
 
 	// Decompress the compressed data in blocks and reference the uncompressed blocks.
@@ -622,7 +622,7 @@ FSharedBuffer FCompressedBuffer::Decompress() const
 			FUniqueBuffer RawData = FUniqueBuffer::Alloc(Header.TotalRawSize);
 			if (Method->TryDecompressTo(Header, CompressedData, RawData))
 			{
-				return FSharedBuffer(MoveTemp(RawData));
+				return RawData.MoveToShared();
 			}
 		}
 	}
diff --git a/Engine/Source/Runtime/Core/Private/Memory/CompositeBuffer.cpp b/Engine/Source/Runtime/Core/Private/Memory/CompositeBuffer.cpp
--- a/Engine/Source/Runtime/Core/Private/Memory/CompositeBuffer.cpp
+++ b/Engine/Source/Runtime/Core/Private/Memory/CompositeBuffer.cpp
@@ -46,7 +46,7 @@ FSharedBuffer FCompositeBuffer::Flatten() const &
 	default:
 		FUniqueBuffer Buffer = FUniqueBuffer::Alloc(GetSize());
 		Algo::TransformAccumulate(Segments, &FSharedBuffer::GetView, Buffer.GetView(), &FMutableMemoryView::CopyFrom);
-		return FSharedBuffer(MoveTemp(Buffer));
+		return Buffer.MoveToShared();
 	}
 }
 
diff --git a/Engine/Source/Runtime/Core/Private/Memory/SharedBuffer.cpp b/Engine/Source/Runtime/Core/Private/Memory/SharedBuffer.cpp
--- a/Engine/Source/Runtime/Core/Private/Memory/SharedBuffer.cpp
+++ b/Engine/Source/Runtime/Core/Private/Memory/SharedBuffer.cpp
@@ -128,21 +128,13 @@ FUniqueBuffer FUniqueBuffer::MakeView(void* Data, uint64 Size)
 	return FUniqueBuffer(new BufferOwnerPrivate::FBufferOwnerView(Data, Size));
 }
 
-FUniqueBuffer FUniqueBuffer::MakeUnique(FSharedBuffer&& Buffer)
+FUniqueBuffer::FUniqueBuffer(FBufferOwner* InOwner)
+	: Owner(InOwner)
 {
-	OwnerPtrType ExistingOwner = ToPrivateOwnerPtr(MoveTemp(Buffer));
-	if (!ExistingOwner || (ExistingOwner->IsOwned() && ExistingOwner->GetTotalRefCount() == 1))
-	{
-		return FUniqueBuffer(MoveTemp(ExistingOwner));
-	}
-	else
-	{
-		return Clone(ExistingOwner->GetData(), ExistingOwner->GetSize());
-	}
 }
 
-FUniqueBuffer::FUniqueBuffer(FBufferOwner* InOwner)
-	: Owner(InOwner)
+FUniqueBuffer::FUniqueBuffer(BufferOwnerPrivate::TBufferOwnerPtr<BufferOwnerPrivate::FSharedOps>&& SharedOwner)
+	: Owner(MoveTemp(SharedOwner))
 {
 }
 
@@ -164,16 +156,21 @@ void FUniqueBuffer::Materialize() const
 	}
 }
 
+FSharedBuffer FUniqueBuffer::MoveToShared()
+{
+	return FSharedBuffer(ToPrivateOwnerPtr(MoveTemp(*this)));
+}
+
 ///////////////////////////////////////////////////////////////////////////////////////////////////////////////////////
 
 FSharedBuffer FSharedBuffer::Clone(FMemoryView View)
 {
-	return FSharedBuffer(FUniqueBuffer::Clone(View));
+	return FUniqueBuffer::Clone(View).MoveToShared();
 }
 
 FSharedBuffer FSharedBuffer::Clone(const void* Data, uint64 Size)
 {
-	return FSharedBuffer(FUniqueBuffer::Clone(Data, Size));
+	return FUniqueBuffer::Clone(Data, Size).MoveToShared();
 }
 
 FSharedBuffer FSharedBuffer::MakeView(FMemoryView View)
@@ -229,6 +226,11 @@ FSharedBuffer::FSharedBuffer(FBufferOwner* InOwner)
 {
 }
 
+FSharedBuffer::FSharedBuffer(BufferOwnerPrivate::TBufferOwnerPtr<BufferOwnerPrivate::FSharedOps>&& SharedOwner)
+	: Owner(MoveTemp(SharedOwner))
+{
+}
+
 FSharedBuffer::FSharedBuffer(const BufferOwnerPrivate::TBufferOwnerPtr<BufferOwnerPrivate::FWeakOps>& WeakOwner)
 	: Owner(WeakOwner)
 {
@@ -257,6 +259,19 @@ void FSharedBuffer::Materialize() const
 	}
 }
 
+FUniqueBuffer FSharedBuffer::MoveToUnique()
+{
+	OwnerPtrType ExistingOwner = ToPrivateOwnerPtr(MoveTemp(*this));
+	if (!ExistingOwner || (ExistingOwner->IsOwned() && ExistingOwner->GetTotalRefCount() == 1))
+	{
+		return FUniqueBuffer(MoveTemp(ExistingOwner));
+	}
+	else
+	{
+		return FUniqueBuffer::Clone(ExistingOwner->GetData(), ExistingOwner->GetSize());
+	}
+}
+
 ///////////////////////////////////////////////////////////////////////////////////////////////////////////////////////
 
 FWeakSharedBuffer::FWeakSharedBuffer(const FSharedBuffer& Buffer)
diff --git a/Engine/Source/Runtime/Core/Private/Serialization/CompactBinary.cpp b/Engine/Source/Runtime/Core/Private/Serialization/CompactBinary.cpp
--- a/Engine/Source/Runtime/Core/Private/Serialization/CompactBinary.cpp
+++ b/Engine/Source/Runtime/Core/Private/Serialization/CompactBinary.cpp
@@ -917,7 +917,7 @@ FCbFieldIterator FCbFieldIterator::CloneRange(const FCbFieldViewIterator& It)
 	{
 		FUniqueBuffer Buffer = FUniqueBuffer::Alloc(It.GetRangeSize());
 		It.CopyRangeTo(Buffer);
-		return MakeRange(FSharedBuffer(MoveTemp(Buffer)));
+		return MakeRange(Buffer.MoveToShared());
 	}
 }
 
diff --git a/Engine/Source/Runtime/Core/Private/Serialization/CompactBinarySerialization.cpp b/Engine/Source/Runtime/Core/Private/Serialization/CompactBinarySerialization.cpp
--- a/Engine/Source/Runtime/Core/Private/Serialization/CompactBinarySerialization.cpp
+++ b/Engine/Source/Runtime/Core/Private/Serialization/CompactBinarySerialization.cpp
@@ -186,7 +186,7 @@ FCbField LoadCompactBinary(FArchive& Ar, FCbBufferAllocator Allocator)
 	{
 		Ar.Serialize(View.GetData(), static_cast<int64>(View.GetSize()));
 	}
-	return FCbField(FSharedBuffer(MoveTemp(Buffer)));
+	return FCbField(Buffer.MoveToShared());
 }
 
 ///////////////////////////////////////////////////////////////////////////////////////////////////////////////////////
diff --git a/Engine/Source/Runtime/Core/Private/Serialization/CompactBinaryWriter.cpp b/Engine/Source/Runtime/Core/Private/Serialization/CompactBinaryWriter.cpp
--- a/Engine/Source/Runtime/Core/Private/Serialization/CompactBinaryWriter.cpp
+++ b/Engine/Source/Runtime/Core/Private/Serialization/CompactBinaryWriter.cpp
@@ -90,7 +90,7 @@ FCbFieldIterator FCbWriter::Save() const
 	const uint64 Size = GetSaveSize();
 	FUniqueBuffer Buffer = FUniqueBuffer::Alloc(Size);
 	const FCbFieldViewIterator Output = Save(Buffer);
-	return FCbFieldIterator::MakeRangeView(Output, FSharedBuffer(MoveTemp(Buffer)));
+	return FCbFieldIterator::MakeRangeView(Output, Buffer.MoveToShared());
 }
 
 FCbFieldViewIterator FCbWriter::Save(const FMutableMemoryView Buffer) const
diff --git a/Engine/Source/Runtime/Core/Private/Tests/Memory/SharedBufferTest.cpp b/Engine/Source/Runtime/Core/Private/Tests/Memory/SharedBufferTest.cpp
--- a/Engine/Source/Runtime/Core/Private/Tests/Memory/SharedBufferTest.cpp
+++ b/Engine/Source/Runtime/Core/Private/Tests/Memory/SharedBufferTest.cpp
@@ -17,7 +17,7 @@ static_assert(!std::is_constructible<FUniqueBuffer, const FSharedBuffer&>::value
 static_assert(!std::is_constructible<FUniqueBuffer, const FWeakSharedBuffer&>::value, "Invalid constructor");
 
 static_assert( std::is_constructible<FSharedBuffer>::value, "Missing constructor");
-static_assert( std::is_constructible<FSharedBuffer, FUniqueBuffer&&>::value, "Missing constructor");
+static_assert(!std::is_constructible<FSharedBuffer, FUniqueBuffer&&>::value, "Invalid constructor");
 static_assert( std::is_constructible<FSharedBuffer, FSharedBuffer&&>::value, "Missing constructor");
 static_assert(!std::is_constructible<FSharedBuffer, FWeakSharedBuffer&&>::value, "Invalid constructor");
 static_assert(!std::is_constructible<FSharedBuffer, const FUniqueBuffer&>::value, "Invalid constructor");
@@ -41,7 +41,7 @@ static_assert(!std::is_assignable<FUniqueBuffer, const FUniqueBuffer&>::value, "
 static_assert(!std::is_assignable<FUniqueBuffer, const FSharedBuffer&>::value, "Invalid assignment");
 static_assert(!std::is_assignable<FUniqueBuffer, const FWeakSharedBuffer&>::value, "Invalid assignment");
 
-static_assert( std::is_assignable<FSharedBuffer, FUniqueBuffer&&>::value, "Missing assignment");
+static_assert(!std::is_assignable<FSharedBuffer, FUniqueBuffer&&>::value, "Invalid assignment");
 static_assert( std::is_assignable<FSharedBuffer, FSharedBuffer&&>::value, "Missing assignment");
 static_assert(!std::is_assignable<FSharedBuffer, FWeakSharedBuffer&&>::value, "Invalid assignment");
 static_assert(!std::is_assignable<FSharedBuffer, const FUniqueBuffer&>::value, "Invalid assignment");
@@ -183,45 +183,6 @@ bool FUniqueBufferTest::RunTest(const FString& Parameters)
 		TestTrue(TEXT("FUniqueBuffer::TakeOwnership() Deleted"), bDeleted);
 	}
 
-	// Test MakeUnique
-	{
-		FUniqueBuffer Buffer = FUniqueBuffer::MakeUnique(FSharedBuffer());
-		TestTrue(TEXT("FUniqueBuffer::MakeUnique(Null).IsOwned()"), Buffer.IsOwned());
-		TestEqual(TEXT("FUniqueBuffer::MakeUnique(Null).GetSize()"), Buffer.GetSize(), uint64(0));
-	}
-	{
-		FSharedBuffer Shared = FSharedBuffer::Clone(MakeMemoryView<uint8>({1, 2, 3, 4}));
-		FUniqueBuffer Buffer = FUniqueBuffer::MakeUnique(FSharedBuffer(Shared));
-		TestTrue(TEXT("FUniqueBuffer::MakeUnique(Shared).IsOwned()"), Buffer.IsOwned());
-		TestEqual(TEXT("FUniqueBuffer::MakeUnique(Shared).GetSize()"), Buffer.GetSize(), Shared.GetSize());
-		TestNotEqual(TEXT("FUniqueBuffer::MakeUnique(Shared).GetData()"), const_cast<const void*>(Buffer.GetData()), Shared.GetData());
-		TestTrue(TEXT("FUniqueBuffer::MakeUnique(Shared)->Equal"), Buffer.GetView().EqualBytes(Shared.GetView()));
-	}
-	{
-		bool bDeleted = false;
-		constexpr uint64 Size = 64;
-		void* Data = new uint8[Size];
-		FUniqueBuffer Buffer = FUniqueBuffer::MakeUnique(FSharedBuffer::TakeOwnership(Data, Size,
-			[&bDeleted](void* InData, uint64 InSize)
-		{
-			bDeleted = true;
-			delete[] static_cast<uint8*>(InData);
-		}));
-		TestTrue(TEXT("FUniqueBuffer::MakeUnique(Owned).IsOwned()"), Buffer.IsOwned());
-		TestEqual(TEXT("FUniqueBuffer::MakeUnique(Owned).GetSize()"), Buffer.GetSize(), Size);
-		TestEqual(TEXT("FUniqueBuffer::MakeUnique(Owned).GetData()"), Buffer.GetData(), Data);
-		Buffer.Reset();
-		TestTrue(TEXT("FUniqueBuffer::MakeUnique(Owned) Deleted"), bDeleted);
-	}
-	{
-		constexpr uint64 Size = 64;
-		uint8 Data[Size]{};
-		FUniqueBuffer Buffer = FUniqueBuffer::MakeUnique(FSharedBuffer::MakeView(MakeMemoryView(Data)));
-		TestTrue(TEXT("FUniqueBuffer::MakeUnique(View).IsOwned()"), Buffer.IsOwned());
-		TestEqual(TEXT("FUniqueBuffer::MakeUnique(View).GetSize()"), Buffer.GetSize(), Size);
-		TestNotEqual(TEXT("FUniqueBuffer::MakeUnique(View).GetData()"), Buffer.GetData(), static_cast<void*>(Data));
-	}
-
 	// Test MakeOwned
 	{
 		constexpr uint64 Size = 64;
@@ -299,7 +260,7 @@ bool FSharedBufferTest::RunTest(const FString& Parameters)
 		{
 			FSharedBuffer Buffer;
 			{
-				FSharedBuffer OuterBuffer = FSharedBuffer(FUniqueBuffer::Alloc(Size));
+				FSharedBuffer OuterBuffer = FUniqueBuffer::Alloc(Size).MoveToShared();
 				Buffer = FSharedBuffer::MakeView(OuterBuffer.GetData(), OuterBuffer.GetSize(), OuterBuffer);
 				TestEqual(TEXT("FSharedBuffer::MakeView(OuterOwned, SameView)"), Buffer, OuterBuffer);
 				Buffer = FSharedBuffer::MakeView(OuterBuffer.GetData(), OuterBuffer.GetSize() / 2, OuterBuffer);
@@ -345,6 +306,49 @@ bool FSharedBufferTest::RunTest(const FString& Parameters)
 		TestEqual(TEXT("FSharedBuffer::MakeOwned(Owned).GetData()"), BufferCopy.GetData(), Buffer.GetData());
 	}
 
+	// Test MoveToUnique
+	{
+		FUniqueBuffer Buffer = FSharedBuffer().MoveToUnique();
+		TestTrue(TEXT("FSharedBuffer::MoveToUnique(Null).IsOwned()"), Buffer.IsOwned());
+		TestEqual(TEXT("FSharedBuffer::MoveToUnique(Null).GetSize()"), Buffer.GetSize(), uint64(0));
+	}
+	{
+		FSharedBuffer Shared = FSharedBuffer::Clone(MakeMemoryView<uint8>({1, 2, 3, 4}));
+		FUniqueBuffer Buffer = FSharedBuffer(Shared).MoveToUnique();
+		TestTrue(TEXT("FSharedBuffer::MoveToUnique(Shared).IsOwned()"), Buffer.IsOwned());
+		TestEqual(TEXT("FSharedBuffer::MoveToUnique(Shared).GetSize()"), Buffer.GetSize(), Shared.GetSize());
+		TestNotEqual(TEXT("FSharedBuffer::MoveToUnique(Shared).GetData()"), const_cast<const void*>(Buffer.GetData()), Shared.GetData());
+		TestTrue(TEXT("FSharedBuffer::MoveToUnique(Shared)->Equal"), Buffer.GetView().EqualBytes(Shared.GetView()));
+	}
+	{
+		bool bDeleted = false;
+		constexpr uint64 Size = 64;
+		void* Data = new uint8[Size];
+		FSharedBuffer Shared = FSharedBuffer::TakeOwnership(Data, Size,
+			[&bDeleted](void* InData, uint64 InSize)
+		{
+			bDeleted = true;
+			delete[] static_cast<uint8*>(InData);
+		});
+		FUniqueBuffer Buffer = Shared.MoveToUnique();
+		TestTrue(TEXT("FSharedBuffer::MoveToUnique(Owned).IsNull()"), Shared.IsNull());
+		TestTrue(TEXT("FSharedBuffer::MoveToUnique(Owned).IsOwned()"), Buffer.IsOwned());
+		TestEqual(TEXT("FSharedBuffer::MoveToUnique(Owned).GetSize()"), Buffer.GetSize(), Size);
+		TestEqual(TEXT("FSharedBuffer::MoveToUnique(Owned).GetData()"), Buffer.GetData(), Data);
+		Buffer.Reset();
+		TestTrue(TEXT("FSharedBuffer::MoveToUnique(Owned) Deleted"), bDeleted);
+	}
+	{
+		constexpr uint64 Size = 64;
+		uint8 Data[Size]{};
+		FSharedBuffer Shared = FSharedBuffer::MakeView(MakeMemoryView(Data));
+		FUniqueBuffer Buffer = Shared.MoveToUnique();
+		TestTrue(TEXT("FSharedBuffer::MoveToUnique(View).IsNull()"), Shared.IsNull());
+		TestTrue(TEXT("FSharedBuffer::MoveToUnique(View).IsOwned()"), Buffer.IsOwned());
+		TestEqual(TEXT("FSharedBuffer::MoveToUnique(View).GetSize()"), Buffer.GetSize(), Size);
+		TestNotEqual(TEXT("FSharedBuffer::MoveToUnique(View).GetData()"), Buffer.GetData(), static_cast<void*>(Data));
+	}
+
 	return true;
 }
 
@@ -400,10 +404,10 @@ bool FBufferOwnerTest::RunTest(const FString& Parameters)
 	TestFalse(TEXT("FSharedBuffer(BufferOwner)->GetData()->!Materialized"), bMaterialized);
 	bMaterialized = false;
 
-	Buffer = FUniqueBuffer(new FTestBufferOwner(bMaterialized, bFreed, bDeleted));
-	TestFalse(TEXT("FSharedBuffer(BufferOwner)->Assign(FUniqueBuffer)->!Materialized"), bMaterialized);
-	TestTrue(TEXT("FSharedBuffer(BufferOwner)->Assign(FUniqueBuffer)->Freed"), bFreed);
-	TestTrue(TEXT("FSharedBuffer(BufferOwner)->Assign(FUniqueBuffer)->Deleted"), bDeleted);
+	Buffer = FSharedBuffer(new FTestBufferOwner(bMaterialized, bFreed, bDeleted));
+	TestFalse(TEXT("FSharedBuffer(BufferOwner)->Assign(FSharedBuffer)->!Materialized"), bMaterialized);
+	TestTrue(TEXT("FSharedBuffer(BufferOwner)->Assign(FSharedBuffer)->Freed"), bFreed);
+	TestTrue(TEXT("FSharedBuffer(BufferOwner)->Assign(FSharedBuffer)->Deleted"), bDeleted);
 	bFreed = false;
 	bDeleted = false;
 
@@ -424,7 +428,7 @@ bool FBufferOwnerTest::RunTest(const FString& Parameters)
 	bMaterialized = false;
 	bFreed = false;
 	bDeleted = false;
-	Buffer = FUniqueBuffer(new FTestBufferOwner(bMaterialized, bFreed, bDeleted));
+	Buffer = FSharedBuffer(new FTestBufferOwner(bMaterialized, bFreed, bDeleted));
 
 	TestFalse(TEXT("FSharedBuffer(BufferOwner)->IsMaterialized()"), Buffer.IsMaterialized());
 	Buffer.Materialize();
@@ -438,7 +442,7 @@ bool FBufferOwnerTest::RunTest(const FString& Parameters)
 
 	bFreed = false;
 	bDeleted = false;
-	Buffer = FUniqueBuffer(new FTestBufferOwner(bMaterialized, bFreed, bDeleted));
+	Buffer = FSharedBuffer(new FTestBufferOwner(bMaterialized, bFreed, bDeleted));
 	Buffer.Reset();
 	TestFalse(TEXT("FSharedBuffer(BufferOwner)->Reset(!Materialized)->!Materialized"), bMaterialized);
 	TestTrue(TEXT("FSharedBuffer(BufferOwner)->Reset(!Materialized)->Freed"), bFreed);
diff --git a/Engine/Source/Runtime/Core/Private/Tests/Serialization/CompactBinaryPackageTest.cpp b/Engine/Source/Runtime/Core/Private/Tests/Serialization/CompactBinaryPackageTest.cpp
--- a/Engine/Source/Runtime/Core/Private/Tests/Serialization/CompactBinaryPackageTest.cpp
+++ b/Engine/Source/Runtime/Core/Private/Tests/Serialization/CompactBinaryPackageTest.cpp
@@ -207,7 +207,7 @@ bool FCbAttachmentTest::RunTest(const FString& Parameters)
 
 	// Binary Empty
 	{
-		const FCbAttachment Attachment(FSharedBuffer(FUniqueBuffer::Alloc(0)));
+		const FCbAttachment Attachment(FUniqueBuffer::Alloc(0).MoveToShared());
 		TestTrue(TEXT("FCbAttachment(BinaryEmpty).IsNull()"), Attachment.IsNull());
 		TestFalse(TEXT("FCbAttachment(BinaryEmpty).IsBinary()"), Attachment.IsBinary());
 		TestFalse(TEXT("FCbAttachment(BinaryEmpty).IsCompactBinary()"), Attachment.IsCompactBinary());
diff --git a/Engine/Source/Runtime/Core/Public/Memory/SharedBuffer.h b/Engine/Source/Runtime/Core/Public/Memory/SharedBuffer.h
--- a/Engine/Source/Runtime/Core/Public/Memory/SharedBuffer.h
+++ b/Engine/Source/Runtime/Core/Public/Memory/SharedBuffer.h
@@ -4,6 +4,7 @@
 
 #include "CoreTypes.h"
 #include "Containers/ContainersFwd.h"
+#include "Memory/MemoryFwd.h"
 #include "Memory/MemoryView.h"
 #include "Misc/EnumClassFlags.h"
 #include "Templates/Invoke.h"
@@ -193,7 +194,7 @@ private:
  * A reference to a single-ownership mutable buffer.
  *
  * Ownership can be transferred by moving to FUniqueBuffer or it can be converted to an immutable
- * shared buffer by moving to FSharedBuffer.
+ * shared buffer with MoveToShared().
  *
  * @see FBufferOwner
  */
@@ -229,14 +230,6 @@ public:
 		decltype(Invoke(std::declval<DeleteFunctionType>(), std::declval<void*>(), std::declval<uint64>()))* = nullptr>
 	UE_NODISCARD static inline FUniqueBuffer TakeOwnership(void* Data, uint64 Size, DeleteFunctionType&& DeleteFunction);
 
-	/**
-	 * Make a unique buffer from a shared buffer.
-	 *
-	 * Steals the buffer owner from the shared buffer if this is the last reference to it, otherwise
-	 * clones the shared buffer to guarantee unique ownership. An non-owned buffer is always cloned.
-	 */
-	UE_NODISCARD CORE_API static FUniqueBuffer MakeUnique(FSharedBuffer&& Buffer);
-
 	/** Construct a null unique buffer. */
 	FUniqueBuffer() = default;
 
@@ -291,14 +284,23 @@ public:
 	 */
 	CORE_API void Materialize() const;
 
+	/**
+	 * Convert this to an immutable shared buffer, leaving this null.
+	 *
+	 * Steals the buffer owner from the unique buffer.
+	 */
+	UE_NODISCARD CORE_API FSharedBuffer MoveToShared();
+
+	friend class FSharedBuffer;
+
 private:
+	explicit FUniqueBuffer(BufferOwnerPrivate::TBufferOwnerPtr<BufferOwnerPrivate::FSharedOps>&& SharedOwner);
+
 	using OwnerPtrType = BufferOwnerPrivate::TBufferOwnerPtr<BufferOwnerPrivate::FSharedOps>;
 
 	inline friend const OwnerPtrType& ToPrivateOwnerPtr(const FUniqueBuffer& Buffer) { return Buffer.Owner; }
 	inline friend OwnerPtrType ToPrivateOwnerPtr(FUniqueBuffer&& Buffer) { return MoveTemp(Buffer.Owner); }
 
-	inline explicit FUniqueBuffer(OwnerPtrType&& InOwner) : Owner(MoveTemp(InOwner)) {}
-
 	OwnerPtrType Owner;
 };
 
@@ -350,19 +352,6 @@ public:
 	/** Construct a shared buffer from a new unreferenced buffer owner. */
 	CORE_API explicit FSharedBuffer(FBufferOwner* Owner);
 
-	/** Construct a shared buffer from a unique buffer, making it immutable. */
-	inline explicit FSharedBuffer(FUniqueBuffer&& Buffer)
-		: Owner(ToPrivateOwnerPtr(MoveTemp(Buffer)))
-	{
-	}
-
-	/** Assign a shared buffer from a unique buffer, making it immutable. */
-	inline FSharedBuffer& operator=(FUniqueBuffer&& Buffer)
-	{
-		Owner = ToPrivateOwnerPtr(MoveTemp(Buffer));
-		return *this;
-	}
-
 	/** Reset this to null. */
 	CORE_API void Reset();
 
@@ -403,10 +392,20 @@ public:
 	 */
 	CORE_API void Materialize() const;
 
+	/**
+	 * Convert this to a unique buffer, leaving this null.
+	 *
+	 * Steals the buffer owner from the shared buffer if this is the last reference to it, otherwise
+	 * clones the shared buffer to guarantee unique ownership. A non-owned buffer is always cloned.
+	 */
+	UE_NODISCARD CORE_API FUniqueBuffer MoveToUnique();
+
+	friend class FUniqueBuffer;
 	friend class FWeakSharedBuffer;
 
 private:
-	FSharedBuffer(const BufferOwnerPrivate::TBufferOwnerPtr<BufferOwnerPrivate::FWeakOps>& WeakOwner);
+	explicit FSharedBuffer(BufferOwnerPrivate::TBufferOwnerPtr<BufferOwnerPrivate::FSharedOps>&& SharedOwner);
+	explicit FSharedBuffer(const BufferOwnerPrivate::TBufferOwnerPtr<BufferOwnerPrivate::FWeakOps>& WeakOwner);
 
 	using OwnerPtrType = BufferOwnerPrivate::TBufferOwnerPtr<BufferOwnerPrivate::FSharedOps>;
 
diff --git a/Engine/Source/Runtime/Core/Public/Serialization/CompactBinary.h b/Engine/Source/Runtime/Core/Public/Serialization/CompactBinary.h
--- a/Engine/Source/Runtime/Core/Public/Serialization/CompactBinary.h
+++ b/Engine/Source/Runtime/Core/Public/Serialization/CompactBinary.h
@@ -1094,7 +1094,7 @@ public:
 			FUniqueBuffer MutableBuffer = FUniqueBuffer::Alloc(ViewType::GetSize());
 			ViewType::CopyTo(MutableBuffer);
 			ViewType::operator=(ViewType(MutableBuffer.GetData()));
-			Buffer = MoveTemp(MutableBuffer);
+			Buffer = MutableBuffer.MoveToShared();
 		}
 	}
 
diff --git a/Engine/Source/Runtime/CoreUObject/Private/Tests/Virtualization/VirtualizedBulkDataTests.cpp b/Engine/Source/Runtime/CoreUObject/Private/Tests/Virtualization/VirtualizedBulkDataTests.cpp
--- a/Engine/Source/Runtime/CoreUObject/Private/Tests/Virtualization/VirtualizedBulkDataTests.cpp
+++ b/Engine/Source/Runtime/CoreUObject/Private/Tests/Virtualization/VirtualizedBulkDataTests.cpp
@@ -55,7 +55,7 @@ bool FVirtualizationWrapperTestBasic::RunTest(const FString& Parameters)
 
 	// Create a new bulkdata object and copy by assignment (note we assign some junk data that will get overwritten)
 	FByteVirtualizedBulkData BulkDataAssignment;
-	BulkDataAssignment.UpdatePayload(FSharedBuffer(FUniqueBuffer::Alloc(128)));
+	BulkDataAssignment.UpdatePayload(FUniqueBuffer::Alloc(128).MoveToShared());
 	BulkDataAssignment = BulkData;
 
 	// Test both bulkdata objects
@@ -129,7 +129,7 @@ bool FVirtualizationWrapperTestUpdatePayload::RunTest(const FString& Parameters)
 		{
 			FUniqueBuffer EditablePayload = FUniqueBuffer::Clone(Payload);
 			FMemory::Memset(EditablePayload.GetData(), NewValue, EditablePayload.GetSize());
-			EditedPayload = MoveTemp(EditablePayload);
+			EditedPayload = EditablePayload.MoveToShared();
 		}
 
 		// Update the bulkdata object with the new edited payload
diff --git a/Engine/Source/Runtime/CoreUObject/Private/Virtualization/VirtualizationManager.cpp b/Engine/Source/Runtime/CoreUObject/Private/Virtualization/VirtualizationManager.cpp
--- a/Engine/Source/Runtime/CoreUObject/Private/Virtualization/VirtualizationManager.cpp
+++ b/Engine/Source/Runtime/CoreUObject/Private/Virtualization/VirtualizationManager.cpp
@@ -80,7 +80,7 @@ FSharedBuffer FVirtualizationManager::PullData(const FGuid& Guid)
 
 	FUniqueBuffer Buffer = FUniqueBuffer::Alloc(DataLength);
 	FileAr->Serialize(Buffer.GetData(), DataLength);
-	return FSharedBuffer(MoveTemp(Buffer));
+	return Buffer.MoveToShared();
 }
 
 void FVirtualizationManager::LoadSettingsFromConfigFiles()
diff --git a/Engine/Source/Runtime/CoreUObject/Private/Virtualization/VirtualizedBulkData.cpp b/Engine/Source/Runtime/CoreUObject/Private/Virtualization/VirtualizedBulkData.cpp
--- a/Engine/Source/Runtime/CoreUObject/Private/Virtualization/VirtualizedBulkData.cpp
+++ b/Engine/Source/Runtime/CoreUObject/Private/Virtualization/VirtualizedBulkData.cpp
@@ -254,7 +254,7 @@ bool FVirtualizedUntypedBulkData::SerializeData(FArchive& Ar, FSharedBuffer& InP
 
 		if (Ar.IsLoading())
 		{
-			InPayload = FSharedBuffer(MoveTemp(LoadPayload));
+			InPayload = LoadPayload.MoveToShared();
 		}
 
 		return true;
