#rnx

diff --git a/Engine/Source/Runtime/Core/Public/Misc/AssertionMacros.h b/Engine/Source/Runtime/Core/Public/Misc/AssertionMacros.h
--- a/Engine/Source/Runtime/Core/Public/Misc/AssertionMacros.h
+++ b/Engine/Source/Runtime/Core/Public/Misc/AssertionMacros.h
@@ -397,10 +397,10 @@ namespace UE4Asserts_Private
 // Returns FName(TEXT("FunctionName")), while statically verifying that the function exists in ClassName
 // Handles overloaded functions by specifying the argument of the overload to use
 #define GET_FUNCTION_NAME_CHECKED_OneParam(ClassName, FunctionName, ArgType) \
-	((void)sizeof(std::declval<ClassName&>().FunctionName(std::declval<PREPROCESSOR_REMOVE_OPTIONAL_PARENS(ArgType)>()), (int)0), FName(TEXT(#FunctionName)))
+	((void)sizeof((std::declval<ClassName&>().FunctionName(std::declval<PREPROCESSOR_REMOVE_OPTIONAL_PARENS(ArgType)>()), (int)0)), FName(TEXT(#FunctionName)))
 
 #define GET_FUNCTION_NAME_STRING_CHECKED_OneParam(ClassName, FunctionName, ArgType) \
-	((void)sizeof(std::declval<ClassName&>().FunctionName(std::declval<PREPROCESSOR_REMOVE_OPTIONAL_PARENS(ArgType)>()), (int)0), TEXT(#FunctionName))
+	((void)sizeof((std::declval<ClassName&>().FunctionName(std::declval<PREPROCESSOR_REMOVE_OPTIONAL_PARENS(ArgType)>()), (int)0)), TEXT(#FunctionName))
 
 /*----------------------------------------------------------------------------
 	Low level error macros
