Fix memory leak in FPreloadFile.

- Cache the AsyncSizeHandle and delete it when the size callback is complete.

diff --git a/Engine/Source/Runtime/Core/Private/Misc/PreLoadFile.cpp b/Engine/Source/Runtime/Core/Private/Misc/PreLoadFile.cpp
--- a/Engine/Source/Runtime/Core/Private/Misc/PreLoadFile.cpp
+++ b/Engine/Source/Runtime/Core/Private/Misc/PreLoadFile.cpp
@@ -87,10 +87,13 @@ void FPreLoadFile::KickOffRead()
 			CompletionEvent->Trigger();
 			delete AsyncReadHandle;
 		}
+
+		delete SizeRequestHandle;
+		SizeRequestHandle = nullptr;
 	};
 
 	AsyncReadHandle = FPlatformFileManager::Get().GetPlatformFile().OpenAsyncRead(*Path);
-	AsyncReadHandle->SizeRequest(&SizeCallbackFunction);
+	SizeRequestHandle = AsyncReadHandle->SizeRequest(&SizeCallbackFunction);
 
 #else
 
diff --git a/Engine/Source/Runtime/Core/Public/Misc/PreLoadFile.h b/Engine/Source/Runtime/Core/Public/Misc/PreLoadFile.h
--- a/Engine/Source/Runtime/Core/Public/Misc/PreLoadFile.h
+++ b/Engine/Source/Runtime/Core/Public/Misc/PreLoadFile.h
@@ -22,5 +22,6 @@ protected:
 	int64 FileSize;
 	FString Path;
 	FEvent* CompletionEvent;
-	class IAsyncReadFileHandle* AsyncReadHandle;
+	class IAsyncReadFileHandle* AsyncReadHandle = nullptr;
+	class IAsyncReadRequest* SizeRequestHandle = nullptr;
 };
