Cleaned-up potential unsafe use of FRegexPattern

FRegexPattern could be implicitly constructed from an FString, which made it easy to accidentally pass an FString as the first argument to the FRegexMatcher constructor, and get a matcher which outlived the pattern it was using.

This changes the FRegexPattern constructor to be explicit (as you should choose to make a pattern, since it's not cheap), and also has FRegexMatcher keep a copy of the pattern it was constructed with, to ensure that the matcher can't outlive its pattern.

FRegexMatcher has also been made non-copyable, as it contains state that shouldn't be accidentally shared. FRegexPattern is now both movable and copyable, as internally it references an immutable instance so can be safely shared.

diff --git a/Engine/Source/Runtime/Core/Private/Internationalization/ICURegex.cpp b/Engine/Source/Runtime/Core/Private/Internationalization/ICURegex.cpp
--- a/Engine/Source/Runtime/Core/Private/Internationalization/ICURegex.cpp
+++ b/Engine/Source/Runtime/Core/Private/Internationalization/ICURegex.cpp
@@ -23,13 +23,13 @@ public:
 		}
 	}
 
-	TSharedPtr<icu::RegexPattern> GetInternalRegexPattern() const
+	TSharedPtr<const icu::RegexPattern> GetInternalRegexPattern() const
 	{
 		return ICURegexPattern.Pin();
 	}
 
 private:
-	TWeakPtr<icu::RegexPattern> ICURegexPattern;
+	TWeakPtr<const icu::RegexPattern> ICURegexPattern;
 };
 
 
@@ -94,13 +94,13 @@ FICURegexManager& FICURegexManager::Get()
 	return *Singleton;
 }
 
-TWeakPtr<icu::RegexPattern> FICURegexManager::CreateRegexPattern(const FString& InSourceString)
+TWeakPtr<const icu::RegexPattern> FICURegexManager::CreateRegexPattern(const FString& InSourceString)
 {
 	icu::UnicodeString ICUSourceString;
 	ICUUtilities::ConvertString(InSourceString, ICUSourceString);
 
 	UErrorCode ICUStatus = U_ZERO_ERROR;
-	TSharedPtr<icu::RegexPattern>ICURegexPattern = MakeShareable(icu::RegexPattern::compile(ICUSourceString, 0, ICUStatus));
+	TSharedPtr<const icu::RegexPattern> ICURegexPattern = MakeShareable(icu::RegexPattern::compile(ICUSourceString, 0, ICUStatus));
 
 	if (ICURegexPattern.IsValid())
 	{
@@ -111,9 +111,9 @@ TWeakPtr<icu::RegexPattern> FICURegexManager::CreateRegexPattern(const FString&
 	return ICURegexPattern;
 }
 
-void FICURegexManager::DestroyRegexPattern(TWeakPtr<icu::RegexPattern>& InICURegexPattern)
+void FICURegexManager::DestroyRegexPattern(TWeakPtr<const icu::RegexPattern>& InICURegexPattern)
 {
-	TSharedPtr<icu::RegexPattern> ICURegexPattern = InICURegexPattern.Pin();
+	TSharedPtr<const icu::RegexPattern> ICURegexPattern = InICURegexPattern.Pin();
 	if (ICURegexPattern.IsValid())
 	{
 		FScopeLock ScopeLock(&AllocatedRegexPatternsCS);
@@ -122,7 +122,7 @@ void FICURegexManager::DestroyRegexPattern(TWeakPtr<icu::RegexPattern>& InICUReg
 	InICURegexPattern.Reset();
 }
 
-TWeakPtr<icu::RegexMatcher> FICURegexManager::CreateRegexMatcher(icu::RegexPattern* InPattern, const icu::UnicodeString* InInputString)
+TWeakPtr<icu::RegexMatcher> FICURegexManager::CreateRegexMatcher(const icu::RegexPattern* InPattern, const icu::UnicodeString* InInputString)
 {
 	TSharedPtr<icu::RegexMatcher> ICURegexMatcher;
 
@@ -154,15 +154,22 @@ void FICURegexManager::DestroyRegexMatcher(TWeakPtr<icu::RegexMatcher>& InICUReg
 
 
 FRegexPattern::FRegexPattern(const FString& SourceString) 
-	: Implementation(new FRegexPatternImplementation(SourceString))
+	: Implementation(MakeShared<FRegexPatternImplementation>(SourceString))
 {
 }
 
 
-FRegexMatcher::FRegexMatcher(const FRegexPattern& Pattern, const FString& InputString) 
-	: Implementation(new FRegexMatcherImplementation(Pattern.Implementation.Get(), InputString))
+FRegexMatcher::FRegexMatcher(const FRegexPattern& SourcePattern, const FString& InputString)
+	: Pattern(SourcePattern)
+	, Implementation(MakeShared<FRegexMatcherImplementation>(Pattern.Implementation.Get(), InputString))
 {
-}	
+}
+
+FRegexMatcher::FRegexMatcher(FRegexPattern&& SourcePattern, const FString& InputString)
+	: Pattern(MoveTemp(SourcePattern))
+	, Implementation(MakeShared<FRegexMatcherImplementation>(Pattern.Implementation.Get(), InputString))
+{
+}
 
 bool FRegexMatcher::FindNext()
 {
diff --git a/Engine/Source/Runtime/Core/Private/Internationalization/ICURegex.h b/Engine/Source/Runtime/Core/Private/Internationalization/ICURegex.h
--- a/Engine/Source/Runtime/Core/Private/Internationalization/ICURegex.h
+++ b/Engine/Source/Runtime/Core/Private/Internationalization/ICURegex.h
@@ -23,17 +23,17 @@ public:
 	static bool IsInitialized();
 	static FICURegexManager& Get();
 
-	TWeakPtr<icu::RegexPattern> CreateRegexPattern(const FString& InSourceString);
-	void DestroyRegexPattern(TWeakPtr<icu::RegexPattern>& InICURegexPattern);
+	TWeakPtr<const icu::RegexPattern> CreateRegexPattern(const FString& InSourceString);
+	void DestroyRegexPattern(TWeakPtr<const icu::RegexPattern>& InICURegexPattern);
 
-	TWeakPtr<icu::RegexMatcher> CreateRegexMatcher(icu::RegexPattern* InPattern, const icu::UnicodeString* InInputString);
+	TWeakPtr<icu::RegexMatcher> CreateRegexMatcher(const icu::RegexPattern* InPattern, const icu::UnicodeString* InInputString);
 	void DestroyRegexMatcher(TWeakPtr<icu::RegexMatcher>& InICURegexMatcher);
 
 private:
 	static FICURegexManager* Singleton;
 
 	FCriticalSection AllocatedRegexPatternsCS;
-	TSet<TSharedPtr<icu::RegexPattern>> AllocatedRegexPatterns;
+	TSet<TSharedPtr<const icu::RegexPattern>> AllocatedRegexPatterns;
 
 	FCriticalSection AllocatedRegexMatchersCS;
 	TSet<TSharedPtr<icu::RegexMatcher>> AllocatedRegexMatchers;
diff --git a/Engine/Source/Runtime/Core/Private/Internationalization/LegacyRegex.cpp b/Engine/Source/Runtime/Core/Private/Internationalization/LegacyRegex.cpp
--- a/Engine/Source/Runtime/Core/Private/Internationalization/LegacyRegex.cpp
+++ b/Engine/Source/Runtime/Core/Private/Internationalization/LegacyRegex.cpp
@@ -11,7 +11,8 @@ class FRegexPatternImplementation
 
 };
 
-FRegexPattern::FRegexPattern(const FString& SourceString) : Implementation( MakeShareable( new FRegexPatternImplementation() ) )
+FRegexPattern::FRegexPattern(const FString& SourceString)
+	: Implementation(MakeShared<FRegexPatternImplementation>())
 {
 }
 
@@ -20,9 +21,17 @@ class FRegexMatcherImplementation
 
 };
 
-FRegexMatcher::FRegexMatcher(const FRegexPattern& Pattern, const FString& InputString) : Implementation( MakeShareable( new FRegexMatcherImplementation() ) )
+FRegexMatcher::FRegexMatcher(const FRegexPattern& SourcePattern, const FString& InputString)
+	: Pattern(SourcePattern)
+	, Implementation(MakeShared<FRegexMatcherImplementation>())
 {
-}	
+}
+
+FRegexMatcher::FRegexMatcher(FRegexPattern&& SourcePattern, const FString& InputString)
+	: Pattern(MoveTemp(SourcePattern))
+	, Implementation(MakeShared<FRegexMatcherImplementation>())
+{
+}
 
 bool FRegexMatcher::FindNext()
 {
diff --git a/Engine/Source/Runtime/Core/Public/Internationalization/Regex.h b/Engine/Source/Runtime/Core/Public/Internationalization/Regex.h
--- a/Engine/Source/Runtime/Core/Public/Internationalization/Regex.h
+++ b/Engine/Source/Runtime/Core/Public/Internationalization/Regex.h
@@ -17,7 +17,13 @@ class CORE_API FRegexPattern
 	friend class FRegexMatcher;
 
 public:
-	FRegexPattern(const FString& SourceString);
+	explicit FRegexPattern(const FString& SourceString);
+
+	FRegexPattern(const FRegexPattern&) = default;
+	FRegexPattern& operator=(const FRegexPattern&) = default;
+
+	FRegexPattern(FRegexPattern&&) = default;
+	FRegexPattern& operator=(FRegexPattern&&) = default;
 
 private:
 	TSharedRef<FRegexPatternImplementation> Implementation;
@@ -30,7 +36,14 @@ private:
 class CORE_API FRegexMatcher
 {
 public:
-	FRegexMatcher(const FRegexPattern& Pattern, const FString& Input);
+	FRegexMatcher(const FRegexPattern& SourcePattern, const FString& Input);
+	FRegexMatcher(FRegexPattern&& SourcePattern, const FString& Input);
+
+	FRegexMatcher(const FRegexMatcher&) = delete;
+	FRegexMatcher& operator=(const FRegexMatcher&) = delete;
+
+	FRegexMatcher(FRegexMatcher&&) = default;
+	FRegexMatcher& operator=(FRegexMatcher&&) = default;
 
 	bool FindNext();
 
@@ -46,5 +59,6 @@ public:
 	void SetLimits(const int32 BeginIndex, const int32 EndIndex);
 
 private:
+	FRegexPattern Pattern;
 	TSharedRef<FRegexMatcherImplementation> Implementation;
 };
