Swapped backtracing's RWLock out for a critical section until the need to lock can be removed.

#rnx

diff --git a/Engine/Source/Runtime/Core/Private/ProfilingDebugging/Windows/CallstackTrace.cpp b/Engine/Source/Runtime/Core/Private/ProfilingDebugging/Windows/CallstackTrace.cpp
--- a/Engine/Source/Runtime/Core/Private/ProfilingDebugging/Windows/CallstackTrace.cpp
+++ b/Engine/Source/Runtime/Core/Private/ProfilingDebugging/Windows/CallstackTrace.cpp
@@ -6,7 +6,6 @@
 #include "CoreTypes.h"
 #include "HAL/CriticalSection.h"
 #include "HAL/MemoryBase.h"
-#include "Misc/ScopeRWLock.h"
 #include "HAL/RunnableThread.h"
 #include "HAL/LowLevelMemTracker.h"
 #include "ProfilingDebugging/MemoryTrace.h"
@@ -185,7 +184,7 @@ private:
 	void					StopWorkerThread();
 	const FFunction*		LookupFunction(UPTRINT Address, FLookupState& State) const;
 	static FBacktracer*		Instance;
-	mutable FRWLock			Lock;
+	mutable FCriticalSection Lock;
 	FModule*				Modules;
 	int32					ModulesNum;
 	int32					ModulesCapacity;
@@ -477,7 +476,7 @@ void FBacktracer::AddModule(UPTRINT ModuleBase, const TCHAR* Name)
 	};
 
 	{
-		FWriteScopeLock _(Lock);
+		FScopeLock _(&Lock);
 
 		TArrayView<FModule> ModulesView(Modules, ModulesNum);
 		int32 Index = Algo::UpperBound(ModulesView, Module.Id, FIdPredicate());
@@ -511,7 +510,7 @@ void FBacktracer::RemoveModule(UPTRINT ModuleBase)
 	return;
 #endif
 
-	FWriteScopeLock _(Lock);
+	FScopeLock _(&Lock);
 
 	uint32 ModuleId = AddressToId(ModuleBase);
 	TArrayView<FModule> ModulesView(Modules, ModulesNum);
@@ -613,7 +612,7 @@ void* FBacktracer::GetBacktraceId(void* AddressOfReturnAddress) const
 	uint64 BacktraceId = 0;
 	uint32 FrameIdx = 0;
 
-	FReadScopeLock _(Lock);
+	FScopeLock _(&Lock);
 	do
 	{
 		UPTRINT RetAddr = *StackPointer;
@@ -680,7 +679,7 @@ void* FBacktracer::GetBacktraceId(void* AddressOfReturnAddress) const
 ////////////////////////////////////////////////////////////////////////////////
 int32 FBacktracer::GetFrameSize(void* FunctionAddress) const
 {
-	FReadScopeLock _(Lock);
+	FScopeLock _(&Lock);
 
 	FLookupState LookupState;
 	const FFunction* Function = LookupFunction(UPTRINT(FunctionAddress), LookupState);
