Corrected the remote debugging prompt's message.

#rnx

diff --git a/Engine/Source/Runtime/Core/Private/Windows/WindowsPlatformMisc.cpp b/Engine/Source/Runtime/Core/Private/Windows/WindowsPlatformMisc.cpp
--- a/Engine/Source/Runtime/Core/Private/Windows/WindowsPlatformMisc.cpp
+++ b/Engine/Source/Runtime/Core/Private/Windows/WindowsPlatformMisc.cpp
@@ -2045,9 +2045,9 @@ void FWindowsPlatformMisc::PromptForRemoteDebugging(bool bIsEnsure)
 			TEXT("Have a programmer remote debug this crash?\n")
 			TEXT("Hit NO to exit and submit error report as normal.\n")
 			TEXT("Otherwise, contact a programmer for remote debugging,\n")
-			TEXT("giving him the changelist number below.\n")
-			TEXT("Once he confirms he is connected to the machine,\n")
-			TEXT("hit YES to allow him to debug the crash.\n")
+			TEXT("giving them the changelist number below.\n")
+			TEXT("Once they have confirmed they are connected to the machine,\n")
+			TEXT("hit YES to allow them to debug the crash.\n")
 			TEXT("[Changelist = %d]"),
 			FEngineVersion::Current().GetChangelist());
 		FSlowHeartBeatScope SuspendHeartBeat;
