CompressedBuffer: Added FHeader::TotalCompressedSize and FCompressedBuffer::FromCompressed(FArchive&)

#rnx

diff --git a/Engine/Source/Developer/DerivedDataCache/Private/DerivedDataCacheRecord.cpp b/Engine/Source/Developer/DerivedDataCache/Private/DerivedDataCacheRecord.cpp
--- a/Engine/Source/Developer/DerivedDataCache/Private/DerivedDataCacheRecord.cpp
+++ b/Engine/Source/Developer/DerivedDataCache/Private/DerivedDataCacheRecord.cpp
@@ -172,7 +172,7 @@ void FCacheRecordBuilderInternal::SetMeta(FCbObject&& InMeta)
 FPayloadId FCacheRecordBuilderInternal::SetValue(const FSharedBuffer& Buffer, const FPayloadId& Id)
 {
 	const FPayloadId ValueId = GetOrCreatePayloadId(Id, Buffer);
-	return SetValue(FPayload(ValueId, FCompressedBuffer::Compress(NAME_LZ4, Buffer)));
+	return SetValue(FPayload(ValueId, FCompressedBuffer::Compress(NAME_Default, Buffer)));
 }
 
 FPayloadId FCacheRecordBuilderInternal::SetValue(FPayload&& Payload)
@@ -186,7 +186,7 @@ FPayloadId FCacheRecordBuilderInternal::SetValue(FPayload&& Payload)
 FPayloadId FCacheRecordBuilderInternal::AddAttachment(const FSharedBuffer& Buffer, const FPayloadId& Id)
 {
 	const FPayloadId AttachmentId = GetOrCreatePayloadId(Id, Buffer);
-	return AddAttachment(FPayload(AttachmentId, FCompressedBuffer::Compress(NAME_LZ4, Buffer)));
+	return AddAttachment(FPayload(AttachmentId, FCompressedBuffer::Compress(NAME_Default, Buffer)));
 }
 
 FPayloadId FCacheRecordBuilderInternal::AddAttachment(FPayload&& Payload)
diff --git a/Engine/Source/Editor/UnrealEd/Private/EditorDomain/EditorDomainUtils.cpp b/Engine/Source/Editor/UnrealEd/Private/EditorDomain/EditorDomainUtils.cpp
--- a/Engine/Source/Editor/UnrealEd/Private/EditorDomain/EditorDomainUtils.cpp
+++ b/Engine/Source/Editor/UnrealEd/Private/EditorDomain/EditorDomainUtils.cpp
@@ -19,7 +19,7 @@ namespace EditorDomain
 {
 
 // Change to a new guid when EditorDomain needs to be invalidated
-const TCHAR* EditorDomainVersion = TEXT("BF57FE551C4F4FAFAFFF2EDC17132DC9");
+const TCHAR* EditorDomainVersion = TEXT("25F3D3DC3B1246DB9BB1384185578B8A");
 // Identifier of the CacheBucket for EditorDomainPackages
 const TCHAR* EditorDomainPackageBucketName = TEXT("EditorDomainPackage");
 
diff --git a/Engine/Source/Runtime/Core/Private/Compression/CompressedBuffer.cpp b/Engine/Source/Runtime/Core/Private/Compression/CompressedBuffer.cpp
--- a/Engine/Source/Runtime/Core/Private/Compression/CompressedBuffer.cpp
+++ b/Engine/Source/Runtime/Core/Private/Compression/CompressedBuffer.cpp
@@ -6,6 +6,7 @@
 #include "Math/UnrealMathUtility.h"
 #include "Misc/ByteSwap.h"
 #include "Misc/Crc.h"
+#include "Serialization/Archive.h"
 #include "Templates/Function.h"
 #include "UObject/NameTypes.h"
 
@@ -52,6 +53,8 @@ struct FHeader
 	uint32 BlockCount = 0;
 	/** The total size of the uncompressed data. */
 	uint64 TotalRawSize = 0;
+	/** The total size of the compressed data including the header. */
+	uint64 TotalCompressedSize = 0;
 
 	/** Checks validity of the buffer based on the magic number, method, and CRC-32. */
 	static bool IsValid(const FCompositeBuffer& CompressedData);
@@ -69,25 +72,47 @@ struct FHeader
 		return Header;
 	}
 
-	/** Write a header to a memory view that is at least sizeof(FHeader). */
-	void Write(FMutableMemoryView CompressedData) const
+	/**
+	 * Write a header to a memory view that is at least sizeof(FHeader).
+	 *
+	 * @param HeaderView   View of the header to write, including any method-specific header data.
+	 */
+	void Write(FMutableMemoryView HeaderView) const
 	{
 		FHeader Header = *this;
 		Header.ByteSwap();
-		CompressedData.CopyFrom(MakeMemoryView(&Header, &Header + 1));
+		HeaderView.CopyFrom(MakeMemoryView(&Header, &Header + 1));
+		Header.ByteSwap();
+		Header.Crc32 = CalculateCrc32(HeaderView);
+		Header.ByteSwap();
+		HeaderView.CopyFrom(MakeMemoryView(&Header, &Header + 1));
+	}
+
+	/** Calculate the CRC-32 from a view of a header including any method-specific header data. */
+	static uint32 CalculateCrc32(FMemoryView HeaderView)
+	{
+		uint32 Crc32 = 0;
+		constexpr uint64 MethodOffset = STRUCT_OFFSET(FHeader, Method);
+		for (FMemoryView View = HeaderView + MethodOffset; const uint64 ViewSize = View.GetSize();)
+		{
+			const int32 Size = static_cast<int32>(FMath::Min<uint64>(ViewSize, MAX_int32));
+			Crc32 = FCrc::MemCrc32(View.GetData(), Size, Crc32);
+			View += Size;
+		}
+		return Crc32;
 	}
 
-private:
 	void ByteSwap()
 	{
 		Magic = NETWORK_ORDER32(Magic);
 		Crc32 = NETWORK_ORDER32(Crc32);
 		BlockCount = NETWORK_ORDER32(BlockCount);
 		TotalRawSize = NETWORK_ORDER64(TotalRawSize);
+		TotalCompressedSize = NETWORK_ORDER64(TotalCompressedSize);
 	}
 };
 
-static_assert(sizeof(FHeader) == 24, "FHeader is the wrong size.");
+static_assert(sizeof(FHeader) == 32, "FHeader is the wrong size.");
 
 ///////////////////////////////////////////////////////////////////////////////////////////////////////////////////////
 
@@ -97,36 +122,9 @@ public:
 	virtual FCompositeBuffer Compress(const FCompositeBuffer& RawData, uint64 BlockSize = DefaultBlockSize) const = 0;
 	virtual FCompositeBuffer Decompress(const FHeader& Header, const FCompositeBuffer& CompressedData) const = 0;
 	virtual bool TryDecompressTo(const FHeader& Header, const FCompositeBuffer& CompressedData, FMutableMemoryView RawView) const = 0;
-	virtual uint32 CalculateCrc32(const FHeader& Header, const FCompositeBuffer& CompressedData) const;
-
-protected:
-	static uint32 CalculateRangeCrc32(const FCompositeBuffer& CompressedData, uint64 Offset, uint64 Size, uint32 Crc32 = 0);
+	virtual uint64 GetHeaderSize(const FHeader& Header) const = 0;
 };
 
-uint32 FMethod::CalculateCrc32(const FHeader& Header, const FCompositeBuffer& CompressedData) const
-{
-	if (sizeof(FHeader) <= CompressedData.GetSize())
-	{
-		constexpr uint64 MethodOffset = STRUCT_OFFSET(FHeader, Method);
-		return CalculateRangeCrc32(CompressedData, MethodOffset, sizeof(FHeader) - MethodOffset);
-	}
-	return 0;
-}
-
-uint32 FMethod::CalculateRangeCrc32(const FCompositeBuffer& CompressedData, uint64 Offset, uint64 Size, uint32 Crc32)
-{
-	CompressedData.IterateRange(Offset, Size, [&Crc32](FMemoryView Crc32View)
-	{
-		while (const uint64 ViewSize = Crc32View.GetSize())
-		{
-			const int32 Crc32Size = static_cast<int32>(FMath::Min<uint64>(ViewSize, MAX_int32));
-			Crc32 = FCrc::MemCrc32(Crc32View.GetData(), Crc32Size, Crc32);
-			Crc32View += Crc32Size;
-		}
-	});
-	return Crc32;
-}
-
 ///////////////////////////////////////////////////////////////////////////////////////////////////////////////////////
 
 class FMethodNone final : public FMethod
@@ -138,21 +136,18 @@ public:
 		Header.Method = EMethod::None;
 		Header.BlockCount = 1;
 		Header.TotalRawSize = RawData.GetSize();
+		Header.TotalCompressedSize = Header.TotalRawSize + sizeof(FHeader);
 
 		FUniqueBuffer HeaderData = FUniqueBuffer::Alloc(sizeof(FHeader));
-		const FMutableMemoryView HeaderDataView = HeaderData;
-
-		FCompositeBuffer CompressedData(HeaderData.MoveToShared(), RawData.MakeOwned());
-		Header.Write(HeaderDataView);
-		Header.Crc32 = CalculateCrc32(Header, CompressedData);
-		Header.Write(HeaderDataView);
-		return CompressedData;
+		Header.Write(HeaderData);
+		return FCompositeBuffer(HeaderData.MoveToShared(), RawData.MakeOwned());
 	}
 
 	FCompositeBuffer Decompress(const FHeader& Header, const FCompositeBuffer& CompressedData) const final
 	{
 		if (Header.Method == EMethod::None &&
-			sizeof(FHeader) + Header.TotalRawSize == CompressedData.GetSize())
+			Header.TotalCompressedSize == CompressedData.GetSize() &&
+			Header.TotalCompressedSize == Header.TotalRawSize + sizeof(FHeader))
 		{
 			return CompressedData.Mid(sizeof(FHeader), Header.TotalRawSize).MakeOwned();
 		}
@@ -163,13 +158,19 @@ public:
 	{
 		if (Header.Method == EMethod::None &&
 			Header.TotalRawSize == RawView.GetSize() &&
-			sizeof(FHeader) + Header.TotalRawSize == CompressedData.GetSize())
+			Header.TotalCompressedSize == CompressedData.GetSize() &&
+			Header.TotalCompressedSize == Header.TotalRawSize + sizeof(FHeader))
 		{
 			CompressedData.CopyTo(RawView, sizeof(FHeader));
 			return true;
 		}
 		return false;
 	}
+
+	uint64 GetHeaderSize(const FHeader& Header) const final
+	{
+		return sizeof(FHeader);
+	}
 };
 
 ///////////////////////////////////////////////////////////////////////////////////////////////////////////////////////
@@ -181,18 +182,12 @@ public:
 	FCompositeBuffer Decompress(const FHeader& Header, const FCompositeBuffer& CompressedData) const final;
 	bool TryDecompressTo(const FHeader& Header, const FCompositeBuffer& CompressedData, FMutableMemoryView RawView) const final;
 
-protected:
-	uint32 CalculateCrc32(const FHeader& Header, const FCompositeBuffer& CompressedData) const final
+	uint64 GetHeaderSize(const FHeader& Header) const final
 	{
-		const uint64 MetaSize = uint64(Header.BlockCount) * sizeof(uint32);
-		if (sizeof(FHeader) + MetaSize <= CompressedData.GetSize())
-		{
-			const uint32 BaseCrc32 = FMethod::CalculateCrc32(Header, CompressedData);
-			return CalculateRangeCrc32(CompressedData, sizeof(FHeader), MetaSize, BaseCrc32);
-		}
-		return 0;
+		return sizeof(FHeader) + sizeof(uint32) * uint64(Header.BlockCount);
 	}
 
+protected:
 	virtual EMethod GetMethod() const = 0;
 	virtual uint64 CompressBlockBound(uint64 RawSize) const = 0;
 	virtual bool CompressBlock(FMutableMemoryView& CompressedData, FMemoryView RawData) const = 0;
@@ -215,7 +210,6 @@ FCompositeBuffer FMethodBlock::Compress(const FCompositeBuffer& RawData, const u
 		BlockCount == 0 ? 0 :
 		BlockCount == 1 ? CompressBlockBound(RawSize) : CompressBlockBound(BlockSize) - BlockSize + RawSize;
 	FUniqueBuffer CompressedData = FUniqueBuffer::Alloc(CompressedDataSize);
-	const FMutableMemoryView CompressedDataView = CompressedData.GetView();
 
 	// Compress the raw data in blocks and store the raw data for incompressible blocks.
 	TArray64<uint32> CompressedBlockSizes;
@@ -223,7 +217,7 @@ FCompositeBuffer FMethodBlock::Compress(const FCompositeBuffer& RawData, const u
 	uint64 CompressedSize = 0;
 	{
 		FUniqueBuffer RawBlockCopy;
-		FMutableMemoryView CompressedBlocksView = CompressedDataView + sizeof(FHeader) + MetaSize;
+		FMutableMemoryView CompressedBlocksView = CompressedData.GetView() + sizeof(FHeader) + MetaSize;
 		for (uint64 RawOffset = 0; RawOffset < RawSize;)
 		{
 			const uint64 RawBlockSize = FMath::Min(RawSize - RawOffset, BlockSize);
@@ -261,26 +255,24 @@ FCompositeBuffer FMethodBlock::Compress(const FCompositeBuffer& RawData, const u
 
 	// Write the header and calculate the CRC-32.
 	Algo::ForEach(CompressedBlockSizes, [](uint32& Size) { Size = NETWORK_ORDER32(Size); });
-	CompressedDataView.Mid(sizeof(FHeader), MetaSize).CopyFrom(MakeMemoryView(CompressedBlockSizes));
-
-	const FMemoryView CompositeView = CompressedDataView.Left(sizeof(FHeader) + MetaSize + CompressedSize);
-	const FCompositeBuffer Composite(FSharedBuffer::MakeView(CompositeView, CompressedData.MoveToShared()));
+	CompressedData.GetView().Mid(sizeof(FHeader), MetaSize).CopyFrom(MakeMemoryView(CompressedBlockSizes));
 
 	FHeader Header;
 	Header.Method = GetMethod();
 	Header.BlockSizeExponent = static_cast<uint8>(FMath::FloorLog2_64(BlockSize));
 	Header.BlockCount = static_cast<uint32>(BlockCount);
 	Header.TotalRawSize = RawSize;
-	Header.Write(CompressedDataView);
-	Header.Crc32 = CalculateCrc32(Header, Composite);
-	Header.Write(CompressedDataView);
+	Header.TotalCompressedSize = sizeof(FHeader) + MetaSize + CompressedSize;
+	Header.Write(CompressedData.GetView().Left(sizeof(FHeader) + MetaSize));
 
-	return Composite;
+	const FMemoryView CompositeView = CompressedData.GetView().Left(Header.TotalCompressedSize);
+	return FCompositeBuffer(FSharedBuffer::MakeView(CompositeView, CompressedData.MoveToShared()));
 }
 
 FCompositeBuffer FMethodBlock::Decompress(const FHeader& Header, const FCompositeBuffer& CompressedData) const
 {
-	if (Header.BlockCount == 0)
+	if (Header.BlockCount == 0 || 
+		Header.TotalCompressedSize != CompressedData.GetSize())
 	{
 		return FCompositeBuffer();
 	}
@@ -391,7 +383,8 @@ FCompositeBuffer FMethodBlock::Decompress(const FHeader& Header, const FComposit
 
 bool FMethodBlock::TryDecompressTo(const FHeader& Header, const FCompositeBuffer& CompressedData, FMutableMemoryView RawView) const
 {
-	if (Header.TotalRawSize != RawView.GetSize())
+	if (Header.TotalRawSize != RawView.GetSize() ||
+		Header.TotalCompressedSize != CompressedData.GetSize())
 	{
 		return false;
 	}
@@ -506,7 +499,7 @@ static const FMethod* GetMethod(FName FormatName)
 	{
 		return GetMethod(EMethod::None);
 	}
-	if (FormatName == NAME_LZ4)
+	if (FormatName == NAME_LZ4 || FormatName == NAME_Default)
 	{
 		return GetMethod(EMethod::LZ4);
 	}
@@ -528,9 +521,11 @@ bool FHeader::IsValid(const FCompositeBuffer& CompressedData)
 		const FHeader Header = Read(CompressedData);
 		if (Header.Magic == FHeader::ExpectedMagic)
 		{
-			if (const FMethod* Method = GetMethod(Header.Method))
+			if (const FMethod* const Method = GetMethod(Header.Method))
 			{
-				if (Header.Crc32 == Method->CalculateCrc32(Header, CompressedData))
+				FUniqueBuffer HeaderCopy;
+				const FMemoryView HeaderView = CompressedData.ViewOrCopyRange(0, Method->GetHeaderSize(Header), HeaderCopy);
+				if (Header.Crc32 == FHeader::CalculateCrc32(HeaderView))
 				{
 					return true;
 				}
@@ -548,7 +543,7 @@ FCompressedBuffer FCompressedBuffer::Compress(FName FormatName, const FComposite
 {
 	using namespace UE::CompressedBuffer;
 	FCompressedBuffer Local;
-	if (const FMethod* Method = GetMethod(FormatName))
+	if (const FMethod* const Method = GetMethod(FormatName))
 	{
 		Local.CompressedData = Method->Compress(RawData);
 	}
@@ -588,6 +583,33 @@ FCompressedBuffer FCompressedBuffer::FromCompressed(FSharedBuffer&& InCompressed
 	return Local;
 }
 
+FCompressedBuffer FCompressedBuffer::FromCompressed(FArchive& Ar)
+{
+	using namespace UE::CompressedBuffer;
+	check(Ar.IsLoading());
+
+	if (sizeof(FHeader) > Ar.TotalSize() - Ar.Tell())
+	{
+		return FCompressedBuffer();
+	}
+
+	FHeader Header;
+	Ar.Serialize(&Header, sizeof(FHeader));
+	Header.ByteSwap();
+
+	if (Header.TotalCompressedSize + Ar.Tell() > Ar.TotalSize() + sizeof(FHeader))
+	{
+		return FCompressedBuffer();
+	}
+
+	FUniqueBuffer MutableBuffer = FUniqueBuffer::Alloc(Header.TotalCompressedSize);
+	Header.ByteSwap();
+	const FMutableMemoryView MutableView = MutableBuffer.GetView().CopyFrom(MakeMemoryView(&Header, &Header + 1));
+	Ar.Serialize(MutableView.GetData(), static_cast<int64>(MutableView.GetSize()));
+
+	return FromCompressed(MutableBuffer.MoveToShared());
+}
+
 uint64 FCompressedBuffer::GetRawSize() const
 {
 	return CompressedData ? UE::CompressedBuffer::FHeader::Read(CompressedData).TotalRawSize : 0;
@@ -599,7 +621,7 @@ bool FCompressedBuffer::TryDecompressTo(FMutableMemoryView RawView) const
 	if (CompressedData)
 	{
 		const FHeader Header = FHeader::Read(CompressedData);
-		if (const FMethod* Method = GetMethod(Header.Method))
+		if (const FMethod* const Method = GetMethod(Header.Method))
 		{
 			return Method->TryDecompressTo(Header, CompressedData, RawView);
 		}
@@ -613,7 +635,7 @@ FSharedBuffer FCompressedBuffer::Decompress() const
 	if (CompressedData)
 	{
 		const FHeader Header = FHeader::Read(CompressedData);
-		if (const FMethod* Method = GetMethod(Header.Method))
+		if (const FMethod* const Method = GetMethod(Header.Method))
 		{
 			if (Header.Method == EMethod::None)
 			{
@@ -635,7 +657,7 @@ FCompositeBuffer FCompressedBuffer::DecompressToComposite() const
 	if (CompressedData)
 	{
 		const FHeader Header = FHeader::Read(CompressedData);
-		if (const FMethod* Method = GetMethod(Header.Method))
+		if (const FMethod* const Method = GetMethod(Header.Method))
 		{
 			return Method->Decompress(Header, CompressedData);
 		}
diff --git a/Engine/Source/Runtime/Core/Public/Compression/CompressedBuffer.h b/Engine/Source/Runtime/Core/Public/Compression/CompressedBuffer.h
--- a/Engine/Source/Runtime/Core/Public/Compression/CompressedBuffer.h
+++ b/Engine/Source/Runtime/Core/Public/Compression/CompressedBuffer.h
@@ -5,6 +5,7 @@
 #include "Memory/CompositeBuffer.h"
 #include "Memory/MemoryFwd.h"
 
+class FArchive;
 class FName;
 
 /**
@@ -19,9 +20,9 @@ public:
 	/**
 	 * Compress the buffer using the requested compression format.
 	 *
-	 * @param FormatName   One of NAME_None, NAME_LZ4.
+	 * @param FormatName   One of NAME_None, NAME_Default, NAME_LZ4.
 	 * @param RawData      Raw data to compress. NAME_None will reference owned raw data.
-	 * @return An owned compressed buffer, or empty on error.
+	 * @return An owned compressed buffer, or null on error.
 	 */
 	[[nodiscard]] CORE_API static FCompressedBuffer Compress(FName FormatName, const FCompositeBuffer& RawData);
 	[[nodiscard]] CORE_API static FCompressedBuffer Compress(FName FormatName, const FSharedBuffer& RawData);
@@ -29,18 +30,22 @@ public:
 	/**
 	 * Construct from a compressed buffer previously created by Compress().
 	 *
-	 * @return A compressed buffer, or empty on error, such as an invalid format or corrupt header.
+	 * @return A compressed buffer, or null on error, such as an invalid format or corrupt header.
 	 */
 	[[nodiscard]] CORE_API static FCompressedBuffer FromCompressed(const FCompositeBuffer& CompressedData);
 	[[nodiscard]] CORE_API static FCompressedBuffer FromCompressed(FCompositeBuffer&& CompressedData);
 	[[nodiscard]] CORE_API static FCompressedBuffer FromCompressed(const FSharedBuffer& CompressedData);
 	[[nodiscard]] CORE_API static FCompressedBuffer FromCompressed(FSharedBuffer&& CompressedData);
+	[[nodiscard]] CORE_API static FCompressedBuffer FromCompressed(FArchive& Ar);
 
-	/** Returns true if the compressed buffer is not empty. */
-	[[nodiscard]] inline explicit operator bool() const { return !IsEmpty(); }
+	/** Reset this to null. */
+	inline void Reset() { CompressedData.Reset(); }
 
-	/** Returns true if the compressed buffer is empty. */
-	[[nodiscard]] inline bool IsEmpty() const { return CompressedData.IsEmpty(); }
+	/** Returns true if the compressed buffer is not null. */
+	[[nodiscard]] inline explicit operator bool() const { return !IsNull(); }
+
+	/** Returns true if the compressed buffer is null. */
+	[[nodiscard]] inline bool IsNull() const { return CompressedData.IsNull(); }
 
 	/** Returns true if the composite buffer is owned. */
 	[[nodiscard]] inline bool IsOwned() const { return CompressedData.IsOwned(); }
@@ -49,14 +54,14 @@ public:
 	[[nodiscard]] inline FCompressedBuffer MakeOwned() const & { return FromCompressed(CompressedData.MakeOwned()); }
 	[[nodiscard]] inline FCompressedBuffer MakeOwned() && { return FromCompressed(MoveTemp(CompressedData).MakeOwned()); }
 
-	/** Returns a composite buffer containing the compressed data. May be empty. May not be owned. */
+	/** Returns a composite buffer containing the compressed data. May be null. May not be owned. */
 	[[nodiscard]] inline const FCompositeBuffer& GetCompressed() const & { return CompressedData; }
 	[[nodiscard]] inline FCompositeBuffer GetCompressed() && { return MoveTemp(CompressedData); }
 
-	/** Returns the size of the compressed data. Zero if this is empty. */
+	/** Returns the size of the compressed data. Zero if this is null. */
 	[[nodiscard]] inline uint64 GetCompressedSize() const { return CompressedData.GetSize(); }
 
-	/** Returns the size of the raw data. Zero on error or if this is empty. */
+	/** Returns the size of the raw data. Zero on error or if this is empty or null. */
 	[[nodiscard]] CORE_API uint64 GetRawSize() const;
 
 	/**
@@ -67,14 +72,14 @@ public:
 	/**
 	 * Decompress into an owned buffer.
 	 *
-	 * @return An owned buffer containing the raw data, or empty on error or if this is empty.
+	 * @return An owned buffer containing the raw data, or null on error.
 	 */
 	[[nodiscard]] CORE_API FSharedBuffer Decompress() const;
 
 	/**
 	 * Decompress into an owned composite buffer.
 	 *
-	 * @return An owned buffer containing the raw data, or empty on error or if this is empty.
+	 * @return An owned buffer containing the raw data, or null on error.
 	 */
 	[[nodiscard]] CORE_API FCompositeBuffer DecompressToComposite() const;
 
diff --git a/Engine/Source/Runtime/Core/Public/Memory/CompositeBuffer.h b/Engine/Source/Runtime/Core/Public/Memory/CompositeBuffer.h
--- a/Engine/Source/Runtime/Core/Public/Memory/CompositeBuffer.h
+++ b/Engine/Source/Runtime/Core/Public/Memory/CompositeBuffer.h
@@ -36,7 +36,7 @@ public:
 		}
 	}
 
-	/** Reset this to empty. */
+	/** Reset this to null. */
 	CORE_API void Reset();
 
 	/** Returns the total size of the composite buffer in bytes. */
@@ -45,11 +45,11 @@ public:
 	/** Returns the segments that the buffer is composed from. */
 	[[nodiscard]] inline TConstArrayView<FSharedBuffer> GetSegments() const { return Segments; }
 
-	/** Returns true if the composite buffer has a total size greater than zero. */
-	[[nodiscard]] inline explicit operator bool() const { return !IsEmpty(); }
+	/** Returns true if the composite buffer is not null. */
+	[[nodiscard]] inline explicit operator bool() const { return !IsNull(); }
 
-	/** Returns true if the composite buffer has a total size of zero. */
-	[[nodiscard]] inline bool IsEmpty() const { return Segments.IsEmpty() || GetSize() == 0; }
+	/** Returns true if the composite buffer is null. */
+	[[nodiscard]] inline bool IsNull() const { return Segments.IsEmpty(); }
 
 	/** Returns true if every segment in the composite buffer is owned. */
 	[[nodiscard]] CORE_API bool IsOwned() const;
