CompressedBuffer: Fixed decompression of empty buffers to produce an empty buffer instead of a null buffer

#rnx

diff --git a/Engine/Source/Runtime/Core/Private/Compression/CompressedBuffer.cpp b/Engine/Source/Runtime/Core/Private/Compression/CompressedBuffer.cpp
--- a/Engine/Source/Runtime/Core/Private/Compression/CompressedBuffer.cpp
+++ b/Engine/Source/Runtime/Core/Private/Compression/CompressedBuffer.cpp
@@ -278,7 +278,8 @@ FCompositeBuffer FMethodBlock::Decompress(const FHeader& Header, const FComposit
 	}
 
 	// The raw data cannot reference the compressed data unless it is owned.
-	if (!CompressedData.IsOwned())
+	// An empty raw buffer requires an empty segment, which this path creates.
+	if (!CompressedData.IsOwned() || Header.TotalRawSize == 0)
 	{
 		FUniqueBuffer Buffer = FUniqueBuffer::Alloc(Header.TotalRawSize);
 		return TryDecompressTo(Header, CompressedData, Buffer) ? FCompositeBuffer(Buffer.MoveToShared()) : FCompositeBuffer();
diff --git a/Engine/Source/Runtime/Core/Private/Memory/CompositeBuffer.cpp b/Engine/Source/Runtime/Core/Private/Memory/CompositeBuffer.cpp
--- a/Engine/Source/Runtime/Core/Private/Memory/CompositeBuffer.cpp
+++ b/Engine/Source/Runtime/Core/Private/Memory/CompositeBuffer.cpp
@@ -57,6 +57,9 @@ FSharedBuffer FCompositeBuffer::Flatten() &&
 
 FCompositeBuffer FCompositeBuffer::Mid(uint64 Offset, uint64 Size) const
 {
+	const uint64 BufferSize = GetSize();
+	Offset = FMath::Min(Offset, BufferSize);
+	Size = FMath::Min(Size, BufferSize - Offset);
 	FCompositeBuffer Buffer;
 	IterateRange(Offset, Size, [&Buffer](FMemoryView View, const FSharedBuffer& ViewOuter)
 		{
@@ -82,7 +85,7 @@ FMemoryView FCompositeBuffer::ViewOrCopyRange(uint64 Offset, uint64 Size, FUniqu
 					{
 						CopyBuffer = FUniqueBuffer::Alloc(Size);
 					}
-					View = WriteView = CopyBuffer;
+					View = WriteView = CopyBuffer.GetView().Left(Size);
 				}
 				WriteView = WriteView.CopyFrom(Segment);
 			}
@@ -110,20 +113,23 @@ void FCompositeBuffer::IterateRange(uint64 Offset, uint64 Size,
 		TEXT(" of a composite buffer containing %" UINT64_FMT " bytes."), Size, Offset, GetSize());
 	for (const FSharedBuffer& Segment : Segments)
 	{
-		if (Size == 0)
+		if (const uint64 SegmentSize = Segment.GetSize(); Offset <= SegmentSize)
 		{
-			break;
-		}
-		const FMemoryView View = Segment.GetView().Mid(Offset, Size);
-		if (const uint64 ViewSize = View.GetSize())
-		{
-			Visitor(View, Segment);
+			const FMemoryView View = Segment.GetView().Mid(Offset, Size);
 			Offset = 0;
-			Size -= ViewSize;
+			if (Size == 0 || !View.IsEmpty())
+			{
+				Visitor(View, Segment);
+			}
+			Size -= View.GetSize();
+			if (Size == 0)
+			{
+				break;
+			}
 		}
 		else
 		{
-			Offset -= Segment.GetSize();
+			Offset -= SegmentSize;
 		}
 	}
 }
diff --git a/Engine/Source/Runtime/Core/Private/Tests/Memory/CompositeBufferTest.cpp b/Engine/Source/Runtime/Core/Private/Tests/Memory/CompositeBufferTest.cpp
new file mode 100644
--- /dev/null
+++ b/Engine/Source/Runtime/Core/Private/Tests/Memory/CompositeBufferTest.cpp
@@ -0,0 +1,176 @@
+// Copyright Epic Games, Inc. All Rights Reserved.
+
+#include "Memory/CompositeBuffer.h"
+
+#include "Misc/AutomationTest.h"
+
+#if WITH_DEV_AUTOMATION_TESTS
+
+IMPLEMENT_SIMPLE_AUTOMATION_TEST(FCompositeBufferTest, "System.Core.Memory.CompositeBuffer", EAutomationTestFlags::ApplicationContextMask | EAutomationTestFlags::SmokeFilter)
+bool FCompositeBufferTest::RunTest(const FString& Parameters)
+{
+	// Test Null
+	{
+		FCompositeBuffer Buffer;
+		TestTrue(TEXT("FCompositeBuffer().IsNull()"), Buffer.IsNull());
+		TestTrue(TEXT("FCompositeBuffer().IsOwned()"), Buffer.IsOwned());
+		TestTrue(TEXT("FCompositeBuffer().MakeOwned().IsNull()"), Buffer.MakeOwned().IsNull());
+		TestTrue(TEXT("FCompositeBuffer().Flatten().IsNull()"), Buffer.Flatten().IsNull());
+		TestTrue(TEXT("FCompositeBuffer().Mid(0, 0).IsNull()"), Buffer.Mid(0, 0).IsNull());
+		TestEqual(TEXT("FCompositeBuffer().GetSize()"), Buffer.GetSize(), uint64(0));
+		TestEqual(TEXT("FCompositeBuffer().GetSegments().Num()"), Buffer.GetSegments().Num(), 0);
+
+		FUniqueBuffer CopyBuffer;
+		TestTrue(TEXT("FCompositeBuffer().ViewOrCopyRange(0, 0).IsEmpty()"), Buffer.ViewOrCopyRange(0, 0, CopyBuffer).IsEmpty() && CopyBuffer.IsNull());
+
+		FMutableMemoryView CopyView;
+		Buffer.CopyTo(CopyView);
+
+		uint32 VisitCount = 0;
+		Buffer.IterateRange(0, 0, [&VisitCount](FMemoryView) { ++VisitCount; });
+		TestEqual(TEXT("FCompositeBuffer().IterateRange(0, 0)"), VisitCount, uint32(0));
+	}
+	{
+		FCompositeBuffer Buffer(FSharedBuffer(), FCompositeBuffer(), TArray<FSharedBuffer>{FSharedBuffer()});
+		TestTrue(TEXT("FCompositeBuffer(Nulls).IsNull()"), Buffer.IsNull());
+		TestEqual(TEXT("FCompositeBuffer(Nulls).GetSegments().Num()"), Buffer.GetSegments().Num(), 0);
+	}
+
+	// Test Empty
+	{
+		const uint8 EmptyArray[]{0};
+		const FSharedBuffer EmptyView = FSharedBuffer::MakeView(EmptyArray, 0);
+		FCompositeBuffer Buffer(EmptyView);
+		TestFalse(TEXT("FCompositeBuffer(Empty).IsNull()"), Buffer.IsNull());
+		TestFalse(TEXT("FCompositeBuffer(Empty).IsOwned()"), Buffer.IsOwned());
+		TestFalse(TEXT("FCompositeBuffer(Empty).MakeOwned().IsNull()"), Buffer.MakeOwned().IsNull());
+		TestTrue(TEXT("FCompositeBuffer(Empty).MakeOwned().IsOwned()"), Buffer.MakeOwned().IsOwned());
+		TestEqual(TEXT("FCompositeBuffer(Empty).Flatten()"), Buffer.Flatten(), EmptyView);
+		TestEqual(TEXT("FCompositeBuffer(Empty).Mid(0, 0)"), Buffer.Mid(0, 0).Flatten(), EmptyView);
+		TestEqual(TEXT("FCompositeBuffer(Empty).GetSize()"), Buffer.GetSize(), uint64(0));
+		TestEqual(TEXT("FCompositeBuffer(Empty).GetSegments().Num()"), Buffer.GetSegments().Num(), 1);
+		TestEqual(TEXT("FCompositeBuffer(Empty).GetSegments()[0]"), Buffer.GetSegments()[0], EmptyView);
+
+		FUniqueBuffer CopyBuffer;
+		TestTrue(TEXT("FCompositeBuffer(Empty).ViewOrCopyRange(0, 0)"), Buffer.ViewOrCopyRange(0, 0, CopyBuffer) == EmptyView.GetView() && CopyBuffer.IsNull());
+
+		FMutableMemoryView CopyView;
+		Buffer.CopyTo(CopyView);
+
+		uint32 VisitCount = 0;
+		Buffer.IterateRange(0, 0, [&VisitCount](FMemoryView) { ++VisitCount; });
+		TestEqual(TEXT("FCompositeBuffer(Empty).IterateRange(0, 0)"), VisitCount, uint32(1));
+	}
+	{
+		const uint8 EmptyArray[1]{};
+		const FSharedBuffer EmptyView1 = FSharedBuffer::MakeView(EmptyArray, 0);
+		const FSharedBuffer EmptyView2 = FSharedBuffer::MakeView(EmptyArray + 1, 0);
+		FCompositeBuffer Buffer(TArray<FSharedBuffer>{EmptyView1}, FCompositeBuffer(EmptyView2));
+		TestEqual(TEXT("FCompositeBuffer(Empty, Empty).Mid(0, 0)"), Buffer.Mid(0, 0).Flatten(), EmptyView1);
+		TestEqual(TEXT("FCompositeBuffer(Empty, Empty).GetSize()"), Buffer.GetSize(), uint64(0));
+		TestEqual(TEXT("FCompositeBuffer(Empty, Empty).GetSegments().Num()"), Buffer.GetSegments().Num(), 2);
+		TestEqual(TEXT("FCompositeBuffer(Empty, Empty).GetSegments()[0]"), Buffer.GetSegments()[0], EmptyView1);
+		TestEqual(TEXT("FCompositeBuffer(Empty, Empty).GetSegments()[1]"), Buffer.GetSegments()[1], EmptyView2);
+
+		FUniqueBuffer CopyBuffer;
+		TestTrue(TEXT("FCompositeBuffer(Empty, Empty).ViewOrCopyRange(0, 0)"), Buffer.ViewOrCopyRange(0, 0, CopyBuffer) == EmptyView1.GetView() && CopyBuffer.IsNull());
+
+		FMutableMemoryView CopyView;
+		Buffer.CopyTo(CopyView);
+
+		uint32 VisitCount = 0;
+		Buffer.IterateRange(0, 0, [&VisitCount](FMemoryView) { ++VisitCount; });
+		TestEqual(TEXT("FCompositeBuffer(Empty, Empty).IterateRange(0, 0)"), VisitCount, uint32(1));
+	}
+
+	// Test Flat
+	{
+		const uint8 FlatArray[]{1, 2, 3, 4, 5, 6, 7, 8};
+		const FSharedBuffer FlatView = FSharedBuffer::Clone(MakeMemoryView(FlatArray));
+		FCompositeBuffer Buffer(FlatView);
+		TestFalse(TEXT("FCompositeBuffer(Flat).IsNull()"), Buffer.IsNull());
+		TestTrue(TEXT("FCompositeBuffer(Flat).IsOwned()"), Buffer.IsOwned());
+		TestEqual(TEXT("FCompositeBuffer(Flat).Flatten()"), Buffer.Flatten(), FlatView);
+		TestEqual(TEXT("FCompositeBuffer(Flat).MakeOwned().Flatten()"), Buffer.MakeOwned().Flatten(), FlatView);
+		TestEqual(TEXT("FCompositeBuffer(Flat).Mid(0)"), Buffer.Mid(0).Flatten(), FlatView);
+		TestEqual(TEXT("FCompositeBuffer(Flat).Mid(4)"), Buffer.Mid(4).Flatten().GetView(), FlatView.GetView().Mid(4));
+		TestEqual(TEXT("FCompositeBuffer(Flat).Mid(8)"), Buffer.Mid(8).Flatten().GetView(), FlatView.GetView().Mid(8));
+		TestEqual(TEXT("FCompositeBuffer(Flat).Mid(4, 2)"), Buffer.Mid(4, 2).Flatten().GetView(), FlatView.GetView().Mid(4, 2));
+		TestEqual(TEXT("FCompositeBuffer(Flat).Mid(8).GetData()"), Buffer.Mid(8).Flatten().GetView().GetData(), FlatView.GetView().Mid(8).GetData());
+		TestEqual(TEXT("FCompositeBuffer(Flat).Mid(4, 2).GetData()"), Buffer.Mid(4, 2).Flatten().GetView().GetData(), FlatView.GetView().Mid(4, 2).GetData());
+		TestEqual(TEXT("FCompositeBuffer(Flat).Mid(8, 0).GetData()"), Buffer.Mid(8, 0).Flatten().GetView().GetData(), FlatView.GetView().Mid(8, 0).GetData());
+		TestEqual(TEXT("FCompositeBuffer(Flat).GetSize()"), Buffer.GetSize(), uint64(sizeof(FlatArray)));
+		TestEqual(TEXT("FCompositeBuffer(Flat).GetSegments().Num()"), Buffer.GetSegments().Num(), 1);
+		TestEqual(TEXT("FCompositeBuffer(Flat).GetSegments()[0]"), Buffer.GetSegments()[0], FlatView);
+
+		FUniqueBuffer CopyBuffer;
+		TestTrue(TEXT("FCompositeBuffer(Flat).ViewOrCopyRange(0, 0)"), Buffer.ViewOrCopyRange(0, sizeof(FlatArray), CopyBuffer) == FlatView.GetView() && CopyBuffer.IsNull());
+
+		uint8 CopyArray[sizeof(FlatArray) - 3];
+		Buffer.CopyTo(MakeMemoryView(CopyArray), 3);
+		TestTrue(TEXT("FCompositeBuffer(Flat).CopyTo()"), MakeMemoryView(CopyArray).EqualBytes(MakeMemoryView(FlatArray) + 3));
+
+		uint32 VisitCount = 0;
+		Buffer.IterateRange(0, sizeof(FlatArray), [&VisitCount](FMemoryView) { ++VisitCount; });
+		TestEqual(TEXT("FCompositeBuffer(Flat).IterateRange(0, N)"), VisitCount, uint32(1));
+	}
+
+	// Test Composite
+	{
+		const uint8 FlatArray[]{1, 2, 3, 4, 5, 6, 7, 8};
+		const FSharedBuffer FlatView1 = FSharedBuffer::MakeView(MakeMemoryView(FlatArray).Left(4));
+		const FSharedBuffer FlatView2 = FSharedBuffer::MakeView(MakeMemoryView(FlatArray).Right(4));
+		FCompositeBuffer Buffer(FlatView1, FlatView2);
+		TestFalse(TEXT("FCompositeBuffer(Composite).IsNull()"), Buffer.IsNull());
+		TestFalse(TEXT("FCompositeBuffer(Composite).IsOwned()"), Buffer.IsOwned());
+		TestTrue(TEXT("FCompositeBuffer(Composite).Flatten()"), Buffer.Flatten().GetView().EqualBytes(MakeMemoryView(FlatArray)));
+		TestTrue(TEXT("FCompositeBuffer(Composite).Mid(2, 4)"), Buffer.Mid(2, 4).Flatten().GetView().EqualBytes(MakeMemoryView(FlatArray).Mid(2, 4)));
+		TestEqual(TEXT("FCompositeBuffer(Composite).Mid(0, 4)"), Buffer.Mid(0, 4).Flatten(), FlatView1);
+		TestEqual(TEXT("FCompositeBuffer(Composite).Mid(4, 4)"), Buffer.Mid(4, 4).Flatten(), FlatView2);
+		TestEqual(TEXT("FCompositeBuffer(Composite).Mid(4, 0).GetData()"), Buffer.Mid(4, 0).Flatten().GetView().GetData(), static_cast<const void*>(FlatArray + 4));
+		TestEqual(TEXT("FCompositeBuffer(Composite).Mid(8, 0).GetData()"), Buffer.Mid(8, 0).Flatten().GetView().GetData(), static_cast<const void*>(FlatArray + 8));
+		TestEqual(TEXT("FCompositeBuffer(Composite).GetSize()"), Buffer.GetSize(), uint64(sizeof(FlatArray)));
+		TestEqual(TEXT("FCompositeBuffer(Composite).GetSegments().Num()"), Buffer.GetSegments().Num(), 2);
+		TestEqual(TEXT("FCompositeBuffer(Composite).GetSegments()[0]"), Buffer.GetSegments()[0], FlatView1);
+		TestEqual(TEXT("FCompositeBuffer(Composite).GetSegments()[2]"), Buffer.GetSegments()[1], FlatView2);
+
+		FUniqueBuffer CopyBuffer;
+		TestTrue(TEXT("FCompositeBuffer(Composite).ViewOrCopyRange(0, 4)"), Buffer.ViewOrCopyRange(0, 4, CopyBuffer) == FlatView1.GetView() && CopyBuffer.IsNull());
+		TestTrue(TEXT("FCompositeBuffer(Composite).ViewOrCopyRange(4, 4)"), Buffer.ViewOrCopyRange(4, 4, CopyBuffer) == FlatView2.GetView() && CopyBuffer.IsNull());
+		TestTrue(TEXT("FCompositeBuffer(Composite).ViewOrCopyRange(3, 2)"), Buffer.ViewOrCopyRange(3, 2, CopyBuffer).EqualBytes(MakeMemoryView(FlatArray).Mid(3, 2)) && CopyBuffer.GetSize() == 2);
+		TestTrue(TEXT("FCompositeBuffer(Composite).ViewOrCopyRange(1, 6)"), Buffer.ViewOrCopyRange(1, 6, CopyBuffer).EqualBytes(MakeMemoryView(FlatArray).Mid(1, 6)) && CopyBuffer.GetSize() == 6);
+		TestTrue(TEXT("FCompositeBuffer(Composite).ViewOrCopyRange(2, 4)"), Buffer.ViewOrCopyRange(2, 4, CopyBuffer).EqualBytes(MakeMemoryView(FlatArray).Mid(2, 4)) && CopyBuffer.GetSize() == 6);
+
+		uint8 CopyArray[4];
+		Buffer.CopyTo(MakeMemoryView(CopyArray), 2);
+		TestTrue(TEXT("FCompositeBuffer(Composite).CopyTo()"), MakeMemoryView(CopyArray).EqualBytes(MakeMemoryView(FlatArray).Mid(2, 4)));
+
+		uint32 VisitCount = 0;
+		Buffer.IterateRange(0, sizeof(FlatArray), [&VisitCount](FMemoryView) { ++VisitCount; });
+		TestEqual(TEXT("FCompositeBuffer(Composite).IterateRange(0, N)"), VisitCount, uint32(2));
+
+		const auto TestIterateRange = [this, &Buffer](uint64 Offset, uint64 Size, FMemoryView ExpectedView, const FSharedBuffer& ExpectedViewOuter)
+		{
+			uint32 VisitCount = 0;
+			FMemoryView ActualView;
+			FSharedBuffer ActualViewOuter;
+			Buffer.IterateRange(Offset, Size, [&VisitCount, &ActualView, &ActualViewOuter](FMemoryView View, const FSharedBuffer& ViewOuter)
+			{
+				++VisitCount;
+				ActualView = View;
+				ActualViewOuter = ViewOuter;
+			});
+			TestEqual(FString::Printf(TEXT("FCompositeBuffer(Composite).IterateRange(%" UINT64_FMT ", %" UINT64_FMT ")->VisitCount"), Offset, Size), VisitCount, uint32(1));
+			TestEqual(FString::Printf(TEXT("FCompositeBuffer(Composite).IterateRange(%" UINT64_FMT ", %" UINT64_FMT ")->View"), Offset, Size), ActualView, ExpectedView);
+			TestEqual(FString::Printf(TEXT("FCompositeBuffer(Composite).IterateRange(%" UINT64_FMT ", %" UINT64_FMT ")->ViewOuter"), Offset, Size), ActualViewOuter, ExpectedViewOuter);
+		};
+		TestIterateRange(0, 4, MakeMemoryView(FlatArray).Mid(0, 4), FlatView1);
+		TestIterateRange(4, 0, MakeMemoryView(FlatArray).Mid(4, 0), FlatView1);
+		TestIterateRange(4, 4, MakeMemoryView(FlatArray).Mid(4, 4), FlatView2);
+		TestIterateRange(8, 0, MakeMemoryView(FlatArray).Mid(8, 0), FlatView2);
+	}
+
+	return true;
+}
+
+#endif // WITH_DEV_AUTOMATION_TESTS
