LLM: Clamp value of Default Untracked to 0 in the same way as we clamp Platform Untracked to 0.
#rnx

diff --git a/Engine/Source/Runtime/Core/Private/HAL/LowLevelMemTracker.cpp b/Engine/Source/Runtime/Core/Private/HAL/LowLevelMemTracker.cpp
--- a/Engine/Source/Runtime/Core/Private/HAL/LowLevelMemTracker.cpp
+++ b/Engine/Source/Runtime/Core/Private/HAL/LowLevelMemTracker.cpp
@@ -770,10 +770,12 @@ void FLowLevelMemTracker::TickInternal()
 	PlatformTracker.SetTagAmountInUpdate(FindOrAddTagData(ELLMTag::PlatformOSAvailable), PlatformStats.AvailablePhysical, false);
 
 	int64 TrackedTotal = DefaultTracker.GetTrackedTotal();
-	// remove the MemoryUsageCurrentOverhead from the default LLM as it's not something anyone needs to investigate when finding what to reduce
+	// remove the MemoryUsageCurrentOverhead from the "Total" for the default LLM as it's not something anyone needs to investigate when finding what to reduce
 	// the platform LLM will have the info 
-	DefaultTracker.SetTagAmountInUpdate(FindOrAddTagData(ELLMTag::Total), PlatformProcessMemory - MemoryUsageCurrentOverhead, false);
-	DefaultTracker.SetTagAmountInUpdate(FindOrAddTagData(ELLMTag::Untracked), PlatformProcessMemory - (TrackedTotal + MemoryUsageCurrentOverhead), false);
+	int64 DefaultProcessMemory = PlatformProcessMemory - MemoryUsageCurrentOverhead;
+	int64 DefaultUntracked = FMath::Max<int64>(0, DefaultProcessMemory - TrackedTotal);
+	DefaultTracker.SetTagAmountInUpdate(FindOrAddTagData(ELLMTag::Total), DefaultProcessMemory, false);
+	DefaultTracker.SetTagAmountInUpdate(FindOrAddTagData(ELLMTag::Untracked), DefaultUntracked, false);
 
 #if PLATFORM_WINDOWS
 	DefaultTracker.SetTagAmountInUpdate(FindOrAddTagData(ELLMTag::WorkingSetSize), PlatformStats.UsedPhysical, false);
