CompactBinary: Added CompositeBuffer overloads and made non-view arrays and objects behave more consistently

#rnx

diff --git a/Engine/Source/Runtime/Core/Private/Serialization/CompactBinaryWriter.cpp b/Engine/Source/Runtime/Core/Private/Serialization/CompactBinaryWriter.cpp
--- a/Engine/Source/Runtime/Core/Private/Serialization/CompactBinaryWriter.cpp
+++ b/Engine/Source/Runtime/Core/Private/Serialization/CompactBinaryWriter.cpp
@@ -4,6 +4,8 @@
 
 #include "Containers/StringConv.h"
 #include "IO/IoHash.h"
+#include "Memory/CompositeBuffer.h"
+#include "Memory/SharedBuffer.h"
 #include "Misc/ByteSwap.h"
 #include "Misc/DateTime.h"
 #include "Misc/Guid.h"
@@ -368,11 +370,16 @@ void FCbWriter::AddBinary(const void* const Value, const uint64 Size)
 	EndField(ECbFieldType::Binary);
 }
 
-void FCbWriter::AddBinary(FSharedBuffer Buffer)
+void FCbWriter::AddBinary(const FSharedBuffer& Buffer)
 {
 	AddBinary(Buffer.GetData(), Buffer.GetSize());
 }
 
+void FCbWriter::AddBinary(const FCompositeBuffer& Buffer)
+{
+	AddBinary(Buffer.Flatten());
+}
+
 void FCbWriter::AddString(const FAnsiStringView Value)
 {
 	BeginField();
diff --git a/Engine/Source/Runtime/Core/Public/Serialization/CompactBinary.h b/Engine/Source/Runtime/Core/Public/Serialization/CompactBinary.h
--- a/Engine/Source/Runtime/Core/Public/Serialization/CompactBinary.h
+++ b/Engine/Source/Runtime/Core/Public/Serialization/CompactBinary.h
@@ -1281,6 +1281,10 @@ public:
 	{
 		return FCbField(FCbArrayView::AsFieldView(), MoveTemp(*this));
 	}
+
+private:
+	friend inline FCbFieldIterator begin(const FCbArray& Array) { return Array.CreateIterator(); }
+	friend inline FCbFieldIterator end(const FCbArray&) { return FCbFieldIterator(); }
 };
 
 ///////////////////////////////////////////////////////////////////////////////////////////////////////////////////////
@@ -1322,6 +1326,9 @@ public:
 		return FCbField();
 	}
 
+	/** Find a field by case-sensitive name comparison. */
+	inline FCbField operator[](FAnsiStringView Name) const { return Find(Name); }
+
 	/** Access the object as an object field. */
 	inline FCbField AsField() const &
 	{
@@ -1333,6 +1340,10 @@ public:
 	{
 		return FCbField(FCbObjectView::AsFieldView(), MoveTemp(*this));
 	}
+
+private:
+	friend inline FCbFieldIterator begin(const FCbObject& Object) { return Object.CreateIterator(); }
+	friend inline FCbFieldIterator end(const FCbObject&) { return FCbFieldIterator(); }
 };
 
 ///////////////////////////////////////////////////////////////////////////////////////////////////////////////////////
diff --git a/Engine/Source/Runtime/Core/Public/Serialization/CompactBinaryWriter.h b/Engine/Source/Runtime/Core/Public/Serialization/CompactBinaryWriter.h
--- a/Engine/Source/Runtime/Core/Public/Serialization/CompactBinaryWriter.h
+++ b/Engine/Source/Runtime/Core/Public/Serialization/CompactBinaryWriter.h
@@ -153,8 +153,10 @@ public:
 	inline void AddBinary(FAnsiStringView Name, FMemoryView Value) { SetName(Name); AddBinary(Value); }
 	inline void AddBinary(FMemoryView Value) { AddBinary(Value.GetData(), Value.GetSize()); }
 	/** Write a binary field by copying the buffer. Holds a reference if owned. */
-	inline void AddBinary(FAnsiStringView Name, FSharedBuffer Value) { SetName(Name); AddBinary(MoveTemp(Value)); }
-	CORE_API void AddBinary(FSharedBuffer Value);
+	inline void AddBinary(FAnsiStringView Name, const FSharedBuffer& Value) { SetName(Name); AddBinary(Value); }
+	CORE_API void AddBinary(const FSharedBuffer& Value);
+	inline void AddBinary(FAnsiStringView Name, const FCompositeBuffer& Value) { SetName(Name); AddBinary(Value); }
+	CORE_API void AddBinary(const FCompositeBuffer& Value);
 
 	/** Write a string field by copying the UTF-8 value. */
 	inline void AddString(FAnsiStringView Name, FAnsiStringView Value) { SetName(Name); AddString(Value); }
