CompressedBuffer: Added the raw data hash to the header

#rnx

diff --git a/Engine/Source/Editor/UnrealEd/Private/EditorDomain/EditorDomainUtils.cpp b/Engine/Source/Editor/UnrealEd/Private/EditorDomain/EditorDomainUtils.cpp
--- a/Engine/Source/Editor/UnrealEd/Private/EditorDomain/EditorDomainUtils.cpp
+++ b/Engine/Source/Editor/UnrealEd/Private/EditorDomain/EditorDomainUtils.cpp
@@ -19,7 +19,7 @@ namespace EditorDomain
 {
 
 // Change to a new guid when EditorDomain needs to be invalidated
-const TCHAR* EditorDomainVersion = TEXT("25F3D3DC3B1246DB9BB1384185578B8A");
+const TCHAR* EditorDomainVersion = TEXT("C217EB656E9B4C04816D3DC0E21901F6");
 // Identifier of the CacheBucket for EditorDomainPackages
 const TCHAR* EditorDomainPackageBucketName = TEXT("EditorDomainPackage");
 
diff --git a/Engine/Source/Runtime/Core/Private/Compression/CompressedBuffer.cpp b/Engine/Source/Runtime/Core/Private/Compression/CompressedBuffer.cpp
--- a/Engine/Source/Runtime/Core/Private/Compression/CompressedBuffer.cpp
+++ b/Engine/Source/Runtime/Core/Private/Compression/CompressedBuffer.cpp
@@ -3,6 +3,7 @@
 #include "Compression/CompressedBuffer.h"
 
 #include "Algo/ForEach.h"
+#include "Hash/Blake3.h"
 #include "Math/UnrealMathUtility.h"
 #include "Misc/ByteSwap.h"
 #include "Misc/Crc.h"
@@ -55,6 +56,8 @@ struct FHeader
 	uint64 TotalRawSize = 0;
 	/** The total size of the compressed data including the header. */
 	uint64 TotalCompressedSize = 0;
+	/** The hash of the uncompressed data. */
+	FBlake3Hash RawHash;
 
 	/** Checks validity of the buffer based on the magic number, method, and CRC-32. */
 	static bool IsValid(const FCompositeBuffer& CompressedData);
@@ -112,7 +115,7 @@ struct FHeader
 	}
 };
 
-static_assert(sizeof(FHeader) == 32, "FHeader is the wrong size.");
+static_assert(sizeof(FHeader) == 64, "FHeader is the wrong size.");
 
 ///////////////////////////////////////////////////////////////////////////////////////////////////////////////////////
 
@@ -137,6 +140,7 @@ public:
 		Header.BlockCount = 1;
 		Header.TotalRawSize = RawData.GetSize();
 		Header.TotalCompressedSize = Header.TotalRawSize + sizeof(FHeader);
+		Header.RawHash = FBlake3::HashBuffer(RawData);
 
 		FUniqueBuffer HeaderData = FUniqueBuffer::Alloc(sizeof(FHeader));
 		Header.Write(HeaderData);
@@ -199,6 +203,7 @@ FCompositeBuffer FMethodBlock::Compress(const FCompositeBuffer& RawData, const u
 	checkf(FMath::IsPowerOfTwo(BlockSize) && BlockSize <= MAX_uint32,
 		TEXT("BlockSize must be a 32-bit power of two but was %" UINT64_FMT "."), BlockSize);
 	const uint64 RawSize = RawData.GetSize();
+	FBlake3 RawHash;
 
 	const uint64 BlockCount = FMath::DivideAndRoundUp(RawSize, BlockSize);
 	checkf(BlockCount <= MAX_uint32, TEXT("Raw data of size %" UINT64_FMT " with block size %" UINT64_FMT " requires ")
@@ -222,6 +227,7 @@ FCompositeBuffer FMethodBlock::Compress(const FCompositeBuffer& RawData, const u
 		{
 			const uint64 RawBlockSize = FMath::Min(RawSize - RawOffset, BlockSize);
 			const FMemoryView RawBlock = RawData.ViewOrCopyRange(RawOffset, RawBlockSize, RawBlockCopy);
+			RawHash.Update(RawBlock);
 
 			FMutableMemoryView CompressedBlock = CompressedBlocksView;
 			if (!CompressBlock(CompressedBlock, RawBlock))
@@ -263,6 +269,7 @@ FCompositeBuffer FMethodBlock::Compress(const FCompositeBuffer& RawData, const u
 	Header.BlockCount = static_cast<uint32>(BlockCount);
 	Header.TotalRawSize = RawSize;
 	Header.TotalCompressedSize = sizeof(FHeader) + MetaSize + CompressedSize;
+	Header.RawHash = RawHash.Finalize();
 	Header.Write(CompressedData.GetView().Left(sizeof(FHeader) + MetaSize));
 
 	const FMemoryView CompositeView = CompressedData.GetView().Left(Header.TotalCompressedSize);
@@ -271,7 +278,7 @@ FCompositeBuffer FMethodBlock::Compress(const FCompositeBuffer& RawData, const u
 
 FCompositeBuffer FMethodBlock::Decompress(const FHeader& Header, const FCompositeBuffer& CompressedData) const
 {
-	if (Header.BlockCount == 0 || 
+	if (Header.BlockCount == 0 ||
 		Header.TotalCompressedSize != CompressedData.GetSize())
 	{
 		return FCompositeBuffer();
@@ -616,6 +623,11 @@ uint64 FCompressedBuffer::GetRawSize() const
 	return CompressedData ? UE::CompressedBuffer::FHeader::Read(CompressedData).TotalRawSize : 0;
 }
 
+FBlake3Hash FCompressedBuffer::GetRawHash() const
+{
+	return CompressedData ? UE::CompressedBuffer::FHeader::Read(CompressedData).RawHash : FBlake3Hash();
+}
+
 bool FCompressedBuffer::TryDecompressTo(FMutableMemoryView RawView) const
 {
 	using namespace UE::CompressedBuffer;
diff --git a/Engine/Source/Runtime/Core/Public/Compression/CompressedBuffer.h b/Engine/Source/Runtime/Core/Public/Compression/CompressedBuffer.h
--- a/Engine/Source/Runtime/Core/Public/Compression/CompressedBuffer.h
+++ b/Engine/Source/Runtime/Core/Public/Compression/CompressedBuffer.h
@@ -7,6 +7,7 @@
 
 class FArchive;
 class FName;
+struct FBlake3Hash;
 
 /**
  * A compressed buffer stores compressed data in a self-contained format.
@@ -64,6 +65,9 @@ public:
 	/** Returns the size of the raw data. Zero on error or if this is empty or null. */
 	[[nodiscard]] CORE_API uint64 GetRawSize() const;
 
+	/** Returns the hash of the raw data. Zero on error or if this is null. */
+	[[nodiscard]] CORE_API FBlake3Hash GetRawHash() const;
+
 	/**
 	 * Decompress into a memory view that is exactly GetRawSize() bytes.
 	 */
