Reduce Stats allocations during early editor startup

diff --git a/Engine/Source/Runtime/Core/Private/Stats/Stats2.cpp b/Engine/Source/Runtime/Core/Private/Stats/Stats2.cpp
--- a/Engine/Source/Runtime/Core/Private/Stats/Stats2.cpp
+++ b/Engine/Source/Runtime/Core/Private/Stats/Stats2.cpp
@@ -10,6 +10,7 @@
 #include "UObject/NameTypes.h"
 #include "Logging/LogMacros.h"
 #include "Misc/Parse.h"
+#include "Misc/StringBuilder.h"
 #include "HAL/Runnable.h"
 #include "HAL/RunnableThread.h"
 #include "Misc/SingleThreadRunnable.h"
@@ -18,6 +19,7 @@
 #include "Misc/CommandLine.h"
 #include "Containers/Map.h"
 #include "Stats/Stats.h"
+#include "String/Find.h"
 #include "Async/AsyncWork.h"
 #include "Containers/Ticker.h"
 #include "Stats/StatsData.h"
@@ -727,8 +729,7 @@ IStatGroupEnableManager& IStatGroupEnableManager::Get()
 
 FName FStatNameAndInfo::ToLongName(FName InStatName, char const* InGroup, char const* InCategory, TCHAR const* InDescription, bool InSortByName)
 {
-	FString LongName;
-	LongName.Reserve(255);
+	TStringBuilder<256> LongName;
 	if (InGroup)
 	{
 		LongName += TEXT("//");
@@ -739,7 +740,7 @@ FName FStatNameAndInfo::ToLongName(FName InStatName, char const* InGroup, char c
 	if (InDescription)
 	{
 		LongName += TEXT("///");
-		LongName += FStatsUtils::ToEscapedFString(InDescription);
+		FStatsUtils::ToEscapedString(InDescription, LongName);
 		LongName += TEXT("///");
 	}
 	if (InCategory)
@@ -759,56 +760,60 @@ FName FStatNameAndInfo::ToLongName(FName InStatName, char const* InGroup, char c
 
 FName FStatNameAndInfo::GetShortNameFrom(FName InLongName)
 {
-	FString Input(InLongName.ToString());
+	TStringBuilder<256> Temp;
+	Temp << InLongName;
+	FStringView Input(Temp);
 
 	if (Input.StartsWith(TEXT("//"), ESearchCase::CaseSensitive))
 	{
-		Input.RightChopInline(2, false);
-		const int32 IndexEnd = Input.Find(TEXT("//"), ESearchCase::CaseSensitive);
+		Input.RemovePrefix(2);
+		const int32 IndexEnd = UE::String::FindFirst(Input, TEXT("//"));
 		if (IndexEnd == INDEX_NONE)
 		{
 			checkStats(0);
 			return InLongName;
 		}
-		Input.RightChopInline(IndexEnd + 2, false);
+		Input.RightChopInline(IndexEnd + 2);
 	}
-	const int32 DescIndexEnd = Input.Find(TEXT("///"), ESearchCase::CaseSensitive);
+	const int32 DescIndexEnd = UE::String::FindFirst(Input, TEXT("///"), ESearchCase::CaseSensitive);
 	if (DescIndexEnd != INDEX_NONE)
 	{
-		Input.LeftInline(DescIndexEnd, false);
+		Input.LeftInline(DescIndexEnd);
 	}
-	const int32 CategoryIndexEnd = Input.Find( TEXT( "####" ), ESearchCase::CaseSensitive );
+	const int32 CategoryIndexEnd = UE::String::FindFirst(Input, TEXT( "####" ), ESearchCase::CaseSensitive );
 	if( DescIndexEnd == INDEX_NONE && CategoryIndexEnd != INDEX_NONE )
 	{
-		Input.LeftInline(CategoryIndexEnd, false);
+		Input.LeftInline(CategoryIndexEnd);
 	}
-	const int32 SortByNameIndexEnd = Input.Find( TEXT( "/#/#" ), ESearchCase::CaseSensitive );
+	const int32 SortByNameIndexEnd = UE::String::FindFirst(Input, TEXT( "/#/#" ), ESearchCase::CaseSensitive );
 	if( DescIndexEnd == INDEX_NONE && CategoryIndexEnd == INDEX_NONE && SortByNameIndexEnd != INDEX_NONE )
 	{
-		Input.LeftInline(SortByNameIndexEnd, false);
+		Input.LeftInline(SortByNameIndexEnd);
 	}
-	return FName(*Input);
+	return FName(Input);
 }
 
 FName FStatNameAndInfo::GetGroupNameFrom(FName InLongName)
 {
-	FString Input(InLongName.ToString());
+	TStringBuilder<256> Temp;
+	Temp << InLongName;
+	FStringView Input(Temp);
 
 	if (Input.StartsWith(TEXT("//"), ESearchCase::CaseSensitive))
 	{
-		Input.RightChopInline(2, false);
+		Input.RemovePrefix(2);
 		if (Input.StartsWith(TEXT("Groups//")))
 		{
-			Input.RightChopInline(8, false);
+			Input.RemovePrefix(8);
 		}
-		const int32 IndexEnd = Input.Find(TEXT("//"), ESearchCase::CaseSensitive);
+		const int32 IndexEnd = UE::String::FindFirst(Input, TEXT("//"));
 		if (IndexEnd != INDEX_NONE)
 		{
-			return FName(*Input.Left(IndexEnd));
+			return FName(Input.Left(IndexEnd));
 		}
 		checkStats(0);
 	}
-	return NAME_None;
+	return FName();
 }
 
 FString FStatNameAndInfo::GetDescriptionFrom(FName InLongName)
@@ -822,7 +827,7 @@ FString FStatNameAndInfo::GetDescriptionFrom(FName InLongName)
 		const int32 IndexEnd = Input.Find(TEXT("///"), ESearchCase::CaseSensitive);
 		if (IndexEnd != INDEX_NONE)
 		{
-			return FStatsUtils::FromEscapedFString(*Input.Left(IndexEnd));
+			return FStatsUtils::FromEscapedString(*Input.Left(IndexEnd));
 		}
 	}
 	return FString();
diff --git a/Engine/Source/Runtime/Core/Private/Stats/StatsData.cpp b/Engine/Source/Runtime/Core/Private/Stats/StatsData.cpp
--- a/Engine/Source/Runtime/Core/Private/Stats/StatsData.cpp
+++ b/Engine/Source/Runtime/Core/Private/Stats/StatsData.cpp
@@ -5,9 +5,12 @@
 #include "Misc/CoreStats.h"
 #if STATS
 
+#include "Async/TaskGraphInterfaces.h"
 #include "Containers/LockFreeFixedSizeAllocator.h"
+#include "Containers/StringView.h"
+#include "Misc/StringBuilder.h"
 #include "HAL/IConsoleManager.h"
-#include "Async/TaskGraphInterfaces.h"
+#include "Misc/AsciiSet.h"
 
 DECLARE_CYCLE_STAT(TEXT("Broadcast"),STAT_StatsBroadcast,STATGROUP_StatSystem);
 DECLARE_CYCLE_STAT(TEXT("Condense"),STAT_StatsCondense,STATGROUP_StatSystem);
@@ -2043,7 +2046,7 @@ void FStatsUtils::AccumulateStat(FStatMessage& Dest, FStatMessage const& Item, E
 	}
 }
 
-FString FStatsUtils::FromEscapedFString(const TCHAR* Escaped)
+FString FStatsUtils::FromEscapedString(const TCHAR* Escaped)
 {
 	FString Result;
 	FString Input(Escaped);
@@ -2076,48 +2079,24 @@ FString FStatsUtils::FromEscapedFString(const TCHAR* Escaped)
 	return Result;
 }
 
-FString FStatsUtils::ToEscapedFString(const TCHAR* Source)
+void FStatsUtils::ToEscapedString(FStringView Input, FStringBuilderBase& Output)
 {
-	FString Invalid(INVALID_NAME_CHARACTERS);
-	Invalid += TEXT("$");
-
-	FString Output;
-	FString Input(Source);
-	int32 StartValid = 0;
-	int32 NumValid = 0;
-
-	for (int32 i = 0; i < Input.Len(); i++)
+	constexpr FAsciiSet InvalidChars = FAsciiSet(INVALID_NAME_CHARACTERS) + '$';
+	while (true)
 	{
-		int32 Index = 0;
-		if (!Invalid.FindChar(Input[i], Index))
+		FStringView ValidInput = FAsciiSet::FindPrefixWithout(Input, InvalidChars);
+		Output.Append(ValidInput);
+
+		if (ValidInput.end() == Input.end())
 		{
-			NumValid++;
+			break;
 		}
-		else
-		{
-			// Copy the valid range so far
-			Output += Input.Mid(StartValid, NumValid);
-
-			// Reset valid ranges
-			StartValid = i + 1;
-			NumValid = 0;
 
-			// Replace the invalid character with a special string
-			Output += FString::Printf(TEXT("$%u$"), uint32(Input[i]));
-		}
-	}
+		// Replace the invalid character with a special string
+		Output << '$' << static_cast<uint32>(*ValidInput.end()) << '$';
 
-	// Just return the input if the entire string was valid
-	if (StartValid == 0 && NumValid == Input.Len())
-	{
-		return Input;
-	}
-	else if (NumValid > 0)
-	{
-		// Copy the remaining valid part
-		Output += Input.Mid(StartValid, NumValid);
+		Input.RemovePrefix(ValidInput.Len() + 1);
 	}
-	return Output;
 }
 
 
diff --git a/Engine/Source/Runtime/Core/Public/Stats/StatsData.h b/Engine/Source/Runtime/Core/Public/Stats/StatsData.h
--- a/Engine/Source/Runtime/Core/Public/Stats/StatsData.h
+++ b/Engine/Source/Runtime/Core/Public/Stats/StatsData.h
@@ -777,8 +777,8 @@ struct CORE_API FStatsUtils
 	}
 
 	/** Internal use, converts arbitrary string to and from an escaped notation for storage in an FName. **/
-	static FString ToEscapedFString(const TCHAR* Source);
-	static FString FromEscapedFString(const TCHAR* Escaped);
+	static void ToEscapedString(FStringView Input, FStringBuilderBase& Output);
+	static FString FromEscapedString(const TCHAR* Escaped);
 	
 	static FString BuildUniqueThreadName( uint32 InThreadID )
 	{
