Reduce FName allocations during early editor startup

diff --git a/Engine/Source/Runtime/Core/Private/UObject/UnrealNames.cpp b/Engine/Source/Runtime/Core/Private/UObject/UnrealNames.cpp
--- a/Engine/Source/Runtime/Core/Private/UObject/UnrealNames.cpp
+++ b/Engine/Source/Runtime/Core/Private/UObject/UnrealNames.cpp
@@ -2217,12 +2217,12 @@ void FName::ToString(FString& Out) const
 
 	if (GetNumber() == NAME_NO_NUMBER_INTERNAL)
 	{
-		Out.Empty(NameEntry->GetNameLength());
+		Out.Reset(NameEntry->GetNameLength());
 		NameEntry->AppendNameToString(Out);
 	}	
 	else
 	{
-		Out.Empty(NameEntry->GetNameLength() + 6);
+		Out.Reset(NameEntry->GetNameLength() + 6);
 		NameEntry->AppendNameToString(Out);
 
 		Out += TEXT('_');
