Reduce ModuleManager allocations during early editor startup

diff --git a/Engine/Source/Runtime/Core/Private/Modules/ModuleManager.cpp b/Engine/Source/Runtime/Core/Private/Modules/ModuleManager.cpp
--- a/Engine/Source/Runtime/Core/Private/Modules/ModuleManager.cpp
+++ b/Engine/Source/Runtime/Core/Private/Modules/ModuleManager.cpp
@@ -1084,12 +1084,30 @@ void FModuleManager::FindModulePaths(const TCHAR* NamePattern, TMap<FName, FStri
 		return;
 	}
 
-	// Search the cache
-	for (const TPair<FName, FString>& Pair : ModulePathsCache)
+	// Avoid wildcard pattern matching if possible
+	if (FCString::Strchr(NamePattern, '*') == nullptr)
 	{
-		if (Pair.Key.ToString().MatchesWildcard(NamePattern))
+		FName Key(NamePattern, FNAME_Find);
+		if (Key != FName())
 		{
-			OutModulePaths.Add(Pair.Key, Pair.Value);
+			if (const FString* Value = ModulePathsCache.Find(Key))
+			{
+				OutModulePaths.Add(Key, *Value);
+			}
+		}
+	}
+	else
+	{
+		// Search the cache
+		FString KeyTemp;
+		KeyTemp.Reserve(256);
+		for (const TPair<FName, FString>& Pair : ModulePathsCache)
+		{
+			Pair.Key.ToString(KeyTemp);
+			if (KeyTemp.MatchesWildcard(NamePattern))
+			{
+				OutModulePaths.Add(Pair.Key, Pair.Value);
+			}
 		}
 	}
 }
diff --git a/Engine/Source/Runtime/Core/Private/Modules/ModuleManifest.cpp b/Engine/Source/Runtime/Core/Private/Modules/ModuleManifest.cpp
--- a/Engine/Source/Runtime/Core/Private/Modules/ModuleManifest.cpp
+++ b/Engine/Source/Runtime/Core/Private/Modules/ModuleManifest.cpp
@@ -3,6 +3,7 @@
 #include "Modules/ModuleManifest.h"
 #include "Misc/FileHelper.h"
 #include "Misc/App.h"
+#include "Misc/StringBuilder.h"
 #include "Modules/ModuleManager.h"
 #include "Modules/SimpleParse.h"
 
@@ -19,6 +20,15 @@ FString FModuleManifest::GetFileName(const FString& DirectoryName, bool bIsGameF
 #endif
 }
 
+// Assumes GetTypeHash(AltKeyType) matches GetTypeHash(KeyType)
+template<class KeyType, class ValueType, class AltKeyType, class AltValueType>
+ValueType& FindOrAddHeterogeneous(TMap<KeyType, ValueType>& Map, const AltKeyType& Key, const AltValueType& Value) 
+{
+	checkSlow(GetTypeHash(KeyType(Key)) == GetTypeHash(Key));
+	ValueType* Existing = Map.FindByHash(GetTypeHash(Key), Key);
+	return Existing ? *Existing : Map.Emplace(KeyType(Key), AltValueType(Value));
+}
+
 bool FModuleManifest::TryRead(const FString& FileName, FModuleManifest& OutManifest)
 {
 	// Read the file to a string
@@ -36,7 +46,7 @@ bool FModuleManifest::TryRead(const FString& FileName, FModuleManifest& OutManif
 
 	for (;;)
 	{
-		FString Field;
+		TStringBuilder<64> Field;
 		if (!FSimpleParse::ParseString(Ptr, Field))
 		{
 			return false;
@@ -47,14 +57,14 @@ bool FModuleManifest::TryRead(const FString& FileName, FModuleManifest& OutManif
 			return false;
 		}
 
-		if (Field == TEXT("BuildId"))
+		if (Field.ToView() == TEXT("BuildId"_SV))
 		{
 			if (!FSimpleParse::ParseString(Ptr, OutManifest.BuildId))
 			{
 				return false;
 			}
 		}
-		else if (Field == TEXT("Modules"))
+		else if (Field.ToView() == TEXT("Modules"_SV))
 		{
 			if (!FSimpleParse::MatchZeroOrMoreWhitespace(Ptr) || !FSimpleParse::MatchChar(Ptr, TEXT('{')) || !FSimpleParse::MatchZeroOrMoreWhitespace(Ptr))
 			{
@@ -65,14 +75,14 @@ bool FModuleManifest::TryRead(const FString& FileName, FModuleManifest& OutManif
 			{
 				for (;;)
 				{
-					FString ModuleName;
-					FString ModulePath;
+					TStringBuilder<64> ModuleName;
+					TStringBuilder<80> ModulePath;
 					if (!FSimpleParse::ParseString(Ptr, ModuleName) || !FSimpleParse::MatchZeroOrMoreWhitespace(Ptr) || !FSimpleParse::MatchChar(Ptr, TEXT(':')) || !FSimpleParse::MatchZeroOrMoreWhitespace(Ptr) || !FSimpleParse::ParseString(Ptr, ModulePath) || !FSimpleParse::MatchZeroOrMoreWhitespace(Ptr))
 					{
 						return false;
 					}
 
-					OutManifest.ModuleNameToFileName.FindOrAdd(MoveTemp(ModuleName)) = MoveTemp(ModulePath);
+					FindOrAddHeterogeneous(OutManifest.ModuleNameToFileName, ModuleName.ToView(), ModulePath.ToView());
 
 					if (FSimpleParse::MatchChar(Ptr, TEXT('}')))
 					{
