Reduce FWindowsPlatformFile allocations

diff --git a/Engine/Source/Runtime/Core/Private/Windows/WindowsPlatformFile.cpp b/Engine/Source/Runtime/Core/Private/Windows/WindowsPlatformFile.cpp
--- a/Engine/Source/Runtime/Core/Private/Windows/WindowsPlatformFile.cpp
+++ b/Engine/Source/Runtime/Core/Private/Windows/WindowsPlatformFile.cpp
@@ -4,6 +4,8 @@
 #include "CoreTypes.h"
 #include "Misc/DateTime.h"
 #include "Misc/AssertionMacros.h"
+#include "Misc/PathViews.h"
+#include "Misc/StringBuilder.h"
 #include "Logging/LogMacros.h"
 #include "Math/UnrealMathUtility.h"
 #include "HAL/UnrealMemory.h"
@@ -855,65 +857,82 @@ FMappedFileRegionWindows::~FMappedFileRegionWindows()
 **/
 class CORE_API FWindowsPlatformFile : public IPhysicalPlatformFile
 {
-private:
-	FString WindowsNormalizedFilename(const TCHAR* Filename)
-	{
-		const bool bIsFilename = true; // TODO: Create a platform-independent EPathType enum and use that here
-		return WindowsNormalizedPath(Filename, bIsFilename);
-	}
-
-	FString WindowsNormalizedDirname(const TCHAR* Directory)
-	{
-		const bool bIsFilename = false;
-		return WindowsNormalizedPath(Directory, bIsFilename);
-	}
-
 	/**
 	  * Convert from a valid Unreal Path to a canonical and strict-valid Windows Path.
 	  * An Unreal Path may have either \ or / and may have empty directories (two / in a row), and may have .. and may be relative
 	  * A canonical and strict-valid Windows Path has only \, does not have .., does not have empty directories, and is an absolute path, either \\UNC or D:\
 	  * We need to use strict-valid Windows Paths when calling Windows API calls so that we can support the long-path prefix \\?\
 	  */
-	FString WindowsNormalizedPath(const TCHAR* PathString, bool bIsFilename)
+	static void NormalizeWindowsPath(FStringBuilderBase& Path, bool bIsFilename)
 	{
-		FString Result = FPaths::ConvertRelativePathToFull(FString(PathString));
+		FPathViews::ToAbsolutePathInline(Path);
+
 		// NormalizeFilename was already called by ConvertRelativePathToFull, but we still need to do the extra steps in NormalizeDirectoryName if it is a directory
 		if (!bIsFilename)
 		{
-			FPaths::NormalizeDirectoryName(Result);
+			FPathViews::NormalizeDirectoryName(Path);
 		}
 
 		// Remove duplicate slashes
-		const bool bIsUNCPath = Result.StartsWith(TEXT("//"));
+		const bool bIsUNCPath = Path.ToView().StartsWith(TEXT("//"_SV));
+		
+		FPathViews::RemoveDuplicateSlashes(Path);
+
 		if (bIsUNCPath)
 		{
 			// Keep // at the beginning.  If There are more than two / at the beginning, replace them with just //.
-			FPaths::RemoveDuplicateSlashes(Result);
-			Result = TEXT("/") + Result;
-		}
-		else
-		{
-			FPaths::RemoveDuplicateSlashes(Result);
+			Path.Prepend(TEXT("/"_SV));
 		}
 
 		// We now have a canonical, strict-valid, absolute Unreal Path.  Convert it to a Windows Path.
-		Result.ReplaceCharInline(TEXT('/'), TEXT('\\'), ESearchCase::CaseSensitive);
+		for (TCHAR& Char : TArrayView<TCHAR>(Path.GetData(), Path.Len()))
+		{
+			if (Char == '/')
+			{
+				Char = '\\';
+			}
+		}
 
 		// Handle Windows Path length over MAX_PATH
-		if (Result.Len() > MAX_PATH)
+		if (Path.Len() > MAX_PATH)
 		{
 			if (bIsUNCPath)
 			{
-				Result = TEXT("\\\\?\\UNC") + Result.RightChop(1);
+				Path.ReplaceAt(0, 1, TEXT("\\\\?\\UNC"_SV));
 			}
 			else
 			{
-				Result = TEXT("\\\\?\\") + Result;
+				Path.Prepend(TEXT("\\\\?\\"_SV));
 			}
 		}
-
-		return Result;
 	}
+	
+	class FNormalizedFilename : public TStringBuilder<256>
+	{
+	public:
+		explicit FNormalizedFilename(const TCHAR* Filename)
+		{
+			Append(Filename);
+			NormalizeWindowsPath(*this, /* bIsFilename */ true);
+		}
+
+		explicit FNormalizedFilename(FStringView Dir, FStringView Filename)
+		{
+			Append(Dir);
+			FPathViews::Append(*this, Filename);
+			NormalizeWindowsPath(*this, /* bIsFilename */ true);
+		}
+	};
+
+	class FNormalizedDirectory : public TStringBuilder<256>
+	{
+	public:
+		explicit FNormalizedDirectory(const TCHAR* Directory)
+		{
+			Append(Directory);
+			NormalizeWindowsPath(*this, /* bIsFilename */ false);
+		}
+	};
 
 public:
 	//~ For visibility of overloads we don't override
@@ -922,7 +941,7 @@ public:
 
 	virtual bool FileExists(const TCHAR* Filename) override
 	{
-		uint32 Result = GetFileAttributesW(*WindowsNormalizedFilename(Filename));
+		uint32 Result = GetFileAttributesW(*FNormalizedFilename(Filename));
 		if (Result != 0xFFFFFFFF && !(Result & FILE_ATTRIBUTE_DIRECTORY))
 		{
 			return true;
@@ -932,7 +951,7 @@ public:
 	virtual int64 FileSize(const TCHAR* Filename) override
 	{
 		WIN32_FILE_ATTRIBUTE_DATA Info;
-		if (!!GetFileAttributesExW(*WindowsNormalizedFilename(Filename), GetFileExInfoStandard, &Info))
+		if (!!GetFileAttributesExW(*FNormalizedFilename(Filename), GetFileExInfoStandard, &Info))
 		{
 			if ((Info.dwFileAttributes & FILE_ATTRIBUTE_DIRECTORY) == 0)
 			{
@@ -946,12 +965,12 @@ public:
 	}
 	virtual bool DeleteFile(const TCHAR* Filename) override
 	{
-		const FString NormalizedFilename = WindowsNormalizedFilename(Filename);
+		FNormalizedFilename NormalizedFilename(Filename);
 		return !!DeleteFileW(*NormalizedFilename);
 	}
 	virtual bool IsReadOnly(const TCHAR* Filename) override
 	{
-		uint32 Result = GetFileAttributesW(*WindowsNormalizedFilename(Filename));
+		uint32 Result = GetFileAttributesW(*FNormalizedFilename(Filename));
 		if (Result != 0xFFFFFFFF)
 		{
 			return !!(Result & FILE_ATTRIBUTE_READONLY);
@@ -960,17 +979,17 @@ public:
 	}
 	virtual bool MoveFile(const TCHAR* To, const TCHAR* From) override
 	{
-		return !!MoveFileW(*WindowsNormalizedFilename(From), *WindowsNormalizedFilename(To));
+		return !!MoveFileW(*FNormalizedFilename(From), *FNormalizedFilename(To));
 	}
 	virtual bool SetReadOnly(const TCHAR* Filename, bool bNewReadOnlyValue) override
 	{
-		return !!SetFileAttributesW(*WindowsNormalizedFilename(Filename), bNewReadOnlyValue ? FILE_ATTRIBUTE_READONLY : FILE_ATTRIBUTE_NORMAL);
+		return !!SetFileAttributesW(*FNormalizedFilename(Filename), bNewReadOnlyValue ? FILE_ATTRIBUTE_READONLY : FILE_ATTRIBUTE_NORMAL);
 	}
 
 	virtual FDateTime GetTimeStamp(const TCHAR* Filename) override
 	{
 		WIN32_FILE_ATTRIBUTE_DATA Info;
-		if (GetFileAttributesExW(*WindowsNormalizedFilename(Filename), GetFileExInfoStandard, &Info))
+		if (GetFileAttributesExW(*FNormalizedFilename(Filename), GetFileExInfoStandard, &Info))
 		{
 			return WindowsFileTimeToUEDateTime(Info.ftLastWriteTime);
 		}
@@ -981,7 +1000,7 @@ public:
 	virtual void SetTimeStamp(const TCHAR* Filename, FDateTime DateTime) override
 	{
 		TRACE_PLATFORMFILE_BEGIN_OPEN(Filename);
-		HANDLE Handle = CreateFileW(*WindowsNormalizedFilename(Filename), FILE_WRITE_ATTRIBUTES, FILE_SHARE_READ, nullptr, OPEN_EXISTING, FILE_ATTRIBUTE_NORMAL | FILE_FLAG_OVERLAPPED, nullptr);
+		HANDLE Handle = CreateFileW(*FNormalizedFilename(Filename), FILE_WRITE_ATTRIBUTES, FILE_SHARE_READ, nullptr, OPEN_EXISTING, FILE_ATTRIBUTE_NORMAL | FILE_FLAG_OVERLAPPED, nullptr);
 		if (Handle != INVALID_HANDLE_VALUE)
 		{
 			TRACE_PLATFORMFILE_END_OPEN(Handle);
@@ -1019,7 +1038,7 @@ public:
 	virtual FDateTime GetAccessTimeStamp(const TCHAR* Filename) override
 	{
 		WIN32_FILE_ATTRIBUTE_DATA Info;
-		if (GetFileAttributesExW(*WindowsNormalizedFilename(Filename), GetFileExInfoStandard, &Info))
+		if (GetFileAttributesExW(*FNormalizedFilename(Filename), GetFileExInfoStandard, &Info))
 		{
 			return WindowsFileTimeToUEDateTime(Info.ftLastAccessTime);
 		}
@@ -1029,7 +1048,7 @@ public:
 
 	virtual FString GetFilenameOnDisk(const TCHAR* Filename) override
 	{
-		FString NormalizedFileName = WindowsNormalizedFilename(Filename);
+		FString NormalizedFileName = *FNormalizedFilename(Filename);
 		TRACE_PLATFORMFILE_BEGIN_OPEN(Filename);
 		HANDLE hFile = CreateFile(*NormalizedFileName, FILE_READ_ATTRIBUTES, FILE_SHARE_READ | FILE_SHARE_WRITE | FILE_SHARE_DELETE, NULL, OPEN_EXISTING, 0, NULL);
 		// If the file exists on disk, read the capitalization from the path on disk, otherwise just return the (normalized) input filename
@@ -1044,7 +1063,7 @@ public:
 				Length = GetFinalPathNameByHandle(hFile, CharArray.GetData(), CharArray.Num(), FILE_NAME_NORMALIZED);
 				if (Length == 0)
 				{
-					NormalizedFileName = WindowsNormalizedFilename(Filename);
+					NormalizedFileName = *FNormalizedFilename(Filename);
 					break;
 				}
 				if (Length < (uint32)CharArray.Num())
@@ -1131,14 +1150,14 @@ public:
 
 		TRACE_PLATFORMFILE_BEGIN_OPEN(Filename);
 #if USE_OVERLAPPED_IO
-		HANDLE Handle    = CreateFileW(*WindowsNormalizedFilename(Filename), Access, WinFlags, NULL, Create, FILE_ATTRIBUTE_NORMAL | FILE_FLAG_OVERLAPPED, NULL);
+		HANDLE Handle    = CreateFileW(*FNormalizedFilename(Filename), Access, WinFlags, NULL, Create, FILE_ATTRIBUTE_NORMAL | FILE_FLAG_OVERLAPPED, NULL);
 		if (Handle != INVALID_HANDLE_VALUE)
 		{
 			TRACE_PLATFORMFILE_END_OPEN(Handle);
 			return new FAsyncBufferedFileReaderWindows(Handle, Access, WinFlags, FILE_FLAG_OVERLAPPED);
 		}
 #else
-		HANDLE Handle = CreateFileW(*WindowsNormalizedFilename(Filename), Access, WinFlags, NULL, Create, FILE_ATTRIBUTE_NORMAL, NULL);
+		HANDLE Handle = CreateFileW(*FNormalizedFilename(Filename), Access, WinFlags, NULL, Create, FILE_ATTRIBUTE_NORMAL, NULL);
 		if (Handle != INVALID_HANDLE_VALUE)
 		{
 			TRACE_PLATFORMFILE_END_OPEN(Handle);
@@ -1158,7 +1177,7 @@ public:
 		uint32  WinFlags = FILE_SHARE_READ | (bAllowWrite ? FILE_SHARE_WRITE : 0);
 		uint32  Create = OPEN_EXISTING;
 		TRACE_PLATFORMFILE_BEGIN_OPEN(Filename);
-		HANDLE Handle = CreateFileW(*WindowsNormalizedFilename(Filename), Access, WinFlags, NULL, Create, FILE_ATTRIBUTE_NORMAL | FILE_FLAG_OVERLAPPED, NULL);
+		HANDLE Handle = CreateFileW(*FNormalizedFilename(Filename), Access, WinFlags, NULL, Create, FILE_ATTRIBUTE_NORMAL | FILE_FLAG_OVERLAPPED, NULL);
 		if (Handle != INVALID_HANDLE_VALUE)
 		{
 			TRACE_PLATFORMFILE_END_OPEN(Handle);
@@ -1177,7 +1196,7 @@ public:
 		uint32  WinFlags  = bAllowRead ? FILE_SHARE_READ : 0;
 		uint32  Create    = bAppend ? OPEN_ALWAYS : CREATE_ALWAYS;
 		TRACE_PLATFORMFILE_BEGIN_OPEN(Filename);
-		HANDLE Handle    = CreateFileW(*WindowsNormalizedFilename(Filename), Access, WinFlags, NULL, Create, FILE_ATTRIBUTE_NORMAL, NULL);
+		HANDLE Handle    = CreateFileW(*FNormalizedFilename(Filename), Access, WinFlags, NULL, Create, FILE_ATTRIBUTE_NORMAL, NULL);
 		if(Handle != INVALID_HANDLE_VALUE)
 		{
 			TRACE_PLATFORMFILE_END_OPEN(Handle);
@@ -1207,7 +1226,7 @@ public:
 		uint32  WinFlags = FILE_SHARE_READ;
 		uint32  Create = OPEN_EXISTING;
 		TRACE_PLATFORMFILE_BEGIN_OPEN(Filename);
-		HANDLE Handle = CreateFileW(*WindowsNormalizedFilename(Filename), Access, WinFlags, NULL, Create, FILE_ATTRIBUTE_NORMAL, NULL);
+		HANDLE Handle = CreateFileW(*FNormalizedFilename(Filename), Access, WinFlags, NULL, Create, FILE_ATTRIBUTE_NORMAL, NULL);
 		if (Handle != INVALID_HANDLE_VALUE)
 		{
 			TRACE_PLATFORMFILE_END_OPEN(Handle);
@@ -1248,18 +1267,18 @@ public:
 		bool bExists = !FCString::Strlen(Directory);
 		if (!bExists) 
 		{
-			uint32 Result = GetFileAttributesW(*WindowsNormalizedDirname(Directory));
+			uint32 Result = GetFileAttributesW(*FNormalizedDirectory(Directory));
 			bExists = (Result != 0xFFFFFFFF && (Result & FILE_ATTRIBUTE_DIRECTORY));
 		}
 		return bExists;
 	}
 	virtual bool CreateDirectory(const TCHAR* Directory) override
 	{
-		return CreateDirectoryW(*WindowsNormalizedDirname(Directory), NULL) || GetLastError() == ERROR_ALREADY_EXISTS;
+		return CreateDirectoryW(*FNormalizedDirectory(Directory), NULL) || GetLastError() == ERROR_ALREADY_EXISTS;
 	}
 	virtual bool DeleteDirectory(const TCHAR* Directory) override
 	{
-		RemoveDirectoryW(*WindowsNormalizedDirname(Directory));
+		RemoveDirectoryW(*FNormalizedDirectory(Directory));
 		uint32 LastError = GetLastError();
 		const bool bSucceeded = !DirectoryExists(Directory);
 		if (!bSucceeded)
@@ -1271,7 +1290,7 @@ public:
 	virtual FFileStatData GetStatData(const TCHAR* FilenameOrDirectory) override
 	{
 		WIN32_FILE_ATTRIBUTE_DATA Info;
-		if (GetFileAttributesExW(*WindowsNormalizedFilename(FilenameOrDirectory), GetFileExInfoStandard, &Info))
+		if (GetFileAttributesExW(*FNormalizedFilename(FilenameOrDirectory), GetFileExInfoStandard, &Info))
 		{
 			const bool bIsDirectory = !!(Info.dwFileAttributes & FILE_ATTRIBUTE_DIRECTORY);
 
@@ -1298,11 +1317,15 @@ public:
 	}
 	virtual bool IterateDirectory(const TCHAR* Directory, FDirectoryVisitor& Visitor) override
 	{
-		const FString DirectoryStr = Directory;
+		TStringBuilder<256> DirTemp;
+		DirTemp << Directory;
+		const int32 DirLen = DirTemp.Len();
 		return IterateDirectoryCommon(Directory, [&](const WIN32_FIND_DATAW& InData) -> bool
 		{
+			DirTemp.RemoveSuffix(DirTemp.Len() - DirLen);
+			FPathViews::Append(DirTemp, InData.cFileName);
 			const bool bIsDirectory = !!(InData.dwFileAttributes & FILE_ATTRIBUTE_DIRECTORY);
-			return Visitor.Visit(*(DirectoryStr / InData.cFileName), bIsDirectory);
+			return Visitor.Visit(*DirTemp, bIsDirectory);
 		});
 	}
 	virtual bool IterateDirectoryStat(const TCHAR* Directory, FDirectoryStatVisitor& Visitor) override
@@ -1334,12 +1357,18 @@ public:
 				);
 		});
 	}
+	
+	// Outline to reduce stack space usage since IterateDirectoryCommon might be recursive
+	FORCENOINLINE static HANDLE FindFirstFileWithWildcard(const TCHAR* Directory, WIN32_FIND_DATAW& OutData)
+	{
+		return FindFirstFileW(*(FNormalizedFilename(Directory, TEXT("*.*"_SV))), &OutData);
+	}
+
 	bool IterateDirectoryCommon(const TCHAR* Directory, const TFunctionRef<bool(const WIN32_FIND_DATAW&)>& Visitor)
 	{
 		bool bResult = true;
 		WIN32_FIND_DATAW Data;
-		FString SearchWildcard = FString(Directory) / TEXT("*.*");
-		HANDLE Handle = FindFirstFileW(*(WindowsNormalizedFilename(*SearchWildcard)), &Data);
+		HANDLE Handle = FindFirstFileWithWildcard(Directory, /* Out */ Data);
 		if (Handle != INVALID_HANDLE_VALUE)
 		{
 			do
