Reduce INI allocations during early editor startup

diff --git a/Engine/Source/Runtime/Core/Private/Misc/ConfigCacheIni.cpp b/Engine/Source/Runtime/Core/Private/Misc/ConfigCacheIni.cpp
--- a/Engine/Source/Runtime/Core/Private/Misc/ConfigCacheIni.cpp
+++ b/Engine/Source/Runtime/Core/Private/Misc/ConfigCacheIni.cpp
@@ -733,13 +733,22 @@ bool FConfigFile::Combine(const FString& Filename)
 	return false;
 }
 
+// Assumes GetTypeHash(AltKeyType) matches GetTypeHash(KeyType)
+template<class KeyType, class ValueType, class AltKeyType>
+ValueType& FindOrAddHeterogeneous(TMap<KeyType, ValueType>& Map, const AltKeyType& Key) 
+{
+	checkSlow(GetTypeHash(KeyType(Key)) == GetTypeHash(Key));
+	ValueType* Existing = Map.FindByHash(GetTypeHash(Key), Key);
+	return Existing ? *Existing : Map.Emplace(KeyType(Key));
+}
 
 void FConfigFile::CombineFromBuffer(const FString& Buffer)
 {
 	const TCHAR* Ptr = *Buffer;
 	FConfigSection* CurrentSection = nullptr;
-	FString CurrentSectionName;
-	FString TheLine;
+	TStringBuilder<64> CurrentSectionName;
+	TStringBuilder<128> TheLine;
+	FString ProcessedValue;
 	bool Done = false;
 	while( !Done )
 	{
@@ -825,8 +834,6 @@ void FConfigFile::CombineFromBuffer(const FString& Buffer)
 					Start[FCString::Strlen(Start)-1] = 0;
 				}
 
-				FString ProcessedValue;
-
 				// Strip leading whitespace from the property value
 				while ( *Value && FChar::IsWhitespace(*Value) )
 				{
@@ -839,6 +846,8 @@ void FConfigFile::CombineFromBuffer(const FString& Buffer)
 					Value[FCString::Strlen(Value)-1] = 0;
 				}
 
+				ProcessedValue.Reset();
+
 				// If this line is delimited by quotes
 				if( *Value=='\"' )
 				{
@@ -876,7 +885,7 @@ void FConfigFile::CombineFromBuffer(const FString& Buffer)
 				else if (Cmd == '*')
 				{
 					// track a key to show uniqueness for arrays of structs
-					TMap<FName, FString>& POCKeys = PerObjectConfigArrayOfStructKeys.FindOrAdd(CurrentSectionName);
+					TMap<FName, FString>& POCKeys = FindOrAddHeterogeneous(PerObjectConfigArrayOfStructKeys, CurrentSectionName.ToView());
 					POCKeys.Add(Start, MoveTemp(ProcessedValue));
 				}
 				else
@@ -920,6 +929,7 @@ void FConfigFile::ProcessInputFileContents(const FString& Contents)
 {
 	const TCHAR* Ptr = Contents.Len() > 0 ? *Contents : nullptr;
 	FConfigSection* CurrentSection = nullptr;
+	TStringBuilder<128> TheLine;
 	bool Done = false;
 	while( !Done && Ptr != nullptr )
 	{
@@ -929,7 +939,6 @@ void FConfigFile::ProcessInputFileContents(const FString& Contents)
 			Ptr++;
 		}			
 		// read the next line
-		FString TheLine;
 		int32 LinesConsumed = 0;
 		FParse::LineExtended(&Ptr, TheLine, LinesConsumed, false);
 		if (Ptr == nullptr || *Ptr == 0)
@@ -1759,6 +1768,7 @@ void FConfigFile::AddMissingProperties( const FConfigFile& InSourceFile )
 		{
 			// If we don't already have this section, go ahead and add it now
 			FConfigSection* DestSection = FindOrAddSection( SourceSectionName );
+			DestSection->Reserve(SourceSection.Num());
 
 			for( FConfigSection::TConstIterator SourcePropertyIt( SourceSection ); SourcePropertyIt; ++SourcePropertyIt )
 			{
@@ -1767,12 +1777,11 @@ void FConfigFile::AddMissingProperties( const FConfigFile& InSourceFile )
 				// If we don't already have this property, go ahead and add it now
 				if( DestSection->Find( SourcePropertyName ) == nullptr )
 				{
-					TArray<FConfigValue> Results;
-					SourceSection.MultiFind(SourcePropertyName, Results, true);
-					for (const FConfigValue& Result : Results)
+					TArray<const FConfigValue*, TInlineAllocator<32>> Results;
+					SourceSection.MultiFindPointer(SourcePropertyName, Results, true);
+					for (const FConfigValue* Result : Results)
 					{
-
-						DestSection->Add(SourcePropertyName, Result.GetSavedValue());
+						DestSection->Add(SourcePropertyName, *Result);
 						Dirty = true;
 					}
 				}
