PR #7837: Allow us to bypass a physical log file using the -NODEFAULTLOG arg (Contributed by [REDACTED_PERSON])

diff --git a/Engine/Source/Runtime/Core/Private/GenericPlatform/GenericPlatformOutputDevices.cpp b/Engine/Source/Runtime/Core/Private/GenericPlatform/GenericPlatformOutputDevices.cpp
--- a/Engine/Source/Runtime/Core/Private/GenericPlatform/GenericPlatformOutputDevices.cpp
+++ b/Engine/Source/Runtime/Core/Private/GenericPlatform/GenericPlatformOutputDevices.cpp
@@ -23,7 +23,12 @@ void FGenericPlatformOutputDevices::SetupOutputDevices()
 	check(GLog);
 
 	ResetCachedAbsoluteFilename();
-	GLog->AddOutputDevice(FPlatformOutputDevices::GetLog());
+	
+	// Add the default log device (typically file) unless the commandline says otherwise.
+	if (!FParse::Param(FCommandLine::Get(), TEXT("NODEFAULTLOG")))
+	{
+		GLog->AddOutputDevice(FPlatformOutputDevices::GetLog());
+	}
 
 	TArray<FOutputDevice*> ChannelFileOverrides;
 	FPlatformOutputDevices::GetPerChannelFileOverrides(ChannelFileOverrides);
