Enable GPU-Scene instance culling for all platforms.

diff --git a/Engine/Shaders/Public/ShaderVersion.ush b/Engine/Shaders/Public/ShaderVersion.ush
--- a/Engine/Shaders/Public/ShaderVersion.ush
+++ b/Engine/Shaders/Public/ShaderVersion.ush
@@ -2,4 +2,4 @@
 // Each time the console command is executed it generates a new GUID. As this file is included
 // in Platform.ush (which should be included in any shader) it allows to invalidate the shader DDC.
 // 
-// GUID = 196E4485-47A1-477A-8434-7000E7CA65CC
\ No newline at end of file
+// GUID = DDCF4C7AB9BB4D43BB2F9D07CD995890
\ No newline at end of file
diff --git a/Engine/Source/Runtime/Core/Core.Build.cs b/Engine/Source/Runtime/Core/Core.Build.cs
--- a/Engine/Source/Runtime/Core/Core.Build.cs
+++ b/Engine/Source/Runtime/Core/Core.Build.cs
@@ -292,18 +292,7 @@ public class Core : ModuleRules
 		PrivateDefinitions.Add("PLATFORM_SUPPORTS_BINARYCONFIG=" + (SupportsBinaryConfig(Target) ? "1" : "0"));
 
 		// temporary thing to enable backing out in case of disaster, remove after initial testing period.
-		if (   Target.Platform == UnrealTargetPlatform.Mac
-			|| Target.Platform == UnrealTargetPlatform.IOS
-			|| Target.Platform == UnrealTargetPlatform.TVOS
-			|| Target.Platform == UnrealTargetPlatform.Linux
-			|| Target.Platform == UnrealTargetPlatform.LinuxAArch64)
-		{
-			PublicDefinitions.Add("GPUCULL_TODO=0");
-		}
-		else
-		{
-			PublicDefinitions.Add("GPUCULL_TODO=1");
-		}
+		PublicDefinitions.Add("GPUCULL_TODO=1");
 
 		PublicDefinitions.Add("WITH_MALLOC_STOMP=" + (bWithMallocStomp ? "1" : "0"));
 
