Remove "fake" warning causing non-deterministic build health issue.

diff --git a/Engine/Source/Runtime/Core/Public/Windows/COMPointer.h b/Engine/Source/Runtime/Core/Public/Windows/COMPointer.h
--- a/Engine/Source/Runtime/Core/Public/Windows/COMPointer.h
+++ b/Engine/Source/Runtime/Core/Public/Windows/COMPointer.h
@@ -2,6 +2,7 @@
 
 #pragma once
 
-#pragma message(__FILE__"(5): warning : #include Microsoft/COMPointer.h instead of Windows/COMPointer.h")
+// Producing a warning that is not treated as an error here calls non-deterministic incremental build warnings
+//#pragma message(__FILE__"(5): warning : #include Microsoft/COMPointer.h instead of Windows/COMPointer.h")
 
 #include "Microsoft/COMPointer.h"
