Adding VectorLoadTwoPairsFloat and VectorDeinterleave #rb charles.bloom  Zak.Middleton

diff --git a/Engine/Source/Runtime/Core/Private/Tests/Math/UnrealMathTest.cpp b/Engine/Source/Runtime/Core/Private/Tests/Math/UnrealMathTest.cpp
--- a/Engine/Source/Runtime/Core/Private/Tests/Math/UnrealMathTest.cpp
+++ b/Engine/Source/Runtime/Core/Private/Tests/Math/UnrealMathTest.cpp
@@ -1113,6 +1113,18 @@ bool FVectorRegisterAbstractionTest::RunTest(const FString& Parameters)
 	V3 = VectorMultiply(VectorMultiply(V1, V1), V0);
 	LogTest( TEXT("VectorReciprocalSqrtAccurate"), TestVectorsEqual( VectorOne(), V3, 1e-6f ) );
 
+	SetScratch(1.0f, 2.0f, 3.0f, 4.0f);
+	V0 = VectorLoadTwoPairsFloat(GScratch + 0, GScratch + 1);
+	V1 = MakeVectorRegister(1.0f, 2.0f, 2.0f, 3.0f);
+	LogTest(TEXT("VectorLoadTwoPairsFloat"), TestVectorsEqual(V0, V1));
+
+	V0 = MakeVectorRegister(0.0f, 1.0f, 2.0f, 3.0f);
+	V1 = MakeVectorRegister(4.0f, 5.0f, 6.0f, 7.0f);
+	VectorDeinterleave(V2, V3, V0, V1);
+	V0 = MakeVectorRegister(0.0f, 2.0f, 4.0f, 6.0f);
+	V1 = MakeVectorRegister(1.0f, 3.0f, 5.0f, 7.0f);
+	LogTest(TEXT("VectorDeinterleave"), TestVectorsEqual(V2, V0) && TestVectorsEqual(V3, V1));
+
 	// VectorMod
 	V0 = MakeVectorRegister(0.0f, 3.2f, 2.8f,  1.5f);
 	V1 = MakeVectorRegister(2.0f, 1.2f, 2.0f,  3.0f);
diff --git a/Engine/Source/Runtime/Core/Public/Math/UnrealMathDirectX.h b/Engine/Source/Runtime/Core/Public/Math/UnrealMathDirectX.h
--- a/Engine/Source/Runtime/Core/Public/Math/UnrealMathDirectX.h
+++ b/Engine/Source/Runtime/Core/Public/Math/UnrealMathDirectX.h
@@ -154,6 +154,22 @@ FORCEINLINE VectorRegisterInt MakeVectorRegisterInt(int32 X, int32 Y, int32 Z, i
  */
 #define VectorLoadFloat2( Ptr )			MakeVectorRegister( ((const float*)(Ptr))[0], ((const float*)(Ptr))[1], ((const float*)(Ptr))[0], ((const float*)(Ptr))[1] )
 
+/**
+ * Loads 4 unaligned floats - 2 from the first pointer, 2 from the second, and packs
+ * them in to 1 vector.
+ *
+ * @param Ptr1	Unaligned memory pointer to the first 2 floats
+ * @param Ptr2	Unaligned memory pointer to the second 2 floats
+ * @return		VectorRegister(Ptr1[0], Ptr1[1], Ptr2[0], Ptr2[1])
+ */
+FORCEINLINE VectorRegister VectorLoadTwoPairsFloat(const float* Ptr1, const float* Ptr2)
+{
+	__m128 Ret = _mm_castpd_ps(_mm_load_sd((double const*)(Ptr1)));
+	Ret = _mm_loadh_pi(Ret, (__m64 const*)(Ptr2));
+	return Ret;
+}
+
+
 /**
  * Creates a vector out of three FLOATs and leaves W undefined.
  *
@@ -636,6 +652,22 @@ FORCEINLINE VectorRegister VectorCombineLow(const VectorRegister& Vec1, const Ve
 	return VectorShuffle(Vec1, Vec2, 0, 1, 0, 1);
 }
 
+/**
+ * Deinterleaves the components of the two given vectors such that the even components
+ * are in one vector and the odds in another.
+ *
+ * @param Lo	[Even0, Odd0, Even1, Odd1]
+ * @param Hi	[Even2, Odd2, Even3, Odd3]
+ * @param OutEvens [Even0, Even1, Even2, Even3]
+ * @param OutOdds [Odd0, Odd1, Odd2, Odd3]
+*/
+FORCEINLINE void VectorDeinterleave(VectorRegister& OutEvens, VectorRegister& OutOdds, const VectorRegister& Lo, const VectorRegister& Hi)
+{
+	OutEvens = _mm_shuffle_ps(Lo, Hi, _MM_SHUFFLE(2, 0, 2, 0));
+	OutOdds = _mm_shuffle_ps(Lo, Hi, _MM_SHUFFLE(3, 1, 3, 1));
+}
+
+
 /**
  * These functions return a vector mask to indicate which components pass the comparison.
  * Each component is 0xffffffff if it passes, 0x00000000 if it fails.
diff --git a/Engine/Source/Runtime/Core/Public/Math/UnrealMathFPU.h b/Engine/Source/Runtime/Core/Public/Math/UnrealMathFPU.h
--- a/Engine/Source/Runtime/Core/Public/Math/UnrealMathFPU.h
+++ b/Engine/Source/Runtime/Core/Public/Math/UnrealMathFPU.h
@@ -169,6 +169,19 @@ FORCEINLINE VectorRegisterInt MakeVectorRegisterInt(int32 X, int32 Y, int32 Z, i
  */
 #define VectorLoadFloat2( Ptr )			MakeVectorRegister( ((const float*)(Ptr))[0], ((const float*)(Ptr))[1], ((const float*)(Ptr))[0], ((const float*)(Ptr))[1] )
 
+/**
+ * Loads 4 unaligned floats - 2 from the first pointer, 2 from the second, and packs
+ * them in to 1 vector.
+ *
+ * @param Ptr1	Unaligned memory pointer to the first 2 floats
+ * @param Ptr2	Unaligned memory pointer to the second 2 floats
+ * @return		VectorRegister(Ptr1[0], Ptr1[1], Ptr2[0], Ptr2[1])
+ */
+FORCEINLINE VectorRegister VectorLoadTwoPairsFloat(const float* Ptr1, const float* Ptr2)
+{
+	return MakeVectorRegister(Ptr1[0], Ptr1[1], Ptr2[0], Ptr2[1]);
+}
+
 /**
  * Creates a vector out of three FLOATs and leaves W undefined.
  *
@@ -1007,6 +1020,21 @@ FORCEINLINE VectorRegister VectorCombineLow(const VectorRegister& Vec1, const Ve
 	return VectorShuffle(Vec1, Vec2, 0, 1, 0, 1);
 }
 
+/**
+ * Deinterleaves the components of the two given vectors such that the even components
+ * are in one vector and the odds in another.
+ *
+ * @param Lo	[Even0, Odd0, Even1, Odd1]
+ * @param Hi	[Even2, Odd2, Even3, Odd3]
+ * @param OutEvens [Even0, Even1, Even2, Even3]
+ * @param OutOdds [Odd0, Odd1, Odd2, Odd3]
+*/
+FORCEINLINE void VectorDeinterleave(VectorRegister& OutEvens, VectorRegister& OutOdds, const VectorRegister& Lo, const VectorRegister& Hi)
+{
+	OutEvens = MakeVectorRegister(Lo.V[0], Lo.V[2], Hi.V[0], Hi.V[2]);
+	OutOdds = MakeVectorRegister(Lo.V[1], Lo.V[3], Hi.V[1], Hi.V[3]);
+}
+
 /**
  * Merges the XYZ components of one vector with the W component of another vector and returns the result.
  *
diff --git a/Engine/Source/Runtime/Core/Public/Math/UnrealMathNeon.h b/Engine/Source/Runtime/Core/Public/Math/UnrealMathNeon.h
--- a/Engine/Source/Runtime/Core/Public/Math/UnrealMathNeon.h
+++ b/Engine/Source/Runtime/Core/Public/Math/UnrealMathNeon.h
@@ -230,6 +230,22 @@ FORCEINLINE VectorRegister VectorLoadFloat1( const void *Ptr )
 {
 	return vdupq_n_f32( ((float32_t *)Ptr)[0] );
 }
+
+/**
+ * Loads 4 unaligned floats - 2 from the first pointer, 2 from the second, and packs
+ * them in to 1 vector.
+ *
+ * @param Ptr1	Unaligned memory pointer to the first 2 floats
+ * @param Ptr2	Unaligned memory pointer to the second 2 floats
+ * @return		VectorRegister(Ptr1[0], Ptr1[1], Ptr2[0], Ptr2[1])
+ */
+FORCEINLINE VectorRegister VectorLoadTwoPairsFloat(const float* Ptr1, const float* Ptr2)
+{
+	float32x2_t Lo = vld1_f32(Ptr1);
+	float32x2_t Hi = vld1_f32(Ptr2);
+	return vcombine_f32(Lo, Hi);
+}
+
 /**
  * Creates a vector out of three floats and leaves W undefined.
  *
@@ -737,6 +753,22 @@ FORCEINLINE VectorRegister VectorCombineLow(const VectorRegister& Vec1, const Ve
 	return vcombine_f32(vget_low_f32(Vec1), vget_low_f32(Vec2));
 }
 
+/**
+ * Deinterleaves the components of the two given vectors such that the even components
+ * are in one vector and the odds in another.
+ *
+ * @param Lo	[Even0, Odd0, Even1, Odd1]
+ * @param Hi	[Even2, Odd2, Even3, Odd3]
+ * @param OutEvens [Even0, Even1, Even2, Even3]
+ * @param OutOdds [Odd0, Odd1, Odd2, Odd3]
+*/
+FORCEINLINE void VectorDeinterleave(VectorRegister& OutEvens, VectorRegister& OutOdds, const VectorRegister& Lo, const VectorRegister& Hi)
+{
+	float32x4x2_t deinterleaved = vuzpq_f32(Lo, Hi);
+	OutEvens = deinterleaved.val[0];
+	OutOdds = deinterleaved.val[1];
+}
+
 /**
  * Calculates the cross product of two vectors (XYZ components). W is set to 0.
  *
diff --git a/Engine/Source/Runtime/Core/Public/Math/UnrealMathSSE.h b/Engine/Source/Runtime/Core/Public/Math/UnrealMathSSE.h
--- a/Engine/Source/Runtime/Core/Public/Math/UnrealMathSSE.h
+++ b/Engine/Source/Runtime/Core/Public/Math/UnrealMathSSE.h
@@ -190,6 +190,21 @@ FORCEINLINE VectorRegister VectorLoad( const void* Ptr )
  */
 #define VectorLoadFloat2( Ptr )			_mm_castpd_ps(_mm_load1_pd((const double*)(Ptr)))
 
+/** 
+ * Loads 4 unaligned floats - 2 from the first pointer, 2 from the second, and packs
+ * them in to 1 vector.
+ *
+ * @param Ptr1	Unaligned memory pointer to the first 2 floats
+ * @param Ptr2	Unaligned memory pointer to the second 2 floats
+ * @return		VectorRegister(Ptr1[0], Ptr1[1], Ptr2[0], Ptr2[1])
+ */
+FORCEINLINE VectorRegister VectorLoadTwoPairsFloat(const float* Ptr1, const float* Ptr2)
+{
+	__m128 Ret = _mm_castpd_ps(_mm_load_sd((double const *)(Ptr1)));
+	Ret = _mm_loadh_pi(Ret, (__m64 const*)(Ptr2));
+	return Ret;
+}
+
 /**
  * Creates a vector out of three FLOATs and leaves W undefined.
  *
@@ -905,6 +920,21 @@ FORCEINLINE VectorRegister VectorCombineLow(const VectorRegister& Vec1, const Ve
 	return VectorShuffle(Vec1, Vec2, 0, 1, 0, 1);
 }
 
+/**
+ * Deinterleaves the components of the two given vectors such that the even components
+ * are in one vector and the odds in another.
+ *
+ * @param Lo	[Even0, Odd0, Even1, Odd1]
+ * @param Hi	[Even2, Odd2, Even3, Odd3]
+ * @param OutEvens [Even0, Even1, Even2, Even3]
+ * @param OutOdds [Odd0, Odd1, Odd2, Odd3]
+*/
+FORCEINLINE void VectorDeinterleave(VectorRegister& OutEvens, VectorRegister& OutOdds, const VectorRegister& Lo, const VectorRegister& Hi)
+{
+	OutEvens = _mm_shuffle_ps(Lo, Hi, SHUFFLEMASK(0, 2, 0, 2));
+	OutOdds = _mm_shuffle_ps(Lo, Hi, SHUFFLEMASK(1, 3, 1, 3));
+}
+
 /**
  * These functions return a vector mask to indicate which components pass the comparison.
  * Each component is 0xffffffff if it passes, 0x00000000 if it fails.
