CompactBinary: Improved usability and performance of field iterators

#rnx

diff --git a/Engine/Source/Runtime/Core/Private/Serialization/CompactBinary.cpp b/Engine/Source/Runtime/Core/Private/Serialization/CompactBinary.cpp
--- a/Engine/Source/Runtime/Core/Private/Serialization/CompactBinary.cpp
+++ b/Engine/Source/Runtime/Core/Private/Serialization/CompactBinary.cpp
@@ -620,6 +620,28 @@ FMemoryView FCbFieldView::GetViewNoType() const
 	return MakeMemoryView(static_cast<const uint8*>(Payload) - NameSize, NameSize + PayloadSize);
 }
 
+FCbFieldViewIterator FCbFieldView::CreateViewIterator() const
+{
+	const ECbFieldType LocalType = Type;
+	if (FCbFieldType::HasFields(LocalType))
+	{
+		const uint8* PayloadBytes = static_cast<const uint8*>(GetPayload());
+		uint32 PayloadSizeByteCount;
+		const uint64 PayloadSize = ReadVarUInt(PayloadBytes, PayloadSizeByteCount);
+		PayloadBytes += PayloadSizeByteCount;
+		const uint64 NumByteCount = FCbFieldType::IsArray(LocalType) ? MeasureVarUInt(PayloadBytes) : 0;
+		if (PayloadSize > NumByteCount)
+		{
+			const void* const PayloadEnd = PayloadBytes + PayloadSize;
+			PayloadBytes += NumByteCount;
+			const ECbFieldType UniformType = FCbFieldType::HasUniformFields(LocalType) ?
+				ECbFieldType(*PayloadBytes++) : ECbFieldType::HasFieldType;
+			return FCbFieldViewIterator::MakeRange(MakeMemoryView(PayloadBytes, PayloadEnd), UniformType);
+		}
+	}
+	return FCbFieldViewIterator();
+}
+
 ///////////////////////////////////////////////////////////////////////////////////////////////////////////////////////
 
 FCbArrayView::FCbArrayView()
@@ -635,24 +657,6 @@ uint64 FCbArrayView::Num() const
 	return ReadVarUInt(PayloadBytes, NumByteCount);
 }
 
-FCbFieldViewIterator FCbArrayView::CreateViewIterator() const
-{
-	const uint8* PayloadBytes = static_cast<const uint8*>(GetPayload());
-	uint32 PayloadSizeByteCount;
-	const uint64 PayloadSize = ReadVarUInt(PayloadBytes, PayloadSizeByteCount);
-	PayloadBytes += PayloadSizeByteCount;
-	const uint64 NumByteCount = MeasureVarUInt(PayloadBytes);
-	if (PayloadSize > NumByteCount)
-	{
-		const void* const PayloadEnd = PayloadBytes + PayloadSize;
-		PayloadBytes += NumByteCount;
-		const ECbFieldType UniformType = FCbFieldType::GetType(GetType()) == ECbFieldType::UniformArray ?
-			ECbFieldType(*PayloadBytes++) : ECbFieldType::HasFieldType;
-		return FCbFieldViewIterator::MakeRange(MakeMemoryView(PayloadBytes, PayloadEnd), UniformType);
-	}
-	return FCbFieldViewIterator();
-}
-
 uint64 FCbArrayView::GetSize() const
 {
 	return sizeof(ECbFieldType) + GetPayloadSize();
@@ -704,22 +708,6 @@ FCbObjectView::FCbObjectView()
 {
 }
 
-FCbFieldViewIterator FCbObjectView::CreateViewIterator() const
-{
-	const uint8* PayloadBytes = static_cast<const uint8*>(GetPayload());
-	uint32 PayloadSizeByteCount;
-	const uint64 PayloadSize = ReadVarUInt(PayloadBytes, PayloadSizeByteCount);
-	PayloadBytes += PayloadSizeByteCount;
-	if (PayloadSize)
-	{
-		const void* const PayloadEnd = PayloadBytes + PayloadSize;
-		const ECbFieldType UniformType = FCbFieldType::GetType(GetType()) == ECbFieldType::UniformObject ?
-			ECbFieldType(*PayloadBytes++) : ECbFieldType::HasFieldType;
-		return FCbFieldViewIterator::MakeRange(MakeMemoryView(PayloadBytes, PayloadEnd), UniformType);
-	}
-	return FCbFieldViewIterator();
-}
-
 FCbFieldView FCbObjectView::FindView(const FAnsiStringView Name) const
 {
 	for (const FCbFieldView Field : *this)
@@ -790,6 +778,19 @@ void FCbObjectView::CopyTo(FArchive& Ar) const
 
 ///////////////////////////////////////////////////////////////////////////////////////////////////////////////////////
 
+template <typename FieldType>
+TCbFieldIterator<FieldType>& TCbFieldIterator<FieldType>::operator++()
+{
+	const void* const PayloadEnd = FieldType::GetPayloadEnd();
+	const int64 AtEndMask = int64(PayloadEnd == FieldsEnd) - 1;
+	const ECbFieldType NextType = ECbFieldType(int64(FieldType::GetType()) & AtEndMask);
+	const void* const NextField = reinterpret_cast<const void*>(int64(PayloadEnd) & AtEndMask);
+	const void* const NextFieldsEnd = reinterpret_cast<const void*>(int64(FieldsEnd) & AtEndMask);
+	FieldType::Assign(NextField, NextType);
+	FieldsEnd = NextFieldsEnd;
+	return *this;
+}
+
 template <typename FieldType>
 uint64 TCbFieldIterator<FieldType>::GetRangeSize() const
 {
diff --git a/Engine/Source/Runtime/Core/Private/Tests/Serialization/CompactBinaryTest.cpp b/Engine/Source/Runtime/Core/Private/Tests/Serialization/CompactBinaryTest.cpp
--- a/Engine/Source/Runtime/Core/Private/Tests/Serialization/CompactBinaryTest.cpp
+++ b/Engine/Source/Runtime/Core/Private/Tests/Serialization/CompactBinaryTest.cpp
@@ -326,6 +326,15 @@ bool FCbFieldObjectTest::RunTest(const FString& Parameters)
 			TestEqual(TEXT("FCbFieldView(Object) Range"), Field.AsInt32(), ActualNum);
 		}
 		TestEqual(TEXT("FCbFieldView(Object)::AsObjectView() Range -> Count"), ActualNum, ExpectedNum);
+
+		ActualNum = 0;
+		for (FCbFieldView Field : Object.AsFieldView())
+		{
+			++ActualNum;
+			TestNotEqual(TEXT("FCbFieldView(ObjectField) Iterator Name"), Field.GetName().Len(), 0);
+			TestEqual(TEXT("FCbFieldView(ObjectField) Range"), Field.AsInt32(), ActualNum);
+		}
+		TestEqual(TEXT("FCbFieldView(ObjectField)::AsObjectView() Range -> Count"), ActualNum, ExpectedNum);
 	};
 
 	// Test FCbFieldView(Object, Empty)
@@ -447,6 +456,9 @@ bool FCbFieldObjectTest::RunTest(const FString& Parameters)
 		for (FCbFieldIterator It = ObjectClone.CreateIterator(), End; It != End; ++It)
 		{
 		}
+		for (FCbField Field : ObjectClone)
+		{
+		}
 
 		// CopyTo
 		uint8 CopyBytes[6];
@@ -512,6 +524,14 @@ bool FCbFieldArrayTest::RunTest(const FString& Parameters)
 			TestEqual(TEXT("FCbFieldView(Array) Range"), Field.AsInt32(), ActualNum);
 		}
 		TestEqual(TEXT("FCbFieldView(Array)::AsArrayView() Range -> Count"), ActualNum, ExpectedNum);
+
+		ActualNum = 0;
+		for (FCbFieldView Field : Array.AsFieldView())
+		{
+			++ActualNum;
+			TestEqual(TEXT("FCbFieldView(ArrayField) Range"), Field.AsInt32(), ActualNum);
+		}
+		TestEqual(TEXT("FCbFieldView(ArrayField)::AsArrayView() Range -> Count"), ActualNum, ExpectedNum);
 	};
 
 	// Test FCbFieldView(Array, Empty)
@@ -611,6 +631,9 @@ bool FCbFieldArrayTest::RunTest(const FString& Parameters)
 		for (FCbFieldIterator It = ArrayClone.CreateIterator(), End; It != End; ++It)
 		{
 		}
+		for (FCbField Field : ArrayClone)
+		{
+		}
 
 		// CopyTo
 		uint8 CopyBytes[5];
diff --git a/Engine/Source/Runtime/Core/Public/Serialization/CompactBinary.h b/Engine/Source/Runtime/Core/Public/Serialization/CompactBinary.h
--- a/Engine/Source/Runtime/Core/Public/Serialization/CompactBinary.h
+++ b/Engine/Source/Runtime/Core/Public/Serialization/CompactBinary.h
@@ -65,6 +65,7 @@ class FArchive;
 class FBlake3;
 class FCbArrayView;
 class FCbFieldView;
+class FCbFieldViewIterator;
 class FCbObjectId;
 class FCbObjectView;
 struct FDateTime;
@@ -282,6 +283,16 @@ public:
 	static constexpr inline bool IsCustomById(ECbFieldType Type)   { return GetType(Type) == ECbFieldType::CustomById; }
 	static constexpr inline bool IsCustomByName(ECbFieldType Type) { return GetType(Type) == ECbFieldType::CustomByName; }
 
+	static constexpr inline bool HasFields(ECbFieldType Type)
+	{
+		return GetType(Type) >= ECbFieldType::Object && GetType(Type) <= ECbFieldType::UniformArray;
+	}
+
+	static constexpr inline bool HasUniformFields(ECbFieldType Type)
+	{
+		return GetType(Type) == ECbFieldType::UniformObject || GetType(Type) == ECbFieldType::UniformArray;
+	}
+
 	/** Whether the type is or may contain fields of any attachment type. */
 	static constexpr inline bool MayContainAttachments(ECbFieldType Type)
 	{
@@ -297,6 +308,11 @@ using FCbFieldVisitor = TFunctionRef<void (FCbFieldView)>;
 
 ///////////////////////////////////////////////////////////////////////////////////////////////////////////////////////
 
+/** Iterator that can be used as a sentinel for the end of a range. */
+class FCbIteratorSentinel
+{
+};
+
 /**
  * Iterator for FCbField[View] that can operate on any contiguous range of fields.
  *
@@ -312,17 +328,7 @@ public:
 	/** Construct an empty field range. */
 	constexpr TCbFieldIterator() = default;
 
-	inline TCbFieldIterator& operator++()
-	{
-		const void* const PayloadEnd = FieldType::GetPayloadEnd();
-		const int64 AtEndMask = int64(PayloadEnd == FieldsEnd) - 1;
-		const ECbFieldType NextType = ECbFieldType(int64(FieldType::GetType()) & AtEndMask);
-		const void* const NextField = reinterpret_cast<const void*>(int64(PayloadEnd) & AtEndMask);
-		const void* const NextFieldsEnd = reinterpret_cast<const void*>(int64(FieldsEnd) & AtEndMask);
-		FieldType::Assign(NextField, NextType);
-		FieldsEnd = NextFieldsEnd;
-		return *this;
-	}
+	CORE_API TCbFieldIterator& operator++();
 
 	inline TCbFieldIterator operator++(int)
 	{
@@ -365,6 +371,9 @@ public:
 		return !Equals(Other);
 	}
 
+	constexpr inline bool operator==(const FCbIteratorSentinel&) const { return !FieldType::HasValue(); }
+	constexpr inline bool operator!=(const FCbIteratorSentinel&) const { return FieldType::HasValue(); }
+
 	/** Copy the field range into a buffer of exactly GetRangeSize() bytes. */
 	CORE_API void CopyRangeTo(FMutableMemoryView Buffer) const;
 
@@ -396,6 +405,10 @@ public:
 		return false;
 	}
 
+	/** DO NOT USE DIRECTLY. These functions enable range-based for loop support. */
+	constexpr inline TCbFieldIterator begin() const { return *this; }
+	constexpr inline FCbIteratorSentinel end() const { return FCbIteratorSentinel(); }
+
 protected:
 	/** Construct a field range that contains exactly one field. */
 	constexpr inline explicit TCbFieldIterator(FieldType InField)
@@ -423,10 +436,6 @@ protected:
 		return It.FieldsEnd;
 	}
 
-private:
-	friend inline TCbFieldIterator begin(const TCbFieldIterator& Iterator) { return Iterator; }
-	friend inline TCbFieldIterator end(const TCbFieldIterator&) { return TCbFieldIterator(); }
-
 private:
 	template <typename OtherType>
 	friend class TCbFieldIterator;
@@ -741,6 +750,13 @@ public:
 		return false;
 	}
 
+	/** Create an iterator for the fields of an array or object, otherwise an empty iterator. */
+	CORE_API FCbFieldViewIterator CreateViewIterator() const;
+
+	/** DO NOT USE DIRECTLY. These functions enable range-based for loop support. */
+	inline FCbFieldViewIterator begin() const;
+	constexpr inline FCbIteratorSentinel end() const { return FCbIteratorSentinel(); }
+
 protected:
 	/** Returns a view of the name and value payload, which excludes the type. */
 	CORE_API FMemoryView GetViewNoType() const;
@@ -851,6 +867,11 @@ private:
 	using TCbFieldIterator::TCbFieldIterator;
 };
 
+inline FCbFieldViewIterator FCbFieldView::begin() const
+{
+	return CreateViewIterator();
+}
+
 ///////////////////////////////////////////////////////////////////////////////////////////////////////////////////////
 
 /**
@@ -874,9 +895,6 @@ public:
 	/** Returns the number of items in the array. */
 	CORE_API uint64 Num() const;
 
-	/** Create an iterator for the fields of this array. */
-	CORE_API FCbFieldViewIterator CreateViewIterator() const;
-
 	/** Access the array as an array field. */
 	inline FCbFieldView AsFieldView() const { return static_cast<const FCbFieldView&>(*this); }
 
@@ -925,10 +943,12 @@ public:
 		return !FCbFieldView::HasName() && FCbFieldView::TryGetSerializedView(OutView);
 	}
 
-private:
-	friend inline FCbFieldViewIterator begin(const FCbArrayView& Array) { return Array.CreateViewIterator(); }
-	friend inline FCbFieldViewIterator end(const FCbArrayView&) { return FCbFieldViewIterator(); }
+	/** @see FCbFieldView::CreateViewIterator */
+	using FCbFieldView::CreateViewIterator;
+	using FCbFieldView::begin;
+	using FCbFieldView::end;
 
+private:
 	/** Construct an array from an array field. No type check is performed! Use via FromFieldNoCheck. */
 	inline explicit FCbArrayView(const FCbFieldView& Field) : FCbFieldView(Field) {}
 };
@@ -955,9 +975,6 @@ public:
 	/** Construct an object with no fields. */
 	CORE_API FCbObjectView();
 
-	/** Create an iterator for the fields of this object. */
-	CORE_API FCbFieldViewIterator CreateViewIterator() const;
-
 	/**
 	 * Find a field by case-sensitive name comparison.
 	 *
@@ -1023,16 +1040,20 @@ public:
 		return !FCbFieldView::HasName() && FCbFieldView::TryGetSerializedView(OutView);
 	}
 
-private:
-	friend inline FCbFieldViewIterator begin(const FCbObjectView& Object) { return Object.CreateViewIterator(); }
-	friend inline FCbFieldViewIterator end(const FCbObjectView&) { return FCbFieldViewIterator(); }
+	/** @see FCbFieldView::CreateViewIterator */
+	using FCbFieldView::CreateViewIterator;
+	using FCbFieldView::begin;
+	using FCbFieldView::end;
 
+private:
 	/** Construct an object from an object field. No type check is performed! Use via FromFieldNoCheck. */
 	inline explicit FCbObjectView(const FCbFieldView& Field) : FCbFieldView(Field) {}
 };
 
 ///////////////////////////////////////////////////////////////////////////////////////////////////////////////////////
 
+class FCbFieldIterator;
+
 /** A reference to a function that is used to allocate buffers for compact binary data. */
 using FCbBufferAllocator = TFunctionRef<FUniqueBuffer (uint64 Size)>;
 
@@ -1110,6 +1131,13 @@ public:
 	inline const FSharedBuffer& GetOuterBuffer() const & { return Buffer; }
 	inline FSharedBuffer GetOuterBuffer() && { return MoveTemp(Buffer); }
 
+	/** Create an iterator for the fields of an array or object, otherwise an empty iterator. */
+	inline FCbFieldIterator CreateIterator() const;
+
+	/** DO NOT USE DIRECTLY. These functions enable range-based for loop support. */
+	inline FCbFieldIterator begin() const;
+	constexpr inline FCbIteratorSentinel end() const { return FCbIteratorSentinel(); }
+
 private:
 	template <typename OtherType>
 	friend class TCbBuffer;
@@ -1251,6 +1279,18 @@ private:
 	using TCbFieldIterator::TCbFieldIterator;
 };
 
+template <typename ViewType>
+inline FCbFieldIterator TCbBuffer<ViewType>::CreateIterator() const
+{
+	return FCbFieldIterator::MakeRangeView(ViewType::CreateViewIterator(), GetOuterBuffer());
+}
+
+template <typename ViewType>
+inline FCbFieldIterator TCbBuffer<ViewType>::begin() const
+{
+	return CreateIterator();
+}
+
 ///////////////////////////////////////////////////////////////////////////////////////////////////////////////////////
 
 /**
@@ -1264,12 +1304,6 @@ class FCbArray : public TCbBuffer<FCbArrayView>, public TCbBufferFactory<FCbArra
 public:
 	using TCbBuffer::TCbBuffer;
 
-	/** Create an iterator for the fields of this array. */
-	inline FCbFieldIterator CreateIterator() const
-	{
-		return FCbFieldIterator::MakeRangeView(CreateViewIterator(), GetOuterBuffer());
-	}
-
 	/** Access the array as an array field. */
 	inline FCbField AsField() const &
 	{
@@ -1281,10 +1315,6 @@ public:
 	{
 		return FCbField(FCbArrayView::AsFieldView(), MoveTemp(*this));
 	}
-
-private:
-	friend inline FCbFieldIterator begin(const FCbArray& Array) { return Array.CreateIterator(); }
-	friend inline FCbFieldIterator end(const FCbArray&) { return FCbFieldIterator(); }
 };
 
 ///////////////////////////////////////////////////////////////////////////////////////////////////////////////////////
@@ -1300,12 +1330,6 @@ class FCbObject : public TCbBuffer<FCbObjectView>, public TCbBufferFactory<FCbOb
 public:
 	using TCbBuffer::TCbBuffer;
 
-	/** Create an iterator for the fields of this object. */
-	inline FCbFieldIterator CreateIterator() const
-	{
-		return FCbFieldIterator::MakeRangeView(CreateViewIterator(), GetOuterBuffer());
-	}
-
 	/** Find a field by case-sensitive name comparison. */
 	inline FCbField Find(FAnsiStringView Name) const
 	{
@@ -1340,10 +1364,6 @@ public:
 	{
 		return FCbField(FCbObjectView::AsFieldView(), MoveTemp(*this));
 	}
-
-private:
-	friend inline FCbFieldIterator begin(const FCbObject& Object) { return Object.CreateIterator(); }
-	friend inline FCbFieldIterator end(const FCbObject&) { return FCbFieldIterator(); }
 };
 
 ///////////////////////////////////////////////////////////////////////////////////////////////////////////////////////
