fix static analysis warning

diff --git a/Engine/Source/Runtime/Core/Public/Math/UnrealMathSSE.h b/Engine/Source/Runtime/Core/Public/Math/UnrealMathSSE.h
--- a/Engine/Source/Runtime/Core/Public/Math/UnrealMathSSE.h
+++ b/Engine/Source/Runtime/Core/Public/Math/UnrealMathSSE.h
@@ -12,8 +12,10 @@
 
 #include "Math/sse_mathfun.h"
 
-// We suppress static analysis warnings for the cast from (double*) to (float*) in VectorLoadFloat2 below:
+// We suppress static analysis warnings for the cast from (double*) to (float*) in VectorLoadFloat2
+// and VectorLoadTwoPairsFloat below:
 // -V:VectorLoadFloat2:615
+// -V:VectorLoadTwoPairsFloat:615
 
 /*=============================================================================
  *	Helpers:
