Added TStridedView, which is designed to allow functions to take pointers to arbitrarily structured data. It is similar to TArrayView, but allows custom byte stride between elements.

diff --git a/Engine/Source/Runtime/Core/Private/Tests/Containers/StridedViewTest.cpp b/Engine/Source/Runtime/Core/Private/Tests/Containers/StridedViewTest.cpp
new file mode 100644
--- /dev/null
+++ b/Engine/Source/Runtime/Core/Private/Tests/Containers/StridedViewTest.cpp
@@ -0,0 +1,323 @@
+// Copyright Epic Games, Inc. All Rights Reserved.
+
+#include "Containers/StridedView.h"
+#include "Misc/AutomationTest.h"
+
+#if WITH_DEV_AUTOMATION_TESTS
+
+#include "Algo/Accumulate.h"
+#include "Containers/Array.h"
+#include "Math/Vector.h"
+
+#include <type_traits>
+
+namespace UE
+{
+namespace StridedViewTest
+{
+
+struct FMyStruct
+{
+	uint32  SomeData = 123;
+	FVector Position = FVector(5.0f, 6.0f, 7.0f);
+	uint32  MoreData = 890;
+};
+
+struct FMyStructDerived : FMyStruct
+{
+	FVector Normal = FVector(0.0f, 0.0f, 0.0f);
+};
+
+static void SetValue(TStridedView<FVector> Output, const FVector& Value)
+{
+	for (FVector& It : Output)
+	{
+		It = Value;
+	}
+}
+
+static FVector ComputeSum(TStridedView<const FVector> Values)
+{
+	return Algo::Accumulate(Values, FVector(0.0f));
+}
+
+static FVector ComputeMean(TStridedView<const FVector> Values)
+{
+	return ComputeSum(Values) / (float)Values.Num();
+}
+
+static FVector ComputeMeanPosition(TArrayView<const FMyStruct> Structs)
+{
+	return ComputeMean(MakeStridedView(sizeof(FMyStruct), &Structs[0].Position, Structs.Num()));
+}
+
+static FVector ComputeMeanPositionStrided(TStridedView<const FMyStruct> Structs)
+{
+	return ComputeMean(MakeStridedView(Structs.GetStride(), &Structs[0].Position, Structs.Num()));
+}
+
+static TStridedView<const FVector> TestCreateConstArrayView(const TArray<FVector>& VecsArray)
+{
+	return MakeStridedView(VecsArray);
+}
+
+static_assert(std::is_convertible_v<TStridedView<FMyStruct>, TStridedView<FMyStruct>>);
+static_assert(std::is_convertible_v<TStridedView<FMyStruct>, TStridedView<const FMyStruct>>);
+static_assert(!std::is_convertible_v<TStridedView<const FMyStruct>, TStridedView<FMyStruct>>);
+static_assert(!std::is_convertible_v<TStridedView<FMyStructDerived>, TStridedView<FMyStruct>>);
+static_assert(!std::is_convertible_v<TStridedView<FMyStruct>, TStridedView<FMyStructDerived>>);
+static_assert(std::is_convertible_v<TStridedView<const FMyStruct*>, TStridedView<const FMyStruct*>>);
+static_assert(std::is_convertible_v<TStridedView<const FMyStruct*>, TStridedView<const FMyStruct* const>>);
+static_assert(!std::is_convertible_v<TStridedView<const FMyStruct* const>, TStridedView<const FMyStruct*>>);
+static_assert(!std::is_convertible_v<TStridedView<const FMyStructDerived*>, TStridedView<const FMyStruct*>>);
+static_assert(!std::is_convertible_v<TStridedView<const FMyStruct*>, TStridedView<const FMyStructDerived*>>);
+
+} // namespace StridedViewTest
+} // namespace UE
+
+IMPLEMENT_SIMPLE_AUTOMATION_TEST(FStridedViewTest, "System.Core.Containers.StridedView.Misc", EAutomationTestFlags::ApplicationContextMask | EAutomationTestFlags::SmokeFilter)
+bool FStridedViewTest::RunTest(const FString& Parameters)
+{
+	using namespace UE::StridedViewTest;
+
+	static constexpr int32 Num = 17;
+
+	FMyStruct Structs[Num];
+
+	// Initialize test data and validate that basic iteration works
+
+	{
+		TStridedView<FVector> PositionsView(sizeof(FMyStruct), &Structs[0].Position, Num);
+		for (int32 i = 0; i < PositionsView.Num(); ++i)
+		{
+			PositionsView[i] = FVector((float)i);
+		}
+
+		for (int32 i = 0; i < Num; ++i)
+		{
+			TestEqual(TEXT("Structs[i].Position"), Structs[i].Position, FVector((float)i));
+		}
+
+		const FVector Multiplier = FVector(2.0f, 3.0f, 4.0f);
+		for (FVector& It : PositionsView)
+		{
+			It *= Multiplier;
+		}
+
+		for (int32 i = 0; i < Num; ++i)
+		{
+			TestEqual(TEXT("Structs[i].Position"), Structs[i].Position, FVector((float)i) * Multiplier);
+		}
+	}
+
+	FVector ExpectedAvg = FVector(0.0f);
+	for (int32 i = 0; i < Num; ++i)
+	{
+		ExpectedAvg += Structs[i].Position;
+	}
+	ExpectedAvg /= (float)Num;
+
+	// Test various ways to construct and use the strided view
+
+	{
+		// Incorrect API usage: trying to pass a const strided view to a mutating function. Generates compile error (as expected).
+		// TStridedView<const FVector> ConstView(sizeof(FMyStruct), &Structs[0].Position, Num);
+		// SetValue(ConstView, FVector(123.0f));
+	}
+
+	{
+		FVector Avg = ComputeMeanPosition(MakeArrayView(Structs));
+		TestEqual(TEXT("ComputeMeanPosition(MakeArrayView(Structs))"), Avg, ExpectedAvg);
+	}
+
+	{
+		FVector Avg = ComputeMeanPositionStrided(MakeStridedView(Structs));
+		TestEqual(TEXT("ComputeMeanPositionStrided(MakeStridedView(Structs))"), Avg, ExpectedAvg);
+	}
+
+	{
+		FVector Avg = ComputeMean(TStridedView<FVector>(sizeof(FMyStruct), &Structs[0].Position, Num));
+		TestEqual(TEXT("ComputeMean(TStridedView<FVector>(sizeof(FMyStruct), &Structs[0].Position, Num))"), Avg, ExpectedAvg);
+	}
+
+	{
+		FVector Avg = ComputeMean(TStridedView<const FVector>(sizeof(FMyStruct), &Structs[0].Position, Num));
+		TestEqual(TEXT("ComputeMean(TStridedView<const FVector>(sizeof(FMyStruct), &Structs[0].Position, Num))"), Avg, ExpectedAvg);
+	}
+
+	{
+		FVector Avg = ComputeMean(MakeStridedView(sizeof(FMyStruct), &Structs[0].Position, Num));
+		TestEqual(TEXT("ComputeMean(MakeStridedView(sizeof(FMyStruct), &Structs[0].Position, Num))"), Avg, ExpectedAvg);
+	}
+
+	{
+		FVector Avg = ComputeMean(MakeStridedView(MakeArrayView(Structs), &FMyStruct::Position));
+		TestEqual(TEXT("ComputeMean(MakeStridedView(MakeArrayView(Structs), &FMyStruct::Position))"), Avg, ExpectedAvg);
+	}
+
+	{
+		FVector Avg = ComputeMean(MakeStridedView(Structs, &FMyStruct::Position));
+		TestEqual(TEXT("ComputeMean(MakeStridedView(Structs, &FMyStruct::Position))"), Avg, ExpectedAvg);
+	}
+
+	{
+		TArray<FMyStruct> StructsArray = TArray<FMyStruct>(Structs, Num);
+		FVector Avg = ComputeMean(MakeStridedView(StructsArray, &FMyStruct::Position));
+		TestEqual(TEXT("ComputeMean(MakeStridedView(StructsArray, &FMyStruct::Position))"), Avg, ExpectedAvg);
+	}
+
+	// Test support for trivial strided views (stride = sizeof(element))
+
+	FVector Vecs[Num];
+	for (int32 i = 0; i < Num; ++i)
+	{
+		Vecs[i] = Structs[i].Position;
+	}
+
+	{
+		FVector Avg = ComputeMean(MakeStridedView(Vecs));
+		TestEqual(TEXT("ComputeMean(MakeStridedView(Vecs))"), Avg, ExpectedAvg);
+	}
+
+	{
+		TArrayView<const FVector> VecsView = Vecs;
+		auto VecsStridedView = MakeStridedView(VecsView);
+		static_assert(std::is_const_v<decltype(VecsStridedView)::ElementType>);
+		FVector Avg = ComputeMean(VecsStridedView);
+		TestEqual(TEXT("ComputeMean(MakeStridedView(VecsView))"), Avg, ExpectedAvg);
+	}
+
+	{
+		TArray<FVector> VecsArray = TArray<FVector>(Vecs, Num);
+		auto VecsStridedView = TestCreateConstArrayView(VecsArray);
+		static_assert(std::is_const_v<decltype(VecsStridedView)::ElementType>);
+		FVector Avg = ComputeMean(VecsStridedView);
+		TestEqual(TEXT("ComputeMean(MakeStridedView(VecsArray))"), Avg, ExpectedAvg);
+	}
+
+	// Test smearing an element
+
+	{
+		FVector V = FVector(1.0f, 2.0f, 3.0f);
+		FVector Avg = ComputeMean(TStridedView<FVector>(0, &V, Num));
+		TestEqual(TEXT("ComputeMean(TStridedView<FVector>(0, &V, Num))"), Avg, V);
+
+		FVector Sum = ComputeSum(TStridedView<FVector>(0, &V, Num));
+		FVector ExpectedSum = V * (float)Num;
+		TestEqual(TEXT("ComputeSum(TStridedView<FVector>(0, &V, Num))"), Sum, ExpectedSum);
+	}
+
+	// Test iterating through derived struct
+
+	FMyStructDerived StructsDerived[Num];
+
+	// Initialize derived struct data by copying positions from base and set normals to (1,2,3)*i
+
+	{
+		TStridedView<FVector> PositionsViewDerived(sizeof(FMyStructDerived), &StructsDerived[0].Position, Num);
+		TStridedView<FVector> NormalsView(sizeof(FMyStructDerived), &StructsDerived[0].Normal, Num);
+		TStridedView<FVector> PositionsView(sizeof(FMyStruct), &Structs[0].Position, Num);
+
+		for (int32 i=0; i < PositionsViewDerived.Num(); ++i)
+		{
+			PositionsViewDerived[i] = PositionsView[i];
+			NormalsView[i] = FVector(1.0f, 2.0f, 3.0f) * (float)i;
+		}
+
+		for (int32 i = 0; i < PositionsViewDerived.Num(); ++i)
+		{
+			TestEqual(TEXT("StructsDerived[i].Position"), StructsDerived[i].Position, Structs[i].Position);
+			TestEqual(TEXT("StructsDerived[i].Normal"), StructsDerived[i].Normal, FVector(1.0f, 2.0f, 3.0f) * (float)i);
+		}
+	}
+
+	// Test different ways to construct a view of a derived struct
+
+	{
+		// Incorrect API usage: trying to pass a const strided view to a mutating function. Generates compile error (as expected).
+		// TStridedView<const FVector> ConstView(sizeof(FMyStructDerived), &StructsDerived[0].Position, Num);
+		// SetValue(ConstView, FVector(123.0f));
+	}
+
+	{
+		FVector Avg = ComputeMean(TStridedView<FVector>(sizeof(FMyStructDerived), &StructsDerived[0].Position, Num));
+		TestEqual(TEXT("ComputeMean(TStridedView<FVector>(sizeof(FMyStructDerived), &StructsDerived[0].Position, Num))"), Avg, ExpectedAvg);
+	}
+
+	{
+		FVector Avg = ComputeMean(TStridedView<const FVector>(sizeof(FMyStructDerived), &StructsDerived[0].Position, Num));
+		TestEqual(TEXT("ComputeMean(TStridedView<const FVector>(sizeof(FMyStructDerived), &StructsDerived[0].Position, Num))"), Avg, ExpectedAvg);
+	}
+
+	{
+		FVector Avg = ComputeMean(MakeStridedView(sizeof(FMyStructDerived), &StructsDerived[0].Position, Num));
+		TestEqual(TEXT("ComputeMean(MakeStridedView(sizeof(FMyStructDerived), &StructsDerived[0].Position, Num))"), Avg, ExpectedAvg);
+	}
+
+	{
+		// Incorrect API usage: trying to create view of derived class from base. Generates compile error (as expected).
+		// TStridedView<FMyStructDerived> ViewDerivedError = MakeStridedView(Structs);
+
+		TStridedView<FMyStructDerived> ViewDerived = MakeStridedView(StructsDerived);
+
+		// Incorrect API usage : trying to create view of base class from derived. Generates compile error (as expected).
+		// TStridedView<FMyStruct> ViewDerivedAsBase = MakeStridedView(StructsDerived);
+
+		// Explicitly construct view of base from derived
+		TStridedView<FMyStruct> ViewDerivedAsBase(sizeof(StructsDerived[0]), static_cast<FMyStruct*>(StructsDerived), Num);
+
+		TestEqual(TEXT("ViewDerived.GetStride()"), ViewDerived.GetStride(), (int32)sizeof(FMyStructDerived));
+		TestEqual(TEXT("ViewDerivedAsBase.GetStride()"), ViewDerivedAsBase.GetStride(), (int32)sizeof(FMyStructDerived));
+
+		// Incorrect API usage : implicit conversion of derived view to base. Generates compile error (as expected).
+		// FVector Avg = ComputeMeanPositionStrided(ViewDerived);
+
+		FVector AvgDerivedAsBase = ComputeMeanPositionStrided(ViewDerivedAsBase);
+		TestEqual(TEXT("ComputeMeanPositionStrided(ViewDerivedAsBase)"), AvgDerivedAsBase, ExpectedAvg);
+	}
+
+	{
+		TArrayView<FMyStructDerived> DerivedArrayView = MakeArrayView(StructsDerived);
+		TStridedView<FVector> View = MakeStridedView(DerivedArrayView, &FMyStructDerived::Position);
+		TestEqual(TEXT("MakeStridedView(MakeArrayView(StructsDerived), &FMyStructDerived::Position).GetStride()"), View.GetStride(), (int32)sizeof(StructsDerived[0]));
+		FVector Avg = ComputeMean(View);
+		TestEqual(TEXT("ComputeMean(MakeStridedView(MakeArrayView(StructsDerived), &FMyStructDerived::Position))"), Avg, ExpectedAvg);
+	}
+
+	{
+		TArrayView<FMyStructDerived> DerivedArrayView = MakeArrayView(StructsDerived);
+		TStridedView<FVector> View = MakeStridedView(DerivedArrayView, &FMyStruct::Position);
+		TestEqual(TEXT("MakeStridedView(MakeArrayView(StructsDerived), &FMyStruct::Position).GetStride()"), View.GetStride(), (int32)sizeof(StructsDerived[0]));
+		FVector Avg = ComputeMean(View);
+		TestEqual(TEXT("ComputeMean(MakeStridedView(MakeArrayView(StructsDerived), &FMyStruct::Position))"), Avg, ExpectedAvg);
+	}
+
+	{
+		FVector Avg = ComputeMean(MakeStridedView(StructsDerived, &FMyStructDerived::Position));
+		TestEqual(TEXT("ComputeMean(MakeStridedView(StructsDerived, &FMyStructDerived::Position))"), Avg, ExpectedAvg);
+	}
+
+	{
+		TArray<FMyStructDerived> StructsArrayDerived = TArray<FMyStructDerived>(StructsDerived, Num);
+		FVector Avg = ComputeMean(MakeStridedView(StructsArrayDerived, &FMyStructDerived::Position));
+		TestEqual(TEXT("ComputeMean(MakeStridedView(StructsArrayDerived, &FMyStructDerived::Position))"), Avg, ExpectedAvg);
+	}
+
+	// Test iterating through nested fields
+
+	{
+		TStridedView<float> ViewOfZs = MakeStridedView((int32)sizeof(Structs[0]), &Structs[0].Position.Z, Num);
+		float ExpectedSum = 0.0;
+		for (int32 i = 0; i < Num; ++i)
+		{
+			ExpectedSum += Structs[i].Position.Z;
+		}
+		float Sum = Algo::Accumulate(ViewOfZs, 0.0f);
+		TestEqual(TEXT("Algo::Accumulate(MakeStridedView((int32)sizeof(Structs[0]), &Structs[0].Position.Z, Num), 0.0f)"), Sum, ExpectedSum);
+	}
+
+	return !HasAnyErrors();
+}
+
+
+#endif // WITH_DEV_AUTOMATION_TESTS
diff --git a/Engine/Source/Runtime/Core/Public/Containers/StridedView.h b/Engine/Source/Runtime/Core/Public/Containers/StridedView.h
new file mode 100644
--- /dev/null
+++ b/Engine/Source/Runtime/Core/Public/Containers/StridedView.h
@@ -0,0 +1,184 @@
+// Copyright Epic Games, Inc. All Rights Reserved.
+
+#pragma once
+
+#include "CoreTypes.h"
+#include "Templates/UnrealTemplate.h"
+#include "Misc/AssertionMacros.h"
+#include "Containers/ArrayView.h"
+#include "HAL/PlatformString.h" // for INT64_FMT
+
+#include <type_traits>
+
+/*
+* TStridedView is similar to TArrayView, but allows flexible byte stride between elements.
+* Stride must be positive and a multiple of the element alignment.
+* Stride can be zero to duplicate the same element over the whole range.
+* 
+* Example usage:
+* 
+* struct FMyStruct
+* {
+*     uint32  SomeData;
+*     FVector Position;
+* };
+*
+* FVector ComputeMean(TStridedView<const FVector> Positions)
+* {
+*     return Algo::Accumulate(Positions, FVector(0.0f)) / Positions.Num();
+* }
+*
+* FVector ComputeMeanPosition(TArrayView<const FMyStruct> Structs)
+* {
+*     return ComputeMean(MakeStridedView(Structs, &FMyStruct::Position))
+* }
+* 
+* See StridedViewTest.cpp for more examples.
+* 
+*/
+
+/** Pointer with extent and a stride, similar to TArrayView. Designed to allow functions to take pointers to arbitrarily structured data. */
+template<typename InElementType, typename InSizeType = int32>
+class TStridedView
+{
+public:
+
+	using ElementType = InElementType;
+	using SizeType = InSizeType;
+
+	static_assert(std::is_signed_v<SizeType>, "TStridedView only supports signed index types");
+
+	TStridedView() = default;
+
+	template <typename OtherElementType, decltype(ImplicitConv<ElementType* const*>((OtherElementType**)nullptr))* = nullptr>
+	FORCEINLINE TStridedView(SizeType InBytesBetweenElements, OtherElementType* InFirstElementPtr, SizeType InNumElements)
+		: FirstElementPtr(InFirstElementPtr)
+		, BytesBetweenElements(InBytesBetweenElements)
+		, NumElements(InNumElements)
+	{
+		check(NumElements >= 0);
+		check(BytesBetweenElements >= 0); // NOTE: Zero stride is valid to allow duplicating a single element.
+		check(BytesBetweenElements % alignof(ElementType) == 0);
+	}
+
+	template <typename OtherElementType, decltype(ImplicitConv<ElementType* const*>((OtherElementType**)nullptr))* = nullptr>
+	FORCEINLINE TStridedView(const TStridedView<OtherElementType, SizeType>& Other)
+		: FirstElementPtr(nullptr)
+		, BytesBetweenElements(Other.GetStride())
+		, NumElements(Other.Num())
+	{
+		if (Other.Num())
+		{
+			FirstElementPtr = &Other[0];
+		}
+	}
+
+	FORCEINLINE SizeType Num() const
+	{
+		return NumElements;
+	}
+
+	FORCEINLINE SizeType GetStride() const
+	{
+		return BytesBetweenElements;
+	}
+
+	FORCEINLINE ElementType& operator[](SizeType Index) const
+	{
+		return *GetElementPtr(Index);
+	}
+
+	struct FIterator
+	{
+		const TStridedView* Owner;
+		SizeType Index;
+
+		FORCEINLINE FIterator& operator++()
+		{
+			++Index;
+			return *this;
+		}
+
+		FORCEINLINE ElementType& operator*()
+		{
+			return *Owner->GetElementPtrUnsafe(Index);
+		}
+
+		FORCEINLINE bool operator == (const FIterator& Other) const
+		{
+			return Owner == Other.Owner
+				&& Index == Other.Index;
+		}
+
+		FORCEINLINE bool operator != (const FIterator& Other) const
+		{
+			return !(*this == Other);
+		}
+	};
+
+	/** Ranged iteration support. DO NOT USE DIRECTLY. */
+
+	FORCEINLINE FIterator begin() const { return FIterator{ this, 0 }; }
+	FORCEINLINE FIterator end() const { return FIterator{ this, Num() }; }
+
+private:
+
+	FORCEINLINE void RangeCheck(SizeType Index) const
+	{
+		checkf((Index >= 0) & (Index < NumElements), TEXT("Array index out of bounds: %" INT64_FMT " from an array of size %" INT64_FMT), int64(Index), int64(NumElements))
+	}
+
+	FORCEINLINE ElementType* GetElementPtrUnsafe(SizeType Index) const
+	{
+		using ByteType = typename std::conditional_t<std::is_const_v<ElementType>, const uint8, uint8>;
+
+		ByteType* AsBytes = reinterpret_cast<ByteType*>(FirstElementPtr);
+		ElementType* AsElement = reinterpret_cast<ElementType*>(AsBytes + Index * BytesBetweenElements);
+
+		return AsElement;
+	}
+
+	FORCEINLINE ElementType* GetElementPtr(SizeType Index) const
+	{
+		RangeCheck(Index);
+		return GetElementPtrUnsafe(Index);
+	}
+
+private:
+
+	ElementType* FirstElementPtr = nullptr;
+	SizeType BytesBetweenElements = 0;
+	SizeType NumElements = 0;
+};
+
+template <typename ElementType>
+TStridedView<ElementType> MakeStridedView(int32 BytesBetweenElements, ElementType* FirstElement, int32 Count)
+{
+	return TStridedView<ElementType>(BytesBetweenElements, FirstElement, Count);
+}
+
+template <typename StructureType>
+TStridedView<StructureType> MakeStridedView(TArrayView<StructureType> StructuredView)
+{
+	return MakeStridedView((int32)sizeof(StructureType), GetData(StructuredView), (int32)GetNum(StructuredView));
+}
+
+template <
+	typename StructuredRangeType,
+	typename ElementType,
+	typename StructureType,
+	decltype(GetData(std::declval<StructuredRangeType>()))* = nullptr,
+	decltype(GetNum(std::declval<StructuredRangeType>()))* = nullptr
+>
+TStridedView<ElementType> MakeStridedView(StructuredRangeType&& StructuredRange, ElementType StructureType::* Member)
+{
+	return TStridedView<ElementType>((int32)sizeof(*GetData(StructuredRange)), &(GetData(StructuredRange)->*Member), (int32)GetNum(StructuredRange));
+}
+
+template <typename StructuredRangeType>
+auto MakeStridedView(StructuredRangeType&& StructuredRange) -> TStridedView<std::remove_reference_t<decltype(*GetData(StructuredRange))>>
+{
+	auto StructuredView = MakeArrayView(Forward<StructuredRangeType>(StructuredRange));
+	return MakeStridedView((int32)sizeof(*GetData(StructuredRange)), GetData(StructuredView), (int32)GetNum(StructuredView));
+}
+
