CompressedBuffer: Added GetFormatName()

#rnx

diff --git a/Engine/Source/Runtime/Core/Private/Compression/CompressedBuffer.cpp b/Engine/Source/Runtime/Core/Private/Compression/CompressedBuffer.cpp
--- a/Engine/Source/Runtime/Core/Private/Compression/CompressedBuffer.cpp
+++ b/Engine/Source/Runtime/Core/Private/Compression/CompressedBuffer.cpp
@@ -514,6 +514,19 @@ static const FMethod* GetMethod(FName FormatName)
 	return nullptr;
 }
 
+static FName GetMethodName(EMethod Method)
+{
+	switch (Method)
+	{
+	default:
+		return NAME_Error;
+	case EMethod::None:
+		return NAME_None;
+	case EMethod::LZ4:
+		return NAME_LZ4;
+	}
+}
+
 template <typename BufferType>
 inline FCompositeBuffer ValidBufferOrEmpty(BufferType&& CompressedData)
 {
@@ -628,6 +641,12 @@ FBlake3Hash FCompressedBuffer::GetRawHash() const
 	return CompressedData ? UE::CompressedBuffer::FHeader::Read(CompressedData).RawHash : FBlake3Hash();
 }
 
+FName FCompressedBuffer::GetFormatName() const
+{
+	using namespace UE::CompressedBuffer;
+	return CompressedData ? GetMethodName(FHeader::Read(CompressedData).Method) : NAME_None;
+}
+
 bool FCompressedBuffer::TryDecompressTo(FMutableMemoryView RawView) const
 {
 	using namespace UE::CompressedBuffer;
diff --git a/Engine/Source/Runtime/Core/Public/Compression/CompressedBuffer.h b/Engine/Source/Runtime/Core/Public/Compression/CompressedBuffer.h
--- a/Engine/Source/Runtime/Core/Public/Compression/CompressedBuffer.h
+++ b/Engine/Source/Runtime/Core/Public/Compression/CompressedBuffer.h
@@ -68,6 +68,17 @@ public:
 	/** Returns the hash of the raw data. Zero on error or if this is null. */
 	[[nodiscard]] CORE_API FBlake3Hash GetRawHash() const;
 
+	/**
+	 * Returns the name of the compression format used by this buffer.
+	 *
+	 * The format name may differ from the format name specified when creating the compressed buffer
+	 * because an incompressible buffer is stored with NAME_None, and a request of NAME_Default will
+	 * be stored in a specific format such as NAME_LZ4.
+	 *
+	 * @return The format name, or NAME_None if this null, or NAME_Error if the format is unknown.
+	 */
+	[[nodiscard]] CORE_API FName GetFormatName() const;
+
 	/**
 	 * Decompress into a memory view that is exactly GetRawSize() bytes.
 	 */
