CompressedBuffer: Fixed the compressed buffer size calculation

#rnx

diff --git a/Engine/Source/Runtime/Core/Private/Compression/CompressedBuffer.cpp b/Engine/Source/Runtime/Core/Private/Compression/CompressedBuffer.cpp
--- a/Engine/Source/Runtime/Core/Private/Compression/CompressedBuffer.cpp
+++ b/Engine/Source/Runtime/Core/Private/Compression/CompressedBuffer.cpp
@@ -196,6 +196,17 @@ protected:
 	virtual uint64 CompressBlockBound(uint64 RawSize) const = 0;
 	virtual bool CompressBlock(FMutableMemoryView& CompressedData, FMemoryView RawData) const = 0;
 	virtual bool DecompressBlock(FMutableMemoryView RawData, FMemoryView CompressedData) const = 0;
+
+private:
+	uint64 GetCompressedBlocksBound(uint64 BlockCount, uint64 BlockSize, uint64 RawSize) const
+	{
+		switch (BlockCount)
+		{
+		case 0:  return 0;
+		case 1:  return CompressBlockBound(RawSize);
+		default: return CompressBlockBound(BlockSize) - BlockSize + RawSize;
+		}
+	}
 };
 
 FCompositeBuffer FMethodBlock::Compress(const FCompositeBuffer& RawData, const uint64 BlockSize) const
@@ -211,9 +222,7 @@ FCompositeBuffer FMethodBlock::Compress(const FCompositeBuffer& RawData, const u
 
 	// Allocate the buffer for the header, metadata, and compressed blocks.
 	const uint64 MetaSize = BlockCount * sizeof(uint32);
-	const uint64 CompressedDataSize = sizeof(FHeader) + MetaSize +
-		BlockCount == 0 ? 0 :
-		BlockCount == 1 ? CompressBlockBound(RawSize) : CompressBlockBound(BlockSize) - BlockSize + RawSize;
+	const uint64 CompressedDataSize = sizeof(FHeader) + MetaSize + GetCompressedBlocksBound(BlockCount, BlockSize, RawSize);
 	FUniqueBuffer CompressedData = FUniqueBuffer::Alloc(CompressedDataSize);
 
 	// Compress the raw data in blocks and store the raw data for incompressible blocks.
