Fixed static analysis warning

diff --git a/Engine/Source/Runtime/Core/Private/Tests/Containers/StridedViewTest.cpp b/Engine/Source/Runtime/Core/Private/Tests/Containers/StridedViewTest.cpp
--- a/Engine/Source/Runtime/Core/Private/Tests/Containers/StridedViewTest.cpp
+++ b/Engine/Source/Runtime/Core/Private/Tests/Containers/StridedViewTest.cpp
@@ -264,7 +264,7 @@ bool FStridedViewTest::RunTest(const FString& Parameters)
 		// TStridedView<FMyStruct> ViewDerivedAsBase = MakeStridedView(StructsDerived);
 
 		// Explicitly construct view of base from derived
-		TStridedView<FMyStruct> ViewDerivedAsBase(sizeof(StructsDerived[0]), static_cast<FMyStruct*>(StructsDerived), Num);
+		TStridedView<FMyStruct> ViewDerivedAsBase(sizeof(StructsDerived[0]), &static_cast<FMyStruct&>(StructsDerived[0]), Num);
 
 		TestEqual(TEXT("ViewDerived.GetStride()"), ViewDerived.GetStride(), (int32)sizeof(FMyStructDerived));
 		TestEqual(TEXT("ViewDerivedAsBase.GetStride()"), ViewDerivedAsBase.GetStride(), (int32)sizeof(FMyStructDerived));
