Add per asset memory usage approximation during async compilation to offer more precise throttling when memory is scarce or we're processing huge assets
  - FQueuedThreadPoolWrapper will now try to retract any task it can until the max concurrency goal is achieved to respect priority as much as possible
  - Add memory requirement estimate on IQueuedWork to allow the threadpool wrapper to take action when required
  - Add logic in the asset compiler's threadpool to handle per asset memory requirement instead of using a default for all assets
  - Add estimate for static mesh and textures
  - This change should improve CPU usage when processing small assets on memory constrained systems and should prevent OOM when dealing with several excessively large assets

diff --git a/Engine/Source/Runtime/Core/Private/HAL/ThreadingBase.cpp b/Engine/Source/Runtime/Core/Private/HAL/ThreadingBase.cpp
--- a/Engine/Source/Runtime/Core/Private/HAL/ThreadingBase.cpp
+++ b/Engine/Source/Runtime/Core/Private/HAL/ThreadingBase.cpp
@@ -869,12 +869,12 @@ IQueuedWork* FThreadPoolPriorityQueue::Dequeue(EQueuedWorkPriority* OutDequeuedW
 	return Work;
 }
 
-IQueuedWork* FThreadPoolPriorityQueue::Peek(EQueuedWorkPriority* OutDequeuedWorkPriority)
+IQueuedWork* FThreadPoolPriorityQueue::Peek(EQueuedWorkPriority* OutDequeuedWorkPriority) const
 {
 	IQueuedWork* Work = nullptr;
 	for (int32 QueueIndex = FirstNonEmptyQueueIndex, Num = PriorityQueuedWork.Num(); QueueIndex < Num; ++QueueIndex)
 	{
-		TArray<IQueuedWork*>& QueuedWork = PriorityQueuedWork[QueueIndex];
+		const TArray<IQueuedWork*>& QueuedWork = PriorityQueuedWork[QueueIndex];
 		if (QueuedWork.Num() > 0)
 		{
 			Work = QueuedWork[0];
diff --git a/Engine/Source/Runtime/Core/Private/Misc/QueuedThreadPoolWrapper.cpp b/Engine/Source/Runtime/Core/Private/Misc/QueuedThreadPoolWrapper.cpp
--- a/Engine/Source/Runtime/Core/Private/Misc/QueuedThreadPoolWrapper.cpp
+++ b/Engine/Source/Runtime/Core/Private/Misc/QueuedThreadPoolWrapper.cpp
@@ -7,32 +7,57 @@
 
 struct FQueuedThreadPoolWrapper::FScheduledWork : public IQueuedWork
 {
-	FScheduledWork(FQueuedThreadPoolWrapper* InParentPool)
+	FScheduledWork(FQueuedThreadPoolWrapper* InParentPool, IQueuedWork* InWork, EQueuedWorkPriority InPriority)
 		: ParentPool(InParentPool)
+		, Work(InWork)
+		, Priority(InPriority)
+		, RequiredMemory(InWork->GetRequiredMemory())
 	{
 	}
 
 	void DoThreadedWork() override
 	{
-		InnerWork->DoThreadedWork();
+		Work->DoThreadedWork();
 		ParentPool->Schedule(this);
 	}
 
 	void Abandon() override
 	{
-		InnerWork->Abandon();
+		Work->Abandon();
 		ParentPool->Schedule(this);
 	}
 
 	EQueuedWorkFlags GetQueuedWorkFlags() const override 
 	{
-		return InnerWork->GetQueuedWorkFlags();
+		return Work->GetQueuedWorkFlags();
 	}
 
-	IQueuedWork* InnerWork = nullptr;
-	EQueuedWorkPriority Priority = EQueuedWorkPriority::Normal;
+	int64 GetRequiredMemory() const override 
+	{
+		return RequiredMemory;
+	}
+
+	IQueuedWork* GetInnerWork() const
+	{
+		return Work;
+	}
+	
+	EQueuedWorkPriority GetPriority() const
+	{
+		return Priority;
+	}
+
+	void Reset()
+	{
+		Work = nullptr;
+	}
 private:
 	FQueuedThreadPoolWrapper* ParentPool;
+	IQueuedWork* Work;
+	EQueuedWorkPriority Priority;
+	
+	// Store the memory of the inner task to ensure it stays constant
+	int64 RequiredMemory;
 };
 
 FQueuedThreadPoolWrapper::FQueuedThreadPoolWrapper(FQueuedThreadPool* InWrappedQueuedThreadPool, int32 InMaxConcurrency, TFunction<EQueuedWorkPriority(EQueuedWorkPriority)> InPriorityMapper)
@@ -74,7 +99,7 @@ void FQueuedThreadPoolWrapper::Destroy()
 		{
 			if (WrappedQueuedThreadPool->RetractQueuedWork(Work))
 			{
-				Work->InnerWork->Abandon();
+				Work->GetInnerWork()->Abandon();
 				ReleaseWorkNoLock(Work);
 			}
 		}
@@ -96,7 +121,7 @@ void FQueuedThreadPoolWrapper::Destroy()
 		FRWScopeLock ScopeLock(Lock, SLT_Write);
 		for (FScheduledWork* Work : WorkPool)
 		{
-			check(Work->InnerWork == nullptr);
+			check(Work->GetInnerWork() == nullptr);
 			delete Work;
 		}
 		WorkPool.Empty();
@@ -107,7 +132,7 @@ void FQueuedThreadPoolWrapper::SetMaxConcurrency(int32 InMaxConcurrency)
 {
 	MaxConcurrency = InMaxConcurrency == -1 ? WrappedQueuedThreadPool->GetNumThreads() : InMaxConcurrency;
 
-	// In case we just increased the concurrency, we might need to schedule new tasks
+	// We might need to schedule or unshedule tasks depending on how MaxConcurrency has changed.
 	Schedule();
 }
 
@@ -169,11 +194,57 @@ int32 FQueuedThreadPoolWrapper::GetNumThreads() const
 void FQueuedThreadPoolWrapper::ReleaseWorkNoLock(FScheduledWork* Work)
 {
 	CurrentConcurrency--;
-	ScheduledWork.Remove(Work->InnerWork);
-	Work->InnerWork = nullptr;
+	OnUnscheduled(Work);
+	
+	ScheduledWork.Remove(Work->GetInnerWork());
+	Work->Reset();
 	WorkPool.Push(Work);
 }
 
+bool FQueuedThreadPoolWrapper::CanSchedule() const
+{
+	return (MaxTaskToSchedule == -1 || MaxTaskToSchedule > 0) && CurrentConcurrency < GetMaxConcurrency();
+}
+
+FQueuedThreadPoolWrapper::FScheduledWork* FQueuedThreadPoolWrapper::AllocateWork(IQueuedWork* InnerWork, EQueuedWorkPriority Priority)
+{
+	if (WorkPool.Num() > 0)
+	{
+		FScheduledWork* Work = WorkPool.Pop(false);
+		*Work = FScheduledWork(this, InnerWork, Priority);
+		return Work;
+	}
+	
+	return new FScheduledWork(this, InnerWork, Priority);
+}
+
+bool FQueuedThreadPoolWrapper::TryRetractWorkNoLock(EQueuedWorkPriority InPriority)
+{
+	// Scheduled work is bound by MaxConcurrency which is normally limited by Core count. 
+	// The linear scan should be small and pretty fast.
+	for (TTuple<IQueuedWork*, FScheduledWork*>& Pair : ScheduledWork)
+	{
+		// higher number means lower priority
+		if (Pair.Value->GetPriority() > InPriority)
+		{
+			if (WrappedQueuedThreadPool->RetractQueuedWork(Pair.Value))
+			{
+				QueuedWork.Enqueue(Pair.Key, Pair.Value->GetPriority());
+				ReleaseWorkNoLock(Pair.Value);
+
+				if (MaxTaskToSchedule != -1)
+				{
+					MaxTaskToSchedule++;
+				}
+
+				return true;
+			}
+		}
+	}
+
+	return false;
+}
+
 void FQueuedThreadPoolWrapper::Schedule(FScheduledWork* Work)
 {
 	FRWScopeLock ScopeLock(Lock, SLT_Write);
@@ -186,37 +257,19 @@ void FQueuedThreadPoolWrapper::Schedule(FScheduledWork* Work)
 		ReleaseWorkNoLock(Work);
 	}
 
-	// If a higher priority task comes in, try to retract a lower priority one if possible to make room
-	if (CurrentConcurrency >= GetMaxConcurrency() && (MaxTaskToSchedule == -1 || MaxTaskToSchedule > 0))
+	// If a higher priority task comes in, try to retract lower priority ones if possible to make room
+	EQueuedWorkPriority NextWorkPriority;
+	IQueuedWork* NextWork = QueuedWork.Peek(&NextWorkPriority);
+	if (NextWork)
 	{
-		EQueuedWorkPriority NextWorkPriority;
-		IQueuedWork* NextWork = QueuedWork.Peek(&NextWorkPriority);
-		if (NextWork)
+		// Continue retracting more work until nothing can be retracted anymore or we can finally squeeze the higher priority task in
+		while (!CanSchedule() && TryRetractWorkNoLock(NextWorkPriority))
 		{
-			// Scheduled work is bound by MaxConcurrency which is normally limited by Core count. 
-			// The linear scan should be small and pretty fast.
-			for (TTuple<IQueuedWork*, FScheduledWork*>& Pair : ScheduledWork)
-			{
-				// higher number means lower priority
-				if (Pair.Value->Priority > NextWorkPriority)
-				{
-					if (WrappedQueuedThreadPool->RetractQueuedWork(Pair.Value))
-					{
-						QueuedWork.Enqueue(Pair.Key, Pair.Value->Priority);
-						ReleaseWorkNoLock(Pair.Value);
-
-						if (MaxTaskToSchedule != -1)
-						{
-							MaxTaskToSchedule++;
-						}
-						break;
-					}
-				}
-			}
 		}
 	}
 
-	while (CurrentConcurrency < GetMaxConcurrency() && (MaxTaskToSchedule == -1 || MaxTaskToSchedule > 0))
+	// Schedule as many tasks we can fit
+	while (CanSchedule())
 	{
 		EQueuedWorkPriority WorkPriority;
 		IQueuedWork* InnerWork = QueuedWork.Dequeue(&WorkPriority);
@@ -224,18 +277,10 @@ void FQueuedThreadPoolWrapper::Schedule(FScheduledWork* Work)
 		if (InnerWork)
 		{
 			CurrentConcurrency++;
-			if (WorkPool.Num() > 0)
-			{
-				Work = WorkPool.Pop(false);
-			}
-			else
-			{
-				Work = new FScheduledWork(this);
-			}
-
-			Work->Priority = WorkPriority;
-			Work->InnerWork = InnerWork;
+			
+			Work = AllocateWork(InnerWork, WorkPriority);
 			ScheduledWork.Add(InnerWork, Work);
+			OnScheduled(Work);
 			WrappedQueuedThreadPool->AddQueuedWork(Work, PriorityMapper(WorkPriority));
 
 			if (MaxTaskToSchedule > 0)
diff --git a/Engine/Source/Runtime/Core/Public/Async/AsyncWork.h b/Engine/Source/Runtime/Core/Public/Async/AsyncWork.h
--- a/Engine/Source/Runtime/Core/Public/Async/AsyncWork.h
+++ b/Engine/Source/Runtime/Core/Public/Async/AsyncWork.h
@@ -229,6 +229,8 @@ class FAsyncTask
 	EQueuedWorkPriority Priority;
 	/** Current flags */
 	EQueuedWorkFlags Flags;
+	/** Approximation of the peak memory (in bytes) this task could require during it's execution. */
+	int64 RequiredMemory = -1;
 	/** optional LLM tag */
 	LLM(const UE::LLMPrivate::FTagData* InheritedLLMTag);
 	/** Memory trace tag */
@@ -252,7 +254,7 @@ class FAsyncTask
 	/* Generic start function, not called directly
 		* @param bForceSynchronous if true, this job will be started synchronously, now, on this thread
 	**/
-	void Start(bool bForceSynchronous, FQueuedThreadPool* InQueuedPool, EQueuedWorkPriority InQueuedWorkPriority, EQueuedWorkFlags InQueuedWorkFlags)
+	void Start(bool bForceSynchronous, FQueuedThreadPool* InQueuedPool, EQueuedWorkPriority InQueuedWorkPriority, EQueuedWorkFlags InQueuedWorkFlags, int64 InRequiredMemory)
 	{
 		FScopeCycleCounter Scope( Task.GetStatId(), true );
 		DECLARE_SCOPE_CYCLE_COUNTER( TEXT( "FAsyncTask::Start" ), STAT_FAsyncTask_Start, STATGROUP_ThreadPoolAsyncTasks );
@@ -260,6 +262,7 @@ class FAsyncTask
 #if USE_MEMORY_TRACE_TAGS
 		InheritedTraceTag = MemoryTrace_GetActiveTag();
 #endif
+		RequiredMemory = InRequiredMemory;
 
 		FPlatformMisc::MemoryBarrier();
 		CheckIdle();  // can't start a job twice without it being completed first
@@ -431,21 +434,29 @@ public:
 		return Task;
 	}
 
+	/**
+	 * Returns an approximation of the peak memory (in bytes) this task could require during it's execution.
+	 **/
+	int64 GetRequiredMemory() const override
+	{
+		return RequiredMemory;
+	}
+
 	/** 
 	* Run this task on this thread
 	* @param bDoNow if true then do the job now instead of at EnsureCompletion
 	**/
-	void StartSynchronousTask(EQueuedWorkPriority InQueuedWorkPriority = EQueuedWorkPriority::Normal, EQueuedWorkFlags InQueuedWorkFlags = EQueuedWorkFlags::None)
+	void StartSynchronousTask(EQueuedWorkPriority InQueuedWorkPriority = EQueuedWorkPriority::Normal, EQueuedWorkFlags InQueuedWorkFlags = EQueuedWorkFlags::None, int64 InRequiredMemory = -1)
 	{
-		Start(true, GThreadPool, InQueuedWorkPriority, InQueuedWorkFlags);
+		Start(true, GThreadPool, InQueuedWorkPriority, InQueuedWorkFlags, InRequiredMemory);
 	}
 
 	/** 
 	* Queue this task for processing by the background thread pool
 	**/
-	void StartBackgroundTask(FQueuedThreadPool* InQueuedPool = GThreadPool, EQueuedWorkPriority InQueuedWorkPriority = EQueuedWorkPriority::Normal, EQueuedWorkFlags InQueuedWorkFlags = EQueuedWorkFlags::None)
+	void StartBackgroundTask(FQueuedThreadPool* InQueuedPool = GThreadPool, EQueuedWorkPriority InQueuedWorkPriority = EQueuedWorkPriority::Normal, EQueuedWorkFlags InQueuedWorkFlags = EQueuedWorkFlags::None, int64 InRequiredMemory = -1)
 	{
-		Start(false, InQueuedPool, InQueuedWorkPriority, InQueuedWorkFlags);
+		Start(false, InQueuedPool, InQueuedWorkPriority, InQueuedWorkFlags, InRequiredMemory);
 	}
 
 	/** 
diff --git a/Engine/Source/Runtime/Core/Public/Misc/IQueuedWork.h b/Engine/Source/Runtime/Core/Public/Misc/IQueuedWork.h
--- a/Engine/Source/Runtime/Core/Public/Misc/IQueuedWork.h
+++ b/Engine/Source/Runtime/Core/Public/Misc/IQueuedWork.h
@@ -79,6 +79,11 @@ public:
 	 * Returns any special work flags.
 	 */
 	virtual EQueuedWorkFlags GetQueuedWorkFlags() const { return EQueuedWorkFlags::None; }
+
+	/**
+	 * Returns an approximation of the peak memory (in bytes) this task could require during it's execution.
+	 */
+	virtual int64 GetRequiredMemory() const { return -1 /* Negative value means unknown */; }
 public:
 
 	/**
diff --git a/Engine/Source/Runtime/Core/Public/Misc/QueuedThreadPool.h b/Engine/Source/Runtime/Core/Public/Misc/QueuedThreadPool.h
--- a/Engine/Source/Runtime/Core/Public/Misc/QueuedThreadPool.h
+++ b/Engine/Source/Runtime/Core/Public/Misc/QueuedThreadPool.h
@@ -66,7 +66,7 @@ public:
 	/**
 	 * Get the next work item in priority order without actually dequeuing.
 	 */
-	IQueuedWork* Peek(EQueuedWorkPriority* OutDequeuedWorkPriority = nullptr);
+	IQueuedWork* Peek(EQueuedWorkPriority* OutDequeuedWorkPriority = nullptr) const;
 
 	/**
 	 * Empty the queue.
diff --git a/Engine/Source/Runtime/Core/Public/Misc/QueuedThreadPoolWrapper.h b/Engine/Source/Runtime/Core/Public/Misc/QueuedThreadPoolWrapper.h
--- a/Engine/Source/Runtime/Core/Public/Misc/QueuedThreadPoolWrapper.h
+++ b/Engine/Source/Runtime/Core/Public/Misc/QueuedThreadPoolWrapper.h
@@ -47,29 +47,38 @@ public:
 	void AddQueuedWork(IQueuedWork* InQueuedWork, EQueuedWorkPriority InPriority = EQueuedWorkPriority::Normal) override;
 	bool RetractQueuedWork(IQueuedWork* InQueuedWork) override;
 	int32 GetNumThreads() const override;
+	int32 GetCurrentConcurrency() const { return CurrentConcurrency.load(std::memory_order_relaxed); }
 
 protected:
 	FRWLock Lock;
 	FThreadPoolPriorityQueue QueuedWork;
 
 	// Can be overriden to dynamically control the maximum concurrency
-	virtual int32 GetMaxConcurrency() const { return MaxConcurrency.Load(EMemoryOrder::Relaxed); }
+	virtual int32 GetMaxConcurrency() const { return MaxConcurrency.load(std::memory_order_relaxed); }
+	
+	// Can be overriden to know when work has been scheduled.
+	virtual void OnScheduled(const IQueuedWork*) {}
+	
+	// Can be overriden to know when work has been unscheduled.
+	virtual void OnUnscheduled(const IQueuedWork*) {}
 private:
 	struct FScheduledWork;
-
+	FScheduledWork* AllocateWork(IQueuedWork* InnerWork, EQueuedWorkPriority Priority);
+	bool CanSchedule() const;
 	bool Create(uint32 InNumQueuedThreads, uint32 StackSize, EThreadPriority ThreadPriority, const TCHAR* Name) override;
 	void Destroy() override;
 	void Schedule(FScheduledWork* Work = nullptr);
 	void ReleaseWorkNoLock(FScheduledWork* Work);
+	bool TryRetractWorkNoLock(EQueuedWorkPriority InPriority);
 	
 	TFunction<EQueuedWorkPriority(EQueuedWorkPriority)> PriorityMapper;
 
 	FQueuedThreadPool* WrappedQueuedThreadPool;
 	TArray<FScheduledWork*> WorkPool;
 	TMap<IQueuedWork*, FScheduledWork*> ScheduledWork;
-	TAtomic<int32> MaxConcurrency;
+	std::atomic<int32> MaxConcurrency;
 	int32 MaxTaskToSchedule;
-	TAtomic<int32> CurrentConcurrency;
+	std::atomic<int32> CurrentConcurrency;
 	EQueuedWorkPriority WrappedQueuePriority;
 };
 
diff --git a/Engine/Source/Runtime/Engine/Classes/Engine/StaticMesh.h b/Engine/Source/Runtime/Engine/Classes/Engine/StaticMesh.h
--- a/Engine/Source/Runtime/Engine/Classes/Engine/StaticMesh.h
+++ b/Engine/Source/Runtime/Engine/Classes/Engine/StaticMesh.h
@@ -1781,6 +1781,11 @@ private:
 	 */
 	void FinishBuildInternal(const TArray<UStaticMeshComponent*>& InAffectedComponents, bool bHasRenderDataChanged, bool bShouldComputeExtendedBounds = true);
 
+	/**
+	 * Get an estimate of the peak amount of memory required to build this mesh.
+	 */
+	int64 GetBuildRequiredMemory() const;
+
 #if WITH_EDITORONLY_DATA
 	/**
 	 * Deserialize MeshDescription for the specified LodIndex from BulkData, DDC or RawMesh.
diff --git a/Engine/Source/Runtime/Engine/Classes/Engine/Texture.h b/Engine/Source/Runtime/Engine/Classes/Engine/Texture.h
--- a/Engine/Source/Runtime/Engine/Classes/Engine/Texture.h
+++ b/Engine/Source/Runtime/Engine/Classes/Engine/Texture.h
@@ -1020,6 +1020,11 @@ public:
 	 */
 	ENGINE_API void ForceRebuildPlatformData();
 
+	/**
+	 * Get an estimate of the peak amount of memory required to build this texture.
+	 */
+	ENGINE_API int64 GetBuildRequiredMemory() const;
+
 	/**
 	 * Marks platform data as transient. This optionally removes persistent or cached data associated with the platform.
 	 */
diff --git a/Engine/Source/Runtime/Engine/Private/AssetCompilingManager.cpp b/Engine/Source/Runtime/Engine/Private/AssetCompilingManager.cpp
--- a/Engine/Source/Runtime/Engine/Private/AssetCompilingManager.cpp
+++ b/Engine/Source/Runtime/Engine/Private/AssetCompilingManager.cpp
@@ -38,11 +38,19 @@ static AsyncCompilationHelpers::FAsyncCompilationStandardCVars CVarAsyncCompilat
 static TAutoConsoleVariable<int32> CVarAsyncAssetCompilationMemoryPerCore(
 	TEXT("Editor.AsyncAssetCompilationMemoryPerCore"),
 	4,
-	TEXT("0 - No memory limit per core.\n")
+	TEXT("0 - No memory limit per asset.\n")
 	TEXT("N - Dynamically adjust concurrency limit by dividing free system memory by this number (in GB).\n")
 	TEXT("Limit concurrency for async processing based on RAM available.\n"),
 	ECVF_Default);
 
+static TAutoConsoleVariable<int32> CVarAsyncAssetCompilationMaxMemoryUsage(
+	TEXT("Editor.AsyncAssetCompilationMaxMemoryUsage"),
+	0,
+	TEXT("0 - No hard memory limit, will be tuned against system available memory (recommended default).\n")
+	TEXT("N - Try to limit total memory usage for asset compilation to this amount (in GB).\n")
+	TEXT("Try to stay under specified memory limit for asset compilation by reducing concurrency when under memory pressure.\n"),
+	ECVF_Default);
+
 namespace AssetCompilingManagerImpl
 {
 	static void EnsureInitializedCVars()
@@ -74,6 +82,9 @@ int32 FAssetCompilingManager::GetNumRemainingAssets() const
 }
 
 TRACE_DECLARE_INT_COUNTER(AsyncCompilationMaxConcurrency, TEXT("AsyncCompilation/MaxConcurrency"));
+TRACE_DECLARE_INT_COUNTER(AsyncCompilationConcurrency, TEXT("AsyncCompilation/Concurrency"));
+TRACE_DECLARE_MEMORY_COUNTER(AsyncCompilationTotalMemoryLimit, TEXT("AsyncCompilation/TotalMemoryLimit"));
+TRACE_DECLARE_MEMORY_COUNTER(AsyncCompilationTotalEstimatedMemory, TEXT("AsyncCompilation/TotalEstimatedMemory"));
 
 class FMemoryBoundQueuedThreadPoolWrapper : public FQueuedThreadPoolWrapper
 {
@@ -88,22 +99,91 @@ public:
 	{
 	}
 
+	int64 GetHardMemoryLimit() const
+	{
+		return (int64)CVarAsyncAssetCompilationMaxMemoryUsage.GetValueOnAnyThread() * 1024 * 1024 * 1024;
+	}
+
+	int64 GetDefaultMemoryPerAsset() const
+	{
+		return FMath::Max(0ll, (int64)CVarAsyncAssetCompilationMemoryPerCore.GetValueOnAnyThread(false) * 1024 * 1024 * 1024);
+	}
+
+	int64 GetMemoryLimit() const
+	{
+		// Just make sure available physical fits in a int64, if it's bigger than that, we're not expecting to be memory limited anyway
+		int64 AvailablePhysical = (int64)FMath::Min(FPlatformMemory::GetStats().AvailablePhysical, (uint64)INT64_MAX);
+
+		int64 HardMemoryLimit = GetHardMemoryLimit();
+		if (HardMemoryLimit > 0)
+		{
+			return FMath::Min(HardMemoryLimit, AvailablePhysical);
+		}
+		
+		return AvailablePhysical;
+	}
+
+	int64 GetRequiredMemory(const IQueuedWork* InWork) const
+	{
+		int64 RequiredMemory = InWork->GetRequiredMemory();
+
+		// If the task doesn't support a memory estimate, use global defined default
+		if (RequiredMemory == -1)
+		{
+			RequiredMemory = GetDefaultMemoryPerAsset();
+		}
+
+		return RequiredMemory;
+	}
+
+	void UpdateCounters()
+	{
+		TRACE_COUNTER_SET(AsyncCompilationTotalMemoryLimit, GetMemoryLimit());
+		TRACE_COUNTER_SET(AsyncCompilationTotalEstimatedMemory, TotalEstimatedMemory);
+		TRACE_COUNTER_SET(AsyncCompilationConcurrency, GetCurrentConcurrency());
+	}
+
+	void OnScheduled(const IQueuedWork* InWork) override
+	{
+		TotalEstimatedMemory += GetRequiredMemory(InWork);
+		UpdateCounters();
+	}
+
+	void OnUnscheduled(const IQueuedWork* InWork) override
+	{
+		TotalEstimatedMemory -= GetRequiredMemory(InWork);
+		check(TotalEstimatedMemory >= 0);
+		UpdateCounters();
+	}
+
 	int32 GetMaxConcurrency() const override
 	{
-		int32 MemoryPerCore = CVarAsyncAssetCompilationMemoryPerCore.GetValueOnAnyThread(false);
+		int64 RequiredMemory = TotalEstimatedMemory;
+		
+		// Add next work memory requirement to see if it still fits in available memory
+		if (IQueuedWork* NextWork = QueuedWork.Peek())
+		{
+			RequiredMemory += GetRequiredMemory(NextWork);
+		}
+
 		int32 DynamicMaxConcurrency = FQueuedThreadPoolWrapper::GetMaxConcurrency();
+		int32 Concurrency = FQueuedThreadPoolWrapper::GetCurrentConcurrency();
 
-		if (MemoryPerCore > 0)
+		// Never limit below a concurrency of 1 or else we'll starve the asset processing and never be scheduled again
+		// The idea for now is to let assets get built one by one when starving on memory, which should not increase OOM compared to the old synchronous behavior.
+		if (RequiredMemory > 0 && Concurrency > 0 && RequiredMemory >= GetMemoryLimit())
 		{
-			FPlatformMemoryStats MemoryStats = FPlatformMemory::GetStats();
-			uint64 MemoryMaxConcurrency = FMath::Max(1llu, MemoryStats.AvailablePhysical / ((uint64)MemoryPerCore * 1024 * 1024 * 1024));
-			DynamicMaxConcurrency = FMath::Min((int32)MemoryMaxConcurrency, DynamicMaxConcurrency);
-
-			TRACE_COUNTER_SET(AsyncCompilationMaxConcurrency, DynamicMaxConcurrency);
+			DynamicMaxConcurrency = Concurrency; // Limit concurrency to what we already allowed
 		}
 
+		TRACE_COUNTER_SET(AsyncCompilationMaxConcurrency, DynamicMaxConcurrency);
+
+		check(DynamicMaxConcurrency > 0);
 		return DynamicMaxConcurrency;
 	}
+
+private:
+	std::atomic<int64> TotalEstimatedMemory {0};
 };
 
 FQueuedThreadPool* FAssetCompilingManager::GetThreadPool() const
@@ -174,4 +254,4 @@ void FAssetCompilingManager::ProcessAsyncTasks(bool bLimitExecutionTime)
 }
 
 #undef LOCTEXT_NAMESPACE
-#endif // #if WITH_EDITOR
+#endif // #if WITH_EDITOR
\ No newline at end of file
diff --git a/Engine/Source/Runtime/Engine/Private/StaticMesh.cpp b/Engine/Source/Runtime/Engine/Private/StaticMesh.cpp
--- a/Engine/Source/Runtime/Engine/Private/StaticMesh.cpp
+++ b/Engine/Source/Runtime/Engine/Private/StaticMesh.cpp
@@ -45,6 +45,7 @@
 #include "EngineGlobals.h"
 #include "HAL/LowLevelMemTracker.h"
 #include "DynamicMeshBuilder.h"
+#include "Math/UnrealMathUtility.h"
 #include "Model.h"
 #include "Async/Async.h"
 #include "SplineMeshSceneProxy.h"
@@ -5029,6 +5030,35 @@ void UStaticMesh::ReleaseAsyncProperty(EStaticMeshAsyncProperties AsyncPropertie
 	LockedProperties &= ~(uint32)AsyncProperties;
 }
 
+int64 UStaticMesh::GetBuildRequiredMemory() const
+{
+	// We have to base our estimate on something accessible and known before the build, for now use the biggest bulk data size.
+	int64 BiggestBulkDataSize = -1;
+	if (GetHiResSourceModel().GetMeshDescriptionBulkData())
+	{
+		BiggestBulkDataSize = FMath::Max(BiggestBulkDataSize, GetHiResSourceModel().GetMeshDescriptionBulkData()->GetBulkDataSize());
+	}
+
+	for (const FStaticMeshSourceModel& Model : GetSourceModels())
+	{
+		if (Model.GetMeshDescriptionBulkData())
+		{
+			BiggestBulkDataSize = FMath::Max(BiggestBulkDataSize, Model.GetMeshDescriptionBulkData()->GetBulkDataSize());
+		}
+
+		if (Model.RawMeshBulkData)
+		{
+			BiggestBulkDataSize = FMath::Max(BiggestBulkDataSize, Model.RawMeshBulkData->GetBulkData().GetBulkDataSize());
+		}
+	}
+
+	// Rough estimate of the memory that should be required to build that mesh.
+	// Used -ddc=cold on the command line and opened some big meshes from the content browser one by one while
+	// monitoring the editor commit memory spike generated by the compilation.
+	// For smaller assets, -llm with LLM tags have been used and came down to mostly the same ratio.
+	return BiggestBulkDataSize * 7;
+}
+
 #endif // WITH_EDITOR
 
 //
@@ -5081,7 +5111,7 @@ void UStaticMesh::PostLoad()
 
 		Context.StaticMeshOwner.Reset(GetPackage());
 		AsyncTask = MakeUnique<FStaticMeshAsyncBuildTask>(this, MakeUnique<FStaticMeshPostLoadContext>(MoveTemp(Context)));
-		AsyncTask->StartBackgroundTask(StaticMeshThreadPool, BasePriority);
+		AsyncTask->StartBackgroundTask(StaticMeshThreadPool, BasePriority, EQueuedWorkFlags::None, GetBuildRequiredMemory());
 		FStaticMeshCompilingManager::Get().AddStaticMeshes({this});
 	}
 	else
diff --git a/Engine/Source/Runtime/Engine/Private/StaticMeshBuild.cpp b/Engine/Source/Runtime/Engine/Private/StaticMeshBuild.cpp
--- a/Engine/Source/Runtime/Engine/Private/StaticMeshBuild.cpp
+++ b/Engine/Source/Runtime/Engine/Private/StaticMeshBuild.cpp
@@ -254,7 +254,7 @@ void UStaticMesh::BatchBuild(const TArray<UStaticMesh*>& InStaticMeshes, bool bI
 					EQueuedWorkPriority BasePriority = FStaticMeshCompilingManager::Get().GetBasePriority(StaticMesh);
 					check(StaticMesh->AsyncTask == nullptr);
 					StaticMesh->AsyncTask = MakeUnique<FStaticMeshAsyncBuildTask>(StaticMesh, MoveTemp(Context));
-					StaticMesh->AsyncTask->StartBackgroundTask(StaticMeshThreadPool, BasePriority);
+					StaticMesh->AsyncTask->StartBackgroundTask(StaticMeshThreadPool, BasePriority, EQueuedWorkFlags::None, StaticMesh->GetBuildRequiredMemory());
 					FStaticMeshCompilingManager::Get().AddStaticMeshes({ StaticMesh });
 					return true;
 				}
diff --git a/Engine/Source/Runtime/Engine/Private/Texture.cpp b/Engine/Source/Runtime/Engine/Private/Texture.cpp
--- a/Engine/Source/Runtime/Engine/Private/Texture.cpp
+++ b/Engine/Source/Runtime/Engine/Private/Texture.cpp
@@ -1593,6 +1593,30 @@ void UTexture::SetLayerFormatSettings(int32 LayerIndex, const FTextureFormatSett
 	}
 }
 
+int64 UTexture::GetBuildRequiredMemory() const
+{
+	// Compute the memory it should take to uncompress the bulkdata in memory
+	int64 MemoryEstimate = 0;
+
+	// Compute the amount of memory necessary to uncompress the bulkdata in memory
+	for (int32 BlockIndex = 0; BlockIndex < Source.GetNumBlocks(); ++BlockIndex)
+	{
+		FTextureSourceBlock SourceBlock;
+		Source.GetBlock(BlockIndex, SourceBlock);
+
+		for (int32 LayerIndex = 0; LayerIndex < Source.GetNumLayers(); ++LayerIndex)
+		{
+			for (int32 MipIndex = 0; MipIndex < SourceBlock.NumMips; ++MipIndex)
+			{
+				MemoryEstimate += Source.CalcMipSize(BlockIndex, LayerIndex, MipIndex);
+			}
+		}
+	}
+
+	// Account for the multiple copies that are currently carried over during the compression phase
+	return MemoryEstimate <= 0 ? -1 /* Unknown */ : MemoryEstimate * 5;
+}
+
 #endif // #if WITH_EDITOR
 
 FName GetDefaultTextureFormatName( const ITargetPlatform* TargetPlatform, const UTexture* Texture, int32 LayerIndex, bool bSupportDX11TextureFormats, bool bSupportCompressedVolumeTexture, int32 BlockSize )
diff --git a/Engine/Source/Runtime/Engine/Private/TextureDerivedData.cpp b/Engine/Source/Runtime/Engine/Private/TextureDerivedData.cpp
--- a/Engine/Source/Runtime/Engine/Private/TextureDerivedData.cpp
+++ b/Engine/Source/Runtime/Engine/Private/TextureDerivedData.cpp
@@ -909,7 +909,7 @@ void FTexturePlatformData::Cache(
 		EQueuedWorkPriority BasePriority      = FTextureCompilingManager::Get().GetBasePriority(&InTexture);
 
 		AsyncTask = new FTextureAsyncCacheDerivedDataTask(Compressor, this, &InTexture, InSettingsPerLayer, Flags);
-		AsyncTask->StartBackgroundTask(TextureThreadPool, BasePriority);
+		AsyncTask->StartBackgroundTask(TextureThreadPool, BasePriority, EQueuedWorkFlags::None, AsyncTask->GetTask().GetRequiredMemoryEstimate());
 	}
 	else
 	{
diff --git a/Engine/Source/Runtime/Engine/Private/TextureDerivedDataTask.cpp b/Engine/Source/Runtime/Engine/Private/TextureDerivedDataTask.cpp
--- a/Engine/Source/Runtime/Engine/Private/TextureDerivedDataTask.cpp
+++ b/Engine/Source/Runtime/Engine/Private/TextureDerivedDataTask.cpp
@@ -256,8 +256,13 @@ void FTextureCacheDerivedDataWorker::BuildTexture(bool bReplaceExistingDDC)
 	Args.Add(TEXT("TextureFormatName"), FText::FromString(BuildSettingsPerLayer[0].TextureFormatName.GetPlainNameString()));
 	Args.Add(TEXT("TextureResolutionX"), FText::FromString(FString::FromInt(bHasValidMip0 ? TextureData.Blocks[0].MipsPerLayer[0][0].SizeX : 0)));
 	Args.Add(TEXT("TextureResolutionY"), FText::FromString(FString::FromInt(bHasValidMip0 ? TextureData.Blocks[0].MipsPerLayer[0][0].SizeY : 0)));
+	Args.Add(TEXT("EstimatedMemory"), FText::FromString(FString::SanitizeFloat(double(RequiredMemoryEstimate) / (1024.0*1024.0), 3)));
 
-	FTextureStatusMessageContext StatusMessage(FText::Format(NSLOCTEXT("Engine", "BuildTextureStatus", "Building textures: {TextureName} ({TextureFormatName}, {TextureResolutionX}X{TextureResolutionY})"), Args));
+	FTextureStatusMessageContext StatusMessage(
+		FText::Format(
+			NSLOCTEXT("Engine", "BuildTextureStatus", "Building textures: {TextureName} ({TextureFormatName}, {TextureResolutionX}X{TextureResolutionY}) (Required Memory Estimate: {EstimatedMemory} MB)"), 
+			Args
+		));
 
 	if (!ensure(Compressor))
 	{
@@ -413,6 +418,7 @@ FTextureCacheDerivedDataWorker::FTextureCacheDerivedDataWorker(
 	, DerivedData(InDerivedData)
 	, Texture(*InTexture)
 	, CacheFlags(InCacheFlags)
+	, RequiredMemoryEstimate(InTexture->GetBuildRequiredMemory())
 	, bSucceeded(false)
 {
 	check(DerivedData);
diff --git a/Engine/Source/Runtime/Engine/Private/TextureDerivedDataTask.h b/Engine/Source/Runtime/Engine/Private/TextureDerivedDataTask.h
--- a/Engine/Source/Runtime/Engine/Private/TextureDerivedDataTask.h
+++ b/Engine/Source/Runtime/Engine/Private/TextureDerivedDataTask.h
@@ -125,6 +125,8 @@ class FTextureCacheDerivedDataWorker : public FNonAbandonableTask
 	uint32 CacheFlags;
 	/** Have many bytes were loaded from DDC or built (for telemetry) */
 	uint32 BytesCached = 0;
+	/** Estimate of the peak amount of memory required to complete this task. */
+	int64 RequiredMemoryEstimate = -1;
 
 	/** true if caching has succeeded. */
 	bool bSucceeded;
@@ -147,7 +149,6 @@ public:
 	/** Does the work to cache derived data. Safe to call from any thread. */
 	void DoWork();
 
-
 	/** Finalize work. Must be called ONLY by the game thread! */
 	void Finalize();
 
@@ -157,6 +158,12 @@ public:
 		return BytesCached;
 	}
 
+	/** Estimate of the peak amount of memory required to complete this task. */
+	int64 GetRequiredMemoryEstimate() const
+	{
+		return RequiredMemoryEstimate;
+	}
+
 	/** Expose how the resource was returned for telemetry. */
 	bool WasLoadedFromDDC() const
 	{
diff --git a/Engine/Source/Runtime/MeshDescription/Public/MeshDescription.h b/Engine/Source/Runtime/MeshDescription/Public/MeshDescription.h
--- a/Engine/Source/Runtime/MeshDescription/Public/MeshDescription.h
+++ b/Engine/Source/Runtime/MeshDescription/Public/MeshDescription.h
@@ -1318,6 +1318,9 @@ public:
 	/** Uses a hash as the GUID, useful to prevent recomputing content already in cache. */
 	void UseHashAsGuid();
 
+	/** Gets the size of the serialized bulk data */
+	int64 GetBulkDataSize() { return BulkData.GetBulkDataSize(); }
+
 private:
 #if WITH_EDITOR
 	/** Protects simultaneous access to BulkData */
