MemStack remove unused MinMarksToAlloc fixes Debug build failing to pack
Error: Assertion failed: NumMarks >= MinMarksToAlloc [File:D:\ue5_main\Engine\Source\Runtime\Core\Public\Misc/MemStack.h] [Line: 108]

MinMarksToAlloc was only used for checkSlow and wasn't being used by any clients
default value of 1 causes this assertion to fail when NumMarks = 0

diff --git a/Engine/Source/Runtime/Core/Public/Misc/MemStack.h b/Engine/Source/Runtime/Core/Public/Misc/MemStack.h
--- a/Engine/Source/Runtime/Core/Public/Misc/MemStack.h
+++ b/Engine/Source/Runtime/Core/Public/Misc/MemStack.h
@@ -69,25 +69,21 @@ private:
 class CORE_API FMemStackBase //-V1062
 {
 public:
-#if ( PLATFORM_WINDOWS && defined(__clang__) )
-	FMemStackBase()			// @todo clang: parameterless constructor is needed to prevent an ICE in clang
-		: FMemStackBase(1)	// https://llvm.org/bugs/show_bug.cgi?id=28137
-	{
-	}
 
-	FMemStackBase(int32 InMinMarksToAlloc)
-#else
-	FMemStackBase(int32 InMinMarksToAlloc = 1)
-#endif
+	FMemStackBase()
 		: Top(nullptr)
 		, End(nullptr)
 		, TopChunk(nullptr)
 		, TopMark(nullptr)
 		, NumMarks(0)
-		, MinMarksToAlloc(InMinMarksToAlloc)
 	{
 	}
 
+	FMemStackBase(int unused) : FMemStackBase()
+	{
+		// @todo : remove uses of this constructor
+	}
+
 	~FMemStackBase()
 	{
 		check((GIsCriticalError || !NumMarks));
@@ -105,8 +101,6 @@ public:
 		checkSlow(AllocSize>=0);
 		checkSlow((Alignment&(Alignment-1))==0);
 		checkSlow(Top<=End);
-		checkSlow(NumMarks >= MinMarksToAlloc);
-
 
 		// Try to get memory from the current chunk.
 		uint8* Result = Align( Top, Alignment );
@@ -136,7 +130,7 @@ public:
 
 	FORCEINLINE void Flush()
 	{
-		check(!NumMarks && !MinMarksToAlloc);
+		check(!NumMarks);
 		FreeChunks(nullptr);
 	}
 	FORCEINLINE int32 GetNumMarks()
@@ -192,8 +186,6 @@ private:
 	/** The number of marks on this stack. */
 	int32 NumMarks;
 
-	/** Used for a checkSlow. Most stacks require a mark to allocate. Command lists don't because they never mark, only flush*/
-	int32 MinMarksToAlloc;
 };
 
 
