Fix variadic template build errors caused by workaround for compiler bug in prior VS2019 compiler versions.

diff --git a/Engine/Source/Runtime/Core/Public/Misc/TVariantMeta.h b/Engine/Source/Runtime/Core/Public/Misc/TVariantMeta.h
--- a/Engine/Source/Runtime/Core/Public/Misc/TVariantMeta.h
+++ b/Engine/Source/Runtime/Core/Public/Misc/TVariantMeta.h
@@ -14,7 +14,9 @@
 // Due to a bug in Visual Studio, we must use a recursive template to determine max sizeof() and alignof() of types in a template parameter pack.
 // On other compilers, we use a constexpr array and pluck out the largest
 // Bug reported to Microsoft https://developercommunity.visualstudio.com/content/problem/528990/constexpr-expansion-inside-a-lambda-fails-to-be-ev.html
-#if defined(_MSC_VER) && !defined(__clang__)
+// 
+// Bug has since been fixed in Visual Studio 2019 Update 1: _MSC_VER 1921
+#if defined(_MSC_VER) && _MSC_VER < 1921 && !defined(__clang__)
 #define TVARIANT_STORAGE_USE_RECURSIVE_TEMPLATE 1
 #else
 #define TVARIANT_STORAGE_USE_RECURSIVE_TEMPLATE 0
