TVariantMeta.h workaround for VC 16.0 removed, as UE5 requires 16.4 at minimum.

diff --git a/Engine/Source/Runtime/Core/Public/Misc/TVariantMeta.h b/Engine/Source/Runtime/Core/Public/Misc/TVariantMeta.h
--- a/Engine/Source/Runtime/Core/Public/Misc/TVariantMeta.h
+++ b/Engine/Source/Runtime/Core/Public/Misc/TVariantMeta.h
@@ -11,17 +11,6 @@
 
 #include "Misc/AssertionMacros.h"
 
-// Due to a bug in Visual Studio, we must use a recursive template to determine max sizeof() and alignof() of types in a template parameter pack.
-// On other compilers, we use a constexpr array and pluck out the largest
-// Bug reported to Microsoft https://developercommunity.visualstudio.com/content/problem/528990/constexpr-expansion-inside-a-lambda-fails-to-be-ev.html
-// 
-// Bug has since been fixed in Visual Studio 2019 Update 1: _MSC_VER 1921
-#if defined(_MSC_VER) && _MSC_VER < 1921 && !defined(__clang__)
-#define TVARIANT_STORAGE_USE_RECURSIVE_TEMPLATE 1
-#else
-#define TVARIANT_STORAGE_USE_RECURSIVE_TEMPLATE 0
-#endif
-
 template <typename T, typename... Ts>
 class TVariant;
 
@@ -72,67 +61,6 @@ namespace Private
 		static constexpr bool Value = TOr<TIsReferenceType<Ts>...>::Value;
 	};
 
-#if TVARIANT_STORAGE_USE_RECURSIVE_TEMPLATE
-	/** Determine the max alignof and sizeof of all types in a template parameter pack */
-	template <typename... Ts>
-	struct TVariantStorageTraits;
-
-	template <typename T, typename... Ts>
-	struct TVariantStorageTraits<T, Ts...>
-	{
-		static constexpr SIZE_T MaxSizeof(SIZE_T CurrentSize)
-		{
-			return TVariantStorageTraits<Ts...>::MaxSizeof(TVariantStorageTraits<T>::MaxSizeof(CurrentSize));
-		}
-
-		static constexpr SIZE_T MaxAlignof(SIZE_T CurrentSize)
-		{
-			return TVariantStorageTraits<Ts...>::MaxAlignof(TVariantStorageTraits<T>::MaxAlignof(CurrentSize));
-		}
-	};
-
-	template <typename T>
-	struct TVariantStorageTraits<T>
-	{
-		static constexpr SIZE_T MaxSizeof(SIZE_T CurrentSize)
-		{
-			return CurrentSize > sizeof(T) ? CurrentSize : sizeof(T);
-		}
-
-		static constexpr SIZE_T MaxAlignof(SIZE_T CurrentSize)
-		{
-			return CurrentSize > alignof(T) ? CurrentSize : alignof(T);
-		}
-	};
-
-	/** Expose a type that is suitable for storing any of the types in a template parameter pack */
-	template <typename T, typename... Ts>
-	struct TVariantStorage
-	{
-		static constexpr SIZE_T SizeofValue = TVariantStorageTraits<T, Ts...>::MaxSizeof(0);
-		static constexpr SIZE_T AlignofValue = TVariantStorageTraits<T, Ts...>::MaxAlignof(0);
-		static_assert(SizeofValue > 0, "MaxSizeof must be greater than 0");
-		static_assert(AlignofValue > 0, "MaxAlignof must be greater than 0");
-
-		/** Interpret the underlying data as the type in the variant parameter pack at the compile-time index. This function is used to implement Visit and should not be used directly */
-		template <SIZE_T N>
-		auto& GetValueAsIndexedType()
-		{
-			using ReturnType = typename TNthTypeFromParameterPack<N, T, Ts...>::Type;
-			return *reinterpret_cast<ReturnType*>(&Storage);
-		}
-
-		/** Interpret the underlying data as the type in the variant parameter pack at the compile-time index. This function is used to implement Visit and should not be used directly */
-		template <SIZE_T N>
-		const auto& GetValueAsIndexedType() const
-		{
-			// Temporarily remove the const qualifier so we can implement GetValueAsIndexedType in one location.
-			return const_cast<TVariantStorage*>(this)->template GetValueAsIndexedType<N>();
-		}
-
-		TAlignedBytes<SizeofValue, AlignofValue> Storage;
-	};
-#else
 	/** Determine the max alignof and sizeof of all types in a template parameter pack and provide a type that is compatible with those sizes */
 	template <typename... Ts>
 	struct TVariantStorage
@@ -183,7 +111,6 @@ namespace Private
 
 		TAlignedBytes<SizeofValue, AlignofValue> Storage;
 	};
-#endif
 
 	/** Helper to lookup indices of each type in a template parameter pack */
 	template <SIZE_T N, typename LookupType, typename... Ts>
