Ensure that ::operator new always allocates at least one byte, even if Malloc doesn't.
Replacement operator/delete overloads for C++17 which are alignment-aware.

diff --git a/Engine/Source/Runtime/Core/Public/Modules/Boilerplate/ModuleBoilerplate.h b/Engine/Source/Runtime/Core/Public/Modules/Boilerplate/ModuleBoilerplate.h
--- a/Engine/Source/Runtime/Core/Public/Modules/Boilerplate/ModuleBoilerplate.h
+++ b/Engine/Source/Runtime/Core/Public/Modules/Boilerplate/ModuleBoilerplate.h
@@ -17,6 +17,10 @@ class FChunkedFixedUObjectArray;
  * Override new + delete operators (and array versions) in every module.
  * This prevents the possibility of mismatched new/delete calls such as a new[] that
  * uses Unreal's allocator and a delete[] that uses the system allocator.
+ *
+ * Overloads have to guarantee at least 1 byte is allocated because
+ * otherwise new T[0] could return a null pointer, as could ::operator new(0), depending
+ * on the allocator (e.g. TBB), which is non-standard behaviour.
  */
 #if USING_CODE_ANALYSIS
 	#define OPERATOR_NEW_MSVC_PRAGMA MSVC_PRAGMA( warning( suppress : 28251 ) )	//	warning C28251: Inconsistent annotation for 'new': this instance has no annotations
@@ -26,18 +30,30 @@ class FChunkedFixedUObjectArray;
 
 #if !FORCE_ANSI_ALLOCATOR
 #define REPLACEMENT_OPERATOR_NEW_AND_DELETE \
-	OPERATOR_NEW_MSVC_PRAGMA void* operator new  ( size_t Size                        ) OPERATOR_NEW_THROW_SPEC      { return FMemory::Malloc( Size ); } \
-	OPERATOR_NEW_MSVC_PRAGMA void* operator new[]( size_t Size                        ) OPERATOR_NEW_THROW_SPEC      { return FMemory::Malloc( Size ); } \
-	OPERATOR_NEW_MSVC_PRAGMA void* operator new  ( size_t Size, const std::nothrow_t& ) OPERATOR_NEW_NOTHROW_SPEC    { return FMemory::Malloc( Size ); } \
-	OPERATOR_NEW_MSVC_PRAGMA void* operator new[]( size_t Size, const std::nothrow_t& ) OPERATOR_NEW_NOTHROW_SPEC    { return FMemory::Malloc( Size ); } \
-	void operator delete  ( void* Ptr )                                                 OPERATOR_DELETE_THROW_SPEC   { FMemory::Free( Ptr ); } \
-	void operator delete[]( void* Ptr )                                                 OPERATOR_DELETE_THROW_SPEC   { FMemory::Free( Ptr ); } \
-	void operator delete  ( void* Ptr, const std::nothrow_t& )                          OPERATOR_DELETE_NOTHROW_SPEC { FMemory::Free( Ptr ); } \
-	void operator delete[]( void* Ptr, const std::nothrow_t& )                          OPERATOR_DELETE_NOTHROW_SPEC { FMemory::Free( Ptr ); } \
-	void operator delete  ( void* Ptr, size_t Size )                                    OPERATOR_DELETE_THROW_SPEC   { FMemory::Free( Ptr ); } \
-	void operator delete[]( void* Ptr, size_t Size )                                    OPERATOR_DELETE_THROW_SPEC   { FMemory::Free( Ptr ); } \
-	void operator delete  ( void* Ptr, size_t Size, const std::nothrow_t& )             OPERATOR_DELETE_NOTHROW_SPEC { FMemory::Free( Ptr ); } \
-	void operator delete[]( void* Ptr, size_t Size, const std::nothrow_t& )             OPERATOR_DELETE_NOTHROW_SPEC { FMemory::Free( Ptr ); }
+	OPERATOR_NEW_MSVC_PRAGMA void* operator new  ( size_t Size                                                    ) OPERATOR_NEW_THROW_SPEC      { return FMemory::Malloc( Size ? Size : 1 ); } \
+	OPERATOR_NEW_MSVC_PRAGMA void* operator new[]( size_t Size                                                    ) OPERATOR_NEW_THROW_SPEC      { return FMemory::Malloc( Size ? Size : 1 ); } \
+	OPERATOR_NEW_MSVC_PRAGMA void* operator new  ( size_t Size,                             const std::nothrow_t& ) OPERATOR_NEW_NOTHROW_SPEC    { return FMemory::Malloc( Size ? Size : 1 ); } \
+	OPERATOR_NEW_MSVC_PRAGMA void* operator new[]( size_t Size,                             const std::nothrow_t& ) OPERATOR_NEW_NOTHROW_SPEC    { return FMemory::Malloc( Size ? Size : 1 ); } \
+	OPERATOR_NEW_MSVC_PRAGMA void* operator new  ( size_t Size, std::align_val_t Alignment                        ) OPERATOR_NEW_THROW_SPEC      { return FMemory::Malloc( Size ? Size : 1, (std::size_t)Alignment ); } \
+	OPERATOR_NEW_MSVC_PRAGMA void* operator new[]( size_t Size, std::align_val_t Alignment                        ) OPERATOR_NEW_THROW_SPEC      { return FMemory::Malloc( Size ? Size : 1, (std::size_t)Alignment ); } \
+	OPERATOR_NEW_MSVC_PRAGMA void* operator new  ( size_t Size, std::align_val_t Alignment, const std::nothrow_t& ) OPERATOR_NEW_NOTHROW_SPEC    { return FMemory::Malloc( Size ? Size : 1, (std::size_t)Alignment ); } \
+	OPERATOR_NEW_MSVC_PRAGMA void* operator new[]( size_t Size, std::align_val_t Alignment, const std::nothrow_t& ) OPERATOR_NEW_NOTHROW_SPEC    { return FMemory::Malloc( Size ? Size : 1, (std::size_t)Alignment ); } \
+	void operator delete  ( void* Ptr                                                                             ) OPERATOR_DELETE_THROW_SPEC   { FMemory::Free( Ptr ); } \
+	void operator delete[]( void* Ptr                                                                             ) OPERATOR_DELETE_THROW_SPEC   { FMemory::Free( Ptr ); } \
+	void operator delete  ( void* Ptr,                                                      const std::nothrow_t& ) OPERATOR_DELETE_NOTHROW_SPEC { FMemory::Free( Ptr ); } \
+	void operator delete[]( void* Ptr,                                                      const std::nothrow_t& ) OPERATOR_DELETE_NOTHROW_SPEC { FMemory::Free( Ptr ); } \
+	void operator delete  ( void* Ptr,             size_t Size                                                    ) OPERATOR_DELETE_THROW_SPEC   { FMemory::Free( Ptr ); } \
+	void operator delete[]( void* Ptr,             size_t Size                                                    ) OPERATOR_DELETE_THROW_SPEC   { FMemory::Free( Ptr ); } \
+	void operator delete  ( void* Ptr,             size_t Size,                             const std::nothrow_t& ) OPERATOR_DELETE_NOTHROW_SPEC { FMemory::Free( Ptr ); } \
+	void operator delete[]( void* Ptr,             size_t Size,                             const std::nothrow_t& ) OPERATOR_DELETE_NOTHROW_SPEC { FMemory::Free( Ptr ); } \
+	void operator delete  ( void* Ptr,                          std::align_val_t Alignment                        ) OPERATOR_DELETE_THROW_SPEC   { FMemory::Free( Ptr ); } \
+	void operator delete[]( void* Ptr,                          std::align_val_t Alignment                        ) OPERATOR_DELETE_THROW_SPEC   { FMemory::Free( Ptr ); } \
+	void operator delete  ( void* Ptr,                          std::align_val_t Alignment, const std::nothrow_t& ) OPERATOR_DELETE_NOTHROW_SPEC { FMemory::Free( Ptr ); } \
+	void operator delete[]( void* Ptr,                          std::align_val_t Alignment, const std::nothrow_t& ) OPERATOR_DELETE_NOTHROW_SPEC { FMemory::Free( Ptr ); } \
+	void operator delete  ( void* Ptr,             size_t Size, std::align_val_t Alignment                        ) OPERATOR_DELETE_THROW_SPEC   { FMemory::Free( Ptr ); } \
+	void operator delete[]( void* Ptr,             size_t Size, std::align_val_t Alignment                        ) OPERATOR_DELETE_THROW_SPEC   { FMemory::Free( Ptr ); } \
+	void operator delete  ( void* Ptr,             size_t Size, std::align_val_t Alignment, const std::nothrow_t& ) OPERATOR_DELETE_NOTHROW_SPEC { FMemory::Free( Ptr ); } \
+	void operator delete[]( void* Ptr,             size_t Size, std::align_val_t Alignment, const std::nothrow_t& ) OPERATOR_DELETE_NOTHROW_SPEC { FMemory::Free( Ptr ); }
 #else
 	#define REPLACEMENT_OPERATOR_NEW_AND_DELETE
 #endif
