another attempt to disable a PVS warning #rb none

diff --git a/Engine/Source/Runtime/Core/Public/Math/UnrealMathSSE.h b/Engine/Source/Runtime/Core/Public/Math/UnrealMathSSE.h
--- a/Engine/Source/Runtime/Core/Public/Math/UnrealMathSSE.h
+++ b/Engine/Source/Runtime/Core/Public/Math/UnrealMathSSE.h
@@ -202,7 +202,7 @@ FORCEINLINE VectorRegister VectorLoad( const void* Ptr )
  */
 FORCEINLINE VectorRegister VectorLoadTwoPairsFloat(const float* Ptr1, const float* Ptr2)
 {
-	__m128 Ret = _mm_castpd_ps(_mm_load_sd((double const *)(Ptr1)));
+	__m128 Ret = _mm_castpd_ps(_mm_load_sd((double const *)(Ptr1))); // -V615
 	Ret = _mm_loadh_pi(Ret, (__m64 const*)(Ptr2));
 	return Ret;
 }
