SharedBuffer: Fixed MoveToUnique() with outer buffers

#rnx

diff --git a/Engine/Source/Runtime/Core/Private/Memory/SharedBuffer.cpp b/Engine/Source/Runtime/Core/Private/Memory/SharedBuffer.cpp
--- a/Engine/Source/Runtime/Core/Private/Memory/SharedBuffer.cpp
+++ b/Engine/Source/Runtime/Core/Private/Memory/SharedBuffer.cpp
@@ -6,7 +6,7 @@
 
 ///////////////////////////////////////////////////////////////////////////////////////////////////////////////////////
 
-namespace BufferOwnerPrivate
+namespace UE::SharedBuffer::Private
 {
 
 template <typename FOps>
@@ -54,6 +54,12 @@ public:
 		SetIsMaterialized();
 	}
 
+	inline FBufferOwnerView(const void* Data, uint64 Size)
+		: FBufferOwnerView(const_cast<void*>(Data), Size)
+	{
+		SetIsImmutable();
+	}
+
 protected:
 	virtual void FreeBuffer() final
 	{
@@ -63,18 +69,18 @@ protected:
 class FBufferOwnerOuterView final : public FBufferOwner
 {
 public:
-	inline FBufferOwnerOuterView(void* Data, uint64 Size, FSharedBuffer&& InOuterBuffer)
+	template <typename OuterBufferType, decltype(FSharedBuffer(DeclVal<OuterBufferType>()))* = nullptr>
+	inline FBufferOwnerOuterView(void* Data, uint64 Size, OuterBufferType&& InOuterBuffer)
 		: FBufferOwner(Data, Size)
-		, OuterBuffer(MoveTemp(InOuterBuffer))
+		, OuterBuffer(Forward<OuterBufferType>(InOuterBuffer))
 	{
-		Initialize();
-	}
-
-	inline FBufferOwnerOuterView(void* Data, uint64 Size, const FSharedBuffer& InOuterBuffer)
-		: FBufferOwner(Data, Size)
-		, OuterBuffer(InOuterBuffer)
-	{
-		Initialize();
+		check(OuterBuffer.GetView().Contains(MakeMemoryView(GetData(), GetSize())));
+		SetIsMaterialized();
+		SetIsImmutable();
+		if (OuterBuffer.IsOwned())
+		{
+			SetIsOwned();
+		}
 	}
 
 protected:
@@ -84,26 +90,16 @@ protected:
 	}
 
 private:
-	inline void Initialize()
-	{
-		check(OuterBuffer.GetView().Contains(MakeMemoryView(GetData(), GetSize())));
-		SetIsMaterialized();
-		if (OuterBuffer.IsOwned())
-		{
-			SetIsOwned();
-		}
-	}
-
 	FSharedBuffer OuterBuffer;
 };
 
-} // BufferOwnerPrivate
+} // UE::SharedBuffer::Private
 
 ///////////////////////////////////////////////////////////////////////////////////////////////////////////////////////
 
 FUniqueBuffer FUniqueBuffer::Alloc(uint64 InSize)
 {
-	return FUniqueBuffer(new BufferOwnerPrivate::FBufferOwnerHeap(InSize));
+	return FUniqueBuffer(new UE::SharedBuffer::Private::FBufferOwnerHeap(InSize));
 }
 
 FUniqueBuffer FUniqueBuffer::Clone(FMemoryView View)
@@ -125,7 +121,7 @@ FUniqueBuffer FUniqueBuffer::MakeView(FMutableMemoryView View)
 
 FUniqueBuffer FUniqueBuffer::MakeView(void* Data, uint64 Size)
 {
-	return FUniqueBuffer(new BufferOwnerPrivate::FBufferOwnerView(Data, Size));
+	return FUniqueBuffer(new UE::SharedBuffer::Private::FBufferOwnerView(Data, Size));
 }
 
 FUniqueBuffer::FUniqueBuffer(FBufferOwner* InOwner)
@@ -133,7 +129,7 @@ FUniqueBuffer::FUniqueBuffer(FBufferOwner* InOwner)
 {
 }
 
-FUniqueBuffer::FUniqueBuffer(BufferOwnerPrivate::TBufferOwnerPtr<BufferOwnerPrivate::FSharedOps>&& SharedOwner)
+FUniqueBuffer::FUniqueBuffer(FOwnerPtrType&& SharedOwner)
 	: Owner(MoveTemp(SharedOwner))
 {
 }
@@ -163,6 +159,21 @@ FSharedBuffer FUniqueBuffer::MoveToShared()
 
 ///////////////////////////////////////////////////////////////////////////////////////////////////////////////////////
 
+template <typename OuterBufferType, decltype(FSharedBuffer(DeclVal<OuterBufferType>()))* = nullptr>
+static FSharedBuffer MakeSharedBufferViewWithOuter(FMemoryView View, OuterBufferType&& OuterBuffer)
+{
+	if (OuterBuffer.IsNull())
+	{
+		return FSharedBuffer::MakeView(View);
+	}
+	if (View == OuterBuffer.GetView())
+	{
+		return Forward<OuterBufferType>(OuterBuffer);
+	}
+	return FSharedBuffer(new UE::SharedBuffer::Private::FBufferOwnerOuterView(
+		const_cast<void*>(View.GetData()), View.GetSize(), Forward<OuterBufferType>(OuterBuffer)));
+}
+
 FSharedBuffer FSharedBuffer::Clone(FMemoryView View)
 {
 	return FUniqueBuffer::Clone(View).MoveToShared();
@@ -180,45 +191,27 @@ FSharedBuffer FSharedBuffer::MakeView(FMemoryView View)
 
 FSharedBuffer FSharedBuffer::MakeView(FMemoryView View, FSharedBuffer&& OuterBuffer)
 {
-	if (OuterBuffer.IsNull())
-	{
-		return MakeView(View);
-	}
-	if (View == OuterBuffer.GetView())
-	{
-		return MoveTemp(OuterBuffer);
-	}
-	return FSharedBuffer(new BufferOwnerPrivate::FBufferOwnerOuterView(
-		const_cast<void*>(View.GetData()), View.GetSize(), MoveTemp(OuterBuffer)));
+	return MakeSharedBufferViewWithOuter(View, MoveTemp(OuterBuffer));
 }
 
 FSharedBuffer FSharedBuffer::MakeView(FMemoryView View, const FSharedBuffer& OuterBuffer)
 {
-	if (OuterBuffer.IsNull())
-	{
-		return MakeView(View);
-	}
-	if (View == OuterBuffer.GetView())
-	{
-		return OuterBuffer;
-	}
-	return FSharedBuffer(new BufferOwnerPrivate::FBufferOwnerOuterView(
-		const_cast<void*>(View.GetData()), View.GetSize(), OuterBuffer));
+	return MakeSharedBufferViewWithOuter(View, OuterBuffer);
 }
 
 FSharedBuffer FSharedBuffer::MakeView(const void* Data, uint64 Size)
 {
-	return FSharedBuffer(new BufferOwnerPrivate::FBufferOwnerView(const_cast<void*>(Data), Size));
+	return FSharedBuffer(new UE::SharedBuffer::Private::FBufferOwnerView(Data, Size));
 }
 
 FSharedBuffer FSharedBuffer::MakeView(const void* Data, uint64 Size, FSharedBuffer&& OuterBuffer)
 {
-	return MakeView(MakeMemoryView(Data, Size), MoveTemp(OuterBuffer));
+	return MakeSharedBufferViewWithOuter(MakeMemoryView(Data, Size), MoveTemp(OuterBuffer));
 }
 
 FSharedBuffer FSharedBuffer::MakeView(const void* Data, uint64 Size, const FSharedBuffer& OuterBuffer)
 {
-	return MakeView(MakeMemoryView(Data, Size), OuterBuffer);
+	return MakeSharedBufferViewWithOuter(MakeMemoryView(Data, Size), OuterBuffer);
 }
 
 FSharedBuffer::FSharedBuffer(FBufferOwner* InOwner)
@@ -226,12 +219,12 @@ FSharedBuffer::FSharedBuffer(FBufferOwner* InOwner)
 {
 }
 
-FSharedBuffer::FSharedBuffer(BufferOwnerPrivate::TBufferOwnerPtr<BufferOwnerPrivate::FSharedOps>&& SharedOwner)
+FSharedBuffer::FSharedBuffer(FOwnerPtrType&& SharedOwner)
 	: Owner(MoveTemp(SharedOwner))
 {
 }
 
-FSharedBuffer::FSharedBuffer(const BufferOwnerPrivate::TBufferOwnerPtr<BufferOwnerPrivate::FWeakOps>& WeakOwner)
+FSharedBuffer::FSharedBuffer(const FWeakOwnerPtrType& WeakOwner)
 	: Owner(WeakOwner)
 {
 }
@@ -261,8 +254,8 @@ void FSharedBuffer::Materialize() const
 
 FUniqueBuffer FSharedBuffer::MoveToUnique()
 {
-	OwnerPtrType ExistingOwner = ToPrivateOwnerPtr(MoveTemp(*this));
-	if (!ExistingOwner || (ExistingOwner->IsOwned() && ExistingOwner->GetTotalRefCount() == 1))
+	FOwnerPtrType ExistingOwner = ToPrivateOwnerPtr(MoveTemp(*this));
+	if (!ExistingOwner || ExistingOwner->IsUniqueOwnedMutable())
 	{
 		return FUniqueBuffer(MoveTemp(ExistingOwner));
 	}
diff --git a/Engine/Source/Runtime/Core/Private/Tests/Memory/SharedBufferTest.cpp b/Engine/Source/Runtime/Core/Private/Tests/Memory/SharedBufferTest.cpp
--- a/Engine/Source/Runtime/Core/Private/Tests/Memory/SharedBufferTest.cpp
+++ b/Engine/Source/Runtime/Core/Private/Tests/Memory/SharedBufferTest.cpp
@@ -169,16 +169,11 @@ bool FUniqueBufferTest::RunTest(const FString& Parameters)
 	{
 		bool bDeleted = false;
 		constexpr uint64 Size = 64;
-		void* Data = new uint8[Size];
-		FUniqueBuffer Buffer = FUniqueBuffer::TakeOwnership(Data, Size,
-			[&bDeleted](void* InData, uint64 InSize)
-		{
-			bDeleted = true;
-			delete[] static_cast<uint8*>(InData);
-		});
+		uint8 Data[Size]{};
+		FUniqueBuffer Buffer = FUniqueBuffer::TakeOwnership(Data, Size, [&bDeleted](void*, uint64) { bDeleted = true; });
 		TestTrue(TEXT("FUniqueBuffer::TakeOwnership().IsOwned()"), Buffer.IsOwned());
 		TestEqual(TEXT("FUniqueBuffer::TakeOwnership().GetSize()"), Buffer.GetSize(), Size);
-		TestEqual(TEXT("FUniqueBuffer::TakeOwnership().GetData()"), Buffer.GetData(), Data);
+		TestEqual(TEXT("FUniqueBuffer::TakeOwnership().GetData()"), Buffer.GetData(), static_cast<void*>(Data));
 		Buffer.Reset();
 		TestTrue(TEXT("FUniqueBuffer::TakeOwnership() Deleted"), bDeleted);
 	}
@@ -279,12 +274,7 @@ bool FSharedBufferTest::RunTest(const FString& Parameters)
 		bool bDeleted = false;
 		constexpr uint64 Size = 64;
 		const void* Data = new uint8[Size];
-		FSharedBuffer Buffer = FSharedBuffer::TakeOwnership(Data, Size,
-			[&bDeleted](void* InData, uint64 InSize)
-		{
-			bDeleted = true;
-			delete[] static_cast<uint8*>(InData);
-		});
+		FSharedBuffer Buffer = FSharedBuffer::TakeOwnership(Data, Size, [&bDeleted](void*, uint64) { bDeleted = true; });
 		TestTrue(TEXT("FSharedBuffer::TakeOwnership().IsOwned()"), Buffer.IsOwned());
 		TestEqual(TEXT("FSharedBuffer::TakeOwnership().GetSize()"), Buffer.GetSize(), Size);
 		TestEqual(TEXT("FSharedBuffer::TakeOwnership().GetData()"), Buffer.GetData(), Data);
@@ -323,31 +313,56 @@ bool FSharedBufferTest::RunTest(const FString& Parameters)
 	{
 		bool bDeleted = false;
 		constexpr uint64 Size = 64;
-		void* Data = new uint8[Size];
-		FSharedBuffer Shared = FSharedBuffer::TakeOwnership(Data, Size,
-			[&bDeleted](void* InData, uint64 InSize)
-		{
-			bDeleted = true;
-			delete[] static_cast<uint8*>(InData);
-		});
+		uint8 Data[Size]{};
+		FSharedBuffer Shared = FUniqueBuffer::TakeOwnership(Data, Size, [&bDeleted](void*) { bDeleted = true; }).MoveToShared();
 		FUniqueBuffer Buffer = Shared.MoveToUnique();
 		TestTrue(TEXT("FSharedBuffer::MoveToUnique(Owned).IsNull()"), Shared.IsNull());
 		TestTrue(TEXT("FSharedBuffer::MoveToUnique(Owned).IsOwned()"), Buffer.IsOwned());
 		TestEqual(TEXT("FSharedBuffer::MoveToUnique(Owned).GetSize()"), Buffer.GetSize(), Size);
-		TestEqual(TEXT("FSharedBuffer::MoveToUnique(Owned).GetData()"), Buffer.GetData(), Data);
+		TestEqual(TEXT("FSharedBuffer::MoveToUnique(Owned).GetData()"), Buffer.GetData(), static_cast<void*>(Data));
 		Buffer.Reset();
 		TestTrue(TEXT("FSharedBuffer::MoveToUnique(Owned) Deleted"), bDeleted);
 	}
 	{
 		constexpr uint64 Size = 64;
 		uint8 Data[Size]{};
-		FSharedBuffer Shared = FSharedBuffer::MakeView(MakeMemoryView(Data));
+		FSharedBuffer Shared = FSharedBuffer::TakeOwnership(Data, Size, [](void*){});
+		FUniqueBuffer Buffer = Shared.MoveToUnique();
+		TestTrue(TEXT("FSharedBuffer::MoveToUnique(ImmutableOwned).IsNull()"), Shared.IsNull());
+		TestTrue(TEXT("FSharedBuffer::MoveToUnique(ImmutableOwned).IsOwned()"), Buffer.IsOwned());
+		TestEqual(TEXT("FSharedBuffer::MoveToUnique(ImmutableOwned).GetSize()"), Buffer.GetSize(), Size);
+		TestNotEqual(TEXT("FSharedBuffer::MoveToUnique(ImmutableOwned).GetData()"), Buffer.GetData(), static_cast<void*>(Data));
+	}
+	{
+		constexpr uint64 Size = 64;
+		uint8 Data[Size]{};
+		FSharedBuffer Shared = FUniqueBuffer::MakeView(MakeMemoryView(Data)).MoveToShared();
 		FUniqueBuffer Buffer = Shared.MoveToUnique();
 		TestTrue(TEXT("FSharedBuffer::MoveToUnique(View).IsNull()"), Shared.IsNull());
 		TestTrue(TEXT("FSharedBuffer::MoveToUnique(View).IsOwned()"), Buffer.IsOwned());
 		TestEqual(TEXT("FSharedBuffer::MoveToUnique(View).GetSize()"), Buffer.GetSize(), Size);
 		TestNotEqual(TEXT("FSharedBuffer::MoveToUnique(View).GetData()"), Buffer.GetData(), static_cast<void*>(Data));
 	}
+	{
+		constexpr uint64 Size = 64;
+		uint8 Data[Size]{};
+		FSharedBuffer Shared = FSharedBuffer::MakeView(MakeMemoryView(Data));
+		FUniqueBuffer Buffer = Shared.MoveToUnique();
+		TestTrue(TEXT("FSharedBuffer::MoveToUnique(ImmutableView).IsNull()"), Shared.IsNull());
+		TestTrue(TEXT("FSharedBuffer::MoveToUnique(ImmutableView).IsOwned()"), Buffer.IsOwned());
+		TestEqual(TEXT("FSharedBuffer::MoveToUnique(ImmutableView).GetSize()"), Buffer.GetSize(), Size);
+		TestNotEqual(TEXT("FSharedBuffer::MoveToUnique(ImmutableView).GetData()"), Buffer.GetData(), static_cast<void*>(Data));
+	}
+	{
+		constexpr uint64 Size = 64;
+		uint8 Data[Size]{};
+		FSharedBuffer Shared = FUniqueBuffer::TakeOwnership(Data, Size, [](void*){}).MoveToShared();
+		FUniqueBuffer Buffer = FSharedBuffer::MakeView(Shared.GetView().Left(Size / 2), Shared).MoveToUnique();
+		TestFalse(TEXT("FSharedBuffer::MoveToUnique(OwnedView).IsNull()"), Shared.IsNull());
+		TestTrue(TEXT("FSharedBuffer::MoveToUnique(OwnedView).IsOwned()"), Buffer.IsOwned());
+		TestEqual(TEXT("FSharedBuffer::MoveToUnique(OwnedView).GetSize()"), Buffer.GetSize(), Size / 2);
+		TestNotEqual(TEXT("FSharedBuffer::MoveToUnique(OwnedView).GetData()"), Buffer.GetData(), static_cast<void*>(Data));
+	}
 
 	return true;
 }
diff --git a/Engine/Source/Runtime/Core/Public/Memory/SharedBuffer.h b/Engine/Source/Runtime/Core/Public/Memory/SharedBuffer.h
--- a/Engine/Source/Runtime/Core/Public/Memory/SharedBuffer.h
+++ b/Engine/Source/Runtime/Core/Public/Memory/SharedBuffer.h
@@ -16,8 +16,8 @@
 
 ///////////////////////////////////////////////////////////////////////////////////////////////////////////////////////
 
-namespace BufferOwnerPrivate { struct FSharedOps; }
-namespace BufferOwnerPrivate { struct FWeakOps; }
+namespace UE::SharedBuffer::Private { struct FSharedOps; }
+namespace UE::SharedBuffer::Private { struct FWeakOps; }
 
 ///////////////////////////////////////////////////////////////////////////////////////////////////////////////////////
 
@@ -39,6 +39,8 @@ namespace BufferOwnerPrivate { struct FWeakOps; }
  * of the buffer to be deferred.
  *
  * A derived type must call SetIsOwned from its constructor if they own (or will own) the buffer.
+ * A derived type must call SetIsImmutable from its constructor if they want the buffer to become
+ * permanently immutable after it is stored in FSharedBuffer, forcing MoveToUnique() to clone it.
  * A derived type must call SetIsMaterialized from its constructor, unless it implements deferred
  * materialization by overriding MaterializeBuffer.
  */
@@ -46,6 +48,7 @@ class FBufferOwner
 {
 private:
 	enum class EBufferOwnerFlags : uint8;
+	FRIEND_ENUM_CLASS_FLAGS(EBufferOwnerFlags);
 
 protected:
 	FBufferOwner() = default;
@@ -81,6 +84,9 @@ protected:
 	inline bool IsOwned() const;
 	inline void SetIsOwned();
 
+	inline bool IsImmutable() const;
+	inline void SetIsImmutable();
+
 	inline void Materialize();
 	inline bool IsMaterialized() const;
 	inline void SetIsMaterialized();
@@ -88,12 +94,18 @@ protected:
 	inline uint32 GetTotalRefCount() const;
 
 private:
-	static uint32 GetSharedRefCount(uint64 RefCountsAndFlags) { return uint32(RefCountsAndFlags >> 0) & 0x7fffffff; }
-	static uint64 SetSharedRefCount(uint32 RefCount) { return uint64(RefCount) << 0; }
-	static uint32 GetWeakRefCount(uint64 RefCountsAndFlags) { return uint32(RefCountsAndFlags >> 31) & 0x7fffffff; }
-	static uint64 SetWeakRefCount(uint32 RefCount) { return uint64(RefCount) << 31; }
-	static EBufferOwnerFlags GetFlags(uint64 RefCountsAndFlags) { return EBufferOwnerFlags(RefCountsAndFlags >> 62); }
-	static uint64 SetFlags(EBufferOwnerFlags Flags) { return uint64(Flags) << 62; }
+	static constexpr uint32 RefCountMask = 0x3fffffff;
+
+	static constexpr inline uint32 GetTotalRefCount(uint64 RefCountsAndFlags);
+	static constexpr inline uint32 GetSharedRefCount(uint64 RefCountsAndFlags) { return uint32(RefCountsAndFlags >> 0) & RefCountMask; }
+	static constexpr inline uint64 SetSharedRefCount(uint32 RefCount) { return uint64(RefCount) << 0; }
+	static constexpr inline uint32 GetWeakRefCount(uint64 RefCountsAndFlags) { return uint32(RefCountsAndFlags >> 30) & RefCountMask; }
+	static constexpr inline uint64 SetWeakRefCount(uint32 RefCount) { return uint64(RefCount) << 30; }
+	static constexpr inline EBufferOwnerFlags GetFlags(uint64 RefCountsAndFlags) { return EBufferOwnerFlags(RefCountsAndFlags >> 60); }
+	static constexpr inline uint64 SetFlags(EBufferOwnerFlags Flags) { return uint64(Flags) << 60; }
+
+	/** Returns whether this has one total reference, is owned, and is not immutable. */
+	inline bool IsUniqueOwnedMutable() const;
 
 	inline void AddSharedReference();
 	inline void ReleaseSharedReference();
@@ -104,8 +116,8 @@ private:
 	friend class FUniqueBuffer;
 	friend class FSharedBuffer;
 	friend class FWeakSharedBuffer;
-	friend struct BufferOwnerPrivate::FSharedOps;
-	friend struct BufferOwnerPrivate::FWeakOps;
+	friend struct UE::SharedBuffer::Private::FSharedOps;
+	friend struct UE::SharedBuffer::Private::FWeakOps;
 
 private:
 	void* Data = nullptr;
@@ -115,7 +127,7 @@ private:
 
 ///////////////////////////////////////////////////////////////////////////////////////////////////////////////////////
 
-namespace BufferOwnerPrivate
+namespace UE::SharedBuffer::Private
 {
 
 struct FSharedOps final
@@ -186,7 +198,7 @@ private:
 	FBufferOwner* Owner = nullptr;
 };
 
-} // BufferOwnerPrivate
+} // UE::SharedBuffer::Private
 
 ///////////////////////////////////////////////////////////////////////////////////////////////////////////////////////
 
@@ -202,15 +214,15 @@ class FUniqueBuffer
 {
 public:
 	/** Make an uninitialized owned buffer of the specified size. */
-	UE_NODISCARD CORE_API static FUniqueBuffer Alloc(uint64 Size);
+	[[nodiscard]] CORE_API static FUniqueBuffer Alloc(uint64 Size);
 
 	/** Make an owned clone of the input. */
-	UE_NODISCARD CORE_API static FUniqueBuffer Clone(FMemoryView View);
-	UE_NODISCARD CORE_API static FUniqueBuffer Clone(const void* Data, uint64 Size);
+	[[nodiscard]] CORE_API static FUniqueBuffer Clone(FMemoryView View);
+	[[nodiscard]] CORE_API static FUniqueBuffer Clone(const void* Data, uint64 Size);
 
 	/** Make a non-owned view of the input. */
-	UE_NODISCARD CORE_API static FUniqueBuffer MakeView(FMutableMemoryView View);
-	UE_NODISCARD CORE_API static FUniqueBuffer MakeView(void* Data, uint64 Size);
+	[[nodiscard]] CORE_API static FUniqueBuffer MakeView(FMutableMemoryView View);
+	[[nodiscard]] CORE_API static FUniqueBuffer MakeView(void* Data, uint64 Size);
 
 	/**
 	 * Make an owned buffer by taking ownership of the input.
@@ -219,7 +231,7 @@ public:
 	 */
 	template <typename DeleteFunctionType,
 		decltype(Invoke(std::declval<DeleteFunctionType>(), std::declval<void*>()))* = nullptr>
-	UE_NODISCARD static inline FUniqueBuffer TakeOwnership(void* Data, uint64 Size, DeleteFunctionType&& DeleteFunction);
+	[[nodiscard]] static inline FUniqueBuffer TakeOwnership(void* Data, uint64 Size, DeleteFunctionType&& DeleteFunction);
 
 	/**
 	 * Make an owned buffer by taking ownership of the input.
@@ -228,7 +240,7 @@ public:
 	 */
 	template <typename DeleteFunctionType,
 		decltype(Invoke(std::declval<DeleteFunctionType>(), std::declval<void*>(), std::declval<uint64>()))* = nullptr>
-	UE_NODISCARD static inline FUniqueBuffer TakeOwnership(void* Data, uint64 Size, DeleteFunctionType&& DeleteFunction);
+	[[nodiscard]] static inline FUniqueBuffer TakeOwnership(void* Data, uint64 Size, DeleteFunctionType&& DeleteFunction);
 
 	/** Construct a null unique buffer. */
 	FUniqueBuffer() = default;
@@ -246,36 +258,36 @@ public:
 	CORE_API void Reset();
 
 	/** Returns a pointer to the start of the buffer. */
-	UE_NODISCARD inline void* GetData() { return Owner ? Owner->GetData() : nullptr; }
-	UE_NODISCARD inline const void* GetData() const { return Owner ? Owner->GetData() : nullptr; }
+	[[nodiscard]] inline void* GetData() { return Owner ? Owner->GetData() : nullptr; }
+	[[nodiscard]] inline const void* GetData() const { return Owner ? Owner->GetData() : nullptr; }
 
 	/** Returns the size of the buffer in bytes. */
-	UE_NODISCARD inline uint64 GetSize() const { return Owner ? Owner->GetSize() : 0; }
+	[[nodiscard]] inline uint64 GetSize() const { return Owner ? Owner->GetSize() : 0; }
 
 	/** Returns a view of the buffer. */
-	UE_NODISCARD inline FMutableMemoryView GetView() { return FMutableMemoryView(GetData(), GetSize()); }
-	UE_NODISCARD inline FMemoryView GetView() const { return FMemoryView(GetData(), GetSize()); }
-	UE_NODISCARD inline operator FMutableMemoryView() { return GetView(); }
-	UE_NODISCARD inline operator FMemoryView() const { return GetView(); }
+	[[nodiscard]] inline FMutableMemoryView GetView() { return FMutableMemoryView(GetData(), GetSize()); }
+	[[nodiscard]] inline FMemoryView GetView() const { return FMemoryView(GetData(), GetSize()); }
+	[[nodiscard]] inline operator FMutableMemoryView() { return GetView(); }
+	[[nodiscard]] inline operator FMemoryView() const { return GetView(); }
 
 	/** Returns true if this points to a buffer owner. */
-	UE_NODISCARD inline explicit operator bool() const { return !IsNull(); }
+	[[nodiscard]] inline explicit operator bool() const { return !IsNull(); }
 
 	/**
 	 * Returns true if this does not point to a buffer owner.
 	 *
 	 * A null buffer is always owned, materialized, and empty.
 	 */
-	UE_NODISCARD inline bool IsNull() const { return Owner.IsNull(); }
+	[[nodiscard]] inline bool IsNull() const { return Owner.IsNull(); }
 
 	/** Returns true if this keeps the referenced buffer alive. */
-	UE_NODISCARD inline bool IsOwned() const { return !Owner || Owner->IsOwned(); }
+	[[nodiscard]] inline bool IsOwned() const { return !Owner || Owner->IsOwned(); }
 
 	/** Returns a buffer that is owned, by cloning if not owned. */
-	UE_NODISCARD CORE_API FUniqueBuffer MakeOwned() &&;
+	[[nodiscard]] CORE_API FUniqueBuffer MakeOwned() &&;
 
 	/** Returns true if the referenced buffer has been materialized. */
-	UE_NODISCARD inline bool IsMaterialized() const { return !Owner || Owner->IsMaterialized(); }
+	[[nodiscard]] inline bool IsMaterialized() const { return !Owner || Owner->IsMaterialized(); }
 
 	/**
 	 * Materialize the buffer by making its data and size available.
@@ -289,19 +301,19 @@ public:
 	 *
 	 * Steals the buffer owner from the unique buffer.
 	 */
-	UE_NODISCARD CORE_API FSharedBuffer MoveToShared();
+	[[nodiscard]] CORE_API FSharedBuffer MoveToShared();
 
 	friend class FSharedBuffer;
 
 private:
-	explicit FUniqueBuffer(BufferOwnerPrivate::TBufferOwnerPtr<BufferOwnerPrivate::FSharedOps>&& SharedOwner);
+	using FOwnerPtrType = UE::SharedBuffer::Private::TBufferOwnerPtr<UE::SharedBuffer::Private::FSharedOps>;
 
-	using OwnerPtrType = BufferOwnerPrivate::TBufferOwnerPtr<BufferOwnerPrivate::FSharedOps>;
+	explicit FUniqueBuffer(FOwnerPtrType&& SharedOwner);
 
-	inline friend const OwnerPtrType& ToPrivateOwnerPtr(const FUniqueBuffer& Buffer) { return Buffer.Owner; }
-	inline friend OwnerPtrType ToPrivateOwnerPtr(FUniqueBuffer&& Buffer) { return MoveTemp(Buffer.Owner); }
+	inline friend const FOwnerPtrType& ToPrivateOwnerPtr(const FUniqueBuffer& Buffer) { return Buffer.Owner; }
+	inline friend FOwnerPtrType ToPrivateOwnerPtr(FUniqueBuffer&& Buffer) { return MoveTemp(Buffer.Owner); }
 
-	OwnerPtrType Owner;
+	FOwnerPtrType Owner;
 };
 
 ///////////////////////////////////////////////////////////////////////////////////////////////////////////////////////
@@ -315,18 +327,18 @@ class FSharedBuffer
 {
 public:
 	/** Make an owned clone of the input. */
-	UE_NODISCARD CORE_API static FSharedBuffer Clone(FMemoryView View);
-	UE_NODISCARD CORE_API static FSharedBuffer Clone(const void* Data, uint64 Size);
+	[[nodiscard]] CORE_API static FSharedBuffer Clone(FMemoryView View);
+	[[nodiscard]] CORE_API static FSharedBuffer Clone(const void* Data, uint64 Size);
 
 	/** Make a non-owned view of the input. */
-	UE_NODISCARD CORE_API static FSharedBuffer MakeView(FMemoryView View);
-	UE_NODISCARD CORE_API static FSharedBuffer MakeView(const void* Data, uint64 Size);
+	[[nodiscard]] CORE_API static FSharedBuffer MakeView(FMemoryView View);
+	[[nodiscard]] CORE_API static FSharedBuffer MakeView(const void* Data, uint64 Size);
 
 	/** Make a view of the input within its outer buffer. Ownership matches OuterBuffer. */
-	UE_NODISCARD CORE_API static FSharedBuffer MakeView(FMemoryView View, FSharedBuffer&& OuterBuffer);
-	UE_NODISCARD CORE_API static FSharedBuffer MakeView(FMemoryView View, const FSharedBuffer& OuterBuffer);
-	UE_NODISCARD CORE_API static FSharedBuffer MakeView(const void* Data, uint64 Size, FSharedBuffer&& OuterBuffer);
-	UE_NODISCARD CORE_API static FSharedBuffer MakeView(const void* Data, uint64 Size, const FSharedBuffer& OuterBuffer);
+	[[nodiscard]] CORE_API static FSharedBuffer MakeView(FMemoryView View, FSharedBuffer&& OuterBuffer);
+	[[nodiscard]] CORE_API static FSharedBuffer MakeView(FMemoryView View, const FSharedBuffer& OuterBuffer);
+	[[nodiscard]] CORE_API static FSharedBuffer MakeView(const void* Data, uint64 Size, FSharedBuffer&& OuterBuffer);
+	[[nodiscard]] CORE_API static FSharedBuffer MakeView(const void* Data, uint64 Size, const FSharedBuffer& OuterBuffer);
 
 	/**
 	 * Make an owned buffer by taking ownership of the input.
@@ -335,7 +347,7 @@ public:
 	 */
 	template <typename DeleteFunctionType,
 		decltype(Invoke(std::declval<DeleteFunctionType>(), std::declval<void*>()))* = nullptr>
-	UE_NODISCARD static inline FSharedBuffer TakeOwnership(const void* Data, uint64 Size, DeleteFunctionType&& DeleteFunction);
+	[[nodiscard]] static inline FSharedBuffer TakeOwnership(const void* Data, uint64 Size, DeleteFunctionType&& DeleteFunction);
 
 	/**
 	 * Make an owned buffer by taking ownership of the input.
@@ -344,7 +356,7 @@ public:
 	 */
 	template <typename DeleteFunctionType,
 		decltype(Invoke(std::declval<DeleteFunctionType>(), std::declval<void*>(), std::declval<uint64>()))* = nullptr>
-	UE_NODISCARD static inline FSharedBuffer TakeOwnership(const void* Data, uint64 Size, DeleteFunctionType&& DeleteFunction);
+	[[nodiscard]] static inline FSharedBuffer TakeOwnership(const void* Data, uint64 Size, DeleteFunctionType&& DeleteFunction);
 
 	/** Construct a null shared buffer. */
 	FSharedBuffer() = default;
@@ -356,34 +368,34 @@ public:
 	CORE_API void Reset();
 
 	/** Returns a pointer to the start of the buffer. */
-	UE_NODISCARD inline const void* GetData() const { return Owner ? Owner->GetData() : nullptr; }
+	[[nodiscard]] inline const void* GetData() const { return Owner ? Owner->GetData() : nullptr; }
 
 	/** Returns the size of the buffer in bytes. */
-	UE_NODISCARD inline uint64 GetSize() const { return Owner ? Owner->GetSize() : 0; }
+	[[nodiscard]] inline uint64 GetSize() const { return Owner ? Owner->GetSize() : 0; }
 
 	/** Returns a view of the buffer. */
-	UE_NODISCARD inline FMemoryView GetView() const { return FMemoryView(GetData(), GetSize()); }
-	UE_NODISCARD inline operator FMemoryView() const { return GetView(); }
+	[[nodiscard]] inline FMemoryView GetView() const { return FMemoryView(GetData(), GetSize()); }
+	[[nodiscard]] inline operator FMemoryView() const { return GetView(); }
 
 	/** Returns true if this points to a buffer owner. */
-	UE_NODISCARD inline explicit operator bool() const { return !IsNull(); }
+	[[nodiscard]] inline explicit operator bool() const { return !IsNull(); }
 
 	/**
 	 * Returns true if this does not point to a buffer owner.
 	 *
 	 * A null buffer is always owned, materialized, and empty.
 	 */
-	UE_NODISCARD inline bool IsNull() const { return Owner.IsNull(); }
+	[[nodiscard]] inline bool IsNull() const { return Owner.IsNull(); }
 
 	/** Returns true if this keeps the referenced buffer alive. */
-	UE_NODISCARD inline bool IsOwned() const { return !Owner || Owner->IsOwned(); }
+	[[nodiscard]] inline bool IsOwned() const { return !Owner || Owner->IsOwned(); }
 
 	/** Returns a buffer that is owned, by cloning if not owned. */
-	UE_NODISCARD CORE_API FSharedBuffer MakeOwned() const &;
-	UE_NODISCARD CORE_API FSharedBuffer MakeOwned() &&;
+	[[nodiscard]] CORE_API FSharedBuffer MakeOwned() const &;
+	[[nodiscard]] CORE_API FSharedBuffer MakeOwned() &&;
 
 	/** Returns true if the referenced buffer has been materialized. */
-	UE_NODISCARD inline bool IsMaterialized() const { return !Owner || Owner->IsMaterialized(); }
+	[[nodiscard]] inline bool IsMaterialized() const { return !Owner || Owner->IsMaterialized(); }
 
 	/**
 	 * Materialize the buffer by making its data and size available.
@@ -398,21 +410,22 @@ public:
 	 * Steals the buffer owner from the shared buffer if this is the last reference to it, otherwise
 	 * clones the shared buffer to guarantee unique ownership. A non-owned buffer is always cloned.
 	 */
-	UE_NODISCARD CORE_API FUniqueBuffer MoveToUnique();
+	[[nodiscard]] CORE_API FUniqueBuffer MoveToUnique();
 
 	friend class FUniqueBuffer;
 	friend class FWeakSharedBuffer;
 
 private:
-	explicit FSharedBuffer(BufferOwnerPrivate::TBufferOwnerPtr<BufferOwnerPrivate::FSharedOps>&& SharedOwner);
-	explicit FSharedBuffer(const BufferOwnerPrivate::TBufferOwnerPtr<BufferOwnerPrivate::FWeakOps>& WeakOwner);
+	using FOwnerPtrType = UE::SharedBuffer::Private::TBufferOwnerPtr<UE::SharedBuffer::Private::FSharedOps>;
+	using FWeakOwnerPtrType = UE::SharedBuffer::Private::TBufferOwnerPtr<UE::SharedBuffer::Private::FWeakOps>;
 
-	using OwnerPtrType = BufferOwnerPrivate::TBufferOwnerPtr<BufferOwnerPrivate::FSharedOps>;
+	explicit FSharedBuffer(FOwnerPtrType&& SharedOwner);
+	explicit FSharedBuffer(const FWeakOwnerPtrType& WeakOwner);
 
-	inline friend const OwnerPtrType& ToPrivateOwnerPtr(const FSharedBuffer& Buffer) { return Buffer.Owner; }
-	inline friend OwnerPtrType ToPrivateOwnerPtr(FSharedBuffer&& Buffer) { return MoveTemp(Buffer.Owner); }
+	inline friend const FOwnerPtrType& ToPrivateOwnerPtr(const FSharedBuffer& Buffer) { return Buffer.Owner; }
+	inline friend FOwnerPtrType ToPrivateOwnerPtr(FSharedBuffer&& Buffer) { return MoveTemp(Buffer.Owner); }
 
-	OwnerPtrType Owner;
+	FOwnerPtrType Owner;
 };
 
 ///////////////////////////////////////////////////////////////////////////////////////////////////////////////////////
@@ -438,14 +451,14 @@ public:
 	CORE_API void Reset();
 
 	/** Convert this to a shared buffer if it has any remaining shared references. */
-	UE_NODISCARD CORE_API FSharedBuffer Pin() const;
+	[[nodiscard]] CORE_API FSharedBuffer Pin() const;
 
 private:
-	using OwnerPtrType = BufferOwnerPrivate::TBufferOwnerPtr<BufferOwnerPrivate::FWeakOps>;
+	using FWeakOwnerPtrType = UE::SharedBuffer::Private::TBufferOwnerPtr<UE::SharedBuffer::Private::FWeakOps>;
 
-	inline friend const OwnerPtrType& ToPrivateOwnerPtr(const FWeakSharedBuffer& Buffer) { return Buffer.Owner; }
+	inline friend const FWeakOwnerPtrType& ToPrivateOwnerPtr(const FWeakSharedBuffer& Buffer) { return Buffer.Owner; }
 
-	OwnerPtrType Owner;
+	FWeakOwnerPtrType Owner;
 };
 
 ///////////////////////////////////////////////////////////////////////////////////////////////////////////////////////
@@ -478,7 +491,7 @@ template<> struct TIsWeakPointerType<FWeakSharedBuffer> { enum { Value = true };
 
 ///////////////////////////////////////////////////////////////////////////////////////////////////////////////////////
 
-namespace BufferOwnerPrivate
+namespace UE::SharedBuffer::Private
 {
 
 template <typename DeleteFunctionType>
@@ -493,6 +506,12 @@ public:
 		SetIsOwned();
 	}
 
+	explicit TBufferOwnerDeleteFunction(const void* Data, uint64 Size, DeleteFunctionType&& InDeleteFunction)
+		: TBufferOwnerDeleteFunction(const_cast<void*>(Data), Size, Forward<DeleteFunctionType>(InDeleteFunction))
+	{
+		SetIsImmutable();
+	}
+
 private:
 	virtual void FreeBuffer() final
 	{
@@ -502,15 +521,15 @@ private:
 	std::decay_t<DeleteFunctionType> DeleteFunction;
 };
 
-} // BufferOwnerPrivate
+} // UE::SharedBuffer::Private
 
 template <typename DeleteFunctionType,
 	decltype(Invoke(std::declval<DeleteFunctionType>(), std::declval<void*>()))*>
 inline FUniqueBuffer FUniqueBuffer::TakeOwnership(void* Data, uint64 Size, DeleteFunctionType&& DeleteFunction)
 {
-	return TakeOwnership(Data, Size, [Delete=Forward<DeleteFunctionType>(DeleteFunction)](void* InData, uint64 InSize)
+	return TakeOwnership(Data, Size, [Delete = Forward<DeleteFunctionType>(DeleteFunction)](void* InData, uint64 InSize)
 	{
-		Delete(InData);
+		Invoke(Delete, InData);
 	});
 }
 
@@ -518,7 +537,7 @@ template <typename DeleteFunctionType,
 	decltype(Invoke(std::declval<DeleteFunctionType>(), std::declval<void*>(), std::declval<uint64>()))*>
 inline FUniqueBuffer FUniqueBuffer::TakeOwnership(void* Data, uint64 Size, DeleteFunctionType&& DeleteFunction)
 {
-	using OwnerType = BufferOwnerPrivate::TBufferOwnerDeleteFunction<DeleteFunctionType>;
+	using OwnerType = UE::SharedBuffer::Private::TBufferOwnerDeleteFunction<DeleteFunctionType>;
 	return FUniqueBuffer(new OwnerType(Data, Size, Forward<DeleteFunctionType>(DeleteFunction)));
 }
 
@@ -526,9 +545,9 @@ template <typename DeleteFunctionType,
 	decltype(Invoke(std::declval<DeleteFunctionType>(), std::declval<void*>()))*>
 inline FSharedBuffer FSharedBuffer::TakeOwnership(const void* Data, uint64 Size, DeleteFunctionType&& DeleteFunction)
 {
-	return TakeOwnership(Data, Size, [Delete=Forward<DeleteFunctionType>(DeleteFunction)](void* InData, uint64 InSize)
+	return TakeOwnership(Data, Size, [Delete = Forward<DeleteFunctionType>(DeleteFunction)](void* InData, uint64 InSize)
 	{
-		Delete(InData);
+		Invoke(Delete, InData);
 	});
 }
 
@@ -536,8 +555,8 @@ template <typename DeleteFunctionType,
 	decltype(Invoke(std::declval<DeleteFunctionType>(), std::declval<void*>(), std::declval<uint64>()))*>
 inline FSharedBuffer FSharedBuffer::TakeOwnership(const void* Data, uint64 Size, DeleteFunctionType&& DeleteFunction)
 {
-	using OwnerType = BufferOwnerPrivate::TBufferOwnerDeleteFunction<DeleteFunctionType>;
-	return FSharedBuffer(new OwnerType(const_cast<void*>(Data), Size, Forward<DeleteFunctionType>(DeleteFunction)));
+	using OwnerType = UE::SharedBuffer::Private::TBufferOwnerDeleteFunction<DeleteFunctionType>;
+	return FSharedBuffer(new OwnerType(Data, Size, Forward<DeleteFunctionType>(DeleteFunction)));
 }
 
 ///////////////////////////////////////////////////////////////////////////////////////////////////////////////////////
@@ -546,9 +565,12 @@ enum class FBufferOwner::EBufferOwnerFlags : uint8
 {
 	None         = 0,
 	Owned        = 1 << 0,
-	Materialized = 1 << 1,
+	Immutable    = 1 << 1,
+	Materialized = 1 << 2,
 };
 
+ENUM_CLASS_FLAGS(FBufferOwner::EBufferOwnerFlags);
+
 inline FBufferOwner::FBufferOwner(void* InData, uint64 InSize)
 	: Data(InData)
 	, Size(InSize)
@@ -585,7 +607,7 @@ inline void FBufferOwner::SetBuffer(void* InData, uint64 InSize)
 
 inline bool FBufferOwner::IsOwned() const
 {
-	return (uint8(GetFlags(ReferenceCountsAndFlags)) & uint8(EBufferOwnerFlags::Owned)) != 0;
+	return EnumHasAnyFlags(GetFlags(ReferenceCountsAndFlags.load(std::memory_order_relaxed)), EBufferOwnerFlags::Owned);
 }
 
 inline void FBufferOwner::SetIsOwned()
@@ -593,6 +615,16 @@ inline void FBufferOwner::SetIsOwned()
 	ReferenceCountsAndFlags.fetch_or(SetFlags(EBufferOwnerFlags::Owned));
 }
 
+inline bool FBufferOwner::IsImmutable() const
+{
+	return EnumHasAnyFlags(GetFlags(ReferenceCountsAndFlags.load(std::memory_order_relaxed)), EBufferOwnerFlags::Immutable);
+}
+
+inline void FBufferOwner::SetIsImmutable()
+{
+	ReferenceCountsAndFlags.fetch_or(SetFlags(EBufferOwnerFlags::Immutable));
+}
+
 inline void FBufferOwner::Materialize()
 {
 	if (!IsMaterialized())
@@ -604,7 +636,7 @@ inline void FBufferOwner::Materialize()
 
 inline bool FBufferOwner::IsMaterialized() const
 {
-	return (uint8(GetFlags(ReferenceCountsAndFlags)) & uint8(EBufferOwnerFlags::Materialized)) != 0;
+	return EnumHasAnyFlags(GetFlags(ReferenceCountsAndFlags.load(std::memory_order_relaxed)), EBufferOwnerFlags::Materialized);
 }
 
 inline void FBufferOwner::SetIsMaterialized()
@@ -614,17 +646,28 @@ inline void FBufferOwner::SetIsMaterialized()
 
 inline uint32 FBufferOwner::GetTotalRefCount() const
 {
-	const uint64 Value = ReferenceCountsAndFlags.load(std::memory_order_relaxed);
-	const uint32 SharedRefCount = GetSharedRefCount(Value);
+	return GetTotalRefCount(ReferenceCountsAndFlags.load(std::memory_order_relaxed));
+}
+
+constexpr inline uint32 FBufferOwner::GetTotalRefCount(const uint64 RefCountsAndFlags)
+{
+	const uint32 SharedRefCount = GetSharedRefCount(RefCountsAndFlags);
 	// A non-zero SharedRefCount adds 1 to WeakRefCount to keep the owner alive.
 	// Subtract that extra reference when it is present to return an accurate count.
-	return GetWeakRefCount(Value) + SharedRefCount - !!SharedRefCount;
+	return GetWeakRefCount(RefCountsAndFlags) + SharedRefCount - !!SharedRefCount;
+}
+
+inline bool FBufferOwner::IsUniqueOwnedMutable() const
+{
+	const uint64 RefCountsAndFlags = ReferenceCountsAndFlags.load(std::memory_order_relaxed);
+	return GetTotalRefCount(RefCountsAndFlags) == 1 &&
+		(GetFlags(RefCountsAndFlags) & (EBufferOwnerFlags::Owned | EBufferOwnerFlags::Immutable)) == EBufferOwnerFlags::Owned;
 }
 
 inline void FBufferOwner::AddSharedReference()
 {
 	const uint64 PreviousValue = ReferenceCountsAndFlags.fetch_add(SetSharedRefCount(1), std::memory_order_relaxed);
-	checkSlow(GetSharedRefCount(PreviousValue) < 0x7fffffff);
+	checkSlow(GetSharedRefCount(PreviousValue) < RefCountMask);
 	if (GetSharedRefCount(PreviousValue) == 0)
 	{
 		AddWeakReference();
@@ -663,7 +706,7 @@ inline bool FBufferOwner::TryAddSharedReference()
 inline void FBufferOwner::AddWeakReference()
 {
 	const uint64 PreviousValue = ReferenceCountsAndFlags.fetch_add(SetWeakRefCount(1), std::memory_order_relaxed);
-	checkSlow(GetWeakRefCount(PreviousValue) < 0x7fffffff);
+	checkSlow(GetWeakRefCount(PreviousValue) < RefCountMask);
 }
 
 inline void FBufferOwner::ReleaseWeakReference()
@@ -678,7 +721,7 @@ inline void FBufferOwner::ReleaseWeakReference()
 
 ///////////////////////////////////////////////////////////////////////////////////////////////////////////////////////
 
-namespace BufferOwnerPrivate
+namespace UE::SharedBuffer::Private
 {
 
 template <typename FOps>
@@ -807,11 +850,11 @@ inline bool TBufferOwnerPtr<FOps>::operator!=(const TBufferOwnerPtr<FOtherOps>&
 	return Owner != Ptr.Owner;
 }
 
-} // BufferOwnerPrivate
+} // UE::SharedBuffer::Private
 
 ///////////////////////////////////////////////////////////////////////////////////////////////////////////////////////
 
-namespace BufferOwnerPrivate
+namespace UE::SharedBuffer::Private
 {
 
 template <typename T, typename Allocator>
@@ -835,13 +878,13 @@ private:
 	TArray<T, Allocator> Array;
 };
 
-} // BufferOwnerPrivate
+} // UE::SharedBuffer::Private
 
 /** Construct a shared buffer by taking ownership of an array. */
 template <typename T, typename Allocator>
-UE_NODISCARD inline FSharedBuffer MakeSharedBufferFromArray(TArray<T, Allocator>&& Array)
+[[nodiscard]] inline FSharedBuffer MakeSharedBufferFromArray(TArray<T, Allocator>&& Array)
 {
-	return FSharedBuffer(new BufferOwnerPrivate::TBufferOwnerTArray<T, Allocator>(MoveTemp(Array)));
+	return FSharedBuffer(new UE::SharedBuffer::Private::TBufferOwnerTArray<T, Allocator>(MoveTemp(Array)));
 }
 
 ///////////////////////////////////////////////////////////////////////////////////////////////////////////////////////
