When copying FHashedName, avoid making shallow copy of frozen debug string.  Instead make sure the copy is pointing to an interned string

diff --git a/Engine/Source/Runtime/Core/Private/Serialization/MemoryImage.cpp b/Engine/Source/Runtime/Core/Private/Serialization/MemoryImage.cpp
--- a/Engine/Source/Runtime/Core/Private/Serialization/MemoryImage.cpp
+++ b/Engine/Source/Runtime/Core/Private/Serialization/MemoryImage.cpp
@@ -1076,6 +1076,44 @@ public:
 		return Result.String;
 	}
 
+	const char* RegisterStringUTF8(const char* InStringUTF8, int32 InLength, uint64 InHash)
+	{
+		const char* Result = nullptr;
+		{
+			FReadScopeLock ReadLock(Lock);
+			const FStringEntry* Entry = Entries.Find(InHash);
+			if (Entry)
+			{
+				Result = Entry->String;
+			}
+		}
+
+		if (!Result)
+		{
+			FWriteScopeLock WriteLock(Lock);
+			FStringEntry* Entry = Entries.Find(InHash);
+			if (Entry)
+			{
+				Result = Entry->String;
+			}
+			else
+			{
+				char* InternedString = (char*)MemStack.Alloc(InLength + 1, 4);
+				FMemory::Memcpy(InternedString, InStringUTF8, InLength);
+				InternedString[InLength] = 0;
+
+				const FUTF8ToTCHAR NameTCHAR(InStringUTF8);
+
+				Entry = &Entries.Add(InHash);
+				Entry->String = InternedString;
+				Entry->Name = NameTCHAR.Get();
+				Result = InternedString;
+			}
+		}
+
+		return Result;
+	}
+
 	struct FStringEntry
 	{
 		const char* String = nullptr;
@@ -1180,6 +1218,24 @@ FHashedName::FHashedName(const FName& InName)
 	}
 }
 
+FHashedName::FHashedName(const FHashedName& InName) : Hash(InName.Hash)
+{
+#if WITH_EDITORONLY_DATA
+	if (InName.DebugString.String.IsFrozen())
+	{
+		// If the source debug string is frozen, we need to make ourselved an interned unfrozen copy
+		// It's not safe to point directly to the frozen source, as there's no guarantee of its lifetime
+		const char* StringUTF8 = InName.DebugString.String.Get();
+		DebugString.String = FHashedNameRegistry::Get().RegisterStringUTF8(StringUTF8, FCStringAnsi::Strlen(StringUTF8), InName.Hash);
+	}
+	else
+	{
+		// If the source isn't frozen, it must already be interned, safe to copy the pointer directly
+		DebugString.String = InName.DebugString.String;
+	}
+#endif
+}
+
 void FPointerTableBase::SaveToArchive(FArchive& Ar, const FPlatformTypeLayoutParameters& LayoutParams, const void* FrozenObject) const
 {
 	int32 NumDependencies = 0;
diff --git a/Engine/Source/Runtime/Core/Public/Serialization/MemoryImage.h b/Engine/Source/Runtime/Core/Public/Serialization/MemoryImage.h
--- a/Engine/Source/Runtime/Core/Public/Serialization/MemoryImage.h
+++ b/Engine/Source/Runtime/Core/Public/Serialization/MemoryImage.h
@@ -302,10 +302,10 @@ public:
 	inline bool IsValid() const { return UnfrozenPtr != nullptr; }
 	inline bool IsNull() const { return UnfrozenPtr == nullptr; }
 
-	inline TMemoryImagePtr(T* InPtr = nullptr) : UnfrozenPtr(InPtr) {}
-	inline TMemoryImagePtr(const TMemoryImagePtr<T>& InPtr) : UnfrozenPtr(InPtr.Get()) {}
-	inline TMemoryImagePtr& operator=(T* InPtr) { UnfrozenPtr = InPtr; check((Frozen.OffsetFromThis & AllFlags) == 0u); return *this; }
-	inline TMemoryImagePtr& operator=(const TMemoryImagePtr<T>& InPtr) { UnfrozenPtr = InPtr.Get(); check((Frozen.OffsetFromThis & AllFlags) == 0u); return *this; }
+	inline TMemoryImagePtr(T* InPtr = nullptr) : UnfrozenPtr(InPtr) { check(((uint64)UnfrozenPtr & AllFlags) == 0u); }
+	inline TMemoryImagePtr(const TMemoryImagePtr<T>& InPtr) : UnfrozenPtr(InPtr.Get()) { check(((uint64)UnfrozenPtr & AllFlags) == 0u); }
+	inline TMemoryImagePtr& operator=(T* InPtr) { UnfrozenPtr = InPtr; check(((uint64)UnfrozenPtr & AllFlags) == 0u); return *this; }
+	inline TMemoryImagePtr& operator=(const TMemoryImagePtr<T>& InPtr) { UnfrozenPtr = InPtr.Get(); check(((uint64)UnfrozenPtr & AllFlags) == 0u); return *this; }
 
 	inline ~TMemoryImagePtr() 
 	{
@@ -758,6 +758,7 @@ class FHashedName
 public:
 	inline FHashedName() : Hash(0u) {}
 	CORE_API explicit FHashedName(uint64 InHash);
+	CORE_API FHashedName(const FHashedName& InName);
 	CORE_API FHashedName(const TCHAR* InString);
 	CORE_API FHashedName(const FString& InString);
 	CORE_API FHashedName(const FName& InName);
