Set the default AA method to FXAA on mobile.

diff --git a/Engine/Source/Runtime/Core/Private/HAL/ConsoleManager.cpp b/Engine/Source/Runtime/Core/Private/HAL/ConsoleManager.cpp
--- a/Engine/Source/Runtime/Core/Private/HAL/ConsoleManager.cpp
+++ b/Engine/Source/Runtime/Core/Private/HAL/ConsoleManager.cpp
@@ -2825,9 +2825,9 @@ static TAutoConsoleVariable<int32> CVarSkyLightingQuality(
 	ECVF_Scalability | ECVF_RenderThreadSafe
 );
 
-static TAutoConsoleVariable<int32> CVarDefaultAntiAliasing(
+static TAutoConsoleVariable<int32> CVarMobileDefaultAntiAliasing(
 	TEXT("r.Mobile.AntiAliasing"),
-	2,
+	1,
 	TEXT("Mobile default AntiAliasingMethod\n")
 	TEXT(" 0: off (no anti-aliasing)\n")
 	TEXT(" 1: FXAA (default, faster than TemporalAA but much more shimmering for non static cases)\n")
