Color: make it clear that QuantizeRound is preferred and prepare for deprecation of old Quantize
call "QuantizeFloor" if you want the old Quantize behavior

diff --git a/Engine/Source/Runtime/Core/Private/Math/Color.cpp b/Engine/Source/Runtime/Core/Private/Math/Color.cpp
--- a/Engine/Source/Runtime/Core/Private/Math/Color.cpp
+++ b/Engine/Source/Runtime/Core/Private/Math/Color.cpp
@@ -8,6 +8,7 @@
 #include "Math/Vector.h"
 #include "Math/Vector4.h"
 #include "Math/Float16Color.h"
+#include "Templates/CheckValueCast.h"
 
 // Common colors.
 const FLinearColor FLinearColor::White(1.f,1.f,1.f);
@@ -34,18 +35,14 @@ const FColor FColor::Turquoise(26, 188, 156);
 const FColor FColor::Silver(189, 195, 199);
 const FColor FColor::Emerald(46, 204, 113);
 
-/**
-* Helper used by FColor -> FLinearColor conversion. We don't use a lookup table as unlike pow, multiplication is fast.
-*/
-static const float OneOver255 = 1.0f / 255.0f;
-
-//	FColor->FLinearColor conversion.
+// FLinearColor(FColor) does sRGB -> Linear conversion.
+// to get direct conversion use ReinterpretAsLinear
 FLinearColor::FLinearColor(const FColor& Color)
 {
 	R = sRGBToLinearTable[Color.R];
 	G = sRGBToLinearTable[Color.G];
 	B =	sRGBToLinearTable[Color.B];
-	A =	float(Color.A) * OneOver255;
+	A =	float(Color.A) * (1.0f / 255.0f);
 }
 
 FLinearColor::FLinearColor(const FVector& Vector) :
@@ -70,24 +67,13 @@ FLinearColor::FLinearColor(const FFloat16Color& C)
 	A =	C.A.GetFloat();
 }
 
-FLinearColor FLinearColor::FromSRGBColor(const FColor& Color)
-{
-	FLinearColor LinearColor;
-	LinearColor.R = sRGBToLinearTable[Color.R];
-	LinearColor.G = sRGBToLinearTable[Color.G];
-	LinearColor.B =	sRGBToLinearTable[Color.B];
-	LinearColor.A =	float(Color.A) * OneOver255;
-
-	return LinearColor;
-}
-
 FLinearColor FLinearColor::FromPow22Color(const FColor& Color)
 {
 	FLinearColor LinearColor;
 	LinearColor.R = Pow22OneOver255Table[Color.R];
 	LinearColor.G = Pow22OneOver255Table[Color.G];
 	LinearColor.B =	Pow22OneOver255Table[Color.B];
-	LinearColor.A =	float(Color.A) * OneOver255;
+	LinearColor.A =	float(Color.A) * (1.0f / 255.0f);
 
 	return LinearColor;
 }
@@ -112,10 +98,11 @@ FColor FLinearColor::ToRGBE() const
 		int32 Exponent = 1 + (int32)logbf(Primary);
 		const float Scale = ldexpf(1.f, -Exponent + 8);
 
-		Color.R = (uint8)FMath::Clamp(FMath::TruncToInt(R * Scale), 0, 255);
-		Color.G = (uint8)FMath::Clamp(FMath::TruncToInt(G * Scale), 0, 255);
-		Color.B = (uint8)FMath::Clamp(FMath::TruncToInt(B * Scale), 0, 255);
-		Color.A = (uint8)(FMath::Clamp(Exponent,-128,127) + 128);
+		// no clamp needed, should always fit in uint8 :
+		Color.R = TCheckValueCast<uint8>(FMath::TruncToInt(R * Scale));
+		Color.G = TCheckValueCast<uint8>(FMath::TruncToInt(G * Scale));
+		Color.B = TCheckValueCast<uint8>(FMath::TruncToInt(B * Scale));
+		Color.A = TCheckValueCast<uint8>(Exponent + 128);
 	}
 
 	return Color;
@@ -150,6 +137,16 @@ FColor FLinearColor::ToFColor(const bool bSRGB) const
 
 FColor FLinearColor::Quantize() const
 {
+	// @todo deprecate me
+	// change callers to use QuantizeRound mostly, and explicit QuantizeFloor where needed
+	return QuantizeFloor();
+}
+
+FColor FLinearColor::QuantizeFloor() const
+{
+	// Do not use this for graphics or textures or images, they should use QuantizeRound
+	// replicates behavior of old Quantize() which is deprecated
+	// restoration should be done to bucket centers (+0.5 on dequantize)
 	return FColor(
 		(uint8)FMath::Clamp<int32>(FMath::TruncToInt(R*255.f),0,255),
 		(uint8)FMath::Clamp<int32>(FMath::TruncToInt(G*255.f),0,255),
diff --git a/Engine/Source/Runtime/Core/Public/Math/Color.h b/Engine/Source/Runtime/Core/Public/Math/Color.h
--- a/Engine/Source/Runtime/Core/Public/Math/Color.h
+++ b/Engine/Source/Runtime/Core/Public/Math/Color.h
@@ -51,6 +51,7 @@ struct FLinearColor
 	/**
 	 * Converts an FColor which is assumed to be in sRGB space, into linear color space.
 	 * @param Color The sRGB color that needs to be converted into linear space.
+	 * to get direct conversion use ReinterpretAsLinear
 	 */
 	CORE_API FLinearColor(const FColor& Color);
 
@@ -81,7 +82,10 @@ struct FLinearColor
 	 * Converts an FColor coming from an observed sRGB output, into a linear color.
 	 * @param Color The sRGB color that needs to be converted into linear space.
 	 */
-	CORE_API static FLinearColor FromSRGBColor(const FColor& Color);
+	CORE_API static FLinearColor FromSRGBColor(const FColor& Color)
+	{
+		return FLinearColor(Color);
+	}
 
 	/**
 	 * Converts an FColor coming from an observed Pow(1/2.2) output, into a linear color.
@@ -302,13 +306,27 @@ struct FLinearColor
 	 */
 	static CORE_API FLinearColor LerpUsingHSV( const FLinearColor& From, const FLinearColor& To, const float Progress );
 
-	/** Quantizes the linear color and returns the result as a FColor.  This bypasses the SRGB conversion. */
-	CORE_API FColor Quantize() const;
-
-	/** Quantizes the linear color with rounding and returns the result as a FColor.  This bypasses the SRGB conversion. */
+	/** Quantizes the linear color with rounding and returns the result as a FColor.  This bypasses the SRGB conversion. 
+	* QuantizeRound can be dequantized back to linear with FColor::ReinterpretAsLinear (just /255.f)
+	* this matches the GPU U8<->float conversion spec and should be preferred
+	*/
 	CORE_API FColor QuantizeRound() const;
+	
+	/** Quantizes the linear color and returns the result as a FColor.  This bypasses the SRGB conversion.
+	* Uses floor quantization, which does not match the GPU standard conversion.
+	* Restoration to float should be done with a +0.5 bias to restore to centered buckets.
+	* Do NOT use this for graphics or textures or images, use QuantizeRound instead.
+	*/
+	CORE_API FColor QuantizeFloor() const;
+
+	/** backwards compatible Quantize function name, does QuantizeFloor.
+	* @todo deprecate me
+	*/
+	CORE_API FColor Quantize() const;
 
-	/** Quantizes the linear color and returns the result as a FColor with optional sRGB conversion and quality as goal. */
+	/** Quantizes the linear color and returns the result as a FColor with optional sRGB conversion. 
+	* Clamps in [0,1] range before conversion.
+	*/
 	CORE_API FColor ToFColor(const bool bSRGB) const;
 
 	/**
@@ -529,7 +547,8 @@ public:
 
 	/**
 	 * Reinterprets the color as a linear color.
-	 *
+	 * This is the correct dequantizer for QuantizeRound.
+	 * This matches the GPU spec conversion for U8<->float
 	 * @return The linear color representation.
 	 */
 	FORCEINLINE FLinearColor ReinterpretAsLinear() const
diff --git a/Engine/Source/Runtime/Core/Public/Templates/CheckValueCast.h b/Engine/Source/Runtime/Core/Public/Templates/CheckValueCast.h
new file mode 100644
--- /dev/null
+++ b/Engine/Source/Runtime/Core/Public/Templates/CheckValueCast.h
@@ -0,0 +1,22 @@
+// Copyright Epic Games, Inc. All Rights Reserved.
+
+#pragma once
+
+#include "CoreTypes.h"
+#include "UnrealTemplate.h"
+
+/**
+ * Template to cast to another type and ensure the value is not changed
+ *
+ */
+
+template <typename T_to,typename T_fm>
+T_to TCheckValueCast(T_fm FromValue)
+{
+	// cast to T_to type :
+	T_to ReturnValue = static_cast<T_to>(FromValue);
+	// cast back to T_fm to check equality was preserved :
+	check( static_cast<T_fm>(ReturnValue) == FromValue );
+	return ReturnValue;
+}
+
diff --git a/Engine/Source/Runtime/Engine/Private/ImageUtils.cpp b/Engine/Source/Runtime/Engine/Private/ImageUtils.cpp
--- a/Engine/Source/Runtime/Engine/Private/ImageUtils.cpp
+++ b/Engine/Source/Runtime/Engine/Private/ImageUtils.cpp
@@ -707,7 +707,6 @@ private:
 	template<typename TSourceColorType>
 	void WriteHDRBits(FArchive& Ar, TSourceColorType* SourceTexels)
 	{
-		const FRandomStream RandomStream(0xA1A1);
 		const int32 NumChannels = 4;
 		const int32 SizeX = Size.X;
 		const int32 SizeY = Size.Y;
@@ -735,9 +734,8 @@ private:
 			for (int32 x = 0; x < SizeX; x++)
 			{
 				FLinearColor LinearColor(*SourceTexels);
-				FColor RGBEColor = ToRGBEDithered(LinearColor, RandomStream);
+				FColor RGBEColor = LinearColor.ToRGBE();
 
-				FLinearColor lintest = RGBEColor.FromRGBE();
 				ScanLine[0].Add(RGBEColor.R);
 				ScanLine[1].Add(RGBEColor.G);
 				ScanLine[2].Add(RGBEColor.B);
@@ -775,38 +773,6 @@ private:
 		}
 	}
 
-	/**
-	 * Converts from a linear float color to RGBE as outlined in Gregory Ward's Real Pixels article, Graphics Gems II, page 80.
-	 * Implementation details in https://cbloomrants.blogspot.com/2020/06/widespread-error-in-radiance-hdr-rgbe.html
-	 */ 
-	static FColor ToRGBEDithered(const FLinearColor& ColorIN, const FRandomStream& Rand)
-	{
-		const float R = ColorIN.R;
-		const float G = ColorIN.G;
-		const float B = ColorIN.B;
-		const float Primary = FMath::Max3(R, G, B);
-		FColor	ReturnColor;
-
-		if (Primary < 1E-32)
-		{
-			ReturnColor = FColor(0, 0, 0, 0);
-		}
-		else
-		{
-			// The following replaces a call to frexpf, because frexpf would have a warning for an unused return value.
-			// Additionally, this usage of logbf assumes FLT_RADIX == 2
-			int32 Exponent = 1 + (int32)logbf(Primary);
-			const float Scale = ldexpf(1.f, -Exponent + 8);
-
-			ReturnColor.R = FMath::Clamp(FMath::TruncToInt((R * Scale) + Rand.GetFraction()), 0, 255);
-			ReturnColor.G = FMath::Clamp(FMath::TruncToInt((G * Scale) + Rand.GetFraction()), 0, 255);
-			ReturnColor.B = FMath::Clamp(FMath::TruncToInt((B * Scale) + Rand.GetFraction()), 0, 255);
-			ReturnColor.A = FMath::Clamp(FMath::TruncToInt(Exponent), -128, 127) + 128;
-		}
-
-		return ReturnColor;
-	}
-
 	FIntPoint Size;
 	EPixelFormat Format;
 };
