Remove incorrect uses of FColor::Quantize that should have been QuantizeRound
QuantizeRound matches GPU color conversion convention and is almost all preferred
make more direct requantizers to change bit depth from 16 and 10 to 8
fix incorrect dequantizer in ProxyLOD

diff --git a/Engine/Plugins/Experimental/PanoramicCapture/Source/PanoramicCapture/Private/SceneCapturer.cpp b/Engine/Plugins/Experimental/PanoramicCapture/Source/PanoramicCapture/Private/SceneCapturer.cpp
--- a/Engine/Plugins/Experimental/PanoramicCapture/Source/PanoramicCapture/Private/SceneCapturer.cpp
+++ b/Engine/Plugins/Experimental/PanoramicCapture/Source/PanoramicCapture/Private/SceneCapturer.cpp
@@ -1252,10 +1252,10 @@ void USceneCapturer::Tick( float DeltaTime )
 			{
 				// switch to 8 bit/channel
 				TArray<FColor> CombinedAtlas8bit;
+				CombinedAtlas8bit.Empty( CombinedAtlas.Num() );
 				for (FLinearColor& Color : CombinedAtlas)
 				{
-					FColor t = Color.Quantize();
-					CombinedAtlas8bit.Add(t);
+					CombinedAtlas8bit.Add( Color.QuantizeRound() );
 				}
 
 				// save as png 8bit/channel
diff --git a/Engine/Plugins/Experimental/ProxyLODPlugin/Source/ProxyLOD/Private/ProxyLODMaterialTransferUtilities.cpp b/Engine/Plugins/Experimental/ProxyLODPlugin/Source/ProxyLOD/Private/ProxyLODMaterialTransferUtilities.cpp
--- a/Engine/Plugins/Experimental/ProxyLODPlugin/Source/ProxyLOD/Private/ProxyLODMaterialTransferUtilities.cpp
+++ b/Engine/Plugins/Experimental/ProxyLODPlugin/Source/ProxyLOD/Private/ProxyLODMaterialTransferUtilities.cpp
@@ -1014,7 +1014,7 @@ namespace
 	// Unpack a quantized vector into the [-1,1]x[-1,1]x[-1,1] cube
 	FVector UnPackQuantizedVector(const FColor& QuantizedNormal)
 	{
-		FVector Vector(QuantizedNormal.R / 256.f, QuantizedNormal.G / 256.f, QuantizedNormal.B / 256.f);
+		FVector Vector(QuantizedNormal.R / 255.f, QuantizedNormal.G / 255.f, QuantizedNormal.B / 255.f);
 
 		Vector = 2.f * Vector - FVector(1., 1., 1.);
 		return Vector;
@@ -1024,7 +1024,7 @@ namespace
 	{
 		FVector Tmp = 0.5 * (Vector + FVector(1., 1., 1.));
 		FLinearColor LinearColor(Tmp);
-		return LinearColor.Quantize();
+		return LinearColor.QuantizeRound();
 	}
 
 	// Specialized version for the normal map. 
diff --git a/Engine/Plugins/FX/Niagara/Source/Niagara/Private/NiagaraRendererComponents.cpp b/Engine/Plugins/FX/Niagara/Source/Niagara/Private/NiagaraRendererComponents.cpp
--- a/Engine/Plugins/FX/Niagara/Source/Niagara/Private/NiagaraRendererComponents.cpp
+++ b/Engine/Plugins/FX/Niagara/Source/Niagara/Private/NiagaraRendererComponents.cpp
@@ -77,7 +77,7 @@ void ConvertVariableToType(const FNiagaraVariable& SourceVariable, FNiagaraVaria
 	else if (SourceType == FNiagaraTypeDefinition::GetColorDef() && TargetType == UNiagaraComponentRendererProperties::GetFColorDef())
 	{
 		FLinearColor Data = SourceVariable.GetValue<FLinearColor>();
-		TargetVariable.SetValue<FColor>(Data.Quantize());
+		TargetVariable.SetValue<FColor>(Data.QuantizeRound());
 	}
 	else if (SourceType == FNiagaraTypeDefinition::GetVec3Def() && TargetType == UNiagaraComponentRendererProperties::GetFRotatorDef())
 	{
diff --git a/Engine/Source/Editor/UnrealEd/Private/EditorComponents.cpp b/Engine/Source/Editor/UnrealEd/Private/EditorComponents.cpp
--- a/Engine/Source/Editor/UnrealEd/Private/EditorComponents.cpp
+++ b/Engine/Source/Editor/UnrealEd/Private/EditorComponents.cpp
@@ -635,7 +635,7 @@ void FEditorCommonDrawHelper::DrawOriginAxisLine(FVector* A,FVector* B,float* AX
 	*AX = 0;
 	*BX = 0;
 
-	PDI->DrawLine(*A,*B,FLinearColor(Color.Quantize()),SDPG_World);
+	PDI->DrawLine(*A,*B,Color,SDPG_World);
 }
 
 
diff --git a/Engine/Source/Runtime/Core/Public/Math/Color.h b/Engine/Source/Runtime/Core/Public/Math/Color.h
--- a/Engine/Source/Runtime/Core/Public/Math/Color.h
+++ b/Engine/Source/Runtime/Core/Public/Math/Color.h
@@ -342,6 +342,13 @@ struct FLinearColor
 	{		
 		return R * 0.3f + G * 0.59f + B * 0.11f;
 	}
+	
+	FORCEINLINE float GetLuminance() const 
+	{
+		// duplicate function
+		// @todo remove me
+		return ComputeLuminance();
+	}
 
 	/**
 	 * Returns the maximum value in this color structure
@@ -369,11 +376,6 @@ struct FLinearColor
 		return FMath::Min( FMath::Min( FMath::Min( R, G ), B ), A );
 	}
 
-	FORCEINLINE float GetLuminance() const 
-	{ 
-		return R * 0.3f + G * 0.59f + B * 0.11f; 
-	}
-
 	FString ToString() const
 	{
 		return FString::Printf(TEXT("(R=%f,G=%f,B=%f,A=%f)"),R,G,B,A);
@@ -536,6 +538,50 @@ public:
 	*/
 	static CORE_API FColor MakeFromColorTemperature( float Temp );
 
+	static uint8 Requantize10to8( int Value10 )
+	{
+		check( Value10 >= 0 && Value10 <= 1023 );
+
+		// Dequantize from 10 bit (Value10/1023.f)
+		// requantize to 8 bit with rounding (GPU convention)
+		//  this is the computation we want :
+		// (int)( (Value10/1023.f)*255.f + 0.5f );
+		// this gives the exactly the same results :
+		int Temp = Value10*255 + (1<<9);
+		int Value8 = (Temp + (Temp >> 10)) >> 10;
+		return (uint8)Value8;
+	}
+	
+	static uint8 Requantize16to8(int Value16)
+	{
+		check( Value16 >= 0 && Value16 <= 65535 );
+
+		// Dequantize x from 16 bit (Value16/65535.f)
+		// then requantize to 8 bit with rounding (GPU convention)
+
+		// matches exactly with :
+		//  (int)( (Value16/65535.f) * 255.f + 0.5f );
+		int Value8 = (Value16*255 + 32895)>>16;
+		return (uint8)Value8;
+	}
+
+	/**
+	* Return 8-bit color Quantized from 10-bit RGB , 2-bit A
+	*/
+	static FColor MakeRequantizeFrom1010102( int R, int G, int B, int A )
+	{
+		check( A >= 0 && A <= 3 );
+
+		// requantize 2 bits to 8 ; could bit-replicate or just table lookup :
+		const uint8 Requantize2to8[4] = { 0, 0x55, 0xAA, 0xFF };
+		return FColor(
+			Requantize10to8(R),
+			Requantize10to8(G),
+			Requantize10to8(B),
+			Requantize2to8[A] );
+
+	}
+
 	/**
 	 *	@return a new FColor based of this color with the new alpha value.
 	 *	Usage: const FColor& MyColor = FColorList::Green.WithAlpha(128);
diff --git a/Engine/Source/Runtime/Engine/Private/CheatManager.cpp b/Engine/Source/Runtime/Engine/Private/CheatManager.cpp
--- a/Engine/Source/Runtime/Engine/Private/CheatManager.cpp
+++ b/Engine/Source/Runtime/Engine/Private/CheatManager.cpp
@@ -841,7 +841,7 @@ void UCheatManager::TickCollisionDebug()
 			}
 			else
 			{
-				DrawDebugCapsule(GetWorld(), TraceInfo.HitLocation, TraceInfo.CapsuleHalfHeight, TraceInfo.CapsuleRadius, FQuat::Identity, CurrentColor.Quantize());
+				DrawDebugCapsule(GetWorld(), TraceInfo.HitLocation, TraceInfo.CapsuleHalfHeight, TraceInfo.CapsuleRadius, FQuat::Identity, CurrentColor.QuantizeRound());
 			}
 			DrawDebugDirectionalArrow(GetWorld(), TraceInfo.HitNormalStart, TraceInfo.HitNormalEnd, 5.f, FColor(255,64,64), SDPG_World);
 		}
diff --git a/Engine/Source/Runtime/Engine/Private/LightMap.cpp b/Engine/Source/Runtime/Engine/Private/LightMap.cpp
--- a/Engine/Source/Runtime/Engine/Private/LightMap.cpp
+++ b/Engine/Source/Runtime/Engine/Private/LightMap.cpp
@@ -920,7 +920,7 @@ static void GenerateLightmapMipsAndDilateColor(int32 NumMips, int32 TextureSizeX
 				}
 				else if(Coverage)
 				{
-					DestColor = ( AccumulatedColor / Coverage ).Quantize();
+					DestColor = ( AccumulatedColor / Coverage ).QuantizeRound();
 					DestCoverage = Coverage / (MipFactorX * MipFactorY);
 				}
 				else
@@ -977,7 +977,7 @@ static void GenerateLightmapMipsAndDilateColor(int32 NumMips, int32 TextureSizeX
 
 					if(Coverage)
 					{
-						DestColor = (AccumulatedColor / Coverage).Quantize();
+						DestColor = (AccumulatedColor / Coverage).QuantizeRound();
 						DestCoverage = -1;
 					}
 				}
diff --git a/Engine/Source/Runtime/RHI/Public/RHISurfaceDataConversion.h b/Engine/Source/Runtime/RHI/Public/RHISurfaceDataConversion.h
--- a/Engine/Source/Runtime/RHI/Public/RHISurfaceDataConversion.h
+++ b/Engine/Source/Runtime/RHI/Public/RHISurfaceDataConversion.h
@@ -44,9 +44,9 @@ static inline void ConvertRawR16DataToFColor(uint32 Width, uint32 Height, uint8
 		for (uint32 X = 0; X < Width; X++)
 		{
 			uint16 Value16 = *SrcPtr;
-			float Value = Value16 / (float)(0xffff);
+			int Value = FColor::Requantize16to8(Value16);
 
-			*DestPtr = FLinearColor(Value, Value, Value).Quantize();
+			*DestPtr = FColor(Value, Value, Value);
 			++SrcPtr;
 			++DestPtr;
 		}
@@ -99,12 +99,12 @@ static inline void ConvertRawR10G10B10A2DataToFColor(uint32 Width, uint32 Height
 		FColor* DestPtr = Out + Y * Width;
 		for (uint32 X = 0; X < Width; X++)
 		{
-			*DestPtr = FLinearColor(
-				(float)SrcPtr->R / 1023.0f,
-				(float)SrcPtr->G / 1023.0f,
-				(float)SrcPtr->B / 1023.0f,
-				(float)SrcPtr->A / 3.0f
-			).Quantize();
+			*DestPtr = FColor::MakeRequantizeFrom1010102(
+				SrcPtr->R,
+				SrcPtr->G,
+				SrcPtr->B,
+				SrcPtr->A
+			);
 			++SrcPtr;
 			++DestPtr;
 		}
@@ -280,12 +280,12 @@ static inline void ConvertRawR16G16B16A16DataToFColor(uint32 Width, uint32 Heigh
 		FColor* DestPtr = Out + Y * Width;
 		for (uint32 X = 0; X < Width; X++)
 		{
-			*DestPtr = FLinearColor(
-				(float)SrcPtr->R / 65535.0f,
-				(float)SrcPtr->G / 65535.0f,
-				(float)SrcPtr->B / 65535.0f,
-				(float)SrcPtr->A / 65535.0f
-			).Quantize();
+			*DestPtr = FColor(
+				FColor::Requantize16to8(SrcPtr->R),
+				FColor::Requantize16to8(SrcPtr->G),
+				FColor::Requantize16to8(SrcPtr->B),
+				FColor::Requantize16to8(SrcPtr->A)
+			);
 			++SrcPtr;
 			++DestPtr;
 		}
@@ -300,10 +300,10 @@ static inline void ConvertRawR16G16DataToFColor(uint32 Width, uint32 Height, uin
 		FColor* DestPtr = Out + Y * Width;
 		for (uint32 X = 0; X < Width; X++)
 		{
-			*DestPtr = FLinearColor(
-				(float)SrcPtr->R / 65535.0f,
-				(float)SrcPtr->G / 65535.0f,
-				0).Quantize();
+			*DestPtr = FColor(
+				FColor::Requantize16to8(SrcPtr->R),
+				FColor::Requantize16to8(SrcPtr->G),
+				0);
 			++SrcPtr;
 			++DestPtr;
 		}
