fix GExitPurge order of conditional so that already-destructed class is not accessed

diff --git a/Engine/Source/Runtime/Core/Private/Misc/ConfigCacheIni.cpp b/Engine/Source/Runtime/Core/Private/Misc/ConfigCacheIni.cpp
--- a/Engine/Source/Runtime/Core/Private/Misc/ConfigCacheIni.cpp
+++ b/Engine/Source/Runtime/Core/Private/Misc/ConfigCacheIni.cpp
@@ -655,10 +655,15 @@ FConfigFile::FConfigFile()
 
 FConfigFile::~FConfigFile()
 {
+	// this destructor can run at file scope, static shutdown
 
-	if (FCoreDelegates::OnFConfigDeleted.IsBound() && !(GExitPurge))
+	if ( !GExitPurge )
 	{
-		FCoreDelegates::OnFConfigDeleted.Broadcast(this);
+		// FCoreDelegates::OnFConfigDeleted may already be deleted
+		if ( FCoreDelegates::OnFConfigDeleted.IsBound() )
+		{
+			FCoreDelegates::OnFConfigDeleted.Broadcast(this);
+		}
 	}
 
 	if (SourceConfigFile != nullptr)
@@ -2115,6 +2120,7 @@ FConfigCacheIni::FConfigCacheIni()
 
 FConfigCacheIni::~FConfigCacheIni()
 {
+	// this destructor can run at file scope, static shutdown
 	Flush( 1 );
 }
 
