added FTask alias for TTask<void>, for convenience as it's expected to be the majority of cases

diff --git a/Engine/Source/Runtime/Core/Private/Stats/Stats2.cpp b/Engine/Source/Runtime/Core/Private/Stats/Stats2.cpp
--- a/Engine/Source/Runtime/Core/Private/Stats/Stats2.cpp
+++ b/Engine/Source/Runtime/Core/Private/Stats/Stats2.cpp
@@ -1369,7 +1369,7 @@ void FThreadStats::WaitForStats()
 
 		if (FPlatformProcess::SupportsMultithreading())
 		{
-			static UE::Tasks::TTask<void> LastFramesEvents[MAX_STAT_LAG];
+			static UE::Tasks::FTask LastFramesEvents[MAX_STAT_LAG];
 		
 			{
 				SCOPE_CYCLE_COUNTER(STAT_WaitForStats);
diff --git a/Engine/Source/Runtime/Core/Private/Stats/StatsCommand.cpp b/Engine/Source/Runtime/Core/Private/Stats/StatsCommand.cpp
--- a/Engine/Source/Runtime/Core/Private/Stats/StatsCommand.cpp
+++ b/Engine/Source/Runtime/Core/Private/Stats/StatsCommand.cpp
@@ -2351,7 +2351,7 @@ bool DirectStatsCommand(const TCHAR* Cmd, bool bBlockForCompletion /*= false*/,
 				STAT_FSimpleDelegateGraphTask_StatCmd,
 				STATGROUP_TaskGraphTasks);
 
-			UE::Tasks::TTask<void> Task = GStatsPipe.Launch(UE_SOURCE_LOCATION, [FullCmd, bStatCommand, Ar] { StatCmd(FullCmd, bStatCommand, Ar); });
+			UE::Tasks::FTask Task = GStatsPipe.Launch(UE_SOURCE_LOCATION, [FullCmd, bStatCommand, Ar] { StatCmd(FullCmd, bStatCommand, Ar); });
 			if (bBlockForCompletion)
 			{
 				Task.Wait();
diff --git a/Engine/Source/Runtime/Core/Private/Tests/Tasks/TasksTest.cpp b/Engine/Source/Runtime/Core/Private/Tests/Tasks/TasksTest.cpp
--- a/Engine/Source/Runtime/Core/Private/Tests/Tasks/TasksTest.cpp
+++ b/Engine/Source/Runtime/Core/Private/Tests/Tasks/TasksTest.cpp
@@ -25,14 +25,14 @@ namespace UE { namespace TasksTests
 	void BasicStressTest()
 	{
 		constexpr uint32 SpawnerGroupsNum = 50;
-		TArray<TTask<void>> SpawnerGroups;
+		TArray<FTask> SpawnerGroups;
 		SpawnerGroups.Reserve(SpawnerGroupsNum);
 
 		constexpr uint32 SpawnersPerGroupNum = 100;
 		constexpr uint32 TasksNum = SpawnerGroupsNum * SpawnersPerGroupNum;
-		TArray<TTask<void>> Spawners;
+		TArray<FTask> Spawners;
 		Spawners.AddDefaulted(TasksNum);
-		TArray<TTask<void>> Tasks;
+		TArray<FTask> Tasks;
 		Tasks.AddDefaulted(TasksNum);
 
 		std::atomic<uint32> TasksExecutedNum{ 0 };
@@ -68,17 +68,17 @@ namespace UE { namespace TasksTests
 			));
 		}
 
-		for (TTask<void>& SpawnerGroup : SpawnerGroups)
+		for (FTask& SpawnerGroup : SpawnerGroups)
 		{
 			SpawnerGroup.Wait();
 		}
 
-		for (TTask<void>& Spawner : Spawners)
+		for (FTask& Spawner : Spawners)
 		{
 			Spawner.Wait();
 		}
 
-		for (TTask<void>& Task : Tasks)
+		for (FTask& Task : Tasks)
 		{
 			Task.Wait();
 		}
@@ -111,7 +111,7 @@ namespace UE { namespace TasksTests
 			check(!Event.IsTriggered());
 
 			// check that waiting blocks
-			TTask<void> Task = Launch(UE_SOURCE_LOCATION, [&Event] { Event.Wait(); });
+			FTask Task = Launch(UE_SOURCE_LOCATION, [&Event] { Event.Wait(); });
 			FPlatformProcess::Sleep(0.1f);
 			check(!Task.IsCompleted());
 
@@ -122,7 +122,7 @@ namespace UE { namespace TasksTests
 
 		{	// postpone execution so waiting kicks in first
 			std::atomic<int32> Counter{ 0 };
-			TTask<void> Task = Launch(UE_SOURCE_LOCATION, [&Counter] { ++Counter; FPlatformProcess::Sleep(0.1f); });
+			FTask Task = Launch(UE_SOURCE_LOCATION, [&Counter] { ++Counter; FPlatformProcess::Sleep(0.1f); });
 
 			ensure(!Task.Wait(FTimespan::Zero()));
 			Task.Wait();
@@ -131,7 +131,7 @@ namespace UE { namespace TasksTests
 
 		{	// same but using `FTaskEvent`
 			FTaskEvent Event;
-			TTask<void> Task = Launch(UE_SOURCE_LOCATION, [&Event] { Event.Wait(); });
+			FTask Task = Launch(UE_SOURCE_LOCATION, [&Event] { Event.Wait(); });
 			ensure(!Task.Wait(FTimespan::FromMilliseconds(100)));
 			Event.Trigger();
 			Task.Wait();
@@ -139,7 +139,7 @@ namespace UE { namespace TasksTests
 
 		{	// basic use-case, postpone waiting so the task is executed first
 			std::atomic<bool> Done{ false };
-			TTask<void> Task = Launch(UE_SOURCE_LOCATION, [&Done] { Done = true; });
+			FTask Task = Launch(UE_SOURCE_LOCATION, [&Done] { Done = true; });
 			while (!Task.IsCompleted())
 			{
 				FPlatformProcess::Yield();
@@ -211,7 +211,7 @@ namespace UE { namespace TasksTests
 		}
 
 		{	// free memory occupied by a private task instance, can be required if task instance is held as a member var
-			TTask<void> Task = Launch(UE_SOURCE_LOCATION, [] {});
+			FTask Task = Launch(UE_SOURCE_LOCATION, [] {});
 			Task.Wait();
 			Task = {};
 		}
@@ -238,8 +238,8 @@ namespace UE { namespace TasksTests
 				// `UE_SOURCE_LOCATION` - source file name and line number
 			// launch two tasks in the pipe, they will be executed sequentially, but in parallel
 			// with other tasks (including TaskGraph's old API tasks)
-			Tasks::TTask<void> Task1 = Pipe.Launch(UE_SOURCE_LOCATION, [] {});
-			Tasks::TTask<void> Task2 = Pipe.Launch(UE_SOURCE_LOCATION, [] {});
+			Tasks::FTask Task1 = Pipe.Launch(UE_SOURCE_LOCATION, [] {});
+			Tasks::FTask Task2 = Pipe.Launch(UE_SOURCE_LOCATION, [] {});
 			Task2.Wait(); // wait for `Task2` completion
 		}
 
@@ -252,7 +252,7 @@ namespace UE { namespace TasksTests
 					return Pipe.Launch(TEXT("DoSomething()"), [this] { return DoSomethingImpl(); });
 				}
 
-				TTask<void> DoSomethingElse()
+				FTask DoSomethingElse()
 				{
 					return Pipe.Launch(TEXT("DoSomethingElse()"), [this] { DoSomethingElseImpl(); });
 				}
@@ -297,7 +297,7 @@ namespace UE { namespace TasksTests
 		{	// hold the first piped task execution until the next one is piped to test for non-concurrent execution
 			FPipe Pipe{ UE_SOURCE_LOCATION };
 			bool bTask1Done = false;
-			TTask<void> Task1 = Pipe.Launch(UE_SOURCE_LOCATION, 
+			FTask Task1 = Pipe.Launch(UE_SOURCE_LOCATION, 
 				[&bTask1Done] 
 				{ 
 					FPlatformProcess::Sleep(0.1f); 
@@ -320,7 +320,7 @@ namespace UE { namespace TasksTests
 			FPipe Pipe{ UE_SOURCE_LOCATION };
 			std::atomic<bool> bBlocked;
 			FTaskEvent Event;
-			TTask<void> Task = Pipe.Launch(UE_SOURCE_LOCATION,
+			FTask Task = Pipe.Launch(UE_SOURCE_LOCATION,
 				[&bBlocked, &Event]
 				{
 					bBlocked = true;
@@ -346,14 +346,14 @@ namespace UE { namespace TasksTests
 	void PipeStressTest()
 	{
 		constexpr uint32 SpawnerGroupsNum = 50;
-		TArray<TTask<void>> SpawnerGroups;
+		TArray<FTask> SpawnerGroups;
 		SpawnerGroups.Reserve(SpawnerGroupsNum);
 
 		constexpr uint32 SpawnersPerGroupNum = 100;
 		constexpr uint32 TasksNum = SpawnerGroupsNum * SpawnersPerGroupNum;
-		TArray<TTask<void>> Spawners;
+		TArray<FTask> Spawners;
 		Spawners.AddDefaulted(TasksNum);
-		TArray<TTask<void>> Tasks;
+		TArray<FTask> Tasks;
 		Tasks.AddDefaulted(TasksNum);
 
 		std::atomic<bool> bExecuting{ false };
@@ -399,17 +399,17 @@ namespace UE { namespace TasksTests
 				));
 		}
 
-		for (TTask<void>& SpawnerGroup : SpawnerGroups)
+		for (FTask& SpawnerGroup : SpawnerGroups)
 		{
 			SpawnerGroup.Wait();
 		}
 
-		for (TTask<void>& Spawner : Spawners)
+		for (FTask& Spawner : Spawners)
 		{
 			Spawner.Wait();
 		}
 
-		for (TTask<void>& Task : Tasks)
+		for (FTask& Task : Tasks)
 		{
 			Task.Wait();
 		}
diff --git a/Engine/Source/Runtime/Core/Public/Tasks/Task.h b/Engine/Source/Runtime/Core/Public/Tasks/Task.h
--- a/Engine/Source/Runtime/Core/Public/Tasks/Task.h
+++ b/Engine/Source/Runtime/Core/Public/Tasks/Task.h
@@ -184,4 +184,6 @@ namespace UE { namespace Tasks
 		Task->TryLaunch();
 		return TTask<FResult>{ Task };
 	}
+
+	using FTask = TTask<void>;
 }}
\ No newline at end of file
