VarInt: Add support for FArchive

#rnx

diff --git a/Engine/Source/Runtime/Core/Private/Serialization/VarInt.cpp b/Engine/Source/Runtime/Core/Private/Serialization/VarInt.cpp
new file mode 100644
--- /dev/null
+++ b/Engine/Source/Runtime/Core/Private/Serialization/VarInt.cpp
@@ -0,0 +1,59 @@
+// Copyright Epic Games, Inc. All Rights Reserved.
+
+#include "Serialization/VarInt.h"
+
+#include "Serialization/Archive.h"
+
+int64 ReadVarIntFromArchive(FArchive& Ar)
+{
+	const uint64 Value = ReadVarUIntFromArchive(Ar);
+	return -int64(Value & 1) ^ int64(Value >> 1);
+}
+
+void WriteVarIntToArchive(FArchive& Ar, int64 Value)
+{
+	WriteVarUIntToArchive(Ar, uint64((Value >> 63) ^ (Value << 1)));
+}
+
+void SerializeVarInt(FArchive& Ar, int64& Value)
+{
+	if (Ar.IsLoading())
+	{
+		Value = ReadVarIntFromArchive(Ar);
+	}
+	else
+	{
+		WriteVarIntToArchive(Ar, Value);
+	}
+}
+
+uint64 ReadVarUIntFromArchive(FArchive& Ar)
+{
+	uint8 Buffer[9];
+	Ar.Serialize(Buffer, 1);
+	uint32 Size = MeasureVarUInt(Buffer);
+	if (Size > 1)
+	{
+		Ar.Serialize(Buffer + 1, Size - 1);
+	}
+	return ReadVarUInt(Buffer, Size);
+}
+
+void WriteVarUIntToArchive(FArchive& Ar, uint64 Value)
+{
+	uint8 Buffer[9];
+	const uint32 Size = WriteVarUInt(Value, Buffer);
+	Ar.Serialize(Buffer, Size);
+}
+
+void SerializeVarUInt(FArchive& Ar, uint64& Value)
+{
+	if (Ar.IsLoading())
+	{
+		Value = ReadVarUIntFromArchive(Ar);
+	}
+	else
+	{
+		WriteVarUIntToArchive(Ar, Value);
+	}
+}
diff --git a/Engine/Source/Runtime/Core/Private/Tests/Serialization/VarIntTest.cpp b/Engine/Source/Runtime/Core/Private/Tests/Serialization/VarIntTest.cpp
--- a/Engine/Source/Runtime/Core/Private/Tests/Serialization/VarIntTest.cpp
+++ b/Engine/Source/Runtime/Core/Private/Tests/Serialization/VarIntTest.cpp
@@ -3,6 +3,8 @@
 #include "Serialization/VarInt.h"
 
 #include "Misc/AutomationTest.h"
+#include "Serialization/BufferReader.h"
+#include "Serialization/BufferWriter.h"
 
 #if WITH_DEV_AUTOMATION_TESTS
 
@@ -155,7 +157,15 @@ bool FVarIntSerializeTest::RunTest(const FString& Parameters)
 		uint8 Buffer[5];
 		const uint32 WriteByteCount = WriteVarInt(Value, Buffer);
 		uint32 ReadByteCount = 0;
-		return WriteByteCount <= 5 && ReadVarInt(Buffer, ReadByteCount) == Value && ReadByteCount == WriteByteCount;
+		const bool bBufferPass = WriteByteCount <= 5 && ReadVarInt(Buffer, ReadByteCount) == Value && ReadByteCount == WriteByteCount;
+
+		uint8 ArBuffer[5];
+		FBufferWriter WriteAr(ArBuffer, 5);
+		WriteVarIntToArchive(WriteAr, Value);
+		FBufferReader ReadAr(ArBuffer, 5, /*bFreeOnClose*/ false);
+		const bool bArchivePass = WriteAr.Tell() == WriteByteCount && ReadVarIntFromArchive(ReadAr) == Value && ReadAr.Tell() == ReadByteCount;
+
+		return bBufferPass && bArchivePass;
 	};
 	TestTrue(TEXT("Read/WriteVarInt(0x0000'0000)"), TestSerializeVarInt32(int32(0x0000'0000)));
 	TestTrue(TEXT("Read/WriteVarInt(0x0000'0001)"), TestSerializeVarInt32(int32(0x0000'0001)));
@@ -185,7 +195,15 @@ bool FVarIntSerializeTest::RunTest(const FString& Parameters)
 		uint8 Buffer[5];
 		const uint32 WriteByteCount = WriteVarUInt(Value, Buffer);
 		uint32 ReadByteCount = 0;
-		return WriteByteCount <= 5 && ReadVarUInt(Buffer, ReadByteCount) == Value && ReadByteCount == WriteByteCount;
+		const bool bBufferPass = WriteByteCount <= 5 && ReadVarUInt(Buffer, ReadByteCount) == Value && ReadByteCount == WriteByteCount;
+
+		uint8 ArBuffer[5];
+		FBufferWriter WriteAr(ArBuffer, 5);
+		WriteVarUIntToArchive(WriteAr, Value);
+		FBufferReader ReadAr(ArBuffer, 5, /*bFreeOnClose*/ false);
+		const bool bArchivePass = WriteAr.Tell() == WriteByteCount && ReadVarUIntFromArchive(ReadAr) == Value && ReadAr.Tell() == ReadByteCount;
+
+		return bBufferPass && bArchivePass;
 	};
 	TestTrue(TEXT("Read/WriteVarUInt(0x0000'0000)"), TestSerializeVarUInt32(uint32(0x0000'0000)));
 	TestTrue(TEXT("Read/WriteVarUInt(0x0000'007f)"), TestSerializeVarUInt32(uint32(0x0000'007f)));
@@ -207,7 +225,17 @@ bool FVarIntSerializeTest::RunTest(const FString& Parameters)
 		uint8 Buffer[9];
 		const uint32 WriteByteCount = WriteVarInt(Value, Buffer);
 		uint32 ReadByteCount = 0;
-		return WriteByteCount <= 9 && ReadVarInt(Buffer, ReadByteCount) == Value && ReadByteCount == WriteByteCount;
+		const bool bBufferPass = WriteByteCount <= 9 && ReadVarInt(Buffer, ReadByteCount) == Value && ReadByteCount == WriteByteCount;
+
+		uint8 ArBuffer[9];
+		FBufferWriter WriteAr(ArBuffer, 9);
+		SerializeVarInt(WriteAr, Value);
+		FBufferReader ReadAr(ArBuffer, 9, /*bFreeOnClose*/ false);
+		int64 ReadValue;
+		SerializeVarInt(ReadAr, ReadValue);
+		const bool bArchivePass = WriteAr.Tell() == WriteByteCount && ReadValue == Value && ReadAr.Tell() == ReadByteCount;
+
+		return bBufferPass && bArchivePass;
 	};
 	TestTrue(TEXT("Read/WriteVarInt(0x0000'0000'0000'0000)"), TestSerializeVarInt64(int64(0x0000'0000'0000'0000)));
 	TestTrue(TEXT("Read/WriteVarInt(0x0000'0000'0000'0001)"), TestSerializeVarInt64(int64(0x0000'0000'0000'0001)));
@@ -253,7 +281,17 @@ bool FVarIntSerializeTest::RunTest(const FString& Parameters)
 		uint8 Buffer[9];
 		const uint32 WriteByteCount = WriteVarUInt(Value, Buffer);
 		uint32 ReadByteCount = 0;
-		return WriteByteCount <= 9 && ReadVarUInt(Buffer, ReadByteCount) == Value && ReadByteCount == WriteByteCount;
+		const bool bBufferPass = WriteByteCount <= 9 && ReadVarUInt(Buffer, ReadByteCount) == Value && ReadByteCount == WriteByteCount;
+
+		uint8 ArBuffer[9];
+		FBufferWriter WriteAr(ArBuffer, 9);
+		SerializeVarUInt(WriteAr, Value);
+		FBufferReader ReadAr(ArBuffer, 9, /*bFreeOnClose*/ false);
+		uint64 ReadValue;
+		SerializeVarUInt(ReadAr, ReadValue);
+		const bool bArchivePass = WriteAr.Tell() == WriteByteCount && ReadValue == Value && ReadAr.Tell() == ReadByteCount;
+
+		return bBufferPass && bArchivePass;
 	};
 	TestTrue(TEXT("Read/WriteVarUInt(0x0000'0000'0000'0000)"), TestSerializeVarUInt64(uint64(0x0000'0000'0000'0000)));
 	TestTrue(TEXT("Read/WriteVarUInt(0x0000'0000'0000'007f)"), TestSerializeVarUInt64(uint64(0x0000'0000'0000'007f)));
diff --git a/Engine/Source/Runtime/Core/Public/Serialization/VarInt.h b/Engine/Source/Runtime/Core/Public/Serialization/VarInt.h
--- a/Engine/Source/Runtime/Core/Public/Serialization/VarInt.h
+++ b/Engine/Source/Runtime/Core/Public/Serialization/VarInt.h
@@ -198,3 +198,13 @@ FORCEINLINE uint32 WriteVarInt(int64 InValue, void* OutData)
 	const uint64 Value = uint64((InValue >> 63) ^ (InValue << 1));
 	return WriteVarUInt(Value, OutData);
 }
+
+class FArchive;
+
+CORE_API int64 ReadVarIntFromArchive(FArchive& Ar);
+CORE_API void WriteVarIntToArchive(FArchive& Ar, int64 Value);
+CORE_API void SerializeVarInt(FArchive& Ar, int64& Value);
+
+CORE_API uint64 ReadVarUIntFromArchive(FArchive& Ar);
+CORE_API void WriteVarUIntToArchive(FArchive& Ar, uint64 Value);
+CORE_API void SerializeVarUInt(FArchive& Ar, uint64& Value);
