Array don't volatile store DebugGet in destructor in non-debug builds

diff --git a/Engine/Source/Runtime/Core/Public/Containers/Array.h b/Engine/Source/Runtime/Core/Public/Containers/Array.h
--- a/Engine/Source/Runtime/Core/Public/Containers/Array.h
+++ b/Engine/Source/Runtime/Core/Public/Containers/Array.h
@@ -623,13 +623,19 @@ public:
 	/** Destructor. */
 	~TArray()
 	{
-		DestructItems(GetData(), ArrayNum);
-
+		#if UE_BUILD_DEBUG || UE_BUILD_DEVELOPMENT
 		#if defined(_MSC_VER) && !defined(__clang__)	// Relies on MSVC-specific lazy template instantiation to support arrays of incomplete types
 			// ensure that DebugGet gets instantiated.
+			// this is done to ensure DebugGet is available for the debugger watch window
 			//@todo it would be nice if we had a cleaner solution for DebugGet
 			volatile const ElementType* Dummy = &DebugGet(0);
 		#endif
+		#endif
+		
+		DestructItems(GetData(), ArrayNum);
+
+		// note ArrayNum, ArrayMax and data pointer are not invalidated
+		// they are left unchanged and use-after-destruct will see them the same as before destruct
 	}
 
 	/**
