VarInt: Added comments about signed integer encoding

#rnx

diff --git a/Engine/Source/Runtime/Core/Private/Serialization/VarInt.cpp b/Engine/Source/Runtime/Core/Private/Serialization/VarInt.cpp
--- a/Engine/Source/Runtime/Core/Private/Serialization/VarInt.cpp
+++ b/Engine/Source/Runtime/Core/Private/Serialization/VarInt.cpp
@@ -7,12 +7,13 @@
 int64 ReadVarIntFromArchive(FArchive& Ar)
 {
 	const uint64 Value = ReadVarUIntFromArchive(Ar);
-	return -int64(Value & 1) ^ int64(Value >> 1);
+	return -int64(Value & 1) ^ int64(Value >> 1); // ZigZag Encoding
 }
 
 void WriteVarIntToArchive(FArchive& Ar, int64 Value)
 {
-	WriteVarUIntToArchive(Ar, uint64((Value >> 63) ^ (Value << 1)));
+	Value = uint64((Value >> 63) ^ (Value << 1)); // ZigZag Encoding
+	WriteVarUIntToArchive(Ar, Value);
 }
 
 void SerializeVarInt(FArchive& Ar, int64& Value)
