Repeated emplacement constructor for TStaticArray.

diff --git a/Engine/Source/Runtime/Core/Public/Containers/StaticArray.h b/Engine/Source/Runtime/Core/Public/Containers/StaticArray.h
--- a/Engine/Source/Runtime/Core/Public/Containers/StaticArray.h
+++ b/Engine/Source/Runtime/Core/Public/Containers/StaticArray.h
@@ -16,11 +16,19 @@ public:
 
 	TStaticArray() 
 		: Storage()
-	{}
+	{
+	}
 
 	explicit TStaticArray(const InElementType& DefaultElement)
-		: Storage(TMakeIntegerSequence<uint32, NumElements>(), DefaultElement)
-	{}
+		: Storage(InPlace, TMakeIntegerSequence<uint32, NumElements>(), DefaultElement)
+	{
+	}
+
+	template <typename... ArgTypes>
+	explicit TStaticArray(EInPlace, ArgTypes&&... Args)
+		: Storage(InPlace, TMakeIntegerSequence<uint32, NumElements>(), Forward<ArgTypes>(Args)...)
+	{
+	}
 
 	TStaticArray(TStaticArray&& Other) = default;
 	TStaticArray(const TStaticArray& Other) = default;
@@ -110,8 +118,11 @@ private:
 	struct alignas(Alignment) TArrayStorageElementAligned
 	{
 		TArrayStorageElementAligned() {}
-		TArrayStorageElementAligned(const InElementType& InElement)
-			: Element(InElement)
+
+		// Index is used to achieve pack expansion in TArrayStorage, but is unused here
+		template <typename... ArgTypes>
+		explicit TArrayStorageElementAligned(EInPlace, uint32 /*Index*/, ArgTypes&&... Args)
+			: Element(Forward<ArgTypes>(Args)...)
 		{
 		}
 
@@ -125,10 +136,17 @@ private:
 		{
 		}
 
-		template<uint32... Indices>
-		TArrayStorage(TIntegerSequence<uint32, Indices...>, const InElementType& DefaultElement)
-			: Elements { ((void)Indices, DefaultElement)... } //Integer Sequence pack expansion duplicates DefaultElement NumElements times and the comma operator throws away the index
+		template<uint32... Indices, typename... ArgTypes>
+		explicit TArrayStorage(EInPlace, TIntegerSequence<uint32, Indices...>, ArgTypes&&... Args)
+			: Elements{ TArrayStorageElementAligned(InPlace, Indices, Args...)... }
 		{
+			// The arguments are deliberately not forwarded arguments here, because we're initializing multiple elements
+			// and don't want an argument to be mutated by the first element's constructor and then that moved-from state
+			// be used to construct the remaining elements.
+			//
+			// This'll mean that it'll be a compile error to use move-only types like TUniquePtr when in-place constructing
+			// TStaticArray elements, which is a natural expectation because that TUniquePtr can only transfer ownership to
+			// a single element.
 		}
 
 		TArrayStorageElementAligned Elements[NumElements];
