Avoid delaying process shutdown by up to 500ms due to the crash reporter thread.  This can be significant for worker processes that run at high frequency (eg: SCW).

diff --git a/Engine/Source/Runtime/Core/Private/Windows/WindowsPlatformCrashContext.cpp b/Engine/Source/Runtime/Core/Private/Windows/WindowsPlatformCrashContext.cpp
--- a/Engine/Source/Runtime/Core/Private/Windows/WindowsPlatformCrashContext.cpp
+++ b/Engine/Source/Runtime/Core/Private/Windows/WindowsPlatformCrashContext.cpp
@@ -1067,14 +1067,20 @@ void CreateExceptionInfoString(EXCEPTION_RECORD* ExceptionRecord, TCHAR* OutErro
 class FCrashReportingThread
 {
 private:
+	enum InputEvent
+	{
+		IE_Crash = 0,
+		IE_StopThread,
+
+		IE_MAX
+	};
+
 	/** Thread Id of reporter thread*/
 	DWORD ThreadId;
 	/** Thread handle to reporter thread */
 	HANDLE Thread;
-	/** Stops this thread */
-	FThreadSafeCounter StopTaskCounter;
-	/** Signals that the game has crashed */
-	HANDLE CrashEvent;
+	/** Signals that are consumed inputs - crash, stop request, etc... */
+	HANDLE InputEvents[IE_MAX];
 	/** Event that signals the crash reporting thread has finished processing the crash */
 	HANDLE CrashHandledEvent;
 
@@ -1111,19 +1117,25 @@ private:
 		__try
 #endif
 		{
-			while (StopTaskCounter.GetValue() == 0)
+			bool bStopThreadRequested = true;
+			while (!bStopThreadRequested)
 			{
-				if (WaitForSingleObject(CrashEvent, 500) == WAIT_OBJECT_0)
+				DWORD WaitResult = WaitForMultipleObjects((DWORD)IE_MAX, InputEvents, false, 500);
+				if (WaitResult == (WAIT_OBJECT_0 + IE_Crash))
 				{
 					ResetEvent(CrashHandledEvent);
 					HandleCrashInternal();
 
-					ResetEvent(CrashEvent);
+					ResetEvent(InputEvents[IE_Crash]);
 					// Let the thread that crashed know we're done.
 					SetEvent(CrashHandledEvent);
 
 					break;
 				}
+				else if (WaitResult ==  (WAIT_OBJECT_0 + IE_StopThread))
+				{
+					bStopThreadRequested = true;
+				}
 				
 				if (CrashClientHandle.IsValid() && !FPlatformProcess::IsProcRunning(CrashClientHandle))
 				{
@@ -1151,7 +1163,7 @@ private:
 	/** Called by the destructor to terminate the thread */
 	void Stop()
 	{
-		StopTaskCounter.Increment();
+		SetEvent(InputEvents[IE_StopThread]);
 	}
 
 public:
@@ -1159,7 +1171,6 @@ public:
 	FCrashReportingThread()
 		: ThreadId(0)
 		, Thread(nullptr)
-		, CrashEvent(nullptr)
 		, CrashHandledEvent(nullptr)
 		, ExceptionInfo(nullptr)
 		, CrashingThreadId(0)
@@ -1169,7 +1180,10 @@ public:
 		, CrashMonitorPid(0)
 	{
 		// Synchronization objects
-		CrashEvent = CreateEvent(nullptr, true, 0, nullptr);
+		for (HANDLE& InputEvent : InputEvents)
+		{
+			InputEvent = CreateEvent(nullptr, true, 0, nullptr);
+		}
 		CrashHandledEvent = CreateEvent(nullptr, true, 0, nullptr);
 
 		// Add an exception handler to catch issues during static initialization. It is replaced by the engine handler once
@@ -1220,8 +1234,11 @@ public:
 
 		FCoreDelegates::GetPreMainInitDelegate().RemoveAll(this);
 
-		CloseHandle(CrashEvent);
-		CrashEvent = nullptr;
+		for (HANDLE& InputEvent : InputEvents)
+		{
+			CloseHandle(InputEvent);
+			InputEvent = nullptr;
+		}
 
 		CloseHandle(CrashHandledEvent);
 		CrashHandledEvent = nullptr;
@@ -1280,7 +1297,7 @@ public:
 		ExceptionInfo = InExceptionInfo;
 		CrashingThreadId = GetCurrentThreadId();
 		CrashingThreadHandle = GetCurrentThread();
-		SetEvent(CrashEvent);
+		SetEvent(InputEvents[IE_Crash]);
 	}
 
 	/** The thread that crashed calls this function to wait for the report to be generated */
