Fix compiler warning in StatsMallocProfilerProxy.cpp when building debug with Clang.

diff --git a/Engine/Source/Runtime/Core/Private/Stats/StatsMallocProfilerProxy.cpp b/Engine/Source/Runtime/Core/Private/Stats/StatsMallocProfilerProxy.cpp
--- a/Engine/Source/Runtime/Core/Private/Stats/StatsMallocProfilerProxy.cpp
+++ b/Engine/Source/Runtime/Core/Private/Stats/StatsMallocProfilerProxy.cpp
@@ -132,7 +132,7 @@ void FStatsMallocProfilerProxy::TrackAlloc( void* Ptr, int64 Size, int32 Sequenc
 			{
 				ThreadStats->MemoryMessageScope++;
 				const double InvMB = 1.0f / 1024.0f / 1024.0f;
-				UE_LOG( LogStats, Verbose, TEXT( "ThreadID: %i, Current: %.1f" ), FPlatformTLS::GetCurrentThreadId(), InvMB*(int64)ThreadStats->Packet.StatMessages.Num()*sizeof( FStatMessage ) );
+				UE_LOG(LogStats, Verbose, TEXT("ThreadID: %i, Current: %.1f"), FPlatformTLS::GetCurrentThreadId(), InvMB * double(int64(ThreadStats->Packet.StatMessages.Num()) * sizeof(FStatMessage)));
 				ThreadStats->MemoryMessageScope--;
 			}
 #endif // UE_BUILD_DEBUG
