Unshelved from pending changelist '15666024':

PR #7826: Fixed bug in FIntVector4 that equality operator to ignore the W component. (Contributed by [REDACTED_PERSON])

diff --git a/Engine/Source/Runtime/Core/Public/Math/IntVector.h b/Engine/Source/Runtime/Core/Public/Math/IntVector.h
--- a/Engine/Source/Runtime/Core/Public/Math/IntVector.h
+++ b/Engine/Source/Runtime/Core/Public/Math/IntVector.h
@@ -843,7 +843,7 @@ FORCEINLINE int32& FIntVector4::operator[](int32 ComponentIndex)
 
 FORCEINLINE bool FIntVector4::operator==(const FIntVector4& Other) const
 {
-	return X == Other.X && Y == Other.Y && Z == Other.Z;
+	return X == Other.X && Y == Other.Y && Z == Other.Z && W == Other.W;
 }
 
 
