Color use fast SRGB conversion
also fix Round calling floorf !

diff --git a/Engine/Source/Runtime/Core/Private/Math/Color.cpp b/Engine/Source/Runtime/Core/Private/Math/Color.cpp
--- a/Engine/Source/Runtime/Core/Private/Math/Color.cpp
+++ b/Engine/Source/Runtime/Core/Private/Math/Color.cpp
@@ -61,6 +61,7 @@ FLinearColor::FLinearColor(const FVector4& Vector) :
 
 FLinearColor::FLinearColor(const FFloat16Color& C)
 {
+	// @todo Oodle use 4X half to float SIMD (VectorLoadHalf/VectorStoreHalf)
 	R = C.R.GetFloat();
 	G = C.G.GetFloat();
 	B =	C.B.GetFloat();
@@ -78,6 +79,13 @@ FLinearColor FLinearColor::FromPow22Color(const FColor& Color)
 	return LinearColor;
 }
 
+static inline uint8 ClampU8(int Value)
+{
+	// unsigned check against > 255 also checks < 0
+	uint32 UnsignedValue = (uint32)Value;
+	return ( UnsignedValue > 255 ) ? 255 : (uint8)UnsignedValue;
+}
+
 /**
  * Converts from a linear float color to RGBE as outlined in Gregory Ward's Real Pixels article, Graphics Gems II, page 80.
  * Implementation details in https://cbloomrants.blogspot.com/2020/06/widespread-error-in-radiance-hdr-rgbe.html
@@ -85,11 +93,10 @@ FLinearColor FLinearColor::FromPow22Color(const FColor& Color)
 FColor FLinearColor::ToRGBE() const
 {
 	const float	Primary = FMath::Max3( R, G, B );
-	FColor	Color;
 
 	if( Primary < 1E-32f )
 	{
-		Color = FColor(0,0,0,0);
+		return FColor(0,0,0,0);
 	}
 	else
 	{
@@ -97,48 +104,92 @@ FColor FLinearColor::ToRGBE() const
 		// Additionally, this usage of logbf assumes FLT_RADIX == 2
 		int32 Exponent = 1 + (int32)logbf(Primary);
 		const float Scale = ldexpf(1.f, -Exponent + 8);
-
+		
+		FColor	Color;
 		// no clamp needed, should always fit in uint8 :
-		Color.R = TCheckValueCast<uint8>(FMath::TruncToInt(R * Scale));
-		Color.G = TCheckValueCast<uint8>(FMath::TruncToInt(G * Scale));
-		Color.B = TCheckValueCast<uint8>(FMath::TruncToInt(B * Scale));
+		Color.R = TCheckValueCast<uint8>( (int)(R * Scale));
+		Color.G = TCheckValueCast<uint8>( (int)(G * Scale));
+		Color.B = TCheckValueCast<uint8>( (int)(B * Scale));
 		Color.A = TCheckValueCast<uint8>(Exponent + 128);
+		return Color;
 	}
-
-	return Color;
 }
 
 
-/** Quantizes the linear color and returns the result as a FColor with optional sRGB conversion and quality as goal. */
-FColor FLinearColor::ToFColor(const bool bSRGB) const
-{
-	float FloatR = FMath::Clamp(R, 0.0f, 1.0f);
-	float FloatG = FMath::Clamp(G, 0.0f, 1.0f);
-	float FloatB = FMath::Clamp(B, 0.0f, 1.0f);
-	float FloatA = FMath::Clamp(A, 0.0f, 1.0f);
-
-	if(bSRGB)
-	{
-		FloatR = FloatR <= 0.0031308f ? FloatR * 12.92f : FMath::Pow( FloatR, 1.0f / 2.4f ) * 1.055f - 0.055f;
-		FloatG = FloatG <= 0.0031308f ? FloatG * 12.92f : FMath::Pow( FloatG, 1.0f / 2.4f ) * 1.055f - 0.055f;
-		FloatB = FloatB <= 0.0031308f ? FloatB * 12.92f : FMath::Pow( FloatB, 1.0f / 2.4f ) * 1.055f - 0.055f;
-	}
+// fast Linear to SRGB uint8 conversion
+// https://gist.github.com/rygorous/2203834
+//
+// round-trips exactly
+// quantization bucket boundaries vary by max of 0.11%
+// biggest difference at i = 1
+// thresholds[i] = 0.000456
+// c_linear_float_srgb_thresholds[i] = 0.000455
 
-	FColor Result;
-
-	Result.A = (uint8)FMath::RoundToInt(FloatA * 255.0f);
-	Result.R = (uint8)FMath::RoundToInt(FloatR * 255.0f);
-	Result.G = (uint8)FMath::RoundToInt(FloatG * 255.0f);
-	Result.B = (uint8)FMath::RoundToInt(FloatB * 255.0f);
+typedef union
+{
+    uint32 u;
+    float f;
+} stbir__FP32;
+
+static const uint32 stb_fp32_to_srgb8_tab4[104] = {
+    0x0073000d, 0x007a000d, 0x0080000d, 0x0087000d, 0x008d000d, 0x0094000d, 0x009a000d, 0x00a1000d,
+    0x00a7001a, 0x00b4001a, 0x00c1001a, 0x00ce001a, 0x00da001a, 0x00e7001a, 0x00f4001a, 0x0101001a,
+    0x010e0033, 0x01280033, 0x01410033, 0x015b0033, 0x01750033, 0x018f0033, 0x01a80033, 0x01c20033,
+    0x01dc0067, 0x020f0067, 0x02430067, 0x02760067, 0x02aa0067, 0x02dd0067, 0x03110067, 0x03440067,
+    0x037800ce, 0x03df00ce, 0x044600ce, 0x04ad00ce, 0x051400ce, 0x057b00c5, 0x05dd00bc, 0x063b00b5,
+    0x06970158, 0x07420142, 0x07e30130, 0x087b0120, 0x090b0112, 0x09940106, 0x0a1700fc, 0x0a9500f2,
+    0x0b0f01cb, 0x0bf401ae, 0x0ccb0195, 0x0d950180, 0x0e56016e, 0x0f0d015e, 0x0fbc0150, 0x10630143,
+    0x11070264, 0x1238023e, 0x1357021d, 0x14660201, 0x156601e9, 0x165a01d3, 0x174401c0, 0x182401af,
+    0x18fe0331, 0x1a9602fe, 0x1c1502d2, 0x1d7e02ad, 0x1ed4028d, 0x201a0270, 0x21520256, 0x227d0240,
+    0x239f0443, 0x25c003fe, 0x27bf03c4, 0x29a10392, 0x2b6a0367, 0x2d1d0341, 0x2ebe031f, 0x304d0300,
+    0x31d105b0, 0x34a80555, 0x37520507, 0x39d504c5, 0x3c37048b, 0x3e7c0458, 0x40a8042a, 0x42bd0401,
+    0x44c20798, 0x488e071e, 0x4c1c06b6, 0x4f76065d, 0x52a50610, 0x55ac05cc, 0x5892058f, 0x5b590559,
+    0x5e0c0a23, 0x631c0980, 0x67db08f6, 0x6c55087f, 0x70940818, 0x74a007bd, 0x787d076c, 0x7c330723,
+};
+ 
+static uint8 stbir__linear_to_srgb_uchar_fast(float in)
+{
+    static const stbir__FP32 almostone = { 0x3f7fffff }; // 1-eps
+    static const stbir__FP32 minval = { (127-13) << 23 };
+    uint32 tab,bias,scale,t;
+    stbir__FP32 f;
+ 
+    // Clamp to [2^(-13), 1-eps]; these two values map to 0 and 1, respectively.
+    // The tests are carefully written so that NaNs map to 0, same as in the reference
+    // implementation.
+    if (!(in > minval.f)) // written this way to catch NaNs
+        in = minval.f;
+    if (in > almostone.f)
+        in = almostone.f;
+ 
+    // Do the table lookup and unpack bias, scale
+    f.f = in;
+    tab = stb_fp32_to_srgb8_tab4[(f.u - minval.u) >> 20];
+    bias = (tab >> 16) << 9;
+    scale = tab & 0xffff;
+ 
+    // Grab next-highest mantissa bits and perform linear interpolation
+    t = (f.u >> 12) & 0xff;
+    return (uint8) ((bias + scale*t) >> 16);
+}
 
-	return Result;
+/** Quantizes the linear color and returns the result as a FColor with optional sRGB conversion and quality as goal. */
+FColor FLinearColor::ToFColorSRGB() const
+{
+	return FColor(
+		stbir__linear_to_srgb_uchar_fast(R),
+		stbir__linear_to_srgb_uchar_fast(G),
+		stbir__linear_to_srgb_uchar_fast(B),
+		ClampU8( (int)(0.5f + A*255.f) )
+		);
 }
 
 
 FColor FLinearColor::Quantize() const
 {
-	// @todo deprecate me
-	// change callers to use QuantizeRound mostly, and explicit QuantizeFloor where needed
+	// Quantize is deprecated
+	// almost all callers should use QuantizeRound instead
+	// QuantizeFloor maintains old Quantize behavior
 	return QuantizeFloor();
 }
 
@@ -148,20 +199,21 @@ FColor FLinearColor::QuantizeFloor() const
 	// replicates behavior of old Quantize() which is deprecated
 	// restoration should be done to bucket centers (+0.5 on dequantize)
 	return FColor(
-		(uint8)FMath::Clamp<int32>(FMath::TruncToInt(R*255.f),0,255),
-		(uint8)FMath::Clamp<int32>(FMath::TruncToInt(G*255.f),0,255),
-		(uint8)FMath::Clamp<int32>(FMath::TruncToInt(B*255.f),0,255),
-		(uint8)FMath::Clamp<int32>(FMath::TruncToInt(A*255.f),0,255)
+		ClampU8( (int)(R*255.f) ),
+		ClampU8( (int)(G*255.f) ),
+		ClampU8( (int)(B*255.f) ),
+		ClampU8( (int)(A*255.f) )
 		);
 }
 
 FColor FLinearColor::QuantizeRound() const
 {
+	// do not use FMath::RoundToInt because it call floorf
 	return FColor(
-		(uint8)FMath::Clamp<int32>(FMath::RoundToInt(R*255.f),0,255),
-		(uint8)FMath::Clamp<int32>(FMath::RoundToInt(G*255.f),0,255),
-		(uint8)FMath::Clamp<int32>(FMath::RoundToInt(B*255.f),0,255),
-		(uint8)FMath::Clamp<int32>(FMath::RoundToInt(A*255.f),0,255)
+		ClampU8( (int)(0.5f + R*255.f) ),
+		ClampU8( (int)(0.5f + G*255.f) ),
+		ClampU8( (int)(0.5f + B*255.f) ),
+		ClampU8( (int)(0.5f + A*255.f) )
 		);
 }
 
diff --git a/Engine/Source/Runtime/Core/Public/Math/Color.h b/Engine/Source/Runtime/Core/Public/Math/Color.h
--- a/Engine/Source/Runtime/Core/Public/Math/Color.h
+++ b/Engine/Source/Runtime/Core/Public/Math/Color.h
@@ -327,9 +327,12 @@ struct FLinearColor
 
 	/** Quantizes the linear color and returns the result as a FColor with optional sRGB conversion. 
 	* Clamps in [0,1] range before conversion.
+	* ToFColor(false) is QuantizeRound
 	*/
-	CORE_API FColor ToFColor(const bool bSRGB) const;
-
+	CORE_API FColor ToFColorSRGB() const;
+	
+	FORCEINLINE FColor ToFColor(const bool bSRGB) const;
+	
 	/**
 	 * Returns a desaturated color, with 0 meaning no desaturation and 1 == full desaturation
 	 *
@@ -701,6 +704,19 @@ private:
 };
 DECLARE_INTRINSIC_TYPE_LAYOUT(FColor);
 
+
+FORCEINLINE FColor FLinearColor::ToFColor(const bool bSRGB) const
+{
+	if ( bSRGB )
+	{
+		return ToFColorSRGB();
+	}
+	else
+	{
+		return QuantizeRound();
+	}
+}
+
 FORCEINLINE uint32 GetTypeHash( const FColor& Color )
 {
 	return Color.DWColor();
