Color fix silly broken clamp I just checked in

diff --git a/Engine/Source/Runtime/Core/Private/Math/Color.cpp b/Engine/Source/Runtime/Core/Private/Math/Color.cpp
--- a/Engine/Source/Runtime/Core/Private/Math/Color.cpp
+++ b/Engine/Source/Runtime/Core/Private/Math/Color.cpp
@@ -81,9 +81,7 @@ FLinearColor FLinearColor::FromPow22Color(const FColor& Color)
 
 static inline uint8 ClampU8(int Value)
 {
-	// unsigned check against > 255 also checks < 0
-	uint32 UnsignedValue = (uint32)Value;
-	return ( UnsignedValue > 255 ) ? 255 : (uint8)UnsignedValue;
+	return (uint8) FMath::Clamp(Value,0,255);
 }
 
 /**
