CompactBinary: Added explicit operator bool for array and object types

#rnx

diff --git a/Engine/Source/Runtime/Core/Private/Serialization/CompactBinary.cpp b/Engine/Source/Runtime/Core/Private/Serialization/CompactBinary.cpp
--- a/Engine/Source/Runtime/Core/Private/Serialization/CompactBinary.cpp
+++ b/Engine/Source/Runtime/Core/Private/Serialization/CompactBinary.cpp
@@ -732,6 +732,11 @@ FCbFieldView FCbObjectView::FindViewIgnoreCase(const FAnsiStringView Name) const
 	return FCbFieldView();
 }
 
+FCbObjectView::operator bool() const
+{
+	return GetSize() > sizeof(CompactBinaryPrivate::GEmptyObjectPayload);
+}
+
 uint64 FCbObjectView::GetSize() const
 {
 	return sizeof(ECbFieldType) + GetPayloadSize();
diff --git a/Engine/Source/Runtime/Core/Private/Serialization/CompactBinaryPackage.cpp b/Engine/Source/Runtime/Core/Private/Serialization/CompactBinaryPackage.cpp
--- a/Engine/Source/Runtime/Core/Private/Serialization/CompactBinaryPackage.cpp
+++ b/Engine/Source/Runtime/Core/Private/Serialization/CompactBinaryPackage.cpp
@@ -179,7 +179,7 @@ void FCbAttachment::Save(FArchive& Ar) const
 
 void FCbPackage::SetObject(FCbObject InObject, const FIoHash* InObjectHash, FAttachmentResolver* InResolver)
 {
-	if (InObject.CreateViewIterator())
+	if (InObject)
 	{
 		Object = InObject.IsOwned() ? MoveTemp(InObject) : FCbObject::Clone(InObject);
 		if (InObjectHash)
@@ -292,7 +292,7 @@ void FCbPackage::Load(FCbFieldIterator& Fields)
 			Object = Fields.AsObject();
 			Object.MakeOwned();
 			++Fields;
-			if (Object.CreateViewIterator())
+			if (Object)
 			{
 				ObjectHash = Fields.AsCompactBinaryAttachment();
 				checkf(!Fields.HasError(), TEXT("Object must be followed by a CompactBinaryAttachment with the object hash."));
@@ -346,7 +346,7 @@ void FCbPackage::Load(FArchive& Ar, FCbBufferAllocator Allocator)
 			checkf(ValueField.IsObject(), TEXT("Expected Object, Binary, or Null field when loading a package."));
 			Object = ValueField.AsObject();
 			Object.MakeOwned();
-			if (Object.CreateViewIterator())
+			if (Object)
 			{
 				FCbField HashField = LoadCompactBinary(Ar, StackAllocator);
 				ObjectHash = HashField.AsCompactBinaryAttachment();
@@ -362,7 +362,7 @@ void FCbPackage::Load(FArchive& Ar, FCbBufferAllocator Allocator)
 
 void FCbPackage::Save(FCbWriter& Writer) const
 {
-	if (Object.CreateViewIterator())
+	if (Object)
 	{
 		Writer.AddObject(Object);
 		Writer.AddCompactBinaryAttachment(ObjectHash);
diff --git a/Engine/Source/Runtime/Core/Private/Serialization/CompactBinaryValidation.cpp b/Engine/Source/Runtime/Core/Private/Serialization/CompactBinaryValidation.cpp
--- a/Engine/Source/Runtime/Core/Private/Serialization/CompactBinaryValidation.cpp
+++ b/Engine/Source/Runtime/Core/Private/Serialization/CompactBinaryValidation.cpp
@@ -455,7 +455,7 @@ static FIoHash ValidateCbPackageObject(FCbFieldView& Value, FMemoryView& View, E
 		const FIoHash Hash = HashField.AsAttachment();
 		if (EnumHasAnyFlags(Mode, ECbValidateMode::Package))
 		{
-			if (!Object.CreateViewIterator())
+			if (!Object)
 			{
 				AddError(Error, ECbValidateError::NullPackageObject);
 			}
diff --git a/Engine/Source/Runtime/Core/Private/Tests/Serialization/CompactBinaryTest.cpp b/Engine/Source/Runtime/Core/Private/Tests/Serialization/CompactBinaryTest.cpp
--- a/Engine/Source/Runtime/Core/Private/Tests/Serialization/CompactBinaryTest.cpp
+++ b/Engine/Source/Runtime/Core/Private/Tests/Serialization/CompactBinaryTest.cpp
@@ -1436,6 +1436,10 @@ bool FCbFieldIterateAttachmentsTest::RunTest(const FString& Parameters)
 IMPLEMENT_CUSTOM_SIMPLE_AUTOMATION_TEST(FCbFieldBufferTest, FCbFieldTestBase, "System.Core.Serialization.CbFieldBuffer", CompactBinaryTestFlags)
 bool FCbFieldBufferTest::RunTest(const FString& Parameters)
 {
+	static_assert(std::is_constructible<FCbField>::value, "Missing constructor for FCbField");
+	static_assert(std::is_constructible<FCbField, const FCbField&>::value, "Missing constructor for FCbField");
+	static_assert(std::is_constructible<FCbField, FCbField&&>::value, "Missing constructor for FCbField");
+
 	static_assert(std::is_constructible<FCbField, const FSharedBuffer&>::value, "Missing constructor for FCbField");
 	static_assert(std::is_constructible<FCbField, FSharedBuffer&&>::value, "Missing constructor for FCbField");
 
@@ -1562,6 +1566,10 @@ bool FCbFieldBufferTest::RunTest(const FString& Parameters)
 IMPLEMENT_SIMPLE_AUTOMATION_TEST(FCbArrayBufferTest, "System.Core.Serialization.CbArrayBuffer", CompactBinaryTestFlags)
 bool FCbArrayBufferTest::RunTest(const FString& Parameters)
 {
+	static_assert(std::is_constructible<FCbArray>::value, "Missing constructor for FCbArray");
+	static_assert(std::is_constructible<FCbArray, const FCbArray&>::value, "Missing constructor for FCbArray");
+	static_assert(std::is_constructible<FCbArray, FCbArray&&>::value, "Missing constructor for FCbArray");
+
 	static_assert(std::is_constructible<FCbArray, const FCbArrayView&, const FSharedBuffer&>::value, "Missing constructor for FCbArray");
 	static_assert(std::is_constructible<FCbArray, const FCbArrayView&, const FCbFieldIterator&>::value, "Missing constructor for FCbArray");
 	static_assert(std::is_constructible<FCbArray, const FCbArrayView&, const FCbField&>::value, "Missing constructor for FCbArray");
@@ -1590,6 +1598,10 @@ bool FCbArrayBufferTest::RunTest(const FString& Parameters)
 IMPLEMENT_SIMPLE_AUTOMATION_TEST(FCbObjectBufferTest, "System.Core.Serialization.CbObjectBuffer", CompactBinaryTestFlags)
 bool FCbObjectBufferTest::RunTest(const FString& Parameters)
 {
+	static_assert(std::is_constructible<FCbObject>::value, "Missing constructor for FCbObject");
+	static_assert(std::is_constructible<FCbObject, const FCbObject&&>::value, "Missing constructor for FCbObject");
+	static_assert(std::is_constructible<FCbObject, FCbObject&&>::value, "Missing constructor for FCbObject");
+
 	static_assert(std::is_constructible<FCbObject, const FCbObjectView&, const FSharedBuffer&>::value, "Missing constructor for FCbObject");
 	static_assert(std::is_constructible<FCbObject, const FCbObjectView&, const FCbFieldIterator&>::value, "Missing constructor for FCbObject");
 	static_assert(std::is_constructible<FCbObject, const FCbObjectView&, const FCbField&>::value, "Missing constructor for FCbObject");
diff --git a/Engine/Source/Runtime/Core/Public/Serialization/CompactBinary.h b/Engine/Source/Runtime/Core/Public/Serialization/CompactBinary.h
--- a/Engine/Source/Runtime/Core/Public/Serialization/CompactBinary.h
+++ b/Engine/Source/Runtime/Core/Public/Serialization/CompactBinary.h
@@ -54,7 +54,7 @@
  *         }
  *     }
  *
- *     for (FCbFieldView Format : Params["Formats"].AsArrayView())
+ *     for (FCbFieldView Format : Params.FindView("Formats"_ASV))
  *     {
  *         BeginCompress(FName(Format.AsString()));
  *     }
@@ -901,6 +901,9 @@ public:
 	/** Construct an array from an array field. No type check is performed! */
 	static inline FCbArrayView FromFieldNoCheck(const FCbFieldView& Field) { return FCbArrayView(Field); }
 
+	/** Whether the array has any fields. */
+	inline explicit operator bool() const { return Num() > 0; }
+
 	/** Returns the size of the array in bytes if serialized by itself with no name. */
 	CORE_API uint64 GetSize() const;
 
@@ -998,6 +1001,9 @@ public:
 	/** Construct an object from an object field. No type check is performed! */
 	static inline FCbObjectView FromFieldNoCheck(const FCbFieldView& Field) { return FCbObjectView(Field); }
 
+	/** Whether the object has any fields. */
+	CORE_API explicit operator bool() const;
+
 	/** Returns the size of the object in bytes if serialized by itself with no name. */
 	CORE_API uint64 GetSize() const;
 
diff --git a/Engine/Source/Runtime/Core/Public/Serialization/CompactBinaryPackage.h b/Engine/Source/Runtime/Core/Public/Serialization/CompactBinaryPackage.h
--- a/Engine/Source/Runtime/Core/Public/Serialization/CompactBinaryPackage.h
+++ b/Engine/Source/Runtime/Core/Public/Serialization/CompactBinaryPackage.h
@@ -214,10 +214,7 @@ public:
 	inline explicit operator bool() const { return !IsNull(); }
 
 	/** Whether the package has an empty object and no attachments. */
-	inline bool IsNull() const
-	{
-		return !Object.CreateViewIterator() && Attachments.Num() == 0;
-	}
+	inline bool IsNull() const { return !Object && Attachments.IsEmpty(); }
 
 	/** Returns the compact binary object for the package. */
 	inline const FCbObject& GetObject() const { return Object; }
