CompactBinary: Made IsHash treat an attachment as a hash

#rnx

diff --git a/Engine/Source/Runtime/Core/Private/Tests/Serialization/CompactBinaryTest.cpp b/Engine/Source/Runtime/Core/Private/Tests/Serialization/CompactBinaryTest.cpp
--- a/Engine/Source/Runtime/Core/Private/Tests/Serialization/CompactBinaryTest.cpp
+++ b/Engine/Source/Runtime/Core/Private/Tests/Serialization/CompactBinaryTest.cpp
@@ -1050,6 +1050,18 @@ bool FCbFieldHashTest::RunTest(const FString& Parameters)
 		TestFieldError<ECbFieldType::Hash>(TEXT("Hash, None"), DefaultField, ECbFieldError::TypeError, FIoHash(SequentialBytes));
 	}
 
+	// Test FCbFieldView(CompactBinaryAttachment) as Hash
+	{
+		FCbFieldView Field(SequentialBytes, ECbFieldType::CompactBinaryAttachment);
+		TestField<ECbFieldType::Hash>(TEXT("CompactBinaryAttachment, NonZero, AsHash"), Field, FIoHash(SequentialBytes));
+	}
+
+	// Test FCbFieldView(BinaryAttachment) as Hash
+	{
+		FCbFieldView Field(SequentialBytes, ECbFieldType::BinaryAttachment);
+		TestField<ECbFieldType::Hash>(TEXT("BinaryAttachment, NonZero, AsHash"), Field, FIoHash(SequentialBytes));
+	}
+
 	return true;
 }
 
diff --git a/Engine/Source/Runtime/Core/Public/Serialization/CompactBinary.h b/Engine/Source/Runtime/Core/Public/Serialization/CompactBinary.h
--- a/Engine/Source/Runtime/Core/Public/Serialization/CompactBinary.h
+++ b/Engine/Source/Runtime/Core/Public/Serialization/CompactBinary.h
@@ -272,7 +272,7 @@ public:
 	static constexpr inline bool IsBinaryAttachment(ECbFieldType Type)        { return GetType(Type) == ECbFieldType::BinaryAttachment; }
 	static constexpr inline bool IsAttachment(ECbFieldType Type)              { return (Type & AttachmentMask) == AttachmentBase; }
 
-	static constexpr inline bool IsHash(ECbFieldType Type)       { return GetType(Type) == ECbFieldType::Hash; }
+	static constexpr inline bool IsHash(ECbFieldType Type)       { return GetType(Type) == ECbFieldType::Hash || IsAttachment(Type); }
 	static constexpr inline bool IsUuid(ECbFieldType Type)       { return GetType(Type) == ECbFieldType::Uuid; }
 
 	static constexpr inline bool IsDateTime(ECbFieldType Type)   { return GetType(Type) == ECbFieldType::DateTime; }
