Fix TSan compilation for UnrealGame

Overall issue is with TSan enabled GFrameCounter becomes:

std::atomic<uint64> GFrameCounter

While with TSan disabled its:

uint64 GFrameCounter

The fix is to avoid setting up an atomic for GFrameCounter. This seems like this would prevent issues TSan can pick up

diff --git a/Engine/Source/Runtime/Core/Private/Misc/CoreGlobals.cpp b/Engine/Source/Runtime/Core/Private/Misc/CoreGlobals.cpp
--- a/Engine/Source/Runtime/Core/Private/Misc/CoreGlobals.cpp
+++ b/Engine/Source/Runtime/Core/Private/Misc/CoreGlobals.cpp
@@ -346,7 +346,7 @@ bool					GIsInitialLoad					= true;
 bool					GEventDrivenLoaderEnabled = false;
 
 /** Steadily increasing frame counter.																		*/
-TSAN_ATOMIC(uint64)		GFrameCounter(0);
+uint64					GFrameCounter					= 0;
 uint64					GLastGCFrame					= 0;
 /** The time input was sampled, in cycles. */
 uint64					GInputTime					= 0;
diff --git a/Engine/Source/Runtime/Core/Private/Misc/MonitoredProcess.cpp b/Engine/Source/Runtime/Core/Private/Misc/MonitoredProcess.cpp
--- a/Engine/Source/Runtime/Core/Private/Misc/MonitoredProcess.cpp
+++ b/Engine/Source/Runtime/Core/Private/Misc/MonitoredProcess.cpp
@@ -21,6 +21,7 @@ FMonitoredProcess::FMonitoredProcess( const FString& InURL, const FString& InPar
 FMonitoredProcess::FMonitoredProcess( const FString& InURL, const FString& InParams, const FString& InWorkingDir, bool InHidden, bool InCreatePipes )
 	: Hidden(InHidden)
 	, Params(InParams)
+	, bIsRunning(false)
 	, URL(InURL)
 	, WorkingDir(InWorkingDir)
 	, bCreatePipes(InCreatePipes)
diff --git a/Engine/Source/Runtime/Core/Public/CoreGlobals.h b/Engine/Source/Runtime/Core/Public/CoreGlobals.h
--- a/Engine/Source/Runtime/Core/Public/CoreGlobals.h
+++ b/Engine/Source/Runtime/Core/Public/CoreGlobals.h
@@ -435,7 +435,7 @@ extern CORE_API bool GEventDrivenLoaderEnabled;
 extern CORE_API bool GIsRetrievingVTablePtr;
 
 /** Steadily increasing frame counter. */
-extern CORE_API TSAN_ATOMIC(uint64) GFrameCounter;
+extern CORE_API uint64 GFrameCounter;
 
 /** GFrameCounter the last time GC was run. */
 extern CORE_API uint64 GLastGCFrame;
diff --git a/Engine/Source/Runtime/Core/Public/Misc/MonitoredProcess.h b/Engine/Source/Runtime/Core/Public/Misc/MonitoredProcess.h
--- a/Engine/Source/Runtime/Core/Public/Misc/MonitoredProcess.h
+++ b/Engine/Source/Runtime/Core/Public/Misc/MonitoredProcess.h
@@ -222,7 +222,7 @@ protected:
 	FRunnableThread* Thread = nullptr;
 
 	// Is the thread running? 
-	TSAN_ATOMIC(bool) bIsRunning = false;
+	TSAN_ATOMIC(bool) bIsRunning;
 
 	// Holds the URL of the executable to launch. */
 	FString URL;
