[TraceLog] Remove callstack worker thread

Removing callstack tracing worker thread. Profiling shows that the performance impact is negligable. Additionally it solves some deadlock issues and makes the code simpler. In order to only trace unique callstack we track emitted ones in a TSherwoodSet. The insertion performance for this collection is very efficient (see https://probablydance.com/2017/02/26/i-wrote-the-fastest-hashtable/) which means offloading it to a worker thread was not giving the expected gains.

#rnx

diff --git a/Engine/Source/Runtime/Core/Private/ProfilingDebugging/Windows/CallstackTrace.cpp b/Engine/Source/Runtime/Core/Private/ProfilingDebugging/Windows/CallstackTrace.cpp
--- a/Engine/Source/Runtime/Core/Private/ProfilingDebugging/Windows/CallstackTrace.cpp
+++ b/Engine/Source/Runtime/Core/Private/ProfilingDebugging/Windows/CallstackTrace.cpp
@@ -100,12 +100,20 @@ enum
 	UWOP_PUSH_MACHFRAME		= 10,	// 1 node
 };
 
+////////////////////////////////////////////////////////////////////////////////
+UE_TRACE_CHANNEL(CallstackChannel);
+
+UE_TRACE_EVENT_BEGIN(Memory, CallstackSpec, NoSync)
+	UE_TRACE_EVENT_FIELD(uint64, Id)
+	UE_TRACE_EVENT_FIELD(uint64[], Frames)
+UE_TRACE_EVENT_END()
+
 ////////////////////////////////////////////////////////////////////////////////
 namespace {
 
-	class FCallstackProcWorker : public FRunnable 
+	class FCallstackTracer 
 	{
-		public:
+	public:
 		struct FBacktraceEntry
 		{
 			enum {MaxStackDepth = 256};
@@ -114,27 +122,29 @@ namespace {
 			uint64	Frames[MaxStackDepth] = { 0 };
 		};
 
-		FCallstackProcWorker()
-			: bRun(true)
+		FCallstackTracer()
 		{
 			KnownSet.Reserve(1024 * 1024 * 2);
 		}
 		
-		uint32 Run() override;
-		void Stop() override { bRun = false; }
-		void AddCallstack(const FBacktraceEntry& Entry);
-		void AddWork(const FBacktraceEntry& Entry);
-
+		void AddCallstack(const FBacktraceEntry& Entry)
+		{
+			FScopeLock _(&ProducerCs);
+			bool bAlreadyAdded = false;
+			KnownSet.Add(Entry.Id, &bAlreadyAdded);
+			if (!bAlreadyAdded)
+			{
+				UE_TRACE_LOG(Memory, CallstackSpec, CallstackChannel)
+					<< CallstackSpec.Id(Entry.Id)
+					<< CallstackSpec.Frames(Entry.Frames, Entry.FrameCount);
+			}
+		}
+	private:	
 		Experimental::TSherwoodSet<uint64> 	KnownSet;
 		FCriticalSection					ProducerCs;
-		atomic_queue::AtomicQueue2<FBacktraceEntry, 256>	Queue;
-		std::atomic<bool>					bRun;
-		std::atomic<bool>					bStarted;
-		static FRunnableThread*				Thread;
 	};
 }
 
-FRunnableThread* FCallstackProcWorker::Thread = nullptr;
 
 ////////////////////////////////////////////////////////////////////////////////
 class FBacktracer
@@ -143,8 +153,6 @@ public:
 							FBacktracer(FMalloc* InMalloc);
 							~FBacktracer();
 	static FBacktracer*		Get();
-	static void				StartWorker();
-	static void				StopWorker();
 	void					AddModule(UPTRINT Base, const TCHAR* Name);
 	void					RemoveModule(UPTRINT Base);
 	void*					GetBacktraceId(void* AddressOfReturnAddress) const;
@@ -180,8 +188,6 @@ private:
 		FModule				Module;
 	};
 
-	void					StartWorkerThread();
-	void					StopWorkerThread();
 	const FFunction*		LookupFunction(UPTRINT Address, FLookupState& State) const;
 	static FBacktracer*		Instance;
 	mutable FCriticalSection Lock;
@@ -189,7 +195,7 @@ private:
 	int32					ModulesNum;
 	int32					ModulesCapacity;
 	FMalloc*				Malloc;
-	FCallstackProcWorker*	ProcessingThreadRunnable;
+	FCallstackTracer*		CallstackTracer;
 #if BACKTRACE_DBGLVL >= 1
 	mutable uint32			NumFpTruncations = 0;
 	mutable uint32			TotalFunctions = 0;
@@ -208,12 +214,8 @@ FBacktracer::FBacktracer(FMalloc* InMalloc)
 	Modules = (FModule*)Malloc->Malloc(sizeof(FModule) * ModulesCapacity);
 
 	Instance = this;
-	ProcessingThreadRunnable = (FCallstackProcWorker*) Malloc->Malloc(sizeof(FCallstackProcWorker));
-	ProcessingThreadRunnable = new(ProcessingThreadRunnable) FCallstackProcWorker();
-
-	// We cannot start worker threads directly on creation. Delay them until the
-	// engine has gotten a little bit further.
-	FCoreDelegates::GetPreMainInitDelegate().AddStatic(FBacktracer::StartWorker);
+	CallstackTracer = (FCallstackTracer*) Malloc->Malloc(sizeof(FCallstackTracer), alignof(FCallstackTracer));
+	CallstackTracer = new(CallstackTracer) FCallstackTracer();
 }
 
 ////////////////////////////////////////////////////////////////////////////////
@@ -225,7 +227,6 @@ FBacktracer::~FBacktracer()
 		Malloc->Free(Module.Functions);
 	}
 
-	StopWorkerThread();
 }
 
 ////////////////////////////////////////////////////////////////////////////////
@@ -233,48 +234,6 @@ FBacktracer* FBacktracer::Get()
 {
 	return Instance;
 }
-////////////////////////////////////////////////////////////////////////////////
-void FBacktracer::StartWorker()
-{
-	if (Instance)
-	{
-		Instance->StartWorkerThread();
-	}
-}
-
-////////////////////////////////////////////////////////////////////////////////
-void FBacktracer::StopWorker()
-{
-	if (Instance)
-	{
-		Instance->StopWorkerThread();
-	}
-}
-
-////////////////////////////////////////////////////////////////////////////////
-void FBacktracer::StartWorkerThread()
-{
-	check(ProcessingThreadRunnable);
-	if (!FCallstackProcWorker::Thread)
-	{
-		FCallstackProcWorker::Thread = FRunnableThread::Create(ProcessingThreadRunnable, TEXT("TraceMemCallstacks"), 0, TPri_BelowNormal);
-	}
-	// The exit() call will just kill all threads, need controlled shutdown
-	FCoreDelegates::OnExit.AddStatic(FBacktracer::StopWorker);
-}
-
-////////////////////////////////////////////////////////////////////////////////
-void FBacktracer::StopWorkerThread()
-{
-	check(ProcessingThreadRunnable);
-	if (FCallstackProcWorker::Thread)
-	{
-		FRunnableThread* WorkerThread = FCallstackProcWorker::Thread;
-		FCallstackProcWorker::Thread = nullptr;
-		ProcessingThreadRunnable->Stop();
-		WorkerThread->WaitForCompletion();
-	}
-}
 
 ////////////////////////////////////////////////////////////////////////////////
 void FBacktracer::AddModule(UPTRINT ModuleBase, const TCHAR* Name)
@@ -593,7 +552,7 @@ const FBacktracer::FFunction* FBacktracer::LookupFunction(UPTRINT Address, FLook
 void* FBacktracer::GetBacktraceId(void* AddressOfReturnAddress) const
 {
 	FLookupState LookupState;
-	FCallstackProcWorker::FBacktraceEntry BacktraceEntry;
+	FCallstackTracer::FBacktraceEntry BacktraceEntry;
 
 	UPTRINT* StackPointer = (UPTRINT*)AddressOfReturnAddress;
 
@@ -654,7 +613,7 @@ void* FBacktracer::GetBacktraceId(void* AddressOfReturnAddress) const
 		StackPointer += Function->RspBias;
 	}
 	// Trunkate callstacks longer than MaxStackDepth
-	while (*StackPointer && FrameIdx < FCallstackProcWorker::FBacktraceEntry::MaxStackDepth);
+	while (*StackPointer && FrameIdx < FCallstackTracer::FBacktraceEntry::MaxStackDepth);
 
 	// Save the collected id
 	BacktraceEntry.Id = BacktraceId;
@@ -672,7 +631,7 @@ void* FBacktracer::GetBacktraceId(void* AddressOfReturnAddress) const
 
 	// Add to queue to be processed. This might block until there is room in the
 	// queue (i.e. the processing thread has caught up processing).
-	ProcessingThreadRunnable->AddWork(BacktraceEntry);
+	CallstackTracer->AddCallstack(BacktraceEntry);
 	return (void*)BacktraceId;
 }
 
@@ -696,63 +655,6 @@ int32 FBacktracer::GetFrameSize(void* FunctionAddress) const
 	return Function->RspBias;
 }
 
-////////////////////////////////////////////////////////////////////////////////
-UE_TRACE_CHANNEL(CallstackChannel);
-
-UE_TRACE_EVENT_BEGIN(Memory, CallstackSpec, NoSync)
-	UE_TRACE_EVENT_FIELD(uint64, Id)
-	UE_TRACE_EVENT_FIELD(uint64[], Frames)
-UE_TRACE_EVENT_END()
-
-////////////////////////////////////////////////////////////////////////////////
-uint32 FCallstackProcWorker::Run()
-{
-	while (bRun)
-	{
-		while (!Queue.was_empty()) 
-		{
-			FBacktraceEntry Entry = Queue.pop();
-			AddCallstack(Entry);
-		}
- 		FPlatformProcess::Yield();
-	}
-	return 0;
-}
-
-////////////////////////////////////////////////////////////////////////////////
-void FCallstackProcWorker::AddCallstack(const FBacktraceEntry& Entry)
-{
-	bool bAlreadyAdded = false;
-	KnownSet.Add(Entry.Id, &bAlreadyAdded);
-	if (!bAlreadyAdded)
-	{
-		UE_TRACE_LOG(Memory, CallstackSpec, CallstackChannel)
-			<< CallstackSpec.Id(Entry.Id)
-			<< CallstackSpec.Frames(Entry.Frames, Entry.FrameCount);
-	}
-}
-
-////////////////////////////////////////////////////////////////////////////////
-void FCallstackProcWorker::AddWork(const FBacktraceEntry& Entry)
-{
-	// The queue only supports single producer, single consumer. So we use a lock
-	// on the producer side.
-	FScopeLock _(&ProducerCs);
-
-	if (Thread) 
-	{
-		while(!Queue.try_push(Entry))
-		{
-			FPlatformProcess::Yield();
-		}
-	}
-	else
-	{
-		// Worker thread has not yet been started, manually process callstack for now.
-		AddCallstack(Entry);
-	}
-}
-
 ////////////////////////////////////////////////////////////////////////////////
 void Modules_Create(FMalloc*);
 void Modules_Subscribe(void (*)(bool, void*, const TCHAR*));
@@ -767,7 +669,7 @@ void Backtracer_Create(FMalloc* Malloc)
 	}
 
 	// Allocate, construct and intentionally leak backtracer
-	void* Alloc = Malloc->Malloc(sizeof(FBacktracer));
+	void* Alloc = Malloc->Malloc(sizeof(FBacktracer), alignof(FBacktracer));
 	new (Alloc) FBacktracer(Malloc);
 
 	Modules_Create(Malloc);
