minor cleanup: FPlatformProcess::CreateSynchEvent is not deprecated so doesn't need deprecation warning disabled. Also moved the check to the correct place.

diff --git a/Engine/Source/Runtime/Core/Private/Misc/EventPool.h b/Engine/Source/Runtime/Core/Private/Misc/EventPool.h
--- a/Engine/Source/Runtime/Core/Private/Misc/EventPool.h
+++ b/Engine/Source/Runtime/Core/Private/Misc/EventPool.h
@@ -103,14 +103,11 @@ public:
 		if (!Result)
 #endif
 		{
-			// FEventPool is allowed to create synchronization events.
-			PRAGMA_DISABLE_DEPRECATION_WARNINGS
 			Result = FPlatformProcess::CreateSynchEvent((PoolType == EEventPoolTypes::ManualReset));
-			PRAGMA_ENABLE_DEPRECATION_WARNINGS
 		}
+		check(Result);
 		Result->AdvanceStats();
 
-		check(Result);
 		return new FSafeRecyclableEvent(Result);
 	}
 
