removed deprecation of FGenericPlatformProcess::CreateSynchEvent. it was deprecated only in the generic implementation which is used only by some platforms. all platform specific implementations didn't deprecate the method and it's used in ~30 places across the engine. While maybe it's a good idea to deprecate it, consistently across all platforms, for now I'll just fix this inconsistency

diff --git a/Engine/Source/Runtime/Core/Public/GenericPlatform/GenericPlatformProcess.h b/Engine/Source/Runtime/Core/Public/GenericPlatform/GenericPlatformProcess.h
--- a/Engine/Source/Runtime/Core/Public/GenericPlatform/GenericPlatformProcess.h
+++ b/Engine/Source/Runtime/Core/Public/GenericPlatform/GenericPlatformProcess.h
@@ -525,7 +525,6 @@ struct CORE_API FGenericPlatformProcess
 	 * @see GetSynchEventFromPool, ReturnSynchEventToPool
 	 */
 	// Message to others in the future, don't try to delete this function as it isn't exactly deprecated, but it should only ever be called from FEventPool::GetEventFromPool()
-	UE_DEPRECATED(4.8, "Please use GetSynchEventFromPool to create a new event, and ReturnSynchEventToPool to release the event.")
 	static class FEvent* CreateSynchEvent(bool bIsManualReset = false);
 
 	/**
