Upgrade FramePro to version 1.10.10

diff --git a/Engine/Source/Runtime/Core/Private/FramePro/FramePro.cpp b/Engine/Source/Runtime/Core/Private/FramePro/FramePro.cpp
--- a/Engine/Source/Runtime/Core/Private/FramePro/FramePro.cpp
+++ b/Engine/Source/Runtime/Core/Private/FramePro/FramePro.cpp
@@ -1,25 +1,17 @@
 /*
-	This software is provided 'as-is', without any express or implied warranty.
-	In no event will the author(s) be held liable for any damages arising from
-	the use of this software.
-
-	Permission is granted to anyone to use this software for any purpose, including
-	commercial applications, and to alter it and redistribute it freely, subject to
-	the following restrictions:
-
-	1. The origin of this software must not be misrepresented; you must not
-	claim that you wrote the original software. If you use this software
-	in a product, an acknowledgment in the product documentation would be
-	appreciated but is not required.
-	2. Altered source versions must be plainly marked as such, and must not be
-	misrepresented as being the original software.
-	3. This notice may not be removed or altered from any source distribution.
-
-	Author: [REDACTED_PERSON]
-	www.puredevsoftware.com
-	[REDACTED_EMAIL]
-
-	Add FramePro.cpp to your project to allow FramePro to communicate with your application.
+Copyright 2019 PureDev Software Limited
+
+Permission to use, copy, modify, and/or distribute this software for any
+purpose with or without fee is hereby granted, provided that the above
+copyright notice and this permission notice appear in all copies.
+
+THE SOFTWARE IS PROVIDED "AS IS" AND THE AUTHOR DISCLAIMS ALL WARRANTIES
+WITH REGARD TO THIS SOFTWARE INCLUDING ALL IMPLIED WARRANTIES OF
+MERCHANTABILITY AND FITNESS. IN NO EVENT SHALL THE AUTHOR BE LIABLE FOR
+ANY SPECIAL, DIRECT, INDIRECT, OR CONSEQUENTIAL DAMAGES OR ANY DAMAGES
+WHATSOEVER RESULTING FROM LOSS OF USE, DATA OR PROFITS, WHETHER IN AN
+ACTION OF CONTRACT, NEGLIGENCE OR OTHER TORTIOUS ACTION, ARISING OUT OF
+OR IN CONNECTION WITH THE USE OR PERFORMANCE OF THIS SOFTWARE.
 */
 
 #if defined(__UNREAL__)
@@ -31,11 +23,15 @@
 //------------------------------------------------------------------------
 #if FRAMEPRO_ENABLED
 
-
+//#ifndef PLATFORM_SWITCH //@EPIC begin
+	#ifndef SOMAXCONN //@EPIC begin
+		#define SOMAXCONN 16
+	#endif //@EPIC: end
+//#endif //@EPIC: end
 
 //------------------------------------------------------------------------
 //
-// FrameProLib.hpp
+// FrameProLib.h
 //
 
 //------------------------------------------------------------------------
@@ -50,7 +46,101 @@
 namespace FramePro
 {
 	//------------------------------------------------------------------------
-	enum { g_FrameProLibVersion = 14 };
+	enum { g_FrameProLibVersion = 17 };
+
+	//------------------------------------------------------------------------
+	void DebugWrite(const char* p_str, ...);
+
+	//------------------------------------------------------------------------
+	class CheckedAllocator : public Allocator
+	{
+	public:
+		CheckedAllocator()
+		:	mp_Allocator(NULL)
+		{
+		}
+
+		bool HasAllocator() const
+		{
+			return mp_Allocator != NULL;
+		}
+
+		void SetAllocator(Allocator* p_allocator)
+		{
+			mp_Allocator = p_allocator;
+		}
+
+		void* Alloc(size_t size)
+		{
+			if (!mp_Allocator)
+				return NULL;
+
+			void* p = mp_Allocator->Alloc(size);
+			if (!p)
+				HandleOOM();
+
+			return p;
+		}
+
+		void Free(void* p)
+		{
+			if(mp_Allocator)
+				mp_Allocator->Free(p);
+		}
+
+	private:
+		void HandleOOM()
+		{
+			DebugWrite("*******************************\n");
+			DebugWrite("FramePro: ERROR. Out of memory!\n");
+			DebugWrite("*******************************\n");
+
+			DebugBreak();
+		}
+
+		//------------------------------------------------------------------------
+		// data
+	private:
+		Allocator* mp_Allocator;
+	};
+
+	//------------------------------------------------------------------------
+	template<class T>
+	inline T* New(Allocator* p_allocator)
+	{
+		T* p = (T*)p_allocator->Alloc(sizeof(T));
+		if (p)
+			new (p)T();
+		return p;
+	}
+
+	//------------------------------------------------------------------------
+	template<class T, typename Targ1>
+	inline T* New(Allocator* p_allocator, Targ1 arg1)
+	{
+		T* p = (T*)p_allocator->Alloc(sizeof(T));
+		if (p)
+			new (p)T(arg1);
+		return p;
+	}
+
+	//------------------------------------------------------------------------
+	template<class T, typename Targ1, typename Targ2, typename TArg3>
+	inline T* New(Allocator* p_allocator, Targ1 arg1, Targ2 arg2, TArg3 arg3)
+	{
+		T* p = (T*)p_allocator->Alloc(sizeof(T));
+		if (p)
+			new (p)T(arg1, arg2, arg3);
+		return p;
+	}
+
+	//------------------------------------------------------------------------
+	template<typename T>
+	inline void Delete(Allocator* p_allocator, T* p)
+	{
+		p->~T();
+		p_allocator->Free(p);
+	}
 
 	//------------------------------------------------------------------------
 	namespace StringLiteralType
@@ -98,9 +188,6 @@ namespace FramePro
 		Platform::SetTLSValue(GetFrameProTLSSlot(), NULL);
 	}
 
-	//------------------------------------------------------------------------
-	void DebugWrite(const char* p_str, ...);
-
 	//------------------------------------------------------------------------
 	inline bool IsPow2(int value)
 	{
@@ -138,80 +225,6 @@ namespace FramePro
 		b = temp;
 	}
 
-	//------------------------------------------------------------------------
-	namespace ThreadState
-	{
-		enum Enum
-		{
-			Initialized = 0,
-			Ready,
-			Running,
-			Standby,
-			Terminated,
-			Waiting,
-			Transition,
-			DeferredReady,
-		};
-	}
-
-	//------------------------------------------------------------------------
-	namespace ThreadWaitReason
-	{
-		enum Enum
-		{
-			Executive = 0,
-			FreePage,
-			PageIn,
-			PoolAllocation,
-			DelayExecution,
-			Suspended,
-			UserRequest,
-			WrExecutive,
-			WrFreePage,
-			WrPageIn,
-			WrPoolAllocation,
-			WrDelayExecution,
-			WrSuspended,
-			WrUserRequest,
-			WrEventPair,
-			WrQueue,
-			WrLpcReceive,
-			WrLpcReply,
-			WrVirtualMemory,
-			WrPageOut,
-			WrRendezvous,
-			WrKeyedEvent,
-			WrTerminated,
-			WrProcessInSwap,
-			WrCpuRateControl,
-			WrCalloutStack,
-			WrKernel,
-			WrResource,
-			WrPushLock,
-			WrMutex,
-			WrQuantumEnd,
-			WrDispatchInt,
-			WrPreempted,
-			WrYieldExecution,
-			WrFastMutex,
-			WrGuardedMutex,
-			WrRundown,
-			MaximumWaitReason,
-		};
-	}
-
-	//------------------------------------------------------------------------
-	struct ContextSwitch
-	{
-		int64 m_Timestamp;
-		int m_ProcessId;
-		int m_CPUId;
-		int m_OldThreadId;
-		int m_NewThreadId;
-		ThreadState::Enum m_OldThreadState;
-		ThreadWaitReason::Enum m_OldThreadWaitReason;
-	};
-
 	//------------------------------------------------------------------------
 	void SPrintf(char* p_buffer, size_t const buffer_size, const char* p_format, ...);
 }
@@ -219,7 +232,7 @@ namespace FramePro
 
 //------------------------------------------------------------------------
 //
-// EventTraceWin32.hpp
+// EventTraceWin32.h
 //
 
 //------------------------------------------------------------------------
@@ -280,7 +293,7 @@ namespace FramePro
 
 //------------------------------------------------------------------------
 //
-// CriticalSection.hpp
+// CriticalSection.h
 //
 
 //------------------------------------------------------------------------
@@ -310,7 +323,7 @@ namespace FramePro
 		//------------------------------------------------------------------------
 		void Enter()
 		{
-			FRAMEPRO_ASSERT(Platform::GetCurrentThreadId() != m_LockedOnThread);
+			FRAMEPRO_ASSERT((uint64)Platform::GetCurrentThreadId() != m_LockedOnThread);
 
 			Platform::TakeLock(m_OSLockMem);
 
@@ -368,7 +381,7 @@ namespace FramePro
 
 //------------------------------------------------------------------------
 //
-// HashMap.hpp
+// HashMap.h
 //
 
 //------------------------------------------------------------------------
@@ -448,7 +461,7 @@ namespace FramePro
 		//------------------------------------------------------------------------
 		// Add a value to this set.
 		// If this set already contains the value does nothing.
-		void Add(const TKey& key, const TValue& value)
+		bool Add(const TKey& key, const TValue& value)
 		{
 			int index = GetItemIndex(key);
 
@@ -460,12 +473,17 @@ namespace FramePro
 			{
 				if(m_Capacity == 0 || m_Count == (m_Margin * m_Capacity) / 256)
 				{
-					Resize(2*m_Capacity);
+					if (!Resize(2 * m_Capacity))
+						return false;
+
 					index = GetItemIndex(key);
 				}
 
 				// make a copy of the value
 				Pair* p_pair = AllocPair();
+				if (!p_pair)
+					return false;
+
 				p_pair->m_Key = key;
 				p_pair->m_Value = value;
 
@@ -474,6 +492,8 @@ namespace FramePro
 
 				++m_Count;
 			}
+
+			return true;
 		}
 
 		//------------------------------------------------------------------------
@@ -503,7 +523,7 @@ namespace FramePro
 		}
 
 		//------------------------------------------------------------------------
-		void Resize(int new_capacity)
+		bool Resize(int new_capacity)
 		{
 			new_capacity = GetNextPow2(new_capacity);
 
@@ -512,7 +532,8 @@ namespace FramePro
 			const int old_capacity = m_Capacity;
 
 			// allocate the new table
-			AllocTable(new_capacity);
+			if (!AllocTable(new_capacity))
+				return false;
 
 			// copy the values from the old to the new table
 			Pair** p_old_pair = p_old_table;
@@ -527,6 +548,8 @@ namespace FramePro
 			}
 
 			mp_Allocator->Free(p_old_table);
+
+			return true;
 		}
 
 		//------------------------------------------------------------------------
@@ -556,7 +579,7 @@ namespace FramePro
 		}
 
 		//------------------------------------------------------------------------
-		void AllocTable(const int capacity)
+		bool AllocTable(const int capacity)
 		{
 			FRAMEPRO_ASSERT(capacity < m_MaxCapacity);
 			m_Capacity = capacity;
@@ -566,8 +589,13 @@ namespace FramePro
 			{
 				const int size = capacity * sizeof(Pair*);
 				mp_Table = (Pair**)mp_Allocator->Alloc(size);
+				if (!mp_Table)
+					return false;
+
 				memset(mp_Table, 0, size);
 			}
+
+			return true;
 		}
 
 		//------------------------------------------------------------------------
@@ -644,6 +672,9 @@ namespace FramePro
 			{
 				// allocate a new pool and link to pool list
 				byte* p_new_pool = (byte*)mp_Allocator->Alloc(m_ItemBlockSize);
+				if (!p_new_pool)
+					return NULL;
+
 				*(byte**)p_new_pool = mp_ItemPool;
 				mp_ItemPool = p_new_pool;
 
@@ -706,7 +737,7 @@ namespace FramePro
 
 //------------------------------------------------------------------------
 //
-// IncrementingBlockAllocator.hpp
+// IncrementingBlockAllocator.h
 //
 
 //------------------------------------------------------------------------
@@ -740,7 +771,7 @@ namespace FramePro
 		size_t GetMemorySize() const { return m_MemorySize; }
 
 	private:
-		void AllocateBlock();
+		bool AllocateBlock();
 
 		//------------------------------------------------------------------------
 		// data
@@ -758,7 +789,7 @@ namespace FramePro
 
 //------------------------------------------------------------------------
 //
-// FrameProString.hpp
+// FrameProString.h
 //
 
 //------------------------------------------------------------------------
@@ -997,7 +1028,8 @@ namespace FramePro
 			FRAMEPRO_ASSERT(!mp_Value);
 			size_t len = strlen(p_value);
 			mp_Value = (char*)mp_Allocator->Alloc(len + 1);
-			StringCopy(mp_Value, len + 1, p_value, len);
+			if(mp_Value)
+				StringCopy(mp_Value, len + 1, p_value, len);
 		}
 
 		//------------------------------------------------------------------------
@@ -1076,10 +1108,15 @@ namespace FramePro
 			FRAMEPRO_ASSERT(mp_Allocator);
 
 			int len = (int)strlen(p_value);
+			
 			mp_Value = (wchar_t*)mp_Allocator->Alloc((len + 1) * sizeof(wchar_t));
+			if (!mp_Value)
+				return;
+
 			for(int i = 0; i < len; ++i)
 				mp_Value[i] = p_value[i];
-			mp_Value[len] = L'\0';
+
+			mp_Value[len - 1] = L'\0';
 		}
 
 		//------------------------------------------------------------------------
@@ -1089,10 +1126,14 @@ namespace FramePro
 			FRAMEPRO_ASSERT(mp_Allocator);
 
 			int len = (int)wcslen(p_value);
+
 			mp_Value = (wchar_t*)mp_Allocator->Alloc((len + 1) * sizeof(wchar_t));
+			if (!mp_Value)
+				return;
+
 			for(int i = 0; i < len; ++i)
 				mp_Value[i] = p_value[i];
-			mp_Value[len] = L'\0';
+			mp_Value[len - 1] = L'\0';
 		}
 
 		//------------------------------------------------------------------------
@@ -1252,6 +1293,9 @@ namespace FramePro
 			{
 				mp_Allocator->Free(mp_EventInfoBuffer);
 				mp_EventInfoBuffer = mp_Allocator->Alloc(buffer_size);
+				if (!mp_EventInfoBuffer)
+					return 1;
+
 				FRAMEPRO_ASSERT(mp_EventInfoBuffer);
 				m_EventInfoBufferSize = buffer_size;
 			}
@@ -1327,7 +1371,8 @@ namespace FramePro
 						CloseHandle(thread);
 					}
 
-					m_ThreadProcessHashMap.Add(process_thread_id, process_id);
+					if (!m_ThreadProcessHashMap.Add(process_thread_id, process_id))
+						return;
 				}
 			}
 
@@ -1396,7 +1441,7 @@ namespace FramePro
 			default:
 			{
 				char temp[128];
-				sprintf_s(temp, "Error code: %lu", error_code);
+				sprintf_s(temp, "Error code: %ld", error_code);
 				error_string = temp;
 			}
 		}
@@ -1455,11 +1500,6 @@ namespace FramePro
 
 		// start the processing thread
 		HANDLE thread = CreateThread(0, 0, TracingThread_Static, this, 0, NULL);
-		if (thread == NULL)
-		{
-			error = "CreateThread returned NULL";
-			return false;
-		}
 		CloseHandle(thread);
 
 		return true;
@@ -1627,14 +1667,14 @@ namespace FramePro
 
 //------------------------------------------------------------------------
 //
-// FrameProTLS.hpp
+// FrameProTLS.h
 //
 
 //------------------------------------------------------------------------
 
 //------------------------------------------------------------------------
 //
-// Socket.hpp
+// Socket.h
 //
 
 //------------------------------------------------------------------------
@@ -1684,7 +1724,7 @@ namespace FramePro
 
 //------------------------------------------------------------------------
 //
-// Packets.hpp
+// Packets.h
 //
 
 //------------------------------------------------------------------------
@@ -1731,14 +1771,13 @@ namespace FramePro
 	};
 
 	//------------------------------------------------------------------------
-	struct SessionDetailsPacket
+	struct SessionInfoPacket
 	{
-		SessionDetailsPacket(StringId name, StringId build_id, StringId date)
-		:	m_PacketType(PacketType::SessionDetailsPacket),
+		SessionInfoPacket(StringId name, StringId value)
+		:	m_PacketType(PacketType::SessionInfoPacket),
 			m_Padding(0),
 			m_Name(name),
-			m_BuildId(build_id),
-			m_Date(date)
+			m_Value(value)
 		{
 		}
 
@@ -1746,8 +1785,7 @@ namespace FramePro
 		PacketType::Enum m_PacketType;
 		int m_Padding;
 		StringId m_Name;
-		StringId m_BuildId;
-		StringId m_Date;
+		StringId m_Value;
 	};
 
 	//------------------------------------------------------------------------
@@ -1866,10 +1904,10 @@ namespace FramePro
 	};
 
 	//------------------------------------------------------------------------
-	struct SessionInfoPacket
+	struct SessionStatsPacket
 	{
-		SessionInfoPacket()
-		:	m_PacketType(PacketType::SessionInfoPacket),
+		SessionStatsPacket()
+		:	m_PacketType(PacketType::SessionStatsPacket),
 			m_Padding(0xffffffff),
 			m_SendBufferSize(0),
 			m_StringMemorySize(0),
@@ -1930,6 +1968,7 @@ namespace FramePro
 		int m_Count;
 		StringId m_Name;
 		int64 m_Value;
+		int64 m_Time;
 	};
 
 	//------------------------------------------------------------------------
@@ -1939,6 +1978,7 @@ namespace FramePro
 		int m_Count;
 		StringId m_Name;
 		double m_Value;
+		int64 m_Time;
 	};
 
 	//------------------------------------------------------------------------
@@ -1991,6 +2031,14 @@ namespace FramePro
 
 	//------------------------------------------------------------------------
 	struct CallstackPacket
+	{
+		PacketType::Enum m_PacketType;
+		int m_CallstackId;
+		int m_CallstackSize;	// size of the callstack that follows in the send buffer, or 0 if we have already sent this callstack
+	};
+
+	//------------------------------------------------------------------------
+	struct TimeSpanCallstackPacket
 	{
 		// we don't have a packet type here because it always follows a time span packet
 		int m_CallstackId;
@@ -2067,7 +2115,7 @@ namespace FramePro
 
 //------------------------------------------------------------------------
 //
-// PointerSet.hpp
+// PointerSet.h
 //
 
 //------------------------------------------------------------------------
@@ -2109,7 +2157,7 @@ namespace FramePro
 		}
 
 	private:
-		void Grow();
+		bool Grow();
 
 		bool AddInternal(const void* p, int64 hash, int index);
 
@@ -2128,7 +2176,7 @@ namespace FramePro
 
 //------------------------------------------------------------------------
 //
-// SendBuffer.hpp
+// SendBuffer.h
 //
 
 //------------------------------------------------------------------------
@@ -2148,7 +2196,7 @@ namespace FramePro
 
 		const void* GetBuffer() const { return mp_Buffer; }
 
-		void AllocateBuffer(int capacity);
+		bool AllocateBuffer(int capacity);
 
 		void ClearBuffer();
 
@@ -2203,7 +2251,7 @@ namespace FramePro
 
 //------------------------------------------------------------------------
 //
-// Buffer.hpp
+// Buffer.h
 //
 
 //------------------------------------------------------------------------
@@ -2266,7 +2314,8 @@ namespace FramePro
 			if(new_size > m_Capacity)
 			{
 				int double_capacity = 2*m_Capacity;
-				Resize(double_capacity > new_size ? double_capacity : new_size);
+				if (!Resize(double_capacity > new_size ? double_capacity : new_size))
+					return NULL;
 			}
 			void* p = (char*)mp_Buffer + old_size;
 			m_Size = new_size;
@@ -2275,9 +2324,14 @@ namespace FramePro
 		}
 
 	private:
-		void Resize(int new_capacity)
+		bool Resize(int new_capacity)
 		{
 			void* p_new_buffer = mp_Allocator->Alloc(new_capacity);
+			if (!p_new_buffer)
+			{
+				mp_Buffer = NULL;
+				return false;
+			}
 
 			int current_size = m_Size;
 			if(current_size)
@@ -2287,6 +2341,8 @@ namespace FramePro
 			mp_Buffer = p_new_buffer;
 
 			m_Capacity = new_capacity;
+
+			return true;
 		}
 
 		//------------------------------------------------------------------------
@@ -2304,7 +2360,7 @@ namespace FramePro
 
 //------------------------------------------------------------------------
 //
-// List.hpp
+// List.h
 //
 
 //------------------------------------------------------------------------
@@ -2459,7 +2515,7 @@ namespace FramePro
 
 //------------------------------------------------------------------------
 //
-// ConditionalParentScope.hpp
+// ConditionalParentScope.h
 //
 
 //------------------------------------------------------------------------
@@ -2497,7 +2553,7 @@ namespace FramePro
 
 //------------------------------------------------------------------------
 //
-// FrameProCallstackSet.hpp
+// FrameProCallstackSet.h
 //
 
 //------------------------------------------------------------------------
@@ -2545,13 +2601,15 @@ namespace FramePro
 
 		Allocator* mp_Allocator;
 		IncrementingBlockAllocator m_BlockAllocator;
+
+		static std::atomic<int> m_NextUniqueId;
 	};
 }
 
 
 //------------------------------------------------------------------------
 //
-// FrameProStackTrace.hpp
+// FrameProStackTrace.h
 //
 
 //------------------------------------------------------------------------
@@ -2563,6 +2621,12 @@ namespace FramePro
 	//------------------------------------------------------------------------
 	struct CallstackResult
 	{
+		CallstackResult()
+		:	mp_Callstack(NULL),
+			m_IsNew(false)
+		{
+		}
+
 		Callstack* mp_Callstack;
 		bool m_IsNew;
 	};
@@ -2627,13 +2691,13 @@ namespace FramePro
 		//---------------------------------------------------
 		// these functions are called from the main thread, so care needs to be taken with thread safety
 
-		void OnConnected(bool recording_to_file);
+		bool OnConnected(bool recording_to_file);
 
 		void OnDisconnected();
 
 		void SendSessionInfoBuffer();
 
-		void OnFrameStart();
+		bool OnFrameStart();
 
 		void FlushSendBuffers();
 
@@ -2686,7 +2750,7 @@ namespace FramePro
 
 		void SetThreadOrder(StringId thread_name);
 
-		void SetMainThread(int main_thraed_id);
+		void SetMainThread(int main_thread_id);
 
 		StringId RegisterString(const char* p_str);
 
@@ -2696,8 +2760,6 @@ namespace FramePro
 
 		FRAMEPRO_NO_INLINE void SendString(const wchar_t* p_string, PacketType::Enum packet_type);
 
-		void SendFrameStartPacket(int64 wait_for_send_complete_time);
-
 		void SendConnectPacket(int64 clock_frequency, int process_id, Platform::Enum platform);
 
 		void SendStringLiteral(StringLiteralType::Enum string_literal_type, StringId string_id);
@@ -2724,7 +2786,7 @@ namespace FramePro
 
 		bool ShuttingDown() const { return m_ShuttingDown; }
 
-		FRAMEPRO_NO_INLINE void FlushCurrentSendBuffer();
+		FRAMEPRO_NO_INLINE bool FlushCurrentSendBuffer();
 
 		void PushConditionalParentScope(const char* p_name, int64 pre_duration, int64 post_duration);
 		
@@ -2746,6 +2808,13 @@ namespace FramePro
 		{
 			FRAMEPRO_ASSERT(IsOnTLSThread());
 
+			#if FRAMEPRO_ALLOW_UNPARENTED_HIRES_SCOPES
+				if(!m_ScopeDepth)
+					return;
+			#else
+				FRAMEPRO_ASSERT(m_ScopeDepth);	// all hires scopes must have a parent
+			#endif
+
 			// try and find the timer of the specified name
 			int count = m_HiResTimers.GetCount();
 			HiResTimer* p_timer = NULL;
@@ -2763,6 +2832,11 @@ namespace FramePro
 			// add the timer if not found
 			if (!p_timer)
 			{
+				if (!m_HiResTimers.GetCount())
+				{
+					FRAMEPRO_GET_CLOCK_COUNT(m_FirstHiresTimerTime);
+				}
+
 				HiResTimer hires_timer;
 				hires_timer.mp_Name = p_name;
 				hires_timer.m_Duration = 0;
@@ -2791,6 +2865,13 @@ namespace FramePro
 		{
 			FRAMEPRO_ASSERT(IsOnTLSThread());
 
+			#if FRAMEPRO_ALLOW_UNPARENTED_HIRES_SCOPES
+				if(!m_ScopeDepth)
+					return;
+			#else
+				FRAMEPRO_ASSERT(m_ScopeDepth);	// all hires scopes must have a parent
+			#endif
+
 			// get time (do this as early as possible)
 			int64 now;
 			FRAMEPRO_GET_CLOCK_COUNT(now);
@@ -2812,17 +2893,15 @@ namespace FramePro
 			return m_HiResTimers.GetCount() != 0;
 		}
 
-		FRAMEPRO_FORCE_INLINE void SubmitHiResTimers(int64 current_time)
+		FRAMEPRO_FORCE_INLINE void SubmitHiResTimers(int64 start_time, int64 end_time)
 		{
 			FRAMEPRO_ASSERT(IsOnTLSThread());
 
 			if (m_HiResTimers.GetCount() != 0)
-				SendHiResTimersScope(current_time);
-
-			m_HiResTimerScopeStartTime = current_time;
+				SendHiResTimersScope(start_time, end_time);
 		}
 
-		FRAMEPRO_NO_INLINE void SendHiResTimersScope(int64 current_time);
+		FRAMEPRO_NO_INLINE void SendHiResTimersScope(int64 start_time, int64 end_time);
 
 #ifdef FRAMEPRO_SCOPE_MIN_TIME
 		int64 GetScopeMinTime() const { return m_ScopeMinTime; }
@@ -2850,6 +2929,15 @@ namespace FramePro
 
 		void SetCustomStatInfo(StringId name, StringId graph, StringId unit, uint colour);
 
+		int GetAndClearCurrentCallstackId(int64 start_time, int64 end_time);
+
+		void SetCurrentCallstackId(int id);
+
+#if FRAMEPRO_DEBUG || FRAMEPRO_ALLOW_UNPARENTED_HIRES_SCOPES
+		void IncScopeDepth() { ++m_ScopeDepth; }
+		void DecScopeDepth() { --m_ScopeDepth; FRAMEPRO_ASSERT(m_ScopeDepth >= 0); }
+#endif
+
 #if FRAMEPRO_ENABLE_CALLSTACKS
 		bool ShouldSendCallstacks() const { return m_SendCallstacks; }
 
@@ -2865,15 +2953,13 @@ namespace FramePro
 
 		void SendString(StringId string_id, const wchar_t* p_str, PacketType::Enum packet_type);
 
-		void ShowMemoryWarning() const;
-
 		void SendSessionInfo(const void* p_data, int size);
 
 		void UpdateStringMemorySize();
 
 		FRAMEPRO_NO_INLINE void FlushCurrentSendBuffer_no_lock();
 
-		void AllocateCurrentSendBuffer();
+		bool AllocateCurrentSendBuffer();
 
 		void FreeCurrentSendBuffer();
 		
@@ -2894,8 +2980,6 @@ namespace FramePro
 
 		void FlushConditionalChildSendBuffers();
 
-		FRAMEPRO_NO_INLINE static void AddHiResTimer(const char* p_name, HiResTimerList* p_timers);
-
 		void PushHiResTimerList();
 
 		void PopHiResTimerList();
@@ -2946,7 +3030,7 @@ namespace FramePro
 		
 		int m_ThreadId;
 
-		int64 m_HiResTimerScopeStartTime;
+		int64 m_FirstHiresTimerTime;
 
 		// everything else
 
@@ -2999,17 +3083,24 @@ namespace FramePro
 
 		char m_FalseSharingSpacerBuffer[128];		// separate TLS classes to avoid false sharing
 
+		int m_CurrentCallstackId;
+		int64 m_CurrentCallstackIdTime;
+
 #if FRAMEPRO_ENABLE_CALLSTACKS
 		StackTrace m_StackTrace;
 		bool m_SendCallstacks;
 #endif
+
+#if FRAMEPRO_DEBUG || FRAMEPRO_ALLOW_UNPARENTED_HIRES_SCOPES
+		int m_ScopeDepth;
+#endif
 	};
 }
 
 
 //------------------------------------------------------------------------
 //
-// Event.hpp
+// Event.h
 //
 
 //------------------------------------------------------------------------
@@ -3054,13 +3145,13 @@ namespace FramePro
 	private:
 		static const int m_OSEventMemMaxSize = 96;
 		mutable char m_OSEventMem[m_OSEventMemMaxSize];
-	};
+	} FRAMEPRO_ALIGN_STRUCT(16);
 }
 
 
 //------------------------------------------------------------------------
 //
-// Thread.hpp
+// Thread.h
 //
 
 //------------------------------------------------------------------------
@@ -3074,7 +3165,7 @@ namespace FramePro
 
 		~Thread();
 
-		void CreateThread(ThreadMain p_thread_main, void* p_param, Allocator* p_allocator);
+		bool CreateThread(ThreadMain p_thread_main, void* p_param, Allocator* p_allocator);
 
 		bool IsAlive() const { return m_Alive; }
 
@@ -3107,7 +3198,7 @@ namespace FramePro
 
 //------------------------------------------------------------------------
 //
-// File.hpp
+// File.h
 //
 
 //------------------------------------------------------------------------
@@ -3225,7 +3316,7 @@ namespace FramePro
 
 //------------------------------------------------------------------------
 //
-// FrameProSession.hpp
+// FrameProSession.h
 //
 
 //------------------------------------------------------------------------
@@ -3255,7 +3346,7 @@ namespace FramePro
 
 		int64 GetClockFrequency();
 
-		void AddFrameProTLS(FrameProTLS* p_framepro_tls);
+		bool AddFrameProTLS(FrameProTLS* p_framepro_tls);
 
 		void RemoveFrameProTLS(FrameProTLS* p_framepro_tls);
 
@@ -3265,12 +3356,16 @@ namespace FramePro
 
 		Allocator* GetAllocator()
 		{
-			Allocator* p_allocator = mp_Allocator;
-			return p_allocator ? p_allocator : CreateDefaultAllocator();
+			if(!m_Allocator.HasAllocator())
+				CreateDefaultAllocator();
+
+			return &m_Allocator;
 		}
 
 		void SetThreadName(const char* p_name);
 
+		void SetThreadName(int thread_id, const char* p_name);
+
 		void StartRecording(const char* p_filename, bool context_switches, bool callstacks, int64 max_file_size);
 
 		void StartRecording(const wchar_t* p_filename, bool context_switches, bool callstacks, int64 max_file_size);
@@ -3285,11 +3380,11 @@ namespace FramePro
 		
 		void SetThreadAffinity(int affinity);
 
-		void SendSessionDetails(const char* p_name, const char* p_build_id);
+		void SendSessionInfo(const char* p_name, const char* p_value);
 
-		void SendSessionDetails(const wchar_t* p_name, const wchar_t* p_build_id);
+		void SendSessionInfo(const wchar_t* p_name, const wchar_t* p_value);
 
-		void AddGlobalHiResTimer(GlobalHiResTimer* p_timer);
+		void AddCustomStatTimer(CustomStatTimer* p_timer);
 
 		bool CallConditionalParentScopeCallback(ConditionalParentScopeCallback p_callback, const char* p_name, int64 start_time, int64 end_time);
 
@@ -3303,34 +3398,36 @@ namespace FramePro
 
 		void SetCustomStatColour(StringId name, uint colour);
 
+		void SendScopeCallstack();
+
 	private:
 		void Initialise(FrameProTLS* p_framepro_tls);
 
 		void StopRecording_NoLock();
 
-		void SendSessionDetails(StringId name, StringId build_id);
+		void SendSessionInfo(StringId name, StringId value);
 
-		Allocator* CreateDefaultAllocator();
+		void CreateDefaultAllocator();
 
 		void SetAllocatorInternal(Allocator* p_allocator);
 
-		void InitialiseConnection(FrameProTLS* p_framepro_tls);
+		bool InitialiseConnection(FrameProTLS* p_framepro_tls);
 
 		FRAMEPRO_FORCE_INLINE void CalculateTimerFrequency();
 
 		void SetConnected(bool value);
 
-		void WriteSendBuffer(SendBuffer* p_send_buffer, File& file, int64& file_size);
-		
+		bool Send(const void* p_data, int size);
+
 		void SendFrameBuffer();
 
 		static int StaticSendThreadMain(void*);
 
 		int SendThreadMain();
 
-		void Disconect(bool wait_for_threads_to_exit=true);
+		void Disconnect(bool wait_for_threads_to_exit=true);
 
-		void Disconect_NoLock(bool wait_for_threads_to_exit=true);
+		void Disconnect_NoLock(bool wait_for_threads_to_exit=true);
 
 		void SendRecordedDataAndDisconnect();
 
@@ -3354,15 +3451,13 @@ namespace FramePro
 
 		void StartRecordingContextSitches();
 
-		void FlushGlobalHiResTimers(FrameProTLS* p_framepro_tls);
+		void FlushCustomStatTimers(FrameProTLS* p_framepro_tls);
 
-		void ClearGlobalHiResTimers();
+		void ClearCustomStatTimers();
 
 		void SendExtraModuleInfo(int64 ModuleBase, FrameProTLS* p_framepro_tls);
 
 		#if FRAMEPRO_SOCKETS_ENABLED
-			bool SendSendBuffer(SendBuffer* p_send_buffer, Socket& socket);
-
 			bool InitialiseFileCache();
 
 			static int StaticConnectThreadMain(void*);
@@ -3392,6 +3487,8 @@ namespace FramePro
 
 		#if FRAMEPRO_ENABLE_CALLSTACKS
 			void SetCallstacksEnabled(bool enabled);
+
+			void SendModules();
 		#endif
 
 		void SendScopeColours();
@@ -3411,8 +3508,8 @@ namespace FramePro
 
 		char m_Port[8];
 
-		Allocator* mp_Allocator;
-		bool m_CreatedAllocator;
+		CheckedAllocator m_Allocator;
+		Allocator* mp_CreatedAllocator;
 
 		bool m_Initialised;
 
@@ -3479,15 +3576,14 @@ namespace FramePro
 
 		Array<int> m_ProcessIds;
 
-		Buffer m_MainThreadSendBuffer;
+		Buffer m_MainThreadSendBuffer[2];
+		int m_MainThreadSendBufferIndex;
 		CriticalSection m_MainThreadSendBufferLock;
 
 		Array<RequestStringLiteralPacket> m_StringRequestPackets;
 		CriticalSection m_StringRequestPacketsLock;
 
-		GlobalHiResTimer* mp_GlobalHiResTimers;
-
-		int m_ModulesSent;
+		CustomStatTimer* m_CustomStatTimers;
 
 		Array<ModulePacket*> m_ModulePackets;
 
@@ -3556,9 +3652,14 @@ namespace FramePro
 		Allocator* p_allocator = framepro_session.GetAllocator();
 
 		FrameProTLS* p_framepro_tls = (FrameProTLS*)p_allocator->Alloc(sizeof(FrameProTLS));
+		
+		if (!p_framepro_tls)
+			return nullptr;
+
 		new (p_framepro_tls)FrameProTLS(p_allocator, framepro_session.GetClockFrequency());
 
-		framepro_session.AddFrameProTLS(p_framepro_tls);
+		if (!framepro_session.AddFrameProTLS(p_framepro_tls))
+			return NULL;
 
 		#if FRAMEPRO_USE_TLS_SLOTS
 			int slot = GetFrameProTLSSlot();
@@ -3584,12 +3685,18 @@ namespace FramePro
 	}
 
 	//------------------------------------------------------------------------
-	void SendWaitEventPacket(int64 event_id, int64 time, PacketType::Enum packet_type)
+	void SendWaitEventPacket(
+		int64 event_id,
+		int64 time,
+		int core,
+		PacketType::Enum packet_type)
 	{
 		if (!g_Connected)
 			return;
 
 		FrameProTLS* p_framepro_tls = GetFrameProTLS();
+		if (!p_framepro_tls)
+			return;
 
 		CriticalSectionScope lock(p_framepro_tls->GetCurrentSendBufferCriticalSection());
 
@@ -3597,7 +3704,7 @@ namespace FramePro
 
 		p_packet->m_PacketType = packet_type;
 		p_packet->m_Thread = p_framepro_tls->GetThreadId();
-		p_packet->m_Core = Platform::GetCore();
+		p_packet->m_Core = core;
 		p_packet->m_EventId = event_id;
 		p_packet->m_Time = time;
 	}
@@ -3643,24 +3750,38 @@ void FramePro::UnregisterConnectionChangedcallback(ConnectionChangedCallback p_c
 void FramePro::AddTimeSpan(const char* p_name_and_source_info, int64 start_time, int64 end_time)
 {
 	FrameProTLS* p_framepro_tls = GetFrameProTLS();
+	if (!p_framepro_tls)
+		return;
 
 #ifdef FRAMEPRO_SCOPE_MIN_TIME
 	if (end_time - start_time < p_framepro_tls->GetScopeMinTime())
 		return;
 #endif
 
-	p_framepro_tls->SubmitHiResTimers(end_time);
+	p_framepro_tls->SubmitHiResTimers(start_time, end_time);
 
 	// if we are connecting to FramePro, FramePro will ask for the string value later, otherwise we need to send it now
 	if(p_framepro_tls->SendStringsImmediately())
 		p_framepro_tls->SendString(p_name_and_source_info, PacketType::NameAndSourceInfoPacket);
 
 #if FRAMEPRO_ENABLE_CALLSTACKS
-	if (p_framepro_tls->ShouldSendCallstacks())
+	int current_callstack_id = p_framepro_tls->GetAndClearCurrentCallstackId(start_time, end_time);
+
+	if (current_callstack_id != -1 || p_framepro_tls->ShouldSendCallstacks())
 	{
-		CallstackResult callstack_result = p_framepro_tls->GetCallstack();
+		CallstackResult callstack_result;
 
-		int send_size = sizeof(TimeSpanPacket) + sizeof(CallstackPacket);
+		if (current_callstack_id == -1)
+		{
+			callstack_result = p_framepro_tls->GetCallstack();
+			current_callstack_id = callstack_result.mp_Callstack->m_ID;
+		}
+		else
+		{
+			p_framepro_tls->SetCurrentCallstackId(-1);
+		}
+
+		int send_size = sizeof(TimeSpanPacket) + sizeof(TimeSpanCallstackPacket);
 		if (callstack_result.m_IsNew)
 			send_size += callstack_result.mp_Callstack->m_Size * sizeof(uint64);
 
@@ -3675,8 +3796,8 @@ void FramePro::AddTimeSpan(const char* p_name_and_source_info, int64 start_time,
 			p_packet->m_StartTime = start_time;
 			p_packet->m_EndTime = end_time;
 
-			CallstackPacket* p_callstack_packet = (CallstackPacket*)(p_packet + 1);
-			p_callstack_packet->m_CallstackId = callstack_result.mp_Callstack->m_ID;
+			TimeSpanCallstackPacket* p_callstack_packet = (TimeSpanCallstackPacket*)(p_packet + 1);
+			p_callstack_packet->m_CallstackId = current_callstack_id;
 			p_callstack_packet->m_CallstackSize = 0;
 
 			if (callstack_result.m_IsNew)
@@ -3710,8 +3831,10 @@ void FramePro::AddTimeSpan(const wchar_t* p_name_and_source_info, int64 start_ti
 	FRAMEPRO_ASSERT(start_time <= end_time);
 
 	FrameProTLS* p_framepro_tls = GetFrameProTLS();
+	if (!p_framepro_tls)
+		return;
 
-	p_framepro_tls->SubmitHiResTimers(end_time);
+	p_framepro_tls->SubmitHiResTimers(start_time, end_time);
 
 #ifdef FRAMEPRO_SCOPE_MIN_TIME
 	if (end_time - start_time < p_framepro_tls->GetScopeMinTime())
@@ -3723,11 +3846,23 @@ void FramePro::AddTimeSpan(const wchar_t* p_name_and_source_info, int64 start_ti
 		p_framepro_tls->SendString(p_name_and_source_info, PacketType::NameAndSourceInfoPacketW);
 
 #if FRAMEPRO_ENABLE_CALLSTACKS
-	if (p_framepro_tls->ShouldSendCallstacks())
+	int current_callstack_id = p_framepro_tls->GetAndClearCurrentCallstackId(start_time, end_time);
+
+	if (current_callstack_id != -1 || p_framepro_tls->ShouldSendCallstacks())
 	{
-		CallstackResult callstack_result = p_framepro_tls->GetCallstack();
+		CallstackResult callstack_result;
+
+		if (current_callstack_id == -1)
+		{
+			callstack_result = p_framepro_tls->GetCallstack();
+			current_callstack_id = callstack_result.mp_Callstack->m_ID;
+		}
+		else
+		{
+			p_framepro_tls->SetCurrentCallstackId(-1);
+		}
 
-		int send_size = sizeof(TimeSpanPacket) + sizeof(CallstackPacket);
+		int send_size = sizeof(TimeSpanPacket) + sizeof(TimeSpanCallstackPacket);
 		if (callstack_result.m_IsNew)
 			send_size += callstack_result.mp_Callstack->m_Size * sizeof(uint64);
 
@@ -3742,8 +3877,8 @@ void FramePro::AddTimeSpan(const wchar_t* p_name_and_source_info, int64 start_ti
 			p_packet->m_StartTime = start_time;
 			p_packet->m_EndTime = end_time;
 
-			CallstackPacket* p_callstack_packet = (CallstackPacket*)(p_packet + 1);
-			p_callstack_packet->m_CallstackId = callstack_result.mp_Callstack->m_ID;
+			TimeSpanCallstackPacket* p_callstack_packet = (TimeSpanCallstackPacket*)(p_packet + 1);
+			p_callstack_packet->m_CallstackId = current_callstack_id;
 			p_callstack_packet->m_CallstackSize = 0;
 
 			if (callstack_result.m_IsNew)
@@ -3777,8 +3912,10 @@ void FramePro::AddTimeSpan(StringId name, const char* p_source_info, int64 start
 	FRAMEPRO_ASSERT(start_time <= end_time);
 
 	FrameProTLS* p_framepro_tls = GetFrameProTLS();
+	if (!p_framepro_tls)
+		return;
 
-	p_framepro_tls->SubmitHiResTimers(end_time);
+	p_framepro_tls->SubmitHiResTimers(start_time, end_time);
 
 #ifdef FRAMEPRO_SCOPE_MIN_TIME
 	if (end_time - start_time < p_framepro_tls->GetScopeMinTime())
@@ -3790,11 +3927,23 @@ void FramePro::AddTimeSpan(StringId name, const char* p_source_info, int64 start
 		p_framepro_tls->SendString(p_source_info, PacketType::SourceInfoPacket);
 
 #if FRAMEPRO_ENABLE_CALLSTACKS
-	if (p_framepro_tls->ShouldSendCallstacks())
+	int current_callstack_id = p_framepro_tls->GetAndClearCurrentCallstackId(start_time, end_time);
+
+	if (current_callstack_id != -1 || p_framepro_tls->ShouldSendCallstacks())
 	{
-		CallstackResult callstack_result = p_framepro_tls->GetCallstack();
+		CallstackResult callstack_result;
+
+		if (current_callstack_id == -1)
+		{
+			callstack_result = p_framepro_tls->GetCallstack();
+			current_callstack_id = callstack_result.mp_Callstack->m_ID;
+		}
+		else
+		{
+			p_framepro_tls->SetCurrentCallstackId(-1);
+		}
 
-		int send_size = sizeof(NamedTimeSpanPacket) + sizeof(CallstackPacket);
+		int send_size = sizeof(NamedTimeSpanPacket) + sizeof(TimeSpanCallstackPacket);
 		if (callstack_result.m_IsNew)
 			send_size += callstack_result.mp_Callstack->m_Size * sizeof(uint64);
 
@@ -3810,8 +3959,8 @@ void FramePro::AddTimeSpan(StringId name, const char* p_source_info, int64 start
 			p_packet->m_StartTime = start_time;
 			p_packet->m_EndTime = end_time;
 
-			CallstackPacket* p_callstack_packet = (CallstackPacket*)(p_packet + 1);
-			p_callstack_packet->m_CallstackId = callstack_result.mp_Callstack->m_ID;
+			TimeSpanCallstackPacket* p_callstack_packet = (TimeSpanCallstackPacket*)(p_packet + 1);
+			p_callstack_packet->m_CallstackId = current_callstack_id;
 			p_callstack_packet->m_CallstackSize = 0;
 
 			if (callstack_result.m_IsNew)
@@ -3846,8 +3995,10 @@ void FramePro::AddTimeSpan(StringId name, const char* p_source_info, int64 start
 	FRAMEPRO_ASSERT(start_time <= end_time);
 
 	FrameProTLS* p_framepro_tls = GetFrameProTLS();
+	if (!p_framepro_tls)
+		return;
 
-	p_framepro_tls->SubmitHiResTimers(end_time);
+	p_framepro_tls->SubmitHiResTimers(start_time, end_time);
 
 #ifdef FRAMEPRO_SCOPE_MIN_TIME
 	if (end_time - start_time < p_framepro_tls->GetScopeMinTime())
@@ -3859,11 +4010,23 @@ void FramePro::AddTimeSpan(StringId name, const char* p_source_info, int64 start
 		p_framepro_tls->SendString(p_source_info, PacketType::SourceInfoPacket);
 
 #if FRAMEPRO_ENABLE_CALLSTACKS
-	if (p_framepro_tls->ShouldSendCallstacks())
+	int current_callstack_id = p_framepro_tls->GetAndClearCurrentCallstackId(start_time, end_time);
+
+	if (current_callstack_id != -1 || p_framepro_tls->ShouldSendCallstacks())
 	{
-		CallstackResult callstack_result = p_framepro_tls->GetCallstack();
+		CallstackResult callstack_result;
 
-		int send_size = sizeof(NamedTimeSpanPacket) + sizeof(CallstackPacket);
+		if (current_callstack_id == -1)
+		{
+			callstack_result = p_framepro_tls->GetCallstack();
+			current_callstack_id = callstack_result.mp_Callstack->m_ID;
+		}
+		else
+		{
+			p_framepro_tls->SetCurrentCallstackId(-1);
+		}
+
+		int send_size = sizeof(NamedTimeSpanPacket) + sizeof(TimeSpanCallstackPacket);
 		if (callstack_result.m_IsNew)
 			send_size += callstack_result.mp_Callstack->m_Size * sizeof(uint64);
 
@@ -3879,8 +4042,8 @@ void FramePro::AddTimeSpan(StringId name, const char* p_source_info, int64 start
 			p_packet->m_StartTime = start_time;
 			p_packet->m_EndTime = end_time;
 
-			CallstackPacket* p_callstack_packet = (CallstackPacket*)(p_packet + 1);
-			p_callstack_packet->m_CallstackId = callstack_result.mp_Callstack->m_ID;
+			TimeSpanCallstackPacket* p_callstack_packet = (TimeSpanCallstackPacket*)(p_packet + 1);
+			p_callstack_packet->m_CallstackId = current_callstack_id;
 			p_callstack_packet->m_CallstackSize = 0;
 
 			if (callstack_result.m_IsNew)
@@ -3914,8 +4077,10 @@ void FramePro::AddTimeSpan(StringId name, const char* p_source_info, int64 start
 void FramePro::AddTimeSpan(const char* p_name, const char* p_source_info, int64 start_time, int64 end_time)
 {
 	FrameProTLS* p_framepro_tls = GetFrameProTLS();
+	if (!p_framepro_tls)
+		return;
 
-	p_framepro_tls->SubmitHiResTimers(end_time);
+	p_framepro_tls->SubmitHiResTimers(start_time, end_time);
 
 #ifdef FRAMEPRO_SCOPE_MIN_TIME
 	if (end_time - start_time < p_framepro_tls->GetScopeMinTime())
@@ -3925,16 +4090,28 @@ void FramePro::AddTimeSpan(const char* p_name, const char* p_source_info, int64
 	// if we are connecting to FramePro, FramePro will ask for the string value later, otherwise we need to send it now
 	if(p_framepro_tls->SendStringsImmediately())
 	{
-		p_framepro_tls->SendString(p_name, PacketType::StringPacket);
+		p_framepro_tls->SendString(p_name, PacketType::StringLiteralTimerNamePacket);
 		p_framepro_tls->SendString(p_source_info, PacketType::SourceInfoPacket);
 	}
 
 #if FRAMEPRO_ENABLE_CALLSTACKS
-	if (p_framepro_tls->ShouldSendCallstacks())
+	int current_callstack_id = p_framepro_tls->GetAndClearCurrentCallstackId(start_time, end_time);
+
+	if (current_callstack_id != -1 || p_framepro_tls->ShouldSendCallstacks())
 	{
-		CallstackResult callstack_result = p_framepro_tls->GetCallstack();
+		CallstackResult callstack_result;
+
+		if (current_callstack_id == -1)
+		{
+			callstack_result = p_framepro_tls->GetCallstack();
+			current_callstack_id = callstack_result.mp_Callstack->m_ID;
+		}
+		else
+		{
+			p_framepro_tls->SetCurrentCallstackId(-1);
+		}
 
-		int send_size = sizeof(NamedTimeSpanPacket) + sizeof(CallstackPacket);
+		int send_size = sizeof(NamedTimeSpanPacket) + sizeof(TimeSpanCallstackPacket);
 		if (callstack_result.m_IsNew)
 			send_size += callstack_result.mp_Callstack->m_Size * sizeof(uint64);
 
@@ -3950,8 +4127,8 @@ void FramePro::AddTimeSpan(const char* p_name, const char* p_source_info, int64
 			p_packet->m_StartTime = start_time;
 			p_packet->m_EndTime = end_time;
 
-			CallstackPacket* p_callstack_packet = (CallstackPacket*)(p_packet + 1);
-			p_callstack_packet->m_CallstackId = callstack_result.mp_Callstack->m_ID;
+			TimeSpanCallstackPacket* p_callstack_packet = (TimeSpanCallstackPacket*)(p_packet + 1);
+			p_callstack_packet->m_CallstackId = current_callstack_id;
 			p_callstack_packet->m_CallstackSize = 0;
 
 			if (callstack_result.m_IsNew)
@@ -3990,6 +4167,8 @@ void FramePro::AddCustomStat(const char* p_name, int value, const char* p_graph,
 void FramePro::AddCustomStat(const char* p_name, int64 value, const char* p_graph, const char* p_unit, uint colour)
 {
 	FrameProTLS* p_framepro_tls = GetFrameProTLS();
+	if (!p_framepro_tls)
+		return;
 
 	// if we are connecting to FramePro, FramePro will ask for the string value later, otherwise we need to send it now
 	if (p_framepro_tls->SendStringsImmediately())
@@ -4007,6 +4186,7 @@ void FramePro::AddCustomStat(const char* p_name, int64 value, const char* p_grap
 	p_packet->m_Count = 1;
 	p_packet->m_Name = (StringId)p_name;
 	p_packet->m_Value = value;
+	FRAMEPRO_GET_CLOCK_COUNT(p_packet->m_Time);
 }
 
 //------------------------------------------------------------------------
@@ -4019,6 +4199,8 @@ void FramePro::AddCustomStat(const char* p_name, float value, const char* p_grap
 void FramePro::AddCustomStat(const char* p_name, double value, const char* p_graph, const char* p_unit, uint colour)
 {
 	FrameProTLS* p_framepro_tls = GetFrameProTLS();
+	if (!p_framepro_tls)
+		return;
 
 	// if we are connecting to FramePro, FramePro will ask for the string value later, otherwise we need to send it now
 	if (p_framepro_tls->SendStringsImmediately())
@@ -4036,6 +4218,7 @@ void FramePro::AddCustomStat(const char* p_name, double value, const char* p_gra
 	p_packet->m_Count = 1;
 	p_packet->m_Name = (StringId)p_name;
 	p_packet->m_Value = value;
+	FRAMEPRO_GET_CLOCK_COUNT(p_packet->m_Time);
 }
 
 //------------------------------------------------------------------------
@@ -4048,6 +4231,8 @@ void FramePro::AddCustomStat(const wchar_t* p_name, int value, const wchar_t* p_
 void FramePro::AddCustomStat(const wchar_t* p_name, int64 value, const wchar_t* p_graph, const wchar_t* p_unit, uint colour)
 {
 	FrameProTLS* p_framepro_tls = GetFrameProTLS();
+	if (!p_framepro_tls)
+		return;
 
 	// if we are connecting to FramePro, FramePro will ask for the string value later, otherwise we need to send it now
 	if (p_framepro_tls->SendStringsImmediately())
@@ -4065,6 +4250,7 @@ void FramePro::AddCustomStat(const wchar_t* p_name, int64 value, const wchar_t*
 	p_packet->m_Count = 1;
 	p_packet->m_Name = (StringId)p_name;
 	p_packet->m_Value = value;
+	FRAMEPRO_GET_CLOCK_COUNT(p_packet->m_Time);
 }
 
 //------------------------------------------------------------------------
@@ -4077,6 +4263,8 @@ void FramePro::AddCustomStat(const wchar_t* p_name, float value, const wchar_t*
 void FramePro::AddCustomStat(const wchar_t* p_name, double value, const wchar_t* p_graph, const wchar_t* p_unit, uint colour)
 {
 	FrameProTLS* p_framepro_tls = GetFrameProTLS();
+	if (!p_framepro_tls)
+		return;
 
 	// if we are connecting to FramePro, FramePro will ask for the string value later, otherwise we need to send it now
 	if (p_framepro_tls->SendStringsImmediately())
@@ -4094,6 +4282,7 @@ void FramePro::AddCustomStat(const wchar_t* p_name, double value, const wchar_t*
 	p_packet->m_Count = 1;
 	p_packet->m_Name = (StringId)p_name;
 	p_packet->m_Value = value;
+	FRAMEPRO_GET_CLOCK_COUNT(p_packet->m_Time);
 }
 
 //------------------------------------------------------------------------
@@ -4106,6 +4295,8 @@ void FramePro::AddCustomStat(StringId name, int value, const char* p_graph, cons
 void FramePro::AddCustomStat(StringId name, int64 value, const char* p_graph, const char* p_unit, uint colour)
 {
 	FrameProTLS* p_framepro_tls = GetFrameProTLS();
+	if (!p_framepro_tls)
+		return;
 
 	p_framepro_tls->SetCustomStatInfo(name, p_graph, p_unit, colour);	// only sends first time
 
@@ -4119,6 +4310,7 @@ void FramePro::AddCustomStat(StringId name, int64 value, const char* p_graph, co
 	p_packet->m_Count = 1;
 	p_packet->m_Name = name;
 	p_packet->m_Value = value;
+	FRAMEPRO_GET_CLOCK_COUNT(p_packet->m_Time);
 }
 
 //------------------------------------------------------------------------
@@ -4131,6 +4323,8 @@ void FramePro::AddCustomStat(StringId name, float value, const char* p_graph, co
 void FramePro::AddCustomStat(StringId name, double value, const char* p_graph, const char* p_unit, uint colour)
 {
 	FrameProTLS* p_framepro_tls = GetFrameProTLS();
+	if (!p_framepro_tls)
+		return;
 
 	p_framepro_tls->SetCustomStatInfo(name, p_graph, p_unit, colour);	// only sends first time
 
@@ -4144,6 +4338,7 @@ void FramePro::AddCustomStat(StringId name, double value, const char* p_graph, c
 	p_packet->m_Count = 1;
 	p_packet->m_Name = name;
 	p_packet->m_Value = value;
+	FRAMEPRO_GET_CLOCK_COUNT(p_packet->m_Time);
 }
 
 //------------------------------------------------------------------------
@@ -4156,6 +4351,8 @@ void FramePro::AddCustomStat(StringId name, int value, const wchar_t* p_graph, c
 void FramePro::AddCustomStat(StringId name, int64 value, const wchar_t* p_graph, const wchar_t* p_unit, uint colour)
 {
 	FrameProTLS* p_framepro_tls = GetFrameProTLS();
+	if (!p_framepro_tls)
+		return;
 
 	p_framepro_tls->SetCustomStatInfo(name, p_graph, p_unit, colour);	// only sends first time
 
@@ -4169,6 +4366,7 @@ void FramePro::AddCustomStat(StringId name, int64 value, const wchar_t* p_graph,
 	p_packet->m_Count = 1;
 	p_packet->m_Name = name;
 	p_packet->m_Value = value;
+	FRAMEPRO_GET_CLOCK_COUNT(p_packet->m_Time);
 }
 
 //------------------------------------------------------------------------
@@ -4181,6 +4379,8 @@ void FramePro::AddCustomStat(StringId name, float value, const wchar_t* p_graph,
 void FramePro::AddCustomStat(StringId name, double value, const wchar_t* p_graph, const wchar_t* p_unit, uint colour)
 {
 	FrameProTLS* p_framepro_tls = GetFrameProTLS();
+	if (!p_framepro_tls)
+		return;
 
 	p_framepro_tls->SetCustomStatInfo(name, p_graph, p_unit, colour);	// only sends first time
 
@@ -4194,6 +4394,7 @@ void FramePro::AddCustomStat(StringId name, double value, const wchar_t* p_graph
 	p_packet->m_Count = 1;
 	p_packet->m_Name = name;
 	p_packet->m_Value = value;
+	FRAMEPRO_GET_CLOCK_COUNT(p_packet->m_Time);
 }
 
 //------------------------------------------------------------------------
@@ -4206,6 +4407,8 @@ void FramePro::AddCustomStat(StringId name, int value, StringId graph, StringId
 void FramePro::AddCustomStat(StringId name, int64 value, StringId graph, StringId unit, uint colour)
 {
 	FrameProTLS* p_framepro_tls = GetFrameProTLS();
+	if (!p_framepro_tls)
+		return;
 
 	p_framepro_tls->SetCustomStatInfo(name, graph, unit, colour);	// only sends first time
 
@@ -4219,6 +4422,7 @@ void FramePro::AddCustomStat(StringId name, int64 value, StringId graph, StringI
 	p_packet->m_Count = 1;
 	p_packet->m_Name = name;
 	p_packet->m_Value = value;
+	FRAMEPRO_GET_CLOCK_COUNT(p_packet->m_Time);
 }
 
 //------------------------------------------------------------------------
@@ -4231,6 +4435,8 @@ void FramePro::AddCustomStat(StringId name, float value, StringId graph, StringI
 void FramePro::AddCustomStat(StringId name, double value, StringId graph, StringId unit, uint colour)
 {
 	FrameProTLS* p_framepro_tls = GetFrameProTLS();
+	if (!p_framepro_tls)
+		return;
 
 	p_framepro_tls->SetCustomStatInfo(name, graph, unit, colour);	// only sends first time
 
@@ -4244,6 +4450,7 @@ void FramePro::AddCustomStat(StringId name, double value, StringId graph, String
 	p_packet->m_Count = 1;
 	p_packet->m_Name = name;
 	p_packet->m_Value = value;
+	FRAMEPRO_GET_CLOCK_COUNT(p_packet->m_Time);
 }
 
 //------------------------------------------------------------------------
@@ -4252,22 +4459,32 @@ void FramePro::SetThreadName(const char* p_name)
 	GetFrameProSession().SetThreadName(p_name);
 }
 
+//------------------------------------------------------------------------
+void FramePro::SetThreadName(int thread_id, const char* p_name)
+{
+	GetFrameProSession().SetThreadName(thread_id, p_name);
+}
+
 //------------------------------------------------------------------------
 void FramePro::SetThreadOrder(StringId thread_name)
 {
-	GetFrameProTLS()->SetThreadOrder(thread_name);
+	FrameProTLS* p_framepro_tls = GetFrameProTLS();
+	if (p_framepro_tls)
+		p_framepro_tls->SetThreadOrder(thread_name);
 }
 
 //------------------------------------------------------------------------
 FramePro::StringId FramePro::RegisterString(const char* p_str)
 {
-	return GetFrameProTLS()->RegisterString(p_str);
+	FrameProTLS* p_framepro_tls = GetFrameProTLS();
+	return p_framepro_tls ? p_framepro_tls->RegisterString(p_str) : 0;
 }
 
 //------------------------------------------------------------------------
 FramePro::StringId FramePro::RegisterString(const wchar_t* p_str)
 {
-	return GetFrameProTLS()->RegisterString(p_str);
+	FrameProTLS* p_framepro_tls = GetFrameProTLS();
+	return p_framepro_tls ? p_framepro_tls->RegisterString(p_str) : 0;
 }
 
 //------------------------------------------------------------------------
@@ -4319,21 +4536,21 @@ void FramePro::SetPort(int port)
 }
 
 //------------------------------------------------------------------------
-void FramePro::SendSessionInfo(const char* p_name, const char* p_build_id)
+void FramePro::SendSessionInfo(const char* p_name, const char* p_value)
 {
-	GetFrameProSession().SendSessionDetails(p_name, p_build_id);
+	GetFrameProSession().SendSessionInfo(p_name, p_value);
 }
 
 //------------------------------------------------------------------------
-void FramePro::SendSessionInfo(const wchar_t* p_name, const wchar_t* p_build_id)
+void FramePro::SendSessionInfo(const wchar_t* p_name, const wchar_t* p_value)
 {
-	GetFrameProSession().SendSessionDetails(p_name, p_build_id);
+	GetFrameProSession().SendSessionInfo(p_name, p_value);
 }
 
 //------------------------------------------------------------------------
-void FramePro::AddGlobalHiResTimer(GlobalHiResTimer* p_timer)
+void FramePro::AddCustomStatTimer(CustomStatTimer* p_timer)
 {
-	GetFrameProSession().AddGlobalHiResTimer(p_timer);
+	GetFrameProSession().AddCustomStatTimer(p_timer);
 }
 
 //------------------------------------------------------------------------
@@ -4354,13 +4571,17 @@ void FramePro::CleanupThread()
 //------------------------------------------------------------------------
 void FramePro::PushConditionalParentScope(const char* p_name, int64 pre_duration, int64 post_duration)
 {
-	GetFrameProTLS()->PushConditionalParentScope(p_name, pre_duration, post_duration);
+	FrameProTLS* p_framepro_tls = GetFrameProTLS();
+	if (p_framepro_tls)
+		p_framepro_tls->PushConditionalParentScope(p_name, pre_duration, post_duration);
 }
 
 //------------------------------------------------------------------------
 void FramePro::PopConditionalParentScope(bool add_children)
 {
-	GetFrameProTLS()->PopConditionalParentScope(add_children);
+	FrameProTLS* p_framepro_tls = GetFrameProTLS();
+	if (p_framepro_tls)
+		p_framepro_tls->PopConditionalParentScope(add_children);
 }
 
 //------------------------------------------------------------------------
@@ -4372,63 +4593,79 @@ bool FramePro::CallConditionalParentScopeCallback(ConditionalParentScopeCallback
 //------------------------------------------------------------------------
 void FramePro::StartHiResTimer(const char* p_name)
 {
-	GetFrameProTLS()->StartHiResTimer(p_name);
+	FrameProTLS* p_framepro_tls = GetFrameProTLS();
+	if (p_framepro_tls)
+		p_framepro_tls->StartHiResTimer(p_name);
 }
 
 //------------------------------------------------------------------------
 void FramePro::StopHiResTimer()
 {
-	GetFrameProTLS()->StopHiResTimer();
-}
-
-//------------------------------------------------------------------------
-void FramePro::SubmitHiResTimers(int64 current_time)
-{
-	FRAMEPRO_ASSERT(g_Connected);
-
-	GetFrameProTLS()->SubmitHiResTimers(current_time);
+	FrameProTLS* p_framepro_tls = GetFrameProTLS();
+	if (p_framepro_tls)
+		p_framepro_tls->StopHiResTimer();
 }
 
 //------------------------------------------------------------------------
 void FramePro::Log(const char* p_message)
 {
-	if(g_Connected)
-		GetFrameProTLS()->SendLogPacket(p_message);
+	if (g_Connected)
+	{
+		FrameProTLS* p_framepro_tls = GetFrameProTLS();
+		if (p_framepro_tls)
+			p_framepro_tls->SendLogPacket(p_message);
+	}
 }
 
 //------------------------------------------------------------------------
 void FramePro::AddEvent(const char* p_name, uint colour)
 {
 	if (g_Connected)
-		GetFrameProTLS()->SendEventPacket(p_name, colour);
+	{
+		FrameProTLS* p_framepro_tls = GetFrameProTLS();
+		if (p_framepro_tls)
+			p_framepro_tls->SendEventPacket(p_name, colour);
+	}
 }
 
 //------------------------------------------------------------------------
-void FramePro::AddWaitEvent(int64 event_id, int64 start_time, int64 end_time)
+void FramePro::AddWaitEvent(
+	int64 event_id,
+	int64 start_time,
+	int start_core,
+	int64 end_time,
+	int end_core)
 {
 	FrameProTLS* p_framepro_tls = GetFrameProTLS();
+	if (!p_framepro_tls)
+		return;
 
 #ifdef FRAMEPRO_SCOPE_MIN_TIME
 	if (end_time - start_time < p_framepro_tls->GetWaitEventMinTime())
 		return;
 #endif
-	SendWaitEventPacket(event_id, start_time, PacketType::StartWaitEventPacket);
-	SendWaitEventPacket(event_id, end_time, PacketType::StopWaitEventPacket);
+	SendWaitEventPacket(event_id, start_time, start_core, PacketType::StartWaitEventPacket);
+	SendWaitEventPacket(event_id, end_time, end_core, PacketType::StopWaitEventPacket);
 }
 
 //------------------------------------------------------------------------
 void FramePro::TriggerWaitEvent(int64 event_id)
 {
+	if (!g_Connected)
+		return;
+
 	int64 time;
 	FRAMEPRO_GET_CLOCK_COUNT(time);
 
-	SendWaitEventPacket(event_id, time, PacketType::TriggerWaitEventPacket);
+	SendWaitEventPacket(event_id, time, Platform::GetCore(), PacketType::TriggerWaitEventPacket);
 }
 
 //------------------------------------------------------------------------
 void FramePro::SetScopeCustomStat(const char* p_name, int64 value, const char* p_graph, const char* p_unit, uint colour)
 {
 	FrameProTLS* p_framepro_tls = GetFrameProTLS();
+	if (!p_framepro_tls)
+		return;
 
 	if (p_framepro_tls->SendStringsImmediately())
 		p_framepro_tls->SendString(p_name, PacketType::StringPacket);
@@ -4442,9 +4679,11 @@ void FramePro::SetScopeCustomStat(const char* p_name, int64 value, const char* p
 void FramePro::SetScopeCustomStat(const wchar_t* p_name, int64 value, const wchar_t* p_graph, const wchar_t* p_unit, uint colour)
 {
 	FrameProTLS* p_framepro_tls = GetFrameProTLS();
+	if (!p_framepro_tls)
+		return;
 
 	if (p_framepro_tls->SendStringsImmediately())
-		p_framepro_tls->SendString(p_name, PacketType::StringPacket);
+		p_framepro_tls->SendString(p_name, PacketType::WStringPacket);
 
 	p_framepro_tls->SetCustomTimeSpanStatW((StringId)p_name, value);
 
@@ -4454,8 +4693,12 @@ void FramePro::SetScopeCustomStat(const wchar_t* p_name, int64 value, const wcha
 //------------------------------------------------------------------------
 void FramePro::SetScopeCustomStat(StringId name, int64 value, StringId graph, StringId unit, uint colour)
 {
+	FrameProTLS* p_framepro_tls = GetFrameProTLS();
+	if (!p_framepro_tls)
+		return;
+
 	// don't care about whether it is W or not because string has already been registerd
-	GetFrameProTLS()->SetCustomTimeSpanStat(name, value);
+	p_framepro_tls->SetCustomTimeSpanStat(name, value);
 
 	AddCustomStat(name, value, graph, unit, colour);
 }
@@ -4500,6 +4743,8 @@ void FramePro::SetScopeCustomStat(StringId name, float value, StringId graph, St
 void FramePro::SetScopeCustomStat(const char* p_name, double value, const char* p_graph, const char* p_unit, uint colour)
 {
 	FrameProTLS* p_framepro_tls = GetFrameProTLS();
+	if (!p_framepro_tls)
+		return;
 
 	if (p_framepro_tls->SendStringsImmediately())
 		p_framepro_tls->SendString(p_name, PacketType::StringPacket);
@@ -4513,6 +4758,8 @@ void FramePro::SetScopeCustomStat(const char* p_name, double value, const char*
 void FramePro::SetScopeCustomStat(const wchar_t* p_name, double value, const wchar_t* p_graph, const wchar_t* p_unit, uint colour)
 {
 	FrameProTLS* p_framepro_tls = GetFrameProTLS();
+	if (!p_framepro_tls)
+		return;
 
 	if (p_framepro_tls->SendStringsImmediately())
 		p_framepro_tls->SendString(p_name, PacketType::StringPacket);
@@ -4525,7 +4772,11 @@ void FramePro::SetScopeCustomStat(const wchar_t* p_name, double value, const wch
 //------------------------------------------------------------------------
 void FramePro::SetScopeCustomStat(StringId name, double value, StringId graph, StringId unit, uint colour)
 {
-	GetFrameProTLS()->SetCustomTimeSpanStat(name, value);
+	FrameProTLS* p_framepro_tls = GetFrameProTLS();
+	if (!p_framepro_tls)
+		return;
+
+	p_framepro_tls->SetCustomTimeSpanStat(name, value);
 
 	AddCustomStat(name, value, graph, unit, colour);
 }
@@ -4560,6 +4811,38 @@ void FramePro::SetCustomStatColour(StringId name, uint colour)
 	GetFrameProSession().SetCustomStatColour(name, colour);
 }
 
+//------------------------------------------------------------------------
+void FramePro::SendScopeCallstack()
+{
+	GetFrameProSession().SendScopeCallstack();
+}
+
+//------------------------------------------------------------------------
+int FramePro::GetCore()
+{
+	return Platform::GetCore();
+}
+
+//------------------------------------------------------------------------
+#if FRAMEPRO_DEBUG || FRAMEPRO_ALLOW_UNPARENTED_HIRES_SCOPES
+void FramePro::StartScope()
+{
+	FrameProTLS* p_framepro_tls = GetFrameProTLS();
+	if (p_framepro_tls)
+		p_framepro_tls->IncScopeDepth();
+}
+#endif
+
+//------------------------------------------------------------------------
+#if FRAMEPRO_DEBUG || FRAMEPRO_ALLOW_UNPARENTED_HIRES_SCOPES
+void FramePro::StopScope()
+{
+	FrameProTLS* p_framepro_tls = GetFrameProTLS();
+	if (p_framepro_tls)
+		p_framepro_tls->DecScopeDepth();
+}
+#endif
+
 
 //------------------------------------------------------------------------
 //
@@ -4572,6 +4855,9 @@ namespace FramePro
 	//------------------------------------------------------------------------
 	const int g_CallstackSetInitialCapacity = 4096;		// must be a power of 2
 
+	//------------------------------------------------------------------------
+	std::atomic<int> CallstackSet::m_NextUniqueId(0);
+
 	//------------------------------------------------------------------------
 	inline bool StacksMatch(FramePro::Callstack* p_callstack, uint64* p_stack, int stack_size, unsigned int hash)
 	{
@@ -4597,7 +4883,8 @@ namespace FramePro
 		mp_Allocator(p_allocator),
 		m_BlockAllocator(p_allocator)
 	{
-		memset(mp_Data, 0, g_CallstackSetInitialCapacity*sizeof(Callstack*));
+		if(mp_Data)
+			memset(mp_Data, 0, g_CallstackSetInitialCapacity*sizeof(Callstack*));
 	}
 
 	//------------------------------------------------------------------------
@@ -4616,7 +4903,11 @@ namespace FramePro
 		m_Capacity *= 2;
 		m_CapacityMask = m_Capacity - 1;
 		int size = m_Capacity * sizeof(Callstack*);
+
 		mp_Data = (Callstack**)mp_Allocator->Alloc(size);
+		if (!mp_Data)
+			return;
+
 		memset(mp_Data, 0, size);
 
 		// transfer callstacks from old set
@@ -4635,6 +4926,9 @@ namespace FramePro
 	//------------------------------------------------------------------------
 	FramePro::Callstack* CallstackSet::Get(uint64* p_stack, int stack_size, unsigned int hash)
 	{
+		if (!mp_Data)
+			return NULL;
+
 		int index = hash & m_CapacityMask;
 
 		while(mp_Data[index] && !StacksMatch(mp_Data[index], p_stack, stack_size, hash))
@@ -4652,7 +4946,7 @@ namespace FramePro
 
 		// create a new callstack
 		Callstack* p_callstack = (Callstack*)m_BlockAllocator.Alloc(sizeof(Callstack));
-		p_callstack->m_ID = m_Count;
+		p_callstack->m_ID = m_NextUniqueId++;
 		p_callstack->m_Size = stack_size;
 		p_callstack->mp_Stack = (uint64*)m_BlockAllocator.Alloc(stack_size*sizeof(uint64));
 		p_callstack->m_Hash = hash;
@@ -4666,6 +4960,9 @@ namespace FramePro
 	//------------------------------------------------------------------------
 	void CallstackSet::Add(Callstack* p_callstack)
 	{
+		if (!mp_Data)
+			return;
+
 		// find a clear index
 		int index = p_callstack->m_Hash & m_CapacityMask;
 		while(mp_Data[index])
@@ -4684,11 +4981,21 @@ namespace FramePro
 		mp_Allocator->Free(mp_Data);
 
 		size_t size = g_CallstackSetInitialCapacity*sizeof(Callstack*);
+		
 		mp_Data = (Callstack**)mp_Allocator->Alloc((int)size);
-		memset(mp_Data, 0, size);
-		m_CapacityMask = g_CallstackSetInitialCapacity-1;
+		
+		if (mp_Data)
+		{
+			memset(mp_Data, 0, size);
+			m_CapacityMask = g_CallstackSetInitialCapacity - 1;
+			m_Capacity = g_CallstackSetInitialCapacity;
+		}
+		else
+		{
+			m_Capacity = 0;
+		}
+
 		m_Count = 0;
-		m_Capacity = g_CallstackSetInitialCapacity;
 	}
 }
 
@@ -4767,8 +5074,7 @@ namespace FramePro
 
 	//------------------------------------------------------------------------
 	FrameProSession::FrameProSession()
-	:	mp_Allocator(NULL),	// must be first (before any call to GetAllocator)
-		m_CreatedAllocator(false),
+	:	mp_CreatedAllocator(NULL),
 		m_Initialised(false),
 		m_InitialiseConnectionNextFrame(false),
 		m_StartContextSwitchRecording(false),
@@ -4793,7 +5099,8 @@ namespace FramePro
 		m_SendThreadExit(false),
 		m_SendThreadFinished(false, true),
 		m_SocketsBlocked(FRAMEPRO_SOCKETS_BLOCKED_BY_DEFAULT),
-		mp_GlobalHiResTimers(NULL),
+		m_MainThreadSendBufferIndex(0),
+		m_CustomStatTimers(NULL),
 		mp_ContextSwitchRecorder(NULL)
 #if FRAMEPRO_ENABLE_CALLSTACKS
 		,m_SendModules(false)
@@ -4810,7 +5117,11 @@ namespace FramePro
 	//------------------------------------------------------------------------
 	FrameProSession::~FrameProSession()
 	{
-		Disconect(false);
+		Disconnect(true);
+
+		// shut down the connect thread
+		m_ListenSocket.Disconnect();
+		m_ConnectThread.WaitForThreadToTerminate(10 * 1000);
 
 		m_NamedThreads.Clear();
 
@@ -4819,9 +5130,15 @@ namespace FramePro
 		// must clear all arrays and buffers and detach the allocator before deleting the allocator
 		m_ProcessIds.Clear();
 		m_ProcessIds.SetAllocator(NULL);
-		
-		m_MainThreadSendBuffer.ClearAndFree();
-		m_MainThreadSendBuffer.SetAllocator(NULL);
+
+		{
+			CriticalSectionScope tls_lock(m_MainThreadSendBufferLock);
+			for (Buffer& buffer : m_MainThreadSendBuffer)
+			{
+				buffer.ClearAndFree();
+				buffer.SetAllocator(NULL);
+			}
+		}
 
 		m_StringRequestPackets.Clear();
 		m_StringRequestPackets.SetAllocator(NULL);
@@ -4847,8 +5164,7 @@ namespace FramePro
 		m_Connectionchangedcallbacks.Clear();
 		m_Connectionchangedcallbacks.SetAllocator(NULL);
 
-		if(m_CreatedAllocator)
-			delete mp_Allocator;
+		delete mp_CreatedAllocator;
 	}
 
 	//------------------------------------------------------------------------
@@ -4860,7 +5176,7 @@ namespace FramePro
 	//------------------------------------------------------------------------
 	void FrameProSession::SetAllocator(Allocator* p_allocator)
 	{
-		if(mp_Allocator)
+		if(m_Allocator.HasAllocator())
 			FramePro::DebugBreak();		// allocator already set. You must call Allocator BEFORE calling FRAMEPRO_FRAME_START
 
 		SetAllocatorInternal(p_allocator);
@@ -4869,27 +5185,25 @@ namespace FramePro
 	//------------------------------------------------------------------------
 	void FrameProSession::SetAllocatorInternal(Allocator* p_allocator)
 	{
-		FRAMEPRO_ASSERT(!mp_Allocator);
+		FRAMEPRO_ASSERT(!m_Allocator.HasAllocator());
 		FRAMEPRO_ASSERT(p_allocator);
 
-		mp_Allocator = p_allocator;
+		m_Allocator.SetAllocator(p_allocator);
 
 		m_NonInteractiveRecordingFile.SetAllocator(p_allocator);
 		m_RecordingFile.SetAllocator(p_allocator);
 	}
 
 	//------------------------------------------------------------------------
-	Allocator* FrameProSession::CreateDefaultAllocator()
+	void FrameProSession::CreateDefaultAllocator()
 	{
 		CriticalSectionScope lock(m_CriticalSection);
 
-		if(!mp_Allocator)
+		if(!m_Allocator.HasAllocator())
 		{
-			SetAllocatorInternal(new DefaultAllocator());
-			m_CreatedAllocator = true;
+			mp_CreatedAllocator = new DefaultAllocator();
+			SetAllocatorInternal(mp_CreatedAllocator);
 		}
-
-		return mp_Allocator;
 	}
 
 	//------------------------------------------------------------------------
@@ -4958,8 +5272,12 @@ namespace FramePro
 		int ret = p_this->ConnectThreadMain();
 
 		#if !FRAMEPRO_PLATFORM_UE4
-			DestroyFrameProTLS(GetFrameProTLS());
-			ClearFrameProTLS();
+			FrameProTLS* p_framepro_tls = TryGetFrameProTLS();
+			if (p_framepro_tls)
+			{
+				DestroyFrameProTLS(p_framepro_tls);
+				ClearFrameProTLS();
+			}
 		#endif
 
 		return ret;
@@ -4999,8 +5317,12 @@ namespace FramePro
 		int ret = p_this->ReceiveThreadMain();
 
 		#if !FRAMEPRO_PLATFORM_UE4
-			DestroyFrameProTLS(GetFrameProTLS());
-			ClearFrameProTLS();
+			FrameProTLS* p_framepro_tls = GetFrameProTLS();
+			if (p_framepro_tls)
+			{
+				DestroyFrameProTLS(p_framepro_tls);
+				ClearFrameProTLS();
+			}
 		#endif
 
 		return ret;
@@ -5012,8 +5334,9 @@ namespace FramePro
 	{
 		CriticalSectionScope tls_lock(m_MainThreadSendBufferLock);
 
-		void* p_dst = m_MainThreadSendBuffer.Allocate(size);
-		memcpy(p_dst, p_src, size);
+		void* p_dst = m_MainThreadSendBuffer[m_MainThreadSendBufferIndex].Allocate(size);
+		if(p_dst)
+			memcpy(p_dst, p_src, size);
 	}
 
 	//------------------------------------------------------------------------
@@ -5125,9 +5448,24 @@ namespace FramePro
 	//------------------------------------------------------------------------
 #if FRAMEPRO_ENABLE_CALLSTACKS
 	void FrameProSession::SetCallstacksEnabled(bool enabled)
+	{
+		if (enabled)
+			SendModules();
+
+		{
+			CriticalSectionScope tls_lock(m_TLSListCriticalSection);
+			for (FrameProTLS* p_tls = m_FrameProTLSList.GetHead(); p_tls != NULL; p_tls = p_tls->GetNext())
+				p_tls->SetSendCallstacks(enabled);
+		}
+	}
+#endif
+
+	//------------------------------------------------------------------------
+#if FRAMEPRO_ENABLE_CALLSTACKS
+	void FrameProSession::SendModules()
 	{
 		// enumerate modules
-		if (enabled && !m_SendModules)
+		if (!m_SendModules)
 		{
 			Platform::EnumerateModules(m_ModulePackets, GetAllocator());
 
@@ -5136,18 +5474,12 @@ namespace FramePro
 			{
 				ModulePacket* p_module_packet = m_ModulePackets[i];
 				SendImmediate(p_module_packet, sizeof(ModulePacket), GetFrameProTLS());
-				mp_Allocator->Free(p_module_packet);
+				m_Allocator.Free(p_module_packet);
 			}
 			m_ModulePackets.Clear();
 
 			m_SendModules = true;
 		}
-
-		{
-			CriticalSectionScope tls_lock(m_TLSListCriticalSection);
-			for (FrameProTLS* p_tls = m_FrameProTLSList.GetHead(); p_tls != NULL; p_tls = p_tls->GetNext())
-				p_tls->SetSendCallstacks(enabled);
-		}
 	}
 #endif
 
@@ -5159,7 +5491,10 @@ namespace FramePro
 		if (!mp_ContextSwitchRecorder)
 			mp_ContextSwitchRecorder = Platform::CreateContextSwitchRecorder(GetAllocator());
 
-		bool started = Platform::StartRecordingContextSitches(mp_ContextSwitchRecorder, ContextSwitchCallback_Static, this, error);
+		if (!mp_ContextSwitchRecorder)
+			return;
+
+		bool started = Platform::StartRecordingContextSitches(mp_ContextSwitchRecorder, ContextSwitchCallback_Static, this, error, GetAllocator());
 
 		if (!started)
 		{
@@ -5181,7 +5516,7 @@ namespace FramePro
 #if FRAMEPRO_SOCKETS_ENABLED
 	int FrameProSession::OnReceiveThreadExit()
 	{
-		Disconect();
+		Disconnect();
 		return 0;
 	}
 #endif
@@ -5191,7 +5526,8 @@ namespace FramePro
 	{
 		m_CriticalSection.Leave();
 
-		m_SendThread.CreateThread(StaticSendThreadMain, this, GetAllocator());
+		if (!m_SendThread.CreateThread(StaticSendThreadMain, this, GetAllocator()))
+			return;
 
 		if(m_ThreadPrioritySet)
 			m_SendThread.SetPriority(m_ThreadPriority);
@@ -5210,7 +5546,8 @@ namespace FramePro
 	{
 		m_ReceiveThreadTerminatedEvent.Reset();
 
-		m_ReceiveThread.CreateThread(StaticReceiveThreadMain, this, GetAllocator());
+		if (!m_ReceiveThread.CreateThread(StaticReceiveThreadMain, this, GetAllocator()))
+			return;
 
 		if(m_ThreadPrioritySet)
 			m_ReceiveThread.SetPriority(m_ThreadPriority);
@@ -5236,7 +5573,9 @@ namespace FramePro
 		if(!m_ProcessIds.Contains(context_switch.m_ProcessId))
 		{
 			m_ProcessIds.SetAllocator(GetAllocator());
-			m_ProcessIds.Add(context_switch.m_ProcessId);
+
+			if (!m_ProcessIds.Add(context_switch.m_ProcessId))
+				return;
 
 			const int max_process_name_length = 260;
 			char process_name[max_process_name_length];
@@ -5264,7 +5603,7 @@ namespace FramePro
 	}
 
 	//------------------------------------------------------------------------
-	void FrameProSession::InitialiseConnection(FrameProTLS* p_framepro_tls)
+	bool FrameProSession::InitialiseConnection(FrameProTLS* p_framepro_tls)
 	{
 		// start the Send thread FIRST, but paused (because it adds another TLS that we need to call OnConnected on)
 		m_SendComplete.Reset();
@@ -5275,8 +5614,11 @@ namespace FramePro
 		bool recording_to_file = m_RecordingFile.IsOpened();
 		{
 			CriticalSectionScope tls_lock(m_TLSListCriticalSection);
-			for(FrameProTLS* p_tls = m_FrameProTLSList.GetHead(); p_tls!=NULL; p_tls=p_tls->GetNext())
-				p_tls->OnConnected(recording_to_file);
+			for (FrameProTLS* p_tls = m_FrameProTLSList.GetHead(); p_tls != NULL; p_tls = p_tls->GetNext())
+			{
+				if (!p_tls->OnConnected(recording_to_file))
+					return false;
+			}
 		}
 
 		p_framepro_tls->SendConnectPacket(m_ClockFrequency, Platform::GetCurrentProcessId(), Platform::GetPlatformEnum());
@@ -5297,7 +5639,9 @@ namespace FramePro
 		// lock the session info buffers of all threads
 		for(FrameProTLS* p_tls = m_FrameProTLSList.GetHead(); p_tls!=NULL; p_tls=p_tls->GetNext())
 		{
-			p_tls->OnConnected(recording_to_file);		// in case any threads have been added since sending the connect packet
+			if (!p_tls->OnConnected(recording_to_file))		// in case any threads have been added since sending the connect packet
+				return false;
+
 			p_tls->LockSessionInfoBuffer();
 		}
 
@@ -5305,7 +5649,11 @@ namespace FramePro
 		for(FrameProTLS* p_tls = m_FrameProTLSList.GetHead(); p_tls!=NULL; p_tls=p_tls->GetNext())
 			p_tls->SendSessionInfoBuffer();
 
-		p_framepro_tls->SendFrameStartPacket(0);
+		// start the new frame
+		int64 frame_start_time;
+		FRAMEPRO_GET_CLOCK_COUNT(frame_start_time);
+		FrameStartPacket frame_start_packet(frame_start_time, 0);
+		SendOnMainThread(frame_start_packet);
 
 		if(g_ConditionalScopeMinTime == UINT_MAX)
 			g_ConditionalScopeMinTime = (unsigned int)((((int64)FRAMEPRO_DEFAULT_COND_SCOPE_MIN_TIME) * m_ClockFrequency) / 1000000LL);
@@ -5333,7 +5681,10 @@ namespace FramePro
 			m_StartContextSwitchRecording = false;
 		}
 
-		ClearGlobalHiResTimers();
+		ClearCustomStatTimers();
+
+		for (CustomStatTimer* p_custom_stat_timer = m_CustomStatTimers; p_custom_stat_timer; p_custom_stat_timer = p_custom_stat_timer->GetNext())
+			SetCustomStatUnit(FramePro::RegisterString(p_custom_stat_timer->GetName()), FramePro::RegisterString("cycles"));
 
 		SendScopeColours();
 		SendCustomStatGraphs();
@@ -5342,6 +5693,15 @@ namespace FramePro
 
 		OnConnectionChanged(true, m_RecordingFile.GetFilename());
 
+		// send date
+		m_CriticalSection.Leave();
+		StringId date_name = RegisterString("Date");
+		char date_str[64];
+		GetDateString(date_str, sizeof(date_str));
+		StringId date_value = RegisterString(date_str);
+		p_framepro_tls->SendPacket(SessionInfoPacket(date_name, date_value));
+		m_CriticalSection.Enter();
+
 #if FRAMEPRO_ENABLE_CALLSTACKS
 		bool enable_callstacks = m_StartRecordingCallstacks;
 		m_StartRecordingCallstacks = false;
@@ -5349,6 +5709,8 @@ namespace FramePro
 		SetCallstacksEnabled(enable_callstacks);
 		m_CriticalSection.Enter();
 #endif
+
+		return true;
 	}
 
 	//------------------------------------------------------------------------
@@ -5362,7 +5724,8 @@ namespace FramePro
 
 		{
 			CriticalSectionScope tls_lock(m_MainThreadSendBufferLock);
-			m_MainThreadSendBuffer.SetAllocator(GetAllocator());
+			for(Buffer& buffer : m_MainThreadSendBuffer)
+				buffer.SetAllocator(GetAllocator());
 		}
 
 		{
@@ -5409,35 +5772,47 @@ namespace FramePro
 #if FRAMEPRO_SOCKETS_ENABLED
 	void FrameProSession::StartConnectThread()
 	{
-		m_ConnectThread.CreateThread(StaticConnectThreadMain, this, GetAllocator());
-	}
-#endif
-
-	//------------------------------------------------------------------------
-#if FRAMEPRO_SOCKETS_ENABLED
-	bool FrameProSession::SendSendBuffer(SendBuffer* p_send_buffer, Socket& socket)
-	{
-		#if FRAMEPRO_DEBUG_TCP
-			static File file;
-			if (!file.IsOpened())
-			{
-				bool opened = file.OpenForWrite("framepro_network_data.framepro_recording");
-				FRAMEPRO_ASSERT(opened);
-				FRAMEPRO_UNREFERENCED(opened);
-			}
-			file.Write(p_send_buffer->GetBuffer(), p_send_buffer->GetSize());
-		#endif
-
-		return socket.Send(p_send_buffer->GetBuffer(), p_send_buffer->GetSize());
+		m_ConnectThread.CreateThread(StaticConnectThreadMain, this, GetAllocator());
 	}
 #endif
 
 	//------------------------------------------------------------------------
-	void FrameProSession::WriteSendBuffer(SendBuffer* p_send_buffer, File& file, int64& file_size)
+	bool FrameProSession::Send(const void* p_data, int size)
 	{
-		int size = p_send_buffer->GetSize();
-		file.Write(p_send_buffer->GetBuffer(), size);
-		file_size += size;
+		if (m_RecordingFile.IsOpened())
+		{
+			CriticalSectionScope lock(m_CriticalSection);
+			m_RecordingFile.Write(p_data, size);
+			m_RecordingFileSize += size;
+		}
+		else
+		{
+#if FRAMEPRO_SOCKETS_ENABLED
+			if (m_Interactive)
+			{
+				#if FRAMEPRO_DEBUG_TCP
+					static File file;
+					if (!file.IsOpened())
+					{
+						bool opened = file.OpenForWrite("framepro_network_data.framepro_recording");
+						FRAMEPRO_ASSERT(opened);
+						FRAMEPRO_UNREFERENCED(opened);
+					}
+					file.Write(p_data, size);
+				#endif
+
+				if(!m_ClientSocket.Send(p_data, size))
+					return false;	// disconnected
+			}
+			else
+			{
+				m_NonInteractiveRecordingFile.Write(p_data, size);
+				m_NonInteractiveRecordingFileSize += size;
+			}
+#endif
+		}
+
+		return true;
 	}
 
 	//------------------------------------------------------------------------
@@ -5447,6 +5822,18 @@ namespace FramePro
 
 		List<SendBuffer> send_buffer_list;
 
+		// send the main thread send buffer
+		Buffer* p_main_thread_send_buffer = nullptr;
+		{
+			CriticalSectionScope tls_lock(m_MainThreadSendBufferLock);
+			p_main_thread_send_buffer = m_MainThreadSendBuffer + m_MainThreadSendBufferIndex;
+			m_MainThreadSendBufferIndex = 1 - m_MainThreadSendBufferIndex;
+			m_MainThreadSendBuffer[m_MainThreadSendBufferIndex].Clear();
+		}
+		
+		Send(p_main_thread_send_buffer->GetBuffer(), p_main_thread_send_buffer->GetSize());
+		p_main_thread_send_buffer->Clear();
+
 		// get all of the send buffers
 		{
 			CriticalSectionScope tls_lock(m_TLSListCriticalSection);
@@ -5457,25 +5844,8 @@ namespace FramePro
 		// send the send buffers
 		for(SendBuffer* p_send_buffer = send_buffer_list.GetHead(); p_send_buffer; p_send_buffer = p_send_buffer->GetNext())
 		{
-			if(m_RecordingFile.IsOpened())
-			{
-				CriticalSectionScope lock2(m_CriticalSection);
-				WriteSendBuffer(p_send_buffer, m_RecordingFile, m_RecordingFileSize);
-			}
-			else
-			{
-#if FRAMEPRO_SOCKETS_ENABLED
-				if(m_Interactive)
-				{
-					if(!SendSendBuffer(p_send_buffer, m_ClientSocket))
-						break;		// disconnected
-				}
-				else
-				{
-					WriteSendBuffer(p_send_buffer, m_NonInteractiveRecordingFile, m_NonInteractiveRecordingFileSize);
-				}
-#endif
-			}
+			if(!Send(p_send_buffer->GetBuffer(), p_send_buffer->GetSize()))
+				break;		// disconnected
 		}
 
 		// give the empty send buffers back to the TLS objects
@@ -5538,7 +5908,11 @@ namespace FramePro
 		size_t bytes_to_read = read_file.GetSize();
 
 		const int block_size = 64*1024;
-		char* p_read_buffer = (char*)mp_Allocator->Alloc(block_size);
+
+		char* p_read_buffer = (char*)m_Allocator.Alloc(block_size);
+		if (!p_read_buffer)
+			return;
+
 		while(bytes_to_read)
 		{
 			size_t size_to_read = block_size < bytes_to_read ? block_size : bytes_to_read;
@@ -5547,24 +5921,25 @@ namespace FramePro
 			bytes_to_read -= size_to_read;
 		}
 		read_file.Close();
-		mp_Allocator->Free(p_read_buffer);
+		m_Allocator.Free(p_read_buffer);
 
-		Disconect_NoLock();
+		Disconnect_NoLock();
 #endif
 	}
 
 	//------------------------------------------------------------------------
-	void FrameProSession::Disconect(bool wait_for_threads_to_exit)
+	void FrameProSession::Disconnect(bool wait_for_threads_to_exit)
 	{
 		CriticalSectionScope lock(m_CriticalSection);
 		if(g_Connected)
-			Disconect_NoLock(wait_for_threads_to_exit);
+			Disconnect_NoLock(wait_for_threads_to_exit);
 	}
 
 	//------------------------------------------------------------------------
-	void FrameProSession::Disconect_NoLock(bool wait_for_threads_to_exit)
+	void FrameProSession::Disconnect_NoLock(bool wait_for_threads_to_exit)
 	{
-		Platform::StopRecordingContextSitches(mp_ContextSwitchRecorder);
+		if(mp_ContextSwitchRecorder)
+			Platform::StopRecordingContextSitches(mp_ContextSwitchRecorder);
 
 #if FRAMEPRO_SOCKETS_ENABLED
 		m_ClientSocket.Disconnect();
@@ -5604,7 +5979,7 @@ namespace FramePro
 		m_InitialiseConnectionNextFrame = false;
 
 		DynamicWString recording_filename;
-		recording_filename.SetAllocator(mp_Allocator);
+		recording_filename.SetAllocator(&m_Allocator);
 
 		if(m_RecordingFile.IsOpened())
 		{
@@ -5638,21 +6013,21 @@ namespace FramePro
 			}
 
 			// send session info
-			SessionInfoPacket session_info_packet;
+			SessionStatsPacket session_stats_packet;
 
 			{
 				CriticalSectionScope tls_lock(m_TLSListCriticalSection);
 				for(FrameProTLS* p_framepro_tls_iter = m_FrameProTLSList.GetHead(); p_framepro_tls_iter!=NULL; p_framepro_tls_iter=p_framepro_tls_iter->GetNext())
 				{
-					session_info_packet.m_SendBufferSize += p_framepro_tls_iter->GetSendBufferMemorySize();
-					session_info_packet.m_StringMemorySize += p_framepro_tls_iter->GetStringMemorySize();
-					session_info_packet.m_MiscMemorySize += sizeof(FrameProTLS);
+					session_stats_packet.m_SendBufferSize += p_framepro_tls_iter->GetSendBufferMemorySize();
+					session_stats_packet.m_StringMemorySize += p_framepro_tls_iter->GetStringMemorySize();
+					session_stats_packet.m_MiscMemorySize += sizeof(FrameProTLS);
 				}
 			}
 
-			session_info_packet.m_RecordingFileSize = m_NonInteractiveRecordingFileSize;
+			session_stats_packet.m_RecordingFileSize = m_NonInteractiveRecordingFileSize;
 
-			SendImmediate(&session_info_packet, sizeof(session_info_packet), p_framepro_tls);
+			SendImmediate(&session_stats_packet, sizeof(session_stats_packet), p_framepro_tls);
 		}
 	}
 
@@ -5676,25 +6051,25 @@ namespace FramePro
 	}
 
 	//------------------------------------------------------------------------
-	void FrameProSession::SendSessionDetails(const char* p_name, const char* p_build_id)
+	void FrameProSession::SendSessionInfo(const char* p_name, const char* p_value)
 	{
 		StringId name = RegisterString(p_name);
-		StringId build_id = RegisterString(p_build_id);
+		StringId value = RegisterString(p_value);
 
-		SendSessionDetails(name, build_id);
+		SendSessionInfo(name, value);
 	}
 
 	//------------------------------------------------------------------------
-	void FrameProSession::SendSessionDetails(const wchar_t* p_name, const wchar_t* p_build_id)
+	void FrameProSession::SendSessionInfo(const wchar_t* p_name, const wchar_t* p_value)
 	{
 		StringId name = RegisterString(p_name);
-		StringId build_id = RegisterString(p_build_id);
+		StringId value = RegisterString(p_value);
 
-		SendSessionDetails(name, build_id);
+		SendSessionInfo(name, value);
 	}
 
 	//------------------------------------------------------------------------
-	void FrameProSession::SendSessionDetails(StringId name, StringId build_id)
+	void FrameProSession::SendSessionInfo(StringId name, StringId value)
 	{
 		// this needs to be outside the critical section lock because it might lock the critical section itself
 		FrameProTLS* p_framepro_tls = GetFrameProTLS();
@@ -5703,11 +6078,7 @@ namespace FramePro
 
 		Initialise(p_framepro_tls);
 
-		char date_str[64];
-		GetDateString(date_str, sizeof(date_str));
-		StringId date = RegisterString(date_str);
-
-		p_framepro_tls->SendSessionInfoPacket(SessionDetailsPacket(name, build_id, date));
+		p_framepro_tls->SendSessionInfoPacket(SessionInfoPacket(name, value));
 	}
 
 	//------------------------------------------------------------------------
@@ -5781,7 +6152,12 @@ namespace FramePro
 		// initialise the connection
 		if(m_InitialiseConnectionNextFrame)
 		{
-			InitialiseConnection(p_framepro_tls);
+			if (!InitialiseConnection(p_framepro_tls))
+			{
+				Disconnect_NoLock(false);
+				return;
+			}
+
 			m_InitialiseConnectionNextFrame = false;
 		}
 
@@ -5800,16 +6176,6 @@ namespace FramePro
 			}
 		}
 
-		// send the m_MainThreadSendBuffer
-		{
-			CriticalSectionScope tls_lock(m_MainThreadSendBufferLock);
-			if(m_MainThreadSendBuffer.GetSize())
-			{
-				p_framepro_tls->Send(m_MainThreadSendBuffer.GetBuffer(), m_MainThreadSendBuffer.GetSize());
-				m_MainThreadSendBuffer.Clear();
-			}
-		}
-
 		if(g_Connected)
 		{
 			Platform::FlushContextSwitches(mp_ContextSwitchRecorder);
@@ -5817,7 +6183,7 @@ namespace FramePro
 			int64 wait_start_time;
 			FRAMEPRO_GET_CLOCK_COUNT(wait_start_time);
 
-			FlushGlobalHiResTimers(p_framepro_tls);
+			FlushCustomStatTimers(p_framepro_tls);
 
 			{
 				FRAMEPRO_NAMED_SCOPE("FramePro Wait For Send");
@@ -5842,16 +6208,32 @@ namespace FramePro
 			m_SendComplete.Reset();
 
 			// tell the TLS objects that the frame has started
+			bool error = false;
 			{
 				CriticalSectionScope tls_lock(m_TLSListCriticalSection);
-				for(FrameProTLS* p_iter=m_FrameProTLSList.GetHead(); p_iter!=NULL; p_iter=p_iter->GetNext())
-					p_iter->OnFrameStart();
+				for (FrameProTLS* p_iter = m_FrameProTLSList.GetHead(); p_iter != NULL; p_iter = p_iter->GetNext())
+				{
+					if (!p_iter->OnFrameStart())
+					{
+						error = true;
+						break;
+					}
+				}
+			}
+
+			if (error)
+			{
+				Disconnect_NoLock(false);
+				return;
 			}
 
 			// send the frame data
 			SendHeartbeatInfo(p_framepro_tls);
 
-			p_framepro_tls->SendFrameStartPacket(wait_for_send_complete_time);
+			int64 frame_start_time;
+			FRAMEPRO_GET_CLOCK_COUNT(frame_start_time);
+			FrameStartPacket frame_start_packet(frame_start_time, wait_for_send_complete_time);
+			SendOnMainThread(frame_start_packet);
 		}
 
 		// stop recording if the file has become too big
@@ -5882,7 +6264,7 @@ namespace FramePro
 	}
 
 	//------------------------------------------------------------------------
-	void FrameProSession::AddFrameProTLS(FrameProTLS* p_framepro_tls)
+	bool FrameProSession::AddFrameProTLS(FrameProTLS* p_framepro_tls)
 	{
 		CriticalSectionScope lock(m_CriticalSection);
 
@@ -5891,8 +6273,13 @@ namespace FramePro
 			m_FrameProTLSList.AddTail(p_framepro_tls);
 		}
 
-		if(g_Connected)
-			p_framepro_tls->OnConnected(m_RecordingFile.IsOpened());
+		if (g_Connected)
+		{
+			if (!p_framepro_tls->OnConnected(m_RecordingFile.IsOpened()))
+				return false;
+		}
+
+		return true;
 	}
 
 	//------------------------------------------------------------------------
@@ -5923,6 +6310,21 @@ namespace FramePro
 		p_framepro_tls->SetThreadName(thread_id, p_name);
 	}
 
+	//------------------------------------------------------------------------
+	void FrameProSession::SetThreadName(int thread_id, const char* p_name)
+	{
+		FrameProTLS* p_framepro_tls = GetFrameProTLS();
+
+		CriticalSectionScope lock(m_CriticalSection);
+
+		m_NamedThreads.SetAllocator(GetAllocator());
+
+		if (!m_NamedThreads.Contains(thread_id))
+			m_NamedThreads.Add(thread_id);
+
+		p_framepro_tls->SetThreadName(thread_id, p_name);
+	}
+
 	//------------------------------------------------------------------------
 	bool FrameProSession::HasSetThreadName(int thread_id) const
 	{
@@ -5987,7 +6389,7 @@ namespace FramePro
 	//------------------------------------------------------------------------
 	void FrameProSession::StartRecording(const char* p_filename, bool context_switches, bool callstacks, int64 max_file_size)
 	{
-		Disconect(true);
+		Disconnect(true);
 
 		CreateDefaultAllocator();
 
@@ -6029,7 +6431,7 @@ namespace FramePro
 	//------------------------------------------------------------------------
 	void FrameProSession::StartRecording(const wchar_t* p_filename, bool context_switches, bool callstacks, int64 max_file_size)
 	{
-		Disconect(true);
+		Disconnect(true);
 
 		CreateDefaultAllocator();
 
@@ -6078,7 +6480,8 @@ namespace FramePro
 			#if FRAMEPRO_SOCKETS_ENABLED
 				OpenListenSocket();					// start the listening socket again so that we can accept new connections
 			#endif
-				Disconect_NoLock();
+
+			Disconnect_NoLock();
 		}
 	}
 
@@ -6109,18 +6512,23 @@ namespace FramePro
 	}
 
 	//------------------------------------------------------------------------
-	void FrameProSession::AddGlobalHiResTimer(GlobalHiResTimer* p_timer)
+	void FrameProSession::AddCustomStatTimer(CustomStatTimer* p_timer)
 	{
-		CriticalSectionScope lock(m_CriticalSection);
+		{
+			CriticalSectionScope lock(m_CriticalSection);
+
+			p_timer->SetNext(m_CustomStatTimers);
+			m_CustomStatTimers = p_timer;
+		}
 
-		p_timer->SetNext(mp_GlobalHiResTimers);
-		mp_GlobalHiResTimers = p_timer;
+		if(g_Connected)
+			SetCustomStatUnit(FramePro::RegisterString(p_timer->GetName()), FramePro::RegisterString("cycles"));
 	}
 
 	//------------------------------------------------------------------------
-	void FrameProSession::FlushGlobalHiResTimers(FrameProTLS* p_framepro_tls)
+	void FrameProSession::FlushCustomStatTimers(FrameProTLS* p_framepro_tls)
 	{
-		for(GlobalHiResTimer* p_timer = mp_GlobalHiResTimers; p_timer != NULL; p_timer = p_timer->GetNext())
+		for(CustomStatTimer* p_timer = m_CustomStatTimers; p_timer != NULL; p_timer = p_timer->GetNext())
 		{
 			uint64 value;
 			uint count;
@@ -6140,13 +6548,14 @@ namespace FramePro
 			p_packet->m_Count = count;
 			p_packet->m_Name = (StringId)p_timer->GetName();
 			p_packet->m_Value = value;
+			FRAMEPRO_GET_CLOCK_COUNT(p_packet->m_Time);
 		}
 	}
 
 	//------------------------------------------------------------------------
-	void FrameProSession::ClearGlobalHiResTimers()
+	void FrameProSession::ClearCustomStatTimers()
 	{
-		for (GlobalHiResTimer* p_timer = mp_GlobalHiResTimers; p_timer != NULL; p_timer = p_timer->GetNext())
+		for (CustomStatTimer* p_timer = m_CustomStatTimers; p_timer != NULL; p_timer = p_timer->GetNext())
 		{
 			uint64 value;
 			uint count;
@@ -6177,6 +6586,9 @@ namespace FramePro
 		{
 			if (m_ScopeColours[i].m_Name == name)
 			{
+				if(m_ScopeColours[i].m_Colour == colour)
+					return;
+
 				m_ScopeColours[i].m_Colour = colour;
 				updated_existing = true;
 				break;
@@ -6311,6 +6723,46 @@ namespace FramePro
 			GetFrameProTLS()->SendCustomStatColourPacket(name, colour);
 	}
 
+	//------------------------------------------------------------------------
+	void FrameProSession::SendScopeCallstack()
+	{
+#if FRAMEPRO_ENABLE_CALLSTACKS
+		if (!g_Connected)
+			return;
+
+		SendModules();
+
+		FrameProTLS* p_framepro_tls = GetFrameProTLS();
+
+		CallstackResult callstack_result = p_framepro_tls->GetCallstack();
+
+		p_framepro_tls->SetCurrentCallstackId(callstack_result.mp_Callstack->m_ID);
+
+		int send_size = sizeof(CallstackPacket);
+		if (callstack_result.m_IsNew)
+			send_size += callstack_result.mp_Callstack->m_Size * sizeof(uint64);
+
+		{
+			CriticalSectionScope lock(p_framepro_tls->GetCurrentSendBufferCriticalSection());
+
+			CallstackPacket* p_packet = (CallstackPacket*)p_framepro_tls->AllocateSpaceInBuffer(send_size);
+
+			p_packet->m_PacketType = PacketType::CallstackPacket;
+			p_packet->m_CallstackId = callstack_result.mp_Callstack->m_ID;
+			p_packet->m_CallstackSize = 0;
+
+			if (callstack_result.m_IsNew)
+			{
+				p_packet->m_CallstackSize = callstack_result.mp_Callstack->m_Size;
+				memcpy(
+					(char*)(p_packet + 1),
+					callstack_result.mp_Callstack->mp_Stack,
+					callstack_result.mp_Callstack->m_Size * sizeof(uint64));
+			}
+		}
+#endif
+	}
+
 	//------------------------------------------------------------------------
 	void FrameProSession::SendCustomStatGraphs()
 	{
@@ -6439,7 +6891,7 @@ namespace FramePro
 		mp_CurrentSendBuffer(NULL),
 		m_CurrentSendBufferSize(0),
 		m_ThreadId(Platform::GetCurrentThreadId()),
-		m_HiResTimerScopeStartTime(0),
+		m_FirstHiresTimerTime(0),
 		m_HiResTimerStartTime(0),
 		m_ActiveHiResTimerIndex(-1),
 		mp_Next(NULL),
@@ -6456,10 +6908,15 @@ namespace FramePro
 		m_StringMemorySize(0),
 		m_ClockFrequency(clock_frequency),
 		m_ShuttingDown(false),
-		mp_CurrentConditionalParentScope(NULL)
+		mp_CurrentConditionalParentScope(NULL),
+		m_CurrentCallstackId(-1),
+		m_CurrentCallstackIdTime(0)
 #if FRAMEPRO_ENABLE_CALLSTACKS
 		,m_StackTrace(p_allocator)
 		,m_SendCallstacks(false)
+#endif
+#if FRAMEPRO_DEBUG || FRAMEPRO_ALLOW_UNPARENTED_HIRES_SCOPES
+		,m_ScopeDepth(0)
 #endif
 	{
 		UpdateSendStringsImmediatelyFlag();
@@ -6531,6 +6988,8 @@ namespace FramePro
 			m_ConditionalParentScopeList.Clear();
 		}
 
+		m_InitialisedCustomStats.Clear();
+
 		UpdateStringMemorySize();
 
 #if FRAMEPRO_ENABLE_CALLSTACKS
@@ -6550,7 +7009,7 @@ namespace FramePro
 
 	//------------------------------------------------------------------------
 	// this is called from the main thread
-	void FrameProTLS::OnConnected(bool recording_to_file)
+	bool FrameProTLS::OnConnected(bool recording_to_file)
 	{
 		CriticalSectionScope lock(m_CriticalSection);
 		
@@ -6565,13 +7024,16 @@ namespace FramePro
 
 			{
 				CriticalSectionScope lock2(m_CurrentSendBufferCriticalSection);
-				AllocateCurrentSendBuffer();
+				if (!AllocateCurrentSendBuffer())
+					return false;
 			}
 		}
+
+		return true;
 	}
 
 	//------------------------------------------------------------------------
-	void FrameProTLS::AllocateCurrentSendBuffer()
+	bool FrameProTLS::AllocateCurrentSendBuffer()
 	{
 		FRAMEPRO_ASSERT(m_CriticalSection.Locked());
 		FRAMEPRO_ASSERT(m_CurrentSendBufferCriticalSection.Locked());
@@ -6580,10 +7042,13 @@ namespace FramePro
 		if(!mp_CurrentSendBuffer)
 		{
 			mp_CurrentSendBuffer = mp_Allocator->Alloc(m_SendBufferCapacity);
-			FRAMEPRO_ASSERT(mp_CurrentSendBuffer);
+			if (!mp_CurrentSendBuffer)
+				return false;
 
 			m_SendBufferMemorySize = m_SendBufferMemorySize + m_SendBufferCapacity;
 		}
+
+		return true;
 	}
 
 	//------------------------------------------------------------------------
@@ -6656,6 +7121,9 @@ namespace FramePro
 
 					// copy current send buffer to a new send buffer object
 					SendBuffer* p_send_buffer = AllocateSendBuffer();
+					if (!p_send_buffer)
+						return;
+
 					p_send_buffer->Swap(mp_CurrentSendBuffer, m_CurrentSendBufferSize, m_SendBufferCapacity);
 					FRAMEPRO_ASSERT(mp_CurrentSendBuffer);
 					available_space = m_SendBufferCapacity;
@@ -6672,15 +7140,18 @@ namespace FramePro
 	}
 
 	//------------------------------------------------------------------------
-	void FrameProTLS::OnFrameStart()
+	bool FrameProTLS::OnFrameStart()
 	{
 		UpdateStringMemorySize();
 
 		m_SessionInfoBufferMemorySize = m_SessionInfoBuffer.GetMemorySize();
 
-		FlushCurrentSendBuffer();
+		if (!FlushCurrentSendBuffer())
+			return false;
 
 		FlushConditionalChildSendBuffers();
+
+		return true;
 	}
 
 	//------------------------------------------------------------------------
@@ -6706,15 +7177,6 @@ namespace FramePro
 		}
 	}
 
-	//------------------------------------------------------------------------
-	void FrameProTLS::SendFrameStartPacket(int64 wait_for_send_complete_time)
-	{
-		// start the new frame
-		int64 frame_start_time;
-		FRAMEPRO_GET_CLOCK_COUNT(frame_start_time);
-		SendPacket(FrameStartPacket(frame_start_time, wait_for_send_complete_time));
-	}
-
 	//------------------------------------------------------------------------
 	void FrameProTLS::SetThreadName(int thread_id, const char* p_name)
 	{
@@ -6730,9 +7192,9 @@ namespace FramePro
 	}
 
 	//------------------------------------------------------------------------
-	void FrameProTLS::SetMainThread(int main_thraed_id)
+	void FrameProTLS::SetMainThread(int main_thread_id)
 	{
-		SendSessionInfoPacket(MainThreadPacket(main_thraed_id));
+		SendSessionInfoPacket(MainThreadPacket(main_thread_id));
 	}
 
 	//------------------------------------------------------------------------
@@ -6747,7 +7209,8 @@ namespace FramePro
 		{
 			string_id = ++m_StringCount;
 			str.TakeCopy(m_StringAllocator);
-			m_StringHashMap.Add(str, string_id);
+			if (!m_StringHashMap.Add(str, string_id))
+				return 0;
 		
 			SendString(string_id, p_str, PacketType::StringPacket);
 
@@ -6769,7 +7232,8 @@ namespace FramePro
 		{
 			string_id = ++m_StringCount;
 			str.TakeCopy(m_StringAllocator);
-			m_WStringHashMap.Add(str, string_id);
+			if (!m_WStringHashMap.Add(str, string_id))
+				return 0;
 		
 			SendString(string_id, p_str, PacketType::WStringPacket);
 
@@ -6786,6 +7250,8 @@ namespace FramePro
 			// copy it to the session buffer
 			CriticalSectionScope lock(m_SessionInfoBufferLock);
 			void* p_dest = m_SessionInfoBuffer.Allocate(size);
+			if (!p_dest)
+				return;
 			memcpy(p_dest, p_data, size);
 		}
 
@@ -6810,10 +7276,7 @@ namespace FramePro
 
 			p_packet = (StringPacket*)(m_SessionInfoBuffer.Allocate(size_to_allocate));
 			if(!p_packet)
-			{
-				ShowMemoryWarning();
 				return;
-			}
 
 			p_packet->m_PacketType = packet_type;
 			p_packet->m_Length = (int)string_len;
@@ -6847,10 +7310,7 @@ namespace FramePro
 
 				p_packet = (StringPacket*)(m_SessionInfoBuffer.Allocate(size_to_allocate));
 				if(!p_packet)
-				{
-					ShowMemoryWarning();
 					return;
-				}
 
 				p_packet->m_PacketType = packet_type;
 				p_packet->m_Length = (int)string_len;
@@ -6870,10 +7330,7 @@ namespace FramePro
 
 				p_packet = (StringPacket*)(m_SessionInfoBuffer.Allocate(size_to_allocate));
 				if (!p_packet)
-				{
-					ShowMemoryWarning();
 					return;
-				}
 
 				p_packet->m_PacketType = packet_type;
 				p_packet->m_Length = (int)string_len;
@@ -7058,23 +7515,11 @@ namespace FramePro
 	}
 
 	//------------------------------------------------------------------------
-	void FrameProTLS::ShowMemoryWarning() const
-	{
-		static int64 last_warn_time = 0;
-		int64 now;
-		FRAMEPRO_GET_CLOCK_COUNT(now);
-
-		if(now - last_warn_time >= m_ClockFrequency)
-		{
-			Platform::DebugWrite("Warning: FramePro failed to allocate enough memory.");
-			last_warn_time = now;
-		}
-	}
-
-	//------------------------------------------------------------------------
-	void FrameProTLS::FlushCurrentSendBuffer()
+	bool FrameProTLS::FlushCurrentSendBuffer()
 	{
 		SendBuffer* p_send_buffer = AllocateSendBuffer();
+		if (!p_send_buffer)
+			return false;
 
 		{
 			CriticalSectionScope lock(m_CurrentSendBufferCriticalSection);
@@ -7094,6 +7539,8 @@ namespace FramePro
 			CriticalSectionScope lock(m_CriticalSection);
 			m_SendBufferList.AddTail(p_send_buffer);
 		}
+
+		return true;
 	}
 
 	//------------------------------------------------------------------------
@@ -7103,6 +7550,8 @@ namespace FramePro
 		FRAMEPRO_ASSERT(IsOnTLSThread() || !g_Connected);		// can only be accessed from TLS thread, unless we haven't connected yet
 
 		SendBuffer* p_send_buffer = AllocateSendBuffer();
+		if (!p_send_buffer)
+			return;
 
 		p_send_buffer->Swap(mp_CurrentSendBuffer, m_CurrentSendBufferSize, m_SendBufferCapacity);
 
@@ -7112,6 +7561,8 @@ namespace FramePro
 		if (mp_CurrentConditionalParentScope)
 		{
 			SendBuffer* p_new_parent_send_buffer = AllocateSendBuffer();
+			if (!p_new_parent_send_buffer)
+				return;
 
 			{
 				// move the current child send buffer to the parents child buffer list
@@ -7148,6 +7599,9 @@ namespace FramePro
 		else
 		{
 			p_send_buffer = New<SendBuffer>(mp_Allocator, mp_Allocator, m_SendBufferCapacity, this);
+			if (!p_send_buffer)
+				return NULL;
+
 			m_SendBufferMemorySize = m_SendBufferMemorySize + m_SendBufferCapacity + sizeof(SendBuffer);		// doesn't need to be atomic, it's only for stats
 		}
 
@@ -7156,7 +7610,9 @@ namespace FramePro
 
 		if (!p_send_buffer->GetBuffer())
 		{
-			p_send_buffer->AllocateBuffer(m_SendBufferCapacity);
+			if (!p_send_buffer->AllocateBuffer(m_SendBufferCapacity))
+				return NULL;
+
 			m_SendBufferMemorySize = m_SendBufferMemorySize + m_SendBufferCapacity;
 		}
 
@@ -7196,6 +7652,9 @@ namespace FramePro
 		FRAMEPRO_ASSERT(m_ConditionalParentScopeListCritSec.Locked());
 
 		ConditionalParentScope* p_scope = New<ConditionalParentScope>(mp_Allocator, p_name);
+		if (!p_scope)
+			return NULL;
+	
 		m_ConditionalParentScopeList.AddTail(p_scope);
 
 		return p_scope;
@@ -7209,11 +7668,17 @@ namespace FramePro
 		FRAMEPRO_ASSERT(!mp_CurrentConditionalParentScope);		// nested conditional parent scopes not supported
 
 		ConditionalParentScope* p_scope = GetConditionalParentScope(p_name);
-		if(!p_scope)
+		if (!p_scope)
+		{
 			p_scope = CreateConditionalParentScope(p_name);
+			if (!p_scope)
+				return;
+		}
 
 		FRAMEPRO_ASSERT(!p_scope->mp_SendBuffer);
 		p_scope->mp_SendBuffer = AllocateSendBuffer();
+		if (!p_scope->mp_SendBuffer)
+			return;
 
 		p_scope->m_PreDuration = pre_duration;
 		p_scope->m_PostDuration = post_duration;
@@ -7291,10 +7756,17 @@ namespace FramePro
 	}
 
 	//------------------------------------------------------------------------
-	void FrameProTLS::SendHiResTimersScope(int64 current_time)
+	// this is called when a normal scope ends. If the hires timers are within
+	// this scope send them.
+	void FrameProTLS::SendHiResTimersScope(int64 start_time, int64 end_time)
 	{
 		FRAMEPRO_ASSERT(IsOnTLSThread());
 
+		FRAMEPRO_ASSERT(m_FirstHiresTimerTime < end_time);
+
+		if(m_FirstHiresTimerTime < start_time)
+			return;
+
 		int count = m_HiResTimers.GetCount();
 		FRAMEPRO_ASSERT(count);
 
@@ -7312,8 +7784,8 @@ namespace FramePro
 
 			HiResTimerScopePacket* p_packet = (HiResTimerScopePacket*)AllocateSpaceInBuffer(size_to_send);
 			p_packet->m_PacketType = PacketType::HiResTimerScopePacket;
-			p_packet->m_StartTime = m_HiResTimerScopeStartTime;
-			p_packet->m_EndTime = current_time;
+			p_packet->m_StartTime = start_time;
+			p_packet->m_EndTime = end_time;
 			p_packet->m_Count = count;
 			p_packet->m_ThreadId = m_ThreadId;
 			p_packet->m_Padding = 0;
@@ -7323,6 +7795,8 @@ namespace FramePro
 		}
 
 		m_HiResTimers.ClearNoFree();
+
+		m_FirstHiresTimerTime = 0;
 	}
 
 	//------------------------------------------------------------------------
@@ -7437,7 +7911,10 @@ namespace FramePro
 		if (old_count <= name_index)
 		{
 			int new_count = name_index + 1;
-			m_InitialisedCustomStats.Resize(new_count);
+			
+			if (!m_InitialisedCustomStats.Resize(new_count))
+				return;
+
 			memset(&m_InitialisedCustomStats[old_count], 0, (new_count - old_count) * sizeof(bool));
 		}
 		m_InitialisedCustomStats[name_index] = true;
@@ -7465,7 +7942,7 @@ namespace FramePro
 	{
 		StringId name = RegisterString(p_name);
 
-		if (HaveSentCustomStatInfo(name))
+		if (!HaveSentCustomStatInfo(name))
 		{
 			SendCustomStatGraphPacket(name, RegisterString(p_graph));
 			SendCustomStatUnitPacket(name, RegisterString(p_unit));
@@ -7480,7 +7957,7 @@ namespace FramePro
 	//------------------------------------------------------------------------
 	void FrameProTLS::SetCustomStatInfo(StringId name, const char* p_graph, const char* p_unit, uint colour)
 	{
-		if (HaveSentCustomStatInfo(name))
+		if (!HaveSentCustomStatInfo(name))
 		{
 			SendCustomStatGraphPacket(name, RegisterString(p_graph));
 			SendCustomStatUnitPacket(name, RegisterString(p_unit));
@@ -7495,7 +7972,7 @@ namespace FramePro
 	//------------------------------------------------------------------------
 	void FrameProTLS::SetCustomStatInfo(StringId name, const wchar_t* p_graph, const wchar_t* p_unit, uint colour)
 	{
-		if (HaveSentCustomStatInfo(name))
+		if (!HaveSentCustomStatInfo(name))
 		{
 			SendCustomStatGraphPacket(name, RegisterString(p_graph));
 			SendCustomStatUnitPacket(name, RegisterString(p_unit));
@@ -7510,7 +7987,7 @@ namespace FramePro
 	//------------------------------------------------------------------------
 	void FrameProTLS::SetCustomStatInfo(StringId name, StringId graph, StringId unit, uint colour)
 	{
-		if (HaveSentCustomStatInfo(name))
+		if (!HaveSentCustomStatInfo(name))
 		{
 			SendCustomStatGraphPacket(name, graph);
 			SendCustomStatUnitPacket(name, unit);
@@ -7529,6 +8006,32 @@ namespace FramePro
 		return m_StackTrace.Capture();
 	}
 #endif
+
+	//------------------------------------------------------------------------
+	int FrameProTLS::GetAndClearCurrentCallstackId(int64 start_time, int64 end_time)
+	{
+		if (m_CurrentCallstackId == -1)
+			return -1;
+
+		int64 callstack_time = m_CurrentCallstackIdTime;
+
+		if (callstack_time >= start_time && callstack_time < end_time)
+		{
+			int callstack_id = m_CurrentCallstackId;
+			m_CurrentCallstackIdTime = 0;
+			m_CurrentCallstackId = -1;
+			return callstack_id;
+		}
+
+		return -1;
+	}
+
+	//------------------------------------------------------------------------
+	void FrameProTLS::SetCurrentCallstackId(int id)
+	{
+		m_CurrentCallstackId = id;
+		FRAMEPRO_GET_CLOCK_COUNT(m_CurrentCallstackIdTime);
+	}
 }
 
 //------------------------------------------------------------------------
@@ -7579,8 +8082,11 @@ namespace FramePro
 	//------------------------------------------------------------------------
 	void* IncrementingBlockAllocator::Alloc(size_t size)
 	{
-		if(m_CurrentBlockSize + size > m_MemoryBlockSize)
-			AllocateBlock();
+		if (m_CurrentBlockSize + size > m_MemoryBlockSize)
+		{
+			if (!AllocateBlock())
+				return NULL;
+		}
 
 		void* p_mem = mp_BlockList->m_Memory + m_CurrentBlockSize;
 		m_CurrentBlockSize += size;
@@ -7588,14 +8094,19 @@ namespace FramePro
 	}
 
 	//------------------------------------------------------------------------
-	void IncrementingBlockAllocator::AllocateBlock()
+	bool IncrementingBlockAllocator::AllocateBlock()
 	{
 		Block* p_block = (Block*)mp_Allocator->Alloc(sizeof(Block));
+		if (!p_block)
+			return false;
+
 		p_block->mp_Next = mp_BlockList;
 		mp_BlockList = p_block;
 		m_CurrentBlockSize = 0;
 
 		m_MemorySize += m_BlockSize;
+
+		return true;
 	}
 }
 
@@ -7619,7 +8130,8 @@ namespace FramePro
 		m_Capacity(g_InitialCapacity),
 		mp_Allocator(p_allocator)
 	{
-		memset(mp_Data, 0, g_InitialCapacity*sizeof(const void*));
+		if(mp_Data)
+			memset(mp_Data, 0, g_InitialCapacity*sizeof(const void*));
 	}
 
 	//------------------------------------------------------------------------
@@ -7629,7 +8141,7 @@ namespace FramePro
 	}
 
 	//------------------------------------------------------------------------
-	void PointerSet::Grow()
+	bool PointerSet::Grow()
 	{
 		int old_capacity = m_Capacity;
 		const void** p_old_data = mp_Data;
@@ -7641,6 +8153,8 @@ namespace FramePro
 		m_CapacityMask = m_Capacity - 1;
 		size_t alloc_size = m_Capacity * sizeof(const void*);
 		mp_Data = (const void**)mp_Allocator->Alloc(alloc_size);
+		if (!mp_Data)
+			return false;
 
 		int size = m_Capacity * sizeof(void*);
 		memset(mp_Data, 0, size);
@@ -7658,6 +8172,8 @@ namespace FramePro
 		mp_Allocator->Free(p_old_data);
 
 		alloc_size -= old_capacity * sizeof(const void*);
+
+		return true;
 	}
 
 	//------------------------------------------------------------------------
@@ -7666,7 +8182,9 @@ namespace FramePro
 	{
 		if(m_Count >= m_Capacity/4)
 		{
-			Grow();
+			if (!Grow())
+				return false;
+
 			index = hash & m_CapacityMask;
 		}
 
@@ -7698,15 +8216,17 @@ namespace FramePro
 {
 	//------------------------------------------------------------------------
 	SendBuffer::SendBuffer(Allocator* p_allocator, int capacity, FrameProTLS* p_owner)
-	:	mp_Buffer(p_allocator->Alloc(capacity)),
+	:	mp_Buffer(NULL),
 		m_Size(0),
-		m_Capacity(capacity),
+		m_Capacity(0),
 		mp_Next(NULL),
 		mp_Allocator(p_allocator),
 		mp_Owner(p_owner),
 		m_CreationTime(0)
 	{
 		SetCreationTime();
+
+		AllocateBuffer(capacity);
 	}
 
 	//------------------------------------------------------------------------
@@ -7716,12 +8236,17 @@ namespace FramePro
 	}
 
 	//------------------------------------------------------------------------
-	void SendBuffer::AllocateBuffer(int capacity)
+	bool SendBuffer::AllocateBuffer(int capacity)
 	{
 		FRAMEPRO_ASSERT(!mp_Buffer);
 
 		mp_Buffer = mp_Allocator->Alloc(capacity);
+		if (!mp_Buffer)
+			return false;
+
 		m_Capacity = capacity;
+
+		return true;
 	}
 
 	//------------------------------------------------------------------------
@@ -7899,7 +8424,7 @@ namespace FramePro
 	}
 
 	//------------------------------------------------------------------------
-	void Thread::CreateThread(FramePro::ThreadMain p_thread_main, void* p_param, Allocator* p_allocator)
+	bool Thread::CreateThread(FramePro::ThreadMain p_thread_main, void* p_param, Allocator* p_allocator)
 	{
 		if(m_Created)
 			Platform::DestroyThread(m_OSThread);
@@ -7907,9 +8432,10 @@ namespace FramePro
 		mp_ThreadMain = p_thread_main;
 		mp_Param = p_param;
 
-		Platform::CreateThread(m_OSThread, sizeof(m_OSThread), ThreadMain, this, p_allocator);
+		if (Platform::CreateThread(m_OSThread, sizeof(m_OSThread), ThreadMain, this, p_allocator))
+			m_Created = true;
 
-		m_Created = true;
+		return m_Created;
 	}
 
 	//------------------------------------------------------------------------
@@ -7945,11 +8471,11 @@ namespace FramePro
 
 //------------------------------------------------------------------------
 //
-// EnumModulesLinux.hpp
+// EnumModulesLinux.h
 //
 
 //------------------------------------------------------------------------
-#if FRAMEPRO_LINUX_BASED_PLATFORM
+#if FRAMEPRO_LINUX_BASED_PLATFORM && !FRAMEPRO_PLATFORM_PS4
 
 //------------------------------------------------------------------------
 #if FRAMEPRO_ENABLE_CALLSTACKS
@@ -7977,7 +8503,7 @@ namespace FramePro
 //
 
 //------------------------------------------------------------------------
-#if FRAMEPRO_LINUX_BASED_PLATFORM
+#if FRAMEPRO_LINUX_BASED_PLATFORM && !FRAMEPRO_PLATFORM_PS4 && !(FRAMEPRO_PLATFORM_UE4 && PLATFORM_PS4)
 
 //------------------------------------------------------------------------
 #if FRAMEPRO_ENABLE_CALLSTACKS
@@ -8017,6 +8543,9 @@ namespace FramePro
 		ModuleCallbackContext* p_context)
 	{
 		ModulePacket* p_module_packet = (ModulePacket*)p_context->mp_Allocator->Alloc(sizeof(ModulePacket));
+		if (!p_module_packet)
+			return;
+
 		memset(p_module_packet, 0, sizeof(ModulePacket));
 		p_module_packet->m_PacketType = PacketType::ModulePacket;
 
@@ -8091,6 +8620,9 @@ namespace FramePro
 		// if you are having problems compiling this on your platform unset FRAMEPRO_ENUMERATE_ALL_MODULES and it send info for just the main module
 #if FRAMEPRO_ENUMERATE_ALL_MODULES
 		ModuleCallbackContext* p_context = (ModuleCallbackContext*)p_allocator->Alloc(sizeof(ModuleCallbackContext));
+		if (!p_context)
+			return;
+
 		p_context->mp_ModulePackets = &module_packets;
 		p_context->mp_Allocator = p_allocator;
 
@@ -8105,6 +8637,9 @@ namespace FramePro
 			// to getting the base address for the main module. This will always work for for all platforms.
 
 			ModulePacket* p_module_packet = (ModulePacket*)p_allocator->Alloc(sizeof(ModulePacket));
+			if (!p_module_packet)
+				return;
+
 			memset(p_module_packet, 0, sizeof(ModulePacket));
 
 			p_module_packet->m_PacketType = PacketType::ModulePacket;
@@ -8140,7 +8675,7 @@ namespace FramePro
 
 //------------------------------------------------------------------------
 //
-// EnumModulesWindows.hpp
+// EnumModulesWindows.h
 //
 
 //------------------------------------------------------------------------
@@ -8291,6 +8826,9 @@ namespace FramePro
 			ModuleCallbackContext* p_context)
 		{
 			ModulePacket* p_module_packet = (ModulePacket*)p_context->mp_Allocator->Alloc(sizeof(ModulePacket));
+			if (!p_module_packet)
+				return;
+
 			memset(p_module_packet, 0, sizeof(ModulePacket));
 			p_module_packet->m_PacketType = PacketType::ModulePacket;
 
@@ -8304,7 +8842,8 @@ namespace FramePro
 
 			GetExtraModuleInfo(module_base, p_module_packet);
 
-			p_context->mp_ModulePackets->Add(p_module_packet);
+			if (!p_context->mp_ModulePackets->Add(p_module_packet))
+				return;
 		}
 
 		#if !defined(_IMAGEHLP_SOURCE_) && defined(_IMAGEHLP64)
@@ -8328,6 +8867,9 @@ namespace FramePro
 		// if you are having problems compiling this on your platform unset FRAMEPRO_ENUMERATE_ALL_MODULES and it send info for just the main module
 		#if FRAMEPRO_ENUMERATE_ALL_MODULES
 			ModuleCallbackContext* p_context = (ModuleCallbackContext*)p_allocator->Alloc(sizeof(ModuleCallbackContext));
+			if (!p_context)
+				return;
+
 			p_context->mp_ModulePackets = &module_packets;
 			p_context->mp_Allocator = p_allocator;
 
@@ -8342,6 +8884,9 @@ namespace FramePro
 			// to getting the base address for the main module. This will always work for for all platforms.
 
 			ModulePacket* p_module_packet = (ModulePacket*)p_allocator->Alloc(sizeof(ModulePacket));
+			if (!p_module_packet)
+				return;
+
 			memset(p_module_packet, 0, sizeof(ModulePacket));
 
 			p_module_packet->m_PacketType = PacketType::ModulePacket;
@@ -8363,7 +8908,8 @@ namespace FramePro
 				strcpy_s(p_module_packet->m_ModuleName, tchar_filename);
 			#endif
 
-			module_packets.Add(p_module_packet);
+			if(!module_packets.Add(p_module_packet))
+				return;
 		}
 	}
 }
@@ -8418,9 +8964,11 @@ namespace FramePro
 	#endif
 
 	//------------------------------------------------------------------------
-	#if FRAMEPRO_MAX_PATH != MAX_PATH
-		#error
-	#endif
+//@EPIC BEGIN: MAX_PATH might not be defined
+	//#if FRAMEPRO_MAX_PATH != MAX_PATH
+	//	#error
+	//#endif
+//@EPIC END
 
 	//------------------------------------------------------------------------
 	#if FRAMEPRO_SOCKETS_ENABLED
@@ -8435,7 +8983,7 @@ namespace FramePro
 		#include <winsock2.h>
 		#include <ws2tcpip.h>
 		#if FRAMEPRO_PLATFORM_UE4
-			#include "Windows/HideWindowsPlatformTypes.h"
+			#include "Windows//HideWindowsPlatformTypes.h"
 		#endif
 	#endif
 
@@ -8563,23 +9111,21 @@ namespace FramePro
 
 					int buffer_size = 1024;
 					FRAMEPRO_TCHAR* p_buffer = (FRAMEPRO_TCHAR*)HeapAlloc(GetProcessHeap(), 0, buffer_size * sizeof(FRAMEPRO_TCHAR));
-					if (p_buffer)
-					{
-						memset(p_buffer, 0, buffer_size * sizeof(FRAMEPRO_TCHAR));
+					memset(p_buffer, 0, buffer_size * sizeof(FRAMEPRO_TCHAR));
 
-						FormatMessage(
-							FORMAT_MESSAGE_FROM_SYSTEM,
-							NULL,
-							WSAGetLastError(),
-							0,
-							p_buffer,
-							buffer_size,
-							NULL);
+					va_list args;
+					FormatMessage(
+						FORMAT_MESSAGE_FROM_SYSTEM,
+						NULL,
+						WSAGetLastError(),
+						0,
+						p_buffer,
+						buffer_size,
+						&args);
 
-						FramePro::DebugWrite("FramePro Network Error: %s\n", p_buffer);
+					FramePro::DebugWrite("FramePro Network Error: %ws\n", p_buffer);
 
-						HeapFree(GetProcessHeap(), 0, p_buffer);
-					}
+					HeapFree(GetProcessHeap(), 0, p_buffer);
 				#endif
 			}
 
@@ -8714,7 +9260,7 @@ namespace FramePro
 
 					// Resolve the server address and port
 					addrinfo* p_result_info;
-					int result = getaddrinfo(NULL, p_port, &info, &p_result_info);
+					HRESULT result = getaddrinfo(NULL, p_port, &info, &p_result_info);
 					if (result != 0)
 					{
 						HandleSocketError();
@@ -8762,6 +9308,10 @@ namespace FramePro
 					SOCKET& target_socket = GetOSSocket(p_target_os_socket_mem);
 
 					target_socket = accept(source_socket, NULL, NULL);
+
+					if (target_socket == INVALID_SOCKET)
+						HandleSocketError();
+
 					return target_socket != INVALID_SOCKET;
 				#else
 					FRAMEPRO_UNREFERENCED(p_source_os_socket_mem);
@@ -8841,7 +9391,7 @@ namespace FramePro
 			}
 
 			//------------------------------------------------------------------------
-			void CreateThread(
+			bool CreateThread(
 				void* p_os_thread_mem,
 				int os_thread_mem_size,
 				ThreadMain p_thread_main,
@@ -8852,11 +9402,16 @@ namespace FramePro
 				FRAMEPRO_UNREFERENCED(os_thread_mem_size);
 
 				ThreadContext* p_thread_context = New<ThreadContext>(p_allocator);
+				if (!p_thread_context)
+					return false;
+
 				p_thread_context->mp_ThreadMain = p_thread_main;
 				p_thread_context->mp_Context = p_context;
 				p_thread_context->mp_Allocator = p_allocator;
 
 				GetOSThread(p_os_thread_mem) = ::CreateThread(NULL, 0, PlatformThreadMain, p_thread_context, 0, NULL);
+
+				return true;
 			}
 
 			//------------------------------------------------------------------------
@@ -9369,7 +9924,7 @@ namespace FramePro
 			}
 
 			//------------------------------------------------------------------------
-			void CreateThread(
+			bool CreateThread(
 				void* p_os_thread_mem,
 				int os_thread_mem_size,
 				ThreadMain p_thread_main,
@@ -9379,11 +9934,18 @@ namespace FramePro
 				FRAMEPRO_ASSERT((size_t)os_thread_mem_size >= sizeof(pthread_t));
 
 				ThreadContext* p_thread_context = New<ThreadContext>(p_allocator);
+				if (!p_thread_context)
+					return false;
+
 				p_thread_context->mp_ThreadMain = p_thread_main;
 				p_thread_context->mp_Context = p_context;
 				p_thread_context->mp_Allocator = p_allocator;
 
 				pthread_create(&GetOSThread(p_os_thread_mem), NULL, PlatformThreadMain, p_thread_context);
+
+				Platform::SetThreadAffinity(p_os_thread_mem, 0xffffffff);
+
+				return true;
 			}
 
 			//------------------------------------------------------------------------
@@ -9524,7 +10086,8 @@ namespace FramePro
 			void* p_context_switch_recorder,
 			ContextSwitchCallbackFunction p_callback,
 			void* p_context,
-			DynamicString& error)
+			DynamicString& error,
+			Allocator*)
 		{
 			return GenericPlatform::StartRecordingContextSitches(
 				p_context_switch_recorder,
@@ -9880,14 +10443,14 @@ namespace FramePro
 		}
 
 		//------------------------------------------------------------------------
-		void Platform::CreateThread(
+		bool Platform::CreateThread(
 			void* p_os_thread_mem,
 			int os_thread_mem_size,
 			ThreadMain p_thread_main,
 			void* p_context,
 			Allocator* p_allocator)
 		{
-			GenericPlatform::CreateThread(
+			return GenericPlatform::CreateThread(
 				p_os_thread_mem,
 				os_thread_mem_size,
 				p_thread_main,
@@ -9948,9 +10511,9 @@ namespace FramePro
 	}
 
 //------------------------------------------------------------------------
-//                         FRAMEPRO_PLATFORM_HOLOLENS
+//                         FRAMEPRO_PLATFORM_UWP
 //------------------------------------------------------------------------
-#elif FRAMEPRO_PLATFORM_HOLOLENS
+#elif FRAMEPRO_PLATFORM_UWP
 
 	//------------------------------------------------------------------------
 	// need to include windows.h in header for QueryPerformanceCounter()
@@ -9998,7 +10561,7 @@ namespace FramePro
 		//------------------------------------------------------------------------
 		Platform::Enum Platform::GetPlatformEnum()
 		{
-			return Platform::Windows_HoloLens;
+			return Platform::Windows_UWP;
 		}
 
 		//------------------------------------------------------------------------
@@ -10017,7 +10580,8 @@ namespace FramePro
 			void*,
 			ContextSwitchCallbackFunction,
 			void*,
-			DynamicString&)
+			DynamicString&,
+			Allocator*)
 		{
 			FramePro::DebugWrite("FramePro Warning: Failed to start recording context switches. Context switches may not be supported for this platform\n");
 			return false;
@@ -10313,14 +10877,14 @@ namespace FramePro
 		}
 
 		//------------------------------------------------------------------------
-		void Platform::CreateThread(
+		bool Platform::CreateThread(
 			void* p_os_thread_mem,
 			int os_thread_mem_size,
 			ThreadMain p_thread_main,
 			void* p_context,
 			Allocator* p_allocator)
 		{
-			GenericPlatform::CreateThread(
+			return GenericPlatform::CreateThread(
 				p_os_thread_mem,
 				os_thread_mem_size,
 				p_thread_main,
@@ -10717,14 +11281,14 @@ namespace FramePro
 		}
 
 		//------------------------------------------------------------------------
-		void Platform::CreateThread(
+		bool Platform::CreateThread(
 			void* p_os_thread_mem,
 			int os_thread_mem_size,
 			ThreadMain p_thread_main,
 			void* p_context,
 			Allocator* p_allocator)
 		{
-			GenericPlatform::CreateThread(
+			return GenericPlatform::CreateThread(
 				p_os_thread_mem,
 				os_thread_mem_size,
 				p_thread_main,
@@ -10799,7 +11363,8 @@ namespace FramePro
 			void* p_context_switch_recorder,
 			ContextSwitchCallbackFunction p_callback,
 			void* p_context,
-			DynamicString& error)
+			DynamicString& error,
+			Allocator*)
 		{
 			return GenericPlatform::StartRecordingContextSitches(
 				p_context_switch_recorder,
@@ -10831,7 +11396,8 @@ namespace FramePro
 	#include <sys/types.h>
 	#include <unistd.h>
 	#include <execinfo.h>
-       	
+	#include <sys/syscall.h>
+
 	//------------------------------------------------------------------------
 	namespace FramePro
 	{
@@ -11015,7 +11581,7 @@ namespace FramePro
 		//------------------------------------------------------------------------
 		int Platform::GetCurrentThreadId()
 		{
-			return gettid();
+			return syscall(SYS_gettid);
 		}
 
 		//------------------------------------------------------------------------
@@ -11127,14 +11693,14 @@ namespace FramePro
 		}
 
 		//------------------------------------------------------------------------
-		void Platform::CreateThread(
+		bool Platform::CreateThread(
 			void* p_os_thread_mem,
 			int os_thread_mem_size,
 			ThreadMain p_thread_main,
 			void* p_context,
 			Allocator* p_allocator)
 		{
-			GenericPlatform::CreateThread(
+			return GenericPlatform::CreateThread(
 				p_os_thread_mem,
 				os_thread_mem_size,
 				p_thread_main,
@@ -11209,7 +11775,8 @@ namespace FramePro
 			void* p_context_switch_recorder,
 			ContextSwitchCallbackFunction p_callback,
 			void* p_context,
-			DynamicString& error)
+			DynamicString& error,
+			Allocator*)
 		{
 			return GenericPlatform::StartRecordingContextSitches(
 				p_context_switch_recorder,
diff --git a/Engine/Source/Runtime/Core/Private/FramePro/FrameProUE4.cpp b/Engine/Source/Runtime/Core/Private/FramePro/FrameProUE4.cpp
--- a/Engine/Source/Runtime/Core/Private/FramePro/FrameProUE4.cpp
+++ b/Engine/Source/Runtime/Core/Private/FramePro/FrameProUE4.cpp
@@ -22,8 +22,13 @@
 	Add FramePro.cpp to your project to allow FramePro to communicate with your application.
 */
 //------------------------------------------------------------------------
+#ifdef __UNREAL__
 #include "FramePro/FrameProUE4.h"
 #include "FramePro/FramePro.h"
+#else
+#include "FrameProUE4.h"
+#include "FramePro.h"
+#endif
 
 //------------------------------------------------------------------------
 //                         FRAMEPRO_PLATFORM_UE4
@@ -126,7 +131,7 @@
 		//------------------------------------------------------------------------
 		int64 Platform::GetTimerFrequency()
 		{
-			return (int64)(1.0 / FPlatformTime::GetSecondsPerCycle());
+			return (int64)(1.0 / FPlatformTime::GetSecondsPerCycle64());
 		}
 
 		//------------------------------------------------------------------------
@@ -146,6 +151,8 @@
 		{
 			#if defined(FRAMEPRO_UNREAL_PLATFORM) //@EPIC: allow external definition
 				return FRAMEPRO_UNREAL_PLATFORM;  //@EPIC end
+			#elif PLATFORM_XBOXONE
+				return Platform::XBoxOne;
 			#elif PLATFORM_WINDOWS
 				return Platform::Windows;
 			#elif PLATFORM_LINUX
@@ -169,16 +176,34 @@
 			#endif
 		}
 
+		//------------------------------------------------------------------------
+		namespace PS4ContextSwitches
+		{
+			void* CreateContextSwitchRecorder(Allocator* p_allocator);
+			void DestroyContextSwitchRecorder(void* p_obj, Allocator* p_allocator);
+			bool StartRecordingContextSitches(void* p_obj, Platform::ContextSwitchCallbackFunction callback_function, void* p_context, DynamicString& error, Allocator* p_allocator);
+			void StopRecordingContextSitches(void* p_obj);
+			void FlushContextSwitches(void*);
+		}
+
 		//------------------------------------------------------------------------
 		void* Platform::CreateContextSwitchRecorder(Allocator* p_allocator)
 		{
-			return GenericPlatform::CreateContextSwitchRecorder(p_allocator);
+			#if PLATFORM_PS4 && FRAMEPRO_PS4_CONTEXT_SWITCH_RECORDING
+				return PS4ContextSwitches::CreateContextSwitchRecorder(p_allocator);
+			#else
+				return GenericPlatform::CreateContextSwitchRecorder(p_allocator);
+			#endif
 		}
 
 		//------------------------------------------------------------------------
 		void Platform::DestroyContextSwitchRecorder(void* p_context_switch_recorder, Allocator* p_allocator)
 		{
-			GenericPlatform::DestroyContextSwitchRecorder(p_context_switch_recorder, p_allocator);
+			#if PLATFORM_PS4 && FRAMEPRO_PS4_CONTEXT_SWITCH_RECORDING
+				PS4ContextSwitches::DestroyContextSwitchRecorder(p_context_switch_recorder, p_allocator);
+			#else
+				GenericPlatform::DestroyContextSwitchRecorder(p_context_switch_recorder, p_allocator);
+			#endif
 		}
 
 		//------------------------------------------------------------------------
@@ -186,29 +211,49 @@
 			void* p_context_switch_recorder,
 			ContextSwitchCallbackFunction p_callback,
 			void* p_context,
-			DynamicString& error)
-		{
-			return GenericPlatform::StartRecordingContextSitches(
-				p_context_switch_recorder,
-				p_callback,
-				p_context,
-				error);
+			DynamicString& error,
+			Allocator* p_allocator)
+		{
+			#if PLATFORM_PS4 && FRAMEPRO_PS4_CONTEXT_SWITCH_RECORDING
+				return PS4ContextSwitches::StartRecordingContextSitches(
+					p_context_switch_recorder,
+					p_callback,
+					p_context,
+					error,
+					p_allocator);
+			#else
+				FRAMEPRO_UNREFERENCED(p_allocator);
+
+				return GenericPlatform::StartRecordingContextSitches(
+					p_context_switch_recorder,
+					p_callback,
+					p_context,
+					error);
+			#endif
 		}
 
 		//------------------------------------------------------------------------
 		void Platform::StopRecordingContextSitches(void* p_context_switch_recorder)
 		{
-			GenericPlatform::StopRecordingContextSitches(p_context_switch_recorder);
+			#if PLATFORM_PS4 && FRAMEPRO_PS4_CONTEXT_SWITCH_RECORDING
+				PS4ContextSwitches::StopRecordingContextSitches(p_context_switch_recorder);
+			#else
+				GenericPlatform::StopRecordingContextSitches(p_context_switch_recorder);
+			#endif
 		}
 
 		//------------------------------------------------------------------------
 		void Platform::FlushContextSwitches(void* p_context_switch_recorder)
 		{
-			GenericPlatform::FlushContextSwitches(p_context_switch_recorder);
+			#if PLATFORM_PS4 && FRAMEPRO_PS4_CONTEXT_SWITCH_RECORDING
+				PS4ContextSwitches::FlushContextSwitches(p_context_switch_recorder);
+			#else
+				GenericPlatform::FlushContextSwitches(p_context_switch_recorder);
+			#endif
 		}
 
 		//------------------------------------------------------------------------
-		#if PLATFORM_PS4
+		#if PLATFORM_PS4 && FRAMEPRO_ENUMERATE_ALL_MODULES
 			namespace EnumModulesPS4 { void EnumerateModules(Array<ModulePacket*>& module_packets, Allocator* p_allocator); }
 		#endif
 
@@ -236,6 +281,9 @@
 					EnumModulesSwitch::EnumerateModules(module_packets, p_allocator);
 				#else
 					ModulePacket* p_module_packet = (ModulePacket*)p_allocator->Alloc(sizeof(ModulePacket));
+					if (!p_module_packet)
+						return;
+
 					memset(p_module_packet, 0, sizeof(ModulePacket));
 
 					p_module_packet->m_PacketType = PacketType::ModulePacket;
@@ -432,7 +480,7 @@
 			{
 				return false;
 			}
-			
+
 			FString ProcessName = FPaths::GetCleanFilename(ProcessNameOrPath);
 			const char* p_process_name = TCHAR_TO_ANSI(*ProcessName);
 //@EPIC END
@@ -440,7 +488,7 @@
 			size_t max_length = max_name_length - 1;
 
 			int copy_length = length < max_length ? length : max_length;
-			FCStringAnsi::Strncpy( p_name, p_process_name, copy_length);
+			FCStringAnsi::Strncpy(p_name, p_process_name, copy_length);
 
 			return true;
 		}
@@ -461,7 +509,7 @@
 				p_event = FPlatformProcess::GetSynchEventFromPool(!auto_reset);
 			}
 			else
-			{				
+			{
 #if PLATFORM_WINDOWS
 				p_event = new FEventWin();
 				p_event->Create(!auto_reset);
@@ -605,13 +653,13 @@
 				else
 				{
 #if PLATFORM_WINDOWS
-					CreateThread( NULL, 0, (LPTHREAD_START_ROUTINE)mp_ThreadMain, mp_Context, 0, NULL );
+					CreateThread(NULL, 0, (LPTHREAD_START_ROUTINE)mp_ThreadMain, mp_Context, 0, NULL);
 #elif PLATFORM_USE_PTHREADS
-					typedef void *(*PthreadEntryPoint)(void *arg);
+					typedef void* (*PthreadEntryPoint)(void* arg);
 					pthread_t thread_id;
-					pthread_create( &thread_id, nullptr, (PthreadEntryPoint)mp_ThreadMain, mp_Context );
+					pthread_create(&thread_id, nullptr, (PthreadEntryPoint)mp_ThreadMain, mp_Context);
 #else
-					checkf(false,TEXT("unsupported platform for -nothreading"));
+					checkf(false, TEXT("unsupported platform for -nothreading"));
 #endif
 				}
 //@EPIC END
@@ -641,7 +689,7 @@
 		}
 
 		//------------------------------------------------------------------------
-		void Platform::CreateThread(
+		bool Platform::CreateThread(
 			void* p_os_thread_mem,
 			int os_thread_mem_size,
 			ThreadMain p_thread_main,
@@ -651,6 +699,8 @@
 			FRAMEPRO_ASSERT(os_thread_mem_size >= sizeof(UE4Thread*));
 
 			GetOSThread(p_os_thread_mem) = new UE4Thread(p_thread_main, p_context);
+
+			return true;
 		}
 
 		//------------------------------------------------------------------------
diff --git a/Engine/Source/Runtime/Core/Public/FramePro/FramePro.h b/Engine/Source/Runtime/Core/Public/FramePro/FramePro.h
--- a/Engine/Source/Runtime/Core/Public/FramePro/FramePro.h
+++ b/Engine/Source/Runtime/Core/Public/FramePro/FramePro.h
@@ -1,35 +1,27 @@
 /*
-	This software is provided 'as-is', without any express or implied warranty.
-	In no event will the author(s) be held liable for any damages arising from
-	the use of this software.
-
-	Permission is granted to anyone to use this software for any purpose, including
-	commercial applications, and to alter it and redistribute it freely, subject to
-	the following restrictions:
-
-	1. The origin of this software must not be misrepresented; you must not
-	claim that you wrote the original software. If you use this software
-	in a product, an acknowledgment in the product documentation would be
-	appreciated but is not required.
-	2. Altered source versions must be plainly marked as such, and must not be
-	misrepresented as being the original software.
-	3. This notice may not be removed or altered from any source distribution.
-
-	Author: [REDACTED_PERSON]
-	www.puredevsoftware.com
-	[REDACTED_EMAIL]
-
-	Add FramePro.cpp to your project to allow FramePro to communicate with your application.
+Copyright 2019 PureDev Software Limited
+
+Permission to use, copy, modify, and/or distribute this software for any
+purpose with or without fee is hereby granted, provided that the above
+copyright notice and this permission notice appear in all copies.
+
+THE SOFTWARE IS PROVIDED "AS IS" AND THE AUTHOR DISCLAIMS ALL WARRANTIES
+WITH REGARD TO THIS SOFTWARE INCLUDING ALL IMPLIED WARRANTIES OF
+MERCHANTABILITY AND FITNESS. IN NO EVENT SHALL THE AUTHOR BE LIABLE FOR
+ANY SPECIAL, DIRECT, INDIRECT, OR CONSEQUENTIAL DAMAGES OR ANY DAMAGES
+WHATSOEVER RESULTING FROM LOSS OF USE, DATA OR PROFITS, WHETHER IN AN
+ACTION OF CONTRACT, NEGLIGENCE OR OTHER TORTIOUS ACTION, ARISING OUT OF
+OR IN CONNECTION WITH THE USE OR PERFORMANCE OF THIS SOFTWARE.
 */
 
 //------------------------------------------------------------------------
 //
-// FramePro.hpp
+// FramePro.h
 //
 //------------------------------------------------------------------------
 /*
 	FramePro
-	Version:	1.5.20.0
+	Version:	1.10.10.0
 */
 //------------------------------------------------------------------------
 #ifndef FRAMEPRO_H_INCLUDED
@@ -38,11 +30,13 @@
 //------------------------------------------------------------------------
 #ifndef FRAMEPRO_ENABLED
 	// *** enable/disable FramePro here
+	//@EPIC BEGIN: Only enable for engine / game, not external tools
 	#ifdef __UNREAL__
 		#define FRAMEPRO_ENABLED 0
 	#else
 		#define FRAMEPRO_ENABLED 1
 	#endif
+	//@EPIC END
 #endif
 
 //------------------------------------------------------------------------
@@ -53,9 +47,10 @@
 	#define FRAMEPRO_FRAME_START()															FramePro::FrameStart()
 	#define FRAMEPRO_SHUTDOWN()																FramePro::Shutdown()
 	#define FRAMEPRO_SET_PORT(port)															FramePro::SetPort(port)
-	#define FRAMEPRO_SET_SESSION_INFO(name, build_id)										FramePro::SendSessionInfo(name, build_id)
+	#define FRAMEPRO_SESSION_INFO(name, value)												FramePro::SendSessionInfo(name, value)
 	#define FRAMEPRO_SET_ALLOCATOR(p_allocator)												FramePro::SetAllocator(p_allocator)
 	#define FRAMEPRO_SET_THREAD_NAME(name)													FramePro::SetThreadName(name)
+	#define FRAMEPRO_SET_THREAD_ID_NAME(thread_id, name)									FramePro::SetThreadName(thread_id, name)
 	#define FRAMEPRO_THREAD_ORDER(thread_name)												FramePro::SetThreadOrder(FramePro::RegisterString(thread_name))
 	#define FRAMEPRO_REGISTER_STRING(str)													FramePro::RegisterString(str)
 	#define FRAMEPRO_START_RECORDING(filename, context_switches, callstacks, max_file_size)	FramePro::StartRecording(filename, context_switches, callstacks, max_file_size)
@@ -73,42 +68,43 @@
 	#define FRAMEPRO_SET_CONDITIONAL_SCOPE_MIN_TIME(microseconds)							FramePro::SetConditionalScopeMinTimeInMicroseconds(microseconds)
 
 	// scope macros
-	#define FRAMEPRO_SCOPE()																FramePro::TimerScope FRAMEPRO_UNIQUE(timer_scope)(FRAMEPRO_FUNCTION_NAME "|" FRAMEPRO_SOURCE_STRING)
+	#define FRAMEPRO_SCOPE()																FramePro::TimerScope FRAMEPRO_UNIQUE(timer_scope)(FRAMEPRO_FUNCTION_NAME_LITERAL "|" FRAMEPRO_SOURCE_STRING)
 	#define FRAMEPRO_NAMED_SCOPE(name)														FramePro::TimerScope FRAMEPRO_UNIQUE(timer_scope)(name "|" FRAMEPRO_SOURCE_STRING)
 	#define FRAMEPRO_NAMED_SCOPE_W(name)													FramePro::TimerScopeW FRAMEPRO_UNIQUE(timer_scope)(name L"|" FRAMEPRO_SOURCE_STRING_W)
 	#define FRAMEPRO_ID_SCOPE(name_id)														FramePro::IdTimerScope FRAMEPRO_UNIQUE(timer_scope)(name_id, FRAMEPRO_SOURCE_STRING)
 	#define FRAMEPRO_DYNAMIC_SCOPE(dynamic_string)											FramePro::IdTimerScope FRAMEPRO_UNIQUE(timer_scope)(FramePro::IsConnected() ? FramePro::RegisterString(dynamic_string) : -1, FRAMEPRO_SOURCE_STRING)
-	#define FRAMEPRO_CONDITIONAL_SCOPE()													FramePro::ConditionalTimerScope FRAMEPRO_UNIQUE(timer_scope)(FRAMEPRO_FUNCTION_NAME "|" FRAMEPRO_SOURCE_STRING)
+	#define FRAMEPRO_CONDITIONAL_SCOPE()													FramePro::ConditionalTimerScope FRAMEPRO_UNIQUE(timer_scope)(FRAMEPRO_FUNCTION_NAME_LITERAL "|" FRAMEPRO_SOURCE_STRING)
 	#define FRAMEPRO_CONDITIONAL_ID_SCOPE(name)												FramePro::ConditionalTimerScopeId FRAMEPRO_UNIQUE(timer_scope)(name, FRAMEPRO_SOURCE_STRING)
 	#define FRAMEPRO_CONDITIONAL_NAMED_SCOPE(name)											FramePro::ConditionalTimerScope FRAMEPRO_UNIQUE(timer_scope)(name "|" FRAMEPRO_SOURCE_STRING)
 	#define FRAMEPRO_CONDITIONAL_NAMED_SCOPE_W(name)										FramePro::ConditionalTimerScopeW FRAMEPRO_UNIQUE(timer_scope)(name L"|" FRAMEPRO_SOURCE_STRING_W)
-	#define FRAMEPRO_CONDITIONAL_BOOL_SCOPE(b)												FramePro::ConditionalBoolTimerScope FRAMEPRO_UNIQUE(timer_scope)(FRAMEPRO_FUNCTION_NAME "|" FRAMEPRO_SOURCE_STRING, (b))
+	#define FRAMEPRO_CONDITIONAL_BOOL_SCOPE(b)												FramePro::ConditionalBoolTimerScope FRAMEPRO_UNIQUE(timer_scope)(FRAMEPRO_FUNCTION_NAME_LITERAL "|" FRAMEPRO_SOURCE_STRING, (b))
 	#define FRAMEPRO_CONDITIONAL_BOOL_ID_SCOPE(name, b)										FramePro::ConditionalBoolTimerScopeId FRAMEPRO_UNIQUE(timer_scope)(name, FRAMEPRO_SOURCE_STRING, (b))
 	#define FRAMEPRO_CONDITIONAL_BOOL_NAMED_SCOPE(name, b)									FramePro::ConditionalBoolTimerScope FRAMEPRO_UNIQUE(timer_scope)(name "|" FRAMEPRO_SOURCE_STRING, (b))
 	#define FRAMEPRO_CONDITIONAL_BOOL_NAMED_SCOPE_W(name, b)								FramePro::ConditionalBoolTimerScopeW FRAMEPRO_UNIQUE(timer_scope)(name L"|" FRAMEPRO_SOURCE_STRING_W, (b))
-	#define FRAMEPRO_START_NAMED_SCOPE(name)												FramePro::int64 framepro_start_##name=0; FRAMEPRO_GET_CLOCK_COUNT(framepro_start_##name);
-	#define FRAMEPRO_STOP_NAMED_SCOPE(name)													MULTI_STATEMENT( if(FramePro::IsConnected()) { FramePro::int64 framepro_end_##name; FRAMEPRO_GET_CLOCK_COUNT(framepro_end_##name); FramePro::AddTimeSpan(#name "|" FRAMEPRO_SOURCE_STRING, framepro_start_##name, framepro_end_##name); } )
-	#define FRAMEPRO_CONDITIONAL_START_SCOPE()												FramePro::int64 framepro_start=0; MULTI_STATEMENT( if(FramePro::IsConnected()) { FRAMEPRO_GET_CLOCK_COUNT(framepro_start); } )
-	#define FRAMEPRO_CONDITIONAL_STOP_NAMED_SCOPE(name)										MULTI_STATEMENT( if(FramePro::IsConnected()) { FramePro::int64 framepro_end; FRAMEPRO_GET_CLOCK_COUNT(framepro_end); if(framepro_end - framepro_start > FramePro::GetConditionalScopeMinTime()) FramePro::AddTimeSpan(name "|" FRAMEPRO_SOURCE_STRING, framepro_start, framepro_end); } )
-	#define FRAMEPRO_CONDITIONAL_STOP_DYNAMIC_SCOPE(dynamic_string)							MULTI_STATEMENT( if(FramePro::IsConnected()) { FramePro::int64 framepro_end; FRAMEPRO_GET_CLOCK_COUNT(framepro_end); if(framepro_end - framepro_start > FramePro::GetConditionalScopeMinTime()) FramePro::AddTimeSpan(FramePro::RegisterString(dynamic_string), FRAMEPRO_SOURCE_STRING, framepro_start, framepro_end); } )
+	#define FRAMEPRO_START_NAMED_SCOPE(name)												FramePro::int64 framepro_start_##name=0; MULTI_STATEMENT( if(FramePro::IsConnected()) { FRAMEPRO_GET_CLOCK_COUNT(framepro_start_##name); FramePro::StartScope(); } )
+	#define FRAMEPRO_STOP_NAMED_SCOPE(name)													MULTI_STATEMENT( if(FramePro::IsConnected()) { FramePro::int64 framepro_end_##name; FRAMEPRO_GET_CLOCK_COUNT(framepro_end_##name); FramePro::AddTimeSpan(#name "|" FRAMEPRO_SOURCE_STRING, framepro_start_##name, framepro_end_##name); FramePro::StopScope(); } )
+	#define FRAMEPRO_CONDITIONAL_START_SCOPE()												FramePro::int64 framepro_start=0; MULTI_STATEMENT( if(FramePro::IsConnected()) { FRAMEPRO_GET_CLOCK_COUNT(framepro_start); FramePro::StartScope(); } )
+	#define FRAMEPRO_CONDITIONAL_STOP_NAMED_SCOPE(name)										MULTI_STATEMENT( if(FramePro::IsConnected()) { FramePro::int64 framepro_end; FRAMEPRO_GET_CLOCK_COUNT(framepro_end); if(framepro_end - framepro_start > FramePro::GetConditionalScopeMinTime()) FramePro::AddTimeSpan(name "|" FRAMEPRO_SOURCE_STRING, framepro_start, framepro_end); FramePro::StopScope();} )
+	#define FRAMEPRO_CONDITIONAL_STOP_DYNAMIC_SCOPE(dynamic_string)							MULTI_STATEMENT( if(FramePro::IsConnected()) { FramePro::int64 framepro_end; FRAMEPRO_GET_CLOCK_COUNT(framepro_end); if(framepro_end - framepro_start > FramePro::GetConditionalScopeMinTime()) FramePro::AddTimeSpan(FramePro::RegisterString(dynamic_string), FRAMEPRO_SOURCE_STRING, framepro_start, framepro_end); FramePro::StopScope(); } )
 	#define FRAMEPRO_CONDITIONAL_PARENT_SCOPE(name, callback, pre_duration, post_duration)	FramePro::ConditionalParentTimerScope FRAMEPRO_UNIQUE(timer_scope)(name, FRAMEPRO_SOURCE_STRING, callback, pre_duration, post_duration)
 	#define FRAMEPRO_SET_SCOPE_COLOUR(name, colour)											MULTI_STATEMENT( FramePro::SetScopeColour(FramePro::RegisterString(name), colour); )
+	#define FRAMEPRO_CALLSTACK()															FramePro::SendScopeCallstack()
 
 	// idle scope macros
-	#define FRAMEPRO_IDLE_SCOPE()															FramePro::TimerScope FRAMEPRO_UNIQUE(timer_scope)(FRAMEPRO_FUNCTION_NAME "|" FRAMEPRO_SOURCE_STRING_IDLE)
+	#define FRAMEPRO_IDLE_SCOPE()															FramePro::TimerScope FRAMEPRO_UNIQUE(timer_scope)(FRAMEPRO_FUNCTION_NAME_LITERAL "|" FRAMEPRO_SOURCE_STRING_IDLE)
 	#define FRAMEPRO_IDLE_NAMED_SCOPE(name)													FramePro::TimerScope FRAMEPRO_UNIQUE(timer_scope)(name "|" FRAMEPRO_SOURCE_STRING_IDLE)
 	#define FRAMEPRO_IDLE_NAMED_SCOPE_W(name)												FramePro::TimerScopeW FRAMEPRO_UNIQUE(timer_scope)(name L"|" FRAMEPRO_SOURCE_STRING_IDLE_W)
 	#define FRAMEPRO_IDLE_ID_SCOPE(name_id)													FramePro::IdTimerScope FRAMEPRO_UNIQUE(timer_scope)(name_id, FRAMEPRO_SOURCE_STRING_IDLE)
 	#define FRAMEPRO_IDLE_DYNAMIC_SCOPE(dynamic_string)										FramePro::IdTimerScope FRAMEPRO_UNIQUE(timer_scope)(FramePro::IsConnected() ? FramePro::RegisterString(dynamic_string) : -1, FRAMEPRO_SOURCE_STRING_IDLE)
-	#define FRAMEPRO_IDLE_CONDITIONAL_SCOPE()												FramePro::ConditionalTimerScope FRAMEPRO_UNIQUE(timer_scope)(FRAMEPRO_FUNCTION_NAME "|" FRAMEPRO_SOURCE_STRING_IDLE)
+	#define FRAMEPRO_IDLE_CONDITIONAL_SCOPE()												FramePro::ConditionalTimerScope FRAMEPRO_UNIQUE(timer_scope)(FRAMEPRO_FUNCTION_NAME_LITERAL "|" FRAMEPRO_SOURCE_STRING_IDLE)
 	#define FRAMEPRO_IDLE_CONDITIONAL_ID_SCOPE(name)										FramePro::ConditionalTimerScopeId FRAMEPRO_UNIQUE(timer_scope)(name, FRAMEPRO_SOURCE_STRING_IDLE)
 	#define FRAMEPRO_IDLE_CONDITIONAL_NAMED_SCOPE(name)										FramePro::ConditionalTimerScope FRAMEPRO_UNIQUE(timer_scope)(name "|" FRAMEPRO_SOURCE_STRING_IDLE)
 	#define FRAMEPRO_IDLE_CONDITIONAL_NAMED_SCOPE_W(name)									FramePro::ConditionalTimerScopeW FRAMEPRO_UNIQUE(timer_scope)(name L"|" FRAMEPRO_SOURCE_STRING_IDLE_W)
-	#define FRAMEPRO_IDLE_START_NAMED_SCOPE(name)											FramePro::int64 framepro_start_##name=0; FRAMEPRO_GET_CLOCK_COUNT(framepro_start_##name);
-	#define FRAMEPRO_IDLE_STOP_NAMED_SCOPE(name)											MULTI_STATEMENT( if(FramePro::IsConnected()) { FramePro::int64 framepro_end_##name; FRAMEPRO_GET_CLOCK_COUNT(framepro_end_##name); FramePro::AddTimeSpan(#name "|" FRAMEPRO_SOURCE_STRING_IDLE, framepro_start_##name, framepro_end_##name); } )
-	#define FRAMEPRO_IDLE_CONDITIONAL_START_SCOPE()											FramePro::int64 framepro_start=0; MULTI_STATEMENT( if(FramePro::IsConnected()) { FRAMEPRO_GET_CLOCK_COUNT(framepro_start); } )
-	#define FRAMEPRO_IDLE_CONDITIONAL_STOP_NAMED_SCOPE(name)								MULTI_STATEMENT( if(FramePro::IsConnected()) { FramePro::int64 framepro_end; FRAMEPRO_GET_CLOCK_COUNT(framepro_end); if(framepro_end - framepro_start > FramePro::GetConditionalScopeMinTime()) FramePro::AddTimeSpan(name "|" FRAMEPRO_SOURCE_STRING_IDLE, framepro_start, framepro_end); } )
-	#define FRAMEPRO_IDLE_CONDITIONAL_STOP_DYNAMIC_SCOPE(dynamic_string)					MULTI_STATEMENT( if(FramePro::IsConnected()) { FramePro::int64 framepro_end; FRAMEPRO_GET_CLOCK_COUNT(framepro_end); if(framepro_end - framepro_start > FramePro::GetConditionalScopeMinTime()) FramePro::AddTimeSpan(FramePro::RegisterString(dynamic_string), FRAMEPRO_SOURCE_STRING_IDLE, framepro_start, framepro_end); } )
+	#define FRAMEPRO_IDLE_START_NAMED_SCOPE(name)											FramePro::int64 framepro_start_##name=0; if(FramePro::IsConnected()) { FRAMEPRO_GET_CLOCK_COUNT(framepro_start_##name); FramePro::StartScope(); }
+	#define FRAMEPRO_IDLE_STOP_NAMED_SCOPE(name)											MULTI_STATEMENT( if(FramePro::IsConnected()) { FramePro::int64 framepro_end_##name; FRAMEPRO_GET_CLOCK_COUNT(framepro_end_##name); FramePro::AddTimeSpan(#name "|" FRAMEPRO_SOURCE_STRING_IDLE, framepro_start_##name, framepro_end_##name); FramePro::StopScope(); } )
+	#define FRAMEPRO_IDLE_CONDITIONAL_START_SCOPE()											FramePro::int64 framepro_start=0; MULTI_STATEMENT( if(FramePro::IsConnected()) { FRAMEPRO_GET_CLOCK_COUNT(framepro_start); FramePro::StartScope(); } )
+	#define FRAMEPRO_IDLE_CONDITIONAL_STOP_NAMED_SCOPE(name)								MULTI_STATEMENT( if(FramePro::IsConnected()) { FramePro::int64 framepro_end; FRAMEPRO_GET_CLOCK_COUNT(framepro_end); if(framepro_end - framepro_start > FramePro::GetConditionalScopeMinTime()) FramePro::AddTimeSpan(name "|" FRAMEPRO_SOURCE_STRING_IDLE, framepro_start, framepro_end); FramePro::StopScope(); } )
+	#define FRAMEPRO_IDLE_CONDITIONAL_STOP_DYNAMIC_SCOPE(dynamic_string)					MULTI_STATEMENT( if(FramePro::IsConnected()) { FramePro::int64 framepro_end; FRAMEPRO_GET_CLOCK_COUNT(framepro_end); if(framepro_end - framepro_start > FramePro::GetConditionalScopeMinTime()) FramePro::AddTimeSpan(FramePro::RegisterString(dynamic_string), FRAMEPRO_SOURCE_STRING_IDLE, framepro_start, framepro_end); FramePro::StopScope(); } )
 
 	// custom stat macros
 	#define FRAMEPRO_CUSTOM_STAT(name, value, graph, unit, colour)							MULTI_STATEMENT( if(FramePro::IsConnected()) FramePro::AddCustomStat(name, value, graph, unit, colour); )
@@ -121,9 +117,10 @@
 	// high-res timers
 	#define FRAMEPRO_HIRES_SCOPE(name)														FramePro::HiResTimerScope FRAMEPRO_UNIQUE(hires_scope)(name)
 
-	// global high-res timers
-	#define FRAMEPRO_DECL_GLOBAL_HIRES_TIMER(name)											FramePro::GlobalHiResTimer g_FrameProHiResTimer##name(#name)
-	#define FRAMEPRO_GLOBAL_HIRES_SCOPE(name)												FramePro::GlobalHiResTimerScope FRAMEPRO_UNIQUE(timer_scope)(g_FrameProHiResTimer##name)
+	// custom stat timers
+	#define FRAMEPRO_DECL_CUSTOM_STAT_TIMER(name)											FramePro::CustomStatTimer g_FrameProCustomStatTimer##name(#name)
+	#define FRAMEPRO_EXTERN_CUSTOM_STAT_TIMER(name)											extern FramePro::CustomStatTimer g_FrameProCustomStatTimer##name
+	#define FRAMEPRO_CUSTOM_STAT_TIMER_SCOPE(name)											FramePro::CustomStatTimerScope FRAMEPRO_UNIQUE(timer_scope)(g_FrameProCustomStatTimer##name)
 
 	// events
 	#define FRAMEPRO_EVENT(name, colour)													FramePro::AddEvent(name, colour)
@@ -136,9 +133,10 @@
 	#define FRAMEPRO_FRAME_START()															((void)0)
 	#define FRAMEPRO_SHUTDOWN()																((void)0)
 	#define FRAMEPRO_SET_PORT(port)															((void)0)
-	#define FRAMEPRO_SET_SESSION_INFO(name, id)												((void)0)
+	#define FRAMEPRO_SESSION_INFO(name, id)													((void)0)
 	#define FRAMEPRO_SET_ALLOCATOR(p_allocator)												((void)0)
 	#define FRAMEPRO_SET_THREAD_NAME(name)													((void)0)
+	#define FRAMEPRO_SET_THREAD_ID_NAME(thread_id, name)									((void)0)
 	#define FRAMEPRO_THREAD_ORDER(thread_name)												((void)0)
 	#define FRAMEPRO_REGISTER_STRING(str)													0
 	#define FRAMEPRO_START_RECORDING(filename, context_switches, callstacks, max_file_size)	((void)0)
@@ -174,6 +172,7 @@
 	#define FRAMEPRO_CONDITIONAL_STOP_DYNAMIC_SCOPE(dynamic_string)							((void)0)
 	#define FRAMEPRO_CONDITIONAL_PARENT_SCOPE(name, callback, pre_duration, post_duration)	((void)0)
 	#define FRAMEPRO_SET_SCOPE_COLOUR(name, colour)											((void)0)
+	#define FRAMEPRO_CALLSTACK()															((void)0)
 
 	#define FRAMEPRO_IDLE_SCOPE()															((void)0)
 	#define FRAMEPRO_IDLE_NAMED_SCOPE(name)													((void)0)
@@ -199,8 +198,9 @@
 
 	#define FRAMEPRO_HIRES_SCOPE(name)														((void)0)
 
-	#define FRAMEPRO_DECL_GLOBAL_HIRES_TIMER(name)											
-	#define FRAMEPRO_GLOBAL_HIRES_SCOPE(name)												((void)0)
+	#define FRAMEPRO_DECL_CUSTOM_STAT_TIMER(name)											
+	#define FRAMEPRO_EXTERN_CUSTOM_STAT_TIMER(name)
+	#define FRAMEPRO_CUSTOM_STAT_TIMER_SCOPE(name)											((void)0)
 
 	#define FRAMEPRO_EVENT(name, colour)													((void)0)
 
@@ -257,7 +257,7 @@ namespace FramePro
 #define FRAMEPRO_ENABLE_CONTEXT_SWITCH_TRACKING 1
 
 //------------------------------------------------------------------------
-#define FRAMEPRO_ENABLE_CALLSTACKS 1
+#define FRAMEPRO_ENABLE_CALLSTACKS 0
 
 //------------------------------------------------------------------------
 // If this is disabled and two dynamic strings happen to hash to the same value you may see
@@ -269,6 +269,13 @@ namespace FramePro
 //------------------------------------------------------------------------
 #define FRAMEPRO_STACK_TRACE_SIZE 128
 
+//------------------------------------------------------------------------
+// FramePro will assert if a hires scope doesn't have a parent scope. Enabling
+// this define make the code detect that case and throw the hires scope away
+// instead of asserting.
+// This adds extra overhead to all scopes (function calls to increment/decrement
+// on each start/end scope) and so if off by default
+#define FRAMEPRO_ALLOW_UNPARENTED_HIRES_SCOPES 0
 
 //------------------------------------------------------------------------
 // DOCUMENTATION
@@ -344,11 +351,11 @@ namespace FramePro
 // Cleans up all FramePro resources. No thread can call into FramePro after this is called.
 
 
-// FRAMEPRO_SET_SESSION_INFO(name, build_id)
-// -----------------------------------------
-// name: string - name of the session
-// build_id: string - build id of the build being profiled
-// The session name is typically the game name. The build can be anything. These values are shown in the session info panel in FramePro.
+// FRAMEPRO_SESSION_INFO(name, value)
+// ----------------------------------
+// name: string
+// value: string
+// The name/value pair will be displayed in the FramePro SessionInfo panel. Use it for things such as the build name, version, etc.
 
 
 // FRAMEPRO_SET_ALLOCATOR(p_allocator)
@@ -366,6 +373,14 @@ namespace FramePro
 // FramePro will pick a colour based on the thread name which will always stay the same.
 
 
+// FRAMEPRO_SET_THREAD_ID_NAME(thread_id, name)
+// -----------------------------------------
+// thread_id: id of the thread to set
+// name: string (literal or dynamic)
+// Tell FramePro the name of the specified thread.
+// FramePro will pick a colour based on the thread name which will always stay the same.
+
+
 // FRAMEPRO_THREAD_ORDER(thread_name)
 // ----------------------------------
 // thread_name: string (literal or dynamic)
@@ -455,7 +470,7 @@ namespace FramePro
 
 // FRAMEPRO_THREAD_SCOPE(thread_name)
 // ----------------------------------
-// thraed_name: string (literal or dynamic)
+// thread_name: string (literal or dynamic)
 // Calls FRAMEPRO_SET_THREAD_NAME at the start of the scope, and FRAMEPRO_CLEANUP_THREAD at the end of the scope.
 // Ensure that this is called outside of any other scopes.
 
@@ -659,6 +674,11 @@ namespace FramePro
 // This is a slow macro, so please only call once, not every frame.
 
 
+// FRAMEPRO_CALLSTACK()
+// --------------------------
+// capture the callstack and send for the current scope
+
+
 //--------------------------------------------------------------------------------------------------------
 // idle scope macros
 
@@ -808,33 +828,38 @@ namespace FramePro
 
 
 //--------------------------------------------------------------------------------------------------------
-// global hires timers
+// custom stat timers
 
-// Global high resolution timers are used when you want to time scopes that happen many thousands of times a frame
-// and don't want the overhead of recording a timespan for each scope. The scope simple adds to a global
-// value which stores total duration and count, and then sends that value once a frame. The timer is reset
-// at the start of each frame.
+// Custom stat timers are very low overhead high resolution timers. A custom stat timer is simply
+// a global timer value that is added to on every CUSTOM_STAT_TIMER_SCOPE. The timer is sent and
+// reset at the end of a frame. Custom stat timers are useful for timing events that happen
+// thousands of timer a frame.
 
 // Usage:
-//		// this defines a global timer with name "my_timer"
-//		FRAMEPRO_DECL_GLOBAL_HIRES_TIMER(my_timer)
+//		// this defines a custom stat timer with name "my_timer"
+//		FRAMEPRO_DECL_CUSTOM_STAT_TIMER(my_timer)
 //
 //		void MyFunction()
 //		{
-//			FRAMEPRO_HIRES_SCOPE(my_timer)
+//			FRAMEPRO_CUSTOM_STAT_TIMER_SCOPE(my_timer)
 //			...
 //		}
 
-// FRAMEPRO_DECL_GLOBAL_HIRES_TIMER(name)
-// ---------------------------------------------
-// name: unquoted string		the name of the global timer object
-// Declares a high resolution timer with the specified name.
+// FRAMEPRO_DECL_CUSTOM_STAT_TIMER(name)
+// -------------------------------------
+// name: unquoted string		the name of the custom stat timer object
+// Declares a custom stat timer with the specified name.
 
-// FRAMEPRO_GLOBAL_HIRES_SCOPE(name)
-// ---------------------------------
-// name: unquoted string		the name of the global timer declared with FRAMEPRO_GLOBAL_HIRES_SCOPE
-// Times the current scope and adds the value to the specified high-res timer. The high-res timer must
-// have already been declared with FRAMEPRO_GLOBAL_HIRES_SCOPE.
+// FRAMEPRO_EXTERN_CUSTOM_STAT_TIMER(name)
+// ---------------------------------------
+// name: unquoted string		the name of the custom stat timer object
+// Declares an extern to a custom stat timer with the specified name.
+
+// FRAMEPRO_CUSTOM_STAT_TIMER_SCOPE(name)
+// --------------------------------------
+// name: unquoted string		the name of the custom stat timer
+// Times the current scope and adds the value to the specified timer. The timer must have already been declared
+// with FRAMEPRO_DECL_CUSTOM_STAT_TIMER.
 
 
 //--------------------------------------------------------------------------------------------------------
@@ -860,7 +885,6 @@ namespace FramePro
 // Measure how long a thread waits on an event. For Win32 Events put this around the WaitForSingleObject
 // call. For Win32 critical sections put this around the EnterCriticalSection call. In order to not flood
 // the channel with events, the wait event is only sent if it is longer than FRAMEPRO_WAIT_EVENT_MIN_TIME.
-// FramePro will also only show a wait event if it was triggered with FRAMEPRO_STOP_WAIT_EVENT.
 
 // FRAMEPRO_TRIGGER_WAIT_EVENT(event_id)
 // ----------------------------------
@@ -891,10 +915,12 @@ namespace FramePro
 
 //------------------------------------------------------------------------
 // enable debug here to get FramePro asserts
-#if defined(DEBUG) || defined(_DEBUG)
-	#define FRAMEPRO_DEBUG 1
-#else
-	#define FRAMEPRO_DEBUG 0
+#ifndef FRAMEPRO_DEBUG
+	#if defined(DEBUG) || defined(_DEBUG)
+		#define FRAMEPRO_DEBUG 1
+	#else
+		#define FRAMEPRO_DEBUG 0
+	#endif
 #endif
 
 //------------------------------------------------------------------------
@@ -914,19 +940,15 @@ namespace FramePro
 #if defined(__UNREAL__)
 	#define FRAMEPRO_PLATFORM_UE4 1
 	#define FRAMEPRO_PLATFORM_XBOXONE 0
-	#define FRAMEPRO_PLATFORM_HOLOLENS 0
-#if defined(_WIN32) || defined(_WIN64) || defined(WIN32) || defined(WIN64) || defined(__WIN32__) || defined(__WINDOWS__)
-	#define FRAMEPRO_PLATFORM_WIN 1
-#else
+	#define FRAMEPRO_PLATFORM_UWP 0
 	#define FRAMEPRO_PLATFORM_WIN 0
-#endif
 	#define FRAMEPRO_PLATFORM_LINUX 0
 	#define FRAMEPRO_PLATFORM_PS4 0
 	#define FRAMEPRO_PLATFORM_ANDROID 0
 #elif defined(_XBOX_ONE) || defined(_DURANGO)
 	#define FRAMEPRO_PLATFORM_UE4 0
 	#define FRAMEPRO_PLATFORM_XBOXONE 1
-	#define FRAMEPRO_PLATFORM_HOLOLENS 0
+	#define FRAMEPRO_PLATFORM_UWP 0
 	#define FRAMEPRO_PLATFORM_WIN 0
 	#define FRAMEPRO_PLATFORM_LINUX 0
 	#define FRAMEPRO_PLATFORM_PS4 0
@@ -934,7 +956,7 @@ namespace FramePro
 #elif defined(WINAPI_FAMILY) && (WINAPI_FAMILY == WINAPI_FAMILY_APP)
 	#define FRAMEPRO_PLATFORM_UE4 0
 	#define FRAMEPRO_PLATFORM_XBOXONE 0
-	#define FRAMEPRO_PLATFORM_HOLOLENS 1
+	#define FRAMEPRO_PLATFORM_UWP 1
 	#define FRAMEPRO_PLATFORM_WIN 0
 	#define FRAMEPRO_PLATFORM_LINUX 0
 	#define FRAMEPRO_PLATFORM_PS4 0
@@ -942,7 +964,7 @@ namespace FramePro
 #elif defined(_WIN32) || defined(_WIN64) || defined(WIN32) || defined(WIN64) || defined(__WIN32__) || defined(__WINDOWS__)
 	#define FRAMEPRO_PLATFORM_UE4 0
 	#define FRAMEPRO_PLATFORM_XBOXONE 0
-	#define FRAMEPRO_PLATFORM_HOLOLENS 0
+	#define FRAMEPRO_PLATFORM_UWP 0
 	#define FRAMEPRO_PLATFORM_WIN 1
 	#define FRAMEPRO_PLATFORM_LINUX 0
 	#define FRAMEPRO_PLATFORM_PS4 0
@@ -950,7 +972,7 @@ namespace FramePro
 #elif defined(__ORBIS__)
 	#define FRAMEPRO_PLATFORM_UE4 0
 	#define FRAMEPRO_PLATFORM_XBOXONE 0
-	#define FRAMEPRO_PLATFORM_HOLOLENS 0
+	#define FRAMEPRO_PLATFORM_UWP 0
 	#define FRAMEPRO_PLATFORM_WIN 0
 	#define FRAMEPRO_PLATFORM_LINUX 0
 	#define FRAMEPRO_PLATFORM_PS4 1
@@ -958,7 +980,7 @@ namespace FramePro
 #elif defined(__ANDROID__)
 	#define FRAMEPRO_PLATFORM_UE4 0
 	#define FRAMEPRO_PLATFORM_XBOXONE 0
-	#define FRAMEPRO_PLATFORM_HOLOLENS 0
+	#define FRAMEPRO_PLATFORM_UWP 0
 	#define FRAMEPRO_PLATFORM_WIN 0
 	#define FRAMEPRO_PLATFORM_LINUX 0
 	#define FRAMEPRO_PLATFORM_PS4 0
@@ -966,7 +988,7 @@ namespace FramePro
 #elif defined(unix) || defined(__unix__) || defined(__unix)
 	#define FRAMEPRO_PLATFORM_UE4 0
 	#define FRAMEPRO_PLATFORM_XBOXONE 0
-	#define FRAMEPRO_PLATFORM_HOLOLENS 0
+	#define FRAMEPRO_PLATFORM_UWP 0
 	#define FRAMEPRO_PLATFORM_WIN 0
 	#define FRAMEPRO_PLATFORM_LINUX 1
 	#define FRAMEPRO_PLATFORM_PS4 0
@@ -976,10 +998,10 @@ namespace FramePro
 #endif
 
 //------------------------------------------------------------------------
-#define FRAMEPRO_SOURCE_STRING			__FILE__ "|" FRAMEPRO_FUNCTION_NAME "|" FRAMEPRO_STRINGIZE(__LINE__) "|"
-#define FRAMEPRO_SOURCE_STRING_W		FRAMEPRO_WIDESTR(__FILE__) L"|" FRAMEPRO_WIDESTR(FRAMEPRO_FUNCTION_NAME) L"|" FRAMEPRO_WIDESTR(FRAMEPRO_STRINGIZE(__LINE__)) L"|"
-#define FRAMEPRO_SOURCE_STRING_IDLE		__FILE__ "|" FRAMEPRO_FUNCTION_NAME "|" FRAMEPRO_STRINGIZE(__LINE__) "|Idle"
-#define FRAMEPRO_SOURCE_STRING_IDLE_W	FRAMEPRO_WIDESTR(__FILE__) L"|" FRAMEPRO_WIDESTR(FRAMEPRO_FUNCTION_NAME) L"|" FRAMEPRO_WIDESTR(FRAMEPRO_STRINGIZE(__LINE__)) L"|Idle"
+#define FRAMEPRO_SOURCE_STRING			__FILE__ "|" FRAMEPRO_FUNCTION_NAME_LITERAL "|" FRAMEPRO_STRINGIZE(__LINE__) "|"
+#define FRAMEPRO_SOURCE_STRING_W		FRAMEPRO_WIDESTR(__FILE__) L"|" FRAMEPRO_WIDESTR(FRAMEPRO_FUNCTION_NAME_LITERAL) L"|" FRAMEPRO_WIDESTR(FRAMEPRO_STRINGIZE(__LINE__)) L"|"
+#define FRAMEPRO_SOURCE_STRING_IDLE		__FILE__ "|" FRAMEPRO_FUNCTION_NAME_LITERAL "|" FRAMEPRO_STRINGIZE(__LINE__) "|Idle"
+#define FRAMEPRO_SOURCE_STRING_IDLE_W	FRAMEPRO_WIDESTR(__FILE__) L"|" FRAMEPRO_WIDESTR(FRAMEPRO_FUNCTION_NAME_LITERAL) L"|" FRAMEPRO_WIDESTR(FRAMEPRO_STRINGIZE(__LINE__)) L"|Idle"
 
 //------------------------------------------------------------------------
 #if FRAMEPRO_DEBUG
@@ -1036,9 +1058,9 @@ namespace FramePro
 			RequestStringLiteralPacket,
 			SetConditionalScopeMinTimePacket,
 			ConnectResponsePacket,
-			SessionInfoPacket,
+			SessionStatsPacket,
 			RequestRecordedDataPacket,
-			SessionDetailsPacket,
+			legacy_SessionDetailsPacket,
 			ContextSwitchPacket,
 			ContextSwitchRecordingStartedPacket,
 			ProcessNamePacket,
@@ -1057,14 +1079,18 @@ namespace FramePro
 			StringLiteralNamedTimeSpanWithCallstack,
 			ModulePacket,
 			SetCallstackRecordingEnabledPacket,
-			CustomStatPacketW,
+			CustomStatPacketW_Depreciated2,
 			TimeSpanCustomStatPacketW,
-			CustomStatPacket,
+			CustomStatPacket_Depreciated2,
 			TimeSpanCustomStatPacket,
 			ScopeColourPacket,
 			CustomStatColourPacket,
 			CustomStatGraphPacket,
 			CustomStatUnitPacket,
+			CallstackPacket,
+			SessionInfoPacket,
+			CustomStatPacket,
+			CustomStatPacketW,
 		};
 	};
 }
@@ -1093,7 +1119,7 @@ namespace FramePro
 		enum Enum
 		{
 			Windows = 0,
-			Windows_HoloLens,
+			Windows_UWP,
 			XBoxOne,
 			Unused,
 			Linux,
@@ -1122,7 +1148,7 @@ namespace FramePro
 
 //------------------------------------------------------------------------
 //
-// FrameProPlatform.hpp
+// FrameProPlatform.h
 //
 #if FRAMEPRO_ENABLED
 
@@ -1190,9 +1216,9 @@ namespace FramePro
 	#define FRAMEPRO_ENUMERATE_ALL_MODULES (FRAMEPRO_X64 && 1)
 
 //------------------------------------------------------------------------
-//                         FRAMEPRO_PLATFORM_HOLOLENS
+//                         FRAMEPRO_PLATFORM_UWP
 //------------------------------------------------------------------------
-#elif FRAMEPRO_PLATFORM_HOLOLENS
+#elif FRAMEPRO_PLATFORM_UWP
 
 	__int64 FramePro_QueryPerformanceCounter();
 	#define FRAMEPRO_GET_CLOCK_COUNT(time) time = FramePro_QueryPerformanceCounter()
@@ -1215,7 +1241,7 @@ namespace FramePro
 	#define FRAMEPRO_TCHAR _TCHAR
 	
 	#define MULTI_STATEMENT(s) do { s } while(true,false)
-	
+
 	#if !defined(__clang__)
 		#define FRAMEPRO_FUNCTION_DEFINE_IS_STRING_LITERAL 1
 	#else
@@ -1415,7 +1441,7 @@ namespace FramePro
 
 		void DestroyContextSwitchRecorder(void* p_context_switch_recorder, Allocator* p_allocator);
 
-		bool StartRecordingContextSitches(void* p_context_switch_recorder, ContextSwitchCallbackFunction p_callback, void* p_context, DynamicString& error);
+		bool StartRecordingContextSitches(void* p_context_switch_recorder, ContextSwitchCallbackFunction p_callback, void* p_context, DynamicString& error, Allocator* p_allocator);
 		
 		void StopRecordingContextSitches(void* p_context_switch_recorder);
 
@@ -1479,7 +1505,7 @@ namespace FramePro
 
 		void HandleSocketError();
 
-		void CreateThread(
+		bool CreateThread(
 			void* p_os_thread_mem,
 			int os_thread_mem_size,
 			ThreadMain p_thread_main,
@@ -1522,29 +1548,37 @@ namespace FramePro
 #endif		// #if FRAMEPRO_ENABLED
 #include <atomic>
 
+//------------------------------------------------------------------------
+#if FRAMEPRO_WIN_BASED_PLATFORM
+	#define FRAMEPRO_FUNCTION_NAME __FUNCSIG__
+#else
+	#define FRAMEPRO_FUNCTION_NAME __PRETTY_FUNCTION__
+#endif
+
 //------------------------------------------------------------------------
 #if FRAMEPRO_FUNCTION_DEFINE_IS_STRING_LITERAL
-	#define FRAMEPRO_FUNCTION_NAME __FUNCTION__
+	#define FRAMEPRO_FUNCTION_NAME_LITERAL FRAMEPRO_FUNCTION_NAME
 #else
-	#define FRAMEPRO_FUNCTION_NAME ""		// __FUNCTION__ is not a string literal on Linux platforms so we can't combine it with other string literals
+	#define FRAMEPRO_FUNCTION_NAME_LITERAL ""		// FRAMEPRO_FUNCTION_NAME is not a string literal on Linux platforms so we can't combine it with other string literals
 
-	// On non-windows based platforms __FUNCTION__ can't be concatonated with
-	// other string literals so we need to pass in the name separately.
+	// If FRAMEPRO_FUNCTION_NAME is not a string literal it can't be
+	// concatonated with other string literals so we need to pass in
+	// the name separately.
 	#ifdef FRAMEPRO_SCOPE
 		#undef FRAMEPRO_SCOPE
-		#define FRAMEPRO_SCOPE() 					FramePro::StringLiteralNamedTimerScope FRAMEPRO_UNIQUE(timer_scope)(__FUNCTION__, FRAMEPRO_SOURCE_STRING)
+		#define FRAMEPRO_SCOPE() 					FramePro::StringLiteralNamedTimerScope FRAMEPRO_UNIQUE(timer_scope)(FRAMEPRO_FUNCTION_NAME, FRAMEPRO_SOURCE_STRING)
 	#endif
 	#ifdef FRAMEPRO_CONDITIONAL_SCOPE
 		#undef FRAMEPRO_CONDITIONAL_SCOPE
-		#define FRAMEPRO_CONDITIONAL_SCOPE()		FramePro::StringLiteralNamedConditionalTimerScope FRAMEPRO_UNIQUE(timer_scope)(__FUNCTION__, FRAMEPRO_SOURCE_STRING)
+		#define FRAMEPRO_CONDITIONAL_SCOPE()		FramePro::StringLiteralNamedConditionalTimerScope FRAMEPRO_UNIQUE(timer_scope)(FRAMEPRO_FUNCTION_NAME, FRAMEPRO_SOURCE_STRING)
 	#endif
 	#ifdef FRAMEPRO_IDLE_SCOPE
 		#undef FRAMEPRO_IDLE_SCOPE
-		#define FRAMEPRO_IDLE_SCOPE()				FramePro::StringLiteralNamedTimerScope FRAMEPRO_UNIQUE(timer_scope)(__FUNCTION__, FRAMEPRO_SOURCE_STRING_IDLE)
+		#define FRAMEPRO_IDLE_SCOPE()				FramePro::StringLiteralNamedTimerScope FRAMEPRO_UNIQUE(timer_scope)(FRAMEPRO_FUNCTION_NAME, FRAMEPRO_SOURCE_STRING_IDLE)
 	#endif
 	#ifdef FRAMEPRO_IDLE_CONDITIONAL_SCOPE
 		#undef FRAMEPRO_IDLE_CONDITIONAL_SCOPE
-		#define FRAMEPRO_IDLE_CONDITIONAL_SCOPE()	FramePro::StringLiteralNamedConditionalTimerScope FRAMEPRO_UNIQUE(timer_scope)(__FUNCTION__, FRAMEPRO_SOURCE_STRING_IDLE)
+		#define FRAMEPRO_IDLE_CONDITIONAL_SCOPE()	FramePro::StringLiteralNamedConditionalTimerScope FRAMEPRO_UNIQUE(timer_scope)(FRAMEPRO_FUNCTION_NAME, FRAMEPRO_SOURCE_STRING_IDLE)
 	#endif
 #endif
 
@@ -1589,9 +1623,9 @@ namespace FramePro
 
 	FRAMEPRO_API void DebugBreak();
 
-	FRAMEPRO_API void SendSessionInfo(const char* p_name, const char* p_build_id);
+	FRAMEPRO_API void SendSessionInfo(const char* p_name, const char* p_value);
 
-	FRAMEPRO_API void SendSessionInfo(const wchar_t* p_name, const wchar_t* p_build_id);
+	FRAMEPRO_API void SendSessionInfo(const wchar_t* p_name, const wchar_t* p_value);
 
 	FRAMEPRO_API void SetAllocator(class Allocator* p_allocator);		// if you call this you must call it BEFORE any other calls
 
@@ -1649,12 +1683,14 @@ namespace FramePro
 
 	FRAMEPRO_API void AddEvent(const char* p_name, uint colour);
 
-	FRAMEPRO_API void AddWaitEvent(int64 event_id, int64 start_time, int64 end_time);
+	FRAMEPRO_API void AddWaitEvent(int64 event_id, int64 start_time, int start_core, int64 end_time, int end_core);
 
 	FRAMEPRO_API void TriggerWaitEvent(int64 event_id);
 
 	FRAMEPRO_API void SetThreadName(const char* p_name);
 
+	FRAMEPRO_API void SetThreadName(int thread_id, const char* p_name);
+
 	FRAMEPRO_API void SetThreadOrder(StringId thread_name);
 
 	FRAMEPRO_API StringId RegisterString(const char* p_str);
@@ -1679,7 +1715,7 @@ namespace FramePro
 
 	FRAMEPRO_API void UnblockSockets();
 
-	FRAMEPRO_API void AddGlobalHiResTimer(class GlobalHiResTimer* p_timer);
+	FRAMEPRO_API void AddCustomStatTimer(class CustomStatTimer* p_timer);
 
 	FRAMEPRO_API void CleanupThread();
 
@@ -1693,8 +1729,6 @@ namespace FramePro
 
 	FRAMEPRO_API void StopHiResTimer();
 
-	FRAMEPRO_API void SubmitHiResTimers(int64 current_time);
-
 	FRAMEPRO_API void Log(const char* p_message);
 
 	FRAMEPRO_API void SetScopeCustomStat(const char* p_name, int value, const char* p_graph, const char* p_unit, uint colour);
@@ -1730,6 +1764,18 @@ namespace FramePro
 	FRAMEPRO_API void SetCustomStatUnit(StringId name, StringId unit);
 
 	FRAMEPRO_API void SetCustomStatColour(StringId name, uint colour);
+
+	FRAMEPRO_API void SendScopeCallstack();
+
+	FRAMEPRO_API int GetCore();
+
+#if FRAMEPRO_DEBUG || FRAMEPRO_ALLOW_UNPARENTED_HIRES_SCOPES
+	FRAMEPRO_API void StartScope();
+	FRAMEPRO_API void StopScope();
+#else
+	inline void StartScope() {}
+	inline void StopScope() {}
+#endif
 }
 
 //------------------------------------------------------------------------
@@ -1745,70 +1791,37 @@ namespace FramePro
 		virtual ~Allocator() {}
 	};
 
-	//------------------------------------------------------------------------
-	template<class T>
-	inline T* New(Allocator* p_allocator)
-	{
-		T* p = (T*)p_allocator->Alloc(sizeof(T));
-		new (p)T();
-		return p;
-	}
-
-	//------------------------------------------------------------------------
-	template<class T, typename Targ1>
-	inline T* New(Allocator* p_allocator, Targ1 arg1)
-	{
-		T* p = (T*)p_allocator->Alloc(sizeof(T));
-		new (p)T(arg1);
-		return p;
-	}
-
-	//------------------------------------------------------------------------
-	template<class T, typename Targ1, typename Targ2, typename TArg3>
-	inline T* New(Allocator* p_allocator, Targ1 arg1, Targ2 arg2, TArg3 arg3)
-	{
-		T* p = (T*)p_allocator->Alloc(sizeof(T));
-		new (p)T(arg1, arg2, arg3);
-		return p;
-	}
-
-	//------------------------------------------------------------------------
-	template<typename T>
-	inline void Delete(Allocator* p_allocator, T* p)
-	{
-		p->~T();
-		p_allocator->Free(p);
-	}
-
 	//------------------------------------------------------------------------
 	class TimerScope
 	{
 	public:
 		TimerScope(const char* p_name_and_source_info)
-		:	mp_NameAndSourceInfo(p_name_and_source_info)
 		{
-			bool connected = FramePro::IsConnected();
-			m_Connected = connected;
-
-			int64 start_time;
-			FRAMEPRO_GET_CLOCK_COUNT(start_time);
-			if (connected)
-				FramePro::SubmitHiResTimers(start_time);
-			m_StartTime = start_time;
+			if (FramePro::IsConnected())
+			{
+				mp_NameAndSourceInfo = p_name_and_source_info;
+				FRAMEPRO_GET_CLOCK_COUNT(m_StartTime);
+				FramePro::StartScope();
+			}
+			else
+			{
+				mp_NameAndSourceInfo = nullptr;
+				m_StartTime = 0;
+			}
 		}
 
 		~TimerScope()
 		{
-			if(m_Connected)
+			if(mp_NameAndSourceInfo)
 			{
 				int64 end_time = 0;
 				FRAMEPRO_GET_CLOCK_COUNT(end_time);
 				FramePro::AddTimeSpan(mp_NameAndSourceInfo, m_StartTime, end_time);
+				FramePro::StopScope();
 			}
 		}
 
 	private:
-		bool m_Connected;
 		const char* mp_NameAndSourceInfo;
 		int64 m_StartTime;
 	};
@@ -1818,30 +1831,32 @@ namespace FramePro
 	{
 	public:
 		TimerScopeW(const wchar_t* p_name_and_source_info)
-		:	mp_NameAndSourceInfo(p_name_and_source_info)
 		{
-			bool connected = FramePro::IsConnected();
-			m_Connected = connected;
-
-			int64 start_time;
-			FRAMEPRO_GET_CLOCK_COUNT(start_time);
-			if (connected)
-				FramePro::SubmitHiResTimers(start_time);
-			m_StartTime = start_time;
+			if (FramePro::IsConnected())
+			{
+				mp_NameAndSourceInfo = p_name_and_source_info;
+				FRAMEPRO_GET_CLOCK_COUNT(m_StartTime);
+				FramePro::StartScope();
+			}
+			else
+			{
+				mp_NameAndSourceInfo = nullptr;
+				m_StartTime = 0;
+			}
 		}
 
 		~TimerScopeW()
 		{
-			if(m_Connected)
+			if(mp_NameAndSourceInfo)
 			{
 				int64 end_time = 0;
 				FRAMEPRO_GET_CLOCK_COUNT(end_time);
 				FramePro::AddTimeSpan(mp_NameAndSourceInfo, m_StartTime, end_time);
+				FramePro::StopScope();
 			}
 		}
 
 	private:
-		bool m_Connected;
 		const wchar_t* mp_NameAndSourceInfo;
 		int64 m_StartTime;
 	};
@@ -1851,29 +1866,30 @@ namespace FramePro
 	{
 	public:
 		ConditionalParentTimerScope(const char* p_name, const char* p_source_info, ConditionalParentScopeCallback p_callback, int64 pre_duration, int64 post_duration)
-		:	mp_Name(p_name),
-			mp_SourceInfo(p_source_info),
-			m_StartTime(0),
-			mp_Callback(p_callback)
+		:	mp_Callback(p_callback)
 		{
-			bool connected = FramePro::IsConnected();
-			m_Connected = connected;
-
-			if (connected)
+			if (FramePro::IsConnected())
 			{
+				mp_Name = p_name;
+				mp_SourceInfo = p_source_info;
+
 				FramePro::PushConditionalParentScope(p_name, pre_duration, post_duration);
 
-				int64 start_time;
-				FRAMEPRO_GET_CLOCK_COUNT(start_time);
-				if (connected)
-					FramePro::SubmitHiResTimers(start_time);
-				m_StartTime = start_time;
+				FRAMEPRO_GET_CLOCK_COUNT(m_StartTime);
+
+				FramePro::StartScope();
+			}
+			else
+			{
+				mp_Name = nullptr;
+				mp_SourceInfo = nullptr;
+				m_StartTime = 0;
 			}
 		}
 
 		~ConditionalParentTimerScope()
 		{
-			if (m_Connected)
+			if (mp_SourceInfo)
 			{
 				int64 end_time = 0;
 				FRAMEPRO_GET_CLOCK_COUNT(end_time);
@@ -1882,11 +1898,12 @@ namespace FramePro
 				FramePro::PopConditionalParentScope(add_children);
 
 				FramePro::AddTimeSpan(mp_Name, mp_SourceInfo, m_StartTime, end_time);
+
+				FramePro::StopScope();
 			}
 		}
 
 	private:
-		bool m_Connected;
 		const char* mp_Name;
 		const char* mp_SourceInfo;
 		int64 m_StartTime;
@@ -1899,31 +1916,34 @@ namespace FramePro
 	{
 	public:
 		IdTimerScope(StringId name, const char* p_source_info)
-		:	m_Name(name),
-			mp_SourceInfo(p_source_info)
 		{
-			bool connected = FramePro::IsConnected();
-			m_Connected = connected;
-
-			int64 start_time;
-			FRAMEPRO_GET_CLOCK_COUNT(start_time);
-			if (connected)
-				FramePro::SubmitHiResTimers(start_time);
-			m_StartTime = start_time;
+			if (FramePro::IsConnected())
+			{
+				m_Name = name;
+				mp_SourceInfo = p_source_info;
+				FRAMEPRO_GET_CLOCK_COUNT(m_StartTime);
+				FramePro::StartScope();
+			}
+			else
+			{
+				m_Name = 0;
+				mp_SourceInfo = nullptr;
+				m_StartTime = 0;
+			}
 		}
 
 		~IdTimerScope()
 		{
-			if(m_Connected)
+			if(mp_SourceInfo)
 			{
 				int64 end_time = 0;
 				FRAMEPRO_GET_CLOCK_COUNT(end_time);
 				FramePro::AddTimeSpan(m_Name, mp_SourceInfo, m_StartTime, end_time);
+				FramePro::StopScope();
 			}
 		}
 
 	private:
-		bool m_Connected;
 		StringId m_Name;
 		const char* mp_SourceInfo;
 		int64 m_StartTime;
@@ -1935,31 +1955,34 @@ namespace FramePro
 	public:
 		// p_name must be a string literal
 		StringLiteralNamedTimerScope(const char* p_name, const char* p_source_info)
-		:	mp_Name(p_name),
-			mp_SourceInfo(p_source_info)
 		{
-			bool connected = FramePro::IsConnected();
-			m_Connected = connected;
-
-			int64 start_time;
-			FRAMEPRO_GET_CLOCK_COUNT(start_time);
-			if (connected)
-				FramePro::SubmitHiResTimers(start_time);
-			m_StartTime = start_time;
+			if (FramePro::IsConnected())
+			{
+				mp_Name = p_name;
+				mp_SourceInfo = p_source_info;
+				FRAMEPRO_GET_CLOCK_COUNT(m_StartTime);
+				FramePro::StartScope();
+			}
+			else
+			{
+				mp_Name = nullptr;
+				mp_SourceInfo = nullptr;
+				m_StartTime = 0;
+			}
 		}
 
 		~StringLiteralNamedTimerScope()
 		{
-			if(m_Connected)
+			if(mp_SourceInfo)
 			{
 				int64 end_time = 0;
 				FRAMEPRO_GET_CLOCK_COUNT(end_time);
 				FramePro::AddTimeSpan(mp_Name, mp_SourceInfo, m_StartTime, end_time);
+				FramePro::StopScope();
 			}
 		}
 
 	private:
-		bool m_Connected;
 		const char* mp_Name;
 		const char* mp_SourceInfo;
 		int64 m_StartTime;
@@ -1970,31 +1993,33 @@ namespace FramePro
 	{
 	public:
 		ConditionalTimerScope(const char* p_name_and_source_info)
-		:	mp_NameAndSourceInfo(p_name_and_source_info)
 		{
-			bool connected = FramePro::IsConnected();
-			m_Connected = connected;
-
-			int64 start_time;
-			FRAMEPRO_GET_CLOCK_COUNT(start_time);
-			if (connected)
-				FramePro::SubmitHiResTimers(start_time);
-			m_StartTime = start_time;
+			if (FramePro::IsConnected())
+			{
+				mp_NameAndSourceInfo = p_name_and_source_info;
+				FRAMEPRO_GET_CLOCK_COUNT(m_StartTime);
+				FramePro::StartScope();
+			}
+			else
+			{
+				mp_NameAndSourceInfo = nullptr;
+				m_StartTime = 0;
+			}
 		}
 
 		~ConditionalTimerScope()
 		{
-			if(m_Connected)
+			if(mp_NameAndSourceInfo)
 			{
 				int64 end_time = 0;
 				FRAMEPRO_GET_CLOCK_COUNT(end_time);
 				if(end_time - m_StartTime > FramePro::GetConditionalScopeMinTime())
 					FramePro::AddTimeSpan(mp_NameAndSourceInfo, m_StartTime, end_time);
+				FramePro::StopScope();
 			}
 		}
 
 	private:
-		bool m_Connected;
 		const char* mp_NameAndSourceInfo;
 		int64 m_StartTime;
 	};
@@ -2004,32 +2029,35 @@ namespace FramePro
 	{
 	public:
 		ConditionalTimerScopeId(StringId name, const char* p_source_info)
-		:	m_Name(name),
-			mp_SourceInfo(p_source_info)
 		{
-			bool connected = FramePro::IsConnected();
-			m_Connected = connected;
-
-			int64 start_time;
-			FRAMEPRO_GET_CLOCK_COUNT(start_time);
-			if (connected)
-				FramePro::SubmitHiResTimers(start_time);
-			m_StartTime = start_time;
+			if (FramePro::IsConnected())
+			{
+				m_Name = name;
+				mp_SourceInfo = p_source_info;
+				FRAMEPRO_GET_CLOCK_COUNT(m_StartTime);
+				FramePro::StartScope();
+			}
+			else
+			{
+				m_Name = 0;
+				mp_SourceInfo = nullptr;
+				m_StartTime = 0;
+			}
 		}
 
 		~ConditionalTimerScopeId()
 		{
-			if(m_Connected)
+			if(mp_SourceInfo)
 			{
 				int64 end_time = 0;
 				FRAMEPRO_GET_CLOCK_COUNT(end_time);
 				if(end_time - m_StartTime > FramePro::GetConditionalScopeMinTime())
 					FramePro::AddTimeSpan(m_Name, mp_SourceInfo, m_StartTime, end_time);
+				FramePro::StopScope();
 			}
 		}
 
 	private:
-		bool m_Connected;
 		StringId m_Name;
 		const char* mp_SourceInfo;
 		int64 m_StartTime;
@@ -2040,31 +2068,33 @@ namespace FramePro
 	{
 	public:
 		ConditionalTimerScopeW(const wchar_t* p_name_and_source_info)
-		:	mp_NameAndSourceInfo(p_name_and_source_info)
 		{
-			bool connected = FramePro::IsConnected();
-			m_Connected = connected;
-
-			int64 start_time;
-			FRAMEPRO_GET_CLOCK_COUNT(start_time);
-			if (connected)
-				FramePro::SubmitHiResTimers(start_time);
-			m_StartTime = start_time;
+			if (FramePro::IsConnected())
+			{
+				mp_NameAndSourceInfo = p_name_and_source_info;
+				FRAMEPRO_GET_CLOCK_COUNT(m_StartTime);
+				FramePro::StartScope();
+			}
+			else
+			{
+				mp_NameAndSourceInfo = nullptr;
+				m_StartTime = 0;
+			}
 		}
 
 		~ConditionalTimerScopeW()
 		{
-			if(m_Connected)
+			if(mp_NameAndSourceInfo)
 			{
 				int64 end_time = 0;
 				FRAMEPRO_GET_CLOCK_COUNT(end_time);
 				if(end_time - m_StartTime > FramePro::GetConditionalScopeMinTime())
 					FramePro::AddTimeSpan(mp_NameAndSourceInfo, m_StartTime, end_time);
+				FramePro::StopScope();
 			}
 		}
 
 	private:
-		bool m_Connected;
 		const wchar_t* mp_NameAndSourceInfo;
 		int64 m_StartTime;
 	};
@@ -2074,33 +2104,33 @@ namespace FramePro
 	{
 	public:
 		ConditionalBoolTimerScope(const char* p_name_and_source_info, bool b)
-		:	mp_NameAndSourceInfo(p_name_and_source_info)
 		{
-			bool send_scope = b && FramePro::IsConnected();
-			m_SendScope = send_scope;
-
-			if (send_scope)
+			if (b && FramePro::IsConnected())
+			{
+				mp_NameAndSourceInfo = p_name_and_source_info;
+				FRAMEPRO_GET_CLOCK_COUNT(m_StartTime);
+				FramePro::StartScope();
+			}
+			else
 			{
-				int64 start_time;
-				FRAMEPRO_GET_CLOCK_COUNT(start_time);
-				FramePro::SubmitHiResTimers(start_time);
-				m_StartTime = start_time;
+				mp_NameAndSourceInfo = nullptr;
+				m_StartTime = 0;
 			}
 		}
 
 		~ConditionalBoolTimerScope()
 		{
-			if(m_SendScope)
+			if(mp_NameAndSourceInfo)
 			{
 				int64 end_time = 0;
 				FRAMEPRO_GET_CLOCK_COUNT(end_time);
 				if(end_time - m_StartTime > FramePro::GetConditionalScopeMinTime())
 					FramePro::AddTimeSpan(mp_NameAndSourceInfo, m_StartTime, end_time);
+				FramePro::StopScope();
 			}
 		}
 
 	private:
-		bool m_SendScope;
 		const char* mp_NameAndSourceInfo;
 		int64 m_StartTime;
 	};
@@ -2110,34 +2140,35 @@ namespace FramePro
 	{
 	public:
 		ConditionalBoolTimerScopeId(StringId name, const char* p_source_info, bool b)
-		:	m_Name(name),
-			mp_SourceInfo(p_source_info)
 		{
-			bool send_scope = b && FramePro::IsConnected();
-			m_SendScope = send_scope;
-
-			if (send_scope)
+			if (b && FramePro::IsConnected())
+			{
+				m_Name = name;
+				mp_SourceInfo = p_source_info;
+				FRAMEPRO_GET_CLOCK_COUNT(m_StartTime);
+				FramePro::StartScope();
+			}
+			else
 			{
-				int64 start_time;
-				FRAMEPRO_GET_CLOCK_COUNT(start_time);
-				FramePro::SubmitHiResTimers(start_time);
-				m_StartTime = start_time;
+				m_Name = 0;
+				mp_SourceInfo = nullptr;
+				m_StartTime = 0;
 			}
 		}
 
 		~ConditionalBoolTimerScopeId()
 		{
-			if(m_SendScope)
+			if(mp_SourceInfo)
 			{
 				int64 end_time = 0;
 				FRAMEPRO_GET_CLOCK_COUNT(end_time);
 				if(end_time - m_StartTime > FramePro::GetConditionalScopeMinTime())
 					FramePro::AddTimeSpan(m_Name, mp_SourceInfo, m_StartTime, end_time);
+				FramePro::StopScope();
 			}
 		}
 
 	private:
-		bool m_SendScope;
 		StringId m_Name;
 		const char* mp_SourceInfo;
 		int64 m_StartTime;
@@ -2148,33 +2179,33 @@ namespace FramePro
 	{
 	public:
 		ConditionalBoolTimerScopeW(const wchar_t* p_name_and_source_info, bool b)
-		:	mp_NameAndSourceInfo(p_name_and_source_info)
 		{
-			bool send_scope = b && FramePro::IsConnected();
-			m_SendScope = send_scope;
-
-			if (send_scope)
+			if (b && FramePro::IsConnected())
 			{
-				int64 start_time;
-				FRAMEPRO_GET_CLOCK_COUNT(start_time);
-				FramePro::SubmitHiResTimers(start_time);
-				m_StartTime = start_time;
+				mp_NameAndSourceInfo = p_name_and_source_info;
+				FRAMEPRO_GET_CLOCK_COUNT(m_StartTime);
+				FramePro::StartScope();
+			}
+			else
+			{
+				mp_NameAndSourceInfo = nullptr;
+				m_StartTime = 0;
 			}
 		}
 
 		~ConditionalBoolTimerScopeW()
 		{
-			if(m_SendScope)
+			if(mp_NameAndSourceInfo)
 			{
 				int64 end_time = 0;
 				FRAMEPRO_GET_CLOCK_COUNT(end_time);
 				if(end_time - m_StartTime > FramePro::GetConditionalScopeMinTime())
 					FramePro::AddTimeSpan(mp_NameAndSourceInfo, m_StartTime, end_time);
+				FramePro::StopScope();
 			}
 		}
 
 	private:
-		bool m_SendScope;
 		const wchar_t* mp_NameAndSourceInfo;
 		int64 m_StartTime;
 	};
@@ -2184,62 +2215,71 @@ namespace FramePro
 	{
 	public:
 		StringLiteralNamedConditionalTimerScope(const char* p_name, const char* p_source_info)
-		:	mp_Name(p_name),
-			mp_SourceInfo(p_source_info)
 		{
-			bool connected = FramePro::IsConnected();
-			m_Connected = connected;
-
-			int64 start_time;
-			FRAMEPRO_GET_CLOCK_COUNT(start_time);
-			if (connected)
-				FramePro::SubmitHiResTimers(start_time);
-			m_StartTime = start_time;
+			if (FramePro::IsConnected())
+			{
+				mp_Name = p_name;
+				mp_SourceInfo = p_source_info;
+				FRAMEPRO_GET_CLOCK_COUNT(m_StartTime);
+				FramePro::StartScope();
+			}
+			else
+			{
+				mp_Name = nullptr;
+				mp_SourceInfo = nullptr;
+				m_StartTime = 0;
+			}
 		}
 
 		~StringLiteralNamedConditionalTimerScope()
 		{
-			if(m_Connected)
+			if(mp_SourceInfo)
 			{
 				int64 end_time = 0;
 				FRAMEPRO_GET_CLOCK_COUNT(end_time);
 				if(end_time - m_StartTime > FramePro::GetConditionalScopeMinTime())
 					FramePro::AddTimeSpan(mp_Name, mp_SourceInfo, m_StartTime, end_time);
+				FramePro::StopScope();
 			}
 		}
 
 	private:
-		bool m_Connected;
 		const char* mp_Name;
 		const char* mp_SourceInfo;
 		int64 m_StartTime;
 	};
 
 	//------------------------------------------------------------------------
-	class GlobalHiResTimer
+	class CustomStatTimer
 	{
 	public:
-		GlobalHiResTimer(const char* p_name) : m_Value(0), mp_Next(NULL), mp_Name(p_name)
+		CustomStatTimer(const char* p_name)
+		:	m_Value(0),
+			mp_Next(NULL)
 		{
-			FramePro::AddGlobalHiResTimer(this);
+			mp_Name = p_name;
+			FramePro::AddCustomStatTimer(this);
 		}
 
 		void Add(uint value)
 		{
-			uint64 existing_value = m_Value.load(std::memory_order_relaxed);
-			uint64 new_value;
-			do
+			if (mp_Name)
 			{
-				uint64 duration = existing_value & 0xffffffffff;
-				uint64 count = (existing_value >> 40) & 0xffffff;
-				duration += value;
-				++count;
-				new_value = (count << 40) | duration;
-
-				FRAMEPRO_ASSERT(count <= 0xffffff);
-				FRAMEPRO_ASSERT(duration <= 0xffffffffffL);
-
-			} while (!m_Value.compare_exchange_weak(existing_value, new_value, std::memory_order_relaxed, std::memory_order_relaxed));
+				uint64 existing_value = m_Value.load(std::memory_order_relaxed);
+				uint64 new_value;
+				do
+				{
+					uint64 duration = existing_value & 0xffffffffff;
+					uint64 count = (existing_value >> 40) & 0xffffff;
+					duration += value;
+					++count;
+					new_value = (count << 40) | duration;
+
+					FRAMEPRO_ASSERT(count <= 0xffffff);
+					FRAMEPRO_ASSERT(duration <= 0xffffffffffL);
+
+				} while (!m_Value.compare_exchange_weak(existing_value, new_value, std::memory_order_relaxed, std::memory_order_relaxed));
+			}
 		}
 
 		void GetAndClear(uint64& value, uint& count)
@@ -2252,47 +2292,54 @@ namespace FramePro
 			count = (existing_value >> 40) & 0xffffff;
 		}
 
-		void SetNext(GlobalHiResTimer* p_next) { mp_Next = p_next; }
+		void SetNext(CustomStatTimer* p_next) { mp_Next = p_next; }
 
-		GlobalHiResTimer* GetNext() const { return mp_Next; }
+		CustomStatTimer* GetNext() const { return mp_Next; }
 
 		const char* GetName() const { return mp_Name; }
 
 	private:
 		std::atomic<uint64> m_Value;
 
-		GlobalHiResTimer* mp_Next;
+		CustomStatTimer* mp_Next;
 
 		const char* mp_Name;
 	};
 
 	//------------------------------------------------------------------------
-	class GlobalHiResTimerScope
+	class CustomStatTimerScope
 	{
 	public:
-		GlobalHiResTimerScope(GlobalHiResTimer& timer)
-		:	m_Timer(timer)
+		CustomStatTimerScope(CustomStatTimer& timer)
 		{
-			FRAMEPRO_GET_CLOCK_COUNT(m_StartTime);
+			if (FramePro::IsConnected())
+			{
+				mp_Timer = &timer;
+				FRAMEPRO_GET_CLOCK_COUNT(m_StartTime);
+			}
+			else
+			{
+				mp_Timer = nullptr;
+			}
 		}
 
-		~GlobalHiResTimerScope()
+		~CustomStatTimerScope()
 		{
-			if (FramePro::IsConnected())
+			if (mp_Timer)
 			{
 				int64 end_time;
 				FRAMEPRO_GET_CLOCK_COUNT(end_time);
 
-				m_Timer.Add((uint)(end_time - m_StartTime));
+				mp_Timer->Add((uint)(end_time - m_StartTime));
 			}
 		}
 
 	private:
-		GlobalHiResTimerScope(const GlobalHiResTimerScope&);
-		GlobalHiResTimerScope& operator=(const GlobalHiResTimerScope&);
+		CustomStatTimerScope(const CustomStatTimerScope&);
+		CustomStatTimerScope& operator=(const CustomStatTimerScope&);
 
 		int64 m_StartTime;
-		GlobalHiResTimer& m_Timer;
+		CustomStatTimer* mp_Timer;
 	};
 
 	//------------------------------------------------------------------------
@@ -2338,22 +2385,36 @@ namespace FramePro
 	{
 	public:
 		WaitEventScope(int64 event_id)
-		:	m_EventId(event_id)
 		{
-			FRAMEPRO_GET_CLOCK_COUNT(m_StartTime);
+			if (FramePro::IsConnected())
+			{
+				m_EventId = event_id;
+				m_StartCore = FramePro::GetCore();
+				FRAMEPRO_GET_CLOCK_COUNT(m_StartTime);
+			}
+			else
+			{
+				m_EventId = 0;
+				m_StartCore = 0;
+				m_StartTime = 0;
+			}
 		}
 
 		~WaitEventScope()
 		{
-			int64 end_time;
-			FRAMEPRO_GET_CLOCK_COUNT(end_time);
+			if (m_EventId)
+			{
+				int64 end_time;
+				FRAMEPRO_GET_CLOCK_COUNT(end_time);
 
-			FramePro::AddWaitEvent(m_EventId, m_StartTime, end_time);
+				FramePro::AddWaitEvent(m_EventId, m_StartTime, m_StartCore, end_time, FramePro::GetCore());
+			}
 		}
 
 	private:
 		int64 m_EventId;
 		int64 m_StartTime;
+		int m_StartCore;
 	};
 
 	//------------------------------------------------------------------------
@@ -2391,13 +2452,100 @@ namespace FramePro
 			return hash;
 		#endif
 	}
+
+	//------------------------------------------------------------------------
+	namespace ThreadState
+	{
+		enum Enum
+		{
+			Initialized = 0,
+			Ready,
+			Running,
+			Standby,
+			Terminated,
+			Waiting,
+			Transition,
+			DeferredReady,
+
+			Reserved1,
+			Reserved2,
+			Reserved3,
+			Reserved4,
+			Reserved5,
+			Reserved6,
+			Reserved7,
+			Reserved8,
+			Reserved9,
+			Reserved10,
+			Reserved11,
+			Reserved12,
+		};
+	}
+
+	//------------------------------------------------------------------------
+	namespace ThreadWaitReason
+	{
+		enum Enum
+		{
+			Executive = 0,
+			FreePage,
+			PageIn,
+			PoolAllocation,
+			DelayExecution,
+			Suspended,
+			UserRequest,
+			WrExecutive,
+			WrFreePage,
+			WrPageIn,
+			WrPoolAllocation,
+			WrDelayExecution,
+			WrSuspended,
+			WrUserRequest,
+			WrEventPair,
+			WrQueue,
+			WrLpcReceive,
+			WrLpcReply,
+			WrVirtualMemory,
+			WrPageOut,
+			WrRendezvous,
+			WrKeyedEvent,
+			WrTerminated,
+			WrProcessInSwap,
+			WrCpuRateControl,
+			WrCalloutStack,
+			WrKernel,
+			WrResource,
+			WrPushLock,
+			WrMutex,
+			WrQuantumEnd,
+			WrDispatchInt,
+			WrPreempted,
+			WrYieldExecution,
+			WrFastMutex,
+			WrGuardedMutex,
+			WrRundown,
+			MaximumWaitReason,
+		};
+	}
+
+	//------------------------------------------------------------------------
+	struct ContextSwitch
+	{
+		int64 m_Timestamp;
+		int m_ProcessId;
+		int m_CPUId;
+		int m_OldThreadId;
+		int m_NewThreadId;
+		ThreadState::Enum m_OldThreadState;
+		ThreadWaitReason::Enum m_OldThreadWaitReason;
+	};
 }
 
 //------------------------------------------------------------------------
 
 //------------------------------------------------------------------------
 //
-// Array.hpp
+// Array.h
 //
 #include <memory.h>
 
@@ -2460,13 +2608,18 @@ namespace FramePro
 		}
 
 		//--------------------------------------------------------------------
-		void Add(const T& value)
+		bool Add(const T& value)
 		{
-			if(m_Count == m_Capacity)
-				Grow();
+			if (m_Count == m_Capacity)
+			{
+				if (!Grow())
+					return false;
+			}
 
 			mp_Array[m_Count] = value;
 			++m_Count;
+
+			return true;
 		}
 
 		//--------------------------------------------------------------------
@@ -2512,7 +2665,7 @@ namespace FramePro
 		}
 
 		//--------------------------------------------------------------------
-		void Resize(int count)
+		bool Resize(int count)
 		{
 			int new_capacity = m_Count + count;
 
@@ -2523,6 +2676,12 @@ namespace FramePro
 					m_Capacity *= 2;
 
 				T* p_new_array = (T*)mp_Allocator->Alloc(sizeof(T)*m_Capacity);
+				if (!p_new_array)
+				{
+					mp_Array = NULL;
+					return false;
+				}
+
 				if (mp_Array)
 					memcpy(p_new_array, mp_Array, sizeof(T)*m_Count);
 				mp_Allocator->Free(mp_Array);
@@ -2530,18 +2689,30 @@ namespace FramePro
 			}
 
 			m_Count = count;
+
+			return true;
 		}
 
 	private:
 		//--------------------------------------------------------------------
-		void Grow()
+		bool Grow()
 		{
 			m_Capacity = m_Capacity ? 2*m_Capacity : 32;
+			
 			T* p_new_array = (T*)mp_Allocator->Alloc(sizeof(T)*m_Capacity);
+			if (!p_new_array)
+			{
+				mp_Array = NULL;
+				return false;
+			}
+
 			if(mp_Array)
 				memcpy(p_new_array, mp_Array, sizeof(T)*m_Count);
+
 			mp_Allocator->Free(mp_Array);
 			mp_Array = p_new_array;
+
+			return true;
 		}
 
 		//------------------------------------------------------------------------
@@ -2564,3 +2735,4 @@ namespace FramePro
 
 //------------------------------------------------------------------------
 #endif		// #ifndef FRAMEPRO_H_INCLUDED
+
diff --git a/Engine/Source/Runtime/Core/Public/FramePro/FrameProUE4.h b/Engine/Source/Runtime/Core/Public/FramePro/FrameProUE4.h
--- a/Engine/Source/Runtime/Core/Public/FramePro/FrameProUE4.h
+++ b/Engine/Source/Runtime/Core/Public/FramePro/FrameProUE4.h
@@ -45,7 +45,7 @@
 
 		// Windows or Linux based platform
 #if !defined(FRAMEPRO_WIN_BASED_PLATFORM) //@EPIC begin - allow external definition
-		#define FRAMEPRO_WIN_BASED_PLATFORM (PLATFORM_WINDOWS)
+		#define FRAMEPRO_WIN_BASED_PLATFORM (PLATFORM_WINDOWS || PLATFORM_XBOXONE)
 #endif //@EPIC: end
 		#define FRAMEPRO_LINUX_BASED_PLATFORM (!FRAMEPRO_WIN_BASED_PLATFORM && !PLATFORM_SWITCH)
 
@@ -53,8 +53,13 @@
 
 		// Port
 #if !defined(FRAMEPRO_PORT) //@EPIC begin - allow external definition
-		#define FRAMEPRO_PORT "8428"
+		#if PLATFORM_XBOXONE
+			#define FRAMEPRO_PORT "4420"
+		#else
+			#define FRAMEPRO_PORT "8428"
+		#endif
 #endif //@EPIC: end
+
 		// x64 or x32
 		#define FRAMEPRO_X64 PLATFORM_64BITS
 
@@ -91,6 +96,8 @@
 
 		#define FRAMEPRO_ENUMERATE_ALL_MODULES 0
 
+		#define FRAMEPRO_PS4_CONTEXT_SWITCH_RECORDING 0
+
 	#endif		// #ifndef FRAMEPRO_UE4_INCLUDED
 
 //------------------------------------------------------------------------
