Remove SOMAXCONN definitions from FramePro because they clash with Winsock (and are not well defined on other platforms).

diff --git a/Engine/Source/Runtime/Core/Private/FramePro/FramePro.cpp b/Engine/Source/Runtime/Core/Private/FramePro/FramePro.cpp
--- a/Engine/Source/Runtime/Core/Private/FramePro/FramePro.cpp
+++ b/Engine/Source/Runtime/Core/Private/FramePro/FramePro.cpp
@@ -23,12 +23,6 @@ OR IN CONNECTION WITH THE USE OR PERFORMANCE OF THIS SOFTWARE.
 //------------------------------------------------------------------------
 #if FRAMEPRO_ENABLED
 
-//#ifndef PLATFORM_SWITCH //@EPIC begin
-	#ifndef SOMAXCONN //@EPIC begin
-		#define SOMAXCONN 16
-	#endif //@EPIC: end
-//#endif //@EPIC: end
-
 //------------------------------------------------------------------------
 //
 // FrameProLib.h
@@ -9239,7 +9233,7 @@ namespace FramePro
 			bool StartSocketListening(void* p_os_socket_mem)
 			{
 				#if FRAMEPRO_SOCKETS_ENABLED
-					return listen(GetOSSocket(p_os_socket_mem), SOMAXCONN) != SOCKET_ERROR;
+					return listen(GetOSSocket(p_os_socket_mem), 16) != SOCKET_ERROR; //@EPIC BEGIN 16 is SOMAXCONN //@EPIC END
 				#else
 					FRAMEPRO_UNREFERENCED(p_os_socket_mem);
 					return false;
@@ -9798,7 +9792,7 @@ namespace FramePro
 			bool StartSocketListening(void* p_os_socket_mem)
 			{
 				#if FRAMEPRO_SOCKETS_ENABLED
-					return listen(GetOSSocket(p_os_socket_mem), SOMAXCONN) != g_SocketErrorId;
+					return listen(GetOSSocket(p_os_socket_mem), 16) != g_SocketErrorId; //@EPIC BEGIN 16 is SOMAXCONN //@EPIC END
 				#else
 					return false;
 				#endif
