Removal of unnecessary ESPMode::Fast.

diff --git a/Engine/Source/Runtime/Core/Public/HAL/IConsoleManager.h b/Engine/Source/Runtime/Core/Public/HAL/IConsoleManager.h
--- a/Engine/Source/Runtime/Core/Public/HAL/IConsoleManager.h
+++ b/Engine/Source/Runtime/Core/Public/HAL/IConsoleManager.h
@@ -195,12 +195,12 @@ struct FNullConsoleVariableDelegate
 	}
 
 	template <typename UserClass, typename... VarTypes>
-	inline static DerivedType CreateSP(const TSharedRef<UserClass, ESPMode::Fast>&, typename TMemFunPtrType<false, UserClass, void (ParamTypes..., VarTypes...)>::Type, VarTypes...)
+	inline static DerivedType CreateSP(const TSharedRef<UserClass>&, typename TMemFunPtrType<false, UserClass, void (ParamTypes..., VarTypes...)>::Type, VarTypes...)
 	{
 		return {};
 	}
 	template <typename UserClass, typename... VarTypes>
-	inline static DerivedType CreateSP(const TSharedRef<UserClass, ESPMode::Fast>&, typename TMemFunPtrType<true, UserClass, void (ParamTypes..., VarTypes...)>::Type, VarTypes...)
+	inline static DerivedType CreateSP(const TSharedRef<UserClass>&, typename TMemFunPtrType<true, UserClass, void (ParamTypes..., VarTypes...)>::Type, VarTypes...)
 	{
 		return {};
 	}
