Mode template parameter exposed through the TSharedPtr type.

diff --git a/Engine/Source/Runtime/Core/Public/Templates/SharedPointer.h b/Engine/Source/Runtime/Core/Public/Templates/SharedPointer.h
--- a/Engine/Source/Runtime/Core/Public/Templates/SharedPointer.h
+++ b/Engine/Source/Runtime/Core/Public/Templates/SharedPointer.h
@@ -154,11 +154,12 @@ namespace UE4SharedPointer_Private
  * This shared reference will be conditionally thread-safe when the optional Mode template argument is set to ThreadSafe.
  */
 // NOTE: TSharedRef is an Unreal extension to standard smart pointer feature set
-template< class ObjectType, ESPMode Mode >
+template< class ObjectType, ESPMode InMode >
 class TSharedRef
 {
 public:
 	using ElementType = ObjectType;
+	static constexpr ESPMode Mode = InMode;
 
 	// NOTE: TSharedRef has no default constructor as it does not support empty references.  You must
 	//		 initialize your TSharedRef to a valid object at construction time.
@@ -555,11 +556,12 @@ struct FMakeReferenceTo<const void>
  * TSharedPtr is a non-intrusive reference-counted authoritative object pointer.  This shared pointer
  * will be conditionally thread-safe when the optional Mode template argument is set to ThreadSafe.
  */
-template< class ObjectType, ESPMode Mode >
+template< class ObjectType, ESPMode InMode >
 class TSharedPtr
 {
 public:
 	using ElementType = ObjectType;
+	static constexpr ESPMode Mode = InMode;
 
 	/**
 	 * Constructs an empty shared pointer
@@ -1028,11 +1030,12 @@ template<class ObjectType, ESPMode Mode> struct TIsZeroConstructType<TSharedPtr<
  * TWeakPtr is a non-intrusive reference-counted weak object pointer.  This weak pointer will be
  * conditionally thread-safe when the optional Mode template argument is set to ThreadSafe.
  */
-template< class ObjectType, ESPMode Mode >
+template< class ObjectType, ESPMode InMode >
 class TWeakPtr
 {
 public:
 	using ElementType = ObjectType;
+	static constexpr ESPMode Mode = InMode;
 
 	/** Constructs an empty TWeakPtr */
 	// NOTE: FNullTag parameter is an Unreal extension to standard shared_ptr behavior
