TSharedPtr is now thread-safe by default.
ESPMode::Fast is deprecated.
UNIQUENETID_ESPMODE and UNIQUENETID_CONSTRUCTORVIS removed.
Delegate creation functions (DelegateType::CreateSP, DelegateType::BindSP and DelegateType::AddSP) now work in all modes, making the DelegateType::*ThreadSafeSP functions redundant (retained for backwards compatibliity).

diff --git a/Engine/Plugins/FX/Niagara/Source/NiagaraEditor/Private/ViewModels/NiagaraEmitterViewModel.cpp b/Engine/Plugins/FX/Niagara/Source/NiagaraEditor/Private/ViewModels/NiagaraEmitterViewModel.cpp
--- a/Engine/Plugins/FX/Niagara/Source/NiagaraEditor/Private/ViewModels/NiagaraEmitterViewModel.cpp
+++ b/Engine/Plugins/FX/Niagara/Source/NiagaraEditor/Private/ViewModels/NiagaraEmitterViewModel.cpp
@@ -508,15 +508,15 @@ void FNiagaraEmitterViewModel::AddScriptEventHandlers()
 		{
 			UNiagaraScriptSource* ScriptSource = CastChecked<UNiagaraScriptSource>(Script->GetLatestSource());
 			FDelegateHandle OnGraphChangedHandle = ScriptSource->NodeGraph->AddOnGraphChangedHandler(
-				FOnGraphChanged::FDelegate::CreateSP<FNiagaraEmitterViewModel, const UNiagaraScript&>(this->AsShared(), &FNiagaraEmitterViewModel::ScriptGraphChanged, *Script));
+				FOnGraphChanged::FDelegate::CreateThreadSafeSP<FNiagaraEmitterViewModel, const UNiagaraScript&>(this->AsShared(), &FNiagaraEmitterViewModel::ScriptGraphChanged, *Script));
 			FDelegateHandle OnGraphNeedRecompileHandle = ScriptSource->NodeGraph->AddOnGraphNeedsRecompileHandler(
-				FOnGraphChanged::FDelegate::CreateSP<FNiagaraEmitterViewModel, const UNiagaraScript&>(this->AsShared(), &FNiagaraEmitterViewModel::ScriptGraphChanged, *Script));
+				FOnGraphChanged::FDelegate::CreateThreadSafeSP<FNiagaraEmitterViewModel, const UNiagaraScript&>(this->AsShared(), &FNiagaraEmitterViewModel::ScriptGraphChanged, *Script));
 
 			ScriptToOnGraphChangedHandleMap.Add(FObjectKey(Script), OnGraphChangedHandle);
 			ScriptToRecompileHandleMap.Add(FObjectKey(Script), OnGraphNeedRecompileHandle);
 
 			FDelegateHandle OnParameterStoreChangedHandle = Script->RapidIterationParameters.AddOnChangedHandler(
-				FNiagaraParameterStore::FOnChanged::FDelegate::CreateSP<FNiagaraEmitterViewModel, const FNiagaraParameterStore&, const UNiagaraScript&>(
+				FNiagaraParameterStore::FOnChanged::FDelegate::CreateThreadSafeSP<FNiagaraEmitterViewModel, const FNiagaraParameterStore&, const UNiagaraScript&>(
 					this->AsShared(), &FNiagaraEmitterViewModel::ScriptParameterStoreChanged, Script->RapidIterationParameters, *Script));
 			ScriptToOnParameterStoreChangedHandleMap.Add(FObjectKey(Script), OnParameterStoreChangedHandle);
 		}
diff --git a/Engine/Plugins/FX/Niagara/Source/NiagaraEditor/Private/ViewModels/NiagaraSystemViewModel.cpp b/Engine/Plugins/FX/Niagara/Source/NiagaraEditor/Private/ViewModels/NiagaraSystemViewModel.cpp
--- a/Engine/Plugins/FX/Niagara/Source/NiagaraEditor/Private/ViewModels/NiagaraSystemViewModel.cpp
+++ b/Engine/Plugins/FX/Niagara/Source/NiagaraEditor/Private/ViewModels/NiagaraSystemViewModel.cpp
@@ -2107,14 +2107,14 @@ void FNiagaraSystemViewModel::AddSystemEventHandlers()
 			if (Script != nullptr)
 			{
 				FDelegateHandle OnParameterStoreChangedHandle = Script->RapidIterationParameters.AddOnChangedHandler(
-					FNiagaraParameterStore::FOnChanged::FDelegate::CreateSP<FNiagaraSystemViewModel, const FNiagaraParameterStore&, const UNiagaraScript*>(
+					FNiagaraParameterStore::FOnChanged::FDelegate::CreateThreadSafeSP<FNiagaraSystemViewModel, const FNiagaraParameterStore&, const UNiagaraScript*>(
 						this->AsShared(), &FNiagaraSystemViewModel::SystemParameterStoreChanged, Script->RapidIterationParameters, Script));
 				ScriptToOnParameterStoreChangedHandleMap.Add(FObjectKey(Script), OnParameterStoreChangedHandle);
 			}
 		}
 
 		UserParameterStoreChangedHandle = System->GetExposedParameters().AddOnChangedHandler(
-			FNiagaraParameterStore::FOnChanged::FDelegate::CreateSP<FNiagaraSystemViewModel, const FNiagaraParameterStore&, const UNiagaraScript*>(
+			FNiagaraParameterStore::FOnChanged::FDelegate::CreateThreadSafeSP<FNiagaraSystemViewModel, const FNiagaraParameterStore&, const UNiagaraScript*>(
 				this->AsShared(), &FNiagaraSystemViewModel::SystemParameterStoreChanged, System->GetExposedParameters(), nullptr));
 
 		SystemScriptGraphChangedHandle = SystemScriptViewModel->GetGraphViewModel()->GetGraph()->AddOnGraphChangedHandler(
diff --git a/Engine/Plugins/Online/OnlineSubsystem/Source/Public/OnlineSubsystemTypes.h b/Engine/Plugins/Online/OnlineSubsystem/Source/Public/OnlineSubsystemTypes.h
--- a/Engine/Plugins/Online/OnlineSubsystem/Source/Public/OnlineSubsystemTypes.h
+++ b/Engine/Plugins/Online/OnlineSubsystem/Source/Public/OnlineSubsystemTypes.h
@@ -815,8 +815,8 @@ public:
 // placeholder "type" until we can make FUniqueNetIdString sufficiently abstract
 static FName NAME_Unset = TEXT("UNSET");
 
-using FUniqueNetIdStringRef = TSharedRef<const class FUniqueNetIdString, UNIQUENETID_ESPMODE>;
-using FUniqueNetIdStringPtr = TSharedPtr<const class FUniqueNetIdString, UNIQUENETID_ESPMODE>;
+using FUniqueNetIdStringRef = TSharedRef<const class FUniqueNetIdString>;
+using FUniqueNetIdStringPtr = TSharedPtr<const class FUniqueNetIdString>;
 
 /**
  * Unique net id wrapper for a string
@@ -832,7 +832,7 @@ public:
 	template<typename... TArgs>
 	static FUniqueNetIdStringRef Create(TArgs&&... Args)
 	{
-		return MakeShared<FUniqueNetIdString, UNIQUENETID_ESPMODE>(Forward<TArgs>(Args)...);
+		return MakeShared<FUniqueNetIdString>(Forward<TArgs>(Args)...);
 	}
 
 	/** Allow MakeShared to see private constructors */
@@ -891,8 +891,7 @@ public:
 		return ::GetTypeHash(A.UniqueNetIdStr);
 	}
 
-	// Public constructors are deprecated, but we are piggybacking on the ESPMode deprecation and making these public in ESPMode::Fast and protected in ESPMode::ThreadSafe
-UNIQUENETID_CONSTRUCTORVIS:
+public:
 	FUniqueNetIdString() = default;
 
 	/**
@@ -940,15 +939,15 @@ UNIQUENETID_CONSTRUCTORVIS:
 
 
 #define TEMP_UNIQUENETIDSTRING_SUBCLASS(SUBCLASSNAME, TYPE) \
-using SUBCLASSNAME##Ptr = TSharedPtr<const class SUBCLASSNAME, UNIQUENETID_ESPMODE>; \
-using SUBCLASSNAME##Ref = TSharedRef<const class SUBCLASSNAME, UNIQUENETID_ESPMODE>; \
+using SUBCLASSNAME##Ptr = TSharedPtr<const class SUBCLASSNAME>; \
+using SUBCLASSNAME##Ref = TSharedRef<const class SUBCLASSNAME>; \
 class SUBCLASSNAME : public FUniqueNetIdString \
 { \
 public: \
 	template<typename... TArgs> \
 	static SUBCLASSNAME##Ref Create(TArgs&&... Args) \
 	{ \
-		return MakeShared<SUBCLASSNAME, UNIQUENETID_ESPMODE>(Forward<TArgs>(Args)...); \
+		return MakeShared<SUBCLASSNAME>(Forward<TArgs>(Args)...); \
 	} \
 	friend class SharedPointerInternals::TIntrusiveReferenceController<SUBCLASSNAME>; \
 	static SUBCLASSNAME##Ref Cast(const FUniqueNetIdRef& InNetId) \
@@ -983,7 +982,7 @@ public: \
 		static const SUBCLASSNAME##Ref EmptyId(Create()); \
 		return EmptyId; \
 	} \
-UNIQUENETID_CONSTRUCTORVIS: \
+public: \
 	SUBCLASSNAME() \
 		: FUniqueNetIdString() \
 	{ \
diff --git a/Engine/Plugins/Online/OnlineSubsystemEOS/Source/OnlineSubsystemEOS/Private/OnlineSubsystemEOSTypes.h b/Engine/Plugins/Online/OnlineSubsystemEOS/Source/OnlineSubsystemEOS/Private/OnlineSubsystemEOSTypes.h
--- a/Engine/Plugins/Online/OnlineSubsystemEOS/Source/OnlineSubsystemEOS/Private/OnlineSubsystemEOSTypes.h
+++ b/Engine/Plugins/Online/OnlineSubsystemEOS/Source/OnlineSubsystemEOS/Private/OnlineSubsystemEOSTypes.h
@@ -32,8 +32,8 @@ class FOnlineSubsystemEOS;
 #define ID_HALF_BYTE_SIZE 16
 #define EOS_ID_BYTE_SIZE (ID_HALF_BYTE_SIZE * 2)
 
-typedef TSharedPtr<const class FUniqueNetIdEOS, UNIQUENETID_ESPMODE> FUniqueNetIdEOSPtr;
-typedef TSharedRef<const class FUniqueNetIdEOS, UNIQUENETID_ESPMODE> FUniqueNetIdEOSRef;
+typedef TSharedPtr<const class FUniqueNetIdEOS> FUniqueNetIdEOSPtr;
+typedef TSharedRef<const class FUniqueNetIdEOS> FUniqueNetIdEOSRef;
 
 /**
  * Unique net id wrapper for a EOS account ids. The underlying string is a combination
@@ -46,7 +46,7 @@ public:
 	template<typename... TArgs>
 	static FUniqueNetIdEOSRef Create(TArgs&&... Args)
 	{
-		return MakeShared<FUniqueNetIdEOS, UNIQUENETID_ESPMODE>(Forward<TArgs>(Args)...);
+		return MakeShared<FUniqueNetIdEOS>(Forward<TArgs>(Args)...);
 	}
 
 	/** Allow MakeShared to see private constructors */
diff --git a/Engine/Plugins/Online/OnlineSubsystemEOS/Source/OnlineSubsystemEOSPlus/Private/OnlineUserEOSPlus.h b/Engine/Plugins/Online/OnlineSubsystemEOS/Source/OnlineSubsystemEOSPlus/Private/OnlineUserEOSPlus.h
--- a/Engine/Plugins/Online/OnlineSubsystemEOS/Source/OnlineSubsystemEOSPlus/Private/OnlineUserEOSPlus.h
+++ b/Engine/Plugins/Online/OnlineSubsystemEOS/Source/OnlineSubsystemEOSPlus/Private/OnlineUserEOSPlus.h
@@ -17,8 +17,8 @@ class FOnlineSubsystemEOSPlus;
 
 #define BASE_NETID_TYPE_SIZE 1
 
-using FUniqueNetIdEOSPlusPtr = TSharedPtr<const class FUniqueNetIdEOSPlus, UNIQUENETID_ESPMODE>;
-using FUniqueNetIdEOSPlusRef = TSharedRef<const class FUniqueNetIdEOSPlus, UNIQUENETID_ESPMODE>;
+using FUniqueNetIdEOSPlusPtr = TSharedPtr<const class FUniqueNetIdEOSPlus>;
+using FUniqueNetIdEOSPlusRef = TSharedRef<const class FUniqueNetIdEOSPlus>;
 
 /**
  * Unique net id wrapper for a EOS plus another account id. The underlying string is a combination
@@ -31,7 +31,7 @@ public:
 	template<typename... TArgs>
 	static FUniqueNetIdEOSPlusRef Create(TArgs&&... Args)
 	{
-		return MakeShared<FUniqueNetIdEOSPlus, UNIQUENETID_ESPMODE>(Forward<TArgs>(Args)...);
+		return MakeShared<FUniqueNetIdEOSPlus>(Forward<TArgs>(Args)...);
 	}
 
 	/** Allow MakeShared to see private constructors */
@@ -61,8 +61,8 @@ private:
 	TArray<uint8> RawBytes;
 };
 
-using FUniqueNetIdBinaryPtr = TSharedPtr<const class FUniqueNetIdBinary, UNIQUENETID_ESPMODE>;
-using FUniqueNetIdBinaryRef = TSharedRef<const class FUniqueNetIdBinary, UNIQUENETID_ESPMODE>;
+using FUniqueNetIdBinaryPtr = TSharedPtr<const class FUniqueNetIdBinary>;
+using FUniqueNetIdBinaryRef = TSharedRef<const class FUniqueNetIdBinary>;
 
 class FUniqueNetIdBinary :
 	public FUniqueNetId
@@ -73,7 +73,7 @@ public:
 	template<typename... TArgs>
 	static FUniqueNetIdBinaryRef Create(TArgs&&... Args)
 	{
-		return MakeShared<FUniqueNetIdBinary, UNIQUENETID_ESPMODE>(Forward<TArgs>(Args)...);
+		return MakeShared<FUniqueNetIdBinary>(Forward<TArgs>(Args)...);
 	}
 
 	/** Allow MakeShared to see private constructors */
diff --git a/Engine/Plugins/Online/OnlineSubsystemFacebook/Source/Private/OnlineSubsystemFacebookTypes.h b/Engine/Plugins/Online/OnlineSubsystemFacebook/Source/Private/OnlineSubsystemFacebookTypes.h
--- a/Engine/Plugins/Online/OnlineSubsystemFacebook/Source/Private/OnlineSubsystemFacebookTypes.h
+++ b/Engine/Plugins/Online/OnlineSubsystemFacebook/Source/Private/OnlineSubsystemFacebookTypes.h
@@ -10,8 +10,8 @@
 #define PICTURE_DATA_SILHOUETTE "is_silhouette"
 #define PICTURE_DATA_URL "url"
 
-using FUniqueNetIdFacebookPtr = TSharedPtr<const class FUniqueNetIdFacebook, UNIQUENETID_ESPMODE>;
-using FUniqueNetIdFacebookRef = TSharedRef<const class FUniqueNetIdFacebook, UNIQUENETID_ESPMODE>;
+using FUniqueNetIdFacebookPtr = TSharedPtr<const class FUniqueNetIdFacebook>;
+using FUniqueNetIdFacebookRef = TSharedRef<const class FUniqueNetIdFacebook>;
 
 /**
  * Facebook specific implementation of the unique net id
@@ -27,7 +27,7 @@ public:
 	template<typename... TArgs>
 	static FUniqueNetIdFacebookRef Create(TArgs&&... Args)
 	{
-		return MakeShared<FUniqueNetIdFacebook, UNIQUENETID_ESPMODE>(Forward<TArgs>(Args)...);
+		return MakeShared<FUniqueNetIdFacebook>(Forward<TArgs>(Args)...);
 	}
 
 	/** Allow MakeShared to see private constructors */
diff --git a/Engine/Plugins/Online/OnlineSubsystemOculus/Source/Private/OnlineSubsystemOculusTypes.h b/Engine/Plugins/Online/OnlineSubsystemOculus/Source/Private/OnlineSubsystemOculusTypes.h
--- a/Engine/Plugins/Online/OnlineSubsystemOculus/Source/Private/OnlineSubsystemOculusTypes.h
+++ b/Engine/Plugins/Online/OnlineSubsystemOculus/Source/Private/OnlineSubsystemOculusTypes.h
@@ -7,8 +7,8 @@
 #include "OnlineSubsystemOculusPackage.h"
 #include "OVR_Platform.h"
 
-using FUniqueNetIdOculusPtr = TSharedPtr<const class FUniqueNetIdOculus, UNIQUENETID_ESPMODE>;
-using FUniqueNetIdOculusRef = TSharedRef<const class FUniqueNetIdOculus, UNIQUENETID_ESPMODE>;
+using FUniqueNetIdOculusPtr = TSharedPtr<const class FUniqueNetIdOculus>;
+using FUniqueNetIdOculusRef = TSharedRef<const class FUniqueNetIdOculus>;
 
 class FUniqueNetIdOculus : public FUniqueNetId {
 private:
@@ -34,7 +34,7 @@ public:
 	template<typename... TArgs>
 	static FUniqueNetIdOculusRef Create(TArgs&&... Args)
 	{
-		return MakeShared<FUniqueNetIdOculus, UNIQUENETID_ESPMODE>(Forward<TArgs>(Args)...);
+		return MakeShared<FUniqueNetIdOculus>(Forward<TArgs>(Args)...);
 	}
 
 	/** Allow MakeShared to see private constructors */
diff --git a/Engine/Plugins/Online/OnlineSubsystemSteam/Source/Private/OnlineSubsystemSteamTypes.h b/Engine/Plugins/Online/OnlineSubsystemSteam/Source/Private/OnlineSubsystemSteamTypes.h
--- a/Engine/Plugins/Online/OnlineSubsystemSteam/Source/Private/OnlineSubsystemSteamTypes.h
+++ b/Engine/Plugins/Online/OnlineSubsystemSteam/Source/Private/OnlineSubsystemSteamTypes.h
@@ -56,8 +56,8 @@ namespace ESteamSession
 	}
 }
 
-using FUniqueNetIdSteamRef = TSharedRef<const class FUniqueNetIdSteam, UNIQUENETID_ESPMODE>;
-using FUniqueNetIdSteamPtr = TSharedPtr<const class FUniqueNetIdSteam, UNIQUENETID_ESPMODE>;
+using FUniqueNetIdSteamRef = TSharedRef<const class FUniqueNetIdSteam>;
+using FUniqueNetIdSteamPtr = TSharedPtr<const class FUniqueNetIdSteam>;
 
 /**
  * Steam specific implementation of the unique net id
@@ -73,7 +73,7 @@ public:
 	template<typename... TArgs>
 	static FUniqueNetIdSteamRef Create(TArgs&&... Args)
 	{
-		return MakeShared<FUniqueNetIdSteam, UNIQUENETID_ESPMODE>(Forward<TArgs>(Args)...);
+		return MakeShared<FUniqueNetIdSteam>(Forward<TArgs>(Args)...);
 	}
 
 	static const FUniqueNetIdSteam& Cast(const FUniqueNetId& NetId)
diff --git a/Engine/Plugins/Online/OnlineSubsystemTencent/Source/Public/OnlineSubsystemTencentTypes.h b/Engine/Plugins/Online/OnlineSubsystemTencent/Source/Public/OnlineSubsystemTencentTypes.h
--- a/Engine/Plugins/Online/OnlineSubsystemTencent/Source/Public/OnlineSubsystemTencentTypes.h
+++ b/Engine/Plugins/Online/OnlineSubsystemTencent/Source/Public/OnlineSubsystemTencentTypes.h
@@ -27,8 +27,8 @@ using TRailIdMap = TMap<rail::RailID, ValueType, FDefaultSetAllocator, FRailIdKe
 /** Key/Value pairs stored per user on the Rail platform */
 using FMetadataPropertiesRail = FOnlineKeyValuePairs<FString, FVariantData>;
 
-using FUniqueNetIdRailRef = TSharedRef<const class FUniqueNetIdRail, UNIQUENETID_ESPMODE>;
-using FUniqueNetIdRailPtr = TSharedPtr<const class FUniqueNetIdRail, UNIQUENETID_ESPMODE>;
+using FUniqueNetIdRailRef = TSharedRef<const class FUniqueNetIdRail>;
+using FUniqueNetIdRailPtr = TSharedPtr<const class FUniqueNetIdRail>;
 
 /**
  * Rail specific implementation of the unique net id
@@ -47,7 +47,7 @@ public:
 	template<typename... TArgs>
 	static FUniqueNetIdRailRef Create(TArgs&&... Args)
 	{
-		return MakeShared<FUniqueNetIdRail, UNIQUENETID_ESPMODE>(Forward<TArgs>(Args)...);
+		return MakeShared<FUniqueNetIdRail>(Forward<TArgs>(Args)...);
 	}
 
 	/** Allow MakeShared to see private constructors */
diff --git a/Engine/Source/Runtime/AudioAnalyzer/Public/IAudioAnalyzerInterface.h b/Engine/Source/Runtime/AudioAnalyzer/Public/IAudioAnalyzerInterface.h
--- a/Engine/Source/Runtime/AudioAnalyzer/Public/IAudioAnalyzerInterface.h
+++ b/Engine/Source/Runtime/AudioAnalyzer/Public/IAudioAnalyzerInterface.h
@@ -79,7 +79,7 @@ namespace Audio
 		virtual TUniquePtr<IAnalyzerResult> NewResult() const = 0;
 
 		// Convenience function to create a new shared result by calling NewResult.
-		template<ESPMode Mode = ESPMode::Fast>
+		template<ESPMode Mode = ESPMode::ThreadSafe>
 		TSharedPtr<IAnalyzerResult, Mode> NewResultShared() const
 		{
 			TUniquePtr<Audio::IAnalyzerResult> Result = NewResult();
diff --git a/Engine/Source/Runtime/AudioAnalyzer/Public/IAudioAnalyzerNRTInterface.h b/Engine/Source/Runtime/AudioAnalyzer/Public/IAudioAnalyzerNRTInterface.h
--- a/Engine/Source/Runtime/AudioAnalyzer/Public/IAudioAnalyzerNRTInterface.h
+++ b/Engine/Source/Runtime/AudioAnalyzer/Public/IAudioAnalyzerNRTInterface.h
@@ -116,7 +116,7 @@ namespace Audio
 		virtual TUniquePtr<IAnalyzerNRTResult> NewResult() const = 0;
 
 		// Convenience function to create a new shared result by calling NewResult.
-		template<ESPMode Mode = ESPMode::Fast>
+		template<ESPMode Mode = ESPMode::ThreadSafe>
 		TSharedPtr<IAnalyzerNRTResult, Mode> NewResultShared() const
 		{
 			TUniquePtr<Audio::IAnalyzerNRTResult> Result = NewResult();
diff --git a/Engine/Source/Runtime/Core/Private/Tests/Misc/TypeContainerTest.cpp b/Engine/Source/Runtime/Core/Private/Tests/Misc/TypeContainerTest.cpp
--- a/Engine/Source/Runtime/Core/Private/Tests/Misc/TypeContainerTest.cpp
+++ b/Engine/Source/Runtime/Core/Private/Tests/Misc/TypeContainerTest.cpp
@@ -41,7 +41,7 @@ struct FStrawberry : public IBerry
 	virtual FString Name() override { return TEXT("Strawberry"); }
 };
 
-template<ESPMode Mode = ESPMode::Fast>
+template<ESPMode Mode = ESPMode::ThreadSafe>
 struct ISmoothie
 {
 	virtual ~ISmoothie() { }
@@ -49,7 +49,7 @@ struct ISmoothie
 	virtual TSharedRef<IFruit, Mode> GetFruit() = 0;
 };
 
-template<ESPMode Mode = ESPMode::Fast>
+template<ESPMode Mode = ESPMode::ThreadSafe>
 struct TSmoothie : public ISmoothie<Mode>
 {
 	TSmoothie(TSharedRef<IFruit, Mode> InFruit, TSharedRef<IBerry, Mode> InBerry) : Berry(InBerry), Fruit(InFruit) { }
@@ -60,7 +60,7 @@ struct TSmoothie : public ISmoothie<Mode>
 	TSharedRef<IFruit, Mode> Fruit;
 };
 
-template<ESPMode Mode = ESPMode::Fast>
+template<ESPMode Mode = ESPMode::ThreadSafe>
 struct TTwoSmoothies
 {
 	TSharedPtr<ISmoothie<Mode>, Mode> One;
@@ -75,9 +75,6 @@ DECLARE_DELEGATE_RetVal_TwoParams(TSharedRef<ISmoothie<>>, FSmoothieFactoryDeleg
 Expose_TNameOf(IBerry)
 Expose_TNameOf(IFruit)
 Expose_TNameOf(ISmoothie<>)
-#if !FORCE_THREADSAFE_SHAREDPTRS
-Expose_TNameOf(ISmoothie<ESPMode::ThreadSafe>)
-#endif
 
 
 /* Tests
diff --git a/Engine/Source/Runtime/Core/Public/Delegates/DelegateSignatureImpl.inl b/Engine/Source/Runtime/Core/Public/Delegates/DelegateSignatureImpl.inl
--- a/Engine/Source/Runtime/Core/Public/Delegates/DelegateSignatureImpl.inl
+++ b/Engine/Source/Runtime/Core/Public/Delegates/DelegateSignatureImpl.inl
@@ -74,16 +74,16 @@ public:
 	template <typename UserClass, typename... VarTypes> using TConstMethodPtr = typename TMemFunPtrType<true,  UserClass, RetValType(ParamTypes..., VarTypes...)>::Type;
 
 	/** Declare the user's "fast" shared pointer-based delegate instance types. */
-	template <typename UserClass                                                                            > struct TSPMethodDelegate                 : TBaseSPMethodDelegateInstance<false, UserClass, ESPMode::Fast, FuncType, UserPolicy                                        > { typedef TBaseSPMethodDelegateInstance<false, UserClass, ESPMode::Fast, FuncType, UserPolicy                                        > Super; TSPMethodDelegate                (const TSharedRef< UserClass, ESPMode::Fast >& InUserObject, typename Super::FMethodPtr InMethodPtr                                                            ) : Super(InUserObject, InMethodPtr                        ) {} };
-	template <typename UserClass                                                                            > struct TSPMethodDelegate_Const           : TBaseSPMethodDelegateInstance<true , UserClass, ESPMode::Fast, FuncType, UserPolicy                                        > { typedef TBaseSPMethodDelegateInstance<true , UserClass, ESPMode::Fast, FuncType, UserPolicy                                        > Super; TSPMethodDelegate_Const          (const TSharedRef< UserClass, ESPMode::Fast >& InUserObject, typename Super::FMethodPtr InMethodPtr                                                            ) : Super(InUserObject, InMethodPtr                        ) {} };
-	template <typename UserClass, typename Var1Type                                                         > struct TSPMethodDelegate_OneVar          : TBaseSPMethodDelegateInstance<false, UserClass, ESPMode::Fast, FuncType, UserPolicy, Var1Type                              > { typedef TBaseSPMethodDelegateInstance<false, UserClass, ESPMode::Fast, FuncType, UserPolicy, Var1Type                              > Super; TSPMethodDelegate_OneVar         (const TSharedRef< UserClass, ESPMode::Fast >& InUserObject, typename Super::FMethodPtr InMethodPtr, Var1Type Var1                                             ) : Super(InUserObject, InMethodPtr, Var1                  ) {} };
-	template <typename UserClass, typename Var1Type                                                         > struct TSPMethodDelegate_OneVar_Const    : TBaseSPMethodDelegateInstance<true , UserClass, ESPMode::Fast, FuncType, UserPolicy, Var1Type                              > { typedef TBaseSPMethodDelegateInstance<true , UserClass, ESPMode::Fast, FuncType, UserPolicy, Var1Type                              > Super; TSPMethodDelegate_OneVar_Const   (const TSharedRef< UserClass, ESPMode::Fast >& InUserObject, typename Super::FMethodPtr InMethodPtr, Var1Type Var1                                             ) : Super(InUserObject, InMethodPtr, Var1                  ) {} };
-	template <typename UserClass, typename Var1Type, typename Var2Type                                      > struct TSPMethodDelegate_TwoVars         : TBaseSPMethodDelegateInstance<false, UserClass, ESPMode::Fast, FuncType, UserPolicy, Var1Type, Var2Type                    > { typedef TBaseSPMethodDelegateInstance<false, UserClass, ESPMode::Fast, FuncType, UserPolicy, Var1Type, Var2Type                    > Super; TSPMethodDelegate_TwoVars        (const TSharedRef< UserClass, ESPMode::Fast >& InUserObject, typename Super::FMethodPtr InMethodPtr, Var1Type Var1, Var2Type Var2                              ) : Super(InUserObject, InMethodPtr, Var1, Var2            ) {} };
-	template <typename UserClass, typename Var1Type, typename Var2Type                                      > struct TSPMethodDelegate_TwoVars_Const   : TBaseSPMethodDelegateInstance<true , UserClass, ESPMode::Fast, FuncType, UserPolicy, Var1Type, Var2Type                    > { typedef TBaseSPMethodDelegateInstance<true , UserClass, ESPMode::Fast, FuncType, UserPolicy, Var1Type, Var2Type                    > Super; TSPMethodDelegate_TwoVars_Const  (const TSharedRef< UserClass, ESPMode::Fast >& InUserObject, typename Super::FMethodPtr InMethodPtr, Var1Type Var1, Var2Type Var2                              ) : Super(InUserObject, InMethodPtr, Var1, Var2            ) {} };
-	template <typename UserClass, typename Var1Type, typename Var2Type, typename Var3Type                   > struct TSPMethodDelegate_ThreeVars       : TBaseSPMethodDelegateInstance<false, UserClass, ESPMode::Fast, FuncType, UserPolicy, Var1Type, Var2Type, Var3Type          > { typedef TBaseSPMethodDelegateInstance<false, UserClass, ESPMode::Fast, FuncType, UserPolicy, Var1Type, Var2Type, Var3Type          > Super; TSPMethodDelegate_ThreeVars      (const TSharedRef< UserClass, ESPMode::Fast >& InUserObject, typename Super::FMethodPtr InMethodPtr, Var1Type Var1, Var2Type Var2, Var3Type Var3               ) : Super(InUserObject, InMethodPtr, Var1, Var2, Var3      ) {} };
-	template <typename UserClass, typename Var1Type, typename Var2Type, typename Var3Type                   > struct TSPMethodDelegate_ThreeVars_Const : TBaseSPMethodDelegateInstance<true , UserClass, ESPMode::Fast, FuncType, UserPolicy, Var1Type, Var2Type, Var3Type          > { typedef TBaseSPMethodDelegateInstance<true , UserClass, ESPMode::Fast, FuncType, UserPolicy, Var1Type, Var2Type, Var3Type          > Super; TSPMethodDelegate_ThreeVars_Const(const TSharedRef< UserClass, ESPMode::Fast >& InUserObject, typename Super::FMethodPtr InMethodPtr, Var1Type Var1, Var2Type Var2, Var3Type Var3               ) : Super(InUserObject, InMethodPtr, Var1, Var2, Var3      ) {} };
-	template <typename UserClass, typename Var1Type, typename Var2Type, typename Var3Type, typename Var4Type> struct TSPMethodDelegate_FourVars        : TBaseSPMethodDelegateInstance<false, UserClass, ESPMode::Fast, FuncType, UserPolicy, Var1Type, Var2Type, Var3Type, Var4Type> { typedef TBaseSPMethodDelegateInstance<false, UserClass, ESPMode::Fast, FuncType, UserPolicy, Var1Type, Var2Type, Var3Type, Var4Type> Super; TSPMethodDelegate_FourVars       (const TSharedRef< UserClass, ESPMode::Fast >& InUserObject, typename Super::FMethodPtr InMethodPtr, Var1Type Var1, Var2Type Var2, Var3Type Var3, Var4Type Var4) : Super(InUserObject, InMethodPtr, Var1, Var2, Var3, Var4) {} };
-	template <typename UserClass, typename Var1Type, typename Var2Type, typename Var3Type, typename Var4Type> struct TSPMethodDelegate_FourVars_Const  : TBaseSPMethodDelegateInstance<true , UserClass, ESPMode::Fast, FuncType, UserPolicy, Var1Type, Var2Type, Var3Type, Var4Type> { typedef TBaseSPMethodDelegateInstance<true , UserClass, ESPMode::Fast, FuncType, UserPolicy, Var1Type, Var2Type, Var3Type, Var4Type> Super; TSPMethodDelegate_FourVars_Const (const TSharedRef< UserClass, ESPMode::Fast >& InUserObject, typename Super::FMethodPtr InMethodPtr, Var1Type Var1, Var2Type Var2, Var3Type Var3, Var4Type Var4) : Super(InUserObject, InMethodPtr, Var1, Var2, Var3, Var4) {} };
+	template <typename UserClass                                                                            > struct TSPMethodDelegate                 : TBaseSPMethodDelegateInstance<false, UserClass, ESPMode::ThreadSafe, FuncType, UserPolicy                                        > { typedef TBaseSPMethodDelegateInstance<false, UserClass, ESPMode::ThreadSafe, FuncType, UserPolicy                                        > Super; TSPMethodDelegate                (const TSharedRef< UserClass, ESPMode::ThreadSafe >& InUserObject, typename Super::FMethodPtr InMethodPtr                                                            ) : Super(InUserObject, InMethodPtr                        ) {} };
+	template <typename UserClass                                                                            > struct TSPMethodDelegate_Const           : TBaseSPMethodDelegateInstance<true , UserClass, ESPMode::ThreadSafe, FuncType, UserPolicy                                        > { typedef TBaseSPMethodDelegateInstance<true , UserClass, ESPMode::ThreadSafe, FuncType, UserPolicy                                        > Super; TSPMethodDelegate_Const          (const TSharedRef< UserClass, ESPMode::ThreadSafe >& InUserObject, typename Super::FMethodPtr InMethodPtr                                                            ) : Super(InUserObject, InMethodPtr                        ) {} };
+	template <typename UserClass, typename Var1Type                                                         > struct TSPMethodDelegate_OneVar          : TBaseSPMethodDelegateInstance<false, UserClass, ESPMode::ThreadSafe, FuncType, UserPolicy, Var1Type                              > { typedef TBaseSPMethodDelegateInstance<false, UserClass, ESPMode::ThreadSafe, FuncType, UserPolicy, Var1Type                              > Super; TSPMethodDelegate_OneVar         (const TSharedRef< UserClass, ESPMode::ThreadSafe >& InUserObject, typename Super::FMethodPtr InMethodPtr, Var1Type Var1                                             ) : Super(InUserObject, InMethodPtr, Var1                  ) {} };
+	template <typename UserClass, typename Var1Type                                                         > struct TSPMethodDelegate_OneVar_Const    : TBaseSPMethodDelegateInstance<true , UserClass, ESPMode::ThreadSafe, FuncType, UserPolicy, Var1Type                              > { typedef TBaseSPMethodDelegateInstance<true , UserClass, ESPMode::ThreadSafe, FuncType, UserPolicy, Var1Type                              > Super; TSPMethodDelegate_OneVar_Const   (const TSharedRef< UserClass, ESPMode::ThreadSafe >& InUserObject, typename Super::FMethodPtr InMethodPtr, Var1Type Var1                                             ) : Super(InUserObject, InMethodPtr, Var1                  ) {} };
+	template <typename UserClass, typename Var1Type, typename Var2Type                                      > struct TSPMethodDelegate_TwoVars         : TBaseSPMethodDelegateInstance<false, UserClass, ESPMode::ThreadSafe, FuncType, UserPolicy, Var1Type, Var2Type                    > { typedef TBaseSPMethodDelegateInstance<false, UserClass, ESPMode::ThreadSafe, FuncType, UserPolicy, Var1Type, Var2Type                    > Super; TSPMethodDelegate_TwoVars        (const TSharedRef< UserClass, ESPMode::ThreadSafe >& InUserObject, typename Super::FMethodPtr InMethodPtr, Var1Type Var1, Var2Type Var2                              ) : Super(InUserObject, InMethodPtr, Var1, Var2            ) {} };
+	template <typename UserClass, typename Var1Type, typename Var2Type                                      > struct TSPMethodDelegate_TwoVars_Const   : TBaseSPMethodDelegateInstance<true , UserClass, ESPMode::ThreadSafe, FuncType, UserPolicy, Var1Type, Var2Type                    > { typedef TBaseSPMethodDelegateInstance<true , UserClass, ESPMode::ThreadSafe, FuncType, UserPolicy, Var1Type, Var2Type                    > Super; TSPMethodDelegate_TwoVars_Const  (const TSharedRef< UserClass, ESPMode::ThreadSafe >& InUserObject, typename Super::FMethodPtr InMethodPtr, Var1Type Var1, Var2Type Var2                              ) : Super(InUserObject, InMethodPtr, Var1, Var2            ) {} };
+	template <typename UserClass, typename Var1Type, typename Var2Type, typename Var3Type                   > struct TSPMethodDelegate_ThreeVars       : TBaseSPMethodDelegateInstance<false, UserClass, ESPMode::ThreadSafe, FuncType, UserPolicy, Var1Type, Var2Type, Var3Type          > { typedef TBaseSPMethodDelegateInstance<false, UserClass, ESPMode::ThreadSafe, FuncType, UserPolicy, Var1Type, Var2Type, Var3Type          > Super; TSPMethodDelegate_ThreeVars      (const TSharedRef< UserClass, ESPMode::ThreadSafe >& InUserObject, typename Super::FMethodPtr InMethodPtr, Var1Type Var1, Var2Type Var2, Var3Type Var3               ) : Super(InUserObject, InMethodPtr, Var1, Var2, Var3      ) {} };
+	template <typename UserClass, typename Var1Type, typename Var2Type, typename Var3Type                   > struct TSPMethodDelegate_ThreeVars_Const : TBaseSPMethodDelegateInstance<true , UserClass, ESPMode::ThreadSafe, FuncType, UserPolicy, Var1Type, Var2Type, Var3Type          > { typedef TBaseSPMethodDelegateInstance<true , UserClass, ESPMode::ThreadSafe, FuncType, UserPolicy, Var1Type, Var2Type, Var3Type          > Super; TSPMethodDelegate_ThreeVars_Const(const TSharedRef< UserClass, ESPMode::ThreadSafe >& InUserObject, typename Super::FMethodPtr InMethodPtr, Var1Type Var1, Var2Type Var2, Var3Type Var3               ) : Super(InUserObject, InMethodPtr, Var1, Var2, Var3      ) {} };
+	template <typename UserClass, typename Var1Type, typename Var2Type, typename Var3Type, typename Var4Type> struct TSPMethodDelegate_FourVars        : TBaseSPMethodDelegateInstance<false, UserClass, ESPMode::ThreadSafe, FuncType, UserPolicy, Var1Type, Var2Type, Var3Type, Var4Type> { typedef TBaseSPMethodDelegateInstance<false, UserClass, ESPMode::ThreadSafe, FuncType, UserPolicy, Var1Type, Var2Type, Var3Type, Var4Type> Super; TSPMethodDelegate_FourVars       (const TSharedRef< UserClass, ESPMode::ThreadSafe >& InUserObject, typename Super::FMethodPtr InMethodPtr, Var1Type Var1, Var2Type Var2, Var3Type Var3, Var4Type Var4) : Super(InUserObject, InMethodPtr, Var1, Var2, Var3, Var4) {} };
+	template <typename UserClass, typename Var1Type, typename Var2Type, typename Var3Type, typename Var4Type> struct TSPMethodDelegate_FourVars_Const  : TBaseSPMethodDelegateInstance<true , UserClass, ESPMode::ThreadSafe, FuncType, UserPolicy, Var1Type, Var2Type, Var3Type, Var4Type> { typedef TBaseSPMethodDelegateInstance<true , UserClass, ESPMode::ThreadSafe, FuncType, UserPolicy, Var1Type, Var2Type, Var3Type, Var4Type> Super; TSPMethodDelegate_FourVars_Const (const TSharedRef< UserClass, ESPMode::ThreadSafe >& InUserObject, typename Super::FMethodPtr InMethodPtr, Var1Type Var1, Var2Type Var2, Var3Type Var3, Var4Type Var4) : Super(InUserObject, InMethodPtr, Var1, Var2, Var3, Var4) {} };
 
 	/** Declare the user's "thread-safe" shared pointer-based delegate instance types. */
 	template <typename UserClass                                                                            > struct TThreadSafeSPMethodDelegate                 : TBaseSPMethodDelegateInstance<false, UserClass, ESPMode::ThreadSafe, FuncType, UserPolicy                                        > { typedef TBaseSPMethodDelegateInstance<false, UserClass, ESPMode::ThreadSafe, FuncType, UserPolicy                                        > Super; TThreadSafeSPMethodDelegate                (const TSharedRef<UserClass, ESPMode::ThreadSafe>& InUserObject, typename Super::FMethodPtr InMethodPtr                                                            ) : Super(InUserObject, InMethodPtr                        ) {} };
@@ -196,30 +196,30 @@ public:
 	}
 
 	/**
-	 * Static: Creates a shared pointer-based (fast, not thread-safe) member function delegate.
+	 * Static: Creates a shared pointer-based member function delegate.
 	 *
 	 * Shared pointer delegates keep a weak reference to your object.
 	 * You can use ExecuteIfBound() to call them.
 	 */
-	template <typename UserClass, typename... VarTypes>
-	UE_NODISCARD inline static TDelegate<RetValType(ParamTypes...), UserPolicy> CreateSP(const TSharedRef<UserClass, ESPMode::Fast>& InUserObjectRef, typename TMemFunPtrType<false, UserClass, RetValType (ParamTypes..., VarTypes...)>::Type InFunc, VarTypes... Vars)
+	template <typename UserClass, ESPMode Mode, typename... VarTypes>
+	UE_NODISCARD inline static TDelegate<RetValType(ParamTypes...), UserPolicy> CreateSP(const TSharedRef<UserClass, Mode>& InUserObjectRef, typename TMemFunPtrType<false, UserClass, RetValType (ParamTypes..., VarTypes...)>::Type InFunc, VarTypes... Vars)
 	{
 		static_assert(!TIsConst<UserClass>::Value, "Attempting to bind a delegate with a const object pointer and non-const member function.");
 
 		TDelegate<RetValType(ParamTypes...), UserPolicy> Result;
-		TBaseSPMethodDelegateInstance<false, UserClass, ESPMode::Fast, FuncType, UserPolicy, VarTypes...>::Create(Result, InUserObjectRef, InFunc, Vars...);
+		TBaseSPMethodDelegateInstance<false, UserClass, Mode, FuncType, UserPolicy, VarTypes...>::Create(Result, InUserObjectRef, InFunc, Vars...);
 		return Result;
 	}
-	template <typename UserClass, typename... VarTypes>
-	UE_NODISCARD inline static TDelegate<RetValType(ParamTypes...), UserPolicy> CreateSP(const TSharedRef<UserClass, ESPMode::Fast>& InUserObjectRef, typename TMemFunPtrType<true, UserClass, RetValType (ParamTypes..., VarTypes...)>::Type InFunc, VarTypes... Vars)
+	template <typename UserClass, ESPMode Mode, typename... VarTypes>
+	UE_NODISCARD inline static TDelegate<RetValType(ParamTypes...), UserPolicy> CreateSP(const TSharedRef<UserClass, Mode>& InUserObjectRef, typename TMemFunPtrType<true, UserClass, RetValType (ParamTypes..., VarTypes...)>::Type InFunc, VarTypes... Vars)
 	{
 		TDelegate<RetValType(ParamTypes...), UserPolicy> Result;
-		TBaseSPMethodDelegateInstance<true, const UserClass, ESPMode::Fast, FuncType, UserPolicy, VarTypes...>::Create(Result, InUserObjectRef, InFunc, Vars...);
+		TBaseSPMethodDelegateInstance<true, const UserClass, Mode, FuncType, UserPolicy, VarTypes...>::Create(Result, InUserObjectRef, InFunc, Vars...);
 		return Result;
 	}
 
 	/**
-	 * Static: Creates a shared pointer-based (fast, not thread-safe) member function delegate.
+	 * Static: Creates a shared pointer-based member function delegate.
 	 *
 	 * Shared pointer delegates keep a weak reference to your object.
 	 * You can use ExecuteIfBound() to call them.
@@ -238,10 +238,12 @@ public:
 	}
 
 	/**
-	 * Static: Creates a shared pointer-based (slower, conditionally thread-safe) member function delegate.
+	 * Static: Creates a shared pointer-based (thread-safe) member function delegate.
 	 *
 	 * Shared pointer delegates keep a weak reference to your object.
 	 * You can use ExecuteIfBound() to call them.
+	 *
+	 * Note: This function is redundant, but is retained for backwards compatibility.  CreateSP() works in both thread-safe and not-thread-safe modes and should be preferred.
 	 */
 	template <typename UserClass, typename... VarTypes>
 	UE_NODISCARD inline static TDelegate<RetValType(ParamTypes...), UserPolicy> CreateThreadSafeSP(const TSharedRef<UserClass, ESPMode::ThreadSafe>& InUserObjectRef, typename TMemFunPtrType<false, UserClass, RetValType (ParamTypes..., VarTypes...)>::Type InFunc, VarTypes... Vars)
@@ -261,10 +263,12 @@ public:
 	}
 
 	/**
-	 * Static: Creates a shared pointer-based (slower, conditionally thread-safe) member function delegate.
+	 * Static: Creates a shared pointer-based (thread-safe) member function delegate.
 	 *
 	 * Shared pointer delegates keep a weak reference to your object.
 	 * You can use ExecuteIfBound() to call them.
+	 *
+	 * Note: This function is redundant, but is retained for backwards compatibility.  CreateSP() works in both thread-safe and not-thread-safe modes and should be preferred.
 	 */
 	template <typename UserClass, typename... VarTypes>
 	UE_NODISCARD inline static TDelegate<RetValType(ParamTypes...), UserPolicy> CreateThreadSafeSP(UserClass* InUserObject, typename TMemFunPtrType<false, UserClass, RetValType (ParamTypes..., VarTypes...)>::Type InFunc, VarTypes... Vars)
@@ -461,23 +465,23 @@ public:
 	}
 
 	/**
-	 * Binds a shared pointer-based (fast, not thread-safe) member function delegate.  Shared pointer delegates keep a weak reference to your object.  You can use ExecuteIfBound() to call them.
+	 * Binds a shared pointer-based member function delegate.  Shared pointer delegates keep a weak reference to your object.  You can use ExecuteIfBound() to call them.
 	 */
-	template <typename UserClass, typename... VarTypes>
-	inline void BindSP(const TSharedRef<UserClass, ESPMode::Fast>& InUserObjectRef, typename TMemFunPtrType<false, UserClass, RetValType (ParamTypes..., VarTypes...)>::Type InFunc, VarTypes... Vars)
+	template <typename UserClass, ESPMode Mode, typename... VarTypes>
+	inline void BindSP(const TSharedRef<UserClass, Mode>& InUserObjectRef, typename TMemFunPtrType<false, UserClass, RetValType (ParamTypes..., VarTypes...)>::Type InFunc, VarTypes... Vars)
 	{
 		static_assert(!TIsConst<UserClass>::Value, "Attempting to bind a delegate with a const object pointer and non-const member function.");
 
 		*this = CreateSP(InUserObjectRef, InFunc, Vars...);
 	}
-	template <typename UserClass, typename... VarTypes>
-	inline void BindSP(const TSharedRef<UserClass, ESPMode::Fast>& InUserObjectRef, typename TMemFunPtrType<true, UserClass, RetValType (ParamTypes..., VarTypes...)>::Type InFunc, VarTypes... Vars)
+	template <typename UserClass, ESPMode Mode, typename... VarTypes>
+	inline void BindSP(const TSharedRef<UserClass, Mode>& InUserObjectRef, typename TMemFunPtrType<true, UserClass, RetValType (ParamTypes..., VarTypes...)>::Type InFunc, VarTypes... Vars)
 	{
 		*this = CreateSP(InUserObjectRef, InFunc, Vars...);
 	}
 
 	/**
-	 * Binds a shared pointer-based (fast, not thread-safe) member function delegate.
+	 * Binds a shared pointer-based member function delegate.
 	 *
 	 * Shared pointer delegates keep a weak reference to your object.
 	 * You can use ExecuteIfBound() to call them.
@@ -496,10 +500,12 @@ public:
 	}
 
 	/**
-	 * Binds a shared pointer-based (slower, conditionally thread-safe) member function delegate.
+	 * Binds a shared pointer-based (thread-safe) member function delegate.
 	 *
 	 * Shared pointer delegates keep a weak reference to your object.
 	 * You can use ExecuteIfBound() to call them.
+	 *
+	 * Note: This function is redundant, but is retained for backwards compatibility.  BindSP() works in both thread-safe and not-thread-safe modes and should be preferred.
 	 */
 	template <typename UserClass, typename... VarTypes>
 	inline void BindThreadSafeSP(const TSharedRef<UserClass, ESPMode::ThreadSafe>& InUserObjectRef, typename TMemFunPtrType<false, UserClass, RetValType (ParamTypes..., VarTypes...)>::Type InFunc, VarTypes... Vars)
@@ -515,10 +521,12 @@ public:
 	}
 
 	/**
-	 * Binds a shared pointer-based (slower, conditionally thread-safe) member function delegate.
+	 * Binds a shared pointer-based (thread-safe) member function delegate.
 	 *
 	 * Shared pointer delegates keep a weak reference to your object.
 	 * You can use ExecuteIfBound() to call them.
+	 *
+	 * Note: This function is redundant, but is retained for backwards compatibility.  BindSP() works in both thread-safe and not-thread-safe modes and should be preferred.
 	 */
 	template <typename UserClass, typename... VarTypes>
 	inline void BindThreadSafeSP(UserClass* InUserObject, typename TMemFunPtrType<false, UserClass, RetValType (ParamTypes..., VarTypes...)>::Type InFunc, VarTypes... Vars)
@@ -759,28 +767,28 @@ public:
 	}
 
 	/**
-	 * Adds a shared pointer-based (fast, not thread-safe) member function delegate.
+	 * Adds a shared pointer-based member function delegate.
 	 *
 	 * Shared pointer delegates keep a weak reference to your object.
 	 *
 	 * @param	InUserObjectRef	User object to bind to
 	 * @param	InFunc			Class method function address
 	 */
-	template <typename UserClass, typename... VarTypes>
-	inline FDelegateHandle AddSP(const TSharedRef<UserClass, ESPMode::Fast>& InUserObjectRef, typename TMemFunPtrType<false, UserClass, void (ParamTypes..., VarTypes...)>::Type InFunc, VarTypes... Vars)
+	template <typename UserClass, ESPMode Mode, typename... VarTypes>
+	inline FDelegateHandle AddSP(const TSharedRef<UserClass, Mode>& InUserObjectRef, typename TMemFunPtrType<false, UserClass, void (ParamTypes..., VarTypes...)>::Type InFunc, VarTypes... Vars)
 	{
 		static_assert(!TIsConst<UserClass>::Value, "Attempting to bind a delegate with a const object pointer and non-const member function.");
 
 		return Add(FDelegate::CreateSP(InUserObjectRef, InFunc, Vars...));
 	}
-	template <typename UserClass, typename... VarTypes>
-	inline FDelegateHandle AddSP(const TSharedRef<UserClass, ESPMode::Fast>& InUserObjectRef, typename TMemFunPtrType<true, UserClass, void (ParamTypes..., VarTypes...)>::Type InFunc, VarTypes... Vars)
+	template <typename UserClass, ESPMode Mode, typename... VarTypes>
+	inline FDelegateHandle AddSP(const TSharedRef<UserClass, Mode>& InUserObjectRef, typename TMemFunPtrType<true, UserClass, void (ParamTypes..., VarTypes...)>::Type InFunc, VarTypes... Vars)
 	{
 		return Add(FDelegate::CreateSP(InUserObjectRef, InFunc, Vars...));
 	}
 
 	/**
-	 * Adds a shared pointer-based (fast, not thread-safe) member function delegate.
+	 * Adds a shared pointer-based member function delegate.
 	 *
 	 * Shared pointer delegates keep a weak reference to your object.
 	 *
@@ -801,10 +809,12 @@ public:
 	}
 
 	/**
-	 * Adds a shared pointer-based (slower, conditionally thread-safe) member function delegate.  Shared pointer delegates keep a weak reference to your object.
+	 * Adds a shared pointer-based (thread-safe) member function delegate.  Shared pointer delegates keep a weak reference to your object.
 	 *
 	 * @param	InUserObjectRef	User object to bind to
 	 * @param	InFunc			Class method function address
+	 *
+	 * Note: This function is redundant, but is retained for backwards compatibility.  AddSP() works in both thread-safe and not-thread-safe modes and should be preferred.
 	 */
 	template <typename UserClass, typename... VarTypes>
 	inline FDelegateHandle AddThreadSafeSP(const TSharedRef<UserClass, ESPMode::ThreadSafe>& InUserObjectRef, typename TMemFunPtrType<false, UserClass, void (ParamTypes..., VarTypes...)>::Type InFunc, VarTypes... Vars)
@@ -820,12 +830,14 @@ public:
 	}
 
 	/**
-	 * Adds a shared pointer-based (slower, conditionally thread-safe) member function delegate.
+	 * Adds a shared pointer-based (thread-safe) member function delegate.
 	 *
 	 * Shared pointer delegates keep a weak reference to your object.
 	 *
 	 * @param	InUserObject	User object to bind to
 	 * @param	InFunc			Class method function address
+	 *
+	 * Note: This function is redundant, but is retained for backwards compatibility.  AddSP() works in both thread-safe and not-thread-safe modes and should be preferred.
 	 */
 	template <typename UserClass, typename... VarTypes>
 	inline FDelegateHandle AddThreadSafeSP(UserClass* InUserObject, typename TMemFunPtrType<false, UserClass, void (ParamTypes..., VarTypes...)>::Type InFunc, VarTypes... Vars)
diff --git a/Engine/Source/Runtime/Core/Public/Misc/TypeContainer.h b/Engine/Source/Runtime/Core/Public/Misc/TypeContainer.h
--- a/Engine/Source/Runtime/Core/Public/Misc/TypeContainer.h
+++ b/Engine/Source/Runtime/Core/Public/Misc/TypeContainer.h
@@ -63,7 +63,7 @@ enum class ETypeContainerScope
  * Note: Type containers depend on variadic templates and are therefore not available
  * on XboxOne at this time, which means they should only be used for desktop applications.
  */
-template<ESPMode Mode = ESPMode::Fast>
+template<ESPMode Mode = ESPMode::ThreadSafe>
 class TTypeContainer
 {
 	/** Interface for object instance providers. */
@@ -423,5 +423,5 @@ private:
 };
 
 
-/** Thread-unsafe type container (for backwards compatibility). */
-class FTypeContainer : public TTypeContainer<ESPMode::Fast> { };
+/** For backwards compatibility. */
+class FTypeContainer : public TTypeContainer<ESPMode::ThreadSafe> { };
diff --git a/Engine/Source/Runtime/Core/Public/Templates/SharedPointer.h b/Engine/Source/Runtime/Core/Public/Templates/SharedPointer.h
--- a/Engine/Source/Runtime/Core/Public/Templates/SharedPointer.h
+++ b/Engine/Source/Runtime/Core/Public/Templates/SharedPointer.h
@@ -1841,7 +1841,7 @@ FORCEINLINE SharedPointerInternals::FRawPtrProxy< ObjectType > MakeShareable( Ob
  * 	  template <typename ObjectType>
  *	  friend class SharedPointerInternals::TIntrusiveReferenceController;
  */
-template <typename InObjectType, ESPMode InMode = ESPMode::Fast, typename... InArgTypes>
+template <typename InObjectType, ESPMode InMode = ESPMode::ThreadSafe, typename... InArgTypes>
 FORCEINLINE TSharedRef<InObjectType, InMode> MakeShared(InArgTypes&&... Args)
 {
 	SharedPointerInternals::TIntrusiveReferenceController<InObjectType>* Controller = SharedPointerInternals::NewIntrusiveReferenceController<InObjectType>(Forward<InArgTypes>(Args)...);
diff --git a/Engine/Source/Runtime/Core/Public/Templates/SharedPointerInternals.h b/Engine/Source/Runtime/Core/Public/Templates/SharedPointerInternals.h
--- a/Engine/Source/Runtime/Core/Public/Templates/SharedPointerInternals.h
+++ b/Engine/Source/Runtime/Core/Public/Templates/SharedPointerInternals.h
@@ -12,7 +12,6 @@
 
 
 /** Default behavior. */
-#define	FORCE_THREADSAFE_SHAREDPTRS PLATFORM_WEAKLY_CONSISTENT_MEMORY
 #define THREAD_SANITISE_UNSAFEPTR 0
 
 #if THREAD_SANITISE_UNSAFEPTR
@@ -30,25 +29,23 @@ enum class ESPMode
 	/** Forced to be not thread-safe. */
 	NotThreadSafe = 0,
 
+	/** Thread-safe, never spin locks, but slower */
+	ThreadSafe = 1,
+
 	/**
-		*	Fast, doesn't ever use atomic interlocks.
-		*	Some code requires that all shared pointers are thread-safe.
-		*	It's better to change it here, instead of replacing ESPMode::Fast to ESPMode::ThreadSafe throughout the code.
-		*/
-	Fast = FORCE_THREADSAFE_SHAREDPTRS ? 1 : 0,
-
-	/** Conditionally thread-safe, never spin locks, but slower */
-	ThreadSafe = 1
+	 * Fast was sometimes ThreadSafe and sometimes NotThreadSafe, but Fast is a misnomer.
+	 * Deprecated since ThreadSafe became the default.
+	 */
+	Fast UE_DEPRECATED(5.0, "ESPMode::Fast has been deprecated - please use ESPMode::ThreadSafe instead") = 1
 };
 
 
-// Forward declarations.  Note that in the interest of fast performance, thread safety
-// features are mostly turned off (Mode = ESPMode::Fast).  If you need to access your
-// object on multiple threads, you should use ESPMode::ThreadSafe!
-template< class ObjectType, ESPMode Mode = ESPMode::Fast > class TSharedRef;
-template< class ObjectType, ESPMode Mode = ESPMode::Fast > class TSharedPtr;
-template< class ObjectType, ESPMode Mode = ESPMode::Fast > class TWeakPtr;
-template< class ObjectType, ESPMode Mode = ESPMode::Fast > class TSharedFromThis;
+// Forward declarations.  By default, thread safety features are turned on. (Mode = ESPMode::ThreadSafe).
+// If you need more concerned with performance of ref-counting, you should use ESPMode::NotThreadSafe.
+template< class ObjectType, ESPMode Mode = ESPMode::ThreadSafe > class TSharedRef;
+template< class ObjectType, ESPMode Mode = ESPMode::ThreadSafe > class TSharedPtr;
+template< class ObjectType, ESPMode Mode = ESPMode::ThreadSafe > class TWeakPtr;
+template< class ObjectType, ESPMode Mode = ESPMode::ThreadSafe > class TSharedFromThis;
 
 
 /**
diff --git a/Engine/Source/Runtime/CoreOnline/CoreOnline.Build.cs b/Engine/Source/Runtime/CoreOnline/CoreOnline.Build.cs
--- a/Engine/Source/Runtime/CoreOnline/CoreOnline.Build.cs
+++ b/Engine/Source/Runtime/CoreOnline/CoreOnline.Build.cs
@@ -18,10 +18,6 @@ public class CoreOnline : ModuleRules
 		);
 
 		ConfigCache.ReadSettings(DirectoryReference.FromFile(Target.ProjectFile), Target.Platform, this);
-		PublicDefinitions.Add("UNIQUENETID_ESPMODE=ESPMode::" + (bEnableThreadSafeUniqueNetIds ? "ThreadSafe" : "Fast"));
-		// This is to ease migration to ESPMode::ThreadSafe. We have deprecated public FUniqueNetId constructors, by including it in the
-		// ESPMode::Fast deprecation mechanism. The constructors are public when the ESPMode is Fast, and protected when it is ThreadSafe.
-		PublicDefinitions.Add("UNIQUENETID_CONSTRUCTORVIS=" + (bEnableThreadSafeUniqueNetIds ? "protected" : "public"));
 	}
 
 }
diff --git a/Engine/Source/Runtime/CoreOnline/Public/Online/CoreOnline.h b/Engine/Source/Runtime/CoreOnline/Public/Online/CoreOnline.h
--- a/Engine/Source/Runtime/CoreOnline/Public/Online/CoreOnline.h
+++ b/Engine/Source/Runtime/CoreOnline/Public/Online/CoreOnline.h
@@ -71,7 +71,7 @@ struct FJoinabilitySettings
  * Abstraction of a profile service online Id
  * The class is meant to be opaque
  */
-class FUniqueNetId : public TSharedFromThis<FUniqueNetId, UNIQUENETID_ESPMODE>
+class FUniqueNetId : public TSharedFromThis<FUniqueNetId>
 {
 protected:
 
diff --git a/Engine/Source/Runtime/CoreOnline/Public/Online/CoreOnlineFwd.h b/Engine/Source/Runtime/CoreOnline/Public/Online/CoreOnlineFwd.h
--- a/Engine/Source/Runtime/CoreOnline/Public/Online/CoreOnlineFwd.h
+++ b/Engine/Source/Runtime/CoreOnline/Public/Online/CoreOnlineFwd.h
@@ -7,6 +7,6 @@
 class FUniqueNetId;
 struct FUniqueNetIdWrapper;
 
-using FUniqueNetIdPtr = TSharedPtr<const FUniqueNetId, UNIQUENETID_ESPMODE>;
-using FUniqueNetIdRef = TSharedRef<const FUniqueNetId, UNIQUENETID_ESPMODE>;
-using FUniqueNetIdWeakPtr = TWeakPtr<const FUniqueNetId, UNIQUENETID_ESPMODE>;
\ No newline at end of file
+using FUniqueNetIdPtr = TSharedPtr<const FUniqueNetId>;
+using FUniqueNetIdRef = TSharedRef<const FUniqueNetId>;
+using FUniqueNetIdWeakPtr = TWeakPtr<const FUniqueNetId>;
\ No newline at end of file
diff --git a/Engine/Source/Runtime/Net/Core/Public/Net/Core/Analytics/NetAnalytics.h b/Engine/Source/Runtime/Net/Core/Public/Net/Core/Analytics/NetAnalytics.h
--- a/Engine/Source/Runtime/Net/Core/Public/Net/Core/Analytics/NetAnalytics.h
+++ b/Engine/Source/Runtime/Net/Core/Public/Net/Core/Analytics/NetAnalytics.h
@@ -48,11 +48,7 @@ class FNetAnalyticsAggregator;
 #define NET_ANALYTICS_MULTITHREADING 0
 
 
-#if NET_ANALYTICS_MULTITHREADING || FORCE_THREADSAFE_SHAREDPTRS
-	constexpr const ESPMode NetAnalyticsThreadSafety = ESPMode::ThreadSafe;
-#else
-	constexpr const ESPMode NetAnalyticsThreadSafety = ESPMode::NotThreadSafe;
-#endif
+constexpr const ESPMode NetAnalyticsThreadSafety = ESPMode::ThreadSafe;
 
 
 // Forward Declarations
