Remove 0.7M of 13M allocations during an internal game's startup

diff --git a/Engine/Source/Runtime/Core/Private/Misc/Parse.cpp b/Engine/Source/Runtime/Core/Private/Misc/Parse.cpp
--- a/Engine/Source/Runtime/Core/Private/Misc/Parse.cpp
+++ b/Engine/Source/Runtime/Core/Private/Misc/Parse.cpp
@@ -351,17 +351,21 @@ bool ParseQuotedString( const TCHAR* Buffer, T& Value, int32* OutNumCharsRead)
 		return false;
 	}
 
-	auto ShouldParse = [](const TCHAR Ch)
+	constexpr FAsciiSet StopCharacters = FAsciiSet("\"\n\r") + '\0';
+	constexpr FAsciiSet StopAndEscapeCharacters = StopCharacters + '\\';
+	auto ShouldParse = [=](const TCHAR Ch) { return StopCharacters.Test(Ch) == 0; };
+	
+	while (true)
 	{
-		constexpr FAsciiSet StopCharacters = FAsciiSet("\"\n\r") + '\0';
-		return StopCharacters.Test(Ch) == 0;
-	};
+		// Append unescaped substring
+		const TCHAR* UnescapedEnd = FAsciiSet::FindFirstOrEnd(Buffer, StopAndEscapeCharacters);
+		FStringView UnescapedSubstring(Buffer, int32(UnescapedEnd - Buffer));
+		Value += UnescapedSubstring;
+		Buffer = UnescapedEnd;
 
-	while (ShouldParse(*Buffer))
-	{
-		if (*Buffer != TCHAR('\\')) // unescaped character
+		if (*Buffer != '\\') // Found a stop character
 		{
-			Value += *Buffer++;
+			break;
 		}
 		else if (*++Buffer == TCHAR('\\')) // escaped backslash "\\"
 		{
