Fix crash reporting thread exiting immediately.

#rnx

diff --git a/Engine/Source/Runtime/Core/Private/Windows/WindowsPlatformCrashContext.cpp b/Engine/Source/Runtime/Core/Private/Windows/WindowsPlatformCrashContext.cpp
--- a/Engine/Source/Runtime/Core/Private/Windows/WindowsPlatformCrashContext.cpp
+++ b/Engine/Source/Runtime/Core/Private/Windows/WindowsPlatformCrashContext.cpp
@@ -1117,7 +1117,7 @@ private:
 		__try
 #endif
 		{
-			bool bStopThreadRequested = true;
+			bool bStopThreadRequested = false;
 			while (!bStopThreadRequested)
 			{
 				DWORD WaitResult = WaitForMultipleObjects((DWORD)IE_MAX, InputEvents, false, 500);
