U5EA - Don't use IsInActualRenderingThread because it is used in so many places that actual tagging is implactical ATM

diff --git a/Engine/Source/Runtime/Core/Public/HAL/ThreadManager.h b/Engine/Source/Runtime/Core/Public/HAL/ThreadManager.h
--- a/Engine/Source/Runtime/Core/Public/HAL/ThreadManager.h
+++ b/Engine/Source/Runtime/Core/Public/HAL/ThreadManager.h
@@ -46,14 +46,16 @@ public:
 	{
 		static FString GameThreadName(TEXT("GameThread"));
 		static FString RenderThreadName(TEXT("RenderThread"));
+		PRAGMA_DISABLE_DEPRECATION_WARNINGS
 		if (ThreadId == GGameThreadId)
 		{
 			return GameThreadName;
 		}
-		else if (IsInActualRenderingThread())
+		else if (ThreadId == GRenderThreadId)
 		{
 			return RenderThreadName;
 		}
+		PRAGMA_ENABLE_DEPRECATION_WARNINGS
 		return Get().GetThreadNameInternal(ThreadId);
 	}
 
