Added UE_PUSH_MACRO/UE_POP_MACRO and deprecated PUSH_MACRO/POP_MACRO

PUSH_MACRO/POP_MACRO had issues with non-parameterized macros, as they would stringize the expanded macro value rather than the name of the macro itself. The UE_PUSH_MACRO/UE_POP_MACRO versions take the macro name as a string to avoid this issue.

diff --git a/Engine/Plugins/Compositing/OpenCVLensDistortion/Source/OpenCVHelper/Public/OpenCVHelper.h b/Engine/Plugins/Compositing/OpenCVLensDistortion/Source/OpenCVHelper/Public/OpenCVHelper.h
--- a/Engine/Plugins/Compositing/OpenCVLensDistortion/Source/OpenCVHelper/Public/OpenCVHelper.h
+++ b/Engine/Plugins/Compositing/OpenCVLensDistortion/Source/OpenCVHelper/Public/OpenCVHelper.h
@@ -27,11 +27,11 @@ cannot undefine the value.
 	__pragma(warning(disable: 6297))  /* Arithmetic overflow:  32-bit value is shifted, then cast to 64-bit value.  Results might not be an expected value. */ \
 	__pragma(warning(disable: 6294))  /* Ill-defined for-loop:  initial condition does not satisfy test.  Loop body not executed. */ \
 	__pragma(warning(disable: 6201))  /* Index '<x>' is out of valid index range '<a>' to '<b>' for possibly stack allocated buffer '<variable>'. */ \
-	PUSH_MACRO(check)
+	UE_PUSH_MACRO("check")
 #else
 // TODO: when adding support for other platforms, this definition may require updating
-#define OPENCV_INCLUDES_START THIRD_PARTY_INCLUDES_START PUSH_MACRO(check)
+#define OPENCV_INCLUDES_START THIRD_PARTY_INCLUDES_START UE_PUSH_MACRO("check")
 #endif
 
-#define OPENCV_INCLUDES_END THIRD_PARTY_INCLUDES_END POP_MACRO(check)
+#define OPENCV_INCLUDES_END THIRD_PARTY_INCLUDES_END UE_POP_MACRO("check")
 	
\ No newline at end of file
diff --git a/Engine/Source/Runtime/ClothingSystemRuntimeNv/Public/NvClothIncludes.h b/Engine/Source/Runtime/ClothingSystemRuntimeNv/Public/NvClothIncludes.h
--- a/Engine/Source/Runtime/ClothingSystemRuntimeNv/Public/NvClothIncludes.h
+++ b/Engine/Source/Runtime/ClothingSystemRuntimeNv/Public/NvClothIncludes.h
@@ -11,7 +11,7 @@
 
 #if WITH_NVCLOTH
 
-PUSH_MACRO(check)
+UE_PUSH_MACRO("check")
 #undef check
 
 #if PLATFORM_MICROSOFT
@@ -43,6 +43,6 @@ PUSH_MACRO(check)
 	#endif
 #endif
 
-POP_MACRO(check)
+UE_POP_MACRO("check")
 
 #endif
diff --git a/Engine/Source/Runtime/Core/Public/Misc/CoreMiscDefines.h b/Engine/Source/Runtime/Core/Public/Misc/CoreMiscDefines.h
--- a/Engine/Source/Runtime/Core/Public/Misc/CoreMiscDefines.h
+++ b/Engine/Source/Runtime/Core/Public/Misc/CoreMiscDefines.h
@@ -117,20 +117,14 @@ typedef int32 FPlatformUserId;
 const FPlatformUserId PLATFORMUSERID_NONE = INDEX_NONE;
 #endif // RC_INVOKED
 
-
-
 // When passed to pragma message will result in clickable warning in VS
 #define WARNING_LOCATION(Line) __FILE__ "(" PREPROCESSOR_TO_STRING(Line) ")"
 
 // Push and pop macro definitions
-#ifdef __clang__
-	#define PUSH_MACRO(name) _Pragma(PREPROCESSOR_TO_STRING(push_macro(PREPROCESSOR_TO_STRING(name))))
-	#define POP_MACRO(name) _Pragma(PREPROCESSOR_TO_STRING(pop_macro(PREPROCESSOR_TO_STRING(name))))
-#else
-	#define PUSH_MACRO(name) __pragma(push_macro(PREPROCESSOR_TO_STRING(name)))
-	#define POP_MACRO(name) __pragma(pop_macro(PREPROCESSOR_TO_STRING(name)))
-#endif
-
+#define UE_PUSH_MACRO(name) _Pragma(PREPROCESSOR_TO_STRING(push_macro(name)))
+#define UE_POP_MACRO(name) _Pragma(PREPROCESSOR_TO_STRING(pop_macro(name)))
+#define PUSH_MACRO(name) DEPRECATED_MACRO(5.0, "PUSH_MACRO is deprecated. Use UE_PUSH_MACRO and pass the macro name as a string.") UE_PUSH_MACRO(PREPROCESSOR_TO_STRING(name))
+#define POP_MACRO(name) DEPRECATED_MACRO(5.0, "POP_MACRO is deprecated. Use UE_POP_MACRO and pass the macro name as a string.") UE_POP_MACRO(PREPROCESSOR_TO_STRING(name))
 
 #ifdef __COUNTER__
 	// Created a variable with a unique name
diff --git a/Engine/Source/Runtime/NVidia/GeForceNOW/Private/GeForceNOWWrapper.cpp b/Engine/Source/Runtime/NVidia/GeForceNOW/Private/GeForceNOWWrapper.cpp
--- a/Engine/Source/Runtime/NVidia/GeForceNOW/Private/GeForceNOWWrapper.cpp
+++ b/Engine/Source/Runtime/NVidia/GeForceNOW/Private/GeForceNOWWrapper.cpp
@@ -9,12 +9,12 @@
 #include "HAL/FileManager.h"
 #include "Misc/Paths.h"
 
-PUSH_MACRO(TEXT)
+UE_PUSH_MACRO("TEXT")
 #undef TEXT
 THIRD_PARTY_INCLUDES_START
 #include "GfnRuntimeSdk_Wrapper.h"
 THIRD_PARTY_INCLUDES_END
-POP_MACRO(TEXT)
+UE_POP_MACRO("TEXT")
 
 GeForceNOWWrapper::GeForceNOWWrapper()
 	: bIsInitialized(false)
