CsvProfiler metadata improvements
- Implement ability to unset metadata values by setting with nullptr
- Sanitize metadata to remove commas from values and keys (the commandline is excluded because it's special)
- Set a starttimestamp metadata value (in seconds since the Unix epoch).
- Add metadata for OS, CPU, Namedevents (at capture start)
- Log metadata changes

diff --git a/Engine/Source/Runtime/Core/Private/ProfilingDebugging/CsvProfiler.cpp b/Engine/Source/Runtime/Core/Private/ProfilingDebugging/CsvProfiler.cpp
--- a/Engine/Source/Runtime/Core/Private/ProfilingDebugging/CsvProfiler.cpp
+++ b/Engine/Source/Runtime/Core/Private/ProfilingDebugging/CsvProfiler.cpp
@@ -2505,11 +2505,19 @@ FCsvProfiler::FCsvProfiler()
 	FString BuildVersionString = FApp::GetBuildVersion();
 	FString EngineVersionString = FEngineVersion::Current().ToString();
 
-	MetadataMap.FindOrAdd(TEXT("Platform")) = PlatformStr;
-	MetadataMap.FindOrAdd(TEXT("Config")) = BuildConfigurationStr;
-	MetadataMap.FindOrAdd(TEXT("BuildVersion")) = BuildVersionString;
-	MetadataMap.FindOrAdd(TEXT("EngineVersion")) = EngineVersionString;
-	MetadataMap.FindOrAdd(TEXT("Commandline")) = CommandlineStr;
+	FString OSMajor, OSMinor;
+	FPlatformMisc::GetOSVersions(OSMajor, OSMinor);
+	OSMajor.TrimStartAndEndInline();
+	OSMinor.TrimStartAndEndInline();
+	FString OSString = FString::Printf(TEXT("%s %s"), *OSMajor, *OSMinor);
+
+	SetMetadataInternal(TEXT("Platform"), *PlatformStr);
+	SetMetadataInternal(TEXT("Config"), *BuildConfigurationStr);
+	SetMetadataInternal(TEXT("BuildVersion"), *BuildVersionString);
+	SetMetadataInternal(TEXT("EngineVersion"), *EngineVersionString);
+	SetMetadataInternal(TEXT("Commandline"), *CommandlineStr, false);
+	SetMetadataInternal(TEXT("OS"), *OSString);
+	SetMetadataInternal(TEXT("CPU"), *FPlatformMisc::GetDeviceMakeAndModel());
 }
 
 FCsvProfiler::~FCsvProfiler()
@@ -2633,7 +2641,7 @@ void FCsvProfiler::BeginFrame()
 							GCsvUseProcessingThread = false;
 						}
 					}
-
+					 
 					// Set the CSV ID and mirror it to the log
 					FString CsvId = FGuid::NewGuid().ToString();
 					SetMetadata(TEXT("CsvID"), *CsvId);
@@ -2660,6 +2668,11 @@ void FCsvProfiler::BeginFrame()
 
 					SetMetadata(TEXT("PGOEnabled"), FPlatformMisc::IsPGOEnabled() ? TEXT("1") : TEXT("0"));
 
+					// Output the current time in seconds seconds since the Unix Epoch
+					SetMetadata(TEXT("StartTimestamp"), *FString::Printf(TEXT("%lld"), FDateTime::UtcNow().ToUnixTimestamp()));
+
+					SetMetadata(TEXT("NamedEvents"), (GCycleStatsShouldEmitNamedEvents > 0) ? TEXT("1") : TEXT("0"));
+
 					GCsvStatCounts = !!CVarCsvStatCounts.GetValueOnGameThread();
 
 					// Initialize tls before setting the capturing flag to true.
@@ -3117,19 +3130,42 @@ void FCsvProfiler::RecordEvent(int32 CategoryIndex, const FString& EventText)
 	}
 }
 
-void FCsvProfiler::SetMetadata(const TCHAR* Key, const TCHAR* Value)
+void FCsvProfiler::SetMetadataInternal(const TCHAR* Key, const TCHAR* Value, bool bSanitize)
 {
-	TRACE_CSV_PROFILER_METADATA(Key, Value);
-
-	LLM_SCOPE(ELLMTag::CsvProfiler);
-
 	// Always gather CSV metadata, even if we're not currently capturing.
 	// Metadata is applied to the next CSV profile, when the file is written.
-	FCsvProfiler* CsvProfiler = FCsvProfiler::Get();
+	LLM_SCOPE(ELLMTag::CsvProfiler);
 	FString KeyLower = FString(Key).ToLower();
 
-	FScopeLock Lock(&CsvProfiler->MetadataCS);
-	CsvProfiler->MetadataMap.FindOrAdd(KeyLower) = Value;
+	if (Value == nullptr)
+	{
+		FScopeLock Lock(&MetadataCS);
+		if (MetadataMap.Contains(KeyLower))
+		{
+			UE_LOG(LogCsvProfiler, Display, TEXT("Metadata unset : %s"), *KeyLower);
+			MetadataMap.Remove(KeyLower);
+		}
+	}
+	else
+	{
+		TRACE_CSV_PROFILER_METADATA(Key, Value);
+		FString ValueStr = Value;
+		if (bSanitize)
+		{
+			check(!KeyLower.Contains(TEXT(",")));
+			if (ValueStr.ReplaceInline(TEXT(","), TEXT("&#44;")) > 0)
+			{
+				UE_LOG(LogCsvProfiler, Warning, TEXT("Metadata value sanitized due to invalid characters: %s=\"%s\""), *KeyLower, Value);
+			}
+		}
+		// Only log if the metadata changed, to prevent logspam 
+		FScopeLock Lock(&MetadataCS);
+		if (!MetadataMap.Contains(KeyLower) || MetadataMap[KeyLower]!=ValueStr)
+		{
+			UE_LOG(LogCsvProfiler, Display, TEXT("Metadata set : %s=\"%s\""), *KeyLower, *ValueStr);
+		}
+		MetadataMap.FindOrAdd(KeyLower) = ValueStr;
+	}
 }
 
 void FCsvProfiler::SetThreadName(const FString& InThreadName)
diff --git a/Engine/Source/Runtime/Core/Public/ProfilingDebugging/CsvProfiler.h b/Engine/Source/Runtime/Core/Public/ProfilingDebugging/CsvProfiler.h
--- a/Engine/Source/Runtime/Core/Public/ProfilingDebugging/CsvProfiler.h
+++ b/Engine/Source/Runtime/Core/Public/ProfilingDebugging/CsvProfiler.h
@@ -246,7 +246,10 @@ public:
 	CORE_API static void RecordEvent(int32 CategoryIndex, const FString& EventText);
 	CORE_API static void RecordEventAtTimestamp(int32 CategoryIndex, const FString& EventText, uint64 Cycles64);
 
-	CORE_API static void SetMetadata(const TCHAR* Key, const TCHAR* Value);
+	CORE_API static void SetMetadata(const TCHAR* Key, const TCHAR* Value)
+	{
+		FCsvProfiler::Get()->SetMetadataInternal(Key, Value, true);
+	}
 	
 	/** Set Thread name for a TLS. Needs to be called before the first event of that thread is sent. */
 	CORE_API static void SetThreadName(const FString& ThreadName);
@@ -331,6 +334,8 @@ public:
 private:
 	CORE_API static void VARARGS RecordEventfInternal(int32 CategoryIndex, const TCHAR* Fmt, ...);
 
+	void SetMetadataInternal(const TCHAR* Key, const TCHAR* Value, bool bSanitize=true);
+
 	static int32 GetCategoryIndex(const FString& Name);
 
 	void FinalizeCsvFile();
