Add CSV metadata for LLM

diff --git a/Engine/Source/Runtime/Core/Private/HAL/LowLevelMemTracker.cpp b/Engine/Source/Runtime/Core/Private/HAL/LowLevelMemTracker.cpp
--- a/Engine/Source/Runtime/Core/Private/HAL/LowLevelMemTracker.cpp
+++ b/Engine/Source/Runtime/Core/Private/HAL/LowLevelMemTracker.cpp
@@ -14,6 +14,7 @@
 #include "Misc/ScopeRWLock.h"
 #include "Misc/StringBuilder.h"
 #include "Misc/VarArgs.h"
+#include "ProfilingDebugging/CsvProfiler.h"
 #include "Templates/Atomic.h"
 #include "Trace/Trace.inl"
 
@@ -937,6 +938,8 @@ void FLowLevelMemTracker::SetProgramSize(uint64 InProgramSize)
 
 void FLowLevelMemTracker::ProcessCommandLine(const TCHAR* CmdLine)
 {
+	CSV_METADATA(TEXT("LLM"), TEXT("0"));
+
 #if LLM_AUTO_ENABLE
 	// LLM is always on, regardless of command line
 	bool bShouldDisable = false;
@@ -981,6 +984,7 @@ void FLowLevelMemTracker::ProcessCommandLine(const TCHAR* CmdLine)
 		Clear();
 		return;
 	}
+	CSV_METADATA(TEXT("LLM"), TEXT("1"));
 	bIsDisabled = false;
 	bCsvWriterEnabled = bLocalCsvWriterEnabled;
 	bTraceWriterEnabled = bLocalTraceWriterEnabled;
