Mac/Linux build fix

diff --git a/Engine/Source/Runtime/Core/Private/ProfilingDebugging/CsvProfiler.cpp b/Engine/Source/Runtime/Core/Private/ProfilingDebugging/CsvProfiler.cpp
--- a/Engine/Source/Runtime/Core/Private/ProfilingDebugging/CsvProfiler.cpp
+++ b/Engine/Source/Runtime/Core/Private/ProfilingDebugging/CsvProfiler.cpp
@@ -3130,6 +3130,11 @@ void FCsvProfiler::RecordEvent(int32 CategoryIndex, const FString& EventText)
 	}
 }
 
+void FCsvProfiler::SetMetadata(const TCHAR* Key, const TCHAR* Value)
+{
+	FCsvProfiler::Get()->SetMetadataInternal(Key, Value, true);
+}
+
 void FCsvProfiler::SetMetadataInternal(const TCHAR* Key, const TCHAR* Value, bool bSanitize)
 {
 	// Always gather CSV metadata, even if we're not currently capturing.
diff --git a/Engine/Source/Runtime/Core/Public/ProfilingDebugging/CsvProfiler.h b/Engine/Source/Runtime/Core/Public/ProfilingDebugging/CsvProfiler.h
--- a/Engine/Source/Runtime/Core/Public/ProfilingDebugging/CsvProfiler.h
+++ b/Engine/Source/Runtime/Core/Public/ProfilingDebugging/CsvProfiler.h
@@ -246,10 +246,7 @@ public:
 	CORE_API static void RecordEvent(int32 CategoryIndex, const FString& EventText);
 	CORE_API static void RecordEventAtTimestamp(int32 CategoryIndex, const FString& EventText, uint64 Cycles64);
 
-	CORE_API static void SetMetadata(const TCHAR* Key, const TCHAR* Value)
-	{
-		FCsvProfiler::Get()->SetMetadataInternal(Key, Value, true);
-	}
+	CORE_API static void SetMetadata(const TCHAR* Key, const TCHAR* Value);
 	
 	/** Set Thread name for a TLS. Needs to be called before the first event of that thread is sent. */
 	CORE_API static void SetThreadName(const FString& ThreadName);
