Use MSVC specific pragma for UE_PUSH_MACRO/UE_POP_MACRO

Older versions of MSVC don't support the C++11 compliant _Pragma correctly

diff --git a/Engine/Source/Runtime/Core/Public/Misc/CoreMiscDefines.h b/Engine/Source/Runtime/Core/Public/Misc/CoreMiscDefines.h
--- a/Engine/Source/Runtime/Core/Public/Misc/CoreMiscDefines.h
+++ b/Engine/Source/Runtime/Core/Public/Misc/CoreMiscDefines.h
@@ -121,8 +121,13 @@ const FPlatformUserId PLATFORMUSERID_NONE = INDEX_NONE;
 #define WARNING_LOCATION(Line) __FILE__ "(" PREPROCESSOR_TO_STRING(Line) ")"
 
 // Push and pop macro definitions
-#define UE_PUSH_MACRO(name) _Pragma(PREPROCESSOR_TO_STRING(push_macro(name)))
-#define UE_POP_MACRO(name) _Pragma(PREPROCESSOR_TO_STRING(pop_macro(name)))
+#ifdef __clang__
+	#define UE_PUSH_MACRO(name) _Pragma(PREPROCESSOR_TO_STRING(push_macro(name)))
+	#define UE_POP_MACRO(name) _Pragma(PREPROCESSOR_TO_STRING(pop_macro(name)))
+#else
+	#define UE_PUSH_MACRO(name) __pragma(push_macro(name))
+	#define UE_POP_MACRO(name) __pragma(pop_macro(name))
+#endif
 #define PUSH_MACRO(name) DEPRECATED_MACRO(5.0, "PUSH_MACRO is deprecated. Use UE_PUSH_MACRO and pass the macro name as a string.") UE_PUSH_MACRO(PREPROCESSOR_TO_STRING(name))
 #define POP_MACRO(name) DEPRECATED_MACRO(5.0, "POP_MACRO is deprecated. Use UE_POP_MACRO and pass the macro name as a string.") UE_POP_MACRO(PREPROCESSOR_TO_STRING(name))
 
