Fixed missing PCallstack happening when the Editor has more than 256 threads and the crashing thread is not in the 256 first visited by the OS.
  - Bumped the limit from 256 to 512
  - Always reserve one spot for the crashing thread in the list transmitted to CRC, possibly ignoring some thread.
  - Added diagnostic logs in CRC to captures cases where the number of thread would reach the new limit of 512 or if the crashing thread is 0.

diff --git a/Engine/Source/Programs/CrashReportClient/Private/CrashReportClientApp.cpp b/Engine/Source/Programs/CrashReportClient/Private/CrashReportClientApp.cpp
--- a/Engine/Source/Programs/CrashReportClient/Private/CrashReportClientApp.cpp
+++ b/Engine/Source/Programs/CrashReportClient/Private/CrashReportClientApp.cpp
@@ -443,6 +443,14 @@ FPlatformErrorReport CollectErrorReport(FRecoveryService* RecoveryService, uint3
 	{
 		FDiagnosticLogger::Get().LogEvent(TEXT("Report/OpenProcessFail"));
 	}
+	else if (SharedCrashContext.CrashingThreadId == 0)
+	{
+		FDiagnosticLogger::Get().LogEvent(TEXT("Report/BadCrashThreadId"));
+	}
+	else if (SharedCrashContext.NumThreads == CR_MAX_THREADS)
+	{
+		FDiagnosticLogger::Get().LogEvent(TEXT("Report/BumpThreadLimits"));
+	}
 
 	// First init the static crash context state
 	FPlatformCrashContext::InitializeFromContext(
diff --git a/Engine/Source/Runtime/Core/Private/Windows/WindowsPlatformCrashContext.cpp b/Engine/Source/Runtime/Core/Private/Windows/WindowsPlatformCrashContext.cpp
--- a/Engine/Source/Runtime/Core/Private/Windows/WindowsPlatformCrashContext.cpp
+++ b/Engine/Source/Runtime/Core/Private/Windows/WindowsPlatformCrashContext.cpp
@@ -735,6 +735,7 @@ int32 ReportCrashForMonitor(
 	uint32 ThreadIdx = 0;
 	DWORD CurrentProcessId = GetCurrentProcessId();
 	HANDLE ThreadSnapshot = CreateToolhelp32Snapshot(TH32CS_SNAPTHREAD, 0);
+	bool bCapturedCrashingThreadId = false;
 	if (ThreadSnapshot != INVALID_HANDLE_VALUE)
 	{
 		THREADENTRY32 ThreadEntry;
@@ -745,15 +746,24 @@ int32 ReportCrashForMonitor(
 			{
 				if (ThreadEntry.th32OwnerProcessID == CurrentProcessId)
 				{
-					SharedContext->ThreadIds[ThreadIdx] = ThreadEntry.th32ThreadID;
-					const FString& TmThreadName = FThreadManager::GetThreadName(ThreadEntry.th32ThreadID);
-					const TCHAR* ThreadName = TmThreadName.IsEmpty() ? TEXT("Unknown") : *TmThreadName;
-					FCString::Strcpy(
-						&SharedContext->ThreadNames[ThreadIdx*CR_MAX_THREAD_NAME_CHARS],
-						CR_MAX_THREAD_NAME_CHARS - 1,
-						ThreadName
-					);
-					ThreadIdx++;
+					if (CrashingThreadId == ThreadEntry.th32ThreadID)
+					{
+						bCapturedCrashingThreadId = true;
+					}
+
+					// Always keep one spot for the crashing thread in case the number of captured threads is about to reach our limit.
+					if (bCapturedCrashingThreadId || ThreadIdx < CR_MAX_THREADS - 1)
+					{
+						SharedContext->ThreadIds[ThreadIdx] = ThreadEntry.th32ThreadID;
+						const FString& TmThreadName = FThreadManager::GetThreadName(ThreadEntry.th32ThreadID);
+						const TCHAR* ThreadName = TmThreadName.IsEmpty() ? TEXT("Unknown") : *TmThreadName;
+						FCString::Strcpy(
+							&SharedContext->ThreadNames[ThreadIdx*CR_MAX_THREAD_NAME_CHARS],
+							CR_MAX_THREAD_NAME_CHARS - 1,
+							ThreadName
+						);
+						ThreadIdx++;
+					}
 				}
 			} while (Thread32Next(ThreadSnapshot, &ThreadEntry) && (ThreadIdx < CR_MAX_THREADS));
 		}
diff --git a/Engine/Source/Runtime/Core/Public/GenericPlatform/GenericPlatformCrashContext.h b/Engine/Source/Runtime/Core/Public/GenericPlatform/GenericPlatformCrashContext.h
--- a/Engine/Source/Runtime/Core/Public/GenericPlatform/GenericPlatformCrashContext.h
+++ b/Engine/Source/Runtime/Core/Public/GenericPlatform/GenericPlatformCrashContext.h
@@ -117,7 +117,7 @@ enum class ECrashTrigger
 #define CR_MAX_DIRECTORY_CHARS 256
 #define CR_MAX_STACK_FRAMES 256
 #define CR_MAX_THREAD_NAME_CHARS 64
-#define CR_MAX_THREADS 256
+#define CR_MAX_THREADS 512
 #define CR_MAX_GENERIC_FIELD_CHARS 64
 #define CR_MAX_COMMANDLINE_CHARS 1024
 #define CR_MAX_RICHTEXT_FIELD_CHARS 512
