Made EName an enum class to prevent implicit conversion to int

The following code would call the int overload rather than the name overload when passed an EName constant, which could easily lead to confusion and bugs:
    using FMyId = uint8;
    void DoThing(FName);
    void DoThing(FMyId);
    DoThing(NAME_Actor);

The change to use enum class requires that any code that actually deals with EName as an int (mostly internal name code and serialisation) be updated to explicitly cast, but prevents the implicit conversion that causes the issue above.

In order to preserve the NAME_X aliases that the old-style enum added to the global scope, new aliases have been added that point to the EName scoped versions. Unfortunately these can cause shadowing warnings if NAME_X is defined in the local scope, as the old-style enum used to allow that without shadowing, however there is no way to prevent this so we'll just need to fix any warnings that occur.

diff --git a/Engine/Plugins/AI/MLAdapter/Source/MLAdapter/Private/Sensors/MLAdapterSensor_Movement.cpp b/Engine/Plugins/AI/MLAdapter/Source/MLAdapter/Private/Sensors/MLAdapterSensor_Movement.cpp
--- a/Engine/Plugins/AI/MLAdapter/Source/MLAdapter/Private/Sensors/MLAdapterSensor_Movement.cpp
+++ b/Engine/Plugins/AI/MLAdapter/Source/MLAdapter/Private/Sensors/MLAdapterSensor_Movement.cpp
@@ -22,18 +22,18 @@ bool UMLAdapterSensor_Movement::ConfigureForAgent(UMLAdapterAgent& Agent)
 
 void UMLAdapterSensor_Movement::Configure(const TMap<FName, FString>& Params)
 {
-	const FName NAME_Location = TEXT("location");
-	const FName NAME_Velocity = TEXT("velocity");
+	const FName LocationKeyName = TEXT("location");
+	const FName VelocityKeyName = TEXT("velocity");
 	
 	Super::Configure(Params);
 
 	for (auto KeyValue : Params)
 	{
-		if (KeyValue.Key == NAME_Location)
+		if (KeyValue.Key == LocationKeyName)
 		{
 			bAbsoluteLocation = (KeyValue.Value.Find(TEXT("absolute")) != INDEX_NONE);
 		}
-		else if (KeyValue.Key == NAME_Velocity)
+		else if (KeyValue.Key == VelocityKeyName)
 		{
 			bAbsoluteLocation = (KeyValue.Value.Find(TEXT("absolute")) != INDEX_NONE);
 		}
diff --git a/Engine/Plugins/Developer/Concert/ConcertMain/Source/ConcertTransport/Private/IdentifierTable/ConcertTransportArchives.cpp b/Engine/Plugins/Developer/Concert/ConcertMain/Source/ConcertTransport/Private/IdentifierTable/ConcertTransportArchives.cpp
--- a/Engine/Plugins/Developer/Concert/ConcertMain/Source/ConcertTransport/Private/IdentifierTable/ConcertTransportArchives.cpp
+++ b/Engine/Plugins/Developer/Concert/ConcertMain/Source/ConcertTransport/Private/IdentifierTable/ConcertTransportArchives.cpp
@@ -38,7 +38,7 @@ FArchive& FConcertIdentifierWriter::operator<<(FName& Name)
 	if (Ename && ShouldReplicateAsInteger(*Ename))
 	{
 		SerializeConcertIdentifierSource(EConcertIdentifierSource::HardcodedIndex);
-		SerializeIndexValue(*Ename);
+		SerializeIndexValue((int32)*Ename);
 	}
 	else if (LocalIdentifierTable)
 	{
diff --git a/Engine/Plugins/FX/Niagara/Source/Niagara/Classes/NiagaraCollision.h b/Engine/Plugins/FX/Niagara/Source/Niagara/Classes/NiagaraCollision.h
--- a/Engine/Plugins/FX/Niagara/Source/Niagara/Classes/NiagaraCollision.h
+++ b/Engine/Plugins/FX/Niagara/Source/Niagara/Classes/NiagaraCollision.h
@@ -49,7 +49,7 @@ class FNiagaraDICollisionQueryBatch
 {
 public:
 	FNiagaraDICollisionQueryBatch()
-		: BatchID(NAME_None)
+		: BatchID(0)
 		, CurrBuffer(0)
 	{
 	}
diff --git a/Engine/Plugins/FX/Niagara/Source/Niagara/Private/NiagaraDataInterfaceSkeletalMesh.cpp b/Engine/Plugins/FX/Niagara/Source/Niagara/Private/NiagaraDataInterfaceSkeletalMesh.cpp
--- a/Engine/Plugins/FX/Niagara/Source/Niagara/Private/NiagaraDataInterfaceSkeletalMesh.cpp
+++ b/Engine/Plugins/FX/Niagara/Source/Niagara/Private/NiagaraDataInterfaceSkeletalMesh.cpp
@@ -2416,8 +2416,8 @@ UNiagaraDataInterfaceSkeletalMesh::UNiagaraDataInterfaceSkeletalMesh(FObjectInit
 	FNiagaraTypeDefinition Def(UObject::StaticClass());
 	MeshUserParameter.Parameter.SetType(Def);
 
-	static const FName NAME_Root("root");
-	ExcludeBoneName = NAME_Root;
+	static const FName RootBoneName("root");
+	ExcludeBoneName = RootBoneName;
 	bExcludeBone = false;
 
 	Proxy.Reset(new FNiagaraDataInterfaceProxySkeletalMesh());
diff --git a/Engine/Plugins/Online/OnlineFramework/Source/Party/Private/Party/SocialParty.cpp b/Engine/Plugins/Online/OnlineFramework/Source/Party/Private/Party/SocialParty.cpp
--- a/Engine/Plugins/Online/OnlineFramework/Source/Party/Private/Party/SocialParty.cpp
+++ b/Engine/Plugins/Online/OnlineFramework/Source/Party/Private/Party/SocialParty.cpp
@@ -1452,11 +1452,11 @@ void USocialParty::ConnectToReservationBeacon()
 			IOnlineSessionPtr SessionInterface = Online::GetSessionInterface(World);
 			if (SessionInterface.IsValid())
 			{
-				const FName GameSessionName = GetGameSessionName();
-				if (FNamedOnlineSession* Session = SessionInterface->GetNamedSession(GameSessionName))
+				const FName PartyGameSessionName = GetGameSessionName();
+				if (FNamedOnlineSession* Session = SessionInterface->GetNamedSession(PartyGameSessionName))
 				{
 					FString URL;
-					if (ensure(SessionInterface->GetResolvedConnectString(GameSessionName, URL, NAME_BeaconPort)))
+					if (ensure(SessionInterface->GetResolvedConnectString(PartyGameSessionName, URL, NAME_BeaconPort)))
 					{
 						// Reconnect to the reservation beacon to maintain our place in the game (just until actual joined, holds place for all party members)
 						ReservationBeaconClient = World->SpawnActor<APartyBeaconClient>(ReservationBeaconClientClass);
diff --git a/Engine/Source/Runtime/AIModule/Private/AIController.cpp b/Engine/Source/Runtime/AIModule/Private/AIController.cpp
--- a/Engine/Source/Runtime/AIModule/Private/AIController.cpp
+++ b/Engine/Source/Runtime/AIModule/Private/AIController.cpp
@@ -113,7 +113,6 @@ void AAIController::DisplayDebug(UCanvas* Canvas, const FDebugDisplayInfo& Debug
 {
 	Super::DisplayDebug(Canvas, DebugDisplay, YL, YPos);
 
-	static FName NAME_AI = FName(TEXT("AI"));
 	if (DebugDisplay.IsDisplayOn(NAME_AI))
 	{
 		if (PathFollowingComponent)
diff --git a/Engine/Source/Runtime/Core/Private/Logging/LogSuppressionInterface.cpp b/Engine/Source/Runtime/Core/Private/Logging/LogSuppressionInterface.cpp
--- a/Engine/Source/Runtime/Core/Private/Logging/LogSuppressionInterface.cpp
+++ b/Engine/Source/Runtime/Core/Private/Logging/LogSuppressionInterface.cpp
@@ -230,8 +230,6 @@ class FLogSuppressionImplementation: public FLogSuppressionInterface, private FS
 					FName CmdToken = FName(*CommandParts[PartIndex]);
 					static FName NAME_Verbose(TEXT("Verbose"));
 					static FName NAME_VeryVerbose(TEXT("VeryVerbose"));
-					static FName NAME_All(TEXT("All"));
-					static FName NAME_Default(TEXT("Default"));
 					static FName NAME_On(TEXT("On"));
 					static FName NAME_Off(TEXT("Off"));
 					static FName NAME_Break(TEXT("Break"));
diff --git a/Engine/Source/Runtime/Core/Private/Stats/Stats2.cpp b/Engine/Source/Runtime/Core/Private/Stats/Stats2.cpp
--- a/Engine/Source/Runtime/Core/Private/Stats/Stats2.cpp
+++ b/Engine/Source/Runtime/Core/Private/Stats/Stats2.cpp
@@ -866,7 +866,7 @@ bool FStatNameAndInfo::GetSortByNameFrom(FName InLongName)
 		}
 		checkStats(0);
 	}
-	return NAME_None;
+	return false;
 }
 
 /*-----------------------------------------------------------------------------
diff --git a/Engine/Source/Runtime/Core/Private/UObject/UnrealNames.cpp b/Engine/Source/Runtime/Core/Private/UObject/UnrealNames.cpp
--- a/Engine/Source/Runtime/Core/Private/UObject/UnrealNames.cpp
+++ b/Engine/Source/Runtime/Core/Private/UObject/UnrealNames.cpp
@@ -49,7 +49,7 @@ const TCHAR* LexToString(EName Ename)
 {
 	switch (Ename)
 	{
-#define REGISTER_NAME(num,namestr) case num: return TEXT(#namestr);
+#define REGISTER_NAME(num,name) case EName::name: return TEXT(PREPROCESSOR_TO_STRING(name));
 #include "UObject/UnrealNames.inl"
 #undef REGISTER_NAME
 		default:
@@ -1031,7 +1031,7 @@ private:
 	FNamePoolShard<ENameCase::IgnoreCase> ComparisonShards[FNamePoolShards];
 
 	// Put constant lookup on separate cache line to avoid it being constantly invalidated by insertion
-	alignas(PLATFORM_CACHE_LINE_SIZE) FNameEntryId ENameToEntry[NAME_MaxHardcodedNameIndex] = {};
+	alignas(PLATFORM_CACHE_LINE_SIZE) FNameEntryId ENameToEntry[(uint32)EName::MaxHardcodedNameIndex] = {};
 	uint32 LargestEnameUnstableId;
 	TMap<FNameEntryId, EName, TInlineSetAllocator<MaxENames>> EntryToEName;
 };
@@ -1057,9 +1057,9 @@ FNamePool::FNamePool()
 
 	// Make reverse mapping
 	LargestEnameUnstableId = 0;
-	for (uint32 ENameIndex = 0; ENameIndex < NAME_MaxHardcodedNameIndex; ++ENameIndex)
+	for (uint32 ENameIndex = 0; ENameIndex < (uint32)EName::MaxHardcodedNameIndex; ++ENameIndex)
 	{
-		if (ENameIndex == NAME_None || ENameToEntry[ENameIndex])
+		if (ENameIndex == (uint32)NAME_None || ENameToEntry[ENameIndex])
 		{
 			EntryToEName.Add(ENameToEntry[ENameIndex], (EName)ENameIndex);
 			LargestEnameUnstableId = FMath::Max(LargestEnameUnstableId, ENameToEntry[ENameIndex].ToUnstableInt());
@@ -1096,8 +1096,8 @@ static bool IsPureAnsi(const WIDECHAR* Str, const int32 Len)
 
 FNameEntryId FNamePool::Find(EName Ename) const
 {
-	checkSlow(Ename < NAME_MaxHardcodedNameIndex);
-	return ENameToEntry[Ename];
+	checkSlow(Ename < EName::MaxHardcodedNameIndex);
+	return ENameToEntry[(uint32)Ename];
 }
 
 FNameEntryId FNamePool::Find(FNameStringView Name) const
diff --git a/Engine/Source/Runtime/Core/Public/UObject/UnrealNames.h b/Engine/Source/Runtime/Core/Public/UObject/UnrealNames.h
--- a/Engine/Source/Runtime/Core/Public/UObject/UnrealNames.h
+++ b/Engine/Source/Runtime/Core/Public/UObject/UnrealNames.h
@@ -11,16 +11,19 @@ struct FNameEntryId;
 //
 
 // Define a message as an enumeration.
-#define REGISTER_NAME(num,name) NAME_##name = num,
-enum EName 
+#define REGISTER_NAME(num,name) name = num,
+enum class EName : uint32
 {
 	// Include all the hard-coded names
 	#include "UnrealNames.inl"
-
 	// Special constant for the last hard-coded name index
-	NAME_MaxHardcodedNameIndex,
+	MaxHardcodedNameIndex,
 };
 #undef REGISTER_NAME
+// Define aliases for the old-style EName enum syntax
+#define REGISTER_NAME(num,name) inline constexpr EName NAME_##name = EName::name;
+#include "UnrealNames.inl"
+#undef REGISTER_NAME
 
 CORE_API const TCHAR* LexToString(EName Ename);
 
@@ -36,5 +39,5 @@ CORE_API const TCHAR* LexToString(EName Ename);
 
 inline bool ShouldReplicateAsInteger(EName Ename)
 {
-	return Ename <= MAX_NETWORKED_HARDCODED_NAME;
+	return Ename <= EName(MAX_NETWORKED_HARDCODED_NAME);
 }
diff --git a/Engine/Source/Runtime/Core/Public/UObject/UnrealNames.inl b/Engine/Source/Runtime/Core/Public/UObject/UnrealNames.inl
--- a/Engine/Source/Runtime/Core/Public/UObject/UnrealNames.inl
+++ b/Engine/Source/Runtime/Core/Public/UObject/UnrealNames.inl
@@ -2,6 +2,12 @@
 
 // [[ IncludeTool: Inline ]] // Markup to tell IncludeTool that this file is state changing and cannot be optimized out.
 
+// Suppress any macros that may conflict with name constant symbols.
+UE_PUSH_MACRO("TRUE")
+#undef TRUE
+UE_PUSH_MACRO("FALSE")
+#undef FALSE
+
 // Special zero value, meaning no name.
 REGISTER_NAME(0,None)
 
@@ -210,4 +216,8 @@ REGISTER_NAME(602,GameUserSettings)
 // Commonly multicased names - defined globally to prevent non-determinism in package name tables
 REGISTER_NAME(700,Filename)
 REGISTER_NAME(701,Lerp)
-REGISTER_NAME(702,Root)
\ No newline at end of file
+REGISTER_NAME(702,Root)
+
+// Restore any macros suppressed at the top of this file.
+UE_POP_MACRO("TRUE")
+UE_POP_MACRO("FALSE")
diff --git a/Engine/Source/Runtime/CoreUObject/Private/Serialization/AsyncLoading.cpp b/Engine/Source/Runtime/CoreUObject/Private/Serialization/AsyncLoading.cpp
--- a/Engine/Source/Runtime/CoreUObject/Private/Serialization/AsyncLoading.cpp
+++ b/Engine/Source/Runtime/CoreUObject/Private/Serialization/AsyncLoading.cpp
@@ -1652,8 +1652,6 @@ FObjectImport* FAsyncPackage::FindExistingImport(int32 LocalImportIndex)
 					{
 						static const FName NAME_BlueprintGeneratedClass("BlueprintGeneratedClass");
 						static const FName NAME_DynamicClass("DynamicClass");
-
-						static const FName NAME_Function("Function");
 						static const FName NAME_DelegateFunction("DelegateFunction");
 
 						bool bSafeException = (NameImportClass == NAME_BlueprintGeneratedClass && NameActualImportClass == NAME_DynamicClass)
diff --git a/Engine/Source/Runtime/CoreUObject/Private/UObject/CoreNet.cpp b/Engine/Source/Runtime/CoreUObject/Private/UObject/CoreNet.cpp
--- a/Engine/Source/Runtime/CoreUObject/Private/UObject/CoreNet.cpp
+++ b/Engine/Source/Runtime/CoreUObject/Private/UObject/CoreNet.cpp
@@ -299,7 +299,7 @@ bool UPackageMap::StaticSerializeName(FArchive& Ar, FName& InName)
 				Ar.SerializeIntPacked(NameIndex);
 			}
 
-			if (NameIndex < NAME_MaxHardcodedNameIndex)
+			if (NameIndex < (uint32)EName::MaxHardcodedNameIndex)
 			{
 				InName = EName(NameIndex);
 				// hardcoded names never have a Number
@@ -327,7 +327,7 @@ bool UPackageMap::StaticSerializeName(FArchive& Ar, FName& InName)
 		{
 			// send by hardcoded index
 			checkSlow(InName.GetNumber() <= 0); // hardcoded names should never have a Number
-			uint32 NameIndex = *InEName;
+			uint32 NameIndex = (uint32)*InEName;
 			Ar.SerializeIntPacked(NameIndex);
 		}
 		else
diff --git a/Engine/Source/Runtime/Engine/Classes/Animation/AnimNodeBase.h b/Engine/Source/Runtime/Engine/Classes/Animation/AnimNodeBase.h
--- a/Engine/Source/Runtime/Engine/Classes/Animation/AnimNodeBase.h
+++ b/Engine/Source/Runtime/Engine/Classes/Animation/AnimNodeBase.h
@@ -963,8 +963,8 @@ private:
 				static UE::Anim::FNodeDataId CachedId_##Identifier; \
 				if(!CachedId_##Identifier.IsValid()) \
 				{ \
-					static const FName NAME_##Identifier(#Identifier); \
-					CachedId_##Identifier = UE::Anim::FNodeDataId(NAME_##Identifier, this, StaticStruct()); \
+					static const FName AnimName_##Identifier(#Identifier); \
+					CachedId_##Identifier = UE::Anim::FNodeDataId(AnimName_##Identifier, this, StaticStruct()); \
 				} \
 				return CachedId_##Identifier; \
 			}() \
@@ -980,8 +980,8 @@ private:
 				static UE::Anim::FNodeDataId CachedId_##Identifier; \
 				if(!CachedId_##Identifier.IsValid()) \
 				{ \
-					static const FName NAME_##Identifier(#Identifier); \
-					CachedId_##Identifier = UE::Anim::FNodeDataId(NAME_##Identifier, this, StaticStruct()); \
+					static const FName AnimName_##Identifier(#Identifier); \
+					CachedId_##Identifier = UE::Anim::FNodeDataId(AnimName_##Identifier, this, StaticStruct()); \
 				} \
 				return CachedId_##Identifier; \
 			}() \
diff --git a/Engine/Source/Runtime/Engine/Private/PlayerCameraManager.cpp b/Engine/Source/Runtime/Engine/Private/PlayerCameraManager.cpp
--- a/Engine/Source/Runtime/Engine/Private/PlayerCameraManager.cpp
+++ b/Engine/Source/Runtime/Engine/Private/PlayerCameraManager.cpp
@@ -38,8 +38,6 @@ DECLARE_CYCLE_STAT(TEXT("Camera ProcessViewRotation"), STAT_Camera_ProcessViewRo
 APlayerCameraManager::APlayerCameraManager(const FObjectInitializer& ObjectInitializer)
 	: Super(ObjectInitializer)
 {
-	static FName NAME_Default(TEXT("Default"));
-
 	DefaultFOV = 90.0f;
 	DefaultAspectRatio = 1.33333f;
 	bDefaultConstrainAspectRatio = false;
diff --git a/Engine/Source/Runtime/Engine/Private/PrimitiveComponentPhysics.cpp b/Engine/Source/Runtime/Engine/Private/PrimitiveComponentPhysics.cpp
--- a/Engine/Source/Runtime/Engine/Private/PrimitiveComponentPhysics.cpp
+++ b/Engine/Source/Runtime/Engine/Private/PrimitiveComponentPhysics.cpp
@@ -296,13 +296,13 @@ void UPrimitiveComponent::SetAllPhysicsAngularVelocityInRadians(FVector const& N
 
 void UPrimitiveComponent::SetAllPhysicsPosition(FVector NewPos)
 {
-	SetWorldLocation(NewPos, NAME_None);
+	SetWorldLocation(NewPos);
 }
 
 
 void UPrimitiveComponent::SetAllPhysicsRotation(FRotator NewRot)
 {
-	SetWorldRotation(NewRot, NAME_None);
+	SetWorldRotation(NewRot);
 }
 
 void UPrimitiveComponent::SetAllPhysicsRotation(const FQuat& NewRot)
diff --git a/Engine/Source/Runtime/JsonUtilities/Private/JsonObjectConverter.cpp b/Engine/Source/Runtime/JsonUtilities/Private/JsonObjectConverter.cpp
--- a/Engine/Source/Runtime/JsonUtilities/Private/JsonObjectConverter.cpp
+++ b/Engine/Source/Runtime/JsonUtilities/Private/JsonObjectConverter.cpp
@@ -606,8 +606,6 @@ namespace
 	else if (FStructProperty *StructProperty = CastField<FStructProperty>(Property))
 	{
 		static const FName NAME_DateTime(TEXT("DateTime"));
-		static const FName NAME_Color(TEXT("Color"));
-		static const FName NAME_LinearColor(TEXT("LinearColor"));
 		if (JsonValue->Type == EJson::Object)
 		{
 			TSharedPtr<FJsonObject> Obj = JsonValue->AsObject();
diff --git a/Engine/Source/Runtime/Online/XMPP/Private/XmppModule.cpp b/Engine/Source/Runtime/Online/XMPP/Private/XmppModule.cpp
--- a/Engine/Source/Runtime/Online/XMPP/Private/XmppModule.cpp
+++ b/Engine/Source/Runtime/Online/XMPP/Private/XmppModule.cpp
@@ -606,16 +606,14 @@ bool FXmppModule::HandleXmppCommand( const TCHAR* Cmd, FOutputDevice& Ar )
 	}
 	else if (FParse::Command(&Cmd, TEXT("LogVerbosity")))
 	{
-		FString Verbosity = FParse::Token(Cmd, false);
-
-		const FString NAME_NoLogging(TEXT("NoLogging"));
-		const FString NAME_Fatal(TEXT("Fatal"));
-		const FString NAME_Error(TEXT("Error"));
-		const FString NAME_Warning(TEXT("Warning"));
-		const FString NAME_Display(TEXT("Display"));
-		const FString NAME_Log(TEXT("Log"));
-		const FString NAME_Verbose(TEXT("Verbose"));
-		const FString NAME_VeryVerbose(TEXT("VeryVerbose"));
+		FName Verbosity = *FParse::Token(Cmd, false);
+
+		const FName NAME_NoLogging(TEXT("NoLogging"));
+		const FName NAME_Fatal(TEXT("Fatal"));
+		const FName NAME_Display(TEXT("Display"));
+		const FName NAME_Log(TEXT("Log"));
+		const FName NAME_Verbose(TEXT("Verbose"));
+		const FName NAME_VeryVerbose(TEXT("VeryVerbose"));
 
 		if (Verbosity == NAME_NoLogging)
 		{
diff --git a/Engine/Source/Runtime/RenderCore/Private/ShaderParameters.cpp b/Engine/Source/Runtime/RenderCore/Private/ShaderParameters.cpp
--- a/Engine/Source/Runtime/RenderCore/Private/ShaderParameters.cpp
+++ b/Engine/Source/Runtime/RenderCore/Private/ShaderParameters.cpp
@@ -431,7 +431,6 @@ void FShaderType::DumpDebugInfo()
 void FShaderType::GetShaderStableKeyParts(FStableShaderKeyAndValue& SaveKeyVal)
 {
 #if WITH_EDITOR
-	static FName NAME_Global(TEXT("Global"));
 	static FName NAME_Material(TEXT("Material"));
 	static FName NAME_MeshMaterial(TEXT("MeshMaterial"));
 	static FName NAME_Niagara(TEXT("Niagara"));
