GenericPlatformMath avoid calling floorf for common FloorToInt used by RoundToInt
not used where UnrealPlatformMathSSE overrides

diff --git a/Engine/Source/Runtime/Core/Public/GenericPlatform/GenericPlatformMath.h b/Engine/Source/Runtime/Core/Public/GenericPlatform/GenericPlatformMath.h
--- a/Engine/Source/Runtime/Core/Public/GenericPlatform/GenericPlatformMath.h
+++ b/Engine/Source/Runtime/Core/Public/GenericPlatform/GenericPlatformMath.h
@@ -51,6 +51,7 @@ struct FGenericPlatformMath
 	}
 
 	//https://gist.github.com/rygorous/2156668
+	// float_to_half_fast3_rtne
 	static FORCEINLINE void StoreHalf(uint16* Ptr, float Value)
 	{
 		union FP32T
@@ -226,7 +227,9 @@ struct FGenericPlatformMath
 	 */
 	static FORCEINLINE int32 FloorToInt(float F)
 	{
-		return TruncToInt(floorf(F));
+		int32 I = (int32)F;
+		I -= ( (float)I > F );
+		return I;
 	}
 
 	/**
@@ -286,7 +289,9 @@ struct FGenericPlatformMath
 	*/
 	static FORCEINLINE int32 CeilToInt(float F)
 	{
-		return TruncToInt(ceilf(F));
+		int32 I = (int32)F;
+		I += ( (float)I < F );
+		return I;
 	}
 
 	/**
