fix Color.h warnings

diff --git a/Engine/Source/Runtime/Core/Public/Math/Color.h b/Engine/Source/Runtime/Core/Public/Math/Color.h
--- a/Engine/Source/Runtime/Core/Public/Math/Color.h
+++ b/Engine/Source/Runtime/Core/Public/Math/Color.h
@@ -556,14 +556,14 @@ public:
 	{
 		check( Value8 >= 0 && Value8 <= 255 );
 
-		return Value8 / 255.f;
+		return (float)Value8 / 255.f;
 	}
 	
 	static float DequantizeUNorm16ToFloat( int Value16 )
 	{
 		check( Value16 >= 0 && Value16 <= 65535 );
 
-		return Value16 / 65535.f;
+		return (float)Value16 / 65535.f;
 	}
 
 	static uint8 Requantize10to8( int Value10 )
