Old deprecated GetVarArgs overloads removed.

diff --git a/Engine/Source/Runtime/Core/Public/Android/AndroidPlatformString.h b/Engine/Source/Runtime/Core/Public/Android/AndroidPlatformString.h
--- a/Engine/Source/Runtime/Core/Public/Android/AndroidPlatformString.h
+++ b/Engine/Source/Runtime/Core/Public/Android/AndroidPlatformString.h
@@ -413,12 +413,6 @@ struct FAndroidPlatformString : public FGenericPlatformString
 		return StrToken + (Pos - AnsiStrToken);
 	}
 
-	UE_DEPRECATED(4.22, "GetVarArgs with DestSize and Count arguments has been deprecated - only DestSize should be passed")
-	static FORCEINLINE int32 GetVarArgs( WIDECHAR* Dest, SIZE_T DestSize, int32 Count, const WIDECHAR*& Fmt, va_list ArgPtr )
-	{
-		return GetVarArgs(Dest, DestSize, Fmt, ArgPtr);
-	}
-
 	static FORCEINLINE int32 GetVarArgs( WIDECHAR* Dest, SIZE_T DestSize, const WIDECHAR*& Fmt, va_list ArgPtr )
 	{
 #if PLATFORM_USE_LS_SPEC_FOR_WIDECHAR
@@ -573,12 +567,6 @@ struct FAndroidPlatformString : public FGenericPlatformString
 		return strtok(StrToken, Delim);
 	}
 
-	UE_DEPRECATED(4.22, "GetVarArgs with DestSize and Count arguments has been deprecated - only DestSize should be passed")
-	static FORCEINLINE int32 GetVarArgs(ANSICHAR* Dest, SIZE_T DestSize, int32 Count, const ANSICHAR*& Fmt, va_list ArgPtr)
-	{
-		return GetVarArgs(Dest, DestSize, Count, Fmt, ArgPtr);
-	}
-
 	static FORCEINLINE int32 GetVarArgs( ANSICHAR* Dest, SIZE_T DestSize, const ANSICHAR*& Fmt, va_list ArgPtr )
 	{
 		int32 Result = vsnprintf(Dest, DestSize, Fmt, ArgPtr);
diff --git a/Engine/Source/Runtime/Core/Public/GenericPlatform/GenericWidePlatformString.h b/Engine/Source/Runtime/Core/Public/GenericPlatform/GenericWidePlatformString.h
--- a/Engine/Source/Runtime/Core/Public/GenericPlatform/GenericWidePlatformString.h
+++ b/Engine/Source/Runtime/Core/Public/GenericPlatform/GenericWidePlatformString.h
@@ -198,12 +198,6 @@ public:
 	
 	static CORE_API WIDECHAR* Strtok( WIDECHAR* StrToken, const WIDECHAR* Delim, WIDECHAR** Context );
 
-	UE_DEPRECATED(4.22, "GetVarArgs with DestSize and Count arguments has been deprecated - only DestSize should be passed")
-	static 	CORE_API int32 GetVarArgs( WIDECHAR* Dest, SIZE_T DestSize, int32 Count, const WIDECHAR*& Fmt, va_list ArgPtr )
-	{
-		return GetVarArgs(Dest, DestSize, Fmt, ArgPtr);
-	}
-
 	static CORE_API int32 GetVarArgs( WIDECHAR* Dest, SIZE_T DestSize, const WIDECHAR*& Fmt, va_list ArgPtr );
 
 	/**
@@ -301,12 +295,6 @@ public:
 		return strtok(StrToken, Delim);
 	}
 
-	UE_DEPRECATED(4.22, "GetVarArgs with DestSize and Count arguments has been deprecated - only DestSize should be passed")
-	static CORE_API int32 GetVarArgs( ANSICHAR* Dest, SIZE_T DestSize, int32 Count, const ANSICHAR*& Fmt, va_list ArgPtr )
-	{
-		return GetVarArgs(Dest, DestSize, Fmt, ArgPtr);
-	}
-
 	static CORE_API int32 GetVarArgs( ANSICHAR* Dest, SIZE_T DestSize, const ANSICHAR*& Fmt, va_list ArgPtr )
 	{
 		int32 Result = vsnprintf(Dest, DestSize, Fmt, ArgPtr);
diff --git a/Engine/Source/Runtime/Core/Public/GenericPlatform/MicrosoftPlatformString.h b/Engine/Source/Runtime/Core/Public/GenericPlatform/MicrosoftPlatformString.h
--- a/Engine/Source/Runtime/Core/Public/GenericPlatform/MicrosoftPlatformString.h
+++ b/Engine/Source/Runtime/Core/Public/GenericPlatform/MicrosoftPlatformString.h
@@ -125,12 +125,6 @@ struct FMicrosoftPlatformString : public FGenericPlatformString
 		return _tcstok_s(StrToken, Delim, Context);
 	}
 
-	UE_DEPRECATED(4.22, "GetVarArgs with DestSize and Count arguments has been deprecated - only DestSize should be passed")
-	static FORCEINLINE int32 GetVarArgs( WIDECHAR* Dest, SIZE_T DestSize, int32 Count, const WIDECHAR*& Fmt, va_list ArgPtr )
-	{
-		return GetVarArgs(Dest, DestSize, Fmt, ArgPtr);
-	}
-
 	static FORCEINLINE int32 GetVarArgs( WIDECHAR* Dest, SIZE_T DestSize, const WIDECHAR*& Fmt, va_list ArgPtr )
 	{
 		int32 Result = vswprintf(Dest, DestSize, Fmt, ArgPtr);
@@ -237,12 +231,6 @@ struct FMicrosoftPlatformString : public FGenericPlatformString
 		return strtok_s(StrToken, Delim, Context);
 	}
 
-	UE_DEPRECATED(4.22, "GetVarArgs with DestSize and Count arguments has been deprecated - only DestSize should be passed")
-	static FORCEINLINE int32 GetVarArgs( ANSICHAR* Dest, SIZE_T DestSize, int32 Count, const ANSICHAR*& Fmt, va_list ArgPtr )
-	{
-		return GetVarArgs(Dest, DestSize, Fmt, ArgPtr);
-	}
-
 	static FORCEINLINE int32 GetVarArgs( ANSICHAR* Dest, SIZE_T DestSize, const ANSICHAR*& Fmt, va_list ArgPtr )
 	{
 		int32 Result = vsnprintf( Dest, DestSize, Fmt, ArgPtr );
diff --git a/Engine/Source/Runtime/Core/Public/GenericPlatform/StandardPlatformString.h b/Engine/Source/Runtime/Core/Public/GenericPlatform/StandardPlatformString.h
--- a/Engine/Source/Runtime/Core/Public/GenericPlatform/StandardPlatformString.h
+++ b/Engine/Source/Runtime/Core/Public/GenericPlatform/StandardPlatformString.h
@@ -129,12 +129,6 @@ public:
 		return wcstok(StrToken, Delim, Context);
 	}
 
-	UE_DEPRECATED(4.22, "GetVarArgs with DestSize and Count arguments has been deprecated - only DestSize should be passed")
-	static CORE_API int32 GetVarArgs(WIDECHAR* Dest, SIZE_T DestSize, int32 Count, const WIDECHAR*& Fmt, va_list ArgPtr)
-	{
-		return GetVarArgs(Dest, DestSize, Fmt, ArgPtr);
-	}
-
 #if PLATFORM_USE_SYSTEM_VSWPRINTF
 	static CORE_API int32 GetVarArgs( WIDECHAR* Dest, SIZE_T DestSize, const WIDECHAR*& Fmt, va_list ArgPtr )
 	{
@@ -295,12 +289,6 @@ public:
 		return strtok(StrToken, Delim);
 	}
 
-	UE_DEPRECATED(4.22, "GetVarArgs with DestSize and Count arguments has been deprecated - only DestSize should be passed")
-	static CORE_API int32 GetVarArgs( ANSICHAR* Dest, SIZE_T DestSize, int32 Count, const ANSICHAR*& Fmt, va_list ArgPtr )
-	{
-		return GetVarArgs(Dest, DestSize, Fmt, ArgPtr);
-	}
-
 	static CORE_API int32 GetVarArgs( ANSICHAR* Dest, SIZE_T DestSize, const ANSICHAR*& Fmt, va_list ArgPtr )
 	{
 		int32 Result = vsnprintf(Dest, DestSize, Fmt, ArgPtr);
diff --git a/Engine/Source/Runtime/Core/Public/Misc/CString.h b/Engine/Source/Runtime/Core/Public/Misc/CString.h
--- a/Engine/Source/Runtime/Core/Public/Misc/CString.h
+++ b/Engine/Source/Runtime/Core/Public/Misc/CString.h
@@ -394,22 +394,6 @@ public:
 		return SnprintfImpl(Dest, DestSize, Fmt, Args...);
 	}
 
-	/**
-	 * Helper function to write formatted output using an argument list
-	 *
-	 * @param Dest - destination string buffer
-	 * @param DestSize - size of destination buffer
-	 * @param Count - number of characters to write (not including null terminating character)
-	 * @param Fmt - string to print
-	 * @param Args - argument list
-	 * @return number of characters written or -1 if truncated
-	 */
-	UE_DEPRECATED(4.22, "GetVarArgs with DestSize and Count arguments has been deprecated - only DestSize should be passed")
-	static FORCEINLINE int32 GetVarArgs(CharType* Dest, SIZE_T DestSize, int32 Count, const CharType*& Fmt, va_list ArgPtr)
-	{
-		return GetVarArgs(Dest, DestSize, Fmt, ArgPtr);
-	}
-
 	/**
 	 * Helper function to write formatted output using an argument list
 	 *
