Float16Color use VectorStoreHalf/VectorLoadHalf
add SSE2 half conversions

diff --git a/Engine/Plugins/Experimental/MeshModelingToolset/Source/ModelingComponents/Private/AssetUtils/Texture2DUtil.cpp b/Engine/Plugins/Experimental/MeshModelingToolset/Source/ModelingComponents/Private/AssetUtils/Texture2DUtil.cpp
--- a/Engine/Plugins/Experimental/MeshModelingToolset/Source/ModelingComponents/Private/AssetUtils/Texture2DUtil.cpp
+++ b/Engine/Plugins/Experimental/MeshModelingToolset/Source/ModelingComponents/Private/AssetUtils/Texture2DUtil.cpp
@@ -103,9 +103,7 @@ static bool ReadTexture_SourceData(
 		for (int32 i = 0; i < Num; ++i)
 		{
 			const uint8* PixelPtr = SourceDataPtr + (i * BytesPerPixel);
-			FFloat16Color PixelColor = *((FFloat16Color*)PixelPtr);
-			FLinearColor FloatColor(float(PixelColor.R), float(PixelColor.G), float(PixelColor.B), float(PixelColor.A));
-			DestImage.SetPixel(i, FVector4f(FloatColor));
+			DestImage.SetPixel(i, FVector4f( ((const FFloat16Color*)PixelPtr)->GetFloats() ));
 		}
 	}
 	else if (SourceFormat == TSF_G8)
diff --git a/Engine/Plugins/FX/Niagara/Source/Niagara/Private/NiagaraRenderer.cpp b/Engine/Plugins/FX/Niagara/Source/Niagara/Private/NiagaraRenderer.cpp
--- a/Engine/Plugins/FX/Niagara/Source/Niagara/Private/NiagaraRenderer.cpp
+++ b/Engine/Plugins/FX/Niagara/Source/Niagara/Private/NiagaraRenderer.cpp
@@ -371,7 +371,7 @@ struct FParticleOrderAsUint
 	{
 		const uint32 SortKeySignBit = 0x8000;
 		uint32 InOrderAsUint = InOrder.Encoded;
-		InOrderAsUint = (bStrictlyPositive || InOrder.Components.Sign != 0) ? (InOrderAsUint | SortKeySignBit) : ~InOrderAsUint;
+		InOrderAsUint = (bStrictlyPositive || InOrder.IsNegative()) ? (InOrderAsUint | SortKeySignBit) : ~InOrderAsUint;
 		OrderAsUint = bAscending ? InOrderAsUint : ~InOrderAsUint;
 		OrderAsUint &= 0xFFFF;
 		Index = InIndex;
diff --git a/Engine/Source/Developer/TextureFormatIntelISPCTexComp/Private/TextureFormatIntelISPCTexComp.cpp b/Engine/Source/Developer/TextureFormatIntelISPCTexComp/Private/TextureFormatIntelISPCTexComp.cpp
--- a/Engine/Source/Developer/TextureFormatIntelISPCTexComp/Private/TextureFormatIntelISPCTexComp.cpp
+++ b/Engine/Source/Developer/TextureFormatIntelISPCTexComp/Private/TextureFormatIntelISPCTexComp.cpp
@@ -653,10 +653,10 @@ public:
 		{
 			// Flush negative values to 0, as those aren't supported by BC6H_UF16.
 			FFloat16& F16Value = Data[TexelIndex];
-			const bool bNegative = F16Value.Components.Sign == 1;
-			if (bNegative)
+
+			if ( F16Value.IsNegative() )
 			{
-				F16Value = 0;
+				F16Value.Encoded = 0;
 			}
 		}
 	}
diff --git a/Engine/Source/Runtime/Core/Private/Math/Color.cpp b/Engine/Source/Runtime/Core/Private/Math/Color.cpp
--- a/Engine/Source/Runtime/Core/Private/Math/Color.cpp
+++ b/Engine/Source/Runtime/Core/Private/Math/Color.cpp
@@ -61,11 +61,7 @@ FLinearColor::FLinearColor(const FVector4& Vector) :
 
 FLinearColor::FLinearColor(const FFloat16Color& C)
 {
-	// @todo Oodle use 4X half to float SIMD (VectorLoadHalf/VectorStoreHalf)
-	R = C.R.GetFloat();
-	G = C.G.GetFloat();
-	B =	C.B.GetFloat();
-	A =	C.A.GetFloat();
+	*this = C.GetFloats();
 }
 
 FLinearColor FLinearColor::FromPow22Color(const FColor& Color)
diff --git a/Engine/Source/Runtime/Core/Public/Math/Color.h b/Engine/Source/Runtime/Core/Public/Math/Color.h
--- a/Engine/Source/Runtime/Core/Public/Math/Color.h
+++ b/Engine/Source/Runtime/Core/Public/Math/Color.h
@@ -59,6 +59,7 @@ struct FLinearColor
 
 	CORE_API explicit FLinearColor(const FVector4& Vector);
 	
+	// use Float16Color::GetFloats() directly
 	CORE_API explicit FLinearColor(const FFloat16Color& C);
 
 	// Serializer.
diff --git a/Engine/Source/Runtime/Core/Public/Math/Float16.h b/Engine/Source/Runtime/Core/Public/Math/Float16.h
--- a/Engine/Source/Runtime/Core/Public/Math/Float16.h
+++ b/Engine/Source/Runtime/Core/Public/Math/Float16.h
@@ -22,28 +22,15 @@
 * 0<E<31, M=any		== (1 + M / 2^10) * 2^(E-15)
 * E=31, M=0			== Infinity
 * E=31, M!=0		== NAN
+* 
+* conversion from 32 bit float is with RTNE (round to nearest even)
 *
 */
 class FFloat16
 {
 public:
-	union
-	{
-		struct
-		{
-#if PLATFORM_LITTLE_ENDIAN
-			uint16	Mantissa : 10;
-			uint16	Exponent : 5;
-			uint16	Sign : 1;
-#else
-			uint16	Sign : 1;
-			uint16	Exponent : 5;
-			uint16	Mantissa : 10;			
-#endif
-		} Components;
-
-		uint16	Encoded;
-	};
+
+	uint16 Encoded;
 
 	/** Default constructor */
 	FFloat16();
@@ -66,19 +53,17 @@ public:
 	/** Convert from Fp32 to Fp16. */
 	void Set(float FP32Value);
 	
-	/**
-	 * Convert from Fp32 to Fp16 without doing any checks if
-	 * the Fp32 exponent is too large or too small. This is a
-	 * faster alternative to Set() when you know the values
-	 * within the single precision float don't need the checks.
-	 *
-	 * @param FP32Value Single precision float to be set as half precision.
-	 */
-	void SetWithoutBoundsChecks(const float FP32Value);
-
 	/** Convert from Fp16 to Fp32. */
 	float GetFloat() const;
 
+	/** Is the float negative without converting */
+	bool IsNegative() const
+	{
+		// negative if sign bit is on
+		// can be tested with int compare
+		return (int16)Encoded < 0;
+	}
+
 	/**
 	 * Serializes the FFloat16.
 	 *
@@ -135,96 +120,13 @@ FORCEINLINE FFloat16::operator float() const
 
 FORCEINLINE void FFloat16::Set(float FP32Value)
 {
-	FFloat32 FP32(FP32Value);
-
-	// Copy sign-bit
-	Components.Sign = FP32.Components.Sign;
-
-	// Check for zero, denormal or too small value.
-	if (FP32.Components.Exponent <= 112)			// Too small exponent? (0+127-15)
-	{
-		// Set to 0.
-		Components.Exponent = 0;
-		Components.Mantissa = 0;
-
-         // Exponent unbias the single, then bias the halfp
-         const int32 NewExp = FP32.Components.Exponent - 127 + 15;
- 
-         if ( (14 - NewExp) <= 24 ) // Mantissa might be non-zero
-         {
-             uint32 Mantissa = FP32.Components.Mantissa | 0x800000; // Hidden 1 bit
-             Components.Mantissa = (uint16)(Mantissa >> (14 - NewExp));
-			 // Check for rounding
-             if ( (Mantissa >> (13 - NewExp)) & 1 ) //-V1051
-			 {
-                 Encoded++; // Round, might overflow into exp bit, but this is OK
-			 }
-         }
-	}
-	// Check for INF or NaN, or too high value
-	else if (FP32.Components.Exponent >= 143)		// Too large exponent? (31+127-15)
-	{
-		// Set to 65504.0 (max value)
-		Components.Exponent = 30;
-		Components.Mantissa = 1023;
-	}
-	// Handle normal number.
-	else
-	{
-		Components.Exponent = uint16(int32(FP32.Components.Exponent) - 127 + 15);
-		Components.Mantissa = uint16(FP32.Components.Mantissa >> 13);
-	}
+	// FPlatformMath::StoreHalf follows RTNE (round-to-nearest-even) rounding default convention
+	FPlatformMath::StoreHalf(&Encoded, FP32Value);
 }
 
 
-FORCEINLINE void FFloat16::SetWithoutBoundsChecks(const float FP32Value)
-{
-	const FFloat32 FP32(FP32Value);
-
-	// Make absolutely sure that you never pass in a single precision floating
-	// point value that may actually need the checks. If you are not 100% sure
-	// of that just use Set().
-
-	Components.Sign = FP32.Components.Sign;
-	Components.Exponent = uint16(int32(FP32.Components.Exponent) - 127 + 15);
-	Components.Mantissa = uint16(FP32.Components.Mantissa >> 13);
-}
-
 
 FORCEINLINE float FFloat16::GetFloat() const
 {
-	FFloat32	Result;
-
-	Result.Components.Sign = Components.Sign;
-	if (Components.Exponent == 0)
-	{
-		uint32 Mantissa = Components.Mantissa;
-		if(Mantissa == 0)
-		{
-			// Zero.
-			Result.Components.Exponent = 0;
-			Result.Components.Mantissa = 0;
-		}
-		else
-		{
-			// Denormal.
-			uint32 MantissaShift = 10 - (uint32)FMath::TruncToInt(FMath::Log2((float)Mantissa));
-			Result.Components.Exponent = 127 - (15 - 1) - MantissaShift;
-			Result.Components.Mantissa = Mantissa << (MantissaShift + 23 - 10);
-		}
-	}
-	else if (Components.Exponent == 31)		// 2^5 - 1
-	{
-		// Infinity or NaN. Set to 65504.0
-		Result.Components.Exponent = 142;
-		Result.Components.Mantissa = 8380416;
-	}
-	else
-	{
-		// Normal number.
-		Result.Components.Exponent = int32(Components.Exponent) - 15 + 127; // Stored exponents are biased by half their range.
-		Result.Components.Mantissa = uint32(Components.Mantissa) << 13;
-	}
-
-	return Result.FloatValue;
+	return FPlatformMath::LoadHalf(&Encoded);
 }
diff --git a/Engine/Source/Runtime/Core/Public/Math/Float16Color.h b/Engine/Source/Runtime/Core/Public/Math/Float16Color.h
--- a/Engine/Source/Runtime/Core/Public/Math/Float16Color.h
+++ b/Engine/Source/Runtime/Core/Public/Math/Float16Color.h
@@ -18,6 +18,12 @@ public:
 	FFloat16 B;
 	FFloat16 A;
 
+	/* Get as a pointer to four half floats */
+	uint16 * GetFourHalves() { return (uint16 *)this; }
+	const uint16 * GetFourHalves() const { return (const uint16 *)this; }
+
+	const FLinearColor GetFloats() const;
+
 	/** Default constructor */
 	FFloat16Color();
 
@@ -36,8 +42,6 @@ public:
 	 * @param Src The other color.
 	 * @return true if the two colors are identical, otherwise false.
 	 */
-	bool operator==(const FFloat16Color& Src);
-
 	bool operator==(const FFloat16Color& Src) const;
 };
 
@@ -54,12 +58,18 @@ FORCEINLINE FFloat16Color::FFloat16Color(const FFloat16Color& Src)
 }
 
 
-FORCEINLINE FFloat16Color::FFloat16Color(const FLinearColor& Src) :
-	R(Src.R),
-	G(Src.G),
-	B(Src.B),
-	A(Src.A)
-{ }
+FORCEINLINE FFloat16Color::FFloat16Color(const FLinearColor& Src)
+{
+	FPlatformMath::VectorStoreHalf( GetFourHalves(), (const float *)&Src );
+}
+
+	
+FORCEINLINE const FLinearColor FFloat16Color::GetFloats() const
+{
+	FLinearColor Ret;
+	FPlatformMath::VectorLoadHalf( (float *)&Ret, GetFourHalves() );
+	return Ret;
+}
 
 
 FORCEINLINE FFloat16Color& FFloat16Color::operator=(const FFloat16Color& Src)
@@ -71,17 +81,6 @@ FORCEINLINE FFloat16Color& FFloat16Color::operator=(const FFloat16Color& Src)
 	return *this;
 }
 
-
-FORCEINLINE bool FFloat16Color::operator==(const FFloat16Color& Src)
-{
-	return (
-		(R == Src.R) &&
-		(G == Src.G) &&
-		(B == Src.B) &&
-		(A == Src.A)
-		);
-}
-
 FORCEINLINE bool FFloat16Color::operator==(const FFloat16Color& Src) const
 {
 	return (
diff --git a/Engine/Source/Runtime/Core/Public/Math/UnrealPlatformMathSSE.h b/Engine/Source/Runtime/Core/Public/Math/UnrealPlatformMathSSE.h
--- a/Engine/Source/Runtime/Core/Public/Math/UnrealPlatformMathSSE.h
+++ b/Engine/Source/Runtime/Core/Public/Math/UnrealPlatformMathSSE.h
@@ -106,6 +106,88 @@ namespace SSE
 		// We therefore need to do the same times-2 transform (with a slighly different formula) that RoundToInt does; see the note on RoundToInt
 		return -(_mm_cvt_ss2si(_mm_set_ss(-0.5f - (F + F))) >> 1);
 	}
+
+	// https://gist.github.com/rygorous/2156668
+	// float_to_half_rtne_SSE2
+	static __m128i FloatToHalf(__m128 f)
+	{
+		__m128i mask_sign       = _mm_set1_epi32(0x80000000u);
+		__m128i c_f16max        = _mm_set1_epi32((127 + 16) << 23); // all FP32 values >=this round to +inf
+		__m128i c_nanbit        = _mm_set1_epi32(0x200);
+		__m128i c_infty_as_fp16 = _mm_set1_epi32(0x7c00);
+		__m128i c_min_normal    = _mm_set1_epi32((127 - 14) << 23); // smallest FP32 that yields a normalized FP16
+		__m128i c_subnorm_magic = _mm_set1_epi32(((127 - 15) + (23 - 10) + 1) << 23);
+		__m128i c_normal_bias   = _mm_set1_epi32(0xfff - ((127 - 15) << 23)); // adjust exponent and add mantissa rounding
+
+		__m128  msign       = _mm_castsi128_ps(mask_sign);
+		__m128  justsign    = _mm_and_ps(msign, f);
+		__m128  absf        = _mm_xor_ps(f, justsign);
+		__m128i absf_int    = _mm_castps_si128(absf); // the cast is "free" (extra bypass latency, but no thruput hit)
+		__m128i f16max      = c_f16max;
+		__m128  b_isnan     = _mm_cmpunord_ps(absf, absf); // is this a NaN?
+		__m128i b_isregular = _mm_cmpgt_epi32(f16max, absf_int); // (sub)normalized or special?
+		__m128i nanbit      = _mm_and_si128(_mm_castps_si128(b_isnan), c_nanbit);
+		__m128i inf_or_nan  = _mm_or_si128(nanbit, c_infty_as_fp16); // output for specials
+
+		__m128i min_normal  = c_min_normal;
+		__m128i b_issub     = _mm_cmpgt_epi32(min_normal, absf_int);
+
+		// "result is subnormal" path
+		__m128  subnorm1    = _mm_add_ps(absf, _mm_castsi128_ps(c_subnorm_magic)); // magic value to round output mantissa
+		__m128i subnorm2    = _mm_sub_epi32(_mm_castps_si128(subnorm1), c_subnorm_magic); // subtract out bias
+
+		// "result is normal" path
+		__m128i mantoddbit  = _mm_slli_epi32(absf_int, 31 - 13); // shift bit 13 (mantissa LSB) to sign
+		__m128i mantodd     = _mm_srai_epi32(mantoddbit, 31); // -1 if FP16 mantissa odd, else 0
+
+		__m128i round1      = _mm_add_epi32(absf_int, c_normal_bias);
+		__m128i round2      = _mm_sub_epi32(round1, mantodd); // if mantissa LSB odd, bias towards rounding up (RTNE)
+		__m128i normal      = _mm_srli_epi32(round2, 13); // rounded result
+
+		// combine the two non-specials
+		__m128i nonspecial  = _mm_or_si128(_mm_and_si128(subnorm2, b_issub), _mm_andnot_si128(b_issub, normal));
+
+		// merge in specials as well
+		__m128i joined      = _mm_or_si128(_mm_and_si128(nonspecial, b_isregular), _mm_andnot_si128(b_isregular, inf_or_nan));
+
+		__m128i sign_shift  = _mm_srli_epi32(_mm_castps_si128(justsign), 16);
+		__m128i rgba_half32 = _mm_or_si128(joined, sign_shift);
+		
+		// there's now a half in each 32-bit lane
+		// pack down to 64 bits :
+        __m128i four_halfs_u64 = _mm_packus_epi32(rgba_half32, _mm_setzero_si128());
+
+		return four_halfs_u64;
+	}
+	
+	// four halfs should be in the u64 part of input
+	static __m128 HalfToFloat(__m128i four_halfs_u64)
+	{
+		__m128i rgba_half32 = _mm_unpacklo_epi16(four_halfs_u64, _mm_setzero_si128());
+
+		const __m128i mask_nosign = _mm_set1_epi32(0x7fff);
+		const __m128 magic_mult	= _mm_castsi128_ps(_mm_set1_epi32((254 - 15) << 23));
+		const __m128i was_infnan = _mm_set1_epi32(0x7bff);
+		const __m128 exp_infnan	= _mm_castsi128_ps(_mm_set1_epi32(255 << 23));
+		const __m128i was_nan = _mm_set1_epi32(0x7c00);
+		const __m128i nan_quiet	= _mm_set1_epi32(1 << 22);
+
+		__m128i expmant		= _mm_and_si128(mask_nosign, rgba_half32);
+		__m128i justsign	= _mm_xor_si128(rgba_half32, expmant);
+		__m128i shifted		= _mm_slli_epi32(expmant, 13);
+		__m128 scaled		= _mm_mul_ps(_mm_castsi128_ps(shifted), magic_mult);
+		__m128i b_wasinfnan = _mm_cmpgt_epi32(expmant, was_infnan);
+		__m128i sign	    = _mm_slli_epi32(justsign, 16);
+		__m128 infnanexp    = _mm_and_ps(_mm_castsi128_ps(b_wasinfnan), exp_infnan);
+		__m128i b_wasnan    = _mm_cmpgt_epi32(expmant, was_nan);
+		__m128i nanquiet    = _mm_and_si128(b_wasnan, nan_quiet);
+		__m128 infnandone   = _mm_or_ps(infnanexp, _mm_castsi128_ps(nanquiet));
+
+		__m128 sign_inf	= _mm_or_ps(_mm_castsi128_ps(sign), infnandone);
+		__m128 result	= _mm_or_ps(scaled, sign_inf);
+
+		return result;
+	}
 }
 }
 
@@ -141,6 +223,16 @@ struct TUnrealPlatformMathSSEBase : public Base
 	{
 		return UE4::SSE::InvSqrtEst(F);
 	}
+	
+	static FORCEINLINE void VectorStoreHalf(uint16* RESTRICT Dst, const float* RESTRICT Src)
+	{
+		_mm_storeu_si64((__m128i*)Dst, UE4::SSE::FloatToHalf(_mm_loadu_ps(Src)));
+	}
+	
+	static FORCEINLINE void VectorLoadHalf(float* RESTRICT Dst, const uint16* RESTRICT Src)
+	{
+		_mm_storeu_ps(Dst, UE4::SSE::HalfToFloat(_mm_loadu_si64((__m128i*)Src)));
+	}
 
 };
 
diff --git a/Engine/Source/Runtime/Engine/Private/Particles/ParticleGpuSimulation.cpp b/Engine/Source/Runtime/Engine/Private/Particles/ParticleGpuSimulation.cpp
--- a/Engine/Source/Runtime/Engine/Private/Particles/ParticleGpuSimulation.cpp
+++ b/Engine/Source/Runtime/Engine/Private/Particles/ParticleGpuSimulation.cpp
@@ -2065,8 +2065,11 @@ static void BuildParticleVertexBuffer(FRHIBuffer* VertexBufferRHI, const TArray<
 			{
 				const float IndexX = TileOffset.X + ((float)ParticleX / (float)GParticleSimulationTextureSizeX) + (0.5f / (float)GParticleSimulationTextureSizeX);
 				const float IndexY = TileOffset.Y + ((float)ParticleY / (float)GParticleSimulationTextureSizeY) + (0.5f / (float)GParticleSimulationTextureSizeY);
-				ParticleIndices->X.SetWithoutBoundsChecks(IndexX);
-				ParticleIndices->Y.SetWithoutBoundsChecks(IndexY);					
+
+				// @todo faster float32 -> float16 conversion
+				//	use AVX2/F16C for _mm_cvtps_ph
+				ParticleIndices->X.Set(IndexX);
+				ParticleIndices->Y.Set(IndexY);					
 
 				// move to next particle
 				ParticleIndices += Stride;
diff --git a/Engine/Source/Runtime/ImageCore/Private/ImageCore.cpp b/Engine/Source/Runtime/ImageCore/Private/ImageCore.cpp
--- a/Engine/Source/Runtime/ImageCore/Private/ImageCore.cpp
+++ b/Engine/Source/Runtime/ImageCore/Private/ImageCore.cpp
@@ -277,7 +277,7 @@ static void CopyImage(const FImage& SrcImage, FImage& DestImage)
 				TArrayView64<const FFloat16Color> SrcColors = SrcImage.AsRGBA16F();
 				for (int64 TexelIndex = 0; TexelIndex < NumTexels; ++TexelIndex)
 				{
-					DestColors[TexelIndex] = FLinearColor(SrcColors[TexelIndex]);
+					DestColors[TexelIndex] = SrcColors[TexelIndex].GetFloats();
 				}
 			}
 			break;
