fix accidental use of SSE4 in last

diff --git a/Engine/Source/Runtime/Core/Public/Math/UnrealPlatformMathSSE.h b/Engine/Source/Runtime/Core/Public/Math/UnrealPlatformMathSSE.h
--- a/Engine/Source/Runtime/Core/Public/Math/UnrealPlatformMathSSE.h
+++ b/Engine/Source/Runtime/Core/Public/Math/UnrealPlatformMathSSE.h
@@ -155,7 +155,8 @@ namespace SSE
 		
 		// there's now a half in each 32-bit lane
 		// pack down to 64 bits :
-        __m128i four_halfs_u64 = _mm_packus_epi32(rgba_half32, _mm_setzero_si128());
+		// packs works because rgba_half32 is sign-extended
+        __m128i four_halfs_u64 = _mm_packs_epi32(rgba_half32, _mm_setzero_si128());
 
 		return four_halfs_u64;
 	}
