Require that a FMemMark has been pushed for any thread using FMemStack::Get().Alloc();
This check prevents memory leaks caused by threads allocating using their tls version of FMemStack, and never releasing the memory.

diff --git a/Engine/Source/Runtime/Core/Public/Misc/MemStack.h b/Engine/Source/Runtime/Core/Public/Misc/MemStack.h
--- a/Engine/Source/Runtime/Core/Public/Misc/MemStack.h
+++ b/Engine/Source/Runtime/Core/Public/Misc/MemStack.h
@@ -76,7 +76,9 @@ public:
 		, TopChunk(nullptr)
 		, TopMark(nullptr)
 		, NumMarks(0)
+		, bShouldEnforceAllocMarks(false)
 	{
+		
 	}
 
 	~FMemStackBase()
@@ -96,6 +98,7 @@ public:
 		checkSlow(AllocSize>=0);
 		checkSlow((Alignment&(Alignment-1))==0);
 		checkSlow(Top<=End);
+		check(!bShouldEnforceAllocMarks || NumMarks > 0);
 
 		// Try to get memory from the current chunk.
 		uint8* Result = Align( Top, Alignment );
@@ -181,11 +184,19 @@ private:
 	/** The number of marks on this stack. */
 	int32 NumMarks;
 
+
+protected:
+	bool bShouldEnforceAllocMarks;	
 };
 
 
 class CORE_API FMemStack : public TThreadSingleton<FMemStack>, public FMemStackBase
 {
+public:
+	FMemStack()
+	{
+		bShouldEnforceAllocMarks = true;
+	}
 };
 
 
