Fixed a bug causing StringBuilder::ReplaceAt to allocate new buffers more frequently than it needed to.

- EnsureCapacity expects the potential increase in length  as it's input but in the method ::ReplaceAt we were passing in the current string length as well as the expected increase in size. This would cause EnsureCapacity to count the length of the current string twice and so try to make sure that the buffer was large enough for a much larger string than needed.
- ReplaceAt now only passes in the expected delta.
- Renamed EnsureCapacity to EnsureAdditionalCapacity to try and avoid this in the future.

diff --git a/Engine/Source/Runtime/Core/Public/Misc/StringBuilder.h b/Engine/Source/Runtime/Core/Public/Misc/StringBuilder.h
--- a/Engine/Source/Runtime/Core/Public/Misc/StringBuilder.h
+++ b/Engine/Source/Runtime/Core/Public/Misc/StringBuilder.h
@@ -125,7 +125,7 @@ public:
 	 */
 	inline int32 AddUninitialized(int32 InCount)
 	{
-		EnsureCapacity(InCount);
+		EnsureAdditionalCapacity(InCount);
 		const int32 OldCount = Len();
 		CurPos += InCount;
 		return OldCount;
@@ -142,7 +142,7 @@ public:
 
 	inline BuilderType& Append(CharType Char)
 	{
-		EnsureCapacity(1);
+		EnsureAdditionalCapacity(1);
 
 		*CurPos++ = Char;
 
@@ -166,7 +166,7 @@ public:
 
 	inline BuilderType& AppendAnsi(const ANSICHAR* String, const int32 Length)
 	{
-		EnsureCapacity(Length);
+		EnsureAdditionalCapacity(Length);
 
 		CharType* RESTRICT Dest = CurPos;
 		CurPos += Length;
@@ -196,7 +196,7 @@ public:
 
 	inline BuilderType& Append(const CharType* String, int32 Length)
 	{
-		EnsureCapacity(Length);
+		EnsureAdditionalCapacity(Length);
 		CharType* RESTRICT Dest = CurPos;
 		CurPos += Length;
 
@@ -224,7 +224,7 @@ public:
 		}
 		else if (DeltaLen > 0)
 		{
-			EnsureCapacity(Len() + DeltaLen);
+			EnsureAdditionalCapacity(DeltaLen);
 			CurPos += DeltaLen;
 
 			for (CharType* It = CurPos - 1, *StopIt = Base + Pos + Str.Len() - 1; It != StopIt; --It)
@@ -350,7 +350,7 @@ protected:
 		}
 	}
 
-	inline void EnsureCapacity(int32 RequiredAdditionalCapacity)
+	inline void EnsureAdditionalCapacity(int32 RequiredAdditionalCapacity)
 	{
 		// precondition: we know the current buffer has enough capacity
 		// for the existing string including NUL terminator
