TTuple::ApplyBefore and TTuple::ApplyAfter now correctly forward the elements to the callable, according to their value category.

diff --git a/Engine/Source/Runtime/Core/Public/Templates/Tuple.h b/Engine/Source/Runtime/Core/Public/Templates/Tuple.h
--- a/Engine/Source/Runtime/Core/Public/Templates/Tuple.h
+++ b/Engine/Source/Runtime/Core/Public/Templates/Tuple.h
@@ -300,17 +300,23 @@ namespace UE4Tuple_Private
 		template <typename T, typename TEnableIf<TTypeCountInParameterPack<T, Types...>::Value == 1>::Type* = nullptr> FORCEINLINE decltype(auto) Get()       volatile&& { return TTupleElementGetterByType<T, sizeof...(Types)>::Get(static_cast<      volatile TTupleBase&&>(*this)); }
 		template <typename T, typename TEnableIf<TTypeCountInParameterPack<T, Types...>::Value == 1>::Type* = nullptr> FORCEINLINE decltype(auto) Get() const volatile&& { return TTupleElementGetterByType<T, sizeof...(Types)>::Get(static_cast<const volatile TTupleBase&&>(*this)); }
 
-		template <typename FuncType, typename... ArgTypes>
-		decltype(auto) ApplyAfter(FuncType&& Func, ArgTypes&&... Args) const
-		{
-			return ::Invoke(Func, Forward<ArgTypes>(Args)..., this->template Get<Indices>()...);
-		}
-
-		template <typename FuncType, typename... ArgTypes>
-		decltype(auto) ApplyBefore(FuncType&& Func, ArgTypes&&... Args) const
-		{
-			return ::Invoke(Func, this->template Get<Indices>()..., Forward<ArgTypes>(Args)...);
-		}
+		template <typename FuncType, typename... ArgTypes> decltype(auto) ApplyAfter(FuncType&& Func, ArgTypes&&... Args)               &  { return ::Invoke(Func, Forward<ArgTypes>(Args)..., static_cast<               TTupleBase& >(*this).template Get<Indices>()...); }
+		template <typename FuncType, typename... ArgTypes> decltype(auto) ApplyAfter(FuncType&& Func, ArgTypes&&... Args) const         &  { return ::Invoke(Func, Forward<ArgTypes>(Args)..., static_cast<const          TTupleBase& >(*this).template Get<Indices>()...); }
+		template <typename FuncType, typename... ArgTypes> decltype(auto) ApplyAfter(FuncType&& Func, ArgTypes&&... Args)       volatile&  { return ::Invoke(Func, Forward<ArgTypes>(Args)..., static_cast<      volatile TTupleBase& >(*this).template Get<Indices>()...); }
+		template <typename FuncType, typename... ArgTypes> decltype(auto) ApplyAfter(FuncType&& Func, ArgTypes&&... Args) const volatile&  { return ::Invoke(Func, Forward<ArgTypes>(Args)..., static_cast<const volatile TTupleBase& >(*this).template Get<Indices>()...); }
+		template <typename FuncType, typename... ArgTypes> decltype(auto) ApplyAfter(FuncType&& Func, ArgTypes&&... Args)               && { return ::Invoke(Func, Forward<ArgTypes>(Args)..., static_cast<               TTupleBase&&>(*this).template Get<Indices>()...); }
+		template <typename FuncType, typename... ArgTypes> decltype(auto) ApplyAfter(FuncType&& Func, ArgTypes&&... Args) const         && { return ::Invoke(Func, Forward<ArgTypes>(Args)..., static_cast<const          TTupleBase&&>(*this).template Get<Indices>()...); }
+		template <typename FuncType, typename... ArgTypes> decltype(auto) ApplyAfter(FuncType&& Func, ArgTypes&&... Args)       volatile&& { return ::Invoke(Func, Forward<ArgTypes>(Args)..., static_cast<      volatile TTupleBase&&>(*this).template Get<Indices>()...); }
+		template <typename FuncType, typename... ArgTypes> decltype(auto) ApplyAfter(FuncType&& Func, ArgTypes&&... Args) const volatile&& { return ::Invoke(Func, Forward<ArgTypes>(Args)..., static_cast<const volatile TTupleBase&&>(*this).template Get<Indices>()...); }
+
+		template <typename FuncType, typename... ArgTypes> decltype(auto) ApplyBefore(FuncType&& Func, ArgTypes&&... Args)               &  { return ::Invoke(Func, static_cast<               TTupleBase& >(*this).template Get<Indices>()..., Forward<ArgTypes>(Args)...); }
+		template <typename FuncType, typename... ArgTypes> decltype(auto) ApplyBefore(FuncType&& Func, ArgTypes&&... Args) const         &  { return ::Invoke(Func, static_cast<const          TTupleBase& >(*this).template Get<Indices>()..., Forward<ArgTypes>(Args)...); }
+		template <typename FuncType, typename... ArgTypes> decltype(auto) ApplyBefore(FuncType&& Func, ArgTypes&&... Args)       volatile&  { return ::Invoke(Func, static_cast<      volatile TTupleBase& >(*this).template Get<Indices>()..., Forward<ArgTypes>(Args)...); }
+		template <typename FuncType, typename... ArgTypes> decltype(auto) ApplyBefore(FuncType&& Func, ArgTypes&&... Args) const volatile&  { return ::Invoke(Func, static_cast<const volatile TTupleBase& >(*this).template Get<Indices>()..., Forward<ArgTypes>(Args)...); }
+		template <typename FuncType, typename... ArgTypes> decltype(auto) ApplyBefore(FuncType&& Func, ArgTypes&&... Args)               && { return ::Invoke(Func, static_cast<               TTupleBase&&>(*this).template Get<Indices>()..., Forward<ArgTypes>(Args)...); }
+		template <typename FuncType, typename... ArgTypes> decltype(auto) ApplyBefore(FuncType&& Func, ArgTypes&&... Args) const         && { return ::Invoke(Func, static_cast<const          TTupleBase&&>(*this).template Get<Indices>()..., Forward<ArgTypes>(Args)...); }
+		template <typename FuncType, typename... ArgTypes> decltype(auto) ApplyBefore(FuncType&& Func, ArgTypes&&... Args)       volatile&& { return ::Invoke(Func, static_cast<      volatile TTupleBase&&>(*this).template Get<Indices>()..., Forward<ArgTypes>(Args)...); }
+		template <typename FuncType, typename... ArgTypes> decltype(auto) ApplyBefore(FuncType&& Func, ArgTypes&&... Args) const volatile&& { return ::Invoke(Func, static_cast<const volatile TTupleBase&&>(*this).template Get<Indices>()..., Forward<ArgTypes>(Args)...); }
 
 		FORCEINLINE friend FArchive& operator<<(FArchive& Ar, TTupleBase& Tuple)
 		{
@@ -417,17 +423,23 @@ namespace UE4Tuple_Private
 		template <typename T, typename TEnableIf<TTypeCountInParameterPack<T, KeyType, ValueType>::Value == 1>::Type* = nullptr> FORCEINLINE decltype(auto) Get()       volatile&& { return TTupleElementGetterByType<T, 2>::Get(static_cast<      volatile TTupleBase&&>(*this)); }
 		template <typename T, typename TEnableIf<TTypeCountInParameterPack<T, KeyType, ValueType>::Value == 1>::Type* = nullptr> FORCEINLINE decltype(auto) Get() const volatile&& { return TTupleElementGetterByType<T, 2>::Get(static_cast<const volatile TTupleBase&&>(*this)); }
 
-		template <typename FuncType, typename... ArgTypes>
-		decltype(auto) ApplyAfter(FuncType&& Func, ArgTypes&&... Args) const
-		{
-			return ::Invoke(Func, Forward<ArgTypes>(Args)..., this->Key, this->Value);
-		}
-
-		template <typename FuncType, typename... ArgTypes>
-		decltype(auto) ApplyBefore(FuncType&& Func, ArgTypes&&... Args) const
-		{
-			return ::Invoke(Func, this->Key, this->Value, Forward<ArgTypes>(Args)...);
-		}
+		template <typename FuncType, typename... ArgTypes> decltype(auto) ApplyAfter(FuncType&& Func, ArgTypes&&... Args)               &  { return ::Invoke(Func, Forward<ArgTypes>(Args)..., static_cast<               TTupleBase& >(*this).template Get<0>(), static_cast<               TTupleBase& >(*this).template Get<1>()); }
+		template <typename FuncType, typename... ArgTypes> decltype(auto) ApplyAfter(FuncType&& Func, ArgTypes&&... Args) const         &  { return ::Invoke(Func, Forward<ArgTypes>(Args)..., static_cast<const          TTupleBase& >(*this).template Get<0>(), static_cast<const          TTupleBase& >(*this).template Get<1>()); }
+		template <typename FuncType, typename... ArgTypes> decltype(auto) ApplyAfter(FuncType&& Func, ArgTypes&&... Args)       volatile&  { return ::Invoke(Func, Forward<ArgTypes>(Args)..., static_cast<      volatile TTupleBase& >(*this).template Get<0>(), static_cast<      volatile TTupleBase& >(*this).template Get<1>()); }
+		template <typename FuncType, typename... ArgTypes> decltype(auto) ApplyAfter(FuncType&& Func, ArgTypes&&... Args) const volatile&  { return ::Invoke(Func, Forward<ArgTypes>(Args)..., static_cast<const volatile TTupleBase& >(*this).template Get<0>(), static_cast<const volatile TTupleBase& >(*this).template Get<1>()); }
+		template <typename FuncType, typename... ArgTypes> decltype(auto) ApplyAfter(FuncType&& Func, ArgTypes&&... Args)               && { return ::Invoke(Func, Forward<ArgTypes>(Args)..., static_cast<               TTupleBase&&>(*this).template Get<0>(), static_cast<               TTupleBase&&>(*this).template Get<1>()); }
+		template <typename FuncType, typename... ArgTypes> decltype(auto) ApplyAfter(FuncType&& Func, ArgTypes&&... Args) const         && { return ::Invoke(Func, Forward<ArgTypes>(Args)..., static_cast<const          TTupleBase&&>(*this).template Get<0>(), static_cast<const          TTupleBase&&>(*this).template Get<1>()); }
+		template <typename FuncType, typename... ArgTypes> decltype(auto) ApplyAfter(FuncType&& Func, ArgTypes&&... Args)       volatile&& { return ::Invoke(Func, Forward<ArgTypes>(Args)..., static_cast<      volatile TTupleBase&&>(*this).template Get<0>(), static_cast<      volatile TTupleBase&&>(*this).template Get<1>()); }
+		template <typename FuncType, typename... ArgTypes> decltype(auto) ApplyAfter(FuncType&& Func, ArgTypes&&... Args) const volatile&& { return ::Invoke(Func, Forward<ArgTypes>(Args)..., static_cast<const volatile TTupleBase&&>(*this).template Get<0>(), static_cast<const volatile TTupleBase&&>(*this).template Get<1>()); }
+
+		template <typename FuncType, typename... ArgTypes> decltype(auto) ApplyBefore(FuncType&& Func, ArgTypes&&... Args)               &  { return ::Invoke(Func, static_cast<               TTupleBase& >(*this).template Get<0>(), static_cast<               TTupleBase& >(*this).template Get<1>(), Forward<ArgTypes>(Args)...); }
+		template <typename FuncType, typename... ArgTypes> decltype(auto) ApplyBefore(FuncType&& Func, ArgTypes&&... Args) const         &  { return ::Invoke(Func, static_cast<const          TTupleBase& >(*this).template Get<0>(), static_cast<const          TTupleBase& >(*this).template Get<1>(), Forward<ArgTypes>(Args)...); }
+		template <typename FuncType, typename... ArgTypes> decltype(auto) ApplyBefore(FuncType&& Func, ArgTypes&&... Args)       volatile&  { return ::Invoke(Func, static_cast<      volatile TTupleBase& >(*this).template Get<0>(), static_cast<      volatile TTupleBase& >(*this).template Get<1>(), Forward<ArgTypes>(Args)...); }
+		template <typename FuncType, typename... ArgTypes> decltype(auto) ApplyBefore(FuncType&& Func, ArgTypes&&... Args) const volatile&  { return ::Invoke(Func, static_cast<const volatile TTupleBase& >(*this).template Get<0>(), static_cast<const volatile TTupleBase& >(*this).template Get<1>(), Forward<ArgTypes>(Args)...); }
+		template <typename FuncType, typename... ArgTypes> decltype(auto) ApplyBefore(FuncType&& Func, ArgTypes&&... Args)               && { return ::Invoke(Func, static_cast<               TTupleBase&&>(*this).template Get<0>(), static_cast<               TTupleBase&&>(*this).template Get<1>(), Forward<ArgTypes>(Args)...); }
+		template <typename FuncType, typename... ArgTypes> decltype(auto) ApplyBefore(FuncType&& Func, ArgTypes&&... Args) const         && { return ::Invoke(Func, static_cast<const          TTupleBase&&>(*this).template Get<0>(), static_cast<const          TTupleBase&&>(*this).template Get<1>(), Forward<ArgTypes>(Args)...); }
+		template <typename FuncType, typename... ArgTypes> decltype(auto) ApplyBefore(FuncType&& Func, ArgTypes&&... Args)       volatile&& { return ::Invoke(Func, static_cast<      volatile TTupleBase&&>(*this).template Get<0>(), static_cast<      volatile TTupleBase&&>(*this).template Get<1>(), Forward<ArgTypes>(Args)...); }
+		template <typename FuncType, typename... ArgTypes> decltype(auto) ApplyBefore(FuncType&& Func, ArgTypes&&... Args) const volatile&& { return ::Invoke(Func, static_cast<const volatile TTupleBase&&>(*this).template Get<0>(), static_cast<const volatile TTupleBase&&>(*this).template Get<1>(), Forward<ArgTypes>(Args)...); }
 
 		FORCEINLINE friend FArchive& operator<<(FArchive& Ar, TTupleBase& Tuple)
 		{
