Added constructors for FBlake3Hash and FIoHash from FMemoryView

#rnx

diff --git a/Engine/Source/Runtime/Core/Public/Hash/Blake3.h b/Engine/Source/Runtime/Core/Public/Hash/Blake3.h
--- a/Engine/Source/Runtime/Core/Public/Hash/Blake3.h
+++ b/Engine/Source/Runtime/Core/Public/Hash/Blake3.h
@@ -27,6 +27,9 @@ public:
 	/** Construct a hash from an array of 32 bytes. */
 	inline explicit FBlake3Hash(const ByteArray& Hash);
 
+	/** Construct a hash from a view of 32 bytes. */
+	inline explicit FBlake3Hash(FMemoryView Hash);
+
 	/** Construct a hash from a 64-character hex string. */
 	inline explicit FBlake3Hash(FAnsiStringView HexHash);
 	inline explicit FBlake3Hash(FWideStringView HexHash);
@@ -87,6 +90,13 @@ inline FBlake3Hash::FBlake3Hash(const ByteArray& InHash)
 	FMemory::Memcpy(Hash, InHash, sizeof(ByteArray));
 }
 
+inline FBlake3Hash::FBlake3Hash(const FMemoryView InHash)
+{
+	checkf(InHash.GetSize() == sizeof(ByteArray),
+		TEXT("FBlake3Hash cannot be constructed from a view of %" UINT64_FMT " bytes."), InHash.GetSize());
+	FMemory::Memcpy(Hash, InHash.GetData(), sizeof(ByteArray));
+}
+
 inline FBlake3Hash::FBlake3Hash(const FAnsiStringView HexHash)
 {
 	check(HexHash.Len() == sizeof(ByteArray) * 2);
diff --git a/Engine/Source/Runtime/Core/Public/IO/IoHash.h b/Engine/Source/Runtime/Core/Public/IO/IoHash.h
--- a/Engine/Source/Runtime/Core/Public/IO/IoHash.h
+++ b/Engine/Source/Runtime/Core/Public/IO/IoHash.h
@@ -33,6 +33,9 @@ public:
 	/** Construct a hash from an array of 20 bytes. */
 	inline explicit FIoHash(const ByteArray& Hash);
 
+	/** Construct a hash from a view of 20 bytes. */
+	inline explicit FIoHash(FMemoryView Hash);
+
 	/** Construct a hash from a BLAKE3-256 hash. */
 	inline FIoHash(const FBlake3Hash& Hash);
 
@@ -64,6 +67,13 @@ inline FIoHash::FIoHash(const ByteArray& InHash)
 	FMemory::Memcpy(Hash, InHash, sizeof(ByteArray));
 }
 
+inline FIoHash::FIoHash(const FMemoryView InHash)
+{
+	checkf(InHash.GetSize() == sizeof(ByteArray),
+		TEXT("FIoHash cannot be constructed from a view of %" UINT64_FMT " bytes."), InHash.GetSize());
+	FMemory::Memcpy(Hash, InHash.GetData(), sizeof(ByteArray));
+}
+
 inline FIoHash::FIoHash(const FBlake3Hash& InHash)
 {
 	static_assert(sizeof(ByteArray) <= sizeof(decltype(InHash.GetBytes())), "Reading too many bytes from source.");
