Some comment fixes.

diff --git a/Engine/Source/Runtime/Core/Public/Delegates/DelegateSignatureImpl.inl b/Engine/Source/Runtime/Core/Public/Delegates/DelegateSignatureImpl.inl
--- a/Engine/Source/Runtime/Core/Public/Delegates/DelegateSignatureImpl.inl
+++ b/Engine/Source/Runtime/Core/Public/Delegates/DelegateSignatureImpl.inl
@@ -73,7 +73,7 @@ public:
 	template <typename UserClass, typename... VarTypes> using TMethodPtr      = typename TMemFunPtrType<false, UserClass, RetValType(ParamTypes..., VarTypes...)>::Type;
 	template <typename UserClass, typename... VarTypes> using TConstMethodPtr = typename TMemFunPtrType<true,  UserClass, RetValType(ParamTypes..., VarTypes...)>::Type;
 
-	/** Declare the user's "fast" shared pointer-based delegate instance types. */
+	/** Declare the user's shared pointer-based delegate instance types. */
 	template <typename UserClass                                                                            > struct TSPMethodDelegate                 : TBaseSPMethodDelegateInstance<false, UserClass, ESPMode::ThreadSafe, FuncType, UserPolicy                                        > { typedef TBaseSPMethodDelegateInstance<false, UserClass, ESPMode::ThreadSafe, FuncType, UserPolicy                                        > Super; TSPMethodDelegate                (const TSharedRef< UserClass, ESPMode::ThreadSafe >& InUserObject, typename Super::FMethodPtr InMethodPtr                                                            ) : Super(InUserObject, InMethodPtr                        ) {} };
 	template <typename UserClass                                                                            > struct TSPMethodDelegate_Const           : TBaseSPMethodDelegateInstance<true , UserClass, ESPMode::ThreadSafe, FuncType, UserPolicy                                        > { typedef TBaseSPMethodDelegateInstance<true , UserClass, ESPMode::ThreadSafe, FuncType, UserPolicy                                        > Super; TSPMethodDelegate_Const          (const TSharedRef< UserClass, ESPMode::ThreadSafe >& InUserObject, typename Super::FMethodPtr InMethodPtr                                                            ) : Super(InUserObject, InMethodPtr                        ) {} };
 	template <typename UserClass, typename Var1Type                                                         > struct TSPMethodDelegate_OneVar          : TBaseSPMethodDelegateInstance<false, UserClass, ESPMode::ThreadSafe, FuncType, UserPolicy, Var1Type                              > { typedef TBaseSPMethodDelegateInstance<false, UserClass, ESPMode::ThreadSafe, FuncType, UserPolicy, Var1Type                              > Super; TSPMethodDelegate_OneVar         (const TSharedRef< UserClass, ESPMode::ThreadSafe >& InUserObject, typename Super::FMethodPtr InMethodPtr, Var1Type Var1                                             ) : Super(InUserObject, InMethodPtr, Var1                  ) {} };
@@ -85,7 +85,7 @@ public:
 	template <typename UserClass, typename Var1Type, typename Var2Type, typename Var3Type, typename Var4Type> struct TSPMethodDelegate_FourVars        : TBaseSPMethodDelegateInstance<false, UserClass, ESPMode::ThreadSafe, FuncType, UserPolicy, Var1Type, Var2Type, Var3Type, Var4Type> { typedef TBaseSPMethodDelegateInstance<false, UserClass, ESPMode::ThreadSafe, FuncType, UserPolicy, Var1Type, Var2Type, Var3Type, Var4Type> Super; TSPMethodDelegate_FourVars       (const TSharedRef< UserClass, ESPMode::ThreadSafe >& InUserObject, typename Super::FMethodPtr InMethodPtr, Var1Type Var1, Var2Type Var2, Var3Type Var3, Var4Type Var4) : Super(InUserObject, InMethodPtr, Var1, Var2, Var3, Var4) {} };
 	template <typename UserClass, typename Var1Type, typename Var2Type, typename Var3Type, typename Var4Type> struct TSPMethodDelegate_FourVars_Const  : TBaseSPMethodDelegateInstance<true , UserClass, ESPMode::ThreadSafe, FuncType, UserPolicy, Var1Type, Var2Type, Var3Type, Var4Type> { typedef TBaseSPMethodDelegateInstance<true , UserClass, ESPMode::ThreadSafe, FuncType, UserPolicy, Var1Type, Var2Type, Var3Type, Var4Type> Super; TSPMethodDelegate_FourVars_Const (const TSharedRef< UserClass, ESPMode::ThreadSafe >& InUserObject, typename Super::FMethodPtr InMethodPtr, Var1Type Var1, Var2Type Var2, Var3Type Var3, Var4Type Var4) : Super(InUserObject, InMethodPtr, Var1, Var2, Var3, Var4) {} };
 
-	/** Declare the user's "thread-safe" shared pointer-based delegate instance types. */
+	/** Declare the user's shared pointer-based delegate instance types. */
 	template <typename UserClass                                                                            > struct TThreadSafeSPMethodDelegate                 : TBaseSPMethodDelegateInstance<false, UserClass, ESPMode::ThreadSafe, FuncType, UserPolicy                                        > { typedef TBaseSPMethodDelegateInstance<false, UserClass, ESPMode::ThreadSafe, FuncType, UserPolicy                                        > Super; TThreadSafeSPMethodDelegate                (const TSharedRef<UserClass, ESPMode::ThreadSafe>& InUserObject, typename Super::FMethodPtr InMethodPtr                                                            ) : Super(InUserObject, InMethodPtr                        ) {} };
 	template <typename UserClass                                                                            > struct TThreadSafeSPMethodDelegate_Const           : TBaseSPMethodDelegateInstance<true , UserClass, ESPMode::ThreadSafe, FuncType, UserPolicy                                        > { typedef TBaseSPMethodDelegateInstance<true , UserClass, ESPMode::ThreadSafe, FuncType, UserPolicy                                        > Super; TThreadSafeSPMethodDelegate_Const          (const TSharedRef<UserClass, ESPMode::ThreadSafe>& InUserObject, typename Super::FMethodPtr InMethodPtr                                                            ) : Super(InUserObject, InMethodPtr                        ) {} };
 	template <typename UserClass, typename Var1Type                                                         > struct TThreadSafeSPMethodDelegate_OneVar          : TBaseSPMethodDelegateInstance<false, UserClass, ESPMode::ThreadSafe, FuncType, UserPolicy, Var1Type                              > { typedef TBaseSPMethodDelegateInstance<false, UserClass, ESPMode::ThreadSafe, FuncType, UserPolicy, Var1Type                              > Super; TThreadSafeSPMethodDelegate_OneVar         (const TSharedRef<UserClass, ESPMode::ThreadSafe>& InUserObject, typename Super::FMethodPtr InMethodPtr, Var1Type Var1                                             ) : Super(InUserObject, InMethodPtr, Var1                  ) {} };
