fixed compilation error: https://horde.devtools.epicgames.com/job/608b3fc0ff44100001a3cca5?step=9073&issue=67780

diff --git a/Engine/Source/Runtime/Core/Private/FramePro/FrameProUE4.cpp b/Engine/Source/Runtime/Core/Private/FramePro/FrameProUE4.cpp
--- a/Engine/Source/Runtime/Core/Private/FramePro/FrameProUE4.cpp
+++ b/Engine/Source/Runtime/Core/Private/FramePro/FrameProUE4.cpp
@@ -510,6 +510,7 @@
 			}
 			else
 			{
+				PRAGMA_DISABLE_DEPRECATION_WARNINGS
 #if PLATFORM_WINDOWS
 				p_event = new FEventWin();
 				p_event->Create(!auto_reset);
@@ -520,6 +521,7 @@
 				checkf(false, TEXT("unsupported platform for -nothreading"));
 				return;
 #endif
+				PRAGMA_ENABLE_DEPRECATION_WARNINGS
 			}
 //@EPIC END
 			GetOSEvent(p_os_event_mem) = p_event;
