CompactBinary: Changed Get[Range]Hash(FBlake3&) to Append[Range]Hash(FBlake3&)

#rnx

diff --git a/Engine/Source/Runtime/Core/Private/Serialization/CompactBinary.cpp b/Engine/Source/Runtime/Core/Private/Serialization/CompactBinary.cpp
--- a/Engine/Source/Runtime/Core/Private/Serialization/CompactBinary.cpp
+++ b/Engine/Source/Runtime/Core/Private/Serialization/CompactBinary.cpp
@@ -548,18 +548,18 @@ uint64 FCbFieldView::GetPayloadSize() const
 	}
 }
 
-FIoHash FCbFieldView::GetHash() const
+FBlake3Hash FCbFieldView::GetHash() const
 {
 	FBlake3 Hash;
-	GetHash(Hash);
-	return FIoHash(Hash.Finalize());
+	AppendHash(Hash);
+	return Hash.Finalize();
 }
 
-void FCbFieldView::GetHash(FBlake3& Hash) const
+void FCbFieldView::AppendHash(FBlake3& Builder) const
 {
 	const ECbFieldType SerializedType = FCbFieldType::GetSerializedType(Type);
-	Hash.Update(&SerializedType, sizeof(SerializedType));
-	Hash.Update(GetViewNoType());
+	Builder.Update(&SerializedType, sizeof(SerializedType));
+	Builder.Update(GetViewNoType());
 }
 
 bool FCbFieldView::Equals(const FCbFieldView& Other) const
@@ -662,18 +662,18 @@ uint64 FCbArrayView::GetSize() const
 	return sizeof(ECbFieldType) + GetPayloadSize();
 }
 
-FIoHash FCbArrayView::GetHash() const
+FBlake3Hash FCbArrayView::GetHash() const
 {
 	FBlake3 Hash;
-	GetHash(Hash);
-	return FIoHash(Hash.Finalize());
+	AppendHash(Hash);
+	return Hash.Finalize();
 }
 
-void FCbArrayView::GetHash(FBlake3& Hash) const
+void FCbArrayView::AppendHash(FBlake3& Builder) const
 {
 	const ECbFieldType SerializedType = FCbFieldType::GetType(GetType());
-	Hash.Update(&SerializedType, sizeof(SerializedType));
-	Hash.Update(GetPayloadView());
+	Builder.Update(&SerializedType, sizeof(SerializedType));
+	Builder.Update(GetPayloadView());
 }
 
 bool FCbArrayView::Equals(const FCbArrayView& Other) const
@@ -742,18 +742,18 @@ uint64 FCbObjectView::GetSize() const
 	return sizeof(ECbFieldType) + GetPayloadSize();
 }
 
-FIoHash FCbObjectView::GetHash() const
+FBlake3Hash FCbObjectView::GetHash() const
 {
 	FBlake3 Hash;
-	GetHash(Hash);
-	return FIoHash(Hash.Finalize());
+	AppendHash(Hash);
+	return Hash.Finalize();
 }
 
-void FCbObjectView::GetHash(FBlake3& Hash) const
+void FCbObjectView::AppendHash(FBlake3& Builder) const
 {
 	const ECbFieldType SerializedType = FCbFieldType::GetType(GetType());
-	Hash.Update(&SerializedType, sizeof(SerializedType));
-	Hash.Update(GetPayloadView());
+	Builder.Update(&SerializedType, sizeof(SerializedType));
+	Builder.Update(GetPayloadView());
 }
 
 bool FCbObjectView::Equals(const FCbObjectView& Other) const
@@ -816,26 +816,26 @@ uint64 TCbFieldIterator<FieldType>::GetRangeSize() const
 }
 
 template <typename FieldType>
-FIoHash TCbFieldIterator<FieldType>::GetRangeHash() const
+FBlake3Hash TCbFieldIterator<FieldType>::GetRangeHash() const
 {
 	FBlake3 Hash;
-	GetRangeHash(Hash);
-	return FIoHash(Hash.Finalize());
+	AppendRangeHash(Hash);
+	return Hash.Finalize();
 }
 
 template <typename FieldType>
-void TCbFieldIterator<FieldType>::GetRangeHash(FBlake3& Hash) const
+void TCbFieldIterator<FieldType>::AppendRangeHash(FBlake3& Builder) const
 {
 	FMemoryView View;
 	if (TryGetSerializedRangeView(View))
 	{
-		Hash.Update(View);
+		Builder.Update(View);
 	}
 	else
 	{
 		for (TCbFieldIterator It(*this); It; ++It)
 		{
-			It.GetHash(Hash);
+			It.AppendHash(Builder);
 		}
 	}
 }
diff --git a/Engine/Source/Runtime/Core/Private/Tests/Serialization/CompactBinaryPackageTest.cpp b/Engine/Source/Runtime/Core/Private/Tests/Serialization/CompactBinaryPackageTest.cpp
--- a/Engine/Source/Runtime/Core/Private/Tests/Serialization/CompactBinaryPackageTest.cpp
+++ b/Engine/Source/Runtime/Core/Private/Tests/Serialization/CompactBinaryPackageTest.cpp
@@ -87,7 +87,7 @@ bool FCbAttachmentTest::RunTest(const FString& Parameters)
 		TestEqual(TEXT("FCbAttachment(CompactBinary).AsCompactBinary()"), Attachment.AsCompactBinary(), Fields);
 		TestTrue(TEXT("FCbAttachment(CompactBinary).IsBinary()"), Attachment.IsBinary());
 		TestTrue(TEXT("FCbAttachment(CompactBinary).IsCompactBinary()"), Attachment.IsCompactBinary());
-		TestEqual(TEXT("FCbAttachment(CompactBinary).GetHash()"), Attachment.GetHash(), Fields.GetRangeHash());
+		TestEqual(TEXT("FCbAttachment(CompactBinary).GetHash()"), Attachment.GetHash(), FIoHash(Fields.GetRangeHash()));
 		TestSaveLoadValidate(TEXT("CompactBinary"), Attachment);
 	}
 
@@ -123,7 +123,7 @@ bool FCbAttachmentTest::RunTest(const FString& Parameters)
 			Attachment.AsCompactBinary().GetRangeView().EqualBytes(Fields.GetRangeView()));
 		TestTrue(TEXT("FCbAttachment(CompactBinaryView).IsBinary()"), Attachment.IsBinary());
 		TestTrue(TEXT("FCbAttachment(CompactBinaryView).IsCompactBinary()"), Attachment.IsCompactBinary());
-		TestEqual(TEXT("FCbAttachment(CompactBinaryView).GetHash()"), Attachment.GetHash(), Fields.GetRangeHash());
+		TestEqual(TEXT("FCbAttachment(CompactBinaryView).GetHash()"), Attachment.GetHash(), FIoHash(Fields.GetRangeHash()));
 	}
 
 	// Binary Load from View
@@ -174,7 +174,7 @@ bool FCbAttachmentTest::RunTest(const FString& Parameters)
 			FieldsView.GetRangeBuffer().GetView().Contains(Attachment.AsCompactBinary().GetRangeBuffer().GetView()));
 		TestTrue(TEXT("FCbAttachment(LoadCompactBinaryView).IsBinary()"), Attachment.IsBinary());
 		TestTrue(TEXT("FCbAttachment(LoadCompactBinaryView).IsCompactBinary()"), Attachment.IsCompactBinary());
-		TestEqual(TEXT("FCbAttachment(LoadCompactBinaryView).GetHash()"), Attachment.GetHash(), Value.GetRangeHash());
+		TestEqual(TEXT("FCbAttachment(LoadCompactBinaryView).GetHash()"), Attachment.GetHash(), FIoHash(Value.GetRangeHash()));
 	}
 
 	// Compact Binary Uniform Sub-View
@@ -192,7 +192,7 @@ bool FCbAttachmentTest::RunTest(const FString& Parameters)
 		TestTrue(TEXT("FCbAttachment(CompactBinaryUniformSubView).AsBinaryView()->EqualBytes()"),
 			Binary.GetView().EqualBytes(SavedFields.GetRangeView()));
 		TestEqual(TEXT("FCbAttachment(CompactBinaryUniformSubView).GetHash()"),
-			Attachment.GetHash(), SavedFields.GetRangeHash());
+			Attachment.GetHash(), FIoHash(SavedFields.GetRangeHash()));
 		TestSaveLoadValidate(TEXT("CompactBinaryUniformSubView"), Attachment);
 	}
 
@@ -279,7 +279,7 @@ bool FCbPackageTest::RunTest(const FString& Parameters)
 		TestEqual(TEXT("FCbPackage(Object).GetAttachments()"), Package.GetAttachments().Num(), 0);
 		TestEqual(TEXT("FCbPackage(Object).GetObject()->IsClone"), Package.GetObject().GetOuterBuffer(), Object.GetOuterBuffer());
 		TestEqual(TEXT("FCbPackage(Object).GetObject()"), Package.GetObject()["Field"].AsInt32(), 42);
-		TestEqual(TEXT("FCbPackage(Object).GetObjectHash()"), Package.GetObjectHash(), Package.GetObject().GetHash());
+		TestEqual(TEXT("FCbPackage(Object).GetObjectHash()"), Package.GetObjectHash(), FIoHash(Package.GetObject().GetHash()));
 		TestSaveLoadValidate(TEXT("Object"), Package);
 	}
 
@@ -297,7 +297,7 @@ bool FCbPackageTest::RunTest(const FString& Parameters)
 		TestEqual(TEXT("FCbPackage(Object).GetAttachments()"), Package.GetAttachments().Num(), 0);
 		TestNotEqual(TEXT("FCbPackage(Object).GetObject()->IsClone"), Package.GetObject().GetOuterBuffer(), Object.GetOuterBuffer());
 		TestEqual(TEXT("FCbPackage(Object).GetObject()"), Package.GetObject()["Field"].AsInt32(), 42);
-		TestEqual(TEXT("FCbPackage(Object).GetObjectHash()"), Package.GetObjectHash(), Package.GetObject().GetHash());
+		TestEqual(TEXT("FCbPackage(Object).GetObjectHash()"), Package.GetObjectHash(), FIoHash(Package.GetObject().GetHash()));
 		TestSaveLoadValidate(TEXT("Object"), Package);
 	}
 
diff --git a/Engine/Source/Runtime/Core/Private/Tests/Serialization/CompactBinaryTest.cpp b/Engine/Source/Runtime/Core/Private/Tests/Serialization/CompactBinaryTest.cpp
--- a/Engine/Source/Runtime/Core/Private/Tests/Serialization/CompactBinaryTest.cpp
+++ b/Engine/Source/Runtime/Core/Private/Tests/Serialization/CompactBinaryTest.cpp
@@ -182,7 +182,7 @@ bool FCbFieldNoneTest::RunTest(const FString& Parameters)
 		TestFalse(TEXT("!FCbFieldView()::HasValue()"), DefaultField.HasValue());
 		TestFalse(TEXT("!FCbFieldView()::HasError()"), DefaultField.HasError());
 		TestEqual(TEXT("FCbFieldView()::GetError() == None"), DefaultField.GetError(), ECbFieldError::None);
-		TestEqual(TEXT("FCbFieldView()::GetHash()"), DefaultField.GetHash(), FIoHash::HashBuffer(MakeMemoryView<uint8>({uint8(ECbFieldType::None)})));
+		TestEqual(TEXT("FCbFieldView()::GetHash()"), DefaultField.GetHash(), FBlake3::HashBuffer(MakeMemoryView<uint8>({uint8(ECbFieldType::None)})));
 		TestEqual(TEXT("FCbFieldView()::GetView()"), DefaultField.GetView(), FMemoryView());
 		FMemoryView SerializedView;
 		TestFalse(TEXT("FCbFieldView()::TryGetSerializedView()"), DefaultField.TryGetSerializedView(SerializedView));
@@ -212,7 +212,7 @@ bool FCbFieldNoneTest::RunTest(const FString& Parameters)
 		TestEqual(TEXT("FCbFieldView(None|Type|Name)::GetName()"), NoneField.GetName(), "Name"_ASV);
 		TestTrue(TEXT("FCbFieldView(None|Type|Name)::HasName()"), NoneField.HasName());
 		TestFalse(TEXT("!FCbFieldView(None|Type|Name)::HasValue()"), NoneField.HasValue());
-		TestEqual(TEXT("FCbFieldView(None|Type|Name)::GetHash()"), NoneField.GetHash(), FIoHash::HashBuffer(MakeMemoryView(NoneBytes)));
+		TestEqual(TEXT("FCbFieldView(None|Type|Name)::GetHash()"), NoneField.GetHash(), FBlake3::HashBuffer(MakeMemoryView(NoneBytes)));
 		TestEqual(TEXT("FCbFieldView(None|Type|Name)::GetView()"), NoneField.GetView(), MakeMemoryView(NoneBytes));
 		FMemoryView SerializedView;
 		TestTrue(TEXT("FCbFieldView(None|Type|Name)::TryGetSerializedView()"), NoneField.TryGetSerializedView(SerializedView) && SerializedView == MakeMemoryView(NoneBytes));
@@ -246,7 +246,7 @@ bool FCbFieldNoneTest::RunTest(const FString& Parameters)
 		TestEqual(TEXT("FCbFieldView(None|Name)::GetName()"), NoneField.GetName(), "Name"_ASV);
 		TestTrue(TEXT("FCbFieldView(None|Name)::HasName()"), NoneField.HasName());
 		TestFalse(TEXT("!FCbFieldView(None|Name)::HasValue()"), NoneField.HasValue());
-		TestEqual(TEXT("FCbFieldView(None|Name)::GetHash()"), NoneField.GetHash(), FIoHash::HashBuffer(MakeMemoryView(NoneBytes)));
+		TestEqual(TEXT("FCbFieldView(None|Name)::GetHash()"), NoneField.GetHash(), FBlake3::HashBuffer(MakeMemoryView(NoneBytes)));
 		TestEqual(TEXT("FCbFieldView(None|Name)::GetView()"), NoneField.GetView(), MakeMemoryView(NoneBytes) + 1);
 		FMemoryView SerializedView;
 		TestFalse(TEXT("FCbFieldView(None|Name)::TryGetSerializedView()"), NoneField.TryGetSerializedView(SerializedView));
@@ -265,7 +265,7 @@ bool FCbFieldNoneTest::RunTest(const FString& Parameters)
 		TestEqual(TEXT("FCbFieldView(None|EmptyName)::GetName()"), NoneField.GetName(), ""_ASV);
 		TestTrue(TEXT("FCbFieldView(None|EmptyName)::HasName()"), NoneField.HasName());
 		TestFalse(TEXT("!FCbFieldView(None|EmptyName)::HasValue()"), NoneField.HasValue());
-		TestEqual(TEXT("FCbFieldView(None|EmptyName)::GetHash()"), NoneField.GetHash(), FIoHash::HashBuffer(MakeMemoryView(NoneBytes)));
+		TestEqual(TEXT("FCbFieldView(None|EmptyName)::GetHash()"), NoneField.GetHash(), FBlake3::HashBuffer(MakeMemoryView(NoneBytes)));
 		TestEqual(TEXT("FCbFieldView(None|EmptyName)::GetView()"), NoneField.GetView(), MakeMemoryView(NoneBytes) + 1);
 		FMemoryView SerializedView;
 		TestFalse(TEXT("FCbFieldView(None|EmptyName)::TryGetSerializedView()"), NoneField.TryGetSerializedView(SerializedView));
@@ -285,7 +285,7 @@ bool FCbFieldNullTest::RunTest(const FString& Parameters)
 		TestTrue(TEXT("FCbFieldView(Null)::HasValue()"), NullField.HasValue());
 		TestFalse(TEXT("!FCbFieldView(Null)::HasError()"), NullField.HasError());
 		TestEqual(TEXT("FCbFieldView(Null)::GetError() == None"), NullField.GetError(), ECbFieldError::None);
-		TestEqual(TEXT("FCbFieldView(Null)::GetHash()"), NullField.GetHash(), FIoHash::HashBuffer(MakeMemoryView<uint8>({uint8(ECbFieldType::Null)})));
+		TestEqual(TEXT("FCbFieldView(Null)::GetHash()"), NullField.GetHash(), FBlake3::HashBuffer(MakeMemoryView<uint8>({uint8(ECbFieldType::Null)})));
 	}
 
 	// Test FCbFieldView(None) as Null
@@ -1678,8 +1678,8 @@ bool FCbFieldBufferIteratorTest::RunTest(const FString& Parameters)
 		const FCbFieldViewIterator FieldViewIt = FCbFieldViewIterator::MakeRange(View);
 		const FCbFieldIterator FieldIt = FCbFieldIterator::MakeRange(View);
 
-		TestEqual(TEXT("FCbFieldViewIterator::GetRangeHash()"), FieldViewIt.GetRangeHash(), FIoHash::HashBuffer(View));
-		TestEqual(TEXT("FCbFieldIterator::GetRangeHash()"), FieldIt.GetRangeHash(), FIoHash::HashBuffer(View));
+		TestEqual(TEXT("FCbFieldViewIterator::GetRangeHash()"), FieldViewIt.GetRangeHash(), FBlake3::HashBuffer(View));
+		TestEqual(TEXT("FCbFieldIterator::GetRangeHash()"), FieldIt.GetRangeHash(), FBlake3::HashBuffer(View));
 
 		TestEqual(TEXT("FCbFieldViewIterator::GetRangeView()"), FieldViewIt.GetRangeView(), MakeMemoryView(Payload));
 		TestEqual(TEXT("FCbFieldIterator::GetRangeView()"), FieldIt.GetRangeView(), MakeMemoryView(Payload));
diff --git a/Engine/Source/Runtime/Core/Public/Serialization/CompactBinary.h b/Engine/Source/Runtime/Core/Public/Serialization/CompactBinary.h
--- a/Engine/Source/Runtime/Core/Public/Serialization/CompactBinary.h
+++ b/Engine/Source/Runtime/Core/Public/Serialization/CompactBinary.h
@@ -347,9 +347,9 @@ public:
 	CORE_API uint64 GetRangeSize() const;
 
 	/** Calculate the hash of every field in the range. */
-	CORE_API FIoHash GetRangeHash() const;
-	/** Calculate the hash of every field in the range. */
-	CORE_API void GetRangeHash(FBlake3& Hash) const;
+	CORE_API FBlake3Hash GetRangeHash() const;
+	/** Append the hash of every field in the range. */
+	CORE_API void AppendRangeHash(FBlake3& Builder) const;
 
 	using FieldType::Equals;
 
@@ -707,9 +707,9 @@ public:
 	CORE_API uint64 GetSize() const;
 
 	/** Calculate the hash of the field, including the type and name. */
-	CORE_API FIoHash GetHash() const;
-	/** Calculate the hash of the field, including the type and name. */
-	CORE_API void GetHash(FBlake3& Hash) const;
+	CORE_API FBlake3Hash GetHash() const;
+	/** Append the hash of the field, including the type and name. */
+	CORE_API void AppendHash(FBlake3& Builder) const;
 
 	/**
 	 * Whether this field is identical to the other field.
@@ -908,9 +908,9 @@ public:
 	CORE_API uint64 GetSize() const;
 
 	/** Calculate the hash of the array if serialized by itself with no name. */
-	CORE_API FIoHash GetHash() const;
-	/** Calculate the hash of the array if serialized by itself with no name. */
-	CORE_API void GetHash(FBlake3& Hash) const;
+	CORE_API FBlake3Hash GetHash() const;
+	/** Append the hash of the array if serialized by itself with no name. */
+	CORE_API void AppendHash(FBlake3& Builder) const;
 
 	/**
 	 * Whether this array is identical to the other array.
@@ -1008,9 +1008,9 @@ public:
 	CORE_API uint64 GetSize() const;
 
 	/** Calculate the hash of the object if serialized by itself with no name. */
-	CORE_API FIoHash GetHash() const;
-	/** Calculate the hash of the object if serialized by itself with no name. */
-	CORE_API void GetHash(FBlake3& Hash) const;
+	CORE_API FBlake3Hash GetHash() const;
+	/** Append the hash of the object if serialized by itself with no name. */
+	CORE_API void AppendHash(FBlake3& Builder) const;
 
 	/**
 	 * Whether this object is identical to the other object.
