Work around Clang build failure on Win64.

- Revert a typo correction in the THIRD_PARTY_INCLUDES_END macro. Will file a separate issue to get it resolved.

#rnx

diff --git a/Engine/Source/Runtime/Core/Public/Windows/WindowsPlatformCompilerPreSetup.h b/Engine/Source/Runtime/Core/Public/Windows/WindowsPlatformCompilerPreSetup.h
--- a/Engine/Source/Runtime/Core/Public/Windows/WindowsPlatformCompilerPreSetup.h
+++ b/Engine/Source/Runtime/Core/Public/Windows/WindowsPlatformCompilerPreSetup.h
@@ -41,8 +41,8 @@
 	#ifndef THIRD_PARTY_INCLUDES_END
 		#define THIRD_PARTY_INCLUDES_END \
 	        PRAGMA_ENABLE_UNINITIALIZED_CONST_REFERENCE_WARNINGS \
-			PRAGMA_ENABLE_MISSING_BRACES_WARNINGS \
-			PRAGMA_ENABLE_OVERLOADED_VIRTUAL_WARNINGS \
+			PRAGMA_DISABLE_MISSING_BRACES_WARNINGS \
+			PRAGMA_DISABLE_OVERLOADED_VIRTUAL_WARNINGS \
 			PRAGMA_ENABLE_DEPRECATION_WARNINGS \
 			PRAGMA_ENABLE_UNDEFINED_IDENTIFIER_WARNINGS \
 			PRAGMA_ENABLE_SHADOW_VARIABLE_WARNINGS \
