Some warning fixes for when running Commandlets with -NoEnginePlugins

* Silence module load warnings unless doing LoadModuleChecked or specified
* Silence class load warnings when loading types from plugins inside APawn::APawn

diff --git a/Engine/Source/Runtime/Core/Private/Modules/ModuleManager.cpp b/Engine/Source/Runtime/Core/Private/Modules/ModuleManager.cpp
--- a/Engine/Source/Runtime/Core/Private/Modules/ModuleManager.cpp
+++ b/Engine/Source/Runtime/Core/Private/Modules/ModuleManager.cpp
@@ -363,7 +363,7 @@ void FModuleManager::RefreshModuleFilenameFromManifest(const FName InModuleName)
 }
 #endif	// !IS_MONOLITHIC
 
-IModuleInterface* FModuleManager::LoadModule( const FName InModuleName )
+IModuleInterface* FModuleManager::LoadModule(const FName InModuleName, ELoadModuleFlags InLoadModuleFlags)
 {
 	// We allow an already loaded module to be returned in other threads to simplify
 	// parallel processing scenarios but they must have been loaded from the main thread beforehand.
@@ -373,7 +373,7 @@ IModuleInterface* FModuleManager::LoadModule( const FName InModuleName )
 	}
 
 	EModuleLoadResult FailureReason;
-	IModuleInterface* Result = LoadModuleWithFailureReason(InModuleName, FailureReason );
+	IModuleInterface* Result = LoadModuleWithFailureReason(InModuleName, FailureReason, InLoadModuleFlags);
 
 	// This should return a valid pointer only if and only if the module is loaded
 	checkSlow((Result != nullptr) == IsModuleLoaded(InModuleName));
@@ -384,14 +384,14 @@ IModuleInterface* FModuleManager::LoadModule( const FName InModuleName )
 
 IModuleInterface& FModuleManager::LoadModuleChecked( const FName InModuleName )
 {
-	IModuleInterface* Module = LoadModule(InModuleName);
+	IModuleInterface* Module = LoadModule(InModuleName, ELoadModuleFlags::LogFailures);
 	checkf(Module, TEXT("%s"), *InModuleName.ToString());
 
 	return *Module;
 }
 
 
-IModuleInterface* FModuleManager::LoadModuleWithFailureReason(const FName InModuleName, EModuleLoadResult& OutFailureReason)
+IModuleInterface* FModuleManager::LoadModuleWithFailureReason(const FName InModuleName, EModuleLoadResult& OutFailureReason, ELoadModuleFlags InLoadModuleFlags)
 {
 #if 0
 	ensureMsgf(IsInGameThread(), TEXT("ModuleManager: Attempting to load '%s' outside the main thread.  Please call LoadModule on the main/game thread only.  You can use GetModule or GetModuleChecked instead, those are safe to call outside the game thread."), *InModuleName.ToString());
@@ -485,7 +485,9 @@ IModuleInterface* FModuleManager::LoadModuleWithFailureReason(const FName InModu
 		}
 		else
 		{
-			UE_LOG(LogModuleManager, Warning, TEXT("ModuleManager: Unable to load module '%s' because InitializeModule function failed (returned nullptr.)"), *InModuleName.ToString());
+			UE_CLOG((InLoadModuleFlags & ELoadModuleFlags::LogFailures) != ELoadModuleFlags::None,
+				LogModuleManager, Warning, TEXT("ModuleManager: Unable to load module '%s' because InitializeModule function failed (returned nullptr.)"), *InModuleName.ToString());
+
 			OutFailureReason = EModuleLoadResult::FailedToInitialize;
 		}
 	}
@@ -494,7 +496,9 @@ IModuleInterface* FModuleManager::LoadModuleWithFailureReason(const FName InModu
 	{
 		// Monolithic builds that do not have the initializer were *not found* during the build step, so return FileNotFound
 		// (FileNotFound is an acceptable error in some case - ie loading a content only project)
-		UE_LOG(LogModuleManager, Warning, TEXT("ModuleManager: Module '%s' not found - its StaticallyLinkedModuleInitializers function is null."), *InModuleName.ToString());
+		UE_CLOG((InLoadModuleFlags & ELoadModuleFlags::LogFailures) != ELoadModuleFlags::None,
+			LogModuleManager, Warning, TEXT("ModuleManager: Module '%s' not found - its StaticallyLinkedModuleInitializers function is null."), *InModuleName.ToString());
+
 		OutFailureReason = EModuleLoadResult::FileNotFound;
 	}
 #else
@@ -505,7 +509,7 @@ IModuleInterface* FModuleManager::LoadModuleWithFailureReason(const FName InModu
 		// in the module being loaded.
 		if (bCanProcessNewlyLoadedObjects)
 		{
-				ProcessLoadedObjectsCallback.Broadcast(NAME_None, bCanProcessNewlyLoadedObjects);
+			ProcessLoadedObjectsCallback.Broadcast(NAME_None, bCanProcessNewlyLoadedObjects);
 		}
 
 		// Try to dynamically load the DLL
@@ -519,7 +523,9 @@ IModuleInterface* FModuleManager::LoadModuleWithFailureReason(const FName InModu
 
 			if (ModulePathMap.Num() != 1)
 			{
-				UE_LOG(LogModuleManager, Warning, TEXT("ModuleManager: Unable to load module '%s'  - %d instances of that module name found."), *InModuleName.ToString(), ModulePathMap.Num());
+				UE_CLOG((InLoadModuleFlags & ELoadModuleFlags::LogFailures) != ELoadModuleFlags::None,
+					LogModuleManager, Warning, TEXT("ModuleManager: Unable to load module '%s'  - %d instances of that module name found."), *InModuleName.ToString(), ModulePathMap.Num());
+
 				OutFailureReason = EModuleLoadResult::FileNotFound;
 				return nullptr;
 			}
@@ -546,7 +552,7 @@ IModuleInterface* FModuleManager::LoadModuleWithFailureReason(const FName InModu
 						// these modules aren't using UObjects.
 							// OK, we've verified that loading the module caused new UObject classes to be
 							// registered, so we'll treat this module as a module with UObjects in it.
-					ProcessLoadedObjectsCallback.Broadcast(InModuleName, bCanProcessNewlyLoadedObjects);
+				ProcessLoadedObjectsCallback.Broadcast(InModuleName, bCanProcessNewlyLoadedObjects);
 
 				// Find our "InitializeModule" global function, which must exist for all module DLLs
 				FInitializeModuleFunctionPtr InitializeModuleFunctionPtr =
@@ -583,7 +589,8 @@ IModuleInterface* FModuleManager::LoadModuleWithFailureReason(const FName InModu
 						}
 						else
 						{
-							UE_LOG(LogModuleManager, Warning, TEXT("ModuleManager: Unable to load module '%s' because InitializeModule function failed (returned nullptr.)"), *ModuleFileToLoad);
+							UE_CLOG((InLoadModuleFlags & ELoadModuleFlags::LogFailures) != ELoadModuleFlags::None,
+								LogModuleManager, Warning, TEXT("ModuleManager: Unable to load module '%s' because InitializeModule function failed (returned nullptr.)"), *ModuleFileToLoad);
 
 							FPlatformProcess::FreeDllHandle(ModuleInfo->Handle);
 							ModuleInfo->Handle = nullptr;
@@ -593,7 +600,8 @@ IModuleInterface* FModuleManager::LoadModuleWithFailureReason(const FName InModu
 				}
 				else
 				{
-					UE_LOG(LogModuleManager, Warning, TEXT("ModuleManager: Unable to load module '%s' because InitializeModule function was not found."), *ModuleFileToLoad);
+					UE_CLOG((InLoadModuleFlags & ELoadModuleFlags::LogFailures) != ELoadModuleFlags::None,
+						LogModuleManager, Warning, TEXT("ModuleManager: Unable to load module '%s' because InitializeModule function was not found."), *ModuleFileToLoad);
 
 					FPlatformProcess::FreeDllHandle(ModuleInfo->Handle);
 					ModuleInfo->Handle = nullptr;
@@ -602,13 +610,17 @@ IModuleInterface* FModuleManager::LoadModuleWithFailureReason(const FName InModu
 			}
 			else
 			{
-				UE_LOG(LogModuleManager, Warning, TEXT("ModuleManager: Unable to load module '%s' because the file couldn't be loaded by the OS."), *ModuleFileToLoad);
+				UE_CLOG((InLoadModuleFlags & ELoadModuleFlags::LogFailures) != ELoadModuleFlags::None,
+					LogModuleManager, Warning, TEXT("ModuleManager: Unable to load module '%s' because the file couldn't be loaded by the OS."), *ModuleFileToLoad);
+
 				OutFailureReason = EModuleLoadResult::CouldNotBeLoadedByOS;
 			}
 		}
 		else
 		{
-			UE_LOG(LogModuleManager, Warning, TEXT("ModuleManager: Unable to load module '%s' because the file '%s' was not found."), *InModuleName.ToString(), *ModuleFileToLoad);
+			UE_CLOG((InLoadModuleFlags & ELoadModuleFlags::LogFailures) != ELoadModuleFlags::None,
+				LogModuleManager, Warning, TEXT("ModuleManager: Unable to load module '%s' because the file '%s' was not found."), *InModuleName.ToString(), *ModuleFileToLoad);
+
 			OutFailureReason = EModuleLoadResult::FileNotFound;
 		}
 	}
diff --git a/Engine/Source/Runtime/Core/Public/Modules/ModuleManager.h b/Engine/Source/Runtime/Core/Public/Modules/ModuleManager.h
--- a/Engine/Source/Runtime/Core/Public/Modules/ModuleManager.h
+++ b/Engine/Source/Runtime/Core/Public/Modules/ModuleManager.h
@@ -80,6 +80,17 @@ enum class ECheckModuleCompatibilityFlags
 ENUM_CLASS_FLAGS(ECheckModuleCompatibilityFlags)
 
 
+enum class ELoadModuleFlags
+{
+	None = 0x0,
+
+	// Print to the log any failure information
+	LogFailures = 1 << 0,
+};
+
+ENUM_CLASS_FLAGS(ELoadModuleFlags)
+
+
 /**
  * Structure for reporting module statuses.
  */
@@ -175,10 +186,11 @@ public:
 	 * Loads the specified module.
 	 *
 	 * @param InModuleName The base name of the module file.  Should not include path, extension or platform/configuration info.  This is just the "module name" part of the module file name.  Names should be globally unique.
+	 * @param InLoadModuleFlags Optional flags for module load operation.
 	 * @return The loaded module, or nullptr if the load operation failed.
 	 * @see AbandonModule, IsModuleLoaded, LoadModuleChecked, LoadModulePtr, LoadModuleWithFailureReason, UnloadModule
 	 */
-	IModuleInterface* LoadModule( const FName InModuleName );
+	IModuleInterface* LoadModule( const FName InModuleName, ELoadModuleFlags InLoadModuleFlags = ELoadModuleFlags::None );
 
 	/**
 	 * Loads the specified module, checking to ensure it exists.
@@ -204,10 +216,11 @@ public:
 	 *
 	 * @param InModuleName The base name of the module file.  Should not include path, extension or platform/configuration info.  This is just the "module name" part of the module file name.  Names should be globally unique.
 	 * @param OutFailureReason Will contain the result.
+	 * @param InLoadModuleFlags Optional flags for module load operation.
 	 * @return The loaded module (null if the load operation failed).
 	 * @see AbandonModule, IsModuleLoaded, LoadModule, LoadModuleChecked, LoadModulePtr, UnloadModule
 	 */
-	IModuleInterface* LoadModuleWithFailureReason( const FName InModuleName, EModuleLoadResult& OutFailureReason );
+	IModuleInterface* LoadModuleWithFailureReason( const FName InModuleName, EModuleLoadResult& OutFailureReason, ELoadModuleFlags InLoadModuleFlags = ELoadModuleFlags::None);
 
 	/**
 	 * Queries information about a specific module name.
diff --git a/Engine/Source/Runtime/Engine/Private/Pawn.cpp b/Engine/Source/Runtime/Engine/Private/Pawn.cpp
--- a/Engine/Source/Runtime/Engine/Private/Pawn.cpp
+++ b/Engine/Source/Runtime/Engine/Private/Pawn.cpp
@@ -44,8 +44,27 @@ APawn::APawn(const FObjectInitializer& ObjectInitializer)
 
 	if (HasAnyFlags(RF_ClassDefaultObject) && GetClass() == APawn::StaticClass())
 	{
-		// WARNING: This line is why the AISupport plugin has to load the AIModule before UObject initialization, otherwise this load fails and CDOs are corrupt in the editor
-		AIControllerClass = LoadClass<AController>(nullptr, *((UEngine*)(UEngine::StaticClass()->GetDefaultObject()))->AIControllerClassName.ToString(), nullptr, LOAD_None, nullptr);
+		bool bLoadPluginClass = true;
+
+#if !(UE_BUILD_SHIPPING || UE_BUILD_TEST)
+		// disabling engine plugins will cause loader warnings, this was added because it added non-valuable warnings for things like basic commandlets
+		if (FParse::Param(FCommandLine::Get(), TEXT("NoEnginePlugins")))
+		{
+			bLoadPluginClass = false;
+		}
+#endif
+
+		FString AIControllerClassName = GetDefault<UEngine>()->AIControllerClassName.ToString();
+
+		// resolve the name to a UClass
+		AIControllerClass = FindObject<UClass>(ANY_PACKAGE, *AIControllerClassName);
+
+		// if we failed to resolve, and plugins are expected to be loaded, proceed with loading it
+		if (AIControllerClass == nullptr && bLoadPluginClass)
+		{
+			// WARNING: This line is why the AISupport plugin has to load the AIModule before UObject initialization, otherwise this load fails and CDOs are corrupt in the editor
+			AIControllerClass = LoadClass<AController>(nullptr, *AIControllerClassName, nullptr, LOAD_None, nullptr);
+		}
 	}
 	else
 	{
