fix SSE FloatToHalf was broken for negatives, affected F32 to F16 texture conversion
single float version was fine
caused by difference of packs and packus variants

diff --git a/Engine/Source/Runtime/Core/Public/Math/UnrealPlatformMathSSE.h b/Engine/Source/Runtime/Core/Public/Math/UnrealPlatformMathSSE.h
--- a/Engine/Source/Runtime/Core/Public/Math/UnrealPlatformMathSSE.h
+++ b/Engine/Source/Runtime/Core/Public/Math/UnrealPlatformMathSSE.h
@@ -111,52 +111,54 @@ namespace SSE
 	// float_to_half_rtne_SSE2
 	static __m128i FloatToHalf(__m128 f)
 	{
-		__m128i mask_sign       = _mm_set1_epi32(0x80000000u);
-		__m128i c_f16max        = _mm_set1_epi32((127 + 16) << 23); // all FP32 values >=this round to +inf
-		__m128i c_nanbit        = _mm_set1_epi32(0x200);
-		__m128i c_infty_as_fp16 = _mm_set1_epi32(0x7c00);
-		__m128i c_min_normal    = _mm_set1_epi32((127 - 14) << 23); // smallest FP32 that yields a normalized FP16
-		__m128i c_subnorm_magic = _mm_set1_epi32(((127 - 15) + (23 - 10) + 1) << 23);
-		__m128i c_normal_bias   = _mm_set1_epi32(0xfff - ((127 - 15) << 23)); // adjust exponent and add mantissa rounding
-
-		__m128  msign       = _mm_castsi128_ps(mask_sign);
-		__m128  justsign    = _mm_and_ps(msign, f);
-		__m128  absf        = _mm_xor_ps(f, justsign);
-		__m128i absf_int    = _mm_castps_si128(absf); // the cast is "free" (extra bypass latency, but no thruput hit)
-		__m128i f16max      = c_f16max;
-		__m128  b_isnan     = _mm_cmpunord_ps(absf, absf); // is this a NaN?
-		__m128i b_isregular = _mm_cmpgt_epi32(f16max, absf_int); // (sub)normalized or special?
-		__m128i nanbit      = _mm_and_si128(_mm_castps_si128(b_isnan), c_nanbit);
-		__m128i inf_or_nan  = _mm_or_si128(nanbit, c_infty_as_fp16); // output for specials
-
-		__m128i min_normal  = c_min_normal;
-		__m128i b_issub     = _mm_cmpgt_epi32(min_normal, absf_int);
+		const __m128 mask_sign		= _mm_set1_ps(-0.0f);
+		const __m128i c_f16max			= _mm_set1_epi32((127 + 16) << 23); // all FP32 values >=this round to +inf
+		const __m128i c_nanbit			= _mm_set1_epi32(0x200);
+		const __m128i c_nanlobits        = _mm_set1_epi32(0x1ff);
+		const __m128i c_infty_as_fp16	= _mm_set1_epi32(0x7c00);
+		const __m128i c_min_normal		= _mm_set1_epi32((127 - 14) << 23); // smallest FP32 that yields a normalized FP16
+		const __m128i c_subnorm_magic	= _mm_set1_epi32(((127 - 15) + (23 - 10) + 1) << 23);
+		const __m128i c_normal_bias		= _mm_set1_epi32(0xfff - ((127 - 15) << 23)); // adjust exponent and add mantissa rounding
+
+		//__m128 justsign	= f & mask_sign;
+		__m128 justsign	= _mm_and_ps( f , mask_sign );
+		//__m128 absf		= andnot(f, mask_sign); // f & ~mask_sign
+		__m128 absf		= _mm_andnot_ps(mask_sign, f); // f & ~mask_sign
+		__m128i absf_int		= _mm_castps_si128(absf); // the cast is "free" (extra bypass latency, but no thruput hit)
+		__m128 b_isnan	= _mm_cmpunord_ps(absf, absf); // is this a NaN?
+		__m128i b_isregular	= _mm_cmpgt_epi32(c_f16max, absf_int); // (sub)normalized or special?
+		__m128i nan_payload  = _mm_and_si128(_mm_srli_epi32(absf_int, 13), c_nanlobits); // payload bits for NaNs
+		__m128i nan_quiet    = _mm_or_si128(nan_payload, c_nanbit); // and set quiet bit
+		__m128i nanfinal		= _mm_and_si128(_mm_castps_si128(b_isnan), nan_quiet);
+		__m128i inf_or_nan	= _mm_or_si128(nanfinal, c_infty_as_fp16); // output for specials
+
+		__m128i b_issub		= _mm_cmpgt_epi32(c_min_normal, absf_int);
 
 		// "result is subnormal" path
-		__m128  subnorm1    = _mm_add_ps(absf, _mm_castsi128_ps(c_subnorm_magic)); // magic value to round output mantissa
-		__m128i subnorm2    = _mm_sub_epi32(_mm_castps_si128(subnorm1), c_subnorm_magic); // subtract out bias
+		__m128 subnorm1	= _mm_add_ps( absf , __m128(_mm_castsi128_ps(c_subnorm_magic)) ); // magic value to round output mantissa
+		__m128i subnorm2		= _mm_sub_epi32(_mm_castps_si128(subnorm1), c_subnorm_magic); // subtract out bias
 
 		// "result is normal" path
-		__m128i mantoddbit  = _mm_slli_epi32(absf_int, 31 - 13); // shift bit 13 (mantissa LSB) to sign
-		__m128i mantodd     = _mm_srai_epi32(mantoddbit, 31); // -1 if FP16 mantissa odd, else 0
+		__m128i mantoddbit	= _mm_slli_epi32(absf_int, 31 - 13); // shift bit 13 (mantissa LSB) to sign
+		__m128i mantodd		= _mm_srai_epi32(mantoddbit, 31); // -1 if FP16 mantissa odd, else 0
 
-		__m128i round1      = _mm_add_epi32(absf_int, c_normal_bias);
-		__m128i round2      = _mm_sub_epi32(round1, mantodd); // if mantissa LSB odd, bias towards rounding up (RTNE)
-		__m128i normal      = _mm_srli_epi32(round2, 13); // rounded result
+		__m128i round1		= _mm_add_epi32(absf_int, c_normal_bias);
+		__m128i round2		= _mm_sub_epi32(round1, mantodd); // if mantissa LSB odd, bias towards rounding up (RTNE)
+		__m128i normal		= _mm_srli_epi32(round2, 13); // rounded result
 
 		// combine the two non-specials
-		__m128i nonspecial  = _mm_or_si128(_mm_and_si128(subnorm2, b_issub), _mm_andnot_si128(b_issub, normal));
+		__m128i nonspecial	= _mm_or_si128(_mm_and_si128(subnorm2, b_issub), _mm_andnot_si128(b_issub, normal));
 
 		// merge in specials as well
-		__m128i joined      = _mm_or_si128(_mm_and_si128(nonspecial, b_isregular), _mm_andnot_si128(b_isregular, inf_or_nan));
+		__m128i joined		= _mm_or_si128(_mm_and_si128(nonspecial, b_isregular), _mm_andnot_si128(b_isregular, inf_or_nan));
+
+		__m128i sign_shift	= _mm_srai_epi32(_mm_castps_si128(justsign), 16);
+		__m128i rgba_half32		= _mm_or_si128(joined, sign_shift);
 
-		__m128i sign_shift  = _mm_srli_epi32(_mm_castps_si128(justsign), 16);
-		__m128i rgba_half32 = _mm_or_si128(joined, sign_shift);
-		
 		// there's now a half in each 32-bit lane
 		// pack down to 64 bits :
 		// packs works because rgba_half32 is sign-extended
-        __m128i four_halfs_u64 = _mm_packs_epi32(rgba_half32, _mm_setzero_si128());
+        __m128i four_halfs_u64 = _mm_packs_epi32(rgba_half32, rgba_half32);
 
 		return four_halfs_u64;
 	}
