Ensure WriteToString is treated as a contiguous container.

diff --git a/Engine/Source/Runtime/Core/Public/Misc/StringBuilder.h b/Engine/Source/Runtime/Core/Public/Misc/StringBuilder.h
--- a/Engine/Source/Runtime/Core/Public/Misc/StringBuilder.h
+++ b/Engine/Source/Runtime/Core/Public/Misc/StringBuilder.h
@@ -478,3 +478,5 @@ public:
 	#endif
 	}
 };
+
+template <int32 BufferSize> struct TIsContiguousContainer<WriteToString<BufferSize>> { static constexpr bool Value = true; };
