CompressedBuffer: Added a serialization operator for FArchive

#rnx

diff --git a/Engine/Source/Runtime/Core/Private/Compression/CompressedBuffer.cpp b/Engine/Source/Runtime/Core/Private/Compression/CompressedBuffer.cpp
--- a/Engine/Source/Runtime/Core/Private/Compression/CompressedBuffer.cpp
+++ b/Engine/Source/Runtime/Core/Private/Compression/CompressedBuffer.cpp
@@ -705,3 +705,21 @@ FCompositeBuffer FCompressedBuffer::DecompressToComposite() const
 	}
 	return FCompositeBuffer();
 }
+
+///////////////////////////////////////////////////////////////////////////////////////////////////////////////////////
+
+FArchive& operator<<(FArchive& Ar, FCompressedBuffer& Buffer)
+{
+	if (Ar.IsLoading())
+	{
+		Buffer = FCompressedBuffer::FromCompressed(Ar);
+	}
+	else
+	{
+		for (const FSharedBuffer& Segment : Buffer.GetCompressed().GetSegments())
+		{
+			Ar.Serialize(const_cast<void*>(Segment.GetData()), static_cast<int64>(Segment.GetSize()));
+		}
+	}
+	return Ar;
+}
diff --git a/Engine/Source/Runtime/Core/Public/Compression/CompressedBuffer.h b/Engine/Source/Runtime/Core/Public/Compression/CompressedBuffer.h
--- a/Engine/Source/Runtime/Core/Public/Compression/CompressedBuffer.h
+++ b/Engine/Source/Runtime/Core/Public/Compression/CompressedBuffer.h
@@ -101,3 +101,5 @@ public:
 private:
 	FCompositeBuffer CompressedData;
 };
+
+CORE_API FArchive& operator<<(FArchive& Ar, FCompressedBuffer& Buffer);
