Fixed LinuxAArch64 build

diff --git a/Engine/Source/Runtime/Core/Public/Math/UnrealMathNeon.h b/Engine/Source/Runtime/Core/Public/Math/UnrealMathNeon.h
--- a/Engine/Source/Runtime/Core/Public/Math/UnrealMathNeon.h
+++ b/Engine/Source/Runtime/Core/Public/Math/UnrealMathNeon.h
@@ -260,6 +260,18 @@ FORCEINLINE VectorRegister2Double MakeVectorRegister2Double(double X, double Y)
 	return Tmp.V;
 }
 
+FORCEINLINE VectorRegister2Double MakeVectorRegister2Double(uint64 X, uint64 Y)
+{
+	union U
+	{
+		VectorRegister2Double V; uint64_t D[2];
+		FORCEINLINE U() : V() {}
+	} Tmp;
+	Tmp.D[0] = X;
+	Tmp.D[1] = Y;
+	return Tmp.V;
+}
+
 /**
 * Returns a vector based on 4 int32.
 *
