TSharedPtr: let a function pointer to be used as a custom deleter

diff --git a/Engine/Source/Runtime/Core/Public/Templates/SharedPointerInternals.h b/Engine/Source/Runtime/Core/Public/Templates/SharedPointerInternals.h
--- a/Engine/Source/Runtime/Core/Public/Templates/SharedPointerInternals.h
+++ b/Engine/Source/Runtime/Core/Public/Templates/SharedPointerInternals.h
@@ -98,19 +98,56 @@ namespace SharedPointerInternals
 		FReferenceControllerBase& operator=( FReferenceControllerBase const& );
 	};
 
+	// A helper class that efficiently stores a custom deleter and is intended to be derived from. If the custom deleter is an empty class,
+	// TDeleterHolder derives from it exploiting empty base optimisation (https://en.cppreference.com/w/cpp/language/ebo). Otherwise
+	// it stores the custom deleter as a member to allow a function pointer to be used as a custom deleter (a function pointer can't 
+	// be a base class)
+	template <typename DeleterType, bool bIsZeroSize = std::is_empty_v<DeleterType>>
+	struct TDeleterHolder : private DeleterType
+	{
+		explicit TDeleterHolder(DeleterType&& Arg)
+			: DeleterType(MoveTemp(Arg))
+		{
+		}
+
+		template <typename ObjectType>
+		void InvokeDeleter(ObjectType * Object)
+		{
+			Invoke(*static_cast<DeleterType*>(this), Object);
+		}
+	};
+
+	template <typename DeleterType>
+	struct TDeleterHolder<DeleterType, false>
+	{
+		explicit TDeleterHolder(DeleterType&& Arg)
+			: Deleter(MoveTemp(Arg))
+		{
+		}
+
+		template <typename ObjectType>
+		void InvokeDeleter(ObjectType * Object)
+		{
+			Invoke(Deleter, Object);
+		}
+
+	private:
+		DeleterType Deleter;
+	};
+
 	template <typename ObjectType, typename DeleterType>
-	class TReferenceControllerWithDeleter : private DeleterType, public FReferenceControllerBase
+	class TReferenceControllerWithDeleter : private TDeleterHolder<DeleterType>, public FReferenceControllerBase
 	{
 	public:
 		explicit TReferenceControllerWithDeleter(ObjectType* InObject, DeleterType&& Deleter)
-			: DeleterType(MoveTemp(Deleter))
+			: TDeleterHolder<DeleterType>(MoveTemp(Deleter))
 			, Object(InObject)
 		{
 		}
 
 		virtual void DestroyObject() override
 		{
-			(*static_cast<DeleterType*>(this))(Object);
+			this->InvokeDeleter(Object);
 		}
 
 		// Non-copyable
