U5 - Remove controversal TryCancelAndLaunchContinuation in favor of just rolling all that functionallity into TryCancel as the Scheduler is not used anymore since an optimization where we run the Continuation inline in case the Task has not yet been launched. Also added TryExecute for cases where one just wants to run the Task inline and it will be executed if it has not yet been launched.

diff --git a/Engine/Source/Runtime/Core/Public/Async/Fundamental/Scheduler.h b/Engine/Source/Runtime/Core/Public/Async/Fundamental/Scheduler.h
--- a/Engine/Source/Runtime/Core/Public/Async/Fundamental/Scheduler.h
+++ b/Engine/Source/Runtime/Core/Public/Async/Fundamental/Scheduler.h
@@ -64,12 +64,7 @@ namespace LowLevelTasks
 		CORE_API void StopWorkers();
 
 		//try to launch the task, the return value will specify if the task was in the ready state and has been launced
-		inline bool TryLaunch(FTask& Task, EQueuePreference QueuePreference = EQueuePreference::DefaultPreference, bool bWakeUpWorker = true);
-		
-		//try to cancel the task and launching it if the task was in ready state
-		//you still need to wait for task completion before recycling the handle
-		//you can alternatively use FTask::TryCancel if you want to launch the task manually
-		inline bool TryCancelAndLaunchContinuation(FTask& Task, EQueuePreference QueuePreference = EQueuePreference::DefaultPreference);		
+		inline bool TryLaunch(FTask& Task, EQueuePreference QueuePreference = EQueuePreference::DefaultPreference, bool bWakeUpWorker = true);	
 
 		//tries to do some work until the Task is completed
 		template<typename TaskType>
@@ -152,13 +147,11 @@ namespace LowLevelTasks
 
 			void Push(FSleepEvent* Item)
 			{
-#if defined(_MSC_VER) && USING_CODE_ANALYSIS
-#pragma warning(push)
-#pragma warning(disable:6011) // Dereferencing NULL pointer 'Item'
-#endif
-				check(Item != nullptr);
+				checkSlow(Item != nullptr);
+#if !USING_CODE_ANALYSIS //MS SA thowing warning C6011 on Item->Next access, even when it is validated or branched over
 				checkSlow(reinterpret_cast<uintptr_t>(Item) < (1ull << 48));
 				checkSlow((reinterpret_cast<uintptr_t>(Item) & 0x7) == 0);
+#endif
 				checkSlow(Item->Next == nullptr);
 			
 				FTopNode LocalTop = Top.load(std::memory_order_relaxed);
@@ -178,9 +171,6 @@ namespace LowLevelTasks
 #endif
 				}
 			}
-#if defined(_MSC_VER) && USING_CODE_ANALYSIS
-#pragma warning(pop)
-#endif
 		};
 
 	public:
@@ -213,11 +203,6 @@ namespace LowLevelTasks
 		return FScheduler::Get().TryLaunch(Task, QueuePreference, bWakeUpWorker);
 	}
 
-	FORCEINLINE_DEBUGGABLE bool TryCancelAndLaunchContinuation(FTask& Task, EQueuePreference QueuePreference = EQueuePreference::DefaultPreference)
-	{
-		return FScheduler::Get().TryCancelAndLaunchContinuation(Task, QueuePreference);
-	}
-
 	FORCEINLINE_DEBUGGABLE void BusyWaitForTask(const FTask& Task)
 	{
 		FScheduler::Get().BusyWait(Task);
@@ -248,17 +233,6 @@ namespace LowLevelTasks
 		return false;
 	}
 
-	inline bool FScheduler::TryCancelAndLaunchContinuation(FTask& Task, EQueuePreference QueuePreference)
-	{
-		bool WasCanceled = Task.TryCancel();
-		if(WasCanceled && Task.TryPrepareLaunch())
-		{
-			Task.ExecuteTask();
-			return true;
-		}
-		return WasCanceled;
-	}
-
 	inline uint32 FScheduler::GetNumWorkers() const
 	{
 		return ActiveWorkers.load(std::memory_order_relaxed);
diff --git a/Engine/Source/Runtime/Core/Public/Async/Fundamental/Task.h b/Engine/Source/Runtime/Core/Public/Async/Fundamental/Task.h
--- a/Engine/Source/Runtime/Core/Public/Async/Fundamental/Task.h
+++ b/Engine/Source/Runtime/Core/Public/Async/Fundamental/Task.h
@@ -15,8 +15,8 @@ namespace LowLevelTasks
 		High,
 		Normal,
 		Default = Normal,
-		BackgroundHigh,
-		ForegroundCount = BackgroundHigh,
+		ForegroundCount,
+		BackgroundHigh = ForegroundCount,
 		BackgroundNormal,
 		BackgroundLow,
 		Count
@@ -171,11 +171,12 @@ namespace LowLevelTasks
 		}
 
 		
-		//try to cancel the task without launching it. Even if the task is canceled it sill needs to be launched manually
-		//to launch the task manually you use FScheduler::TryLaunch
-		//you can alternatively use FScheduler::TryCancelAndLaunch if you want to launch the task in this case automatically
+		//try to cancel the task if it has not been launched yet the continuation will run immediately.
 		inline bool TryCancel();
 
+		//try to execute the task if it has not been launched yet the task will execute immediately.
+		inline bool TryExecute();
+
 		template<typename TRunnable, typename TContinuation>
 		inline void Init(const TCHAR* InDebugName, ETaskPriority InPriority, TRunnable&& InRunnable, TContinuation&& InContinuation, bool bAllowBusyWaiting);
 
@@ -335,8 +336,25 @@ namespace LowLevelTasks
 		FPackedData ScheduledState(LocalPackedData, ETaskState::Scheduled);
 		//we can use memory_order_relaxed in this case because, when a task is canceled it does not launch the task
 		//to launch a canceled  task it has to go though TryPrepareLaunch which is doing the memory_order_acquire
-		return PackedData.compare_exchange_strong(ReadyState, FPackedData(LocalPackedData, ETaskState::CanceledAndReady), std::memory_order_relaxed)
+		bool WasCanceled = PackedData.compare_exchange_strong(ReadyState, FPackedData(LocalPackedData, ETaskState::CanceledAndReady), std::memory_order_relaxed)
 			|| PackedData.compare_exchange_strong(ScheduledState, FPackedData(LocalPackedData, ETaskState::Canceled), std::memory_order_relaxed);
+
+		if(WasCanceled && TryPrepareLaunch())
+		{
+			ExecuteTask();
+			return true;
+		}
+		return WasCanceled;
+	}
+
+	inline bool FTask::TryExecute()
+	{
+		if(TryPrepareLaunch())
+		{
+			ExecuteTask();
+			return true;
+		}
+		return false;
 	}
 
 	inline void FTask::ExecuteTask()
diff --git a/Engine/Source/Runtime/Core/Public/Experimental/Async/AwaitableTask.h b/Engine/Source/Runtime/Core/Public/Experimental/Async/AwaitableTask.h
--- a/Engine/Source/Runtime/Core/Public/Experimental/Async/AwaitableTask.h
+++ b/Engine/Source/Runtime/Core/Public/Experimental/Async/AwaitableTask.h
@@ -260,8 +260,8 @@ namespace AwaitableTask_Detail
 			int LocalCounter = ReferenceCounter.fetch_sub(1, std::memory_order_release);
 			if (LocalCounter == 2 && !IsLaunched()) //this is 2 because fetch_sub returns the value before the decrement (value is actually 1)
 			{
-				//Cancel the Task and Try to Launch the Continuation if we are the last reference
-				verify(LowLevelTasks::TryCancelAndLaunchContinuation(Task));
+				//Cancel the Task and Try to run the Continuation if we are the last reference
+				verify(Task.TryCancel());
 			}
 			else if (LocalCounter == 1) //this is 1 because fetch_sub returns the value before the decrement (value is actually 0)
 			{
@@ -271,7 +271,7 @@ namespace AwaitableTask_Detail
 
 		inline ReturnType GetResult()
 		{
-			if (LowLevelTasks::TryCancelAndLaunchContinuation(Task))
+			if (Task.TryCancel())
 			{
 				Execute();
 				return ReturnValue;
@@ -344,7 +344,7 @@ namespace AwaitableTask_Detail
 			if (LocalCounter == 2 && !IsLaunched()) //this is 2 because fetch_sub returns the value before the decrement (value is actually 1)
 			{
 				//Cancel the Task and Try to Launch the Continuation if we are the last reference
-				verify(LowLevelTasks::TryCancelAndLaunchContinuation(Task));
+				verify(Task.TryCancel());
 			}
 			else if (LocalCounter == 1) //this is 1 because fetch_sub returns the value before the decrement (value is actually 0)
 			{
@@ -354,7 +354,7 @@ namespace AwaitableTask_Detail
 
 		inline void GetResult()
 		{
-			if (LowLevelTasks::TryCancelAndLaunchContinuation(Task))
+			if (Task.TryCancel())
 			{
 				Execute();
 			}
diff --git a/Engine/Source/Runtime/Core/Public/Tasks/TaskPrivate.h b/Engine/Source/Runtime/Core/Public/Tasks/TaskPrivate.h
--- a/Engine/Source/Runtime/Core/Public/Tasks/TaskPrivate.h
+++ b/Engine/Source/Runtime/Core/Public/Tasks/TaskPrivate.h
@@ -137,7 +137,7 @@ namespace UE { namespace Tasks
 				// the task skip waiting in the scheduler's queue.
 				// we can't retract it right here because it can be not scheduled yet (e.g. when it's blocked by a pipe), so we need to wait for this
 				LowLevelTasks::BusyWaitUntil([this] { return !LowLevelTask.IsReady(); }); // `!IsReady()` here means "the task is scheduled"
-				LowLevelTasks::TryCancelAndLaunchContinuation(LowLevelTask);
+				LowLevelTask.TryCancel(); //This will try to cancel the task and run it's continuation (thereby executing the workload) 
 				LowLevelTasks::BusyWaitUntil([this] { return LowLevelTask.IsCompleted(); });
 			}
 
@@ -152,7 +152,7 @@ namespace UE { namespace Tasks
 					return false;
 				}
 
-				LowLevelTasks::TryCancelAndLaunchContinuation(LowLevelTask);
+				LowLevelTask.TryCancel(); //This will try to cancel the task and run it's continuation (thereby executing the workload) 
 				LowLevelTasks::BusyWaitUntil([this, Timeout] { return LowLevelTask.IsCompleted() || Timeout; });
 				return LowLevelTask.IsCompleted();
 			}
@@ -171,7 +171,7 @@ namespace UE { namespace Tasks
 					return false;
 				}
 
-				LowLevelTasks::TryCancelAndLaunchContinuation(LowLevelTask);
+				LowLevelTask.TryCancel(); //This will try to cancel the task and run it's continuation (thereby executing the workload) 
 				LowLevelTasks::BusyWaitUntil(
 					[this, Condition = Forward<ConditionType>(Condition)] { return LowLevelTask.IsCompleted() || Condition(); }
 				);
