Workaround case when GetStack returns an empty stack, and fix some of the cases that can cause it to return an empty stack.
#rnx

diff --git a/Engine/Source/Runtime/Core/Private/Misc/AutomationTest.cpp b/Engine/Source/Runtime/Core/Private/Misc/AutomationTest.cpp
--- a/Engine/Source/Runtime/Core/Private/Misc/AutomationTest.cpp
+++ b/Engine/Source/Runtime/Core/Private/Misc/AutomationTest.cpp
@@ -1057,7 +1057,7 @@ void FAutomationTestExecutionInfo::AddEvent(const FAutomationEvent& Event, int S
 	int32 EntryIndex = 0;
 	if (FAutomationTestFramework::Get().GetCaptureStack())
 	{
-		TArray<FProgramCounterSymbolInfo> Stack = FPlatformStackWalk::GetStack(StackOffset + 1, 1);
+		SAFE_GETSTACK(Stack, StackOffset + 1, 1);
 		EntryIndex = Entries.Add(FAutomationExecutionEntry(Event, Stack[0].Filename, Stack[0].LineNumber));
 	}
 	else
diff --git a/Engine/Source/Runtime/Core/Private/Windows/WindowsPlatformStackWalk.cpp b/Engine/Source/Runtime/Core/Private/Windows/WindowsPlatformStackWalk.cpp
--- a/Engine/Source/Runtime/Core/Private/Windows/WindowsPlatformStackWalk.cpp
+++ b/Engine/Source/Runtime/Core/Private/Windows/WindowsPlatformStackWalk.cpp
@@ -370,7 +370,19 @@ uint32 FWindowsPlatformStackWalk::CaptureStackBackTrace( uint64* BackTrace, uint
 			DetermineMaxCallstackDepth();
 		}
 		PVOID WinBackTrace[MAX_CALLSTACK_DEPTH];
-		uint16 NumFrames = RtlCaptureStackBackTrace(0, FMath::Min<ULONG>(GMaxCallstackDepth, MaxDepth), WinBackTrace, NULL);
+		ULONG RequestedDepth = FMath::Min<ULONG>(GMaxCallstackDepth, MaxDepth);
+		uint16 NumFrames = RtlCaptureStackBackTrace(0, RequestedDepth, WinBackTrace, NULL);
+		if (NumFrames == 0 && RequestedDepth > 0)
+		{
+			// RtlCaptureStackBackTrace sometimes returns 0, which NTAPI Undocumented Functions says is the value it returns to indicate an error
+			// In all the cases we have tested, calling it again after it returns 0 will make it succeed and return non-zero
+			NumFrames = RtlCaptureStackBackTrace(0, RequestedDepth, WinBackTrace, NULL);
+			if (NumFrames == 0)
+			{
+				// Don't send a log since this might be called from inside the log system
+				FPlatformMisc::LocalPrint(TEXT("RtlCaptureStackBackTrace failed; returning an empty backtrace."));
+			}
+		}
 		Depth = NumFrames;
 		for (uint16 FrameIndex = 0; FrameIndex < NumFrames; ++FrameIndex)
 		{
diff --git a/Engine/Source/Runtime/Core/Public/Misc/AutomationTest.h b/Engine/Source/Runtime/Core/Public/Misc/AutomationTest.h
--- a/Engine/Source/Runtime/Core/Public/Misc/AutomationTest.h
+++ b/Engine/Source/Runtime/Core/Public/Misc/AutomationTest.h
@@ -18,6 +18,7 @@
 #include "HAL/PlatformStackWalk.h"
 #include "GenericPlatform/GenericPlatformStackWalk.h" 
 #include "Containers/Queue.h"
+#include "Misc/CString.h"
 #include "Misc/FeedbackContext.h"
 #include "Async/Future.h"
 #include "Async/Async.h"
@@ -34,6 +35,17 @@
 	#define WITH_AUTOMATION_TESTS (WITH_DEV_AUTOMATION_TESTS || WITH_PERF_AUTOMATION_TESTS)
 #endif
 
+/** Call GetStack, with a guarantee of a non-empty return; a placeholder "Unknown",1) is used if necessary */
+#define SAFE_GETSTACK(VariableName, IgnoreCount, MaxDepth)													\
+	TArray<FProgramCounterSymbolInfo> VariableName = FPlatformStackWalk::GetStack(IgnoreCount, MaxDepth);	\
+	if (VariableName.Num() == 0)																			\
+	{																										\
+		/* This is a rare failure that can occur in some circumstances */									\
+		FProgramCounterSymbolInfo& Info = VariableName.Emplace_GetRef();									\
+		TCString<ANSICHAR>::Strcpy(Info.Filename, FProgramCounterSymbolInfo::MAX_NAME_LENGTH, "Unknown");	\
+		Info.LineNumber = 1;																				\
+	}
+
 /** Flags for specifying automation test requirements/behavior */
 namespace EAutomationTestFlags
 {
@@ -2452,7 +2464,7 @@ public:
 	void It(const FString& InDescription, TFunction<void()> DoWork)
 	{
 		const TSharedRef<FSpecDefinitionScope> CurrentScope = DefinitionScopeStack.Last();
-		const TArray<FProgramCounterSymbolInfo> Stack = FPlatformStackWalk::GetStack(1, 1);
+		SAFE_GETSTACK(Stack, 1, 1);
 
 		PushDescription(InDescription);
 		CurrentScope->It.Push(MakeShareable(new FSpecIt(GetDescription(), GetId(), Stack[0].Filename, Stack[0].LineNumber, MakeShareable(new FSingleExecuteLatentCommand(this, DoWork, bEnableSkipIfError)))));
@@ -2462,7 +2474,7 @@ public:
 	void It(const FString& InDescription, EAsyncExecution Execution, TFunction<void()> DoWork)
 	{
 		const TSharedRef<FSpecDefinitionScope> CurrentScope = DefinitionScopeStack.Last();
-		const TArray<FProgramCounterSymbolInfo> Stack = FPlatformStackWalk::GetStack(1, 1);
+		SAFE_GETSTACK(Stack, 1, 1);
 
 		PushDescription(InDescription);
 		CurrentScope->It.Push(MakeShareable(new FSpecIt(GetDescription(), GetId(), Stack[0].Filename, Stack[0].LineNumber, MakeShareable(new FAsyncLatentCommand(this, Execution, DoWork, DefaultTimeout, bEnableSkipIfError)))));
@@ -2472,7 +2484,7 @@ public:
 	void It(const FString& InDescription, EAsyncExecution Execution, const FTimespan& Timeout, TFunction<void()> DoWork)
 	{
 		const TSharedRef<FSpecDefinitionScope> CurrentScope = DefinitionScopeStack.Last();
-		const TArray<FProgramCounterSymbolInfo> Stack = FPlatformStackWalk::GetStack(1, 1);
+		SAFE_GETSTACK(Stack, 1, 1);
 
 		PushDescription(InDescription);
 		CurrentScope->It.Push(MakeShareable(new FSpecIt(GetDescription(), GetId(), Stack[0].Filename, Stack[0].LineNumber, MakeShareable(new FAsyncLatentCommand(this, Execution, DoWork, Timeout, bEnableSkipIfError)))));
@@ -2482,7 +2494,7 @@ public:
 	void LatentIt(const FString& InDescription, TFunction<void(const FDoneDelegate&)> DoWork)
 	{
 		const TSharedRef<FSpecDefinitionScope> CurrentScope = DefinitionScopeStack.Last();
-		const TArray<FProgramCounterSymbolInfo> Stack = FPlatformStackWalk::GetStack(1, 1);
+		SAFE_GETSTACK(Stack, 1, 1);
 
 		PushDescription(InDescription);
 		CurrentScope->It.Push(MakeShareable(new FSpecIt(GetDescription(), GetId(), Stack[0].Filename, Stack[0].LineNumber, MakeShareable(new FUntilDoneLatentCommand(this, DoWork, DefaultTimeout, bEnableSkipIfError)))));
@@ -2492,7 +2504,7 @@ public:
 	void LatentIt(const FString& InDescription, const FTimespan& Timeout, TFunction<void(const FDoneDelegate&)> DoWork)
 	{
 		const TSharedRef<FSpecDefinitionScope> CurrentScope = DefinitionScopeStack.Last();
-		const TArray<FProgramCounterSymbolInfo> Stack = FPlatformStackWalk::GetStack(1, 1);
+		SAFE_GETSTACK(Stack, 1, 1);
 
 		PushDescription(InDescription);
 		CurrentScope->It.Push(MakeShareable(new FSpecIt(GetDescription(), GetId(), Stack[0].Filename, Stack[0].LineNumber, MakeShareable(new FUntilDoneLatentCommand(this, DoWork, Timeout, bEnableSkipIfError)))));
@@ -2502,7 +2514,7 @@ public:
 	void LatentIt(const FString& InDescription, EAsyncExecution Execution, TFunction<void(const FDoneDelegate&)> DoWork)
 	{
 		const TSharedRef<FSpecDefinitionScope> CurrentScope = DefinitionScopeStack.Last();
-		const TArray<FProgramCounterSymbolInfo> Stack = FPlatformStackWalk::GetStack(1, 1);
+		SAFE_GETSTACK(Stack, 1, 1);
 
 		PushDescription(InDescription);
 		CurrentScope->It.Push(MakeShareable(new FSpecIt(GetDescription(), GetId(), Stack[0].Filename, Stack[0].LineNumber, MakeShareable(new FAsyncUntilDoneLatentCommand(this, Execution, DoWork, DefaultTimeout, bEnableSkipIfError)))));
@@ -2512,7 +2524,7 @@ public:
 	void LatentIt(const FString& InDescription, EAsyncExecution Execution, const FTimespan& Timeout, TFunction<void(const FDoneDelegate&)> DoWork)
 	{
 		const TSharedRef<FSpecDefinitionScope> CurrentScope = DefinitionScopeStack.Last();
-		const TArray<FProgramCounterSymbolInfo> Stack = FPlatformStackWalk::GetStack(1, 1);
+		SAFE_GETSTACK(Stack, 1, 1);
 
 		PushDescription(InDescription);
 		CurrentScope->It.Push(MakeShareable(new FSpecIt(GetDescription(), GetId(), Stack[0].Filename, Stack[0].LineNumber, MakeShareable(new FAsyncUntilDoneLatentCommand(this, Execution, DoWork, Timeout, bEnableSkipIfError)))));
