Fix conversion from uint32 -> double, should be uint64 -> double

diff --git a/Engine/Source/Runtime/Core/Public/Math/UnrealMathVectorConstants.h b/Engine/Source/Runtime/Core/Public/Math/UnrealMathVectorConstants.h
--- a/Engine/Source/Runtime/Core/Public/Math/UnrealMathVectorConstants.h
+++ b/Engine/Source/Runtime/Core/Public/Math/UnrealMathVectorConstants.h
@@ -126,5 +126,5 @@ namespace GlobalVectorConstants
 
 	static const VectorRegister4Double DoubleTwo = MakeVectorRegisterDouble(2.0, 2.0, 2.0, 2.0);
 	static const uint64 DoubleAlmostTwoBits = 0x3FFFFFFFFFFFFFFF;
-	static const VectorRegister4Double DoubleAlmostTwo = MakeVectorRegisterDouble(*(double*)&AlmostTwoBits, *(double*)&AlmostTwoBits, *(double*)&AlmostTwoBits, *(double*)&AlmostTwoBits);
+	static const VectorRegister4Double DoubleAlmostTwo = MakeVectorRegisterDouble(*(double*)&DoubleAlmostTwoBits, *(double*)&DoubleAlmostTwoBits, *(double*)&DoubleAlmostTwoBits, *(double*)&DoubleAlmostTwoBits);
 }
