Fixup ASAN and autoreleasepool issues.

#rnx

diff --git a/Engine/Source/Runtime/Core/Private/Apple/ApplePlatformFile.cpp b/Engine/Source/Runtime/Core/Private/Apple/ApplePlatformFile.cpp
--- a/Engine/Source/Runtime/Core/Private/Apple/ApplePlatformFile.cpp
+++ b/Engine/Source/Runtime/Core/Private/Apple/ApplePlatformFile.cpp
@@ -677,13 +677,12 @@ bool FApplePlatformFile::DirectoryExists(const TCHAR* Directory)
 
 bool FApplePlatformFile::CreateDirectory(const TCHAR* Directory)
 {
-	@autoreleasepool
-	{
-		CFStringRef CFDirectory = FPlatformString::TCHARToCFString(*NormalizeFilename(Directory));
-		bool Result = [[NSFileManager defaultManager] createDirectoryAtPath:(NSString*)CFDirectory withIntermediateDirectories:true attributes:nil error:nil];
-		CFRelease(CFDirectory);
-		return Result;
-	}
+	SCOPED_AUTORELEASE_POOL;
+
+	CFStringRef CFDirectory = FPlatformString::TCHARToCFString(*NormalizeFilename(Directory));
+	bool Result = [[NSFileManager defaultManager] createDirectoryAtPath:(NSString*)CFDirectory withIntermediateDirectories:true attributes:nil error:nil];
+	CFRelease(CFDirectory);
+	return Result;
 }
 
 bool FApplePlatformFile::DeleteDirectory(const TCHAR* Directory)
@@ -704,53 +703,51 @@ FFileStatData FApplePlatformFile::GetStatData(const TCHAR* FilenameOrDirectory)
 
 bool FApplePlatformFile::IterateDirectory(const TCHAR* Directory, FDirectoryVisitor& Visitor)
 {
-	@autoreleasepool
+	const FString DirectoryStr = Directory;
+	const FString NormalizedDirectoryStr = NormalizeFilename(Directory);
+
+	return IterateDirectoryCommon(Directory, [&Visitor, &DirectoryStr, &NormalizedDirectoryStr](struct dirent* InEntry) -> bool
 	{
-		const FString DirectoryStr = Directory;
-		const FString NormalizedDirectoryStr = NormalizeFilename(Directory);
+		SCOPED_AUTORELEASE_POOL;
+
+		// Normalize any unicode forms so we match correctly
+		const FString NormalizedFilename = UTF8_TO_TCHAR(([[[NSString stringWithUTF8String:InEntry->d_name] precomposedStringWithCanonicalMapping] cStringUsingEncoding:NSUTF8StringEncoding]));
 
-		return IterateDirectoryCommon(Directory, [&](struct dirent* InEntry) -> bool
+		// Figure out whether it's a directory. Some protocols (like NFS) do not voluntarily return this as part of the directory entry, and need to be queried manually.
+		bool bIsDirectory = (InEntry->d_type == DT_DIR);
+		if (InEntry->d_type == DT_UNKNOWN || InEntry->d_type == DT_LNK)
 		{
-			// Normalize any unicode forms so we match correctly
-			const FString NormalizedFilename = UTF8_TO_TCHAR(([[[NSString stringWithUTF8String:InEntry->d_name] precomposedStringWithCanonicalMapping] cStringUsingEncoding:NSUTF8StringEncoding]));
-				
-			// Figure out whether it's a directory. Some protocols (like NFS) do not voluntarily return this as part of the directory entry, and need to be queried manually.
-			bool bIsDirectory = (InEntry->d_type == DT_DIR);
-			if (InEntry->d_type == DT_UNKNOWN || InEntry->d_type == DT_LNK)
+			struct stat StatInfo;
+			if (stat(TCHAR_TO_UTF8(*(NormalizedDirectoryStr / NormalizedFilename)), &StatInfo) == 0)
 			{
-				struct stat StatInfo;
-				if (stat(TCHAR_TO_UTF8(*(NormalizedDirectoryStr / NormalizedFilename)), &StatInfo) == 0)
-				{
-					bIsDirectory = S_ISDIR(StatInfo.st_mode);
-				}
+				bIsDirectory = S_ISDIR(StatInfo.st_mode);
 			}
-					
-			return Visitor.Visit(*(DirectoryStr / NormalizedFilename), bIsDirectory);
-		});
-	}
+		}
+
+		return Visitor.Visit(*(DirectoryStr / NormalizedFilename), bIsDirectory);
+	});
 }
 
 bool FApplePlatformFile::IterateDirectoryStat(const TCHAR* Directory, FDirectoryStatVisitor& Visitor)
 {
-	@autoreleasepool
+	const FString DirectoryStr = Directory;
+	const FString NormalizedDirectoryStr = NormalizeFilename(Directory);
+
+	return IterateDirectoryCommon(Directory, [&Visitor, &DirectoryStr, &NormalizedDirectoryStr](struct dirent* InEntry) -> bool
 	{
-		const FString DirectoryStr = Directory;
-		const FString NormalizedDirectoryStr = NormalizeFilename(Directory);
+		SCOPED_AUTORELEASE_POOL;
 
-		return IterateDirectoryCommon(Directory, [&](struct dirent* InEntry) -> bool
+		// Normalize any unicode forms so we match correctly
+		const FString NormalizedFilename = UTF8_TO_TCHAR(([[[NSString stringWithUTF8String:InEntry->d_name] precomposedStringWithCanonicalMapping] cStringUsingEncoding:NSUTF8StringEncoding]));
+
+		struct stat StatInfo;
+		if (stat(TCHAR_TO_UTF8(*(NormalizedDirectoryStr / NormalizedFilename)), &StatInfo) == 0)
 		{
-			// Normalize any unicode forms so we match correctly
-			const FString NormalizedFilename = UTF8_TO_TCHAR(([[[NSString stringWithUTF8String:InEntry->d_name] precomposedStringWithCanonicalMapping] cStringUsingEncoding:NSUTF8StringEncoding]));
-				
-			struct stat StatInfo;
-			if (stat(TCHAR_TO_UTF8(*(NormalizedDirectoryStr / NormalizedFilename)), &StatInfo) == 0)
-			{
-				return Visitor.Visit(*(DirectoryStr / NormalizedFilename), MacStatToUEFileData(StatInfo));
-			}
+			return Visitor.Visit(*(DirectoryStr / NormalizedFilename), MacStatToUEFileData(StatInfo));
+		}
 
-			return true;
-		});
-	}
+		return true;
+	});
 }
 
 bool FApplePlatformFile::IterateDirectoryCommon(const TCHAR* Directory, const TFunctionRef<bool(struct dirent*)>& Visitor)
diff --git a/Engine/Source/Runtime/Engine/Private/VT/VirtualTextureChunkDDCCache.cpp b/Engine/Source/Runtime/Engine/Private/VT/VirtualTextureChunkDDCCache.cpp
--- a/Engine/Source/Runtime/Engine/Private/VT/VirtualTextureChunkDDCCache.cpp
+++ b/Engine/Source/Runtime/Engine/Private/VT/VirtualTextureChunkDDCCache.cpp
@@ -211,7 +211,7 @@ void FVirtualTextureChunkDDCCache::Initialize()
 	const FTimespan UnusedFileTime = FTimespan(UnusedFileAge, 0, 0, 0);
 
 	// find all files in the directory
-	IPlatformFile::FDirectoryStatVisitorFunc CacheCleanupFunc = [UnusedFileTime](const TCHAR* FileName, const FFileStatData& Stat)
+	IFileManager::Get().IterateDirectoryStatRecursively(*AbsoluteCachePath, [UnusedFileTime](const TCHAR* FileName, const FFileStatData& Stat)
 	{
 		if (!Stat.bIsDirectory && (Stat.AccessTime != FDateTime::MinValue() || Stat.ModificationTime != FDateTime::MinValue()))
 		{
@@ -226,8 +226,7 @@ void FVirtualTextureChunkDDCCache::Initialize()
 		}
 
 		return true;
-	};
-	IFileManager::Get().IterateDirectoryStatRecursively(*AbsoluteCachePath, CacheCleanupFunc);
+	});
 }
 
 void FVirtualTextureChunkDDCCache::ShutDown()
@@ -292,4 +291,4 @@ bool FVirtualTextureChunkDDCCache::MakeChunkAvailable(struct FVirtualTextureData
 }
 
 
-#endif
\ No newline at end of file
+#endif
