CompressMemoryIfWorthDecompressing to avoid sending compressed data that expanded or barely compressed

diff --git a/Engine/Plugins/Compression/OodleData/Source/Private/OodleDataCompressionFormat.cpp b/Engine/Plugins/Compression/OodleData/Source/Private/OodleDataCompressionFormat.cpp
--- a/Engine/Plugins/Compression/OodleData/Source/Private/OodleDataCompressionFormat.cpp
+++ b/Engine/Plugins/Compression/OodleData/Source/Private/OodleDataCompressionFormat.cpp
@@ -186,6 +186,12 @@ struct FOodleDataCompressionFormat : ICompressionFormat
 		}
 		return TEXT("Unknown");
 	}
+	
+	virtual bool DoesOwnWorthDecompressingCheck() override
+	{
+		// Oodle does own "worth it" check internally, don't add one
+		return true;
+	}
 
 	virtual FName GetCompressionFormatName() override
 	{
diff --git a/Engine/Source/Developer/PakFileUtilities/Private/PakFileUtilities.cpp b/Engine/Source/Developer/PakFileUtilities/Private/PakFileUtilities.cpp
--- a/Engine/Source/Developer/PakFileUtilities/Private/PakFileUtilities.cpp
+++ b/Engine/Source/Developer/PakFileUtilities/Private/PakFileUtilities.cpp
@@ -2060,16 +2060,15 @@ bool CreatePakFile(const TCHAR* Filename, TArray<FPakInputPair>& FilesToAdd, con
 					if ( CompressionMethod == NAME_Zlib )
 					{
 						// for forced-Zlib files still use the old heuristic :
-						// TODO : move this inside the zlib compressor
 
 						// Zlib must save at least 1K regardless of percentage (for small files)
 						bNotEnoughCompression = (OriginalFileSize - CompressedFileBuffer.TotalCompressedSize) < 1024;
 						if ( ! bNotEnoughCompression )
 						{
-					// Check the compression ratio, if it's too low just store uncompressed. Also take into account read size
-					// if we still save 64KB it's probably worthwhile compressing, as that saves a file read operation in the runtime.
-								// TODO: drive this threshold from the command line
-					float PercentLess = ((float)CompressedFileBuffer.TotalCompressedSize / (OriginalFileSize / 100.f));
+							// Check the compression ratio, if it's too low just store uncompressed. Also take into account read size
+							// if we still save 64KB it's probably worthwhile compressing, as that saves a file read operation in the runtime.
+							// TODO: drive this threshold from the command line
+							float PercentLess = ((float)CompressedFileBuffer.TotalCompressedSize / (OriginalFileSize / 100.f));
 							bNotEnoughCompression = (PercentLess > 90.f) && ((OriginalFileSize - CompressedFileBuffer.TotalCompressedSize) < 65536);
 						}
 					}				
diff --git a/Engine/Source/Runtime/Core/Private/Misc/Compression.cpp b/Engine/Source/Runtime/Core/Private/Misc/Compression.cpp
--- a/Engine/Source/Runtime/Core/Private/Misc/Compression.cpp
+++ b/Engine/Source/Runtime/Core/Private/Misc/Compression.cpp
@@ -529,6 +529,86 @@ int32 FCompression::CompressMemoryBound(FName FormatName, int32 UncompressedSize
 }
 
 
+
+bool FCompression::CompressMemoryIfWorthDecompressing(FName FormatName, void* CompressedBuffer, int32& CompressedSize, const void* UncompressedBuffer, int32 UncompressedSize, ECompressionFlags Flags, int32 CompressionData)
+{
+	// returns false if we could compress,
+	//	but it's not worth the time to decompress
+	//	you should store the data uncompressed instead
+
+	const int MinBytesSaved = 1024; // @todo make this configurable?
+	if ( UncompressedSize <= MinBytesSaved )
+	{
+		// if input size is smaller than the number of bytes we need to save
+		// no need to even try encoding
+		// also saves encode time
+		return false;
+	}
+	
+	bool bNeedsWorthItCheck = false;
+	if (FormatName == NAME_Zlib)
+	{
+		// hardcoded zlib
+		bNeedsWorthItCheck = true;
+	}
+	else if (FormatName == NAME_Gzip)
+	{
+		// hardcoded gzip
+		bNeedsWorthItCheck = true;
+	}
+	else if (FormatName == NAME_LZ4)
+	{
+		// hardcoded lz4
+		bNeedsWorthItCheck = true;
+	}
+	else
+	{
+		ICompressionFormat* Format = GetCompressionFormat(FormatName);
+
+		if ( ! Format ) return false;
+
+		bNeedsWorthItCheck = ! Format->DoesOwnWorthDecompressingCheck();
+	}
+
+	bool bCompressSucceeded = CompressMemory(FormatName,CompressedBuffer,CompressedSize,UncompressedBuffer,UncompressedSize,Flags,CompressionData);
+
+	if ( ! bCompressSucceeded )
+	{
+		// compression actually failed
+		return false;
+	}
+
+	if ( ! bNeedsWorthItCheck )
+	{
+		// ICompressionFormat does own "worth it" check, don't do our own
+		// do check for expansion because that's how they signal "not worth it"
+		//	(CompressMemory is not allowed to return false)
+		return ( CompressedSize < UncompressedSize );
+	}
+
+	// we got compression, but do we want it ?
+	
+	// check if the decode time on load is worth the size savings
+	// Oodle uses much more sophisticated models for this
+	// here we replicate the Pak file logic :
+
+	// must save at least MinBytesSaved regardless of percentage (for small files)
+	// also checks CompressedSize >= UncompressedSize
+	int32 BytesSaved = UncompressedSize - CompressedSize;
+	if ( BytesSaved < MinBytesSaved )
+	{
+		return false;
+	}
+
+	// Check the compression ratio, if it's too low just store uncompressed. 
+	// 	saving 64 KB per 1 MB is 6%
+	// @todo make this configurable?
+	float PercentSaved = ((float)BytesSaved * 100.f / (float)UncompressedSize);
+	bool bWorthIt = PercentSaved >= 5.f;
+
+	return bWorthIt;
+}
+
 bool FCompression::CompressMemory(FName FormatName, void* CompressedBuffer, int32& CompressedSize, const void* UncompressedBuffer, int32 UncompressedSize, ECompressionFlags Flags, int32 CompressionData)
 {
 	TRACE_CPUPROFILER_EVENT_SCOPE(FCompression::CompressMemory);
diff --git a/Engine/Source/Runtime/Core/Public/Misc/Compression.h b/Engine/Source/Runtime/Core/Public/Misc/Compression.h
--- a/Engine/Source/Runtime/Core/Public/Misc/Compression.h
+++ b/Engine/Source/Runtime/Core/Public/Misc/Compression.h
@@ -53,6 +53,8 @@ struct FCompression
 	/**
 	 * Thread-safe abstract compression routine. Compresses memory from uncompressed buffer and writes it to compressed
 	 * buffer. Updates CompressedSize with size of compressed data. Compression controlled by the passed in flags.
+	 * CompressMemory is expected to return true and write valid data even if it expanded bytes.
+	 * Always check CompressedSize >= UncompressedSize and fall back to uncompressed, or use CompressMemoryIfWorthDecompressing
 	 *
 	 * @param	Flags						Flags to control what method to use and optionally control memory vs speed
 	 * @param	CompressedBuffer			Buffer compressed data is going to be written to
@@ -63,6 +65,21 @@ struct FCompression
 	 * @return true if compression succeeds, false if it fails because CompressedBuffer was too small or other reasons
 	 */
 	CORE_API static bool CompressMemory(FName FormatName, void* CompressedBuffer, int32& CompressedSize, const void* UncompressedBuffer, int32 UncompressedSize, ECompressionFlags Flags=COMPRESS_NoFlags, int32 CompressionData=0);
+	
+	/**
+	* Same as CompressMemory but evaluates if the compression gain is worth the runtime decode time
+	* returns false if the size saving is not worth it (also if CompressedSize >= UncompressedSize)
+	* if false is returned, send the data uncompressed instead
+	 *
+	 * @param	Flags						Flags to control what method to use and optionally control memory vs speed
+	 * @param	CompressedBuffer			Buffer compressed data is going to be written to
+	 * @param	CompressedSize	[in/out]	Size of CompressedBuffer, at exit will be size of compressed data
+	 * @param	UncompressedBuffer			Buffer containing uncompressed data
+	 * @param	UncompressedSize			Size of uncompressed data in bytes
+	 * @param	BitWindow					Bit window to use in compression
+	 * @return true if compression succeeds, false if it fails because CompressedBuffer was too small or other reasons
+	 */
+	CORE_API static bool CompressMemoryIfWorthDecompressing(FName FormatName, void* CompressedBuffer, int32& CompressedSize, const void* UncompressedBuffer, int32 UncompressedSize, ECompressionFlags Flags=COMPRESS_NoFlags, int32 CompressionData=0);
 
 	/**
 	 * Thread-safe abstract decompression routine. Uncompresses memory from compressed buffer and writes it to uncompressed
diff --git a/Engine/Source/Runtime/Core/Public/Misc/ICompressionFormat.h b/Engine/Source/Runtime/Core/Public/Misc/ICompressionFormat.h
--- a/Engine/Source/Runtime/Core/Public/Misc/ICompressionFormat.h
+++ b/Engine/Source/Runtime/Core/Public/Misc/ICompressionFormat.h
@@ -19,4 +19,5 @@ struct ICompressionFormat : public IModularFeature, public IModuleInterface
 	virtual int32 GetCompressedBufferSize(int32 UncompressedSize, int32 CompressionData) = 0;
 	virtual uint32 GetVersion() = 0;
 	virtual FString GetDDCKeySuffix() = 0;
+	virtual bool DoesOwnWorthDecompressingCheck() = 0;
 };
