use CompressMemoryIfWorthDecompressing in iostore

diff --git a/Engine/Source/Runtime/Core/Private/IO/IoStore.cpp b/Engine/Source/Runtime/Core/Private/IO/IoStore.cpp
--- a/Engine/Source/Runtime/Core/Private/IO/IoStore.cpp
+++ b/Engine/Source/Runtime/Core/Private/IO/IoStore.cpp
@@ -1020,23 +1020,26 @@ private:
 		if (Block->CompressionMethod != NAME_None)
 		{
 			int32 CompressedBlockSize = int32(Block->IoBuffer->DataSize());
+			bool bCompressed;
+
 			{
 				TRACE_CPUPROFILER_EVENT_SCOPE(CompressMemory);
-				const bool bCompressed = FCompression::CompressMemory(
+				bCompressed = FCompression::CompressMemoryIfWorthDecompressing(
 					Block->CompressionMethod,
 					Block->IoBuffer->Data(),
 					CompressedBlockSize,
 					Block->UncompressedData,
 					static_cast<int32>(Block->UncompressedSize));
-				check(bCompressed);
 			}
-			check(CompressedBlockSize > 0);
-			if (CompressedBlockSize >= Block->UncompressedSize)
+			if ( !bCompressed )
 			{
 				Block->CompressionMethod = NAME_None;
 			}
 			else
 			{
+				check( CompressedBlockSize > 0 );
+				check( CompressedBlockSize < Block->UncompressedSize)
+
 				Block->CompressedSize = CompressedBlockSize;
 			}
 		}
@@ -1055,7 +1058,7 @@ private:
 			for (uint64 FillIndex = Block->Size; FillIndex < AlignedCompressedBlockSize; ++FillIndex)
 			{
 				check(FillIndex < Block->IoBuffer->DataSize());
-				CompressedData[FillIndex] = CompressedData[(FillIndex - Block->CompressedSize) % Block->Size];
+				CompressedData[FillIndex] = CompressedData[(FillIndex - Block->Size) % Block->Size];
 			}
 			Block->Size = AlignedCompressedBlockSize;
 		}
