Matrix ispc
-- Rolled back matrix multiply as most recent version is causing issues with bounds calculations.

diff --git a/Engine/Source/Runtime/Core/Public/Math/Matrix.isph b/Engine/Source/Runtime/Core/Public/Math/Matrix.isph
--- a/Engine/Source/Runtime/Core/Public/Math/Matrix.isph
+++ b/Engine/Source/Runtime/Core/Public/Math/Matrix.isph
@@ -100,22 +100,26 @@ inline FMatrix operator*(const FMatrix &A, const FMatrix &B)
 	return Result;
 }
 
-inline FMatrix operator*(const FMatrix &Matrix1, const uniform FMatrix &Matrix2)
+inline FMatrix operator*(const FMatrix &A, const uniform FMatrix &B)
 {
-	FVector4 R[4];
-
-	const FVector4 *uniform A	= (const FVector4 *uniform) &Matrix1;
-	const uniform FVector4 *uniform B	= (const uniform FVector4 *uniform) &Matrix2;
+	FMatrix Result;
 
-	for (uniform unsigned int i = 0; i < 4; i++)
+	for (uniform unsigned int m = 0; m < 4; m++) 
 	{
-		R[i] = VectorReplicate( A[i], 0 ) * B[0] +
-				VectorReplicate( A[i], 1 ) * B[1] +
-				VectorReplicate( A[i], 2 ) * B[2] +
-				VectorReplicate( A[i], 3 ) * B[3];
+		varying float Sum;
+		for (uniform unsigned int k = 0; k < 4; k++) 
+		{
+			Sum = 0.0f;
+			for (uniform unsigned int n = 0; n < 4; n++) 
+			{
+				Sum += A.M[m * 4 + n] * B.M[n * 4 + k];
+			}
+			
+			Result.M[m * 4 + k] = Sum;
+		}
 	}
 
-	return SetMatrix(R[0], R[1], R[2], R[3]);
+	return Result;
 }
 
 // Remove any scaling from this matrix (ie magnitude of each row is 1)
