U5 - disable this check when a crash happens because it is likely to fire if the crash happens within a scope.

diff --git a/Engine/Source/Runtime/Core/Private/HAL/ThreadingBase.cpp b/Engine/Source/Runtime/Core/Private/HAL/ThreadingBase.cpp
--- a/Engine/Source/Runtime/Core/Private/HAL/ThreadingBase.cpp
+++ b/Engine/Source/Runtime/Core/Private/HAL/ThreadingBase.cpp
@@ -93,7 +93,7 @@ FTaskTagScope::FTaskTagScope(bool InTagOnlyIfNone, ETaskTag InTag) : Tag(InTag),
 			//Try to catch other Threads that already opened a non parallel scope
 			ActiveNamedThreads.store(int32(ETaskTag::ENone));
 		}
-		checkf(IsOK, TEXT("Only Scopes tagged with ETaskTag::EParallelThread can be tagged multiple times. ActiveNamedThreads(%x) cannot be tagged multiple times in the same callstack you can use FOptionalTaskTagScope to avoid retagging check the ActiveNamedThreads(%x) with the current Tag(%x)"), OldTag, FTaskTagScope::GetCurrentTag(), Tag);
+		checkf(IsOK || !IsEngineExitRequested(), TEXT("Only Scopes tagged with ETaskTag::EParallelThread can be tagged multiple times. ActiveNamedThreads(%x) cannot be tagged multiple times in the same callstack you can use FOptionalTaskTagScope to avoid retagging check the ActiveNamedThreads(%x) with the current Tag(%x)"), OldTag, FTaskTagScope::GetCurrentTag(), Tag);
 	}
 	
 	ParentTag = ActiveTaskTag;
