Clean up FUE5CookerObjectVersion/UE_USE_VIRTUALBULKDATA in preperation for enabling virtualized bulkdata.

- Moved UE_USE_VIRTUALBULKDATA to VirtualizationManager.h as it will no longer control versioning in any fashion.
- Remove virtualization versions from FUE5CookerObjectVersion (which were never enabled so this is safe to do)
- This is a partial extraction from the CL that will enable virtualized bulkdata but it will be easier to do some of the cleanup in DevCooker first and cherrypick to main then do a full merge up back to DevCooker to avoid a number of merging issues.
- Note that with this submit enabling UE_USE_VIRTUALBULKDATA will not work until the main submit to main, however nobody should be enabling that so it shouldn't matter.

#rnx

diff --git a/Engine/Source/Runtime/Core/Public/UObject/UE5CookerObjectVersion.h b/Engine/Source/Runtime/Core/Public/UObject/UE5CookerObjectVersion.h
--- a/Engine/Source/Runtime/Core/Public/UObject/UE5CookerObjectVersion.h
+++ b/Engine/Source/Runtime/Core/Public/UObject/UE5CookerObjectVersion.h
@@ -4,10 +4,6 @@
 #include "CoreTypes.h"
 #include "Misc/Guid.h"
 
-// TODO: Remove this once we are sure we can run with FVirtualizedUntypedBulkData with no issues
-/** When set to 1 we use FVirtualizedUntypedBulkData, when 0 we use what ever code was there before. */
-#define UE_USE_VIRTUALBULKDATA 0
-
 // Custom serialization version for changes made in the UE5 Dev-Cooker stream
 struct CORE_API FUE5CookerObjectVersion
 {
@@ -18,16 +14,6 @@ struct CORE_API FUE5CookerObjectVersion
 		
 		// -----<new versions can be added above this line>-------------------------------------------------
 
-		// Virtualization version numbers are added at this point. 
-		// When UE_USE_VIRTUALBULKDATA is removed then they can be moved to the normal section
-
-#if UE_USE_VIRTUALBULKDATA
-		// Switch FMeshDescriptionBulkData to use virtualized bulkdata
-		MeshDescriptionVirtualization,
-		// Switch FTextureSource to use virtualized bulkdata
-		TextureSourceVirtualization,
-#endif //UE_USE_VIRTUALBULKDATA
-
 		VersionPlusOne,
 		LatestVersion = VersionPlusOne - 1
 	};
diff --git a/Engine/Source/Runtime/CoreUObject/Public/Virtualization/VirtualizationManager.h b/Engine/Source/Runtime/CoreUObject/Public/Virtualization/VirtualizationManager.h
--- a/Engine/Source/Runtime/CoreUObject/Public/Virtualization/VirtualizationManager.h
+++ b/Engine/Source/Runtime/CoreUObject/Public/Virtualization/VirtualizationManager.h
@@ -7,8 +7,7 @@
 #include "Logging/LogMacros.h"
 #include "Virtualization/PayloadId.h"
 
-// For UE_USE_VIRTUALBULKDATA
-#include "UObject/UE5CookerObjectVersion.h"
+#define UE_USE_VIRTUALBULKDATA 0
 
 // TODO: Do we want to keep this header public for UE5 release? If the only interaction should be
 //		 via FVirtualizedUntypedBulkData or other such classes we might not want to expose this 
