PR #7962: Fixed possible build error on namespace collision (Contributed by [REDACTED_PERSON])

diff --git a/Engine/Source/Runtime/Core/Public/Windows/WindowsPlatformAtomics.h b/Engine/Source/Runtime/Core/Public/Windows/WindowsPlatformAtomics.h
--- a/Engine/Source/Runtime/Core/Public/Windows/WindowsPlatformAtomics.h
+++ b/Engine/Source/Runtime/Core/Public/Windows/WindowsPlatformAtomics.h
@@ -20,17 +20,17 @@ struct CORE_API FWindowsPlatformAtomics
 
 	static FORCEINLINE int8 InterlockedIncrement( volatile int8* Value )
 	{
-		return (int8)_InterlockedExchangeAdd8((char*)Value, 1) + 1;
+		return (int8)::_InterlockedExchangeAdd8((char*)Value, 1) + 1;
 	}
 
 	static FORCEINLINE int16 InterlockedIncrement( volatile int16* Value )
 	{
-		return (int16)_InterlockedIncrement16((short*)Value);
+		return (int16)::_InterlockedIncrement16((short*)Value);
 	}
 
 	static FORCEINLINE int32 InterlockedIncrement( volatile int32* Value )
 	{
-		return (int32)_InterlockedIncrement((long*)Value);
+		return (int32)::_InterlockedIncrement((long*)Value);
 	}
 
 	static FORCEINLINE int64 InterlockedIncrement( volatile int64* Value )
