Implement fast path when comparing FNames against NAME_None.

diff --git a/Engine/Source/Runtime/Core/Private/UObject/UnrealNames.cpp b/Engine/Source/Runtime/Core/Private/UObject/UnrealNames.cpp
--- a/Engine/Source/Runtime/Core/Private/UObject/UnrealNames.cpp
+++ b/Engine/Source/Runtime/Core/Private/UObject/UnrealNames.cpp
@@ -1303,7 +1303,7 @@ static FNamePool& GetNamePoolPostInit()
 
 bool operator==(FNameEntryId Id, EName Ename)
 {
-	return Id == GetNamePoolPostInit().Find(Ename);
+	return Ename == NAME_None ? !Id : Id == GetNamePoolPostInit().Find(Ename);
 }
 
 static int32 CompareDifferentIdsAlphabetically(FNameEntryId AId, FNameEntryId BId)
