Add move construction to MemStack and delete the copy constructor.

diff --git a/Engine/Source/Runtime/Core/Public/Misc/MemStack.h b/Engine/Source/Runtime/Core/Public/Misc/MemStack.h
--- a/Engine/Source/Runtime/Core/Public/Misc/MemStack.h
+++ b/Engine/Source/Runtime/Core/Public/Misc/MemStack.h
@@ -81,6 +81,28 @@ public:
 		
 	}
 
+	FMemStackBase(const FMemStackBase&) = delete;
+	FMemStackBase(FMemStackBase&& Other)
+	{
+		*this = MoveTemp(Other);
+	}
+
+	FMemStackBase& operator=(FMemStackBase&& Other)
+	{
+		Top = Other.Top;
+		End = Other.End;
+		TopChunk = Other.TopChunk;
+		TopMark = Other.TopMark;
+		NumMarks = Other.NumMarks;
+		bShouldEnforceAllocMarks = Other.bShouldEnforceAllocMarks;
+		Other.Top = nullptr;
+		Other.End = nullptr;
+		Other.TopChunk = nullptr;
+		Other.NumMarks = 0;
+		Other.bShouldEnforceAllocMarks = false;
+		return *this;
+	}
+
 	~FMemStackBase()
 	{
 		check((GIsCriticalError || !NumMarks));
