Don't call destructors when destroying frozen objects
- Frozen memory images may have de-duplicated sections, where the same object is referenced multiple times; this means default destructors might be called multiple times, which is not supported in general
- Add a hook to allow custom behavior to destroy frozen objects that may hold references to external resources
- Within shader system, no need to run destructors, since all memory should be owned by the top-level memory image

diff --git a/Engine/Source/Runtime/Core/Private/Serialization/MemoryImage.cpp b/Engine/Source/Runtime/Core/Private/Serialization/MemoryImage.cpp
--- a/Engine/Source/Runtime/Core/Private/Serialization/MemoryImage.cpp
+++ b/Engine/Source/Runtime/Core/Private/Serialization/MemoryImage.cpp
@@ -280,7 +280,7 @@ void InternalDeleteObjectFromLayout(void* Object, const FTypeLayoutDesc& TypeDes
 	// DestroyFunc may be nullptr for types with trivial destructors
 	if (TypeDesc.DestroyFunc)
 	{
-		TypeDesc.DestroyFunc(Object, TypeDesc, PtrTable);
+		TypeDesc.DestroyFunc(Object, TypeDesc, PtrTable, bIsFrozen);
 	}
 	if (!bIsFrozen)
 	{
diff --git a/Engine/Source/Runtime/Core/Public/Serialization/MemoryLayout.h b/Engine/Source/Runtime/Core/Public/Serialization/MemoryLayout.h
--- a/Engine/Source/Runtime/Core/Public/Serialization/MemoryLayout.h
+++ b/Engine/Source/Runtime/Core/Public/Serialization/MemoryLayout.h
@@ -108,7 +108,7 @@ struct FFieldLayoutDesc
 
 struct FTypeLayoutDesc
 {
-	typedef void (FDestroyFunc)(void* Object, const FTypeLayoutDesc& TypeDesc, const FPointerTableBase* PtrTable);
+	typedef void (FDestroyFunc)(void* Object, const FTypeLayoutDesc& TypeDesc, const FPointerTableBase* PtrTable, bool bIsFrozen);
 	typedef void (FWriteFrozenMemoryImageFunc)(FMemoryImageWriter& Writer, const void* Object, const FTypeLayoutDesc& TypeDesc, const FTypeLayoutDesc& DerivedTypeDesc);
 	typedef uint32 (FUnfrozenCopyFunc)(const FMemoryUnfreezeContent& Context, const void* Object, const FTypeLayoutDesc& TypeDesc, void* OutDst);
 	typedef uint32 (FAppendHashFunc)(const FTypeLayoutDesc& TypeDesc, const FPlatformTypeLayoutParameters& LayoutParams, FSHA1& Hasher);
@@ -245,10 +245,20 @@ namespace Freeze
 	CORE_API void IntrinsicWriteMemoryImage(FMemoryImageWriter& Writer, void*, const FTypeLayoutDesc&);
 
 	// Override for types that need access to a PointerTable in order to destroy frozen data
-	FORCEINLINE void CleanupObject(void* Object, const FPointerTableBase* PtrTable) {}
-
+	// WARNING! Frozen data may have deduplicated portions, so it's possible that a given object may have multiple references
+	// Any attempt to cleanup/destroy frozen data must take care to safely allow DestroyObject to be called multiple times on the same frozen object
+	// Generally speaking, it's *NOT* safe to invoke the built-in destructor multiple times
 	template<typename T>
-	FORCEINLINE void CallDestructor(T* Object) { Object->~T(); }
+	FORCEINLINE void DestroyObject(T* Object, const FPointerTableBase* PtrTable, bool bIsFrozen)
+	{
+		// Only call destructor for non-frozen objects
+		if (!bIsFrozen)
+		{
+			Object->~T();
+		}
+		// Wipe destroyed memory to a common pattern
+		FMemory::Memset(Object, 0xfe, sizeof(T));
+	}
 
 	template<typename T>
 	FORCEINLINE void IntrinsicWriteMemoryImage(FMemoryImageWriter& Writer, const T& Object, const FTypeLayoutDesc& TypeDesc)
@@ -531,9 +541,8 @@ PRAGMA_ENABLE_DEPRECATION_WARNINGS
 	UE_DECLARE_INTERNAL_LINK_BASE(InternalLinkType) { UE_STATIC_ONLY(InternalLinkType); static FORCEINLINE void Initialize(FTypeLayoutDesc& TypeDesc) {} }
 
 #define INTERNAL_DECLARE_INLINE_TYPE_LAYOUT(T, InInterface) \
-	private: static void InternalDestroy(void* Object, const FTypeLayoutDesc&, const FPointerTableBase* PtrTable) { \
-		Freeze::CleanupObject(static_cast<PREPROCESSOR_REMOVE_OPTIONAL_PARENS(T)*>(Object), PtrTable); \
-		Freeze::CallDestructor(static_cast<PREPROCESSOR_REMOVE_OPTIONAL_PARENS(T)*>(Object)); \
+	private: static void InternalDestroy(void* Object, const FTypeLayoutDesc&, const FPointerTableBase* PtrTable, bool bIsFrozen) { \
+		Freeze::DestroyObject(static_cast<PREPROCESSOR_REMOVE_OPTIONAL_PARENS(T)*>(Object), PtrTable, bIsFrozen); \
 	} \
 	public: static FTypeLayoutDesc& StaticGetTypeLayout() { \
 		static_assert(TValidateInterfaceHelper<PREPROCESSOR_REMOVE_OPTIONAL_PARENS(T), ETypeLayoutInterface::InInterface>::Value, #InInterface " is invalid interface for " #T); \
@@ -562,7 +571,7 @@ PRAGMA_ENABLE_DEPRECATION_WARNINGS
 	INTERNAL_DECLARE_TYPE_LAYOUT_COMMON(T, InInterface)
 
 #define INTERNAL_DECLARE_TYPE_LAYOUT(T, InInterface, RequiredAPI) \
-	private: static void InternalDestroy(void* Object, const FTypeLayoutDesc&, const FPointerTableBase* PtrTable); \
+	private: static void InternalDestroy(void* Object, const FTypeLayoutDesc&, const FPointerTableBase* PtrTable, bool bIsFrozen); \
 	public: RequiredAPI static FTypeLayoutDesc& StaticGetTypeLayout(); \
 	public: INTERNAL_LAYOUT_INTERFACE_PREFIX(InInterface)(RequiredAPI) const FTypeLayoutDesc& GetTypeLayout() const INTERNAL_LAYOUT_INTERFACE_SUFFIX(InInterface) \
 	INTERNAL_DECLARE_TYPE_LAYOUT_COMMON(T, InInterface)
@@ -583,9 +592,8 @@ PRAGMA_ENABLE_DEPRECATION_WARNINGS
 #define DECLARE_EXPORTED_TYPE_LAYOUT_EXPLICIT_BASES(T, RequiredAPI, Interface, ...) INTERNAL_DECLARE_LAYOUT_EXPLICIT_BASES(T, __VA_ARGS__); INTERNAL_DECLARE_TYPE_LAYOUT(T, Interface, RequiredAPI)
 
 #define INTERNAL_IMPLEMENT_TYPE_LAYOUT_COMMON(TemplatePrefix, T) \
-	PREPROCESSOR_REMOVE_OPTIONAL_PARENS(TemplatePrefix) void PREPROCESSOR_REMOVE_OPTIONAL_PARENS(T)::InternalDestroy(void* Object, const FTypeLayoutDesc&, const FPointerTableBase* PtrTable) { \
-		Freeze::CleanupObject(static_cast<PREPROCESSOR_REMOVE_OPTIONAL_PARENS(T)*>(Object), PtrTable); \
-		Freeze::CallDestructor(static_cast<PREPROCESSOR_REMOVE_OPTIONAL_PARENS(T)*>(Object)); \
+	PREPROCESSOR_REMOVE_OPTIONAL_PARENS(TemplatePrefix) void PREPROCESSOR_REMOVE_OPTIONAL_PARENS(T)::InternalDestroy(void* Object, const FTypeLayoutDesc&, const FPointerTableBase* PtrTable, bool bIsFrozen) { \
+		Freeze::DestroyObject(static_cast<PREPROCESSOR_REMOVE_OPTIONAL_PARENS(T)*>(Object), PtrTable, bIsFrozen); \
 	} \
 	PREPROCESSOR_REMOVE_OPTIONAL_PARENS(TemplatePrefix) FTypeLayoutDesc& PREPROCESSOR_REMOVE_OPTIONAL_PARENS(T)::StaticGetTypeLayout() { \
 		static_assert(TValidateInterfaceHelper<PREPROCESSOR_REMOVE_OPTIONAL_PARENS(T), InterfaceType>::Value, "Invalid interface for " #T); \
@@ -698,9 +706,8 @@ inline void DeleteObjectFromLayout(T* Object, const FPointerTableBase* PtrTable
 		static void CallToString(const void* Object, const FTypeLayoutDesc& TypeDesc, const FPlatformTypeLayoutParameters& LayoutParams, FMemoryToStringContext& OutContext) { \
 			return Freeze::IntrinsicToString(*static_cast<PREPROCESSOR_REMOVE_OPTIONAL_PARENS(T) const*>(Object), TypeDesc, LayoutParams, OutContext); \
 		} \
-		static void CallDestroy(void* Object, const FTypeLayoutDesc&, const FPointerTableBase* PtrTable) { \
-			Freeze::CleanupObject(static_cast<PREPROCESSOR_REMOVE_OPTIONAL_PARENS(T)*>(Object), PtrTable); \
-			Freeze::CallDestructor(static_cast<PREPROCESSOR_REMOVE_OPTIONAL_PARENS(T)*>(Object)); \
+		static void CallDestroy(void* Object, const FTypeLayoutDesc&, const FPointerTableBase* PtrTable, bool bIsFrozen) { \
+			Freeze::DestroyObject(static_cast<PREPROCESSOR_REMOVE_OPTIONAL_PARENS(T)*>(Object), PtrTable, bIsFrozen); \
 		} \
 		static const FTypeLayoutDesc& Do() { \
 			alignas(FTypeLayoutDesc) static uint8 TypeBuffer[sizeof(FTypeLayoutDesc)] = { 0 }; \
@@ -734,7 +741,7 @@ inline void DeleteObjectFromLayout(T* Object, const FPointerTableBase* PtrTable
 		static uint32 CallAppendHash(const FTypeLayoutDesc& TypeDesc, const FPlatformTypeLayoutParameters& LayoutParams, FSHA1& Hasher); \
 		static uint32 CallGetTargetAlignment(const FTypeLayoutDesc& TypeDesc, const FPlatformTypeLayoutParameters& LayoutParams); \
 		static void CallToString(const void* Object, const FTypeLayoutDesc& TypeDesc, const FPlatformTypeLayoutParameters& LayoutParams, FMemoryToStringContext& OutContext); \
-		static void CallDestroy(void* Object, const FTypeLayoutDesc&, const FPointerTableBase* PtrTable); \
+		static void CallDestroy(void* Object, const FTypeLayoutDesc&, const FPointerTableBase* PtrTable, bool bIsFrozen); \
 		RequiredAPI static const FTypeLayoutDesc& Do(); }; \
 	PREPROCESSOR_REMOVE_OPTIONAL_PARENS(TemplatePrefix) struct TGetTypeLayoutHelper<PREPROCESSOR_REMOVE_OPTIONAL_PARENS(T)> { \
 		UE_STATIC_ONLY(TGetTypeLayoutHelper); \
@@ -756,9 +763,8 @@ inline void DeleteObjectFromLayout(T* Object, const FPointerTableBase* PtrTable
 		void TStaticGetTypeLayoutHelper<PREPROCESSOR_REMOVE_OPTIONAL_PARENS(T)>::CallToString(const void* Object, const FTypeLayoutDesc& TypeDesc, const FPlatformTypeLayoutParameters& LayoutParams, FMemoryToStringContext& OutContext) { \
 			return Freeze::IntrinsicToString(*static_cast<PREPROCESSOR_REMOVE_OPTIONAL_PARENS(T) const*>(Object), TypeDesc, LayoutParams, OutContext); \
 		} \
-		void TStaticGetTypeLayoutHelper<PREPROCESSOR_REMOVE_OPTIONAL_PARENS(T)>::CallDestroy(void* Object, const FTypeLayoutDesc&, const FPointerTableBase* PtrTable) { \
-			Freeze::CleanupObject(static_cast<PREPROCESSOR_REMOVE_OPTIONAL_PARENS(T)*>(Object), PtrTable); \
-			Freeze::CallDestructor(static_cast<PREPROCESSOR_REMOVE_OPTIONAL_PARENS(T)*>(Object)); \
+		void TStaticGetTypeLayoutHelper<PREPROCESSOR_REMOVE_OPTIONAL_PARENS(T)>::CallDestroy(void* Object, const FTypeLayoutDesc&, const FPointerTableBase* PtrTable, bool bIsFrozen) { \
+			Freeze::DestroyObject(static_cast<PREPROCESSOR_REMOVE_OPTIONAL_PARENS(T)*>(Object), PtrTable, bIsFrozen); \
 		} \
 		const FTypeLayoutDesc& TStaticGetTypeLayoutHelper<PREPROCESSOR_REMOVE_OPTIONAL_PARENS(T)>::Do() { \
 			alignas(FTypeLayoutDesc) static uint8 TypeBuffer[sizeof(FTypeLayoutDesc)] = { 0 }; \
diff --git a/Engine/Source/Runtime/RenderCore/Private/GlobalShader.cpp b/Engine/Source/Runtime/RenderCore/Private/GlobalShader.cpp
--- a/Engine/Source/Runtime/RenderCore/Private/GlobalShader.cpp
+++ b/Engine/Source/Runtime/RenderCore/Private/GlobalShader.cpp
@@ -392,7 +392,7 @@ void FGlobalShaderMap::Empty()
 {
 	for (const auto& It : SectionMap)
 	{
-		It.Value->GetMutableContent()->Empty(&It.Value->GetPointerTable());
+		It.Value->GetMutableContent()->Empty();
 	}
 }
 
diff --git a/Engine/Source/Runtime/RenderCore/Private/ShaderMap.cpp b/Engine/Source/Runtime/RenderCore/Private/ShaderMap.cpp
--- a/Engine/Source/Runtime/RenderCore/Private/ShaderMap.cpp
+++ b/Engine/Source/Runtime/RenderCore/Private/ShaderMap.cpp
@@ -684,24 +684,13 @@ void FShaderMapContent::UpdateHash(FSHA1& Hasher) const
 	}
 }
 
-void FShaderMapContent::Empty(const FPointerTableBase* PointerTable)
+void FShaderMapContent::Empty()
 {
-	EmptyShaderPipelines(PointerTable);
+	EmptyShaderPipelines();
 	for (int32 i = 0; i < Shaders.Num(); ++i)
 	{
 		TMemoryImagePtr<FShader>& Shader = Shaders[i];
-		// It's possible that frozen shader map may have certain shaders embedded that are compiled out of the target build
-		// In this case, we won't be able to find the shader type, and SafeDelete() will crash, as DeleteObjectFromLayout() relies on getting FTypeLayoutDesc from the shader type
-		// In the future, we should ensure that we're not including these shaders at all, but for now it should be OK to skip them
-		if (Shader->GetType(PointerTable))
-		{
-			Shader.SafeDelete(PointerTable);
-		}
-		else
-		{
-			// If we can't find the type, and the shadermap isn't frozen, then something has gone wrong
-			checkf(Shader.IsFrozen(), TEXT("Shader type %016X is missing, but shader isn't frozen"), ShaderTypes[i].GetHash());
-		}
+		Shader.SafeDelete();
 	}
 	Shaders.Empty();
 	ShaderTypes.Empty();
@@ -709,11 +698,11 @@ void FShaderMapContent::Empty(const FPointerTableBase* PointerTable)
 	ShaderHash.Clear();
 }
 
-void FShaderMapContent::EmptyShaderPipelines(const FPointerTableBase* PointerTable)
+void FShaderMapContent::EmptyShaderPipelines()
 {
 	for (TMemoryImagePtr<FShaderPipeline>& Pipeline : ShaderPipelines)
 	{
-		Pipeline.SafeDelete(PointerTable);
+		Pipeline.SafeDelete();
 	}
 	ShaderPipelines.Empty();
 }
\ No newline at end of file
diff --git a/Engine/Source/Runtime/RenderCore/Public/Shader.h b/Engine/Source/Runtime/RenderCore/Public/Shader.h
--- a/Engine/Source/Runtime/RenderCore/Public/Shader.h
+++ b/Engine/Source/Runtime/RenderCore/Public/Shader.h
@@ -1839,7 +1839,7 @@ public:
 	/** Destructor ensures pipelines cleared up. */
 	~FShaderMapContent()
 	{
-		Empty(nullptr);
+		Empty();
 	}
 
 	EShaderPlatform GetShaderPlatform() const { return Platform; }
@@ -1947,7 +1947,7 @@ public:
 	}
 
 	/** clears out all shaders and deletes shader pipelines held in the map */
-	void Empty(const FPointerTableBase* PointerTable);
+	void Empty();
 
 	inline FShaderPipeline* GetShaderPipeline(const FHashedName& PipelineTypeName) const
 	{
@@ -1970,7 +1970,7 @@ public:
 	void UpdateHash(FSHA1& Hasher) const;
 
 protected:
-	void EmptyShaderPipelines(const FPointerTableBase* PointerTable);
+	void EmptyShaderPipelines();
 
 	using FMemoryImageHashTable = THashTable<FMemoryImageAllocator>;
 
@@ -1983,14 +1983,6 @@ protected:
 	LAYOUT_FIELD(TEnumAsByte<EShaderPlatform>, Platform);
 };
 
-namespace Freeze
-{
-	inline void CleanupObject(FShaderMapContent* Object, const FPointerTableBase* PtrTable)
-	{
-		Object->Empty(PtrTable);
-	}
-}
-
 class RENDERCORE_API FShaderMapBase
 {
 public:
