Color.cpp fix checks in HDR convert
RGBE HDR can not store negative floats
Unreal passes negative floats here
 just clamp them to zero

diff --git a/Engine/Source/Runtime/Core/Private/Math/Color.cpp b/Engine/Source/Runtime/Core/Private/Math/Color.cpp
--- a/Engine/Source/Runtime/Core/Private/Math/Color.cpp
+++ b/Engine/Source/Runtime/Core/Private/Math/Color.cpp
@@ -101,6 +101,13 @@ FColor FLinearColor::ToRGBE() const
 	}
 	else
 	{
+		// RGBE HDR can not store negative floats
+		// Unreal passes negative floats here
+		// just clamp them to zero :
+		const float NonNegativeR = FMath::Max(R,0.f);
+		const float NonNegativeG = FMath::Max(G,0.f);
+		const float NonNegativeB = FMath::Max(B,0.f);
+
 		// The following replaces a call to frexpf, because frexpf would have a warning for an unused return value.
 		// Additionally, this usage of logbf assumes FLT_RADIX == 2
 		int32 Exponent = 1 + (int32)logbf(Primary);
@@ -108,10 +115,10 @@ FColor FLinearColor::ToRGBE() const
 		
 		FColor	Color;
 		// no clamp needed, should always fit in uint8 :
-		Color.R = TCheckValueCast<uint8>( (int)(R * Scale));
-		Color.G = TCheckValueCast<uint8>( (int)(G * Scale));
-		Color.B = TCheckValueCast<uint8>( (int)(B * Scale));
-		Color.A = TCheckValueCast<uint8>(Exponent + 128);
+		Color.R = TCheckValueCast<uint8>( (int)(NonNegativeR * Scale) );
+		Color.G = TCheckValueCast<uint8>( (int)(NonNegativeG * Scale) );
+		Color.B = TCheckValueCast<uint8>( (int)(NonNegativeB * Scale) );
+		Color.A = TCheckValueCast<uint8>( Exponent + 128 );
 		return Color;
 	}
 }
