Added new MakeStridedViewOfBase to able to convert a ArrayView of a derived structure to a parent strided view.

#rnx

diff --git a/Engine/Source/Runtime/Core/Private/Tests/Containers/StridedViewTest.cpp b/Engine/Source/Runtime/Core/Private/Tests/Containers/StridedViewTest.cpp
--- a/Engine/Source/Runtime/Core/Private/Tests/Containers/StridedViewTest.cpp
+++ b/Engine/Source/Runtime/Core/Private/Tests/Containers/StridedViewTest.cpp
@@ -265,15 +265,20 @@ bool FStridedViewTest::RunTest(const FString& Parameters)
 
 		// Explicitly construct view of base from derived
 		TStridedView<FMyStruct> ViewDerivedAsBase(sizeof(StructsDerived[0]), &static_cast<FMyStruct&>(StructsDerived[0]), Num);
+		TStridedView<FMyStruct> ViewDerivedAsBase2 = MakeStridedViewOfBase<FMyStruct>(MakeArrayView(StructsDerived));
 
 		TestEqual(TEXT("ViewDerived.GetStride()"), ViewDerived.GetStride(), (int32)sizeof(FMyStructDerived));
 		TestEqual(TEXT("ViewDerivedAsBase.GetStride()"), ViewDerivedAsBase.GetStride(), (int32)sizeof(FMyStructDerived));
+		TestEqual(TEXT("ViewDerivedAsBase2.GetStride()"), ViewDerivedAsBase2.GetStride(), (int32)sizeof(FMyStructDerived));
 
 		// Incorrect API usage : implicit conversion of derived view to base. Generates compile error (as expected).
 		// FVector Avg = ComputeMeanPositionStrided(ViewDerived);
 
 		FVector AvgDerivedAsBase = ComputeMeanPositionStrided(ViewDerivedAsBase);
 		TestEqual(TEXT("ComputeMeanPositionStrided(ViewDerivedAsBase)"), AvgDerivedAsBase, ExpectedAvg);
+
+		FVector AvgDerivedAsBase2 = ComputeMeanPositionStrided(ViewDerivedAsBase2);
+		TestEqual(TEXT("ComputeMeanPositionStrided(ViewDerivedAsBase2)"), AvgDerivedAsBase2, ExpectedAvg);
 	}
 
 	{
diff --git a/Engine/Source/Runtime/Core/Public/Containers/StridedView.h b/Engine/Source/Runtime/Core/Public/Containers/StridedView.h
--- a/Engine/Source/Runtime/Core/Public/Containers/StridedView.h
+++ b/Engine/Source/Runtime/Core/Public/Containers/StridedView.h
@@ -157,6 +157,13 @@ TStridedView<ElementType> MakeStridedView(int32 BytesBetweenElements, ElementTyp
 	return TStridedView<ElementType>(BytesBetweenElements, FirstElement, Count);
 }
 
+template <typename BaseStructureType, typename DerivedStructureType>
+TStridedView<BaseStructureType> MakeStridedViewOfBase(TArrayView<DerivedStructureType> StructuredView)
+{
+	static_assert(TIsDerivedFrom<DerivedStructureType, BaseStructureType>::IsDerived, "Expecting derived structure type");
+	return MakeStridedView<BaseStructureType>((int32)sizeof(DerivedStructureType), GetData(StructuredView), (int32)GetNum(StructuredView));
+}
+
 template <typename StructureType>
 TStridedView<StructureType> MakeStridedView(TArrayView<StructureType> StructuredView)
 {
