Set the size of the trace tail from the command line with -tracetailmb=...

#rnx

diff --git a/Engine/Source/Runtime/Core/Private/ProfilingDebugging/TraceAuxiliary.cpp b/Engine/Source/Runtime/Core/Private/ProfilingDebugging/TraceAuxiliary.cpp
--- a/Engine/Source/Runtime/Core/Private/ProfilingDebugging/TraceAuxiliary.cpp
+++ b/Engine/Source/Runtime/Core/Private/ProfilingDebugging/TraceAuxiliary.cpp
@@ -447,6 +447,10 @@ void FTraceAuxiliary::Initialize(const TCHAR* CommandLine)
 	// Initialize Trace
 	UE::Trace::FInitializeDesc Desc;
 	Desc.bUseWorkerThread = FPlatformProcess::SupportsMultithreading();
+	if (FParse::Value(CommandLine, TEXT("-tracetailmb="), Desc.TailSizeBytes))
+	{
+		Desc.TailSizeBytes <<= 20;
+	}
 	UE::Trace::Initialize(Desc);
 
 	FCoreDelegates::OnEndFrame.AddStatic(UE::Trace::Update);
