Some extra support for UTF8CHAR and related compile time checks for char encodings.

diff --git a/Engine/Source/Runtime/Core/Private/GenericPlatform/GenericPlatformString.cpp b/Engine/Source/Runtime/Core/Private/GenericPlatform/GenericPlatformString.cpp
--- a/Engine/Source/Runtime/Core/Private/GenericPlatform/GenericPlatformString.cpp
+++ b/Engine/Source/Runtime/Core/Private/GenericPlatform/GenericPlatformString.cpp
@@ -12,6 +12,7 @@ DEFINE_LOG_CATEGORY_STATIC(LogGenericPlatformString, Log, All);
 template <> const TCHAR* FGenericPlatformString::GetEncodingTypeName<ANSICHAR>() { return TEXT("ANSICHAR"); }
 template <> const TCHAR* FGenericPlatformString::GetEncodingTypeName<WIDECHAR>() { return TEXT("WIDECHAR"); }
 template <> const TCHAR* FGenericPlatformString::GetEncodingTypeName<UCS2CHAR>() { return TEXT("UCS2CHAR"); }
+template <> const TCHAR* FGenericPlatformString::GetEncodingTypeName<UTF8CHAR>() { return TEXT("UTF8CHAR"); }
 #if PLATFORM_TCHAR_IS_CHAR16
 template <> const TCHAR* FGenericPlatformString::GetEncodingTypeName<wchar_t>() { return TEXT("WCHAR_T"); }
 #endif
diff --git a/Engine/Source/Runtime/Core/Public/GenericPlatform/GenericPlatformString.h b/Engine/Source/Runtime/Core/Public/GenericPlatform/GenericPlatformString.h
--- a/Engine/Source/Runtime/Core/Public/GenericPlatform/GenericPlatformString.h
+++ b/Engine/Source/Runtime/Core/Public/GenericPlatform/GenericPlatformString.h
@@ -74,11 +74,12 @@ struct FGenericPlatformString : public FGenericPlatformStricmp
 		enum { Value = false };
 	};
 
-	template <bool Dummy> struct TIsFixedWidthEncoding_Helper<Dummy, ANSICHAR> { enum { Value = true }; };
-	template <bool Dummy> struct TIsFixedWidthEncoding_Helper<Dummy, WIDECHAR> { enum { Value = true }; };
-	template <bool Dummy> struct TIsFixedWidthEncoding_Helper<Dummy, UCS2CHAR> { enum { Value = true }; };
+	template <bool Dummy> struct TIsFixedWidthEncoding_Helper<Dummy, ANSICHAR>  { enum { Value = true }; };
+	template <bool Dummy> struct TIsFixedWidthEncoding_Helper<Dummy, WIDECHAR>  { enum { Value = true }; };
+	template <bool Dummy> struct TIsFixedWidthEncoding_Helper<Dummy, UCS2CHAR>  { enum { Value = true }; };
+	template <bool Dummy> struct TIsFixedWidthEncoding_Helper<Dummy, UTF32CHAR> { enum { Value = true }; };
 #if PLATFORM_TCHAR_IS_CHAR16
-	template <bool Dummy> struct TIsFixedWidthEncoding_Helper<Dummy, wchar_t> { enum { Value = true }; };
+	template <bool Dummy> struct TIsFixedWidthEncoding_Helper<Dummy, wchar_t>   { enum { Value = true }; };
 #endif
 
 	template <typename T>
diff --git a/Engine/Source/Runtime/Core/Public/HAL/Platform.h b/Engine/Source/Runtime/Core/Public/HAL/Platform.h
--- a/Engine/Source/Runtime/Core/Public/HAL/Platform.h
+++ b/Engine/Source/Runtime/Core/Public/HAL/Platform.h
@@ -1014,9 +1014,11 @@ namespace TypeTests
 
 	static_assert(char(-1) < char(0), "Unsigned char type test failed.");
 
-	static_assert((!TAreTypesEqual<ANSICHAR, WIDECHAR>::Value), "ANSICHAR and WIDECHAR should be different types.");
-	static_assert((!TAreTypesEqual<ANSICHAR, UCS2CHAR>::Value), "ANSICHAR and CHAR16 should be different types.");
-	static_assert((!TAreTypesEqual<WIDECHAR, UCS2CHAR>::Value), "WIDECHAR and CHAR16 should be different types.");
+	static_assert((!TAreTypesEqual<ANSICHAR, WIDECHAR>::Value),  "ANSICHAR and WIDECHAR should be different types.");
+	static_assert((!TAreTypesEqual<ANSICHAR, UTF8CHAR>::Value),  "ANSICHAR and UTF8CHAR should be different types.");
+//	static_assert((!TAreTypesEqual<UCS2CHAR, UTF16CHAR>::Value), "UCS2CHAR and UTF16CHAR should be different types."); // should be true - not currently
+	static_assert((!TAreTypesEqual<ANSICHAR, UCS2CHAR>::Value),  "ANSICHAR and CHAR16 should be different types.");
+	static_assert((!TAreTypesEqual<WIDECHAR, UCS2CHAR>::Value),  "WIDECHAR and CHAR16 should be different types.");
 	static_assert((TAreTypesEqual<TCHAR, ANSICHAR>::Value == true || TAreTypesEqual<TCHAR, WIDECHAR>::Value == true), "TCHAR should either be ANSICHAR or WIDECHAR.");
 
 	static_assert(sizeof(uint8) == 1, "uint8 type size test failed.");
@@ -1049,7 +1051,10 @@ namespace TypeTests
 
 	static_assert(sizeof(WIDECHAR) == 2 || sizeof(WIDECHAR) == 4, "WIDECHAR type size test failed.");
 
-	static_assert(sizeof(UCS2CHAR) == 2, "UCS2CHAR type size test failed.");
+	static_assert(sizeof(UTF8CHAR)  == 1, "UTF8CHAR type size test failed.");
+	static_assert(sizeof(UCS2CHAR)  == 2, "UCS2CHAR type size test failed.");
+	static_assert(sizeof(UTF16CHAR) == 2, "UTF16CHAR type size test failed.");
+	static_assert(sizeof(UTF32CHAR) == 4, "UTF32CHAR type size test failed.");
 
 	static_assert(sizeof(PTRINT) == sizeof(void *), "PTRINT type size test failed.");
 	static_assert(PTRINT(-1) < PTRINT(0), "PTRINT type sign test failed.");
