New char type compatibility constexpr function to replace TAreEncodingsCompatible (now deprecated).

diff --git a/Engine/Source/Runtime/Core/Public/Containers/StringConv.h b/Engine/Source/Runtime/Core/Public/Containers/StringConv.h
--- a/Engine/Source/Runtime/Core/Public/Containers/StringConv.h
+++ b/Engine/Source/Runtime/Core/Public/Containers/StringConv.h
@@ -1332,32 +1332,16 @@ typedef TStringPointer<wchar_t, TCHAR> FWCharToTCHAR;
  * @param Str The null-terminated source string to convert.
  */
 template <typename To, typename From>
-FORCEINLINE typename TEnableIf<FPlatformString::TAreEncodingsCompatible<To, From>::Value, TStringPointer<To>>::Type StringCast(const From* Str)
+FORCEINLINE auto StringCast(const From* Str)
 {
-	return TStringPointer<To>((const To*)Str);
-}
-
-/**
- * Creates an object which acts as a source of a given string type.  See example above.
- *
- * @param Str The null-terminated source string to convert.
- */
-template <typename To, typename From>
-FORCEINLINE typename TEnableIf<!FPlatformString::TAreEncodingsCompatible<To, From>::Value, TStringConversion<TStringConvert<From, To>>>::Type StringCast(const From* Str)
-{
-	return TStringConversion<TStringConvert<From, To>>(Str);
-}
-
-/**
- * Creates an object which acts as a source of a given string type.  See example above.
- *
- * @param Str The source string to convert, not necessarily null-terminated.
- * @param Len The number of From elements in Str.
- */
-template <typename To, typename From>
-FORCEINLINE typename TEnableIf<FPlatformString::TAreEncodingsCompatible<To, From>::Value, TStringPointer<To>>::Type StringCast(const From* Str, int32 Len)
-{
-	return TStringPointer<To>((const To*)Str, Len);
+	if constexpr (FPlatformString::IsCharEncodingCompatibleWith<From, To>())
+	{
+		return TStringPointer<To>((const To*)Str);
+	}
+	else
+	{
+		return TStringConversion<TStringConvert<From, To>>(Str);
+	}
 }
 
 /**
@@ -1367,9 +1351,16 @@ FORCEINLINE typename TEnableIf<FPlatformString::TAreEncodingsCompatible<To, From
  * @param Len The number of From elements in Str.
  */
 template <typename To, typename From>
-FORCEINLINE typename TEnableIf<!FPlatformString::TAreEncodingsCompatible<To, From>::Value, TStringConversion<TStringConvert<From, To>>>::Type StringCast(const From* Str, int32 Len)
+FORCEINLINE auto StringCast(const From* Str, int32 Len)
 {
-	return TStringConversion<TStringConvert<From, To>>(Str, Len);
+	if constexpr (FPlatformString::IsCharEncodingCompatibleWith<From, To>())
+	{
+		return TStringPointer<To>((const To*)Str, Len);
+	}
+	else
+	{
+		return TStringConversion<TStringConvert<From, To>>(Str, Len);
+	}
 }
 
 
diff --git a/Engine/Source/Runtime/Core/Public/GenericPlatform/GenericPlatformString.h b/Engine/Source/Runtime/Core/Public/GenericPlatform/GenericPlatformString.h
--- a/Engine/Source/Runtime/Core/Public/GenericPlatform/GenericPlatformString.h
+++ b/Engine/Source/Runtime/Core/Public/GenericPlatform/GenericPlatformString.h
@@ -5,7 +5,7 @@
 #include "CoreTypes.h"
 #include "Templates/EnableIf.h"
 #include "GenericPlatform/GenericPlatformStricmp.h"
-
+#include <type_traits>
 
 /**
  * Generic string implementation for most platforms
@@ -97,38 +97,48 @@ struct FGenericPlatformString : public FGenericPlatformStricmp
 	 * Same size is a minimum requirement.
 	 */
 	template <typename EncodingA, typename EncodingB>
-	struct TAreEncodingsCompatible
+	struct UE_DEPRECATED(5.0, "TAreEncodingsCompatible is deprecated, use IsCharEncodingCompatibleWith<SrcEncoding, DestEncoding>() instead.") TAreEncodingsCompatible
 	{
 		enum { Value = TIsFixedWidthEncoding<EncodingA>::Value && TIsFixedWidthEncoding<EncodingB>::Value && sizeof(EncodingA) == sizeof(EncodingB) };
 	};
 
 	/**
-	 * Converts the [Src, Src+SrcSize) string range from SourceChar to DestChar and writes it to the [Dest, Dest+DestSize) range.
-	 * The Src range should contain a null terminator if a null terminator is required in the output.
-	 * If the Dest range is not big enough to hold the converted output, NULL is returned.  In this case, nothing should be assumed about the contents of Dest.
+	 * Function which returns whether one encoding type is binary compatible with another.
 	 *
-	 * @param Dest      The start of the destination buffer.
-	 * @param DestSize  The size of the destination buffer.
-	 * @param Src       The start of the string to convert.
-	 * @param SrcSize   The number of Src elements to convert.
-	 * @param BogusChar The char to use when the conversion process encounters a character it cannot convert.
-	 * @return          A pointer to one past the last-written element.
+	 * Unlike TAreEncodingsCompatible, this is not commutative.  For example, ANSI is compatible with
+	 * UTF-8, but UTF-8 is not compatible with ANSI.
 	 */
-	template <typename SourceEncoding, typename DestEncoding>
-	static FORCEINLINE typename TEnableIf<
-		// This overload should be called when SourceEncoding and DestEncoding are 'compatible', i.e. they're the same type or equivalent (e.g. like UCS2CHAR and WIDECHAR are on Windows).
-		TAreEncodingsCompatible<SourceEncoding, DestEncoding>::Value,
-		DestEncoding*
-	>::Type Convert(DestEncoding* Dest, int32 DestSize, const SourceEncoding* Src, int32 SrcSize, DestEncoding BogusChar = (DestEncoding)'?')
+	template <typename SrcEncoding, typename DestEncoding>
+	static constexpr bool IsCharEncodingCompatibleWith()
 	{
-		if (DestSize < SrcSize)
-			return nullptr;
-
-		return (DestEncoding*)Memcpy(Dest, Src, SrcSize * sizeof(SourceEncoding)) + SrcSize;
-	}
+		if constexpr (std::is_same_v<SrcEncoding, DestEncoding>)
+		{
+			return true;
+		}
+		else if constexpr (std::is_same_v<SrcEncoding, ANSICHAR> && std::is_same_v<DestEncoding, UTF8CHAR>)
+		{
+			return true;
+		}
+		else if constexpr (std::is_same_v<SrcEncoding, UCS2CHAR> && std::is_same_v<DestEncoding, UTF16CHAR>)
+		{
+			return true;
+		}
+		else if constexpr (std::is_same_v<SrcEncoding, WIDECHAR> && std::is_same_v<DestEncoding, UCS2CHAR>)
+		{
+			return true;
+		}
+		else if constexpr (std::is_same_v<SrcEncoding, UCS2CHAR> && std::is_same_v<DestEncoding, WIDECHAR>)
+		{
+			return true;
+		}
+		else
+		{
+			return false;
+		}
+	};
 
 	/**
-	 * Converts the [Src, Src+SrcSize) string range from SourceEncoding to DestEncoding and writes it to the [Dest, Dest+DestSize) range.
+	 * Converts the [Src, Src+SrcSize) string range from SourceChar to DestChar and writes it to the [Dest, Dest+DestSize) range.
 	 * The Src range should contain a null terminator if a null terminator is required in the output.
 	 * If the Dest range is not big enough to hold the converted output, NULL is returned.  In this case, nothing should be assumed about the contents of Dest.
 	 *
@@ -140,37 +150,46 @@ struct FGenericPlatformString : public FGenericPlatformStricmp
 	 * @return          A pointer to one past the last-written element.
 	 */
 	template <typename SourceEncoding, typename DestEncoding>
-	static typename TEnableIf<
-		// This overload should be called when the types are not compatible but the source is fixed-width, e.g. ANSICHAR->WIDECHAR.
-		!TAreEncodingsCompatible<SourceEncoding, DestEncoding>::Value && TIsFixedWidthEncoding<SourceEncoding>::Value,
-		DestEncoding*
-	>::Type Convert(DestEncoding* Dest, int32 DestSize, const SourceEncoding* Src, int32 SrcSize, DestEncoding BogusChar = (DestEncoding)'?')
+	static FORCEINLINE DestEncoding* Convert(DestEncoding* Dest, int32 DestSize, const SourceEncoding* Src, int32 SrcSize, DestEncoding BogusChar = (DestEncoding)'?')
 	{
-		const int32 Size = DestSize <= SrcSize ? DestSize : SrcSize;
-		bool bInvalidChars = false;
-		for (int I = 0; I < Size; ++I)
+		if constexpr (IsCharEncodingCompatibleWith<SourceEncoding, DestEncoding>())
 		{
-			SourceEncoding SrcCh = Src[I];
-			Dest[I] = (DestEncoding)SrcCh;
-			bInvalidChars |= !CanConvertChar<DestEncoding>(SrcCh);
-		}
+			if (DestSize < SrcSize)
+			{
+				return nullptr;
+			}
 
-		if (bInvalidChars)
+			return (DestEncoding*)Memcpy(Dest, Src, SrcSize * sizeof(SourceEncoding)) + SrcSize;
+		}
+		else
 		{
+			const int32 Size = DestSize <= SrcSize ? DestSize : SrcSize;
+			bool bInvalidChars = false;
 			for (int I = 0; I < Size; ++I)
 			{
-				if (!CanConvertChar<DestEncoding>(Src[I]))
+				SourceEncoding SrcCh = Src[I];
+				Dest[I] = (DestEncoding)SrcCh;
+				bInvalidChars |= !CanConvertChar<DestEncoding>(SrcCh);
+			}
+
+			if (bInvalidChars)
+			{
+				for (int I = 0; I < Size; ++I)
 				{
-					Dest[I] = BogusChar;
+					if (!CanConvertChar<DestEncoding>(Src[I]))
+					{
+						Dest[I] = BogusChar;
+					}
 				}
+
+				LogBogusChars<DestEncoding>(Src, Size);
 			}
 
-			LogBogusChars<DestEncoding>(Src, Size);
+			return DestSize < SrcSize ? nullptr : Dest + Size;
 		}
-
-		return DestSize < SrcSize ? nullptr : Dest + Size;
 	}
 
+
 	/**
 	 * Returns the required buffer length for the [Src, Src+SrcSize) string when converted to the DestChar encoding.
 	 * The Src range should contain a null terminator if a null terminator is required in the output.
diff --git a/Engine/Source/Runtime/Engine/Public/ReplayTypes.h b/Engine/Source/Runtime/Engine/Public/ReplayTypes.h
--- a/Engine/Source/Runtime/Engine/Public/ReplayTypes.h
+++ b/Engine/Source/Runtime/Engine/Public/ReplayTypes.h
@@ -436,7 +436,7 @@ public:
 private:
 	void EnableFastStringSerialization()
 	{
-		if (FPlatformString::TAreEncodingsCompatible<WIDECHAR, TCHAR>::Value)
+		if (FPlatformString::IsCharEncodingCompatibleWith<WIDECHAR, TCHAR>())
 		{
 			Archive.SetForceUnicode(true);
 		}
@@ -444,7 +444,7 @@ private:
 
 	void RestoreStringSerialization()
 	{
-		if (FPlatformString::TAreEncodingsCompatible<WIDECHAR, TCHAR>::Value)
+		if (FPlatformString::IsCharEncodingCompatibleWith<WIDECHAR, TCHAR>())
 		{
 			Archive.SetForceUnicode(bWasUnicode);
 		}
