FPlatformString support for UTF8CHAR.

diff --git a/Engine/Source/Runtime/Core/Private/GenericPlatform/GenericPlatformString.cpp b/Engine/Source/Runtime/Core/Private/GenericPlatform/GenericPlatformString.cpp
--- a/Engine/Source/Runtime/Core/Private/GenericPlatform/GenericPlatformString.cpp
+++ b/Engine/Source/Runtime/Core/Private/GenericPlatform/GenericPlatformString.cpp
@@ -100,17 +100,28 @@ int32 StrncmpImpl(const CharType1* String1, const CharType2* String2, SIZE_T Cou
 }
 
 int32 FGenericPlatformString::Strncmp(const ANSICHAR* Str1, const ANSICHAR* Str2, SIZE_T Count) { return GenericPlatformStringPrivate::StrncmpImpl(Str1, Str2, Count); }
-int32 FGenericPlatformString::Strncmp(const WIDECHAR* Str1, const WIDECHAR* Str2, SIZE_T Count) { return GenericPlatformStringPrivate::StrncmpImpl(Str1, Str2, Count); }
-int32 FGenericPlatformString::Strncmp(const ANSICHAR* Str1, const WIDECHAR* Str2, SIZE_T Count) { return GenericPlatformStringPrivate::StrncmpImpl(Str1, Str2, Count); }
 int32 FGenericPlatformString::Strncmp(const WIDECHAR* Str1, const ANSICHAR* Str2, SIZE_T Count) { return GenericPlatformStringPrivate::StrncmpImpl(Str1, Str2, Count); }
+int32 FGenericPlatformString::Strncmp(const UTF8CHAR* Str1, const ANSICHAR* Str2, SIZE_T Count) { return GenericPlatformStringPrivate::StrncmpImpl(Str1, Str2, Count); }
+int32 FGenericPlatformString::Strncmp(const ANSICHAR* Str1, const WIDECHAR* Str2, SIZE_T Count) { return GenericPlatformStringPrivate::StrncmpImpl(Str1, Str2, Count); }
+int32 FGenericPlatformString::Strncmp(const WIDECHAR* Str1, const WIDECHAR* Str2, SIZE_T Count) { return GenericPlatformStringPrivate::StrncmpImpl(Str1, Str2, Count); }
+int32 FGenericPlatformString::Strncmp(const UTF8CHAR* Str1, const WIDECHAR* Str2, SIZE_T Count) { return GenericPlatformStringPrivate::StrncmpImpl(Str1, Str2, Count); }
+int32 FGenericPlatformString::Strncmp(const ANSICHAR* Str1, const UTF8CHAR* Str2, SIZE_T Count) { return GenericPlatformStringPrivate::StrncmpImpl(Str1, Str2, Count); }
+int32 FGenericPlatformString::Strncmp(const WIDECHAR* Str1, const UTF8CHAR* Str2, SIZE_T Count) { return GenericPlatformStringPrivate::StrncmpImpl(Str1, Str2, Count); }
+int32 FGenericPlatformString::Strncmp(const UTF8CHAR* Str1, const UTF8CHAR* Str2, SIZE_T Count) { return GenericPlatformStringPrivate::StrncmpImpl(Str1, Str2, Count); }
 
 #if !UE_BUILD_DOCS
 template CORE_API void FGenericPlatformString::LogBogusChars<ANSICHAR, WIDECHAR>(const WIDECHAR* Src, int32 SrcSize);
 template CORE_API void FGenericPlatformString::LogBogusChars<ANSICHAR, UCS2CHAR>(const UCS2CHAR* Src, int32 SrcSize);
+template CORE_API void FGenericPlatformString::LogBogusChars<ANSICHAR, UTF8CHAR>(const UTF8CHAR* Src, int32 SrcSize);
 template CORE_API void FGenericPlatformString::LogBogusChars<WIDECHAR, ANSICHAR>(const ANSICHAR* Src, int32 SrcSize);
 template CORE_API void FGenericPlatformString::LogBogusChars<WIDECHAR, UCS2CHAR>(const UCS2CHAR* Src, int32 SrcSize);
+template CORE_API void FGenericPlatformString::LogBogusChars<WIDECHAR, UTF8CHAR>(const UTF8CHAR* Src, int32 SrcSize);
 template CORE_API void FGenericPlatformString::LogBogusChars<UCS2CHAR, ANSICHAR>(const ANSICHAR* Src, int32 SrcSize);
 template CORE_API void FGenericPlatformString::LogBogusChars<UCS2CHAR, WIDECHAR>(const WIDECHAR* Src, int32 SrcSize);
+template CORE_API void FGenericPlatformString::LogBogusChars<UCS2CHAR, UTF8CHAR>(const UTF8CHAR* Src, int32 SrcSize);
+template CORE_API void FGenericPlatformString::LogBogusChars<UTF8CHAR, ANSICHAR>(const ANSICHAR* Src, int32 SrcSize);
+template CORE_API void FGenericPlatformString::LogBogusChars<UTF8CHAR, WIDECHAR>(const WIDECHAR* Src, int32 SrcSize);
+template CORE_API void FGenericPlatformString::LogBogusChars<UTF8CHAR, UCS2CHAR>(const UCS2CHAR* Src, int32 SrcSize);
 #if PLATFORM_TCHAR_IS_CHAR16
 template CORE_API void FGenericPlatformString::LogBogusChars<wchar_t, char16_t>(const char16_t* Src, int32 SrcSize);
 template CORE_API void FGenericPlatformString::LogBogusChars<char16_t, wchar_t>(const wchar_t* Src, int32 SrcSize);
diff --git a/Engine/Source/Runtime/Core/Public/Android/AndroidPlatformString.h b/Engine/Source/Runtime/Core/Public/Android/AndroidPlatformString.h
--- a/Engine/Source/Runtime/Core/Public/Android/AndroidPlatformString.h
+++ b/Engine/Source/Runtime/Core/Public/Android/AndroidPlatformString.h
@@ -480,7 +480,7 @@ struct FAndroidPlatformString : public FGenericPlatformString
 		return strcpy( Dest, Src );
 	}
 
-	static FORCEINLINE ANSICHAR* Strncpy(ANSICHAR* Dest, const ANSICHAR* Src, int32 MaxLen)
+	static FORCEINLINE ANSICHAR* Strncpy(ANSICHAR* Dest, const ANSICHAR* Src, SIZE_T MaxLen)
 	{
 		::strncpy(Dest, Src, MaxLen);
 		Dest[MaxLen-1]=0;
@@ -574,6 +574,104 @@ struct FAndroidPlatformString : public FGenericPlatformString
 		return (Result != -1 && Result < (int32)DestSize) ? Result : -1;
 	}
 
+	/**
+	 * UTF8CHAR implementation.
+	 */
+	static FORCEINLINE UTF8CHAR* Strcpy(UTF8CHAR* Dest, SIZE_T DestCount, const UTF8CHAR* Src)
+	{
+		return (UTF8CHAR*)Strcpy((ANSICHAR*)Dest, DestCount, (const ANSICHAR*)Src);
+	}
+
+	static FORCEINLINE UTF8CHAR* Strncpy(UTF8CHAR* Dest, const UTF8CHAR* Src, SIZE_T MaxLen)
+	{
+		return (UTF8CHAR*)Strncpy((ANSICHAR*)Dest, (const ANSICHAR*)Src, MaxLen);
+	}
+
+	static FORCEINLINE UTF8CHAR* Strcat(UTF8CHAR* Dest, SIZE_T DestCount, const UTF8CHAR* Src)
+	{
+		return (UTF8CHAR*)Strcat((ANSICHAR*)Dest, DestCount, (const ANSICHAR*)Src);
+	}
+
+	static FORCEINLINE int32 Strcmp(const UTF8CHAR* String1, const UTF8CHAR* String2)
+	{
+		return Strcmp((const ANSICHAR*)String1, (const ANSICHAR*)String2);
+	}
+
+	static FORCEINLINE int32 Strncmp(const UTF8CHAR* String1, const UTF8CHAR* String2, SIZE_T Count)
+	{
+		return Strncmp((const ANSICHAR*)String1, (const ANSICHAR*)String2, Count);
+	}
+
+	static FORCEINLINE int32 Strlen(const UTF8CHAR* String)
+	{
+		return Strlen((const ANSICHAR*)String);
+	}
+
+	static FORCEINLINE int32 Strnlen( const UTF8CHAR* String, SIZE_T StringSize )
+	{
+		return Strnlen((const ANSICHAR*)String, StringSize);
+	}
+
+	static FORCEINLINE const UTF8CHAR* Strstr(const UTF8CHAR* String, const UTF8CHAR* Find)
+	{
+		return (const UTF8CHAR*)Strstr((const ANSICHAR*)String, (const ANSICHAR*)Find);
+	}
+
+	static FORCEINLINE const UTF8CHAR* Strchr(const UTF8CHAR* String, UTF8CHAR C)
+	{
+		return (const UTF8CHAR*)Strchr((const ANSICHAR*)String, (ANSICHAR)C);
+	}
+
+	static FORCEINLINE const UTF8CHAR* Strrchr(const UTF8CHAR* String, UTF8CHAR C)
+	{
+		return (const UTF8CHAR*)Strrchr((const ANSICHAR*)String, (ANSICHAR)C);
+	}
+
+	static FORCEINLINE int32 Atoi(const UTF8CHAR* String)
+	{
+		return Atoi((const ANSICHAR*)String);
+	}
+
+	static FORCEINLINE int64 Atoi64(const UTF8CHAR* String)
+	{
+		return Atoi64((const ANSICHAR*)String);
+	}
+
+	static FORCEINLINE float Atof(const UTF8CHAR* String)
+	{
+		return Atof((const ANSICHAR*)String);
+	}
+
+	static FORCEINLINE double Atod(const UTF8CHAR* String)
+	{
+		return Atod((const ANSICHAR*)String);
+	}
+
+	static FORCEINLINE int32 Strtoi(const UTF8CHAR* Start, UTF8CHAR** End, int32 Base)
+	{
+		return Strtoi((const ANSICHAR*)Start, (ANSICHAR**)End, Base);
+	}
+
+	static FORCEINLINE int64 Strtoi64(const UTF8CHAR* Start, UTF8CHAR** End, int32 Base)
+	{
+		return Strtoi64((const ANSICHAR*)Start, (ANSICHAR**)End, Base);
+	}
+
+	static FORCEINLINE uint64 Strtoui64(const UTF8CHAR* Start, UTF8CHAR** End, int32 Base)
+	{
+		return Strtoui64((const ANSICHAR*)Start, (ANSICHAR**)End, Base);
+	}
+
+	static FORCEINLINE UTF8CHAR* Strtok(UTF8CHAR* StrToken, const UTF8CHAR* Delim, UTF8CHAR** Context)
+	{
+		return (UTF8CHAR*)Strtok((ANSICHAR*)StrToken, (const ANSICHAR*)Delim, (ANSICHAR**)Context);
+	}
+
+	static FORCEINLINE int32 GetVarArgs(UTF8CHAR* Dest, SIZE_T DestSize, const UTF8CHAR*& Fmt, va_list ArgPtr)
+	{
+		return GetVarArgs((ANSICHAR*)Dest, DestSize, *(const ANSICHAR**)&Fmt, ArgPtr);
+	}
+
 	/** 
 	 * UCS2 implementation 
 	 **/
diff --git a/Engine/Source/Runtime/Core/Public/GenericPlatform/GenericPlatformString.h b/Engine/Source/Runtime/Core/Public/GenericPlatform/GenericPlatformString.h
--- a/Engine/Source/Runtime/Core/Public/GenericPlatform/GenericPlatformString.h
+++ b/Engine/Source/Runtime/Core/Public/GenericPlatform/GenericPlatformString.h
@@ -205,9 +205,14 @@ struct FGenericPlatformString : public FGenericPlatformStricmp
 	}
 
 	CORE_API static int32 Strncmp(const ANSICHAR* String1, const ANSICHAR* String2, SIZE_T Count);
-	CORE_API static int32 Strncmp(const WIDECHAR* String1, const WIDECHAR* String2, SIZE_T Count);
-	CORE_API static int32 Strncmp(const ANSICHAR* String1, const WIDECHAR* String2, SIZE_T Count);
 	CORE_API static int32 Strncmp(const WIDECHAR* String1, const ANSICHAR* String2, SIZE_T Count);
+	CORE_API static int32 Strncmp(const UTF8CHAR* String1, const ANSICHAR* String2, SIZE_T Count);
+	CORE_API static int32 Strncmp(const ANSICHAR* String1, const WIDECHAR* String2, SIZE_T Count);
+	CORE_API static int32 Strncmp(const WIDECHAR* String1, const WIDECHAR* String2, SIZE_T Count);
+	CORE_API static int32 Strncmp(const UTF8CHAR* String1, const WIDECHAR* String2, SIZE_T Count);
+	CORE_API static int32 Strncmp(const ANSICHAR* String1, const UTF8CHAR* String2, SIZE_T Count);
+	CORE_API static int32 Strncmp(const WIDECHAR* String1, const UTF8CHAR* String2, SIZE_T Count);
+	CORE_API static int32 Strncmp(const UTF8CHAR* String1, const UTF8CHAR* String2, SIZE_T Count);
 
 private:
 	/**
diff --git a/Engine/Source/Runtime/Core/Public/GenericPlatform/GenericWidePlatformString.h b/Engine/Source/Runtime/Core/Public/GenericPlatform/GenericWidePlatformString.h
--- a/Engine/Source/Runtime/Core/Public/GenericPlatform/GenericWidePlatformString.h
+++ b/Engine/Source/Runtime/Core/Public/GenericPlatform/GenericWidePlatformString.h
@@ -208,7 +208,7 @@ public:
 		return strcpy( Dest, Src );
 	}
 
-	CORE_API static FORCEINLINE ANSICHAR* Strncpy(ANSICHAR* Dest, const ANSICHAR* Src, int32 MaxLen)
+	CORE_API static FORCEINLINE ANSICHAR* Strncpy(ANSICHAR* Dest, const ANSICHAR* Src, SIZE_T MaxLen)
 	{
 		::strncpy(Dest, Src, MaxLen);
 		Dest[MaxLen-1]=0;
@@ -302,6 +302,104 @@ public:
 		return (Result != -1 && Result < (int32)DestSize) ? Result : -1;
 	}
 
+	/**
+	 * UTF8CHAR implementation.
+	 */
+	static FORCEINLINE UTF8CHAR* Strcpy(UTF8CHAR* Dest, SIZE_T DestCount, const UTF8CHAR* Src)
+	{
+		return (UTF8CHAR*)Strcpy((ANSICHAR*)Dest, DestCount, (const ANSICHAR*)Src);
+	}
+
+	static FORCEINLINE UTF8CHAR* Strncpy(UTF8CHAR* Dest, const UTF8CHAR* Src, SIZE_T MaxLen)
+	{
+		return (UTF8CHAR*)Strncpy((ANSICHAR*)Dest, (const ANSICHAR*)Src, MaxLen);
+	}
+
+	static FORCEINLINE UTF8CHAR* Strcat(UTF8CHAR* Dest, SIZE_T DestCount, const UTF8CHAR* Src)
+	{
+		return (UTF8CHAR*)Strcat((ANSICHAR*)Dest, DestCount, (const ANSICHAR*)Src);
+	}
+
+	static FORCEINLINE int32 Strcmp(const UTF8CHAR* String1, const UTF8CHAR* String2)
+	{
+		return Strcmp((const ANSICHAR*)String1, (const ANSICHAR*)String2);
+	}
+
+	static FORCEINLINE int32 Strncmp(const UTF8CHAR* String1, const UTF8CHAR* String2, SIZE_T Count)
+	{
+		return Strncmp((const ANSICHAR*)String1, (const ANSICHAR*)String2, Count);
+	}
+
+	static FORCEINLINE int32 Strlen(const UTF8CHAR* String)
+	{
+		return Strlen((const ANSICHAR*)String);
+	}
+
+	static FORCEINLINE int32 Strnlen(const UTF8CHAR* String, SIZE_T StringSize)
+	{
+		return Strnlen((const ANSICHAR*)String, StringSize);
+	}
+
+	static FORCEINLINE const UTF8CHAR* Strstr(const UTF8CHAR* String, const UTF8CHAR* Find)
+	{
+		return (const UTF8CHAR*)Strstr((const ANSICHAR*)String, (const ANSICHAR*)Find);
+	}
+
+	static FORCEINLINE const UTF8CHAR* Strchr(const UTF8CHAR* String, UTF8CHAR C)
+	{
+		return (const UTF8CHAR*)Strchr((const ANSICHAR*)String, (ANSICHAR)C);
+	}
+
+	static FORCEINLINE const UTF8CHAR* Strrchr(const UTF8CHAR* String, UTF8CHAR C)
+	{
+		return (const UTF8CHAR*)Strrchr((const ANSICHAR*)String, (ANSICHAR)C);
+	}
+
+	static FORCEINLINE int32 Atoi(const UTF8CHAR* String)
+	{
+		return Atoi((const ANSICHAR*)String);
+	}
+
+	static FORCEINLINE int64 Atoi64(const UTF8CHAR* String)
+	{
+		return Atoi64((const ANSICHAR*)String);
+	}
+
+	static FORCEINLINE float Atof(const UTF8CHAR* String)
+	{
+		return Atof((const ANSICHAR*)String);
+	}
+
+	static FORCEINLINE double Atod(const UTF8CHAR* String)
+	{
+		return Atod((const ANSICHAR*)String);
+	}
+
+	static FORCEINLINE int32 Strtoi(const UTF8CHAR* Start, UTF8CHAR** End, int32 Base)
+	{
+		return Strtoi((const ANSICHAR*)Start, (ANSICHAR**)End, Base);
+	}
+
+	static FORCEINLINE int64 Strtoi64(const UTF8CHAR* Start, UTF8CHAR** End, int32 Base)
+	{
+		return Strtoi64((const ANSICHAR*)Start, (ANSICHAR**)End, Base);
+	}
+
+	static FORCEINLINE uint64 Strtoui64(const UTF8CHAR* Start, UTF8CHAR** End, int32 Base)
+	{
+		return Strtoui64((const ANSICHAR*)Start, (ANSICHAR**)End, Base);
+	}
+
+	static FORCEINLINE UTF8CHAR* Strtok(UTF8CHAR* StrToken, const UTF8CHAR* Delim, UTF8CHAR** Context)
+	{
+		return (UTF8CHAR*)Strtok((ANSICHAR*)StrToken, (const ANSICHAR*)Delim, (ANSICHAR**)Context);
+	}
+
+	static FORCEINLINE int32 GetVarArgs(UTF8CHAR* Dest, SIZE_T DestSize, const UTF8CHAR*& Fmt, va_list ArgPtr)
+	{
+		return GetVarArgs((ANSICHAR*)Dest, DestSize, *(const ANSICHAR**)&Fmt, ArgPtr);
+	}
+
 	/**
 	 * UCS2 implementation
 	 **/
diff --git a/Engine/Source/Runtime/Core/Public/GenericPlatform/MicrosoftPlatformString.h b/Engine/Source/Runtime/Core/Public/GenericPlatform/MicrosoftPlatformString.h
--- a/Engine/Source/Runtime/Core/Public/GenericPlatform/MicrosoftPlatformString.h
+++ b/Engine/Source/Runtime/Core/Public/GenericPlatform/MicrosoftPlatformString.h
@@ -140,10 +140,11 @@ struct FMicrosoftPlatformString : public FGenericPlatformString
 		return (ANSICHAR*)strcpy(Dest, Src);
 	}
 
-	static FORCEINLINE void Strncpy(ANSICHAR* Dest, const ANSICHAR* Src, SIZE_T MaxLen)
+	static FORCEINLINE ANSICHAR* Strncpy(ANSICHAR* Dest, const ANSICHAR* Src, SIZE_T MaxLen)
 	{
 		strncpy(Dest, Src, MaxLen);
 		Dest[MaxLen-1] = 0;
+		return Dest;
 	}
 
 	static FORCEINLINE ANSICHAR* Strcat(ANSICHAR* Dest, SIZE_T DestCount, const ANSICHAR* Src)
@@ -251,6 +252,109 @@ struct FMicrosoftPlatformString : public FGenericPlatformString
 	{
 		return _tcsnlen( (const WIDECHAR*)String, StringSize );
 	}
+
+	/**
+	 * UTF8CHAR implementation.
+	 */
+	static FORCEINLINE UTF8CHAR* Strcpy(UTF8CHAR* Dest, SIZE_T DestCount, const UTF8CHAR* Src)
+	{
+		return (UTF8CHAR*)Strcpy((ANSICHAR*)Dest, DestCount, (const ANSICHAR*)Src);
+	}
+
+	static FORCEINLINE UTF8CHAR* Strncpy(UTF8CHAR* Dest, const UTF8CHAR* Src, SIZE_T MaxLen)
+	{
+		return (UTF8CHAR*)Strncpy((ANSICHAR*)Dest, (const ANSICHAR*)Src, MaxLen);
+	}
+
+	static FORCEINLINE UTF8CHAR* Strcat(UTF8CHAR* Dest, SIZE_T DestCount, const UTF8CHAR* Src)
+	{
+		return (UTF8CHAR*)Strcat((ANSICHAR*)Dest, DestCount, (const ANSICHAR*)Src);
+	}
+
+	static FORCEINLINE UTF8CHAR* Strupr(UTF8CHAR* Dest, SIZE_T DestCount)
+	{
+		return (UTF8CHAR*)Strupr((ANSICHAR*)Dest, DestCount);
+	}
+
+	static FORCEINLINE int32 Strcmp(const UTF8CHAR* String1, const UTF8CHAR* String2)
+	{
+		return Strcmp((const ANSICHAR*)String1, (const ANSICHAR*)String2);
+	}
+
+	static FORCEINLINE int32 Strncmp(const UTF8CHAR* String1, const UTF8CHAR* String2, SIZE_T Count)
+	{
+		return Strncmp((const ANSICHAR*)String1, (const ANSICHAR*)String2, Count);
+	}
+	
+	static FORCEINLINE int32 Strlen(const UTF8CHAR* String)
+	{
+		return Strlen((const ANSICHAR*)String); 
+	}
+
+	static FORCEINLINE int32 Strnlen(const UTF8CHAR* String, SIZE_T StringSize)
+	{
+		return Strnlen((const ANSICHAR*)String, StringSize);
+	}
+
+	static FORCEINLINE const UTF8CHAR* Strstr(const UTF8CHAR* String, const UTF8CHAR* Find)
+	{
+		return (const UTF8CHAR*)Strstr((const ANSICHAR*)String, (const ANSICHAR*)Find);
+	}
+
+	static FORCEINLINE const UTF8CHAR* Strchr(const UTF8CHAR* String, UTF8CHAR C)
+	{
+		return (const UTF8CHAR*)Strchr((const ANSICHAR*)String, (ANSICHAR)C);
+	}
+
+	static FORCEINLINE const UTF8CHAR* Strrchr(const UTF8CHAR* String, UTF8CHAR C)
+	{
+		return (const UTF8CHAR*)Strrchr((const ANSICHAR*)String, (ANSICHAR)C);
+	}
+
+	static FORCEINLINE int32 Atoi(const UTF8CHAR* String)
+	{
+		return Atoi((const ANSICHAR*)String);
+	}
+
+	static FORCEINLINE int64 Atoi64(const UTF8CHAR* String)
+	{
+		return Atoi64((const ANSICHAR*)String);
+	}
+
+	static FORCEINLINE float Atof(const UTF8CHAR* String)
+	{
+		return Atof((const ANSICHAR*)String);
+	}
+
+	static FORCEINLINE double Atod(const UTF8CHAR* String)
+	{
+		return Atod((const ANSICHAR*)String);
+	}
+
+	static FORCEINLINE int32 Strtoi(const UTF8CHAR* Start, UTF8CHAR** End, int32 Base)
+	{
+		return Strtoi((const ANSICHAR*)Start, (ANSICHAR**)End, Base);
+	}
+
+	static FORCEINLINE int64 Strtoi64(const UTF8CHAR* Start, UTF8CHAR** End, int32 Base)
+	{
+		return Strtoi64((const ANSICHAR*)Start, (ANSICHAR**)End, Base);
+	}
+
+	static FORCEINLINE uint64 Strtoui64(const UTF8CHAR* Start, UTF8CHAR** End, int32 Base)
+	{
+		return Strtoui64((const ANSICHAR*)Start, (ANSICHAR**)End, Base);
+	}
+
+	static FORCEINLINE UTF8CHAR* Strtok(UTF8CHAR* StrToken, const UTF8CHAR* Delim, UTF8CHAR** Context)
+	{
+		return (UTF8CHAR*)Strtok((ANSICHAR*)StrToken, (const ANSICHAR*)Delim, (ANSICHAR**)Context);
+	}
+
+	static FORCEINLINE int32 GetVarArgs(UTF8CHAR* Dest, SIZE_T DestSize, const UTF8CHAR*& Fmt, va_list ArgPtr)
+	{
+		return GetVarArgs((ANSICHAR*)Dest, DestSize, *(const ANSICHAR**)&Fmt, ArgPtr);
+	}
 };
 
 #pragma warning(pop) // 'function' was was declared deprecated  (needed for the secure string functions)
diff --git a/Engine/Source/Runtime/Core/Public/GenericPlatform/StandardPlatformString.h b/Engine/Source/Runtime/Core/Public/GenericPlatform/StandardPlatformString.h
--- a/Engine/Source/Runtime/Core/Public/GenericPlatform/StandardPlatformString.h
+++ b/Engine/Source/Runtime/Core/Public/GenericPlatform/StandardPlatformString.h
@@ -202,7 +202,7 @@ public:
 		return strcpy( Dest, Src );
 	}
 
-	static FORCEINLINE ANSICHAR* Strncpy(ANSICHAR* Dest, const ANSICHAR* Src, int32 MaxLen)
+	static FORCEINLINE ANSICHAR* Strncpy(ANSICHAR* Dest, const ANSICHAR* Src, SIZE_T MaxLen)
 	{
 		::strncpy(Dest, Src, MaxLen);
 		Dest[MaxLen-1]=0;
@@ -296,6 +296,104 @@ public:
 		return (Result != -1 && Result < (int32)DestSize) ? Result : -1;
 	}
 
+	/**
+	 * UTF8CHAR implementation.
+	 */
+	static FORCEINLINE UTF8CHAR* Strcpy(UTF8CHAR* Dest, SIZE_T DestCount, const UTF8CHAR* Src)
+	{
+		return (UTF8CHAR*)Strcpy((ANSICHAR*)Dest, DestCount, (const ANSICHAR*)Src);
+	}
+
+	static FORCEINLINE UTF8CHAR* Strncpy(UTF8CHAR* Dest, const UTF8CHAR* Src, SIZE_T MaxLen)
+	{
+		return (UTF8CHAR*)Strncpy((ANSICHAR*)Dest, (const ANSICHAR*)Src, MaxLen);
+	}
+
+	static FORCEINLINE UTF8CHAR* Strcat(UTF8CHAR* Dest, SIZE_T DestCount, const UTF8CHAR* Src)
+	{
+		return (UTF8CHAR*)Strcat((ANSICHAR*)Dest, DestCount, (const ANSICHAR*)Src);
+	}
+
+	static FORCEINLINE int32 Strcmp(const UTF8CHAR* String1, const UTF8CHAR* String2)
+	{
+		return Strcmp((const ANSICHAR*)String1, (const ANSICHAR*)String2);
+	}
+
+	static FORCEINLINE int32 Strncmp(const UTF8CHAR* String1, const UTF8CHAR* String2, SIZE_T Count)
+	{
+		return Strncmp((const ANSICHAR*)String1, (const ANSICHAR*)String2, Count);
+	}
+
+	static FORCEINLINE int32 Strlen(const UTF8CHAR* String)
+	{
+		return Strlen((const ANSICHAR*)String);
+	}
+
+	static FORCEINLINE int32 Strnlen(const UTF8CHAR* String, SIZE_T StringSize)
+	{
+		return Strnlen((const ANSICHAR*)String, StringSize);
+	}
+
+	static FORCEINLINE const UTF8CHAR* Strstr(const UTF8CHAR* String, const UTF8CHAR* Find)
+	{
+		return (const UTF8CHAR*)Strstr((const ANSICHAR*)String, (const ANSICHAR*)Find);
+	}
+
+	static FORCEINLINE const UTF8CHAR* Strchr(const UTF8CHAR* String, UTF8CHAR C)
+	{
+		return (const UTF8CHAR*)Strchr((const ANSICHAR*)String, (ANSICHAR)C);
+	}
+
+	static FORCEINLINE const UTF8CHAR* Strrchr(const UTF8CHAR* String, UTF8CHAR C)
+	{
+		return (const UTF8CHAR*)Strrchr((const ANSICHAR*)String, (ANSICHAR)C);
+	}
+
+	static FORCEINLINE int32 Atoi(const UTF8CHAR* String)
+	{
+		return Atoi((const ANSICHAR*)String);
+	}
+
+	static FORCEINLINE int64 Atoi64(const UTF8CHAR* String)
+	{
+		return Atoi64((const ANSICHAR*)String);
+	}
+
+	static FORCEINLINE float Atof(const UTF8CHAR* String)
+	{
+		return Atof((const ANSICHAR*)String);
+	}
+
+	static FORCEINLINE double Atod(const UTF8CHAR* String)
+	{
+		return Atod((const ANSICHAR*)String);
+	}
+
+	static FORCEINLINE int32 Strtoi(const UTF8CHAR* Start, UTF8CHAR** End, int32 Base)
+	{
+		return Strtoi((const ANSICHAR*)Start, (ANSICHAR**)End, Base);
+	}
+
+	static FORCEINLINE int64 Strtoi64(const UTF8CHAR* Start, UTF8CHAR** End, int32 Base)
+	{
+		return Strtoi64((const ANSICHAR*)Start, (ANSICHAR**)End, Base);
+	}
+
+	static FORCEINLINE uint64 Strtoui64(const UTF8CHAR* Start, UTF8CHAR** End, int32 Base)
+	{
+		return Strtoui64((const ANSICHAR*)Start, (ANSICHAR**)End, Base);
+	}
+
+	static FORCEINLINE UTF8CHAR* Strtok(UTF8CHAR* StrToken, const UTF8CHAR* Delim, UTF8CHAR** Context)
+	{
+		return (UTF8CHAR*)Strtok((ANSICHAR*)StrToken, (const ANSICHAR*)Delim, (ANSICHAR**)Context);
+	}
+
+	static FORCEINLINE int32 GetVarArgs(UTF8CHAR* Dest, SIZE_T DestSize, const UTF8CHAR*& Fmt, va_list ArgPtr)
+	{
+		return GetVarArgs((ANSICHAR*)Dest, DestSize, *(const ANSICHAR**)&Fmt, ArgPtr);
+	}
+
 	/**
 	 * UCS2 implementation
 	 **/
