Fix for deprecation warnings.

diff --git a/Engine/Source/Runtime/Core/Public/Containers/StringConv.h b/Engine/Source/Runtime/Core/Public/Containers/StringConv.h
--- a/Engine/Source/Runtime/Core/Public/Containers/StringConv.h
+++ b/Engine/Source/Runtime/Core/Public/Containers/StringConv.h
@@ -1479,16 +1479,17 @@ private:
  * }
  */
 template <typename From, typename To, int32 DefaultConversionSize = DEFAULT_STRING_CONVERSION_SIZE>
-FORCEINLINE typename TEnableIf<FPlatformString::TAreEncodingsCompatible<To, From>::Value, TPassthruPointer<From>>::Type StringMemoryPassthru(To* Buffer, int32 BufferSize, int32 SourceLength)
+FORCEINLINE auto StringMemoryPassthru(To* Buffer, int32 BufferSize, int32 SourceLength)
 {
-	check(SourceLength <= BufferSize);
-	return TPassthruPointer<From>((From*)Buffer);
-}
-
-template <typename From, typename To, int32 DefaultConversionSize = DEFAULT_STRING_CONVERSION_SIZE>
-FORCEINLINE typename TEnableIf<!FPlatformString::TAreEncodingsCompatible<To, From>::Value, TStringPassthru<To, From, DefaultConversionSize>>::Type StringMemoryPassthru(To* Buffer, int32 BufferSize, int32 SourceLength)
-{
-	return TStringPassthru<To, From, DefaultConversionSize>(Buffer, BufferSize, SourceLength);
+	if constexpr (FPlatformString::IsCharEncodingCompatibleWith<From, To>())
+	{
+		check(SourceLength <= BufferSize);
+		return TPassthruPointer<From>((From*)Buffer);
+	}
+	else
+	{
+		return TStringPassthru<To, From, DefaultConversionSize>(Buffer, BufferSize, SourceLength);
+	}
 }
 
 template <typename ToType, typename FromType>
