Fixed missing REOPEN macros when PLATFORMFILETRACE_ENABLED is 0.

diff --git a/Engine/Source/Runtime/Core/Public/ProfilingDebugging/PlatformFileTrace.h b/Engine/Source/Runtime/Core/Public/ProfilingDebugging/PlatformFileTrace.h
--- a/Engine/Source/Runtime/Core/Public/ProfilingDebugging/PlatformFileTrace.h
+++ b/Engine/Source/Runtime/Core/Public/ProfilingDebugging/PlatformFileTrace.h
@@ -71,6 +71,8 @@ struct FPlatformFileTrace
 
 #define TRACE_PLATFORMFILE_BEGIN_OPEN(Path)
 #define TRACE_PLATFORMFILE_END_OPEN(FileHandle)
+#define TRACE_PLATFORMFILE_BEGIN_REOPEN(OldFileHandle)
+#define TRACE_PLATFORMFILE_END_REOPEN(NewFileHandle)
 #define TRACE_PLATFORMFILE_FAIL_OPEN(Path)
 #define TRACE_PLATFORMFILE_BEGIN_CLOSE(FileHandle)
 #define TRACE_PLATFORMFILE_END_CLOSE(FileHandle)
