More compatible char representations, when TCHAR is a wchar_t.

diff --git a/Engine/Source/Runtime/Core/Public/GenericPlatform/GenericPlatformString.h b/Engine/Source/Runtime/Core/Public/GenericPlatform/GenericPlatformString.h
--- a/Engine/Source/Runtime/Core/Public/GenericPlatform/GenericPlatformString.h
+++ b/Engine/Source/Runtime/Core/Public/GenericPlatform/GenericPlatformString.h
@@ -131,6 +131,16 @@ struct FGenericPlatformString : public FGenericPlatformStricmp
 		{
 			return true;
 		}
+#if PLATFORM_TCHAR_IS_CHAR16
+		else if constexpr (std::is_same_v<SrcEncoding, WIDECHAR> && std::is_same_v<DestEncoding, wchar_t>)
+		{
+			return true;
+		}
+		else if constexpr (std::is_same_v<SrcEncoding, wchar_t> && std::is_same_v<DestEncoding, WIDECHAR>)
+		{
+			return true;
+		}
+#endif
 		else
 		{
 			return false;
