StringConv.h now works with pointers to compatible char types (e.g. you can pass an ANSICHAR* to a string conversion function that expects a UTF8CHAR*).

diff --git a/Engine/Source/Runtime/Core/Public/Containers/StringConv.h b/Engine/Source/Runtime/Core/Public/Containers/StringConv.h
--- a/Engine/Source/Runtime/Core/Public/Containers/StringConv.h
+++ b/Engine/Source/Runtime/Core/Public/Containers/StringConv.h
@@ -19,6 +19,7 @@
 #include "Templates/UnrealTypeTraits.h"
 #include "Traits/ElementType.h"
 #include "Traits/IsContiguousContainer.h"
+#include <type_traits>
 
 #define DEFAULT_STRING_CONVERSION_SIZE 128u
 #define UNICODE_BOGUS_CHAR_CODEPOINT '?'
@@ -31,15 +32,23 @@ public:
 	typedef From FromType;
 	typedef To   ToType;
 
-	FORCEINLINE static void Convert(To* Dest, int32 DestLen, const From* Source, int32 SourceLen)
+	template <
+		typename CharType,
+		std::enable_if_t<FPlatformString::IsCharEncodingCompatibleWith<CharType, FromType>()>* = nullptr
+	>
+	FORCEINLINE static void Convert(To* Dest, int32 DestLen, const CharType* Source, int32 SourceLen)
 	{
-		To* Result = FPlatformString::Convert(Dest, DestLen, Source, SourceLen, (To)UNICODE_BOGUS_CHAR_CODEPOINT);
+		To* Result = FPlatformString::Convert(Dest, DestLen, (const FromType*)Source, SourceLen, (To)UNICODE_BOGUS_CHAR_CODEPOINT);
 		check(Result);
 	}
 
-	static int32 ConvertedLength(const From* Source, int32 SourceLen)
+	template <
+		typename CharType,
+		std::enable_if_t<FPlatformString::IsCharEncodingCompatibleWith<CharType, FromType>()>* = nullptr
+	>
+	static int32 ConvertedLength(const CharType* Source, int32 SourceLen)
 	{
-		return FPlatformString::ConvertedLength<To>(Source, SourceLen);
+		return FPlatformString::ConvertedLength<To>((const FromType*)Source, SourceLen);
 	}
 };
 
@@ -217,7 +226,7 @@ public:
 	static int32 Utf8FromCodepoint(uint32 Codepoint, BufferType OutputIterator, uint32 OutputIteratorByteSizeRemaining)
 	{
 		// Ensure we have at least one character in size to write
-		if (OutputIteratorByteSizeRemaining < sizeof(ANSICHAR))
+		if (OutputIteratorByteSizeRemaining < sizeof(ToType))
 		{
 			return 0;
 		}
@@ -236,33 +245,33 @@ public:
 		// Do the encoding...
 		if (Codepoint < 0x80)
 		{
-			*(OutputIterator++) = (ANSICHAR)Codepoint;
+			*(OutputIterator++) = (ToType)Codepoint;
 		}
 		else if (Codepoint < 0x800)
 		{
 			if (OutputIteratorByteSizeRemaining >= 2)
 			{
-				*(OutputIterator++) = (ANSICHAR)((Codepoint >> 6)         | 128 | 64);
-				*(OutputIterator++) = (ANSICHAR) (Codepoint       & 0x3F) | 128;
+				*(OutputIterator++) = (ToType)((Codepoint >> 6)         | 128 | 64);
+				*(OutputIterator++) = (ToType) (Codepoint       & 0x3F) | 128;
 			}
 		}
 		else if (Codepoint < 0x10000)
 		{
 			if (OutputIteratorByteSizeRemaining >= 3)
 			{
-				*(OutputIterator++) = (ANSICHAR)((Codepoint >> 12)        | 128 | 64 | 32);
-				*(OutputIterator++) = (ANSICHAR)((Codepoint >> 6) & 0x3F) | 128;
-				*(OutputIterator++) = (ANSICHAR) (Codepoint       & 0x3F) | 128;
+				*(OutputIterator++) = (ToType)((Codepoint >> 12)        | 128 | 64 | 32);
+				*(OutputIterator++) = (ToType)((Codepoint >> 6) & 0x3F) | 128;
+				*(OutputIterator++) = (ToType) (Codepoint       & 0x3F) | 128;
 			}
 		}
 		else
 		{
 			if (OutputIteratorByteSizeRemaining >= 4)
 			{
-				*(OutputIterator++) = (ANSICHAR)((Codepoint >> 18)         | 128 | 64 | 32 | 16);
-				*(OutputIterator++) = (ANSICHAR)((Codepoint >> 12) & 0x3F) | 128;
-				*(OutputIterator++) = (ANSICHAR)((Codepoint >> 6 ) & 0x3F) | 128;
-				*(OutputIterator++) = (ANSICHAR) (Codepoint        & 0x3F) | 128;
+				*(OutputIterator++) = (ToType)((Codepoint >> 18)         | 128 | 64 | 32 | 16);
+				*(OutputIterator++) = (ToType)((Codepoint >> 12) & 0x3F) | 128;
+				*(OutputIterator++) = (ToType)((Codepoint >> 6 ) & 0x3F) | 128;
+				*(OutputIterator++) = (ToType) (Codepoint        & 0x3F) | 128;
 			}
 		}
 
@@ -279,10 +288,14 @@ public:
 	 * @param SourceLen The length of the source string.
 	 * @return          The number of bytes written to Dest, up to DestLen, or -1 if the entire Source string could did not fit in DestLen bytes.
 	 */
-	template <typename DestBufferType>
-	static FORCEINLINE int32 Convert(DestBufferType Dest, int32 DestLen, const TCHAR* Source, int32 SourceLen)
+	template <
+		typename DestBufferType,
+		typename SrcBufferType,
+		std::enable_if_t<FPlatformString::IsCharEncodingCompatibleWith<SrcBufferType, FromType>()>* = nullptr
+	>
+	static FORCEINLINE int32 Convert(DestBufferType Dest, int32 DestLen, const SrcBufferType* Source, int32 SourceLen)
 	{
-		return Convert_Impl(Dest, DestLen, Source, SourceLen);
+		return Convert_Impl(Dest, DestLen, (const FromType*)Source, SourceLen);
 	}
 
 	/**
@@ -290,18 +303,22 @@ public:
 	 *
 	 * @return The length of the string in UTF-8 code units.
 	 */
-	static FORCEINLINE int32 ConvertedLength(const TCHAR* Source, int32 SourceLen)
+	template <
+		typename SrcBufferType,
+		std::enable_if_t<FPlatformString::IsCharEncodingCompatibleWith<SrcBufferType, FromType>()>* = nullptr
+	>
+	static FORCEINLINE int32 ConvertedLength(const SrcBufferType* Source, int32 SourceLen)
 	{
 		UE4StringConv_Private::FCountingOutputIterator Dest;
 		const int32 DestLen = SourceLen * 4;
-		Convert_Impl(Dest, DestLen, Source, SourceLen);
+		Convert_Impl(Dest, DestLen, (const FromType*)Source, SourceLen);
 
 		return Dest.GetCount();
 	}
 
 private:
 	template <typename DestBufferType>
-	static int32 Convert_Impl(DestBufferType& Dest, int32 DestLen, const TCHAR* Source, const int32 SourceLen)
+	static int32 Convert_Impl(DestBufferType& Dest, int32 DestLen, const FromType* Source, const int32 SourceLen)
 	{
 		DestBufferType DestStartingPosition = Dest;
 #if PLATFORM_TCHAR_IS_4_BYTES
@@ -405,9 +422,13 @@ public:
 	 * @param Source    The source string to convert.
 	 * @param SourceLen The length of the source string.
 	 */
-	static FORCEINLINE void Convert(TCHAR* Dest, const int32 DestLen, const ANSICHAR* Source, const int32 SourceLen)
+	template <
+		typename SrcBufferType,
+		std::enable_if_t<FPlatformString::IsCharEncodingCompatibleWith<SrcBufferType, FromType>()>* = nullptr
+	>
+	static FORCEINLINE void Convert(ToType* Dest, const int32 DestLen, const SrcBufferType* Source, const int32 SourceLen)
 	{
-		Convert_Impl(Dest, DestLen, Source, SourceLen);
+		Convert_Impl(Dest, DestLen, (const FromType*)Source, SourceLen);
 	}
 
 	/**
@@ -417,20 +438,24 @@ public:
 	 * @param SourceLen Length of source string; we will not read past this amount, even if the characters tell us to
 	 * @return The length of the string in UTF-16 or UTF-32 characters.
 	 */
-	static int32 ConvertedLength(const ANSICHAR* Source, const int32 SourceLen)
+	template <
+		typename SrcBufferType,
+		std::enable_if_t<FPlatformString::IsCharEncodingCompatibleWith<SrcBufferType, FromType>()>* = nullptr
+	>
+	static int32 ConvertedLength(const SrcBufferType* Source, const int32 SourceLen)
 	{
 		UE4StringConv_Private::FCountingOutputIterator Dest;
-		Convert_Impl(Dest, MAX_int32, Source, SourceLen);
+		Convert_Impl(Dest, MAX_int32, (const FromType*)Source, SourceLen);
 
 		return Dest.GetCount();
 	}
 
 private:
-	static uint32 CodepointFromUtf8(const ANSICHAR*& SourceString, const uint32 SourceLengthRemaining)
+	static uint32 CodepointFromUtf8(const FromType*& SourceString, const uint32 SourceLengthRemaining)
 	{
 		checkSlow(SourceLengthRemaining > 0);
 
-		const ANSICHAR* OctetPtr = SourceString;
+		const FromType* OctetPtr = SourceString;
 
 		uint32 Codepoint = 0;
 		uint32 Octet = (uint32) ((uint8) *SourceString);
@@ -656,9 +681,9 @@ private:
 	 * Read Source string, converting the data from UTF-8 into UTF-16, and placing these in the Destination
 	 */
 	template <typename DestBufferType>
-	static void Convert_Impl(DestBufferType& ConvertedBuffer, int32 DestLen, const ANSICHAR* Source, const int32 SourceLen)
+	static void Convert_Impl(DestBufferType& ConvertedBuffer, int32 DestLen, const FromType* Source, const int32 SourceLen)
 	{
-		const ANSICHAR* SourceEnd = Source + SourceLen;
+		const FromType* SourceEnd = Source + SourceLen;
 
 		const uint64 ExtendedCharMask = 0x8080808080808080;
 		while (Source < SourceEnd && DestLen > 0)
@@ -809,9 +834,13 @@ public:
 	 * @param Source    The source string to convert.
 	 * @param SourceLen The length of the source string.
 	 */
-	static FORCEINLINE void Convert(ToType* Dest, int32 DestLen, const FromType* Source, int32 SourceLen)
+	template <
+		typename SrcBufferType,
+		std::enable_if_t<FPlatformString::IsCharEncodingCompatibleWith<SrcBufferType, FromType>()>* = nullptr
+	>
+	static FORCEINLINE void Convert(ToType* Dest, int32 DestLen, const SrcBufferType* Source, int32 SourceLen)
 	{
-		Convert_Impl(Dest, DestLen, Source, SourceLen);
+		Convert_Impl(Dest, DestLen, (const FromType*)Source, SourceLen);
 	}
 
 	/**
@@ -819,11 +848,15 @@ public:
 	 *
 	 * @return The length of the string in UTF-16 code units.
 	 */
-	static FORCEINLINE int32 ConvertedLength(const FromType* Source, int32 SourceLen)
+	template <
+		typename SrcBufferType,
+		std::enable_if_t<FPlatformString::IsCharEncodingCompatibleWith<SrcBufferType, FromType>()>* = nullptr
+	>
+	static FORCEINLINE int32 ConvertedLength(const SrcBufferType* Source, int32 SourceLen)
 	{
 		UE4StringConv_Private::FCountingOutputIterator Dest;
 		const int32 DestLen = SourceLen * 2;
-		Convert_Impl(Dest, DestLen, Source, SourceLen);
+		Convert_Impl(Dest, DestLen, (const FromType*)Source, SourceLen);
 
 		return Dest.GetCount();
 	}
@@ -877,9 +910,13 @@ public:
 	 * @param Source    The source string to convert.
 	 * @param SourceLen The length of the source string.
 	 */
-	static FORCEINLINE void Convert(ToType* Dest, const int32 DestLen, const FromType* Source, const int32 SourceLen)
+	template <
+		typename SrcBufferType,
+		std::enable_if_t<FPlatformString::IsCharEncodingCompatibleWith<SrcBufferType, FromType>()>* = nullptr
+	>
+	static FORCEINLINE void Convert(ToType* Dest, const int32 DestLen, const SrcBufferType* Source, const int32 SourceLen)
 	{
-		Convert_Impl(Dest, DestLen, Source, SourceLen);
+		Convert_Impl(Dest, DestLen, (const FromType*)Source, SourceLen);
 	}
 
 	/**
@@ -889,10 +926,14 @@ public:
 	 * @param SourceLen Length of source string; we will not read past this amount, even if the characters tell us to
 	 * @return The length of the string in UTF-32 characters.
 	 */
-	static int32 ConvertedLength(const FromType* Source, const int32 SourceLen)
+	template <
+		typename SrcBufferType,
+		std::enable_if_t<FPlatformString::IsCharEncodingCompatibleWith<SrcBufferType, FromType>()>* = nullptr
+	>
+	static int32 ConvertedLength(const SrcBufferType* Source, const int32 SourceLen)
 	{
 		UE4StringConv_Private::FCountingOutputIterator Dest;
-		Convert_Impl(Dest, MAX_int32, Source, SourceLen);
+		Convert_Impl(Dest, MAX_int32, (const FromType*)Source, SourceLen);
 
 		return Dest.GetCount();
 	}
@@ -1008,11 +1049,15 @@ class TStringConversion : private Converter, private TInlineAllocator<DefaultCon
 	}
 
 public:
-	explicit TStringConversion(const FromType* Source)
+	template <
+		typename SrcBufferType,
+		std::enable_if_t<FPlatformString::IsCharEncodingCompatibleWith<SrcBufferType, FromType>()>* = nullptr
+	>
+	explicit TStringConversion(const SrcBufferType* Source)
 	{
 		if (Source)
 		{
-			Init(Source, TCString<FromType>::Strlen(Source), ENullTerminatedString::Yes);
+			Init((const FromType*)Source, TCString<FromType>::Strlen((const FromType*)Source), ENullTerminatedString::Yes);
 		}
 		else
 		{
@@ -1021,19 +1066,23 @@ public:
 		}
 	}
 
-	TStringConversion(const FromType* Source, int32 SourceLen)
+	template <
+		typename SrcBufferType,
+		std::enable_if_t<FPlatformString::IsCharEncodingCompatibleWith<SrcBufferType, FromType>()>* = nullptr
+	>
+	TStringConversion(const SrcBufferType* Source, int32 SourceLen)
 	{
 		if (Source)
 		{
 			ENullTerminatedString::Type NullTerminated = ENullTerminatedString::No;
-			if (SourceLen > 0 && Source[SourceLen-1] == 0)
+			if (SourceLen > 0 && ((const FromType*)Source)[SourceLen-1] == 0)
 			{
 				// Given buffer is null-terminated
 				NullTerminated = ENullTerminatedString::Yes;
 				SourceLen -= 1;
 			}
 
-			Init(Source, SourceLen, NullTerminated);
+			Init((const FromType*)Source, SourceLen, NullTerminated);
 		}
 		else
 		{
@@ -1045,13 +1094,19 @@ public:
 	/**
 	 * Construct from a compatible character range such as TStringView or TStringBuilder.
 	 */
-	template <typename FromRangeType,
-		typename TEnableIf<TAnd<
-			TIsContiguousContainer<FromRangeType>,
-			TNot<TIsArray<typename TRemoveReference<FromRangeType>::Type>>,
-			TIsSame<FromType, typename TRemoveCV<typename TRemovePointer<decltype(GetData(DeclVal<FromRangeType>()))>::Type>::Type>>::Value>::Type* = nullptr>
+	template <
+		typename FromRangeType,
+		typename FromRangeCharType = std::remove_cv_t<std::remove_pointer_t<decltype(GetData(DeclVal<FromRangeType>()))>>,
+		std::enable_if_t<
+			TAnd<
+				TIsContiguousContainer<FromRangeType>,
+				TNot<TIsArray<typename TRemoveReference<FromRangeType>::Type>>
+			>::Value &&
+			FPlatformString::IsCharEncodingCompatibleWith<FromRangeCharType, FromType>()
+		>* = nullptr
+	>
 	TStringConversion(FromRangeType&& Source)
-		: TStringConversion(GetData(Source), GetNum(Source))
+		: TStringConversion((const FromType*)GetData(Source), GetNum(Source))
 	{
 	}
 
@@ -1138,7 +1193,11 @@ public:
 	typedef InToType   ToType;
 
 public:
-	explicit TStringPointer(const FromType* Source)
+	template <
+		typename SrcBufferType,
+		std::enable_if_t<FPlatformString::IsCharEncodingCompatibleWith<SrcBufferType, FromType>()>* = nullptr
+	>
+	explicit TStringPointer(const SrcBufferType* Source)
 	{
 		if (Source)
 		{
@@ -1152,11 +1211,15 @@ public:
 		}
 	}
 
-	TStringPointer(const FromType* Source, int32 SourceLen)
+	template <
+		typename SrcBufferType,
+		std::enable_if_t<FPlatformString::IsCharEncodingCompatibleWith<SrcBufferType, FromType>()>* = nullptr
+	>
+	TStringPointer(const SrcBufferType* Source, int32 SourceLen)
 	{
 		if (Source)
 		{
-			if (SourceLen > 0 && Source[SourceLen-1] == 0)
+			if (SourceLen > 0 && ((const FromType*)Source)[SourceLen-1] == 0)
 			{
 				// Given buffer is null-terminated
 				SourceLen -= 1;
@@ -1175,13 +1238,19 @@ public:
 	/**
 	 * Construct from a compatible character range such as TStringView or TStringBuilder.
 	 */
-	template <typename FromRangeType,
-		typename TEnableIf<TAnd<
-			TIsContiguousContainer<FromRangeType>,
-			TNot<TIsArray<typename TRemoveReference<FromRangeType>::Type>>,
-			TIsSame<FromType, typename TRemoveCV<typename TRemovePointer<decltype(GetData(DeclVal<FromRangeType>()))>::Type>::Type>>::Value>::Type* = nullptr>
+	template <
+		typename FromRangeType,
+		typename FromRangeCharType = std::remove_cv_t<std::remove_pointer_t<decltype(GetData(DeclVal<FromRangeType>()))>>,
+		std::enable_if_t<
+			TAnd<
+				TIsContiguousContainer<FromRangeType>,
+				TNot<TIsArray<typename TRemoveReference<FromRangeType>::Type>>
+			>::Value &&
+			FPlatformString::IsCharEncodingCompatibleWith<FromRangeCharType, FromType>()
+		>* = nullptr
+	>
 	TStringPointer(FromRangeType&& Source)
-		: TStringPointer(GetData(Source), GetNum(Source))
+		: TStringPointer((const FromType*)GetData(Source), GetNum(Source))
 	{
 	}
 
