Added Zero static member to FBlake3Hash and FIoHash

#rnx

diff --git a/Engine/Source/Runtime/Core/Public/Hash/Blake3.h b/Engine/Source/Runtime/Core/Public/Hash/Blake3.h
--- a/Engine/Source/Runtime/Core/Public/Hash/Blake3.h
+++ b/Engine/Source/Runtime/Core/Public/Hash/Blake3.h
@@ -44,10 +44,15 @@ public:
 	inline ByteArray& GetBytes() { return Hash; }
 	inline const ByteArray& GetBytes() const { return Hash; }
 
+	/** A zero hash. */
+	static const FBlake3Hash Zero;
+
 private:
 	alignas(uint32) ByteArray Hash{};
 };
 
+inline const FBlake3Hash FBlake3Hash::Zero;
+
 ///////////////////////////////////////////////////////////////////////////////////////////////////////////////////////
 
 /** Calculates a BLAKE3 hash. */
diff --git a/Engine/Source/Runtime/Core/Public/IO/IoHash.h b/Engine/Source/Runtime/Core/Public/IO/IoHash.h
--- a/Engine/Source/Runtime/Core/Public/IO/IoHash.h
+++ b/Engine/Source/Runtime/Core/Public/IO/IoHash.h
@@ -58,10 +58,15 @@ public:
 	[[nodiscard]] static inline FIoHash HashBuffer(const void* Data, uint64 Size);
 	[[nodiscard]] static inline FIoHash HashBuffer(const FCompositeBuffer& Buffer);
 
+	/** A zero hash. */
+	static const FIoHash Zero;
+
 private:
 	alignas(uint32) ByteArray Hash{};
 };
 
+inline const FIoHash FIoHash::Zero;
+
 inline FIoHash::FIoHash(const ByteArray& InHash)
 {
 	FMemory::Memcpy(Hash, InHash, sizeof(ByteArray));
