Removed extra assignment to V1 in the math test.

diff --git a/Engine/Source/Runtime/Core/Private/Tests/Math/UnrealMathTest.cpp b/Engine/Source/Runtime/Core/Private/Tests/Math/UnrealMathTest.cpp
--- a/Engine/Source/Runtime/Core/Private/Tests/Math/UnrealMathTest.cpp
+++ b/Engine/Source/Runtime/Core/Private/Tests/Math/UnrealMathTest.cpp
@@ -973,7 +973,7 @@ bool RunDoubleVectorTest()
 	LogTest(TEXT("VectorSwizzle2013<double>"), TestVectorsEqual(V0, V1));
 
 	V0 = MakeVectorRegisterDouble(4.0f, 3.0f, 2.0f, 1.0f);
-	V1 = V1 = VectorSwizzle(V0, 2, 3, 0, 1);
+	V1 = VectorSwizzle(V0, 2, 3, 0, 1);
 	V0 = MakeVectorRegisterDouble(2.0f, 1.0f, 4.0f, 3.0f);
 	LogTest(TEXT("VectorSwizzle2301<double>"), TestVectorsEqual(V0, V1));
 
