Ignore protobuf generated files when running PVS-Studio static code analysis (*.pb.cc)

diff --git a/Engine/Source/Runtime/Core/Public/Windows/WindowsPlatformCodeAnalysis.h b/Engine/Source/Runtime/Core/Public/Windows/WindowsPlatformCodeAnalysis.h
--- a/Engine/Source/Runtime/Core/Public/Windows/WindowsPlatformCodeAnalysis.h
+++ b/Engine/Source/Runtime/Core/Public/Windows/WindowsPlatformCodeAnalysis.h
@@ -240,6 +240,9 @@
 	//-V:FScopedDurationTimer(:1050
 	//-V:FQueryFastData(:1050
 
+	// Exclude all generated protobuf files
+	//V_EXCLUDE_PATH *.pb.cc
+
 #else // defined(__clang_analyzer__)
 
 // A fake function marked with noreturn that acts as a marker for CA_ASSUME to ensure the
