Expose operator == and != on TStatId for state filtering.

diff --git a/Engine/Source/Runtime/Core/Public/Stats/Stats.h b/Engine/Source/Runtime/Core/Public/Stats/Stats.h
--- a/Engine/Source/Runtime/Core/Public/Stats/Stats.h
+++ b/Engine/Source/Runtime/Core/Public/Stats/Stats.h
@@ -119,6 +119,16 @@ struct TStatId
 	{
 		return StatString != nullptr;
 	}
+
+	FORCEINLINE bool operator==(TStatId Other) const
+	{
+		return StatString == Other.StatString;
+	}
+
+	FORCEINLINE bool operator!=(TStatId Other) const
+	{
+		return StatString != Other.StatString;
+	}
 };
 
 #if USE_LIGHTWEIGHT_STATS_FOR_HITCH_DETECTION && USE_HITCH_DETECTION
diff --git a/Engine/Source/Runtime/Core/Public/Stats/Stats2.h b/Engine/Source/Runtime/Core/Public/Stats/Stats2.h
--- a/Engine/Source/Runtime/Core/Public/Stats/Stats2.h
+++ b/Engine/Source/Runtime/Core/Public/Stats/Stats2.h
@@ -199,6 +199,16 @@ struct TStatId
 		return StatIdPtr->StatDescriptionWide.Get();
 	}
 
+	FORCEINLINE bool operator==(TStatId Other) const
+	{
+		return StatIdPtr == Other.StatIdPtr;
+	}
+
+	FORCEINLINE bool operator!=(TStatId Other) const
+	{
+		return StatIdPtr != Other.StatIdPtr;
+	}
+
 private:
 	/** NAME_None. */
 	CORE_API static TStatIdData TStatId_NAME_None;
