Disable spurious C4063 warning triggered by user-defined literal case values.

diff --git a/Engine/Source/Runtime/Core/Public/HoloLens/HoloLensPlatformCompilerSetup.h b/Engine/Source/Runtime/Core/Public/HoloLens/HoloLensPlatformCompilerSetup.h
--- a/Engine/Source/Runtime/Core/Public/HoloLens/HoloLensPlatformCompilerSetup.h
+++ b/Engine/Source/Runtime/Core/Public/HoloLens/HoloLensPlatformCompilerSetup.h
@@ -16,7 +16,7 @@
 static_assert(_MSC_VER >= 1900, "Visual Studio 2015 or newer required.");
 
 
-#pragma warning (error:      4001 4002 4003 4004      4006 4007 4008 4009 4010 4011 4012 4013 4014 4015 4016 4017 4018 4019 4020 4021 4022 4023 4024 4025 4026 4027 4028 4029 4030 4031 4032 4033 4034 4035 4036 4037 4038 4039 4040 4041 4042 4043 4044 4045 4046 4047 4048 4049 4050 4051 4052 4053 4054 4055 4056 4057 4058 4059 4060           4063 4064 4065 4066 4067 4068 4069 4070 4071 4072 4073 4074 4075 4076 4077 4078 4079 4080 4081 4082 4083 4084 4085 4086 4087 4088 4089 4090 4091 4092 4093 4094 4095 4096 4097 4098 4099)
+#pragma warning (error:      4001 4002 4003 4004      4006 4007 4008 4009 4010 4011 4012 4013 4014 4015 4016 4017 4018 4019 4020 4021 4022 4023 4024 4025 4026 4027 4028 4029 4030 4031 4032 4033 4034 4035 4036 4037 4038 4039 4040 4041 4042 4043 4044 4045 4046 4047 4048 4049 4050 4051 4052 4053 4054 4055 4056 4057 4058 4059 4060                4064 4065 4066 4067 4068 4069 4070 4071 4072 4073 4074 4075 4076 4077 4078 4079 4080 4081 4082 4083 4084 4085 4086 4087 4088 4089 4090 4091 4092 4093 4094 4095 4096 4097 4098 4099)
 #pragma warning (error: 4100 4101 4102 4103 4104 4105 4106 4107 4108 4109 4110 4111 4112 4113 4114 4115 4116 4117 4118 4119 4120 4121 4122 4123 4124 4125 4126 4127 4128 4129 4130 4131 4132 4133 4134 4135 4136 4137 4138 4139 4140 4141 4142 4143 4144 4145 4146 4147 4148 4149 4150 4151 4152 4153 4154 4155 4156 4157 4158 4159 4160 4161 4162 4163 4164      4166 4167 4168 4169 4170 4171 4172 4173 4174 4175 4176 4177 4178 4179 4180 4181 4182 4183 4184 4185 4186 4187 4188 4189 4190 4191 4192 4193 4194 4195 4196 4197 4198 4199)
 #pragma warning (error: 4200 4201 4202 4203 4204 4205 4206 4207 4208 4209 4210 4211 4212 4213 4214 4215 4216 4217 4218 4219 4220 4221 4222 4223 4224 4225 4226 4227 4228 4229 4230 4231 4232 4233 4234 4235 4236 4237 4238 4239 4240 4241      4243 4244 4245 4246 4247 4248 4249 4250 4251 4252 4253      4255 4256 4257 4258 4259 4260 4261 4262 4263 4264           4267 4268 4269 4270 4271 4272 4273 4274 4275 4276 4277 4278 4279 4280 4281 4282 4283 4284 4285 4286 4287 4288 4289 4290 4291 4292 4293 4294 4295      4297 4298 4299)
 #pragma warning (error: 4300 4301 4302 4303 4304      4306 4307 4308 4309 4310           4313 4314 4315 4316 4317 4318 4319 4320 4321 4322 4323 4324 4325 4326 4327 4328 4329 4330 4331 4332 4333 4334 4335 4336 4337 4338      4340 4341      4343 4344      4346 4347 4348 4349      4351 4352 4353 4354 4355 4356 4357 4358 4359 4360 4361 4362 4363 4364      4366 4367 4368 4369           4372      4374 4375 4376 4377 4378 4379 4380 4381 4382 4383 4384 4385 4386 4387      4389 4390 4391 4392 4393 4394 4395 4396 4397 4398 4399)
@@ -89,6 +89,7 @@ static_assert(_MSC_VER >= 1900, "Visual Studio 2015 or newer required.");
 #pragma warning (default : 4996)
 
 // Unwanted VC++ level 4 warnings to disable.
+#pragma warning(disable : 4063) // case 'val' is not a valid value for switch of enum 'E'
 #pragma warning(disable : 4100) // unreferenced formal parameter										
 #pragma warning(disable : 4127) // Conditional expression is constant									
 #pragma warning(disable : 4200) // Zero-length array item at end of structure, a VC-specific extension	
diff --git a/Engine/Source/Runtime/Core/Public/Windows/WindowsPlatformCompilerSetup.h b/Engine/Source/Runtime/Core/Public/Windows/WindowsPlatformCompilerSetup.h
--- a/Engine/Source/Runtime/Core/Public/Windows/WindowsPlatformCompilerSetup.h
+++ b/Engine/Source/Runtime/Core/Public/Windows/WindowsPlatformCompilerSetup.h
@@ -23,7 +23,7 @@ static_assert(_MSC_VER >= 1924, "Visual Studio 2019 v16.4 or greater is required
 // Warnings that we explicitly disable later are still included in there.
 //
 
-#pragma warning (error:      4001 4002 4003           4006 4007 4008      4010           4013      4015           4018 4019 4020      4022 4023 4024 4025 4026 4027 4028 4029 4030 4031 4032 4033 4034 4035 4036      4038           4041 4042           4045      4047 4048 4049      4051 4052 4053 4054 4055 4056 4057           4060           4063 4064 4065 4066 4067 4068 4069                4073 4074 4075 4076 4077      4079 4080 4081      4083      4085 4086 4087 4088 4089 4090 4091 4092      4094      4096 4097 4098 4099)
+#pragma warning (error:      4001 4002 4003           4006 4007 4008      4010           4013      4015           4018 4019 4020      4022 4023 4024 4025 4026 4027 4028 4029 4030 4031 4032 4033 4034 4035 4036      4038           4041 4042           4045      4047 4048 4049      4051 4052 4053 4054 4055 4056 4057           4060                4064 4065 4066 4067 4068 4069                4073 4074 4075 4076 4077      4079 4080 4081      4083      4085 4086 4087 4088 4089 4090 4091 4092      4094      4096 4097 4098 4099)
 #pragma warning (error: 4100 4101 4102 4103                          4109           4112 4113 4114 4115 4116 4117      4119 4120 4121 4122 4123 4124 4125      4127      4129 4130 4131 4132 4133                4137 4138           4141 4142 4143 4144 4145 4146                4150      4152 4153 4154 4155 4156 4157 4158 4159 4160 4161 4162 4163 4164      4166 4167 4168                4172      4174 4175 4176 4177 4178 4179 4180 4181 4182 4183      4185 4186 4187 4188 4189 4190 4191 4192 4193 4194 4195 4196 4197      4199)
 #pragma warning (error: 4200 4201 4202 4203 4204 4205 4206 4207 4208      4210 4211 4212 4213 4214 4215 4216      4218      4220 4221      4223 4224      4226 4227 4228 4229 4230      4232 4233 4234 4235      4237 4238 4239 4240           4243      4245                     4250 4251                4255 4256      4258                     4263 4264           4267 4268 4269           4272 4273 4274 4275 4276 4277 4278 4279 4280 4281 4282 4283      4285 4286 4287 4288 4289 4290 4291      4293      4295      4297 4298 4299)
 #pragma warning (error:      4301 4302 4303           4306      4308 4309 4310           4313 4314 4315 4316      4318 4319      4321 4322 4323 4324 4325 4326 4327 4328 4329 4330           4333 4334 4335 4336 4337 4338      4340           4343 4344      4346      4348                4352 4353      4355 4356 4357 4358 4359           4362      4364      4366 4367 4368 4369                     4374 4375 4376 4377 4378 4379 4380 4381 4382 4383 4384           4387      4389 4390 4391 4392 4393 4394 4395 4396 4397 4398 4399)
@@ -124,6 +124,8 @@ static_assert(_MSC_VER >= 1924, "Visual Studio 2019 v16.4 or greater is required
 // Disabled warnings
 //
 
+#pragma warning(disable: 4063) // case 'val' is not a valid value for switch of enum 'E'
+
 // @todo HoloLens: Disabled because DbgHelp.h has some anonymous typedefs in it (not allowed in Visual Studio 2015).  We should probably just wrap that header.
 #pragma warning(disable: 4091) // 'keyword' : ignored on left of 'type' when no variable is declared								https://docs.microsoft.com/en-us/cpp/error-messages/compiler-warnings/compiler-warning-level-1-c4091
 
