Removed the move constructor from TScopeGuard because C++17 requires RVO

Also removed the few dependencies on other headers to make ScopeExit.h as slim as possible.

diff --git a/Engine/Source/Runtime/Core/Public/Misc/ScopeExit.h b/Engine/Source/Runtime/Core/Public/Misc/ScopeExit.h
--- a/Engine/Source/Runtime/Core/Public/Misc/ScopeExit.h
+++ b/Engine/Source/Runtime/Core/Public/Misc/ScopeExit.h
@@ -2,10 +2,6 @@
 
 #pragma once
 
-#include "CoreTypes.h"
-#include "Templates/UnrealTemplate.h"
-#include "Misc/Optional.h"
-
 namespace ScopeExitSupport
 {
 	/**
@@ -14,35 +10,29 @@ namespace ScopeExitSupport
 	 * RAII class that calls a lambda when it is destroyed.
 	 */
 	template <typename FuncType>
-	class TScopeGuard : public FNoncopyable
+	class TScopeGuard
 	{
+		TScopeGuard(TScopeGuard&&) = delete;
+		TScopeGuard(const TScopeGuard&) = delete;
+		TScopeGuard& operator=(TScopeGuard&&) = delete;
+		TScopeGuard& operator=(const TScopeGuard&) = delete;
+
 	public:
 		// Given a lambda, constructs an RAII scope guard.
 		explicit TScopeGuard(FuncType&& InFunc)
-			: Func(MoveTemp(InFunc))
-		{
-		}
-
-		// This constructor needs to be available for the code to compile.
-		// It will be almost definitely be RVOed out (even in DEBUG).
-		TScopeGuard(TScopeGuard&& Other)
-			: Func(MoveTemp(Other.Func))
+			: Func((FuncType&&)InFunc)
 		{
-			Other.Func.Reset();
 		}
 
-		// Causes
+		// Causes the lambda to be executed.
 		~TScopeGuard()
 		{
-			if (Func.IsSet())
-			{
-				Func.GetValue()();
-			}
+			Func();
 		}
 
 	private:
 		// The lambda to be executed when this guard goes out of scope.
-		TOptional<FuncType> Func;
+		FuncType Func;
 	};
 
 	struct FScopeGuardSyntaxSupport
@@ -50,11 +40,14 @@ namespace ScopeExitSupport
 		template <typename FuncType>
 		TScopeGuard<FuncType> operator+(FuncType&& InFunc)
 		{
-			return TScopeGuard<FuncType>(Forward<FuncType>(InFunc));
+			return TScopeGuard<FuncType>((FuncType&&)InFunc);
 		}
 	};
 }
 
+#define UE_PRIVATE_SCOPE_EXIT_JOIN(A, B) UE_PRIVATE_SCOPE_EXIT_JOIN_INNER(A, B)
+#define UE_PRIVATE_SCOPE_EXIT_JOIN_INNER(A, B) A##B
+
 
 
 /**
@@ -77,4 +70,4 @@ namespace ScopeExitSupport
  *      // * Execution reaches the end of the block.
  *    }
  */
-#define ON_SCOPE_EXIT const auto ANONYMOUS_VARIABLE(ScopeGuard_) = ::ScopeExitSupport::FScopeGuardSyntaxSupport() + [&]()
+#define ON_SCOPE_EXIT const auto UE_PRIVATE_SCOPE_EXIT_JOIN(ScopeGuard_, __LINE__) = ::ScopeExitSupport::FScopeGuardSyntaxSupport() + [&]()
