Fix issue with access to UE4Game directories

#ue5
#android

diff --git a/Engine/Source/Runtime/Core/Private/Android/AndroidPlatformFile.cpp b/Engine/Source/Runtime/Core/Private/Android/AndroidPlatformFile.cpp
--- a/Engine/Source/Runtime/Core/Private/Android/AndroidPlatformFile.cpp
+++ b/Engine/Source/Runtime/Core/Private/Android/AndroidPlatformFile.cpp
@@ -1138,18 +1138,18 @@ public:
 
 		// make sure the base path directory exists (UnrealGame and UnrealGame/ProjectName)
 		FString FileBaseDir = GFilePathBase + FString(FILEBASE_DIRECTORY);
-		mkdir(TCHAR_TO_UTF8(*FileBaseDir), 0766);
-		mkdir(TCHAR_TO_UTF8(*(FileBaseDir + GAndroidProjectName)), 0766);
+		mkdir(TCHAR_TO_UTF8(*FileBaseDir), 0777);
+		mkdir(TCHAR_TO_UTF8(*(FileBaseDir + GAndroidProjectName)), 0777);
 
 		// make sure the log directory exists if override applied
 		if (GOverrideAndroidLogDir)
 		{
 			FString LogBaseDir = GExternalFilePath + FString(FILEBASE_DIRECTORY);
-			mkdir(TCHAR_TO_UTF8(*LogBaseDir), 0766);
-			mkdir(TCHAR_TO_UTF8(*(LogBaseDir + GAndroidProjectName)), 0766);
-			mkdir(TCHAR_TO_UTF8(*(LogBaseDir + GAndroidProjectName + TEXT("/") + GAndroidProjectName)), 0766);
-			mkdir(TCHAR_TO_UTF8(*(LogBaseDir + GAndroidProjectName + TEXT("/") + GAndroidProjectName + TEXT("/Saved"))), 0766);
-			mkdir(TCHAR_TO_UTF8(*(LogBaseDir + GAndroidProjectName + TEXT("/") + GAndroidProjectName + TEXT("/Saved/Logs"))), 0766);
+			mkdir(TCHAR_TO_UTF8(*LogBaseDir), 0777);
+			mkdir(TCHAR_TO_UTF8(*(LogBaseDir + GAndroidProjectName)), 0777);
+			mkdir(TCHAR_TO_UTF8(*(LogBaseDir + GAndroidProjectName + TEXT("/") + GAndroidProjectName)), 0777);
+			mkdir(TCHAR_TO_UTF8(*(LogBaseDir + GAndroidProjectName + TEXT("/") + GAndroidProjectName + TEXT("/Saved"))), 0777);
+			mkdir(TCHAR_TO_UTF8(*(LogBaseDir + GAndroidProjectName + TEXT("/") + GAndroidProjectName + TEXT("/Saved/Logs"))), 0777);
 
 			AndroidLogDir = LogBaseDir + GAndroidProjectName + TEXT("/") + GAndroidProjectName + TEXT("/Saved/Logs/");
 		}
