TaskTraces for the new TaskGraph API
WIP

diff --git a/Engine/Source/Runtime/Core/Private/Async/Fundamental/Scheduler.cpp b/Engine/Source/Runtime/Core/Private/Async/Fundamental/Scheduler.cpp
--- a/Engine/Source/Runtime/Core/Private/Async/Fundamental/Scheduler.cpp
+++ b/Engine/Source/Runtime/Core/Private/Async/Fundamental/Scheduler.cpp
@@ -2,6 +2,7 @@
 
 #include "Async/Fundamental/Scheduler.h"
 #include "Async/Fundamental/Task.h"
+#include "Async/TaskTrace.h"
 #include "Logging/LogMacros.h"
 #include "Misc/ScopeLock.h"
 #include "CoreGlobals.h"
@@ -80,6 +81,8 @@ namespace LowLevelTasks
 
 	void FScheduler::StartWorkers(uint32 NumForegroundWorkers, uint32 NumBackgroundWorkers, EThreadPriority WorkerPriority,  EThreadPriority BackgroundPriority, bool bIsForkable)
 	{
+		TaskTrace::Init();
+
 		if (NumForegroundWorkers == 0 && NumBackgroundWorkers == 0)
 		{
 			NumForegroundWorkers = FMath::Max<int32>(1, FMath::Min<int32>(2, FPlatformMisc::NumberOfWorkerThreadsToSpawn() - 1));
diff --git a/Engine/Source/Runtime/Core/Private/Async/TaskGraph.cpp b/Engine/Source/Runtime/Core/Private/Async/TaskGraph.cpp
--- a/Engine/Source/Runtime/Core/Private/Async/TaskGraph.cpp
+++ b/Engine/Source/Runtime/Core/Private/Async/TaskGraph.cpp
@@ -1831,8 +1831,6 @@ public:
 	FTaskGraphCompatibilityImplementation(int32 InNumWorkerThreads) 
 		: NumWorkerThreads(FForkProcessHelper::IsForkedMultithreadInstance() ? CVar_ForkedProcess_MaxWorkerThreads : InNumWorkerThreads)
 	{
-		TaskTrace::Init();
-
 		if (FTaskGraphInterface::IsMultithread())
 		{
 			if (NumWorkerThreads <= 3)
diff --git a/Engine/Source/Runtime/Core/Private/Async/TaskTrace.cpp b/Engine/Source/Runtime/Core/Private/Async/TaskTrace.cpp
--- a/Engine/Source/Runtime/Core/Private/Async/TaskTrace.cpp
+++ b/Engine/Source/Runtime/Core/Private/Async/TaskTrace.cpp
@@ -25,7 +25,7 @@ namespace TaskTrace
 	UE_TRACE_EVENT_BEGIN(TaskTrace, Launched)
 		UE_TRACE_EVENT_FIELD(uint64, Timestamp)
 		UE_TRACE_EVENT_FIELD(uint32, TaskId)
-		UE_TRACE_EVENT_FIELD(UE::Trace::AnsiString, DebugName)
+		UE_TRACE_EVENT_FIELD(UE::Trace::WideString, DebugName)
 		UE_TRACE_EVENT_FIELD(bool, Tracked)
 		UE_TRACE_EVENT_FIELD(int32, ThreadToExecuteOn)
 	UE_TRACE_EVENT_END()
@@ -100,7 +100,7 @@ namespace TaskTrace
 			<< Created.TaskId(TaskId);
 	}
 
-	void Launched(FId TaskId, const char* DebugName, bool bTracked, ENamedThreads::Type ThreadToExecuteOn)
+	void Launched(FId TaskId, const TCHAR* DebugName, bool bTracked, ENamedThreads::Type ThreadToExecuteOn)
 	{
 		check(bGTaskTraceInitialized);
 		check(TaskId != InvalidId);
@@ -108,7 +108,7 @@ namespace TaskTrace
 		UE_TRACE_LOG(TaskTrace, Launched, TaskChannel)
 			<< Launched.Timestamp(FPlatformTime::Cycles64())
 			<< Launched.TaskId(TaskId)
-			<< Launched.DebugName(DebugName != nullptr ? DebugName : "")
+			<< Launched.DebugName(DebugName != nullptr ? DebugName : TEXT(""))
 			<< Launched.Tracked(bTracked)
 			<< Launched.ThreadToExecuteOn(ThreadToExecuteOn);
 	}
diff --git a/Engine/Source/Runtime/Core/Public/Async/TaskTrace.h b/Engine/Source/Runtime/Core/Public/Async/TaskTrace.h
--- a/Engine/Source/Runtime/Core/Public/Async/TaskTrace.h
+++ b/Engine/Source/Runtime/Core/Public/Async/TaskTrace.h
@@ -31,7 +31,7 @@ namespace TaskTrace
 
 	void CORE_API Init();
 	void CORE_API Created(FId TaskId); // optional, used only if a task was created but not launched immediately
-	void CORE_API Launched(FId TaskId, const char* DebugName, bool bTracked, ENamedThreads::Type ThreadToExecuteOn);
+	void CORE_API Launched(FId TaskId, const TCHAR* DebugName, bool bTracked, ENamedThreads::Type ThreadToExecuteOn);
 	void CORE_API Scheduled(FId TaskId);
 	void CORE_API SubsequentAdded(FId TaskId, FId SubsequentId);
 	void CORE_API Started(FId TaskId);
@@ -50,7 +50,7 @@ namespace TaskTrace
 	inline FId GenerateTaskId() { return InvalidId; }
 	inline void Init() {}
 	inline void Created(FId TaskId) {}
-	inline void Launched(FId TaskId, const char* DebugName, bool bTracked, ENamedThreads::Type ThreadToExecuteOn) {}
+	inline void Launched(FId TaskId, const TCHAR* DebugName, bool bTracked, ENamedThreads::Type ThreadToExecuteOn) {}
 	inline void Scheduled(FId TaskId) {}
 	inline void SubsequentAdded(FId TaskId, FId SubsequentId) {}
 	inline void Started(FId TaskId) {}
diff --git a/Engine/Source/Runtime/Core/Public/Tasks/TaskPrivate.h b/Engine/Source/Runtime/Core/Public/Tasks/TaskPrivate.h
--- a/Engine/Source/Runtime/Core/Public/Tasks/TaskPrivate.h
+++ b/Engine/Source/Runtime/Core/Public/Tasks/TaskPrivate.h
@@ -4,6 +4,7 @@
 
 #include "Async/Fundamental/Task.h"
 #include "Async/Fundamental/Scheduler.h"
+#include "Async/TaskTrace.h"
 #include "Templates/Invoke.h"
 #include "Templates/TypeCompatibleBytes.h"
 #include "Misc/Timeout.h"
@@ -72,15 +73,24 @@ namespace UE { namespace Tasks
 						Deleter = Forward<DeleterType>(Deleter)
 					]() mutable
 					{
-						StartPipeExecution();
-						Invoke(TaskBody);
-						FinishPipeExecution();
+						TaskTrace::Started(GetTraceId());
+						{
+							TRACE_CPUPROFILER_EVENT_SCOPE(ExecuteTask);
+							StartPipeExecution();
+							Invoke(TaskBody);
+							FinishPipeExecution();
+						}
+						TaskTrace::Finished(GetTraceId());
+
 						FTaskBase* Subsequent = CloseAndReturnSubsequent();
 						if (Subsequent != nullptr)
 						{
 							check(TryPushIntoPipe());
+							TaskTrace::Scheduled(Subsequent->GetTraceId());
 							LowLevelTasks::TryLaunch(Subsequent->LowLevelTask, LowLevelTasks::EQueuePreference::DefaultPreference, /*bWakeUpWorker =*/ false);
 						}
+
+						TaskTrace::Completed(GetTraceId());
 					} // Continuation
 				);
 			}
@@ -102,8 +112,11 @@ namespace UE { namespace Tasks
 			// scheduled immediately
 			bool TryLaunch()
 			{
+				TaskTrace::Launched(GetTraceId(), LowLevelTask.GetDebugName(), true, (ENamedThreads::Type)255);
+
 				if (TryPushIntoPipe())
 				{
+					TaskTrace::Scheduled(GetTraceId());
 					// scheduler's reference was accounted on task creation
 					return LowLevelTasks::TryLaunch(LowLevelTask);
 				}
@@ -120,6 +133,8 @@ namespace UE { namespace Tasks
 				checkf(CurrentState == GetEmptyOpenState() || CurrentState == GetClosedState(), TEXT("only a single subsequent can be set"));
 #endif
 
+				TaskTrace::SubsequentAdded(GetTraceId(), Task.GetTraceId());
+
 				FTaskBase* EmptyNotCompleted = GetEmptyOpenState();
 				FTaskBase* NotCompletedWithTask = &Task;
 				return SubsequentAndState.compare_exchange_strong(EmptyNotCompleted, NotCompletedWithTask, std::memory_order_release, std::memory_order_relaxed);
@@ -178,6 +193,15 @@ namespace UE { namespace Tasks
 				return LowLevelTask.IsCompleted(std::memory_order_relaxed);
 			}
 
+			TaskTrace::FId GetTraceId() const
+			{
+#if UE_TASK_TRACE_ENABLED
+				return TraceId;
+#else
+				return TaskTrace::InvalidId;
+#endif
+			}
+
 		private:
 			// checks if the task is ready to be launched by trying to push it into the pipe
 			bool TryPushIntoPipe()
@@ -235,6 +259,10 @@ namespace UE { namespace Tasks
 			// can be "empty and open" (`nullptr`), "have subsequent and open" (subsequent ptr) and closed (`nullptr` with the most significant
 			// bit set)
 			std::atomic<FTaskBase*> SubsequentAndState{ GetEmptyOpenState() };
+
+#if UE_TASK_TRACE_ENABLED
+			TaskTrace::FId TraceId = TaskTrace::GenerateTaskId();
+#endif
 		};
 
 		// Extends FTaskBase by supporting execution result.
