remove PLATFORM_XBOXONE from FramePro code.

#rnx

diff --git a/Engine/Source/Runtime/Core/Private/FramePro/FrameProUE4.cpp b/Engine/Source/Runtime/Core/Private/FramePro/FrameProUE4.cpp
--- a/Engine/Source/Runtime/Core/Private/FramePro/FrameProUE4.cpp
+++ b/Engine/Source/Runtime/Core/Private/FramePro/FrameProUE4.cpp
@@ -149,10 +149,10 @@
 		//------------------------------------------------------------------------
 		Platform::Enum Platform::GetPlatformEnum()
 		{
-			#if defined(FRAMEPRO_UNREAL_PLATFORM) //@EPIC: allow external definition
-				return FRAMEPRO_UNREAL_PLATFORM;  //@EPIC end
-			#elif PLATFORM_XBOXONE
-				return Platform::XBoxOne;
+			//@EPIC: allow external definition + XBOXONE removed
+			#if defined(FRAMEPRO_UNREAL_PLATFORM)
+				return FRAMEPRO_UNREAL_PLATFORM;
+			//@EPIC end
 			#elif PLATFORM_WINDOWS
 				return Platform::Windows;
 			#elif PLATFORM_LINUX
diff --git a/Engine/Source/Runtime/Core/Public/FramePro/FrameProUE4.h b/Engine/Source/Runtime/Core/Public/FramePro/FrameProUE4.h
--- a/Engine/Source/Runtime/Core/Public/FramePro/FrameProUE4.h
+++ b/Engine/Source/Runtime/Core/Public/FramePro/FrameProUE4.h
@@ -44,20 +44,16 @@
 		#define FRAMEPRO_API CORE_API
 
 		// Windows or Linux based platform
-#if !defined(FRAMEPRO_WIN_BASED_PLATFORM) //@EPIC begin - allow external definition
-		#define FRAMEPRO_WIN_BASED_PLATFORM (PLATFORM_WINDOWS || PLATFORM_XBOXONE)
+#if !defined(FRAMEPRO_WIN_BASED_PLATFORM) //@EPIC begin - allow external definition + XBOXONE removed
+		#define FRAMEPRO_WIN_BASED_PLATFORM (PLATFORM_WINDOWS)
 #endif //@EPIC: end
 		#define FRAMEPRO_LINUX_BASED_PLATFORM (!FRAMEPRO_WIN_BASED_PLATFORM && !PLATFORM_SWITCH)
 
 		#define FRAMEPRO_USE_TLS_SLOTS 1
 
 		// Port
-#if !defined(FRAMEPRO_PORT) //@EPIC begin - allow external definition
-		#if PLATFORM_XBOXONE
-			#define FRAMEPRO_PORT "4420"
-		#else
-			#define FRAMEPRO_PORT "8428"
-		#endif
+#if !defined(FRAMEPRO_PORT) //@EPIC begin - allow external definition + XBOXONE removed
+		#define FRAMEPRO_PORT "8428"
 #endif //@EPIC: end
 
 		// x64 or x32
