Tasks: inline task priority. There's a special kind of tasks that serve as synchronisation primitives, like FTaskEvent, tasks that are used as a join point of other tasks or tasks that are used to wait for other tasks. Such tasks are lightweight and latency-sensitive. Inline task priority allows these tasks to skip scheduling and be executed "inline" as soon as they are unlocked.

diff --git a/Engine/Source/Runtime/Core/Public/Tasks/TaskPrivate.h b/Engine/Source/Runtime/Core/Public/Tasks/TaskPrivate.h
--- a/Engine/Source/Runtime/Core/Public/Tasks/TaskPrivate.h
+++ b/Engine/Source/Runtime/Core/Public/Tasks/TaskPrivate.h
@@ -54,6 +54,10 @@ namespace UE { namespace Tasks
 			}
 
 		public:
+			// a special internal task priority for "inline" task execution - a task is executed as soon as it's launched and has no 
+			// pending dependencies, "inline", w/o scheduling
+			static const ETaskPriority InlineTaskPriority{ ETaskPriority::Count };
+			
 			// initialises the task but doesn't launches it
 			template<typename TaskBodyType, typename DeleterType>
 			void Init(const TCHAR* DebugName, TaskBodyType&& TaskBody, DeleterType&& Deleter, LowLevelTasks::ETaskPriority Priority)
@@ -114,13 +118,24 @@ namespace UE { namespace Tasks
 			{
 				TaskTrace::Launched(GetTraceId(), LowLevelTask.GetDebugName(), true, (ENamedThreads::Type)255);
 
-				if (TryPushIntoPipe())
+				if (!TryPushIntoPipe())
 				{
-					TaskTrace::Scheduled(GetTraceId());
-					// scheduler's reference was accounted on task creation
-					return LowLevelTasks::TryLaunch(LowLevelTask);
+					return false; // pipe is blocked
 				}
-				return false;
+
+				TaskTrace::Scheduled(GetTraceId());
+
+				// "inline" tasks are not scheduled but executed as soon as they are unlocked
+				if (LowLevelTask.GetPriority() == InlineTaskPriority)
+				{
+					// the task is not scheduled yet, so successful retraction is guaranted
+					verify(LowLevelTask.TryCancel()); // this will cancel the task for the scheduler and execute LowLevelTask's continuation, 
+					// which is actually task execution
+					return true;
+				}
+
+				// schedule
+				return LowLevelTasks::TryLaunch(LowLevelTask);
 			}
 
 			// sets a subsequent task that will be launched automatically when this task is done
