Fix for TTuple<T0,T1>::Get<Type>() while in static analysis mode.

diff --git a/Engine/Source/Runtime/Core/Public/Templates/Tuple.h b/Engine/Source/Runtime/Core/Public/Templates/Tuple.h
--- a/Engine/Source/Runtime/Core/Public/Templates/Tuple.h
+++ b/Engine/Source/Runtime/Core/Public/Templates/Tuple.h
@@ -202,6 +202,26 @@ namespace UE4Tuple_Private
 		}
 	};
 
+#if UE_TUPLE_STATIC_ANALYSIS_WORKAROUND
+	template <typename Type>
+	struct TTupleElementGetterByType<Type, 2>
+	{
+		template <typename TupleType>
+		static FORCEINLINE decltype(auto) Get(TupleType&& Tuple)
+		{
+			if constexpr (std::is_same_v<Type, decltype(Tuple.Key)>)
+			{
+				return TTupleElementGetterByIndex<0, 2>::Get(Forward<TupleType>(Tuple));
+			}
+			else
+			{
+				static_assert(std::is_same_v<Type, decltype(Tuple.Value)>, "This has to be true - check the call site");
+				return TTupleElementGetterByIndex<1, 2>::Get(Forward<TupleType>(Tuple));
+			}
+		}
+	};
+#endif
+
 	template <uint32 ArgCount, uint32 ArgToCompare>
 	struct FEqualityHelper
 	{
