tasks dependencies:
* an ability to specify task prerequisites, either TTask or FTaskEvent
* redesigned FTaskEvent so it can be used as a prerequisite
* a helper function `Tasks::Prerequisites()` to specify multiple prerequisites
* added `Tasks::Wait()` function to wait for multiple tasks
* unit- and perf-tests

diff --git a/Engine/Source/Runtime/Core/Private/Tasks/Pipe.cpp b/Engine/Source/Runtime/Core/Private/Tasks/Pipe.cpp
--- a/Engine/Source/Runtime/Core/Private/Tasks/Pipe.cpp
+++ b/Engine/Source/Runtime/Core/Private/Tasks/Pipe.cpp
@@ -17,7 +17,7 @@ namespace UE { namespace Tasks
 		Private::FTaskBase* LastTask_Local = LastTask.exchange(&Task, std::memory_order_release);
 		checkf(LastTask_Local != &Task, TEXT("Dependency cycle: adding itself as a prerequisite (or use after destruction)"));
 
-		return LastTask_Local == nullptr || !LastTask_Local->SetSubsequent(Task);
+		return LastTask_Local == nullptr || !LastTask_Local->AddSubsequent(Task);
 	}
 
 	void FPipe::ClearTask(Private::FTaskBase& Task)
diff --git a/Engine/Source/Runtime/Core/Private/Tasks/TaskPrivate.cpp b/Engine/Source/Runtime/Core/Private/Tasks/TaskPrivate.cpp
--- a/Engine/Source/Runtime/Core/Private/Tasks/TaskPrivate.cpp
+++ b/Engine/Source/Runtime/Core/Private/Tasks/TaskPrivate.cpp
@@ -22,7 +22,6 @@ namespace UE { namespace Tasks { namespace Private
 	{
 		if (GetPipe() != nullptr)
 		{
-			check(bPushedIntoPipe);
 			GetPipe()->ExecutionFinished();
 			GetPipe()->ClearTask(*this);
 		}
diff --git a/Engine/Source/Runtime/Core/Private/Tests/Tasks/TasksTest.cpp b/Engine/Source/Runtime/Core/Private/Tests/Tasks/TasksTest.cpp
--- a/Engine/Source/Runtime/Core/Private/Tests/Tasks/TasksTest.cpp
+++ b/Engine/Source/Runtime/Core/Private/Tests/Tasks/TasksTest.cpp
@@ -95,7 +95,7 @@ namespace UE { namespace TasksTests
 		}
 
 		{	// basic example, fire and forget a high-pri task
-			Tasks::Launch(
+			Launch(
 				UE_SOURCE_LOCATION, // debug name
 				[] {}, // task body
 				Tasks::ETaskPriority::High /* task priority, `Normal` by default */
@@ -103,12 +103,12 @@ namespace UE { namespace TasksTests
 		}
 
 		{	// launch a task and wait till it's executed
-			Tasks::Launch(UE_SOURCE_LOCATION, [] {}).Wait();
+			Launch(UE_SOURCE_LOCATION, [] {}).Wait();
 		}
 
 		{	// FTaskEvent
-			FTaskEvent Event;
-			check(!Event.IsTriggered());
+			FTaskEvent Event{ UE_SOURCE_LOCATION };
+			check(!Event.IsCompleted());
 
 			// check that waiting blocks
 			FTask Task = Launch(UE_SOURCE_LOCATION, [&Event] { Event.Wait(); });
@@ -116,7 +116,7 @@ namespace UE { namespace TasksTests
 			check(!Task.IsCompleted());
 
 			Event.Trigger();
-			check(Event.IsTriggered());
+			check(Event.IsCompleted());
 			verify(Event.Wait(FTimespan::Zero()));
 		}
 
@@ -130,7 +130,7 @@ namespace UE { namespace TasksTests
 		}
 
 		{	// same but using `FTaskEvent`
-			FTaskEvent Event;
+			FTaskEvent Event{ UE_SOURCE_LOCATION };
 			FTask Task = Launch(UE_SOURCE_LOCATION, [&Event] { Event.Wait(); });
 			ensure(!Task.Wait(FTimespan::FromMilliseconds(100)));
 			Event.Trigger();
@@ -319,7 +319,7 @@ namespace UE { namespace TasksTests
 		{	// an example of blocking a pipe
 			FPipe Pipe{ UE_SOURCE_LOCATION };
 			std::atomic<bool> bBlocked;
-			FTaskEvent Event;
+			FTaskEvent Event{ UE_SOURCE_LOCATION };
 			FTask Task = Pipe.Launch(UE_SOURCE_LOCATION,
 				[&bBlocked, &Event]
 				{
@@ -469,6 +469,176 @@ namespace UE { namespace TasksTests
 
 		return true;
 	}
+
+	IMPLEMENT_SIMPLE_AUTOMATION_TEST(FTasksDependenciesTest, "System.Core.Tasks.Dependencies", EAutomationTestFlags::ApplicationContextMask | EAutomationTestFlags::EngineFilter);
+
+	template<uint64 NumBranches, uint64 NumLoops, uint64 NumTasks>
+	void DependenciesPerfTest()
+	{
+		auto Branch = []
+		{
+			FTask Joiner;
+			for (uint64 LoopIndex = 0; LoopIndex != NumLoops; ++LoopIndex)
+			{
+				TArray<FTask> Tasks;
+				Tasks.Reserve(NumTasks);
+				for (uint64 TaskIndex = 0; TaskIndex != NumTasks; ++TaskIndex)
+				{
+					if (Joiner.IsValid())
+					{
+						Tasks.Add(Launch(UE_SOURCE_LOCATION, [] { /*FPlatformProcess::YieldCycles(100000);*/ }, Joiner));
+					}
+					else
+					{
+						Tasks.Add(Launch(UE_SOURCE_LOCATION, [] { /*FPlatformProcess::YieldCycles(100000);*/ }));
+					}
+				}
+				Joiner = Launch(UE_SOURCE_LOCATION, [] {}, Tasks);
+			}
+			return Joiner;
+		};
+
+		TArray<TTask<FTask>> Branches;
+		Branches.Reserve(NumBranches);
+		for (uint64 BranchIndex = 0; BranchIndex != NumBranches; ++BranchIndex)
+		{
+			Branches.Add(Launch(UE_SOURCE_LOCATION, Branch));
+		}
+		TArray<FTask> BranchTasks;
+		BranchTasks.Reserve(NumBranches);
+		for (TTask<FTask>& Task : Branches)
+		{
+			BranchTasks.Add(Task.GetResult());
+		}
+		Wait(Branches);
+	}
+
+	bool FTasksDependenciesTest::RunTest(const FString& Parameters)
+	{
+		{	// a task is not executed until its prerequisite (FTaskEvent) is completed
+			FTaskEvent Prereq{ UE_SOURCE_LOCATION };
+
+			FTask Task{ Launch(UE_SOURCE_LOCATION, [] {}, Prereq) };
+			verify(!Task.Wait(FTimespan::FromMilliseconds(10)));
+
+			Prereq.Trigger();
+			Task.Wait();
+		}
+
+		{	// a task is not executed until its prerequisite (FTaskEvent) is completed. with explicit task priority
+			FTaskEvent Prereq{ UE_SOURCE_LOCATION };
+
+			FTask Task{ Launch(UE_SOURCE_LOCATION, [] {}, Prereq, ETaskPriority::Normal) };
+			verify(!Task.Wait(FTimespan::FromMilliseconds(10)));
+
+			Prereq.Trigger();
+			Task.Wait();
+		}
+
+		{	// a task is not executed until its prerequisite (FTask) is completed
+			FTaskEvent Event{ UE_SOURCE_LOCATION };
+			FTask Prereq{ Launch(UE_SOURCE_LOCATION, [&Event] { Event.Wait(); }) };
+			FTask Task{ Launch(UE_SOURCE_LOCATION, [] {}, Prereq) };
+			verify(!Task.Wait(FTimespan::FromMilliseconds(10)));
+
+			Event.Trigger();
+			Task.Wait();
+		}
+
+		{	// compilation test of an iterable collection as prerequisites
+			TArray<FTask> Prereqs{ Launch(UE_SOURCE_LOCATION, [] {}), FTaskEvent{ UE_SOURCE_LOCATION } };
+
+			FTask Task{ Launch(UE_SOURCE_LOCATION, [] {}, Prereqs) };
+			((FTaskEvent&)Prereqs[1]).Trigger();
+			Task.Wait();
+		}
+
+		{	// compilation test of an initializer list as prerequisites
+			auto Prereqs = { Launch(UE_SOURCE_LOCATION, [] {}), Launch(UE_SOURCE_LOCATION, [] {}) };
+			Launch(UE_SOURCE_LOCATION, [] {}, Prereqs).Wait();
+			// the following line doesn't compile because the compiler can't deduce the initializer list type
+			//Launch(UE_SOURCE_LOCATION, [] {}, { Launch(UE_SOURCE_LOCATION, [] {}), Launch(UE_SOURCE_LOCATION, [] {}) }).Wait();
+		}
+
+		{	// a task is not executed until all its prerequisites (FTask and FTaskEvent instances) are completed
+			FTaskEvent Prereq1{ UE_SOURCE_LOCATION };
+			FTaskEvent Event{ UE_SOURCE_LOCATION };
+			FTask Prereq2{ Launch(UE_SOURCE_LOCATION, [&Event] { Event.Wait(); }) };
+
+			TTask<void> Task{ Launch(UE_SOURCE_LOCATION, [] {}, Prerequisites(Prereq1, Prereq2)) };
+			verify(!Task.Wait(FTimespan::FromMilliseconds(10)));
+
+			Prereq1.Trigger();
+			verify(!Task.Wait(FTimespan::FromMilliseconds(10)));
+
+			Event.Trigger();
+			Task.Wait();
+		}
+
+		{	// a task is not executed until all its prerequisites (FTask and FTaskEvent instances) are completed. with explicit task priority
+			FTaskEvent Prereq1{ UE_SOURCE_LOCATION };
+			FTaskEvent Event{ UE_SOURCE_LOCATION };
+			FTask Prereq2{ Launch(UE_SOURCE_LOCATION, [&Event] { Event.Wait(); }) };
+			TArray<TTask<void>> Prereqs{ Prereq1, Prereq2 }; // to check if a random iterable container works as a prerequisite collection
+
+			TTask<void> Task{ Launch(UE_SOURCE_LOCATION, [] {}, Prereqs, ETaskPriority::Normal) };
+			verify(!Task.Wait(FTimespan::FromMilliseconds(10)));
+
+			Prereq1.Trigger();
+			verify(!Task.Wait(FTimespan::FromMilliseconds(10)));
+
+			Event.Trigger();
+			Task.Wait();
+		}
+
+		{	// a piped task blocked by a prerequisites doesn't block the pipe
+			FPipe Pipe{ UE_SOURCE_LOCATION };
+			FTaskEvent Prereq{ UE_SOURCE_LOCATION };
+
+			FTask Task1{ Pipe.Launch(UE_SOURCE_LOCATION, [] {}, Prereq) };
+			verify(!Task1.Wait(FTimespan::FromMilliseconds(10)));
+
+			FTask Task2{ Pipe.Launch(UE_SOURCE_LOCATION, [] {}) };
+			Task2.Wait();
+
+			Prereq.Trigger();
+			Task1.Wait();
+		}
+
+		{	// a piped task with multiple prerequisites
+			FPipe Pipe{ UE_SOURCE_LOCATION };
+			FTaskEvent Prereq1{ UE_SOURCE_LOCATION };
+			FTaskEvent Event{ UE_SOURCE_LOCATION };
+			FTask Prereq2{ Launch(UE_SOURCE_LOCATION, [&Event] { Event.Wait(); }) };
+
+			FTask Task{ Pipe.Launch(UE_SOURCE_LOCATION, [] {}, Prerequisites(Prereq1, Prereq2)) };
+			verify(!Task.Wait(FTimespan::FromMilliseconds(10)));
+
+			Prereq1.Trigger();
+			verify(!Task.Wait(FTimespan::FromMilliseconds(10)));
+			Event.Trigger();
+			Task.Wait();
+		}
+
+		{	// a piped task with multiple prerequisites. with explicit task priority
+			FPipe Pipe{ UE_SOURCE_LOCATION };
+			FTaskEvent Prereq1{ UE_SOURCE_LOCATION };
+			FTaskEvent Event{ UE_SOURCE_LOCATION };
+			FTask Prereq2{ Launch(UE_SOURCE_LOCATION, [&Event] { Event.Wait(); }) };
+
+			FTask Task{ Pipe.Launch(UE_SOURCE_LOCATION, [] {}, Prerequisites(Prereq1, Prereq2), ETaskPriority::Normal) };
+			verify(!Task.Wait(FTimespan::FromMilliseconds(10)));
+
+			Prereq1.Trigger();
+			verify(!Task.Wait(FTimespan::FromMilliseconds(10)));
+			Event.Trigger();
+			Task.Wait();
+		}
+
+		UE_BENCHMARK(3, DependenciesPerfTest<200, 50, 1000>);
+
+		return true;
+	}
 }}
 
 #endif // WITH_DEV_AUTOMATION_TESTS
diff --git a/Engine/Source/Runtime/Core/Public/Tasks/Pipe.h b/Engine/Source/Runtime/Core/Public/Tasks/Pipe.h
--- a/Engine/Source/Runtime/Core/Public/Tasks/Pipe.h
+++ b/Engine/Source/Runtime/Core/Public/Tasks/Pipe.h
@@ -62,6 +62,24 @@ namespace UE { namespace Tasks
 			return TTask<FResult>{ Task };
 		}
 
+		// launches a task in the pipe, with multiple prerequisites that must be completed before the task is scheduled
+		// @param InDebugName helps to identify the task in debugger and profiler
+		// @param TaskBody a callable with no parameters whose return value is ignored, usually a lambda but can be also a functor object 
+		// or a pointer to a function. TaskBody can return results.
+		// @Priority - task priority, can affect task scheduling once it's passed the pipe
+		// @return Task instance that can be used to wait for task completion or to obtain the result of task execution
+		template<typename TaskBodyType, typename PrerequisitesCollectionType>
+		TTask<TInvokeResult_T<TaskBodyType>> Launch(const TCHAR* InDebugName, TaskBodyType&& TaskBody, PrerequisitesCollectionType&& Prerequisites, LowLevelTasks::ETaskPriority Priority = LowLevelTasks::ETaskPriority::Default)
+		{
+			using FResult = TInvokeResult_T<TaskBodyType>;
+			Private::TTaskWithResult<FResult>* Task = new Private::TTaskWithResult<FResult>;
+			Task->Init(InDebugName, Forward<TaskBodyType>(TaskBody), Priority);
+			Task->AddPrerequisites(Forward<PrerequisitesCollectionType>(Prerequisites));
+			Task->SetPipe(*this);
+			Task->TryLaunch();
+			return TTask<FResult>{ Task };
+		}
+
 		// checks if pipe's task is being executed by the current thread. Allows to check if accessing a resource protected by a pipe
 		// is thread-safe
 		CORE_API bool IsInContext() const;
diff --git a/Engine/Source/Runtime/Core/Public/Tasks/Task.h b/Engine/Source/Runtime/Core/Public/Tasks/Task.h
--- a/Engine/Source/Runtime/Core/Public/Tasks/Task.h
+++ b/Engine/Source/Runtime/Core/Public/Tasks/Task.h
@@ -5,6 +5,7 @@
 #include "Tasks/TaskPrivate.h"
 #include "Async/Fundamental/Task.h"
 #include "Templates/RefCounting.h"
+#include "Containers/StaticArray.h"
 #include "CoreTypes.h"
 
 namespace UE { namespace Tasks
@@ -13,6 +14,11 @@ namespace UE { namespace Tasks
 	template<typename ResultType>
 	class TTaskBase
 	{
+		friend Private::FTaskBase;
+
+		template<typename... TaskTypes>
+		friend class TPrerequisites;
+
 	protected:
 		TTaskBase() = default;
 
@@ -83,6 +89,9 @@ namespace UE { namespace Tasks
 		template<typename TaskBodyType>
 		friend TTask<TInvokeResult_T<TaskBodyType>> Launch(const TCHAR* DebugName, TaskBodyType&& TaskBody, LowLevelTasks::ETaskPriority Priority/* = LowLevelTasks::ETaskPriority::Normal*/);
 
+		template<typename TaskBodyType, typename PrerequisitesCollectionType>
+		friend TTask<TInvokeResult_T<TaskBodyType>> Launch(const TCHAR* DebugName, TaskBodyType&& TaskBody, PrerequisitesCollectionType&& Prerequisites, LowLevelTasks::ETaskPriority Priority/* = LowLevelTasks::ETaskPriority::Normal*/);
+
 		friend FPipe;
 
 		// private constructor, valid instances can be created only by launching (see friends)
@@ -91,6 +100,8 @@ namespace UE { namespace Tasks
 		{}
 	};
 
+	class FTaskEvent;
+
 	template<>
 	class TTask<void> : public TTaskBase<void>
 	{
@@ -107,7 +118,11 @@ namespace UE { namespace Tasks
 		template<typename TaskBodyType>
 		friend TTask<TInvokeResult_T<TaskBodyType>> Launch(const TCHAR* DebugName, TaskBodyType&& TaskBody, LowLevelTasks::ETaskPriority Priority/* = LowLevelTasks::ETaskPriority::Normal*/);
 
+		template<typename TaskBodyType, typename PrerequisitesCollectionType>
+		friend TTask<TInvokeResult_T<TaskBodyType>> Launch(const TCHAR* DebugName, TaskBodyType&& TaskBody, PrerequisitesCollectionType&& Prerequisites, LowLevelTasks::ETaskPriority Priority/* = LowLevelTasks::ETaskPriority::Normal*/);
+
 		friend FPipe;
+		friend FTaskEvent;
 
 		// private constructor, valid instances can be created only by launching (see friends)
 		explicit TTask(Private::TTaskWithResult<void>* Other)
@@ -116,54 +131,19 @@ namespace UE { namespace Tasks
 	};
 
 	// a synchronisation primitive similar to `FEvent` that uses "busy waiting" - executing tasks while waiting
-	class FTaskEvent
+	class FTaskEvent : public TTask<void>
 	{
 	public:
-		UE_NONCOPYABLE(FTaskEvent)
-
-		FTaskEvent() = default;
-
-		void Trigger()
+		explicit FTaskEvent(const TCHAR* DebugName)
+			: TTask<void>(new Private::TTaskWithResult<void>)
 		{
-			bTriggered.store(true, std::memory_order_release);
+			Pimpl->Init(DebugName, [] {}, Private::FTaskBase::InlineTaskPriority);
 		}
 
-		bool IsTriggered() const
-		{
-			return bTriggered.load(std::memory_order_acquire);
-		}
-
-		// waits until the event is triggered
-		void Wait()
-		{
-			LowLevelTasks::BusyWaitUntil([this] { return bTriggered.load(std::memory_order_acquire); });
-		}
-
-		// waits until the event is triggered or the waiting timed out.
-		// the call can return much later than the given timeout
-		// @return true if the event was triggered
-		bool Wait(FTimespan InTimeout)
-		{
-			LowLevelTasks::BusyWaitUntil(
-				[this, Timeout = FTimeout{ InTimeout }] { return bTriggered.load(std::memory_order_acquire) || Timeout; }
-			);
-			return bTriggered.load(std::memory_order_relaxed);
-		}
-
-		// waits until the event is triggered or the given condition returns true.
-		// the call can return much later than the condition has become true
-		// @return true if the event was triggered
-		template<typename ConditionType>
-		bool Wait(ConditionType&& Condition)
+		void Trigger()
 		{
-			LowLevelTasks::BusyWaitUntil(
-				[this, Condition = Forward<ConditionType>(Condition)] { return bTriggered.load(std::memory_order_acquire) || Condition(); }
-			);
-			return bTriggered.load(std::memory_order_relaxed);
+			verify(Pimpl->TryLaunch());
 		}
-
-	private:
-		std::atomic<bool> bTriggered{ false };
 	};
 
 	// launches a task for asynchronous execution
@@ -173,17 +153,87 @@ namespace UE { namespace Tasks
 	// @return a trivially relocatable instance that can be used to wait for task completion or to obtain task execution result
 	template<typename TaskBodyType>
 	TTask<TInvokeResult_T<TaskBodyType>> Launch(
-		const TCHAR* DebugName, 
-		TaskBodyType&& TaskBody, 
+		const TCHAR* DebugName,
+		TaskBodyType&& TaskBody,
+		LowLevelTasks::ETaskPriority Priority/* = LowLevelTasks::ETaskPriority::Normal*/
+	)
+	{
+		using FResult = TInvokeResult_T<TaskBodyType>;
+		auto* Task{ new Private::TTaskWithResult<FResult> };
+		Task->Init(DebugName, Forward<TaskBodyType>(TaskBody), Priority);
+		Task->TryLaunch();
+		return TTask<FResult>{ Task };
+	}
+
+	// launches a task for asynchronous execution, with prerequisites that must be completed before the task is scheduled
+	// @param DebugName - a unique name for task identification in debugger and profiler, is compiled out in test/shipping builds
+	// @param TaskBody - a functor that will be executed asynchronously
+	// @param Prerequisites - tasks or task events that must be completed before the task being launched can be scheduled, accepts any 
+	// iterable collection (.begin()/.end()), `Tasks::Prerequisites()` helper is recommended to create such collection on the fly
+	// @param Priority - task priority that affects when the task will be executed
+	// @return a trivially relocatable instance that can be used to wait for task completion or to obtain task execution result
+	template<typename TaskBodyType, typename PrerequisitesCollectionType>
+	TTask<TInvokeResult_T<TaskBodyType>> Launch(
+		const TCHAR* DebugName,
+		TaskBodyType&& TaskBody,
+		PrerequisitesCollectionType&& Prerequisites,
 		LowLevelTasks::ETaskPriority Priority/* = LowLevelTasks::ETaskPriority::Normal*/
 	)
 	{
 		using FResult = TInvokeResult_T<TaskBodyType>;
 		auto* Task{ new Private::TTaskWithResult<FResult> };
 		Task->Init(DebugName, Forward<TaskBodyType>(TaskBody), Priority);
+		Task->AddPrerequisites(Forward<PrerequisitesCollectionType>(Prerequisites));
 		Task->TryLaunch();
 		return TTask<FResult>{ Task };
 	}
 
+	template<typename TaskCollectionType>
+	void Wait(TaskCollectionType&& Tasks)
+	{
+		for (auto& Task : Tasks)
+		{
+			Task.Wait();
+		}
+	}
+
 	using FTask = TTask<void>;
-}}
\ No newline at end of file
+
+	// a convenient proxy collection for specifying task prerequisites that can include both tasks and task events
+	// usage: Launch(UE_SOURCE_LOCATION, [] {}, Prerequisites(Task1, Task2, TaskEvent1, ...));
+	template<typename... TaskTypes>
+	class TPrerequisites : public TStaticArray<Private::FTaskBase*, sizeof...(TaskTypes)>
+	{
+	public:
+		TPrerequisites(TaskTypes&&... Tasks)
+		{
+			Fill(0, Forward<TaskTypes>(Tasks)...);
+		}
+
+	private:
+		template<typename FirstTaskType, typename... OtherTaskTypes>
+		void Fill(uint32 Index, FirstTaskType&& FirstTask, OtherTaskTypes&&... OtherTasks)
+		{
+			(*this)[Index] = FirstTask.Pimpl.GetReference();
+			Fill(Index + 1, Forward<OtherTaskTypes>(OtherTasks)...);
+		}
+
+		template<typename TaskType>
+		void Fill(uint32 Index, TaskType&& Task)
+		{
+			(*this)[Index] = Task.Pimpl.GetReference();
+		}
+	};
+
+	template<typename... TaskTypes>
+	TPrerequisites<TaskTypes...> Prerequisites(TaskTypes&... Tasks)
+	{
+		return TPrerequisites<TaskTypes...>{ Forward<TaskTypes>(Tasks)... };
+	}
+}}
+
+template <typename... TaskTypes>
+struct TIsContiguousContainer<UE::Tasks::TPrerequisites<TaskTypes...>>
+{
+	static constexpr bool Value = true;
+};
diff --git a/Engine/Source/Runtime/Core/Public/Tasks/TaskPrivate.h b/Engine/Source/Runtime/Core/Public/Tasks/TaskPrivate.h
--- a/Engine/Source/Runtime/Core/Public/Tasks/TaskPrivate.h
+++ b/Engine/Source/Runtime/Core/Public/Tasks/TaskPrivate.h
@@ -8,6 +8,8 @@
 #include "Templates/Invoke.h"
 #include "Templates/TypeCompatibleBytes.h"
 #include "Misc/Timeout.h"
+#include "Containers/ClosableMpscQueue.h"
+#include "Templates/UnrealTemplate.h"
 #include "CoreTypes.h"
 
 #include <atomic>
@@ -19,6 +21,7 @@ namespace UE { namespace Tasks
 
 	// BEGIN: forward declarations
 
+	class FTaskEvent;
 	class FPipe;
 
 	template<typename T>
@@ -31,6 +34,20 @@ namespace UE { namespace Tasks
 		LowLevelTasks::ETaskPriority Priority = LowLevelTasks::ETaskPriority::Normal
 	);
 
+	template<typename TaskBodyType, typename PrerequisitesCollectionType>
+	TTask<TInvokeResult_T<TaskBodyType>> Launch(
+		const TCHAR* DebugName, 
+		TaskBodyType&& TaskBody, 
+		PrerequisitesCollectionType&& Prerequisites, 
+		LowLevelTasks::ETaskPriority Priority = LowLevelTasks::ETaskPriority::Normal
+	);
+
+	template<typename ResultType>
+	class TTaskBase;
+
+	template<typename... TaskTypes>
+	class TPrerequisites;
+
 	// END: forward declarations
 
 	namespace Private
@@ -42,11 +59,7 @@ namespace UE { namespace Tasks
 
 			// Can be used only as a base class.
 		protected:
-			FTaskBase()
-				: PipePtr(reinterpret_cast<uintptr_t>(nullptr))
-				, bPushedIntoPipe(false)
-			{
-			}
+			FTaskBase() = default;
 
 			~FTaskBase()
 			{
@@ -77,87 +90,98 @@ namespace UE { namespace Tasks
 						Deleter = Forward<DeleterType>(Deleter)
 					]() mutable
 					{
-						TaskTrace::Started(GetTraceId());
-						{
-							TRACE_CPUPROFILER_EVENT_SCOPE(ExecuteTask);
-							StartPipeExecution();
-							Invoke(TaskBody);
-							FinishPipeExecution();
-						}
-						TaskTrace::Finished(GetTraceId());
-
-						FTaskBase* Subsequent = CloseAndReturnSubsequent();
-						if (Subsequent != nullptr)
-						{
-							check(TryPushIntoPipe());
-							TaskTrace::Scheduled(Subsequent->GetTraceId());
-							LowLevelTasks::TryLaunch(Subsequent->LowLevelTask, LowLevelTasks::EQueuePreference::DefaultPreference, /*bWakeUpWorker =*/ false);
-						}
-
-						TaskTrace::Completed(GetTraceId());
+						Execute(MoveTemp(TaskBody));
+
 					} // Continuation
 				);
 			}
 
-			// A piped task will be executed after the previous task from this pipe is completed. Tasks from the same pipe are not executed
-			// concurrently (so don't require synchronization), but not necessarily on the same thread.
-			// @See FPipe
-			void SetPipe(FPipe& Pipe)
-			{
-				PipePtr = reinterpret_cast<uintptr_t>(&Pipe);
-			}
-
-			FPipe* GetPipe() const
+			// the task will be executed only when all prerequisites are completed
+			template<typename T>
+			void AddPrerequisites(const TTaskBase<T>& Prerequisite)
 			{
-				return reinterpret_cast<FPipe*>(PipePtr);
-			}
+				checkf(NumLocks.load(std::memory_order_relaxed) >= NumInitialLocks, TEXT("Prerequisites can be added only before the task is launched"));
 
-			// allows the task to be scheduled and executed. Tasks w/o pending dependencies (not blocked by a pipe) are 
-			// scheduled immediately
-			bool TryLaunch()
-			{
-				TaskTrace::Launched(GetTraceId(), LowLevelTask.GetDebugName(), true, (ENamedThreads::Type)255);
+				// registering the task as a subsequent of the given prerequisite can cause its immediate launch by the prerequisite
+				// (if the prerequisite has been completed on another thread), so we need to keep the task locked by assuming that the 
+				// prerequisite can be added successfully, and release the lock if it wasn't
+				NumLocks.fetch_add(1, std::memory_order_acquire); // `acquire` to make it happen before the task is registered as a subsequent
 
-				if (!TryPushIntoPipe())
+				if (!Prerequisite.Pimpl->AddSubsequent(*this))
 				{
-					return false; // pipe is blocked
+					// failed to add the prerequisite (too late), correct the number
+					NumLocks.fetch_sub(1, std::memory_order_release); // `release` to make it happen after the task is registered as a subsequent
 				}
+			}
 
-				TaskTrace::Scheduled(GetTraceId());
+			// the task will be executed only when all prerequisites are completed
+			// @param Prerequisites - an iterable collection of tasks
+			template<typename PrerequisiteCollectionType, decltype(std::declval<PrerequisiteCollectionType&>().begin())* = nullptr>
+			void AddPrerequisites(const PrerequisiteCollectionType& Prerequisites)
+			{
+				checkf(NumLocks.load(std::memory_order_relaxed) >= NumInitialLocks, TEXT("Prerequisites can be added only before the task is launched"));
 
-				// "inline" tasks are not scheduled but executed as soon as they are unlocked
-				if (LowLevelTask.GetPriority() == InlineTaskPriority)
+				// registering the task as a subsequent of the given prerequisite can cause its immediate launch by the prerequisite
+				// (if the prerequisite has been completed on another thread), so we need to keep the task locked by assuming that the 
+				// prerequisite can be added successfully, and release the lock if it wasn't
+				NumLocks.fetch_add(GetNum(Prerequisites), std::memory_order_acquire); // `acquire` to make it happen before the task is registered as a subsequent
+
+				uint32 NumCompletedPrerequisites = 0;
+				for (auto& Prereq : Prerequisites)
 				{
-					// the task is not scheduled yet, so successful retraction is guaranted
-					verify(LowLevelTask.TryCancel()); // this will cancel the task for the scheduler and execute LowLevelTask's continuation, 
-					// which is actually task execution
-					return true;
+					// prerequisites can be either `FTaskBase*` or its Pimpl handle
+					FTaskBase* Prerequisite;
+					using FPrerequisiteType = std::decay_t<decltype(*std::declval<PrerequisiteCollectionType&>().begin())>;
+					if constexpr (std::is_same_v<FPrerequisiteType, FTaskBase*>)
+					{
+						Prerequisite = Prereq;
+					}
+					else
+					{
+						Prerequisite = Prereq.Pimpl;
+					}
+
+					if (!Prerequisite->AddSubsequent(*this))
+					{
+						++NumCompletedPrerequisites;
+					}
 				}
 
-				// schedule
-				return LowLevelTasks::TryLaunch(LowLevelTask);
+				// unlock for prerequisites that weren't added
+				NumLocks.fetch_sub(NumCompletedPrerequisites, std::memory_order_release); // `release` to make it happen after the task is registered as a subsequent
 			}
 
-			// sets a subsequent task that will be launched automatically when this task is done
-			// returns true only if the task is not closed yet ("closed" = is executed and has already checked its subsequent),
-			// if returned false the subsequent is not blocked by this task
-			bool SetSubsequent(FTaskBase& Task)
+			// the task unlocks all its subsequents on completion.
+			// return false if the task is already completed
+			bool AddSubsequent(FTaskBase& Subsequent)
 			{
-#if DO_CHECK
-				FTaskBase* CurrentState = SubsequentAndState.load(std::memory_order_relaxed);
-				checkf(CurrentState == GetEmptyOpenState() || CurrentState == GetClosedState(), TEXT("only a single subsequent can be set"));
-#endif
+				return Subsequents.Enqueue(&Subsequent);
+			}
 
-				TaskTrace::SubsequentAdded(GetTraceId(), Task.GetTraceId());
+			// A piped task will be executed after the previous task from this pipe is completed. Tasks from the same pipe are not executed
+			// concurrently (so don't require synchronization), but not necessarily on the same thread.
+			// @See FPipe
+			void SetPipe(FPipe& InPipe)
+			{
+				Pipe = &InPipe;
+			}
 
-				FTaskBase* EmptyNotCompleted = GetEmptyOpenState();
-				FTaskBase* NotCompletedWithTask = &Task;
-				return SubsequentAndState.compare_exchange_strong(EmptyNotCompleted, NotCompletedWithTask, std::memory_order_release, std::memory_order_relaxed);
+			FPipe* GetPipe() const
+			{
+				return Pipe;
 			}
 
+			// Tasks w/o incomplete prerequisites and not blocked by a pipe are scheduled immediately
+			bool TryLaunch()
+			{
+				TaskTrace::Launched(GetTraceId(), LowLevelTask.GetDebugName(), true, (ENamedThreads::Type)0xff);
+				return TryUnlock();
+			}
+
+			// if the task execution is already completed
 			bool IsCompleted() const
 			{
-				return LowLevelTask.IsCompleted();
+				return Subsequents.IsClosed();
 			}
 
 			// waits until the task is completed
@@ -166,9 +190,11 @@ namespace UE { namespace Tasks
 				// we try to retract the task from the scheduler, if it's execution hasn't started yet, to execute it immediately. This makes
 				// the task skip waiting in the scheduler's queue.
 				// we can't retract it right here because it can be not scheduled yet (e.g. when it's blocked by a pipe), so we need to wait for this
-				LowLevelTasks::BusyWaitUntil([this] { return !LowLevelTask.IsReady(); }); // `!IsReady()` here means "the task is scheduled"
+				// `!IsReady()` here means "the task is scheduled". The task can be completed w/o being scheduled at all (e.g. if it has 
+				// an empty body) so we need to check for this too
+				LowLevelTasks::BusyWaitUntil([this] { return !LowLevelTask.IsReady(); });
 				LowLevelTask.TryCancel(); //This will try to cancel the task and run it's continuation (thereby executing the workload) 
-				LowLevelTasks::BusyWaitUntil([this] { return LowLevelTask.IsCompleted(); });
+				LowLevelTasks::BusyWaitUntil([this] { return IsCompleted(); });
 			}
 
 			// waits until the task is completed or waiting timed out
@@ -177,14 +203,15 @@ namespace UE { namespace Tasks
 			{
 				FTimeout Timeout{ InTimeout };
 				LowLevelTasks::BusyWaitUntil([this, Timeout] { return !LowLevelTask.IsReady() || Timeout; });
+
 				if (LowLevelTask.IsReady()) // timeout
 				{
 					return false;
 				}
 
 				LowLevelTask.TryCancel(); //This will try to cancel the task and run it's continuation (thereby executing the workload) 
-				LowLevelTasks::BusyWaitUntil([this, Timeout] { return LowLevelTask.IsCompleted() || Timeout; });
-				return LowLevelTask.IsCompleted();
+				LowLevelTasks::BusyWaitUntil([this, Timeout] { return IsCompleted() || Timeout; });
+				return IsCompleted();
 			}
 
 			// waits until the task is completed or the condition returns true
@@ -203,9 +230,9 @@ namespace UE { namespace Tasks
 
 				LowLevelTask.TryCancel(); //This will try to cancel the task and run it's continuation (thereby executing the workload) 
 				LowLevelTasks::BusyWaitUntil(
-					[this, Condition = Forward<ConditionType>(Condition)] { return LowLevelTask.IsCompleted() || Condition(); }
+					[this, Condition = Forward<ConditionType>(Condition)] { return IsCompleted() || Condition(); }
 				);
-				return LowLevelTask.IsCompleted(std::memory_order_relaxed);
+				return IsCompleted();
 			}
 
 			TaskTrace::FId GetTraceId() const
@@ -218,17 +245,75 @@ namespace UE { namespace Tasks
 			}
 
 		private:
-			// checks if the task is ready to be launched by trying to push it into the pipe
-			bool TryPushIntoPipe()
+			// allows the task to be scheduled and executed, once for each lock. it's called on actual launching, on every prerequisite 
+			// completion and when an associated pipe (if any) is unblocked
+			bool TryUnlock()
+			{
+				uint32 LocalNumLocks = NumLocks.fetch_sub(1, std::memory_order_acq_rel) - 1; // `acq_rel` to make it happen after task 
+				// preparation and before launching it
+				const bool bReadyForPiping = LocalNumLocks <= 1; // the only lock that can remain is pipe's one
+				if (!bReadyForPiping)
+				{
+					return false;
+				}
+
+				if (!TryPushIntoPipe(LocalNumLocks))
+				{
+					return false; // pipe is blocked
+				}
+
+				TaskTrace::Scheduled(GetTraceId());
+
+				// "inline" tasks are not scheduled but executed as soon as they are unlocked
+				if (LowLevelTask.GetPriority() == InlineTaskPriority)
+				{
+					// the task is not scheduled yet, so successful retraction is guaranted
+					verify(LowLevelTask.TryCancel()); // this will cancel the task for the scheduler and execute LowLevelTask's continuation, 
+					// which is actually task execution
+					return true;
+				}
+
+				// schedule
+				return LowLevelTasks::TryLaunch(LowLevelTask);
+			}
+
+			template<typename TaskBodyType>
+			void Execute(TaskBodyType&& TaskBody)
+			{
+				TaskTrace::Started(GetTraceId());
+				{
+					TRACE_CPUPROFILER_EVENT_SCOPE(ExecuteTask);
+					StartPipeExecution();
+					Invoke(TaskBody);
+					FinishPipeExecution();
+				}
+				TaskTrace::Finished(GetTraceId());
+
+				Close();
+			}
+
+			// checks if the task is ready to be launched by trying to push it into the pipe.
+			// can be called up to two times: first to push into the blocked pipe and then when the pipe is unblocked
+			// `LocalNumLocks` is the value that was used to make a decision to push into the pipe, no need to read `NumLocks` again
+			bool TryPushIntoPipe(uint32 LocalNumLocks)
 			{
 				if (GetPipe() == nullptr)
 				{
+					// the task is locked for a pipe initially even if eventually there's no pipe
+					check(LocalNumLocks == 1);
+					// there's no need to set NumLocks to 0 in non-shipping builds as it's not used anymore, but we do this for debugging
+					// as it's confusing to see scheduled tasks that are still locked
+					check(NumLocks.exchange(0, std::memory_order_relaxed) == 1);
+
 					return true;
 				}
 
-				if (!bPushedIntoPipe)
+				// on the first call we try to push the task into the pipe. if unsuccessful (the pipe is blocked), the second time the method is called
+				// only when the pipe is unblocked, so we know that the task is free to be executed
+
+				bool bFirstAttempt = LocalNumLocks == 1;
+				if (bFirstAttempt)
 				{
-					bPushedIntoPipe = true;
 					return PushIntoPipe();
 				}
 
@@ -239,41 +324,26 @@ namespace UE { namespace Tasks
 			CORE_API void StartPipeExecution();
 			CORE_API void FinishPipeExecution();
 
-			// after being executed the task checks if it has a subsequent to launch. it's the only time the task will do this check so
-			// it "closes" itself and setting a subsequent fill fail after that.
-			// atomic operation
-			FTaskBase* CloseAndReturnSubsequent()
-			{
-				FTaskBase* Subsequent_Local = SubsequentAndState.exchange(GetClosedState(), std::memory_order_acq_rel);
-				checkf(Subsequent_Local != GetClosedState(), TEXT("Task can be closed only once"));
-				return Subsequent_Local;
-			}
-
-			// no subsequent, still open
-			static FTaskBase* GetEmptyOpenState()
+			void Close()
 			{
-				return nullptr;
-			}
-
-			static FTaskBase* GetClosedState()
-			{
-				// nullptr with the most significant bit set to indicate the closed state
-				return UE_LAUNDER(reinterpret_cast<FTaskBase*>(1ull << 63));
+				Subsequents.Close(
+					[this](FTaskBase* Subsequent)
+					{
+						Subsequent->TryUnlock();
+					}
+				);
 			}
 
 		private:
 			LowLevelTasks::FTask LowLevelTask;
 
-			// packed pipe's ptr and the flag if the task is already pushed into pipe
-			uintptr_t PipePtr : 63; // actually `FPipe*`
-			uintptr_t bPushedIntoPipe : 1;	// bool, if the task is pushed into pipe
+			// the number of times that the task should be unlocked before it can be scheduled
+			// initial count is 1 for launching the task (it can't be scheduled before it's launched) and 1 for a potential blocked pipe
+			static const uint32 NumInitialLocks = 1 + 1;
+			std::atomic<uint32> NumLocks{ NumInitialLocks };
+			TClosableMpscQueue<FTaskBase*> Subsequents;
 
-			// a pointer to subsequent task with the most significant bit borrowed to be used as a close flag
-			// `LowLevelTask.IsCompleted()` flag can't be used as a "close" flag because we need to check if the task still can 
-			// accept a subsequent and to set the subsequent atomically.
-			// can be "empty and open" (`nullptr`), "have subsequent and open" (subsequent ptr) and closed (`nullptr` with the most significant
-			// bit set)
-			std::atomic<FTaskBase*> SubsequentAndState{ GetEmptyOpenState() };
+			FPipe* Pipe{ nullptr };
 
 #if UE_TASK_TRACE_ENABLED
 			TaskTrace::FId TraceId = TaskTrace::GenerateTaskId();
@@ -395,6 +465,10 @@ namespace UE { namespace Tasks
 			template<typename TaskBodyType>
 			friend TTask<TInvokeResult_T<TaskBodyType>> UE::Tasks::Launch(const TCHAR* DebugName, TaskBodyType&& TaskBody, LowLevelTasks::ETaskPriority Priority);
 
+			template<typename TaskBodyType, typename PrerequisitesCollectionType>
+			friend TTask<TInvokeResult_T<TaskBodyType>> UE::Tasks::Launch(const TCHAR* DebugName, TaskBodyType&& TaskBody, PrerequisitesCollectionType&& Prerequisites, LowLevelTasks::ETaskPriority Priority);
+
+			friend FTaskEvent;
 			friend FPipe;
 
 			TTaskWithResult()
