Fixes where integers are assumed to be implicitly convertible to char types.

diff --git a/Engine/Source/Runtime/Core/Private/GenericPlatform/GenericPlatformStricmp.cpp b/Engine/Source/Runtime/Core/Private/GenericPlatform/GenericPlatformStricmp.cpp
--- a/Engine/Source/Runtime/Core/Private/GenericPlatform/GenericPlatformStricmp.cpp
+++ b/Engine/Source/Runtime/Core/Private/GenericPlatform/GenericPlatformStricmp.cpp
@@ -14,10 +14,9 @@ static constexpr uint8 LowerAscii[128] = {
 	0x70, 0x71, 0x72, 0x73, 0x74, 0x75, 0x76, 0x77, 0x78, 0x79, 0x7A, 0x7B, 0x7C, 0x7D, 0x7E, 0x7F
 };
 
-template<typename CharType1, typename CharType2>
-FORCEINLINE bool BothAscii(CharType1 C1, CharType2 C2)
+FORCEINLINE bool BothAscii(uint32 C1, uint32 C2)
 {
-	return (((uint32)C1 | (uint32)C2) & 0xffffff80) == 0;
+	return ((C1 | C2) & 0xffffff80) == 0;
 }
 
 template<typename CharType1, typename CharType2>
@@ -28,27 +27,30 @@ int32 StricmpImpl(const CharType1* String1, const CharType2* String2)
 		CharType1 C1 = *String1++;
 		CharType2 C2 = *String2++;
 
+		uint32 U1 = TChar<CharType1>::ToUnsigned(C1);
+		uint32 U2 = TChar<CharType2>::ToUnsigned(C2);
+
 		// Quickly move on if characters are identical but
 		// return equals if we found two null terminators
-		if (C1 == C2)
+		if (U1 == U2)
 		{
-			if (C1)
+			if (U1)
 			{
 				continue;
 			}
 
 			return 0;
 		}
-		else if (BothAscii(C1, C2))
+		else if (BothAscii(U1, U2))
 		{
-			if (int32 Diff = LowerAscii[TChar<CharType1>::ToUnsigned(C1)] - LowerAscii[TChar<CharType2>::ToUnsigned(C2)])
+			if (int32 Diff = LowerAscii[U1] - LowerAscii[U2])
 			{
 				return Diff;
 			}
 		}
 		else
 		{
-			return TChar<CharType1>::ToUnsigned(C1) - TChar<CharType2>::ToUnsigned(C2);
+			return U1 - U2;
 		}
 	}
 }
diff --git a/Engine/Source/Runtime/Core/Public/Misc/Char.h b/Engine/Source/Runtime/Core/Public/Misc/Char.h
--- a/Engine/Source/Runtime/Core/Public/Misc/Char.h
+++ b/Engine/Source/Runtime/Core/Public/Misc/Char.h
@@ -33,13 +33,13 @@ template <> struct TLiteral<WIDECHAR>
 template <typename CharType, const unsigned int Size>
 struct TCharBase
 {
-	static constexpr CharType LineFeed = 0xa;
-	static constexpr CharType VerticalTab = 0xb;
-	static constexpr CharType FormFeed = 0xc;
-	static constexpr CharType CarriageReturn = 0xd;
-	static constexpr CharType NextLine = 0x85;
-	static constexpr CharType LineSeparator = 0x2028;
-	static constexpr CharType ParagraphSeparator = 0x2029;
+	static constexpr CharType LineFeed           = (CharType)0xa;
+	static constexpr CharType VerticalTab        = (CharType)0xb;
+	static constexpr CharType FormFeed           = (CharType)0xc;
+	static constexpr CharType CarriageReturn     = (CharType)0xd;
+	static constexpr CharType NextLine           = (CharType)0x85;
+	static constexpr CharType LineSeparator      = (CharType)0x2028;
+	static constexpr CharType ParagraphSeparator = (CharType)0x2029;
 
 	static bool IsLinebreak(CharType Char)
 	{
@@ -52,10 +52,10 @@ struct TCharBase
 template <typename CharType>
 struct TCharBase<CharType, 1>
 {
-	static constexpr CharType LineFeed = 0xa;
-	static constexpr CharType VerticalTab = 0xb;
-	static constexpr CharType FormFeed = 0xc;
-	static constexpr CharType CarriageReturn = 0xd;
+	static constexpr CharType LineFeed       = (CharType)0xa;
+	static constexpr CharType VerticalTab    = (CharType)0xb;
+	static constexpr CharType FormFeed       = (CharType)0xc;
+	static constexpr CharType CarriageReturn = (CharType)0xd;
 
 	static bool IsLinebreak(CharType Char)
 	{
