Move UTF-8 conversion functions into the platform layer.
Enable UTF8CHAR as a char type, so that you can use a UTF8CHAR* in a generic string conversion context (e.g. FString construction).

diff --git a/Engine/Source/Runtime/Core/Private/GenericPlatform/GenericPlatformString.cpp b/Engine/Source/Runtime/Core/Private/GenericPlatform/GenericPlatformString.cpp
--- a/Engine/Source/Runtime/Core/Private/GenericPlatform/GenericPlatformString.cpp
+++ b/Engine/Source/Runtime/Core/Private/GenericPlatform/GenericPlatformString.cpp
@@ -32,6 +32,629 @@ namespace
 	}
 }
 
+namespace UE::Core::Private
+{
+	/**
+	 * This is a basic object which counts how many times it has been incremented
+	 */
+	template <typename DestType>
+	struct TCountingOutputIterator
+	{
+		TCountingOutputIterator()
+			: Counter(0)
+		{
+		}
+
+		const TCountingOutputIterator& operator* () const { return *this; }
+		const TCountingOutputIterator& operator++() { ++Counter; return *this; }
+		const TCountingOutputIterator& operator++(int) { ++Counter; return *this; }
+		const TCountingOutputIterator& operator+=(const int32 Amount) { Counter += Amount; return *this; }
+
+		const DestType& operator=(const DestType& Val) const
+		{
+			return Val;
+		}
+
+		friend int32 operator-(TCountingOutputIterator Lhs, TCountingOutputIterator Rhs)
+		{
+			return Lhs.Counter - Rhs.Counter;
+		}
+
+		int32 GetCount() const { return Counter; }
+
+	private:
+		int32 Counter;
+	};
+
+	/** Is the provided Codepoint within the range of valid codepoints? */
+	FORCEINLINE bool IsValidCodepoint(const uint32 Codepoint)
+	{
+		if ((Codepoint > 0x10FFFF) ||						// No Unicode codepoints above 10FFFFh, (for now!)
+			(Codepoint == 0xFFFE) || (Codepoint == 0xFFFF)) // illegal values.
+		{
+			return false;
+		}
+		return true;
+	}
+
+	/** Is the provided Codepoint within the range of the high-surrogates? */
+	static FORCEINLINE bool IsHighSurrogate(const uint32 Codepoint)
+	{
+		return Codepoint >= HIGH_SURROGATE_START_CODEPOINT && Codepoint <= HIGH_SURROGATE_END_CODEPOINT;
+	}
+
+	/** Is the provided Codepoint within the range of the low-surrogates? */
+	static FORCEINLINE bool IsLowSurrogate(const uint32 Codepoint)
+	{
+		return Codepoint >= LOW_SURROGATE_START_CODEPOINT && Codepoint <= LOW_SURROGATE_END_CODEPOINT;
+	}
+
+	static FORCEINLINE uint32 EncodeSurrogate(const uint16 HighSurrogate, const uint16 LowSurrogate)
+	{
+		return ((HighSurrogate - HIGH_SURROGATE_START_CODEPOINT) << 10) + (LowSurrogate - LOW_SURROGATE_START_CODEPOINT) + 0x10000;
+	}
+
+	static FORCEINLINE void DecodeSurrogate(const uint32 Codepoint, uint16& OutHighSurrogate, uint16& OutLowSurrogate)
+	{
+		const uint32 TmpCodepoint = Codepoint - 0x10000;
+		OutHighSurrogate = (uint16)((TmpCodepoint >> 10) + HIGH_SURROGATE_START_CODEPOINT);
+		OutLowSurrogate = (TmpCodepoint & 0x3FF) + LOW_SURROGATE_START_CODEPOINT;
+	}
+
+	/** Is the provided Codepoint outside of the range of the basic multilingual plane, but within the valid range of UTF8/16? */
+	static FORCEINLINE bool IsEncodedSurrogate(const uint32 Codepoint)
+	{
+		return Codepoint >= ENCODED_SURROGATE_START_CODEPOINT && Codepoint <= ENCODED_SURROGATE_END_CODEPOINT;
+	}
+
+	/**
+	 * Convert TCHAR Codepoint into UTF-8 characters.
+	 *
+	 * @param Codepoint Codepoint to expand into UTF-8 bytes
+	 * @param OutputIterator Output iterator to write UTF-8 bytes into
+	 * @param OutputIteratorByteSizeRemaining Maximum number of ANSI characters that can be written to OutputIterator
+	 * @return Number of characters written for Codepoint
+	 */
+	template <typename BufferType>
+	static int32 WriteCodepointToUTF8(uint32 Codepoint, BufferType OutputIterator, uint32 OutputIteratorByteSizeRemaining)
+	{
+		// Ensure we have at least one character in size to write
+		if (OutputIteratorByteSizeRemaining < sizeof(UTF8CHAR))
+		{
+			return 0;
+		}
+
+		const BufferType OutputIteratorStartPosition = OutputIterator;
+
+		if (!IsValidCodepoint(Codepoint))
+		{
+			Codepoint = UNICODE_BOGUS_CHAR_CODEPOINT;
+		}
+		else if (IsHighSurrogate(Codepoint) || IsLowSurrogate(Codepoint)) // UTF-8 Characters are not allowed to encode codepoints in the surrogate pair range
+		{
+			Codepoint = UNICODE_BOGUS_CHAR_CODEPOINT;
+		}
+
+		// Do the encoding...
+		if (Codepoint < 0x80)
+		{
+			*(OutputIterator++) = (UTF8CHAR)Codepoint;
+		}
+		else if (Codepoint < 0x800)
+		{
+			if (OutputIteratorByteSizeRemaining >= 2)
+			{
+				*(OutputIterator++) = (UTF8CHAR)((Codepoint >> 6)         | 128 | 64);
+				*(OutputIterator++) = (UTF8CHAR)((Codepoint       & 0x3F) | 128);
+			}
+		}
+		else if (Codepoint < 0x10000)
+		{
+			if (OutputIteratorByteSizeRemaining >= 3)
+			{
+				*(OutputIterator++) = (UTF8CHAR)( (Codepoint >> 12)        | 128 | 64 | 32);
+				*(OutputIterator++) = (UTF8CHAR)(((Codepoint >> 6) & 0x3F) | 128);
+				*(OutputIterator++) = (UTF8CHAR)( (Codepoint       & 0x3F) | 128);
+			}
+		}
+		else
+		{
+			if (OutputIteratorByteSizeRemaining >= 4)
+			{
+				*(OutputIterator++) = (UTF8CHAR)( (Codepoint >> 18)         | 128 | 64 | 32 | 16);
+				*(OutputIterator++) = (UTF8CHAR)(((Codepoint >> 12) & 0x3F) | 128);
+				*(OutputIterator++) = (UTF8CHAR)(((Codepoint >> 6 ) & 0x3F) | 128);
+				*(OutputIterator++) = (UTF8CHAR)( (Codepoint        & 0x3F) | 128);
+			}
+		}
+
+		return UE_PTRDIFF_TO_INT32(OutputIterator - OutputIteratorStartPosition);
+	}
+
+	template <typename DestBufferType>
+	static bool WriteCodepointToBuffer(const uint32 Codepoint, DestBufferType& Dest, int32& DestLen)
+	{
+		int32 WrittenChars = WriteCodepointToUTF8(Codepoint, Dest, DestLen);
+		if (WrittenChars < 1)
+		{
+			return false;
+		}
+
+		Dest += WrittenChars;
+		DestLen -= WrittenChars;
+		return true;
+	}
+
+	template <typename DestBufferType, typename FromType, typename DestType>
+	static int32 ConvertToUTF8(DestBufferType& Dest, int32 DestLen, const FromType* Source, const int32 SourceLen, DestType BogusChar)
+	{
+		DestBufferType DestStartingPosition = Dest;
+		if constexpr (sizeof(FromType) == 4)
+		{
+			for (int32 i = 0; i < SourceLen; ++i)
+			{
+				uint32 Codepoint = static_cast<uint32>(Source[i]);
+
+				if (!WriteCodepointToBuffer(Codepoint, Dest, DestLen))
+				{
+					// Could not write data, bail out
+					return -1;
+				}
+			}
+		}
+		else
+		{
+			uint32 HighSurrogate = MAX_uint32;
+
+			for (int32 i = 0; i < SourceLen; ++i)
+			{
+				const bool bHighSurrogateIsSet = HighSurrogate != MAX_uint32;
+				uint32 Codepoint = static_cast<uint32>(Source[i]);
+
+				// Check if this character is a high-surrogate
+				if (IsHighSurrogate(Codepoint))
+				{
+					// Ensure we don't already have a high-surrogate set or end without a matching low-surrogate
+					if (bHighSurrogateIsSet || i == SourceLen - 1)
+					{
+						// Already have a high-surrogate in this pair or string ends with lone high-surrogate
+						// Write our stored value (will be converted into bogus character)
+						if (!WriteCodepointToBuffer(HighSurrogate, Dest, DestLen))
+						{
+							// Could not write data, bail out
+							return -1;
+						}
+					}
+
+					// Store our code point for our next character
+					HighSurrogate = Codepoint;
+					continue;
+				}
+
+				// If our High Surrogate is set, check if this character is the matching low-surrogate
+				if (bHighSurrogateIsSet)
+				{
+					if (IsLowSurrogate(Codepoint))
+					{
+						const uint32 LowSurrogate = Codepoint;
+						// Combine our high and low surrogates together to a single Unicode codepoint
+						Codepoint = EncodeSurrogate((uint16)HighSurrogate, (uint16)LowSurrogate);
+					}
+					else
+					{
+						// Did not find matching low-surrogate, write out a bogus character for our stored HighSurrogate
+						if (!WriteCodepointToBuffer(HighSurrogate, Dest, DestLen))
+						{
+							// Could not write data, bail out
+							return -1;
+						}
+					}
+
+					// Reset our high-surrogate now that we've used (or discarded) its value
+					HighSurrogate = MAX_uint32;
+				}
+
+				if (!WriteCodepointToBuffer(Codepoint, Dest, DestLen))
+				{
+					// Could not write data, bail out
+					return -1;
+				}
+			}
+		}
+
+		return UE_PTRDIFF_TO_INT32(Dest - DestStartingPosition);
+	}
+
+	template <typename FromType>
+	static uint32 CodepointFromUtf8(const FromType*& SourceString, const uint32 SourceLengthRemaining, uint32 BogusChar)
+	{
+		checkSlow(SourceLengthRemaining > 0);
+
+		const FromType* OctetPtr = SourceString;
+
+		uint32 Codepoint = 0;
+		uint32 Octet = (uint32) ((uint8) *SourceString);
+		uint32 Octet2, Octet3, Octet4;
+
+		if (Octet < 128)  // one octet char: 0 to 127
+		{
+			++SourceString;  // skip to next possible start of codepoint.
+			return Octet;
+		}
+		else if (Octet < 192)  // bad (starts with 10xxxxxx).
+		{
+			// Apparently each of these is supposed to be flagged as a bogus
+			//  char, instead of just resyncing to the next valid codepoint.
+			++SourceString;  // skip to next possible start of codepoint.
+			return BogusChar;
+		}
+		else if (Octet < 224)  // two octets
+		{
+			// Ensure our string has enough characters to read from
+			if (SourceLengthRemaining < 2)
+			{
+				// Skip to end and write out a single char (we always have room for at least 1 char)
+				SourceString += SourceLengthRemaining;
+				return BogusChar;
+			}
+
+			Octet -= (128+64);
+			Octet2 = (uint32) ((uint8) *(++OctetPtr));
+			if ((Octet2 & (128 + 64)) != 128)  // Format isn't 10xxxxxx?
+			{
+				++SourceString;  // Sequence was not valid UTF-8. Skip the first byte and continue.
+				return BogusChar;
+			}
+
+			Codepoint = ((Octet << 6) | (Octet2 - 128));
+			if ((Codepoint >= 0x80) && (Codepoint <= 0x7FF))
+			{
+				SourceString += 2;  // skip to next possible start of codepoint.
+				return Codepoint;
+			}
+		}
+		else if (Octet < 240)  // three octets
+		{
+			// Ensure our string has enough characters to read from
+			if (SourceLengthRemaining < 3)
+			{
+				// Skip to end and write out a single char (we always have room for at least 1 char)
+				SourceString += SourceLengthRemaining;
+				return BogusChar;
+			}
+
+			Octet -= (128+64+32);
+			Octet2 = (uint32) ((uint8) *(++OctetPtr));
+			if ((Octet2 & (128+64)) != 128)  // Format isn't 10xxxxxx?
+			{
+				++SourceString;  // Sequence was not valid UTF-8. Skip the first byte and continue.
+				return BogusChar;
+			}
+
+			Octet3 = (uint32) ((uint8) *(++OctetPtr));
+			if ((Octet3 & (128+64)) != 128)  // Format isn't 10xxxxxx?
+			{
+				++SourceString;  // Sequence was not valid UTF-8. Skip the first byte and continue.
+				return BogusChar;
+			}
+
+			Codepoint = ( ((Octet << 12)) | ((Octet2-128) << 6) | ((Octet3-128)) );
+
+			// UTF-8 characters cannot be in the UTF-16 surrogates range
+			if (IsHighSurrogate(Codepoint) || IsLowSurrogate(Codepoint))
+			{
+				++SourceString;  // Sequence was not valid UTF-8. Skip the first byte and continue.
+				return BogusChar;
+			}
+
+			// 0xFFFE and 0xFFFF are illegal, too, so we check them at the edge.
+			if ((Codepoint >= 0x800) && (Codepoint <= 0xFFFD))
+			{
+				SourceString += 3;  // skip to next possible start of codepoint.
+				return Codepoint;
+			}
+		}
+		else if (Octet < 248)  // four octets
+		{
+			// Ensure our string has enough characters to read from
+			if (SourceLengthRemaining < 4)
+			{
+				// Skip to end and write out a single char (we always have room for at least 1 char)
+				SourceString += SourceLengthRemaining;
+				return BogusChar;
+			}
+
+			Octet -= (128+64+32+16);
+			Octet2 = (uint32) ((uint8) *(++OctetPtr));
+			if ((Octet2 & (128+64)) != 128)  // Format isn't 10xxxxxx?
+			{
+				++SourceString;  // Sequence was not valid UTF-8. Skip the first byte and continue.
+				return BogusChar;
+			}
+
+			Octet3 = (uint32) ((uint8) *(++OctetPtr));
+			if ((Octet3 & (128+64)) != 128)  // Format isn't 10xxxxxx?
+			{
+				++SourceString;  // Sequence was not valid UTF-8. Skip the first byte and continue.
+				return BogusChar;
+			}
+
+			Octet4 = (uint32) ((uint8) *(++OctetPtr));
+			if ((Octet4 & (128+64)) != 128)  // Format isn't 10xxxxxx?
+			{
+				++SourceString;  // Sequence was not valid UTF-8. Skip the first byte and continue.
+				return BogusChar;
+			}
+
+			Codepoint = ( ((Octet << 18)) | ((Octet2 - 128) << 12) |
+						((Octet3 - 128) << 6) | ((Octet4 - 128)) );
+			if ((Codepoint >= 0x10000) && (Codepoint <= 0x10FFFF))
+			{
+				SourceString += 4;  // skip to next possible start of codepoint.
+				return Codepoint;
+			}
+		}
+		// Five and six octet sequences became illegal in rfc3629.
+		//  We throw the codepoint away, but parse them to make sure we move
+		//  ahead the right number of bytes and don't overflow the buffer.
+		else if (Octet < 252)  // five octets
+		{
+			// Ensure our string has enough characters to read from
+			if (SourceLengthRemaining < 5)
+			{
+				// Skip to end and write out a single char (we always have room for at least 1 char)
+				SourceString += SourceLengthRemaining;
+				return BogusChar;
+			}
+
+			Octet = (uint32) ((uint8) *(++OctetPtr));
+			if ((Octet & (128+64)) != 128)  // Format isn't 10xxxxxx?
+			{
+				++SourceString;  // Sequence was not valid UTF-8. Skip the first byte and continue.
+				return BogusChar;
+			}
+
+			Octet = (uint32) ((uint8) *(++OctetPtr));
+			if ((Octet & (128+64)) != 128)  // Format isn't 10xxxxxx?
+			{
+				++SourceString;  // Sequence was not valid UTF-8. Skip the first byte and continue.
+				return BogusChar;
+			}
+
+			Octet = (uint32) ((uint8) *(++OctetPtr));
+			if ((Octet & (128+64)) != 128)  // Format isn't 10xxxxxx?
+			{
+				++SourceString;  // Sequence was not valid UTF-8. Skip the first byte and continue.
+				return BogusChar;
+			}
+
+			Octet = (uint32) ((uint8) *(++OctetPtr));
+			if ((Octet & (128+64)) != 128)  // Format isn't 10xxxxxx?
+			{
+				++SourceString;  // Sequence was not valid UTF-8. Skip the first byte and continue.
+				return BogusChar;
+			}
+
+			SourceString += 5;  // skip to next possible start of codepoint.
+			return BogusChar;
+		}
+
+		else  // six octets
+		{
+			// Ensure our string has enough characters to read from
+			if (SourceLengthRemaining < 6)
+			{
+				// Skip to end and write out a single char (we always have room for at least 1 char)
+				SourceString += SourceLengthRemaining;
+				return BogusChar;
+			}
+
+			Octet = (uint32) ((uint8) *(++OctetPtr));
+			if ((Octet & (128+64)) != 128)  // Format isn't 10xxxxxx?
+			{
+				++SourceString;  // Sequence was not valid UTF-8. Skip the first byte and continue.
+				return BogusChar;
+			}
+
+			Octet = (uint32) ((uint8) *(++OctetPtr));
+			if ((Octet & (128+64)) != 128)  // Format isn't 10xxxxxx?
+			{
+				++SourceString;  // Sequence was not valid UTF-8. Skip the first byte and continue.
+				return BogusChar;
+			}
+
+			Octet = (uint32) ((uint8) *(++OctetPtr));
+			if ((Octet & (128+64)) != 128)  // Format isn't 10xxxxxx?
+			{
+				++SourceString;  // Sequence was not valid UTF-8. Skip the first byte and continue.
+				return BogusChar;
+			}
+
+			Octet = (uint32) ((uint8) *(++OctetPtr));
+			if ((Octet & (128+64)) != 128)  // Format isn't 10xxxxxx?
+			{
+				++SourceString;  // Sequence was not valid UTF-8. Skip the first byte and continue.
+				return BogusChar;
+			}
+
+			Octet = (uint32) ((uint8) *(++OctetPtr));
+			if ((Octet & (128+64)) != 128)  // Format isn't 10xxxxxx?
+			{
+				++SourceString;  // Sequence was not valid UTF-8. Skip the first byte and continue.
+				return BogusChar;
+			}
+
+			SourceString += 6;  // skip to next possible start of codepoint.
+			return BogusChar;
+		}
+
+		++SourceString;  // Sequence was not valid UTF-8. Skip the first byte and continue.
+		return BogusChar;  // catch everything else.
+	}
+
+	/**
+	 * Read Source string, converting the data from UTF-8 into UTF-16, and placing these in the Destination
+	 */
+	template <typename DestBufferType, typename FromType, typename DestType>
+	static void ConvertFromUTF8(DestBufferType& ConvertedBuffer, int32 DestLen, const FromType* Source, const int32 SourceLen, DestType BogusChar)
+	{
+		const FromType* SourceEnd = Source + SourceLen;
+
+		const uint64 ExtendedCharMask = 0x8080808080808080;
+		while (Source < SourceEnd && DestLen > 0)
+		{
+			// In case we're given an unaligned pointer, we'll
+			// fallback to the slow path until properly aligned.
+			if (IsAligned(Source, 8))
+			{
+				// Fast path for most common case
+				while (Source < SourceEnd - 8 && DestLen >= 8)
+				{
+					// Detect any extended characters 8 chars at a time
+					if ((*(const uint64*)Source) & ExtendedCharMask)
+					{
+						// Move to slow path since we got extended characters to process
+						break;
+					}
+
+					// This should get unrolled on most compiler
+					// ROI of diminished return to vectorize this as we 
+					// would have to deal with alignment, endianness and
+					// rewrite the iterators to support bulk writes
+					for (int32 Index = 0; Index < 8; ++Index)
+					{
+						*(ConvertedBuffer++) = (DestType)(uint8)*(Source++);
+					}
+					DestLen -= 8;
+				}
+			}
+
+			// Slow path for extended characters
+			while (Source < SourceEnd && DestLen > 0)
+			{
+				// Read our codepoint, advancing the source pointer
+				uint32 Codepoint = CodepointFromUtf8(Source, UE_PTRDIFF_TO_UINT32(SourceEnd - Source), BogusChar);
+
+				if constexpr (sizeof(DestType) == 4)
+				{
+					// We want to write out two chars
+					if (IsEncodedSurrogate(Codepoint))
+					{
+						// We need two characters to write the surrogate pair
+						if (DestLen >= 2)
+						{
+							uint16 HighSurrogate = 0;
+							uint16 LowSurrogate = 0;
+							DecodeSurrogate(Codepoint, HighSurrogate, LowSurrogate);
+
+							*(ConvertedBuffer++) = (DestType)HighSurrogate;
+							*(ConvertedBuffer++) = (DestType)LowSurrogate;
+							DestLen -= 2;
+							continue;
+						}
+
+						// If we don't have space, write a bogus character instead (we should have space for it)
+						Codepoint = (uint32)BogusChar;
+					}
+					else if (Codepoint > ENCODED_SURROGATE_END_CODEPOINT)
+					{
+						// Ignore values higher than the supplementary plane range
+						Codepoint = (uint32)BogusChar;
+					}
+				}
+
+				*(ConvertedBuffer++) = (DestType)Codepoint;
+				--DestLen;
+
+				// Return to the fast path once aligned and back to simple ASCII chars
+				if (Codepoint < 128 && IsAligned(Source, 8))
+				{
+					break;
+				}
+			}
+		}
+	}
+
+	/**
+	 * Determines the length of the converted string.
+	 *
+	 * @return The length of the string in UTF-16 code units.
+	 */
+	int32 GetConvertedLength(const UTF8CHAR*, const ANSICHAR* Source, int32 SourceLen)
+	{
+		TCountingOutputIterator<UTF8CHAR> Dest;
+		ConvertToUTF8(Dest, INT32_MAX, Source, SourceLen, (UTF8CHAR)'?');
+
+		return Dest.GetCount();
+	}
+	int32 GetConvertedLength(const UTF8CHAR*, const WIDECHAR* Source, int32 SourceLen)
+	{
+		TCountingOutputIterator<UTF8CHAR> Dest;
+		ConvertToUTF8(Dest, INT32_MAX, Source, SourceLen, (UTF8CHAR)'?');
+
+		return Dest.GetCount();
+	}
+	int32 GetConvertedLength(const UTF8CHAR*, const UCS2CHAR* Source, int32 SourceLen)
+	{
+		TCountingOutputIterator<UTF8CHAR> Dest;
+		ConvertToUTF8(Dest, INT32_MAX, Source, SourceLen, (UTF8CHAR)'?');
+
+		return Dest.GetCount();
+	}
+	int32 GetConvertedLength(const ANSICHAR*, const UTF8CHAR* Source, int32 SourceLen)
+	{
+		TCountingOutputIterator<ANSICHAR> Dest;
+		ConvertFromUTF8(Dest, INT32_MAX, Source, SourceLen, (ANSICHAR)'?');
+
+		return Dest.GetCount();
+	}
+	int32 GetConvertedLength(const WIDECHAR*, const UTF8CHAR* Source, int32 SourceLen)
+	{
+		TCountingOutputIterator<WIDECHAR> Dest;
+		ConvertFromUTF8(Dest, INT32_MAX, Source, SourceLen, (WIDECHAR)'?');
+
+		return Dest.GetCount();
+	}
+	int32 GetConvertedLength(const UCS2CHAR*, const UTF8CHAR* Source, int32 SourceLen)
+	{
+		TCountingOutputIterator<UCS2CHAR> Dest;
+		ConvertFromUTF8(Dest, INT32_MAX, Source, SourceLen, (UCS2CHAR)'?');
+
+		return Dest.GetCount();
+	}
+
+	UTF8CHAR* Convert(UTF8CHAR* Dest, int32 DestLen, const ANSICHAR* Src, int32 SrcLen, UTF8CHAR BogusChar)
+	{
+		ConvertToUTF8(Dest, DestLen, Src, SrcLen, BogusChar);
+		return Dest;
+	}
+	UTF8CHAR* Convert(UTF8CHAR* Dest, int32 DestLen, const WIDECHAR* Src, int32 SrcLen, UTF8CHAR BogusChar)
+	{
+		ConvertToUTF8(Dest, DestLen, Src, SrcLen, BogusChar);
+		return Dest;
+	}
+	UTF8CHAR* Convert(UTF8CHAR* Dest, int32 DestLen, const UCS2CHAR* Src, int32 SrcLen, UTF8CHAR BogusChar)
+	{
+		ConvertToUTF8(Dest, DestLen, Src, SrcLen, BogusChar);
+		return Dest;
+	}
+	ANSICHAR* Convert(ANSICHAR* Dest, int32 DestLen, const UTF8CHAR* Src, int32 SrcLen, ANSICHAR BogusChar)
+	{
+		ConvertFromUTF8(Dest, DestLen, Src, SrcLen, BogusChar);
+		return Dest;
+	}
+	WIDECHAR* Convert(WIDECHAR* Dest, int32 DestLen, const UTF8CHAR* Src, int32 SrcLen, WIDECHAR BogusChar)
+	{
+		ConvertFromUTF8(Dest, DestLen, Src, SrcLen, BogusChar);
+		return Dest;
+	}
+	UCS2CHAR* Convert(UCS2CHAR* Dest, int32 DestLen, const UTF8CHAR* Src, int32 SrcLen, UCS2CHAR BogusChar)
+	{
+		ConvertFromUTF8(Dest, DestLen, Src, SrcLen, BogusChar);
+		return Dest;
+	}
+}
+
 template <typename DestEncoding, typename SourceEncoding>
 void FGenericPlatformString::LogBogusChars(const SourceEncoding* Src, int32 SrcSize)
 {
diff --git a/Engine/Source/Runtime/Core/Public/Containers/StringConv.h b/Engine/Source/Runtime/Core/Public/Containers/StringConv.h
--- a/Engine/Source/Runtime/Core/Public/Containers/StringConv.h
+++ b/Engine/Source/Runtime/Core/Public/Containers/StringConv.h
@@ -54,13 +54,6 @@ public:
 
 namespace StringConv
 {
-	constexpr const uint16 HIGH_SURROGATE_START_CODEPOINT = 0xD800;
-	constexpr const uint16 HIGH_SURROGATE_END_CODEPOINT   = 0xDBFF;
-	constexpr const uint16 LOW_SURROGATE_START_CODEPOINT  = 0xDC00;
-	constexpr const uint16 LOW_SURROGATE_END_CODEPOINT    = 0xDFFF;
-	constexpr const uint32 ENCODED_SURROGATE_START_CODEPOINT = 0x10000;
-	constexpr const uint32 ENCODED_SURROGATE_END_CODEPOINT   = 0x10FFFF;
-
 	/** Is the provided Codepoint within the range of valid codepoints? */
 	static FORCEINLINE bool IsValidCodepoint(const uint32 Codepoint)
 	{
@@ -295,7 +288,7 @@ public:
 	>
 	static FORCEINLINE int32 Convert(DestBufferType Dest, int32 DestLen, const SrcBufferType* Source, int32 SourceLen)
 	{
-		return Convert_Impl(Dest, DestLen, (const FromType*)Source, SourceLen);
+		return (int32)(FPlatformString::Convert(Dest, DestLen, (const FromType*)Source, SourceLen) - Dest);
 	}
 
 	/**
@@ -309,102 +302,7 @@ public:
 	>
 	static FORCEINLINE int32 ConvertedLength(const SrcBufferType* Source, int32 SourceLen)
 	{
-		UE4StringConv_Private::FCountingOutputIterator Dest;
-		const int32 DestLen = SourceLen * 4;
-		Convert_Impl(Dest, DestLen, (const FromType*)Source, SourceLen);
-
-		return Dest.GetCount();
-	}
-
-private:
-	template <typename DestBufferType>
-	static int32 Convert_Impl(DestBufferType& Dest, int32 DestLen, const FromType* Source, const int32 SourceLen)
-	{
-		DestBufferType DestStartingPosition = Dest;
-#if PLATFORM_TCHAR_IS_4_BYTES
-		for (int32 i = 0; i < SourceLen; ++i)
-		{
-			uint32 Codepoint = static_cast<uint32>(Source[i]);
-
-			if (!WriteCodepointToBuffer(Codepoint, Dest, DestLen))
-			{
-				// Could not write data, bail out
-				return -1;
-			}
-		}
-#else	// PLATFORM_TCHAR_IS_4_BYTES
-		uint32 HighSurrogate = MAX_uint32;
-
-		for (int32 i = 0; i < SourceLen; ++i)
-		{
-			const bool bHighSurrogateIsSet = HighSurrogate != MAX_uint32;
-			uint32 Codepoint = static_cast<uint32>(Source[i]);
-
-			// Check if this character is a high-surrogate
-			if (StringConv::IsHighSurrogate(Codepoint))
-			{
-				// Ensure we don't already have a high-surrogate set or end without a matching low-surrogate
-				if (bHighSurrogateIsSet || i == SourceLen - 1)
-				{
-					// Already have a high-surrogate in this pair or string ends with lone high-surrogate
-					// Write our stored value (will be converted into bogus character)
-					if (!WriteCodepointToBuffer(HighSurrogate, Dest, DestLen))
-					{
-						// Could not write data, bail out
-						return -1;
-					}
-				}
-
-				// Store our code point for our next character
-				HighSurrogate = Codepoint;
-				continue;
-			}
-
-			// If our High Surrogate is set, check if this character is the matching low-surrogate
-			if (bHighSurrogateIsSet)
-			{
-				if (StringConv::IsLowSurrogate(Codepoint))
-				{
-					const uint32 LowSurrogate = Codepoint;
-					// Combine our high and low surrogates together to a single Unicode codepoint
-					Codepoint = StringConv::EncodeSurrogate((uint16)HighSurrogate, (uint16)LowSurrogate);
-				}
-				else
-				{
-					// Did not find matching low-surrogate, write out a bogus character for our stored HighSurrogate
-					if (!WriteCodepointToBuffer(HighSurrogate, Dest, DestLen))
-					{
-						// Could not write data, bail out
-						return -1;
-					}
-				}
-
-				// Reset our high-surrogate now that we've used (or discarded) its value
-				HighSurrogate = MAX_uint32;
-			}
-
-			if (!WriteCodepointToBuffer(Codepoint, Dest, DestLen))
-			{
-				// Could not write data, bail out
-				return -1;
-			}
-		}
-#endif	// PLATFORM_TCHAR_IS_4_BYTES
-		return UE_PTRDIFF_TO_INT32(Dest - DestStartingPosition);
-	}
-
-	template <typename DestBufferType>
-	static bool WriteCodepointToBuffer(const uint32 Codepoint, DestBufferType& Dest, int32& DestLen)
-	{
-		int32 WrittenChars = Utf8FromCodepoint(Codepoint, Dest, DestLen);
-		if (WrittenChars < 1)
-		{
-			return false;
-		}
-
-		Dest += WrittenChars;
-		DestLen -= WrittenChars;
-		return true;
+		return FPlatformString::ConvertedLength<ToType>((const FromType*)Source, SourceLen);
 	}
 };
 
@@ -428,7 +326,7 @@ public:
 	>
 	static FORCEINLINE void Convert(ToType* Dest, const int32 DestLen, const SrcBufferType* Source, const int32 SourceLen)
 	{
-		Convert_Impl(Dest, DestLen, (const FromType*)Source, SourceLen);
+		FPlatformString::Convert(Dest, DestLen, (const FromType*)Source, SourceLen);
 	}
 
 	/**
@@ -444,319 +342,7 @@ public:
 	>
 	static int32 ConvertedLength(const SrcBufferType* Source, const int32 SourceLen)
 	{
-		UE4StringConv_Private::FCountingOutputIterator Dest;
-		Convert_Impl(Dest, MAX_int32, (const FromType*)Source, SourceLen);
-
-		return Dest.GetCount();
-	}
-
-private:
-	static uint32 CodepointFromUtf8(const FromType*& SourceString, const uint32 SourceLengthRemaining)
-	{
-		checkSlow(SourceLengthRemaining > 0);
-
-		const FromType* OctetPtr = SourceString;
-
-		uint32 Codepoint = 0;
-		uint32 Octet = (uint32) ((uint8) *SourceString);
-		uint32 Octet2, Octet3, Octet4;
-
-		if (Octet < 128)  // one octet char: 0 to 127
-		{
-			++SourceString;  // skip to next possible start of codepoint.
-			return Octet;
-		}
-		else if (Octet < 192)  // bad (starts with 10xxxxxx).
-		{
-			// Apparently each of these is supposed to be flagged as a bogus
-			//  char, instead of just resyncing to the next valid codepoint.
-			++SourceString;  // skip to next possible start of codepoint.
-			return UNICODE_BOGUS_CHAR_CODEPOINT;
-		}
-		else if (Octet < 224)  // two octets
-		{
-			// Ensure our string has enough characters to read from
-			if (SourceLengthRemaining < 2)
-			{
-				// Skip to end and write out a single char (we always have room for at least 1 char)
-				SourceString += SourceLengthRemaining;
-				return UNICODE_BOGUS_CHAR_CODEPOINT;
-			}
-
-			Octet -= (128+64);
-			Octet2 = (uint32) ((uint8) *(++OctetPtr));
-			if ((Octet2 & (128 + 64)) != 128)  // Format isn't 10xxxxxx?
-			{
-				++SourceString;  // Sequence was not valid UTF-8. Skip the first byte and continue.
-				return UNICODE_BOGUS_CHAR_CODEPOINT;
-			}
-
-			Codepoint = ((Octet << 6) | (Octet2 - 128));
-			if ((Codepoint >= 0x80) && (Codepoint <= 0x7FF))
-			{
-				SourceString += 2;  // skip to next possible start of codepoint.
-				return Codepoint;
-			}
-		}
-		else if (Octet < 240)  // three octets
-		{
-			// Ensure our string has enough characters to read from
-			if (SourceLengthRemaining < 3)
-			{
-				// Skip to end and write out a single char (we always have room for at least 1 char)
-				SourceString += SourceLengthRemaining;
-				return UNICODE_BOGUS_CHAR_CODEPOINT;
-			}
-
-			Octet -= (128+64+32);
-			Octet2 = (uint32) ((uint8) *(++OctetPtr));
-			if ((Octet2 & (128+64)) != 128)  // Format isn't 10xxxxxx?
-			{
-				++SourceString;  // Sequence was not valid UTF-8. Skip the first byte and continue.
-				return UNICODE_BOGUS_CHAR_CODEPOINT;
-			}
-
-			Octet3 = (uint32) ((uint8) *(++OctetPtr));
-			if ((Octet3 & (128+64)) != 128)  // Format isn't 10xxxxxx?
-			{
-				++SourceString;  // Sequence was not valid UTF-8. Skip the first byte and continue.
-				return UNICODE_BOGUS_CHAR_CODEPOINT;
-			}
-
-			Codepoint = ( ((Octet << 12)) | ((Octet2-128) << 6) | ((Octet3-128)) );
-
-			// UTF-8 characters cannot be in the UTF-16 surrogates range
-			if (StringConv::IsHighSurrogate(Codepoint) || StringConv::IsLowSurrogate(Codepoint))
-			{
-				++SourceString;  // Sequence was not valid UTF-8. Skip the first byte and continue.
-				return UNICODE_BOGUS_CHAR_CODEPOINT;
-			}
-
-			// 0xFFFE and 0xFFFF are illegal, too, so we check them at the edge.
-			if ((Codepoint >= 0x800) && (Codepoint <= 0xFFFD))
-			{
-				SourceString += 3;  // skip to next possible start of codepoint.
-				return Codepoint;
-			}
-		}
-		else if (Octet < 248)  // four octets
-		{
-			// Ensure our string has enough characters to read from
-			if (SourceLengthRemaining < 4)
-			{
-				// Skip to end and write out a single char (we always have room for at least 1 char)
-				SourceString += SourceLengthRemaining;
-				return UNICODE_BOGUS_CHAR_CODEPOINT;
-			}
-
-			Octet -= (128+64+32+16);
-			Octet2 = (uint32) ((uint8) *(++OctetPtr));
-			if ((Octet2 & (128+64)) != 128)  // Format isn't 10xxxxxx?
-			{
-				++SourceString;  // Sequence was not valid UTF-8. Skip the first byte and continue.
-				return UNICODE_BOGUS_CHAR_CODEPOINT;
-			}
-
-			Octet3 = (uint32) ((uint8) *(++OctetPtr));
-			if ((Octet3 & (128+64)) != 128)  // Format isn't 10xxxxxx?
-			{
-				++SourceString;  // Sequence was not valid UTF-8. Skip the first byte and continue.
-				return UNICODE_BOGUS_CHAR_CODEPOINT;
-			}
-
-			Octet4 = (uint32) ((uint8) *(++OctetPtr));
-			if ((Octet4 & (128+64)) != 128)  // Format isn't 10xxxxxx?
-			{
-				++SourceString;  // Sequence was not valid UTF-8. Skip the first byte and continue.
-				return UNICODE_BOGUS_CHAR_CODEPOINT;
-			}
-
-			Codepoint = ( ((Octet << 18)) | ((Octet2 - 128) << 12) |
-						((Octet3 - 128) << 6) | ((Octet4 - 128)) );
-			if ((Codepoint >= 0x10000) && (Codepoint <= 0x10FFFF))
-			{
-				SourceString += 4;  // skip to next possible start of codepoint.
-				return Codepoint;
-			}
-		}
-		// Five and six octet sequences became illegal in rfc3629.
-		//  We throw the codepoint away, but parse them to make sure we move
-		//  ahead the right number of bytes and don't overflow the buffer.
-		else if (Octet < 252)  // five octets
-		{
-			// Ensure our string has enough characters to read from
-			if (SourceLengthRemaining < 5)
-			{
-				// Skip to end and write out a single char (we always have room for at least 1 char)
-				SourceString += SourceLengthRemaining;
-				return UNICODE_BOGUS_CHAR_CODEPOINT;
-			}
-
-			Octet = (uint32) ((uint8) *(++OctetPtr));
-			if ((Octet & (128+64)) != 128)  // Format isn't 10xxxxxx?
-			{
-				++SourceString;  // Sequence was not valid UTF-8. Skip the first byte and continue.
-				return UNICODE_BOGUS_CHAR_CODEPOINT;
-			}
-
-			Octet = (uint32) ((uint8) *(++OctetPtr));
-			if ((Octet & (128+64)) != 128)  // Format isn't 10xxxxxx?
-			{
-				++SourceString;  // Sequence was not valid UTF-8. Skip the first byte and continue.
-				return UNICODE_BOGUS_CHAR_CODEPOINT;
-			}
-
-			Octet = (uint32) ((uint8) *(++OctetPtr));
-			if ((Octet & (128+64)) != 128)  // Format isn't 10xxxxxx?
-			{
-				++SourceString;  // Sequence was not valid UTF-8. Skip the first byte and continue.
-				return UNICODE_BOGUS_CHAR_CODEPOINT;
-			}
-
-			Octet = (uint32) ((uint8) *(++OctetPtr));
-			if ((Octet & (128+64)) != 128)  // Format isn't 10xxxxxx?
-			{
-				++SourceString;  // Sequence was not valid UTF-8. Skip the first byte and continue.
-				return UNICODE_BOGUS_CHAR_CODEPOINT;
-			}
-
-			SourceString += 5;  // skip to next possible start of codepoint.
-			return UNICODE_BOGUS_CHAR_CODEPOINT;
-		}
-
-		else  // six octets
-		{
-			// Ensure our string has enough characters to read from
-			if (SourceLengthRemaining < 6)
-			{
-				// Skip to end and write out a single char (we always have room for at least 1 char)
-				SourceString += SourceLengthRemaining;
-				return UNICODE_BOGUS_CHAR_CODEPOINT;
-			}
-
-			Octet = (uint32) ((uint8) *(++OctetPtr));
-			if ((Octet & (128+64)) != 128)  // Format isn't 10xxxxxx?
-			{
-				++SourceString;  // Sequence was not valid UTF-8. Skip the first byte and continue.
-				return UNICODE_BOGUS_CHAR_CODEPOINT;
-			}
-
-			Octet = (uint32) ((uint8) *(++OctetPtr));
-			if ((Octet & (128+64)) != 128)  // Format isn't 10xxxxxx?
-			{
-				++SourceString;  // Sequence was not valid UTF-8. Skip the first byte and continue.
-				return UNICODE_BOGUS_CHAR_CODEPOINT;
-			}
-
-			Octet = (uint32) ((uint8) *(++OctetPtr));
-			if ((Octet & (128+64)) != 128)  // Format isn't 10xxxxxx?
-			{
-				++SourceString;  // Sequence was not valid UTF-8. Skip the first byte and continue.
-				return UNICODE_BOGUS_CHAR_CODEPOINT;
-			}
-
-			Octet = (uint32) ((uint8) *(++OctetPtr));
-			if ((Octet & (128+64)) != 128)  // Format isn't 10xxxxxx?
-			{
-				++SourceString;  // Sequence was not valid UTF-8. Skip the first byte and continue.
-				return UNICODE_BOGUS_CHAR_CODEPOINT;
-			}
-
-			Octet = (uint32) ((uint8) *(++OctetPtr));
-			if ((Octet & (128+64)) != 128)  // Format isn't 10xxxxxx?
-			{
-				++SourceString;  // Sequence was not valid UTF-8. Skip the first byte and continue.
-				return UNICODE_BOGUS_CHAR_CODEPOINT;
-			}
-
-			SourceString += 6;  // skip to next possible start of codepoint.
-			return UNICODE_BOGUS_CHAR_CODEPOINT;
-		}
-
-		++SourceString;  // Sequence was not valid UTF-8. Skip the first byte and continue.
-		return UNICODE_BOGUS_CHAR_CODEPOINT;  // catch everything else.
-	}
-
-	/**
-	 * Read Source string, converting the data from UTF-8 into UTF-16, and placing these in the Destination
-	 */
-	template <typename DestBufferType>
-	static void Convert_Impl(DestBufferType& ConvertedBuffer, int32 DestLen, const FromType* Source, const int32 SourceLen)
-	{
-		const FromType* SourceEnd = Source + SourceLen;
-
-		const uint64 ExtendedCharMask = 0x8080808080808080;
-		while (Source < SourceEnd && DestLen > 0)
-		{
-			// In case we're given an unaligned pointer, we'll
-			// fallback to the slow path until properly aligned.
-			if (IsAligned(Source, 8))
-			{
-				// Fast path for most common case
-				while (Source < SourceEnd - 8 && DestLen >= 8)
-				{
-					// Detect any extended characters 8 chars at a time
-					if ((*(const uint64*)Source) & ExtendedCharMask)
-					{
-						// Move to slow path since we got extended characters to process
-						break;
-					}
-
-					// This should get unrolled on most compiler
-					// ROI of diminished return to vectorize this as we 
-					// would have to deal with alignment, endianness and
-					// rewrite the iterators to support bulk writes
-					for (int32 Index = 0; Index < 8; ++Index)
-					{
-						*(ConvertedBuffer++) = (ToType)(uint8)*(Source++);
-					}
-					DestLen -= 8;
-				}
-			}
-
-			// Slow path for extended characters
-			while (Source < SourceEnd && DestLen > 0)
-			{
-				// Read our codepoint, advancing the source pointer
-				uint32 Codepoint = CodepointFromUtf8(Source, UE_PTRDIFF_TO_UINT32(SourceEnd - Source));
-
-	#if !PLATFORM_TCHAR_IS_4_BYTES
-				// We want to write out two chars
-				if (StringConv::IsEncodedSurrogate(Codepoint))
-				{
-					// We need two characters to write the surrogate pair
-					if (DestLen >= 2)
-					{
-						uint16 HighSurrogate = 0;
-						uint16 LowSurrogate = 0;
-						StringConv::DecodeSurrogate(Codepoint, HighSurrogate, LowSurrogate);
-
-						*(ConvertedBuffer++) = (ToType)HighSurrogate;
-						*(ConvertedBuffer++) = (ToType)LowSurrogate;
-						DestLen -= 2;
-						continue;
-					}
-
-					// If we don't have space, write a bogus character instead (we should have space for it)
-					Codepoint = UNICODE_BOGUS_CHAR_CODEPOINT;
-				}
-				else if (Codepoint > StringConv::ENCODED_SURROGATE_END_CODEPOINT)
-				{
-					// Ignore values higher than the supplementary plane range
-					Codepoint = UNICODE_BOGUS_CHAR_CODEPOINT;
-				}
-	#endif	// !PLATFORM_TCHAR_IS_4_BYTES
-
-				*(ConvertedBuffer++) = (ToType)Codepoint;
-				--DestLen;
-
-				// Return to the fast path once aligned and back to simple ASCII chars
-				if (Codepoint < 128 && IsAligned(Source, 8))
-				{
-					break;
-				}
-			}
-		}
+		return FPlatformString::ConvertedLength<ToType>((const FromType*)Source, SourceLen);
 	}
 };
 
@@ -812,7 +398,7 @@ public:
 				*(OutputIterator++) = (ToType)LowSurrogate;
 			}
 		}
-		else if (Codepoint > StringConv::ENCODED_SURROGATE_END_CODEPOINT)
+		else if (Codepoint > ENCODED_SURROGATE_END_CODEPOINT)
 		{
 			// Ignore values higher than the supplementary plane range
 			*(OutputIterator++) = UNICODE_BOGUS_CHAR_CODEPOINT;
diff --git a/Engine/Source/Runtime/Core/Public/GenericPlatform/GenericPlatformString.h b/Engine/Source/Runtime/Core/Public/GenericPlatform/GenericPlatformString.h
--- a/Engine/Source/Runtime/Core/Public/GenericPlatform/GenericPlatformString.h
+++ b/Engine/Source/Runtime/Core/Public/GenericPlatform/GenericPlatformString.h
@@ -7,6 +7,33 @@
 #include "GenericPlatform/GenericPlatformStricmp.h"
 #include <type_traits>
 
+namespace UE::Core::Private
+{
+	// The Dest parameter is just used for overload resolution
+	CORE_API int32 GetConvertedLength(const UTF8CHAR* Dest, const ANSICHAR* Src, int32 SrcLen);
+	CORE_API int32 GetConvertedLength(const UTF8CHAR* Dest, const WIDECHAR* Src, int32 SrcLen);
+	CORE_API int32 GetConvertedLength(const UTF8CHAR* Dest, const UCS2CHAR* Src, int32 SrcLen);
+	CORE_API int32 GetConvertedLength(const ANSICHAR* Dest, const UTF8CHAR* Src, int32 SrcLen);
+	CORE_API int32 GetConvertedLength(const WIDECHAR* Dest, const UTF8CHAR* Src, int32 SrcLen);
+	CORE_API int32 GetConvertedLength(const UCS2CHAR* Dest, const UTF8CHAR* Src, int32 SrcLen);
+
+	CORE_API UTF8CHAR* Convert(UTF8CHAR* Dest, int32 DestLen, const ANSICHAR* Src, int32 SrcLen, UTF8CHAR BogusChar);
+	CORE_API UTF8CHAR* Convert(UTF8CHAR* Dest, int32 DestLen, const WIDECHAR* Src, int32 SrcLen, UTF8CHAR BogusChar);
+	CORE_API UTF8CHAR* Convert(UTF8CHAR* Dest, int32 DestLen, const UCS2CHAR* Src, int32 SrcLen, UTF8CHAR BogusChar);
+	CORE_API ANSICHAR* Convert(ANSICHAR* Dest, int32 DestLen, const UTF8CHAR* Src, int32 SrcLen, ANSICHAR BogusChar);
+	CORE_API WIDECHAR* Convert(WIDECHAR* Dest, int32 DestLen, const UTF8CHAR* Src, int32 SrcLen, WIDECHAR BogusChar);
+	CORE_API UCS2CHAR* Convert(UCS2CHAR* Dest, int32 DestLen, const UTF8CHAR* Src, int32 SrcLen, UCS2CHAR BogusChar);
+}
+
+// These will be moved inside GenericPlatformString.cpp when the platform layer handles UTF-16
+// instead of StringConv.h.
+#define HIGH_SURROGATE_START_CODEPOINT    ((uint16)0xD800)
+#define HIGH_SURROGATE_END_CODEPOINT      ((uint16)0xDBFF)
+#define LOW_SURROGATE_START_CODEPOINT     ((uint16)0xDC00)
+#define LOW_SURROGATE_END_CODEPOINT       ((uint16)0xDFFF)
+#define ENCODED_SURROGATE_START_CODEPOINT ((uint32)0x10000)
+#define ENCODED_SURROGATE_END_CODEPOINT   ((uint32)0x10FFFF)
+
 /**
  * Generic string implementation for most platforms
  */
@@ -148,7 +175,7 @@ struct FGenericPlatformString : public FGenericPlatformStricmp
 	};
 
 	/**
-	 * Converts the [Src, Src+SrcSize) string range from SourceChar to DestChar and writes it to the [Dest, Dest+DestSize) range.
+	 * Converts the [Src, Src+SrcSize) string range from SourceEncoding to DestEncoding and writes it to the [Dest, Dest+DestSize) range.
 	 * The Src range should contain a null terminator if a null terminator is required in the output.
 	 * If the Dest range is not big enough to hold the converted output, NULL is returned.  In this case, nothing should be assumed about the contents of Dest.
 	 *
@@ -171,7 +198,7 @@ struct FGenericPlatformString : public FGenericPlatformStricmp
 
 			return (DestEncoding*)Memcpy(Dest, Src, SrcSize * sizeof(SourceEncoding)) + SrcSize;
 		}
-		else
+		else if constexpr (TIsFixedWidthEncoding<SourceEncoding>::Value && TIsFixedWidthEncoding<DestEncoding>::Value)
 		{
 			const int32 Size = DestSize <= SrcSize ? DestSize : SrcSize;
 			bool bInvalidChars = false;
@@ -197,6 +224,10 @@ struct FGenericPlatformString : public FGenericPlatformStricmp
 
 			return DestSize < SrcSize ? nullptr : Dest + Size;
 		}
+		else
+		{
+			return UE::Core::Private::Convert(Dest, DestSize, Src, SrcSize, BogusChar);
+		}
 	}
 
 
@@ -209,9 +240,16 @@ struct FGenericPlatformString : public FGenericPlatformStricmp
 	 * @return         The number of DestChar elements that Src will be converted into.
 	 */
 	template <typename DestEncoding, typename SourceEncoding>
-	static typename TEnableIf<TIsFixedWidthEncoding<SourceEncoding>::Value && TIsFixedWidthEncoding<DestEncoding>::Value, int32>::Type ConvertedLength(const SourceEncoding* Src, int32 SrcSize)
+	static int32 ConvertedLength(const SourceEncoding* Src, int32 SrcSize)
 	{
-		return SrcSize;
+		if constexpr (std::is_same_v<SourceEncoding, DestEncoding> || (TIsFixedWidthEncoding<SourceEncoding>::Value && TIsFixedWidthEncoding<DestEncoding>::Value))
+		{
+			return SrcSize;
+		}
+		else
+		{
+			return UE::Core::Private::GetConvertedLength((DestEncoding*)nullptr, Src, SrcSize);
+		}
 	}
 
 	CORE_API static int32 Strncmp(const ANSICHAR* String1, const ANSICHAR* String2, SIZE_T Count);
diff --git a/Engine/Source/Runtime/Core/Public/Templates/UnrealTypeTraits.h b/Engine/Source/Runtime/Core/Public/Templates/UnrealTypeTraits.h
--- a/Engine/Source/Runtime/Core/Public/Templates/UnrealTypeTraits.h
+++ b/Engine/Source/Runtime/Core/Public/Templates/UnrealTypeTraits.h
@@ -101,6 +101,7 @@ template<typename T> struct TIsCharType           { enum { Value = false }; };
 template<>           struct TIsCharType<ANSICHAR> { enum { Value = true  }; };
 template<>           struct TIsCharType<UCS2CHAR> { enum { Value = true  }; };
 template<>           struct TIsCharType<WIDECHAR> { enum { Value = true  }; };
+template<>           struct TIsCharType<UTF8CHAR> { enum { Value = true  }; };
 
 /**
  * TFormatSpecifier, only applies to numeric types
