minor TaskGraph and Tasks tests cleanup, replicated some stress tests on the new API for perf comparison

diff --git a/Engine/Source/Runtime/Core/Private/Tests/Async/TaskGraphTest.cpp b/Engine/Source/Runtime/Core/Private/Tests/Async/TaskGraphTest.cpp
--- a/Engine/Source/Runtime/Core/Private/Tests/Async/TaskGraphTest.cpp
+++ b/Engine/Source/Runtime/Core/Private/Tests/Async/TaskGraphTest.cpp
@@ -1296,9 +1296,9 @@ PRAGMA_ENABLE_DEPRECATION_WARNINGS
 		}
 	}
 
-	IMPLEMENT_SIMPLE_AUTOMATION_TEST(FPerfTest, "System.Core.Async.TaskGraph.PerfTest", EAutomationTestFlags::ApplicationContextMask | EAutomationTestFlags::EngineFilter);
+	IMPLEMENT_SIMPLE_AUTOMATION_TEST(FTaskGraphPerfTest, "System.Core.Async.TaskGraph.PerfTest", EAutomationTestFlags::ApplicationContextMask | EAutomationTestFlags::EngineFilter);
 
-	bool FPerfTest::RunTest(const FString& Parameters)
+	bool FTaskGraphPerfTest::RunTest(const FString& Parameters)
 	{
 		TRACE_CPUPROFILER_EVENT_SCOPE(TaskGraphTests_PerfTest);
 
diff --git a/Engine/Source/Runtime/Core/Private/Tests/Tasks/TasksTest.cpp b/Engine/Source/Runtime/Core/Private/Tests/Tasks/TasksTest.cpp
--- a/Engine/Source/Runtime/Core/Private/Tests/Tasks/TasksTest.cpp
+++ b/Engine/Source/Runtime/Core/Private/Tests/Tasks/TasksTest.cpp
@@ -22,13 +22,12 @@ namespace UE { namespace TasksTests
 
 	IMPLEMENT_SIMPLE_AUTOMATION_TEST(FTasksBasicTest, "System.Core.Tasks.Basic", EAutomationTestFlags::ApplicationContextMask | EAutomationTestFlags::EngineFilter);
 
+	template<uint32 SpawnerGroupsNum, uint32 SpawnersPerGroupNum>
 	void BasicStressTest()
 	{
-		constexpr uint32 SpawnerGroupsNum = 50;
 		TArray<FTask> SpawnerGroups;
 		SpawnerGroups.Reserve(SpawnerGroupsNum);
 
-		constexpr uint32 SpawnersPerGroupNum = 100;
 		constexpr uint32 TasksNum = SpawnerGroupsNum * SpawnersPerGroupNum;
 		TArray<FTask> Spawners;
 		Spawners.AddDefaulted(TasksNum);
@@ -43,8 +42,7 @@ namespace UE { namespace TasksTests
 				[
 					Spawners = &Spawners[GroupIndex * SpawnersPerGroupNum],
 					Tasks = &Tasks[GroupIndex * SpawnersPerGroupNum],
-					&TasksExecutedNum,
-					SpawnersPerGroupNum
+					&TasksExecutedNum
 				]
 				{
 					for (uint32 SpawnerIndex = 0; SpawnerIndex != SpawnersPerGroupNum; ++SpawnerIndex)
@@ -216,13 +214,14 @@ namespace UE { namespace TasksTests
 			Task = {};
 		}
 
-		UE_BENCHMARK(5, BasicStressTest);
+		UE_BENCHMARK(5, BasicStressTest<1000, 1000>);
 
 		return true;
 	}
 
 	IMPLEMENT_SIMPLE_AUTOMATION_TEST(FTasksPipeTest, "System.Core.Tasks.Pipe", EAutomationTestFlags::ApplicationContextMask | EAutomationTestFlags::EngineFilter);
 
+	template<uint32 SpawnerGroupsNum, uint32 SpawnersPerGroupNum>
 	void PipeStressTest();
 
 	bool FTasksPipeTest::RunTest(const FString& Parameters)
@@ -337,19 +336,18 @@ namespace UE { namespace TasksTests
 			Task.Wait();
 		}
 
-		UE_BENCHMARK(5, PipeStressTest);
+		UE_BENCHMARK(5, PipeStressTest<500, 500>);
 
 		return true;
 	}
 
 	// stress test for named thread tasks, checks spawning a large number of tasks from multiple threads and executing them
+	template<uint32 SpawnerGroupsNum, uint32 SpawnersPerGroupNum>
 	void PipeStressTest()
 	{
-		constexpr uint32 SpawnerGroupsNum = 50;
 		TArray<FTask> SpawnerGroups;
 		SpawnerGroups.Reserve(SpawnerGroupsNum);
 
-		constexpr uint32 SpawnersPerGroupNum = 100;
 		constexpr uint32 TasksNum = SpawnerGroupsNum * SpawnersPerGroupNum;
 		TArray<FTask> Spawners;
 		Spawners.AddDefaulted(TasksNum);
@@ -369,7 +367,6 @@ namespace UE { namespace TasksTests
 					Tasks = &Tasks[GroupIndex * SpawnersPerGroupNum],
 					&bExecuting,
 					&TasksExecutedNum,
-					SpawnersPerGroupNum,
 					&Pipe
 				]
 				{
@@ -635,7 +632,210 @@ namespace UE { namespace TasksTests
 			Task.Wait();
 		}
 
-		UE_BENCHMARK(3, DependenciesPerfTest<200, 50, 1000>);
+		UE_BENCHMARK(5, DependenciesPerfTest<200, 10, 1000>);
+
+		return true;
+	}
+
+	template<int NumTasks>
+	void TestPerfBasic()
+	{
+		TArray<FTask> Tasks;
+		Tasks.Reserve(NumTasks);
+
+		for (int32 TaskIndex = 0; TaskIndex < NumTasks; ++TaskIndex)
+		{
+			Tasks.Emplace(Launch(UE_SOURCE_LOCATION, [] {}));
+		}
+
+		Wait(Tasks);
+	}
+
+	template<int32 NumTasks, int32 BatchSize>
+	void TestPerfBatch()
+	{
+		static_assert(NumTasks % BatchSize == 0, "`NumTasks` must be divisible by `BatchSize`");
+		constexpr int32 NumBatches = NumTasks / BatchSize;
+
+		TArray<FTask> Batches;
+		Batches.Reserve(NumBatches);
+		TArray<FTask> Tasks;
+		Tasks.AddDefaulted(NumTasks);
+
+		for (int32 BatchIndex = 0; BatchIndex < NumBatches; ++BatchIndex)
+		{
+			Batches.Add(Launch(UE_SOURCE_LOCATION,
+				[&Tasks, BatchIndex]
+				{
+					for (int32 TaskIndex = 0; TaskIndex < BatchSize; ++TaskIndex)
+					{
+						Tasks[BatchIndex * BatchSize + TaskIndex] = Launch(UE_SOURCE_LOCATION, [] {});
+					}
+				}
+			));
+		}
+
+		Wait(Batches);
+		Wait(Tasks);
+	}
+
+	template<int32 NumTasks, int32 BatchSize>
+	void TestPerfBatchOptimised()
+	{
+		static_assert(NumTasks % BatchSize == 0, "`NumTasks` must be divisible by `BatchSize`");
+		constexpr int32 NumBatches = NumTasks / BatchSize;
+
+		FTaskEvent SpawnSignal{ UE_SOURCE_LOCATION };
+		TArray<FTask> AllDone;
+
+		for (int32 BatchIndex = 0; BatchIndex < NumBatches; ++BatchIndex)
+		{
+			AllDone.Add(Launch(UE_SOURCE_LOCATION,
+				[]
+				{
+					FTaskEvent RunSignal{ UE_SOURCE_LOCATION };
+					for (int32 TaskIndex = 0; TaskIndex < BatchSize; ++TaskIndex)
+					{
+						//AddNested(Launch(UE_SOURCE_LOCATION, [] {}, RunSignal));
+					}
+					RunSignal.Trigger();
+				},
+				SpawnSignal
+			));
+		}
+
+		SpawnSignal.Trigger();
+		Wait(AllDone);
+	}
+
+	template<int NumTasks>
+	void TestLatency()
+	{
+		for (uint32 TaskIndex = 0; TaskIndex != NumTasks; ++TaskIndex)
+		{
+			Launch(UE_SOURCE_LOCATION, [] {}).Wait();
+		}
+	}
+
+	template<uint32 NumTasks>
+	void TestFGraphEventPerf()
+	{
+		FTaskEvent Prereq{ UE_SOURCE_LOCATION };
+		std::atomic<uint32> CompletedTasks{ 0 };
+
+		TArray<FTask> Tasks;
+		for (int i = 0; i != NumTasks; ++i)
+		{
+			Tasks.Add(Launch(UE_SOURCE_LOCATION,
+				[&Prereq, &CompletedTasks]()
+				{
+					Prereq.Wait();
+					++CompletedTasks;
+				}
+			));
+		}
+
+		Prereq.Trigger();
+		Wait(Tasks);
+
+		check(CompletedTasks == NumTasks);
+	}
+
+	template<int NumTasks>
+	void TestSpawning()
+	{
+		{
+			TArray<FTask> Tasks;
+			Tasks.Reserve(NumTasks);
+			//double StartTime = FPlatformTime::Seconds();
+			for (uint32 TaskNo = 0; TaskNo != NumTasks; ++TaskNo)
+			{
+				Tasks.Add(Launch(UE_SOURCE_LOCATION, [] {}));
+			}
+
+			//double Duration = FPlatformTime::Seconds() - StartTime;
+			//UE_LOG(LogTemp, Display, TEXT("Spawning %d empty trackable tasks took %f secs"), NumTasks, Duration);
+
+			Wait(Tasks);
+		}
+		{
+			double StartTime = FPlatformTime::Seconds();
+			for (uint32 TaskNo = 0; TaskNo != NumTasks; ++TaskNo)
+			{
+				Launch(UE_SOURCE_LOCATION, [] {});
+			}
+
+			//double Duration = FPlatformTime::Seconds() - StartTime;
+			//UE_LOG(LogTemp, Display, TEXT("Spawning %d empty non-trackable tasks took %f secs"), NumTasks, Duration);
+		}
+	}
+
+	template<int NumTasks>
+	void TestBatchSpawning()
+	{
+		//double StartTime = FPlatformTime::Seconds();
+		FTaskEvent Prereq{ UE_SOURCE_LOCATION };
+		TArray<FTask> Tasks;
+		Tasks.Reserve(NumTasks);
+		for (uint32 TaskNo = 0; TaskNo != NumTasks; ++TaskNo)
+		{
+			Tasks.Add(Launch(UE_SOURCE_LOCATION, [] {}, Prereq));
+		}
+
+		//double SpawnedTime = FPlatformTime::Seconds();
+		Prereq.Trigger();
+
+		//double EndTime = FPlatformTime::Seconds();
+		//UE_LOG(LogTemp, Display, TEXT("Spawning %d empty non-trackable tasks took %f secs total, %f secs spawning and %f secs dispatching"), NumTasks, EndTime - StartTime, SpawnedTime - StartTime, EndTime - SpawnedTime);
+
+		Wait(Tasks);
+	}
+
+	template<int64 NumBatches, int64 NumTasksPerBatch>
+	void TestWorkStealing()
+	{
+		TArray<FTask> Batches;
+		Batches.Reserve(NumBatches);
+
+		TArray<FTask> Tasks[NumBatches];
+		for (int32 BatchIndex = 0; BatchIndex != NumBatches; ++BatchIndex)
+		{
+			Tasks[BatchIndex].Reserve(NumBatches * NumTasksPerBatch);
+			Batches.Add(Launch(UE_SOURCE_LOCATION,
+				[&Tasks, BatchIndex]()
+				{
+					for (int32 TaskIndex = 0; TaskIndex < NumTasksPerBatch; ++TaskIndex)
+					{
+						Tasks[BatchIndex].Add(Launch(UE_SOURCE_LOCATION, [] {}));
+					}
+				}
+			));
+		}
+
+		Wait(Batches);
+		for (int32 BatchIndex = 0; BatchIndex != NumBatches; ++BatchIndex)
+		{
+			Wait(Tasks[BatchIndex]);
+		}
+	}
+
+	IMPLEMENT_SIMPLE_AUTOMATION_TEST(FTasksPerfTest, "System.Core.Async.TaskGraph.PerfTest", EAutomationTestFlags::ApplicationContextMask | EAutomationTestFlags::EngineFilter);
+
+	bool FTasksPerfTest::RunTest(const FString& Parameters)
+	{
+		TRACE_CPUPROFILER_EVENT_SCOPE(TaskGraphTests_PerfTest);
+
+		//UE_BENCHMARK(1, TestSpawning<100000>);
+		//return true;
+
+		UE_BENCHMARK(5, TestPerfBasic<100000>);
+		UE_BENCHMARK(5, TestPerfBatch<100000, 100>);
+		UE_BENCHMARK(5, TestPerfBatchOptimised<100000, 100>);
+		UE_BENCHMARK(5, TestLatency<10000>);
+		//UE_BENCHMARK(5, TestFGraphEventPerf<100000>); // stack overflow
+		UE_BENCHMARK(5, TestWorkStealing<100, 1000>);
+		UE_BENCHMARK(5, TestSpawning<100000>);
+		UE_BENCHMARK(5, TestBatchSpawning<100000>);
 
 		return true;
 	}
