TFunction::Reset() and TUniqueFunction::Reset() added, equivalent to Func = nullptr; but less surprising.

diff --git a/Engine/Source/Runtime/Core/Public/Templates/Function.h b/Engine/Source/Runtime/Core/Public/Templates/Function.h
--- a/Engine/Source/Runtime/Core/Public/Templates/Function.h
+++ b/Engine/Source/Runtime/Core/Public/Templates/Function.h
@@ -975,6 +975,14 @@ public:
 		return *this;
 	}
 
+	/**
+	 * Removes any bound callable from the TFunction, restoring it to the default 'empty' state.
+	 */
+	void Reset()
+	{
+		*this = nullptr;
+	}
+
 	/**
 	 * Tests if the TFunction is callable.
 	 */
@@ -1071,6 +1079,14 @@ public:
 	TUniqueFunction& operator=(const TUniqueFunction& Other) = delete;
 	~TUniqueFunction() = default;
 
+	/**
+	 * Removes any bound callable from the TFunction, restoring it to the default 'empty' state.
+	 */
+	void Reset()
+	{
+		*this = nullptr;
+	}
+
 	/**
 	 * Tests if the TUniqueFunction is callable.
 	 */
