LLM: Fix maxcount calculation to not integer overflow before capacity integer overflows. This expands the number of allocations we can fit into 32-bit LLM tracking, which is needed for some large projects that still want to use 32-bit LLM tracking.
#rnx

diff --git a/Engine/Source/Runtime/Core/Private/HAL/LowLevelMemoryUtils.h b/Engine/Source/Runtime/Core/Private/HAL/LowLevelMemoryUtils.h
--- a/Engine/Source/Runtime/Core/Private/HAL/LowLevelMemoryUtils.h
+++ b/Engine/Source/Runtime/Core/Private/HAL/LowLevelMemoryUtils.h
@@ -269,14 +269,23 @@ public:
 		}
 		else
 		{
-			SizeType MaxCount = (Margin * Capacity) / 256U;
+			// Be careful with the multiplication to avoid overflow, while still using only integer arithmetic for speed
+			// Margin/256U is the actual float we want to multiply by. The exact number doesn't matter, so use one order of operations
+			// when Capacity is low, and another when it is high
+			SizeType MaxCount = (Capacity >= 256U*256U ? (Capacity/256U)*Margin : (Margin * Capacity) / 256U);
 			if (Count >= MaxCount)
 			{
+				if (Capacity * 2 < Capacity)
+				{
+					// Trying to issue a check statement here will cause reentry into this function, use PLATFORM_BREAK directly instead
+					FPlatformMisc::LowLevelOutputDebugString(TEXT("LLM Error: Integer overflow in LLMap::Add, Capacity has reached its maximum size.\n"));
+					PLATFORM_BREAK();
+				}
 				if (Count > MaxCount)
 				{
-					// This shouldn't happen, because Count is only incremented here, and Capacity is only changed here, and Margin does not change, so Count should equal MaxCount before it goes over it
-					FPlatformMisc::LowLevelOutputDebugString(TEXT("LLM Error: Integer overflow in LLMap::Add, Count > MaxCount.\n"));
 					// Trying to issue a check statement here will cause reentry into this function, use PLATFORM_BREAK directly instead
+					// This shouldn't happen, because Count is only incremented here, and Capacity is only changed here, and Margin does not change, so Count should equal MaxCount before it goes over it
+					FPlatformMisc::LowLevelOutputDebugString(TEXT("LLM Assertion failure: Count > MaxCount.\n"));
 					PLATFORM_BREAK();
 				}
 				Grow();
@@ -612,7 +621,17 @@ private:
 	{
 		SizeType p = 2;
 		while (p < value)
-			p *= 2;
+		{
+			SizeType Nextp = p * 2;
+			if (Nextp < p)
+			{
+				// Trying to issue a check statement here will cause reentry into allocation, use PLATFORM_BREAK directly instead
+				FPlatformMisc::LowLevelOutputDebugString(TEXT("LLM Error: Integer overflow in LLMap::Add, GetNextPow2 called on value > 2^(NumBits-1).\n"));
+				PLATFORM_BREAK();
+				break;
+			}
+			p = Nextp;
+		}
 		return p;
 	}
 
