LLM FLlmCsvProfilerWriter implementation for UE5
- LLM stats are now written to the CSV profiler if it's running
- Default tracker stats appear in the LLM/ category
- Platform tracker stats appear in the LLMPlatform/ category
#tests presubmit

diff --git a/Engine/Source/Runtime/Core/Private/HAL/LowLevelMemTracker.cpp b/Engine/Source/Runtime/Core/Private/HAL/LowLevelMemTracker.cpp
--- a/Engine/Source/Runtime/Core/Private/HAL/LowLevelMemTracker.cpp
+++ b/Engine/Source/Runtime/Core/Private/HAL/LowLevelMemTracker.cpp
@@ -40,6 +40,13 @@ UE_TRACE_EVENT_BEGIN(LLM, TagValue)
 	UE_TRACE_EVENT_FIELD(int64[], Values)
 UE_TRACE_EVENT_END()
 
+#define LLM_CSV_PROFILER_WRITER_ENABLED CSV_PROFILER 
+
+#if LLM_CSV_PROFILER_WRITER_ENABLED
+	CSV_DEFINE_CATEGORY(LLM, true);
+	CSV_DEFINE_CATEGORY(LLMPlatform, true);
+#endif
+
 TAutoConsoleVariable<int32> CVarLLMTrackPeaks(TEXT("LLM.TrackPeaks"), 0, TEXT("Track peak memory in each category since process start rather than current frame's value."));
 
 TAutoConsoleVariable<int32> CVarLLMWriteInterval(
@@ -242,6 +249,25 @@ namespace LLMPrivate
 		bool					bTrackerSpecSent = false;
 	};
 
+	/**
+	 * FLLMCsvWriter: class for writing out LLM stats to the Csv Profiler
+	 */
+	class FLLMCsvProfilerWriter
+	{
+	public:
+		FLLMCsvProfilerWriter();
+		void Clear();
+		void SetTracker(ELLMTracker InTracker);
+		void Publish(FLowLevelMemTracker& LLMRef, const FTrackerTagSizeMap& TagSizes, const FTagData* OverrideTrackedTotalName, const FTagData* OverrideUntaggedName, int64 TrackedTotal, bool bTrackPeaks);
+	protected:
+		void RecordTagToCsv(int32 CsvCategoryIndex, const FTagData* TagData, int64 TagSize);
+#if LLM_CSV_PROFILER_WRITER_ENABLED 
+		ELLMTracker Tracker;
+		TFastPointerLLMMap<const FTagData*, FName> TagDataToCsvStatName;
+#endif
+	};
+
+
 	/** per thread state in an LLMTracker */
 	class FLLMThreadState
 	{
@@ -325,6 +351,7 @@ namespace LLMPrivate
 		void PublishStats(bool bTrackPeaks);
 		void PublishCsv(bool bTrackPeaks);
 		void PublishTrace(bool bTrackPeaks);
+		void PublishCsvProfiler(bool bTrackPeaks);
 
 		struct FLowLevelAllocInfo
 		{
@@ -428,6 +455,7 @@ namespace LLMPrivate
 
 		FLLMCsvWriter CsvWriter;
 		FLLMTraceWriter TraceWriter;
+		FLLMCsvProfilerWriter CsvProfilerWriter;
 
 		double LastTrimTime;
 
@@ -900,6 +928,14 @@ void FLowLevelMemTracker::PublishDataPerFrame(const TCHAR* LogName)
 			DefaultTracker.PublishTrace(bTrackPeaks);
 			PlatformTracker.PublishTrace(bTrackPeaks);
 		}
+
+#if LLM_CSV_PROFILER_WRITER_ENABLED
+		if (FCsvProfiler::Get()->IsCapturing())
+		{
+			DefaultTracker.PublishCsvProfiler(bTrackPeaks);
+			PlatformTracker.PublishCsvProfiler(bTrackPeaks);
+		}
+#endif
 	}
 
 	if (LogName != nullptr)
@@ -2762,6 +2798,7 @@ namespace LLMPrivate
 	{
 		CsvWriter.SetTracker(Tracker);
 		TraceWriter.SetTracker(Tracker);
+		CsvProfilerWriter.SetTracker(Tracker);
 
 		AllocationMap.SetAllocator(InAllocator);
 	}
@@ -3064,6 +3101,7 @@ namespace LLMPrivate
 		}
 		CsvWriter.Clear();
 		TraceWriter.Clear();
+		CsvProfilerWriter.Clear();
 	}
 
 	void FLLMTracker::SetTotalTags(const FTagData* InOverrideUntaggedTagData, const FTagData* InOverrideTrackedTotalTagData)
@@ -3184,6 +3222,11 @@ namespace LLMPrivate
 		TraceWriter.Publish(LLMRef, TagSizes, OverrideTrackedTotalTagData, OverrideUntaggedTagData, TrackedTotal, bTrackPeaks);
 	}
 
+	void FLLMTracker::PublishCsvProfiler(bool bTrackPeaks)
+	{
+		CsvProfilerWriter.Publish(LLMRef, TagSizes, OverrideTrackedTotalTagData, OverrideUntaggedTagData, TrackedTotal, bTrackPeaks);
+	}
+
 	void FLLMTracker::OnTagsResorted(FTagDataArray& OldTagDatas)
 	{
 #if LLM_ENABLED_FULL_TAGS
@@ -3893,6 +3936,76 @@ namespace LLMPrivate
 				<< TagsSpec.Name(*TagName, TagName.Len());
 		}
 	}
+
+	/*
+	 * FLLMCsvProfilerWriter implementation
+	*/
+	FLLMCsvProfilerWriter::FLLMCsvProfilerWriter()
+	{
+	}
+
+	void FLLMCsvProfilerWriter::Clear()
+	{
+#if LLM_CSV_PROFILER_WRITER_ENABLED
+		TagDataToCsvStatName.Empty();
+#endif
+	}
+
+	void FLLMCsvProfilerWriter::SetTracker(ELLMTracker InTracker)
+	{
+#if LLM_CSV_PROFILER_WRITER_ENABLED
+		check(InTracker == ELLMTracker::Platform || InTracker == ELLMTracker::Default);
+		Tracker = InTracker;
+#endif
+	}
+
+	void FLLMCsvProfilerWriter::Publish(FLowLevelMemTracker& LLMRef, const FTrackerTagSizeMap& TagSizes, const FTagData* OverrideTrackedTotalTagData, const FTagData* OverrideUntaggedTagData, int64 TrackedTotal, bool bTrackPeaks)
+	{
+#if LLM_CSV_PROFILER_WRITER_ENABLED
+		int32 CsvCategoryIndex = (Tracker == ELLMTracker::Platform) ? CSV_CATEGORY_INDEX(LLMPlatform) : CSV_CATEGORY_INDEX(LLM);
+
+		if (OverrideTrackedTotalTagData)
+		{
+			RecordTagToCsv(CsvCategoryIndex, OverrideTrackedTotalTagData, TrackedTotal);
+		}
+
+		if (OverrideUntaggedTagData)
+		{
+			const FTagData* TagData = LLMRef.FindTagData(TagName_Untagged);
+			const FTrackerTagSizeData* AllocationData = TagData ? TagSizes.Find(TagData) : nullptr;
+			RecordTagToCsv(CsvCategoryIndex, OverrideUntaggedTagData, AllocationData ? AllocationData->GetSize(bTrackPeaks) : 0);
+		}
+
+		for (const TPair<const FTagData*, FTrackerTagSizeData>& It : TagSizes)
+		{
+			const FTagData* TagData = It.Key;
+			if (OverrideUntaggedTagData && TagData->GetName() == TagName_Untagged)
+			{
+				// Handled separately by OverrideUntaggedTagData
+				continue;
+			}
+			const FTrackerTagSizeData* AllocationData = TagSizes.Find(TagData);
+			RecordTagToCsv(CsvCategoryIndex, TagData, AllocationData ? AllocationData->GetSize(bTrackPeaks) : 0);
+		}
+#endif
+	}
+
+	void FLLMCsvProfilerWriter::RecordTagToCsv(int32 CsvCategoryIndex, const FTagData* TagData, int64 TagSize)
+	{
+#if LLM_CSV_PROFILER_WRITER_ENABLED
+		FName NewCsvStatName;
+		FName* CsvStatNamePtr = TagDataToCsvStatName.Find(TagData);
+		if (CsvStatNamePtr == nullptr)
+		{
+			NewCsvStatName = FName(TagData->GetDisplayPath());
+			TagDataToCsvStatName.Add(TagData, NewCsvStatName);
+			CsvStatNamePtr = &NewCsvStatName;
+		}
+		FCsvProfiler::RecordCustomStat(*CsvStatNamePtr, CsvCategoryIndex, (float)((double)TagSize / (1024.0 * 1024.0)), ECsvCustomStatOp::Set);
+#endif
+	}
+
+
 }
 }
 
diff --git a/Engine/Source/Runtime/Core/Public/HAL/LowLevelMemTracker.h b/Engine/Source/Runtime/Core/Public/HAL/LowLevelMemTracker.h
--- a/Engine/Source/Runtime/Core/Public/HAL/LowLevelMemTracker.h
+++ b/Engine/Source/Runtime/Core/Public/HAL/LowLevelMemTracker.h
@@ -353,6 +353,7 @@ namespace LLMPrivate
 	class FLLMCsvWriter;
 	class FLLMThreadState;
 	class FLLMTraceWriter;
+	class FLLMCsvProfilerWriter;
 	class FLLMTracker;
 
 	namespace AllocatorPrivate
@@ -573,6 +574,7 @@ private:
 	friend class UE::LLMPrivate::FLLMTracker;
 	friend class UE::LLMPrivate::FLLMThreadState;
 	friend class UE::LLMPrivate::FLLMTraceWriter;
+	friend class UE::LLMPrivate::FLLMCsvProfilerWriter;
 	friend void GlobalRegisterTagDeclaration(FLLMTagDeclaration& TagDeclaration);
 
 	UE::LLMPrivate::FLLMAllocator Allocator;
