Fixing HoloLens compile.

diff --git a/Engine/Source/Runtime/Core/Private/HoloLens/HoloLensRunnableThread.h b/Engine/Source/Runtime/Core/Private/HoloLens/HoloLensRunnableThread.h
--- a/Engine/Source/Runtime/Core/Private/HoloLens/HoloLensRunnableThread.h
+++ b/Engine/Source/Runtime/Core/Private/HoloLens/HoloLensRunnableThread.h
@@ -99,7 +99,7 @@ public:
 
 	FRunnableThreadHoloLens()
 		: Thread(NULL)
-		, ThreadInitSyncEvent(NULL)
+		, ThreadInitSyncEvent()
 		, ThreadPriority(TPri_Normal)
 	{
 
