Add context switch and stack sampling trace collection on Windows using ETW

diff --git a/Engine/Source/Developer/TraceServices/Private/Analyzers/PlatformEventTraceAnalysis.cpp b/Engine/Source/Developer/TraceServices/Private/Analyzers/PlatformEventTraceAnalysis.cpp
new file mode 100644
--- /dev/null
+++ b/Engine/Source/Developer/TraceServices/Private/Analyzers/PlatformEventTraceAnalysis.cpp
@@ -0,0 +1,59 @@
+// Copyright Epic Games, Inc. All Rights Reserved.
+#include "PlatformEventTraceAnalysis.h"
+#include "TraceServices/Model/AnalysisSession.h"
+#include "Model/ContextSwitchesPrivate.h"
+#include "Model/StackSamplesPrivate.h"
+
+namespace TraceServices
+{
+
+FPlatformEventTraceAnalyzer::FPlatformEventTraceAnalyzer(IAnalysisSession& InSession,
+														 FContextSwitchProvider& ContextSwitchProvider,
+														 FStackSampleProvider& StackSampleProvider)
+	: Session(InSession)
+	, ContextSwitchProvider(ContextSwitchProvider)
+	, StackSampleProvider(StackSampleProvider)
+{
+}
+
+void FPlatformEventTraceAnalyzer::OnAnalysisBegin(const FOnAnalysisContext& Context)
+{
+	auto& Builder = Context.InterfaceBuilder;
+
+	Builder.RouteEvent(RouteId_ContextSwitch, "PlatformEvent", "ContextSwitch");
+	Builder.RouteEvent(RouteId_StackSample, "PlatformEvent", "StackSample");
+}
+
+bool FPlatformEventTraceAnalyzer::OnEvent(uint16 RouteId, EStyle Style, const FOnEventContext& Context)
+{
+	FAnalysisSessionEditScope _(Session);
+
+	const auto& EventData = Context.EventData;
+	switch (RouteId)
+	{
+
+	case RouteId_ContextSwitch:
+	{
+		double Start = Context.EventTime.AsSeconds(EventData.GetValue<uint64>("StartTime"));
+		double End = Context.EventTime.AsSeconds(EventData.GetValue<uint64>("EndTime"));
+		uint32 ThreadId = EventData.GetValue<uint32>("ThreadId");
+		uint32 CoreNumber = EventData.GetValue<uint8>("CoreNumber");
+		ContextSwitchProvider.Add(ThreadId, Start, End, CoreNumber);
+		break;
+	}
+
+	case RouteId_StackSample:
+	{
+		double Time = Context.EventTime.AsSeconds(EventData.GetValue<uint64>("Time"));
+		uint32 ThreadId = EventData.GetValue<uint32>("ThreadId");
+		const TArrayReader<uint64>& Addresses = EventData.GetArray<uint64>("Addresses");
+		StackSampleProvider.Add(ThreadId, Time, Addresses.Num(), Addresses.GetData());
+		break;
+	}
+
+	}
+
+	return true;
+}
+
+} // namespace TraceServices
diff --git a/Engine/Source/Developer/TraceServices/Private/Analyzers/PlatformEventTraceAnalysis.h b/Engine/Source/Developer/TraceServices/Private/Analyzers/PlatformEventTraceAnalysis.h
new file mode 100644
--- /dev/null
+++ b/Engine/Source/Developer/TraceServices/Private/Analyzers/PlatformEventTraceAnalysis.h
@@ -0,0 +1,37 @@
+// Copyright Epic Games, Inc. All Rights Reserved.
+
+#pragma once
+
+#include "Trace/Analyzer.h"
+#include "ProfilingDebugging/MiscTrace.h"
+
+namespace TraceServices
+{
+
+class IAnalysisSession;
+class FContextSwitchProvider;
+class FStackSampleProvider;
+
+class FPlatformEventTraceAnalyzer
+	: public UE::Trace::IAnalyzer
+{
+public:
+	FPlatformEventTraceAnalyzer(IAnalysisSession& Session,
+								FContextSwitchProvider& ContextSwitchProvider,
+								FStackSampleProvider& StackSampleProvider);
+	virtual void OnAnalysisBegin(const FOnAnalysisContext& Context) override;
+	virtual bool OnEvent(uint16 RouteId, EStyle Style, const FOnEventContext& Context) override;
+
+private:
+	enum : uint16
+	{
+		RouteId_ContextSwitch,
+		RouteId_StackSample,
+	};
+
+	IAnalysisSession& Session;
+	FContextSwitchProvider& ContextSwitchProvider;
+	FStackSampleProvider& StackSampleProvider;
+};
+
+} // namespace TraceServices
diff --git a/Engine/Source/Developer/TraceServices/Private/Model/ContextSwitches.cpp b/Engine/Source/Developer/TraceServices/Private/Model/ContextSwitches.cpp
new file mode 100644
--- /dev/null
+++ b/Engine/Source/Developer/TraceServices/Private/Model/ContextSwitches.cpp
@@ -0,0 +1,91 @@
+// Copyright Epic Games, Inc. All Rights Reserved.
+
+#include "TraceServices/Model/ContextSwitches.h"
+#include "Model/ContextSwitchesPrivate.h"
+#include "AnalysisServicePrivate.h"
+#include "Algo/Sort.h"
+
+namespace TraceServices
+{
+
+const FName FContextSwitchProvider::ProviderName = "ContextSwitchProvider";
+
+FContextSwitchProvider::FContextSwitchProvider(IAnalysisSession& InSession)
+	: Session(InSession)
+{
+}
+
+FContextSwitchProvider::~FContextSwitchProvider()
+{
+	for (const auto& ContextSwitches : Threads)
+	{
+		delete ContextSwitches.Value;
+	}
+}
+
+int32 FContextSwitchProvider::GetCoreNumber(uint32 ThreadId, double Time) const
+{
+	Session.ReadAccessCheck();
+
+	const TPagedArray<FContextSwitch>* ContextSwitches = GetContextSwitches(ThreadId);
+	if (!ContextSwitches)
+	{
+		return -1;
+	}
+
+	uint64 PageIndex = Algo::LowerBoundBy(*ContextSwitches, Time, [](const TPagedArrayPage<FContextSwitch>& In) -> uint64
+	{
+		return GetFirstItem(In)->Start;
+	});
+
+	if (PageIndex < GetNum(*ContextSwitches))
+	{
+		const TPagedArrayPage<FContextSwitch>& Page = GetData(*ContextSwitches)[PageIndex];
+		uint64 Index = Algo::LowerBoundBy(Page, Time, [](const FContextSwitch& In) -> uint64
+		{
+			return In.Start;
+		});
+
+		if (Index < GetNum(Page))
+		{
+			const FContextSwitch& ContextSwitch = GetData(Page)[Index];
+			if (ContextSwitch.Start >= Time && ContextSwitch.End < Time)
+			{
+				return ContextSwitch.CoreNumber;
+			}
+		}
+	}
+
+	return -1;
+}
+
+const TPagedArray<FContextSwitch>* FContextSwitchProvider::GetContextSwitches(uint32 ThreadId) const
+{
+	Session.ReadAccessCheck();
+
+	auto ContextSwitches = Threads.Find(ThreadId);
+	return ContextSwitches ? *ContextSwitches : nullptr;
+}
+
+void FContextSwitchProvider::Add(uint32 ThreadId, double Start, double End, uint32 CoreNumber)
+{
+	Session.WriteAccessCheck();
+
+	TPagedArray<FContextSwitch>** ContextSwitches = Threads.Find(ThreadId);
+	if (!ContextSwitches)
+	{
+		ContextSwitches = &Threads.Add(ThreadId, new TPagedArray<FContextSwitch>(Session.GetLinearAllocator(), 4096));
+	}
+
+	FContextSwitch& ContextSwitch = (*ContextSwitches)->PushBack();
+	ContextSwitch.Start = Start;
+	ContextSwitch.End = End;
+	ContextSwitch.CoreNumber = CoreNumber;
+}
+
+const IContextSwitchProvider& ReadContextSwitchProvider(const IAnalysisSession& Session)
+{
+	return *Session.ReadProvider<IContextSwitchProvider>(FContextSwitchProvider::ProviderName);
+}
+
+} // namespace TraceServices
diff --git a/Engine/Source/Developer/TraceServices/Private/Model/ContextSwitchesPrivate.h b/Engine/Source/Developer/TraceServices/Private/Model/ContextSwitchesPrivate.h
new file mode 100644
--- /dev/null
+++ b/Engine/Source/Developer/TraceServices/Private/Model/ContextSwitchesPrivate.h
@@ -0,0 +1,30 @@
+// Copyright Epic Games, Inc. All Rights Reserved.
+
+#pragma once
+
+#include "TraceServices/Model/AnalysisSession.h"
+#include "TraceServices/Model/ContextSwitches.h"
+#include "Containers/Map.h"
+
+namespace TraceServices
+{
+
+class FContextSwitchProvider
+	: public IContextSwitchProvider
+{
+public:
+	static const FName ProviderName;
+	FContextSwitchProvider(IAnalysisSession& Session);
+	~FContextSwitchProvider();
+
+	int32 GetCoreNumber(uint32 ThreadId, double Time) const override;
+	const TPagedArray<FContextSwitch>* GetContextSwitches(uint32 ThreadId) const override;
+
+	void Add(uint32 ThreadId, double Start, double End, uint32 CoreNumber);
+
+private:
+	IAnalysisSession& Session;
+	TMap<uint32, TPagedArray<FContextSwitch>*> Threads;
+};
+
+} // namespace TraceServices
diff --git a/Engine/Source/Developer/TraceServices/Private/Model/StackSamples.cpp b/Engine/Source/Developer/TraceServices/Private/Model/StackSamples.cpp
new file mode 100644
--- /dev/null
+++ b/Engine/Source/Developer/TraceServices/Private/Model/StackSamples.cpp
@@ -0,0 +1,80 @@
+// Copyright Epic Games, Inc. All Rights Reserved.
+
+#include "TraceServices/Model/StackSamples.h"
+#include "Model/StackSamplesPrivate.h"
+#include "AnalysisServicePrivate.h"
+#include "Algo/Sort.h"
+
+namespace TraceServices
+{
+
+const FName FStackSampleProvider::ProviderName = "StackSampleProvider";
+
+FStackSampleProvider::FStackSampleProvider(IAnalysisSession& InSession)
+	: Session(InSession)
+	, AddressValues(Session.GetLinearAllocator(), 4096)
+{
+}
+
+FStackSampleProvider::~FStackSampleProvider()
+{
+	for (const auto& ThreadSamples : Threads)
+	{
+		delete ThreadSamples.Value;
+	}
+}
+
+const TPagedArray<FStackSample>* FStackSampleProvider::GetStackSamples(uint32 ThreadId) const
+{
+	Session.ReadAccessCheck();
+
+	auto StackSamples = Threads.Find(ThreadId);
+	return StackSamples ? *StackSamples : nullptr;
+}
+
+void FStackSampleProvider::Add(uint32 ThreadId, double Time, uint32 Count, const uint64* Addresses)
+{
+	Session.WriteAccessCheck();
+
+	TPagedArray<FStackSample>** StackSamples = Threads.Find(ThreadId);
+	if (!StackSamples)
+	{
+		StackSamples = &Threads.Add(ThreadId, new TPagedArray<FStackSample>(Session.GetLinearAllocator(), 4096));
+	}
+	
+	FStackSample& StackSample = (*StackSamples)->PushBack();
+	StackSample.Time = Time;
+	StackSample.Count = Count;
+
+	if (AddressValues.NumPages() == 0)
+	{
+		for (uint32 Index = 0; Index < Count; Index++)
+		{
+			AddressValues.EmplaceBack(Addresses[Index]);
+		}
+		StackSample.Addresses = &AddressValues.First();
+		return;
+	}
+
+	uint64 Available = AddressValues.GetPageSize() - AddressValues.GetLastPage()->Count;
+	if (Available < Count)
+	{
+		for (uint64 Temp = 0; Temp < Available; Temp++)
+		{
+			AddressValues.PushBack();
+		}
+	}
+
+	for (uint32 Index = 0; Index < Count; Index++)
+	{
+		AddressValues.EmplaceBack(Addresses[Index]);
+	}
+	StackSample.Addresses = &AddressValues.Last() - Count;
+}
+
+const IStackSampleProvider& ReadStackSampleProvider(const IAnalysisSession& Session)
+{
+	return *Session.ReadProvider<IStackSampleProvider>(FStackSampleProvider::ProviderName);
+}
+
+} // namespace TraceServices
diff --git a/Engine/Source/Developer/TraceServices/Private/Model/StackSamplesPrivate.h b/Engine/Source/Developer/TraceServices/Private/Model/StackSamplesPrivate.h
new file mode 100644
--- /dev/null
+++ b/Engine/Source/Developer/TraceServices/Private/Model/StackSamplesPrivate.h
@@ -0,0 +1,30 @@
+// Copyright Epic Games, Inc. All Rights Reserved.
+
+#pragma once
+
+#include "TraceServices/Model/AnalysisSession.h"
+#include "TraceServices/Model/StackSamples.h"
+#include "Containers/Map.h"
+
+namespace TraceServices
+{
+
+class FStackSampleProvider
+	: public IStackSampleProvider
+{
+public:
+	static const FName ProviderName;
+	FStackSampleProvider(IAnalysisSession& Session);
+	~FStackSampleProvider();
+
+	const TPagedArray<FStackSample>* GetStackSamples(uint32 ThreadId) const override;
+
+	void Add(uint32 ThreadId, double Time, uint32 Count, const uint64* Addresses);
+
+private:
+	IAnalysisSession& Session;
+	TMap<uint32, TPagedArray<FStackSample>*> Threads;
+	TPagedArray<uint64> AddressValues;
+};
+
+} // namespace TraceServices
diff --git a/Engine/Source/Developer/TraceServices/Private/Modules/PlatformEventsModule.cpp b/Engine/Source/Developer/TraceServices/Private/Modules/PlatformEventsModule.cpp
new file mode 100644
--- /dev/null
+++ b/Engine/Source/Developer/TraceServices/Private/Modules/PlatformEventsModule.cpp
@@ -0,0 +1,32 @@
+// Copyright Epic Games, Inc. All Rights Reserved.
+
+#include "PlatformEventsModule.h"
+#include "Analyzers/PlatformEventTraceAnalysis.h"
+#include "AnalysisServicePrivate.h"
+#include "Model/ContextSwitchesPrivate.h"
+#include "Model/StackSamplesPrivate.h"
+
+namespace TraceServices
+{
+
+static const FName PlatformEventsModuleName("TraceModule_PlatformEvents");
+
+void FPlatformEventsModule::GetModuleInfo(FModuleInfo& OutModuleInfo)
+{
+	OutModuleInfo.Name = PlatformEventsModuleName;
+	OutModuleInfo.DisplayName = TEXT("PlatformEvents");
+}
+	
+void FPlatformEventsModule::OnAnalysisBegin(IAnalysisSession& InSession)
+{
+	FAnalysisSession& Session = static_cast<FAnalysisSession&>(InSession);
+
+	FContextSwitchProvider* ContextSwitchProvider = new FContextSwitchProvider(Session);
+	FStackSampleProvider* StackSampleProvider = new FStackSampleProvider(Session);
+	
+	Session.AddProvider(FContextSwitchProvider::ProviderName, ContextSwitchProvider);
+	Session.AddProvider(FStackSampleProvider::ProviderName, StackSampleProvider);
+	Session.AddAnalyzer(new FPlatformEventTraceAnalyzer(Session, *ContextSwitchProvider, *StackSampleProvider));
+}
+
+} // namespace TraceServices
diff --git a/Engine/Source/Developer/TraceServices/Private/Modules/PlatformEventsModule.h b/Engine/Source/Developer/TraceServices/Private/Modules/PlatformEventsModule.h
new file mode 100644
--- /dev/null
+++ b/Engine/Source/Developer/TraceServices/Private/Modules/PlatformEventsModule.h
@@ -0,0 +1,29 @@
+// Copyright Epic Games, Inc. All Rights Reserved.
+
+#pragma once
+
+#include "TraceServices/ModuleService.h"
+#include "UObject/NameTypes.h"
+
+namespace TraceServices
+{
+
+class FPlatformEventsModule
+	: public IModule
+{
+public:
+	virtual void GetModuleInfo(FModuleInfo& OutModuleInfo) override;
+	virtual void OnAnalysisBegin(IAnalysisSession& Session) override;
+	virtual void GetLoggers(TArray<const TCHAR *>& OutLoggers) override
+	{
+	}
+	virtual const TCHAR* GetCommandLineArgument() override
+	{
+		return nullptr;
+	}
+	virtual void GenerateReports(const IAnalysisSession& Session, const TCHAR* CmdLine, const TCHAR* OutputDirectory) override
+	{
+	}
+};
+
+} // namespace TraceServices
diff --git a/Engine/Source/Developer/TraceServices/Private/TraceServicesModule.cpp b/Engine/Source/Developer/TraceServices/Private/TraceServicesModule.cpp
--- a/Engine/Source/Developer/TraceServices/Private/TraceServicesModule.cpp
+++ b/Engine/Source/Developer/TraceServices/Private/TraceServicesModule.cpp
@@ -13,6 +13,7 @@
 #include "Modules/NetProfilerModule.h"
 #include "Modules/MemoryModule.h"
 #include "Modules/DiagnosticsModule.h"
+#include "Modules/PlatformEventsModule.h"
 #include "Modules/TasksModule.h"
 
 class FTraceServicesModule
@@ -39,6 +40,7 @@ private:
 	TraceServices::FNetProfilerModule NetProfilerModule;
 	TraceServices::FMemoryModule MemoryModule;
 	TraceServices::FDiagnosticsModule DiagnosticsModule;
+	TraceServices::FPlatformEventsModule PlatformEventsModule;
 	TraceServices::FTasksModule TasksModule;
 };
 
@@ -83,6 +85,7 @@ void FTraceServicesModule::StartupModule()
 	IModularFeatures::Get().RegisterModularFeature(TraceServices::ModuleFeatureName, &CountersModule);
 	IModularFeatures::Get().RegisterModularFeature(TraceServices::ModuleFeatureName, &NetProfilerModule);
 	IModularFeatures::Get().RegisterModularFeature(TraceServices::ModuleFeatureName, &DiagnosticsModule);
+	IModularFeatures::Get().RegisterModularFeature(TraceServices::ModuleFeatureName, &PlatformEventsModule);
 
 #if EXPERIMENTAL_STATSTRACE_ENABLED
 	IModularFeatures::Get().RegisterModularFeature(TraceServices::ModuleFeatureName, &StatsModule);
@@ -107,6 +110,7 @@ void FTraceServicesModule::ShutdownModule()
 	IModularFeatures::Get().UnregisterModularFeature(TraceServices::ModuleFeatureName, &StatsModule);
 #endif
 
+	IModularFeatures::Get().UnregisterModularFeature(TraceServices::ModuleFeatureName, &PlatformEventsModule);
 	IModularFeatures::Get().UnregisterModularFeature(TraceServices::ModuleFeatureName, &DiagnosticsModule);
 	IModularFeatures::Get().UnregisterModularFeature(TraceServices::ModuleFeatureName, &NetProfilerModule);
 	IModularFeatures::Get().UnregisterModularFeature(TraceServices::ModuleFeatureName, &CountersModule);
diff --git a/Engine/Source/Developer/TraceServices/Public/TraceServices/Model/ContextSwitches.h b/Engine/Source/Developer/TraceServices/Public/TraceServices/Model/ContextSwitches.h
new file mode 100644
--- /dev/null
+++ b/Engine/Source/Developer/TraceServices/Public/TraceServices/Model/ContextSwitches.h
@@ -0,0 +1,29 @@
+// Copyright Epic Games, Inc. All Rights Reserved.
+
+#pragma once
+
+#include "TraceServices/Model/AnalysisSession.h"
+#include "Common/PagedArray.h"
+
+namespace TraceServices
+{
+
+struct FContextSwitch
+{
+	double Start;
+	double End;
+	uint32 CoreNumber;
+};
+
+class IContextSwitchProvider
+	: public IProvider
+{
+public:
+	virtual ~IContextSwitchProvider() = default;
+	virtual int32 GetCoreNumber(uint32 ThreadId, double Time) const = 0;
+	virtual const TPagedArray<FContextSwitch>* GetContextSwitches(uint32 ThreadId) const = 0;
+};
+
+TRACESERVICES_API const IContextSwitchProvider& ReadContextSwitchProvider(const IAnalysisSession& Session);
+
+} // namespace TraceServices
diff --git a/Engine/Source/Developer/TraceServices/Public/TraceServices/Model/StackSamples.h b/Engine/Source/Developer/TraceServices/Public/TraceServices/Model/StackSamples.h
new file mode 100644
--- /dev/null
+++ b/Engine/Source/Developer/TraceServices/Public/TraceServices/Model/StackSamples.h
@@ -0,0 +1,28 @@
+// Copyright Epic Games, Inc. All Rights Reserved.
+
+#pragma once
+
+#include "TraceServices/Model/AnalysisSession.h"
+#include "Common/PagedArray.h"
+
+namespace TraceServices
+{
+
+struct FStackSample
+{
+	double Time;
+	uint32 Count;
+	uint64* Addresses;
+};
+
+class IStackSampleProvider
+	: public IProvider
+{
+public:
+	virtual ~IStackSampleProvider() = default;
+	virtual const TPagedArray<FStackSample>* GetStackSamples(uint32 ThreadId) const = 0;
+};
+
+TRACESERVICES_API const IStackSampleProvider& ReadStackSampleProvider(const IAnalysisSession& Session);
+
+} // namespace TraceServices
diff --git a/Engine/Source/Runtime/Core/Private/ProfilingDebugging/PlatformEvents.cpp b/Engine/Source/Runtime/Core/Private/ProfilingDebugging/PlatformEvents.cpp
new file mode 100644
--- /dev/null
+++ b/Engine/Source/Runtime/Core/Private/ProfilingDebugging/PlatformEvents.cpp
@@ -0,0 +1,49 @@
+// Copyright Epic Games, Inc. All Rights Reserved.
+
+#include "ProfilingDebugging/PlatformEvents.h"
+#include "Containers/UnrealString.h"
+
+/////////////////////////////////////////////////////////////////////
+
+UE_TRACE_CHANNEL_DEFINE(ContextSwitchChannel)
+UE_TRACE_CHANNEL_DEFINE(StackSamplingChannel)
+
+/////////////////////////////////////////////////////////////////////
+
+EPlatformEvent PlatformEvents_GetEvent(const FString& Name)
+{
+	if (Name == TEXT("contextswitch"))
+	{
+		return EPlatformEvent::ContextSwitch;
+	}
+	else if (Name == TEXT("stacksampling"))
+	{
+		return EPlatformEvent::StackSampling;
+	}
+	else
+	{
+		return EPlatformEvent::None;
+	}
+}
+
+/////////////////////////////////////////////////////////////////////
+
+#if !UE_PLATFORM_EVENTS_AVAILABLE
+
+void PlatformEvents_Init(uint32 SamplingIntervalUsec)
+{
+}
+
+void PlatformEvents_Enable(EPlatformEvent Event)
+{
+}
+
+void PlatformEvents_Disable(EPlatformEvent Event)
+{
+}
+
+void PlatformEvents_Stop()
+{
+}
+
+#endif
diff --git a/Engine/Source/Runtime/Core/Private/ProfilingDebugging/PlatformEvents.h b/Engine/Source/Runtime/Core/Private/ProfilingDebugging/PlatformEvents.h
new file mode 100644
--- /dev/null
+++ b/Engine/Source/Runtime/Core/Private/ProfilingDebugging/PlatformEvents.h
@@ -0,0 +1,55 @@
+// Copyright Epic Games, Inc. All Rights Reserved.
+#pragma once
+
+#include "CoreTypes.h"
+#include "Trace/Trace.inl"
+
+#if PLATFORM_WINDOWS
+#	define UE_PLATFORM_EVENTS_AVAILABLE 1
+#else
+#	define UE_PLATFORM_EVENTS_AVAILABLE 0
+#endif
+
+/////////////////////////////////////////////////////////////////////
+
+UE_TRACE_CHANNEL_EXTERN(ContextSwitchChannel)
+UE_TRACE_CHANNEL_EXTERN(StackSamplingChannel)
+
+/////////////////////////////////////////////////////////////////////
+
+// represents time interval when thread was running on specific core
+UE_TRACE_EVENT_BEGIN(PlatformEvent, ContextSwitch, NoSync)
+	UE_TRACE_EVENT_FIELD(uint64, StartTime)
+	UE_TRACE_EVENT_FIELD(uint64, EndTime)
+	UE_TRACE_EVENT_FIELD(uint32, ThreadId)
+	UE_TRACE_EVENT_FIELD(uint8, CoreNumber)
+UE_TRACE_EVENT_END()
+
+// represents call stack addresses in stack sampling
+UE_TRACE_EVENT_BEGIN(PlatformEvent, StackSample, NoSync)
+	UE_TRACE_EVENT_FIELD(uint64, Time)
+	UE_TRACE_EVENT_FIELD(uint32, ThreadId)
+	UE_TRACE_EVENT_FIELD(uint64[], Addresses)
+UE_TRACE_EVENT_END()
+
+/////////////////////////////////////////////////////////////////////
+
+enum class EPlatformEvent
+{
+	None = 0x00,
+	ContextSwitch = 0x01,
+	StackSampling = 0x02,
+};
+
+ENUM_CLASS_FLAGS(EPlatformEvent)
+
+/////////////////////////////////////////////////////////////////////
+
+EPlatformEvent PlatformEvents_GetEvent(const FString& Name);
+
+void PlatformEvents_Init(uint32 SamplingIntervalUsec);
+void PlatformEvents_Enable(EPlatformEvent Event);
+void PlatformEvents_Disable(EPlatformEvent Event);
+void PlatformEvents_Stop();
+
+/////////////////////////////////////////////////////////////////////
diff --git a/Engine/Source/Runtime/Core/Private/ProfilingDebugging/TraceAuxiliary.cpp b/Engine/Source/Runtime/Core/Private/ProfilingDebugging/TraceAuxiliary.cpp
--- a/Engine/Source/Runtime/Core/Private/ProfilingDebugging/TraceAuxiliary.cpp
+++ b/Engine/Source/Runtime/Core/Private/ProfilingDebugging/TraceAuxiliary.cpp
@@ -24,6 +24,7 @@
 #include "ProfilingDebugging/CountersTrace.h"
 #include "ProfilingDebugging/MiscTrace.h"
 #include "ProfilingDebugging/PlatformFileTrace.h"
+#include "ProfilingDebugging/PlatformEvents.h"
 #include "String/ParseTokens.h"
 #include "Templates/UnrealTemplate.h"
 #include "Trace/Trace.inl"
@@ -214,6 +215,12 @@ void FTraceAuxiliaryImpl::EnableChannel(FChannel& Channel)
 		return;
 	}
 
+	EPlatformEvent Event = PlatformEvents_GetEvent(Channel.Name);
+	if (Event != EPlatformEvent::None)
+	{
+		PlatformEvents_Enable(Event);
+	}
+
 	UE::Trace::ToggleChannel(*Channel.Name, true);
 	Channel.bActive = true;
 }
@@ -237,6 +244,12 @@ void FTraceAuxiliaryImpl::DisableChannels()
 		{
 			UE::Trace::ToggleChannel(*Channel.Name, false);
 			Channel.bActive = false;
+
+			EPlatformEvent Event = PlatformEvents_GetEvent(Channel.Name);
+			if (Event != EPlatformEvent::None)
+			{
+				PlatformEvents_Disable(Event);
+			}
 		}
 	}
 }
@@ -453,6 +466,11 @@ void FTraceAuxiliary::Initialize(const TCHAR* CommandLine)
 	}
 	UE::Trace::Initialize(Desc);
 
+	// By default use 1msec for stack sampling interval
+	uint32 Microseconds = 1000;
+	FParse::Value(CommandLine, TEXT("-samplinginterval="), Microseconds);
+	PlatformEvents_Init(Microseconds);
+
 	FCoreDelegates::OnEndFrame.AddStatic(UE::Trace::Update);
 	FModuleManager::Get().OnModulesChanged().AddLambda([](FName Name, EModuleChangeReason Reason)
 	{
@@ -504,6 +522,16 @@ void FTraceAuxiliary::InitializePresets(const TCHAR* CommandLine)
 #endif 
 }
 
+////////////////////////////////////////////////////////////////////////////////
+void FTraceAuxiliary::Shutdown()
+{
+#if UE_TRACE_ENABLED
+	// make sure all platform event functionality has shut down as on some
+	// platforms it impacts whole system, even if application has terminated
+	PlatformEvents_Stop();
+#endif
+}
+
 ////////////////////////////////////////////////////////////////////////////////
 void FTraceAuxiliary::EnableChannels()
 {
diff --git a/Engine/Source/Runtime/Core/Private/ProfilingDebugging/Windows/EventTracingForWindows.cpp b/Engine/Source/Runtime/Core/Private/ProfilingDebugging/Windows/EventTracingForWindows.cpp
new file mode 100644
--- /dev/null
+++ b/Engine/Source/Runtime/Core/Private/ProfilingDebugging/Windows/EventTracingForWindows.cpp
@@ -0,0 +1,525 @@
+// Copyright Epic Games, Inc. All Rights Reserved.
+
+#include "ProfilingDebugging/PlatformEvents.h"
+#include "HAL/Runnable.h"
+#include "HAL/Event.h"
+#include "HAL/RunnableThread.h"
+#include "Logging/LogMacros.h"
+#include "Trace/Trace.h"
+
+#include "Windows/AllowWindowsPlatformTypes.h"
+#include <evntrace.h>
+#include <evntcons.h>
+#include "Windows/HideWindowsPlatformTypes.h"
+
+/////////////////////////////////////////////////////////////////////
+
+DEFINE_LOG_CATEGORY_STATIC(LogPlatformEvents, Log, All);
+
+/////////////////////////////////////////////////////////////////////
+
+namespace {
+
+class FPlatformEvents : public FRunnable
+{
+public:
+	FPlatformEvents(uint32 SamplingIntervalUsec);
+	~FPlatformEvents();
+
+	void Enable(EPlatformEvent Event);
+	void Disable(EPlatformEvent Event);
+
+private:
+	virtual uint32 Run() override;
+	virtual void Stop() override;
+
+	bool StartETW();
+	void StopETW();
+
+	static void WINAPI TraceEventCallback(PEVENT_RECORD Event);
+	static uint32 CurrentProcessId;
+
+	bool bStarting = false;
+	bool bError = false;
+	FEventRef ReadyEvent;
+	FRunnableThread* Thread;
+
+	EPlatformEvent EnabledEvents = EPlatformEvent::None;
+
+	struct FTrace
+	{
+		EVENT_TRACE_PROPERTIES Properties;
+		WCHAR SessionName[1024];
+	};
+
+	TRACEHANDLE TraceHandle = INVALID_PROCESSTRACE_HANDLE;
+	FTrace Trace;
+
+	// timestamp for each core when thread start executing on it
+	uint64 ContexSwitchStart[256];
+};
+
+// https://docs.microsoft.com/en-us/windows/win32/etw/cswitch
+struct FContextSwitchEvent
+{
+	enum { Type = 36 };
+
+	uint32 NewThreadId;					// 0
+	uint32 OldThreadId;					// 4
+	uint8  NewThreadPriority;			// 8
+	uint8  OldThreadPriority;			// 9
+	uint8  PreviousCState;				// 10
+	uint8  SpareByte;					// 11
+	uint8  OldThreadWaitReason;			// 12
+	uint8  OldThreadWaitMode;			// 13
+	uint8  OldThreadState;				// 14
+	uint8  OldThreadWaitIdealProcessor;	// 15
+	uint32 NewThreadWaitTime;			// 16
+	uint32 Reserved;					// 20
+};
+
+// https://docs.microsoft.com/en-us/windows/win32/etw/stackwalk-event
+struct FStackWalkEvent
+{
+	enum { Type = 32 };
+
+	uint64 TimeStamp;
+	uint32 ProcessId;
+	uint32 ThreadId;
+
+	// this array can be max 192 elements
+	// actual count is determined by total byte size of event
+	uint64 Stack[];
+};
+
+} // namespace
+
+/////////////////////////////////////////////////////////////////////
+
+static FPlatformEvents* GPlatformEvents;
+
+/////////////////////////////////////////////////////////////////////
+
+uint32 FPlatformEvents::CurrentProcessId;
+
+void FPlatformEvents::TraceEventCallback(PEVENT_RECORD Event)
+{
+	uint32 ProcessId = Event->EventHeader.ProcessId;
+	int Opcode = Event->EventHeader.EventDescriptor.Opcode;
+	if (Opcode == FContextSwitchEvent::Type)
+	{
+		const FContextSwitchEvent* Payload = (const FContextSwitchEvent*)Event->UserData;
+		if (Payload == nullptr || Event->UserDataLength != sizeof(FContextSwitchEvent))
+		{
+			return;
+		}
+
+		uint8 CoreNumber = Event->BufferContext.ProcessorNumber;
+
+		uint64 StartTime = GPlatformEvents->ContexSwitchStart[CoreNumber];
+		uint64 EndTime = Event->EventHeader.TimeStamp.QuadPart;
+
+		if (StartTime != 0)
+		{
+			uint32 ThreadId = Payload->OldThreadId;
+
+			if (ThreadId != 0)
+			{
+				UE_TRACE_LOG(PlatformEvent, ContextSwitch, ContextSwitchChannel)
+					<< ContextSwitch.StartTime(StartTime)
+					<< ContextSwitch.EndTime(EndTime)
+					<< ContextSwitch.ThreadId(ThreadId)
+					<< ContextSwitch.CoreNumber(CoreNumber);
+			}
+		}
+
+		GPlatformEvents->ContexSwitchStart[CoreNumber] = EndTime;
+	}
+	else if (Opcode == FStackWalkEvent::Type)
+	{
+		const FStackWalkEvent* Payload = (const FStackWalkEvent*)Event->UserData;
+		if (Payload == nullptr || Payload->ProcessId != CurrentProcessId || Event->UserDataLength < sizeof(FStackWalkEvent))
+		{
+			return;
+		}
+
+		uint32 Count = (Event->UserDataLength - sizeof(FStackWalkEvent)) / sizeof(uint64);
+
+		// sometimes stack walk from ETW gives a call-stack that is fully in kernel address space
+		// so if all addresses have a high bit set, then ignore this call stack
+		bool bValid = false;
+		for (uint32 Index = 0; Index < Count; Index++)
+		{
+			uint64 Address = Payload->Stack[Index];
+			if ((Address >> 63) == 0)
+			{
+				bValid = true;
+				break;
+			}
+		}
+
+		// sometimes we get call stack with only one or two entries, which obviously is not very useful
+		if (bValid && Count > 2)
+		{
+			UE_TRACE_LOG(PlatformEvent, StackSample, StackSamplingChannel)
+				<< StackSample.Time(Payload->TimeStamp)
+				<< StackSample.ThreadId(Payload->ThreadId)
+				<< StackSample.Addresses(Payload->Stack, Count);
+		}
+	}
+}
+
+/////////////////////////////////////////////////////////////////////
+
+FPlatformEvents::FPlatformEvents(uint32 SamplingIntervalUsec)
+	: ReadyEvent(EEventMode::ManualReset)
+{
+	// set sampling interval, must be done before StartTrace
+	{
+		typedef struct
+		{
+			ULONG Source;
+			ULONG Interval;
+		}
+		TRACE_PROFILE_INTERVAL;
+
+		// interval is specified in 100 nanosecond units, and allowed range is limited
+		// see "-setprofint" argument in https://docs.microsoft.com/en-us/windows-hardware/test/wpt/start
+		TRACE_PROFILE_INTERVAL Interval;
+		Interval.Source = 0;
+		Interval.Interval = FMath::Clamp<ULONG>(10 * SamplingIntervalUsec, 1221, 10000000);
+
+		ULONG Status = ::TraceSetInformation(NULL, TraceSampledProfileIntervalInfo, &Interval, sizeof(Interval));
+		if (Status != ERROR_SUCCESS)
+		{
+			UE_LOG(LogPlatformEvents, Warning, TEXT("Unable to set ETW sampling interval: 0x%08x"), Status);
+		}
+	}
+
+	ZeroMemory(ContexSwitchStart, sizeof(ContexSwitchStart));
+	CurrentProcessId = ::GetCurrentProcessId();
+
+	ReadyEvent->Trigger();
+}
+
+/////////////////////////////////////////////////////////////////////
+
+FPlatformEvents::~FPlatformEvents()
+{
+	Stop();
+}
+
+/////////////////////////////////////////////////////////////////////
+
+void FPlatformEvents::Enable(EPlatformEvent Event)
+{
+	if (bStarting)
+	{
+		// in case ETW is already starting, wait for startup code to complete
+		ReadyEvent->Wait(MAX_uint32);
+	}
+
+	if (bError)
+	{
+		// in case ETW failed to start up earlier, do nothing
+		return;
+	}
+
+	if ((EnabledEvents & Event) != EPlatformEvent::None)
+	{
+		// if event is already enabled, do nothing
+		return;
+	}
+
+	if (TraceHandle == INVALID_PROCESSTRACE_HANDLE)
+	{
+		// in case ETW is not started, start it now
+		ReadyEvent->Reset();
+		bError = false;
+		bStarting = true;
+		Thread = FRunnableThread::Create(this, TEXT("PlatformEvents"), 0, TPri_Normal);
+	}
+	else
+	{
+		// in case ETW was already running, enable the necessary event
+
+		EVENT_TRACE_PROPERTIES* Properties = &Trace.Properties;
+		if (Event == EPlatformEvent::ContextSwitch)
+		{
+			Properties->EnableFlags |= EVENT_TRACE_FLAG_CSWITCH;
+		}
+		else if (Event == EPlatformEvent::StackSampling)
+		{
+			Properties->EnableFlags |= EVENT_TRACE_FLAG_PROFILE;
+		}
+		
+		ULONG Status = ::ControlTrace(NULL, KERNEL_LOGGER_NAME, Properties, EVENT_TRACE_CONTROL_UPDATE);
+		if (Status != ERROR_SUCCESS)
+		{
+			UE_LOG(LogPlatformEvents, Error, TEXT("Unable to update ETW trace: 0x%08x"), Status);
+		}
+	}
+
+	EnumAddFlags(EnabledEvents, Event);
+}
+
+/////////////////////////////////////////////////////////////////////
+
+void FPlatformEvents::Disable(EPlatformEvent Event)
+{
+	if (bStarting)
+	{
+		// in case ETW is already starting, wait for startup code to complete
+		ReadyEvent->Wait(MAX_uint32);
+	}
+
+	if (TraceHandle == INVALID_PROCESSTRACE_HANDLE)
+	{
+		// in case ETW is not running, do nothing
+		return;
+	}
+
+	if ((EnabledEvents & Event) == EPlatformEvent::None)
+	{
+		// if event is already disabled, do nothing
+		return;
+	}
+	
+	EnumRemoveFlags(EnabledEvents, Event);
+
+	if (EnabledEvents == EPlatformEvent::None)
+	{
+		// in case all events are disabled, stop ETW
+		Stop();
+	}
+	else
+	{
+		// in case ETW still needs to run, disable only specific event
+
+		EVENT_TRACE_PROPERTIES* Properties = &Trace.Properties;
+		if (Event == EPlatformEvent::ContextSwitch)
+		{
+			Properties->EnableFlags &= ~EVENT_TRACE_FLAG_CSWITCH;
+		}
+		else if (Event == EPlatformEvent::StackSampling)
+		{
+			Properties->EnableFlags &= ~EVENT_TRACE_FLAG_PROFILE;
+		}
+
+		ReadyEvent->Wait(MAX_uint32);
+		ULONG Status = ::ControlTrace(NULL, KERNEL_LOGGER_NAME, Properties, EVENT_TRACE_CONTROL_UPDATE);
+		if (Status != ERROR_SUCCESS)
+		{
+			UE_LOG(LogPlatformEvents, Error, TEXT("Unable to update ETW trace: 0x%08x"), Status);
+		}
+	}
+}
+
+/////////////////////////////////////////////////////////////////////
+
+uint32 FPlatformEvents::Run()
+{
+	if (StartETW())
+	{
+		ReadyEvent->Trigger();
+		bStarting = false;
+
+		// ProcessTrace is blocking function that enters processing loop
+		// and calls TraceEventCallback callback for each event
+		// it exits processing loop when trace is closed with ::CloseTrace 
+		::ProcessTrace(&TraceHandle, 1, NULL, NULL);
+	}
+	else
+	{
+		// in case of error in StartETW we want to be sure ETW is stopped
+		StopETW();
+		bError = true;
+		bStarting = false;
+		ReadyEvent->Trigger();
+	}
+	return 0;
+}
+
+/////////////////////////////////////////////////////////////////////
+
+void FPlatformEvents::Stop()
+{
+	// because ETW is started in background thread, then first make sure
+	// ETW startup code has completed (sucessfully or not) before actually
+	// stopping it
+	ReadyEvent->Wait(MAX_uint32);
+
+	StopETW();
+	Thread->WaitForCompletion();
+}
+
+/////////////////////////////////////////////////////////////////////
+
+bool FPlatformEvents::StartETW()
+{
+	// explicitly define GUIDs to simplify dependency on special defines in windows headers
+	// https://docs.microsoft.com/en-us/windows/win32/etw/nt-kernel-logger-constants
+	static const GUID ETW_PerfInfoGuid = { 0xce1dbfb4, 0x137e, 0x4da6, { 0x87, 0xb0, 0x3f, 0x59, 0xaa, 0x10, 0x2c, 0xbc } };
+	static const GUID ETW_SystemTraceControlGuid = { 0x9e814aad, 0x3204, 0x11d2, { 0x9a, 0x82, 0x00, 0x60, 0x08, 0xa8, 0x69, 0x39 } };
+
+	TRACEHANDLE Session;
+
+	SYSTEM_INFO SysInfo;
+	GetSystemInfo(&SysInfo);
+
+	EVENT_TRACE_PROPERTIES* Properties = &Trace.Properties;
+
+	// stop existing trace in case it is already running (from other tools)
+	{
+		ZeroMemory(Properties, sizeof(*Properties));
+		Properties->Wnode.BufferSize = sizeof(Trace);
+		Properties->Wnode.Guid = ETW_SystemTraceControlGuid;
+		Properties->LogFileMode = EVENT_TRACE_REAL_TIME_MODE;
+		Properties->LoggerNameOffset = sizeof(Trace.Properties);
+		::ControlTraceW(NULL, KERNEL_LOGGER_NAMEW, Properties, EVENT_TRACE_CONTROL_STOP);
+	}
+
+	// enable profiling privilege for process
+	{
+		HANDLE Token;
+		if (::OpenProcessToken(::GetCurrentProcess(), TOKEN_ADJUST_PRIVILEGES, &Token))
+		{
+			LUID Luid;
+			if (::LookupPrivilegeValue(NULL, SE_SYSTEM_PROFILE_NAME, &Luid))
+			{
+				TOKEN_PRIVILEGES TokenPrivileges;
+				TokenPrivileges.PrivilegeCount = 1;
+				TokenPrivileges.Privileges[0].Luid = Luid;
+				TokenPrivileges.Privileges[0].Attributes = SE_PRIVILEGE_ENABLED;
+				if (!::AdjustTokenPrivileges(Token, 0, &TokenPrivileges, sizeof(TokenPrivileges), NULL, NULL))
+				{
+					uint32 Error = ::GetLastError();
+					UE_LOG(LogPlatformEvents, Warning, TEXT("Cannot enable profile privlege for process: 0x%08x"), Error);
+				}
+			}
+			::CloseHandle(Token);
+		}
+	}
+
+	// setup trace properties
+	ZeroMemory(Properties, sizeof(*Properties));
+	Properties->Wnode.BufferSize = sizeof(Trace);
+	Properties->Wnode.Guid = ETW_SystemTraceControlGuid;
+
+	// NOTE: change this depending on what timestamps you want to collect
+	// 1 - QueryPerformanceCounter (slower)
+	// 3 - RDTSC (faster)
+	Properties->Wnode.ClientContext = 1;
+
+	Properties->Wnode.Flags = WNODE_FLAG_TRACED_GUID;
+	Properties->BufferSize = 256; // 256KB
+	Properties->MinimumBuffers = 4 * SysInfo.dwNumberOfProcessors;
+	Properties->MaximumBuffers = Properties->MinimumBuffers + 20;
+	Properties->LogFileMode = EVENT_TRACE_REAL_TIME_MODE | EVENT_TRACE_SYSTEM_LOGGER_MODE;
+	Properties->LoggerNameOffset = offsetof(FTrace, SessionName);
+	Properties->FlushTimer = 1;
+	Properties->EnableFlags = 0;
+
+	if ((EnabledEvents & EPlatformEvent::ContextSwitch) != EPlatformEvent::None)
+	{
+		Properties->EnableFlags |= EVENT_TRACE_FLAG_CSWITCH;
+	}
+	if ((EnabledEvents & EPlatformEvent::StackSampling) != EPlatformEvent::None)
+	{
+		Properties->EnableFlags |= EVENT_TRACE_FLAG_PROFILE;
+	}
+
+	// start the trace
+	{
+		ULONG Status = ::StartTraceW(&Session, KERNEL_LOGGER_NAMEW, Properties);
+		if (Status == ERROR_ACCESS_DENIED)
+		{
+			UE_LOG(LogPlatformEvents, Error, TEXT("Administrator rights required for ETW"));
+			return false;
+		}
+		else if (Status != ERROR_SUCCESS)
+		{
+			UE_LOG(LogPlatformEvents, Warning, TEXT("Unable to start ETW trace: 0x%08x"), Status);
+			return false;
+		}
+	}
+
+	// select for which events ETW should collect stack (StackWalkEvent), this is harmless if stack sampling is not enabled
+	{
+		// 0x2e -> PERFINFO_LOG_TYPE_SAMPLED_PROFILE
+		CLASSIC_EVENT_ID Events[] = { { ETW_PerfInfoGuid, 0x2e } };
+
+		ULONG Status = ::TraceSetInformation(Session, TraceStackTracingInfo, Events, sizeof(Events));
+		if (Status != ERROR_SUCCESS)
+		{
+			UE_LOG(LogPlatformEvents, Error, TEXT("Cannot enable stack tracing for ETW: 0x%08x"), Status);
+			return false;
+		}
+	}
+
+	// open trace for receving callbacks
+	{
+		EVENT_TRACE_LOGFILEW LogFile;
+		ZeroMemory(&LogFile, sizeof(LogFile));
+		LogFile.LoggerName = (LPWSTR)KERNEL_LOGGER_NAMEW;
+		LogFile.ProcessTraceMode = PROCESS_TRACE_MODE_EVENT_RECORD | PROCESS_TRACE_MODE_RAW_TIMESTAMP | PROCESS_TRACE_MODE_REAL_TIME;
+		LogFile.EventRecordCallback = &TraceEventCallback;
+
+		TraceHandle = ::OpenTraceW(&LogFile);
+		if (TraceHandle == INVALID_PROCESSTRACE_HANDLE)
+		{
+			UE_LOG(LogPlatformEvents, Error, TEXT("Cannot open ETW trace"));
+			return false;
+		}
+	}
+
+	return true;
+}
+
+/////////////////////////////////////////////////////////////////////
+
+void FPlatformEvents::StopETW()
+{
+	// stop the trace
+	::ControlTraceW(0, KERNEL_LOGGER_NAMEW, &Trace.Properties, EVENT_TRACE_CONTROL_STOP);
+
+	if (TraceHandle != INVALID_PROCESSTRACE_HANDLE)
+	{
+		// close processing loop, this will wait until all pending buffers are flushed
+		// and TracingCallback called on all pending events in buffers
+		::CloseTrace(TraceHandle);
+		TraceHandle = INVALID_PROCESSTRACE_HANDLE;
+	}
+}
+
+/////////////////////////////////////////////////////////////////////
+
+void PlatformEvents_Init(uint32 SamplingIntervalUsec)
+{
+	GPlatformEvents = new FPlatformEvents(SamplingIntervalUsec);
+}
+
+void PlatformEvents_Enable(EPlatformEvent Event)
+{
+	if (GPlatformEvents)
+	{
+		GPlatformEvents->Enable(Event);
+	}
+}
+
+void PlatformEvents_Disable(EPlatformEvent Event)
+{
+	if (GPlatformEvents)
+	{
+		GPlatformEvents->Disable(Event);
+	}
+}
+
+void PlatformEvents_Stop()
+{
+	if (GPlatformEvents)
+	{
+		delete GPlatformEvents;
+		GPlatformEvents = nullptr;
+	}
+}
diff --git a/Engine/Source/Runtime/Core/Public/ProfilingDebugging/TraceAuxiliary.h b/Engine/Source/Runtime/Core/Public/ProfilingDebugging/TraceAuxiliary.h
--- a/Engine/Source/Runtime/Core/Public/ProfilingDebugging/TraceAuxiliary.h
+++ b/Engine/Source/Runtime/Core/Public/ProfilingDebugging/TraceAuxiliary.h
@@ -10,6 +10,7 @@ class CORE_API FTraceAuxiliary
 public:
 	static void Initialize(const TCHAR* CommandLine);
 	static void InitializePresets(const TCHAR* CommandLine);
+	static void Shutdown();
 	static void TryAutoConnect();
 	static void EnableChannels();
 };
diff --git a/Engine/Source/Runtime/Launch/Private/LaunchEngineLoop.cpp b/Engine/Source/Runtime/Launch/Private/LaunchEngineLoop.cpp
--- a/Engine/Source/Runtime/Launch/Private/LaunchEngineLoop.cpp
+++ b/Engine/Source/Runtime/Launch/Private/LaunchEngineLoop.cpp
@@ -6077,6 +6077,8 @@ void FEngineLoop::AppExit()
 
 	FTextLocalizationManager::TearDown();
 	FInternationalization::TearDown();
+
+	FTraceAuxiliary::Shutdown();
 }
 
 void FEngineLoop::PostInitRHI()
