#ue5 - LWC: UnrealMathSSE AVX implementation first pass. Mostly implements improved functions for doubles, but also made improvements to some float functions when AVX or SSE 4.1 is available.

There is some remaining work to add more specializations to VectorSwizzle for doubles. Also need to add platform defines for AVX2.

diff --git a/Engine/Source/Runtime/Core/Private/Tests/Math/UnrealMathTest.cpp b/Engine/Source/Runtime/Core/Private/Tests/Math/UnrealMathTest.cpp
--- a/Engine/Source/Runtime/Core/Private/Tests/Math/UnrealMathTest.cpp
+++ b/Engine/Source/Runtime/Core/Private/Tests/Math/UnrealMathTest.cpp
@@ -33,17 +33,18 @@ static bool GPassing;
 
 #define MATHTEST_INLINE FORCENOINLINE // if you want to do performance testing change to FORCEINLINE or FORCEINLINE_DEBUGGABLE
 
-FORCENOINLINE void SetPassing(bool Passing)
+FORCENOINLINE void CheckPassing(bool Passing)
 {
 	// Convenient for setting breakpoints to catch any failures.
 	if (!Passing)
 	{
 		GPassing = false;
 	}
-	else
-	{
-		GPassing = true;
-	}
+}
+
+FORCENOINLINE void ResetPassing()
+{
+	GPassing = true;
 }
 
 /**
@@ -61,7 +62,7 @@ bool TestVectorsEqualBitwise( VectorRegister4Float Vec0, VectorRegister4Float Ve
 
 	const bool Passed = (memcmp(GScratch + 0, GScratch + 4, sizeof(float) * 4) == 0);
 
-	SetPassing(GPassing && Passed);
+	CheckPassing(Passed);
 	return Passed;
 }
 
@@ -72,7 +73,7 @@ bool TestVectorsEqualBitwise(VectorRegister4Double Vec0, VectorRegister4Double V
 
 	const bool Passed = (memcmp(GScratchDouble + 0, GScratchDouble + 4, sizeof(double) * 4) == 0);
 
-	SetPassing(GPassing && Passed);
+	CheckPassing(Passed);
 	return Passed;
 }
 
@@ -96,7 +97,7 @@ bool TestVectorsEqual( VectorRegister4Float Vec0, VectorRegister4Float Vec1, flo
 		float Diff = GScratch[ Component + 0 ] - GScratch[ Component + 4 ];
 		GSum += ( Diff >= 0.0f ) ? Diff : -Diff;
 	}
-	SetPassing(GPassing && GSum <= Tolerance);
+	CheckPassing(GSum <= Tolerance);
 	return GSum <= Tolerance;
 }
 
@@ -110,7 +111,7 @@ bool TestVectorsEqual(VectorRegister4Double Vec0, VectorRegister4Double Vec1, do
 		double Diff = GScratchDouble[Component + 0] - GScratchDouble[Component + 4];
 		GSumDouble += (Diff >= 0.0f) ? Diff : -Diff;
 	}
-	SetPassing(GPassing && GSumDouble <= Tolerance);
+	CheckPassing(GSumDouble <= Tolerance);
 	return GSumDouble <= Tolerance;
 }
 
@@ -127,7 +128,7 @@ bool TestVectorsEqual_ComponentWiseError(VectorRegister4Float Vec0, VectorRegist
 		float Diff = GScratch[Component + 0] - GScratch[Component + 4];
 		bPassing &= FMath::IsNearlyZero(Diff, Tolerance);
 	}
-	SetPassing(GPassing && bPassing);
+	CheckPassing(bPassing);
 	return bPassing;
 }
 
@@ -150,7 +151,7 @@ bool TestVectorsEqual3( VectorRegister4Float Vec0, VectorRegister4Float Vec1, fl
 	{
 		GSum += FMath::Abs<float>( GScratch[ Component + 0 ] - GScratch[ Component + 4 ] );
 	}
-	SetPassing(GPassing && GSum <= Tolerance);
+	CheckPassing(GSum <= Tolerance);
 	return GSum <= Tolerance;
 }
 
@@ -163,7 +164,7 @@ bool TestVectorsEqual3(VectorRegister4Double Vec0, VectorRegister4Double Vec1, d
 	{
 		GSumDouble += FMath::Abs<double>(GScratchDouble[Component + 0] - GScratchDouble[Component + 4]);
 	}
-	SetPassing(GPassing && GSumDouble <= Tolerance);
+	CheckPassing(GSumDouble <= Tolerance);
 	return GSumDouble <= Tolerance;
 }
 
@@ -192,7 +193,7 @@ bool TestFVector3Equal( const FVector3f& Vec0, const FVector3f& Vec1, float Tole
 	{
 		GSum += FMath::Abs<float>( GScratch[ Component + 0 ] - GScratch[ Component + 4 ] );
 	}
-	SetPassing(GPassing && GSum <= Tolerance);
+	CheckPassing(GSum <= Tolerance);
 	return GSum <= Tolerance;
 }
 
@@ -209,7 +210,7 @@ bool TestQuatsEqual(const FQuat& Q0, const FQuat& Q1, float Tolerance)
 	GSum = 0.f;
 
 	const bool bEqual = Q0.Equals(Q1, Tolerance);
-	SetPassing(GPassing && bEqual);
+	CheckPassing(bEqual);
 	return bEqual;
 }
 
@@ -234,7 +235,7 @@ bool TestFVector3Normalized(const FVector3f& Vec0, float Tolerance)
 	GSum = FMath::Sqrt(Vec0.X * Vec0.X + Vec0.Y * Vec0.Y + Vec0.Z * Vec0.Z);
 
 	const bool bNormalized = FMath::IsNearlyEqual(GSum, 1.0f, Tolerance);
-	SetPassing(GPassing && bNormalized);
+	CheckPassing(bNormalized);
 	return bNormalized;
 }
 
@@ -259,7 +260,7 @@ bool TestQuatNormalized(const FQuat& Q0, float Tolerance)
 	GSum = FMath::Sqrt(Q0.X*Q0.X + Q0.Y*Q0.Y + Q0.Z*Q0.Z + Q0.W*Q0.W);
 
 	const bool bNormalized = FMath::IsNearlyEqual(GSum, 1.0f, Tolerance);
-	SetPassing(GPassing && bNormalized);
+	CheckPassing(bNormalized);
 	return bNormalized;
 }
 
@@ -284,7 +285,7 @@ bool TestMatricesEqual( FMatrix44f &Mat0, FMatrix44f &Mat1, float Tolerance = 0.
 		}
 		if (GSum > Tolerance)
 		{
-			SetPassing(false);
+			CheckPassing(false);
 			return false;
 		}
 	}
@@ -670,7 +671,7 @@ void LogRotatorTest(bool bExpected, const TCHAR* TestName, const FRotator& A, co
 	{
 		UE_LOG(LogUnrealMathTest, Log, TEXT("%s: %s"), bHasPassed ? TEXT("PASSED") : TEXT("FAILED"), TestName);
 		UE_LOG(LogUnrealMathTest, Log, TEXT("(%s).Equals(%s) = %d"), *A.ToString(), *B.ToString(), bComparison);
-		SetPassing(false);
+		CheckPassing(false);
 	}
 }
 
@@ -681,7 +682,7 @@ void LogRotatorTest(const TCHAR* TestName, const FRotator& A, const FRotator& B,
 	{
 		UE_LOG(LogUnrealMathTest, Log, TEXT("%s: %s"), bComparison ? TEXT("PASSED") : TEXT("FAILED"), TestName);
 		UE_LOG(LogUnrealMathTest, Log, TEXT("(%s).Equals(%s) = %d"), *A.ToString(), *B.ToString(), bComparison);
-		SetPassing(false);
+		CheckPassing(false);
 	}
 }
 
@@ -691,7 +692,7 @@ void LogQuaternionTest(const TCHAR* TestName, const FQuat& A, const FQuat& B, bo
 	{
 		UE_LOG(LogUnrealMathTest, Log, TEXT("%s: %s"), bComparison ? TEXT("PASSED") : TEXT("FAILED"), TestName);
 		UE_LOG(LogUnrealMathTest, Log, TEXT("(%s).Equals(%s) = %d"), *A.ToString(), *B.ToString(), bComparison);
-		SetPassing(false);
+		CheckPassing(false);
 	}
 }
 
@@ -801,7 +802,7 @@ void LogTest(const TCHAR* TestName, bool bHasPassed)
 	{
 		UE_LOG(LogUnrealMathTest, Log, TEXT("%s: %s"), bHasPassed ? TEXT("PASSED") : TEXT("FAILED"), TestName);
 		UE_LOG(LogUnrealMathTest, Log, TEXT("Bad(%.8f): (%.8f %.8f %.8f %.8f) (%.8f %.8f %.8f %.8f)"), GSum, GScratch[0], GScratch[1], GScratch[2], GScratch[3], GScratch[4], GScratch[5], GScratch[6], GScratch[7]);
-		SetPassing(false);
+		CheckPassing(false);
 	}
 }
 
@@ -839,7 +840,7 @@ IMPLEMENT_SIMPLE_AUTOMATION_TEST(FVectorRegisterAbstractionTest, "System.Core.Ma
 
 bool RunDoubleVectorTest()
 {
-	SetPassing(true);
+	ResetPassing();
 
 	double F1 = 1.0;
 	uint64 U1 = *(uint64*)&F1;
@@ -945,13 +946,33 @@ bool RunDoubleVectorTest()
 
 	// Replicate
 
+	V0 = MakeVectorRegisterDouble(1.0f, 2.0f, 3.0f, 4.0f);
+	V1 = VectorReplicate(V0, 0);
+	V0 = MakeVectorRegisterDouble(1.0f, 1.0f, 1.0f, 1.0f);
+	LogTest(TEXT("VectorReplicate<double>"), TestVectorsEqual(V0, V1));
+
 	V0 = MakeVectorRegisterDouble(1.0f, 2.0f, 3.0f, 4.0f);
 	V1 = VectorReplicate(V0, 1);
 	V0 = MakeVectorRegisterDouble(2.0f, 2.0f, 2.0f, 2.0f);
 	LogTest(TEXT("VectorReplicate<double>"), TestVectorsEqual(V0, V1));
 
+	V0 = MakeVectorRegisterDouble(1.0f, 2.0f, 3.0f, 4.0f);
+	V1 = VectorReplicate(V0, 2);
+	V0 = MakeVectorRegisterDouble(3.0f, 3.0f, 3.0f, 3.0f);
+	LogTest(TEXT("VectorReplicate<double>"), TestVectorsEqual(V0, V1));
+
+	V0 = MakeVectorRegisterDouble(1.0f, 2.0f, 3.0f, 4.0f);
+	V1 = VectorReplicate(V0, 3);
+	V0 = MakeVectorRegisterDouble(4.0f, 4.0f, 4.0f, 4.0f);
+	LogTest(TEXT("VectorReplicate<double>"), TestVectorsEqual(V0, V1));
+
 	// Swizzle
 
+	V0 = MakeVectorRegisterDouble(0.0, 1.0, 2.0, 3.0);
+	V1 = VectorSwizzle(V0, 0, 1, 2, 3);
+	V0 = MakeVectorRegisterDouble(0.0, 1.0, 2.0, 3.0);
+	LogTest(TEXT("VectorShuffle0123<double>"), TestVectorsEqual(V0, V1));
+
 	V0 = MakeVectorRegisterDouble(4.0f, 3.0f, 2.0f, 1.0f);
 	V1 = VectorSwizzle(V0, 1, 0, 1, 0);
 	V0 = MakeVectorRegisterDouble(3.0f, 4.0f, 3.0f, 4.0f);
@@ -992,6 +1013,16 @@ bool RunDoubleVectorTest()
 	V0 = MakeVectorRegisterDouble(0.0, 3.0, 1.0, 2.0);
 	LogTest(TEXT("VectorShuffle0312<double>"), TestVectorsEqual(V0, V1));
 
+	V0 = MakeVectorRegisterDouble(0.0, 1.0, 2.0, 3.0);
+	V1 = VectorSwizzle(V0, 0, 0, 2, 2);
+	V0 = MakeVectorRegisterDouble(0.0, 0.0, 2.0, 2.0);
+	LogTest(TEXT("VectorShuffle0022<double>"), TestVectorsEqual(V0, V1));
+
+	V0 = MakeVectorRegisterDouble(0.0, 1.0, 2.0, 3.0);
+	V1 = VectorSwizzle(V0, 0, 0, 3, 3);
+	V0 = MakeVectorRegisterDouble(0.0, 0.0, 3.0, 3.0);
+	LogTest(TEXT("VectorShuffle0022<double>"), TestVectorsEqual(V0, V1));
+
 	// Shuffle
 
 	V0 = MakeVectorRegisterDouble(0.0, 1.0, 2.0, 3.0);
@@ -1125,11 +1156,18 @@ bool RunDoubleVectorTest()
 	V3 = MakeVectorRegisterDouble((uint64)(0xFFFFFFFFFFFFFFFFU), (uint64)(0xFFFFFFFFFFFFFFFFU), (uint64)(0), (uint64)(0xFFFFFFFFFFFFFFFFU));
 	LogTest(TEXT("VectorCompareNE"), TestVectorsEqualBitwise(V2, V3));
 
-	V0 = MakeVectorRegisterDouble(1.0f, 3.0f, 2.0f, 8.0f);
-	V1 = MakeVectorRegisterDouble(2.0f, 4.0f, 2.0f, 1.0f);
+	V0 = MakeVectorRegisterDouble(1.0f, 3.0f, 5.0f, 7.0f);
+	V1 = MakeVectorRegisterDouble(2.0f, 4.0f, 6.0f, 8.0f);
 	V2 = MakeVectorRegisterDouble((uint64)-1, (uint64)0, (uint64)0, (uint64)-1);
 	V2 = VectorSelect(V2, V0, V1);
-	V3 = MakeVectorRegisterDouble(1.0f, 4.0f, 2.0f, 8.0f);
+	V3 = MakeVectorRegisterDouble(1.0f, 4.0f, 6.0f, 7.0f);
+	LogTest(TEXT("VectorSelect"), TestVectorsEqual(V2, V3));
+
+	V0 = MakeVectorRegisterDouble(1.0f, 3.0f, 5.0f, 7.0f);
+	V1 = MakeVectorRegisterDouble(2.0f, 4.0f, 6.0f, 8.0f);
+	V2 = MakeVectorRegisterDouble((uint64)0, (uint64)-1, (uint64)-1, (uint64)0);
+	V2 = VectorSelect(V2, V0, V1);
+	V3 = MakeVectorRegisterDouble(2.0f, 3.0f, 5.0f, 8.0f);
 	LogTest(TEXT("VectorSelect"), TestVectorsEqual(V2, V3));
 
 	V0 = MakeVectorRegisterDouble(1.0f, 3.0f, 0.0f, 0.0f);
@@ -1178,6 +1216,18 @@ bool RunDoubleVectorTest()
 	V0 = MakeVectorRegisterDouble(4.0f, 4.0f, 6.0f, 8.0f);
 	LogTest(TEXT("VectorMax"), TestVectorsEqual(V0, V1));
 
+	V0 = MakeVectorRegisterDouble(1.0, 3.0, 5.0, 7.0);
+	V1 = MakeVectorRegisterDouble(2.0, 4.0, 6.0, 8.0);
+	V1 = VectorCombineHigh(V0, V1);
+	V0 = MakeVectorRegisterDouble(5.0, 7.0, 6.0, 8.0);
+	LogTest(TEXT("VectorCombineHigh"), TestVectorsEqual(V0, V1));
+
+	V0 = MakeVectorRegisterDouble(1.0, 3.0, 5.0, 7.0);
+	V1 = MakeVectorRegisterDouble(2.0, 4.0, 6.0, 8.0);
+	V1 = VectorCombineLow(V0, V1);
+	V0 = MakeVectorRegisterDouble(1.0, 3.0, 2.0, 4.0);
+	LogTest(TEXT("VectorCombineLow"), TestVectorsEqual(V0, V1));
+
 	//TODO: Rename functions to explicitly mention Double?
 	V0 = MakeVectorRegister(2.0, 4.0, 6.0, 8.0);
 	V1 = MakeVectorRegister(4.0, 3.0, 2.0, 1.0);
@@ -1237,37 +1287,6 @@ bool RunDoubleVectorTest()
 	V3 = MakeVectorRegisterDouble((uint64)(0xFFFFFFFFFFFFFFFF), (uint64)(0xFFFFFFFFFFFFFFFF), (uint64)(0), (uint64)(0xFFFFFFFFFFFFFFFF));
 	LogTest(TEXT("VectorCompareNE"), TestVectorsEqualBitwise(V2, V3));
 
-	V0 = MakeVectorRegister(1.0, 3.0, 2.0, 8.0);
-	V1 = MakeVectorRegister(2.0, 4.0, 2.0, 1.0);
-	V2 = MakeVectorRegisterDouble((uint64)-1, (uint64)0, (uint64)0, (uint64)-1);
-	V2 = VectorSelect(V2, V0, V1);
-	V3 = MakeVectorRegister(1.0, 4.0, 2.0, 8.0);
-	LogTest(TEXT("VectorSelect"), TestVectorsEqual(V2, V3));
-
-	V0 = MakeVectorRegister(1.0, 3.0, 0.0, 0.0);
-	V1 = MakeVectorRegister(0.0, 0.0, 2.0, 1.0);
-	V2 = VectorBitwiseOr(V0, V1);
-	V3 = MakeVectorRegister(1.0, 3.0, 2.0, 1.0);
-	LogTest(TEXT("VectorBitwiseOr-Float1"), TestVectorsEqual(V2, V3));
-
-	V0 = MakeVectorRegister(1.0, 3.0, 24.0, 36.0);
-	V1 = MakeVectorRegisterDouble((uint64)(0x8000000000000000U), (uint64)(0x8000000000000000U), (uint64)(0x8000000000000000U), (uint64)(0x8000000000000000U));
-	V2 = VectorBitwiseOr(V0, V1);
-	V3 = MakeVectorRegister(-1.0, -3.0, -24.0, -36.0);
-	LogTest(TEXT("VectorBitwiseOr-Float2"), TestVectorsEqual(V2, V3));
-
-	V0 = MakeVectorRegister(-1.0, -3.0, -24.0, 36.0);
-	V1 = MakeVectorRegisterDouble((uint64)(0xFFFFFFFFFFFFFFFF), (uint64)(0x7FFFFFFFFFFFFFFFU), (uint64)(0x7FFFFFFFFFFFFFFFU), (uint64)(0xFFFFFFFFFFFFFFFF));
-	V2 = VectorBitwiseAnd(V0, V1);
-	V3 = MakeVectorRegister(-1.0, 3.0, 24.0, 36.0);
-	LogTest(TEXT("VectorBitwiseAnd-Float"), TestVectorsEqual(V2, V3));
-
-	V0 = MakeVectorRegister(-1.0, -3.0, -24.0, 36.0);
-	V1 = MakeVectorRegisterDouble((uint64)(0x8000000000000000), (uint64)(0x0000000000000000U), (uint64)(0x8000000000000000U), (uint64)(0x8000000000000000U));
-	V2 = VectorBitwiseXor(V0, V1);
-	V3 = MakeVectorRegister(1.0, -3.0, 24.0, -36.0);
-	LogTest(TEXT("VectorBitwiseXor-Float"), TestVectorsEqual(V2, V3));
-
 	V0 = MakeVectorRegister(-1.0, -3.0, -24.0, 36.0);
 	V1 = MakeVectorRegister(5.0, 35.0, 23.0, 48.0);
 	V2 = VectorMergeVecXYZ_VecW(V0, V1);
@@ -1404,7 +1423,7 @@ bool FVectorRegisterAbstractionTest::RunTest(const FString& Parameters)
 	uint32 D1 = *(uint32 *)&F1;
 	VectorRegister4Float V0, V1, V2, V3;
 
-	SetPassing(true);
+	ResetPassing();
 
 	V0 = MakeVectorRegister( D1, D1, D1, D1 );
 	V1 = MakeVectorRegister( F1, F1, F1, F1 );
@@ -1487,11 +1506,26 @@ bool FVectorRegisterAbstractionTest::RunTest(const FString& Parameters)
 	V0 = MakeVectorRegister( 1.0f, 5.0f, -0.25f, -0.5f );
 	LogTest( TEXT("VectorStoreFloat1"), TestVectorsEqual( V0, V1 ) );
 
+	V0 = MakeVectorRegister(1.0f, 2.0f, 3.0f, 4.0f);
+	V1 = VectorReplicate(V0, 0);
+	V0 = MakeVectorRegister(1.0f, 1.0f, 1.0f, 1.0f);
+	LogTest(TEXT("VectorReplicate"), TestVectorsEqual(V0, V1));
+
 	V0 = MakeVectorRegister( 1.0f, 2.0f, 3.0f, 4.0f );
 	V1 = VectorReplicate( V0, 1 );
 	V0 = MakeVectorRegister( 2.0f, 2.0f, 2.0f, 2.0f );
 	LogTest( TEXT("VectorReplicate"), TestVectorsEqual( V0, V1 ) );
 
+	V0 = MakeVectorRegister(1.0f, 2.0f, 3.0f, 4.0f);
+	V1 = VectorReplicate(V0, 2);
+	V0 = MakeVectorRegister(3.0f, 3.0f, 3.0f, 3.0f);
+	LogTest(TEXT("VectorReplicate"), TestVectorsEqual(V0, V1));
+
+	V0 = MakeVectorRegister(1.0f, 2.0f, 3.0f, 4.0f);
+	V1 = VectorReplicate(V0, 3);
+	V0 = MakeVectorRegister(4.0f, 4.0f, 4.0f, 4.0f);
+	LogTest(TEXT("VectorReplicate"), TestVectorsEqual(V0, V1));
+
 	V0 = MakeVectorRegister( 1.0f, -2.0f, 3.0f, -4.0f );
 	V1 = VectorAbs( V0 );
 	V0 = MakeVectorRegister( 1.0f, 2.0f, 3.0f, 4.0f );
@@ -1610,6 +1644,25 @@ bool FVectorRegisterAbstractionTest::RunTest(const FString& Parameters)
 	V0 = MakeVectorRegister( 4.0f, 4.0f, 6.0f, 8.0f );
 	LogTest( TEXT("VectorMax"), TestVectorsEqual( V0, V1 ) );
 
+	V0 = MakeVectorRegister(1.0f, 3.0f, 5.0f, 7.0f);
+	V1 = MakeVectorRegister(2.0f, 4.0f, 6.0f, 8.0f);
+	V1 = VectorCombineHigh(V0, V1);
+	V0 = MakeVectorRegister(5.0f, 7.0f, 6.0f, 8.0f);
+	LogTest(TEXT("VectorCombineHigh"), TestVectorsEqual(V0, V1));
+
+	V0 = MakeVectorRegister(1.0f, 3.0f, 5.0f, 7.0f);
+	V1 = MakeVectorRegister(2.0f, 4.0f, 6.0f, 8.0f);
+	V1 = VectorCombineLow(V0, V1);
+	V0 = MakeVectorRegister(1.0f, 3.0f, 2.0f, 4.0f);
+	LogTest(TEXT("VectorCombineLow"), TestVectorsEqual(V0, V1));
+
+	// Swizzle
+
+	V0 = MakeVectorRegister(0.0f, 1.0f, 2.0f, 3.0f);
+	V1 = VectorSwizzle(V0, 0, 1, 2, 3);
+	V0 = MakeVectorRegister(0.0f, 1.0f, 2.0f, 3.0f);
+	LogTest(TEXT("VectorSwizzle0123"), TestVectorsEqual(V0, V1));
+
 	V0 = MakeVectorRegister( 4.0f, 3.0f, 2.0f, 1.0f );
 	V1 = VectorSwizzle( V0, 1, 0, 3, 2 );
 	V0 = MakeVectorRegister( 3.0f, 4.0f, 1.0f, 2.0f );
@@ -1635,6 +1688,13 @@ bool FVectorRegisterAbstractionTest::RunTest(const FString& Parameters)
 	V0 = MakeVectorRegister( 1.0f, 2.0f, 3.0f, 4.0f );
 	LogTest( TEXT("VectorSwizzle3210"), TestVectorsEqual( V0, V1 ) );
 
+	V0 = MakeVectorRegister(0.0f, 1.0f, 2.0f, 3.0f);
+	V1 = VectorSwizzle(V0, 0, 0, 2, 2);
+	V0 = MakeVectorRegister(0.0f, 0.0f, 2.0f, 2.0f);
+	LogTest(TEXT("VectorSwizzle0022"), TestVectorsEqual(V0, V1));
+
+	// LoadByte4
+
 	uint8 Bytes[4] = { 25, 75, 125, 200 };
 	V0 = VectorLoadByte4( Bytes );
 	V1 = MakeVectorRegister( 25.f, 75.f, 125.f, 200.f );
@@ -1715,6 +1775,13 @@ bool FVectorRegisterAbstractionTest::RunTest(const FString& Parameters)
 	V3 = MakeVectorRegister( 1.0f, 4.0f, 2.0f, 8.0f );
 	LogTest( TEXT("VectorSelect"), TestVectorsEqual( V2, V3 ) );
 
+	V0 = MakeVectorRegister(1.0f, 3.0f, 5.0f, 7.0f);
+	V1 = MakeVectorRegister(2.0f, 4.0f, 6.0f, 8.0f);
+	V2 = MakeVectorRegister((uint32)0, (uint32)-1, (uint32)-1, (uint32)0);
+	V2 = VectorSelect(V2, V0, V1);
+	V3 = MakeVectorRegister(2.0f, 3.0f, 5.0f, 8.0f);
+	LogTest(TEXT("VectorSelect"), TestVectorsEqual(V2, V3));
+
 	V0 = MakeVectorRegister( 1.0f, 3.0f, 0.0f, 0.0f );
 	V1 = MakeVectorRegister( 0.0f, 0.0f, 2.0f, 1.0f );
 	V2 = VectorBitwiseOr( V0, V1 );
@@ -2001,7 +2068,7 @@ bool FVectorRegisterAbstractionTest::RunTest(const FString& Parameters)
 			if (Delta > 1e-8f)
 			{
 				UE_LOG(LogUnrealMathTest, Log, TEXT("FMath::Fmod(%f, %f)=%f <-> fmodf(%f, %f)=%f: FAILED"), X, Y, Ours, X, Y, Theirs);
-				SetPassing(false);
+				CheckPassing(false);
 			}
 		}
 	}
@@ -2087,7 +2154,7 @@ bool FVectorRegisterAbstractionTest::RunTest(const FString& Parameters)
 				if (FractionalDelta > 1e-8)
 				{
 					UE_LOG(LogUnrealMathTest, Log, TEXT("FMath::Fmod(%.12f, %.12f)=%.12f <-> fmod(%.12f, %.12f)=%.12f: FAILED"), X, Y, Ours, X, Y, Theirs);
-					SetPassing(false);
+					CheckPassing(false);
 				}				
 			}
 		}
diff --git a/Engine/Source/Runtime/Core/Public/Math/TransformVectorized.h b/Engine/Source/Runtime/Core/Public/Math/TransformVectorized.h
--- a/Engine/Source/Runtime/Core/Public/Math/TransformVectorized.h
+++ b/Engine/Source/Runtime/Core/Public/Math/TransformVectorized.h
@@ -33,7 +33,7 @@
  * Example: LocalToWorld = (LocalToWorld * DeltaRotation) will change rotation in world space by DeltaRotation.
  */
 
-MS_ALIGN(16) struct FTransform
+struct alignas(alignof(VectorRegister)) FTransform
 {
 	friend struct Z_Construct_UScriptStruct_FTransform_Statics;
 
@@ -108,11 +108,11 @@ public:
 	FORCEINLINE FTransform()
 	{
 		// Rotation = {0,0,0,1)
-		Rotation = VectorSet_W1( VectorZero() );
+		Rotation = GlobalVectorConstants::Vector0001; //VectorSet_W1( VectorZero() );
 		// Translation = {0,0,0,0)
 		Translation = VectorZero();
 		// Scale3D = {1,1,1,0);
-		Scale3D = VectorSet_W0( VectorOne() );
+		Scale3D = GlobalVectorConstants::Vector1110; //VectorSet_W0( VectorOne() );
 	}
 
 	/**
@@ -123,11 +123,11 @@ public:
 	FORCEINLINE explicit FTransform(const FVector& InTranslation) 
 	{
 		// Rotation = {0,0,0,1) quaternion identity
-		Rotation =  VectorSet_W1( VectorZero() );
+		Rotation = GlobalVectorConstants::Vector0001; //VectorSet_W1( VectorZero() );
 		//Translation = InTranslation;
 		Translation = MakeVectorRegister(InTranslation.X, InTranslation.Y, InTranslation.Z, 0.0f );
 		// Scale3D = {1,1,1,0);
-		Scale3D = VectorSet_W0( VectorOne() );
+		Scale3D = GlobalVectorConstants::Vector1110; //VectorSet_W0( VectorOne() );
 
 		DiagnosticCheckNaN_All();
 	}
@@ -144,7 +144,7 @@ public:
 		// Translation = {0,0,0,0)
 		Translation = VectorZero();
 		// Scale3D = {1,1,1,0);
-		Scale3D = VectorSet_W0( VectorOne() );
+		Scale3D = GlobalVectorConstants::Vector1110; //VectorSet_W0( VectorOne() );
 
 		DiagnosticCheckNaN_All();
 	}
@@ -162,7 +162,7 @@ public:
 		// Translation = {0,0,0,0)
 		Translation = VectorZero();
 		// Scale3D = {1,1,1,0);
-		Scale3D = VectorSet_W0( VectorOne() );
+		Scale3D = GlobalVectorConstants::Vector1110; //VectorSet_W0( VectorOne() );
 
 		DiagnosticCheckNaN_All();
 	}
@@ -680,7 +680,7 @@ private:
 
 	FORCEINLINE static bool Private_AnyHasNegativeScale(const VectorRegister& InScale3D, const  VectorRegister& InOtherScale3D)
 	{
-		return !!VectorAnyLesserThan(VectorMin(InScale3D, InOtherScale3D), GlobalVectorConstants::FloatZero);
+		return !!VectorAnyLesserThan(VectorMin(InScale3D, InOtherScale3D), VectorZero());
 	}
 
 	FORCEINLINE bool Private_RotationEquals( const VectorRegister& InRotation, const ScalarRegister& Tolerance = ScalarRegister(GlobalVectorConstants::KindaSmallNumber)) const
@@ -790,11 +790,11 @@ public:
 	FORCEINLINE void SetIdentity()
 	{
 		// Rotation = {0,0,0,1)
-		Rotation = VectorSet_W1( VectorZero() );
+		Rotation = GlobalVectorConstants::Vector0001; //VectorSet_W1( VectorZero() );
 		// Translation = {0,0,0,0)
 		Translation = VectorZero();
 		// Scale3D = {1,1,1,0);
-		Scale3D = VectorSet_W0 ( VectorOne() );
+		Scale3D = GlobalVectorConstants::Vector1110; //VectorSet_W0( VectorOne() );
 	}
 
 	/**
@@ -1107,7 +1107,7 @@ public:
 	 */
 	FORCEINLINE static void BlendFromIdentityAndAccumulate(FTransform& FinalAtom, const FTransform& SourceAtom, const ScalarRegister& BlendWeight)
 	{
-		const VectorRegister Const0001 = GlobalVectorConstants::Float0001;
+		const VectorRegister Const0001 = GlobalVectorConstants::Vector0001;
 		const VectorRegister ConstNegative0001 = VectorSubtract(VectorZero(), Const0001);
 		const VectorRegister VOneMinusAlpha = VectorSubtract(VectorOne(), BlendWeight.Value);
 		const VectorRegister DefaultScale = MakeVectorRegister(1.f, 1.f, 1.f, 0.f);
@@ -1416,7 +1416,7 @@ private:
 	* @param  Relative Transform Relative.
 	*/
 	static void GetRelativeTransformUsingMatrixWithScale(FTransform* OutTransform, const FTransform* Base, const FTransform* Relative);
-}  GCC_ALIGN(16);
+};
 
 template <> struct TIsPODType<FTransform> { enum { Value = true }; };
 
diff --git a/Engine/Source/Runtime/Core/Public/Math/UnrealMathFPU.h b/Engine/Source/Runtime/Core/Public/Math/UnrealMathFPU.h
--- a/Engine/Source/Runtime/Core/Public/Math/UnrealMathFPU.h
+++ b/Engine/Source/Runtime/Core/Public/Math/UnrealMathFPU.h
@@ -86,7 +86,7 @@ void VectorStoreAligned(const VectorRegister4Double& Vec, double* Dst);
 
 
 // Helper for conveniently aligning a float array for extraction from VectorRegister4Float
-struct alignas(16) AlignedFloat4
+struct alignas(alignof(VectorRegister4Float)) AlignedFloat4
 {
 	float V[4];
 
@@ -106,7 +106,7 @@ struct alignas(16) AlignedFloat4
 
 
 // Helper for conveniently aligning a double array for extraction from VectorRegister4Double
-struct alignas(16) AlignedDouble4
+struct alignas(alignof(VectorRegister4Double)) AlignedDouble4
 {
 	double V[4];
 
@@ -143,16 +143,6 @@ FORCEINLINE VectorRegister2Double MakeVectorRegister2Double(double X, double Y)
 	return Result;
 }
 
-// Bitwise equivalent from two 64-bit values.
-FORCEINLINE VectorRegister2Double MakeVectorRegister2Double(uint64 X, uint64 Y)
-{
-	VectorRegister2Double Result;
-	Result.V[0] = (double)X;
-	Result.V[1] = (double)Y;
-	return Result;
-}
-
-
 /**
  * Returns a bitwise equivalent vector based on 4 DWORDs.
  *
@@ -187,6 +177,17 @@ FORCEINLINE VectorRegister4Float MakeVectorRegister(uint32 X, uint32 Y, uint32 Z
 	return MakeVectorRegisterFloat(X, Y, Z, W);
 }
 
+// Nicer aliases
+FORCEINLINE VectorRegister4Float MakeVectorRegisterFloatMask(uint32 X, uint32 Y, uint32 Z, uint32 W)
+{
+	return MakeVectorRegisterFloat(X, Y, Z, W);
+}
+
+FORCEINLINE VectorRegister4Double MakeVectorRegisterDoubleMask(uint64 X, uint64 Y, uint64 Z, uint64 W)
+{
+	return MakeVectorRegisterDouble(X, Y, Z, W);
+}
+
 /**
  * Returns a vector based on 4 FLOATs.
  *
@@ -401,11 +402,6 @@ FORCEINLINE VectorRegister4Float VectorLoadFloat1(const float* Ptr)
 	return MakeVectorRegisterFloat(Ptr[0], Ptr[0], Ptr[0], Ptr[0]);
 }
 
-FORCEINLINE VectorRegister4Double VectorLoadFloat1(const double* Ptr)
-{
-	return MakeVectorRegisterDouble(Ptr[0], Ptr[0], Ptr[0], Ptr[0]);
-}
-
 FORCEINLINE VectorRegister4Double VectorLoadDouble1(const double* Ptr)
 {
 	return MakeVectorRegisterDouble(Ptr[0], Ptr[0], Ptr[0], Ptr[0]);
diff --git a/Engine/Source/Runtime/Core/Public/Math/UnrealMathNeon.h b/Engine/Source/Runtime/Core/Public/Math/UnrealMathNeon.h
--- a/Engine/Source/Runtime/Core/Public/Math/UnrealMathNeon.h
+++ b/Engine/Source/Runtime/Core/Public/Math/UnrealMathNeon.h
@@ -31,7 +31,7 @@ PRAGMA_DISABLE_SHADOW_VARIABLE_WARNINGS
 // redundant zero initialization in these cases for non-MSVC (e.g. V()
 // is called now where it wasn't before)
 template<typename T, bool is_int>
-struct MS_ALIGN(alignof(T)) VectorRegisterWrapper
+struct alignas(alignof(T)) VectorRegisterWrapper
 {
 	FORCEINLINE VectorRegisterWrapper() {}
 	FORCEINLINE VectorRegisterWrapper(const T& vec) : m_vec(vec) {}
@@ -67,7 +67,7 @@ typedef uint64x2_t GCC_ALIGN(16) VectorRegister2Int64;
 	typedef VectorRegister4Float VectorRegister;
 #endif
 
-MS_ALIGN(16) struct VectorRegister4Double
+struct alignas(16) VectorRegister4Double
 {
 	union
 	{
@@ -95,7 +95,7 @@ MS_ALIGN(16) struct VectorRegister4Double
 		*this = VectorRegister4Double(From);
 		return *this;
 	}
-} GCC_ALIGN(16);
+};
 
 
 // Forward declarations
@@ -106,7 +106,7 @@ void VectorStoreAligned(const VectorRegister4Double& Vec, double* Dst);
 
 
 // Helper for conveniently aligning a float array for extraction from VectorRegister4Float
-struct alignas(16) AlignedFloat4
+struct alignas(alignof(VectorRegister4Float)) AlignedFloat4
 {
 	float V[4];
 
@@ -126,7 +126,7 @@ struct alignas(16) AlignedFloat4
 
 
 // Helper for conveniently aligning a double array for extraction from VectorRegister4Double
-struct alignas(16) AlignedDouble4
+struct alignas(alignof(VectorRegister4Double)) AlignedDouble4
 {
 	double V[4];
 
@@ -179,6 +179,13 @@ FORCEINLINE VectorRegister4Float MakeVectorRegisterFloat(uint32 X, uint32 Y, uin
 	return MakeVectorRegister(X, Y, Z, W);
 }
 
+// Nicer alias
+FORCEINLINE VectorRegister4Float MakeVectorRegisterFloatMask(uint32 X, uint32 Y, uint32 Z, uint32 W)
+{
+	return MakeVectorRegisterFloat(X, Y, Z, W);
+}
+
+
 /**
  * Returns a vector based on 4 floats.
  *
@@ -248,6 +255,12 @@ FORCEINLINE VectorRegister4Double MakeVectorRegisterDouble(uint64 X, uint64 Y, u
 	return Tmp.V;
 }
 
+// Nicer alias
+FORCEINLINE VectorRegister4Double MakeVectorRegisterDoubleMask(uint64 X, uint64 Y, uint64 Z, uint64 W)
+{
+	return MakeVectorRegisterDouble(X, Y, Z, W);
+}
+
 FORCEINLINE VectorRegister2Double MakeVectorRegister2Double(double X, double Y)
 {
 	union U
diff --git a/Engine/Source/Runtime/Core/Public/Math/UnrealMathSSE.h b/Engine/Source/Runtime/Core/Public/Math/UnrealMathSSE.h
--- a/Engine/Source/Runtime/Core/Public/Math/UnrealMathSSE.h
+++ b/Engine/Source/Runtime/Core/Public/Math/UnrealMathSSE.h
@@ -2,21 +2,38 @@
 
 #pragma once
 
-#if defined(__cplusplus_cli) && !PLATFORM_HOLOLENS
-// there are compile issues with this file in managed mode, so use the FPU version
-#include "Math/UnrealMathFPU.h"
-#else
-
 // We require SSE2
 #include <emmintrin.h>
 
-#define UNREALPLATFORMMATH_SSE_USE_SSE4_1 PLATFORM_ALWAYS_HAS_SSE4_1
+#ifndef UE_PLATFORM_MATH_USE_SSE4_1
+#define UE_PLATFORM_MATH_USE_SSE4_1			PLATFORM_ALWAYS_HAS_SSE4_1
+#endif
+
+#ifndef UE_PLATFORM_MATH_USE_AVX
+#define UE_PLATFORM_MATH_USE_AVX			PLATFORM_ALWAYS_HAS_AVX
+#endif
+
+#ifndef UE_PLATFORM_MATH_USE_AVX_2
+#define UE_PLATFORM_MATH_USE_AVX_2			0 // PLATFORM_ALWAYS_HAS_AVX_2
+#endif
+
+#ifndef UE_PLATFORM_MATH_USE_FMA3
+#define UE_PLATFORM_MATH_USE_FMA3			PLATFORM_ALWAYS_HAS_FMA3
+#endif
 
 // If SSE4.1 is enabled, need additional defines.
-#if UNREALPLATFORMMATH_SSE_USE_SSE4_1
+#if UE_PLATFORM_MATH_USE_SSE4_1
 #include <smmintrin.h>
 #endif
 
+// If AVX is enabled, need additional defines.
+#if UE_PLATFORM_MATH_USE_AVX
+#include <immintrin.h>
+#define UE_SSE_DOUBLE_ALIGNMENT 32
+#else
+#define UE_SSE_DOUBLE_ALIGNMENT 16
+#endif
+
 #include "Math/sse_mathfun.h"
 
 // We suppress static analysis warnings for the cast from (double*) to (float*) in VectorLoadFloat2
@@ -43,28 +60,75 @@ typedef __m128d	VectorRegister2Double;
 
 
 // 4 doubles
-struct alignas(16) VectorRegister4Double
+struct alignas(UE_SSE_DOUBLE_ALIGNMENT) VectorRegister4Double
 {
+#if !UE_PLATFORM_MATH_USE_AVX
 	VectorRegister2Double XY;
 	VectorRegister2Double ZW;
+#else
+	union
+	{
+		struct
+		{
+			VectorRegister2Double XY;
+			VectorRegister2Double ZW;
+		};
+		__m256d XYZW;
+	};
+#endif
+
+	FORCEINLINE VectorRegister4Double() = default;
 
-	VectorRegister4Double() = default;
+	FORCEINLINE VectorRegister4Double(const VectorRegister2Double& InXY, const VectorRegister2Double& InZW)
+		: XY(InXY)
+		, ZW(InZW)
+	{
+	}
 
 	// Construct from a vector of 4 floats
-	VectorRegister4Double(const VectorRegister4Float& FloatVector)
+	FORCEINLINE VectorRegister4Double(const VectorRegister4Float& FloatVector)
 	{
+#if !UE_PLATFORM_MATH_USE_AVX
 		XY = _mm_cvtps_pd(FloatVector);
 		ZW = _mm_cvtps_pd(_mm_movehl_ps(FloatVector, FloatVector));
+#else
+		XYZW = _mm256_cvtps_pd(FloatVector);
+#endif
 	}
 
 	// Assign from a vector of 4 floats
-	VectorRegister4Double& operator=(const VectorRegister4Float& From)
+	FORCEINLINE VectorRegister4Double& operator=(const VectorRegister4Float& FloatVector)
+	{
+#if !UE_PLATFORM_MATH_USE_AVX
+		XY = _mm_cvtps_pd(FloatVector);
+		ZW = _mm_cvtps_pd(_mm_movehl_ps(FloatVector, FloatVector));
+#else
+		XYZW = _mm256_cvtps_pd(FloatVector);
+#endif
+		return *this;
+	}
+
+#if UE_PLATFORM_MATH_USE_AVX
+	// Convenience for things like 'Result = _mm256_add_pd(...)'
+	FORCEINLINE VectorRegister4Double(const __m256d& Register)
 	{
-		XY = _mm_cvtps_pd(From);
-		ZW = _mm_cvtps_pd(_mm_movehl_ps(From, From));
+		XYZW = Register;
+	}
+
+	// Convenience for things like 'Result = _mm256_add_pd(...)'
+	FORCEINLINE VectorRegister4Double& operator=(const __m256d& Register)
+	{
+		XYZW = Register;
 		return *this;
 	}
 
+	// Convenience for passing VectorRegister4Double to _mm256_* functions without needing '.XYZW'
+	FORCEINLINE operator __m256d() const
+	{
+		return XYZW;
+	}
+#endif
+
 };
 
 
@@ -114,7 +178,7 @@ struct alignas(16) AlignedFloat4
 
 
 // Helper for conveniently aligning a double array for extraction from VectorRegister4Double
-struct alignas(16) AlignedDouble4
+struct alignas(alignof(VectorRegister4Double)) AlignedDouble4
 {
 	double V[4];
 
@@ -159,7 +223,7 @@ FORCEINLINE VectorRegister2Double MakeVectorRegister2Double(double X, double Y)
 }
 
 // Bitwise equivalent from two 64-bit values.
-FORCEINLINE VectorRegister2Double MakeVectorRegister2Double(uint64 X, uint64 Y)
+FORCEINLINE VectorRegister2Double MakeVectorRegister2DoubleMask(uint64 X, uint64 Y)
 {
 	union { VectorRegister2Double Vd; __m128i Vi; } Result;
 	// Note: this instruction only exists on 64-bit
@@ -167,7 +231,6 @@ FORCEINLINE VectorRegister2Double MakeVectorRegister2Double(uint64 X, uint64 Y)
 	return Result.Vd;
 }
 
-
 /**
  * Returns a bitwise equivalent vector based on 4 DWORDs.
  *
@@ -187,8 +250,8 @@ FORCEINLINE VectorRegister4Float MakeVectorRegisterFloat( uint32 X, uint32 Y, ui
 FORCEINLINE VectorRegister4Double MakeVectorRegisterDouble(uint64 X, uint64 Y, uint64 Z, uint64 W)
 {
 	VectorRegister4Double Result;
-	Result.XY = MakeVectorRegister2Double(X, Y);
-	Result.ZW = MakeVectorRegister2Double(Z, W);
+	Result.XY = MakeVectorRegister2DoubleMask(X, Y);
+	Result.ZW = MakeVectorRegister2DoubleMask(Z, W);
 	return Result;
 }
 
@@ -197,6 +260,17 @@ FORCEINLINE VectorRegister4Float MakeVectorRegister(uint32 X, uint32 Y, uint32 Z
 	return MakeVectorRegisterFloat(X, Y, Z, W);
 }
 
+// Nicer aliases
+FORCEINLINE VectorRegister4Float MakeVectorRegisterFloatMask(uint32 X, uint32 Y, uint32 Z, uint32 W)
+{
+	return MakeVectorRegisterFloat(X, Y, Z, W);
+}
+
+FORCEINLINE VectorRegister4Double MakeVectorRegisterDoubleMask(uint64 X, uint64 Y, uint64 Z, uint64 W)
+{
+	return MakeVectorRegisterDouble(X, Y, Z, W);
+}
+
 /**
  * Returns a vector based on 4 FLOATs.
  *
@@ -214,8 +288,12 @@ FORCEINLINE VectorRegister4Float MakeVectorRegisterFloat(float X, float Y, float
 FORCEINLINE VectorRegister4Double MakeVectorRegisterDouble(double X, double Y, double Z, double W)
 {
 	VectorRegister4Double Result;
+#if !UE_PLATFORM_MATH_USE_AVX
 	Result.XY = _mm_setr_pd(X, Y);
 	Result.ZW = _mm_setr_pd(Z, W);
+#else
+	Result = _mm256_setr_pd(X, Y, Z, W);
+#endif
 	return Result;
 }
 
@@ -238,7 +316,11 @@ FORCEINLINE VectorRegister4Double MakeVectorRegisterDouble(const VectorRegister4
 // Lossy conversion: double->float vector
 FORCEINLINE VectorRegister4Float MakeVectorRegisterFloatFromDouble(const VectorRegister4Double& Vec4d)
 {
+#if !UE_PLATFORM_MATH_USE_AVX
 	return _mm_movelh_ps(_mm_cvtpd_ps(Vec4d.XY), _mm_cvtpd_ps(Vec4d.ZW));
+#else
+	return _mm256_cvtpd_ps(Vec4d);
+#endif
 }
 
 /**
@@ -278,8 +360,12 @@ FORCEINLINE VectorRegister4Float VectorZeroFloat(void)
 FORCEINLINE VectorRegister4Double VectorZeroDouble(void)
 {
 	VectorRegister4Double Result;
+#if !UE_PLATFORM_MATH_USE_AVX
 	Result.XY = _mm_setzero_pd();
 	Result.ZW = _mm_setzero_pd();
+#else
+	Result = _mm256_setzero_pd();
+#endif
 	return Result;
 }
 
@@ -315,13 +401,11 @@ FORCEINLINE float VectorGetComponentImpl(const VectorRegister4Float& Vec)
 template <uint32 ComponentIndex>
 FORCEINLINE double VectorGetComponentImpl(const VectorRegister4Double& Vec)
 {
-	// TODO: LWC: verify this is equivalent to _mm_store_pd and a read, or better.
+#if !UE_PLATFORM_MATH_USE_AVX
 	return (((double*)&(Vec.XY))[ComponentIndex]);
-
-	/*
-	AlignedDouble4 Components(Vec);
-	return Components[ComponentIndex];
-	*/
+#else
+	return (((double*)&(Vec.XYZW))[ComponentIndex]);
+#endif
 }
 
 #define VectorGetComponent(Vec, ComponentIndex) VectorGetComponentImpl<ComponentIndex>(Vec)
@@ -351,8 +435,12 @@ FORCEINLINE VectorRegister4Float VectorLoad(const float* Ptr)
 FORCEINLINE VectorRegister4Double VectorLoad(const double* Ptr)
 {
 	VectorRegister4Double Result;
+#if !UE_PLATFORM_MATH_USE_AVX
 	Result.XY = _mm_loadu_pd((double*)(Ptr));
 	Result.ZW = _mm_loadu_pd((double*)(Ptr + 2));
+#else
+	Result = _mm256_loadu_pd((double*)Ptr);
+#endif
 	return Result;
 }
 
@@ -364,11 +452,16 @@ FORCEINLINE VectorRegister4Double VectorLoad(const double* Ptr)
  */
 FORCEINLINE VectorRegister4Double VectorLoadFloat3(const double* Ptr)
 {
-	//return MakeVectorRegisterDouble(Ptr[0], Ptr[1], Ptr[2], 0.0);
+#if !UE_PLATFORM_MATH_USE_AVX
 	VectorRegister4Double Result;
 	Result.XY = _mm_loadu_pd((double*)(Ptr));
 	Result.ZW = _mm_load_sd((double*)(Ptr+2));
 	return Result;
+#else
+	// TODO: AVX: compare performance. MaskLoad is supposed to be efficient compared to MaskStore (from Intel and AMD manuals citing instruction latency), but validate this.
+	// Otherwise, use the m128d code above.
+	return _mm256_maskload_pd(Ptr, _mm256_castpd_si256(GlobalVectorConstants::DoubleXYZMask));
+#endif	
 }
 
 /**
@@ -379,11 +472,15 @@ FORCEINLINE VectorRegister4Double VectorLoadFloat3(const double* Ptr)
  */
 FORCEINLINE VectorRegister4Double VectorLoadFloat3_W1(const double* Ptr)
 {
-	//return MakeVectorRegisterDouble(Ptr[0], Ptr[1], Ptr[2], 1.0);
+#if !UE_PLATFORM_MATH_USE_AVX
 	VectorRegister4Double Result;
 	Result.XY = _mm_loadu_pd((double*)(Ptr));
 	Result.ZW = MakeVectorRegister2Double(Ptr[2], 1.0);
 	return Result;
+#else
+	// TODO: AVX: improve?
+	return MakeVectorRegisterDouble(Ptr[0], Ptr[1], Ptr[2], 1.0);
+#endif	
 }
 
 /**
@@ -400,8 +497,13 @@ FORCEINLINE VectorRegister4Float VectorLoadAligned(const float* Ptr)
 FORCEINLINE VectorRegister4Double VectorLoadAligned(const double* Ptr)
 {
 	VectorRegister4Double Result;
+#if !UE_PLATFORM_MATH_USE_AVX
 	Result.XY = _mm_load_pd((const double*)(Ptr));
 	Result.ZW = _mm_load_pd((const double*)(Ptr + 2));
+#else
+	// TODO: AVX using unaligned here, until we ensure 32-byte alignment (was previously 16 across most types)
+	Result = _mm256_loadu_pd(Ptr);
+#endif
 	return Result;
 }
 
@@ -413,22 +515,25 @@ FORCEINLINE VectorRegister4Double VectorLoadAligned(const double* Ptr)
  */
 FORCEINLINE VectorRegister4Float VectorLoadFloat1(const float* Ptr)
 {
+#if !UE_PLATFORM_MATH_USE_AVX
 	return _mm_load1_ps(Ptr);
+#else
+	return _mm_broadcast_ss(Ptr);
+#endif
 }
 
-FORCEINLINE VectorRegister4Double VectorLoadFloat1(const double* Ptr)
+FORCEINLINE VectorRegister4Double VectorLoadDouble1(const double* Ptr)
 {
 	VectorRegister4Double Result;
+#if !UE_PLATFORM_MATH_USE_AVX
 	Result.XY = _mm_load1_pd(Ptr);
 	Result.ZW = Result.XY;
+#else
+	Result = _mm256_broadcast_sd(Ptr);
+#endif
 	return Result;
 }
 
-FORCEINLINE VectorRegister4Double VectorLoadDouble1(const double* Ptr)
-{
-	return VectorLoadFloat1(Ptr);
-}
-
 /**
  * Loads 2 floats from unaligned memory into X and Y and duplicates them in Z and W.
  *
@@ -437,15 +542,20 @@ FORCEINLINE VectorRegister4Double VectorLoadDouble1(const double* Ptr)
  */
 FORCEINLINE VectorRegister4Float VectorLoadFloat2(const float* Ptr)
 {
-	// Note: SA warning for this is disabled at top of file.
-	return _mm_castpd_ps(_mm_load1_pd((const double*)(Ptr)));
+	// This intentionally casts to a double* to be able to load 64 bits of data using the "load 1 double" instruction to fill in the two 32-bit floats.
+	return _mm_castpd_ps(_mm_load1_pd(reinterpret_cast<const double*>(Ptr))); // -V615
 }
 
 FORCEINLINE VectorRegister4Double VectorLoadFloat2(const double* Ptr)
 {
 	VectorRegister4Double Result;
+#if !UE_PLATFORM_MATH_USE_AVX
 	Result.XY = _mm_loadu_pd(Ptr);
 	Result.ZW = Result.XY;
+#else
+	const __m128d Temp = _mm_loadu_pd(Ptr);
+	Result = _mm256_set_m128d(Temp, Temp);
+#endif
 	return Result;
 }
 
@@ -459,7 +569,8 @@ FORCEINLINE VectorRegister4Double VectorLoadFloat2(const double* Ptr)
  */
 FORCEINLINE VectorRegister4Float VectorLoadTwoPairsFloat(const float* Ptr1, const float* Ptr2)
 {
-	__m128 Ret = _mm_castpd_ps(_mm_load_sd((double const *)(Ptr1))); // -V615
+	// This intentionally casts to a double* to be able to load 64 bits of data using the "load 1 double" instruction to fill in the two 32-bit floats.
+	__m128 Ret = _mm_castpd_ps(_mm_load_sd(reinterpret_cast<const double*>(Ptr1))); // -V615
 	Ret = _mm_loadh_pi(Ret, (__m64 const*)(Ptr2));
 	return Ret;
 }
@@ -467,8 +578,12 @@ FORCEINLINE VectorRegister4Float VectorLoadTwoPairsFloat(const float* Ptr1, cons
 FORCEINLINE VectorRegister4Double VectorLoadTwoPairsFloat(const double* Ptr1, const double* Ptr2)
 {
 	VectorRegister4Double Result;
+#if !UE_PLATFORM_MATH_USE_AVX
 	Result.XY = _mm_loadu_pd(Ptr1);
 	Result.ZW = _mm_loadu_pd(Ptr2);
+#else
+	Result = _mm256_loadu2_m128d(Ptr2, Ptr1); // Note: arguments are (hi, lo)
+#endif
 	return Result;
 }
 
@@ -486,8 +601,12 @@ FORCEINLINE VectorRegister4Float VectorSetFloat1(float F)
 FORCEINLINE VectorRegister4Double VectorSetFloat1(double D)
 {
 	VectorRegister4Double Result;
+#if !UE_PLATFORM_MATH_USE_AVX
 	Result.XY = _mm_set1_pd(D);
 	Result.ZW = Result.XY;
+#else
+	Result = _mm256_set1_pd(D);
+#endif
 	return Result;
 }
 
@@ -504,8 +623,13 @@ FORCEINLINE void VectorStoreAligned(const VectorRegister4Float& Vec, float* Dst)
 
 FORCEINLINE void VectorStoreAligned(const VectorRegister4Double& Vec, double* Dst)
 {
+#if !UE_PLATFORM_MATH_USE_AVX
 	_mm_store_pd(Dst, Vec.XY);
 	_mm_store_pd(Dst + 2, Vec.ZW);
+#else
+	// TODO: AVX using unaligned store until we verify 32-byte alignment requirement (was 16 for most types)
+	_mm256_storeu_pd(Dst, Vec);
+#endif
 }
 
 
@@ -522,19 +646,16 @@ FORCEINLINE void VectorStoreAlignedStreamed(const VectorRegister4Float& Vec, flo
 
 FORCEINLINE void VectorStoreAlignedStreamed(const VectorRegister4Double& Vec, double* Dst)
 {
+#if !UE_PLATFORM_MATH_USE_AVX
 	_mm_stream_pd(Dst, Vec.XY);
 	_mm_stream_pd(Dst + 2, Vec.ZW);
+#else
+	// TODO: AVX using two 128-bit stores until we verify 32-byte alignment requirement for our data with _mm256_stream_pd (was 16 for most types)
+	_mm_stream_pd(Dst, Vec.XY);
+	_mm_stream_pd(Dst + 2, Vec.ZW);
+#endif
 }
 
-// TODO: this is lossy, would rather not implement.
-/*
-FORCEINLINE void VectorStoreAlignedStreamed(const VectorRegister4Double& Vec, float* Dst)
-{
-	const VectorRegister4Float FloatXYZW = MakeVectorRegisterFloatFromDouble(Vec);
-	VectorStoreAlignedStreamed(FloatXYZW, Dst);
-}
-*/
-
 /**
  * Stores a vector to memory (aligned or unaligned).
  *
@@ -548,8 +669,12 @@ FORCEINLINE void VectorStore(const VectorRegister4Float& Vec, float* Ptr)
 
 FORCEINLINE void VectorStore(const VectorRegister4Double& Vec, double* Dst)
 {
+#if !UE_PLATFORM_MATH_USE_AVX
 	_mm_storeu_pd(Dst, Vec.XY);
 	_mm_storeu_pd(Dst + 2, Vec.ZW);
+#else
+	_mm256_storeu_pd(Dst, Vec);
+#endif
 }
 
 /**
@@ -570,8 +695,15 @@ FORCEINLINE void VectorStoreFloat3( const VectorRegister4Float& Vec, float* Ptr
 
 FORCEINLINE void VectorStoreFloat3(const VectorRegister4Double& Vec, double* Dst)
 {
+#if !UE_PLATFORM_MATH_USE_AVX
 	_mm_storeu_pd(Dst, Vec.XY);
 	_mm_store_sd(Dst + 2, Vec.ZW);
+#else
+	// TODO: AVX: compare perf. MaskStore is supposedly not so great on AMD Zen1 and Zen2 or older Intel.
+	_mm_storeu_pd(Dst, Vec.XY);
+	_mm_store_sd(Dst + 2, Vec.ZW);
+	//_mm256_maskstore_pd(Dst, _mm256_castpd_si256(GlobalVectorConstants::DoubleXYZMask), Vec);
+#endif
 }
 
 /**
@@ -591,46 +723,115 @@ FORCEINLINE void VectorStoreFloat1(const VectorRegister4Double& Vec, double* Dst
 }
 
 // Templated shuffles required for double shuffles when using __m128d, since we have to break it down in to two separate operations.
+// TODO: more specializations for simple cases. SSE 4.1 and AVX have cheaper blends available.
 
 // [0,1]:[0,1]
 template<int Index0, int Index1, typename std::enable_if< (Index0 <= 1) && (Index1 <= 1), bool >::type = true>
 FORCEINLINE VectorRegister2Double SelectVectorSwizzle2(const VectorRegister4Double& Vec)
 {
-	VectorRegister2Double Result = _mm_shuffle_pd(Vec.XY, Vec.XY, SHUFFLEMASK2(Index0, Index1));
-	return Result;
+	return _mm_shuffle_pd(Vec.XY, Vec.XY, SHUFFLEMASK2(Index0, Index1));
 }
+// <0,1> is simply XY.
+template<> FORCEINLINE VectorRegister2Double SelectVectorSwizzle2<0, 1>(const VectorRegister4Double& Vec) { return Vec.XY; }
 
 // [0,1]:[2,3]
 template<int Index0, int Index1, typename std::enable_if< (Index0 <= 1) && (Index1 > 1), bool >::type = true>
 FORCEINLINE VectorRegister2Double SelectVectorSwizzle2(const VectorRegister4Double& Vec)
 {
-	VectorRegister2Double Result = _mm_shuffle_pd(Vec.XY, Vec.ZW, SHUFFLEMASK2(Index0, Index1 - 2));
-	return Result;
+	return _mm_shuffle_pd(Vec.XY, Vec.ZW, SHUFFLEMASK2(Index0, Index1 - 2));
 }
 
 // [2,3]:[0,1]
 template<int Index0, int Index1, typename std::enable_if< (Index0 > 1) && (Index1 <= 1), bool >::type = true>
 FORCEINLINE VectorRegister2Double SelectVectorSwizzle2(const VectorRegister4Double& Vec)
 {
-	VectorRegister2Double Result = _mm_shuffle_pd(Vec.ZW, Vec.XY, SHUFFLEMASK2(Index0 - 2, Index1));
-	return Result;
+	return _mm_shuffle_pd(Vec.ZW, Vec.XY, SHUFFLEMASK2(Index0 - 2, Index1));
 }
 
 // [2,3]:[2,3]
 template<int Index0, int Index1, typename std::enable_if< (Index0 > 1) && (Index1 > 1), bool >::type = true>
 FORCEINLINE VectorRegister2Double SelectVectorSwizzle2(const VectorRegister4Double& Vec)
 {
-	VectorRegister2Double Result = _mm_shuffle_pd(Vec.ZW, Vec.ZW, SHUFFLEMASK2(Index0 - 2, Index1 - 2));
-	return Result;
+	return _mm_shuffle_pd(Vec.ZW, Vec.ZW, SHUFFLEMASK2(Index0 - 2, Index1 - 2));
 }
+// <2,3> is simply ZW.
+template<> FORCEINLINE VectorRegister2Double SelectVectorSwizzle2<2, 3>(const VectorRegister4Double& Vec) { return Vec.ZW; }
 
 // Double swizzle
 template<int Index0, int Index1, int Index2, int Index3>
 FORCEINLINE VectorRegister4Double VectorSwizzleTemplate(const VectorRegister4Double& Vec)
 {
+#if UE_PLATFORM_MATH_USE_AVX_2
+	VectorRegister4Double Result = _mm256_permute4x64_pd(Vec, SHUFFLEMASK(Index0, Index1, Index2, Index3));
+#else
 	VectorRegister4Double Result;
 	Result.XY = SelectVectorSwizzle2<Index0, Index1>(Vec);
 	Result.ZW = SelectVectorSwizzle2<Index2, Index3>(Vec);
+#endif
+	return Result;
+}
+
+// Double Swizzle specializations
+template<> FORCEINLINE VectorRegister4Double VectorSwizzleTemplate<0, 1, 2, 3>(const VectorRegister4Double& Vec) { return Vec; }
+
+#if UE_PLATFORM_MATH_USE_AVX
+constexpr int PERMUTE_MASK(int A, int B, int C, int D) { return ((A == 1 ? (1 << 0) : 0) | (B == 1 ? (1 << 1) : 0) | (C == 3 ? (1 << 2) : 0) | (D == 3 ? (1 << 3) : 0)); }
+// _mm256_permute4x64_pd has a latency of 3, so here are some specializations using instructions which have a latency of 1 but are restricted to in-lane (128 bit) permutes.
+template<> FORCEINLINE VectorRegister4Double VectorSwizzleTemplate<0, 1, 2, 2>(const VectorRegister4Double& Vec) { return _mm256_permute_pd(Vec, PERMUTE_MASK(0, 1, 2, 2)); }
+template<> FORCEINLINE VectorRegister4Double VectorSwizzleTemplate<0, 1, 3, 3>(const VectorRegister4Double& Vec) { return _mm256_permute_pd(Vec, PERMUTE_MASK(0, 1, 3, 3)); }
+template<> FORCEINLINE VectorRegister4Double VectorSwizzleTemplate<0, 1, 3, 2>(const VectorRegister4Double& Vec) { return _mm256_permute_pd(Vec, PERMUTE_MASK(0, 1, 3, 2)); }
+																																											 
+template<> FORCEINLINE VectorRegister4Double VectorSwizzleTemplate<0, 0, 2, 3>(const VectorRegister4Double& Vec) { return _mm256_permute_pd(Vec, PERMUTE_MASK(0, 0, 2, 3)); }
+template<> FORCEINLINE VectorRegister4Double VectorSwizzleTemplate<0, 0, 2, 2>(const VectorRegister4Double& Vec) { return _mm256_movedup_pd(Vec); }							 
+template<> FORCEINLINE VectorRegister4Double VectorSwizzleTemplate<0, 0, 3, 3>(const VectorRegister4Double& Vec) { return _mm256_permute_pd(Vec, PERMUTE_MASK(0, 0, 3, 3)); }
+template<> FORCEINLINE VectorRegister4Double VectorSwizzleTemplate<0, 0, 3, 2>(const VectorRegister4Double& Vec) { return _mm256_permute_pd(Vec, PERMUTE_MASK(0, 0, 3, 2)); }
+																																											 
+template<> FORCEINLINE VectorRegister4Double VectorSwizzleTemplate<1, 0, 2, 3>(const VectorRegister4Double& Vec) { return _mm256_permute_pd(Vec, PERMUTE_MASK(1, 0, 2, 3)); }
+template<> FORCEINLINE VectorRegister4Double VectorSwizzleTemplate<1, 0, 2, 2>(const VectorRegister4Double& Vec) { return _mm256_permute_pd(Vec, PERMUTE_MASK(1, 0, 2, 2)); }
+template<> FORCEINLINE VectorRegister4Double VectorSwizzleTemplate<1, 0, 3, 3>(const VectorRegister4Double& Vec) { return _mm256_permute_pd(Vec, PERMUTE_MASK(1, 0, 3, 3)); }
+template<> FORCEINLINE VectorRegister4Double VectorSwizzleTemplate<1, 0, 3, 2>(const VectorRegister4Double& Vec) { return _mm256_permute_pd(Vec, PERMUTE_MASK(1, 0, 3, 2)); }
+																																											 
+template<> FORCEINLINE VectorRegister4Double VectorSwizzleTemplate<1, 1, 2, 3>(const VectorRegister4Double& Vec) { return _mm256_permute_pd(Vec, PERMUTE_MASK(1, 1, 2, 3)); }
+template<> FORCEINLINE VectorRegister4Double VectorSwizzleTemplate<1, 1, 2, 2>(const VectorRegister4Double& Vec) { return _mm256_permute_pd(Vec, PERMUTE_MASK(1, 1, 2, 2)); }
+template<> FORCEINLINE VectorRegister4Double VectorSwizzleTemplate<1, 1, 3, 3>(const VectorRegister4Double& Vec) { return _mm256_permute_pd(Vec, PERMUTE_MASK(1, 1, 3, 3)); }
+template<> FORCEINLINE VectorRegister4Double VectorSwizzleTemplate<1, 1, 3, 2>(const VectorRegister4Double& Vec) { return _mm256_permute_pd(Vec, PERMUTE_MASK(1, 1, 3, 2)); }
+#endif
+
+
+// Double replicate (2 doubles)
+template <int Index>
+FORCEINLINE VectorRegister2Double VectorReplicateImpl2(const VectorRegister2Double& Vec)
+{
+	return _mm_shuffle_pd(Vec, Vec, SHUFFLEMASK2(Index, Index));
+}
+
+// Double replicate (4 doubles)
+template <int Index, typename std::enable_if < (Index <= 1), bool >::type = true >
+FORCEINLINE VectorRegister4Double VectorReplicateImpl(const VectorRegister4Double& Vec)
+{
+	VectorRegister4Double Result;
+	Result.XY = VectorReplicateImpl2<Index>(Vec.XY);
+	Result.ZW = Result.XY;
+	return Result;
+}
+
+template <int Index, typename std::enable_if < (Index > 1), bool >::type = true >
+FORCEINLINE VectorRegister4Double VectorReplicateImpl(const VectorRegister4Double& Vec)
+{
+	VectorRegister4Double Result;
+	Result.XY = VectorReplicateImpl2<Index - 2>(Vec.ZW);
+	Result.ZW = Result.XY;
+	return Result;
+}
+
+template<int Index>
+FORCEINLINE VectorRegister4Double VectorReplicateTemplate(const VectorRegister4Double& Vec)
+{
+#if UE_PLATFORM_MATH_USE_AVX_2
+	VectorRegister4Double Result = _mm256_permute4x64_pd(Vec, SHUFFLEMASK(Index, Index, Index, Index));
+#else
+	VectorRegister4Double Result = VectorReplicateImpl<Index>(Vec);
+#endif
 	return Result;
 }
 
@@ -644,6 +845,19 @@ FORCEINLINE VectorRegister4Double VectorShuffleTemplate(const VectorRegister4Dou
 	return Result;
 }
 
+// Double Shuffle specializations 
+#if UE_PLATFORM_MATH_USE_AVX
+template<> FORCEINLINE VectorRegister4Double VectorShuffleTemplate<0, 1, 0, 1>(const VectorRegister4Double& Vec1, const VectorRegister4Double& Vec2) { return _mm256_permute2f128_pd(Vec1, Vec2, (0x00) | (0x02 << 4)); }
+template<> FORCEINLINE VectorRegister4Double VectorShuffleTemplate<0, 1, 2, 3>(const VectorRegister4Double& Vec1, const VectorRegister4Double& Vec2) { return _mm256_blend_pd(Vec1, Vec2, 0b1100); }
+template<> FORCEINLINE VectorRegister4Double VectorShuffleTemplate<2, 3, 0, 1>(const VectorRegister4Double& Vec1, const VectorRegister4Double& Vec2) { return _mm256_permute2f128_pd(Vec1, Vec2, (0x01) | (0x02 << 4)); }
+template<> FORCEINLINE VectorRegister4Double VectorShuffleTemplate<2, 3, 2, 3>(const VectorRegister4Double& Vec1, const VectorRegister4Double& Vec2) { return _mm256_permute2f128_pd(Vec1, Vec2, (0x01) | (0x03 << 4)); }
+#else
+template<> FORCEINLINE VectorRegister4Double VectorShuffleTemplate<0, 1, 0, 1>(const VectorRegister4Double& Vec1, const VectorRegister4Double& Vec2) { return VectorRegister4Double(Vec1.XY, Vec2.XY); }
+template<> FORCEINLINE VectorRegister4Double VectorShuffleTemplate<0, 1, 2, 3>(const VectorRegister4Double& Vec1, const VectorRegister4Double& Vec2) { return VectorRegister4Double(Vec1.XY, Vec2.ZW); }
+template<> FORCEINLINE VectorRegister4Double VectorShuffleTemplate<2, 3, 0, 1>(const VectorRegister4Double& Vec1, const VectorRegister4Double& Vec2) { return VectorRegister4Double(Vec1.ZW, Vec2.XY); }
+template<> FORCEINLINE VectorRegister4Double VectorShuffleTemplate<2, 3, 2, 3>(const VectorRegister4Double& Vec1, const VectorRegister4Double& Vec2) { return VectorRegister4Double(Vec1.ZW, Vec2.ZW); }
+#endif
+
 
 // Float swizzle
 template<int Index0, int Index1, int Index2, int Index3>
@@ -653,14 +867,44 @@ FORCEINLINE VectorRegister4Float VectorSwizzleTemplate(const VectorRegister4Floa
 	return Result;
 }
 
+// Float Swizzle specializations.
+// These can result in no-ops or simpler ops than shuffle which don't compete with the shuffle unit, or which can copy directly to the destination and avoid an intermediate mov.
+// See: https://stackoverflow.com/questions/56238197/what-is-the-difference-between-mm-movehdup-ps-and-mm-shuffle-ps-in-this-case
+template<> FORCEINLINE VectorRegister4Float VectorSwizzleTemplate<0, 1, 2, 3>(const VectorRegister4Float& Vec) { return Vec; }
+template<> FORCEINLINE VectorRegister4Float VectorSwizzleTemplate<0, 1, 0, 1>(const VectorRegister4Float& Vec) { return _mm_movelh_ps(Vec, Vec); }
+template<> FORCEINLINE VectorRegister4Float VectorSwizzleTemplate<2, 3, 2, 3>(const VectorRegister4Float& Vec) { return _mm_movehl_ps(Vec, Vec); }
+template<> FORCEINLINE VectorRegister4Float VectorSwizzleTemplate<0, 0, 1, 1>(const VectorRegister4Float& Vec) { return _mm_unpacklo_ps(Vec, Vec); }
+template<> FORCEINLINE VectorRegister4Float VectorSwizzleTemplate<2, 2, 3, 3>(const VectorRegister4Float& Vec) { return _mm_unpackhi_ps(Vec, Vec); }
+
+#if UE_PLATFORM_MATH_USE_SSE4_1
+template<> FORCEINLINE VectorRegister4Float VectorSwizzleTemplate<0, 0, 2, 2>(const VectorRegister4Float& Vec) { return _mm_moveldup_ps(Vec); }
+template<> FORCEINLINE VectorRegister4Float VectorSwizzleTemplate<1, 1, 3, 3>(const VectorRegister4Float& Vec) { return _mm_movehdup_ps(Vec); }
+#endif
+
+#if UE_PLATFORM_MATH_USE_AVX_2
+template<> FORCEINLINE VectorRegister4Float VectorSwizzleTemplate<0, 0, 0, 0>(const VectorRegister4Float& Vec) { return _mm_broadcastss_ps(Vec); }
+#endif
+
+
+// Float replicate
+template<int Index>
+FORCEINLINE VectorRegister4Float VectorReplicateTemplate(const VectorRegister4Float& Vec)
+{
+	return VectorSwizzleTemplate<Index, Index, Index, Index>(Vec);
+}
+
 // Float shuffle
 template<int Index0, int Index1, int Index2, int Index3>
 FORCEINLINE VectorRegister4Float VectorShuffleTemplate(const VectorRegister4Float& Vec1, const VectorRegister4Float& Vec2)
 {
-	VectorRegister4Float Result = _mm_shuffle_ps(Vec1, Vec2, SHUFFLEMASK(Index0, Index1, Index2, Index3));
-	return Result;
+	return _mm_shuffle_ps(Vec1, Vec2, SHUFFLEMASK(Index0, Index1, Index2, Index3));
 }
 
+// Float Shuffle specializations
+template<> FORCEINLINE VectorRegister4Float VectorShuffleTemplate<0, 1, 0, 1>(const VectorRegister4Float& Vec1, const VectorRegister4Float& Vec2) { return _mm_movelh_ps(Vec1, Vec2); }
+template<> FORCEINLINE VectorRegister4Float VectorShuffleTemplate<2, 3, 2, 3>(const VectorRegister4Float& Vec1, const VectorRegister4Float& Vec2) { return _mm_movehl_ps(Vec2, Vec1); } // Note: movehl copies first from the 2nd argument
+
+
 /**
  * Replicates one element into all four elements and returns the new vector.
  *
@@ -669,7 +913,7 @@ FORCEINLINE VectorRegister4Float VectorShuffleTemplate(const VectorRegister4Floa
  * @return				VectorRegister4Float( Vec[ElementIndex], Vec[ElementIndex], Vec[ElementIndex], Vec[ElementIndex] )
  */
 
-#define VectorReplicate( Vec, ElementIndex )	VectorSwizzleTemplate<ElementIndex, ElementIndex, ElementIndex, ElementIndex>(Vec)
+#define VectorReplicate( Vec, ElementIndex )	VectorReplicateTemplate<ElementIndex>(Vec)
 
  /**
   * Swizzles the 4 components of a vector and returns the result.
@@ -708,11 +952,22 @@ FORCEINLINE VectorRegister4Float VectorAbs(const VectorRegister4Float& Vec)
 	return _mm_and_ps(Vec, GlobalVectorConstants::SignMask);
 }
 
+namespace SSEVectorConstants
+{
+#define DOUBLE_SIGN_BIT (uint64(1) << uint64(63))
+	static const VectorRegister2Double DoubleSignMask2d = MakeVectorRegister2DoubleMask((uint64)(~DOUBLE_SIGN_BIT), (uint64)(~DOUBLE_SIGN_BIT));
+#undef DOUBLE_SIGN_BIT
+}
+
 FORCEINLINE VectorRegister4Double VectorAbs(const VectorRegister4Double& Vec)
 {
 	VectorRegister4Double Result;
-	Result.XY = _mm_and_pd(Vec.XY, GlobalVectorConstants::DoubleSignMask2d);
-	Result.ZW = _mm_and_pd(Vec.ZW, GlobalVectorConstants::DoubleSignMask2d);
+#if !UE_PLATFORM_MATH_USE_AVX
+	Result.XY = _mm_and_pd(Vec.XY, SSEVectorConstants::DoubleSignMask2d);
+	Result.ZW = _mm_and_pd(Vec.ZW, SSEVectorConstants::DoubleSignMask2d);
+#else
+	Result = _mm256_and_pd(Vec, GlobalVectorConstants::DoubleSignMask);
+#endif
 	return Result;
 }
 
@@ -730,8 +985,12 @@ FORCEINLINE VectorRegister4Float VectorNegate(const VectorRegister4Float& Vec)
 FORCEINLINE VectorRegister4Double VectorNegate(const VectorRegister4Double& Vec)
 {
 	VectorRegister4Double Result;
+#if !UE_PLATFORM_MATH_USE_AVX
 	Result.XY = _mm_sub_pd(_mm_setzero_pd(), Vec.XY);
 	Result.ZW = _mm_sub_pd(_mm_setzero_pd(), Vec.ZW);
+#else
+	Result = _mm256_sub_pd(_mm256_setzero_pd(), Vec);
+#endif
 	return Result;
 }
 
@@ -751,8 +1010,12 @@ FORCEINLINE VectorRegister4Float VectorAdd( const VectorRegister4Float& Vec1, co
 FORCEINLINE VectorRegister4Double VectorAdd(const VectorRegister4Double& Vec1, const VectorRegister4Double& Vec2)
 {
 	VectorRegister4Double Result;
+#if !UE_PLATFORM_MATH_USE_AVX
 	Result.XY = _mm_add_pd(Vec1.XY, Vec2.XY);
 	Result.ZW = _mm_add_pd(Vec1.ZW, Vec2.ZW);
+#else
+	Result = _mm256_add_pd(Vec1, Vec2);
+#endif
 	return Result;
 }
 
@@ -771,8 +1034,12 @@ FORCEINLINE VectorRegister4Float VectorSubtract( const VectorRegister4Float& Vec
 FORCEINLINE VectorRegister4Double VectorSubtract(const VectorRegister4Double& Vec1, const VectorRegister4Double& Vec2)
 {
 	VectorRegister4Double Result;
+#if !UE_PLATFORM_MATH_USE_AVX
 	Result.XY = _mm_sub_pd(Vec1.XY, Vec2.XY);
 	Result.ZW = _mm_sub_pd(Vec1.ZW, Vec2.ZW);
+#else
+	Result = _mm256_sub_pd(Vec1, Vec2);
+#endif
 	return Result;
 }
 
@@ -791,8 +1058,12 @@ FORCEINLINE VectorRegister4Float VectorMultiply( const VectorRegister4Float& Vec
 FORCEINLINE VectorRegister4Double VectorMultiply(const VectorRegister4Double& Vec1, const VectorRegister4Double& Vec2)
 {
 	VectorRegister4Double Result;
+#if !UE_PLATFORM_MATH_USE_AVX
 	Result.XY = _mm_mul_pd(Vec1.XY, Vec2.XY);
 	Result.ZW = _mm_mul_pd(Vec1.ZW, Vec2.ZW);
+#else
+	Result = _mm256_mul_pd(Vec1, Vec2);
+#endif
 	return Result;
 }
 
@@ -808,7 +1079,7 @@ FORCEINLINE VectorRegister4Double VectorMultiply(const VectorRegister4Double& Ve
  */
 FORCEINLINE VectorRegister4Float VectorMultiplyAdd(const VectorRegister4Float& Vec1, const VectorRegister4Float& Vec2, const VectorRegister4Float& Vec3)
 {
-#if PLATFORM_ALWAYS_HAS_FMA3
+#if UE_PLATFORM_MATH_USE_FMA3
 	return _mm_fmadd_ps(Vec1, Vec2, Vec3);
 #else
 	return VectorAdd(VectorMultiply(Vec1, Vec2), Vec3);
@@ -817,10 +1088,14 @@ FORCEINLINE VectorRegister4Float VectorMultiplyAdd(const VectorRegister4Float& V
 
 FORCEINLINE VectorRegister4Double VectorMultiplyAdd(const VectorRegister4Double& Vec1, const VectorRegister4Double& Vec2, const VectorRegister4Double& Vec3)
 {
-#if PLATFORM_ALWAYS_HAS_FMA3
+#if UE_PLATFORM_MATH_USE_FMA3
 	VectorRegister4Double Result;
-	Result.XY = _mm_fmadd_pd(Vec1.XY, Vec2.XY, Vec3.XY);
-	Result.ZW = _mm_fmadd_pd(Vec1.ZW, Vec2.ZW, Vec3.ZW);
+	#if !UE_PLATFORM_MATH_USE_AVX
+		Result.XY = _mm_fmadd_pd(Vec1.XY, Vec2.XY, Vec3.XY);
+		Result.ZW = _mm_fmadd_pd(Vec1.ZW, Vec2.ZW, Vec3.ZW);
+	#else
+		Result = _mm256_fmadd_pd(Vec1, Vec2, Vec3);
+	#endif
 	return Result;
 #else
 	return VectorAdd(VectorMultiply(Vec1, Vec2), Vec3);
@@ -837,7 +1112,7 @@ FORCEINLINE VectorRegister4Double VectorMultiplyAdd(const VectorRegister4Double&
  */
 FORCEINLINE VectorRegister4Float VectorNegateMultiplyAdd(const VectorRegister4Float& Vec1, const VectorRegister4Float& Vec2, const VectorRegister4Float& Vec3)
 {
-#if PLATFORM_ALWAYS_HAS_FMA3
+#if UE_PLATFORM_MATH_USE_FMA3
 	return _mm_fnmadd_ps(Vec1, Vec2, Vec3);
 #else
 	return VectorSubtract(Vec3, VectorMultiply(Vec1, Vec2));
@@ -846,10 +1121,14 @@ FORCEINLINE VectorRegister4Float VectorNegateMultiplyAdd(const VectorRegister4Fl
 
 FORCEINLINE VectorRegister4Double VectorNegateMultiplyAdd(const VectorRegister4Double& Vec1, const VectorRegister4Double& Vec2, const VectorRegister4Double& Vec3)
 {
-#if PLATFORM_ALWAYS_HAS_FMA3
+#if UE_PLATFORM_MATH_USE_FMA3
 	VectorRegister4Double Result;
-	Result.XY = _mm_fnmadd_pd(Vec1.XY, Vec2.XY, Vec3.XY);
-	Result.ZW = _mm_fnmadd_pd(Vec1.ZW, Vec2.ZW, Vec3.ZW);
+	#if !UE_PLATFORM_MATH_USE_AVX
+		Result.XY = _mm_fnmadd_pd(Vec1.XY, Vec2.XY, Vec3.XY);
+		Result.ZW = _mm_fnmadd_pd(Vec1.ZW, Vec2.ZW, Vec3.ZW);
+	#else
+		Result = _mm256_fnmadd_pd(Vec1, Vec2, Vec3);
+	#endif
 	return Result;
 #else
 	return VectorSubtract(Vec3, VectorMultiply(Vec1, Vec2));
@@ -872,8 +1151,12 @@ FORCEINLINE VectorRegister4Float VectorDivide(const VectorRegister4Float& Vec1,
 FORCEINLINE VectorRegister4Double VectorDivide(const VectorRegister4Double& Vec1, const VectorRegister4Double& Vec2)
 {
 	VectorRegister4Double Result;
+#if !UE_PLATFORM_MATH_USE_AVX
 	Result.XY = _mm_div_pd(Vec1.XY, Vec2.XY);
 	Result.ZW = _mm_div_pd(Vec1.ZW, Vec2.ZW);
+#else
+	Result = _mm256_div_pd(Vec1, Vec2);
+#endif
 	return Result;
 }
 
@@ -892,6 +1175,7 @@ FORCEINLINE VectorRegister4Float VectorDot3( const VectorRegister4Float& Vec1, c
 
 FORCEINLINE VectorRegister4Double VectorDot3(const VectorRegister4Double& Vec1, const VectorRegister4Double& Vec2)
 {
+	// TODO: AVX version?
 	VectorRegister4Double Temp = VectorMultiply(Vec1, Vec2);
 	return VectorAdd(VectorReplicate(Temp,0), VectorAdd(VectorReplicate(Temp,1), VectorReplicate(Temp,2)));
 }
@@ -901,20 +1185,21 @@ FORCEINLINE VectorRegister4Double VectorDot3(const VectorRegister4Double& Vec1,
  *
  * @param Vec1	1st vector
  * @param Vec2	2nd vector
- * @return		d = dot4(Vec1.xyzw, Vec2.xyzw), VectorRegister4Float( d, d, d, d )
+ * @return		d = dot4(Vec1, Vec2), VectorRegister4Float( d, d, d, d )
  */
 FORCEINLINE VectorRegister4Float VectorDot4( const VectorRegister4Float& Vec1, const VectorRegister4Float& Vec2 )
 {
 	VectorRegister4Float Temp1, Temp2;
 	Temp1 = VectorMultiply( Vec1, Vec2 );
-	Temp2 = _mm_shuffle_ps( Temp1, Temp1, SHUFFLEMASK(2,3,0,1) );	// (Z,W,X,Y).
-	Temp1 = VectorAdd( Temp1, Temp2 );								// (X*X + Z*Z, Y*Y + W*W, Z*Z + X*X, W*W + Y*Y)
-	Temp2 = _mm_shuffle_ps( Temp1, Temp1, SHUFFLEMASK(1,2,3,0) );	// Rotate left 4 bytes (Y,Z,W,X).
-	return VectorAdd( Temp1, Temp2 );								// (X*X + Z*Z + Y*Y + W*W, Y*Y + W*W + Z*Z + X*X, Z*Z + X*X + W*W + Y*Y, W*W + Y*Y + X*X + Z*Z)
+	Temp2 = VectorSwizzle( Temp1, 2, 3, 0, 1 );	// (Z,W,X,Y).
+	Temp1 = VectorAdd( Temp1, Temp2 );			// (X*X + Z*Z, Y*Y + W*W, Z*Z + X*X, W*W + Y*Y)
+	Temp2 = VectorSwizzle( Temp1, 1, 2, 3, 0 );	// Rotate left 4 bytes (Y,Z,W,X).
+	return VectorAdd( Temp1, Temp2 );			// (X*X + Z*Z + Y*Y + W*W, Y*Y + W*W + Z*Z + X*X, Z*Z + X*X + W*W + Y*Y, W*W + Y*Y + X*X + Z*Z)
 }
 
 FORCEINLINE VectorRegister4Double VectorDot4(const VectorRegister4Double& Vec1, const VectorRegister4Double& Vec2)
 {
+	// TODO: AVX version
 	VectorRegister4Double R;
 	VectorRegister2Double T;
 
@@ -951,8 +1236,12 @@ FORCEINLINE VectorRegister4Float VectorCompareEQ(const VectorRegister4Float& Vec
 FORCEINLINE VectorRegister4Double VectorCompareEQ(const VectorRegister4Double& Vec1, const VectorRegister4Double& Vec2)
 {
 	VectorRegister4Double Result;
+#if !UE_PLATFORM_MATH_USE_AVX
 	Result.XY = _mm_cmpeq_pd(Vec1.XY, Vec2.XY);
 	Result.ZW = _mm_cmpeq_pd(Vec1.ZW, Vec2.ZW);
+#else
+	Result = _mm256_cmp_pd(Vec1, Vec2, _CMP_EQ_OQ);
+#endif
 	return Result;
 }
 
@@ -971,8 +1260,12 @@ FORCEINLINE VectorRegister4Float VectorCompareNE(const VectorRegister4Float& Vec
 FORCEINLINE VectorRegister4Double VectorCompareNE(const VectorRegister4Double& Vec1, const VectorRegister4Double& Vec2)
 {
 	VectorRegister4Double Result;
+#if !UE_PLATFORM_MATH_USE_AVX
 	Result.XY = _mm_cmpneq_pd(Vec1.XY, Vec2.XY);
 	Result.ZW = _mm_cmpneq_pd(Vec1.ZW, Vec2.ZW);
+#else
+	Result = _mm256_cmp_pd(Vec1, Vec2, _CMP_NEQ_OQ);
+#endif
 	return Result;
 }
 
@@ -991,8 +1284,12 @@ FORCEINLINE VectorRegister4Float VectorCompareGT(const VectorRegister4Float& Vec
 FORCEINLINE VectorRegister4Double VectorCompareGT(const VectorRegister4Double& Vec1, const VectorRegister4Double& Vec2)
 {
 	VectorRegister4Double Result;
+#if !UE_PLATFORM_MATH_USE_AVX
 	Result.XY = _mm_cmpgt_pd(Vec1.XY, Vec2.XY);
 	Result.ZW = _mm_cmpgt_pd(Vec1.ZW, Vec2.ZW);
+#else
+	Result = _mm256_cmp_pd(Vec1, Vec2, _CMP_GT_OQ);
+#endif
 	return Result;
 }
 
@@ -1011,8 +1308,12 @@ FORCEINLINE VectorRegister4Float VectorCompareGE(const VectorRegister4Float& Vec
 FORCEINLINE VectorRegister4Double VectorCompareGE(const VectorRegister4Double& Vec1, const VectorRegister4Double& Vec2)
 {
 	VectorRegister4Double Result;
+#if !UE_PLATFORM_MATH_USE_AVX
 	Result.XY = _mm_cmpge_pd(Vec1.XY, Vec2.XY);
 	Result.ZW = _mm_cmpge_pd(Vec1.ZW, Vec2.ZW);
+#else
+	Result = _mm256_cmp_pd(Vec1, Vec2, _CMP_GE_OQ);
+#endif
 	return Result;
 }
 
@@ -1031,8 +1332,12 @@ FORCEINLINE VectorRegister4Float VectorCompareLT(const VectorRegister4Float& Vec
 FORCEINLINE VectorRegister4Double VectorCompareLT(const VectorRegister4Double& Vec1, const VectorRegister4Double& Vec2)
 {
 	VectorRegister4Double Result;
+#if !UE_PLATFORM_MATH_USE_AVX
 	Result.XY = _mm_cmplt_pd(Vec1.XY, Vec2.XY);
 	Result.ZW = _mm_cmplt_pd(Vec1.ZW, Vec2.ZW);
+#else
+	Result = _mm256_cmp_pd(Vec1, Vec2, _CMP_LT_OQ);
+#endif
 	return Result;
 }
 
@@ -1051,8 +1356,12 @@ FORCEINLINE VectorRegister4Float VectorCompareLE(const VectorRegister4Float& Vec
 FORCEINLINE VectorRegister4Double VectorCompareLE(const VectorRegister4Double& Vec1, const VectorRegister4Double& Vec2)
 {
 	VectorRegister4Double Result;
+#if !UE_PLATFORM_MATH_USE_AVX
 	Result.XY = _mm_cmple_pd(Vec1.XY, Vec2.XY);
 	Result.ZW = _mm_cmple_pd(Vec1.ZW, Vec2.ZW);
+#else
+	Result = _mm256_cmp_pd(Vec1, Vec2, _CMP_LE_OQ);
+#endif
 	return Result;
 }
 
@@ -1079,8 +1388,12 @@ FORCEINLINE VectorRegister2Double VectorSelect(const VectorRegister2Double& Mask
 FORCEINLINE VectorRegister4Double VectorSelect(const VectorRegister4Double& Mask, const VectorRegister4Double& Vec1, const VectorRegister4Double& Vec2)
 {
 	VectorRegister4Double Result;
+#if !UE_PLATFORM_MATH_USE_AVX
 	Result.XY = VectorSelect(Mask.XY, Vec1.XY, Vec2.XY);
 	Result.ZW = VectorSelect(Mask.ZW, Vec1.ZW, Vec2.ZW);
+#else
+	Result = _mm256_xor_pd(Vec2, _mm256_and_pd(Mask, _mm256_xor_pd(Vec1, Vec2)));
+#endif
 	return Result;
 }
 
@@ -1099,8 +1412,12 @@ FORCEINLINE VectorRegister4Float VectorBitwiseOr(const VectorRegister4Float& Vec
 FORCEINLINE VectorRegister4Double VectorBitwiseOr(const VectorRegister4Double& Vec1, const VectorRegister4Double& Vec2)
 {
 	VectorRegister4Double Result;
+#if !UE_PLATFORM_MATH_USE_AVX
 	Result.XY = _mm_or_pd(Vec1.XY, Vec2.XY);
 	Result.ZW = _mm_or_pd(Vec1.ZW, Vec2.ZW);
+#else
+	Result = _mm256_or_pd(Vec1, Vec2);
+#endif
 	return Result;
 }
 
@@ -1119,8 +1436,12 @@ FORCEINLINE VectorRegister4Float VectorBitwiseAnd(const VectorRegister4Float& Ve
 FORCEINLINE VectorRegister4Double VectorBitwiseAnd(const VectorRegister4Double& Vec1, const VectorRegister4Double& Vec2)
 {
 	VectorRegister4Double Result;
+#if !UE_PLATFORM_MATH_USE_AVX
 	Result.XY = _mm_and_pd(Vec1.XY, Vec2.XY);
 	Result.ZW = _mm_and_pd(Vec1.ZW, Vec2.ZW);
+#else
+	Result = _mm256_and_pd(Vec1, Vec2);
+#endif
 	return Result;
 }
 
@@ -1139,8 +1460,12 @@ FORCEINLINE VectorRegister4Float VectorBitwiseXor(const VectorRegister4Float& Ve
 FORCEINLINE VectorRegister4Double VectorBitwiseXor(const VectorRegister4Double& Vec1, const VectorRegister4Double& Vec2)
 {
 	VectorRegister4Double Result;
+#if !UE_PLATFORM_MATH_USE_AVX
 	Result.XY = _mm_xor_pd(Vec1.XY, Vec2.XY);
 	Result.ZW = _mm_xor_pd(Vec1.ZW, Vec2.ZW);
+#else
+	Result = _mm256_xor_pd(Vec1, Vec2);
+#endif
 	return Result;
 }
 
@@ -1166,6 +1491,8 @@ FORCEINLINE VectorRegister4Float VectorCross( const VectorRegister4Float& Vec1,
 
 FORCEINLINE VectorRegister4Double VectorCross(const VectorRegister4Double& Vec1, const VectorRegister4Double& Vec2)
 {
+	// TODO: AVX?
+
 	// YZX
 	VectorRegister4Double A = VectorSwizzle(Vec2, 1, 2, 0, 3);
 	VectorRegister4Double B = VectorSwizzle(Vec1, 1, 2, 0, 3);
@@ -1226,8 +1553,12 @@ FORCEINLINE VectorRegister4Double VectorReciprocalSqrt(const VectorRegister4Doub
 {
 	// Not an estimate.
 	VectorRegister4Double Result;
+#if !UE_PLATFORM_MATH_USE_AVX
 	Result.XY = _mm_sqrt_pd(Vec.XY);
 	Result.ZW = _mm_sqrt_pd(Vec.ZW);
+#else
+	Result = _mm256_sqrt_pd(Vec);
+#endif
 	return VectorDivide(GlobalVectorConstants::DoubleOne, Result);
 }
 
@@ -1309,8 +1640,12 @@ FORCEINLINE VectorRegister4Float VectorReciprocalSqrtAccurate(const VectorRegist
 FORCEINLINE VectorRegister4Double VectorReciprocalSqrtAccurate(const VectorRegister4Double& Vec)
 {
 	VectorRegister4Double Result;
+#if !UE_PLATFORM_MATH_USE_AVX
 	Result.XY = _mm_sqrt_pd(Vec.XY);
 	Result.ZW = _mm_sqrt_pd(Vec.ZW);
+#else
+	Result = _mm256_sqrt_pd(Vec);
+#endif
 	return VectorDivide(GlobalVectorConstants::DoubleOne, Result);
 }
 
@@ -1370,8 +1705,12 @@ FORCEINLINE VectorRegister4Float VectorSet_W0(const VectorRegister4Float& Vec)
 FORCEINLINE VectorRegister4Double VectorSet_W0(const VectorRegister4Double& Vec)
 {
 	VectorRegister4Double Result;
+#if !UE_PLATFORM_MATH_USE_AVX
 	Result.XY = Vec.XY;
 	Result.ZW = _mm_castsi128_pd(_mm_move_epi64(_mm_castpd_si128(Vec.ZW)));
+#else
+	Result = _mm256_blend_pd(Vec, VectorZeroDouble(), 0b1000);
+#endif
 	return Result;
 }
 
@@ -1387,14 +1726,18 @@ FORCEINLINE VectorRegister4Float VectorSet_W1( const VectorRegister4Float& Vec)
 	VectorRegister4Float Temp = _mm_movehl_ps( VectorOneFloat(), Vec);
 
 	// Return (Vector[0], Vector[1], Vector[2], 1.0f)
-	return _mm_shuffle_ps(Vec, Temp, SHUFFLEMASK(0,1,0,3) );
+	return VectorShuffle(Vec, Temp, 0, 1, 0, 3 );
 }
 
 FORCEINLINE VectorRegister4Double VectorSet_W1(const VectorRegister4Double& Vec)
 {
 	VectorRegister4Double Result;
+#if !UE_PLATFORM_MATH_USE_AVX
 	Result.XY = Vec.XY;
 	Result.ZW = _mm_move_sd(GlobalVectorConstants::DoubleOne2d, Vec.ZW);
+#else
+	Result = _mm256_blend_pd(Vec, VectorOneDouble(), 0b1000);
+#endif
 	return Result;
 }
 
@@ -1570,7 +1913,7 @@ FORCEINLINE void VectorMatrixInverse(FMatrix44d* DstMatrix, const FMatrix44d* Sr
 		M[2][0] * (M[0][1] * M[1][2] - M[0][2] * M[1][1])
 		);
 
-	memcpy( DstMatrix, &Result, 16*sizeof(double) );
+	memcpy( DstMatrix, &Result, sizeof(Result) );
 }
 
 FORCEINLINE void VectorMatrixInverse(FMatrix44f* DstMatrix, const FMatrix44f* SrcMatrix)
@@ -1654,7 +1997,7 @@ FORCEINLINE void VectorMatrixInverse(FMatrix44f* DstMatrix, const FMatrix44f* Sr
 		M[2][0] * (M[0][1] * M[1][2] - M[0][2] * M[1][1])
 		);
 
-	memcpy(DstMatrix, &Result, 16 * sizeof(float));
+	memcpy(DstMatrix, &Result, sizeof(Result));
 }
 
 
@@ -1668,7 +2011,7 @@ FORCEINLINE void VectorMatrixInverse(FMatrix44f* DstMatrix, const FMatrix44f* Sr
  */
 FORCEINLINE VectorRegister4Float VectorTransformVector(const VectorRegister4Float& VecP, const FMatrix44f* MatrixM)
 {
-	const VectorRegister4Float *M	= (const VectorRegister4Float *) MatrixM;
+	const VectorRegister4Float *M = (const VectorRegister4Float *) MatrixM;
 	VectorRegister4Float VTempX, VTempY, VTempZ, VTempW;
 
 	// Splat x,y,z and w
@@ -1759,8 +2102,12 @@ FORCEINLINE VectorRegister4Float VectorMin(const VectorRegister4Float& Vec1, con
 FORCEINLINE VectorRegister4Double VectorMin(const VectorRegister4Double& Vec1, const VectorRegister4Double& Vec2)
 {
 	VectorRegister4Double Result;
+#if !UE_PLATFORM_MATH_USE_AVX
 	Result.XY = _mm_min_pd(Vec1.XY, Vec2.XY);
 	Result.ZW = _mm_min_pd(Vec1.ZW, Vec2.ZW);
+#else
+	Result = _mm256_min_pd(Vec1, Vec2);
+#endif
 	return Result;
 }
 
@@ -1779,8 +2126,12 @@ FORCEINLINE VectorRegister4Float VectorMax(const VectorRegister4Float& Vec1, con
 FORCEINLINE VectorRegister4Double VectorMax(const VectorRegister4Double& Vec1, const VectorRegister4Double& Vec2)
 {
 	VectorRegister4Double Result;
+#if !UE_PLATFORM_MATH_USE_AVX
 	Result.XY = _mm_max_pd(Vec1.XY, Vec2.XY);
 	Result.ZW = _mm_max_pd(Vec1.ZW, Vec2.ZW);
+#else
+	Result = _mm256_max_pd(Vec1, Vec2);
+#endif
 	return Result;
 }
 
@@ -1798,10 +2149,7 @@ FORCEINLINE VectorRegister4Float VectorCombineHigh(const VectorRegister4Float& V
 
 FORCEINLINE VectorRegister4Double VectorCombineHigh(const VectorRegister4Double& Vec1, const VectorRegister4Double& Vec2)
 {
-	VectorRegister4Double Result;
-	Result.XY = Vec1.ZW;
-	Result.ZW = Vec2.ZW;
-	return Result;
+	return VectorShuffle(Vec1, Vec2, 2, 3, 2, 3);
 }
 
 /**
@@ -1818,10 +2166,7 @@ FORCEINLINE VectorRegister4Float VectorCombineLow(const VectorRegister4Float& Ve
 
 FORCEINLINE VectorRegister4Double VectorCombineLow(const VectorRegister4Double& Vec1, const VectorRegister4Double& Vec2)
 {
-	VectorRegister4Double Result;
-	Result.XY = Vec1.XY;
-	Result.ZW = Vec2.XY;
-	return Result;
+	return VectorShuffle(Vec1, Vec2, 0, 1, 0, 1);
 }
 
 /**
@@ -1859,9 +2204,13 @@ FORCEINLINE int VectorMaskBits(const VectorRegister4Float& VecMask)
 
 FORCEINLINE int VectorMaskBits(const VectorRegister4Double& VecMask)
 {
+#if !UE_PLATFORM_MATH_USE_AVX
 	const int MaskXY = _mm_movemask_pd(VecMask.XY);
 	const int MaskZW = _mm_movemask_pd(VecMask.ZW);
 	return (MaskZW << 2) | (MaskXY);
+#else
+	return _mm256_movemask_pd(VecMask);
+#endif
 }
 
 /**
@@ -1881,7 +2230,7 @@ FORCEINLINE VectorRegister4Double VectorMergeVecXYZ_VecW(const VectorRegister4Do
 	//return VectorSelect(GlobalVectorConstants::DoubleXYZMask, VecXYZ, VecW);
 	VectorRegister4Double Result;
 	Result.XY = VecXYZ.XY;
-	Result.ZW = _mm_move_sd(VecW.ZW, VecXYZ.ZW); //VectorSelect(GlobalVectorConstants::DoubleZMask, VecXYZ.ZW, VecW.ZW);
+	Result.ZW = _mm_move_sd(VecW.ZW, VecXYZ.ZW);
 	return Result;
 }
 
@@ -1923,7 +2272,7 @@ FORCEINLINE VectorRegister4Float VectorLoadSignedByte4(const void* Ptr)
 FORCEINLINE VectorRegister4Float VectorLoadByte4Reverse( void* Ptr )
 {
 	VectorRegister4Float Temp = VectorLoadByte4(Ptr);
-	return _mm_shuffle_ps( Temp, Temp, SHUFFLEMASK(3,2,1,0) );
+	return VectorSwizzle( Temp, 3, 2, 1, 0 );
 }
 
 /**
@@ -2183,7 +2532,7 @@ FORCEINLINE bool VectorContainsNaNOrInfinite(const VectorRegister4Double& Vec)
 
 FORCEINLINE VectorRegister4Float VectorTruncate(const VectorRegister4Float& Vec)
 {
-#if UNREALPLATFORMMATH_SSE_USE_SSE4_1
+#if UE_PLATFORM_MATH_USE_SSE4_1
 	return _mm_round_ps(Vec, _MM_FROUND_TRUNC);
 #else
 	return _mm_cvtepi32_ps(_mm_cvttps_epi32(Vec));
@@ -2192,7 +2541,7 @@ FORCEINLINE VectorRegister4Float VectorTruncate(const VectorRegister4Float& Vec)
 
 FORCEINLINE VectorRegister2Double TruncateVectorRegister2d(const VectorRegister2Double& V)
 {
-#if UNREALPLATFORMMATH_SSE_USE_SSE4_1
+#if UE_PLATFORM_MATH_USE_SSE4_1
 	return _mm_round_pd(V, _MM_FROUND_TRUNC);
 #else
 	// TODO: LWC: Optimize
@@ -2202,21 +2551,25 @@ FORCEINLINE VectorRegister2Double TruncateVectorRegister2d(const VectorRegister2
 	VectorRegister2Double A = _mm_cvtsi64_sd(V, X); // Converts to lowest element, copies upper.
 	VectorRegister2Double B = _mm_cvtsi64_sd(V, Y); // Converts to lowest element, copies upper.
 	return _mm_shuffle_pd(A, B, SHUFFLEMASK2(0, 0));
-#endif // UNREALPLATFORMMATH_SSE_USE_SSE4_1	
+#endif // UE_PLATFORM_MATH_USE_SSE4_1	
 }
 
 FORCEINLINE VectorRegister4Double VectorTruncate(const VectorRegister4Double& V)
 {
 	VectorRegister4Double Result;
+#if !UE_PLATFORM_MATH_USE_AVX
 	Result.XY = TruncateVectorRegister2d(V.XY);
 	Result.ZW = TruncateVectorRegister2d(V.ZW);
+#else
+	Result = _mm256_round_pd(V, _MM_FROUND_TO_ZERO | _MM_FROUND_NO_EXC);
+#endif
 	return Result;
 }
 
 
 FORCEINLINE VectorRegister4Float VectorCeil(const VectorRegister4Float& V)
 {
-#if UNREALPLATFORMMATH_SSE_USE_SSE4_1
+#if UE_PLATFORM_MATH_USE_SSE4_1
 	return _mm_ceil_ps(V);
 #else
 	const VectorRegister4Float Trunc = VectorTruncate(V);
@@ -2229,7 +2582,11 @@ FORCEINLINE VectorRegister4Float VectorCeil(const VectorRegister4Float& V)
 
 FORCEINLINE VectorRegister4Double VectorCeil(const VectorRegister4Double& V)
 {
-#if UNREALPLATFORMMATH_SSE_USE_SSE4_1
+#if UE_PLATFORM_MATH_USE_AVX
+	VectorRegister4Double Result;
+	Result = _mm256_round_pd(V, _MM_FROUND_TO_POS_INF | _MM_FROUND_NO_EXC);
+	return Result;
+#elif UE_PLATFORM_MATH_USE_SSE4_1
 	VectorRegister4Double Result;
 	Result.XY = _mm_ceil_pd(V.XY);
 	Result.ZW = _mm_ceil_pd(V.ZW);
@@ -2245,7 +2602,7 @@ FORCEINLINE VectorRegister4Double VectorCeil(const VectorRegister4Double& V)
 
 FORCEINLINE VectorRegister4Float VectorFloor(const VectorRegister4Float& V)
 {
-#if UNREALPLATFORMMATH_SSE_USE_SSE4_1
+#if UE_PLATFORM_MATH_USE_SSE4_1
 	return _mm_floor_ps(V);
 #else
 	const VectorRegister4Float Trunc = VectorTruncate(V);
@@ -2258,7 +2615,11 @@ FORCEINLINE VectorRegister4Float VectorFloor(const VectorRegister4Float& V)
 
 FORCEINLINE VectorRegister4Double VectorFloor(const VectorRegister4Double& V)
 {
-#if UNREALPLATFORMMATH_SSE_USE_SSE4_1
+#if UE_PLATFORM_MATH_USE_AVX
+	VectorRegister4Double Result;
+	Result = _mm256_round_pd(V, _MM_FROUND_TO_NEG_INF | _MM_FROUND_NO_EXC);
+	return Result;
+#elif UE_PLATFORM_MATH_USE_SSE4_1
 	VectorRegister4Double Result;
 	Result.XY = _mm_floor_pd(V.XY);
 	Result.ZW = _mm_floor_pd(V.ZW);
@@ -2757,5 +3118,3 @@ FORCEINLINE VectorRegister4Int VectorFloatToInt(const VectorRegister4Double& A)
 * @return		VectorRegister4Int(*Ptr, *Ptr, *Ptr, *Ptr)
 */
 #define VectorIntLoad1( Ptr )	_mm_set1_epi32(*(Ptr))
-#endif
-
diff --git a/Engine/Source/Runtime/Core/Public/Math/UnrealMathVectorCommon.h b/Engine/Source/Runtime/Core/Public/Math/UnrealMathVectorCommon.h
--- a/Engine/Source/Runtime/Core/Public/Math/UnrealMathVectorCommon.h
+++ b/Engine/Source/Runtime/Core/Public/Math/UnrealMathVectorCommon.h
@@ -5,9 +5,9 @@
 struct FGenericPlatformMath;
 struct FMath;
 
-FORCEINLINE bool VectorIsAligned(const void *Ptr)
+FORCEINLINE bool VectorIsAligned(const void *Ptr, uint32 Alignment = SIMD_ALIGNMENT)
 {
-	return !(PTRINT(Ptr) & (SIMD_ALIGNMENT - 1));
+	return !(PTRINT(Ptr) & (Alignment - 1));
 }
 
 // Overload to resolve compiler ambiguity for things like MakeVectorRegister(V.X, V.Y, V.Z, 0.f) when V is a double type.
@@ -226,6 +226,10 @@ FORCEINLINE VectorRegister4Double VectorLoadFloat1(const VectorRegister4Double*
 	return VectorReplicate(*Ptr, 0);
 }
 
+FORCEINLINE VectorRegister4Double VectorLoadFloat1(const double* Ptr)
+{
+	return VectorLoadDouble1(Ptr);
+}
 
 ///////////////////////////////////////////////////////////////////////////////////////////////////////////////////
 /// VectorStore
diff --git a/Engine/Source/Runtime/Core/Public/Math/UnrealMathVectorConstants.h b/Engine/Source/Runtime/Core/Public/Math/UnrealMathVectorConstants.h
--- a/Engine/Source/Runtime/Core/Public/Math/UnrealMathVectorConstants.h
+++ b/Engine/Source/Runtime/Core/Public/Math/UnrealMathVectorConstants.h
@@ -10,6 +10,7 @@ namespace GlobalVectorConstants
 	static const VectorRegister4Float FloatMinusOne = MakeVectorRegisterFloat(-1.0f, -1.0f, -1.0f, -1.0f);
 	static const VectorRegister4Float Float0001 = MakeVectorRegisterFloat( 0.0f, 0.0f, 0.0f, 1.0f );
 	static const VectorRegister4Float Float1000 = MakeVectorRegisterFloat( 1.0f, 0.0f, 0.0f, 0.0f );
+	static const VectorRegister4Float Float1110 = MakeVectorRegisterFloat(1.0f, 1.0f, 1.0f, 0.0f);
 	static const VectorRegister4Float SmallLengthThreshold = MakeVectorRegisterFloat(1.e-8f, 1.e-8f, 1.e-8f, 1.e-8f);
 	static const VectorRegister4Float FloatOneHundredth = MakeVectorRegisterFloat(0.01f, 0.01f, 0.01f, 0.01f);
 	static const VectorRegister4Float Float111_Minus1 = MakeVectorRegisterFloat( 1.f, 1.f, 1.f, -1.f );
@@ -27,6 +28,7 @@ namespace GlobalVectorConstants
 	static const VectorRegister4Double DoubleMinusOne = MakeVectorRegisterDouble(-1.0, -1.0, -1.0, -1.0);
 	static const VectorRegister4Double Double0001 = MakeVectorRegisterDouble(0.0f, 0.0, 0.0, 1.0);
 	static const VectorRegister4Double Double1000 = MakeVectorRegisterDouble(1.0, 0.0, 0.0, 0.0);
+	static const VectorRegister4Double Double1110 = MakeVectorRegisterDouble(1.0, 1.0, 1.0, 0.0);
 	static const VectorRegister4Double DoubleSmallLengthThreshold = MakeVectorRegisterDouble(1.e-8, 1.e-8, 1.e-8, 1.e-8);
 	static const VectorRegister4Double DoubleOneHundredth = MakeVectorRegisterDouble(0.01, 0.01, 0.01, 0.01);
 	static const VectorRegister4Double Double111_Minus1 = MakeVectorRegisterDouble(1., 1., 1., -1.);
@@ -38,6 +40,9 @@ namespace GlobalVectorConstants
 	static const VectorRegister4Double DoubleThreshQuatNormalized = MakeVectorRegisterDouble(DOUBLE_THRESH_QUAT_NORMALIZED, DOUBLE_THRESH_QUAT_NORMALIZED, DOUBLE_THRESH_QUAT_NORMALIZED, DOUBLE_THRESH_QUAT_NORMALIZED);
 	static const VectorRegister4Double DoubleBigNumber = MakeVectorRegisterDouble(DOUBLE_BIG_NUMBER, DOUBLE_BIG_NUMBER, DOUBLE_BIG_NUMBER, DOUBLE_BIG_NUMBER);
 
+	static const VectorRegister Vector0001 = Float0001;
+	static const VectorRegister Vector1110 = Float1110;
+
 	static const VectorRegister4Int IntOne = MakeVectorRegisterInt(1, 1, 1, 1);
 	static const VectorRegister4Int IntZero = MakeVectorRegisterInt(0, 0, 0, 0);
 	static const VectorRegister4Int IntMinusOne = MakeVectorRegisterInt(-1, -1, -1, -1);
@@ -61,34 +66,28 @@ namespace GlobalVectorConstants
 	static const VectorRegister4Double DOUBLE_RAD_TO_DEG = MakeVectorRegister((180.)/DOUBLE_PI, (180.)/DOUBLE_PI, (180.)/DOUBLE_PI, (180.)/DOUBLE_PI);
 
 	/** Bitmask to AND out the XYZ components in a vector */
-	static const VectorRegister4Float XYZMask = MakeVectorRegister((uint32)0xffffffff, (uint32)0xffffffff, (uint32)0xffffffff, (uint32)0x00000000);
-	static const VectorRegister4Double DoubleXYZMask = MakeVectorRegisterDouble((uint64)0xFFFFFFFFFFFFFFFF, (uint64)0xFFFFFFFFFFFFFFFF, (uint64)0xFFFFFFFFFFFFFFFF, (uint64)0);
+	static const VectorRegister4Float XYZMask = MakeVectorRegisterFloatMask((uint32)0xffffffff, (uint32)0xffffffff, (uint32)0xffffffff, (uint32)0x00000000);
+	static const VectorRegister4Double DoubleXYZMask = MakeVectorRegisterDoubleMask((uint64)0xFFFFFFFFFFFFFFFF, (uint64)0xFFFFFFFFFFFFFFFF, (uint64)0xFFFFFFFFFFFFFFFF, (uint64)0);
 	
-	/** Bitmask to AND out the 2nd component in a double pair vector */
-	static const VectorRegister2Double DoubleZMask = MakeVectorRegister2Double((uint64)0xFFFFFFFFFFFFFFFF, (uint64)0);
-
-
 	/** Bitmask to AND out the sign bit of each components in a vector */
 #define SIGN_BIT ((1 << 31))
-	static const VectorRegister4Float SignBit = MakeVectorRegister((uint32)SIGN_BIT, (uint32)SIGN_BIT, (uint32)SIGN_BIT, (uint32)SIGN_BIT);
-	static const VectorRegister4Float SignMask = MakeVectorRegister((uint32)(~SIGN_BIT), (uint32)(~SIGN_BIT), (uint32)(~SIGN_BIT), (uint32)(~SIGN_BIT));
+	static const VectorRegister4Float SignBit = MakeVectorRegisterFloatMask((uint32)SIGN_BIT, (uint32)SIGN_BIT, (uint32)SIGN_BIT, (uint32)SIGN_BIT);
+	static const VectorRegister4Float SignMask = MakeVectorRegisterFloatMask((uint32)(~SIGN_BIT), (uint32)(~SIGN_BIT), (uint32)(~SIGN_BIT), (uint32)(~SIGN_BIT));
 	static const VectorRegister4Int IntSignBit = MakeVectorRegisterInt(SIGN_BIT, SIGN_BIT, SIGN_BIT, SIGN_BIT);
 	static const VectorRegister4Int IntSignMask = MakeVectorRegisterInt((~SIGN_BIT), (~SIGN_BIT), (~SIGN_BIT), (~SIGN_BIT));
 #undef SIGN_BIT
-	static const VectorRegister4Float AllMask = MakeVectorRegister(0xFFFFFFFF, 0xFFFFFFFF, 0xFFFFFFFF, 0xFFFFFFFF);
-	static const VectorRegister4Double DoubleAllMask = MakeVectorRegisterDouble((uint64)(0xFFFFFFFFFFFFFFFF), (uint64)(0xFFFFFFFFFFFFFFFF), (uint64)(0xFFFFFFFFFFFFFFFF), (uint64)(0xFFFFFFFFFFFFFFFF));
+	static const VectorRegister4Float AllMask = MakeVectorRegisterFloatMask(0xFFFFFFFF, 0xFFFFFFFF, 0xFFFFFFFF, 0xFFFFFFFF);
+	static const VectorRegister4Double DoubleAllMask = MakeVectorRegisterDoubleMask((uint64)(0xFFFFFFFFFFFFFFFF), (uint64)(0xFFFFFFFFFFFFFFFF), (uint64)(0xFFFFFFFFFFFFFFFF), (uint64)(0xFFFFFFFFFFFFFFFF));
 	static const VectorRegister4Int IntAllMask = MakeVectorRegisterInt(0xFFFFFFFF, 0xFFFFFFFF, 0xFFFFFFFF, 0xFFFFFFFF);
 
 #define DOUBLE_SIGN_BIT (uint64(1) << uint64(63))
-	static const VectorRegister2Double DoubleSignBit2d = MakeVectorRegister2Double((uint64)DOUBLE_SIGN_BIT, (uint64)DOUBLE_SIGN_BIT);
-	static const VectorRegister2Double DoubleSignMask2d = MakeVectorRegister2Double((uint64)(~DOUBLE_SIGN_BIT), (uint64)(~DOUBLE_SIGN_BIT));
-	static const VectorRegister4Double DoubleSignBit = MakeVectorRegisterDouble((uint64)DOUBLE_SIGN_BIT, (uint64)DOUBLE_SIGN_BIT, (uint64)DOUBLE_SIGN_BIT, (uint64)DOUBLE_SIGN_BIT);
-	static const VectorRegister4Double DoubleSignMask = MakeVectorRegisterDouble((uint64)(~DOUBLE_SIGN_BIT), (uint64)(~DOUBLE_SIGN_BIT), (uint64)(~DOUBLE_SIGN_BIT), (uint64)(~DOUBLE_SIGN_BIT));
+	static const VectorRegister4Double DoubleSignBit = MakeVectorRegisterDoubleMask((uint64)DOUBLE_SIGN_BIT, (uint64)DOUBLE_SIGN_BIT, (uint64)DOUBLE_SIGN_BIT, (uint64)DOUBLE_SIGN_BIT);
+	static const VectorRegister4Double DoubleSignMask = MakeVectorRegisterDoubleMask((uint64)(~DOUBLE_SIGN_BIT), (uint64)(~DOUBLE_SIGN_BIT), (uint64)(~DOUBLE_SIGN_BIT), (uint64)(~DOUBLE_SIGN_BIT));
 #undef DOUBLE_SIGN_BIT
 
 	/** Vector full of positive infinity */
-	static const VectorRegister4Float FloatInfinity = MakeVectorRegisterFloat((uint32)0x7F800000, (uint32)0x7F800000, (uint32)0x7F800000, (uint32)0x7F800000);
-	static const VectorRegister4Double DoubleInfinity = MakeVectorRegisterDouble((uint64)0x7FF0000000000000, (uint64)0x7FF0000000000000, (uint64)0x7FF0000000000000, (uint64)0x7FF0000000000000);
+	static const VectorRegister4Float FloatInfinity = MakeVectorRegisterFloatMask((uint32)0x7F800000, (uint32)0x7F800000, (uint32)0x7F800000, (uint32)0x7F800000);
+	static const VectorRegister4Double DoubleInfinity = MakeVectorRegisterDoubleMask((uint64)0x7FF0000000000000, (uint64)0x7FF0000000000000, (uint64)0x7FF0000000000000, (uint64)0x7FF0000000000000);
 
 	static const VectorRegister4Float Pi = MakeVectorRegister(PI, PI, PI, PI);
 	static const VectorRegister4Float TwoPi = MakeVectorRegister(2.0f*PI, 2.0f*PI, 2.0f*PI, 2.0f*PI);
@@ -122,9 +121,9 @@ namespace GlobalVectorConstants
 
 	static const VectorRegister4Float FloatTwo = MakeVectorRegister(2.0f, 2.0f, 2.0f, 2.0f);
 	static const uint32 AlmostTwoBits = 0x3fffffff;
-	static const VectorRegister4Float FloatAlmostTwo = MakeVectorRegister(*(float*)&AlmostTwoBits, *(float*)&AlmostTwoBits, *(float*)&AlmostTwoBits, *(float*)&AlmostTwoBits);
+	static const VectorRegister4Float FloatAlmostTwo = MakeVectorRegisterFloatMask(AlmostTwoBits, AlmostTwoBits, AlmostTwoBits, AlmostTwoBits);
 
 	static const VectorRegister4Double DoubleTwo = MakeVectorRegisterDouble(2.0, 2.0, 2.0, 2.0);
 	static const uint64 DoubleAlmostTwoBits = 0x3FFFFFFFFFFFFFFF;
-	static const VectorRegister4Double DoubleAlmostTwo = MakeVectorRegisterDouble(*(double*)&DoubleAlmostTwoBits, *(double*)&DoubleAlmostTwoBits, *(double*)&DoubleAlmostTwoBits, *(double*)&DoubleAlmostTwoBits);
+	static const VectorRegister4Double DoubleAlmostTwo = MakeVectorRegisterDoubleMask(DoubleAlmostTwoBits, DoubleAlmostTwoBits, DoubleAlmostTwoBits, DoubleAlmostTwoBits);
 }
diff --git a/Engine/Source/Runtime/Core/Public/Math/VectorRegister.h b/Engine/Source/Runtime/Core/Public/Math/VectorRegister.h
--- a/Engine/Source/Runtime/Core/Public/Math/VectorRegister.h
+++ b/Engine/Source/Runtime/Core/Public/Math/VectorRegister.h
@@ -10,19 +10,19 @@
 
 // Platform specific vector intrinsics include.
 #if WITH_DIRECTXMATH
-#define SIMD_ALIGNMENT (16)
 #include "Math/UnrealMathDirectX.h"
 #elif PLATFORM_ENABLE_VECTORINTRINSICS_NEON
-#define SIMD_ALIGNMENT (16)
 #include "Math/UnrealMathNeon.h"
+#elif defined(__cplusplus_cli) && !PLATFORM_HOLOLENS
+#include "Math/UnrealMathFPU.h" // there are compile issues with UnrealMathSSE in managed mode, so use the FPU version
 #elif PLATFORM_ENABLE_VECTORINTRINSICS
-#define SIMD_ALIGNMENT (16)
 #include "Math/UnrealMathSSE.h"
 #else
-#define SIMD_ALIGNMENT (4)
 #include "Math/UnrealMathFPU.h"
 #endif
 
+#define SIMD_ALIGNMENT (alignof(VectorRegister))
+
 // 'Cross-platform' vector intrinsics (built on the platform-specific ones defined above)
 #include "Math/UnrealMathVectorCommon.h"
 
diff --git a/Engine/Source/Runtime/Core/Public/Templates/TypeCompatibleBytes.h b/Engine/Source/Runtime/Core/Public/Templates/TypeCompatibleBytes.h
--- a/Engine/Source/Runtime/Core/Public/Templates/TypeCompatibleBytes.h
+++ b/Engine/Source/Runtime/Core/Public/Templates/TypeCompatibleBytes.h
@@ -45,6 +45,7 @@ struct TAlignedBytes<Size,1>
 #endif
 
 // Implement TAlignedBytes for these alignments.
+IMPLEMENT_ALIGNED_STORAGE(32);
 IMPLEMENT_ALIGNED_STORAGE(16);
 IMPLEMENT_ALIGNED_STORAGE(8);
 IMPLEMENT_ALIGNED_STORAGE(4);
