Fix shutdown code of ETW tracing, and do not log warning on startup if ETW is not enabled yet

diff --git a/Engine/Source/Runtime/Core/Private/ProfilingDebugging/Windows/EventTracingForWindows.cpp b/Engine/Source/Runtime/Core/Private/ProfilingDebugging/Windows/EventTracingForWindows.cpp
--- a/Engine/Source/Runtime/Core/Private/ProfilingDebugging/Windows/EventTracingForWindows.cpp
+++ b/Engine/Source/Runtime/Core/Private/ProfilingDebugging/Windows/EventTracingForWindows.cpp
@@ -42,7 +42,8 @@ private:
 	bool bStarting = false;
 	bool bError = false;
 	FEventRef ReadyEvent;
-	FRunnableThread* Thread;
+	FRunnableThread* Thread = nullptr;
+	uint32 SamplingIntervalUsec;
 
 	EPlatformEvent EnabledEvents = EPlatformEvent::None;
 
@@ -173,29 +174,8 @@ void FPlatformEvents::TraceEventCallback(PEVENT_RECORD Event)
 
 FPlatformEvents::FPlatformEvents(uint32 SamplingIntervalUsec)
 	: ReadyEvent(EEventMode::ManualReset)
+	, SamplingIntervalUsec(SamplingIntervalUsec)
 {
-	// set sampling interval, must be done before StartTrace
-	{
-		typedef struct
-		{
-			ULONG Source;
-			ULONG Interval;
-		}
-		TRACE_PROFILE_INTERVAL;
-
-		// interval is specified in 100 nanosecond units, and allowed range is limited
-		// see "-setprofint" argument in https://docs.microsoft.com/en-us/windows-hardware/test/wpt/start
-		TRACE_PROFILE_INTERVAL Interval;
-		Interval.Source = 0;
-		Interval.Interval = FMath::Clamp<ULONG>(10 * SamplingIntervalUsec, 1221, 10000000);
-
-		ULONG Status = ::TraceSetInformation(NULL, TraceSampledProfileIntervalInfo, &Interval, sizeof(Interval));
-		if (Status != ERROR_SUCCESS)
-		{
-			UE_LOG(LogPlatformEvents, Warning, TEXT("Unable to set ETW sampling interval: 0x%08x"), Status);
-		}
-	}
-
 	ZeroMemory(ContexSwitchStart, sizeof(ContexSwitchStart));
 	CurrentProcessId = ::GetCurrentProcessId();
 
@@ -349,8 +329,11 @@ void FPlatformEvents::Stop()
 	// stopping it
 	ReadyEvent->Wait(MAX_uint32);
 
-	StopETW();
-	Thread->WaitForCompletion();
+	if (Thread)
+	{
+		StopETW();
+		Thread->WaitForCompletion();
+	}
 }
 
 /////////////////////////////////////////////////////////////////////
@@ -401,6 +384,28 @@ bool FPlatformEvents::StartETW()
 		}
 	}
 
+	// set up sampling interval, must be done before StartTrace
+	{
+		typedef struct
+		{
+			ULONG Source;
+			ULONG Interval;
+		}
+		TRACE_PROFILE_INTERVAL;
+
+		// interval is specified in 100 nanosecond units, and allowed range is limited
+		// see "-setprofint" argument in https://docs.microsoft.com/en-us/windows-hardware/test/wpt/start
+		TRACE_PROFILE_INTERVAL Interval;
+		Interval.Source = 0;
+		Interval.Interval = FMath::Clamp<ULONG>(10 * SamplingIntervalUsec, 1221, 10000000);
+
+		ULONG Status = ::TraceSetInformation(NULL, TraceSampledProfileIntervalInfo, &Interval, sizeof(Interval));
+		if (Status != ERROR_SUCCESS)
+		{
+			UE_LOG(LogPlatformEvents, Warning, TEXT("Unable to set ETW sampling interval: 0x%08x"), Status);
+		}
+	}
+
 	// setup trace properties
 	ZeroMemory(Properties, sizeof(*Properties));
 	Properties->Wnode.BufferSize = sizeof(Trace);
