Make MinBytesSaved and MinPercentSaved package compression parameters configurable.

diff --git a/Engine/Config/BaseGame.ini b/Engine/Config/BaseGame.ini
--- a/Engine/Config/BaseGame.ini
+++ b/Engine/Config/BaseGame.ini
@@ -93,6 +93,8 @@ PackageCompressionMethod=Kraken
 PackageCompressionLevel_Distribution=7
 PackageCompressionLevel_TestShipping=5
 PackageCompressionLevel_DebugDevelopment=3
+PackageCompressionMinBytesSaved=1024
+PackageCompressionMinPercentSaved=5
 bForceUseProjectCompressionFormatIgnoreHardwareOverride=False
 bGenerateChunks=False
 bChunkHardReferencesOnly=False
diff --git a/Engine/Source/Developer/DeveloperToolSettings/Classes/Settings/ProjectPackagingSettings.h b/Engine/Source/Developer/DeveloperToolSettings/Classes/Settings/ProjectPackagingSettings.h
--- a/Engine/Source/Developer/DeveloperToolSettings/Classes/Settings/ProjectPackagingSettings.h
+++ b/Engine/Source/Developer/DeveloperToolSettings/Classes/Settings/ProjectPackagingSettings.h
@@ -317,6 +317,24 @@ public:
 	UPROPERTY(config, EditAnywhere, Category = Packaging, AdvancedDisplay, meta = (DisplayName = "Encoder Effort Level for Distribution"))
 	int32 PackageCompressionLevel_Distribution;
 
+	/**
+	 * A generic setting which is used to determine whether it is worth using compression for a block of data when creating IoStore or .pak files.
+	 * If the amount of saved bytes is smaller than the specified value, then the block of data remains uncompressed.
+	 * The optimal value of this setting depends on the capabilities of the target platform. For instance PackageCompressionMinBytesSaved=1024
+	 * Note that some compressors (for example Oodle) do their own internal worth it check and only use this value to determine the minimal size of a block which should be compressed.
+	 */
+	UPROPERTY(config, EditAnywhere, Category = Packaging, AdvancedDisplay, meta = (DisplayName = "Minimum amount of bytes which should be saved when compressing a block of data, otherwise data remains uncompressed"))
+	int32 PackageCompressionMinBytesSaved;
+
+	/**
+	 * A generic setting which is used to determine whether it is worth using compression for a block of data when creating IoStore or .pak files.
+	 * If the saved percentage of a compressed block of data is smaller than the specified value, then the block remains uncompressed.
+	 * The optimal value of this setting depends on the capabilities of the target platform. For instance PackageCompressionMinPercentSaved=5
+	 * Note that some compressors (for example Oodle) do their own internal worth it check and ignore this value.
+	 */
+	UPROPERTY(config, EditAnywhere, Category = Packaging, AdvancedDisplay, meta = (DisplayName = "Minimum percentage of a block of data which should be saved when performing compression, otherwise data remains uncompressed"))
+	int32 PackageCompressionMinPercentSaved;
+
 	/** 
 	 * Version name for HTTP Chunk Install Data.
 	 */
diff --git a/Engine/Source/Developer/IoStoreUtilities/Private/IoStoreUtilities.cpp b/Engine/Source/Developer/IoStoreUtilities/Private/IoStoreUtilities.cpp
--- a/Engine/Source/Developer/IoStoreUtilities/Private/IoStoreUtilities.cpp
+++ b/Engine/Source/Developer/IoStoreUtilities/Private/IoStoreUtilities.cpp
@@ -3310,6 +3310,19 @@ int32 CreateIoStoreContainerFiles(const TCHAR* CmdLine)
 		}
 	}
 
+	FString TargetPlatformName;
+	if (FParse::Value(FCommandLine::Get(), TEXT("TargetPlatform="), TargetPlatformName))
+	{
+		UE_LOG(LogIoStore, Display, TEXT("Using target platform '%s'"), *TargetPlatformName);
+		ITargetPlatformManagerModule& TPM = GetTargetPlatformManagerRef();
+		Arguments.TargetPlatform = TPM.FindTargetPlatform(TargetPlatformName);
+		if (!Arguments.TargetPlatform)
+		{
+			UE_LOG(LogIoStore, Error, TEXT("Invalid TargetPlatform: '%s'"), *TargetPlatformName);
+			return 1;
+		}
+	}
+
 	FIoStoreWriterSettings GeneralIoWriterSettings { DefaultCompressionMethod, DefaultCompressionBlockSize, false };
 	GeneralIoWriterSettings.bEnableCsvOutput = FParse::Param(CmdLine, TEXT("-csvoutput"));
 
@@ -3376,9 +3389,16 @@ int32 CreateIoStoreContainerFiles(const TCHAR* CmdLine)
 		GeneralIoWriterSettings.MaxPartitionSize = MaxPartitionSize;
 	}
 
+	if (Arguments.TargetPlatform)
+	{
+		GeneralIoWriterSettings.InitializePlatformSpecificSettings(Arguments.TargetPlatform);
+	}
+
 	UE_LOG(LogIoStore, Display, TEXT("Using memory mapping alignment '%ld'"), GeneralIoWriterSettings.MemoryMappingAlignment);
 	UE_LOG(LogIoStore, Display, TEXT("Using compression block size '%ld'"), GeneralIoWriterSettings.CompressionBlockSize);
 	UE_LOG(LogIoStore, Display, TEXT("Using compression block alignment '%ld'"), GeneralIoWriterSettings.CompressionBlockAlignment);
+	UE_LOG(LogIoStore, Display, TEXT("Using compression min bytes saved '%d'"), GeneralIoWriterSettings.CompressionMinBytesSaved);
+	UE_LOG(LogIoStore, Display, TEXT("Using compression min percent saved '%d'"), GeneralIoWriterSettings.CompressionMinPercentSaved);
 	UE_LOG(LogIoStore, Display, TEXT("Using max partition size '%lld'"), GeneralIoWriterSettings.MaxPartitionSize);
 
 	FParse::Value(CmdLine, TEXT("-MetaOutputDirectory="), Arguments.MetaOutputDir);
@@ -3504,19 +3524,7 @@ int32 CreateIoStoreContainerFiles(const TCHAR* CmdLine)
 
 	if (Arguments.ShouldCreateContainers())
 	{
-		FString TargetPlatform;
-		if (FParse::Value(FCommandLine::Get(), TEXT("TargetPlatform="), TargetPlatform))
-		{
-			UE_LOG(LogIoStore, Display, TEXT("Using target platform '%s'"), *TargetPlatform);
-			ITargetPlatformManagerModule& TPM = GetTargetPlatformManagerRef();
-			Arguments.TargetPlatform = TPM.FindTargetPlatform(TargetPlatform);
-			if (!Arguments.TargetPlatform)
-			{
-				UE_LOG(LogIoStore, Error, TEXT("Invalid TargetPlatform: '%s'"), *TargetPlatform);
-				return 1;
-			}
-		}
-		else
+		if (!Arguments.TargetPlatform)
 		{
 			UE_LOG(LogIoStore, Error, TEXT("TargetPlatform must be specified"));
 			return 1;
diff --git a/Engine/Source/Developer/IoStoreUtilities/Private/PackageStoreWriter.cpp b/Engine/Source/Developer/IoStoreUtilities/Private/PackageStoreWriter.cpp
--- a/Engine/Source/Developer/IoStoreUtilities/Private/PackageStoreWriter.cpp
+++ b/Engine/Source/Developer/IoStoreUtilities/Private/PackageStoreWriter.cpp
@@ -15,6 +15,7 @@ FPackageStoreWriter::FPackageStoreWriter(const FString& InOutputPath, const ITar
 {
 	FIoStoreWriterSettings WriterSettings;
 	WriterSettings.CompressionMethod = NAME_Zlib;
+	WriterSettings.InitializePlatformSpecificSettings(InTargetPlatform);
 	FIoStatus IoStatus = IoStoreWriterContext->Initialize(WriterSettings);
 	check(IoStatus.IsOk());
 	FIoContainerSettings ContainerSettings;
diff --git a/Engine/Source/Runtime/Core/Private/IO/IoStore.cpp b/Engine/Source/Runtime/Core/Private/IO/IoStore.cpp
--- a/Engine/Source/Runtime/Core/Private/IO/IoStore.cpp
+++ b/Engine/Source/Runtime/Core/Private/IO/IoStore.cpp
@@ -17,6 +17,7 @@
 #include "Misc/CoreDelegates.h"
 #include "Serialization/MemoryWriter.h"
 #include "Async/AsyncFileHandle.h"
+#include "Interfaces/ITargetPlatform.h"
 
 //////////////////////////////////////////////////////////////////////////
 
@@ -146,6 +147,8 @@ struct FChunkBlock
 	uint64 CompressedSize = 0;
 	uint64 UncompressedSize = 0;
 	FName CompressionMethod = NAME_None;
+	int32 CompressionMinBytesSaved = 0;
+	int32 CompressionMinPercentSaved = 0;
 	FSHAHash Signature;
 };
 
@@ -997,6 +1000,8 @@ private:
 			FChunkBlock& Block = Entry->ChunkBlocks[BlockIndex];
 			Block.IoBuffer = WriterContext->AllocCompressionBuffer(int32(NumChunkBlocks));
 			Block.CompressionMethod = CompressionMethod;
+			Block.CompressionMinBytesSaved = WriterSettings.CompressionMinBytesSaved;
+			Block.CompressionMinPercentSaved = WriterSettings.CompressionMinPercentSaved;
 			Block.UncompressedSize = FMath::Min(BytesToProcess, WriterSettings.CompressionBlockSize);
 			Block.UncompressedData = UncompressedData;
 			BytesToProcess -= Block.UncompressedSize;
@@ -1026,6 +1031,8 @@ private:
 				TRACE_CPUPROFILER_EVENT_SCOPE(CompressMemory);
 				bCompressed = FCompression::CompressMemoryIfWorthDecompressing(
 					Block->CompressionMethod,
+					Block->CompressionMinBytesSaved,
+					Block->CompressionMinPercentSaved,
 					Block->IoBuffer->Data(),
 					CompressedBlockSize,
 					Block->UncompressedData,
@@ -2042,4 +2049,12 @@ void FIoStoreReader::GetFilenamesByBlockIndex(const TArray<int32>& InBlockIndexL
 
 			return true;
 		});
+}
+
+void FIoStoreWriterSettings::InitializePlatformSpecificSettings(const ITargetPlatform* TargetPlatform)
+{
+	FConfigCacheIni* PlatformConfig = TargetPlatform->GetConfigSystem();
+	const TCHAR* IniSection = TEXT("/Script/UnrealEd.ProjectPackagingSettings");
+	PlatformConfig->GetInt(IniSection, TEXT("PackageCompressionMinBytesSaved"), CompressionMinBytesSaved, GGameIni);
+	PlatformConfig->GetInt(IniSection, TEXT("PackageCompressionMinPercentSaved"), CompressionMinPercentSaved, GGameIni);
 }
\ No newline at end of file
diff --git a/Engine/Source/Runtime/Core/Private/Misc/Compression.cpp b/Engine/Source/Runtime/Core/Private/Misc/Compression.cpp
--- a/Engine/Source/Runtime/Core/Private/Misc/Compression.cpp
+++ b/Engine/Source/Runtime/Core/Private/Misc/Compression.cpp
@@ -530,13 +530,12 @@ int32 FCompression::CompressMemoryBound(FName FormatName, int32 UncompressedSize
 
 
 
-bool FCompression::CompressMemoryIfWorthDecompressing(FName FormatName, void* CompressedBuffer, int32& CompressedSize, const void* UncompressedBuffer, int32 UncompressedSize, ECompressionFlags Flags, int32 CompressionData)
+bool FCompression::CompressMemoryIfWorthDecompressing(FName FormatName, int32 MinBytesSaved, int32 MinPercentSaved, void* CompressedBuffer, int32& CompressedSize, const void* UncompressedBuffer, int32 UncompressedSize, ECompressionFlags Flags, int32 CompressionData)
 {
 	// returns false if we could compress,
 	//	but it's not worth the time to decompress
 	//	you should store the data uncompressed instead
 
-	const int MinBytesSaved = 1024; // @todo make this configurable?
 	if ( UncompressedSize <= MinBytesSaved )
 	{
 		// if input size is smaller than the number of bytes we need to save
@@ -600,13 +599,9 @@ bool FCompression::CompressMemoryIfWorthDecompressing(FName FormatName, void* Co
 		return false;
 	}
 
-	// Check the compression ratio, if it's too low just store uncompressed. 
-	// 	saving 64 KB per 1 MB is 6%
-	// @todo make this configurable?
-	float PercentSaved = ((float)BytesSaved * 100.f / (float)UncompressedSize);
-	bool bWorthIt = PercentSaved >= 5.f;
-
-	return bWorthIt;
+	// Check the saved compression ratio, if it's too low just store uncompressed. 
+	// For example, saving 64 KB per 1 MB is about 6%.
+	return (int64)BytesSaved * 100 >= (int64)UncompressedSize * MinPercentSaved;
 }
 
 bool FCompression::CompressMemory(FName FormatName, void* CompressedBuffer, int32& CompressedSize, const void* UncompressedBuffer, int32 UncompressedSize, ECompressionFlags Flags, int32 CompressionData)
diff --git a/Engine/Source/Runtime/Core/Public/IO/IoDispatcher.h b/Engine/Source/Runtime/Core/Public/IO/IoDispatcher.h
--- a/Engine/Source/Runtime/Core/Public/IO/IoDispatcher.h
+++ b/Engine/Source/Runtime/Core/Public/IO/IoDispatcher.h
@@ -1029,10 +1029,14 @@ struct FIoStoreWriterSettings
 	FName CompressionMethod = NAME_None;
 	uint64 CompressionBlockSize = 64 << 10;
 	uint64 CompressionBlockAlignment = 0;
+	int32 CompressionMinBytesSaved = 0;
+	int32 CompressionMinPercentSaved = 0;
 	uint64 MemoryMappingAlignment = 0;
 	uint64 MaxPartitionSize = 0;
 	bool bEnableCsvOutput = false;
 	bool bEnableFileRegions = false;
+
+	CORE_API void InitializePlatformSpecificSettings(const ITargetPlatform* TargetPlatform);
 };
 
 enum class EIoContainerFlags : uint8
diff --git a/Engine/Source/Runtime/Core/Public/Misc/Compression.h b/Engine/Source/Runtime/Core/Public/Misc/Compression.h
--- a/Engine/Source/Runtime/Core/Public/Misc/Compression.h
+++ b/Engine/Source/Runtime/Core/Public/Misc/Compression.h
@@ -43,9 +43,10 @@ struct FCompression
 	/**
 	 * Thread-safe abstract compression routine to query memory requirements for a compression operation.
 	 *
-	 * @param	Flags						Flags to control what method to use and optionally control memory vs speed
+	 * @param	FormatName					Name of the compression format
 	 * @param	UncompressedSize			Size of uncompressed data in bytes
-	 * @param	BitWindow					Bit window to use in compression
+	 * @param	Flags						Flags to control what method to use and optionally control memory vs speed
+	 * @param	CompressionData				Additional compression parameter (specifies BitWindow value for ZLIB compression format)
 	 * @return The maximum possible bytes needed for compression of data buffer of size UncompressedSize
 	 */
 	CORE_API static int32 CompressMemoryBound(FName FormatName, int32 UncompressedSize, ECompressionFlags Flags=COMPRESS_NoFlags, int32 CompressionData=0);
@@ -56,12 +57,13 @@ struct FCompression
 	 * CompressMemory is expected to return true and write valid data even if it expanded bytes.
 	 * Always check CompressedSize >= UncompressedSize and fall back to uncompressed, or use CompressMemoryIfWorthDecompressing
 	 *
-	 * @param	Flags						Flags to control what method to use and optionally control memory vs speed
+	 * @param	FormatName					Name of the compression format
 	 * @param	CompressedBuffer			Buffer compressed data is going to be written to
 	 * @param	CompressedSize	[in/out]	Size of CompressedBuffer, at exit will be size of compressed data
 	 * @param	UncompressedBuffer			Buffer containing uncompressed data
 	 * @param	UncompressedSize			Size of uncompressed data in bytes
-	 * @param	BitWindow					Bit window to use in compression
+	 * @param	Flags						Flags to control what method to use and optionally control memory vs speed
+	 * @param	CompressionData				Additional compression parameter (specifies BitWindow value for ZLIB compression format)
 	 * @return true if compression succeeds, false if it fails because CompressedBuffer was too small or other reasons
 	 */
 	CORE_API static bool CompressMemory(FName FormatName, void* CompressedBuffer, int32& CompressedSize, const void* UncompressedBuffer, int32 UncompressedSize, ECompressionFlags Flags=COMPRESS_NoFlags, int32 CompressionData=0);
@@ -71,26 +73,30 @@ struct FCompression
 	* returns false if the size saving is not worth it (also if CompressedSize >= UncompressedSize)
 	* if false is returned, send the data uncompressed instead
 	 *
-	 * @param	Flags						Flags to control what method to use and optionally control memory vs speed
+	 * @param	FormatName					Name of the compression format
+	 * @param	MinBytesSaved				Minimum amount of bytes which should be saved when performing compression, otherwise false is returned
+	 * @param	MinPercentSaved				Minimum percentage of the buffer which should be saved when performing compression, otherwise false is returned
 	 * @param	CompressedBuffer			Buffer compressed data is going to be written to
 	 * @param	CompressedSize	[in/out]	Size of CompressedBuffer, at exit will be size of compressed data
 	 * @param	UncompressedBuffer			Buffer containing uncompressed data
 	 * @param	UncompressedSize			Size of uncompressed data in bytes
-	 * @param	BitWindow					Bit window to use in compression
+	 * @param	Flags						Flags to control what method to use and optionally control memory vs speed
+	 * @param	CompressionData				Additional compression parameter (specifies BitWindow value for ZLIB compression format)
 	 * @return true if compression succeeds, false if it fails because CompressedBuffer was too small or other reasons
 	 */
-	CORE_API static bool CompressMemoryIfWorthDecompressing(FName FormatName, void* CompressedBuffer, int32& CompressedSize, const void* UncompressedBuffer, int32 UncompressedSize, ECompressionFlags Flags=COMPRESS_NoFlags, int32 CompressionData=0);
+	CORE_API static bool CompressMemoryIfWorthDecompressing(FName FormatName, int32 MinBytesSaved, int32 MinPercentSaved, void* CompressedBuffer, int32& CompressedSize, const void* UncompressedBuffer, int32 UncompressedSize, ECompressionFlags Flags=COMPRESS_NoFlags, int32 CompressionData=0);
 
 	/**
 	 * Thread-safe abstract decompression routine. Uncompresses memory from compressed buffer and writes it to uncompressed
 	 * buffer. UncompressedSize is expected to be the exact size of the data after decompression.
 	 *
-	 * @param	Flags						Flags to control what method to use to decompress
+	 * @param	FormatName					Name of the compression format
 	 * @param	UncompressedBuffer			Buffer containing uncompressed data
 	 * @param	UncompressedSize			Size of uncompressed data in bytes
 	 * @param	CompressedBuffer			Buffer compressed data is going to be read from
 	 * @param	CompressedSize				Size of CompressedBuffer data in bytes
-	 * @param	bIsSourcePadded		Whether the source memory is padded with a full cache line at the end
+	 * @param	Flags						Flags to control what method to use to decompress
+	 * @param	CompressionData				Additional decompression parameter (specifies BitWindow value for ZLIB compression format)
 	 * @return true if compression succeeds, false if it fails because CompressedBuffer was too small or other reasons
 	 */
 	CORE_API static bool UncompressMemory(FName FormatName, void* UncompressedBuffer, int32 UncompressedSize, const void* CompressedBuffer, int32 CompressedSize, ECompressionFlags Flags=COMPRESS_NoFlags, int32 CompressionData=0);
