Make sure FTransform are using native serialization in order to be compatible with LWC

diff --git a/Engine/Source/Runtime/Core/Public/Math/TransformNonVectorized.h b/Engine/Source/Runtime/Core/Public/Math/TransformNonVectorized.h
--- a/Engine/Source/Runtime/Core/Public/Math/TransformNonVectorized.h
+++ b/Engine/Source/Runtime/Core/Public/Math/TransformNonVectorized.h
@@ -677,6 +677,12 @@ public:
 		return Ar;
 	}
 
+	bool Serialize(FArchive& Ar)
+	{
+		Ar << *this;
+		return true;
+	}
+
 	// Binary comparison operators.
 	/*
 	bool operator==(const FTransform& Other) const
diff --git a/Engine/Source/Runtime/Core/Public/Math/TransformVectorized.h b/Engine/Source/Runtime/Core/Public/Math/TransformVectorized.h
--- a/Engine/Source/Runtime/Core/Public/Math/TransformVectorized.h
+++ b/Engine/Source/Runtime/Core/Public/Math/TransformVectorized.h
@@ -664,6 +664,12 @@ public:
 		return Ar;
 	}
 
+	bool Serialize(FArchive& Ar)
+	{
+		Ar << *this;
+		return true;
+	}
+
 	// Binary comparison operators.
 	/*
 	bool operator==( const FTransform& Other ) const
diff --git a/Engine/Source/Runtime/CoreUObject/Private/UObject/Property.cpp b/Engine/Source/Runtime/CoreUObject/Private/UObject/Property.cpp
--- a/Engine/Source/Runtime/CoreUObject/Private/UObject/Property.cpp
+++ b/Engine/Source/Runtime/CoreUObject/Private/UObject/Property.cpp
@@ -362,6 +362,7 @@ struct TStructOpsTypeTraits<FTransform> : public TStructOpsTypeTraitsBase2<FTran
 {
 	enum
 	{
+		WithSerializer = true,
 		WithIdentical = true,
 	};
 };
