No-unity compile fixes.

#rnx

diff --git a/Engine/Source/Runtime/Core/Private/ProfilingDebugging/PlatformEvents.h b/Engine/Source/Runtime/Core/Private/ProfilingDebugging/PlatformEvents.h
--- a/Engine/Source/Runtime/Core/Private/ProfilingDebugging/PlatformEvents.h
+++ b/Engine/Source/Runtime/Core/Private/ProfilingDebugging/PlatformEvents.h
@@ -2,7 +2,10 @@
 #pragma once
 
 #include "CoreTypes.h"
-#include "Trace/Trace.inl"
+#include "Misc/EnumClassFlags.h"
+#include "Trace/Trace.h"
+
+class FString;
 
 #if PLATFORM_WINDOWS
 #	define UE_PLATFORM_EVENTS_AVAILABLE 1
