Compile fix as my one was incorrect.

#rnx

diff --git a/Engine/Source/Runtime/Core/Private/ProfilingDebugging/PlatformEvents.cpp b/Engine/Source/Runtime/Core/Private/ProfilingDebugging/PlatformEvents.cpp
--- a/Engine/Source/Runtime/Core/Private/ProfilingDebugging/PlatformEvents.cpp
+++ b/Engine/Source/Runtime/Core/Private/ProfilingDebugging/PlatformEvents.cpp
@@ -8,6 +8,9 @@
 UE_TRACE_CHANNEL_DEFINE(ContextSwitchChannel)
 UE_TRACE_CHANNEL_DEFINE(StackSamplingChannel)
 
+UE_TRACE_EVENT_DEFINE(PlatformEvent, ContextSwitch)
+UE_TRACE_EVENT_DEFINE(PlatformEvent, StackSample)
+
 /////////////////////////////////////////////////////////////////////
 
 EPlatformEvent PlatformEvents_GetEvent(const FString& Name)
diff --git a/Engine/Source/Runtime/Core/Private/ProfilingDebugging/PlatformEvents.h b/Engine/Source/Runtime/Core/Private/ProfilingDebugging/PlatformEvents.h
--- a/Engine/Source/Runtime/Core/Private/ProfilingDebugging/PlatformEvents.h
+++ b/Engine/Source/Runtime/Core/Private/ProfilingDebugging/PlatformEvents.h
@@ -4,6 +4,7 @@
 #include "CoreTypes.h"
 #include "Misc/EnumClassFlags.h"
 #include "Trace/Trace.h"
+#include "Trace/Trace.inl"
 
 class FString;
 
@@ -21,7 +22,7 @@ UE_TRACE_CHANNEL_EXTERN(StackSamplingChannel)
 /////////////////////////////////////////////////////////////////////
 
 // represents time interval when thread was running on specific core
-UE_TRACE_EVENT_BEGIN(PlatformEvent, ContextSwitch, NoSync)
+UE_TRACE_EVENT_BEGIN_EXTERN(PlatformEvent, ContextSwitch, NoSync)
 	UE_TRACE_EVENT_FIELD(uint64, StartTime)
 	UE_TRACE_EVENT_FIELD(uint64, EndTime)
 	UE_TRACE_EVENT_FIELD(uint32, ThreadId)
@@ -29,7 +30,7 @@ UE_TRACE_EVENT_BEGIN(PlatformEvent, ContextSwitch, NoSync)
 UE_TRACE_EVENT_END()
 
 // represents call stack addresses in stack sampling
-UE_TRACE_EVENT_BEGIN(PlatformEvent, StackSample, NoSync)
+UE_TRACE_EVENT_BEGIN_EXTERN(PlatformEvent, StackSample, NoSync)
 	UE_TRACE_EVENT_FIELD(uint64, Time)
 	UE_TRACE_EVENT_FIELD(uint32, ThreadId)
 	UE_TRACE_EVENT_FIELD(uint64[], Addresses)
