Fix nonunity compilation error in IoStore.cpp

diff --git a/Engine/Source/Runtime/Core/Private/IO/IoStore.cpp b/Engine/Source/Runtime/Core/Private/IO/IoStore.cpp
--- a/Engine/Source/Runtime/Core/Private/IO/IoStore.cpp
+++ b/Engine/Source/Runtime/Core/Private/IO/IoStore.cpp
@@ -18,6 +18,7 @@
 #include "Serialization/MemoryWriter.h"
 #include "Async/AsyncFileHandle.h"
 #include "Interfaces/ITargetPlatform.h"
+#include "Misc/ConfigCacheIni.h"
 
 //////////////////////////////////////////////////////////////////////////
 
