#ue5 - LWC: Fix FRotator::Quaternion() working in doubles when LWC=on when it just needs to use floats (for now), which also meant it was doing some unsafe bit manipulation to change the sign of some values which only worked for floats.

diff --git a/Engine/Source/Runtime/Core/Private/Math/UnrealMath.cpp b/Engine/Source/Runtime/Core/Private/Math/UnrealMath.cpp
--- a/Engine/Source/Runtime/Core/Private/Math/UnrealMath.cpp
+++ b/Engine/Source/Runtime/Core/Private/Math/UnrealMath.cpp
@@ -410,35 +410,35 @@ FQuat FRotator::Quaternion() const
 	DiagnosticCheckNaN();
 
 #if PLATFORM_ENABLE_VECTORINTRINSICS
-	const VectorRegister Angles = MakeVectorRegister(Pitch, Yaw, Roll, 0.0f);
-	const VectorRegister AnglesNoWinding = VectorMod(Angles, GlobalVectorConstants::Float360);
-	const VectorRegister HalfAngles = VectorMultiply(AnglesNoWinding, GlobalVectorConstants::DEG_TO_RAD_HALF);
+	const VectorRegister4Float Angles = MakeVectorRegister(Pitch, Yaw, Roll, 0.0f);
+	const VectorRegister4Float AnglesNoWinding = VectorMod(Angles, GlobalVectorConstants::Float360);
+	const VectorRegister4Float HalfAngles = VectorMultiply(AnglesNoWinding, GlobalVectorConstants::DEG_TO_RAD_HALF);
 
-	VectorRegister SinAngles, CosAngles;
+	VectorRegister4Float SinAngles, CosAngles;
 	VectorSinCos(&SinAngles, &CosAngles, &HalfAngles);
 
 	// Vectorized conversion, measured 20% faster than using scalar version after VectorSinCos.
 	// Indices within VectorRegister (for shuffles): P=0, Y=1, R=2
-	const VectorRegister SR = VectorReplicate(SinAngles, 2);
-	const VectorRegister CR = VectorReplicate(CosAngles, 2);
+	const VectorRegister4Float SR = VectorReplicate(SinAngles, 2);
+	const VectorRegister4Float CR = VectorReplicate(CosAngles, 2);
 
-	const VectorRegister SY_SY_CY_CY_Temp = VectorShuffle(SinAngles, CosAngles, 1, 1, 1, 1);
+	const VectorRegister4Float SY_SY_CY_CY_Temp = VectorShuffle(SinAngles, CosAngles, 1, 1, 1, 1);
 
-	const VectorRegister SP_SP_CP_CP = VectorShuffle(SinAngles, CosAngles, 0, 0, 0, 0);
-	const VectorRegister SY_CY_SY_CY = VectorShuffle(SY_SY_CY_CY_Temp, SY_SY_CY_CY_Temp, 0, 2, 0, 2);
+	const VectorRegister4Float SP_SP_CP_CP = VectorShuffle(SinAngles, CosAngles, 0, 0, 0, 0);
+	const VectorRegister4Float SY_CY_SY_CY = VectorShuffle(SY_SY_CY_CY_Temp, SY_SY_CY_CY_Temp, 0, 2, 0, 2);
 
-	const VectorRegister CP_CP_SP_SP = VectorShuffle(CosAngles, SinAngles, 0, 0, 0, 0);
-	const VectorRegister CY_SY_CY_SY = VectorShuffle(SY_SY_CY_CY_Temp, SY_SY_CY_CY_Temp, 2, 0, 2, 0);
+	const VectorRegister4Float CP_CP_SP_SP = VectorShuffle(CosAngles, SinAngles, 0, 0, 0, 0);
+	const VectorRegister4Float CY_SY_CY_SY = VectorShuffle(SY_SY_CY_CY_Temp, SY_SY_CY_CY_Temp, 2, 0, 2, 0);
 
 	const uint32 Neg = uint32(1 << 31);
 	const uint32 Pos = uint32(0);
-	const VectorRegister SignBitsLeft  = MakeVectorRegister(Pos, Neg, Pos, Pos);
-	const VectorRegister SignBitsRight = MakeVectorRegister(Neg, Neg, Neg, Pos);
-	const VectorRegister LeftTerm  = VectorBitwiseXor(SignBitsLeft , VectorMultiply(CR, VectorMultiply(SP_SP_CP_CP, SY_CY_SY_CY)));
-	const VectorRegister RightTerm = VectorBitwiseXor(SignBitsRight, VectorMultiply(SR, VectorMultiply(CP_CP_SP_SP, CY_SY_CY_SY)));
+	const VectorRegister4Float SignBitsLeft  = MakeVectorRegister(Pos, Neg, Pos, Pos);
+	const VectorRegister4Float SignBitsRight = MakeVectorRegister(Neg, Neg, Neg, Pos);
+	const VectorRegister4Float LeftTerm  = VectorBitwiseXor(SignBitsLeft , VectorMultiply(CR, VectorMultiply(SP_SP_CP_CP, SY_CY_SY_CY)));
+	const VectorRegister4Float RightTerm = VectorBitwiseXor(SignBitsRight, VectorMultiply(SR, VectorMultiply(CP_CP_SP_SP, CY_SY_CY_SY)));
 
 	FQuat RotationQuat;
-	const VectorRegister Result = VectorAdd(LeftTerm, RightTerm);	
+	const VectorRegister4Float Result = VectorAdd(LeftTerm, RightTerm);	
 	VectorStoreAligned(Result, &RotationQuat);
 #else
 	const float DEG_TO_RAD = PI/(180.f);
