#ue5 - LWC: Implement math functions using SVML (Simple Vector Math Library) when available. Gives faster results for operations like VectorSin(), VectorSinCos(), etc. Temporarily disabled until I add more math tests at startup, though it passes everthing that exists currently.

diff --git a/Engine/Source/Runtime/Core/Public/Math/UnrealMathSSE.h b/Engine/Source/Runtime/Core/Public/Math/UnrealMathSSE.h
--- a/Engine/Source/Runtime/Core/Public/Math/UnrealMathSSE.h
+++ b/Engine/Source/Runtime/Core/Public/Math/UnrealMathSSE.h
@@ -21,13 +21,22 @@
 #define UE_PLATFORM_MATH_USE_FMA3			PLATFORM_ALWAYS_HAS_FMA3
 #endif
 
+#ifndef UE_PLATFORM_MATH_USE_SVML
+//////////// TEMPORARILY disabled until UnrealMathTest is expanded to test more trig functions, though it passes what exists there now. ////////////
+#define UE_PLATFORM_MATH_USE_SVML			(0 && (_MSC_VER >= 1920)) // Support added to MSVC 2019 16.0+
+#endif
+
+#ifndef UE_PLATFORM_MATH_USE_SVML_AVX
+#define UE_PLATFORM_MATH_USE_SVML_AVX		(UE_PLATFORM_MATH_USE_SVML && UE_PLATFORM_MATH_USE_AVX)
+#endif
+
 // If SSE4.1 is enabled, need additional defines.
 #if UE_PLATFORM_MATH_USE_SSE4_1
 #include <smmintrin.h>
 #endif
 
 // If AVX is enabled, need additional defines.
-#if UE_PLATFORM_MATH_USE_AVX
+#if UE_PLATFORM_MATH_USE_AVX || UE_PLATFORM_MATH_USE_SVML
 #include <immintrin.h>
 #define UE_SSE_DOUBLE_ALIGNMENT 32
 #else
@@ -2686,78 +2695,114 @@ FORCEINLINE VectorRegister4Double VectorStep(const VectorRegister4Double& X)
 	return VectorSelect(Mask, GlobalVectorConstants::DoubleOne, GlobalVectorConstants::DoubleZero);
 }
 
-//TODO: Vectorize
 FORCEINLINE VectorRegister4Float VectorExp(const VectorRegister4Float& X)
 {
+#if UE_PLATFORM_MATH_USE_SVML
+	return _mm_exp_ps(X);
+#else
 	return SseMath_exp_ps(X);
+#endif
 }
 
 FORCEINLINE VectorRegister4Double VectorExp(const VectorRegister4Double& X)
 {
+#if UE_PLATFORM_MATH_USE_SVML_AVX
+	return _mm256_exp_pd(X);
+#elif UE_PLATFORM_MATH_USE_SVML
+	return VectorRegister4Double(_mm_exp_pd(X.XY), _mm_exp_pd(X.ZW));
+#else
 	AlignedDouble4 Doubles(X);
 	Doubles[0] = FMath::Exp(Doubles[0]);
 	Doubles[1] = FMath::Exp(Doubles[1]);
 	Doubles[2] = FMath::Exp(Doubles[2]);
 	Doubles[3] = FMath::Exp(Doubles[3]);
 	return Doubles.ToVectorRegister();
+#endif
 }
 
-//TODO: Vectorize
 FORCEINLINE VectorRegister4Float VectorExp2(const VectorRegister4Float& X)
 {
+#if UE_PLATFORM_MATH_USE_SVML
+	return _mm_exp2_ps(X);
+#else
 	AlignedFloat4 Floats(X);
 	Floats[0] = FMath::Exp2(Floats[0]);
 	Floats[1] = FMath::Exp2(Floats[1]);
 	Floats[2] = FMath::Exp2(Floats[2]);
 	Floats[3] = FMath::Exp2(Floats[3]);
 	return Floats.ToVectorRegister();
+#endif
 }
 
 FORCEINLINE VectorRegister4Double VectorExp2(const VectorRegister4Double& X)
 {
+#if UE_PLATFORM_MATH_USE_SVML_AVX
+	return _mm256_exp2_pd(X);
+#elif UE_PLATFORM_MATH_USE_SVML
+	return VectorRegister4Double(_mm_exp2_pd(X.XY), _mm_exp2_pd(X.ZW));
+#else
 	AlignedDouble4 Doubles(X);
 	Doubles[0] = FMath::Exp2(Doubles[0]);
 	Doubles[1] = FMath::Exp2(Doubles[1]);
 	Doubles[2] = FMath::Exp2(Doubles[2]);
 	Doubles[3] = FMath::Exp2(Doubles[3]);
 	return Doubles.ToVectorRegister();
+#endif
 }
 
-//TODO: Vectorize
 FORCEINLINE VectorRegister4Float VectorLog(const VectorRegister4Float& X)
 {
+#if UE_PLATFORM_MATH_USE_SVML
+	return _mm_log_ps(X);
+#else
 	return SseMath_log_ps(X);
+#endif
 }
 
 FORCEINLINE VectorRegister4Double VectorLog(const VectorRegister4Double& X)
 {
+#if UE_PLATFORM_MATH_USE_SVML_AVX
+	return _mm256_log_pd(X);
+#elif UE_PLATFORM_MATH_USE_SVML
+	return VectorRegister4Double(_mm_log_pd(X.XY), _mm_log_pd(X.ZW));
+#else
 	AlignedDouble4 Doubles(X);
 	Doubles[0] = FMath::Loge(Doubles[0]);
 	Doubles[1] = FMath::Loge(Doubles[1]);
 	Doubles[2] = FMath::Loge(Doubles[2]);
 	Doubles[3] = FMath::Loge(Doubles[3]);
 	return Doubles.ToVectorRegister();
+#endif
 }
 
-//TODO: Vectorize
 FORCEINLINE VectorRegister4Float VectorLog2(const VectorRegister4Float& X)
 {
+#if UE_PLATFORM_MATH_USE_SVML
+	return _mm_log2_ps(X);
+#else
 	AlignedFloat4 Floats(X);
 	Floats[0] = FMath::Log2(Floats[0]);
 	Floats[1] = FMath::Log2(Floats[1]);
 	Floats[2] = FMath::Log2(Floats[2]);
 	Floats[3] = FMath::Log2(Floats[3]);
 	return Floats.ToVectorRegister();
+#endif
 }
 
 FORCEINLINE VectorRegister4Double VectorLog2(const VectorRegister4Double& X)
 {
+#if UE_PLATFORM_MATH_USE_SVML_AVX
+	return _mm256_log2_pd(X);
+#elif UE_PLATFORM_MATH_USE_SVML
+	return VectorRegister4Double(_mm_log2_pd(X.XY), _mm_log2_pd(X.ZW));
+#else
 	AlignedDouble4 Doubles(X);
 	Doubles[0] = FMath::Log2(Doubles[0]);
 	Doubles[1] = FMath::Log2(Doubles[1]);
 	Doubles[2] = FMath::Log2(Doubles[2]);
 	Doubles[3] = FMath::Log2(Doubles[3]);
 	return Doubles.ToVectorRegister();
+#endif
 }
 
 
@@ -2776,8 +2821,11 @@ namespace VectorSinConstantsSSE
 	static const VectorRegister4Float B = MakeVectorRegisterFloat(b, b, b, b);
 }
 
-FORCEINLINE VectorRegister4Float VectorSin(const VectorRegister4Float& X)
+FORCEINLINE VectorRegister4Float VectorSin(const VectorRegister4Float& V)
 {
+#if UE_PLATFORM_MATH_USE_SVML
+	return _mm_sin_ps(V);
+#else
 	//Sine approximation using a squared parabola restrained to f(0) = 0, f(PI) = 0, f(PI/2) = 1.
 	//based on a good discussion here http://forum.devmaster.net/t/fast-and-accurate-sine-cosine/9648
 	//After approx 2.5 million tests comparing to sin(): 
@@ -2793,35 +2841,52 @@ FORCEINLINE VectorRegister4Float VectorSin(const VectorRegister4Float& X)
 	// Also - other platforms (NEON) don't vectorize this as of January 2021,
 	// so you'll get accuracy differences between platforms.
 	//
-	VectorRegister4Float y = VectorMultiply(X, GlobalVectorConstants::OneOverTwoPi);
+	VectorRegister4Float y = VectorMultiply(V, GlobalVectorConstants::OneOverTwoPi);
 	y = VectorSubtract(y, VectorFloor(VectorAdd(y, GlobalVectorConstants::FloatOneHalf)));
 	y = VectorMultiply(VectorSinConstantsSSE::A, VectorMultiply(y, VectorSubtract(GlobalVectorConstants::FloatOneHalf, VectorAbs(y))));
 	return VectorMultiply(y, VectorAdd(VectorSinConstantsSSE::B, VectorAbs(y)));
+#endif
 }
 
-FORCEINLINE VectorRegister4Double VectorSin(const VectorRegister4Double& X)
+FORCEINLINE VectorRegister4Double VectorSin(const VectorRegister4Double& V)
 {
-	AlignedDouble4 Doubles(X);
+#if UE_PLATFORM_MATH_USE_SVML_AVX
+		return _mm256_sin_pd(V);
+#elif UE_PLATFORM_MATH_USE_SVML
+		return VectorRegister4Double(_mm_sin_pd(V.XY), _mm_sin_pd(V.ZW));
+#else
+	AlignedDouble4 Doubles(V);
 	Doubles[0] = FMath::Sin(Doubles[0]);
 	Doubles[1] = FMath::Sin(Doubles[1]);
 	Doubles[2] = FMath::Sin(Doubles[2]);
 	Doubles[3] = FMath::Sin(Doubles[3]);
 	return Doubles.ToVectorRegister();
+#endif
 }
 
-FORCEINLINE VectorRegister4Float VectorCos(const VectorRegister4Float& X)
+FORCEINLINE VectorRegister4Float VectorCos(const VectorRegister4Float& V)
 {
-	return VectorSin(VectorAdd(X, GlobalVectorConstants::PiByTwo));
+#if UE_PLATFORM_MATH_USE_SVML
+	return _mm_cos_ps(V);
+#else
+	return VectorSin(VectorAdd(V, GlobalVectorConstants::PiByTwo));
+#endif
 }
 
-FORCEINLINE VectorRegister4Double VectorCos(const VectorRegister4Double& X)
+FORCEINLINE VectorRegister4Double VectorCos(const VectorRegister4Double& V)
 {
-	AlignedDouble4 Doubles(X);
+#if UE_PLATFORM_MATH_USE_SVML_AVX
+		return _mm256_cos_pd(V);
+#elif UE_PLATFORM_MATH_USE_SVML
+		return VectorRegister4Double(_mm_cos_pd(V.XY), _mm_cos_pd(V.ZW));
+#else
+	AlignedDouble4 Doubles(V);
 	Doubles[0] = FMath::Cos(Doubles[0]);
 	Doubles[1] = FMath::Cos(Doubles[1]);
 	Doubles[2] = FMath::Cos(Doubles[2]);
 	Doubles[3] = FMath::Cos(Doubles[3]);
 	return Doubles.ToVectorRegister();
+#endif
 }
 
 /**
@@ -2833,6 +2898,9 @@ FORCEINLINE VectorRegister4Double VectorCos(const VectorRegister4Double& X)
 */
 FORCEINLINE void VectorSinCos(VectorRegister4Float* RESTRICT VSinAngles, VectorRegister4Float* RESTRICT VCosAngles, const VectorRegister4Float* RESTRICT VAngles)
 {
+#if UE_PLATFORM_MATH_USE_SVML
+	*VSinAngles = _mm_sincos_ps(VCosAngles, *VAngles);
+#else
 	// Map to [-pi, pi]
 	// X = A - 2pi * round(A/2pi)
 	// Note the round(), not truncate(). In this case round() can round halfway cases using round-to-nearest-even OR round-to-nearest.
@@ -2881,102 +2949,148 @@ FORCEINLINE void VectorSinCos(VectorRegister4Float* RESTRICT VSinAngles, VectorR
 	C = VectorMultiplyAdd(XSquared, C, VectorReplicate(CosCoeff0, 1));
 	C = VectorMultiplyAdd(XSquared, C, VectorReplicate(CosCoeff0, 0));
 	*VCosAngles = VectorMultiply(C, sign);
+#endif
 }
 
 FORCEINLINE void VectorSinCos(VectorRegister4Double* RESTRICT VSinAngles, VectorRegister4Double* RESTRICT VCosAngles, const VectorRegister4Double* RESTRICT VAngles)
 {
+#if UE_PLATFORM_MATH_USE_SVML_AVX
+	VSinAngles->XYZW = _mm256_sincos_pd(&(VCosAngles->XYZW), VAngles->XYZW);
+#elif UE_PLATFORM_MATH_USE_SVML
+	VSinAngles->XY = _mm_sincos_pd(&(VCosAngles->XY), VAngles->XY);
+	VSinAngles->ZW = _mm_sincos_pd(&(VCosAngles->ZW), VAngles->ZW);
+#else
 	*VSinAngles = VectorSin(*VAngles);
 	*VCosAngles = VectorCos(*VAngles);
+#endif
 }
 
 
-//TODO: Vectorize
 FORCEINLINE VectorRegister4Float VectorTan(const VectorRegister4Float& X)
 {
+#if UE_PLATFORM_MATH_USE_SVML
+	return _mm_tan_ps(X);
+#else
 	AlignedFloat4 Floats(X);
 	Floats[0] = FMath::Tan(Floats[0]);
 	Floats[1] = FMath::Tan(Floats[1]);
 	Floats[2] = FMath::Tan(Floats[2]);
 	Floats[3] = FMath::Tan(Floats[3]);
 	return Floats.ToVectorRegister();
+#endif
 }
 
 FORCEINLINE VectorRegister4Double VectorTan(const VectorRegister4Double& X)
 {
+#if UE_PLATFORM_MATH_USE_SVML_AVX
+	return _mm256_tan_pd(X);
+#elif UE_PLATFORM_MATH_USE_SVML
+	return VectorRegister4Double(_mm_tan_pd(X.XY), _mm_tan_pd(X.ZW));
+#else
 	AlignedDouble4 Doubles(X);
 	Doubles[0] = FMath::Tan(Doubles[0]);
 	Doubles[1] = FMath::Tan(Doubles[1]);
 	Doubles[2] = FMath::Tan(Doubles[2]);
 	Doubles[3] = FMath::Tan(Doubles[3]);
 	return Doubles.ToVectorRegister();
+#endif
 }
 
-//TODO: Vectorize
 FORCEINLINE VectorRegister4Float VectorASin(const VectorRegister4Float& X)
 {
+#if UE_PLATFORM_MATH_USE_SVML
+	return _mm_asin_ps(X);
+#else
 	AlignedFloat4 Floats(X);
 	Floats[0] = FMath::Asin(Floats[0]);
 	Floats[1] = FMath::Asin(Floats[1]);
 	Floats[2] = FMath::Asin(Floats[2]);
 	Floats[3] = FMath::Asin(Floats[3]);
 	return Floats.ToVectorRegister();
+#endif
 }
 
 FORCEINLINE VectorRegister4Double VectorASin(const VectorRegister4Double& X)
 {
+#if UE_PLATFORM_MATH_USE_SVML_AVX
+	return _mm256_asin_pd(X);
+#elif UE_PLATFORM_MATH_USE_SVML
+	return VectorRegister4Double(_mm_asin_pd(X.XY), _mm_asin_pd(X.ZW));
+#else
 	AlignedDouble4 Doubles(X);
 	Doubles[0] = FMath::Asin(Doubles[0]);
 	Doubles[1] = FMath::Asin(Doubles[1]);
 	Doubles[2] = FMath::Asin(Doubles[2]);
 	Doubles[3] = FMath::Asin(Doubles[3]);
 	return Doubles.ToVectorRegister();
+#endif
 }
 
-//TODO: Vectorize
 FORCEINLINE VectorRegister4Float VectorACos(const VectorRegister4Float& X)
 {
+#if UE_PLATFORM_MATH_USE_SVML
+	return _mm_acos_ps(X);
+#else
 	AlignedFloat4 Floats(X);
 	Floats[0] = FMath::Acos(Floats[0]);
 	Floats[1] = FMath::Acos(Floats[1]);
 	Floats[2] = FMath::Acos(Floats[2]);
 	Floats[3] = FMath::Acos(Floats[3]);
 	return Floats.ToVectorRegister();
+#endif
 }
 
 FORCEINLINE VectorRegister4Double VectorACos(const VectorRegister4Double& X)
 {
+#if UE_PLATFORM_MATH_USE_SVML_AVX
+	return _mm256_acos_pd(X);
+#elif UE_PLATFORM_MATH_USE_SVML
+	return VectorRegister4Double(_mm_acos_pd(X.XY), _mm_acos_pd(X.ZW));
+#else
 	AlignedDouble4 Doubles(X);
 	Doubles[0] = FMath::Acos(Doubles[0]);
 	Doubles[1] = FMath::Acos(Doubles[1]);
 	Doubles[2] = FMath::Acos(Doubles[2]);
 	Doubles[3] = FMath::Acos(Doubles[3]);
 	return Doubles.ToVectorRegister();
+#endif
 }
 
-//TODO: Vectorize
 FORCEINLINE VectorRegister4Float VectorATan(const VectorRegister4Float& X)
 {
+#if UE_PLATFORM_MATH_USE_SVML
+	return _mm_atan_ps(X);
+#else
 	AlignedFloat4 Floats(X);
 	Floats[0] = FMath::Atan(Floats[0]);
 	Floats[1] = FMath::Atan(Floats[1]);
 	Floats[2] = FMath::Atan(Floats[2]);
 	Floats[3] = FMath::Atan(Floats[3]);
 	return Floats.ToVectorRegister();
+#endif
 }
 
 FORCEINLINE VectorRegister4Double VectorATan(const VectorRegister4Double& X)
 {
+#if UE_PLATFORM_MATH_USE_SVML_AVX
+	return _mm256_atan_pd(X);
+#elif UE_PLATFORM_MATH_USE_SVML
+	return VectorRegister4Double(_mm_atan_pd(X.XY), _mm_atan_pd(X.ZW));
+#else
 	AlignedDouble4 Doubles(X);
 	Doubles[0] = FMath::Atan(Doubles[0]);
 	Doubles[1] = FMath::Atan(Doubles[1]);
 	Doubles[2] = FMath::Atan(Doubles[2]);
 	Doubles[3] = FMath::Atan(Doubles[3]);
 	return Doubles.ToVectorRegister();
+#endif
 }
 
-//TODO: Vectorize
 FORCEINLINE VectorRegister4Float VectorATan2(const VectorRegister4Float& X, const VectorRegister4Float& Y)
 {
+#if UE_PLATFORM_MATH_USE_SVML
+	return _mm_atan2_ps(X, Y);
+#else
 	AlignedFloat4 FloatsX(X);
 	AlignedFloat4 FloatsY(Y);
 	FloatsX[0] = FMath::Atan2(FloatsX[0], FloatsY[0]);
@@ -2984,10 +3098,16 @@ FORCEINLINE VectorRegister4Float VectorATan2(const VectorRegister4Float& X, cons
 	FloatsX[2] = FMath::Atan2(FloatsX[2], FloatsY[2]);
 	FloatsX[3] = FMath::Atan2(FloatsX[3], FloatsY[3]);
 	return FloatsX.ToVectorRegister();
+#endif
 }
 
 FORCEINLINE VectorRegister4Double VectorATan2(const VectorRegister4Double& X, const VectorRegister4Double& Y)
 {
+#if UE_PLATFORM_MATH_USE_SVML_AVX
+	return _mm256_atan2_pd(X, Y);
+#elif UE_PLATFORM_MATH_USE_SVML
+	return VectorRegister4Double(_mm_atan2_pd(X.XY, Y.XY), _mm_atan2_pd(X.ZW, Y.ZW));
+#else
 	AlignedDouble4 DoublesX(X);
 	AlignedDouble4 DoublesY(Y);
 	DoublesX[0] = FMath::Atan2(DoublesX[0], DoublesY[0]);
@@ -2995,6 +3115,7 @@ FORCEINLINE VectorRegister4Double VectorATan2(const VectorRegister4Double& X, co
 	DoublesX[2] = FMath::Atan2(DoublesX[2], DoublesY[2]);
 	DoublesX[3] = FMath::Atan2(DoublesX[3], DoublesY[3]);
 	return DoublesX.ToVectorRegister();
+#endif
 }
 
 // To be continued...
