#ue5 - LWC: nicer #define checks for using SVML on MSVC.

diff --git a/Engine/Source/Runtime/Core/Public/Math/UnrealMathSSE.h b/Engine/Source/Runtime/Core/Public/Math/UnrealMathSSE.h
--- a/Engine/Source/Runtime/Core/Public/Math/UnrealMathSSE.h
+++ b/Engine/Source/Runtime/Core/Public/Math/UnrealMathSSE.h
@@ -22,8 +22,12 @@
 #endif
 
 #ifndef UE_PLATFORM_MATH_USE_SVML
-//////////// TEMPORARILY disabled until UnrealMathTest is expanded to test more trig functions, though it passes what exists there now. ////////////
-#define UE_PLATFORM_MATH_USE_SVML			(0 && (_MSC_VER >= 1920)) // Support added to MSVC 2019 16.0+
+	#if defined(_MSC_VER)
+	//////////// TEMPORARILY disabled until UnrealMathTest is expanded to test more trig functions, though it passes what exists there now. ////////////
+	#define UE_PLATFORM_MATH_USE_SVML			(0 && (_MSC_VER >= 1920)) // Support added to MSVC 2019 16.0+
+	#else
+	#define UE_PLATFORM_MATH_USE_SVML			0
+	#endif // defined(_MSC_VER)
 #endif
 
 #ifndef UE_PLATFORM_MATH_USE_SVML_AVX
