LLM: Fix race condition when a tag is added from LLM_SCOPE calls on two threads at the same time.
#rnx

diff --git a/Engine/Source/Runtime/Core/Private/HAL/LowLevelMemTracker.cpp b/Engine/Source/Runtime/Core/Private/HAL/LowLevelMemTracker.cpp
--- a/Engine/Source/Runtime/Core/Private/HAL/LowLevelMemTracker.cpp
+++ b/Engine/Source/Runtime/Core/Private/HAL/LowLevelMemTracker.cpp
@@ -1476,6 +1476,26 @@ UE::LLMPrivate::FTagData& FLowLevelMemTracker::RegisterTagData(FName Name, FName
 	}
 
 	FWriteScopeLock ScopeLock(TagDataLock);
+	FTagData*& TagDataForName = TagDataNameMap->FindOrAdd(Name, nullptr);
+	if (TagDataForName)
+	{
+		// The tag already exists; this can happen because two formal registrations
+		// have tried to add the same name, but it can also happen due to a race condition
+		// when auto-adding a tag from LLM_SCOPE, if the tag is added by another thread
+		// in between FindOrAddTagData's check of TagDataNameMap->Find and now.
+		// Note that it is not valid for an LLM_SCOPE to be called before a formal registration
+		// (e.g. LLM_DECLARE_TAG). If a formal registration exists for a tag, it must precede its use
+		// in any LLM_SCOPE calls.
+		if (ReferenceSource != UE::LLMPrivate::ETagReferenceSource::Scope &&
+			ReferenceSource != UE::LLMPrivate::ETagReferenceSource::FunctionAPI)
+		{
+			ReportDuplicateTagName(TagDataForName, ReferenceSource);
+		}
+
+		// Abandon the string work we've done and return the version that was added by the other source.
+		return *TagDataForName;
+	}
+
 	FTagData* ParentData = nullptr;
 	if (!ParentName.IsNone())
 	{
@@ -1498,11 +1518,6 @@ UE::LLMPrivate::FTagData& FLowLevelMemTracker::RegisterTagData(FName Name, FName
 	TagData->SetIndex(TagDatas->Num());
 	TagDatas->Add(TagData);
 
-	FTagData*& TagDataForName = TagDataNameMap->FindOrAdd(Name, nullptr);
-	if (TagDataForName != nullptr)
-	{
-		ReportDuplicateTagName(TagDataForName, ReferenceSource);
-	}
 	TagDataForName = TagData;
 
 	if (bHasEnumTag)
