#ue5 - LWC: Refactored tests for VectorReplicate(), VectorSwizzle(), VectorShuffle() to test many more permutations in more compact code. Added missing tests of VectorShuffle() to float version.

diff --git a/Engine/Source/Runtime/Core/Private/Tests/Math/UnrealMathTest.cpp b/Engine/Source/Runtime/Core/Private/Tests/Math/UnrealMathTest.cpp
--- a/Engine/Source/Runtime/Core/Private/Tests/Math/UnrealMathTest.cpp
+++ b/Engine/Source/Runtime/Core/Private/Tests/Math/UnrealMathTest.cpp
@@ -945,133 +945,134 @@ bool RunDoubleVectorTest()
 	LogTest(TEXT("VectorGetComponent<double>"), TestVectorsEqual(V0, V1));
 
 	// Replicate
+	{
+		double ArrayV0[4] = { 0.0, 1.0, 2.0, 3.0 };
+		V0 = VectorLoad(ArrayV0);
 
-	V0 = MakeVectorRegisterDouble(1.0f, 2.0f, 3.0f, 4.0f);
-	V1 = VectorReplicate(V0, 0);
-	V0 = MakeVectorRegisterDouble(1.0f, 1.0f, 1.0f, 1.0f);
-	LogTest(TEXT("VectorReplicate<double>"), TestVectorsEqual(V0, V1));
-
-	V0 = MakeVectorRegisterDouble(1.0f, 2.0f, 3.0f, 4.0f);
-	V1 = VectorReplicate(V0, 1);
-	V0 = MakeVectorRegisterDouble(2.0f, 2.0f, 2.0f, 2.0f);
-	LogTest(TEXT("VectorReplicate<double>"), TestVectorsEqual(V0, V1));
+#define ReplicateTest(A, X) \
+		V2 = VectorReplicate(A, X); \
+		V3 = MakeVectorRegisterDouble(ArrayV0[X], ArrayV0[X], ArrayV0[X], ArrayV0[X]); \
+		LogTest(*FString::Printf(TEXT("VectorReplicateDouble<%d>"), X), TestVectorsEqual(V2, V3));
 
-	V0 = MakeVectorRegisterDouble(1.0f, 2.0f, 3.0f, 4.0f);
-	V1 = VectorReplicate(V0, 2);
-	V0 = MakeVectorRegisterDouble(3.0f, 3.0f, 3.0f, 3.0f);
-	LogTest(TEXT("VectorReplicate<double>"), TestVectorsEqual(V0, V1));
+		ReplicateTest(V0, 0);
+		ReplicateTest(V0, 1);
+		ReplicateTest(V0, 2);
+		ReplicateTest(V0, 3);
 
-	V0 = MakeVectorRegisterDouble(1.0f, 2.0f, 3.0f, 4.0f);
-	V1 = VectorReplicate(V0, 3);
-	V0 = MakeVectorRegisterDouble(4.0f, 4.0f, 4.0f, 4.0f);
-	LogTest(TEXT("VectorReplicate<double>"), TestVectorsEqual(V0, V1));
+#undef ReplicateTest
+	}
 
 	// Swizzle
+	{
+		double ArrayV0[4] = { 0.0, 1.0, 2.0, 3.0 };
+		V0 = VectorLoad(ArrayV0);
 
-	V0 = MakeVectorRegisterDouble(0.0, 1.0, 2.0, 3.0);
-	V1 = VectorSwizzle(V0, 0, 1, 2, 3);
-	V0 = MakeVectorRegisterDouble(0.0, 1.0, 2.0, 3.0);
-	LogTest(TEXT("VectorShuffle0123<double>"), TestVectorsEqual(V0, V1));
-
-	V0 = MakeVectorRegisterDouble(4.0f, 3.0f, 2.0f, 1.0f);
-	V1 = VectorSwizzle(V0, 1, 0, 1, 0);
-	V0 = MakeVectorRegisterDouble(3.0f, 4.0f, 3.0f, 4.0f);
-	LogTest(TEXT("VectorSwizzle1010<double>"), TestVectorsEqual(V0, V1));
-
-	V0 = MakeVectorRegisterDouble(4.0f, 3.0f, 2.0f, 1.0f);
-	V1 = VectorSwizzle(V0, 1, 0, 3, 2);
-	V0 = MakeVectorRegisterDouble(3.0f, 4.0f, 1.0f, 2.0f);
-	LogTest(TEXT("VectorSwizzle1032<double>"), TestVectorsEqual(V0, V1));
-
-	V0 = MakeVectorRegisterDouble(4.0f, 3.0f, 2.0f, 1.0f);
-	V1 = VectorSwizzle(V0, 1, 2, 0, 1);
-	V0 = MakeVectorRegisterDouble(3.0f, 2.0f, 4.0f, 3.0f);
-	LogTest(TEXT("VectorSwizzle1201<double>"), TestVectorsEqual(V0, V1));
-
-	V0 = MakeVectorRegisterDouble(4.0f, 3.0f, 2.0f, 1.0f);
-	V1 = VectorSwizzle(V0, 2, 0, 1, 3);
-	V0 = MakeVectorRegisterDouble(2.0f, 4.0f, 3.0f, 1.0f);
-	LogTest(TEXT("VectorSwizzle2013<double>"), TestVectorsEqual(V0, V1));
-
-	V0 = MakeVectorRegisterDouble(4.0f, 3.0f, 2.0f, 1.0f);
-	V1 = VectorSwizzle(V0, 2, 3, 0, 1);
-	V0 = MakeVectorRegisterDouble(2.0f, 1.0f, 4.0f, 3.0f);
-	LogTest(TEXT("VectorSwizzle2301<double>"), TestVectorsEqual(V0, V1));
-
-	V0 = MakeVectorRegisterDouble(4.0f, 3.0f, 2.0f, 1.0f);
-	V1 = VectorSwizzle(V0, 3, 2, 1, 0);
-	V0 = MakeVectorRegisterDouble(1.0f, 2.0f, 3.0f, 4.0f);
-	LogTest(TEXT("VectorSwizzle3210<double>"), TestVectorsEqual(V0, V1));
+#define SwizzleTest(A, X, Y, Z, W) \
+		V2 = VectorSwizzle(A, X, Y, Z, W); \
+		V3 = MakeVectorRegisterDouble(ArrayV0[X], ArrayV0[Y], ArrayV0[Z], ArrayV0[W]); \
+		LogTest(*FString::Printf(TEXT("VectorSwizzleDouble<%d,%d,%d,%d>"), X, Y, Z, W), TestVectorsEqual(V2, V3));
 
-	V0 = MakeVectorRegisterDouble(0.0, 1.0, 2.0, 3.0);
-	V1 = VectorSwizzle(V0, 3, 0, 3, 0);
-	V0 = MakeVectorRegisterDouble(3.0, 0.0, 3.0, 0.0);
-	LogTest(TEXT("VectorShuffle3030<double>"), TestVectorsEqual(V0, V1));
+		// This is not an exhaustive list because it would be 4*4*4*4 = 256 entries, but it tries to test a lot of common permutations.
+		// Unfortunately it can't be done in a loop because it uses a #define and compile-time constants for the VectorSwizzle() 'function'.
+		// Many of these were selected to also stress the specializations in certain implementations.
 
-	V0 = MakeVectorRegisterDouble(0.0, 1.0, 2.0, 3.0);
-	V1 = VectorSwizzle(V0, 0, 3, 1, 2);
-	V0 = MakeVectorRegisterDouble(0.0, 3.0, 1.0, 2.0);
-	LogTest(TEXT("VectorShuffle0312<double>"), TestVectorsEqual(V0, V1));
+		SwizzleTest(V0, 0, 1, 2, 3); // Identity
+		SwizzleTest(V0, 0, 0, 0, 0); // Replicate 0
+		SwizzleTest(V0, 1, 1, 1, 1); // Replicate 1
+		SwizzleTest(V0, 2, 2, 2, 2); // Replicate 2
+		SwizzleTest(V0, 3, 3, 3, 3); // Replicate 3
 
-	V0 = MakeVectorRegisterDouble(0.0, 1.0, 2.0, 3.0);
-	V1 = VectorSwizzle(V0, 0, 0, 2, 2);
-	V0 = MakeVectorRegisterDouble(0.0, 0.0, 2.0, 2.0);
-	LogTest(TEXT("VectorShuffle0022<double>"), TestVectorsEqual(V0, V1));
+		SwizzleTest(V0, 0, 3, 1, 2);
+		SwizzleTest(V0, 0, 2, 0, 2);
 
-	V0 = MakeVectorRegisterDouble(0.0, 1.0, 2.0, 3.0);
-	V1 = VectorSwizzle(V0, 0, 0, 3, 3);
-	V0 = MakeVectorRegisterDouble(0.0, 0.0, 3.0, 3.0);
-	LogTest(TEXT("VectorShuffle0022<double>"), TestVectorsEqual(V0, V1));
+		SwizzleTest(V0, 1, 0, 1, 0);
+		SwizzleTest(V0, 1, 0, 3, 2);
+		SwizzleTest(V0, 1, 2, 0, 1);
 
-	// Shuffle
+		SwizzleTest(V0, 2, 0, 1, 3);
+		SwizzleTest(V0, 2, 3, 0, 1);
+
+		SwizzleTest(V0, 3, 2, 1, 0);
+		SwizzleTest(V0, 3, 0, 3, 0);
+		
+		SwizzleTest(V0, 2, 2, 0, 1);
+		SwizzleTest(V0, 3, 3, 1, 0);
+		SwizzleTest(V0, 2, 2, 0, 2);
+		SwizzleTest(V0, 3, 3, 1, 3);
+
+		// Common specializations
+		SwizzleTest(V0, 0, 1, 2, 2);
+		SwizzleTest(V0, 0, 1, 3, 3);
+		SwizzleTest(V0, 0, 1, 3, 2);
+
+		SwizzleTest(V0, 0, 0, 2, 3);
+		SwizzleTest(V0, 0, 0, 2, 2);
+		SwizzleTest(V0, 0, 0, 3, 3);
+		SwizzleTest(V0, 0, 0, 3, 2);
+
+		SwizzleTest(V0, 1, 0, 2, 3);
+		SwizzleTest(V0, 1, 0, 2, 2);
+		SwizzleTest(V0, 1, 0, 3, 3);
+		SwizzleTest(V0, 1, 0, 3, 2);
+
+		SwizzleTest(V0, 1, 1, 2, 3);
+		SwizzleTest(V0, 1, 1, 2, 2);
+		SwizzleTest(V0, 1, 1, 3, 3);
+		SwizzleTest(V0, 1, 1, 3, 2);
+
+		SwizzleTest(V0, 0, 1, 0, 1);
+		SwizzleTest(V0, 2, 3, 2, 3);
+		SwizzleTest(V0, 0, 0, 1, 1);
+		SwizzleTest(V0, 2, 2, 3, 3);
+		SwizzleTest(V0, 0, 0, 2, 2);
+		SwizzleTest(V0, 1, 1, 3, 3);
+#undef SwizzleTest
+	}
 
-	V0 = MakeVectorRegisterDouble(0.0, 1.0, 2.0, 3.0);
-	V1 = MakeVectorRegisterDouble(4.0, 5.0, 6.0, 7.0);
-	V1 = VectorShuffle(V0, V1, 1, 0, 1, 0);
-	V0 = MakeVectorRegisterDouble(1.0, 0.0, 5.0, 4.0);
-	LogTest(TEXT("VectorShuffle1010<double>"), TestVectorsEqual(V0, V1));
-
-	V0 = MakeVectorRegisterDouble(0.0, 1.0, 2.0, 3.0);
-	V1 = MakeVectorRegisterDouble(4.0, 5.0, 6.0, 7.0);
-	V1 = VectorShuffle(V0, V1, 1, 0, 3, 2);
-	V0 = MakeVectorRegisterDouble(1.0, 0.0, 7.0, 6.0);
-	LogTest(TEXT("VectorShuffle1032<double>"), TestVectorsEqual(V0, V1));
-
-	V0 = MakeVectorRegisterDouble(0.0, 1.0, 2.0, 3.0);
-	V1 = MakeVectorRegisterDouble(4.0, 5.0, 6.0, 7.0);
-	V1 = VectorShuffle(V0, V1, 1, 2, 0, 1);
-	V0 = MakeVectorRegisterDouble(1.0, 2.0, 4.0, 5.0);
-	LogTest(TEXT("VectorShuffle1201<double>"), TestVectorsEqual(V0, V1));
-
-	V0 = MakeVectorRegisterDouble(0.0, 1.0, 2.0, 3.0);
-	V1 = MakeVectorRegisterDouble(4.0, 5.0, 6.0, 7.0);
-	V1 = VectorShuffle(V0, V1, 2, 0, 1, 3);
-	V0 = MakeVectorRegisterDouble(2.0, 0.0, 5.0, 7.0);
-	LogTest(TEXT("VectorShuffle2013<double>"), TestVectorsEqual(V0, V1));
-
-	V0 = MakeVectorRegisterDouble(0.0, 1.0, 2.0, 3.0);
-	V1 = MakeVectorRegisterDouble(4.0, 5.0, 6.0, 7.0);
-	V1 = VectorShuffle(V0, V1, 2, 3, 0, 1);
-	V0 = MakeVectorRegisterDouble(2.0, 3.0, 4.0, 5.0);
-	LogTest(TEXT("VectorShuffle2301<double>"), TestVectorsEqual(V0, V1));
-
-	V0 = MakeVectorRegisterDouble(0.0, 1.0, 2.0, 3.0);
-	V1 = MakeVectorRegisterDouble(4.0, 5.0, 6.0, 7.0);
-	V1 = VectorShuffle(V0, V1, 3, 2, 1, 0);
-	V0 = MakeVectorRegisterDouble(3.0, 2.0, 5.0, 4.0);
-	LogTest(TEXT("VectorShuffle3210<double>"), TestVectorsEqual(V0, V1));
-
-	V0 = MakeVectorRegisterDouble(0.0, 1.0, 2.0, 3.0);
-	V1 = MakeVectorRegisterDouble(4.0, 5.0, 6.0, 7.0);
-	V1 = VectorShuffle(V0, V1, 3, 0, 3, 0);
-	V0 = MakeVectorRegisterDouble(3.0, 0.0, 7.0, 4.0);
-	LogTest(TEXT("VectorShuffle3030<double>"), TestVectorsEqual(V0, V1));
-
-	V0 = MakeVectorRegisterDouble(0.0, 1.0, 2.0, 3.0);
-	V1 = MakeVectorRegisterDouble(4.0, 5.0, 6.0, 7.0);
-	V1 = VectorShuffle(V0, V1, 0, 3, 1, 2);
-	V0 = MakeVectorRegisterDouble(0.0, 3.0, 5.0, 6.0);
-	LogTest(TEXT("VectorShuffle0312<double>"), TestVectorsEqual(V0, V1));
+	// Shuffle
+	{
+		double ArrayV0[4] = { 0.0, 1.0, 2.0, 3.0 };
+		double ArrayV1[4] = { 0.0, 1.0, 2.0, 3.0 };
+		V0 = VectorLoad(ArrayV0);
+		V1 = VectorLoad(ArrayV1);
+
+#define ShuffleTest(A, B, X, Y, Z, W) \
+		V2 = VectorShuffle(A, B, X, Y, Z, W); \
+		V3 = MakeVectorRegisterDouble(ArrayV0[X], ArrayV0[Y], ArrayV1[Z], ArrayV1[W]); \
+		LogTest(*FString::Printf(TEXT("VectorShuffleDouble<%d,%d,%d,%d>"), X, Y, Z, W), TestVectorsEqual(V2, V3));
+
+		// This is not an exhaustive list because it would be 4*4*4*4 = 256 entries, but it tries to test a lot of common permutations.
+		// Unfortunately it can't be done in a loop because it uses a #define and compile-time constants for the VectorShuffle() 'function'.
+		// Many of these were selected to also stress the specializations in certain implementations.
+		ShuffleTest(V0, V1, 0, 0, 0, 0);
+		ShuffleTest(V0, V1, 1, 1, 1, 1);
+		ShuffleTest(V0, V1, 2, 2, 2, 2);
+		ShuffleTest(V0, V1, 3, 3, 3, 3);
+
+		ShuffleTest(V0, V1, 0, 1, 2, 3);
+		ShuffleTest(V0, V1, 3, 2, 1, 0);
+
+		ShuffleTest(V0, V1, 0, 1, 0, 1);
+		ShuffleTest(V0, V1, 0, 2, 0, 2);
+		ShuffleTest(V0, V1, 0, 3, 0, 3);
+
+		ShuffleTest(V0, V1, 1, 0, 1, 0);
+		ShuffleTest(V0, V1, 2, 0, 2, 0);
+		ShuffleTest(V0, V1, 3, 0, 3, 0);
+		
+		ShuffleTest(V0, V1, 0, 0, 1, 1);
+		ShuffleTest(V0, V1, 0, 0, 2, 2);
+		ShuffleTest(V0, V1, 1, 1, 3, 3);
+		ShuffleTest(V0, V1, 2, 2, 3, 3);
+		ShuffleTest(V0, V1, 2, 3, 0, 1);
+		ShuffleTest(V0, V1, 2, 3, 2, 3);
+
+		ShuffleTest(V0, V1, 3, 1, 3, 0);
+		ShuffleTest(V0, V1, 2, 2, 0, 1);
+		ShuffleTest(V0, V1, 0, 1, 3, 2);
+		ShuffleTest(V0, V1, 1, 3, 0, 3);
+#undef ShuffleTest
+	}
 
 	// Abs
 	V0 = MakeVectorRegisterDouble(1.0f, -2.0f, 3.0f, -4.0f);
@@ -1656,42 +1657,136 @@ bool FVectorRegisterAbstractionTest::RunTest(const FString& Parameters)
 	V0 = MakeVectorRegister(1.0f, 3.0f, 2.0f, 4.0f);
 	LogTest(TEXT("VectorCombineLow"), TestVectorsEqual(V0, V1));
 
-	// Swizzle
 
-	V0 = MakeVectorRegister(0.0f, 1.0f, 2.0f, 3.0f);
-	V1 = VectorSwizzle(V0, 0, 1, 2, 3);
-	V0 = MakeVectorRegister(0.0f, 1.0f, 2.0f, 3.0f);
-	LogTest(TEXT("VectorSwizzle0123"), TestVectorsEqual(V0, V1));
-
-	V0 = MakeVectorRegister( 4.0f, 3.0f, 2.0f, 1.0f );
-	V1 = VectorSwizzle( V0, 1, 0, 3, 2 );
-	V0 = MakeVectorRegister( 3.0f, 4.0f, 1.0f, 2.0f );
-	LogTest( TEXT("VectorSwizzle1032"), TestVectorsEqual( V0, V1 ) );
+	// Replicate
+	{
+		const float ArrayV0[4] = { 0.0, 1.0, 2.0, 3.0 };
+		V0 = VectorLoad(ArrayV0);
 
-	V0 = MakeVectorRegister( 4.0f, 3.0f, 2.0f, 1.0f );
-	V1 = VectorSwizzle( V0, 1, 2, 0, 1 );
-	V0 = MakeVectorRegister( 3.0f, 2.0f, 4.0f, 3.0f );
-	LogTest( TEXT("VectorSwizzle1201"), TestVectorsEqual( V0, V1 ) );
+#define ReplicateTest(A, X) \
+		V2 = VectorReplicate(A, X); \
+		V3 = MakeVectorRegisterFloat(ArrayV0[X], ArrayV0[X], ArrayV0[X], ArrayV0[X]); \
+		LogTest(*FString::Printf(TEXT("VectorReplicateFloat<%d>"), X), TestVectorsEqual(V2, V3));
 
-	V0 = MakeVectorRegister( 4.0f, 3.0f, 2.0f, 1.0f );
-	V1 = VectorSwizzle( V0, 2, 0, 1, 3 );
-	V0 = MakeVectorRegister( 2.0f, 4.0f, 3.0f, 1.0f );
-	LogTest( TEXT("VectorSwizzle2013"), TestVectorsEqual( V0, V1 ) );
+		ReplicateTest(V0, 0);
+		ReplicateTest(V0, 1);
+		ReplicateTest(V0, 2);
+		ReplicateTest(V0, 3);
 
-	V0 = MakeVectorRegister( 4.0f, 3.0f, 2.0f, 1.0f );
-	V1 = VectorSwizzle( V0, 2, 3, 0, 1 );
-	V0 = MakeVectorRegister( 2.0f, 1.0f, 4.0f, 3.0f );
-	LogTest( TEXT("VectorSwizzle2301"), TestVectorsEqual( V0, V1 ) );
+#undef ReplicateTest
+	}
 
-	V0 = MakeVectorRegister( 4.0f, 3.0f, 2.0f, 1.0f );
-	V1 = VectorSwizzle( V0, 3, 2, 1, 0 );
-	V0 = MakeVectorRegister( 1.0f, 2.0f, 3.0f, 4.0f );
-	LogTest( TEXT("VectorSwizzle3210"), TestVectorsEqual( V0, V1 ) );
+	// Swizzle
+	{
+		const float ArrayV0[4] = { 0.0, 1.0, 2.0, 3.0 };
+		V0 = VectorLoad(ArrayV0);
+
+#define SwizzleTest(A, X, Y, Z, W) \
+		V2 = VectorSwizzle(A, X, Y, Z, W); \
+		V3 = MakeVectorRegisterFloat(ArrayV0[X], ArrayV0[Y], ArrayV0[Z], ArrayV0[W]); \
+		LogTest(*FString::Printf(TEXT("VectorSwizzleFloat<%d,%d,%d,%d>"), X, Y, Z, W), TestVectorsEqual(V2, V3));
+
+		// This is not an exhaustive list because it would be 4*4*4*4 = 256 entries, but it tries to test a lot of common permutations.
+		// Unfortunately it can't be done in a loop because it uses a #define and compile-time constants for the VectorSwizzle() 'function'.
+		// Many of these were selected to also stress the specializations in certain implementations.
+
+		SwizzleTest(V0, 0, 1, 2, 3); // Identity
+		SwizzleTest(V0, 0, 0, 0, 0); // Replicate 0
+		SwizzleTest(V0, 1, 1, 1, 1); // Replicate 1
+		SwizzleTest(V0, 2, 2, 2, 2); // Replicate 2
+		SwizzleTest(V0, 3, 3, 3, 3); // Replicate 3
+
+		SwizzleTest(V0, 0, 3, 1, 2);
+		SwizzleTest(V0, 0, 2, 0, 2);
+
+		SwizzleTest(V0, 1, 0, 1, 0);
+		SwizzleTest(V0, 1, 0, 3, 2);
+		SwizzleTest(V0, 1, 2, 0, 1);
+
+		SwizzleTest(V0, 2, 0, 1, 3);
+		SwizzleTest(V0, 2, 3, 0, 1);
+
+		SwizzleTest(V0, 3, 2, 1, 0);
+		SwizzleTest(V0, 3, 0, 3, 0);
+
+		SwizzleTest(V0, 2, 2, 0, 1);
+		SwizzleTest(V0, 3, 3, 1, 0);
+		SwizzleTest(V0, 2, 2, 0, 2);
+		SwizzleTest(V0, 3, 3, 1, 3);
+
+		// Common specializations
+		SwizzleTest(V0, 0, 1, 2, 2);
+		SwizzleTest(V0, 0, 1, 3, 3);
+		SwizzleTest(V0, 0, 1, 3, 2);
+
+		SwizzleTest(V0, 0, 0, 2, 3);
+		SwizzleTest(V0, 0, 0, 2, 2);
+		SwizzleTest(V0, 0, 0, 3, 3);
+		SwizzleTest(V0, 0, 0, 3, 2);
+
+		SwizzleTest(V0, 1, 0, 2, 3);
+		SwizzleTest(V0, 1, 0, 2, 2);
+		SwizzleTest(V0, 1, 0, 3, 3);
+		SwizzleTest(V0, 1, 0, 3, 2);
+
+		SwizzleTest(V0, 1, 1, 2, 3);
+		SwizzleTest(V0, 1, 1, 2, 2);
+		SwizzleTest(V0, 1, 1, 3, 3);
+		SwizzleTest(V0, 1, 1, 3, 2);
+
+		SwizzleTest(V0, 0, 1, 0, 1);
+		SwizzleTest(V0, 2, 3, 2, 3);
+		SwizzleTest(V0, 0, 0, 1, 1);
+		SwizzleTest(V0, 2, 2, 3, 3);
+		SwizzleTest(V0, 0, 0, 2, 2);
+		SwizzleTest(V0, 1, 1, 3, 3);
+#undef SwizzleTest
+	}
 
-	V0 = MakeVectorRegister(0.0f, 1.0f, 2.0f, 3.0f);
-	V1 = VectorSwizzle(V0, 0, 0, 2, 2);
-	V0 = MakeVectorRegister(0.0f, 0.0f, 2.0f, 2.0f);
-	LogTest(TEXT("VectorSwizzle0022"), TestVectorsEqual(V0, V1));
+	// Shuffle
+	{
+		const float ArrayV0[4] = { 0.0, 1.0, 2.0, 3.0 };
+		const float ArrayV1[4] = { 0.0, 1.0, 2.0, 3.0 };
+		V0 = VectorLoad(ArrayV0);
+		V1 = VectorLoad(ArrayV1);
+
+#define ShuffleTest(A, B, X, Y, Z, W) \
+		V2 = VectorShuffle(A, B, X, Y, Z, W); \
+		V3 = MakeVectorRegisterFloat(ArrayV0[X], ArrayV0[Y], ArrayV1[Z], ArrayV1[W]); \
+		LogTest(*FString::Printf(TEXT("VectorShuffleFloat<%d,%d,%d,%d>"), X, Y, Z, W), TestVectorsEqual(V2, V3));
+
+		// This is not an exhaustive list because it would be 4*4*4*4 = 256 entries, but it tries to test a lot of common permutations.
+		// Unfortunately it can't be done in a loop because it uses a #define and compile-time constants for the VectorShuffle() 'function'.
+		// Many of these were selected to also stress the specializations in certain implementations.
+		ShuffleTest(V0, V1, 0, 0, 0, 0);
+		ShuffleTest(V0, V1, 1, 1, 1, 1);
+		ShuffleTest(V0, V1, 2, 2, 2, 2);
+		ShuffleTest(V0, V1, 3, 3, 3, 3);
+
+		ShuffleTest(V0, V1, 0, 1, 2, 3);
+		ShuffleTest(V0, V1, 3, 2, 1, 0);
+
+		ShuffleTest(V0, V1, 0, 1, 0, 1);
+		ShuffleTest(V0, V1, 0, 2, 0, 2);
+		ShuffleTest(V0, V1, 0, 3, 0, 3);
+
+		ShuffleTest(V0, V1, 1, 0, 1, 0);
+		ShuffleTest(V0, V1, 2, 0, 2, 0);
+		ShuffleTest(V0, V1, 3, 0, 3, 0);
+
+		ShuffleTest(V0, V1, 0, 0, 1, 1);
+		ShuffleTest(V0, V1, 0, 0, 2, 2);
+		ShuffleTest(V0, V1, 1, 1, 3, 3);
+		ShuffleTest(V0, V1, 2, 2, 3, 3);
+		ShuffleTest(V0, V1, 2, 3, 0, 1);
+		ShuffleTest(V0, V1, 2, 3, 2, 3);
+
+		ShuffleTest(V0, V1, 3, 1, 3, 0);
+		ShuffleTest(V0, V1, 2, 2, 0, 1);
+		ShuffleTest(V0, V1, 0, 1, 3, 2);
+		ShuffleTest(V0, V1, 1, 3, 0, 3);
+#undef ShuffleTest
+	}
 
 	// LoadByte4
 
