A better solution for handling negative zero in FString::SanitizeFloat(), as it was producing "0.0" when passed a NaN when built under /fp:fast.

diff --git a/Engine/Source/Runtime/Core/Private/Containers/String.cpp b/Engine/Source/Runtime/Core/Private/Containers/String.cpp
--- a/Engine/Source/Runtime/Core/Private/Containers/String.cpp
+++ b/Engine/Source/Runtime/Core/Private/Containers/String.cpp
@@ -788,10 +788,7 @@ bool FString::ToHexBlob( const FString& Source, uint8* DestBuffer, const uint32
 FString FString::SanitizeFloat( double InFloat, const int32 InMinFractionalDigits )
 {
 	// Avoids negative zero
-	if( InFloat == 0 )
-	{
-		InFloat = 0;
-	}
+	InFloat = InFloat + 0.0f;
 
 	// First create the string
 	FString TempString = FString::Printf(TEXT("%f"), InFloat);
