Fix convex transform having the wrong Rotation in LWC mode
- Changed Rotation to use a float4 vector register in LWC to match its noexporttypes and non-vectorized declarations (FQuat is using float in LWC for now )

diff --git a/Engine/Source/Runtime/Core/Private/Math/TransformVectorized.cpp b/Engine/Source/Runtime/Core/Private/Math/TransformVectorized.cpp
--- a/Engine/Source/Runtime/Core/Private/Math/TransformVectorized.cpp
+++ b/Engine/Source/Runtime/Core/Private/Math/TransformVectorized.cpp
@@ -124,7 +124,7 @@ FTransform FTransform::GetRelativeTransformReverse(const FTransform& Other) cons
 
 	Result.Scale3D = VScale3D;	
 	Result.Translation = VTranslation;
-	Result.Rotation = VRotation;
+	Result.Rotation = MakeVectorRegisterFloatFromDouble(VRotation);
 		
 	Result.DiagnosticCheckNaN_All(); 
 
@@ -172,7 +172,7 @@ void FTransform::SetToRelativeTransform(const FTransform& ParentTransform)
 	Translation = VectorMultiply(VR, VSafeScale3D);
 
 	// Rotation = Q(B)(-1) * Q(A)	
-	Rotation = VectorQuaternionMultiply2(VInverseParentRot, Rotation );
+	Rotation = MakeVectorRegisterFloatFromDouble(VectorQuaternionMultiply2(VInverseParentRot, Rotation ));
 
 	DiagnosticCheckNaN_All(); 
 
@@ -238,7 +238,7 @@ FTransform FTransform::GetRelativeTransform(const FTransform& Other) const
 
 		Result.Scale3D = VScale3D;
 		Result.Translation = VTranslation;
-		Result.Rotation = VRotation;
+		Result.Rotation = MakeVectorRegisterFloatFromDouble(VRotation);
 
 		Result.DiagnosticCheckNaN_All();
 #if DEBUG_INVERSE_TRANSFORM
diff --git a/Engine/Source/Runtime/Core/Public/Math/TransformVectorized.h b/Engine/Source/Runtime/Core/Public/Math/TransformVectorized.h
--- a/Engine/Source/Runtime/Core/Public/Math/TransformVectorized.h
+++ b/Engine/Source/Runtime/Core/Public/Math/TransformVectorized.h
@@ -39,7 +39,7 @@ struct alignas(alignof(VectorRegister)) FTransform
 
 protected:
 	/** Rotation of this transformation, as a quaternion */
-	VectorRegister	Rotation;
+	VectorRegister4Float Rotation; // LWC_TODO : change back to VectorRegister when FQuat uses double 
 	/** Translation of this transformation, as a vector */
 	VectorRegister	Translation;
 	/** 3D scale (always applied in local space) as a vector */
@@ -108,7 +108,7 @@ public:
 	FORCEINLINE FTransform()
 	{
 		// Rotation = {0,0,0,1)
-		Rotation = GlobalVectorConstants::Vector0001; //VectorSet_W1( VectorZero() );
+		Rotation = GlobalVectorConstants::Float0001; //VectorSet_W1( VectorZero() );
 		// Translation = {0,0,0,0)
 		Translation = VectorZero();
 		// Scale3D = {1,1,1,0);
@@ -123,7 +123,7 @@ public:
 	FORCEINLINE explicit FTransform(const FVector& InTranslation) 
 	{
 		// Rotation = {0,0,0,1) quaternion identity
-		Rotation = GlobalVectorConstants::Vector0001; //VectorSet_W1( VectorZero() );
+		Rotation = GlobalVectorConstants::Float0001; //VectorSet_W1( VectorZero() );
 		//Translation = InTranslation;
 		Translation = MakeVectorRegister(InTranslation.X, InTranslation.Y, InTranslation.Z, 0.0f );
 		// Scale3D = {1,1,1,0);
@@ -193,7 +193,7 @@ public:
 	 * @param InTranslation The value to use for the translation component
 	 * @param InScale3D The value to use for the scale component
 	 */
-	FORCEINLINE FTransform(const VectorRegister	& InRotation, const VectorRegister& InTranslation, const VectorRegister& InScale3D) 
+	FORCEINLINE FTransform(const VectorRegister4Float& InRotation, const VectorRegister& InTranslation, const VectorRegister& InScale3D) 
 		: Rotation(InRotation),
 		Translation(InTranslation),
 		Scale3D(InScale3D)
@@ -432,7 +432,8 @@ public:
 			Translation = FMath::Lerp(Atom1.Translation, Atom2.Translation, BlendWeight.Value);			
 			Scale3D = FMath::Lerp(Atom1.Scale3D, Atom2.Scale3D, BlendWeight.Value);			
 
-			VectorRegister VRotation = VectorLerpQuat(Atom1.Rotation, Atom2.Rotation, BlendWeight.Value);
+			VectorRegister4Float BlendWeightFloat = VectorLoadFloat1(&Alpha);
+			VectorRegister4Float VRotation = VectorLerpQuat(Atom1.Rotation, Atom2.Rotation, BlendWeightFloat);
 
 			// ..and renormalize
 			Rotation = VectorNormalizeQuaternion(VRotation);
@@ -465,7 +466,8 @@ public:
 				
 				Scale3D = FMath::Lerp(Scale3D, OtherAtom.Scale3D, BlendWeight.Value);
 				
-				VectorRegister VRotation = VectorLerpQuat(Rotation, OtherAtom.Rotation, BlendWeight.Value);
+				VectorRegister4Float BlendWeightFloat = VectorLoadFloat1(&Alpha);
+				VectorRegister4Float VRotation = VectorLerpQuat(Rotation, OtherAtom.Rotation, BlendWeightFloat);
 				
 				// ..and renormalize
 				Rotation = VectorNormalizeQuaternion(VRotation);
@@ -497,13 +499,13 @@ public:
 
 	FORCEINLINE FTransform operator*(const ScalarRegister& Mult) const
 	{		
-		return FTransform( VectorMultiply(Rotation, Mult), VectorMultiply(Translation, Mult), VectorMultiply(Scale3D, Mult) );
+		return FTransform( VectorMultiply(Rotation, MakeVectorRegisterFloatFromDouble(Mult.Value)), VectorMultiply(Translation, Mult), VectorMultiply(Scale3D, Mult) );
 	}
 
 	FORCEINLINE FTransform& operator*=(const ScalarRegister& Mult)
 	{			
 		Translation= VectorMultiply(Translation, Mult);
-		Rotation = VectorMultiply(Rotation, Mult);
+		Rotation = VectorMultiply(Rotation, MakeVectorRegisterFloatFromDouble(Mult.Value));
 		Scale3D = VectorMultiply(Scale3D, Mult);
 
 		return *this;
@@ -644,15 +646,8 @@ public:
 	// Serializer.
 	inline friend FArchive& operator<<(FArchive& Ar,FTransform& M)
 	{
-		//@TODO: fix this mess when FQuat finally are double and serialization does the right thing
-		// right now we cast rotation to 4 floats and the rest to 4 float/double depending on the LWC switch
-#if UE_LARGE_WORLD_COORDINATES_DISABLED
+		//@LWC_TODO: fix this when FQuat finally are double 
 		Ar << *reinterpret_cast<FVector4*>(&(M.Rotation));
-#else
-		VectorRegister4Float SerializedRotation = MakeVectorRegisterFloatFromDouble(M.Rotation);
-		Ar << *reinterpret_cast<FVector4*>(&(SerializedRotation));
-		M.Rotation = MakeVectorRegisterDouble(SerializedRotation);
-#endif
 		Ar << *reinterpret_cast<FVector*>(&(M.Translation));
 		Ar << *reinterpret_cast<FVector*>(&(M.Scale3D));
 		
@@ -790,7 +785,7 @@ public:
 	FORCEINLINE void SetIdentity()
 	{
 		// Rotation = {0,0,0,1)
-		Rotation = GlobalVectorConstants::Vector0001; //VectorSet_W1( VectorZero() );
+		Rotation = GlobalVectorConstants::Float0001; //VectorSet_W1( VectorZero() );
 		// Translation = {0,0,0,0)
 		Translation = VectorZero();
 		// Scale3D = {1,1,1,0);
@@ -930,11 +925,11 @@ public:
 	*/
 	FORCEINLINE void Accumulate(const FTransform& SourceAtom)
 	{
-		const VectorRegister BlendedRotation = SourceAtom.Rotation;
-		const VectorRegister RotationW = VectorReplicate(BlendedRotation, 3);
+		const VectorRegister4Float BlendedRotation = SourceAtom.Rotation;
+		const VectorRegister4Float RotationW = VectorReplicate(BlendedRotation, 3);
 
 		// if( Square(SourceAtom.Rotation.W) < 1.f - DELTA * DELTA )
-		if (VectorAnyGreaterThan(GlobalVectorConstants::RotationSignificantThreshold, VectorMultiply(RotationW, RotationW)))
+		if (VectorAnyGreaterThan(GlobalVectorConstants::RotationSignificantThresholdFloat, VectorMultiply(RotationW, RotationW)))
 		{
 			// Rotation = SourceAtom.Rotation * Rotation;
 			Rotation = VectorQuaternionMultiply2(BlendedRotation, Rotation);
@@ -966,15 +961,16 @@ public:
 	FORCEINLINE void Accumulate(const FTransform& Atom, const ScalarRegister& BlendWeight)
 	{
 		// SourceAtom = Atom * BlendWeight;
-		const VectorRegister BlendedRotation = VectorMultiply(Atom.Rotation, BlendWeight.Value);
+		VectorRegister4Float BlendWeightFloat = MakeVectorRegisterFloatFromDouble(BlendWeight.Value);
+		const VectorRegister4Float BlendedRotation = VectorMultiply(Atom.Rotation, BlendWeightFloat);
 		const VectorRegister BlendedTranslation = VectorMultiply(Atom.Translation, BlendWeight.Value);
 		const VectorRegister BlendedScale = VectorMultiply(Atom.Scale3D, BlendWeight.Value);
 
-		const VectorRegister RotationW = VectorReplicate(BlendedRotation, 3);
+		const VectorRegister4Float RotationW = VectorReplicate(BlendedRotation, 3);
 
 		// Add ref pose relative animation to base animation, only if rotation is significant.
 		// if( Square(SourceAtom.Rotation.W) < 1.f - DELTA * DELTA )
-		if (VectorAnyGreaterThan(GlobalVectorConstants::RotationSignificantThreshold, VectorMultiply(RotationW, RotationW)))
+		if (VectorAnyGreaterThan(GlobalVectorConstants::RotationSignificantThresholdFloat, VectorMultiply(RotationW, RotationW)))
 		{
 			// Rotation = SourceAtom.Rotation * Rotation;
 			Rotation = VectorQuaternionMultiply2(BlendedRotation, Rotation);
@@ -999,7 +995,7 @@ public:
 	 */
 	FORCEINLINE void AccumulateWithShortestRotation(const FTransform& DeltaAtom, const ScalarRegister& BlendWeight)
 	{
-		const VectorRegister BlendedRotation = VectorMultiply(DeltaAtom.Rotation, BlendWeight.Value);
+		const VectorRegister4Float BlendedRotation = VectorMultiply(DeltaAtom.Rotation, MakeVectorRegisterFloatFromDouble(BlendWeight.Value));
 
 		Rotation = VectorAccumulateQuaternionShortestPath(Rotation, BlendedRotation);
 
@@ -1032,7 +1028,7 @@ public:
 		const VectorRegister DefaultScale = MakeVectorRegister(1.f, 1.f, 1.f, 0.f);
 
 		// SourceAtom = Atom * BlendWeight;
-		const VectorRegister BlendedRotation = VectorMultiply(Atom.Rotation, BlendWeight.Value);
+		const VectorRegister4Float BlendedRotation = VectorMultiply(Atom.Rotation, MakeVectorRegisterFloatFromDouble(BlendWeight.Value));
 		const VectorRegister BlendedScale = VectorMultiply(Atom.Scale3D, BlendWeight.Value);
 		const VectorRegister BlendedTranslation = VectorMultiply(Atom.Translation, BlendWeight.Value);
 
@@ -1040,7 +1036,7 @@ public:
 
 		// Add ref pose relative animation to base animation, only if rotation is significant.
 		// if( Square(SourceAtom.Rotation.W) < 1.f - DELTA * DELTA )
-		if (VectorAnyGreaterThan(GlobalVectorConstants::RotationSignificantThreshold, VectorMultiply(RotationW, RotationW)))
+		if (VectorAnyGreaterThan(GlobalVectorConstants::RotationSignificantThresholdFloat, VectorMultiply(RotationW, RotationW)))
 		{
 			// Rotation = SourceAtom.Rotation * Rotation;
 			Rotation = VectorQuaternionMultiply2(BlendedRotation, Rotation);
@@ -1107,9 +1103,10 @@ public:
 	 */
 	FORCEINLINE static void BlendFromIdentityAndAccumulate(FTransform& FinalAtom, const FTransform& SourceAtom, const ScalarRegister& BlendWeight)
 	{
-		const VectorRegister Const0001 = GlobalVectorConstants::Vector0001;
-		const VectorRegister ConstNegative0001 = VectorSubtract(VectorZero(), Const0001);
-		const VectorRegister VOneMinusAlpha = VectorSubtract(VectorOne(), BlendWeight.Value);
+		const VectorRegister4Float BlendWeightFloat = MakeVectorRegisterFloatFromDouble(BlendWeight.Value);
+		const VectorRegister4Float Const0001 = GlobalVectorConstants::Float0001;
+		const VectorRegister4Float ConstNegative0001 = VectorSubtract(VectorZeroFloat(), Const0001);
+		const VectorRegister4Float VOneMinusAlpha = VectorSubtract(VectorOneFloat(), BlendWeightFloat);
 		const VectorRegister DefaultScale = MakeVectorRegister(1.f, 1.f, 1.f, 0.f);
 
 		// Blend rotation
@@ -1118,15 +1115,15 @@ public:
 		//     BlendedAtom.Rotation = (B * Alpha) + (A * (Bias * (1.f - Alpha)));
 		//     BlendedAtom.Rotation.QuaternionNormalize();
 		//  Note: A = (0,0,0,1), which simplifies things a lot; only care about sign of B.W now, instead of doing a dot product
-		const VectorRegister RotationB = SourceAtom.Rotation;
+		const VectorRegister4Float RotationB = SourceAtom.Rotation;
 
-		const VectorRegister QuatRotationDirMask = VectorCompareGE(RotationB, VectorZero());
-		const VectorRegister BiasTimesA = VectorSelect(QuatRotationDirMask, Const0001, ConstNegative0001);
-		const VectorRegister RotateBTimesWeight = VectorMultiply(RotationB, BlendWeight.Value);
-		const VectorRegister UnnormalizedRotation = VectorMultiplyAdd(BiasTimesA, VOneMinusAlpha, RotateBTimesWeight);
+		const VectorRegister4Float QuatRotationDirMask = VectorCompareGE(RotationB, VectorZero());
+		const VectorRegister4Float BiasTimesA = VectorSelect(QuatRotationDirMask, Const0001, ConstNegative0001);
+		const VectorRegister4Float RotateBTimesWeight = VectorMultiply(RotationB, BlendWeightFloat);
+		const VectorRegister4Float UnnormalizedRotation = VectorMultiplyAdd(BiasTimesA, VOneMinusAlpha, RotateBTimesWeight);
 
 		// Normalize blended rotation ( result = (Q.Q >= 1e-8) ? (Q / |Q|) : (0,0,0,1) )
-		const VectorRegister BlendedRotation = VectorNormalizeSafe(UnnormalizedRotation, Const0001);
+		const VectorRegister4Float BlendedRotation = VectorNormalizeSafe(UnnormalizedRotation, GlobalVectorConstants::Float0001);
 
 		// FinalAtom.Rotation = BlendedAtom.Rotation * FinalAtom.Rotation;
 		FinalAtom.Rotation = VectorQuaternionMultiply2(BlendedRotation, FinalAtom.Rotation);
@@ -1382,7 +1379,7 @@ private:
 		const VectorRegister t2 = VectorQuaternionRotateVector(InvRotation, ScaledTranslation);
 		const VectorRegister InvTranslation = VectorSet_W0(VectorNegate(t2));
 
-		return FTransform(InvRotation, InvTranslation, InvScale);
+		return FTransform(MakeVectorRegisterFloatFromDouble(InvRotation), InvTranslation, InvScale);
 	}
 
 	/**
@@ -1521,8 +1518,8 @@ FORCEINLINE void FTransform::Multiply(FTransform* OutTransform, const FTransform
 	}
 	else
 	{
-		const VectorRegister QuatA = A->Rotation;
-		const VectorRegister QuatB = B->Rotation;
+		const VectorRegister4Float QuatA = A->Rotation;
+		const VectorRegister4Float QuatB = B->Rotation;
 		const VectorRegister TranslateA = A->Translation;
 		const VectorRegister TranslateB = B->Translation;
 		const VectorRegister ScaleA = A->Scale3D;
diff --git a/Engine/Source/Runtime/Core/Public/Math/VectorRegister.h b/Engine/Source/Runtime/Core/Public/Math/VectorRegister.h
--- a/Engine/Source/Runtime/Core/Public/Math/VectorRegister.h
+++ b/Engine/Source/Runtime/Core/Public/Math/VectorRegister.h
@@ -46,5 +46,7 @@ namespace GlobalVectorConstants
 {
 	static const VectorRegister AnimWeightThreshold = MakeVectorRegister(ZERO_ANIMWEIGHT_THRESH, ZERO_ANIMWEIGHT_THRESH, ZERO_ANIMWEIGHT_THRESH, ZERO_ANIMWEIGHT_THRESH);
 	static const VectorRegister RotationSignificantThreshold = MakeVectorRegister(1.0f - DELTA*DELTA, 1.0f - DELTA*DELTA, 1.0f - DELTA*DELTA, 1.0f - DELTA*DELTA);
+	// LWC_TODO : only needed until FQuat is using doubles (see TransformVectorized)
+	static const VectorRegister4Float RotationSignificantThresholdFloat = MakeVectorRegisterFloat(1.0f - DELTA * DELTA, 1.0f - DELTA * DELTA, 1.0f - DELTA * DELTA, 1.0f - DELTA * DELTA);
 }
 
