Undo //UE5/Main/Engine/Source/Runtime/Core/Private/Containers/... changelist 16375887

diff --git a/Engine/Source/Runtime/Core/Private/Containers/String.cpp b/Engine/Source/Runtime/Core/Private/Containers/String.cpp
--- a/Engine/Source/Runtime/Core/Private/Containers/String.cpp
+++ b/Engine/Source/Runtime/Core/Private/Containers/String.cpp
@@ -788,7 +788,10 @@ bool FString::ToHexBlob( const FString& Source, uint8* DestBuffer, const uint32
 FString FString::SanitizeFloat( double InFloat, const int32 InMinFractionalDigits )
 {
 	// Avoids negative zero
-	InFloat = InFloat + 0.0f;
+	if( InFloat == 0 )
+	{
+		InFloat = 0;
+	}
 
 	// First create the string
 	FString TempString = FString::Printf(TEXT("%f"), InFloat);
