Removed redundant include statement.

#rnx

diff --git a/Engine/Source/Runtime/Core/Private/Serialization/MemoryImage.cpp b/Engine/Source/Runtime/Core/Private/Serialization/MemoryImage.cpp
--- a/Engine/Source/Runtime/Core/Private/Serialization/MemoryImage.cpp
+++ b/Engine/Source/Runtime/Core/Private/Serialization/MemoryImage.cpp
@@ -12,7 +12,6 @@
 #include "Misc/ScopeRWLock.h"
 #include "Misc/DataDrivenPlatformInfoRegistry.h"
 #include "Serialization/Archive.h"
-#include "Misc/MemStack.h"
 #include "Async/ParallelFor.h"
 
 DEFINE_LOG_CATEGORY_STATIC(LogMemoryImage, Log, All);
