Guarantee the lifetime of FPageAllocator is long enough for all objects that use it. A FMemStackBase frees its allocs back to the page allocator in its destructor. The page allocator singleton is lazily created the first call to FMemStackBase::Alloc(). Because any number of other global FMemStackBase objects could have been constructed prior to the first call Alloc(), the page allocator will be destructed before these global FMemStackBase objects.

#rnx

diff --git a/Engine/Source/Runtime/Core/Private/Misc/MemStack.cpp b/Engine/Source/Runtime/Core/Private/Misc/MemStack.cpp
--- a/Engine/Source/Runtime/Core/Private/Misc/MemStack.cpp
+++ b/Engine/Source/Runtime/Core/Private/Misc/MemStack.cpp
@@ -12,19 +12,32 @@ DECLARE_MEMORY_STAT(TEXT("MemStack Large Block"), STAT_MemStackLargeBLock,STATGR
 DECLARE_MEMORY_STAT(TEXT("PageAllocator Free"), STAT_PageAllocatorFree, STATGROUP_Memory);
 DECLARE_MEMORY_STAT(TEXT("PageAllocator Used"), STAT_PageAllocatorUsed, STATGROUP_Memory);
 
-FPageAllocator& FPageAllocator::Get()
+FPageAllocator* FPageAllocator::Instance = nullptr;
+
+FPageAllocator::FPageAllocator()
 {
-	static std::atomic<FPageAllocator*> ThePageAllocator;
+	Instance = this;
+}
 
-	FPageAllocator* LocalPageAllocator = ThePageAllocator.load(std::memory_order_acquire);
-	if (!LocalPageAllocator)
+FPageAllocator::~FPageAllocator()
+{
+	Instance = nullptr;
+}
+
+FPageAllocator& FPageAllocator::Get()
+{
+	if (LIKELY(Instance != nullptr))
 	{
-		static FPageAllocator Instance;
-		LocalPageAllocator = &Instance;
-		ThePageAllocator.store(LocalPageAllocator, std::memory_order_release);
+		return *Instance;
 	}
 
-	return *LocalPageAllocator;
+	return Construct();
+}
+
+FORCENOINLINE FPageAllocator& FPageAllocator::Construct()
+{
+	static FPageAllocator PageAllocator;
+	return *Instance;
 }
 
 #define USE_MEMSTACK_PUGATORY (0)
@@ -284,6 +297,19 @@ void FPageAllocator::UpdateStats()
 	FMemStack implementation.
 -----------------------------------------------------------------------------*/
 
+FMemStackBase::FMemStackBase()
+: Top(nullptr)
+, End(nullptr)
+, TopChunk(nullptr)
+, TopMark(nullptr)
+, NumMarks(0)
+, bShouldEnforceAllocMarks(false)
+{
+	// By fetching the FPageAllocator singleton's address here we can guarantee
+	// that its lifetime spans that of all FMemStackBase objects.
+	FPageAllocator::Get();
+}
+
 int32 FMemStackBase::GetByteCount() const
 {
 	int32 Count = 0;
diff --git a/Engine/Source/Runtime/Core/Public/Misc/MemStack.h b/Engine/Source/Runtime/Core/Public/Misc/MemStack.h
--- a/Engine/Source/Runtime/Core/Public/Misc/MemStack.h
+++ b/Engine/Source/Runtime/Core/Public/Misc/MemStack.h
@@ -43,6 +43,7 @@ public:
 
 	static FPageAllocator& Get();
 
+	~FPageAllocator();
 	void* Alloc();
 	void Free(void* Mem);
 	void* AllocSmall();
@@ -52,7 +53,9 @@ public:
 	void LatchProtectedMode();
 private:
 
-	FPageAllocator() {}
+	static FPageAllocator* Instance;
+	static FPageAllocator& Construct();
+	FPageAllocator();
 
 #if STATS
 	void UpdateStats();
@@ -70,16 +73,7 @@ class CORE_API FMemStackBase //-V1062
 {
 public:
 
-	FMemStackBase()
-		: Top(nullptr)
-		, End(nullptr)
-		, TopChunk(nullptr)
-		, TopMark(nullptr)
-		, NumMarks(0)
-		, bShouldEnforceAllocMarks(false)
-	{
-		
-	}
+	FMemStackBase();
 
 	FMemStackBase(const FMemStackBase&) = delete;
 	FMemStackBase(FMemStackBase&& Other)
