CompactBinary: Switched loop variables to references to avoid compile warnings

#rnx

diff --git a/Engine/Source/Runtime/Core/Private/Serialization/CompactBinary.cpp b/Engine/Source/Runtime/Core/Private/Serialization/CompactBinary.cpp
--- a/Engine/Source/Runtime/Core/Private/Serialization/CompactBinary.cpp
+++ b/Engine/Source/Runtime/Core/Private/Serialization/CompactBinary.cpp
@@ -710,7 +710,7 @@ FCbObjectView::FCbObjectView()
 
 FCbFieldView FCbObjectView::FindView(const FAnsiStringView Name) const
 {
-	for (const FCbFieldView Field : *this)
+	for (const FCbFieldView& Field : *this)
 	{
 		if (Name.Equals(Field.GetName(), ESearchCase::CaseSensitive))
 		{
@@ -722,7 +722,7 @@ FCbFieldView FCbObjectView::FindView(const FAnsiStringView Name) const
 
 FCbFieldView FCbObjectView::FindViewIgnoreCase(const FAnsiStringView Name) const
 {
-	for (const FCbFieldView Field : *this)
+	for (const FCbFieldView& Field : *this)
 	{
 		if (Name.Equals(Field.GetName(), ESearchCase::IgnoreCase))
 		{
