A better solution for handling negative zero in FString::SanitizeFloat(), as it was producing "0.0" when passed a NaN when built under /fp:fast.

diff --git a/Engine/Source/Runtime/Core/Private/Containers/String.cpp b/Engine/Source/Runtime/Core/Private/Containers/String.cpp
--- a/Engine/Source/Runtime/Core/Private/Containers/String.cpp
+++ b/Engine/Source/Runtime/Core/Private/Containers/String.cpp
@@ -785,13 +785,20 @@ bool FString::ToHexBlob( const FString& Source, uint8* DestBuffer, const uint32
 	return false;
 }
 
+PRAGMA_DISABLE_OPTIMIZATION
+void StripNegativeZero(double& InFloat)
+{
+	// This works for translating a negative zero into a positive zero,
+	// but if optimizations are enabled when compiling with -ffast-math
+	// or /fp:fast, the compiler can strip it out.
+	InFloat += 0.0f;
+}
+PRAGMA_ENABLE_OPTIMIZATION
+
 FString FString::SanitizeFloat( double InFloat, const int32 InMinFractionalDigits )
 {
 	// Avoids negative zero
-	if( InFloat == 0 )
-	{
-		InFloat = 0;
-	}
+	StripNegativeZero(InFloat);
 
 	// First create the string
 	FString TempString = FString::Printf(TEXT("%f"), InFloat);
