detected that the test for movable-only task result is unreliable, disabled but left there because it's still can be useful for debugging

diff --git a/Engine/Source/Runtime/Core/Private/Tests/Tasks/TasksTest.cpp b/Engine/Source/Runtime/Core/Private/Tests/Tasks/TasksTest.cpp
--- a/Engine/Source/Runtime/Core/Private/Tests/Tasks/TasksTest.cpp
+++ b/Engine/Source/Runtime/Core/Private/Tests/Tasks/TasksTest.cpp
@@ -88,7 +88,7 @@ namespace UE { namespace TasksTests
 	{
 		if (!FPlatformProcess::SupportsMultithreading())
 		{
-			// the new API doesn't support single-threaded execution (`-nothreading`) until it's feature-compatible with the new API and completely replaces it
+			// the new API doesn't support single-threaded execution (`-nothreading`) until it's feature-compatible with the old API and completely replaces it
 			return true;
 		}
 
@@ -168,7 +168,7 @@ namespace UE { namespace TasksTests
 			{
 				FORCENOINLINE FMoveConstructable()
 				{
-					++ConstructionsNum;
+					ConstructionsNum.fetch_add(1, std::memory_order_relaxed);
 				}
 
 				FMoveConstructable(FMoveConstructable&&)
@@ -182,14 +182,37 @@ namespace UE { namespace TasksTests
 
 				FORCENOINLINE ~FMoveConstructable()
 				{
-					++DestructionsNum;
+					DestructionsNum.fetch_add(1, std::memory_order_relaxed);
 				}
 			};
 
-			Launch(UE_SOURCE_LOCATION, [] { return FMoveConstructable{}; }).GetResult(); // consume the result
-				
-			checkf(ConstructionsNum == 1, TEXT("%d result instances were created but one was expected: the value stored in the task"));
-			checkf(ConstructionsNum == DestructionsNum, TEXT("Mismatched number of constructions (%d) and destructions (%d)"));
+			{
+				Launch(UE_SOURCE_LOCATION, [] { return FMoveConstructable{}; }).GetResult();
+			}
+
+#if 0	// unreliable test, destruction can happen on a worker thread, after the task is flagged as completed and so the check can be hit before the destruction
+			uint32 LocalConstructionsNum = ConstructionsNum.load(std::memory_order_relaxed);
+			uint32 LocalDestructionsNum = DestructionsNum.load(std::memory_order_relaxed);
+			checkf(LocalConstructionsNum == 1, TEXT("%d result instances were created but one was expected: the value stored in the task"), LocalConstructionsNum);
+			checkf(LocalConstructionsNum == LocalDestructionsNum, TEXT("Mismatched number of constructions (%d) and destructions (%d)"), LocalConstructionsNum, LocalDestructionsNum);
+
+			ConstructionsNum = 0;
+			DestructionsNum = 0;
+#endif
+
+			{
+				FMoveConstructable Res{ MoveTemp(Launch(UE_SOURCE_LOCATION, [] { return FMoveConstructable{}; }).GetResult()) }; // consume the result
+			}
+
+#if 0	// unreliable test, destruction can happen on a worker thread, after the task is flagged as completed and so the check can be hit before the destruction
+			LocalConstructionsNum = ConstructionsNum.load(std::memory_order_relaxed);
+			LocalDestructionsNum = DestructionsNum.load(std::memory_order_relaxed);
+			checkf(LocalConstructionsNum == 2, TEXT("%d result instances were created but 2 was expected: the value stored in the task"), LocalConstructionsNum);
+			checkf(LocalConstructionsNum == LocalDestructionsNum, TEXT("Mismatched number of constructions (%d) and destructions (%d)"), LocalConstructionsNum, LocalDestructionsNum);
+
+			ConstructionsNum = 0;
+			DestructionsNum = 0;
+#endif
 		}
 
 		// fire and forget: launch a task w/o keeping its reference
