Changed FApp::GetInstanceId() to use a local static declaration (static magic) rather than a class member to avoir static order initialization issues when calling it very early (basically in engine pre-init code).

diff --git a/Engine/Source/Runtime/Core/Private/Misc/App.cpp b/Engine/Source/Runtime/Core/Private/Misc/App.cpp
--- a/Engine/Source/Runtime/Core/Private/Misc/App.cpp
+++ b/Engine/Source/Runtime/Core/Private/Misc/App.cpp
@@ -20,7 +20,6 @@ DEFINE_LOG_CATEGORY_STATIC(LogApp, Log, All);
 bool FApp::bIsDebugGame = false;
 #endif
 
-FGuid FApp::InstanceId = FGuid::NewGuid();
 FGuid FApp::SessionId = FGuid::NewGuid();
 FString FApp::SessionName = FString();
 FString FApp::SessionOwner = FString();
@@ -130,6 +129,7 @@ void FApp::InitializeSession()
 {
 	// parse session details on command line
 	FString InstanceIdString;
+	FGuid InstanceId = GetInstanceId();
 
 	if (FParse::Value(FCommandLine::Get(), TEXT("-InstanceId="), InstanceIdString))
 	{
@@ -139,11 +139,6 @@ void FApp::InitializeSession()
 		}
 	}
 
-	if (!InstanceId.IsValid())
-	{
-		InstanceId = FGuid::NewGuid();
-	}
-
 	FString SessionIdString;
 
 	if (FParse::Value(FCommandLine::Get(), TEXT("-SessionId="), SessionIdString))
diff --git a/Engine/Source/Runtime/Core/Public/Misc/App.h b/Engine/Source/Runtime/Core/Public/Misc/App.h
--- a/Engine/Source/Runtime/Core/Public/Misc/App.h
+++ b/Engine/Source/Runtime/Core/Public/Misc/App.h
@@ -222,6 +222,7 @@ public:
 	 */
 	FORCEINLINE static FGuid GetInstanceId()
 	{
+		static FGuid InstanceId = FGuid::NewGuid();
 		return InstanceId;
 	}
 
@@ -316,7 +317,7 @@ public:
 	 */
 	FORCEINLINE static bool IsThisInstance(const FGuid& InInstanceId)
 	{
-		return (InInstanceId == InstanceId);
+		return (InInstanceId == GetInstanceId());
 	};
 
 	/**
@@ -742,9 +743,6 @@ private:
 	static bool bIsDebugGame;
 #endif
 
-	/** Holds the instance identifier. */
-	static FGuid InstanceId;
-
 	/** Holds the session identifier. */
 	static FGuid SessionId;
 
