minor refactoring

diff --git a/Engine/Source/Runtime/Core/Private/Tests/Tasks/TasksTest.cpp b/Engine/Source/Runtime/Core/Private/Tests/Tasks/TasksTest.cpp
--- a/Engine/Source/Runtime/Core/Private/Tests/Tasks/TasksTest.cpp
+++ b/Engine/Source/Runtime/Core/Private/Tests/Tasks/TasksTest.cpp
@@ -66,20 +66,9 @@ namespace UE { namespace TasksTests
 			));
 		}
 
-		for (FTask& SpawnerGroup : SpawnerGroups)
-		{
-			SpawnerGroup.Wait();
-		}
-
-		for (FTask& Spawner : Spawners)
-		{
-			Spawner.Wait();
-		}
-
-		for (FTask& Task : Tasks)
-		{
-			Task.Wait();
-		}
+		Wait(SpawnerGroups);
+		Wait(Spawners);
+		Wait(Tasks);
 
 		check(TasksExecutedNum == TasksNum);
 	}
@@ -419,20 +408,9 @@ namespace UE { namespace TasksTests
 				));
 		}
 
-		for (FTask& SpawnerGroup : SpawnerGroups)
-		{
-			SpawnerGroup.Wait();
-		}
-
-		for (FTask& Spawner : Spawners)
-		{
-			Spawner.Wait();
-		}
-
-		for (FTask& Task : Tasks)
-		{
-			Task.Wait();
-		}
+		Wait(SpawnerGroups);
+		Wait(Spawners);
+		Wait(Tasks);
 
 		check(TasksExecutedNum == TasksNum);
 	}
