Added static Null members to FCompressedBuffer and FCompositeBuffer

#rnx

diff --git a/Engine/Source/Runtime/Core/Public/Compression/CompressedBuffer.h b/Engine/Source/Runtime/Core/Public/Compression/CompressedBuffer.h
--- a/Engine/Source/Runtime/Core/Public/Compression/CompressedBuffer.h
+++ b/Engine/Source/Runtime/Core/Public/Compression/CompressedBuffer.h
@@ -98,8 +98,13 @@ public:
 	 */
 	[[nodiscard]] CORE_API FCompositeBuffer DecompressToComposite() const;
 
+	/** A null compressed buffer. */
+	static const FCompressedBuffer Null;
+
 private:
 	FCompositeBuffer CompressedData;
 };
 
+inline const FCompressedBuffer FCompressedBuffer::Null;
+
 CORE_API FArchive& operator<<(FArchive& Ar, FCompressedBuffer& Buffer);
diff --git a/Engine/Source/Runtime/Core/Public/Memory/CompositeBuffer.h b/Engine/Source/Runtime/Core/Public/Memory/CompositeBuffer.h
--- a/Engine/Source/Runtime/Core/Public/Memory/CompositeBuffer.h
+++ b/Engine/Source/Runtime/Core/Public/Memory/CompositeBuffer.h
@@ -96,6 +96,9 @@ public:
 	CORE_API void IterateRange(uint64 Offset, uint64 Size,
 		TFunctionRef<void (FMemoryView View, const FSharedBuffer& ViewOuter)> Visitor) const;
 
+	/** A null composite buffer. */
+	static const FCompositeBuffer Null;
+
 private:
 	static inline int32 GetBufferCount(const FCompositeBuffer& Buffer) { return Buffer.Segments.Num(); }
 	inline void AppendBuffers(const FCompositeBuffer& Buffer) { Segments.Append(Buffer.Segments); }
@@ -113,3 +116,5 @@ private:
 private:
 	TArray<FSharedBuffer, TInlineAllocator<1>> Segments;
 };
+
+inline const FCompositeBuffer FCompositeBuffer::Null;
