#ue5 - LWC: Moved VectorReplicate/VectorSwizzle/VectorShuffle tests to templates to avoid code duplication. Added missing VectorGetComponent() tests for floats.

diff --git a/Engine/Source/Runtime/Core/Private/Tests/Math/UnrealMathTest.cpp b/Engine/Source/Runtime/Core/Private/Tests/Math/UnrealMathTest.cpp
--- a/Engine/Source/Runtime/Core/Private/Tests/Math/UnrealMathTest.cpp
+++ b/Engine/Source/Runtime/Core/Private/Tests/Math/UnrealMathTest.cpp
@@ -835,6 +835,150 @@ void SetScratchDouble(float X, float Y, float Z, float W, float U = 0.0f)
 }
 
 
+
+template<typename RealType, typename VectorRegisterType>
+FORCENOINLINE void TestVectorReplicate()
+{
+	const RealType ArrayV0[4] = { 0.0, 1.0, 2.0, 3.0 };
+	VectorRegisterType V0 = VectorLoad(ArrayV0);
+	VectorRegisterType V1, V2;
+
+#define ReplicateTest(A, X) \
+		V1 = VectorReplicate(A, X); \
+		V2 = MakeVectorRegister(ArrayV0[X], ArrayV0[X], ArrayV0[X], ArrayV0[X]); \
+		LogTest(*FString::Printf(TEXT("VectorReplicate<%d>"), X), TestVectorsEqual(V1, V2));
+
+	ReplicateTest(V0, 0);
+	ReplicateTest(V0, 1);
+	ReplicateTest(V0, 2);
+	ReplicateTest(V0, 3);
+
+#undef ReplicateTest
+}
+
+
+template<typename RealType, typename VectorRegisterType>
+FORCENOINLINE void TestVectorSwizzle()
+{
+	const RealType ArrayV0[4] = { 0.0, 1.0, 2.0, 3.0 };
+	VectorRegisterType V0 = VectorLoad(ArrayV0);
+	VectorRegisterType V1, V2;
+
+#define SwizzleTest(A, X, Y, Z, W) \
+		V1 = VectorSwizzle(A, X, Y, Z, W); \
+		V2 = MakeVectorRegister(ArrayV0[X], ArrayV0[Y], ArrayV0[Z], ArrayV0[W]); \
+		LogTest(*FString::Printf(TEXT("VectorSwizzle<%d,%d,%d,%d>"), X, Y, Z, W), TestVectorsEqual(V1, V2));
+
+	// This is not an exhaustive list because it would be 4*4*4*4 = 256 entries, but it tries to test a lot of common permutations.
+	// Unfortunately it can't be done in a loop because it uses a #define and compile-time constants for the VectorSwizzle() 'function'.
+	// Many of these were selected to also stress the specializations in certain implementations.
+
+	SwizzleTest(V0, 0, 1, 2, 3); // Identity
+	SwizzleTest(V0, 0, 0, 0, 0); // Replicate 0
+	SwizzleTest(V0, 1, 1, 1, 1); // Replicate 1
+	SwizzleTest(V0, 2, 2, 2, 2); // Replicate 2
+	SwizzleTest(V0, 3, 3, 3, 3); // Replicate 3
+
+	SwizzleTest(V0, 0, 3, 1, 2);
+	SwizzleTest(V0, 0, 2, 0, 2);
+
+	SwizzleTest(V0, 1, 0, 1, 0);
+	SwizzleTest(V0, 1, 0, 3, 2);
+	SwizzleTest(V0, 1, 2, 0, 1);
+
+	SwizzleTest(V0, 2, 0, 1, 3);
+	SwizzleTest(V0, 2, 3, 0, 1);
+
+	SwizzleTest(V0, 3, 2, 1, 0);
+	SwizzleTest(V0, 3, 0, 3, 0);
+
+	SwizzleTest(V0, 2, 2, 0, 1);
+	SwizzleTest(V0, 3, 3, 1, 0);
+	SwizzleTest(V0, 2, 2, 0, 2);
+	SwizzleTest(V0, 3, 3, 1, 3);
+
+	// Common specializations
+	SwizzleTest(V0, 0, 1, 2, 2);
+	SwizzleTest(V0, 0, 1, 3, 3);
+	SwizzleTest(V0, 0, 1, 3, 2);
+
+	SwizzleTest(V0, 0, 0, 2, 3);
+	SwizzleTest(V0, 0, 0, 2, 2);
+	SwizzleTest(V0, 0, 0, 3, 3);
+	SwizzleTest(V0, 0, 0, 3, 2);
+
+	SwizzleTest(V0, 1, 0, 2, 3);
+	SwizzleTest(V0, 1, 0, 2, 2);
+	SwizzleTest(V0, 1, 0, 3, 3);
+	SwizzleTest(V0, 1, 0, 3, 2);
+
+	SwizzleTest(V0, 1, 1, 2, 3);
+	SwizzleTest(V0, 1, 1, 2, 2);
+	SwizzleTest(V0, 1, 1, 3, 3);
+	SwizzleTest(V0, 1, 1, 3, 2);
+
+	SwizzleTest(V0, 0, 1, 0, 1);
+	SwizzleTest(V0, 2, 3, 2, 3);
+	SwizzleTest(V0, 0, 0, 1, 1);
+	SwizzleTest(V0, 2, 2, 3, 3);
+	SwizzleTest(V0, 0, 0, 2, 2);
+	SwizzleTest(V0, 1, 1, 3, 3);
+
+#undef SwizzleTest
+}
+
+template<typename RealType, typename VectorRegisterType>
+FORCENOINLINE void TestVectorShuffle()
+{
+	RealType ArrayV0[4] = { 0.0, 0.1, 0.2, 0.3 };
+	RealType ArrayV1[4] = { 1.0, 1.1, 1.2, 1.3 };
+	VectorRegisterType V0 = VectorLoad(ArrayV0);
+	VectorRegisterType V1 = VectorLoad(ArrayV1);
+	VectorRegisterType V2, V3;
+
+#define ShuffleTest(A, B, X, Y, Z, W) \
+		V2 = VectorShuffle(A, B, X, Y, Z, W); \
+		V3 = MakeVectorRegister(ArrayV0[X], ArrayV0[Y], ArrayV1[Z], ArrayV1[W]); \
+		LogTest(*FString::Printf(TEXT("VectorShuffle<%d,%d,%d,%d>"), X, Y, Z, W), TestVectorsEqual(V2, V3));
+
+	// This is not an exhaustive list because it would be 4*4*4*4 = 256 entries, but it tries to test a lot of common permutations.
+	// Unfortunately it can't be done in a loop because it uses a #define and compile-time constants for the VectorShuffle() 'function'.
+	// Many of these were selected to also stress the specializations in certain implementations.
+	ShuffleTest(V0, V1, 0, 0, 0, 0);
+	ShuffleTest(V0, V1, 1, 1, 1, 1);
+	ShuffleTest(V0, V1, 2, 2, 2, 2);
+	ShuffleTest(V0, V1, 3, 3, 3, 3);
+
+	ShuffleTest(V0, V1, 0, 1, 2, 3);
+	ShuffleTest(V0, V1, 3, 2, 1, 0);
+
+	ShuffleTest(V0, V1, 0, 1, 0, 1);
+	ShuffleTest(V0, V1, 0, 2, 0, 2);
+	ShuffleTest(V0, V1, 0, 3, 0, 3);
+
+	ShuffleTest(V0, V1, 1, 0, 1, 0);
+	ShuffleTest(V0, V1, 2, 0, 2, 0);
+	ShuffleTest(V0, V1, 3, 0, 3, 0);
+
+	ShuffleTest(V0, V1, 0, 0, 1, 1);
+	ShuffleTest(V0, V1, 0, 0, 2, 2);
+	ShuffleTest(V0, V1, 1, 1, 3, 3);
+	ShuffleTest(V0, V1, 2, 2, 3, 3);
+	ShuffleTest(V0, V1, 2, 3, 0, 1);
+	ShuffleTest(V0, V1, 2, 3, 2, 3);
+
+	ShuffleTest(V0, V1, 3, 1, 3, 0);
+	ShuffleTest(V0, V1, 2, 2, 0, 1);
+	ShuffleTest(V0, V1, 0, 1, 3, 2);
+	ShuffleTest(V0, V1, 1, 3, 0, 3);
+
+#undef ShuffleTest
+}
+
+
+
+
+
 IMPLEMENT_SIMPLE_AUTOMATION_TEST(FVectorRegisterAbstractionTest, "System.Core.Math.Vector Register Abstraction Test", EAutomationTestFlags::ApplicationContextMask | EAutomationTestFlags::SmokeFilter)
 
 
@@ -934,6 +1078,15 @@ bool RunDoubleVectorTest()
 	V0 = MakeVectorRegisterDouble(1.0f, 5.0f, -0.25f, -0.5f);
 	LogTest(TEXT("VectorStoreFloat1<double>"), TestVectorsEqual(V0, V1));
 
+	// Replicate
+	TestVectorReplicate<double, VectorRegister4Double>();
+
+	// Swizzle
+	TestVectorSwizzle<double, VectorRegister4Double>();
+
+	// Shuffle
+	TestVectorShuffle<double, VectorRegister4Double>();
+
 	V0 = MakeVectorRegisterDouble(0., 0., 2., 0.);
 	V1 = MakeVectorRegisterDouble(0., 1., 0., 3.);
 	double D0 = VectorGetComponent(V0, 0);
@@ -944,136 +1097,6 @@ bool RunDoubleVectorTest()
 	V1 = MakeVectorRegisterDouble(D0, D1, D2, D3);
 	LogTest(TEXT("VectorGetComponent<double>"), TestVectorsEqual(V0, V1));
 
-	// Replicate
-	{
-		double ArrayV0[4] = { 0.0, 1.0, 2.0, 3.0 };
-		V0 = VectorLoad(ArrayV0);
-
-#define ReplicateTest(A, X) \
-		V2 = VectorReplicate(A, X); \
-		V3 = MakeVectorRegisterDouble(ArrayV0[X], ArrayV0[X], ArrayV0[X], ArrayV0[X]); \
-		LogTest(*FString::Printf(TEXT("VectorReplicateDouble<%d>"), X), TestVectorsEqual(V2, V3));
-
-		ReplicateTest(V0, 0);
-		ReplicateTest(V0, 1);
-		ReplicateTest(V0, 2);
-		ReplicateTest(V0, 3);
-
-#undef ReplicateTest
-	}
-
-	// Swizzle
-	{
-		double ArrayV0[4] = { 0.0, 1.0, 2.0, 3.0 };
-		V0 = VectorLoad(ArrayV0);
-
-#define SwizzleTest(A, X, Y, Z, W) \
-		V2 = VectorSwizzle(A, X, Y, Z, W); \
-		V3 = MakeVectorRegisterDouble(ArrayV0[X], ArrayV0[Y], ArrayV0[Z], ArrayV0[W]); \
-		LogTest(*FString::Printf(TEXT("VectorSwizzleDouble<%d,%d,%d,%d>"), X, Y, Z, W), TestVectorsEqual(V2, V3));
-
-		// This is not an exhaustive list because it would be 4*4*4*4 = 256 entries, but it tries to test a lot of common permutations.
-		// Unfortunately it can't be done in a loop because it uses a #define and compile-time constants for the VectorSwizzle() 'function'.
-		// Many of these were selected to also stress the specializations in certain implementations.
-
-		SwizzleTest(V0, 0, 1, 2, 3); // Identity
-		SwizzleTest(V0, 0, 0, 0, 0); // Replicate 0
-		SwizzleTest(V0, 1, 1, 1, 1); // Replicate 1
-		SwizzleTest(V0, 2, 2, 2, 2); // Replicate 2
-		SwizzleTest(V0, 3, 3, 3, 3); // Replicate 3
-
-		SwizzleTest(V0, 0, 3, 1, 2);
-		SwizzleTest(V0, 0, 2, 0, 2);
-
-		SwizzleTest(V0, 1, 0, 1, 0);
-		SwizzleTest(V0, 1, 0, 3, 2);
-		SwizzleTest(V0, 1, 2, 0, 1);
-
-		SwizzleTest(V0, 2, 0, 1, 3);
-		SwizzleTest(V0, 2, 3, 0, 1);
-
-		SwizzleTest(V0, 3, 2, 1, 0);
-		SwizzleTest(V0, 3, 0, 3, 0);
-		
-		SwizzleTest(V0, 2, 2, 0, 1);
-		SwizzleTest(V0, 3, 3, 1, 0);
-		SwizzleTest(V0, 2, 2, 0, 2);
-		SwizzleTest(V0, 3, 3, 1, 3);
-
-		// Common specializations
-		SwizzleTest(V0, 0, 1, 2, 2);
-		SwizzleTest(V0, 0, 1, 3, 3);
-		SwizzleTest(V0, 0, 1, 3, 2);
-
-		SwizzleTest(V0, 0, 0, 2, 3);
-		SwizzleTest(V0, 0, 0, 2, 2);
-		SwizzleTest(V0, 0, 0, 3, 3);
-		SwizzleTest(V0, 0, 0, 3, 2);
-
-		SwizzleTest(V0, 1, 0, 2, 3);
-		SwizzleTest(V0, 1, 0, 2, 2);
-		SwizzleTest(V0, 1, 0, 3, 3);
-		SwizzleTest(V0, 1, 0, 3, 2);
-
-		SwizzleTest(V0, 1, 1, 2, 3);
-		SwizzleTest(V0, 1, 1, 2, 2);
-		SwizzleTest(V0, 1, 1, 3, 3);
-		SwizzleTest(V0, 1, 1, 3, 2);
-
-		SwizzleTest(V0, 0, 1, 0, 1);
-		SwizzleTest(V0, 2, 3, 2, 3);
-		SwizzleTest(V0, 0, 0, 1, 1);
-		SwizzleTest(V0, 2, 2, 3, 3);
-		SwizzleTest(V0, 0, 0, 2, 2);
-		SwizzleTest(V0, 1, 1, 3, 3);
-#undef SwizzleTest
-	}
-
-	// Shuffle
-	{
-		double ArrayV0[4] = { 0.0, 1.0, 2.0, 3.0 };
-		double ArrayV1[4] = { 0.0, 1.0, 2.0, 3.0 };
-		V0 = VectorLoad(ArrayV0);
-		V1 = VectorLoad(ArrayV1);
-
-#define ShuffleTest(A, B, X, Y, Z, W) \
-		V2 = VectorShuffle(A, B, X, Y, Z, W); \
-		V3 = MakeVectorRegisterDouble(ArrayV0[X], ArrayV0[Y], ArrayV1[Z], ArrayV1[W]); \
-		LogTest(*FString::Printf(TEXT("VectorShuffleDouble<%d,%d,%d,%d>"), X, Y, Z, W), TestVectorsEqual(V2, V3));
-
-		// This is not an exhaustive list because it would be 4*4*4*4 = 256 entries, but it tries to test a lot of common permutations.
-		// Unfortunately it can't be done in a loop because it uses a #define and compile-time constants for the VectorShuffle() 'function'.
-		// Many of these were selected to also stress the specializations in certain implementations.
-		ShuffleTest(V0, V1, 0, 0, 0, 0);
-		ShuffleTest(V0, V1, 1, 1, 1, 1);
-		ShuffleTest(V0, V1, 2, 2, 2, 2);
-		ShuffleTest(V0, V1, 3, 3, 3, 3);
-
-		ShuffleTest(V0, V1, 0, 1, 2, 3);
-		ShuffleTest(V0, V1, 3, 2, 1, 0);
-
-		ShuffleTest(V0, V1, 0, 1, 0, 1);
-		ShuffleTest(V0, V1, 0, 2, 0, 2);
-		ShuffleTest(V0, V1, 0, 3, 0, 3);
-
-		ShuffleTest(V0, V1, 1, 0, 1, 0);
-		ShuffleTest(V0, V1, 2, 0, 2, 0);
-		ShuffleTest(V0, V1, 3, 0, 3, 0);
-		
-		ShuffleTest(V0, V1, 0, 0, 1, 1);
-		ShuffleTest(V0, V1, 0, 0, 2, 2);
-		ShuffleTest(V0, V1, 1, 1, 3, 3);
-		ShuffleTest(V0, V1, 2, 2, 3, 3);
-		ShuffleTest(V0, V1, 2, 3, 0, 1);
-		ShuffleTest(V0, V1, 2, 3, 2, 3);
-
-		ShuffleTest(V0, V1, 3, 1, 3, 0);
-		ShuffleTest(V0, V1, 2, 2, 0, 1);
-		ShuffleTest(V0, V1, 0, 1, 3, 2);
-		ShuffleTest(V0, V1, 1, 3, 0, 3);
-#undef ShuffleTest
-	}
-
 	// Abs
 	V0 = MakeVectorRegisterDouble(1.0f, -2.0f, 3.0f, -4.0f);
 	V1 = VectorAbs(V0);
@@ -1507,25 +1530,24 @@ bool FVectorRegisterAbstractionTest::RunTest(const FString& Parameters)
 	V0 = MakeVectorRegister( 1.0f, 5.0f, -0.25f, -0.5f );
 	LogTest( TEXT("VectorStoreFloat1"), TestVectorsEqual( V0, V1 ) );
 
-	V0 = MakeVectorRegister(1.0f, 2.0f, 3.0f, 4.0f);
-	V1 = VectorReplicate(V0, 0);
-	V0 = MakeVectorRegister(1.0f, 1.0f, 1.0f, 1.0f);
-	LogTest(TEXT("VectorReplicate"), TestVectorsEqual(V0, V1));
-
-	V0 = MakeVectorRegister( 1.0f, 2.0f, 3.0f, 4.0f );
-	V1 = VectorReplicate( V0, 1 );
-	V0 = MakeVectorRegister( 2.0f, 2.0f, 2.0f, 2.0f );
-	LogTest( TEXT("VectorReplicate"), TestVectorsEqual( V0, V1 ) );
+	// Replicate
+	TestVectorReplicate<float, VectorRegister4Float>();
 
-	V0 = MakeVectorRegister(1.0f, 2.0f, 3.0f, 4.0f);
-	V1 = VectorReplicate(V0, 2);
-	V0 = MakeVectorRegister(3.0f, 3.0f, 3.0f, 3.0f);
-	LogTest(TEXT("VectorReplicate"), TestVectorsEqual(V0, V1));
+	// Swizzle
+	TestVectorSwizzle<float, VectorRegister4Float>();
 
-	V0 = MakeVectorRegister(1.0f, 2.0f, 3.0f, 4.0f);
-	V1 = VectorReplicate(V0, 3);
-	V0 = MakeVectorRegister(4.0f, 4.0f, 4.0f, 4.0f);
-	LogTest(TEXT("VectorReplicate"), TestVectorsEqual(V0, V1));
+	// Shuffle
+	TestVectorShuffle<float, VectorRegister4Float>();
+
+	V0 = MakeVectorRegisterFloat(0.f, 0.f, 2.f, 0.f);
+	V1 = MakeVectorRegisterFloat(0.f, 1.f, 0.f, 3.f);
+	float Float0 = VectorGetComponent(V0, 0);
+	float Float1 = VectorGetComponent(V1, 1);
+	float Float2 = VectorGetComponent(V0, 2);
+	float Float3 = VectorGetComponent(V1, 3);
+	V0 = MakeVectorRegisterFloat(0.f, 1.f, 2.f, 3.f);
+	V1 = MakeVectorRegisterFloat(Float0, Float1, Float2, Float3);
+	LogTest(TEXT("VectorGetComponent<float>"), TestVectorsEqual(V0, V1));
 
 	V0 = MakeVectorRegister( 1.0f, -2.0f, 3.0f, -4.0f );
 	V1 = VectorAbs( V0 );
@@ -1657,137 +1679,6 @@ bool FVectorRegisterAbstractionTest::RunTest(const FString& Parameters)
 	V0 = MakeVectorRegister(1.0f, 3.0f, 2.0f, 4.0f);
 	LogTest(TEXT("VectorCombineLow"), TestVectorsEqual(V0, V1));
 
-
-	// Replicate
-	{
-		const float ArrayV0[4] = { 0.0, 1.0, 2.0, 3.0 };
-		V0 = VectorLoad(ArrayV0);
-
-#define ReplicateTest(A, X) \
-		V2 = VectorReplicate(A, X); \
-		V3 = MakeVectorRegisterFloat(ArrayV0[X], ArrayV0[X], ArrayV0[X], ArrayV0[X]); \
-		LogTest(*FString::Printf(TEXT("VectorReplicateFloat<%d>"), X), TestVectorsEqual(V2, V3));
-
-		ReplicateTest(V0, 0);
-		ReplicateTest(V0, 1);
-		ReplicateTest(V0, 2);
-		ReplicateTest(V0, 3);
-
-#undef ReplicateTest
-	}
-
-	// Swizzle
-	{
-		const float ArrayV0[4] = { 0.0, 1.0, 2.0, 3.0 };
-		V0 = VectorLoad(ArrayV0);
-
-#define SwizzleTest(A, X, Y, Z, W) \
-		V2 = VectorSwizzle(A, X, Y, Z, W); \
-		V3 = MakeVectorRegisterFloat(ArrayV0[X], ArrayV0[Y], ArrayV0[Z], ArrayV0[W]); \
-		LogTest(*FString::Printf(TEXT("VectorSwizzleFloat<%d,%d,%d,%d>"), X, Y, Z, W), TestVectorsEqual(V2, V3));
-
-		// This is not an exhaustive list because it would be 4*4*4*4 = 256 entries, but it tries to test a lot of common permutations.
-		// Unfortunately it can't be done in a loop because it uses a #define and compile-time constants for the VectorSwizzle() 'function'.
-		// Many of these were selected to also stress the specializations in certain implementations.
-
-		SwizzleTest(V0, 0, 1, 2, 3); // Identity
-		SwizzleTest(V0, 0, 0, 0, 0); // Replicate 0
-		SwizzleTest(V0, 1, 1, 1, 1); // Replicate 1
-		SwizzleTest(V0, 2, 2, 2, 2); // Replicate 2
-		SwizzleTest(V0, 3, 3, 3, 3); // Replicate 3
-
-		SwizzleTest(V0, 0, 3, 1, 2);
-		SwizzleTest(V0, 0, 2, 0, 2);
-
-		SwizzleTest(V0, 1, 0, 1, 0);
-		SwizzleTest(V0, 1, 0, 3, 2);
-		SwizzleTest(V0, 1, 2, 0, 1);
-
-		SwizzleTest(V0, 2, 0, 1, 3);
-		SwizzleTest(V0, 2, 3, 0, 1);
-
-		SwizzleTest(V0, 3, 2, 1, 0);
-		SwizzleTest(V0, 3, 0, 3, 0);
-
-		SwizzleTest(V0, 2, 2, 0, 1);
-		SwizzleTest(V0, 3, 3, 1, 0);
-		SwizzleTest(V0, 2, 2, 0, 2);
-		SwizzleTest(V0, 3, 3, 1, 3);
-
-		// Common specializations
-		SwizzleTest(V0, 0, 1, 2, 2);
-		SwizzleTest(V0, 0, 1, 3, 3);
-		SwizzleTest(V0, 0, 1, 3, 2);
-
-		SwizzleTest(V0, 0, 0, 2, 3);
-		SwizzleTest(V0, 0, 0, 2, 2);
-		SwizzleTest(V0, 0, 0, 3, 3);
-		SwizzleTest(V0, 0, 0, 3, 2);
-
-		SwizzleTest(V0, 1, 0, 2, 3);
-		SwizzleTest(V0, 1, 0, 2, 2);
-		SwizzleTest(V0, 1, 0, 3, 3);
-		SwizzleTest(V0, 1, 0, 3, 2);
-
-		SwizzleTest(V0, 1, 1, 2, 3);
-		SwizzleTest(V0, 1, 1, 2, 2);
-		SwizzleTest(V0, 1, 1, 3, 3);
-		SwizzleTest(V0, 1, 1, 3, 2);
-
-		SwizzleTest(V0, 0, 1, 0, 1);
-		SwizzleTest(V0, 2, 3, 2, 3);
-		SwizzleTest(V0, 0, 0, 1, 1);
-		SwizzleTest(V0, 2, 2, 3, 3);
-		SwizzleTest(V0, 0, 0, 2, 2);
-		SwizzleTest(V0, 1, 1, 3, 3);
-#undef SwizzleTest
-	}
-
-	// Shuffle
-	{
-		const float ArrayV0[4] = { 0.0, 1.0, 2.0, 3.0 };
-		const float ArrayV1[4] = { 0.0, 1.0, 2.0, 3.0 };
-		V0 = VectorLoad(ArrayV0);
-		V1 = VectorLoad(ArrayV1);
-
-#define ShuffleTest(A, B, X, Y, Z, W) \
-		V2 = VectorShuffle(A, B, X, Y, Z, W); \
-		V3 = MakeVectorRegisterFloat(ArrayV0[X], ArrayV0[Y], ArrayV1[Z], ArrayV1[W]); \
-		LogTest(*FString::Printf(TEXT("VectorShuffleFloat<%d,%d,%d,%d>"), X, Y, Z, W), TestVectorsEqual(V2, V3));
-
-		// This is not an exhaustive list because it would be 4*4*4*4 = 256 entries, but it tries to test a lot of common permutations.
-		// Unfortunately it can't be done in a loop because it uses a #define and compile-time constants for the VectorShuffle() 'function'.
-		// Many of these were selected to also stress the specializations in certain implementations.
-		ShuffleTest(V0, V1, 0, 0, 0, 0);
-		ShuffleTest(V0, V1, 1, 1, 1, 1);
-		ShuffleTest(V0, V1, 2, 2, 2, 2);
-		ShuffleTest(V0, V1, 3, 3, 3, 3);
-
-		ShuffleTest(V0, V1, 0, 1, 2, 3);
-		ShuffleTest(V0, V1, 3, 2, 1, 0);
-
-		ShuffleTest(V0, V1, 0, 1, 0, 1);
-		ShuffleTest(V0, V1, 0, 2, 0, 2);
-		ShuffleTest(V0, V1, 0, 3, 0, 3);
-
-		ShuffleTest(V0, V1, 1, 0, 1, 0);
-		ShuffleTest(V0, V1, 2, 0, 2, 0);
-		ShuffleTest(V0, V1, 3, 0, 3, 0);
-
-		ShuffleTest(V0, V1, 0, 0, 1, 1);
-		ShuffleTest(V0, V1, 0, 0, 2, 2);
-		ShuffleTest(V0, V1, 1, 1, 3, 3);
-		ShuffleTest(V0, V1, 2, 2, 3, 3);
-		ShuffleTest(V0, V1, 2, 3, 0, 1);
-		ShuffleTest(V0, V1, 2, 3, 2, 3);
-
-		ShuffleTest(V0, V1, 3, 1, 3, 0);
-		ShuffleTest(V0, V1, 2, 2, 0, 1);
-		ShuffleTest(V0, V1, 0, 1, 3, 2);
-		ShuffleTest(V0, V1, 1, 3, 0, 3);
-#undef ShuffleTest
-	}
-
 	// LoadByte4
 
 	uint8 Bytes[4] = { 25, 75, 125, 200 };
