#ue5 - LWC: fixed a few precision losses in FVector functions.

diff --git a/Engine/Source/Runtime/Core/Private/Math/UnrealMath.cpp b/Engine/Source/Runtime/Core/Private/Math/UnrealMath.cpp
--- a/Engine/Source/Runtime/Core/Private/Math/UnrealMath.cpp
+++ b/Engine/Source/Runtime/Core/Private/Math/UnrealMath.cpp
@@ -199,10 +199,10 @@ FRotator FVector3d::ToOrientationRotator() const
 	FRotator R;
 
 	// Find yaw.
-	R.Yaw = static_cast<float>(FMath::Atan2(Y,X) * (180.0 / PI));
+	R.Yaw = static_cast<float>(FMath::Atan2(Y,X) * (180.0 / DOUBLE_PI));
 
 	// Find pitch.
-	R.Pitch = static_cast<float>(FMath::Atan2(Z,FMath::Sqrt(X*X+Y*Y)) * (180.0 / PI));
+	R.Pitch = static_cast<float>(FMath::Atan2(Z,FMath::Sqrt(X*X+Y*Y)) * (180.0 / DOUBLE_PI));
 
 	// Find roll.
 	R.Roll = 0;
@@ -332,10 +332,10 @@ FVector FMath::ClosestPointOnLine(const FVector& LineStart, const FVector& LineE
 {
 	// Solve to find alpha along line that is closest point
 	// Weisstein, Eric W. "Point-Line Distance--3-Dimensional." From MathWorld--A Switchram Web Resource. http://mathworld.wolfram.com/Point-LineDistance3-Dimensional.html 
-	const float A = (LineStart - Point) | (LineEnd - LineStart);
-	const float B = (LineEnd - LineStart).SizeSquared();
+	const FVector::FReal A = (LineStart - Point) | (LineEnd - LineStart);
+	const FVector::FReal B = (LineEnd - LineStart).SizeSquared();
 	// This should be robust to B == 0 (resulting in NaN) because clamp should return 1.
-	const float T = FMath::Clamp(-A/B, 0.f, 1.f);
+	const FVector::FReal T = FMath::Clamp(-A/B, FVector::FReal(0.f), FVector::FReal(1.f));
 
 	// Generate closest point
 	FVector ClosestPoint = LineStart + (T * (LineEnd - LineStart));
@@ -345,13 +345,13 @@ FVector FMath::ClosestPointOnLine(const FVector& LineStart, const FVector& LineE
 
 FVector FMath::ClosestPointOnInfiniteLine(const FVector& LineStart, const FVector& LineEnd, const FVector& Point)
 {
-	const float A = (LineStart - Point) | (LineEnd - LineStart);
-	const float B = (LineEnd - LineStart).SizeSquared();
+	const FVector::FReal A = (LineStart - Point) | (LineEnd - LineStart);
+	const FVector::FReal B = (LineEnd - LineStart).SizeSquared();
 	if (B < SMALL_NUMBER)
 	{
 		return LineStart;
 	}
-	const float T = -A/B;
+	const FVector::FReal T = -A/B;
 
 	// Generate closest point
 	const FVector ClosestPoint = LineStart + (T * (LineEnd - LineStart));
