#ue5 - LWC: fix benign double<->float warning.

diff --git a/Engine/Source/Runtime/Core/Private/Tests/Math/UnrealMathTest.cpp b/Engine/Source/Runtime/Core/Private/Tests/Math/UnrealMathTest.cpp
--- a/Engine/Source/Runtime/Core/Private/Tests/Math/UnrealMathTest.cpp
+++ b/Engine/Source/Runtime/Core/Private/Tests/Math/UnrealMathTest.cpp
@@ -839,7 +839,7 @@ void SetScratchDouble(float X, float Y, float Z, float W, float U = 0.0f)
 template<typename RealType, typename VectorRegisterType>
 FORCENOINLINE void TestVectorReplicate()
 {
-	const RealType ArrayV0[4] = { 0.0, 1.0, 2.0, 3.0 };
+	const RealType ArrayV0[4] = { (RealType)0.0, (RealType)1.0, (RealType)2.0, (RealType)3.0 };
 	VectorRegisterType V0 = VectorLoad(ArrayV0);
 	VectorRegisterType V1, V2;
 
@@ -860,7 +860,7 @@ FORCENOINLINE void TestVectorReplicate()
 template<typename RealType, typename VectorRegisterType>
 FORCENOINLINE void TestVectorSwizzle()
 {
-	const RealType ArrayV0[4] = { 0.0, 1.0, 2.0, 3.0 };
+	const RealType ArrayV0[4] = { (RealType)0.0, (RealType)1.0, (RealType)2.0, (RealType)3.0 };
 	VectorRegisterType V0 = VectorLoad(ArrayV0);
 	VectorRegisterType V1, V2;
 
@@ -930,8 +930,8 @@ FORCENOINLINE void TestVectorSwizzle()
 template<typename RealType, typename VectorRegisterType>
 FORCENOINLINE void TestVectorShuffle()
 {
-	RealType ArrayV0[4] = { 0.0, 0.1, 0.2, 0.3 };
-	RealType ArrayV1[4] = { 1.0, 1.1, 1.2, 1.3 };
+	RealType ArrayV0[4] = { (RealType)0.0, (RealType)0.1, (RealType)0.2, (RealType)0.3 };
+	RealType ArrayV1[4] = { (RealType)1.0, (RealType)1.1, (RealType)1.2, (RealType)1.3 };
 	VectorRegisterType V0 = VectorLoad(ArrayV0);
 	VectorRegisterType V1 = VectorLoad(ArrayV1);
 	VectorRegisterType V2, V3;
