#ue5 - LWC: Improve FBox::TransformBy() to use faster method of loading float constant.

diff --git a/Engine/Source/Runtime/Core/Private/Math/Box.cpp b/Engine/Source/Runtime/Core/Private/Math/Box.cpp
--- a/Engine/Source/Runtime/Core/Private/Math/Box.cpp
+++ b/Engine/Source/Runtime/Core/Private/Math/Box.cpp
@@ -50,15 +50,15 @@ FBox FBox::TransformBy(const FMatrix& M) const
 
 	FBox NewBox;
 
-	const VectorRegister VecMin = VectorLoadFloat3(&Min);
-	const VectorRegister VecMax = VectorLoadFloat3(&Max);
+	const VectorRegister VecMin = VectorLoadFloat3_W0(&Min);
+	const VectorRegister VecMax = VectorLoadFloat3_W0(&Max);
 
 	const VectorRegister m0 = VectorLoadAligned(M.M[0]);
 	const VectorRegister m1 = VectorLoadAligned(M.M[1]);
 	const VectorRegister m2 = VectorLoadAligned(M.M[2]);
 	const VectorRegister m3 = VectorLoadAligned(M.M[3]);
 
-	const VectorRegister Half = VectorSetFloat3(0.5f, 0.5f, 0.5f);
+	const VectorRegister Half = VectorSetFloat1((FVector::FReal)0.5f); // VectorSetFloat1() can be faster than SetFloat3(0.5, 0.5, 0.5, 0.0). Okay if 4th element is 0.5, it's multiplied by 0.0 below and we discard W anyway.
 	const VectorRegister Origin = VectorMultiply(VectorAdd(VecMax, VecMin), Half);
 	const VectorRegister Extent = VectorMultiply(VectorSubtract(VecMax, VecMin), Half);
 
