#ue5 - LWC: Expanded VectorSwizzle() and VectorShuffle() tests to hit more permutations.

diff --git a/Engine/Source/Runtime/Core/Private/Tests/Math/UnrealMathTest.cpp b/Engine/Source/Runtime/Core/Private/Tests/Math/UnrealMathTest.cpp
--- a/Engine/Source/Runtime/Core/Private/Tests/Math/UnrealMathTest.cpp
+++ b/Engine/Source/Runtime/Core/Private/Tests/Math/UnrealMathTest.cpp
@@ -873,24 +873,32 @@ FORCENOINLINE void TestVectorSwizzle()
 	// Unfortunately it can't be done in a loop because it uses a #define and compile-time constants for the VectorSwizzle() 'function'.
 	// Many of these were selected to also stress the specializations in certain implementations.
 
-	SwizzleTest(V0, 0, 1, 2, 3); // Identity
 	SwizzleTest(V0, 0, 0, 0, 0); // Replicate 0
 	SwizzleTest(V0, 1, 1, 1, 1); // Replicate 1
 	SwizzleTest(V0, 2, 2, 2, 2); // Replicate 2
 	SwizzleTest(V0, 3, 3, 3, 3); // Replicate 3
 
-	SwizzleTest(V0, 0, 3, 1, 2);
+	SwizzleTest(V0, 0, 1, 2, 3); // Identity
+	SwizzleTest(V0, 1, 2, 3, 0); // Rotate << 1
+	SwizzleTest(V0, 2, 3, 0, 1); // Rotate << 2
+	SwizzleTest(V0, 3, 0, 1, 2); // Rotate << 3
+	
 	SwizzleTest(V0, 0, 2, 0, 2);
 
 	SwizzleTest(V0, 1, 0, 1, 0);
-	SwizzleTest(V0, 1, 0, 3, 2);
 	SwizzleTest(V0, 1, 2, 0, 1);
 
-	SwizzleTest(V0, 2, 0, 1, 3);
-	SwizzleTest(V0, 2, 3, 0, 1);
-
+	SwizzleTest(V0, 0, 3, 1, 2);
+	SwizzleTest(V0, 0, 3, 2, 1);
+	SwizzleTest(V0, 1, 0, 3, 2);
+	SwizzleTest(V0, 1, 0, 2, 3);
+	SwizzleTest(V0, 2, 1, 3, 1);
+	SwizzleTest(V0, 2, 1, 1, 3);
+	SwizzleTest(V0, 3, 2, 0, 1);
 	SwizzleTest(V0, 3, 2, 1, 0);
+
 	SwizzleTest(V0, 3, 0, 3, 0);
+	SwizzleTest(V0, 0, 3, 0, 3);
 
 	SwizzleTest(V0, 2, 2, 0, 1);
 	SwizzleTest(V0, 3, 3, 1, 0);
@@ -950,6 +958,10 @@ FORCENOINLINE void TestVectorShuffle()
 	ShuffleTest(V0, V1, 3, 3, 3, 3);
 
 	ShuffleTest(V0, V1, 0, 1, 2, 3);
+	ShuffleTest(V0, V1, 1, 2, 3, 0);
+	ShuffleTest(V0, V1, 2, 3, 0, 1);
+	ShuffleTest(V0, V1, 3, 0, 1, 2);
+
 	ShuffleTest(V0, V1, 3, 2, 1, 0);
 
 	ShuffleTest(V0, V1, 0, 1, 0, 1);
