#ue5 - LWC: Fix incorrect/misleading comment on VectorLoadFloat3.

diff --git a/Engine/Source/Runtime/Core/Public/Math/UnrealMathVectorCommon.h b/Engine/Source/Runtime/Core/Public/Math/UnrealMathVectorCommon.h
--- a/Engine/Source/Runtime/Core/Public/Math/UnrealMathVectorCommon.h
+++ b/Engine/Source/Runtime/Core/Public/Math/UnrealMathVectorCommon.h
@@ -114,7 +114,7 @@ FORCEINLINE VectorRegister4Double VectorLoadAligned(const VectorRegister4Double*
 ///////////////////////////////////////////////////////////////////////////////////////////////////////////////////
 
 /**
-* Loads 3 floats from unaligned memory and leaves W undefined.
+* Loads 3 floats from unaligned memory and sets W=0.
 *
 * @param Ptr	Unaligned memory pointer to the 3 floats
 * @return		VectorRegister4Float(Ptr[0], Ptr[1], Ptr[2], 0.0f)
