#ue5 - LWC: Add a couple SSE 4.1 specializations for more efficient VectorSwizzle() when available.

diff --git a/Engine/Source/Runtime/Core/Public/Math/UnrealMathSSE.h b/Engine/Source/Runtime/Core/Public/Math/UnrealMathSSE.h
--- a/Engine/Source/Runtime/Core/Public/Math/UnrealMathSSE.h
+++ b/Engine/Source/Runtime/Core/Public/Math/UnrealMathSSE.h
@@ -736,7 +736,6 @@ FORCEINLINE void VectorStoreFloat1(const VectorRegister4Double& Vec, double* Dst
 }
 
 // Templated shuffles required for double shuffles when using __m128d, since we have to break it down in to two separate operations.
-// TODO: more specializations for simple cases. SSE 4.1 and AVX have cheaper blends available.
 
 // [0,1]:[0,1]
 template<int Index0, int Index1, typename std::enable_if< (Index0 <= 1) && (Index1 <= 1), bool >::type = true>
@@ -744,8 +743,6 @@ FORCEINLINE VectorRegister2Double SelectVectorSwizzle2(const VectorRegister4Doub
 {
 	return _mm_shuffle_pd(Vec.XY, Vec.XY, SHUFFLEMASK2(Index0, Index1));
 }
-// <0,1> is simply XY.
-template<> FORCEINLINE VectorRegister2Double SelectVectorSwizzle2<0, 1>(const VectorRegister4Double& Vec) { return Vec.XY; }
 
 // [0,1]:[2,3]
 template<int Index0, int Index1, typename std::enable_if< (Index0 <= 1) && (Index1 > 1), bool >::type = true>
@@ -767,9 +764,16 @@ FORCEINLINE VectorRegister2Double SelectVectorSwizzle2(const VectorRegister4Doub
 {
 	return _mm_shuffle_pd(Vec.ZW, Vec.ZW, SHUFFLEMASK2(Index0 - 2, Index1 - 2));
 }
-// <2,3> is simply ZW.
+
+template<> FORCEINLINE VectorRegister2Double SelectVectorSwizzle2<0, 1>(const VectorRegister4Double& Vec) { return Vec.XY; }
 template<> FORCEINLINE VectorRegister2Double SelectVectorSwizzle2<2, 3>(const VectorRegister4Double& Vec) { return Vec.ZW; }
 
+#if UE_PLATFORM_MATH_USE_SSE4_1
+// blend can run on more ports than shuffle, so are preferable even if latency is claimed to be the same.
+template<> FORCEINLINE VectorRegister2Double SelectVectorSwizzle2<0, 3>(const VectorRegister4Double& Vec) { return _mm_blend_pd(Vec.XY, Vec.ZW, SHUFFLEMASK2(0, 1)); }
+template<> FORCEINLINE VectorRegister2Double SelectVectorSwizzle2<2, 1>(const VectorRegister4Double& Vec) { return _mm_blend_pd(Vec.ZW, Vec.XY, SHUFFLEMASK2(0, 1)); }
+#endif
+
 // Double swizzle
 template<int Index0, int Index1, int Index2, int Index3>
 FORCEINLINE VectorRegister4Double VectorSwizzleTemplate(const VectorRegister4Double& Vec)
@@ -789,6 +793,7 @@ template<> FORCEINLINE VectorRegister4Double VectorSwizzleTemplate<0, 1, 2, 3>(c
 
 #if UE_PLATFORM_MATH_USE_AVX
 constexpr int PERMUTE_MASK(int A, int B, int C, int D) { return ((A == 1 ? (1 << 0) : 0) | (B == 1 ? (1 << 1) : 0) | (C == 3 ? (1 << 2) : 0) | (D == 3 ? (1 << 3) : 0)); }
+
 // _mm256_permute4x64_pd has a latency of 3, so here are some specializations using instructions which have a latency of 1 but are restricted to in-lane (128 bit) permutes.
 template<> FORCEINLINE VectorRegister4Double VectorSwizzleTemplate<0, 1, 2, 2>(const VectorRegister4Double& Vec) { return _mm256_permute_pd(Vec, PERMUTE_MASK(0, 1, 2, 2)); }
 template<> FORCEINLINE VectorRegister4Double VectorSwizzleTemplate<0, 1, 3, 3>(const VectorRegister4Double& Vec) { return _mm256_permute_pd(Vec, PERMUTE_MASK(0, 1, 3, 3)); }
