U5 - Fix the some of the broken old values when the RenderingThread is not active.

diff --git a/Engine/Source/Runtime/Core/Private/HAL/ThreadingBase.cpp b/Engine/Source/Runtime/Core/Private/HAL/ThreadingBase.cpp
--- a/Engine/Source/Runtime/Core/Private/HAL/ThreadingBase.cpp
+++ b/Engine/Source/Runtime/Core/Private/HAL/ThreadingBase.cpp
@@ -248,14 +248,15 @@ CORE_API bool IsInRenderingThread()
 PRAGMA_DISABLE_DEPRECATION_WARNINGS
 	const bool bLocalIsLoadingThreadSuspended = GIsRenderingThreadSuspended.Load(EMemoryOrder::Relaxed) != 0;
 
-	bool newValue = !GRenderingThread || bLocalIsLoadingThreadSuspended
+	bool newValue = (GRenderThreadId == 0) || bLocalIsLoadingThreadSuspended
 		? FTaskTagScope::IsCurrentTag(ETaskTag::EGameThread) || FTaskTagScope::IsCurrentTag(ETaskTag::ERenderingThread) || FTaskTagScope::IsRunningDuringStaticInit()
 		: FTaskTagScope::IsCurrentTag(ETaskTag::ERenderingThread);
 
 #if !UE_BUILD_SHIPPING && !UE_BUILD_TEST
 	if (!LowLevelTasks::FScheduler::IsBusyWaiting())
 	{
-		bool oldValue = !GRenderingThread || bLocalIsLoadingThreadSuspended || (FPlatformTLS::GetCurrentThreadId() == GRenderingThread->GetThreadID());
+		const uint32 CurrentThreadId = FPlatformTLS::GetCurrentThreadId();
+		bool oldValue = ((GRenderThreadId == 0) || bLocalIsLoadingThreadSuspended) ? (CurrentThreadId == GGameThreadId) || FTaskTagScope::IsRunningDuringStaticInit() : (CurrentThreadId == GRenderThreadId);
 		ensureMsgf(oldValue == newValue, TEXT("oldValue(%i) newValue(%i) If this check fails make sure that there is a FTaskTagScope(ETaskTag::ERenderingThread) as deep as possible on the current callstack, you can see the current value in ActiveNamedThreads(%x), GRenderingThread(%x), GIsRenderingThreadSuspended(%d)"), oldValue, newValue, FTaskTagScope::GetCurrentTag(), GRenderingThread, bLocalIsLoadingThreadSuspended);
 		newValue = oldValue;
 	}
@@ -270,7 +271,7 @@ PRAGMA_DISABLE_DEPRECATION_WARNINGS
 	const bool bLocalIsLoadingThreadSuspended = GIsRenderingThreadSuspended.Load(EMemoryOrder::Relaxed) != 0;
 
 	bool newValue = false;
-	if (!GRenderingThread || bLocalIsLoadingThreadSuspended)
+	if ((GRenderThreadId == 0) || bLocalIsLoadingThreadSuspended)
 	{
 		newValue = FTaskTagScope::IsCurrentTag(ETaskTag::ERenderingThread) || FTaskTagScope::IsCurrentTag(ETaskTag::EGameThread) || FTaskTagScope::IsCurrentTag(ETaskTag::EParallelRenderingThread);
 	}
@@ -285,15 +286,8 @@ PRAGMA_DISABLE_DEPRECATION_WARNINGS
 #if !UE_BUILD_SHIPPING && !UE_BUILD_TEST
 	if (!LowLevelTasks::FScheduler::IsBusyWaiting())
 	{
-		bool oldValue = false;
-		if (!GRenderingThread || bLocalIsLoadingThreadSuspended)
-		{
-			oldValue = true;
-		}
-		else
-		{
-			oldValue = FPlatformTLS::GetCurrentThreadId() != GGameThreadId;
-		}
+		const uint32 CurrentThreadId = FPlatformTLS::GetCurrentThreadId();
+		bool oldValue = ((GRenderThreadId == 0) || bLocalIsLoadingThreadSuspended) ?  CurrentThreadId == GGameThreadId : CurrentThreadId != GGameThreadId;
 		ensureMsgf(oldValue == newValue, TEXT("oldValue(%i) newValue(%i) If this check fails make sure that there is a FTaskTagScope(ETaskTag::EParallelRenderingThread) as deep as possible on the current callstack, you can see the current value in ActiveNamedThreads(%x), GRenderingThread(%x), GIsRenderingThreadSuspended(%d)"), oldValue, newValue, FTaskTagScope::GetCurrentTag(), GRenderingThread, bLocalIsLoadingThreadSuspended);
 		newValue = oldValue;
 	}
