Fixed a type truncation warning in the stall detector (the warning was temporarily globally disabled in UE5 main, which got fixed today)

diff --git a/Engine/Source/Runtime/Core/Public/ProfilingDebugging/StallDetector.h b/Engine/Source/Runtime/Core/Public/ProfilingDebugging/StallDetector.h
--- a/Engine/Source/Runtime/Core/Public/ProfilingDebugging/StallDetector.h
+++ b/Engine/Source/Runtime/Core/Public/ProfilingDebugging/StallDetector.h
@@ -85,7 +85,7 @@ namespace UE
 		struct TabulatedResult
 		{
 			const UE::FStallDetectorStats* Stats;
-			uint32 TriggerCount;
+			uint64 TriggerCount;
 			double OverageSeconds;
 			
 			TabulatedResult()
