Made FPathViews::Append take a template parameter pack of arguments to append

diff --git a/Engine/Source/Runtime/Core/Private/Misc/PathViews.cpp b/Engine/Source/Runtime/Core/Private/Misc/PathViews.cpp
--- a/Engine/Source/Runtime/Core/Private/Misc/PathViews.cpp
+++ b/Engine/Source/Runtime/Core/Private/Misc/PathViews.cpp
@@ -225,31 +225,6 @@ void FPathViews::Split(const FStringView& InPath, FStringView& OutPath, FStringV
 	OutExt = CleanName.RightChop(NameLen + 1);
 }
 
-void FPathViews::Append(FStringBuilderBase& Builder, FStringView Suffix)
-{
-	if (Builder.Len() > 0 && !UE4PathViews_Private::IsSlashOrBackslash(Builder.LastChar()))
-	{
-		Builder.Append('/');
-	}
-
-	Builder.Append(Suffix);
-}
-
-void FPathViews::Append(FStringBuilderBase& Builder, FName Suffix)
-{
-	if (Builder.Len() > 0 && !UE4PathViews_Private::IsSlashOrBackslash(Builder.LastChar()))
-	{
-		Builder.Append('/');
-	}
-
-	Suffix.AppendString(Builder);
-}
-
-void FPathViews::Append(FStringBuilderBase& Builder, const TCHAR* Suffix)
-{
-	Append(Builder, FStringView(Suffix));
-}
-
 FString FPathViews::ChangeExtension(const FStringView& InPath, const FStringView& InNewExtension)
 {
 	// Make sure the period we found was actually for a file extension and not part of the file path.
diff --git a/Engine/Source/Runtime/Core/Private/Tests/Misc/PathViewsTest.cpp b/Engine/Source/Runtime/Core/Private/Tests/Misc/PathViewsTest.cpp
--- a/Engine/Source/Runtime/Core/Private/Tests/Misc/PathViewsTest.cpp
+++ b/Engine/Source/Runtime/Core/Private/Tests/Misc/PathViewsTest.cpp
@@ -292,6 +292,50 @@ bool FPathViewsSplitTest::RunTest(const FString& InParameters)
 	return true;
 }
 
+IMPLEMENT_CUSTOM_SIMPLE_AUTOMATION_TEST(FPathViewsAppendTest, FPathViewsTest, "System.Core.Misc.PathViews.Append", EAutomationTestFlags::ApplicationContextMask | EAutomationTestFlags::SmokeFilter)
+bool FPathViewsAppendTest::RunTest(const FString& InParameters)
+{
+	TStringBuilder<256> Path;
+
+	FPathViews::Append(Path, TEXT("A"), TEXT(""));
+	TestEqual(TEXT("FPathViews::Append('A', '')"), FStringView(Path), TEXT("A/"_SV));
+	Path.Reset();
+
+	FPathViews::Append(Path, TEXT(""), TEXT("B"));
+	TestEqual(TEXT("FPathViews::Append('', 'B')"), FStringView(Path), TEXT("B"_SV));
+	Path.Reset();
+
+	FPathViews::Append(Path, TEXT("/"), TEXT("B"));
+	TestEqual(TEXT("FPathViews::Append('/', 'B')"), FStringView(Path), TEXT("/B"_SV));
+	Path.Reset();
+
+	FPathViews::Append(Path, TEXT("A"), TEXT("B"));
+	TestEqual(TEXT("FPathViews::Append('A', 'B')"), FStringView(Path), TEXT("A/B"_SV));
+	Path.Reset();
+
+	FPathViews::Append(Path, TEXT("A/"), TEXT("B"));
+	TestEqual(TEXT("FPathViews::Append('A/', 'B')"), FStringView(Path), TEXT("A/B"_SV));
+	Path.Reset();
+
+	FPathViews::Append(Path, TEXT("A\\"), TEXT("B"));
+	TestEqual(TEXT("FPathViews::Append('A\\', 'B')"), FStringView(Path), TEXT("A\\B"_SV));
+	Path.Reset();
+
+	FPathViews::Append(Path, TEXT("A/B"), TEXT("C/D"));
+	TestEqual(TEXT("FPathViews::Append('A/B', 'C/D')"), FStringView(Path), TEXT("A/B/C/D"_SV));
+	Path.Reset();
+
+	FPathViews::Append(Path, TEXT("A/"), TEXT("B"), TEXT("C/"), TEXT("D"));
+	TestEqual(TEXT("FPathViews::Append('A/', 'B', 'C/', 'D')"), FStringView(Path), TEXT("A/B/C/D"_SV));
+	Path.Reset();
+
+	FPathViews::Append(Path, TEXT("A/"), 16, TEXT("B"));
+	TestEqual(TEXT("FPathViews::Append('A/', 16, 'B')"), FStringView(Path), TEXT("A/16/B"_SV));
+	Path.Reset();
+
+	return true;
+}
+
 IMPLEMENT_CUSTOM_SIMPLE_AUTOMATION_TEST(FPathViewsChangeExtensionTest, FPathViewsTest, "System.Core.Misc.PathViews.ChangeExtension", EAutomationTestFlags::ApplicationContextMask | EAutomationTestFlags::SmokeFilter)
 bool FPathViewsChangeExtensionTest::RunTest(const FString& InParameters)
 {
diff --git a/Engine/Source/Runtime/Core/Public/Misc/PathViews.h b/Engine/Source/Runtime/Core/Public/Misc/PathViews.h
--- a/Engine/Source/Runtime/Core/Public/Misc/PathViews.h
+++ b/Engine/Source/Runtime/Core/Public/Misc/PathViews.h
@@ -146,7 +146,7 @@ public:
 	static void Split(const FStringView& InPath, FStringView& OutPath, FStringView& OutName, FStringView& OutExt);
 
 	/**
-	 * Appends the suffix to the path in the builder and ensures that there is a separator between them.
+	 * Appends each suffix argument to the path in the builder and ensures that there is a separator between them.
 	 *
 	 * Examples:
 	 * ("",    "")    -> ""
@@ -157,13 +157,24 @@ public:
 	 * ("A/",  "B")   -> "A/B"
 	 * ("A\\", "B")   -> "A\\B"
 	 * ("A/B", "C/D") -> "A/B/C/D"
+	 * ("A/", "B", "C/", "D") -> "A/B/C/D"
 	 *
 	 * @param Builder A possibly-empty path that may end in a separator.
-	 * @param Suffix A possibly-empty suffix that does not start with a separator.
+	 * @param Args Arguments that can write to a string builder and do not start with a separator.
 	 */
-	static void Append(FStringBuilderBase& Builder, FStringView Suffix);
-	static void Append(FStringBuilderBase& Builder, const TCHAR* Suffix);
-	static void Append(FStringBuilderBase& Builder, FName Suffix);
+	template <typename CharType, typename... ArgTypes>
+	static void Append(TStringBuilderBase<CharType>& Builder, ArgTypes&&... Args)
+	{
+		const auto AddSeparator = [&Builder]() -> TStringBuilderBase<CharType>&
+		{
+			if (!(Builder.Len() == 0 || Builder.LastChar() == '/' || Builder.LastChar() == '\\'))
+			{
+				Builder.Append('/');
+			}
+			return Builder;
+		};
+		((AddSeparator() << Args), ...);
+	}
 
 	/**
 	 * Replaces the pre-existing file extension of a filename. 
