LWC: Fix precision loss warnings in Core. Temporarily silence Chaos precision warnings in LWC builds while they are being fixed.

diff --git a/Engine/Source/Runtime/Core/Private/Math/BoxSphereBounds.cpp b/Engine/Source/Runtime/Core/Private/Math/BoxSphereBounds.cpp
--- a/Engine/Source/Runtime/Core/Private/Math/BoxSphereBounds.cpp
+++ b/Engine/Source/Runtime/Core/Private/Math/BoxSphereBounds.cpp
@@ -51,7 +51,7 @@ FBoxSphereBounds FBoxSphereBounds::TransformBy(const FMatrix& M) const
 	Result.SphereRadius = FMath::Sqrt(VectorGetComponent( MaxRadius, 0) ) * SphereRadius;
 
 	// For non-uniform scaling, computing sphere radius from a box results in a smaller sphere.
-	float const BoxExtentMagnitude = FMath::Sqrt(VectorGetComponent(VectorDot3(NewExtent, NewExtent), 0));
+	FVector::FReal const BoxExtentMagnitude = FMath::Sqrt(VectorGetComponent(VectorDot3(NewExtent, NewExtent), 0));
 	Result.SphereRadius = FMath::Min<FVector::FReal>(Result.SphereRadius, BoxExtentMagnitude);
 
 	Result.DiagnosticCheckNaN();
diff --git a/Engine/Source/Runtime/Core/Private/Math/Color.cpp b/Engine/Source/Runtime/Core/Private/Math/Color.cpp
--- a/Engine/Source/Runtime/Core/Private/Math/Color.cpp
+++ b/Engine/Source/Runtime/Core/Private/Math/Color.cpp
@@ -308,9 +308,9 @@ FLinearColor FLinearColor::FGetHSV( uint8 H, uint8 S, uint8 V )
 	float Brightness = V * 1.4f / 255.f;
 	Brightness *= 0.7f/(0.01f + FMath::Sqrt(Brightness));
 	Brightness  = FMath::Clamp(Brightness,0.f,1.f);
-	const FVector Hue = (H<86) ? FVector((85-H)/85.f,(H-0)/85.f,0) : (H<171) ? FVector(0,(170-H)/85.f,(H-85)/85.f) : FVector((H-170)/85.f,0,(255-H)/84.f);
-	const FVector ColorVector = (Hue + S/255.f * (FVector(1,1,1) - Hue)) * Brightness;
-	return FLinearColor(ColorVector.X,ColorVector.Y,ColorVector.Z,1);
+	const FVector3f Hue = (H<86) ? FVector3f((85-H)/85.f,(H-0)/85.f,0) : (H<171) ? FVector3f(0,(170-H)/85.f,(H-85)/85.f) : FVector3f((H-170)/85.f,0,(255-H)/84.f);
+	const FVector3f ColorVector = (Hue + S/255.f * (FVector3f(1,1,1) - Hue)) * Brightness;
+	return FLinearColor(ColorVector.X, ColorVector.Y, ColorVector.Z,1);
 }
 
 /**
diff --git a/Engine/Source/Runtime/Core/Private/Math/Sphere.cpp b/Engine/Source/Runtime/Core/Private/Math/Sphere.cpp
--- a/Engine/Source/Runtime/Core/Private/Math/Sphere.cpp
+++ b/Engine/Source/Runtime/Core/Private/Math/Sphere.cpp
@@ -26,14 +26,14 @@ FSphere::FSphere( const FVector* Points, int32 Num )
 		}
 	}
 
-	float LargestDistSqr = 0.0f;
+	FReal LargestDistSqr = 0.0f;
 	int32 LargestAxis = 0;
 	for( int k = 0; k < 3; k++ )
 	{
 		FVector PointMin = Points[ MinIndex[k] ];
 		FVector PointMax = Points[ MaxIndex[k] ];
 
-		float DistSqr = ( PointMax - PointMin ).SizeSquared();
+		FReal DistSqr = ( PointMax - PointMin ).SizeSquared();
 		if( DistSqr > LargestDistSqr )
 		{
 			LargestDistSqr = DistSqr;
@@ -50,12 +50,12 @@ FSphere::FSphere( const FVector* Points, int32 Num )
 	// Adjust to fit all points
 	for( int i = 0; i < Num; i++ )
 	{
-		float DistSqr = ( Points[i] - Center ).SizeSquared();
+		FReal DistSqr = ( Points[i] - Center ).SizeSquared();
 
 		if( DistSqr > W*W )
 		{
-			float Dist = FMath::Sqrt( DistSqr );
-			float t = 0.5f + 0.5f * ( W / Dist );
+			FReal Dist = FMath::Sqrt( DistSqr );
+			FReal t = 0.5f + 0.5f * ( W / Dist );
 
 			Center = FMath::LerpStable( Points[i], Center, t );
 			W = 0.5f * ( W + Dist );
@@ -80,14 +80,14 @@ FSphere::FSphere( const FSphere* Spheres, int32 Num )
 		}
 	}
 
-	float LargestDist = 0.0f;
+	FReal LargestDist = 0.0f;
 	int32 LargestAxis = 0;
 	for( int k = 0; k < 3; k++ )
 	{
 		FSphere SphereMin = Spheres[ MinIndex[k] ];
 		FSphere SphereMax = Spheres[ MaxIndex[k] ];
 
-		float Dist = ( SphereMax.Center - SphereMin.Center ).Size() + SphereMin.W + SphereMax.W;
+		FReal Dist = ( SphereMax.Center - SphereMin.Center ).Size() + SphereMin.W + SphereMax.W;
 		if( Dist > LargestDist )
 		{
 			LargestDist = Dist;
@@ -132,7 +132,7 @@ FSphere FSphere::TransformBy(const FTransform& M) const
 	return Result;
 }
 
-float FSphere::GetVolume() const
+FSphere::FReal FSphere::GetVolume() const
 {
 	return (4.f / 3.f) * PI * (W * W * W);
 }
@@ -146,7 +146,7 @@ FSphere& FSphere::operator+=(const FSphere &Other)
 	}
 	
 	FVector ToOther = Other.Center - Center;
-	float DistSqr = ToOther.SizeSquared();
+	FReal DistSqr = ToOther.SizeSquared();
 
 	if( FMath::Square( W - Other.W ) + KINDA_SMALL_NUMBER >= DistSqr )
 	{
@@ -158,7 +158,7 @@ FSphere& FSphere::operator+=(const FSphere &Other)
 	}
 	else
 	{
-		float Dist = FMath::Sqrt( DistSqr );
+		FReal Dist = FMath::Sqrt( DistSqr );
 
 		FSphere NewSphere;
 		NewSphere.W = ( Dist + Other.W + W ) * 0.5f;
diff --git a/Engine/Source/Runtime/Core/Private/Math/UnrealMath.cpp b/Engine/Source/Runtime/Core/Private/Math/UnrealMath.cpp
--- a/Engine/Source/Runtime/Core/Private/Math/UnrealMath.cpp
+++ b/Engine/Source/Runtime/Core/Private/Math/UnrealMath.cpp
@@ -335,7 +335,7 @@ FVector FMath::ClosestPointOnLine(const FVector& LineStart, const FVector& LineE
 	const FVector::FReal A = (LineStart - Point) | (LineEnd - LineStart);
 	const FVector::FReal B = (LineEnd - LineStart).SizeSquared();
 	// This should be robust to B == 0 (resulting in NaN) because clamp should return 1.
-	const FVector::FReal T = FMath::Clamp(-A/B, FVector::FReal(0.f), FVector::FReal(1.f));
+	const FVector::FReal T = FMath::Clamp<FVector::FReal>(-A/B, 0.f, 1.f);
 
 	// Generate closest point
 	FVector ClosestPoint = LineStart + (T * (LineEnd - LineStart));
@@ -480,7 +480,8 @@ FVector FRotator::Euler() const
 
 FRotator FRotator::MakeFromEuler(const FVector& Euler)
 {
-	return FRotator(Euler.Y, Euler.Z, Euler.X);
+	using FReal = decltype(FRotator::Pitch);
+	return FRotator((FReal)Euler.Y, (FReal)Euler.Z, (FReal)Euler.X);
 }
 
 FVector FRotator::UnrotateVector(const FVector& V) const
@@ -631,8 +632,10 @@ void FQuat::ToSwingTwist(const FVector& InTwistAxis, FQuat& OutSwing, FQuat& Out
 	// Vector part projected onto twist axis
 	FVector Projection = FVector::DotProduct(InTwistAxis, FVector(X, Y, Z)) * InTwistAxis;
 
+	using FQuatReal = decltype(FQuat::X);
+
 	// Twist quaternion
-	OutTwist = FQuat(Projection.X, Projection.Y, Projection.Z, W);
+	OutTwist = FQuat((FQuatReal)Projection.X, (FQuatReal)Projection.Y, (FQuatReal)Projection.Z, W);
 
 	// Singularity close to 180deg
 	if(OutTwist.SizeSquared() == 0.0f)
@@ -650,8 +653,8 @@ void FQuat::ToSwingTwist(const FVector& InTwistAxis, FQuat& OutSwing, FQuat& Out
 
 float FQuat::GetTwistAngle(const FVector& TwistAxis) const
 {
-	float XYZ = FVector::DotProduct(TwistAxis, FVector(X, Y, Z));
-	return FMath::UnwindRadians(2.0f * FMath::Atan2(XYZ, W));
+	float XYZ = (float)FVector::DotProduct(TwistAxis, FVector(X, Y, Z));
+	return FMath::UnwindRadians(2.0f * (float)FMath::Atan2(XYZ, W));
 }
 
 FMatrix FRotationAboutPointMatrix::Make(const FQuat& Rot, const FVector& Origin)
@@ -882,17 +885,18 @@ bool FQuat::NetSerialize(FArchive& Ar, class UPackageMap*, bool& bOutSuccess)
 // http://lolengine.net/blog/2014/02/24/quaternion-from-two-vectors-final
 // http://www.euclideanspace.com/maths/algebra/vectors/angleBetween/index.htm
 //
-FORCEINLINE_DEBUGGABLE FQuat FindBetween_Helper(const FVector& A, const FVector& B, float NormAB)
+FORCEINLINE_DEBUGGABLE FQuat FindBetween_Helper(const FVector& A, const FVector& B, FVector::FReal NormAB)
 {
-	float W = NormAB + FVector::DotProduct(A, B);
+	using FReal = decltype(FQuat::X);
+	FReal W = FReal(NormAB + FVector::DotProduct(A, B));	// LWC_TODO: Precision loss
 	FQuat Result;
 
 	if (W >= 1e-6f * NormAB)
 	{
 		//Axis = FVector::CrossProduct(A, B);
-		Result = FQuat(A.Y * B.Z - A.Z * B.Y,
-					   A.Z * B.X - A.X * B.Z,
-					   A.X * B.Y - A.Y * B.X,
+		Result = FQuat(FReal(A.Y * B.Z - A.Z * B.Y),
+					   FReal(A.Z * B.X - A.X * B.Z),
+					   FReal(A.X * B.Y - A.Y * B.X),
 					   W);
 	}
 	else
@@ -900,8 +904,8 @@ FORCEINLINE_DEBUGGABLE FQuat FindBetween_Helper(const FVector& A, const FVector&
 		// A and B point in opposite directions
 		W = 0.f;
 		Result = FMath::Abs(A.X) > FMath::Abs(A.Y)
-				? FQuat(-A.Z, 0.f, A.X, W)
-				: FQuat(0.f, -A.Z, A.Y, W);
+				? FQuat(FReal(-A.Z), 0.f, FReal(A.X), W)
+				: FQuat(0.f, FReal(-A.Z), FReal(A.Y), W);
 	}
 
 	Result.Normalize();
@@ -910,13 +914,13 @@ FORCEINLINE_DEBUGGABLE FQuat FindBetween_Helper(const FVector& A, const FVector&
 
 FQuat FQuat::FindBetweenNormals(const FVector& A, const FVector& B)
 {
-	const float NormAB = 1.f;
+	const FVector::FReal NormAB = 1.f;
 	return FindBetween_Helper(A, B, NormAB);
 }
 
 FQuat FQuat::FindBetweenVectors(const FVector& A, const FVector& B)
 {
-	const float NormAB = FMath::Sqrt(A.SizeSquared() * B.SizeSquared());
+	const FVector::FReal NormAB = FMath::Sqrt(A.SizeSquared() * B.SizeSquared());
 	return FindBetween_Helper(A, B, NormAB);
 }
 
@@ -1089,18 +1093,18 @@ bool FMath::LineExtentBoxIntersection(const FBox& inBox,
 	{	
 		if(Time.Y > Time.Z)
 		{
-			HitTime = Time.Y;
+			HitTime = (decltype(HitTime))Time.Y;	// LWC_TODO: Remove decltype
 			HitNormal = FVector(0, faceDir[1], 0);
 		}
 		else
 		{
-			HitTime = Time.Z;
+			HitTime = (decltype(HitTime))Time.Z;
 			HitNormal = FVector(0, 0, faceDir[2]);
 		}
 		
 		if(Time.X > HitTime)
 		{
-			HitTime = Time.X;
+			HitTime = (decltype(HitTime))Time.X;
 			HitNormal = FVector(faceDir[0], 0, 0);
 		}
 		
@@ -1417,14 +1421,14 @@ CORE_API float FMath::PointDistToLine(const FVector &Point, const FVector &Direc
 {
 	const FVector SafeDir = Direction.GetSafeNormal();
 	OutClosestPoint = Origin + (SafeDir * ((Point-Origin) | SafeDir));
-	return (OutClosestPoint-Point).Size();
+	return (float)(OutClosestPoint-Point).Size();		// LWC_TODO: Precision loss
 }
 
 CORE_API float FMath::PointDistToLine(const FVector &Point, const FVector &Direction, const FVector &Origin)
 {
 	const FVector SafeDir = Direction.GetSafeNormal();
 	const FVector OutClosestPoint = Origin + (SafeDir * ((Point-Origin) | SafeDir));
-	return (OutClosestPoint-Point).Size();
+	return (float)(OutClosestPoint-Point).Size();		// LWC_TODO: Precision loss
 }
 
 FVector FMath::ClosestPointOnSegment(const FVector &Point, const FVector &StartPoint, const FVector &EndPoint)
@@ -1433,14 +1437,14 @@ FVector FMath::ClosestPointOnSegment(const FVector &Point, const FVector &StartP
 	const FVector VectToPoint = Point - StartPoint;
 
 	// See if closest point is before StartPoint
-	const float Dot1 = VectToPoint | Segment;
+	const FVector::FReal Dot1 = VectToPoint | Segment;
 	if( Dot1 <= 0 )
 	{
 		return StartPoint;
 	}
 
 	// See if closest point is beyond EndPoint
-	const float Dot2 = Segment | Segment;
+	const FVector::FReal Dot2 = Segment | Segment;
 	if( Dot2 <= Dot1 )
 	{
 		return EndPoint;
@@ -1473,16 +1477,16 @@ FVector2D FMath::ClosestPointOnSegment2D(const FVector2D &Point, const FVector2D
 	return StartPoint + Segment * (Dot1 / Dot2);
 }
 
-float FMath::PointDistToSegment(const FVector &Point, const FVector &StartPoint, const FVector &EndPoint) 
+float FMath::PointDistToSegment(const FVector &Point, const FVector &StartPoint, const FVector &EndPoint)
 {
 	const FVector ClosestPoint = ClosestPointOnSegment(Point, StartPoint, EndPoint);
-	return (Point - ClosestPoint).Size();
+	return float((Point - ClosestPoint).Size());	// LWC_TODO: Precision loss
 }
 
-float FMath::PointDistToSegmentSquared(const FVector &Point, const FVector &StartPoint, const FVector &EndPoint) 
+float FMath::PointDistToSegmentSquared(const FVector &Point, const FVector &StartPoint, const FVector &EndPoint)
 {
 	const FVector ClosestPoint = ClosestPointOnSegment(Point, StartPoint, EndPoint);
-	return (Point - ClosestPoint).SizeSquared();
+	return float((Point - ClosestPoint).SizeSquared());	// LWC_TODO: Precision loss
 }
 
 struct SegmentDistToSegment_Solver
@@ -1508,19 +1512,19 @@ struct SegmentDistToSegment_Solver
 
 	void Solve(FVector& OutP1, FVector& OutP2)
 	{
-		const float Dot11 = S1 | S1;
-		const float Dot12 = S1 | S2;
-		const float Dot13 = S1 | S3;
-		const float Dot22 = S2 | S2;
-		const float Dot23 = S2 | S3;
+		const FVector::FReal Dot11 = S1 | S1;
+		const FVector::FReal Dot12 = S1 | S2;
+		const FVector::FReal Dot13 = S1 | S3;
+		const FVector::FReal Dot22 = S2 | S2;
+		const FVector::FReal Dot23 = S2 | S3;
 
-		const float D = Dot11*Dot22 - Dot12*Dot12;
+		const FVector::FReal D = Dot11*Dot22 - Dot12*Dot12;
 
-		float D1 = D;
-		float D2 = D;
+		FVector::FReal D1 = D;
+		FVector::FReal D2 = D;
 
-		float N1;
-		float N2;
+		FVector::FReal N1;
+		FVector::FReal N2;
 
 		if (bLinesAreNearlyParallel || D < KINDA_SMALL_NUMBER)
 		{
@@ -1594,8 +1598,8 @@ struct SegmentDistToSegment_Solver
 		}
 
 		// finally do the division to get the points' location
-		const float T1 = (FMath::Abs(N1) < KINDA_SMALL_NUMBER ? 0.f : N1 / D1);
-		const float T2 = (FMath::Abs(N2) < KINDA_SMALL_NUMBER ? 0.f : N2 / D2);
+		const FVector::FReal T1 = (FMath::Abs(N1) < KINDA_SMALL_NUMBER ? 0.f : N1 / D1);
+		const FVector::FReal T2 = (FMath::Abs(N2) < KINDA_SMALL_NUMBER ? 0.f : N2 / D2);
 
 		// return the closest points
 		OutP1 = A1 + T1 * S1;
@@ -1630,10 +1634,10 @@ void FMath::SegmentDistToSegmentSafe(FVector A1, FVector B1, FVector A2, FVector
 	}
 	else
 	{
-		const	float	Dot11_norm = S1_norm | S1_norm;	// always >= 0
-		const	float	Dot22_norm = S2_norm | S2_norm;	// always >= 0
-		const	float	Dot12_norm = S1_norm | S2_norm;
-		const	float	D_norm	= Dot11_norm*Dot22_norm - Dot12_norm*Dot12_norm;	// always >= 0
+		const	FVector::FReal	Dot11_norm = S1_norm | S1_norm;	// always >= 0
+		const	FVector::FReal	Dot22_norm = S2_norm | S2_norm;	// always >= 0
+		const	FVector::FReal	Dot12_norm = S1_norm | S2_norm;
+		const	FVector::FReal	D_norm	= Dot11_norm*Dot22_norm - Dot12_norm*Dot12_norm;	// always >= 0
 
 		Solver.bLinesAreNearlyParallel = D_norm < KINDA_SMALL_NUMBER;
 		Solver.Solve(OutP1, OutP2);
@@ -1647,12 +1651,12 @@ void FMath::SegmentDistToSegment(FVector A1, FVector B1, FVector A2, FVector B2,
 
 float FMath::GetTForSegmentPlaneIntersect(const FVector& StartPoint, const FVector& EndPoint, const FPlane& Plane)
 {
-	return ( Plane.W - (StartPoint|Plane) ) / ( (EndPoint - StartPoint)|Plane);	
+	return float(( Plane.W - (StartPoint|Plane) ) / ( (EndPoint - StartPoint)|Plane));		// LWC_TODO: Precision loss
 }
 
 bool FMath::SegmentPlaneIntersection(const FVector& StartPoint, const FVector& EndPoint, const FPlane& Plane, FVector& out_IntersectionPoint)
 {
-	float T = FMath::GetTForSegmentPlaneIntersect(StartPoint, EndPoint, Plane);
+	FVector::FReal T = FMath::GetTForSegmentPlaneIntersect(StartPoint, EndPoint, Plane);
 	// If the parameter value is not between 0 and 1, there is no intersection
 	if (T > -KINDA_SMALL_NUMBER && T < 1.f + KINDA_SMALL_NUMBER)
 	{
@@ -1698,8 +1702,8 @@ bool FMath::SegmentIntersection2D(const FVector& SegmentStartA, const FVector& S
 	const FVector VectorA = SegmentEndA - SegmentStartA;
 	const FVector VectorB = SegmentEndB - SegmentStartB;
 
-	const float S = (-VectorA.Y * (SegmentStartA.X - SegmentStartB.X) + VectorA.X * (SegmentStartA.Y - SegmentStartB.Y)) / (-VectorB.X * VectorA.Y + VectorA.X * VectorB.Y);
-	const float T = (VectorB.X * (SegmentStartA.Y - SegmentStartB.Y) - VectorB.Y * (SegmentStartA.X - SegmentStartB.X)) / (-VectorB.X * VectorA.Y + VectorA.X * VectorB.Y);
+	const FVector::FReal S = (-VectorA.Y * (SegmentStartA.X - SegmentStartB.X) + VectorA.X * (SegmentStartA.Y - SegmentStartB.Y)) / (-VectorB.X * VectorA.Y + VectorA.X * VectorB.Y);
+	const FVector::FReal T = (VectorB.X * (SegmentStartA.Y - SegmentStartB.Y) - VectorB.Y * (SegmentStartA.X - SegmentStartB.X)) / (-VectorB.X * VectorA.Y + VectorA.X * VectorB.Y);
 
 	const bool bIntersects = (S >= 0 && S <= 1 && T >= 0 && T <= 1);
 
@@ -1719,39 +1723,41 @@ bool FMath::SegmentIntersection2D(const FVector& SegmentStartA, const FVector& S
  * and http://sourceforge.net/mailarchive/message.php?msg_id=10501105
  */
 static bool ComputeProjectedSphereShaft(
-	float LightX,
-	float LightZ,
-	float Radius,
+	FVector::FReal LightX,
+	FVector::FReal LightZ,
+	FVector::FReal Radius,
 	const FMatrix& ProjMatrix,
 	const FVector& Axis,
-	float AxisSign,
+	FVector::FReal AxisSign,
 	int32& InOutMinX,
 	int32& InOutMaxX
 	)
 {
-	float ViewX = (float)InOutMinX;
-	float ViewSizeX = (float)(InOutMaxX - InOutMinX);
+	FVector::FReal ViewX = (FVector::FReal)InOutMinX;
+	FVector::FReal ViewSizeX = (FVector::FReal)(InOutMaxX - InOutMinX);
 
 	// Vertical planes: T = <Nx, 0, Nz, 0>
-	float Discriminant = (FMath::Square(LightX) - FMath::Square(Radius) + FMath::Square(LightZ)) * FMath::Square(LightZ);
+	FVector::FReal Discriminant = (FMath::Square(LightX) - FMath::Square(Radius) + FMath::Square(LightZ)) * FMath::Square(LightZ);
 	if(Discriminant >= 0)
 	{
-		float SqrtDiscriminant = FMath::Sqrt(Discriminant);
-		float InvLightSquare = 1.0f / (FMath::Square(LightX) + FMath::Square(LightZ));
+		FVector::FReal SqrtDiscriminant = FMath::Sqrt(Discriminant);
+		FVector::FReal InvLightSquare = 1.0f / (FMath::Square(LightX) + FMath::Square(LightZ));
 
-		float Nxa = (Radius * LightX - SqrtDiscriminant) * InvLightSquare;
-		float Nxb = (Radius * LightX + SqrtDiscriminant) * InvLightSquare;
-		float Nza = (Radius - Nxa * LightX) / LightZ;
-		float Nzb = (Radius - Nxb * LightX) / LightZ;
-		float Pza = LightZ - Radius * Nza;
-		float Pzb = LightZ - Radius * Nzb;
+		FVector::FReal Nxa = (Radius * LightX - SqrtDiscriminant) * InvLightSquare;
+		FVector::FReal Nxb = (Radius * LightX + SqrtDiscriminant) * InvLightSquare;
+		FVector::FReal Nza = (Radius - Nxa * LightX) / LightZ;
+		FVector::FReal Nzb = (Radius - Nxb * LightX) / LightZ;
+		FVector::FReal Pza = LightZ - Radius * Nza;
+		FVector::FReal Pzb = LightZ - Radius * Nzb;
+
+		using FVec4Real = decltype(FVector4::X);
 
 		// Tangent a
 		if(Pza > 0)
 		{
-			float Pxa = -Pza * Nza / Nxa;
-			FVector4 P = ProjMatrix.TransformFVector4(FVector4(Axis.X * Pxa,Axis.Y * Pxa,Pza,1));
-			float X = (Dot3(P,Axis) / P.W + 1.0f * AxisSign) / 2.0f * AxisSign;
+			FVector::FReal Pxa = -Pza * Nza / Nxa;
+			FVector4 P = ProjMatrix.TransformFVector4(FVector4(FVec4Real(Axis.X * Pxa),FVec4Real(Axis.Y * Pxa),FVec4Real(Pza),1));
+			FVector::FReal X = (Dot3(P,Axis) / P.W + 1.0f * AxisSign) / 2.0f * AxisSign;
 			if(FMath::IsNegative(Nxa) ^ FMath::IsNegative(AxisSign))
 			{
 				InOutMaxX = FMath::Min<int32>(FMath::CeilToInt(ViewSizeX * X + ViewX),InOutMaxX);
@@ -1765,10 +1771,10 @@ static bool ComputeProjectedSphereShaft(
 		// Tangent b
 		if(Pzb > 0)
 		{
-			float Pxb = -Pzb * Nzb / Nxb;
-			FVector4 P = ProjMatrix.TransformFVector4(FVector4(Axis.X * Pxb,Axis.Y * Pxb,Pzb,1));
-			float X = (Dot3(P,Axis) / P.W + 1.0f * AxisSign) / 2.0f * AxisSign;
-			if(FMath::IsNegativeFloat(Nxb) ^ FMath::IsNegativeFloat(AxisSign))
+			FVector::FReal Pxb = -Pzb * Nzb / Nxb;
+			FVector4 P = ProjMatrix.TransformFVector4(FVector4(FVec4Real(Axis.X * Pxb), FVec4Real(Axis.Y * Pxb),FVec4Real(Pzb),1));
+			FVector::FReal X = (Dot3(P,Axis) / P.W + 1.0f * AxisSign) / 2.0f * AxisSign;
+			if(FMath::IsNegative(Nxb) ^ FMath::IsNegative(AxisSign))
 			{
 				InOutMaxX = FMath::Min<int32>(FMath::CeilToInt(ViewSizeX * X + ViewX),InOutMaxX);
 			}
@@ -1829,13 +1835,13 @@ int32 FMath::PlaneAABBRelativePosition(const FPlane& P, const FBox& AABB)
 	FVector Vmin, Vmax;
 
 	// Bypass the slow FVector[] operator. Not RESTRICT because it won't update Vmin, Vmax
-	float* VminPtr = (float*)&Vmin;
-	float* VmaxPtr = (float*)&Vmax;
+	FVector::FReal* VminPtr = (FVector::FReal*)&Vmin;
+	FVector::FReal* VmaxPtr = (FVector::FReal*)&Vmax;
 
 	// Use restrict to get better instruction scheduling and to bypass the slow FVector[] operator
-	const float* RESTRICT AABBMinPtr = (const float*)&AABB.Min;
-	const float* RESTRICT AABBMaxPtr = (const float*)&AABB.Max;
-	const float* RESTRICT PlanePtr = (const float*)&P;
+	const FVector::FReal* RESTRICT AABBMinPtr = (const FVector::FReal*)&AABB.Min;
+	const FVector::FReal* RESTRICT AABBMaxPtr = (const FVector::FReal*)&AABB.Max;
+	const FPlane::FReal* RESTRICT PlanePtr = (const FPlane::FReal*)&P;
 
 	for(int32 Idx=0;Idx<3;++Idx)
 	{
@@ -1852,8 +1858,8 @@ int32 FMath::PlaneAABBRelativePosition(const FPlane& P, const FBox& AABB)
 	}
 
 	// if either diagonal is right on the plane, or one is on either side we have an interesection
-	float dMax = P.PlaneDot(Vmax);
-	float dMin = P.PlaneDot(Vmin);
+	FPlane::FReal dMax = P.PlaneDot(Vmax);
+	FPlane::FReal dMin = P.PlaneDot(Vmin);
 
 	// if Max is below plane, or Min is above we know there is no intersection.. otherwise there must be one
 	if (dMax < 0.f)
@@ -1883,8 +1889,8 @@ bool FMath::SphereConeIntersection(const FVector& SphereCenter, float SphereRadi
 
 	FVector U = ConeAxis * (-SphereRadius / ConeAngleSin);
 	FVector D = SphereCenter - U;
-	float dsqr = D | D;
-	float e = ConeAxis | D;
+	FVector::FReal dsqr = D | D;
+	FVector::FReal e = ConeAxis | D;
 
 	if(e > 0 && e * e >= dsqr * FMath::Square(ConeAngleCos))
 	{
@@ -1954,8 +1960,8 @@ FVector FMath::ClosestPointOnTriangleToPoint(const FVector& Point, const FVector
 
 FVector FMath::GetBaryCentric2D(const FVector& Point, const FVector& A, const FVector& B, const FVector& C)
 {
-	float a = ((B.Y-C.Y)*(Point.X-C.X) + (C.X-B.X)*(Point.Y-C.Y)) / ((B.Y-C.Y)*(A.X-C.X) + (C.X-B.X)*(A.Y-C.Y));
-	float b = ((C.Y-A.Y)*(Point.X-C.X) + (A.X-C.X)*(Point.Y-C.Y)) / ((B.Y-C.Y)*(A.X-C.X) + (C.X-B.X)*(A.Y-C.Y));
+	FVector::FReal a = ((B.Y-C.Y)*(Point.X-C.X) + (C.X-B.X)*(Point.Y-C.Y)) / ((B.Y-C.Y)*(A.X-C.X) + (C.X-B.X)*(A.Y-C.Y));
+	FVector::FReal b = ((C.Y-A.Y)*(Point.X-C.X) + (A.X-C.X)*(Point.Y-C.Y)) / ((B.Y-C.Y)*(A.X-C.X) + (C.X-B.X)*(A.Y-C.Y));
 
 	return FVector(a, b, 1.0f - a - b);	
 }
@@ -1983,15 +1989,15 @@ FVector FMath::ComputeBaryCentric2D(const FVector& Point, const FVector& A, cons
 	const FVector N = TriNorm.GetSafeNormal();
 
 	// Compute twice area of triangle ABC
-	const float AreaABCInv = 1.0f / (N | TriNorm);
+	const FVector::FReal AreaABCInv = 1.0f / (N | TriNorm);
 
 	// Compute a contribution
-	const float AreaPBC = N | ((B-Point) ^ (C-Point));
-	const float a = AreaPBC * AreaABCInv;
+	const FVector::FReal AreaPBC = N | ((B-Point) ^ (C-Point));
+	const FVector::FReal a = AreaPBC * AreaABCInv;
 
 	// Compute b contribution
-	const float AreaPCA = N | ((C-Point) ^ (A-Point));
-	const float b = AreaPCA * AreaABCInv;
+	const FVector::FReal AreaPCA = N | ((C-Point) ^ (A-Point));
+	const FVector::FReal b = AreaPCA * AreaABCInv;
 
 	// Compute c contribution
 	return FVector(a, b, 1.0f - a - b);
@@ -2018,7 +2024,7 @@ FVector4 FMath::ComputeBaryCentric3D(const FVector& Point, const FVector& A, con
 	//The point V can be expressed as Ax=v where x is the vector containing the weights {w1...wn}
 	//Solve for x by multiplying both sides by AInv   (AInv * A)x = AInv * v ==> x = AInv * v
 	const FMatrix InvSolvMat = SolvMat.Inverse();
-	const FPlane BaryCoords = InvSolvMat.TransformVector(V);	 
+	const FPlane4f BaryCoords = (FPlane4f)InvSolvMat.TransformVector(V);		// LWC_TODO: Precision loss
 
 	//Reorder the weights to be a, b, c, d
 	return FVector4(1.0f - BaryCoords.X - BaryCoords.Y - BaryCoords.Z, BaryCoords.X, BaryCoords.Y, BaryCoords.Z);
@@ -2116,9 +2122,9 @@ void FMath::SphereDistToLine(FVector SphereOrigin, float SphereRadius, FVector L
 	//const float A = NormalizedLineDir | NormalizedLineDir  (this is 1 because normalized)
 	//solving quadratic formula in terms of t where closest point = LineOrigin + t * NormalizedLineDir
 	const FVector LineOriginToSphereOrigin = SphereOrigin - LineOrigin;
-	const float B = -2.f * (NormalizedLineDir | LineOriginToSphereOrigin);
-	const float C = LineOriginToSphereOrigin.SizeSquared() - FMath::Square(SphereRadius);
-	const float D	= FMath::Square(B) - 4.f * C;
+	const FVector::FReal B = -2.f * (NormalizedLineDir | LineOriginToSphereOrigin);
+	const FVector::FReal C = LineOriginToSphereOrigin.SizeSquared() - FMath::Square(SphereRadius);
+	const FVector::FReal D	= FMath::Square(B) - 4.f * C;
 
 	if( D <= KINDA_SMALL_NUMBER )
 	{
@@ -2129,10 +2135,10 @@ void FMath::SphereDistToLine(FVector SphereOrigin, float SphereRadius, FVector L
 	else
 	{
 		// Line intersecting sphere in 2 points. Pick closest to line origin.
-		const float	E	= FMath::Sqrt(D);
-		const float T1	= (-B + E) * 0.5f;
-		const float T2	= (-B - E) * 0.5f;
-		const float T	= FMath::Abs( T1 ) == FMath::Abs( T2 ) ? FMath::Abs( T1 ) : FMath::Abs( T1 ) < FMath::Abs( T2 ) ? T1 : T2;	// In the case where both points are exactly the same distance we take the one in the direction of LineDir
+		const FVector::FReal	E	= FMath::Sqrt(D);
+		const FVector::FReal T1	= (-B + E) * 0.5f;
+		const FVector::FReal T2	= (-B - E) * 0.5f;
+		const FVector::FReal T	= FMath::Abs( T1 ) == FMath::Abs( T2 ) ? FMath::Abs( T1 ) : FMath::Abs( T1 ) < FMath::Abs( T2 ) ? T1 : T2;	// In the case where both points are exactly the same distance we take the one in the direction of LineDir
 
 		OutClosestPoint	= LineOrigin + T * NormalizedLineDir;
 	}
@@ -2146,16 +2152,16 @@ bool FMath::GetDistanceWithinConeSegment(FVector Point, FVector ConeStartPoint,
 
 	FVector PointOnCone; // Stores the point on the cone's center line closest to our target point.
 
-	const float Distance = FMath::PointDistToLine(Point, ConeLine, ConeStartPoint, PointOnCone); // distance is how far from the viewline we are
+	const FVector::FReal Distance = FMath::PointDistToLine(Point, ConeLine, ConeStartPoint, PointOnCone); // distance is how far from the viewline we are
 
 	PercentageOut = 0.0; // start assuming we're outside cone until proven otherwise.
 
 	const FVector VectToStart = ConeStartPoint - PointOnCone;
 	const FVector VectToEnd = (ConeStartPoint + ConeLine) - PointOnCone;
 	
-	const float ConeLengthSqr = ConeLine.SizeSquared();
-	const float DistToStartSqr = VectToStart.SizeSquared();
-	const float DistToEndSqr = VectToEnd.SizeSquared();
+	const FVector::FReal ConeLengthSqr = ConeLine.SizeSquared();
+	const FVector::FReal DistToStartSqr = VectToStart.SizeSquared();
+	const FVector::FReal DistToEndSqr = VectToEnd.SizeSquared();
 
 	if (DistToStartSqr > ConeLengthSqr || DistToEndSqr > ConeLengthSqr)
 	{
@@ -2163,13 +2169,13 @@ bool FMath::GetDistanceWithinConeSegment(FVector Point, FVector ConeStartPoint,
 		return false;
 	}
 
-	const float PercentAlongCone = FMath::Sqrt(DistToStartSqr) / FMath::Sqrt(ConeLengthSqr); // don't have to catch outside 0->1 due to above code (saves 2 sqrts if outside)
-	const float RadiusAtPoint = RadiusAtStart + ((RadiusAtEnd - RadiusAtStart) * PercentAlongCone);
+	const FVector::FReal PercentAlongCone = FMath::Sqrt(DistToStartSqr) / FMath::Sqrt(ConeLengthSqr); // don't have to catch outside 0->1 due to above code (saves 2 sqrts if outside)
+	const FVector::FReal RadiusAtPoint = RadiusAtStart + ((RadiusAtEnd - RadiusAtStart) * PercentAlongCone);
 
 	if(Distance > RadiusAtPoint) // target is farther from the line than the radius at that distance)
 		return false;
 
-	PercentageOut = RadiusAtPoint > 0.0f ? (RadiusAtPoint - Distance) / RadiusAtPoint : 1.0f;
+	PercentageOut = RadiusAtPoint > 0.0f ? float((RadiusAtPoint - Distance) / RadiusAtPoint) : 1.0f;
 
 	return true;
 }
@@ -2213,11 +2219,12 @@ bool FMath::GetDotDistance
 	const FVector NoZProjDir = (NormalDir - (NormalDir | AxisZ) * AxisZ).GetSafeNormal();
 	
 	// Figure out if projection is on right or left.
-	const float AzimuthSign = ( (NoZProjDir | AxisY) < 0.f ) ? -1.f : 1.f;
+	const FVector::FReal AzimuthSign = ( (NoZProjDir | AxisY) < 0.f ) ? -1.f : 1.f;
 
-	OutDotDist.Y = NormalDir | AxisZ;
-	const float DirDotX	= NoZProjDir | AxisX;
-	OutDotDist.X = AzimuthSign * FMath::Abs(DirDotX);
+	using FVec2Real = decltype(FVector2D::X);
+	OutDotDist.Y = FVec2Real(NormalDir | AxisZ);
+	const FVector::FReal DirDotX	= NoZProjDir | AxisX;
+	OutDotDist.X = FVec2Real(AzimuthSign * FMath::Abs(DirDotX));
 
 	return (DirDotX >= 0.f );
 }
@@ -2234,12 +2241,13 @@ FVector2D FMath::GetAzimuthAndElevation
 	// Find projected point (on AxisX and AxisY, remove AxisZ component)
 	const FVector NoZProjDir = (NormalDir - (NormalDir | AxisZ) * AxisZ).GetSafeNormal();
 	// Figure out if projection is on right or left.
-	const float AzimuthSign = ((NoZProjDir | AxisY) < 0.f) ? -1.f : 1.f;
-	const float ElevationSin = NormalDir | AxisZ;
-	const float AzimuthCos = NoZProjDir | AxisX;
+	const FVector::FReal AzimuthSign = ((NoZProjDir | AxisY) < 0.f) ? -1.f : 1.f;
+	const FVector::FReal ElevationSin = NormalDir | AxisZ;
+	const FVector::FReal AzimuthCos = NoZProjDir | AxisX;
 
 	// Convert to Angles in Radian.
-	return FVector2D(FMath::Acos(AzimuthCos) * AzimuthSign, FMath::Asin(ElevationSin));
+	using FVec2Real = decltype(FVector2D::X);
+	return FVector2D(FVec2Real(FMath::Acos(AzimuthCos) * AzimuthSign), (FVec2Real)FMath::Asin(ElevationSin));
 }
 
 CORE_API FVector FMath::VInterpNormalRotationTo(const FVector& Current, const FVector& Target, float DeltaTime, float RotationSpeedDegrees)
@@ -2267,8 +2275,8 @@ CORE_API FVector FMath::VInterpNormalRotationTo(const FVector& Current, const FV
 CORE_API FVector FMath::VInterpConstantTo(const FVector& Current, const FVector& Target, float DeltaTime, float InterpSpeed)
 {
 	const FVector Delta = Target - Current;
-	const float DeltaM = Delta.Size();
-	const float MaxStep = InterpSpeed * DeltaTime;
+	const FVector::FReal DeltaM = Delta.Size();
+	const FVector::FReal MaxStep = InterpSpeed * DeltaTime;
 
 	if( DeltaM > MaxStep )
 	{
@@ -3293,7 +3301,7 @@ float FMath::ClampAngle(float AngleDegrees, float MinAngleDegrees, float MaxAngl
 void FMath::ApplyScaleToFloat(float& Dst, const FVector& DeltaScale, float Magnitude)
 {
 	const float Multiplier = ( DeltaScale.X > 0.0f || DeltaScale.Y > 0.0f || DeltaScale.Z > 0.0f ) ? Magnitude : -Magnitude;
-	Dst += Multiplier * DeltaScale.Size();
+	Dst += Multiplier * (float)DeltaScale.Size();
 	Dst = FMath::Max( 0.0f, Dst );
 }
 
@@ -3453,15 +3461,15 @@ float FMath::PerlinNoise3D(const FVector& Location)
 {
 	using namespace FMathPerlinHelpers;
 
-	float Xfl = FMath::FloorToFloat(Location.X);
-	float Yfl = FMath::FloorToFloat(Location.Y);
-	float Zfl = FMath::FloorToFloat(Location.Z);
+	float Xfl = FMath::FloorToFloat((float)Location.X);		// LWC_TODO: Precision loss
+	float Yfl = FMath::FloorToFloat((float)Location.Y);
+	float Zfl = FMath::FloorToFloat((float)Location.Z);
 	int32 Xi = (int32)(Xfl) & 255;
 	int32 Yi = (int32)(Yfl) & 255;
 	int32 Zi = (int32)(Zfl) & 255;
-	float X = Location.X - Xfl;
-	float Y = Location.Y - Yfl;
-	float Z = Location.Z - Zfl;
+	float X = (float)Location.X - Xfl;
+	float Y = (float)Location.Y - Yfl;
+	float Z = (float)Location.Z - Zfl;
 	float Xm1 = X - 1.0f;
 	float Ym1 = Y - 1.0f;
 	float Zm1 = Z - 1.0f;
diff --git a/Engine/Source/Runtime/Core/Private/Misc/DefaultValueHelper.cpp b/Engine/Source/Runtime/Core/Private/Misc/DefaultValueHelper.cpp
--- a/Engine/Source/Runtime/Core/Private/Misc/DefaultValueHelper.cpp
+++ b/Engine/Source/Runtime/Core/Private/Misc/DefaultValueHelper.cpp
@@ -535,7 +535,8 @@ bool FDefaultValueHelper::ParseRotator(const FString& Source, FRotator& OutVal)
 	FVector Vector;
 	if( ParseVector( Source, Vector ) )
 	{
-		OutVal = FRotator(Vector.X, Vector.Y, Vector.Z);
+		using FReal = decltype(FRotator::Pitch);
+		OutVal = FRotator((FReal)Vector.X, (FReal)Vector.Y, (FReal)Vector.Z);
 		return true;
 	}
 
diff --git a/Engine/Source/Runtime/Core/Private/Tests/Containers/StridedViewTest.cpp b/Engine/Source/Runtime/Core/Private/Tests/Containers/StridedViewTest.cpp
--- a/Engine/Source/Runtime/Core/Private/Tests/Containers/StridedViewTest.cpp
+++ b/Engine/Source/Runtime/Core/Private/Tests/Containers/StridedViewTest.cpp
@@ -312,12 +312,12 @@ bool FStridedViewTest::RunTest(const FString& Parameters)
 
 	{
 		TStridedView<FVector::FReal> ViewOfZs = MakeStridedView((int32)sizeof(Structs[0]), &Structs[0].Position.Z, Num);
-		float ExpectedSum = 0.0;
+		FVector::FReal ExpectedSum = 0.0;
 		for (int32 i = 0; i < Num; ++i)
 		{
 			ExpectedSum += Structs[i].Position.Z;
 		}
-		float Sum = Algo::Accumulate(ViewOfZs, 0.0f);
+		FVector::FReal Sum = Algo::Accumulate(ViewOfZs, FVector::FReal(0));
 		TestEqual(TEXT("Algo::Accumulate(MakeStridedView((int32)sizeof(Structs[0]), &Structs[0].Position.Z, Num), 0.0f)"), Sum, ExpectedSum);
 	}
 
diff --git a/Engine/Source/Runtime/Core/Public/Math/Box.h b/Engine/Source/Runtime/Core/Public/Math/Box.h
--- a/Engine/Source/Runtime/Core/Public/Math/Box.h
+++ b/Engine/Source/Runtime/Core/Public/Math/Box.h
@@ -169,7 +169,7 @@ public:
 	 * @param Point The point.
 	 * @return The distance.
 	 */
-	FORCEINLINE float ComputeSquaredDistanceToPoint( const FVector& Point ) const
+	FORCEINLINE FVector::FReal ComputeSquaredDistanceToPoint( const FVector& Point ) const
 	{
 		return ComputeSquaredDistanceFromBoxToPoint(Min, Max, Point);
 	}
@@ -177,7 +177,7 @@ public:
 	/**
 	 * Calculates squared distance between two boxes.
 	 */
-	FORCEINLINE float ComputeSquaredDistanceToBox(const FBox& Box) const
+	FORCEINLINE FVector::FReal ComputeSquaredDistanceToBox(const FBox& Box) const
 	{
 		FVector AxisDistances = (GetCenter() - Box.GetCenter()).GetAbs() - (GetExtent() + Box.GetExtent());
 		AxisDistances = FVector::Max(AxisDistances, FVector(0.0f, 0.0f, 0.0f));
@@ -313,9 +313,9 @@ public:
 	 * @return The box volume.
 	 * @see GetCenter, GetCenterAndExtents, GetExtent, GetSize
 	 */
-	FORCEINLINE float GetVolume() const
+	FORCEINLINE FVector::FReal GetVolume() const
 	{
-		return ((Max.X - Min.X) * (Max.Y - Min.Y) * (Max.Z - Min.Z));
+		return (Max.X - Min.X) * (Max.Y - Min.Y) * (Max.Z - Min.Z);
 	}
 
 	/**
@@ -768,7 +768,7 @@ inline bool FMath::LineBoxIntersection
 
 	if(bStartIsOutside)
 	{
-		const float	MaxTime = Max3(Time.X,Time.Y,Time.Z);
+		const FVector::FReal MaxTime = Max3(Time.X,Time.Y,Time.Z);
 
 		if(MaxTime >= 0.0f && MaxTime <= 1.0f)
 		{
@@ -803,10 +803,10 @@ inline bool FMath::LineBoxIntersection
  *
  * @return Whether the sphere/box intersect or not.
  */
-FORCEINLINE bool FMath::SphereAABBIntersection(const FVector& SphereCenter, const float RadiusSquared, const FBox& AABB)
+FORCEINLINE bool FMath::SphereAABBIntersection(const FVector& SphereCenter, const double RadiusSquared, const FBox& AABB)
 {
 	// Accumulates the distance as we iterate axis
-	float DistSquared = 0.f;
+	FVector::FReal DistSquared = 0.f;
 	// Check each axis for min/max and add the distance accordingly
 	// NOTE: Loop manually unrolled for > 2x speed up
 	if (SphereCenter.X < AABB.Min.X)
diff --git a/Engine/Source/Runtime/Core/Public/Math/BoxSphereBounds.h b/Engine/Source/Runtime/Core/Public/Math/BoxSphereBounds.h
--- a/Engine/Source/Runtime/Core/Public/Math/BoxSphereBounds.h
+++ b/Engine/Source/Runtime/Core/Public/Math/BoxSphereBounds.h
@@ -15,6 +15,8 @@
  */
 struct FBoxSphereBounds
 {
+	using FReal = FVector::FReal;
+
 	/** Holds the origin of the bounding box and sphere. */
 	FVector	Origin;
 
@@ -22,7 +24,7 @@ struct FBoxSphereBounds
 	FVector BoxExtent;
 
 	/** Holds the radius of the bounding sphere. */
-	FVector::FReal SphereRadius;
+	FReal SphereRadius;
 
 public:
 
@@ -49,7 +51,7 @@ public:
 	 * @param InBoxExtent half size of box.
 	 * @param InSphereRadius radius of the sphere.
 	 */
-	FBoxSphereBounds( const FVector& InOrigin, const FVector& InBoxExtent, float InSphereRadius )
+	FBoxSphereBounds( const FVector& InOrigin, const FVector& InBoxExtent, FReal InSphereRadius )
 		: Origin(InOrigin)
 		, BoxExtent(InBoxExtent)
 		, SphereRadius(InSphereRadius)
@@ -143,7 +145,7 @@ public:
 	 * @param Point The point.
 	 * @return The distance.
 	 */
-	FORCEINLINE float ComputeSquaredDistanceFromBoxToPoint( const FVector& Point ) const
+	FORCEINLINE FReal ComputeSquaredDistanceFromBoxToPoint( const FVector& Point ) const
 	{
 		FVector Mins = Origin - BoxExtent;
 		FVector Maxs = Origin + BoxExtent;
@@ -161,7 +163,7 @@ public:
 	 */
 	FORCEINLINE static bool SpheresIntersect(const FBoxSphereBounds& A, const FBoxSphereBounds& B, float Tolerance = KINDA_SMALL_NUMBER)
 	{
-		return (A.Origin - B.Origin).SizeSquared() <= FMath::Square(FMath::Max<FVector::FReal>(0.f, A.SphereRadius + B.SphereRadius + Tolerance));
+		return (A.Origin - B.Origin).SizeSquared() <= FMath::Square(FMath::Max<FReal>(0.f, A.SphereRadius + B.SphereRadius + Tolerance));
 	}
 
 	/**
@@ -218,7 +220,7 @@ public:
 	 * @param ExpandAmount The size to increase by.
 	 * @return A new box with the expanded size.
 	 */
-	FORCEINLINE FBoxSphereBounds ExpandBy( float ExpandAmount ) const
+	FORCEINLINE FBoxSphereBounds ExpandBy( FReal ExpandAmount ) const
 	{
 		return FBoxSphereBounds(Origin, BoxExtent + ExpandAmount, SphereRadius + ExpandAmount);
 	}
@@ -333,7 +335,7 @@ FORCEINLINE FBoxSphereBounds::FBoxSphereBounds( const FVector* Points, uint32 Nu
 
 	for (uint32 PointIndex = 0; PointIndex < NumPoints; PointIndex++)
 	{
-		SquaredSphereRadius = FMath::Max<FVector::FReal>(SquaredSphereRadius, (Points[PointIndex] - Origin).SizeSquared());
+		SquaredSphereRadius = (float)FMath::Max<FReal>(SquaredSphereRadius, (Points[PointIndex] - Origin).SizeSquared());	// LWC_TODO: Precision loss
 	}
 
 	SphereRadius = FMath::Sqrt(SquaredSphereRadius);
@@ -354,7 +356,7 @@ FORCEINLINE FBoxSphereBounds FBoxSphereBounds::operator+( const FBoxSphereBounds
 	// build a bounding sphere from the bounding box's origin and the radii of A and B.
 	FBoxSphereBounds Result(BoundingBox);
 
-	Result.SphereRadius = FMath::Min<FVector::FReal>(Result.SphereRadius, FMath::Max((Origin - Result.Origin).Size() + SphereRadius, (Other.Origin - Result.Origin).Size() + Other.SphereRadius));
+	Result.SphereRadius = FMath::Min<FReal>(Result.SphereRadius, FMath::Max((Origin - Result.Origin).Size() + SphereRadius, (Other.Origin - Result.Origin).Size() + Other.SphereRadius));
 	Result.DiagnosticCheckNaN();
 
 	return Result;
diff --git a/Engine/Source/Runtime/Core/Public/Math/ClipProjectionMatrix.h b/Engine/Source/Runtime/Core/Public/Math/ClipProjectionMatrix.h
--- a/Engine/Source/Runtime/Core/Public/Math/ClipProjectionMatrix.h
+++ b/Engine/Source/Runtime/Core/Public/Math/ClipProjectionMatrix.h
@@ -25,7 +25,7 @@ public:
 
 private:
 	/** return sign of a number */
-	FORCEINLINE float sgn( float a );
+	FORCEINLINE FMatrix::FReal sgn(FMatrix::FReal a );
 };
 
 
@@ -55,7 +55,7 @@ FMatrix(SrcProjMat)
 }
 
 
-FORCEINLINE float FClipProjectionMatrix::sgn( float a )
+FORCEINLINE FMatrix::FReal FClipProjectionMatrix::sgn(FMatrix::FReal a )
 {
 	if (a > 0.0f) return (1.0f);
 	if (a < 0.0f) return (-1.0f);
diff --git a/Engine/Source/Runtime/Core/Public/Math/ConvexHull2d.h b/Engine/Source/Runtime/Core/Public/Math/ConvexHull2d.h
--- a/Engine/Source/Runtime/Core/Public/Math/ConvexHull2d.h
+++ b/Engine/Source/Runtime/Core/Public/Math/ConvexHull2d.h
@@ -90,12 +90,12 @@ namespace ConvexHull2D
 	}
 
 	/** Returns <0 if C is left of A-B */
-	inline float ComputeDeterminant(const FVector& A, const FVector& B, const FVector& C)
+	inline FVector::FReal ComputeDeterminant(const FVector& A, const FVector& B, const FVector& C)
 	{
-		const float u1 = B.X - A.X;
-		const float v1 = B.Y - A.Y;
-		const float u2 = C.X - A.X;
-		const float v2 = C.Y - A.Y;
+		const FVector::FReal u1 = B.X - A.X;
+		const FVector::FReal v1 = B.Y - A.Y;
+		const FVector::FReal u2 = C.X - A.X;
+		const FVector::FReal v2 = C.Y - A.Y;
 
 		return u1 * v2 - v1 * u2;
 	}
@@ -186,7 +186,7 @@ namespace ConvexHull2D
 				const FVector & A = Points.GetData()[Hull];
 				const FVector & B = Points.GetData()[NextPoint];
 				const FVector & C = Points.GetData()[j];
-				float Deter = ComputeDeterminant(A, B, C);
+				FVector::FReal Deter = ComputeDeterminant(A, B, C);
 
 				// 0.001 Bias is to stop floating point errors, when comparing points on a straight line; KINDA_SMALL_NUMBER was slightly too small to use.
 				if(Deter < -0.001)
diff --git a/Engine/Source/Runtime/Core/Public/Math/Matrix.inl b/Engine/Source/Runtime/Core/Public/Math/Matrix.inl
--- a/Engine/Source/Runtime/Core/Public/Math/Matrix.inl
+++ b/Engine/Source/Runtime/Core/Public/Math/Matrix.inl
@@ -194,7 +194,8 @@ FORCEINLINE FVector4 TMatrix<T>::TransformFVector4(const FVector4& P) const
 template<typename T>
 FORCEINLINE FVector4 TMatrix<T>::TransformPosition(const TVector<T>& V) const
 {
-	return TransformFVector4(FVector4(V.X, V.Y, V.Z, 1.0f));
+	using FVec4Real = decltype(FVector4::X);
+	return TransformFVector4(FVector4((FVec4Real)V.X, (FVec4Real)V.Y, (FVec4Real)V.Z, 1.0f));
 }
 
 /** Inverts the matrix and then transforms V - correctly handles scaling in this matrix. */
@@ -214,7 +215,8 @@ FORCEINLINE TVector<T> TMatrix<T>::InverseTransformPosition(const TVector<T>& V)
 template<typename T>
 FORCEINLINE FVector4 TMatrix<T>::TransformVector(const TVector<T>& V) const
 {
-	return TransformFVector4(FVector4(V.X, V.Y, V.Z, 0.0f));
+	using FVec4Real = decltype(FVector4::X);
+	return TransformFVector4(FVector4((FVec4Real)V.X, (FVec4Real)V.Y, (FVec4Real)V.Z, 0.0f));
 }
 
 /** Faster version of InverseTransformVector that assumes no scaling. WARNING: Will NOT work correctly if there is scaling in the matrix. */
diff --git a/Engine/Source/Runtime/Core/Public/Math/OrientedBox.h b/Engine/Source/Runtime/Core/Public/Math/OrientedBox.h
--- a/Engine/Source/Runtime/Core/Public/Math/OrientedBox.h
+++ b/Engine/Source/Runtime/Core/Public/Math/OrientedBox.h
@@ -89,13 +89,13 @@ FORCEINLINE void FOrientedBox::CalcVertices( FVector* Verts ) const
 
 FORCEINLINE FFloatInterval FOrientedBox::Project( const FVector& Axis ) const
 {
-	static const float Signs[] = {-1.0f, 1.0f};
+	static const FVector::FReal Signs[] = {-1.0f, 1.0f};
 
 	// calculate the projections of the box center and the extent-scaled axes.
-	float ProjectedCenter = Axis | Center;
-	float ProjectedAxisX = Axis | (ExtentX * AxisX);
-	float ProjectedAxisY = Axis | (ExtentY * AxisY);
-	float ProjectedAxisZ = Axis | (ExtentZ * AxisZ);
+	FVector::FReal ProjectedCenter = Axis | Center;
+	FVector::FReal ProjectedAxisX = Axis | (ExtentX * AxisX);
+	FVector::FReal ProjectedAxisY = Axis | (ExtentY * AxisY);
+	FVector::FReal ProjectedAxisZ = Axis | (ExtentZ * AxisZ);
 
 	FFloatInterval ProjectionInterval;
 
@@ -106,7 +106,7 @@ FORCEINLINE FFloatInterval FOrientedBox::Project( const FVector& Axis ) const
 			for (int32 k = 0; k < 2; k++)
 			{
 				// project the box vertex onto the axis.
-				float ProjectedVertex = ProjectedCenter + Signs[i] * ProjectedAxisX + Signs[j] * ProjectedAxisY + Signs[k] * ProjectedAxisZ;
+				float ProjectedVertex = float(ProjectedCenter + Signs[i] * ProjectedAxisX + Signs[j] * ProjectedAxisY + Signs[k] * ProjectedAxisZ);	// LWC_TODO: Precision loss
 
 				// if necessary, expand the projection interval to include the box vertex projection.
 				ProjectionInterval.Include(ProjectedVertex);
diff --git a/Engine/Source/Runtime/Core/Public/Math/Quat.h b/Engine/Source/Runtime/Core/Public/Math/Quat.h
--- a/Engine/Source/Runtime/Core/Public/Math/Quat.h
+++ b/Engine/Source/Runtime/Core/Public/Math/Quat.h
@@ -14,6 +14,8 @@
 
 class Error;
 
+// LWC_TODO: Float casting in here to inhibit precision loss warnings from LWC enabled types. Revisit(float) casts / usage when LWC is enabled for this struct.
+
 /**
  * Floating point quaternion that can represent a rotation about an axis in 3-D space.
  * The X, Y, Z, W components also double as the Axis/Angle format.
@@ -619,20 +621,20 @@ inline FQuat::FQuat(const FMatrix& M)
 #endif
 
 	//const MeReal *const t = (MeReal *) tm;
-	float	s;
+	FMatrix::FReal	s;
 
 	// Check diagonal (trace)
-	const float tr = M.M[0][0] + M.M[1][1] + M.M[2][2];
+	const FMatrix::FReal tr = M.M[0][0] + M.M[1][1] + M.M[2][2];
 
 	if (tr > 0.0f) 
 	{
-		float InvS = FMath::InvSqrt(tr + 1.f);
-		this->W = 0.5f * (1.f / InvS);
+		FMatrix::FReal InvS = FMath::InvSqrt(tr + 1.f);
+		this->W = float(0.5f * (1.f / InvS));
 		s = 0.5f * InvS;
 
-		this->X = (M.M[1][2] - M.M[2][1]) * s;
-		this->Y = (M.M[2][0] - M.M[0][2]) * s;
-		this->Z = (M.M[0][1] - M.M[1][0]) * s;
+		this->X = float((M.M[1][2] - M.M[2][1]) * s);
+		this->Y = float((M.M[2][0] - M.M[0][2]) * s);
+		this->Z = float((M.M[0][1] - M.M[1][0]) * s);
 	} 
 	else 
 	{
@@ -651,9 +653,9 @@ inline FQuat::FQuat(const FMatrix& M)
  
 		s = M.M[i][i] - M.M[j][j] - M.M[k][k] + 1.0f;
 
-		float InvS = FMath::InvSqrt(s);
+		FMatrix::FReal InvS = FMath::InvSqrt(s);
 
-		float qt[4];
+		FMatrix::FReal qt[4];
 		qt[i] = 0.5f * (1.f / InvS);
 
 		s = 0.5f * InvS;
@@ -662,10 +664,10 @@ inline FQuat::FQuat(const FMatrix& M)
 		qt[j] = (M.M[i][j] + M.M[j][i]) * s;
 		qt[k] = (M.M[i][k] + M.M[k][i]) * s;
 
-		this->X = qt[0];
-		this->Y = qt[1];
-		this->Z = qt[2];
-		this->W = qt[3];
+		this->X = (float)qt[0];
+		this->Y = (float)qt[1];
+		this->Z = (float)qt[2];
+		this->W = (float)qt[3];
 
 		DiagnosticCheckNaN();
 	}
@@ -692,7 +694,7 @@ inline FMatrix FQuat::operator*(const FMatrix& M) const
 	FQuat Inv = Inverse();
 	for (int32 I=0; I<4; ++I)
 	{
-		FQuat VQ(M.M[I][0], M.M[I][1], M.M[I][2], M.M[I][3]);
+		FQuat VQ((float)M.M[I][0], (float)M.M[I][1], (float)M.M[I][2], (float)M.M[I][3]);
 		VectorQuaternionMultiply(&VT, this, &VQ);
 		VectorQuaternionMultiply(&VR, &VT, &Inv);
 		Result.M[I][0] = VR.X;
@@ -774,9 +776,9 @@ FORCEINLINE FQuat::FQuat(FVector Axis, float AngleRad)
 	float s, c;
 	FMath::SinCos(&s, &c, half_a);
 
-	X = s * Axis.X;
-	Y = s * Axis.Y;
-	Z = s * Axis.Z;
+	X = s * (float)Axis.X;
+	Y = s * (float)Axis.Y;
+	Z = s * (float)Axis.Z;
 	W = c;
 
 	DiagnosticCheckNaN();
diff --git a/Engine/Source/Runtime/Core/Public/Math/RandomStream.h b/Engine/Source/Runtime/Core/Public/Math/RandomStream.h
--- a/Engine/Source/Runtime/Core/Public/Math/RandomStream.h
+++ b/Engine/Source/Runtime/Core/Public/Math/RandomStream.h
@@ -142,7 +142,7 @@ public:
 	FVector GetUnitVector() const
 	{
 		FVector Result;
-		float L;
+		FVector::FReal L;
 
 		do
 		{
diff --git a/Engine/Source/Runtime/Core/Public/Math/SHMath.h b/Engine/Source/Runtime/Core/Public/Math/SHMath.h
--- a/Engine/Source/Runtime/Core/Public/Math/SHMath.h
+++ b/Engine/Source/Runtime/Core/Public/Math/SHMath.h
@@ -457,9 +457,9 @@ inline TSHVector<2> TSHVector<2>::SHBasisFunction(const FVector& Vector)
 {
 	TSHVector<2> Result;
 	Result.V[0] = 0.282095f; 
-	Result.V[1] = -0.488603f * Vector.Y;
-	Result.V[2] = 0.488603f * Vector.Z;
-	Result.V[3] = -0.488603f * Vector.X;
+	Result.V[1] = -0.488603f * (float)Vector.Y;	// LWC_TODO: Precision loss
+	Result.V[2] = 0.488603f * (float)Vector.Z;
+	Result.V[3] = -0.488603f * (float)Vector.X;
 	return Result;
 }
 
@@ -469,16 +469,16 @@ inline TSHVector<3> TSHVector<3>::SHBasisFunction(const FVector& Vector)
 {
 	TSHVector<3> Result;
 	Result.V[0] = 0.282095f; 
-	Result.V[1] = -0.488603f * Vector.Y;
-	Result.V[2] = 0.488603f * Vector.Z;
-	Result.V[3] = -0.488603f * Vector.X;
+	Result.V[1] = -0.488603f * (float)Vector.Y;	// LWC_TODO: Precision loss
+	Result.V[2] = 0.488603f * (float)Vector.Z;
+	Result.V[3] = -0.488603f * (float)Vector.X;
 
 	FVector VectorSquared = Vector * Vector;
-	Result.V[4] = 1.092548f * Vector.X * Vector.Y;
-	Result.V[5] = -1.092548f * Vector.Y * Vector.Z;
-	Result.V[6] = 0.315392f * (3.0f * VectorSquared.Z - 1.0f);
-	Result.V[7] = -1.092548f * Vector.X * Vector.Z;
-	Result.V[8] = 0.546274f * (VectorSquared.X - VectorSquared.Y);
+	Result.V[4] = 1.092548f * float(Vector.X * Vector.Y);
+	Result.V[5] = -1.092548f * float(Vector.Y * Vector.Z);
+	Result.V[6] = 0.315392f * float(3.0f * VectorSquared.Z - 1.0f);
+	Result.V[7] = -1.092548f * float(Vector.X * Vector.Z);
+	Result.V[8] = 0.546274f * float(VectorSquared.X - VectorSquared.Y);
 	return Result;
 }
 
diff --git a/Engine/Source/Runtime/Core/Public/Math/Sphere.h b/Engine/Source/Runtime/Core/Public/Math/Sphere.h
--- a/Engine/Source/Runtime/Core/Public/Math/Sphere.h
+++ b/Engine/Source/Runtime/Core/Public/Math/Sphere.h
@@ -12,12 +12,13 @@
 class FSphere
 {
 public:
+	using FReal = FVector::FReal;
 
 	/** The sphere's center point. */
 	FVector Center;
 
 	/** The sphere's radius. */
-	FVector::FReal W;
+	FReal W;
 
 public:
 
@@ -40,7 +41,7 @@ public:
 	 * @param InV Center of sphere.
 	 * @param InW Radius of sphere.
 	 */
-	FSphere(FVector InV, typename FVector::FReal InW)
+	FSphere(FVector InV, FReal InW)
 		: Center(InV)
 		, W(InW)
 	{ }
@@ -146,7 +147,7 @@ public:
 	 *
 	 * @return Volume (in Unreal units).
 	 */
-	CORE_API float GetVolume() const;
+	CORE_API FReal GetVolume() const;
 
 	/**
 	 * Adds to this bounding box to include a new bounding volume.
@@ -204,7 +205,7 @@ public:
  */
 FORCEINLINE bool FMath::SphereAABBIntersection(const FSphere& Sphere,const FBox& AABB)
 {
-	FVector::FReal RadiusSquared = FMath::Square(Sphere.W);
+	FSphere::FReal RadiusSquared = FMath::Square(Sphere.W);
 	// If the distance is less than or equal to the radius, they intersect
 	return SphereAABBIntersection(Sphere.Center,RadiusSquared,AABB);
 }
@@ -222,7 +223,7 @@ FORCEINLINE FSphere FMath::ComputeBoundingSphereForCone(FVector const& ConeOrigi
 	}
 	else
 	{
-		const float BoundingRadius = ConeRadius / (2.0f * CosConeAngle);
+		const FSphere::FReal BoundingRadius = ConeRadius / (2.0f * CosConeAngle);
 		return FSphere(ConeOrigin + ConeDirection * BoundingRadius, BoundingRadius);
 	}
 }
\ No newline at end of file
diff --git a/Engine/Source/Runtime/Core/Public/Math/UnrealMathUtility.h b/Engine/Source/Runtime/Core/Public/Math/UnrealMathUtility.h
--- a/Engine/Source/Runtime/Core/Public/Math/UnrealMathUtility.h
+++ b/Engine/Source/Runtime/Core/Public/Math/UnrealMathUtility.h
@@ -8,8 +8,6 @@
 #include "Templates/Decay.h"
 #include "Templates/IsFloatingPoint.h"
 #include "Templates/IsIntegral.h"
-//#include "Math/Vector.h"
-
 
 //#define IMPLEMENT_ASSIGNMENT_OPERATOR_MANUALLY
 
@@ -922,10 +920,10 @@ public:
 	template< class T, class U > 
 	static FORCEINLINE_DEBUGGABLE T CubicInterp( const T& P0, const T& T0, const T& P1, const T& T1, const U& A )
 	{
-		const float A2 = A  * A;
-		const float A3 = A2 * A;
+		const U A2 = A * A;
+		const U A3 = A2 * A;
 
-		return (T)(((2*A3)-(3*A2)+1) * P0) + ((A3-(2*A2)+A) * T0) + ((A3-A2) * T1) + (((-2*A3)+(3*A2)) * P1);
+		return T((((2*A3)-(3*A2)+1) * P0) + ((A3-(2*A2)+A) * T0) + ((A3-A2) * T1) + (((-2*A3)+(3*A2)) * P1));
 	}
 
 	/**
@@ -944,9 +942,9 @@ public:
 		T b = -6.f*P0 - 4.f*T0 - 2.f*T1 + 6.f*P1;
 		T c = T0;
 
-		const float A2 = A  * A;
+		const U A2 = A * A;
 
-		return (a * A2) + (b * A) + c;
+		return T((a * A2) + (b * A) + c);
 	}
 
 	/**
@@ -1591,7 +1589,7 @@ public:
 	 *
 	 * @return Whether the sphere/box intersect or not.
 	 */
-	static bool SphereAABBIntersection(const FVector& SphereCenter,const float RadiusSquared,const FBox& AABB);
+	static bool SphereAABBIntersection(const FVector& SphereCenter,const double RadiusSquared,const FBox& AABB);
 
 	/**
 	 * Converts a sphere into a point plus radius squared for the test above
diff --git a/Engine/Source/Runtime/Core/Public/Math/UnrealPlatformMathSSE.h b/Engine/Source/Runtime/Core/Public/Math/UnrealPlatformMathSSE.h
--- a/Engine/Source/Runtime/Core/Public/Math/UnrealPlatformMathSSE.h
+++ b/Engine/Source/Runtime/Core/Public/Math/UnrealPlatformMathSSE.h
@@ -247,32 +247,32 @@ struct TUnrealPlatformMathSSEBase : public Base
 	// LWC_TODO: Default to double variant for non-float types, or just for int64?
 	static FORCEINLINE int32 TruncToInt(double F)
 	{
-		return Base::TruncToInt(F);
+		return Base::TruncToInt((float)F);
 	}
 
 	static FORCEINLINE int32 RoundToInt(double F)
 	{
-		return Base::RoundToInt(F);
+		return Base::RoundToInt((float)F);
 	}
 
 	static FORCEINLINE int32 FloorToInt(double F)
 	{
-		return Base::FloorToInt(F);
+		return Base::FloorToInt((float)F);
 	}
 
 	static FORCEINLINE int32 CeilToInt(double F)
 	{
-		return Base::CeilToInt(F);
+		return Base::CeilToInt((float)F);
 	}
 
 	static FORCEINLINE double InvSqrt(double F)
 	{
-		return Base::InvSqrt(F);
+		return Base::InvSqrt((float)F);
 	}
 
 	static FORCEINLINE double InvSqrtEst(double F)
 	{
-		return Base::InvSqrtEst(F);
+		return Base::InvSqrtEst((float)F);
 	}
 
 };
diff --git a/Engine/Source/Runtime/Core/Public/Math/Vector.h b/Engine/Source/Runtime/Core/Public/Math/Vector.h
--- a/Engine/Source/Runtime/Core/Public/Math/Vector.h
+++ b/Engine/Source/Runtime/Core/Public/Math/Vector.h
@@ -2376,8 +2376,8 @@ template<typename T> const TVector<T> TVector<T>::ZAxisVector(0, 0, 1);
  * @param V Vector to scale.
  * @return Result of multiplication.
  */
-template<typename T>
-FORCEINLINE TVector<T> operator*(float Scale, const TVector<T>& V)	// LWC_TODO: T Scale
+template<typename T, typename T2, TEMPLATE_REQUIRES(std::is_fundamental<T2>::value)>
+FORCEINLINE TVector<T> operator*(T2 Scale, const TVector<T>& V)
 {
 	return V.operator*(Scale);
 }
diff --git a/Engine/Source/Runtime/Experimental/Chaos/Chaos.Build.cs b/Engine/Source/Runtime/Experimental/Chaos/Chaos.Build.cs
--- a/Engine/Source/Runtime/Experimental/Chaos/Chaos.Build.cs
+++ b/Engine/Source/Runtime/Experimental/Chaos/Chaos.Build.cs
@@ -1,5 +1,8 @@
 // Copyright Epic Games, Inc. All Rights Reserved.
 
+using System.Linq;
+using System.Collections.Generic;
+
 namespace UnrealBuildTool.Rules
 {
     public class Chaos : ModuleRules
@@ -39,8 +42,11 @@ namespace UnrealBuildTool.Rules
 			{
 				PublicDefinitions.Add("CHAOS_MEMORY_TRACKING=0");
 			}
-
-			UnsafeTypeCastWarningLevel = WarningLevel.Warning;
+			
+			if (!Target.GlobalDefinitions.Contains("UE_LARGE_WORLD_COORDINATES_DISABLED=0")) // LWC_TODO: Temporarily disabled for LWC while these are properly fixed
+			{
+				UnsafeTypeCastWarningLevel = WarningLevel.Warning;
+			}
 		}
     }
 }
diff --git a/Engine/Source/Runtime/Experimental/ChaosCore/Public/Chaos/Vector.h b/Engine/Source/Runtime/Experimental/ChaosCore/Public/Chaos/Vector.h
--- a/Engine/Source/Runtime/Experimental/ChaosCore/Public/Chaos/Vector.h
+++ b/Engine/Source/Runtime/Experimental/ChaosCore/Public/Chaos/Vector.h
@@ -377,9 +377,9 @@ namespace Chaos
 		TVector()
 		    : FVector4() {}
 		explicit TVector(const FReal x)
-		    : FVector4(x, x, x, x) {}
+		    : FVector4((decltype(FVector4::X))x, (decltype(FVector4::X))x, (decltype(FVector4::X))x, (decltype(FVector4::X))x) {}	// LWC_TODO: Remove casts once FVector4 supports variants
 		TVector(const FReal x, const FReal y, const FReal z, const FReal w)
-		    : FVector4(x, y, z, w) {}
+		    : FVector4((decltype(FVector4::X))x, (decltype(FVector4::X))y, (decltype(FVector4::X))z, (decltype(FVector4::X))w) {}
 		TVector(const FVector4& vec)
 		    : FVector4(vec) {}
 	};
@@ -594,9 +594,9 @@ namespace Chaos
 		TVector()
 		    : FVector2D() {}
 		TVector(const FReal x)
-		    : FVector2D(x, x) {}
+		    : FVector2D((decltype(FVector2D::X))x, (decltype(FVector2D::X))x) {}	// LWC_TODO: Remove casts once FVector2D supports variants
 		TVector(const FReal x, const FReal y)
-		    : FVector2D(x, y) {}
+		    : FVector2D((decltype(FVector2D::X))x, (decltype(FVector2D::X))y) {}
 		TVector(const FVector2D& vec)
 		    : FVector2D(vec) {}
 		TVector(std::istream& Stream)
