Add helper function to query whether there is sufficient space in the persistent download dir - optionally showing a UI to clear space, depending on the platform.

diff --git a/Engine/Source/Runtime/Core/Private/GenericPlatform/GenericPlatformMisc.cpp b/Engine/Source/Runtime/Core/Private/GenericPlatform/GenericPlatformMisc.cpp
--- a/Engine/Source/Runtime/Core/Private/GenericPlatform/GenericPlatformMisc.cpp
+++ b/Engine/Source/Runtime/Core/Private/GenericPlatform/GenericPlatformMisc.cpp
@@ -1815,4 +1815,10 @@ int32 FGenericPlatformMisc::GetPakchunkIndexFromPakFile(const FString& InFilenam
 bool FGenericPlatformMisc::IsPGOEnabled()
 {
 	return PLATFORM_COMPILER_OPTIMIZATION_PG;
-}
\ No newline at end of file
+}
+
+bool FGenericPlatformMisc::CheckPersistentDownloadStorageSpaceAvailable( uint64 BytesRequired, bool bAttemptToUseUI )
+{
+	return true;
+}
+
diff --git a/Engine/Source/Runtime/Core/Public/GenericPlatform/GenericPlatformMisc.h b/Engine/Source/Runtime/Core/Public/GenericPlatform/GenericPlatformMisc.h
--- a/Engine/Source/Runtime/Core/Public/GenericPlatform/GenericPlatformMisc.h
+++ b/Engine/Source/Runtime/Core/Public/GenericPlatform/GenericPlatformMisc.h
@@ -982,6 +982,12 @@ public:
 		return true;
 	}
 
+	/** Ensure that there enough space in the persistent download dir. Returns true if there is sufficient space.
+	 *  If there is insufficient space and bAttemptToUseUI is true, the function will return false and the user may be prompted to clear space
+	 */
+	static bool CheckPersistentDownloadStorageSpaceAvailable( uint64 BytesRequired, bool bAttemptToUseUI );
+
+
 	/**
 	 *	Return the GamePersistentDownloadDir.
 	 *	On some platforms, returns the writable directory for downloaded data that persists across play sessions.
