Added missing tregonometry double NEON vector functions

diff --git a/Engine/Source/Runtime/Core/Public/Math/UnrealMathNeon.h b/Engine/Source/Runtime/Core/Public/Math/UnrealMathNeon.h
--- a/Engine/Source/Runtime/Core/Public/Math/UnrealMathNeon.h
+++ b/Engine/Source/Runtime/Core/Public/Math/UnrealMathNeon.h
@@ -2339,6 +2339,12 @@ FORCEINLINE VectorRegister4Float VectorExp(const VectorRegister4Float& X)
 	return MakeVectorRegister(FMath::Exp(Val[0]), FMath::Exp(Val[1]), FMath::Exp(Val[2]), FMath::Exp(Val[3]));
 }
 
+FORCEINLINE VectorRegister4Double VectorExp(const VectorRegister4Double& X)
+{
+	AlignedDouble4 Val(X);
+	return MakeVectorRegister(FMath::Exp(Val[0]), FMath::Exp(Val[1]), FMath::Exp(Val[2]), FMath::Exp(Val[3]));
+}
+
 //TODO: Vectorize
 FORCEINLINE VectorRegister4Float VectorExp2(const VectorRegister4Float& X)
 {
@@ -2346,6 +2352,12 @@ FORCEINLINE VectorRegister4Float VectorExp2(const VectorRegister4Float& X)
 	return MakeVectorRegister(FMath::Exp2(Val[0]), FMath::Exp2(Val[1]), FMath::Exp2(Val[2]), FMath::Exp2(Val[3]));
 }
 
+FORCEINLINE VectorRegister4Double VectorExp2(const VectorRegister4Double& X)
+{
+	AlignedDouble4 Val(X);
+	return MakeVectorRegister(FMath::Exp2(Val[0]), FMath::Exp2(Val[1]), FMath::Exp2(Val[2]), FMath::Exp2(Val[3]));
+}
+
 //TODO: Vectorize
 FORCEINLINE VectorRegister4Float VectorLog(const VectorRegister4Float& X)
 {
@@ -2353,6 +2365,12 @@ FORCEINLINE VectorRegister4Float VectorLog(const VectorRegister4Float& X)
 	return MakeVectorRegister(FMath::Loge(Val[0]), FMath::Loge(Val[1]), FMath::Loge(Val[2]), FMath::Loge(Val[3]));
 }
 
+FORCEINLINE VectorRegister4Double VectorLog(const VectorRegister4Double& X)
+{
+	AlignedDouble4 Val(X);
+	return MakeVectorRegister(FMath::Loge(Val[0]), FMath::Loge(Val[1]), FMath::Loge(Val[2]), FMath::Loge(Val[3]));
+}
+
 //TODO: Vectorize
 FORCEINLINE VectorRegister4Float VectorLog2(const VectorRegister4Float& X)
 {
@@ -2360,6 +2378,12 @@ FORCEINLINE VectorRegister4Float VectorLog2(const VectorRegister4Float& X)
 	return MakeVectorRegister(FMath::Log2(Val[0]), FMath::Log2(Val[1]), FMath::Log2(Val[2]), FMath::Log2(Val[3]));
 }
 
+FORCEINLINE VectorRegister4Double VectorLog2(const VectorRegister4Double& X)
+{
+	AlignedDouble4 Val(X);
+	return MakeVectorRegister(FMath::Log2(Val[0]), FMath::Log2(Val[1]), FMath::Log2(Val[2]), FMath::Log2(Val[3]));
+}
+
 //TODO: Vectorize
 FORCEINLINE VectorRegister4Float VectorTan(const VectorRegister4Float& X)
 {
@@ -2367,6 +2391,12 @@ FORCEINLINE VectorRegister4Float VectorTan(const VectorRegister4Float& X)
 	return MakeVectorRegister(FMath::Tan(Val[0]), FMath::Tan(Val[1]), FMath::Tan(Val[2]), FMath::Tan(Val[3]));
 }
 
+FORCEINLINE VectorRegister4Double VectorTan(const VectorRegister4Double& X)
+{
+	AlignedDouble4 Val(X);
+	return MakeVectorRegister(FMath::Tan(Val[0]), FMath::Tan(Val[1]), FMath::Tan(Val[2]), FMath::Tan(Val[3]));
+}
+
 //TODO: Vectorize
 FORCEINLINE VectorRegister4Float VectorASin(const VectorRegister4Float& X)
 {
@@ -2374,6 +2404,12 @@ FORCEINLINE VectorRegister4Float VectorASin(const VectorRegister4Float& X)
 	return MakeVectorRegister(FMath::Asin(Val[0]), FMath::Asin(Val[1]), FMath::Asin(Val[2]), FMath::Asin(Val[3]));
 }
 
+FORCEINLINE VectorRegister4Double VectorASin(const VectorRegister4Double& X)
+{
+	AlignedDouble4 Val(X);
+	return MakeVectorRegister(FMath::Asin(Val[0]), FMath::Asin(Val[1]), FMath::Asin(Val[2]), FMath::Asin(Val[3]));
+}
+
 //TODO: Vectorize
 FORCEINLINE VectorRegister4Float VectorACos(const VectorRegister4Float& X)
 {
@@ -2381,6 +2417,12 @@ FORCEINLINE VectorRegister4Float VectorACos(const VectorRegister4Float& X)
 	return MakeVectorRegister(FMath::Acos(Val[0]), FMath::Acos(Val[1]), FMath::Acos(Val[2]), FMath::Acos(Val[3]));
 }
 
+FORCEINLINE VectorRegister4Double VectorACos(const VectorRegister4Double& X)
+{
+	AlignedDouble4 Val(X);
+	return MakeVectorRegister(FMath::Acos(Val[0]), FMath::Acos(Val[1]), FMath::Acos(Val[2]), FMath::Acos(Val[3]));
+}
+
 //TODO: Vectorize
 FORCEINLINE VectorRegister4Float VectorATan(const VectorRegister4Float& X)
 {
@@ -2388,6 +2430,12 @@ FORCEINLINE VectorRegister4Float VectorATan(const VectorRegister4Float& X)
 	return MakeVectorRegister(FMath::Atan(Val[0]), FMath::Atan(Val[1]), FMath::Atan(Val[2]), FMath::Atan(Val[3]));
 }
 
+FORCEINLINE VectorRegister4Double VectorATan(const VectorRegister4Double& X)
+{
+	AlignedDouble4 Val(X);
+	return MakeVectorRegister(FMath::Atan(Val[0]), FMath::Atan(Val[1]), FMath::Atan(Val[2]), FMath::Atan(Val[3]));
+}
+
 //TODO: Vectorize
 FORCEINLINE VectorRegister4Float VectorATan2(const VectorRegister4Float& X, const VectorRegister4Float& Y)
 {
@@ -2400,6 +2448,17 @@ FORCEINLINE VectorRegister4Float VectorATan2(const VectorRegister4Float& X, cons
 							  FMath::Atan2(ValX[3], ValY[3]));
 }
 
+FORCEINLINE VectorRegister4Double VectorATan2(const VectorRegister4Double& X, const VectorRegister4Double& Y)
+{
+	AlignedDouble4 ValX(X);
+	AlignedDouble4 ValY(Y);
+
+	return MakeVectorRegister(FMath::Atan2(ValX[0], ValY[0]),
+							  FMath::Atan2(ValX[1], ValY[1]),
+							  FMath::Atan2(ValX[2], ValY[2]),
+							  FMath::Atan2(ValX[3], ValY[3]));
+}
+
 FORCEINLINE VectorRegister4Float VectorCeil(const VectorRegister4Float& X)
 {
 	return vrndpq_f32(X);
@@ -2525,11 +2584,37 @@ FORCEINLINE VectorRegister4Float VectorSin(const VectorRegister4Float& X)
 	return VectorMultiply(Y, VectorAdd(VectorSinConstantsNEON::B, VectorAbs(Y)));
 }
 
+FORCEINLINE VectorRegister4Double VectorSin(const VectorRegister4Double& X)
+{
+	AlignedDouble4 Doubles(X);
+	Doubles[0] = FMath::Sin(Doubles[0]);
+	Doubles[1] = FMath::Sin(Doubles[1]);
+	Doubles[2] = FMath::Sin(Doubles[2]);
+	Doubles[3] = FMath::Sin(Doubles[3]);
+	return Doubles.ToVectorRegister();
+}
+
 FORCEINLINE VectorRegister4Float VectorCos(const VectorRegister4Float& X)
 {
 	return VectorSin(VectorAdd(X, GlobalVectorConstants::PiByTwo));
 }
 
+FORCEINLINE VectorRegister4Double VectorCos(const VectorRegister4Double& X)
+{
+	AlignedDouble4 Doubles(X);
+	Doubles[0] = FMath::Cos(Doubles[0]);
+	Doubles[1] = FMath::Cos(Doubles[1]);
+	Doubles[2] = FMath::Cos(Doubles[2]);
+	Doubles[3] = FMath::Cos(Doubles[3]);
+	return Doubles.ToVectorRegister();
+}
+
+FORCEINLINE void VectorSinCos(VectorRegister4Double* RESTRICT VSinAngles, VectorRegister4Double* RESTRICT VCosAngles, const VectorRegister4Double* RESTRICT VAngles)
+{
+	*VSinAngles = VectorSin(*VAngles);
+	*VCosAngles = VectorCos(*VAngles);
+}
+
 /**
 * Loads packed RGBA16(4 bytes) from unaligned memory and converts them into 4 FLOATs.
 * IMPORTANT: You need to call VectorResetFloatRegisters() before using scalar FLOATs after you've used this intrinsic!
