Compilation time optimization
  - Split extremely long function into multiple parts to improve codegen time
  - 22.04s -> 11.14s for LowLevelMemTracker.cpp
  - 20.40s -> 13.62s for Module.Core.4_of_17.cpp

diff --git a/Engine/Source/Runtime/Core/Private/HAL/LowLevelMemTracker.cpp b/Engine/Source/Runtime/Core/Private/HAL/LowLevelMemTracker.cpp
--- a/Engine/Source/Runtime/Core/Private/HAL/LowLevelMemTracker.cpp
+++ b/Engine/Source/Runtime/Core/Private/HAL/LowLevelMemTracker.cpp
@@ -1301,8 +1301,8 @@ void FLowLevelMemTracker::FinishInitialise()
 	(void)LLMGetTagUniqueName(ELLMTag::Untagged);
 	InitialiseTagDatas();
 }
-
-void FLowLevelMemTracker::InitialiseTagDatas()
+ 
+void FLowLevelMemTracker::InitialiseTagDatas_SetLLMTagNames()
 {
 	using namespace UE::LLMPrivate;
 
@@ -1324,7 +1324,11 @@ void FLowLevelMemTracker::InitialiseTagDatas()
 	}
 	LLM_ENUM_GENERIC_TAGS(SET_ELLMTAG_NAMES);
 #undef SET_ELLMTAG_NAMES
+}
 
+void FLowLevelMemTracker::InitialiseTagDatas_FinishRegister()
+{
+	using namespace UE::LLMPrivate;
 
 	// Record the central list of ELLMTags in TagDataNameMap, and mark that bootstrapping is complete
 	{
@@ -1347,6 +1351,15 @@ void FLowLevelMemTracker::InitialiseTagDatas()
 
 		bIsBootstrapping = false;
 	}
+}
+
+void FLowLevelMemTracker::InitialiseTagDatas()
+{
+	using namespace UE::LLMPrivate;
+
+	InitialiseTagDatas_SetLLMTagNames();
+
+	InitialiseTagDatas_FinishRegister();
 
 	// Construct the remaining startup tags; allocations when constructing these tags are known to consist only of the central list of ELLMTags so we do not need to bootstrap
 	{
diff --git a/Engine/Source/Runtime/Core/Public/HAL/LowLevelMemTracker.h b/Engine/Source/Runtime/Core/Public/HAL/LowLevelMemTracker.h
--- a/Engine/Source/Runtime/Core/Public/HAL/LowLevelMemTracker.h
+++ b/Engine/Source/Runtime/Core/Public/HAL/LowLevelMemTracker.h
@@ -554,6 +554,8 @@ private:
 	void BootstrapTagDatas();
 	/** Called when we have detected enabled in processcommandline, or as late as possible if callers access features that require full initialisation before then */
 	void FinishInitialise();
+	void InitialiseTagDatas_SetLLMTagNames();
+	void InitialiseTagDatas_FinishRegister();
 	void InitialiseTagDatas();
 	void ClearTagDatas();
 	void RegisterTagDeclaration(FLLMTagDeclaration& TagDeclaration);
