Replaced SSE:InvSqrt with non-estimate calculation - caused in  floating point indeterminancy between AMD and Intel.

diff --git a/Engine/Source/Runtime/Core/Public/Math/UnrealPlatformMathSSE.h b/Engine/Source/Runtime/Core/Public/Math/UnrealPlatformMathSSE.h
--- a/Engine/Source/Runtime/Core/Public/Math/UnrealPlatformMathSSE.h
+++ b/Engine/Source/Runtime/Core/Public/Math/UnrealPlatformMathSSE.h
@@ -19,42 +19,14 @@ namespace UE4
 {
 namespace SSE
 {
-	static FORCEINLINE float InvSqrt(float F)
+	static FORCEINLINE float InvSqrt(float InValue)
 	{
-		// Performs two passes of Newton-Raphson iteration on the hardware estimate
-		//    v^-0.5 = x
-		// => x^2 = v^-1
-		// => 1/(x^2) = v
-		// => F(x) = x^-2 - v
-		//    F'(x) = -2x^-3
-
-		//    x1 = x0 - F(x0)/F'(x0)
-		// => x1 = x0 + 0.5 * (x0^-2 - Vec) * x0^3
-		// => x1 = x0 + 0.5 * (x0 - Vec * x0^3)
-		// => x1 = x0 + x0 * (0.5 - 0.5 * Vec * x0^2)
-		//
-		// This final form has one more operation than the legacy factorization (X1 = 0.5*X0*(3-(Y*X0)*X0)
-		// but retains better accuracy (namely InvSqrt(1) = 1 exactly).
-
-		const __m128 fOneHalf = _mm_set_ss(0.5f);
-		__m128 Y0, X0, X1, X2, FOver2;
+		const __m128 One = _mm_set_ss(1.0f);
+		const __m128 Y0 = _mm_set_ss(InValue);
+		const __m128 X0 = _mm_sqrt_ss(Y0);
+		const __m128 R0 = _mm_div_ss(One, X0);
 		float temp;
-
-		Y0 = _mm_set_ss(F);
-		X0 = _mm_rsqrt_ss(Y0);	// 1/sqrt estimate (12 bits)
-		FOver2 = _mm_mul_ss(Y0, fOneHalf);
-
-		// 1st Newton-Raphson iteration
-		X1 = _mm_mul_ss(X0, X0);
-		X1 = _mm_sub_ss(fOneHalf, _mm_mul_ss(FOver2, X1));
-		X1 = _mm_add_ss(X0, _mm_mul_ss(X0, X1));
-
-		// 2nd Newton-Raphson iteration
-		X2 = _mm_mul_ss(X1, X1);
-		X2 = _mm_sub_ss(fOneHalf, _mm_mul_ss(FOver2, X2));
-		X2 = _mm_add_ss(X1, _mm_mul_ss(X1, X2));
-
-		_mm_store_ss(&temp, X2);
+		_mm_store_ss(&temp, R0);
 		return temp;
 	}
 
