FCStringUtf8 typedef added.

diff --git a/Engine/Source/Runtime/Core/Public/Misc/CString.h b/Engine/Source/Runtime/Core/Public/Misc/CString.h
--- a/Engine/Source/Runtime/Core/Public/Misc/CString.h
+++ b/Engine/Source/Runtime/Core/Public/Misc/CString.h
@@ -409,6 +409,7 @@ public:
 typedef TCString<TCHAR>    FCString;
 typedef TCString<ANSICHAR> FCStringAnsi;
 typedef TCString<WIDECHAR> FCStringWide;
+typedef TCString<UTF8CHAR> FCStringUtf8;
 
 /*-----------------------------------------------------------------------------
 	generic TCString implementations
