FTCHARToUTF8_Convert::Convert can take either the old or new char type, for backwards compatibility.

diff --git a/Engine/Source/Runtime/Core/Public/Containers/StringConv.h b/Engine/Source/Runtime/Core/Public/Containers/StringConv.h
--- a/Engine/Source/Runtime/Core/Public/Containers/StringConv.h
+++ b/Engine/Source/Runtime/Core/Public/Containers/StringConv.h
@@ -288,15 +288,24 @@ public:
 		typename SrcBufferType,
 		std::enable_if_t<FPlatformString::IsCharEncodingCompatibleWith<SrcBufferType, FromType>()>* = nullptr
 	>
-	static FORCEINLINE int32 Convert(ToType* Dest, int32 DestLen, const SrcBufferType* Source, int32 SourceLen)
+	static FORCEINLINE int32 Convert(IntendedToType* Dest, int32 DestLen, const SrcBufferType* Source, int32 SourceLen)
 	{
-		IntendedToType* Result = FPlatformString::Convert((IntendedToType*)Dest, DestLen, (const FromType*)Source, SourceLen);
+		IntendedToType* Result = FPlatformString::Convert(Dest, DestLen, (const FromType*)Source, SourceLen);
 		if (!Result)
 		{
 			return -1;
 		}
 
-		return (int32)(Result - (IntendedToType*)Dest);
+		return (int32)(Result - Dest);
+	}
+	template <
+		typename SrcBufferType,
+		std::enable_if_t<FPlatformString::IsCharEncodingCompatibleWith<SrcBufferType, FromType>()>* = nullptr
+	>
+	static FORCEINLINE int32 Convert(ToType* Dest, int32 DestLen, const SrcBufferType* Source, int32 SourceLen)
+	{
+		// This helper function is only for backwards compatibility and won't be necessary once IntendedToType becomes the ToType
+		return FTCHARToUTF8_Convert::Convert((IntendedToType*)Dest, DestLen, Source, SourceLen);
 	}
 
 	/**
