* added methods to get/return "raw" events (as opposite to "safe" events) from/to FEventPool, refactored FScopedEvent and FEventRef to avoid unnecessary memalloc
* removed EEventPoolTypes as we already have equivalent EEventMode
* added FSharedEventRef that is ref-counted and takes an event from the global pool. Is required as FEventPool is in a private header so not directly accessible from user code

diff --git a/Engine/Source/Runtime/Core/Private/GenericPlatform/GenericPlatformProcess.cpp b/Engine/Source/Runtime/Core/Private/GenericPlatform/GenericPlatformProcess.cpp
--- a/Engine/Source/Runtime/Core/Private/GenericPlatform/GenericPlatformProcess.cpp
+++ b/Engine/Source/Runtime/Core/Private/GenericPlatform/GenericPlatformProcess.cpp
@@ -496,14 +496,14 @@ FEvent* FGenericPlatformProcess::CreateSynchEvent(bool bIsManualReset)
 FEvent* FGenericPlatformProcess::GetSynchEventFromPool(bool bIsManualReset)
 {
 	return bIsManualReset
-		? TLazySingleton<FEventPool<EEventPoolTypes::ManualReset>>::Get().GetEventFromPool()
-		: TLazySingleton<FEventPool<EEventPoolTypes::AutoReset>>::Get().GetEventFromPool();
+		? TLazySingleton<TEventPool<EEventMode::ManualReset>>::Get().GetEventFromPool()
+		: TLazySingleton<TEventPool<EEventMode::AutoReset>>::Get().GetEventFromPool();
 }
 
 void FGenericPlatformProcess::FlushPoolSyncEvents()
 {
-	TLazySingleton<FEventPool<EEventPoolTypes::ManualReset>>::Get().EmptyPool();
-	TLazySingleton<FEventPool<EEventPoolTypes::AutoReset>>::Get().EmptyPool();
+	TLazySingleton<TEventPool<EEventMode::ManualReset>>::Get().EmptyPool();
+	TLazySingleton<TEventPool<EEventMode::AutoReset>>::Get().EmptyPool();
 }
 
 void FGenericPlatformProcess::ReturnSynchEventToPool(FEvent* Event)
@@ -515,11 +515,11 @@ void FGenericPlatformProcess::ReturnSynchEventToPool(FEvent* Event)
 
 	if (Event->IsManualReset())
 	{
-		TLazySingleton<FEventPool<EEventPoolTypes::ManualReset>>::Get().ReturnToPool(Event);
+		TLazySingleton<TEventPool<EEventMode::ManualReset>>::Get().ReturnToPool(Event);
 	}
 	else
 	{
-		TLazySingleton<FEventPool<EEventPoolTypes::AutoReset>>::Get().ReturnToPool(Event);
+		TLazySingleton<TEventPool<EEventMode::AutoReset>>::Get().ReturnToPool(Event);
 	}
 }
 
@@ -640,8 +640,8 @@ FSystemWideCriticalSectionNotImplemented::FSystemWideCriticalSectionNotImplement
 
 void FGenericPlatformProcess::TearDown()
 {
-	TLazySingleton<FEventPool<EEventPoolTypes::AutoReset>>::TearDown();
-	TLazySingleton<FEventPool<EEventPoolTypes::ManualReset>>::TearDown();
+	TLazySingleton<TEventPool<EEventMode::AutoReset>>::TearDown();
+	TLazySingleton<TEventPool<EEventMode::ManualReset>>::TearDown();
 }
 
 ENamedThreads::Type FGenericPlatformProcess::GetDesiredThreadForUObjectReferenceCollector()
diff --git a/Engine/Source/Runtime/Core/Private/HAL/ThreadingBase.cpp b/Engine/Source/Runtime/Core/Private/HAL/ThreadingBase.cpp
--- a/Engine/Source/Runtime/Core/Private/HAL/ThreadingBase.cpp
+++ b/Engine/Source/Runtime/Core/Private/HAL/ThreadingBase.cpp
@@ -647,32 +647,24 @@ void FEvent::ResetForStats()
 }
 
 FScopedEvent::FScopedEvent()
-	: Event(TLazySingleton<FEventPool<EEventPoolTypes::AutoReset>>::Get().GetEventFromPool())
+	: Event(TLazySingleton<TEventPool<EEventMode::AutoReset>>::Get().GetRawEvent())
 { }
 
 bool FScopedEvent::IsReady()
 {
-	if ( Event )
+	if ( Event->Wait(1) )
 	{
-		if ( Event->Wait(1) )
-		{
-			TLazySingleton<FEventPool<EEventPoolTypes::AutoReset>>::Get().ReturnToPool(Event);
-			Event = nullptr;
-			return true;
-		}
-		return false;
+		TLazySingleton<TEventPool<EEventMode::AutoReset>>::Get().ReturnRawEvent(Event);
+		Event = nullptr;
+		return true;
 	}
-	return true;
+	return false;
 }
 
 FScopedEvent::~FScopedEvent()
 {
-	if ( Event )
-	{
-		Event->Wait();
-		TLazySingleton<FEventPool<EEventPoolTypes::AutoReset>>::Get().ReturnToPool(Event);
-		Event = nullptr;
-	}
+	Event->Wait();
+	TLazySingleton<TEventPool<EEventMode::AutoReset>>::Get().ReturnRawEvent(Event);
 }
 
 
@@ -681,12 +673,37 @@ FScopedEvent::~FScopedEvent()
 -----------------------------------------------------------------------------*/
 
 FEventRef::FEventRef(EEventMode Mode /* = EEventMode::AutoReset */)
-	: Event(FPlatformProcess::GetSynchEventFromPool(Mode == EEventMode::ManualReset))
-{}
+{
+	if (Mode == EEventMode::AutoReset)
+	{
+		Event = TLazySingleton<TEventPool<EEventMode::AutoReset>>::Get().GetRawEvent();
+	}
+	else
+	{
+		Event = TLazySingleton<TEventPool<EEventMode::ManualReset>>::Get().GetRawEvent();
+	}
+}
 
 FEventRef::~FEventRef()
 {
-	FPlatformProcess::ReturnSynchEventToPool(Event);
+	if (Event->IsManualReset())
+	{
+		TLazySingleton<TEventPool<EEventMode::ManualReset>>::Get().ReturnRawEvent(Event);
+	}
+	else
+	{
+		TLazySingleton<TEventPool<EEventMode::AutoReset>>::Get().ReturnRawEvent(Event);
+	}
+}
+
+/*-----------------------------------------------------------------------------
+	FEventPtr
+-----------------------------------------------------------------------------*/
+
+FSharedEventRef::FSharedEventRef(EEventMode Mode  /* = EEventMode::AutoReset */)
+	: Ptr(TLazySingleton<TEventPool<EEventMode::AutoReset>>::Get().GetRawEvent(),
+		[](FEvent* Event) { TLazySingleton<TEventPool<EEventMode::AutoReset>>::Get().ReturnRawEvent(Event); })
+{
 }
 
 /*-----------------------------------------------------------------------------
diff --git a/Engine/Source/Runtime/Core/Private/Misc/EventPool.h b/Engine/Source/Runtime/Core/Private/Misc/EventPool.h
--- a/Engine/Source/Runtime/Core/Private/Misc/EventPool.h
+++ b/Engine/Source/Runtime/Core/Private/Misc/EventPool.h
@@ -12,18 +12,6 @@
 	#define USE_EVENT_POOLING 1
 #endif
 
-/**
- * Enumerates available event pool types.
- */
-enum class EEventPoolTypes
-{
-	/** Creates events that have their signaled state reset automatically. */
-	AutoReset,
-
-	/** Creates events that have their signaled state reset manually. */
-	ManualReset
-};
-
 class FSafeRecyclableEvent  final: public FEvent
 {
 public:
@@ -79,12 +67,12 @@ public:
  * @param PoolType Specifies the type of pool.
  * @see FEvent
  */
-template<EEventPoolTypes PoolType>
-class FEventPool
+template<EEventMode PoolType>
+class TEventPool
 {
 public:
 #if USE_EVENT_POOLING
-	~FEventPool()
+	~TEventPool()
 	{
 		EmptyPool();
 	}
@@ -98,21 +86,7 @@ public:
 	 */
 	FEvent* GetEventFromPool()
 	{
-		FEvent* Result = nullptr;
-
-#if USE_EVENT_POOLING
-		Result = Pool.Pop();
-		if (!Result)
-#endif
-		{
-			PRAGMA_DISABLE_DEPRECATION_WARNINGS
-			Result = FPlatformProcess::CreateSynchEvent((PoolType == EEventPoolTypes::ManualReset));
-			PRAGMA_ENABLE_DEPRECATION_WARNINGS
-		}
-		check(Result);
-		Result->AdvanceStats();
-
-		return new FSafeRecyclableEvent(Result);
+		return new FSafeRecyclableEvent(GetRawEvent());
 	}
 
 	/**
@@ -124,17 +98,11 @@ public:
 	void ReturnToPool(FEvent* Event)
 	{
 		check(Event);
-		check(Event->IsManualReset() == (PoolType == EEventPoolTypes::ManualReset));
+		check(Event->IsManualReset() == (PoolType == EEventMode::ManualReset));
+
 		FSafeRecyclableEvent* SafeEvent = (FSafeRecyclableEvent*)Event;
-		FEvent* Result = SafeEvent->InnerEvent;
+		ReturnRawEvent(SafeEvent->InnerEvent);
 		delete SafeEvent;
-		check(Result);
-		Result->Reset();
-#if USE_EVENT_POOLING
-		Pool.Push(Result);
-#else
-		delete Result;
-#endif
 	}
 
 	void EmptyPool()
@@ -147,6 +115,50 @@ public:
 #endif
 	}
 
+	/**
+	* Gets a "raw" event (as opposite to `FSafeRecyclableEvent` handle returned by `GetEventFromPool`) from the pool or 
+	* creates one if necessary.
+	* @see ReturnRaw
+	*/
+	FEvent* GetRawEvent()
+	{
+		FEvent* Event =
+#if USE_EVENT_POOLING
+			Pool.Pop();
+#else
+			nullptr;
+#endif
+
+		if (Event == nullptr)
+		{
+			PRAGMA_DISABLE_DEPRECATION_WARNINGS
+			Event = FPlatformProcess::CreateSynchEvent(PoolType == EEventMode::ManualReset);
+			PRAGMA_ENABLE_DEPRECATION_WARNINGS
+		}
+
+		check(Event);
+
+		Event->AdvanceStats();
+
+		return Event;
+	}
+
+	/**
+	 * Returns a "raw" event to the pool.
+	 * @see GetRaw
+	 */
+	void ReturnRawEvent(FEvent* Event)
+	{
+		check(Event);
+
+#if USE_EVENT_POOLING
+		Event->Reset();
+		Pool.Push(Event);
+#else
+		delete Event;
+#endif	
+	}
+
 private:
 #if USE_EVENT_POOLING
 	/** Holds the collection of recycled events. */
diff --git a/Engine/Source/Runtime/Core/Public/GenericPlatform/GenericPlatformProcess.h b/Engine/Source/Runtime/Core/Public/GenericPlatform/GenericPlatformProcess.h
--- a/Engine/Source/Runtime/Core/Public/GenericPlatform/GenericPlatformProcess.h
+++ b/Engine/Source/Runtime/Core/Public/GenericPlatform/GenericPlatformProcess.h
@@ -528,7 +528,7 @@ struct CORE_API FGenericPlatformProcess
 	 * @return A new event, or nullptr none could be created.
 	 * @see GetSynchEventFromPool, ReturnSynchEventToPool
 	 */
-	// Message to others in the future, don't try to delete this function as it isn't exactly deprecated, but it should only ever be called from FEventPool::GetEventFromPool()
+	// Message to others in the future, don't try to delete this function as it isn't exactly deprecated, but it should only ever be called from TEventPool::GetEventFromPool()
 	UE_DEPRECATED(5.0, "Please use GetSynchEventFromPool to create a new event, and ReturnSynchEventToPool to release the event.")
 	static class FEvent* CreateSynchEvent(bool bIsManualReset = false);
 
diff --git a/Engine/Source/Runtime/Core/Public/HAL/Event.h b/Engine/Source/Runtime/Core/Public/HAL/Event.h
--- a/Engine/Source/Runtime/Core/Public/HAL/Event.h
+++ b/Engine/Source/Runtime/Core/Public/HAL/Event.h
@@ -6,6 +6,7 @@
 #include "Math/NumericLimits.h"
 #include "Misc/Timespan.h"
 #include "Templates/Atomic.h"
+#include "Templates/SharedPointer.h"
 
 /**
  * Interface for waitable events.
@@ -126,7 +127,7 @@ protected:
 enum class EEventMode { AutoReset, ManualReset };
 
 /**
- * RAII-style `FEvent`
+ * RAII-style pooled `FEvent`
  *
  * non-copyable, non-movable
  */
@@ -147,6 +148,28 @@ public:
 		return Event;
 	}
 
+	FEvent* Get()
+	{
+		return Event;
+	}
+
 private:
 	FEvent* Event;
 };
+
+/**
+ * RAII-style shared and pooled `FEvent`
+ */
+class CORE_API FSharedEventRef final
+{
+public:
+	explicit FSharedEventRef(EEventMode Mode = EEventMode::AutoReset);
+
+	FEvent* operator->() const
+	{
+		return Ptr.Get();
+	}
+
+private:
+	TSharedPtr<FEvent> Ptr;
+};
\ No newline at end of file
