Some ANSICHAR/UTF8CHAR cast fixes for upcoming changes to StringConv.h.

diff --git a/Engine/Source/Developer/ShaderCompilerCommon/Private/ShaderCompilerCommon.cpp b/Engine/Source/Developer/ShaderCompilerCommon/Private/ShaderCompilerCommon.cpp
--- a/Engine/Source/Developer/ShaderCompilerCommon/Private/ShaderCompilerCommon.cpp
+++ b/Engine/Source/Developer/ShaderCompilerCommon/Private/ShaderCompilerCommon.cpp
@@ -1557,7 +1557,7 @@ void DumpDebugShaderText(const FShaderCompilerInput& Input, const FString& InSou
 	// Provide mutable container to pass string to FArchive inside inner function
 	TArray<ANSICHAR> SourceAnsi;
 	SourceAnsi.SetNum(InSource.Len() + 1);
-	FCStringAnsi::Strncpy(SourceAnsi.GetData(), StringConverter.Get(), SourceAnsi.Num());
+	FCStringAnsi::Strncpy(SourceAnsi.GetData(), (ANSICHAR*)StringConverter.Get(), SourceAnsi.Num());
 
 	// Forward temporary container to primary function
 	DumpDebugShaderText(Input, SourceAnsi.GetData(), InSource.Len(), FileExtension);
diff --git a/Engine/Source/Developer/ShaderPreprocessor/Private/ShaderPreprocessor.cpp b/Engine/Source/Developer/ShaderPreprocessor/Private/ShaderPreprocessor.cpp
--- a/Engine/Source/Developer/ShaderPreprocessor/Private/ShaderPreprocessor.cpp
+++ b/Engine/Source/Developer/ShaderPreprocessor/Private/ShaderPreprocessor.cpp
@@ -20,7 +20,7 @@ static void AddMcppDefines(TArray<TArray<ANSICHAR>>& OutOptions, const TMap<FStr
 	{
 		FString Argument(FString::Printf(TEXT("-D%s=%s"), *(It.Key()), *(It.Value())));
 		FTCHARToUTF8 Converter(Argument.GetCharArray().GetData());
-		OutOptions.Emplace(Converter.Get(), Converter.Length() + 1);
+		OutOptions.Emplace((const ANSICHAR*)Converter.Get(), Converter.Length() + 1);
 	}
 }
 
diff --git a/Engine/Source/Developer/TraceAnalysis/Private/Store/StoreClient.cpp b/Engine/Source/Developer/TraceAnalysis/Private/Store/StoreClient.cpp
--- a/Engine/Source/Developer/TraceAnalysis/Private/Store/StoreClient.cpp
+++ b/Engine/Source/Developer/TraceAnalysis/Private/Store/StoreClient.cpp
@@ -189,7 +189,7 @@ bool FStoreCborClient::Connect(const TCHAR* Host, uint16 Port)
 	FCStringAnsi::Sprintf(PortString, "%d", Port);
 
 	asio::ip::tcp::resolver Resolver(IoContext);
-	asio::ip::tcp::resolver::results_type Endpoints = Resolver.resolve(HostUtf8.Get(), PortString);
+	asio::ip::tcp::resolver::results_type Endpoints = Resolver.resolve((const char*)HostUtf8.Get(), PortString);
 
 	asio::error_code ErrorCode;
 #if PLATFORM_WINDOWS
diff --git a/Engine/Source/Runtime/Analytics/AnalyticsET/Private/AnalyticsProviderETEventCache.cpp b/Engine/Source/Runtime/Analytics/AnalyticsET/Private/AnalyticsProviderETEventCache.cpp
--- a/Engine/Source/Runtime/Analytics/AnalyticsET/Private/AnalyticsProviderETEventCache.cpp
+++ b/Engine/Source/Runtime/Analytics/AnalyticsET/Private/AnalyticsProviderETEventCache.cpp
@@ -114,7 +114,7 @@ namespace EventCacheStatic
 			// make space for the string
 			UTF8Stream.SetNumUninitialized(OldLen + StrLen, false);
 			// convert it to UTF8
-			int32 CharsWritten = FTCHARToUTF8_Convert::Convert((ANSICHAR*)&UTF8Stream[OldLen], StrLen, Str, Len);
+			int32 CharsWritten = FTCHARToUTF8_Convert::Convert((UTF8CHAR*)&UTF8Stream[OldLen], StrLen, Str, Len);
 			// figure out how many characters were actually written 
 			if (CharsWritten >= 0)
 			{
@@ -132,7 +132,7 @@ namespace EventCacheStatic
 					const int32 ActualCharsNeeded = FTCHARToUTF8_Convert::ConvertedLength(Str, Len);
 					UTF8Stream.SetNumUninitialized(OldLen + ActualCharsNeeded, false);
 					// convert it to UTF8 using the known number of charts
-					FTCHARToUTF8_Convert::Convert((ANSICHAR*)&UTF8Stream[OldLen], ActualCharsNeeded, Str, Len);
+					FTCHARToUTF8_Convert::Convert((UTF8CHAR*)&UTF8Stream[OldLen], ActualCharsNeeded, Str, Len);
 					bWroteFullString = true;
 				}
 			}
diff --git a/Engine/Source/Runtime/Cbor/Private/CborWriter.cpp b/Engine/Source/Runtime/Cbor/Private/CborWriter.cpp
--- a/Engine/Source/Runtime/Cbor/Private/CborWriter.cpp
+++ b/Engine/Source/Runtime/Cbor/Private/CborWriter.cpp
@@ -130,7 +130,7 @@ void FCborWriter::WriteValue(const FString& Value)
 	WriteUIntValue(ECborCode::TextString, *Stream, (uint64)UTF8String.Length());
 	// Write string
 	check(sizeof(decltype(*UTF8String.Get())) == 1);
-	Stream->Serialize(const_cast<char*>(UTF8String.Get()), UTF8String.Length());
+	Stream->Serialize(const_cast<char*>(reinterpret_cast<const char*>(UTF8String.Get())), UTF8String.Length());
 }
 
 void FCborWriter::WriteValue(const char* CString, uint64 Length)
diff --git a/Engine/Source/Runtime/Core/Private/Misc/OutputDeviceHelper.cpp b/Engine/Source/Runtime/Core/Private/Misc/OutputDeviceHelper.cpp
--- a/Engine/Source/Runtime/Core/Private/Misc/OutputDeviceHelper.cpp
+++ b/Engine/Source/Runtime/Core/Private/Misc/OutputDeviceHelper.cpp
@@ -116,7 +116,7 @@ void FOutputDeviceHelper::FormatCastAndSerializeLine(FArchive& Output, const TCH
 	const int32 ConvertedPrefixLength = FTCHARToUTF8_Convert::ConvertedLength(*Prefix, Prefix.Len());
 
 	// Allocate the conversion buffer. It's ok to always add some slack for the line terminator even if not required
-	TArray<ANSICHAR, TInlineAllocator<2 * DEFAULT_STRING_CONVERSION_SIZE>> ConvertedText;
+	TArray<UTF8CHAR, TInlineAllocator<2 * DEFAULT_STRING_CONVERSION_SIZE>> ConvertedText;
 	ConvertedText.AddUninitialized(ConvertedPrefixLength + ConvertedDataLength + (bAutoEmitLineTerminator ? ConvertedTerminatorLength : 0));
 	// Do the actual conversion to the pre-allocated buffer
 	FTCHARToUTF8_Convert::Convert(ConvertedText.GetData(), ConvertedPrefixLength, *Prefix, PrefixLength);
diff --git a/Engine/Source/Runtime/Core/Private/Serialization/CompactBinaryWriter.cpp b/Engine/Source/Runtime/Core/Private/Serialization/CompactBinaryWriter.cpp
--- a/Engine/Source/Runtime/Core/Private/Serialization/CompactBinaryWriter.cpp
+++ b/Engine/Source/Runtime/Core/Private/Serialization/CompactBinaryWriter.cpp
@@ -407,7 +407,7 @@ void FCbWriter::AddString(const FWideStringView Value)
 	StringData += SizeByteCount;
 	if (Size > 0)
 	{
-		FTCHARToUTF8_Convert::Convert(reinterpret_cast<ANSICHAR*>(StringData), Size, Value.GetData(), Value.Len());
+		FTCHARToUTF8_Convert::Convert(reinterpret_cast<UTF8CHAR*>(StringData), Size, Value.GetData(), Value.Len());
 	}
 	EndField(ECbFieldType::String);
 }
diff --git a/Engine/Source/Runtime/Core/Private/Serialization/MemoryImage.cpp b/Engine/Source/Runtime/Core/Private/Serialization/MemoryImage.cpp
--- a/Engine/Source/Runtime/Core/Private/Serialization/MemoryImage.cpp
+++ b/Engine/Source/Runtime/Core/Private/Serialization/MemoryImage.cpp
@@ -1185,7 +1185,7 @@ FHashedName::FHashedName(const FName& InName)
 			}
 			UpperNameBuffer.Wide[NameLength] = 0;
 			const FTCHARToUTF8 UpperNameUTF8(UpperNameBuffer.Wide);
-			Hash = CityHash64WithSeed(UpperNameUTF8.Get(), UpperNameUTF8.Length(), InternalNumber);
+			Hash = CityHash64WithSeed((const char*)UpperNameUTF8.Get(), UpperNameUTF8.Length(), InternalNumber);
 #if WITH_EDITORONLY_DATA
 			{
 				const FTCHARToUTF8 NameUTF8(NameBuffer.Wide);
diff --git a/Engine/Source/Runtime/Core/Private/Tests/Misc/TCharToUTF8.cpp b/Engine/Source/Runtime/Core/Private/Tests/Misc/TCharToUTF8.cpp
--- a/Engine/Source/Runtime/Core/Private/Tests/Misc/TCharToUTF8.cpp
+++ b/Engine/Source/Runtime/Core/Private/Tests/Misc/TCharToUTF8.cpp
@@ -16,7 +16,7 @@ bool FTCharToUTF8Test::RunTest(const FString& Parameters)
 		{
 			const FTCHARToUTF8 ConvertedValue(*TestString);
 			const FUTF8ToTCHAR ReverseConvertedValue(ConvertedValue.Get());
-			TestTrue(FString(TEXT("Expected test string 1 to have different encoded value")), FCStringAnsi::Strcmp(ExpectedResult, ConvertedValue.Get()) == 0);
+			TestTrue(FString(TEXT("Expected test string 1 to have different encoded value")), FCStringAnsi::Strcmp(ExpectedResult, (const ANSICHAR*)ConvertedValue.Get()) == 0);
 			TestTrue(FString(TEXT("Expected test string 1 to have different decoded value")), FCString::Strcmp(*TestString, ReverseConvertedValue.Get()) == 0);
 		}
 
@@ -24,7 +24,7 @@ bool FTCharToUTF8Test::RunTest(const FString& Parameters)
 			const FUTF8ToTCHAR ConvertedValue(ExpectedResult);
 			const FTCHARToUTF8 ReverseConvertedValue(ConvertedValue.Get());
 			TestTrue(FString(TEXT("Expected test string 2 to have different encoded value")), FCString::Strcmp(*TestString, ConvertedValue.Get()) == 0);
-			TestTrue(FString(TEXT("Expected test string 2 to have different decoded value")), FCStringAnsi::Strcmp(ExpectedResult, ReverseConvertedValue.Get()) == 0);
+			TestTrue(FString(TEXT("Expected test string 2 to have different decoded value")), FCStringAnsi::Strcmp(ExpectedResult, (const ANSICHAR*)ReverseConvertedValue.Get()) == 0);
 		}
 	}
 
@@ -35,7 +35,7 @@ bool FTCharToUTF8Test::RunTest(const FString& Parameters)
 		{
 			const FTCHARToUTF8 ConvertedValue(*TestString);
 			const FUTF8ToTCHAR ReverseConvertedValue(ConvertedValue.Get());
-			TestTrue(FString(TEXT("Expected test string 3 to have different encoded value")), FCStringAnsi::Strcmp(ExpectedResult, ConvertedValue.Get()) == 0);
+			TestTrue(FString(TEXT("Expected test string 3 to have different encoded value")), FCStringAnsi::Strcmp(ExpectedResult, (const ANSICHAR*)ConvertedValue.Get()) == 0);
 			TestTrue(FString(TEXT("Expected test string 3 to have different decoded value")), FCString::Strcmp(*TestString, ReverseConvertedValue.Get()) == 0);
 		}
 
@@ -43,7 +43,7 @@ bool FTCharToUTF8Test::RunTest(const FString& Parameters)
 			const FUTF8ToTCHAR ConvertedValue(ExpectedResult);
 			const FTCHARToUTF8 ReverseConvertedValue(ConvertedValue.Get());
 			TestTrue(FString(TEXT("Expected test string 4 to have different encoded value")), FCString::Strcmp(*TestString, ConvertedValue.Get()) == 0);
-			TestTrue(FString(TEXT("Expected test string 4 to have different decoded value")), FCStringAnsi::Strcmp(ExpectedResult, ReverseConvertedValue.Get()) == 0);
+			TestTrue(FString(TEXT("Expected test string 4 to have different decoded value")), FCStringAnsi::Strcmp(ExpectedResult, (const ANSICHAR*)ReverseConvertedValue.Get()) == 0);
 		}
 	}
 
@@ -54,7 +54,7 @@ bool FTCharToUTF8Test::RunTest(const FString& Parameters)
 		{
 			const FTCHARToUTF8 ConvertedValue(*TestString);
 			const FUTF8ToTCHAR ReverseConvertedValue(ConvertedValue.Get());
-			TestTrue(FString(TEXT("Expected test string 5 to have different encoded value")), FCStringAnsi::Strcmp(ExpectedResult, ConvertedValue.Get()) == 0);
+			TestTrue(FString(TEXT("Expected test string 5 to have different encoded value")), FCStringAnsi::Strcmp(ExpectedResult, (const ANSICHAR*)ConvertedValue.Get()) == 0);
 			TestTrue(FString(TEXT("Expected test string 5 to have different decoded value")), FCString::Strcmp(*TestString, ReverseConvertedValue.Get()) == 0);
 		}
 
@@ -62,7 +62,7 @@ bool FTCharToUTF8Test::RunTest(const FString& Parameters)
 			const FUTF8ToTCHAR ConvertedValue(ExpectedResult);
 			const FTCHARToUTF8 ReverseConvertedValue(ConvertedValue.Get());
 			TestTrue(FString(TEXT("Expected test string 6 to have different encoded value")), FCString::Strcmp(*TestString, ConvertedValue.Get()) == 0);
-			TestTrue(FString(TEXT("Expected test string 6 to have different decoded value")), FCStringAnsi::Strcmp(ExpectedResult, ReverseConvertedValue.Get()) == 0);
+			TestTrue(FString(TEXT("Expected test string 6 to have different decoded value")), FCStringAnsi::Strcmp(ExpectedResult, (const ANSICHAR*)ReverseConvertedValue.Get()) == 0);
 		}
 	}
 
@@ -73,7 +73,7 @@ bool FTCharToUTF8Test::RunTest(const FString& Parameters)
 		{
 			const FTCHARToUTF8 ConvertedValue(*TestString);
 			const FUTF8ToTCHAR ReverseConvertedValue(ConvertedValue.Get());
-			TestTrue(FString(TEXT("Expected test string 7 to have different encoded value")), FCStringAnsi::Strcmp(ExpectedResult, ConvertedValue.Get()) == 0);
+			TestTrue(FString(TEXT("Expected test string 7 to have different encoded value")), FCStringAnsi::Strcmp(ExpectedResult, (const ANSICHAR*)ConvertedValue.Get()) == 0);
 			TestTrue(FString(TEXT("Expected test string 7 to have different decoded value")), FCString::Strcmp(*TestString, ReverseConvertedValue.Get()) == 0);
 		}
 
@@ -81,7 +81,7 @@ bool FTCharToUTF8Test::RunTest(const FString& Parameters)
 			const FUTF8ToTCHAR ConvertedValue(ExpectedResult);
 			const FTCHARToUTF8 ReverseConvertedValue(ConvertedValue.Get());
 			TestTrue(FString(TEXT("Expected test string 8 to have different encoded value")), FCString::Strcmp(*TestString, ConvertedValue.Get()) == 0);
-			TestTrue(FString(TEXT("Expected test string 8 to have different decoded value")), FCStringAnsi::Strcmp(ExpectedResult, ReverseConvertedValue.Get()) == 0);
+			TestTrue(FString(TEXT("Expected test string 8 to have different decoded value")), FCStringAnsi::Strcmp(ExpectedResult, (const ANSICHAR*)ReverseConvertedValue.Get()) == 0);
 		}
 	}
 
diff --git a/Engine/Source/Runtime/Net/Core/Private/Net/Core/Trace/NetTraceInternal.cpp b/Engine/Source/Runtime/Net/Core/Private/Net/Core/Trace/NetTraceInternal.cpp
--- a/Engine/Source/Runtime/Net/Core/Private/Net/Core/Trace/NetTraceInternal.cpp
+++ b/Engine/Source/Runtime/Net/Core/Private/Net/Core/Trace/NetTraceInternal.cpp
@@ -491,7 +491,7 @@ FNetDebugNameId FNetTrace::TraceName(const TCHAR* Name)
 		ThreadBuffer->DynamicNameHashToNameIdMap.Add(HashedName, NameId);
 
 		FTCHARToUTF8 Converter(Name);
-		FNetTraceInternal::Reporter::ReportAnsiName(NameId, Converter.Length() + 1, Converter.Get());		
+		FNetTraceInternal::Reporter::ReportAnsiName(NameId, Converter.Length() + 1, (const char*)Converter.Get());		
 		
 		return NameId;
 	}
@@ -524,7 +524,7 @@ FNetDebugNameId FNetTrace::TraceName(FName Name)
 		TCHAR Buffer[StringBufferSize];
 		uint32 NameLen = Name.ToString(Buffer);
 		FTCHARToUTF8 Converter(Buffer);
-		FNetTraceInternal::Reporter::ReportAnsiName(NameId, Converter.Length() + 1, Converter.Get());		
+		FNetTraceInternal::Reporter::ReportAnsiName(NameId, Converter.Length() + 1, (const char*)Converter.Get());		
 		
 		return NameId;
 	}
diff --git a/Engine/Source/Runtime/Online/HTTP/Private/Curl/CurlHttp.cpp b/Engine/Source/Runtime/Online/HTTP/Private/Curl/CurlHttp.cpp
--- a/Engine/Source/Runtime/Online/HTTP/Private/Curl/CurlHttp.cpp
+++ b/Engine/Source/Runtime/Online/HTTP/Private/Curl/CurlHttp.cpp
@@ -334,7 +334,7 @@ void FCurlHttpRequest::SetContentAsString(const FString& ContentString)
 	int32 Utf8Length = FTCHARToUTF8_Convert::ConvertedLength(*ContentString, ContentString.Len());
 	TArray<uint8> Buffer;
 	Buffer.SetNumUninitialized(Utf8Length);
-	FTCHARToUTF8_Convert::Convert((ANSICHAR*)Buffer.GetData(), Buffer.Num(), *ContentString, ContentString.Len());
+	FTCHARToUTF8_Convert::Convert((UTF8CHAR*)Buffer.GetData(), Buffer.Num(), *ContentString, ContentString.Len());
 	RequestPayload = MakeUnique<FRequestPayloadInMemory>(MoveTemp(Buffer));
 	bIsRequestPayloadSeekable = true;
 }
diff --git a/Engine/Source/Runtime/Online/HTTP/Private/NullHttp.cpp b/Engine/Source/Runtime/Online/HTTP/Private/NullHttp.cpp
--- a/Engine/Source/Runtime/Online/HTTP/Private/NullHttp.cpp
+++ b/Engine/Source/Runtime/Online/HTTP/Private/NullHttp.cpp
@@ -81,7 +81,7 @@ void FNullHttpRequest::SetContentAsString(const FString& ContentString)
 {
 	int32 Utf8Length = FTCHARToUTF8_Convert::ConvertedLength(*ContentString, ContentString.Len());
 	Payload.SetNumUninitialized(Utf8Length);
-	FTCHARToUTF8_Convert::Convert((ANSICHAR*)Payload.GetData(), Payload.Num(), *ContentString, ContentString.Len());
+	FTCHARToUTF8_Convert::Convert((UTF8CHAR*)Payload.GetData(), Payload.Num(), *ContentString, ContentString.Len());
 }
 
 bool FNullHttpRequest::SetContentAsStreamedFile(const FString& Filename)
diff --git a/Engine/Source/Runtime/Online/WebSockets/Private/Lws/LwsWebSocket.cpp b/Engine/Source/Runtime/Online/WebSockets/Private/Lws/LwsWebSocket.cpp
--- a/Engine/Source/Runtime/Online/WebSockets/Private/Lws/LwsWebSocket.cpp
+++ b/Engine/Source/Runtime/Online/WebSockets/Private/Lws/LwsWebSocket.cpp
@@ -145,7 +145,7 @@ void FLwsWebSocket::Close(int32 Code, const FString& Reason)
 	// We are doing this conversion here so we don't have to do it on the ws thread
 	FTCHARToUTF8 Convert(*Reason);
 	ANSICHAR* ANSIReason = static_cast<ANSICHAR*>(FMemory::Malloc(Convert.Length() + 1));
-	FCStringAnsi::Strcpy(ANSIReason, Convert.Length(), Convert.Get());
+	FCStringAnsi::Strcpy(ANSIReason, Convert.Length(), (const ANSICHAR*)Convert.Get());
 	ANSIReason[Convert.Length()] = 0;
 
 	UE_LOG(LogWebSockets, Verbose, TEXT("FLwsWebSocket[%d]::Close: Close queued with code=%d reason=%s"), Identifier, Code, *Reason);
@@ -713,7 +713,7 @@ void FLwsWebSocket::ConnectInternal(struct lws_context &LwsContext)
 		OptionalCombinedProtocolsUTF8.Emplace(*CombinedProtocols);
 	}
 
-	const char* const CombinedProtocolsUTF8 = OptionalCombinedProtocolsUTF8.IsSet() ? OptionalCombinedProtocolsUTF8.GetValue().Get() : nullptr;
+	const char* const CombinedProtocolsUTF8 = OptionalCombinedProtocolsUTF8.IsSet() ? (const char*)OptionalCombinedProtocolsUTF8.GetValue().Get() : nullptr;
 
 	struct lws_client_connect_info ConnectInfo = {};
 	ConnectInfo.context = &LwsContext;
diff --git a/Engine/Source/Runtime/Online/WebSockets/Private/Lws/LwsWebSocketsManager.cpp b/Engine/Source/Runtime/Online/WebSockets/Private/Lws/LwsWebSocketsManager.cpp
--- a/Engine/Source/Runtime/Online/WebSockets/Private/Lws/LwsWebSocketsManager.cpp
+++ b/Engine/Source/Runtime/Online/WebSockets/Private/Lws/LwsWebSocketsManager.cpp
@@ -74,7 +74,7 @@ void FLwsWebSocketsManager::InitWebSockets(TArrayView<const FString> Protocols)
 
 		// We need to hold on to the converted strings
 		ANSICHAR* Converted = static_cast<ANSICHAR*>(FMemory::Malloc(ConvertName.Length() + 1));
-		FCStringAnsi::Strcpy(Converted, ConvertName.Length(), ConvertName.Get());
+		FCStringAnsi::Strcpy(Converted, ConvertName.Length(), (const ANSICHAR*)ConvertName.Get());
 		lws_protocols LwsProtocol;
 		FMemory::Memset(&LwsProtocol, 0, sizeof(LwsProtocol));
 		LwsProtocol.name = Converted;
@@ -114,7 +114,7 @@ void FLwsWebSocketsManager::InitWebSockets(TArrayView<const FString> Protocols)
 	if (!ProxyAddress.IsEmpty())
 	{
 		Converter.Emplace(*ProxyAddress);
-		ContextInfo.http_proxy_address = Converter->Get();
+		ContextInfo.http_proxy_address = (const char*)Converter->Get();
 	}
 
 #if WITH_SSL
diff --git a/Engine/Source/Runtime/PerfCounters/Private/PerfCounters.cpp b/Engine/Source/Runtime/PerfCounters/Private/PerfCounters.cpp
--- a/Engine/Source/Runtime/PerfCounters/Private/PerfCounters.cpp
+++ b/Engine/Source/Runtime/PerfCounters/Private/PerfCounters.cpp
@@ -385,7 +385,7 @@ bool FPerfCounters::ReportUnplayableCondition(const FString& ConditionDescriptio
 
 	// include description for debugging
 	FTCHARToUTF8 Converter(*FString::Printf(TEXT("Unplayable condition encountered: %s\n"), *ConditionDescription));
-	ReportFile->Serialize(reinterpret_cast<void *>(const_cast<char *>(Converter.Get())), Converter.Length());
+	ReportFile->Serialize(reinterpret_cast<void *>(const_cast<char *>(reinterpret_cast<const char *>(Converter.Get()))), Converter.Length());
 
 	ReportFile->Close();
 	delete ReportFile;
diff --git a/Engine/Source/Runtime/RenderCore/Private/ShaderCodeLibrary.cpp b/Engine/Source/Runtime/RenderCore/Private/ShaderCodeLibrary.cpp
--- a/Engine/Source/Runtime/RenderCore/Private/ShaderCodeLibrary.cpp
+++ b/Engine/Source/Runtime/RenderCore/Private/ShaderCodeLibrary.cpp
@@ -1561,7 +1561,7 @@ struct FEditorShaderCodeArchive
 				FString DumpFile = DebugLibFolder / TEXT("Dump.txt");
 				TUniquePtr<FArchive> DumpAr(IFileManager::Get().CreateFileWriter(*DumpFile));
 				FTCHARToUTF8 Converter(*ExtendedStats.TextualRepresentation);
-				DumpAr->Serialize(const_cast<char*>(Converter.Get()), Converter.Length());
+				DumpAr->Serialize(const_cast<UTF8CHAR*>(reinterpret_cast<const UTF8CHAR*>(Converter.Get())), Converter.Length());
 				UE_LOG(LogShaderLibrary, Display, TEXT("Textual dump saved to '%s'"), *DumpFile);
 			}
 #if 0 // creating a graphviz graph - maybe one day we'll return to this
diff --git a/Engine/Source/Runtime/SlateCore/Private/Fonts/FontCacheHarfBuzz.cpp b/Engine/Source/Runtime/SlateCore/Private/Fonts/FontCacheHarfBuzz.cpp
--- a/Engine/Source/Runtime/SlateCore/Private/Fonts/FontCacheHarfBuzz.cpp
+++ b/Engine/Source/Runtime/SlateCore/Private/Fonts/FontCacheHarfBuzz.cpp
@@ -60,7 +60,7 @@ void AppendStringToBuffer(const FStringView InString, const int32 InStartIndex,
 	//				   In practice this may not be an issue as our platforms should all use the other functions, but to fix it we'd need UTF-8 iteration functions to find the correct points the buffer
 	FStringView SubString = InString.Mid(InStartIndex, InLength);
 	FTCHARToUTF8 SubStringUtf8(SubString.GetData(), SubString.Len());
-	hb_buffer_add_utf8(InHarfBuzzTextBuffer, SubStringUtf8.Get(), SubStringUtf8.Length(), 0, SubStringUtf8.Length());
+	hb_buffer_add_utf8(InHarfBuzzTextBuffer, (const char*)SubStringUtf8.Get(), SubStringUtf8.Length(), 0, SubStringUtf8.Length());
 }
 
 template <>
