#ue5 - LWC: Expand VectorRegister trig and math function tests to run for doubles, and test all existing trig functions for both float and doubles. Cleaned up math test logging to log correct values in failure cases. Templated more functions to reduce code duplication. Enabled SVML on MSVC 2019+ for more efficient vectorized versions of both float and double on SSE/AVX.

diff --git a/Engine/Source/Runtime/Core/Private/Tests/Math/UnrealMathTest.cpp b/Engine/Source/Runtime/Core/Private/Tests/Math/UnrealMathTest.cpp
--- a/Engine/Source/Runtime/Core/Private/Tests/Math/UnrealMathTest.cpp
+++ b/Engine/Source/Runtime/Core/Private/Tests/Math/UnrealMathTest.cpp
@@ -59,6 +59,7 @@ bool TestVectorsEqualBitwise( VectorRegister4Float Vec0, VectorRegister4Float Ve
 {
 	VectorStoreAligned( Vec0, GScratch + 0 );
 	VectorStoreAligned( Vec1, GScratch + 4 );
+	GSum = 0.0f;
 
 	const bool Passed = (memcmp(GScratch + 0, GScratch + 4, sizeof(float) * 4) == 0);
 
@@ -70,6 +71,7 @@ bool TestVectorsEqualBitwise(VectorRegister4Double Vec0, VectorRegister4Double V
 {
 	VectorStoreAligned(Vec0, GScratchDouble + 0);
 	VectorStoreAligned(Vec1, GScratchDouble + 4);
+	GSumDouble = 0.0;
 
 	const bool Passed = (memcmp(GScratchDouble + 0, GScratchDouble + 4, sizeof(double) * 4) == 0);
 
@@ -95,7 +97,7 @@ bool TestVectorsEqual( VectorRegister4Float Vec0, VectorRegister4Float Vec1, flo
 	for ( int32 Component = 0; Component < 4; Component++ ) 
 	{
 		float Diff = GScratch[ Component + 0 ] - GScratch[ Component + 4 ];
-		GSum += ( Diff >= 0.0f ) ? Diff : -Diff;
+		GSum += FMath::Abs(Diff);
 	}
 	CheckPassing(GSum <= Tolerance);
 	return GSum <= Tolerance;
@@ -109,7 +111,7 @@ bool TestVectorsEqual(VectorRegister4Double Vec0, VectorRegister4Double Vec1, do
 	for (int32 Component = 0; Component < 4; Component++)
 	{
 		double Diff = GScratchDouble[Component + 0] - GScratchDouble[Component + 4];
-		GSumDouble += (Diff >= 0.0f) ? Diff : -Diff;
+		GSumDouble += FMath::Abs(Diff);
 	}
 	CheckPassing(GSumDouble <= Tolerance);
 	return GSumDouble <= Tolerance;
@@ -122,11 +124,31 @@ bool TestVectorsEqual_ComponentWiseError(VectorRegister4Float Vec0, VectorRegist
 {
 	VectorStoreAligned(Vec0, GScratch + 0);
 	VectorStoreAligned(Vec1, GScratch + 4);
+	GSum = 0.0f;
 	bool bPassing = true;
+
 	for (int32 Component = 0; Component < 4; Component++)
 	{
 		float Diff = GScratch[Component + 0] - GScratch[Component + 4];
 		bPassing &= FMath::IsNearlyZero(Diff, Tolerance);
+		GSum += FMath::Abs(Diff);
+	}
+	CheckPassing(bPassing);
+	return bPassing;
+}
+
+bool TestVectorsEqual_ComponentWiseError(VectorRegister4Double Vec0, VectorRegister4Double Vec1, double Tolerance = 0.0f)
+{
+	VectorStoreAligned(Vec0, GScratchDouble + 0);
+	VectorStoreAligned(Vec1, GScratchDouble + 4);
+	GSumDouble = 0.0f;
+	bool bPassing = true;
+
+	for (int32 Component = 0; Component < 4; Component++)
+	{
+		double Diff = GScratchDouble[Component + 0] - GScratchDouble[Component + 4];
+		bPassing &= FMath::IsNearlyZero(Diff, Tolerance);
+		GSumDouble += FMath::Abs(Diff);
 	}
 	CheckPassing(bPassing);
 	return bPassing;
@@ -281,7 +303,7 @@ bool TestMatricesEqual( FMatrix44f &Mat0, FMatrix44f &Mat1, float Tolerance = 0.
 		for ( int32 Column = 0; Column < 4; ++Column ) 
 		{
 			float Diff = Mat0.M[Row][Column] - Mat1.M[Row][Column];
-			GSum += ( Diff >= 0.0f ) ? Diff : -Diff;
+			GSum += FMath::Abs(Diff);
 		}
 		if (GSum > Tolerance)
 		{
@@ -696,7 +718,6 @@ void LogQuaternionTest(const TCHAR* TestName, const FQuat& A, const FQuat& B, bo
 	}
 }
 
-
 // Normalize tests
 
 MATHTEST_INLINE VectorRegister4Float TestVectorNormalize_Sqrt(const VectorRegister4Float& V)
@@ -750,9 +771,10 @@ MATHTEST_INLINE VectorRegister4Double TestReferenceMod(const VectorRegister4Doub
 
 
 // SinCos
-MATHTEST_INLINE void TestReferenceSinCos(VectorRegister4Float& S, VectorRegister4Float& C, const VectorRegister4Float& VAngles)
+template<typename FloatType, typename VectorRegisterType>
+MATHTEST_INLINE void TestReferenceSinCos(VectorRegisterType& S, VectorRegisterType& C, const VectorRegisterType& VAngles)
 {
-	float FloatAngles[4];
+	FloatType FloatAngles[4];
 	VectorStore(VAngles, FloatAngles);
 
 	S = MakeVectorRegister(
@@ -770,12 +792,13 @@ MATHTEST_INLINE void TestReferenceSinCos(VectorRegister4Float& S, VectorRegister
 			);
 }
 
-MATHTEST_INLINE void TestFastSinCos(VectorRegister4Float& S, VectorRegister4Float& C, const VectorRegister4Float& VAngles)
+template<typename FloatType, typename VectorRegisterType>
+MATHTEST_INLINE void TestFastSinCos(VectorRegisterType& S, VectorRegisterType& C, const VectorRegisterType& VAngles)
 {
-	float FloatAngles[4];
+	FloatType FloatAngles[4];
 	VectorStore(VAngles, FloatAngles);
 
-	float SFloat[4], CFloat[4];
+	FloatType SFloat[4], CFloat[4];
 	FMath::SinCos(&SFloat[0], &CFloat[0], FloatAngles[0]);
 	FMath::SinCos(&SFloat[1], &CFloat[1], FloatAngles[1]);
 	FMath::SinCos(&SFloat[2], &CFloat[2], FloatAngles[2]);
@@ -785,27 +808,101 @@ MATHTEST_INLINE void TestFastSinCos(VectorRegister4Float& S, VectorRegister4Floa
 	C = VectorLoad(CFloat);
 }
 
-MATHTEST_INLINE void TestVectorSinCos(VectorRegister4Float& S, VectorRegister4Float& C, const VectorRegister4Float& VAngles)
+template<typename FloatType, typename VectorRegisterType>
+MATHTEST_INLINE void TestVectorSinCos(VectorRegisterType& S, VectorRegisterType& C, const VectorRegisterType& VAngles)
 {
 	VectorSinCos(&S, &C, &VAngles);
 }
 
+// Generic way to test functions taking a single FloatType param and returning a FloatType value, and comparing to a vectorized version of the same function.
+template<typename FloatType, typename VectorRegisterType, FloatType (*Func)(FloatType), VectorRegisterType (*VectorFunc)(const VectorRegisterType&)>
+FORCENOINLINE bool TestVectorFunction1Param(VectorRegisterType& ReferenceResult, VectorRegisterType& VectorResult, const VectorRegisterType& VParams, FloatType ErrorTolerance)
+{
+	// Load params
+	FloatType FloatParams[4];
+	VectorStore(VParams, FloatParams);
+
+	// Reference function
+	FloatType Results[4];
+	Results[0] = Func(FloatParams[0]);
+	Results[1] = Func(FloatParams[1]);
+	Results[2] = Func(FloatParams[2]);
+	Results[3] = Func(FloatParams[3]);
+	ReferenceResult = VectorLoad(Results);
+	
+	// Vector function
+	VectorResult = VectorFunc(VParams);
+
+	// Check for errors
+	return TestVectorsEqual_ComponentWiseError(ReferenceResult, VectorResult, ErrorTolerance);
+}
+
+// Generic way to test functions taking two FloatType params and returning a FloatType value, and comparing to a vectorized version of the same function.
+template<typename FloatType, typename VectorRegisterType, FloatType(*Func)(FloatType, FloatType), VectorRegisterType(*VectorFunc)(const VectorRegisterType&, const VectorRegisterType&)>
+FORCENOINLINE bool TestVectorFunction2Params(VectorRegisterType& ReferenceResult, VectorRegisterType& VectorResult, const VectorRegisterType& VParams0, const VectorRegisterType& VParams1, FloatType ErrorTolerance)
+{
+	// Load params
+	FloatType FloatParams0[4];
+	VectorStore(VParams0, FloatParams0);
+	FloatType FloatParams1[4];
+	VectorStore(VParams1, FloatParams1);
+
+	// Reference function
+	FloatType Results[4];
+	Results[0] = Func(FloatParams0[0], FloatParams1[0]);
+	Results[1] = Func(FloatParams0[1], FloatParams1[1]);
+	Results[2] = Func(FloatParams0[2], FloatParams1[2]);
+	Results[3] = Func(FloatParams0[3], FloatParams1[3]);
+	ReferenceResult = VectorLoad(Results);
+
+	// Vector function
+	VectorResult = VectorFunc(VParams0, VParams1);
+
+	// Check for errors
+	return TestVectorsEqual_ComponentWiseError(ReferenceResult, VectorResult, ErrorTolerance);
+}
+
 /**
  * Helper debugf function to print out success or failure information for a test
  *
  * @param TestName Name of the current test
  * @param bHasPassed true if the test has passed
  */
-void LogTest(const TCHAR* TestName, bool bHasPassed)
+
+template<typename FloatType>
+FORCENOINLINE void LogTest(const TCHAR* TestName, bool bHasPassed)
 {
 	if (bHasPassed == false)
 	{
-		UE_LOG(LogUnrealMathTest, Log, TEXT("%s: %s"), bHasPassed ? TEXT("PASSED") : TEXT("FAILED"), TestName);
+		UE_LOG(LogUnrealMathTest, Error, TEXT("Unimplemented type for LogTest()"));
+		CheckPassing(false);
+	}
+}
+
+// Specialization for 'float' so it checks the correct global state filled by the various comparision validation functions
+template<>
+FORCENOINLINE void LogTest<float>(const TCHAR* TestName, bool bHasPassed)
+{
+	if (bHasPassed == false)
+	{
+		UE_LOG(LogUnrealMathTest, Log, TEXT("%s <float>: %s"), bHasPassed ? TEXT("PASSED") : TEXT("FAILED"), TestName);
 		UE_LOG(LogUnrealMathTest, Log, TEXT("Bad(%.8f): (%.8f %.8f %.8f %.8f) (%.8f %.8f %.8f %.8f)"), GSum, GScratch[0], GScratch[1], GScratch[2], GScratch[3], GScratch[4], GScratch[5], GScratch[6], GScratch[7]);
 		CheckPassing(false);
 	}
 }
 
+// Specialization for 'double' so it checks the correct global state filled by the various comparision validation functions
+template<>
+FORCENOINLINE void LogTest<double>(const TCHAR* TestName, bool bHasPassed)
+{
+	if (bHasPassed == false)
+	{
+		UE_LOG(LogUnrealMathTest, Log, TEXT("%s <double>: %s"), bHasPassed ? TEXT("PASSED") : TEXT("FAILED"), TestName);
+		UE_LOG(LogUnrealMathTest, Log, TEXT("Bad(%.8f): (%.8f %.8f %.8f %.8f) (%.8f %.8f %.8f %.8f)"), GSumDouble, GScratchDouble[0], GScratchDouble[1], GScratchDouble[2], GScratchDouble[3], GScratchDouble[4], GScratchDouble[5], GScratchDouble[6], GScratchDouble[7]);
+		CheckPassing(false);
+	}
+}
+
 /**
  * Set the contents of the scratch memory
  *
@@ -846,7 +943,7 @@ FORCENOINLINE void TestVectorReplicate()
 #define ReplicateTest(A, X) \
 		V1 = VectorReplicate(A, X); \
 		V2 = MakeVectorRegister(ArrayV0[X], ArrayV0[X], ArrayV0[X], ArrayV0[X]); \
-		LogTest(*FString::Printf(TEXT("VectorReplicate<%d>"), X), TestVectorsEqual(V1, V2));
+		LogTest<RealType>(*FString::Printf(TEXT("VectorReplicate<%d>"), X), TestVectorsEqual(V1, V2));
 
 	ReplicateTest(V0, 0);
 	ReplicateTest(V0, 1);
@@ -867,7 +964,7 @@ FORCENOINLINE void TestVectorSwizzle()
 #define SwizzleTest(A, X, Y, Z, W) \
 		V1 = VectorSwizzle(A, X, Y, Z, W); \
 		V2 = MakeVectorRegister(ArrayV0[X], ArrayV0[Y], ArrayV0[Z], ArrayV0[W]); \
-		LogTest(*FString::Printf(TEXT("VectorSwizzle<%d,%d,%d,%d>"), X, Y, Z, W), TestVectorsEqual(V1, V2));
+		LogTest<RealType>(*FString::Printf(TEXT("VectorSwizzle<%d,%d,%d,%d>"), X, Y, Z, W), TestVectorsEqual(V1, V2));
 
 	// This is not an exhaustive list because it would be 4*4*4*4 = 256 entries, but it tries to test a lot of common permutations.
 	// Unfortunately it can't be done in a loop because it uses a #define and compile-time constants for the VectorSwizzle() 'function'.
@@ -947,7 +1044,7 @@ FORCENOINLINE void TestVectorShuffle()
 #define ShuffleTest(A, B, X, Y, Z, W) \
 		V2 = VectorShuffle(A, B, X, Y, Z, W); \
 		V3 = MakeVectorRegister(ArrayV0[X], ArrayV0[Y], ArrayV1[Z], ArrayV1[W]); \
-		LogTest(*FString::Printf(TEXT("VectorShuffle<%d,%d,%d,%d>"), X, Y, Z, W), TestVectorsEqual(V2, V3));
+		LogTest<RealType>(*FString::Printf(TEXT("VectorShuffle<%d,%d,%d,%d>"), X, Y, Z, W), TestVectorsEqual(V2, V3));
 
 	// This is not an exhaustive list because it would be 4*4*4*4 = 256 entries, but it tries to test a lot of common permutations.
 	// Unfortunately it can't be done in a loop because it uses a #define and compile-time constants for the VectorShuffle() 'function'.
@@ -988,7 +1085,173 @@ FORCENOINLINE void TestVectorShuffle()
 }
 
 
+template<typename FloatType, typename VectorRegisterType>
+FORCENOINLINE void TestVectorTrigFunctions()
+{
+	const VectorRegisterType Signs[] = {
+			MakeVectorRegister(+1.0f, +1.0f, +1.0f, +1.0f),
+			MakeVectorRegister(-1.0f, -1.0f, -1.0f, -1.0f),
+	};
+
+	// Sin, Cos
+	{
+		const VectorRegisterType QuadrantDegreesArray[] = {
+			MakeVectorRegister(0.0f, 10.0f, 20.0f, 30.0f),
+			MakeVectorRegister(45.0f, 60.0f, 70.0f, 80.0f),
+		};
+
+		const int32 Cycles = 3; // Go through a full circle this many times (negative and positive)
+		for (int32 OffsetQuadrant = -4 * Cycles; OffsetQuadrant <= 4 * Cycles; ++OffsetQuadrant)
+		{
+			const FloatType OffsetFloat = (FloatType)OffsetQuadrant * FloatType(90.0); // Add 90 degrees repeatedly to cover all quadrants and wrap a few times
+			const VectorRegisterType VOffset = VectorLoadFloat1(&OffsetFloat);
+			for (VectorRegisterType const& VDegrees : QuadrantDegreesArray)
+			{
+				const VectorRegisterType VAnglesDegrees = VectorAdd(VOffset, VDegrees);
+				const VectorRegisterType VAngles = VectorMultiply(VAnglesDegrees, GlobalVectorConstants::DEG_TO_RAD);
+				VectorRegisterType S[3], C[3];
+				TestReferenceSinCos<FloatType, VectorRegisterType>(S[0], C[0], VAngles);
+				TestFastSinCos<FloatType, VectorRegisterType>(S[1], C[1], VAngles);
+				TestVectorSinCos<FloatType, VectorRegisterType>(S[2], C[2], VAngles);
+
+				const FloatType SinCosTolerance = FloatType(1e-5);
+				LogTest<FloatType>(TEXT("SinCos (Sin): Ref vs Fast"), TestVectorsEqual_ComponentWiseError(S[0], S[1], SinCosTolerance));
+				LogTest<FloatType>(TEXT("SinCos (Sin): Ref vs Vec"), TestVectorsEqual_ComponentWiseError(S[0], S[2], SinCosTolerance));
+				LogTest<FloatType>(TEXT("SinCos (Cos): Ref vs Fast"), TestVectorsEqual_ComponentWiseError(C[0], C[1], SinCosTolerance));
+				LogTest<FloatType>(TEXT("SinCos (Cos): Ref vs Vec"), TestVectorsEqual_ComponentWiseError(C[0], C[2], SinCosTolerance));
+
+				VectorRegisterType ReferenceResult, VectorResult;
+
+				const FloatType FastSinTolerance = FloatType(0.001091);
+				LogTest<FloatType>(TEXT("Sin: Ref vs Vec"), TestVectorFunction1Param<FloatType, VectorRegisterType, FMath::Sin, VectorSin>(ReferenceResult, VectorResult, VAngles, FastSinTolerance));
+				LogTest<FloatType>(TEXT("Cos: Ref vs Vec"), TestVectorFunction1Param<FloatType, VectorRegisterType, FMath::Cos, VectorCos>(ReferenceResult, VectorResult, VAngles, FastSinTolerance));
+			}
+		}
+	}
+
+	// Tan
+	{
+		const VectorRegisterType QuadrantDegreesArray[] = {
+			MakeVectorRegister(0.0f, 0.1f, 45.0f, 60.0f),
+			MakeVectorRegister(80.0f, 89.f, 91.0f, 120.0f),
+			MakeVectorRegister(179.9f, 180.0f, 245.0f, 269.f),
+			MakeVectorRegister(271.f, 320.f, 359.f, 359.9f),
+		};
+
+		VectorRegisterType ReferenceResult, VectorResult;
+		for (VectorRegisterType const& Sign : Signs)
+		{
+			for (VectorRegisterType const& VDegrees : QuadrantDegreesArray)
+			{
+				VectorRegisterType TestValue = VectorMultiply(Sign, VDegrees);
+				TestValue = VectorMultiply(TestValue, GlobalVectorConstants::DEG_TO_RAD);
 
+				const FloatType TanTolerance = FloatType(1e-5);
+				LogTest<FloatType>(TEXT("Tan: Ref vs Vec"), TestVectorFunction1Param<FloatType, VectorRegisterType, FMath::Tan, VectorTan>(ReferenceResult, VectorResult, TestValue, TanTolerance));
+			}
+		}
+	}
+
+	// ASin, ACos, ATan tests
+	{
+		const VectorRegisterType ValuesArray[] = {
+			MakeVectorRegister(0.0f,  0.1f,  0.2f,  0.3f),
+			MakeVectorRegister(0.4f,  0.5f,  0.6f,  0.7f),
+			MakeVectorRegister(0.8f,  0.9f,  1.0f,  0.995f),
+		};
+		const VectorRegisterType Multipliers[] = {
+			MakeVectorRegister(+1.0f, +1.0f, +1.0f, +1.0f),
+			MakeVectorRegister(-1.0f, -1.0f, -1.0f, -1.0f),
+		};
+
+		VectorRegisterType ReferenceResult, VectorResult;
+		for (VectorRegisterType const& Sign : Signs)
+		{
+			for (VectorRegisterType const& Value : ValuesArray)
+			{
+				const VectorRegisterType TestValue = VectorMultiply(Sign, Value);
+				const FloatType GenericTrigTolerance = FloatType(1e-6);
+				LogTest<FloatType>(TEXT("ASin: Ref vs Vec"), TestVectorFunction1Param<FloatType, VectorRegisterType, FMath::Asin, VectorASin>(ReferenceResult, VectorResult, TestValue, GenericTrigTolerance));
+				LogTest<FloatType>(TEXT("ACos: Ref vs Vec"), TestVectorFunction1Param<FloatType, VectorRegisterType, FMath::Acos, VectorACos>(ReferenceResult, VectorResult, TestValue, GenericTrigTolerance));
+				LogTest<FloatType>(TEXT("ATan: Ref vs Vec"), TestVectorFunction1Param<FloatType, VectorRegisterType, FMath::Atan, VectorATan>(ReferenceResult, VectorResult, TestValue, GenericTrigTolerance));
+			}
+		}
+	}
+
+	// ATan2
+	{
+		const VectorRegisterType ArrayY[] = {
+			MakeVectorRegister( 0.0f, -0.1f,  0.2f, -0.3f),
+			MakeVectorRegister( 0.4f,  0.5f, -0.6f,  0.7f),
+			MakeVectorRegister( 1.0f, -2.1f,  3.2f, -4.3f),
+			MakeVectorRegister(-5.4f,  6.5f, -7.6f,  8.7f),
+		};
+		const VectorRegisterType ArrayX[] = {
+			MakeVectorRegister(-5.4f, -6.5f, -7.6f, -8.7f),
+			MakeVectorRegister( 1.0f, -2.1f,  3.2f,  4.3f),			
+			MakeVectorRegister( 0.4f, -0.5f, -0.6f,  0.7f),
+			MakeVectorRegister( 0.0f,  0.1f,  0.2f,  0.3f),
+		};
+		static_assert(UE_ARRAY_COUNT(ArrayY) == UE_ARRAY_COUNT(ArrayX));
+
+		VectorRegisterType ReferenceResult, VectorResult;
+		for (int32 Index = 0; Index < UE_ARRAY_COUNT(ArrayY); Index++)
+		{
+			const FloatType GenericTrigTolerance = FloatType(1e-6);
+			LogTest<FloatType>(TEXT("ATan2: Ref vs Vec"), TestVectorFunction2Params<FloatType, VectorRegisterType, FMath::Atan2, VectorATan2>(ReferenceResult, VectorResult, ArrayY[Index], ArrayX[Index], GenericTrigTolerance));
+		}
+	}
+}
+
+
+// Exp, Log tests
+template<typename FloatType, typename VectorRegisterType>
+FORCENOINLINE void TestVectorExpLogFunctions()
+{
+	const VectorRegisterType Signs[] = {
+			MakeVectorRegister(+1.0f, +1.0f, +1.0f, +1.0f),
+			MakeVectorRegister(-1.0f, -1.0f, -1.0f, -1.0f),
+	};
+
+	// Exp tests
+	{
+		const VectorRegisterType ValuesArray[] = {
+			MakeVectorRegister(0.0f, 0.1f, 0.2f, 0.3f),
+			MakeVectorRegister(1.4f, 1.5f, 2.0f, 2.7f),
+			MakeVectorRegister(5.1f, 6.8f, 7.7f, 8.0f),
+		};
+		
+
+		VectorRegisterType ReferenceResult, VectorResult;
+		for (VectorRegisterType const& Sign : Signs)
+		{
+			for (VectorRegisterType const& Value : ValuesArray)
+			{
+				const VectorRegisterType TestValue = VectorMultiply(Sign, Value);
+				const FloatType GenericTestTolerance = FloatType(1e-4);
+				LogTest<FloatType>(TEXT("Exp: Ref vs Vec"), TestVectorFunction1Param<FloatType, VectorRegisterType, FMath::Exp, VectorExp>(ReferenceResult, VectorResult, TestValue, GenericTestTolerance));
+				LogTest<FloatType>(TEXT("Exp2: Ref vs Vec"), TestVectorFunction1Param<FloatType, VectorRegisterType, FMath::Exp2, VectorExp2>(ReferenceResult, VectorResult, TestValue, GenericTestTolerance));
+			}
+		}
+	}
+
+	// Log tests
+	{
+		const VectorRegisterType ValuesArray[] = {
+			MakeVectorRegister(0.01f, 0.1f,  0.2f,  0.3f),
+			MakeVectorRegister(1.4f,  1.5f,  2.0f,  2.7f),
+			MakeVectorRegister(9.1f,  8.8f, 13.7f, 17.0f),
+		};
+		VectorRegisterType ReferenceResult, VectorResult;
+
+		for (VectorRegisterType const& Value : ValuesArray)
+		{
+			const FloatType GenericTestTolerance = FloatType(1e-6);
+			LogTest<FloatType>(TEXT("Loge: Ref vs Vec"), TestVectorFunction1Param<FloatType, VectorRegisterType, FMath::Loge, VectorLog>(ReferenceResult, VectorResult, Value, GenericTestTolerance));
+			LogTest<FloatType>(TEXT("Log2: Ref vs Vec"), TestVectorFunction1Param<FloatType, VectorRegisterType, FMath::Log2, VectorLog2>(ReferenceResult, VectorResult, Value, GenericTestTolerance));
+		}
+	}
+}
 
 
 IMPLEMENT_SIMPLE_AUTOMATION_TEST(FVectorRegisterAbstractionTest, "System.Core.Math.Vector Register Abstraction Test", EAutomationTestFlags::ApplicationContextMask | EAutomationTestFlags::SmokeFilter)
@@ -1005,90 +1268,90 @@ bool RunDoubleVectorTest()
 
 	V0 = MakeVectorRegisterDouble(U1, U1, U1, U1);
 	V1 = MakeVectorRegisterDouble(F1, F1, F1, F1);
-	LogTest(TEXT("MakeVectorRegister<double>"), TestVectorsEqual(V0, V1));
+	LogTest<double>(TEXT("MakeVectorRegister"), TestVectorsEqual(V0, V1));
 
 	V0 = MakeVectorRegisterDouble(0., 0., 0., 0.);
 	V1 = VectorZeroDouble();
-	LogTest(TEXT("VectorZero<double>"), TestVectorsEqual(V0, V1));
+	LogTest<double>(TEXT("VectorZero"), TestVectorsEqual(V0, V1));
 
 	V0 = MakeVectorRegisterDouble(1., 1., 1., 1.);
 	V1 = VectorOneDouble();
-	LogTest(TEXT("VectorOne<double>"), TestVectorsEqual(V0, V1));
+	LogTest<double>(TEXT("VectorOne"), TestVectorsEqual(V0, V1));
 
 	V0 = MakeVectorRegisterDouble(0.f, 1.f, 2.f, 3.f); // floats
 	V1 = MakeVectorRegisterDouble(0.0, 1.0, 2.0, 3.0); // doubles
-	LogTest(TEXT("MakeVectorRegister4Double<float>, MakeVectorRegister4Double<double>"), TestVectorsEqual(V0, V1));
+	LogTest<double>(TEXT("MakeVectorRegister4Double<float>, MakeVectorRegister4Double"), TestVectorsEqual(V0, V1));
 
 	SetScratchDouble(1.0f, 2.0f, -0.25f, -0.5f, 5.f);
 	V0 = MakeVectorRegisterDouble(1.0f, 2.0f, -0.25f, -0.5f);
 	V1 = VectorLoad(GScratchDouble);
-	LogTest(TEXT("VectorLoad<double>"), TestVectorsEqual(V0, V1));
+	LogTest<double>(TEXT("VectorLoad"), TestVectorsEqual(V0, V1));
 
 	SetScratch(1.0f, 2.0f, -0.25f, -0.5f, 5.f);
 	SetScratchDouble(1.0f, 2.0f, -0.25f, -0.5f, 5.f);
 	V0 = VectorLoad(GScratch);		// float*
 	V1 = VectorLoad(GScratchDouble);// double*
-	LogTest(TEXT("VectorLoad<float>, VectorLoad<double>"), TestVectorsEqual(V0, V1));
+	LogTest<double>(TEXT("VectorLoad<float>, VectorLoad"), TestVectorsEqual(V0, V1));
 
 	SetScratchDouble(1.0f, 2.0f, -0.25f, -0.5f, 5.f);
 	V0 = MakeVectorRegisterDouble(1.0f, 2.0f, -0.25f, -0.5f);
 	V1 = VectorLoadAligned(GScratchDouble);
-	LogTest(TEXT("VectorLoadAligned<double>"), TestVectorsEqual(V0, V1));
+	LogTest<double>(TEXT("VectorLoadAligned"), TestVectorsEqual(V0, V1));
 
 	SetScratchDouble(1.0f, 2.0f, -0.25f, -0.5f, 5.f);
 	V0 = VectorLoad(GScratchDouble + 1);
 	V1 = VectorLoadDouble3(GScratchDouble + 1);
-	LogTest(TEXT("VectorLoadFloat3<double>"), TestVectorsEqual3(V0, V1));
+	LogTest<double>(TEXT("VectorLoadFloat3"), TestVectorsEqual3(V0, V1));
 
 	SetScratchDouble(1.0f, 2.0f, -0.25f, -0.5f, 5.f);
 	V0 = MakeVectorRegisterDouble(1.0f, 2.0f, -0.25f, 0.0f);
 	V1 = VectorLoadDouble3_W0(GScratchDouble);
-	LogTest(TEXT("VectorLoadFloat3_W0<double>"), TestVectorsEqual(V0, V1));
+	LogTest<double>(TEXT("VectorLoadFloat3_W0"), TestVectorsEqual(V0, V1));
 
 	SetScratchDouble(1.0f, 2.0f, -0.25f, -0.5f, 5.f);
 	V0 = MakeVectorRegisterDouble(1.0f, 2.0f, -0.25f, 1.0f);
 	V1 = VectorLoadDouble3_W1(GScratchDouble);
-	LogTest(TEXT("VectorLoadFloat3_W1<double>"), TestVectorsEqual(V0, V1));
+	LogTest<double>(TEXT("VectorLoadFloat3_W1"), TestVectorsEqual(V0, V1));
 
 	SetScratchDouble(1.0f, 2.0f, -0.25f, -0.5f, 5.f);
 	V0 = MakeVectorRegisterDouble(-0.5f, -0.5f, -0.5f, -0.5f);
 	V1 = VectorLoadDouble1(GScratchDouble + 3);
-	LogTest(TEXT("VectorLoadFloat1<double>"), TestVectorsEqual(V0, V1));
+	LogTest<double>(TEXT("VectorLoadFloat1"), TestVectorsEqual(V0, V1));
 
 	SetScratchDouble(1.0f, 2.0f, -0.25f, -0.5f, 5.f);
 	V0 = VectorSetDouble3(GScratchDouble[1], GScratchDouble[2], GScratchDouble[3]);
 	GScratchDouble[4] = 127.;
 	V1 = VectorLoadDouble3(GScratchDouble + 1);
-	LogTest(TEXT("VectorSet3<double>"), TestVectorsEqual3(V0, V1));
+	LogTest<double>(TEXT("VectorSet3"), TestVectorsEqual3(V0, V1));
 
 	SetScratchDouble(1.0f, 2.0f, -0.25f, -0.5f, 5.f);
 	V0 = VectorSet(GScratchDouble[1], GScratchDouble[2], GScratchDouble[3], GScratchDouble[4]);
 	V1 = VectorLoad(GScratchDouble + 1);
-	LogTest(TEXT("VectorSet<double>"), TestVectorsEqual(V0, V1));
+	LogTest<double>(TEXT("VectorSet"), TestVectorsEqual(V0, V1));
 
 	V0 = MakeVectorRegisterDouble(1.0f, 2.0f, -0.25f, 1.0f);
 	VectorStoreAligned(V0, GScratchDouble + 8);
 	V1 = VectorLoad(GScratchDouble + 8);
-	LogTest(TEXT("VectorStoreAligned<double>"), TestVectorsEqual(V0, V1));
+	LogTest<double>(TEXT("VectorStoreAligned"), TestVectorsEqual(V0, V1));
 
 	V0 = MakeVectorRegisterDouble(1.0f, 2.0f, -0.55f, 1.0f);
 	VectorStore(V0, GScratchDouble + 7);
 	V1 = VectorLoad(GScratchDouble + 7);
-	LogTest(TEXT("VectorStore<double>"), TestVectorsEqual(V0, V1));
+	LogTest<double>(TEXT("VectorStore"), TestVectorsEqual(V0, V1));
 
 	SetScratchDouble(1.0f, 2.0f, -0.25f, -0.5f, 5.f);
 	V0 = MakeVectorRegisterDouble(5.0f, 3.0f, 1.0f, -1.0f);
 	VectorStoreDouble3(V0, GScratchDouble);
 	V1 = VectorLoad(GScratchDouble);
 	V0 = MakeVectorRegisterDouble(5.0f, 3.0f, 1.0f, -0.5f);
-	LogTest(TEXT("VectorStoreFloat3<double>"), TestVectorsEqual(V0, V1));
+	LogTest<double>(TEXT("VectorStoreFloat3"), TestVectorsEqual(V0, V1));
 
 	SetScratchDouble(1.0f, 2.0f, -0.25f, -0.5f, 5.f);
 	V0 = MakeVectorRegisterDouble(5.0f, 3.0f, 1.0f, -1.0f);
 	VectorStoreDouble1(V0, GScratchDouble + 1);
 	V1 = VectorLoad(GScratchDouble);
 	V0 = MakeVectorRegisterDouble(1.0f, 5.0f, -0.25f, -0.5f);
-	LogTest(TEXT("VectorStoreFloat1<double>"), TestVectorsEqual(V0, V1));
+	LogTest<double>(TEXT("VectorStoreFloat1"), TestVectorsEqual(V0, V1));
 
 	// Replicate
 	TestVectorReplicate<double, VectorRegister4Double>();
@@ -1107,327 +1370,339 @@ bool RunDoubleVectorTest()
 	double D3 = VectorGetComponent(V1, 3);
 	V0 = MakeVectorRegisterDouble(0., 1., 2., 3.);
 	V1 = MakeVectorRegisterDouble(D0, D1, D2, D3);
-	LogTest(TEXT("VectorGetComponent<double>"), TestVectorsEqual(V0, V1));
+	LogTest<double>(TEXT("VectorGetComponent"), TestVectorsEqual(V0, V1));
 
 	// Abs
 	V0 = MakeVectorRegisterDouble(1.0f, -2.0f, 3.0f, -4.0f);
 	V1 = VectorAbs(V0);
 	V0 = MakeVectorRegisterDouble(1.0f, 2.0f, 3.0f, 4.0f);
-	LogTest(TEXT("VectorAbs<double>"), TestVectorsEqual(V0, V1));
+	LogTest<double>(TEXT("VectorAbs"), TestVectorsEqual(V0, V1));
 
 	V0 = MakeVectorRegisterDouble(1.0f, -2.0f, 3.0f, -4.0f);
 	V1 = VectorNegate(V0);
 	V0 = MakeVectorRegisterDouble(-1.0f, 2.0f, -3.0f, 4.0f);
-	LogTest(TEXT("VectorNegate<double>"), TestVectorsEqual(V0, V1));
+	LogTest<double>(TEXT("VectorNegate"), TestVectorsEqual(V0, V1));
 
 	V0 = MakeVectorRegisterDouble(1.0f, 2.0f, 3.0f, 4.0f);
 	V1 = MakeVectorRegisterDouble(2.0f, 4.0f, 6.0f, 8.0f);
 	V1 = VectorAdd(V0, V1);
 	V0 = MakeVectorRegisterDouble(3.0f, 6.0f, 9.0f, 12.0f);
-	LogTest(TEXT("VectorAdd<double>"), TestVectorsEqual(V0, V1));
+	LogTest<double>(TEXT("VectorAdd"), TestVectorsEqual(V0, V1));
 
 	V0 = MakeVectorRegisterDouble(2.0f, 4.0f, 6.0f, 8.0f);
 	V1 = MakeVectorRegisterDouble(1.0f, 2.0f, 3.0f, 4.0f);
 	V1 = VectorSubtract(V0, V1);
 	V0 = MakeVectorRegisterDouble(1.0f, 2.0f, 3.0f, 4.0f);
-	LogTest(TEXT("VectorSubtract<double>"), TestVectorsEqual(V0, V1));
+	LogTest<double>(TEXT("VectorSubtract"), TestVectorsEqual(V0, V1));
 
 	V0 = MakeVectorRegisterDouble(2.0f, 4.0f, 6.0f, 8.0f);
 	V1 = MakeVectorRegisterDouble(1.0f, 2.0f, 3.0f, 4.0f);
 	V2 = VectorMultiply(V0, V1);
 	V3 = MakeVectorRegisterDouble(2.0f, 8.0f, 18.0f, 32.0f);
-	LogTest(TEXT("VectorMultiply<double>"), TestVectorsEqual(V3, V2));
+	LogTest<double>(TEXT("VectorMultiply"), TestVectorsEqual(V3, V2));
 
 	V0 = MakeVectorRegisterDouble(2.0f, 4.0f, 6.0f, 8.0f);
 	V1 = MakeVectorRegisterDouble(1.0f, 2.0f, 3.0f, 4.0f);
 	V2 = VectorMultiplyAdd(V0, V1, VectorOneDouble());
 	V3 = MakeVectorRegisterDouble(3.0f, 9.0f, 19.0f, 33.0f);
-	LogTest(TEXT("VectorMultiplyAdd<double>"), TestVectorsEqual(V3, V2));
+	LogTest<double>(TEXT("VectorMultiplyAdd"), TestVectorsEqual(V3, V2));
 
 	V0 = MakeVectorRegisterDouble(2.0f, 4.0f, 6.0f, 8.0f);
 	V1 = MakeVectorRegisterDouble(1.0f, 2.0f, 3.0f, 4.0f);
 	V1 = VectorDot3(V0, V1);
 	V0 = MakeVectorRegisterDouble(28.0f, 28.0f, 28.0f, 28.0f);
-	LogTest(TEXT("VectorDot3<double>"), TestVectorsEqual(V0, V1));
+	LogTest<double>(TEXT("VectorDot3"), TestVectorsEqual(V0, V1));
+
+	V0 = MakeVectorRegisterDouble(1.0f, 2.0f, 3.0f, 5.0f);
+	V1 = MakeVectorRegisterDouble(7.0f, 11.0f, 13.0f, 17.0f);
+	V1 = VectorDot3(V0, V1);
+	V0 = MakeVectorRegisterDouble(68.0f, 68.0f, 68.0f, 68.0f);
+	LogTest<double>(TEXT("VectorDot3"), TestVectorsEqual(V0, V1));
 
 	V0 = MakeVectorRegisterDouble(2.0f, 4.0f, 6.0f, 8.0f);
 	V1 = MakeVectorRegisterDouble(1.0f, 2.0f, 3.0f, 4.0f);
 	V1 = VectorDot4(V0, V1);
 	V0 = MakeVectorRegisterDouble(60.0f, 60.0f, 60.0f, 60.0f);
-	LogTest(TEXT("VectorDot4<double>"), TestVectorsEqual(V0, V1));
+	LogTest<double>(TEXT("VectorDot4"), TestVectorsEqual(V0, V1));
+
+	V0 = MakeVectorRegisterDouble(1.0f, 2.0f, 3.0f, 5.0f);
+	V1 = MakeVectorRegisterDouble(7.0f, 11.0f, 13.0f, 17.0f);
+	V1 = VectorDot4(V0, V1);
+	V0 = MakeVectorRegisterDouble(153.0f, 153.0f, 153.0f, 153.0f);
+	LogTest<double>(TEXT("VectorDot4"), TestVectorsEqual(V0, V1));
 
 	V0 = MakeVectorRegisterDouble(1.0f, 0.0f, 0.0f, 8.0f);
 	V1 = MakeVectorRegisterDouble(0.0f, 2.0f, 0.0f, 4.0f);
 	V1 = VectorCross(V0, V1);
 	V0 = MakeVectorRegisterDouble(0.f, 0.0f, 2.0f, 0.0f);
-	LogTest(TEXT("VectorCross<double>"), TestVectorsEqual(V0, V1));
+	LogTest<double>(TEXT("VectorCross"), TestVectorsEqual(V0, V1));
 
 	V0 = MakeVectorRegisterDouble(2.0, 3.0, 4.0, 8.0);
 	V1 = MakeVectorRegisterDouble(5.0, 6.0, 7.0, 4.0);
 	V1 = VectorCross(V0, V1);
 	V0 = MakeVectorRegisterDouble(-3.0, 6.0, -3.0, 0.0);
-	LogTest(TEXT("VectorCross<double>"), TestVectorsEqual(V0, V1));
+	LogTest<double>(TEXT("VectorCross"), TestVectorsEqual(V0, V1));
 
 	V0 = MakeVectorRegisterDouble(1.0f, 3.0f, 2.0f, 8.0f);
 	V1 = MakeVectorRegisterDouble(2.0f, 4.0f, 2.0f, 1.0f);
 	V2 = VectorCompareGT(V0, V1);
 	V3 = MakeVectorRegisterDouble((uint64)0, (uint64)0, (uint64)0, (uint64)-1);
-	LogTest(TEXT("VectorCompareGT"), TestVectorsEqualBitwise(V2, V3));
+	LogTest<double>(TEXT("VectorCompareGT"), TestVectorsEqualBitwise(V2, V3));
 
 	V0 = MakeVectorRegisterDouble(1.0f, 3.0f, 2.0f, 8.0f);
 	V1 = MakeVectorRegisterDouble(2.0f, 4.0f, 2.0f, 1.0f);
 	V2 = VectorCompareGE(V0, V1);
 	V3 = MakeVectorRegisterDouble((uint64)0, (uint64)0, (uint64)-1, (uint64)-1);
-	LogTest(TEXT("VectorCompareGE"), TestVectorsEqualBitwise(V2, V3));
+	LogTest<double>(TEXT("VectorCompareGE"), TestVectorsEqualBitwise(V2, V3));
 
 	V0 = MakeVectorRegisterDouble(1.0f, 3.0f, 2.0f, 8.0f);
 	V1 = MakeVectorRegisterDouble(2.0f, 4.0f, 2.0f, 1.0f);
 	V2 = VectorCompareEQ(V0, V1);
 	V3 = MakeVectorRegisterDouble((uint64)0, (uint64)0, (uint64)-1, (uint64)0);
-	LogTest(TEXT("VectorCompareEQ"), TestVectorsEqualBitwise(V2, V3));
+	LogTest<double>(TEXT("VectorCompareEQ"), TestVectorsEqualBitwise(V2, V3));
 
 	V0 = MakeVectorRegisterDouble(1.0f, 3.0f, 2.0f, 8.0f);
 	V1 = MakeVectorRegisterDouble(2.0f, 4.0f, 2.0f, 1.0f);
 	V2 = VectorCompareNE(V0, V1);
 	V3 = MakeVectorRegisterDouble((uint64)(0xFFFFFFFFFFFFFFFFU), (uint64)(0xFFFFFFFFFFFFFFFFU), (uint64)(0), (uint64)(0xFFFFFFFFFFFFFFFFU));
-	LogTest(TEXT("VectorCompareNE"), TestVectorsEqualBitwise(V2, V3));
+	LogTest<double>(TEXT("VectorCompareNE"), TestVectorsEqualBitwise(V2, V3));
 
 	V0 = MakeVectorRegisterDouble(1.0f, 3.0f, 5.0f, 7.0f);
 	V1 = MakeVectorRegisterDouble(2.0f, 4.0f, 6.0f, 8.0f);
 	V2 = MakeVectorRegisterDouble((uint64)-1, (uint64)0, (uint64)0, (uint64)-1);
 	V2 = VectorSelect(V2, V0, V1);
 	V3 = MakeVectorRegisterDouble(1.0f, 4.0f, 6.0f, 7.0f);
-	LogTest(TEXT("VectorSelect"), TestVectorsEqual(V2, V3));
+	LogTest<double>(TEXT("VectorSelect"), TestVectorsEqual(V2, V3));
 
 	V0 = MakeVectorRegisterDouble(1.0f, 3.0f, 5.0f, 7.0f);
 	V1 = MakeVectorRegisterDouble(2.0f, 4.0f, 6.0f, 8.0f);
 	V2 = MakeVectorRegisterDouble((uint64)0, (uint64)-1, (uint64)-1, (uint64)0);
 	V2 = VectorSelect(V2, V0, V1);
 	V3 = MakeVectorRegisterDouble(2.0f, 3.0f, 5.0f, 8.0f);
-	LogTest(TEXT("VectorSelect"), TestVectorsEqual(V2, V3));
+	LogTest<double>(TEXT("VectorSelect"), TestVectorsEqual(V2, V3));
 
 	V0 = MakeVectorRegisterDouble(1.0f, 3.0f, 0.0f, 0.0f);
 	V1 = MakeVectorRegisterDouble(0.0f, 0.0f, 2.0f, 1.0f);
 	V2 = VectorBitwiseOr(V0, V1);
 	V3 = MakeVectorRegisterDouble(1.0f, 3.0f, 2.0f, 1.0f);
-	LogTest(TEXT("VectorBitwiseOr-Double1"), TestVectorsEqual(V2, V3));
+	LogTest<double>(TEXT("VectorBitwiseOr-Double1"), TestVectorsEqual(V2, V3));
 
 	V0 = MakeVectorRegisterDouble(1.0f, 3.0f, 24.0f, 36.0f);
 	V1 = MakeVectorRegisterDouble((uint64)(0x8000000000000000U), (uint64)(0x8000000000000000U), (uint64)(0x8000000000000000U), (uint64)(0x8000000000000000U));
 	V2 = VectorBitwiseOr(V0, V1);
 	V3 = MakeVectorRegisterDouble(-1.0f, -3.0f, -24.0f, -36.0f);
-	LogTest(TEXT("VectorBitwiseOr-Double2"), TestVectorsEqual(V2, V3));
+	LogTest<double>(TEXT("VectorBitwiseOr-Double2"), TestVectorsEqual(V2, V3));
 
 	V0 = MakeVectorRegisterDouble(-1.0f, -3.0f, -24.0f, 36.0f);
 	V1 = MakeVectorRegisterDouble((uint64)(0xFFFFFFFFFFFFFFFFU), (uint64)(0x7FFFFFFFFFFFFFFFU), (uint64)(0x7FFFFFFFFFFFFFFFU), (uint64)(0xFFFFFFFFFFFFFFFFU));
 	V2 = VectorBitwiseAnd(V0, V1);
 	V3 = MakeVectorRegisterDouble(-1.0f, 3.0f, 24.0f, 36.0f);
-	LogTest(TEXT("VectorBitwiseAnd-Double"), TestVectorsEqual(V2, V3));
+	LogTest<double>(TEXT("VectorBitwiseAnd-Double"), TestVectorsEqual(V2, V3));
 
 	V0 = MakeVectorRegisterDouble(-1.0f, -3.0f, -24.0f, 36.0f);
 	V1 = MakeVectorRegisterDouble((uint64)(0x8000000000000000U), (uint64)(0x0000000000000000U), (uint64)(0x8000000000000000U), (uint64)(0x8000000000000000U));
 	V2 = VectorBitwiseXor(V0, V1);
 	V3 = MakeVectorRegisterDouble(1.0f, -3.0f, 24.0f, -36.0f);
-	LogTest(TEXT("VectorBitwiseXor-Double"), TestVectorsEqual(V2, V3));
+	LogTest<double>(TEXT("VectorBitwiseXor-Double"), TestVectorsEqual(V2, V3));
 
 	V0 = MakeVectorRegisterDouble(2.0f, -2.0f, 3.0f, -4.0f);
 	V1 = VectorSet_W0(V0);
 	V0 = MakeVectorRegisterDouble(2.0f, -2.0f, 3.0f, 0.0f);
-	LogTest(TEXT("VectorSet_W0"), TestVectorsEqual(V0, V1));
+	LogTest<double>(TEXT("VectorSet_W0"), TestVectorsEqual(V0, V1));
 
 	V0 = MakeVectorRegisterDouble(2.0f, -2.0f, 3.0f, -4.0f);
 	V1 = VectorSet_W1(V0);
 	V0 = MakeVectorRegisterDouble(2.0f, -2.0f, 3.0f, 1.0f);
-	LogTest(TEXT("VectorSet_W1"), TestVectorsEqual(V0, V1));
+	LogTest<double>(TEXT("VectorSet_W1"), TestVectorsEqual(V0, V1));
 
 	V0 = MakeVectorRegisterDouble(2.0f, 4.0f, 6.0f, 8.0f);
 	V1 = MakeVectorRegisterDouble(4.0f, 3.0f, 2.0f, 1.0f);
 	V1 = VectorMin(V0, V1);
 	V0 = MakeVectorRegisterDouble(2.0f, 3.0f, 2.0f, 1.0f);
-	LogTest(TEXT("VectorMin"), TestVectorsEqual(V0, V1));
+	LogTest<double>(TEXT("VectorMin"), TestVectorsEqual(V0, V1));
 
 	V0 = MakeVectorRegisterDouble(2.0f, 4.0f, 6.0f, 8.0f);
 	V1 = MakeVectorRegisterDouble(4.0f, 3.0f, 2.0f, 1.0f);
 	V1 = VectorMax(V0, V1);
 	V0 = MakeVectorRegisterDouble(4.0f, 4.0f, 6.0f, 8.0f);
-	LogTest(TEXT("VectorMax"), TestVectorsEqual(V0, V1));
+	LogTest<double>(TEXT("VectorMax"), TestVectorsEqual(V0, V1));
 
 	V0 = MakeVectorRegisterDouble(1.0, 3.0, 5.0, 7.0);
 	V1 = MakeVectorRegisterDouble(2.0, 4.0, 6.0, 8.0);
 	V1 = VectorCombineHigh(V0, V1);
 	V0 = MakeVectorRegisterDouble(5.0, 7.0, 6.0, 8.0);
-	LogTest(TEXT("VectorCombineHigh"), TestVectorsEqual(V0, V1));
+	LogTest<double>(TEXT("VectorCombineHigh"), TestVectorsEqual(V0, V1));
 
 	V0 = MakeVectorRegisterDouble(1.0, 3.0, 5.0, 7.0);
 	V1 = MakeVectorRegisterDouble(2.0, 4.0, 6.0, 8.0);
 	V1 = VectorCombineLow(V0, V1);
 	V0 = MakeVectorRegisterDouble(1.0, 3.0, 2.0, 4.0);
-	LogTest(TEXT("VectorCombineLow"), TestVectorsEqual(V0, V1));
+	LogTest<double>(TEXT("VectorCombineLow"), TestVectorsEqual(V0, V1));
 
 	//TODO: Rename functions to explicitly mention Double?
 	V0 = MakeVectorRegister(2.0, 4.0, 6.0, 8.0);
 	V1 = MakeVectorRegister(4.0, 3.0, 2.0, 1.0);
 	bool bIsVAGT_TRUE = VectorAnyGreaterThan(V0, V1) != 0;
-	LogTest(TEXT("VectorAnyGreaterThan-true"), bIsVAGT_TRUE);
+	LogTest<double>(TEXT("VectorAnyGreaterThan-true"), bIsVAGT_TRUE);
 
 	V0 = MakeVectorRegister(1.0, 3.0, 2.0, 1.0);
 	V1 = MakeVectorRegister(2.0, 4.0, 6.0, 8.0);
 	bool bIsVAGT_FALSE = VectorAnyGreaterThan(V0, V1) == 0;
-	LogTest(TEXT("VectorAnyGreaterThan-false"), bIsVAGT_FALSE);
+	LogTest<double>(TEXT("VectorAnyGreaterThan-false"), bIsVAGT_FALSE);
 
 	V0 = MakeVectorRegister(1.0, 3.0, 2.0, 1.0);
 	V1 = MakeVectorRegister(2.0, 4.0, 6.0, 8.0);
-	LogTest(TEXT("VectorAnyLesserThan-true"), VectorAnyLesserThan(V0, V1) != 0);
+	LogTest<double>(TEXT("VectorAnyLesserThan-true"), VectorAnyLesserThan(V0, V1) != 0);
 
 	V0 = MakeVectorRegister(3.0, 5.0, 7.0, 9.0);
 	V1 = MakeVectorRegister(2.0, 4.0, 6.0, 8.0);
-	LogTest(TEXT("VectorAnyLesserThan-false"), VectorAnyLesserThan(V0, V1) == 0);
+	LogTest<double>(TEXT("VectorAnyLesserThan-false"), VectorAnyLesserThan(V0, V1) == 0);
 
 	V0 = MakeVectorRegister(3.0, 5.0, 7.0, 9.0);
 	V1 = MakeVectorRegister(2.0, 4.0, 6.0, 8.0);
-	LogTest(TEXT("VectorAllGreaterThan-true"), VectorAllGreaterThan(V0, V1) != 0);
+	LogTest<double>(TEXT("VectorAllGreaterThan-true"), VectorAllGreaterThan(V0, V1) != 0);
 
 	V0 = MakeVectorRegister(3.0, 1.0, 7.0, 9.0);
 	V1 = MakeVectorRegister(2.0, 4.0, 6.0, 8.0);
-	LogTest(TEXT("VectorAllGreaterThan-false"), VectorAllGreaterThan(V0, V1) == 0);
+	LogTest<double>(TEXT("VectorAllGreaterThan-false"), VectorAllGreaterThan(V0, V1) == 0);
 
 	V0 = MakeVectorRegister(1.0, 3.0, 2.0, 1.0);
 	V1 = MakeVectorRegister(2.0, 4.0, 6.0, 8.0);
-	LogTest(TEXT("VectorAllLesserThan-true"), VectorAllLesserThan(V0, V1) != 0);
+	LogTest<double>(TEXT("VectorAllLesserThan-true"), VectorAllLesserThan(V0, V1) != 0);
 
 	V0 = MakeVectorRegister(3.0, 3.0, 2.0, 1.0);
 	V1 = MakeVectorRegister(2.0, 4.0, 6.0, 8.0);
-	LogTest(TEXT("VectorAllLesserThan-false"), VectorAllLesserThan(V0, V1) == 0);
+	LogTest<double>(TEXT("VectorAllLesserThan-false"), VectorAllLesserThan(V0, V1) == 0);
 
 	V0 = MakeVectorRegister(1.0, 3.0, 2.0, 8.0);
 	V1 = MakeVectorRegister(2.0, 4.0, 2.0, 1.0);
 	V2 = VectorCompareGT(V0, V1);
 	V3 = MakeVectorRegisterDouble((uint64)0, (uint64)0, (uint64)0, (uint64)-1);
-	LogTest(TEXT("VectorCompareGT"), TestVectorsEqualBitwise(V2, V3));
+	LogTest<double>(TEXT("VectorCompareGT"), TestVectorsEqualBitwise(V2, V3));
 
 	V0 = MakeVectorRegister(1.0, 3.0, 2.0, 8.0);
 	V1 = MakeVectorRegister(2.0, 4.0, 2.0, 1.0);
 	V2 = VectorCompareGE(V0, V1);
 	V3 = MakeVectorRegisterDouble((uint64)0, (uint64)0, (uint64)-1, (uint64)-1);
-	LogTest(TEXT("VectorCompareGE"), TestVectorsEqualBitwise(V2, V3));
+	LogTest<double>(TEXT("VectorCompareGE"), TestVectorsEqualBitwise(V2, V3));
 
 	V0 = MakeVectorRegister(1.0, 3.0, 2.0, 8.0);
 	V1 = MakeVectorRegister(2.0, 4.0, 2.0, 1.0);
 	V2 = VectorCompareEQ(V0, V1);
 	V3 = MakeVectorRegisterDouble((uint64)0, (uint64)0, (uint64)-1, (uint64)0);
-	LogTest(TEXT("VectorCompareEQ"), TestVectorsEqualBitwise(V2, V3));
+	LogTest<double>(TEXT("VectorCompareEQ"), TestVectorsEqualBitwise(V2, V3));
 
 	V0 = MakeVectorRegister(1.0, 3.0, 2.0, 8.0);
 	V1 = MakeVectorRegister(2.0, 4.0, 2.0, 1.0);
 	V2 = VectorCompareNE(V0, V1);
 	V3 = MakeVectorRegisterDouble((uint64)(0xFFFFFFFFFFFFFFFF), (uint64)(0xFFFFFFFFFFFFFFFF), (uint64)(0), (uint64)(0xFFFFFFFFFFFFFFFF));
-	LogTest(TEXT("VectorCompareNE"), TestVectorsEqualBitwise(V2, V3));
+	LogTest<double>(TEXT("VectorCompareNE"), TestVectorsEqualBitwise(V2, V3));
 
 	V0 = MakeVectorRegister(-1.0, -3.0, -24.0, 36.0);
 	V1 = MakeVectorRegister(5.0, 35.0, 23.0, 48.0);
 	V2 = VectorMergeVecXYZ_VecW(V0, V1);
 	V3 = MakeVectorRegister(-1.0, -3.0, -24.0, 48.0);
-	LogTest(TEXT("VectorMergeXYZ_VecW-1"), TestVectorsEqual(V2, V3));
+	LogTest<double>(TEXT("VectorMergeXYZ_VecW-1"), TestVectorsEqual(V2, V3));
 
 	V0 = MakeVectorRegister(-1.0, -3.0, -24.0, 36.0);
 	V1 = MakeVectorRegister(5.0, 35.0, 23.0, 48.0);
 	V2 = VectorMergeVecXYZ_VecW(V1, V0);
 	V3 = MakeVectorRegister(5.0, 35.0, 23.0, 36.0);
-	LogTest(TEXT("VectorMergeXYZ_VecW-2"), TestVectorsEqual(V2, V3));
+	LogTest<double>(TEXT("VectorMergeXYZ_VecW-2"), TestVectorsEqual(V2, V3));
 
 	V0 = MakeVectorRegister(1.0, 1.0e6, 1.3e-8, 35.0);
 	V1 = VectorReciprocal(V0);
 	V3 = VectorMultiply(V1, V0);
-	LogTest(TEXT("VectorReciprocal"), TestVectorsEqual(VectorOneDouble(), V3, 0.008));
+	LogTest<double>(TEXT("VectorReciprocal"), TestVectorsEqual(VectorOneDouble(), V3, 0.008));
 
 	V0 = MakeVectorRegister(1.0, 1.0e6, 1.3e-8, 35.0);
 	V1 = VectorReciprocalAccurate(V0);
 	V3 = VectorMultiply(V1, V0);
-	LogTest(TEXT("VectorReciprocalAccurate"), TestVectorsEqual(VectorOneDouble(), V3, 1e-7));
+	LogTest<double>(TEXT("VectorReciprocalAccurate"), TestVectorsEqual(VectorOneDouble(), V3, 1e-7));
 
 	V0 = MakeVectorRegister(1.0, 1.0e6, 1.3e-8, 35.0);
 	V1 = VectorReciprocalSqrt(V0);
 	V3 = VectorMultiply(VectorMultiply(V1, V1), V0);
-	LogTest(TEXT("VectorReciprocalSqrt"), TestVectorsEqual(VectorOneDouble(), V3, 0.007));
+	LogTest<double>(TEXT("VectorReciprocalSqrt"), TestVectorsEqual(VectorOneDouble(), V3, 0.007));
 
 	V0 = MakeVectorRegister(1.0, 1.0e6, 1.3e-8, 35.0);
 	V1 = VectorReciprocalSqrtAccurate(V0);
 	V3 = VectorMultiply(VectorMultiply(V1, V1), V0);
-	LogTest(TEXT("VectorReciprocalSqrtAccurate"), TestVectorsEqual(VectorOneDouble(), V3, 1e-6));
+	LogTest<double>(TEXT("VectorReciprocalSqrtAccurate"), TestVectorsEqual(VectorOneDouble(), V3, 1e-6));
 
 	// VectorMod
 	V0 = MakeVectorRegister(0.0, 3.2, 2.8, 1.5);
 	V1 = MakeVectorRegister(2.0, 1.2, 2.0, 3.0);
 	V2 = TestReferenceMod(V0, V1);
 	V3 = VectorMod(V0, V1);
-	LogTest(TEXT("VectorMod positive"), TestVectorsEqual(V2, V3));
+	LogTest<double>(TEXT("VectorMod positive"), TestVectorsEqual(V2, V3));
 
 	V0 = MakeVectorRegister(-2.0, 3.2, -2.8, -1.5);
 	V1 = MakeVectorRegister(-1.5, -1.2, 2.0, 3.0);
 	V2 = TestReferenceMod(V0, V1);
 	V3 = VectorMod(V0, V1);
-	LogTest(TEXT("VectorMod negative"), TestVectorsEqual(V2, V3));
+	LogTest<double>(TEXT("VectorMod negative"), TestVectorsEqual(V2, V3));
 
 	// VectorSign
 	V0 = MakeVectorRegister(2.0, -2.0, 0.0, -3.0);
 	V2 = MakeVectorRegister(1.0, -1.0, 1.0, -1.0);
 	V3 = VectorSign(V0);
-	LogTest(TEXT("VectorSign"), TestVectorsEqual(V2, V3));
+	LogTest<double>(TEXT("VectorSign"), TestVectorsEqual(V2, V3));
 
 	// VectorStep
 	V0 = MakeVectorRegister(2.0, -2.0, 0.0, -3.0);
 	V2 = MakeVectorRegister(1.0, 0.0, 1.0, 0.0);
 	V3 = VectorStep(V0);
-	LogTest(TEXT("VectorStep"), TestVectorsEqual(V2, V3));
+	LogTest<double>(TEXT("VectorStep"), TestVectorsEqual(V2, V3));
 
 	// VectorTruncate
 	V0 = MakeVectorRegister(-1.8, -1.0, -0.8, 0.0);
 	V2 = MakeVectorRegister(-1.0, -1.0, 0.0, 0.0);
 	V3 = VectorTruncate(V0);
-	LogTest(TEXT("VectorTruncate"), TestVectorsEqual(V2, V3, DOUBLE_KINDA_SMALL_NUMBER));
+	LogTest<double>(TEXT("VectorTruncate"), TestVectorsEqual(V2, V3, DOUBLE_KINDA_SMALL_NUMBER));
 
 	V0 = MakeVectorRegister(0.0, 0.8, 1.0, 1.8);
 	V2 = MakeVectorRegister(0.0, 0.0, 1.0, 1.0);
 	V3 = VectorTruncate(V0);
-	LogTest(TEXT("VectorTruncate"), TestVectorsEqual(V2, V3, DOUBLE_KINDA_SMALL_NUMBER));
+	LogTest<double>(TEXT("VectorTruncate"), TestVectorsEqual(V2, V3, DOUBLE_KINDA_SMALL_NUMBER));
 
 	V0 = MakeVectorRegister(0.1, 1.8, 2.4, -2.4);
 	V2 = MakeVectorRegister(0.0, 1.0, 2.0, -2.0);
 	V3 = VectorTruncate(V0);
-	LogTest(TEXT("VectorTruncate"), TestVectorsEqual(V2, V3, DOUBLE_KINDA_SMALL_NUMBER));
+	LogTest<double>(TEXT("VectorTruncate"), TestVectorsEqual(V2, V3, DOUBLE_KINDA_SMALL_NUMBER));
 
 	// VectorFractional
 	V0 = MakeVectorRegister(-1.8, -1.0, -0.8, 0.0);
 	V2 = MakeVectorRegister(-0.8, 0.0, -0.8, 0.0);
 	V3 = VectorFractional(V0);
-	LogTest(TEXT("VectorFractional"), TestVectorsEqual(V2, V3, DOUBLE_KINDA_SMALL_NUMBER));
+	LogTest<double>(TEXT("VectorFractional"), TestVectorsEqual(V2, V3, DOUBLE_KINDA_SMALL_NUMBER));
 
 	V0 = MakeVectorRegister(0.0, 0.8, 1.0, 1.8);
 	V2 = MakeVectorRegister(0.0, 0.8, 0.0, 0.8);
 	V3 = VectorFractional(V0);
-	LogTest(TEXT("VectorFractional"), TestVectorsEqual(V2, V3, DOUBLE_KINDA_SMALL_NUMBER));
+	LogTest<double>(TEXT("VectorFractional"), TestVectorsEqual(V2, V3, DOUBLE_KINDA_SMALL_NUMBER));
 
 	// VectorCeil
 	V0 = MakeVectorRegister(-1.8, -1.0, -0.8, 0.0);
 	V2 = MakeVectorRegister(-1.0, -1.0, -0.0, 0.0);
 	V3 = VectorCeil(V0);
-	LogTest(TEXT("VectorCeil"), TestVectorsEqual(V2, V3, DOUBLE_KINDA_SMALL_NUMBER));
+	LogTest<double>(TEXT("VectorCeil"), TestVectorsEqual(V2, V3, DOUBLE_KINDA_SMALL_NUMBER));
 
 	V0 = MakeVectorRegister(0.0, 0.8, 1.0, 1.8);
 	V2 = MakeVectorRegister(0.0, 1.0, 1.0, 2.0);
 	V3 = VectorCeil(V0);
-	LogTest(TEXT("VectorCeil"), TestVectorsEqual(V2, V3, DOUBLE_KINDA_SMALL_NUMBER));
+	LogTest<double>(TEXT("VectorCeil"), TestVectorsEqual(V2, V3, DOUBLE_KINDA_SMALL_NUMBER));
 
 	// VectorFloor
 	V0 = MakeVectorRegister(-1.8, -1.0, -0.8, 0.0);
 	V2 = MakeVectorRegister(-2.0, -1.0, -1.0, 0.0);
 	V3 = VectorFloor(V0);
-	LogTest(TEXT("VectorFloor"), TestVectorsEqual(V2, V3, DOUBLE_KINDA_SMALL_NUMBER));
+	LogTest<double>(TEXT("VectorFloor"), TestVectorsEqual(V2, V3, DOUBLE_KINDA_SMALL_NUMBER));
 
 	V0 = MakeVectorRegister(0.0, 0.8, 1.0, 1.8);
 	V2 = MakeVectorRegister(0.0, 0.0, 1.0, 1.0);
 	V3 = VectorFloor(V0);
-	LogTest(TEXT("VectorFloor"), TestVectorsEqual(V2, V3));
+	LogTest<double>(TEXT("VectorFloor"), TestVectorsEqual(V2, V3));
 
 	// NaN / Inf tests
 	// Using a union as we need to do a bitwise cast of 0xFFFFFFFF into a float.
@@ -1439,13 +1714,28 @@ bool RunDoubleVectorTest()
 	Int64DoubleUnion NaNU;
 	NaNU.IntNaN = 0xFFFFFFFFFFFFFFFF;
 	const double NaN = NaNU.DoubleNaN;
-	LogTest(TEXT("VectorContainsNaNOrInfinite true"), VectorContainsNaNOrInfinite(MakeVectorRegisterDouble(NaN, NaN, NaN, NaN)));
-	LogTest(TEXT("VectorContainsNaNOrInfinite true"), VectorContainsNaNOrInfinite(MakeVectorRegisterDouble(NaN, 0.0, 0.0, 0.0)));
-	LogTest(TEXT("VectorContainsNaNOrInfinite true"), VectorContainsNaNOrInfinite(MakeVectorRegisterDouble(0.0, 0.0, 0.0, NaN)));
-	LogTest(TEXT("VectorContainsNaNOrInfinite true"), VectorContainsNaNOrInfinite(GlobalVectorConstants::DoubleInfinity));
-	LogTest(TEXT("VectorContainsNaNOrInfinite true"), VectorContainsNaNOrInfinite(MakeVectorRegisterDouble((uint64)0xFFF0000000000000ULL, (uint64)0xFFF0000000000000ULL,
+	SetScratch(0.0, 0.0, 0.0, 0.0);
+	LogTest<double>(TEXT("VectorContainsNaNOrInfinite false"), !VectorContainsNaNOrInfinite(MakeVectorRegisterDouble(0.0, 1.0, -1.0, 0.0)));
+	LogTest<double>(TEXT("VectorContainsNaNOrInfinite true"), VectorContainsNaNOrInfinite(MakeVectorRegisterDouble(NaN, NaN, NaN, NaN)));
+	LogTest<double>(TEXT("VectorContainsNaNOrInfinite true"), VectorContainsNaNOrInfinite(MakeVectorRegisterDouble(NaN, 0.0, 0.0, 0.0)));
+	LogTest<double>(TEXT("VectorContainsNaNOrInfinite true"), VectorContainsNaNOrInfinite(MakeVectorRegisterDouble(0.0, 0.0, 0.0, NaN)));
+	LogTest<double>(TEXT("VectorContainsNaNOrInfinite true"), VectorContainsNaNOrInfinite(GlobalVectorConstants::DoubleInfinity));
+	LogTest<double>(TEXT("VectorContainsNaNOrInfinite true"), VectorContainsNaNOrInfinite(MakeVectorRegisterDouble((uint64)0xFFF0000000000000ULL, (uint64)0xFFF0000000000000ULL,
 																										   (uint64)0xFFF0000000000000ULL, (uint64)0xFFF0000000000000ULL))); // negative infinity
-	LogTest(TEXT("VectorContainsNaNOrInfinite true"), VectorContainsNaNOrInfinite(GlobalVectorConstants::DoubleAllMask));
+	LogTest<double>(TEXT("VectorContainsNaNOrInfinite true"), VectorContainsNaNOrInfinite(GlobalVectorConstants::DoubleAllMask));
+	
+	// Not Nan/Inf
+	LogTest<double>(TEXT("VectorContainsNaNOrInfinite false"), !VectorContainsNaNOrInfinite(GlobalVectorConstants::DoubleZero));
+	LogTest<double>(TEXT("VectorContainsNaNOrInfinite false"), !VectorContainsNaNOrInfinite(GlobalVectorConstants::DoubleOne));
+	LogTest<double>(TEXT("VectorContainsNaNOrInfinite false"), !VectorContainsNaNOrInfinite(GlobalVectorConstants::DoubleMinusOneHalf));
+	LogTest<double>(TEXT("VectorContainsNaNOrInfinite false"), !VectorContainsNaNOrInfinite(GlobalVectorConstants::DoubleSmallNumber));
+	LogTest<double>(TEXT("VectorContainsNaNOrInfinite false"), !VectorContainsNaNOrInfinite(GlobalVectorConstants::DoubleBigNumber));
+
+	// Sin, Cos, Tan tests
+	TestVectorTrigFunctions<double, VectorRegister4Double>();
+
+	// Exp, Log tests
+	TestVectorExpLogFunctions<double, VectorRegister4Double>();
 
 	return GPassing;
 }
@@ -1463,84 +1753,84 @@ bool FVectorRegisterAbstractionTest::RunTest(const FString& Parameters)
 
 	V0 = MakeVectorRegister( D1, D1, D1, D1 );
 	V1 = MakeVectorRegister( F1, F1, F1, F1 );
-	LogTest( TEXT("MakeVectorRegister"), TestVectorsEqual( V0, V1 ) );
+	LogTest<float>( TEXT("MakeVectorRegister"), TestVectorsEqual( V0, V1 ) );
 
 	V0 = MakeVectorRegister( 0.f, 0.f, 0.f, 0.f );
 	V1 = VectorZero();
-	LogTest( TEXT("VectorZero"), TestVectorsEqual( V0, V1 ) );
+	LogTest<float>( TEXT("VectorZero"), TestVectorsEqual( V0, V1 ) );
 
 	V0 = MakeVectorRegister( 1.f, 1.f, 1.f, 1.f );
 	V1 = VectorOne();
-	LogTest( TEXT("VectorOne"), TestVectorsEqual( V0, V1 ) );
+	LogTest<float>( TEXT("VectorOne"), TestVectorsEqual( V0, V1 ) );
 
 	SetScratch( 1.0f, 2.0f, -0.25f, -0.5f, 5.f );
 	V0 = MakeVectorRegister( 1.0f, 2.0f, -0.25f, -0.5f );
 	V1 = VectorLoad( GScratch );
-	LogTest( TEXT("VectorLoad"), TestVectorsEqual( V0, V1 ) );
+	LogTest<float>( TEXT("VectorLoad"), TestVectorsEqual( V0, V1 ) );
 
 	SetScratch( 1.0f, 2.0f, -0.25f, -0.5f, 5.f );
 	V0 = MakeVectorRegister( 1.0f, 2.0f, -0.25f, -0.5f );
 	V1 = VectorLoad( GScratch );
-	LogTest( TEXT("VectorLoad"), TestVectorsEqual( V0, V1 ) );
+	LogTest<float>( TEXT("VectorLoad"), TestVectorsEqual( V0, V1 ) );
 
 	SetScratch( 1.0f, 2.0f, -0.25f, -0.5f, 5.f );
 	V0 = MakeVectorRegister( 1.0f, 2.0f, -0.25f, -0.5f );
 	V1 = VectorLoadAligned( GScratch );
-	LogTest( TEXT("VectorLoadAligned"), TestVectorsEqual( V0, V1 ) );
+	LogTest<float>( TEXT("VectorLoadAligned"), TestVectorsEqual( V0, V1 ) );
 
 	SetScratch( 1.0f, 2.0f, -0.25f, -0.5f, 5.f );
 	V0 = VectorLoad( GScratch + 1 );
 	V1 = VectorLoadFloat3( GScratch + 1 );
-	LogTest( TEXT("VectorLoadFloat3"), TestVectorsEqual3( V0, V1 ) );
+	LogTest<float>( TEXT("VectorLoadFloat3"), TestVectorsEqual3( V0, V1 ) );
 
 	SetScratch( 1.0f, 2.0f, -0.25f, -0.5f, 5.f );
 	V0 = MakeVectorRegister( 1.0f, 2.0f, -0.25f, 0.0f );
 	V1 = VectorLoadFloat3_W0( GScratch );
-	LogTest( TEXT("VectorLoadFloat3_W0"), TestVectorsEqual( V0, V1 ) );
+	LogTest<float>( TEXT("VectorLoadFloat3_W0"), TestVectorsEqual( V0, V1 ) );
 
 	SetScratch( 1.0f, 2.0f, -0.25f, -0.5f, 5.f );
 	V0 = MakeVectorRegister( 1.0f, 2.0f, -0.25f, 1.0f );
 	V1 = VectorLoadFloat3_W1( GScratch );
-	LogTest( TEXT("VectorLoadFloat3_W1"), TestVectorsEqual( V0, V1 ) );
+	LogTest<float>( TEXT("VectorLoadFloat3_W1"), TestVectorsEqual( V0, V1 ) );
 
 	SetScratch( 1.0f, 2.0f, -0.25f, -0.5f, 5.f );
 	V0 = MakeVectorRegister( -0.5f, -0.5f, -0.5f, -0.5f );
 	V1 = VectorLoadFloat1( GScratch + 3 );
-	LogTest( TEXT("VectorLoadFloat1"), TestVectorsEqual( V0, V1 ) );
+	LogTest<float>( TEXT("VectorLoadFloat1"), TestVectorsEqual( V0, V1 ) );
 
 	SetScratch( 1.0f, 2.0f, -0.25f, -0.5f, 5.f );
 	V0 = VectorSetFloat3( GScratch[1], GScratch[2], GScratch[3] );
 	V1 = VectorLoadFloat3( GScratch + 1 );
-	LogTest( TEXT("VectorSetFloat3"), TestVectorsEqual3( V0, V1 ) );
+	LogTest<float>( TEXT("VectorSetFloat3"), TestVectorsEqual3( V0, V1 ) );
 
 	SetScratch( 1.0f, 2.0f, -0.25f, -0.5f, 5.f );
 	V0 = VectorSet( GScratch[1], GScratch[2], GScratch[3], GScratch[4] );
 	V1 = VectorLoad( GScratch + 1 );
-	LogTest( TEXT("VectorSet"), TestVectorsEqual( V0, V1 ) );
+	LogTest<float>( TEXT("VectorSet"), TestVectorsEqual( V0, V1 ) );
 
 	V0 = MakeVectorRegister( 1.0f, 2.0f, -0.25f, 1.0f );
 	VectorStoreAligned( V0, GScratch + 8 );
 	V1 = VectorLoad( GScratch + 8 );
-	LogTest( TEXT("VectorStoreAligned"), TestVectorsEqual( V0, V1 ) );
+	LogTest<float>( TEXT("VectorStoreAligned"), TestVectorsEqual( V0, V1 ) );
 
 	V0 = MakeVectorRegister( 1.0f, 2.0f, -0.55f, 1.0f );
 	VectorStore( V0, GScratch + 7 );
 	V1 = VectorLoad( GScratch + 7 );
-	LogTest( TEXT("VectorStore"), TestVectorsEqual( V0, V1 ) );
+	LogTest<float>( TEXT("VectorStore"), TestVectorsEqual( V0, V1 ) );
 
 	SetScratch( 1.0f, 2.0f, -0.25f, -0.5f, 5.f );
 	V0 = MakeVectorRegister( 5.0f, 3.0f, 1.0f, -1.0f );
 	VectorStoreFloat3( V0, GScratch );
 	V1 = VectorLoad( GScratch );
 	V0 = MakeVectorRegister( 5.0f, 3.0f, 1.0f, -0.5f );
-	LogTest( TEXT("VectorStoreFloat3"), TestVectorsEqual( V0, V1 ) );
+	LogTest<float>( TEXT("VectorStoreFloat3"), TestVectorsEqual( V0, V1 ) );
 
 	SetScratch( 1.0f, 2.0f, -0.25f, -0.5f, 5.f );
 	V0 = MakeVectorRegister( 5.0f, 3.0f, 1.0f, -1.0f );
 	VectorStoreFloat1( V0, GScratch + 1 );
 	V1 = VectorLoad( GScratch );
 	V0 = MakeVectorRegister( 1.0f, 5.0f, -0.25f, -0.5f );
-	LogTest( TEXT("VectorStoreFloat1"), TestVectorsEqual( V0, V1 ) );
+	LogTest<float>( TEXT("VectorStoreFloat1"), TestVectorsEqual( V0, V1 ) );
 
 	// Replicate
 	TestVectorReplicate<float, VectorRegister4Float>();
@@ -1559,363 +1849,375 @@ bool FVectorRegisterAbstractionTest::RunTest(const FString& Parameters)
 	float Float3 = VectorGetComponent(V1, 3);
 	V0 = MakeVectorRegisterFloat(0.f, 1.f, 2.f, 3.f);
 	V1 = MakeVectorRegisterFloat(Float0, Float1, Float2, Float3);
-	LogTest(TEXT("VectorGetComponent<float>"), TestVectorsEqual(V0, V1));
+	LogTest<float>(TEXT("VectorGetComponent<float>"), TestVectorsEqual(V0, V1));
 
 	V0 = MakeVectorRegister( 1.0f, -2.0f, 3.0f, -4.0f );
 	V1 = VectorAbs( V0 );
 	V0 = MakeVectorRegister( 1.0f, 2.0f, 3.0f, 4.0f );
-	LogTest( TEXT("VectorAbs"), TestVectorsEqual( V0, V1 ) );
+	LogTest<float>( TEXT("VectorAbs"), TestVectorsEqual( V0, V1 ) );
 
 	V0 = MakeVectorRegister( 1.0f, -2.0f, 3.0f, -4.0f );
 	V1 = VectorNegate( V0 );
 	V0 = MakeVectorRegister( -1.0f, 2.0f, -3.0f, 4.0f );
-	LogTest( TEXT("VectorNegate"), TestVectorsEqual( V0, V1 ) );
+	LogTest<float>( TEXT("VectorNegate"), TestVectorsEqual( V0, V1 ) );
 
 	V0 = MakeVectorRegister( 1.0f, 2.0f, 3.0f, 4.0f );
 	V1 = MakeVectorRegister( 2.0f, 4.0f, 6.0f, 8.0f );
 	V1 = VectorAdd( V0, V1 );
 	V0 = MakeVectorRegister( 3.0f, 6.0f, 9.0f, 12.0f );
-	LogTest( TEXT("VectorAdd"), TestVectorsEqual( V0, V1 ) );
+	LogTest<float>( TEXT("VectorAdd"), TestVectorsEqual( V0, V1 ) );
 
 	V0 = MakeVectorRegister( 2.0f, 4.0f, 6.0f, 8.0f );
 	V1 = MakeVectorRegister( 1.0f, 2.0f, 3.0f, 4.0f );
 	V1 = VectorSubtract( V0, V1 );
 	V0 = MakeVectorRegister( 1.0f, 2.0f, 3.0f, 4.0f );
-	LogTest( TEXT("VectorSubtract"), TestVectorsEqual( V0, V1 ) );
+	LogTest<float>( TEXT("VectorSubtract"), TestVectorsEqual( V0, V1 ) );
 
 	// Note: Older versions of MSVC had codegen issues with optimization on that caused this to fail. If this passes in Debug builds but not optimized builds, try updating your MSVC compiler.
 	V0 = MakeVectorRegister(2.0f, 4.0f, 6.0f, 8.0f);
 	V1 = MakeVectorRegister(1.0f, 2.0f, 3.0f, 4.0f);
 	V2 = VectorMultiply(V0, V1);
 	V3 = MakeVectorRegister(2.0f, 8.0f, 18.0f, 32.0f);
-	LogTest(TEXT("VectorMultiply"), TestVectorsEqual(V3, V2));
+	LogTest<float>(TEXT("VectorMultiply"), TestVectorsEqual(V3, V2));
 
 	// Note: Older versions of MSVC had codegen issues with optimization on that caused this to fail. If this passes in Debug builds but not optimized builds, try updating your MSVC compiler.
 	V0 = MakeVectorRegister(2.0f, 4.0f, 6.0f, 8.0f);
 	V1 = MakeVectorRegister(1.0f, 2.0f, 3.0f, 4.0f);
 	V2 = VectorMultiplyAdd(V0, V1, VectorOne());
 	V3 = MakeVectorRegister(3.0f, 9.0f, 19.0f, 33.0f);
-	LogTest(TEXT("VectorMultiplyAdd"), TestVectorsEqual(V3, V2));
+	LogTest<float>(TEXT("VectorMultiplyAdd"), TestVectorsEqual(V3, V2));
 
-	V0 = MakeVectorRegister( 2.0f, 4.0f, 6.0f, 8.0f );
-	V1 = MakeVectorRegister( 1.0f, 2.0f, 3.0f, 4.0f );
-	V1 = VectorDot3( V0, V1 );
-	V0 = MakeVectorRegister( 28.0f, 28.0f, 28.0f, 28.0f );
-	LogTest( TEXT("VectorDot3"), TestVectorsEqual( V0, V1 ) );
+	V0 = MakeVectorRegister(2.0f, 4.0f, 6.0f, 8.0f);
+	V1 = MakeVectorRegister(1.0f, 2.0f, 3.0f, 4.0f);
+	V1 = VectorDot3(V0, V1);
+	V0 = MakeVectorRegister(28.0f, 28.0f, 28.0f, 28.0f);
+	LogTest<float>(TEXT("VectorDot3"), TestVectorsEqual(V0, V1));
 
-	V0 = MakeVectorRegister( 2.0f, 4.0f, 6.0f, 8.0f );
-	V1 = MakeVectorRegister( 1.0f, 2.0f, 3.0f, 4.0f );
-	V1 = VectorDot4( V0, V1 );
-	V0 = MakeVectorRegister( 60.0f, 60.0f, 60.0f, 60.0f );
-	LogTest( TEXT("VectorDot4"), TestVectorsEqual( V0, V1 ) );
+	V0 = MakeVectorRegister(1.0f, 2.0f, 3.0f, 5.0f);
+	V1 = MakeVectorRegister(7.0f, 11.0f, 13.0f, 17.0f);
+	V1 = VectorDot3(V0, V1);
+	V0 = MakeVectorRegister(68.0f, 68.0f, 68.0f, 68.0f);
+	LogTest<float>(TEXT("VectorDot3"), TestVectorsEqual(V0, V1));
+
+	V0 = MakeVectorRegister(2.0f, 4.0f, 6.0f, 8.0f);
+	V1 = MakeVectorRegister(1.0f, 2.0f, 3.0f, 4.0f);
+	V1 = VectorDot4(V0, V1);
+	V0 = MakeVectorRegister(60.0f, 60.0f, 60.0f, 60.0f);
+	LogTest<float>(TEXT("VectorDot4"), TestVectorsEqual(V0, V1));
+
+	V0 = MakeVectorRegister(1.0f, 2.0f, 3.0f, 5.0f);
+	V1 = MakeVectorRegister(7.0f, 11.0f, 13.0f, 17.0f);
+	V1 = VectorDot4(V0, V1);
+	V0 = MakeVectorRegister(153.0f, 153.0f, 153.0f, 153.0f);
+	LogTest<float>(TEXT("VectorDot4"), TestVectorsEqual(V0, V1));
 
 	V0 = MakeVectorRegister( 1.0f, 0.0f, 0.0f, 8.0f );
 	V1 = MakeVectorRegister( 0.0f, 2.0f, 0.0f, 4.0f );
 	V1 = VectorCross( V0, V1 );
 	V0 = MakeVectorRegister( 0.f, 0.0f, 2.0f, 0.0f );
-	LogTest( TEXT("VectorCross"), TestVectorsEqual( V0, V1 ) );
+	LogTest<float>( TEXT("VectorCross"), TestVectorsEqual( V0, V1 ) );
 
 	V0 = MakeVectorRegister( 2.0f, 3.0f, 4.0f, 8.0f );
 	V1 = MakeVectorRegister( 5.0f, 6.0f, 7.0f, 4.0f );
 	V1 = VectorCross(V0, V1);
 	V0 = MakeVectorRegister( -3.0f, 6.0f, -3.0f, 0.0f );
-	LogTest(TEXT("VectorCross"), TestVectorsEqual( V0, V1 ));
+	LogTest<float>(TEXT("VectorCross"), TestVectorsEqual( V0, V1 ));
 
 	V0 = MakeVectorRegister( 2.0f, 4.0f, 6.0f, 8.0f );
 	V1 = MakeVectorRegister( 4.0f, 3.0f, 2.0f, 1.0f );
 	V1 = VectorPow( V0, V1 );
 	V0 = MakeVectorRegister( 16.0f, 64.0f, 36.0f, 8.0f );
-	LogTest( TEXT("VectorPow"), TestVectorsEqual( V0, V1, 0.001f ) );
+	LogTest<float>( TEXT("VectorPow"), TestVectorsEqual( V0, V1, 0.001f ) );
 
 	V0 = MakeVectorRegister( 2.0f, -2.0f, 2.0f, -2.0f );
 	V1 = VectorReciprocalLen( V0 );
 	V0 = MakeVectorRegister( 0.25f, 0.25f, 0.25f, 0.25f );
-	LogTest( TEXT("VectorReciprocalLen"), TestVectorsEqual( V0, V1, 0.004f ) );
+	LogTest<float>( TEXT("VectorReciprocalLen"), TestVectorsEqual( V0, V1, 0.004f ) );
 
 	V0 = MakeVectorRegister( 2.0f, -2.0f, 2.0f, -2.0f );
 	V1 = VectorNormalize( V0 );
 	V0 = MakeVectorRegister( 0.5f, -0.5f, 0.5f, -0.5f );
-	LogTest( TEXT("VectorNormalize"), TestVectorsEqual( V0, V1, 0.004f ) );
+	LogTest<float>( TEXT("VectorNormalize"), TestVectorsEqual( V0, V1, 0.004f ) );
 
 	V0 = MakeVectorRegister(2.0f, -2.0f, 2.0f, -2.0f);
 	V1 = VectorNormalizeAccurate(V0);
 	V0 = MakeVectorRegister(0.5f, -0.5f, 0.5f, -0.5f);
-	LogTest(TEXT("VectorNormalizeAccurate"), TestVectorsEqual(V0, V1, 1e-8f));
+	LogTest<float>(TEXT("VectorNormalizeAccurate"), TestVectorsEqual(V0, V1, 1e-8f));
 
 	V0 = MakeVectorRegister(2.0f, -2.0f, 2.0f, -2.0f);
 	V1 = TestVectorNormalize_Sqrt(V0);
 	V0 = MakeVectorRegister(0.5f, -0.5f, 0.5f, -0.5f);
-	LogTest(TEXT("TestVectorNormalize_Sqrt"), TestVectorsEqual(V0, V1, 1e-8f));
+	LogTest<float>(TEXT("TestVectorNormalize_Sqrt"), TestVectorsEqual(V0, V1, 1e-8f));
 
 	V0 = MakeVectorRegister(2.0f, -2.0f, 2.0f, -2.0f);
 	V1 = TestVectorNormalize_InvSqrt(V0);
 	V0 = MakeVectorRegister(0.5f, -0.5f, 0.5f, -0.5f);
-	LogTest(TEXT("TestVectorNormalize_InvSqrt"), TestVectorsEqual(V0, V1, 1e-8f));
+	LogTest<float>(TEXT("TestVectorNormalize_InvSqrt"), TestVectorsEqual(V0, V1, 1e-8f));
 
 	V0 = MakeVectorRegister(2.0f, -2.0f, 2.0f, -2.0f);
 	V1 = TestVectorNormalize_InvSqrtEst(V0);
 	V0 = MakeVectorRegister(0.5f, -0.5f, 0.5f, -0.5f);
-	LogTest(TEXT("TestVectorNormalize_InvSqrtEst"), TestVectorsEqual(V0, V1, 1e-6f));
+	LogTest<float>(TEXT("TestVectorNormalize_InvSqrtEst"), TestVectorsEqual(V0, V1, 1e-6f));
 
 	V0 = MakeVectorRegister( 2.0f, -2.0f, 2.0f, -2.0f );
 	V1 = VectorSet_W0( V0 );
 	V0 = MakeVectorRegister( 2.0f, -2.0f, 2.0f, 0.0f );
-	LogTest( TEXT("VectorSet_W0"), TestVectorsEqual( V0, V1 ) );
+	LogTest<float>( TEXT("VectorSet_W0"), TestVectorsEqual( V0, V1 ) );
 
 	V0 = MakeVectorRegister( 2.0f, -2.0f, 2.0f, -2.0f );
 	V1 = VectorSet_W1( V0 );
 	V0 = MakeVectorRegister( 2.0f, -2.0f, 2.0f, 1.0f );
-	LogTest( TEXT("VectorSet_W1"), TestVectorsEqual( V0, V1 ) );
+	LogTest<float>( TEXT("VectorSet_W1"), TestVectorsEqual( V0, V1 ) );
 
 	V0 = MakeVectorRegister( 2.0f, 4.0f, 6.0f, 8.0f );
 	V1 = MakeVectorRegister( 4.0f, 3.0f, 2.0f, 1.0f );
 	V1 = VectorMin( V0, V1 );
 	V0 = MakeVectorRegister( 2.0f, 3.0f, 2.0f, 1.0f );
-	LogTest( TEXT("VectorMin"), TestVectorsEqual( V0, V1 ) );
+	LogTest<float>( TEXT("VectorMin"), TestVectorsEqual( V0, V1 ) );
 
 	V0 = MakeVectorRegister( 2.0f, 4.0f, 6.0f, 8.0f );
 	V1 = MakeVectorRegister( 4.0f, 3.0f, 2.0f, 1.0f );
 	V1 = VectorMax( V0, V1 );
 	V0 = MakeVectorRegister( 4.0f, 4.0f, 6.0f, 8.0f );
-	LogTest( TEXT("VectorMax"), TestVectorsEqual( V0, V1 ) );
+	LogTest<float>( TEXT("VectorMax"), TestVectorsEqual( V0, V1 ) );
 
 	V0 = MakeVectorRegister(1.0f, 3.0f, 5.0f, 7.0f);
 	V1 = MakeVectorRegister(2.0f, 4.0f, 6.0f, 8.0f);
 	V1 = VectorCombineHigh(V0, V1);
 	V0 = MakeVectorRegister(5.0f, 7.0f, 6.0f, 8.0f);
-	LogTest(TEXT("VectorCombineHigh"), TestVectorsEqual(V0, V1));
+	LogTest<float>(TEXT("VectorCombineHigh"), TestVectorsEqual(V0, V1));
 
 	V0 = MakeVectorRegister(1.0f, 3.0f, 5.0f, 7.0f);
 	V1 = MakeVectorRegister(2.0f, 4.0f, 6.0f, 8.0f);
 	V1 = VectorCombineLow(V0, V1);
 	V0 = MakeVectorRegister(1.0f, 3.0f, 2.0f, 4.0f);
-	LogTest(TEXT("VectorCombineLow"), TestVectorsEqual(V0, V1));
+	LogTest<float>(TEXT("VectorCombineLow"), TestVectorsEqual(V0, V1));
 
 	// LoadByte4
 
 	uint8 Bytes[4] = { 25, 75, 125, 200 };
 	V0 = VectorLoadByte4( Bytes );
 	V1 = MakeVectorRegister( 25.f, 75.f, 125.f, 200.f );
-	LogTest( TEXT("VectorLoadByte4"), TestVectorsEqual( V0, V1 ) );
+	LogTest<float>( TEXT("VectorLoadByte4"), TestVectorsEqual( V0, V1 ) );
 
 	V0 = VectorLoadByte4Reverse( Bytes );
 	V1 = MakeVectorRegister( 25.f, 75.f, 125.f, 200.f );
 	V1 = VectorSwizzle( V1, 3, 2, 1, 0 );
-	LogTest( TEXT("VectorLoadByte4Reverse"), TestVectorsEqual( V0, V1 ) );
+	LogTest<float>( TEXT("VectorLoadByte4Reverse"), TestVectorsEqual( V0, V1 ) );
 
 	V0 = MakeVectorRegister( 4.0f, 3.0f, 2.0f, 1.0f );
 	VectorStoreByte4( V0, Bytes );
 	V1 = VectorLoadByte4( Bytes );
-	LogTest( TEXT("VectorStoreByte4"), TestVectorsEqual( V0, V1 ) );
+	LogTest<float>( TEXT("VectorStoreByte4"), TestVectorsEqual( V0, V1 ) );
 
 	V0 = MakeVectorRegister( 2.0f, 4.0f, 6.0f, 8.0f );
 	V1 = MakeVectorRegister( 4.0f, 3.0f, 2.0f, 1.0f );
 	bool bIsVAGT_TRUE = VectorAnyGreaterThan( V0, V1 ) != 0;
-	LogTest( TEXT("VectorAnyGreaterThan-true"), bIsVAGT_TRUE );
+	LogTest<float>( TEXT("VectorAnyGreaterThan-true"), bIsVAGT_TRUE );
 
 	V0 = MakeVectorRegister( 1.0f, 3.0f, 2.0f, 1.0f );
 	V1 = MakeVectorRegister( 2.0f, 4.0f, 6.0f, 8.0f );
 	bool bIsVAGT_FALSE = VectorAnyGreaterThan( V0, V1 ) == 0;
-	LogTest( TEXT("VectorAnyGreaterThan-false"), bIsVAGT_FALSE );
+	LogTest<float>( TEXT("VectorAnyGreaterThan-false"), bIsVAGT_FALSE );
 
 	V0 = MakeVectorRegister( 1.0f, 3.0f, 2.0f, 1.0f );
 	V1 = MakeVectorRegister( 2.0f, 4.0f, 6.0f, 8.0f );
-	LogTest( TEXT("VectorAnyLesserThan-true"), VectorAnyLesserThan( V0, V1 ) != 0 );
+	LogTest<float>( TEXT("VectorAnyLesserThan-true"), VectorAnyLesserThan( V0, V1 ) != 0 );
 
 	V0 = MakeVectorRegister( 3.0f, 5.0f, 7.0f, 9.0f );
 	V1 = MakeVectorRegister( 2.0f, 4.0f, 6.0f, 8.0f );
-	LogTest( TEXT("VectorAnyLesserThan-false"), VectorAnyLesserThan( V0, V1 ) == 0 );
+	LogTest<float>( TEXT("VectorAnyLesserThan-false"), VectorAnyLesserThan( V0, V1 ) == 0 );
 
 	V0 = MakeVectorRegister( 3.0f, 5.0f, 7.0f, 9.0f );
 	V1 = MakeVectorRegister( 2.0f, 4.0f, 6.0f, 8.0f );
-	LogTest( TEXT("VectorAllGreaterThan-true"), VectorAllGreaterThan( V0, V1 ) != 0 );
+	LogTest<float>( TEXT("VectorAllGreaterThan-true"), VectorAllGreaterThan( V0, V1 ) != 0 );
 
 	V0 = MakeVectorRegister( 3.0f, 1.0f, 7.0f, 9.0f );
 	V1 = MakeVectorRegister( 2.0f, 4.0f, 6.0f, 8.0f );
-	LogTest( TEXT("VectorAllGreaterThan-false"), VectorAllGreaterThan( V0, V1 ) == 0 );
+	LogTest<float>( TEXT("VectorAllGreaterThan-false"), VectorAllGreaterThan( V0, V1 ) == 0 );
 
 	V0 = MakeVectorRegister( 1.0f, 3.0f, 2.0f, 1.0f );
 	V1 = MakeVectorRegister( 2.0f, 4.0f, 6.0f, 8.0f );
-	LogTest( TEXT("VectorAllLesserThan-true"), VectorAllLesserThan( V0, V1 ) != 0 );
+	LogTest<float>( TEXT("VectorAllLesserThan-true"), VectorAllLesserThan( V0, V1 ) != 0 );
 
 	V0 = MakeVectorRegister( 3.0f, 3.0f, 2.0f, 1.0f );
 	V1 = MakeVectorRegister( 2.0f, 4.0f, 6.0f, 8.0f );
-	LogTest( TEXT("VectorAllLesserThan-false"), VectorAllLesserThan( V0, V1 ) == 0 );
+	LogTest<float>( TEXT("VectorAllLesserThan-false"), VectorAllLesserThan( V0, V1 ) == 0 );
 
 	V0 = MakeVectorRegister( 1.0f, 3.0f, 2.0f, 8.0f );
 	V1 = MakeVectorRegister( 2.0f, 4.0f, 2.0f, 1.0f );
 	V2 = VectorCompareGT( V0, V1 );
 	V3 = MakeVectorRegister( (uint32)0, (uint32)0, (uint32)0, (uint32)-1 );
-	LogTest( TEXT("VectorCompareGT"), TestVectorsEqualBitwise( V2, V3 ) );
+	LogTest<float>( TEXT("VectorCompareGT"), TestVectorsEqualBitwise( V2, V3 ) );
 
 	V0 = MakeVectorRegister( 1.0f, 3.0f, 2.0f, 8.0f );
 	V1 = MakeVectorRegister( 2.0f, 4.0f, 2.0f, 1.0f );
 	V2 = VectorCompareGE( V0, V1 );
 	V3 = MakeVectorRegister( (uint32)0, (uint32)0, (uint32)-1, (uint32)-1 );
-	LogTest( TEXT("VectorCompareGE"), TestVectorsEqualBitwise( V2, V3 ) );
+	LogTest<float>( TEXT("VectorCompareGE"), TestVectorsEqualBitwise( V2, V3 ) );
 
 	V0 = MakeVectorRegister( 1.0f, 3.0f, 2.0f, 8.0f );
 	V1 = MakeVectorRegister( 2.0f, 4.0f, 2.0f, 1.0f );
 	V2 = VectorCompareEQ( V0, V1 );
 	V3 = MakeVectorRegister( (uint32)0, (uint32)0, (uint32)-1, (uint32)0 );
-	LogTest( TEXT("VectorCompareEQ"), TestVectorsEqualBitwise( V2, V3 ) );
+	LogTest<float>( TEXT("VectorCompareEQ"), TestVectorsEqualBitwise( V2, V3 ) );
 
 	V0 = MakeVectorRegister( 1.0f, 3.0f, 2.0f, 8.0f );
 	V1 = MakeVectorRegister( 2.0f, 4.0f, 2.0f, 1.0f );
 	V2 = VectorCompareNE( V0, V1 );
 	V3 = MakeVectorRegister( (uint32)(0xFFFFFFFFU), (uint32)(0xFFFFFFFFU), (uint32)(0), (uint32)(0xFFFFFFFFU) );
-	LogTest( TEXT("VectorCompareNE"), TestVectorsEqualBitwise( V2, V3 ) );
+	LogTest<float>( TEXT("VectorCompareNE"), TestVectorsEqualBitwise( V2, V3 ) );
 	
 	V0 = MakeVectorRegister( 1.0f, 3.0f, 2.0f, 8.0f );
 	V1 = MakeVectorRegister( 2.0f, 4.0f, 2.0f, 1.0f );
 	V2 = MakeVectorRegister( (uint32)-1, (uint32)0, (uint32)0, (uint32)-1 );
 	V2 = VectorSelect( V2, V0, V1 );
 	V3 = MakeVectorRegister( 1.0f, 4.0f, 2.0f, 8.0f );
-	LogTest( TEXT("VectorSelect"), TestVectorsEqual( V2, V3 ) );
+	LogTest<float>( TEXT("VectorSelect"), TestVectorsEqual( V2, V3 ) );
 
 	V0 = MakeVectorRegister(1.0f, 3.0f, 5.0f, 7.0f);
 	V1 = MakeVectorRegister(2.0f, 4.0f, 6.0f, 8.0f);
 	V2 = MakeVectorRegister((uint32)0, (uint32)-1, (uint32)-1, (uint32)0);
 	V2 = VectorSelect(V2, V0, V1);
 	V3 = MakeVectorRegister(2.0f, 3.0f, 5.0f, 8.0f);
-	LogTest(TEXT("VectorSelect"), TestVectorsEqual(V2, V3));
+	LogTest<float>(TEXT("VectorSelect"), TestVectorsEqual(V2, V3));
 
 	V0 = MakeVectorRegister( 1.0f, 3.0f, 0.0f, 0.0f );
 	V1 = MakeVectorRegister( 0.0f, 0.0f, 2.0f, 1.0f );
 	V2 = VectorBitwiseOr( V0, V1 );
 	V3 = MakeVectorRegister( 1.0f, 3.0f, 2.0f, 1.0f );
-	LogTest( TEXT("VectorBitwiseOr-Float1"), TestVectorsEqual( V2, V3 ) );
+	LogTest<float>( TEXT("VectorBitwiseOr-Float1"), TestVectorsEqual( V2, V3 ) );
 
 	V0 = MakeVectorRegister( 1.0f, 3.0f, 24.0f, 36.0f );
 	V1 = MakeVectorRegister( (uint32)(0x80000000U), (uint32)(0x80000000U), (uint32)(0x80000000U), (uint32)(0x80000000U) );
 	V2 = VectorBitwiseOr( V0, V1 );
 	V3 = MakeVectorRegister( -1.0f, -3.0f, -24.0f, -36.0f );
-	LogTest( TEXT("VectorBitwiseOr-Float2"), TestVectorsEqual( V2, V3 ) );
+	LogTest<float>( TEXT("VectorBitwiseOr-Float2"), TestVectorsEqual( V2, V3 ) );
 
 	V0 = MakeVectorRegister( -1.0f, -3.0f, -24.0f, 36.0f );
 	V1 = MakeVectorRegister( (uint32)(0xFFFFFFFFU), (uint32)(0x7FFFFFFFU), (uint32)(0x7FFFFFFFU), (uint32)(0xFFFFFFFFU) );
 	V2 = VectorBitwiseAnd( V0, V1 );
 	V3 = MakeVectorRegister( -1.0f, 3.0f, 24.0f, 36.0f );
-	LogTest( TEXT("VectorBitwiseAnd-Float"), TestVectorsEqual( V2, V3 ) );
+	LogTest<float>( TEXT("VectorBitwiseAnd-Float"), TestVectorsEqual( V2, V3 ) );
 
 	V0 = MakeVectorRegister( -1.0f, -3.0f, -24.0f, 36.0f );
 	V1 = MakeVectorRegister( (uint32)(0x80000000U), (uint32)(0x00000000U), (uint32)(0x80000000U), (uint32)(0x80000000U) );
 	V2 = VectorBitwiseXor( V0, V1 );
 	V3 = MakeVectorRegister( 1.0f, -3.0f, 24.0f, -36.0f );
-	LogTest( TEXT("VectorBitwiseXor-Float"), TestVectorsEqual( V2, V3 ) );
+	LogTest<float>( TEXT("VectorBitwiseXor-Float"), TestVectorsEqual( V2, V3 ) );
 
 	V0 = MakeVectorRegister( -1.0f, -3.0f, -24.0f, 36.0f );
 	V1 = MakeVectorRegister( 5.0f, 35.0f, 23.0f, 48.0f );
 	V2 = VectorMergeVecXYZ_VecW( V0, V1 );
 	V3 = MakeVectorRegister( -1.0f, -3.0f, -24.0f, 48.0f );
-	LogTest( TEXT("VectorMergeXYZ_VecW-1"), TestVectorsEqual( V2, V3 ) );
+	LogTest<float>( TEXT("VectorMergeXYZ_VecW-1"), TestVectorsEqual( V2, V3 ) );
 
 	V0 = MakeVectorRegister( -1.0f, -3.0f, -24.0f, 36.0f );
 	V1 = MakeVectorRegister( 5.0f, 35.0f, 23.0f, 48.0f );
 	V2 = VectorMergeVecXYZ_VecW( V1, V0 );
 	V3 = MakeVectorRegister( 5.0f, 35.0f, 23.0f, 36.0f );
-	LogTest( TEXT("VectorMergeXYZ_VecW-2"), TestVectorsEqual( V2, V3 ) );
+	LogTest<float>( TEXT("VectorMergeXYZ_VecW-2"), TestVectorsEqual( V2, V3 ) );
 
 	V0 = MakeVectorRegister( 1.0f, 1.0e6f, 1.3e-8f, 35.0f );
 	V1 = VectorReciprocal( V0 );
 	V3 = VectorMultiply(V1, V0);
-	LogTest( TEXT("VectorReciprocal"), TestVectorsEqual( VectorOne(), V3, 0.008f ) );
+	LogTest<float>( TEXT("VectorReciprocal"), TestVectorsEqual( VectorOne(), V3, 0.008f ) );
 
 	V0 = MakeVectorRegister( 1.0f, 1.0e6f, 1.3e-8f, 35.0f );
 	V1 = VectorReciprocalAccurate( V0 );
 	V3 = VectorMultiply(V1, V0);
-	LogTest( TEXT("VectorReciprocalAccurate"), TestVectorsEqual( VectorOne(), V3, 1e-7f) );
+	LogTest<float>( TEXT("VectorReciprocalAccurate"), TestVectorsEqual( VectorOne(), V3, 1e-7f) );
 
 	V0 = MakeVectorRegister( 1.0f, 1.0e6f, 1.3e-8f, 35.0f );
 	V1 = VectorReciprocalSqrt( V0 );
 	V3 = VectorMultiply(VectorMultiply(V1, V1), V0);
-	LogTest( TEXT("VectorReciprocalSqrt"), TestVectorsEqual( VectorOne(), V3, 0.007f ) );
+	LogTest<float>( TEXT("VectorReciprocalSqrt"), TestVectorsEqual( VectorOne(), V3, 0.007f ) );
 
 	V0 = MakeVectorRegister( 1.0f, 1.0e6f, 1.3e-8f, 35.0f );
 	V1 = VectorReciprocalSqrtAccurate( V0 );
 	V3 = VectorMultiply(VectorMultiply(V1, V1), V0);
-	LogTest( TEXT("VectorReciprocalSqrtAccurate"), TestVectorsEqual( VectorOne(), V3, 1e-6f ) );
+	LogTest<float>( TEXT("VectorReciprocalSqrtAccurate"), TestVectorsEqual( VectorOne(), V3, 1e-6f ) );
 
 	SetScratch(1.0f, 2.0f, 3.0f, 4.0f);
 	V0 = VectorLoadTwoPairsFloat(GScratch + 0, GScratch + 1);
 	V1 = MakeVectorRegister(1.0f, 2.0f, 2.0f, 3.0f);
-	LogTest(TEXT("VectorLoadTwoPairsFloat"), TestVectorsEqual(V0, V1));
+	LogTest<float>(TEXT("VectorLoadTwoPairsFloat"), TestVectorsEqual(V0, V1));
 
 	V0 = MakeVectorRegister(0.0f, 1.0f, 2.0f, 3.0f);
 	V1 = MakeVectorRegister(4.0f, 5.0f, 6.0f, 7.0f);
 	VectorDeinterleave(V2, V3, V0, V1);
 	V0 = MakeVectorRegister(0.0f, 2.0f, 4.0f, 6.0f);
 	V1 = MakeVectorRegister(1.0f, 3.0f, 5.0f, 7.0f);
-	LogTest(TEXT("VectorDeinterleave"), TestVectorsEqual(V2, V0) && TestVectorsEqual(V3, V1));
+	LogTest<float>(TEXT("VectorDeinterleave"), TestVectorsEqual(V2, V0) && TestVectorsEqual(V3, V1));
 
 	// VectorMod
 	V0 = MakeVectorRegister(0.0f, 3.2f, 2.8f,  1.5f);
 	V1 = MakeVectorRegister(2.0f, 1.2f, 2.0f,  3.0f);
 	V2 = TestReferenceMod(V0, V1);
 	V3 = VectorMod(V0, V1);
-	LogTest( TEXT("VectorMod positive"), TestVectorsEqual(V2, V3));
+	LogTest<float>( TEXT("VectorMod positive"), TestVectorsEqual(V2, V3));
 
 	V0 = MakeVectorRegister(-2.0f,  3.2f, -2.8f,  -1.5f);
 	V1 = MakeVectorRegister(-1.5f, -1.2f,  2.0f,   3.0f);
 	V2 = TestReferenceMod(V0, V1);
 	V3 = VectorMod(V0, V1);
-	LogTest( TEXT("VectorMod negative"), TestVectorsEqual(V2, V3));
+	LogTest<float>( TEXT("VectorMod negative"), TestVectorsEqual(V2, V3));
 
 	// VectorSign
 	V0 = MakeVectorRegister(2.0f, -2.0f, 0.0f, -3.0f);
 	V2 = MakeVectorRegister(1.0f, -1.0f, 1.0f, -1.0f);
 	V3 = VectorSign(V0);
-	LogTest(TEXT("VectorSign"), TestVectorsEqual(V2, V3));
+	LogTest<float>(TEXT("VectorSign"), TestVectorsEqual(V2, V3));
 
 	// VectorStep
 	V0 = MakeVectorRegister(2.0f, -2.0f, 0.0f, -3.0f);
 	V2 = MakeVectorRegister(1.0f, 0.0f, 1.0f, 0.0f);
 	V3 = VectorStep(V0);
-	LogTest(TEXT("VectorStep"), TestVectorsEqual(V2, V3));
+	LogTest<float>(TEXT("VectorStep"), TestVectorsEqual(V2, V3));
 
 	// VectorTruncate
 	V0 = MakeVectorRegister(-1.8f, -1.0f, -0.8f, 0.0f);
 	V2 = MakeVectorRegister(-1.0f, -1.0f, 0.0f, 0.0f);
 	V3 = VectorTruncate(V0);
-	LogTest(TEXT("VectorTruncate"), TestVectorsEqual(V2, V3, KINDA_SMALL_NUMBER));
+	LogTest<float>(TEXT("VectorTruncate"), TestVectorsEqual(V2, V3, KINDA_SMALL_NUMBER));
 
 	V0 = MakeVectorRegister(0.0f, 0.8f, 1.0f, 1.8f);
 	V2 = MakeVectorRegister(0.0f, 0.0f, 1.0f, 1.0f);
 	V3 = VectorTruncate(V0);
-	LogTest(TEXT("VectorTruncate"), TestVectorsEqual(V2, V3, KINDA_SMALL_NUMBER));
+	LogTest<float>(TEXT("VectorTruncate"), TestVectorsEqual(V2, V3, KINDA_SMALL_NUMBER));
 
 	// VectorFractional
 	V0 = MakeVectorRegister(-1.8f, -1.0f, -0.8f, 0.0f);
 	V2 = MakeVectorRegister(-0.8f, 0.0f, -0.8f, 0.0f);
 	V3 = VectorFractional(V0);
-	LogTest(TEXT("VectorFractional"), TestVectorsEqual(V2, V3, KINDA_SMALL_NUMBER));
+	LogTest<float>(TEXT("VectorFractional"), TestVectorsEqual(V2, V3, KINDA_SMALL_NUMBER));
 
 	V0 = MakeVectorRegister(0.0f, 0.8f, 1.0f, 1.8f);
 	V2 = MakeVectorRegister(0.0f, 0.8f, 0.0f, 0.8f);
 	V3 = VectorFractional(V0);
-	LogTest(TEXT("VectorFractional"), TestVectorsEqual(V2, V3, KINDA_SMALL_NUMBER));
+	LogTest<float>(TEXT("VectorFractional"), TestVectorsEqual(V2, V3, KINDA_SMALL_NUMBER));
 
 	// VectorCeil
 	V0 = MakeVectorRegister(-1.8f, -1.0f, -0.8f, 0.0f);
 	V2 = MakeVectorRegister(-1.0f, -1.0f, -0.0f, 0.0f);
 	V3 = VectorCeil(V0);
-	LogTest(TEXT("VectorCeil"), TestVectorsEqual(V2, V3, KINDA_SMALL_NUMBER));
+	LogTest<float>(TEXT("VectorCeil"), TestVectorsEqual(V2, V3, KINDA_SMALL_NUMBER));
 
 	V0 = MakeVectorRegister(0.0f, 0.8f, 1.0f, 1.8f);
 	V2 = MakeVectorRegister(0.0f, 1.0f, 1.0f, 2.0f);
 	V3 = VectorCeil(V0);
-	LogTest(TEXT("VectorCeil"), TestVectorsEqual(V2, V3, KINDA_SMALL_NUMBER));
+	LogTest<float>(TEXT("VectorCeil"), TestVectorsEqual(V2, V3, KINDA_SMALL_NUMBER));
 
 	// VectorFloor
 	V0 = MakeVectorRegister(-1.8f, -1.0f, -0.8f, 0.0f);
 	V2 = MakeVectorRegister(-2.0f, -1.0f, -1.0f, 0.0f);
 	V3 = VectorFloor(V0);
-	LogTest(TEXT("VectorFloor"), TestVectorsEqual(V2, V3, KINDA_SMALL_NUMBER));
+	LogTest<float>(TEXT("VectorFloor"), TestVectorsEqual(V2, V3, KINDA_SMALL_NUMBER));
 
 	V0 = MakeVectorRegister(0.0f, 0.8f, 1.0f, 1.8f);
 	V2 = MakeVectorRegister(0.0f, 0.0f, 1.0f, 1.0f);
 	V3 = VectorFloor(V0);
-	LogTest(TEXT("VectorFloor"), TestVectorsEqual(V2, V3));
+	LogTest<float>(TEXT("VectorFloor"), TestVectorsEqual(V2, V3));
 
 	FMatrix44f	M0, M1, M2, M3;		// LWC_TODO: Need double version tests!
 	FVector3f Eye, LookAt, Up;	
@@ -1942,27 +2244,27 @@ bool FVectorRegisterAbstractionTest::RunTest(const FString& Parameters)
 
 	VectorMatrixMultiply( &M2, &M0, &M1 );
 	TestVectorMatrixMultiply( &M3, &M0, &M1 );
-	LogTest( TEXT("VectorMatrixMultiply"), TestMatricesEqual( M2, M3, 0.000001f ) );
+	LogTest<float>( TEXT("VectorMatrixMultiply"), TestMatricesEqual( M2, M3, 0.000001f ) );
 
 	VectorMatrixInverse( &M2, &M1 );
 	TestVectorMatrixInverse( &M3, &M1 );
-	LogTest( TEXT("VectorMatrixInverse"), TestMatricesEqual( M2, M3, 0.000001f) );
+	LogTest<float>( TEXT("VectorMatrixInverse"), TestMatricesEqual( M2, M3, 0.000001f) );
 
 // 	FTransform Transform;
 // 	Transform.SetFromMatrix(M1);
 // 	FTransform InvTransform = Transform.Inverse();
 // 	FTransform InvTransform2 = FTransform(Transform.ToMatrixWithScale().Inverse());
-// 	LogTest( TEXT("FTransform Inverse"), InvTransform.Equals(InvTransform2, 1e-3f ) );
+// 	LogTest<float>( TEXT("FTransform Inverse"), InvTransform.Equals(InvTransform2, 1e-3f ) );
 
 	V0 = MakeVectorRegister( 100.0f, -100.0f, 200.0f, 1.0f );
 	V1 = VectorTransformVector(V0, &M0);
 	V2 = TestVectorTransformVector(V0, &M0);
-	LogTest( TEXT("VectorTransformVector"), TestVectorsEqual( V1, V2 ) );
+	LogTest<float>( TEXT("VectorTransformVector"), TestVectorsEqual( V1, V2 ) );
 
 	V0 = MakeVectorRegister( 32768.0f,131072.0f, -8096.0f, 1.0f );
 	V1 = VectorTransformVector(V0, &M1);
 	V2 = TestVectorTransformVector(V0, &M1);
-	LogTest( TEXT("VectorTransformVector"), TestVectorsEqual( V1, V2 ) );
+	LogTest<float>( TEXT("VectorTransformVector"), TestVectorsEqual( V1, V2 ) );
 
 	// NaN / Inf tests
 	// Using a union as we need to do a bitwise cast of 0xFFFFFFFF into a float.
@@ -1975,19 +2277,20 @@ bool FVectorRegisterAbstractionTest::RunTest(const FString& Parameters)
 	NaNU.IntNaN = 0xFFFFFFFF;
 	const float NaN = NaNU.FloatNaN;
 
-	LogTest(TEXT("VectorContainsNaNOrInfinite true"), VectorContainsNaNOrInfinite(MakeVectorRegister(NaN, NaN, NaN, NaN)));
-	LogTest(TEXT("VectorContainsNaNOrInfinite true"), VectorContainsNaNOrInfinite(MakeVectorRegister(NaN, 0.f, 0.f, 0.f)));
-	LogTest(TEXT("VectorContainsNaNOrInfinite true"), VectorContainsNaNOrInfinite(MakeVectorRegister(0.f, 0.f, 0.f, NaN)));
-	LogTest(TEXT("VectorContainsNaNOrInfinite true"), VectorContainsNaNOrInfinite(GlobalVectorConstants::FloatInfinity));
-	LogTest(TEXT("VectorContainsNaNOrInfinite true"), VectorContainsNaNOrInfinite(MakeVectorRegister((uint32)0xFF800000, (uint32)0xFF800000, (uint32)0xFF800000, (uint32)0xFF800000))); // negative infinity
-	LogTest(TEXT("VectorContainsNaNOrInfinite true"), VectorContainsNaNOrInfinite(GlobalVectorConstants::AllMask));
+	SetScratch(0.0f, 0.0f, 0.0f, 0.0f);
+	LogTest<float>(TEXT("VectorContainsNaNOrInfinite true"), VectorContainsNaNOrInfinite(MakeVectorRegister(NaN, NaN, NaN, NaN)));
+	LogTest<float>(TEXT("VectorContainsNaNOrInfinite true"), VectorContainsNaNOrInfinite(MakeVectorRegister(NaN, 0.f, 0.f, 0.f)));
+	LogTest<float>(TEXT("VectorContainsNaNOrInfinite true"), VectorContainsNaNOrInfinite(MakeVectorRegister(0.f, 0.f, 0.f, NaN)));
+	LogTest<float>(TEXT("VectorContainsNaNOrInfinite true"), VectorContainsNaNOrInfinite(GlobalVectorConstants::FloatInfinity));
+	LogTest<float>(TEXT("VectorContainsNaNOrInfinite true"), VectorContainsNaNOrInfinite(MakeVectorRegister((uint32)0xFF800000, (uint32)0xFF800000, (uint32)0xFF800000, (uint32)0xFF800000))); // negative infinity
+	LogTest<float>(TEXT("VectorContainsNaNOrInfinite true"), VectorContainsNaNOrInfinite(GlobalVectorConstants::AllMask));
 
 	// Not Nan/Inf
-	LogTest(TEXT("VectorContainsNaNOrInfinite false"), !VectorContainsNaNOrInfinite(GlobalVectorConstants::FloatZero));
-	LogTest(TEXT("VectorContainsNaNOrInfinite false"), !VectorContainsNaNOrInfinite(GlobalVectorConstants::FloatOne));
-	LogTest(TEXT("VectorContainsNaNOrInfinite false"), !VectorContainsNaNOrInfinite(GlobalVectorConstants::FloatMinusOneHalf));
-	LogTest(TEXT("VectorContainsNaNOrInfinite false"), !VectorContainsNaNOrInfinite(GlobalVectorConstants::SmallNumber));
-	LogTest(TEXT("VectorContainsNaNOrInfinite false"), !VectorContainsNaNOrInfinite(GlobalVectorConstants::BigNumber));
+	LogTest<float>(TEXT("VectorContainsNaNOrInfinite false"), !VectorContainsNaNOrInfinite(GlobalVectorConstants::FloatZero));
+	LogTest<float>(TEXT("VectorContainsNaNOrInfinite false"), !VectorContainsNaNOrInfinite(GlobalVectorConstants::FloatOne));
+	LogTest<float>(TEXT("VectorContainsNaNOrInfinite false"), !VectorContainsNaNOrInfinite(GlobalVectorConstants::FloatMinusOneHalf));
+	LogTest<float>(TEXT("VectorContainsNaNOrInfinite false"), !VectorContainsNaNOrInfinite(GlobalVectorConstants::SmallNumber));
+	LogTest<float>(TEXT("VectorContainsNaNOrInfinite false"), !VectorContainsNaNOrInfinite(GlobalVectorConstants::BigNumber));
 
 	// FMath::FMod (floats)
 	{
@@ -2162,40 +2465,11 @@ bool FVectorRegisterAbstractionTest::RunTest(const FString& Parameters)
 
 	FQuat Q0, Q1, Q2, Q3;
 
-	// SinCos tests
-	{
-		const VectorRegister4Float QuadrantDegreesArray[] = {
-			MakeVectorRegister( 0.0f, 10.0f, 20.0f, 30.0f),
-			MakeVectorRegister(45.0f, 60.0f, 70.0f, 80.0f),
-		};
+	// Sin, Cos, Tan tests
+	TestVectorTrigFunctions<float, VectorRegister4Float>();
 
-		const float SinCosTolerance = 1e-5f;
-		const int32 Cycles = 3; // Go through a full circle this many times (negative and positive)
-		for (int32 OffsetQuadrant = -4*Cycles; OffsetQuadrant <= 4*Cycles; ++OffsetQuadrant)
-		{
-			const float OffsetFloat = (float)OffsetQuadrant * 90.0f; // Add 90 degrees repeatedly to cover all quadrants and wrap a few times
-			const VectorRegister4Float VOffset = VectorLoadFloat1(&OffsetFloat);
-			for (VectorRegister4Float const& VDegrees : QuadrantDegreesArray)
-			{
-				const VectorRegister4Float VAnglesDegrees = VectorAdd(VOffset, VDegrees);
-				const VectorRegister4Float VAngles = VectorMultiply(VAnglesDegrees, GlobalVectorConstants::DEG_TO_RAD);
-				VectorRegister4Float S[3], C[3];
-				TestReferenceSinCos(S[0], C[0], VAngles);
-				TestFastSinCos(S[1], C[1], VAngles);
-				TestVectorSinCos(S[2], C[2], VAngles);
-				LogTest(TEXT("SinCos (Sin): Ref vs Fast"), TestVectorsEqual_ComponentWiseError(S[0], S[1], SinCosTolerance));
-				LogTest(TEXT("SinCos (Cos): Ref vs Fast"), TestVectorsEqual_ComponentWiseError(C[0], C[1], SinCosTolerance));
-				LogTest(TEXT("SinCos (Sin): Ref vs Vec"), TestVectorsEqual_ComponentWiseError(S[0], S[2], SinCosTolerance));
-				LogTest(TEXT("SinCos (Cos): Ref vs Vec"), TestVectorsEqual_ComponentWiseError(C[0], C[2], SinCosTolerance));
-
-				S[2] = VectorSin(VAngles);
-				LogTest(TEXT("VectorSin: Ref vs Vec"), TestVectorsEqual_ComponentWiseError(S[0], S[2], 0.001091f));
-
-				C[2] = VectorCos(VAngles);
-				LogTest(TEXT("VectorCos: Ref vs Vec"), TestVectorsEqual_ComponentWiseError(C[0], C[2], 0.001091f));
-			}
-		}
-	}
+	// Exp, Log tests
+	TestVectorExpLogFunctions<float, VectorRegister4Float>();
 
 	// Quat<->Rotator conversions and equality
 	{
@@ -2296,7 +2570,7 @@ bool FVectorRegisterAbstractionTest::RunTest(const FString& Parameters)
 
 		// Identity -> X Axis
 		Axis = FQuat::Identity.GetRotationAxis();
-		LogTest(TEXT("FQuat::Identity.GetRotationAxis() == FVector3f::XAxisVector"), TestFVector3Equal(Axis, FVector3f::XAxisVector));
+		LogTest<float>(TEXT("FQuat::Identity.GetRotationAxis() == FVector3f::XAxisVector"), TestFVector3Equal(Axis, FVector3f::XAxisVector));
 
 		const FQuat QuatArray[] = {
 			FQuat(0.0f, 0.0f, 0.0f, 1.0f),
@@ -2309,8 +2583,8 @@ bool FVectorRegisterAbstractionTest::RunTest(const FString& Parameters)
 		for (const FQuat& Q : QuatArray)
 		{
 			Q.ToAxisAndAngle(Axis, Angle);
-			LogTest(TEXT("Quat -> Axis and Angle: Q is Normalized"), TestQuatNormalized(Q, 1e-6f));
-			LogTest(TEXT("Quat -> Axis and Angle: Axis is Normalized"), TestFVector3Normalized(Axis, 1e-6f));
+			LogTest<float>(TEXT("Quat -> Axis and Angle: Q is Normalized"), TestQuatNormalized(Q, 1e-6f));
+			LogTest<float>(TEXT("Quat -> Axis and Angle: Axis is Normalized"), TestFVector3Normalized(Axis, 1e-6f));
 		}
 	}
 
@@ -2320,33 +2594,33 @@ bool FVectorRegisterAbstractionTest::RunTest(const FString& Parameters)
 		Rotator0 = FRotator(30.0f, -45.0f, 90.0f);
 		Q0 = Rotator0.Quaternion();
 		Q1 = TestRotatorToQuaternion(Rotator0);
-		LogTest( TEXT("TestRotatorToQuaternion"), TestQuatsEqual(Q0, Q1, 1e-6f));
+		LogTest<float>( TEXT("TestRotatorToQuaternion"), TestQuatsEqual(Q0, Q1, 1e-6f));
 
 		FVector3f FV0, FV1;
 		FV0 = Rotator0.Vector();
 		FV1 = FRotationMatrix( Rotator0 ).GetScaledAxis( EAxis::X );
-		LogTest( TEXT("Test0 Rotator::Vector()"), TestFVector3Equal(FV0, FV1, 1e-6f));
+		LogTest<float>( TEXT("Test0 Rotator::Vector()"), TestFVector3Equal(FV0, FV1, 1e-6f));
 		
 		FV0 = FRotationMatrix( Rotator0 ).GetScaledAxis( EAxis::X );
 		FV1 = FQuatRotationMatrix( Q0 ).GetScaledAxis( EAxis::X );
-		LogTest( TEXT("Test0 FQuatRotationMatrix"), TestFVector3Equal(FV0, FV1, 1e-5f));
+		LogTest<float>( TEXT("Test0 FQuatRotationMatrix"), TestFVector3Equal(FV0, FV1, 1e-5f));
 
 		Rotator0 = FRotator(45.0f,  60.0f, 120.0f);
 		Q0 = Rotator0.Quaternion();
 		Q1 = TestRotatorToQuaternion(Rotator0);
-		LogTest( TEXT("TestRotatorToQuaternion"), TestQuatsEqual(Q0, Q1, 1e-6f));
+		LogTest<float>( TEXT("TestRotatorToQuaternion"), TestQuatsEqual(Q0, Q1, 1e-6f));
 
 		FV0 = Rotator0.Vector();
 		FV1 = FRotationMatrix( Rotator0 ).GetScaledAxis( EAxis::X );
-		LogTest( TEXT("Test1 Rotator::Vector()"), TestFVector3Equal(FV0, FV1, 1e-6f));
+		LogTest<float>( TEXT("Test1 Rotator::Vector()"), TestFVector3Equal(FV0, FV1, 1e-6f));
 
 		FV0 = FRotationMatrix( Rotator0 ).GetScaledAxis( EAxis::X );
 		FV1 = FQuatRotationMatrix( Q0 ).GetScaledAxis( EAxis::X );
-		LogTest(TEXT("Test1 FQuatRotationMatrix"), TestFVector3Equal(FV0, FV1, 1e-5f));
+		LogTest<float>(TEXT("Test1 FQuatRotationMatrix"), TestFVector3Equal(FV0, FV1, 1e-5f));
 
 		FV0 = FRotationMatrix(FRotator::ZeroRotator).GetScaledAxis(EAxis::X);
 		FV1 = FQuatRotationMatrix(FQuat::Identity).GetScaledAxis(EAxis::X);
-		LogTest(TEXT("Test2 FQuatRotationMatrix"), TestFVector3Equal(FV0, FV1, 1e-6f));
+		LogTest<float>(TEXT("Test2 FQuatRotationMatrix"), TestFVector3Equal(FV0, FV1, 1e-6f));
 	}
 
 	// Quat Rotation tests
@@ -2380,7 +2654,7 @@ bool FVectorRegisterAbstractionTest::RunTest(const FString& Parameters)
 		#define TEST_QUAT_ROTATE(_QIndex, _VIndex, _Quat, _Vec, _Func, _Tolerance) \
 		{ \
 			const FString _TestName = FString::Printf(TEXT("Test Quat%d: Vec%d: %s"), _QIndex, _VIndex, TEXT(#_Func)); \
-			LogTest( *_TestName, TestFVector3Equal(_Quat.RotateVector(_Vec), _Func(_Quat, _Vec), _Tolerance) ); \
+			LogTest<float>( *_TestName, TestFVector3Equal(_Quat.RotateVector(_Vec), _Func(_Quat, _Vec), _Tolerance) ); \
 		}
 
 		// Test loop
@@ -2416,14 +2690,14 @@ bool FVectorRegisterAbstractionTest::RunTest(const FString& Parameters)
 
 					if (A.IsZero() || B.IsZero())
 					{
-						LogTest(TEXT("FindBetween: Old == New (normal)"), TestQuatsEqual(Old, NewNormal, 1e-6f));
-						LogTest(TEXT("FindBetween: Old == New (vector)"), TestQuatsEqual(Old, NewVector, 1e-6f));
+						LogTest<float>(TEXT("FindBetween: Old == New (normal)"), TestQuatsEqual(Old, NewNormal, 1e-6f));
+						LogTest<float>(TEXT("FindBetween: Old == New (vector)"), TestQuatsEqual(Old, NewVector, 1e-6f));
 					}
 					else
 					{
-						LogTest(TEXT("FindBetween: Old A->B"), TestFVector3Equal(RotAOld, BNorm, KINDA_SMALL_NUMBER));
-						LogTest(TEXT("FindBetween: New A->B (normal)"), TestFVector3Equal(RotANewNormal, BNorm, KINDA_SMALL_NUMBER));
-						LogTest(TEXT("FindBetween: New A->B (vector)"), TestFVector3Equal(RotANewVector, BNorm, KINDA_SMALL_NUMBER));
+						LogTest<float>(TEXT("FindBetween: Old A->B"), TestFVector3Equal(RotAOld, BNorm, KINDA_SMALL_NUMBER));
+						LogTest<float>(TEXT("FindBetween: New A->B (normal)"), TestFVector3Equal(RotANewNormal, BNorm, KINDA_SMALL_NUMBER));
+						LogTest<float>(TEXT("FindBetween: New A->B (vector)"), TestFVector3Equal(RotANewVector, BNorm, KINDA_SMALL_NUMBER));
 					}
 				}
 			}
@@ -2444,8 +2718,8 @@ bool FVectorRegisterAbstractionTest::RunTest(const FString& Parameters)
 				const FVector3f Rotated1 = Q1.RotateVector(FVector3f::ForwardVector);
 				const FVector3f Rotated2 = R0.RotateVector(FVector3f::ForwardVector);
 
-				LogTest(TEXT("V.ToOrientationQuat() rotate"), TestFVector3Equal(Rotated0, Rotated1, KINDA_SMALL_NUMBER));
-				LogTest(TEXT("V.ToOrientationRotator() rotate"), TestFVector3Equal(Rotated0, Rotated2, KINDA_SMALL_NUMBER));
+				LogTest<float>(TEXT("V.ToOrientationQuat() rotate"), TestFVector3Equal(Rotated0, Rotated1, KINDA_SMALL_NUMBER));
+				LogTest<float>(TEXT("V.ToOrientationRotator() rotate"), TestFVector3Equal(Rotated0, Rotated2, KINDA_SMALL_NUMBER));
 			}
 		}
 	}
@@ -2456,33 +2730,33 @@ bool FVectorRegisterAbstractionTest::RunTest(const FString& Parameters)
 		Q1 = FQuat(FRotator(45.0f, 60.0f, 120.0f));
 		VectorQuaternionMultiply(&Q2, &Q0, &Q1);
 		TestVectorQuaternionMultiply(&Q3, &Q0, &Q1);
-		LogTest(TEXT("VectorQuaternionMultiply"), TestQuatsEqual(Q2, Q3, 1e-6f));
+		LogTest<float>(TEXT("VectorQuaternionMultiply"), TestQuatsEqual(Q2, Q3, 1e-6f));
 		V0 = VectorLoadAligned(&Q0);
 		V1 = VectorLoadAligned(&Q1);
 		V2 = VectorQuaternionMultiply2(V0, V1);
 		V3 = VectorLoadAligned(&Q3);
-		LogTest(TEXT("VectorQuaternionMultiply2"), TestVectorsEqual(V2, V3, 1e-6f));
+		LogTest<float>(TEXT("VectorQuaternionMultiply2"), TestVectorsEqual(V2, V3, 1e-6f));
 
 		Q0 = FQuat(FRotator(0.0f, 180.0f, 45.0f));
 		Q1 = FQuat(FRotator(-120.0f, -90.0f, 0.0f));
 		VectorQuaternionMultiply(&Q2, &Q0, &Q1);
 		TestVectorQuaternionMultiply(&Q3, &Q0, &Q1);
-		LogTest(TEXT("VectorMatrixInverse"), TestQuatsEqual(Q2, Q3, 1e-6f));
+		LogTest<float>(TEXT("VectorMatrixInverse"), TestQuatsEqual(Q2, Q3, 1e-6f));
 		V0 = VectorLoadAligned(&Q0);
 		V1 = VectorLoadAligned(&Q1);
 		V2 = VectorQuaternionMultiply2(V0, V1);
 		V3 = VectorLoadAligned(&Q3);
-		LogTest(TEXT("VectorQuaternionMultiply2"), TestVectorsEqual(V2, V3, 1e-6f));
+		LogTest<float>(TEXT("VectorQuaternionMultiply2"), TestVectorsEqual(V2, V3, 1e-6f));
 	}
 
 	if (!GPassing)
 	{
-		UE_LOG(LogUnrealMathTest, Fatal,TEXT("VectorIntrinsics Failed."));
+		UE_LOG(LogUnrealMathTest, Fatal, TEXT("VectorIntrinsics Failed."));
 	}
 
 	if (!RunDoubleVectorTest())
 	{
-		UE_LOG(LogUnrealMathTest, Fatal, TEXT("VectorIntrinsics<double> Failed."));
+		UE_LOG(LogUnrealMathTest, Fatal, TEXT("VectorIntrinsics <double> Failed."));
 	}
 
 	return true;
diff --git a/Engine/Source/Runtime/Core/Public/Math/UnrealMathFPU.h b/Engine/Source/Runtime/Core/Public/Math/UnrealMathFPU.h
--- a/Engine/Source/Runtime/Core/Public/Math/UnrealMathFPU.h
+++ b/Engine/Source/Runtime/Core/Public/Math/UnrealMathFPU.h
@@ -2083,22 +2083,22 @@ FORCEINLINE VectorRegister4Double VectorATan(const VectorRegister4Double& Vec)
 }
 
 
-FORCEINLINE VectorRegister4Float VectorATan2(const VectorRegister4Float& X, const VectorRegister4Float& Y)
+FORCEINLINE VectorRegister4Float VectorATan2(const VectorRegister4Float& Y, const VectorRegister4Float& X)
 {
 	return MakeVectorRegisterFloat(
-		FMath::Atan2(X.V[0], Y.V[0]),
-		FMath::Atan2(X.V[1], Y.V[1]),
-		FMath::Atan2(X.V[2], Y.V[2]),
-		FMath::Atan2(X.V[3], Y.V[3]));
+		FMath::Atan2(Y.V[0], X.V[0]),
+		FMath::Atan2(Y.V[1], X.V[1]),
+		FMath::Atan2(Y.V[2], X.V[2]),
+		FMath::Atan2(Y.V[3], X.V[3]));
 }
 
-FORCEINLINE VectorRegister4Double VectorATan2(const VectorRegister4Double& X, const VectorRegister4Double& Y)
+FORCEINLINE VectorRegister4Double VectorATan2(const VectorRegister4Double& Y, const VectorRegister4Double& X)
 {
 	return MakeVectorRegisterDouble(
-		FMath::Atan2(X.V[0], Y.V[0]),
-		FMath::Atan2(X.V[1], Y.V[1]),
-		FMath::Atan2(X.V[2], Y.V[2]),
-		FMath::Atan2(X.V[3], Y.V[3]));
+		FMath::Atan2(Y.V[0], X.V[0]),
+		FMath::Atan2(Y.V[1], X.V[1]),
+		FMath::Atan2(Y.V[2], X.V[2]),
+		FMath::Atan2(Y.V[3], X.V[3]));
 }
 
 
@@ -2125,7 +2125,7 @@ FORCEINLINE void VectorSinCos(VectorRegister4Float* RESTRICT VSinAngles, VectorR
 
 FORCEINLINE void VectorSinCos(VectorRegister4Double* RESTRICT VSinAngles, VectorRegister4Double* RESTRICT VCosAngles, const VectorRegister4Double* RESTRICT VAngles)
 {
-	// TODO: No Fmath::SinCos function exists yet, but need to revisit if one is added
+	// TODO: No FMath::SinCos<double> function exists yet, but need to revisit if one is added
 	*VSinAngles = VectorSin(*VAngles);
 	*VCosAngles = VectorCos(*VAngles);
 }
diff --git a/Engine/Source/Runtime/Core/Public/Math/UnrealMathSSE.h b/Engine/Source/Runtime/Core/Public/Math/UnrealMathSSE.h
--- a/Engine/Source/Runtime/Core/Public/Math/UnrealMathSSE.h
+++ b/Engine/Source/Runtime/Core/Public/Math/UnrealMathSSE.h
@@ -14,7 +14,7 @@
 #endif
 
 #ifndef UE_PLATFORM_MATH_USE_AVX_2
-#define UE_PLATFORM_MATH_USE_AVX_2			0 // PLATFORM_ALWAYS_HAS_AVX_2
+#define UE_PLATFORM_MATH_USE_AVX_2			0 //PLATFORM_ALWAYS_HAS_AVX_2
 #endif
 
 #ifndef UE_PLATFORM_MATH_USE_FMA3
@@ -23,10 +23,9 @@
 
 #ifndef UE_PLATFORM_MATH_USE_SVML
 	#if defined(_MSC_VER)
-	//////////// TEMPORARILY disabled until UnrealMathTest is expanded to test more trig functions, though it passes what exists there now. ////////////
-	#define UE_PLATFORM_MATH_USE_SVML			(0 && (_MSC_VER >= 1920)) // Support added to MSVC 2019 16.0+
+		#define UE_PLATFORM_MATH_USE_SVML			(_MSC_VER >= 1920) // Support added to MSVC 2019 16.0+
 	#else
-	#define UE_PLATFORM_MATH_USE_SVML			0
+		#define UE_PLATFORM_MATH_USE_SVML			0
 	#endif // defined(_MSC_VER)
 #endif
 
@@ -778,6 +777,8 @@ template<> FORCEINLINE VectorRegister2Double SelectVectorSwizzle2<2, 1>(const Ve
 template<int Index0, int Index1, int Index2, int Index3>
 FORCEINLINE VectorRegister4Double VectorSwizzleTemplate(const VectorRegister4Double& Vec)
 {
+	static_assert(Index0 >= 0 && Index0 <= 3 && Index1 >= 0 && Index1 <= 3 && Index2 >= 0 && Index2 <= 3 && Index3 >= 0 && Index3 <= 3);
+
 #if UE_PLATFORM_MATH_USE_AVX_2
 	VectorRegister4Double Result = _mm256_permute4x64_pd(Vec, SHUFFLEMASK(Index0, Index1, Index2, Index3));
 #else
@@ -845,6 +846,8 @@ FORCEINLINE VectorRegister4Double VectorReplicateImpl(const VectorRegister4Doubl
 template<int Index>
 FORCEINLINE VectorRegister4Double VectorReplicateTemplate(const VectorRegister4Double& Vec)
 {
+	static_assert(Index >= 0 && Index <= 3);
+
 #if UE_PLATFORM_MATH_USE_AVX_2
 	VectorRegister4Double Result = _mm256_permute4x64_pd(Vec, SHUFFLEMASK(Index, Index, Index, Index));
 #else
@@ -857,6 +860,7 @@ FORCEINLINE VectorRegister4Double VectorReplicateTemplate(const VectorRegister4D
 template<int Index0, int Index1, int Index2, int Index3>
 FORCEINLINE VectorRegister4Double VectorShuffleTemplate(const VectorRegister4Double& Vec1, const VectorRegister4Double& Vec2)
 {
+	static_assert(Index0 >= 0 && Index0 <= 3 && Index1 >= 0 && Index1 <= 3 && Index2 >= 0 && Index2 <= 3 && Index3 >= 0 && Index3 <= 3);
 	VectorRegister4Double Result;
 	Result.XY = SelectVectorSwizzle2<Index0, Index1>(Vec1);
 	Result.ZW = SelectVectorSwizzle2<Index2, Index3>(Vec2);
@@ -908,6 +912,7 @@ template<> FORCEINLINE VectorRegister4Float VectorSwizzleTemplate<0, 0, 0, 0>(co
 template<int Index>
 FORCEINLINE VectorRegister4Float VectorReplicateTemplate(const VectorRegister4Float& Vec)
 {
+	static_assert(Index >= 0 && Index <= 3);
 	return VectorSwizzleTemplate<Index, Index, Index, Index>(Vec);
 }
 
@@ -915,6 +920,7 @@ FORCEINLINE VectorRegister4Float VectorReplicateTemplate(const VectorRegister4Fl
 template<int Index0, int Index1, int Index2, int Index3>
 FORCEINLINE VectorRegister4Float VectorShuffleTemplate(const VectorRegister4Float& Vec1, const VectorRegister4Float& Vec2)
 {
+	static_assert(Index0 >= 0 && Index0 <= 3 && Index1 >= 0 && Index1 <= 3 && Index2 >= 0 && Index2 <= 3 && Index3 >= 0 && Index3 <= 3);
 	return _mm_shuffle_ps(Vec1, Vec2, SHUFFLEMASK(Index0, Index1, Index2, Index3));
 }
 
@@ -1088,68 +1094,64 @@ FORCEINLINE VectorRegister4Double VectorMultiply(const VectorRegister4Double& Ve
 
 
 /**
- * Multiplies two vectors (component-wise), adds in the third vector and returns the result.
+ * Multiplies two vectors (component-wise), adds in the third vector and returns the result. (A*B + C)
  *
- * @param Vec1	1st vector
- * @param Vec2	2nd vector
- * @param Vec3	3rd vector
- * @return		VectorRegister4Float( Vec1.x*Vec2.x + Vec3.x, Vec1.y*Vec2.y + Vec3.y, Vec1.z*Vec2.z + Vec3.z, Vec1.w*Vec2.w + Vec3.w )
+ * @param A		1st vector
+ * @param B		2nd vector
+ * @param C		3rd vector
+ * @return		VectorRegister4Float( A.x*B.x + C.x, A.y*B.y + C.y, A.z*B.z + C.z, A.w*B.w + C.w )
  */
-FORCEINLINE VectorRegister4Float VectorMultiplyAdd(const VectorRegister4Float& Vec1, const VectorRegister4Float& Vec2, const VectorRegister4Float& Vec3)
+FORCEINLINE VectorRegister4Float VectorMultiplyAdd(const VectorRegister4Float& A, const VectorRegister4Float& B, const VectorRegister4Float& C)
 {
 #if UE_PLATFORM_MATH_USE_FMA3
-	return _mm_fmadd_ps(Vec1, Vec2, Vec3);
+	return _mm_fmadd_ps(A, B, C);
 #else
-	return VectorAdd(VectorMultiply(Vec1, Vec2), Vec3);
+	return VectorAdd(VectorMultiply(A, B), C);
 #endif
 }
 
-FORCEINLINE VectorRegister4Double VectorMultiplyAdd(const VectorRegister4Double& Vec1, const VectorRegister4Double& Vec2, const VectorRegister4Double& Vec3)
+FORCEINLINE VectorRegister4Double VectorMultiplyAdd(const VectorRegister4Double& A, const VectorRegister4Double& B, const VectorRegister4Double& C)
 {
-#if UE_PLATFORM_MATH_USE_FMA3
+#if UE_PLATFORM_MATH_USE_FMA3 && UE_PLATFORM_MATH_USE_AVX
+	return _mm256_fmadd_pd(A, B, C);
+#elif UE_PLATFORM_MATH_USE_FMA3
 	VectorRegister4Double Result;
-	#if !UE_PLATFORM_MATH_USE_AVX
-		Result.XY = _mm_fmadd_pd(Vec1.XY, Vec2.XY, Vec3.XY);
-		Result.ZW = _mm_fmadd_pd(Vec1.ZW, Vec2.ZW, Vec3.ZW);
-	#else
-		Result = _mm256_fmadd_pd(Vec1, Vec2, Vec3);
-	#endif
+	Result.XY = _mm_fmadd_pd(A.XY, B.XY, C.XY);
+	Result.ZW = _mm_fmadd_pd(A.ZW, B.ZW, C.ZW);
 	return Result;
 #else
-	return VectorAdd(VectorMultiply(Vec1, Vec2), Vec3);
+	return VectorAdd(VectorMultiply(A, B), C);
 #endif
 }
 
 /**
- * Multiplies two vectors (component-wise), negates the results and adds it to the third vector i.e. -AB + C = C - AB
+ * Multiplies two vectors (component-wise), negates the results and adds it to the third vector i.e. (-A*B + C) = (C - A*B)
  *
- * @param Vec1	1st vector
- * @param Vec2	2nd vector
- * @param Vec3	3rd vector
- * @return		VectorRegister( Vec3.x - Vec1.x*Vec2.x, Vec3.y - Vec1.y*Vec2.y, Vec3.z - Vec1.z*Vec2.z, Vec3.w - Vec1.w*Vec2.w )
+ * @param A		1st vector
+ * @param B		2nd vector
+ * @param C		3rd vector
+ * @return		VectorRegister( C.x - A.x*B.x, C.y - A.y*B.y, C.z - A.z*B.z, C.w - A.w*B.w )
  */
-FORCEINLINE VectorRegister4Float VectorNegateMultiplyAdd(const VectorRegister4Float& Vec1, const VectorRegister4Float& Vec2, const VectorRegister4Float& Vec3)
+FORCEINLINE VectorRegister4Float VectorNegateMultiplyAdd(const VectorRegister4Float& A, const VectorRegister4Float& B, const VectorRegister4Float& C)
 {
 #if UE_PLATFORM_MATH_USE_FMA3
-	return _mm_fnmadd_ps(Vec1, Vec2, Vec3);
+	return _mm_fnmadd_ps(A, B, C);
 #else
-	return VectorSubtract(Vec3, VectorMultiply(Vec1, Vec2));
+	return VectorSubtract(C, VectorMultiply(A, B));
 #endif
 }
 
-FORCEINLINE VectorRegister4Double VectorNegateMultiplyAdd(const VectorRegister4Double& Vec1, const VectorRegister4Double& Vec2, const VectorRegister4Double& Vec3)
+FORCEINLINE VectorRegister4Double VectorNegateMultiplyAdd(const VectorRegister4Double& A, const VectorRegister4Double& B, const VectorRegister4Double& C)
 {
-#if UE_PLATFORM_MATH_USE_FMA3
+#if UE_PLATFORM_MATH_USE_FMA3 && UE_PLATFORM_MATH_USE_AVX
+	return _mm256_fnmadd_pd(A, B, C);
+#elif UE_PLATFORM_MATH_USE_FMA3
 	VectorRegister4Double Result;
-	#if !UE_PLATFORM_MATH_USE_AVX
-		Result.XY = _mm_fnmadd_pd(Vec1.XY, Vec2.XY, Vec3.XY);
-		Result.ZW = _mm_fnmadd_pd(Vec1.ZW, Vec2.ZW, Vec3.ZW);
-	#else
-		Result = _mm256_fnmadd_pd(Vec1, Vec2, Vec3);
-	#endif
+	Result.XY = _mm_fnmadd_pd(A.XY, B.XY, C.XY);
+	Result.ZW = _mm_fnmadd_pd(A.ZW, B.ZW, C.ZW);
 	return Result;
 #else
-	return VectorSubtract(Vec3, VectorMultiply(Vec1, Vec2));
+	return VectorSubtract(C, VectorMultiply(A, B));
 #endif
 }
 
@@ -3095,40 +3097,38 @@ FORCEINLINE VectorRegister4Double VectorATan(const VectorRegister4Double& X)
 #endif
 }
 
-FORCEINLINE VectorRegister4Float VectorATan2(const VectorRegister4Float& X, const VectorRegister4Float& Y)
+FORCEINLINE VectorRegister4Float VectorATan2(const VectorRegister4Float& Y, const VectorRegister4Float& X)
 {
 #if UE_PLATFORM_MATH_USE_SVML
-	return _mm_atan2_ps(X, Y);
+	return _mm_atan2_ps(Y, X);
 #else
-	AlignedFloat4 FloatsX(X);
 	AlignedFloat4 FloatsY(Y);
-	FloatsX[0] = FMath::Atan2(FloatsX[0], FloatsY[0]);
-	FloatsX[1] = FMath::Atan2(FloatsX[1], FloatsY[1]);
-	FloatsX[2] = FMath::Atan2(FloatsX[2], FloatsY[2]);
-	FloatsX[3] = FMath::Atan2(FloatsX[3], FloatsY[3]);
-	return FloatsX.ToVectorRegister();
+	AlignedFloat4 FloatsX(X);
+	FloatsY[0] = FMath::Atan2(FloatsY[0], FloatsX[0]);
+	FloatsY[1] = FMath::Atan2(FloatsY[1], FloatsX[1]);
+	FloatsY[2] = FMath::Atan2(FloatsY[2], FloatsX[2]);
+	FloatsY[3] = FMath::Atan2(FloatsY[3], FloatsX[3]);
+	return FloatsY.ToVectorRegister();
 #endif
 }
 
-FORCEINLINE VectorRegister4Double VectorATan2(const VectorRegister4Double& X, const VectorRegister4Double& Y)
+FORCEINLINE VectorRegister4Double VectorATan2(const VectorRegister4Double& Y, const VectorRegister4Double& X)
 {
 #if UE_PLATFORM_MATH_USE_SVML_AVX
-	return _mm256_atan2_pd(X, Y);
+	return _mm256_atan2_pd(Y, X);
 #elif UE_PLATFORM_MATH_USE_SVML
-	return VectorRegister4Double(_mm_atan2_pd(X.XY, Y.XY), _mm_atan2_pd(X.ZW, Y.ZW));
+	return VectorRegister4Double(_mm_atan2_pd(Y.XY, X.XY), _mm_atan2_pd(Y.ZW, X.ZW));
 #else
-	AlignedDouble4 DoublesX(X);
 	AlignedDouble4 DoublesY(Y);
-	DoublesX[0] = FMath::Atan2(DoublesX[0], DoublesY[0]);
-	DoublesX[1] = FMath::Atan2(DoublesX[1], DoublesY[1]);
-	DoublesX[2] = FMath::Atan2(DoublesX[2], DoublesY[2]);
-	DoublesX[3] = FMath::Atan2(DoublesX[3], DoublesY[3]);
-	return DoublesX.ToVectorRegister();
+	AlignedDouble4 DoublesX(X);
+	DoublesY[0] = FMath::Atan2(DoublesY[0], DoublesX[0]);
+	DoublesY[1] = FMath::Atan2(DoublesY[1], DoublesX[1]);
+	DoublesY[2] = FMath::Atan2(DoublesY[2], DoublesX[2]);
+	DoublesY[3] = FMath::Atan2(DoublesY[3], DoublesX[3]);
+	return DoublesY.ToVectorRegister();
 #endif
 }
 
-// To be continued...
-
 
 //////////////////////////////////////////////////////////////////////////
 //Integer ops
