Removed references of UE4 in codebase

diff --git a/Engine/Source/Runtime/Core/Private/Misc/CoreMisc.cpp b/Engine/Source/Runtime/Core/Private/Misc/CoreMisc.cpp
--- a/Engine/Source/Runtime/Core/Private/Misc/CoreMisc.cpp
+++ b/Engine/Source/Runtime/Core/Private/Misc/CoreMisc.cpp
@@ -483,7 +483,7 @@ static FAutoConsoleVariableRef CVarGEnsureOnNANDiagnostic(
 #endif
 
 #if DO_CHECK
-namespace UE4Asserts_Private
+namespace UEAsserts_Private
 {
 	void VARARGS InternalLogNANDiagnosticMessage(const TCHAR* FormattedMsg, ...)
 	{		
diff --git a/Engine/Source/Runtime/Core/Public/Logging/LogMacros.h b/Engine/Source/Runtime/Core/Public/Logging/LogMacros.h
--- a/Engine/Source/Runtime/Core/Public/Logging/LogMacros.h
+++ b/Engine/Source/Runtime/Core/Public/Logging/LogMacros.h
@@ -140,7 +140,7 @@ private:
 
 #else
 
-	namespace UE4Asserts_Private
+	namespace UEAsserts_Private
 	{
 		template <int32 VerbosityToCheck, typename CategoryType>
 		FORCEINLINE
@@ -170,7 +170,7 @@ private:
 	 * @param CategoryName name of the logging category
 	 * @param Verbosity, verbosity level to test against
 	**/
-	#define UE_LOG_ACTIVE(CategoryName, Verbosity) (::UE4Asserts_Private::IsLogActive<(int32)ELogVerbosity::Verbosity>(CategoryName))
+	#define UE_LOG_ACTIVE(CategoryName, Verbosity) (::UEAsserts_Private::IsLogActive<(int32)ELogVerbosity::Verbosity>(CategoryName))
 
 	#define UE_GET_LOG_VERBOSITY(CategoryName) \
 		CategoryName.GetVerbosity()
@@ -416,7 +416,7 @@ extern CORE_API int32 GEnsureOnNANDiagnostic;
 
 // Macro to either log an error or ensure on a NaN error.
 #if DO_CHECK && !USING_CODE_ANALYSIS
-namespace UE4Asserts_Private
+namespace UEAsserts_Private
 {
 	CORE_API void VARARGS InternalLogNANDiagnosticMessage(const TCHAR* FormattedMsg, ...); // UE_LOG(LogCore, Error, _FormatString_, ##__VA_ARGS__);
 }
@@ -426,7 +426,7 @@ namespace UE4Asserts_Private
 		static bool OnceOnly = false;\
 		if (!OnceOnly)\
 		{\
-			UE4Asserts_Private::InternalLogNANDiagnosticMessage(_FormatString_, ##__VA_ARGS__); \
+			UEAsserts_Private::InternalLogNANDiagnosticMessage(_FormatString_, ##__VA_ARGS__); \
 			OnceOnly = true;\
 		}\
 	}\
diff --git a/Engine/Source/Runtime/Core/Public/Misc/AssertionMacros.h b/Engine/Source/Runtime/Core/Public/Misc/AssertionMacros.h
--- a/Engine/Source/Runtime/Core/Public/Misc/AssertionMacros.h
+++ b/Engine/Source/Runtime/Core/Public/Misc/AssertionMacros.h
@@ -365,7 +365,7 @@ RetType FORCENOINLINE UE_DEBUG_SECTION DispatchCheckVerify(InnerType&& Inner, Ar
 
 #endif	// DO_CHECK
 
-namespace UE4Asserts_Private
+namespace UEAsserts_Private
 {
 	// A junk function to allow us to use sizeof on a member variable which is potentially a bitfield
 	template <typename T>
@@ -378,14 +378,14 @@ namespace UE4Asserts_Private
 
 // Returns FName(TEXT("EnumeratorName")), while statically verifying that the enumerator exists in the enum
 #define GET_ENUMERATOR_NAME_CHECKED(EnumName, EnumeratorName) \
-	((void)sizeof(UE4Asserts_Private::GetMemberNameCheckedJunk(EnumName::EnumeratorName)), FName(TEXT(#EnumeratorName)))
+	((void)sizeof(UEAsserts_Private::GetMemberNameCheckedJunk(EnumName::EnumeratorName)), FName(TEXT(#EnumeratorName)))
 
 // Returns FName(TEXT("MemberName")), while statically verifying that the member exists in ClassName
 #define GET_MEMBER_NAME_CHECKED(ClassName, MemberName) \
-	((void)sizeof(UE4Asserts_Private::GetMemberNameCheckedJunk(((ClassName*)0)->MemberName)), FName(TEXT(#MemberName)))
+	((void)sizeof(UEAsserts_Private::GetMemberNameCheckedJunk(((ClassName*)0)->MemberName)), FName(TEXT(#MemberName)))
 
 #define GET_MEMBER_NAME_STRING_CHECKED(ClassName, MemberName) \
-	((void)sizeof(UE4Asserts_Private::GetMemberNameCheckedJunk(((ClassName*)0)->MemberName)), TEXT(#MemberName))
+	((void)sizeof(UEAsserts_Private::GetMemberNameCheckedJunk(((ClassName*)0)->MemberName)), TEXT(#MemberName))
 
 // Returns FName(TEXT("FunctionName")), while statically verifying that the function exists in ClassName
 #define GET_FUNCTION_NAME_CHECKED(ClassName, FunctionName) \
diff --git a/Engine/Source/Runtime/Core/Public/Misc/GeneratedTypeName.h b/Engine/Source/Runtime/Core/Public/Misc/GeneratedTypeName.h
--- a/Engine/Source/Runtime/Core/Public/Misc/GeneratedTypeName.h
+++ b/Engine/Source/Runtime/Core/Public/Misc/GeneratedTypeName.h
@@ -20,7 +20,7 @@
 
 #if !UE_BUILD_DOCS
 
-namespace UE4TypeName_Private
+namespace UETypeName_Private
 {
 	// Class representing a substring of another string
 	struct FSubstr
@@ -110,7 +110,7 @@ namespace UE4TypeName_Private
 template <typename T>
 inline const TCHAR* GetGeneratedTypeName()
 {
-	static constexpr auto Result = UE4TypeName_Private::TypeSubstrToCharArray<T>(TMakeIntegerSequence<uint32, UE4TypeName_Private::GetTypeSubstr<T>().Len()>());
+	static constexpr auto Result = UETypeName_Private::TypeSubstrToCharArray<T>(TMakeIntegerSequence<uint32, UETypeName_Private::GetTypeSubstr<T>().Len()>());
 	return Result.Array;
 }
 
diff --git a/Engine/Source/Runtime/Core/Public/Misc/StaticAssertCompleteType.h b/Engine/Source/Runtime/Core/Public/Misc/StaticAssertCompleteType.h
--- a/Engine/Source/Runtime/Core/Public/Misc/StaticAssertCompleteType.h
+++ b/Engine/Source/Runtime/Core/Public/Misc/StaticAssertCompleteType.h
@@ -2,7 +2,7 @@
 
 #pragma once
 
-namespace UE4StaticAssertCompleteType_Private
+namespace UEStaticAssertCompleteType_Private
 {
 	class FIncompleteType;
 
@@ -24,7 +24,7 @@ namespace UE4StaticAssertCompleteType_Private
 }
 
 // Causes a compile error if a type is incomplete
-#define UE_STATIC_ASSERT_COMPLETE_TYPE(TypeToCheck, Message) static_assert(sizeof(UE4StaticAssertCompleteType_Private::TUEStaticAssertTypeChecker<TypeToCheck>::Func()), Message)
+#define UE_STATIC_ASSERT_COMPLETE_TYPE(TypeToCheck, Message) static_assert(sizeof(UEStaticAssertCompleteType_Private::TUEStaticAssertTypeChecker<TypeToCheck>::Func()), Message)
 
 // Tests
 
@@ -34,15 +34,15 @@ namespace UE4StaticAssertCompleteType_Private
 	UE_STATIC_ASSERT_COMPLETE_TYPE(      volatile void,                                     "CV void is incomplete");
 	UE_STATIC_ASSERT_COMPLETE_TYPE(const          void,                                     "CV void is incomplete");
 	UE_STATIC_ASSERT_COMPLETE_TYPE(const volatile void,                                     "CV void is incomplete");
-	UE_STATIC_ASSERT_COMPLETE_TYPE(UE4StaticAssertCompleteType_Private::FIncompleteType,    "A forward-declared-but-undefined class is incomplete");
-	UE_STATIC_ASSERT_COMPLETE_TYPE(UE4StaticAssertCompleteType_Private::FIncompleteType[2], "An array of an incomplete class is incomplete");
+	UE_STATIC_ASSERT_COMPLETE_TYPE(UEStaticAssertCompleteType_Private::FIncompleteType,    "A forward-declared-but-undefined class is incomplete");
+	UE_STATIC_ASSERT_COMPLETE_TYPE(UEStaticAssertCompleteType_Private::FIncompleteType[2], "An array of an incomplete class is incomplete");
 	UE_STATIC_ASSERT_COMPLETE_TYPE(int[],                                                   "An array of a complete type of unspecified bound is incomplete");
 #endif
 
 // Each of these should pass
 #if 0
-	UE_STATIC_ASSERT_COMPLETE_TYPE(UE4StaticAssertCompleteType_Private::FIncompleteType*,                                                         "A pointer to an incomplete type is complete");
-	UE_STATIC_ASSERT_COMPLETE_TYPE(UE4StaticAssertCompleteType_Private::FIncompleteType&,                                                         "A reference to an incomplete type is complete");
-	UE_STATIC_ASSERT_COMPLETE_TYPE(UE4StaticAssertCompleteType_Private::FIncompleteType   (UE4StaticAssertCompleteType_Private::FIncompleteType), "A function type is not incomplete, even if it returns or takes an incomplete type");
-	UE_STATIC_ASSERT_COMPLETE_TYPE(UE4StaticAssertCompleteType_Private::FIncompleteType(&)(UE4StaticAssertCompleteType_Private::FIncompleteType), "References to function types must give a good error");
+	UE_STATIC_ASSERT_COMPLETE_TYPE(UEStaticAssertCompleteType_Private::FIncompleteType*,                                                         "A pointer to an incomplete type is complete");
+	UE_STATIC_ASSERT_COMPLETE_TYPE(UEStaticAssertCompleteType_Private::FIncompleteType&,                                                         "A reference to an incomplete type is complete");
+	UE_STATIC_ASSERT_COMPLETE_TYPE(UEStaticAssertCompleteType_Private::FIncompleteType   (UEStaticAssertCompleteType_Private::FIncompleteType), "A function type is not incomplete, even if it returns or takes an incomplete type");
+	UE_STATIC_ASSERT_COMPLETE_TYPE(UEStaticAssertCompleteType_Private::FIncompleteType(&)(UEStaticAssertCompleteType_Private::FIncompleteType), "References to function types must give a good error");
 #endif
