#ue5 - LWC: Fix typo in VectorNormalizeRotator using wrong constant in the double version. Thankfully we weren't using this yet, as FRotator is still just float.

diff --git a/Engine/Source/Runtime/Core/Public/Math/UnrealMathVectorCommon.h b/Engine/Source/Runtime/Core/Public/Math/UnrealMathVectorCommon.h
--- a/Engine/Source/Runtime/Core/Public/Math/UnrealMathVectorCommon.h
+++ b/Engine/Source/Runtime/Core/Public/Math/UnrealMathVectorCommon.h
@@ -576,7 +576,7 @@ FORCEINLINE VectorRegister4Double VectorNormalizeRotator(const VectorRegister4Do
 
 	// shift to [-180,180]
 	VectorRegister4Double V3 = VectorSubtract(V2, GlobalVectorConstants::Double360);
-	VectorRegister4Double V4 = VectorSelect(VectorCompareGT(V2, GlobalVectorConstants::Double360), V3, V2);
+	VectorRegister4Double V4 = VectorSelect(VectorCompareGT(V2, GlobalVectorConstants::Double180), V3, V2);
 
 	return  V4;
 }
