Pull request #7677 to give actual OS version name data instead of previously hardcoded OS name

diff --git a/Engine/Source/Runtime/ApplicationCore/Private/Unix/UnixPlatformSurvey.cpp b/Engine/Source/Runtime/ApplicationCore/Private/Unix/UnixPlatformSurvey.cpp
--- a/Engine/Source/Runtime/ApplicationCore/Private/Unix/UnixPlatformSurvey.cpp
+++ b/Engine/Source/Runtime/ApplicationCore/Private/Unix/UnixPlatformSurvey.cpp
@@ -10,6 +10,14 @@ bool FUnixPlatformSurvey::GetSurveyResults(FHardwareSurveyResults& OutResults, b
 {
 	FMemory::Memset(&OutResults, 0, sizeof(FHardwareSurveyResults));
 	WriteFStringToResults(OutResults.Platform, TEXT("Unix"));
+	
+	// OS
+	FString OSVersionLabel;
+	FString OSSubVersionLabel;
+	FPlatformMisc::GetOSVersions(OSVersionLabel, OSSubVersionLabel);
+	WriteFStringToResults(OutResults.OSVersion, OSVersionLabel);
+	WriteFStringToResults(OutResults.OSSubVersion, OSSubVersionLabel);
+	OutResults.OSBits = FPlatformMisc::Is64bitOperatingSystem() ? 64 : 32;
 
 	// CPU
 	OutResults.CPUCount = FPlatformMisc::NumberOfCores();  // TODO [RCL] 2015-07-15: parse /proc/cpuinfo for GHz/brand
@@ -32,19 +40,6 @@ bool FUnixPlatformSurvey::GetSurveyResults(FHardwareSurveyResults& OutResults, b
 	return true;
 }
 
-void FUnixPlatformSurvey::GetOSName(FHardwareSurveyResults& OutResults)
-{
-	// TODO [RCL] 2015-07-15: check if /etc/os-release or /etc/redhat-release exist and parse it
-
-	/*
-	TArray<FString> OsReleaseLines;
-	if( FFileHelper::LoadANSITextFileToStrings(TEXT("/etc/os-release"), &IFileManager::Get(), OsReleaseLines))
-	{
-		//...
-	}
-	*/
-}
-
 void FUnixPlatformSurvey::WriteFStringToResults(TCHAR* OutBuffer, const FString& InString)
 {
 	FMemory::Memset( OutBuffer, 0, sizeof(TCHAR) * FHardwareSurveyResults::MaxStringLength );
diff --git a/Engine/Source/Runtime/ApplicationCore/Public/Unix/UnixPlatformSurvey.h b/Engine/Source/Runtime/ApplicationCore/Public/Unix/UnixPlatformSurvey.h
--- a/Engine/Source/Runtime/ApplicationCore/Public/Unix/UnixPlatformSurvey.h
+++ b/Engine/Source/Runtime/ApplicationCore/Public/Unix/UnixPlatformSurvey.h
@@ -19,11 +19,6 @@ struct FUnixPlatformSurvey : public FGenericPlatformSurvey
 
 private:
 
-	/**
-	 * Fills out the actual OS name ("distro" is an OS for all the practical purposes), as detailed as possible.
-	 */
-	static void GetOSName(FHardwareSurveyResults& Results);
-
 	/**
 	 * Safely write strings into the fixed length TCHAR buffers of a FHardwareSurveyResults member
 	 */
diff --git a/Engine/Source/Runtime/Core/Private/Unix/UnixPlatformMisc.cpp b/Engine/Source/Runtime/Core/Private/Unix/UnixPlatformMisc.cpp
--- a/Engine/Source/Runtime/Core/Private/Unix/UnixPlatformMisc.cpp
+++ b/Engine/Source/Runtime/Core/Private/Unix/UnixPlatformMisc.cpp
@@ -20,6 +20,7 @@
 #include "Misc/Parse.h"
 #include "Misc/CommandLine.h"
 #include "Misc/App.h"
+#include "Misc/FileHelper.h"
 #include "Unix/UnixPlatformCrashContext.h"
 #include "Misc/ConfigCacheIni.h"
 #include "GenericPlatform/GenericPlatformChunkInstall.h"
@@ -28,6 +29,7 @@
 	#include <cpuid.h>
 #endif // PLATFORM_HAS_CPUID
 #include <sys/sysinfo.h>
+#include <sys/utsname.h>
 #include <sched.h>
 #include <sys/vfs.h>
 #include <sys/ioctl.h>
@@ -420,11 +422,138 @@ bool FUnixPlatformMisc::HasOverriddenReturnCode(uint8 * OverriddenReturnCodeToUs
 	return GHasOverriddenReturnCode;
 }
 
+namespace
+{
+	/**
+	 * Gets a Linux kernel version string
+	 */
+	FString GetKernelVersion()
+	{
+		struct utsname SysInfo;
+
+		if (uname(&SysInfo) == 0)
+		{
+			return FString(SysInfo.release);
+		}
+
+		return FString();
+	}
+
+	/**
+	 * Reads a Linux style configuration file ignoring # as comments
+	 */
+	TMap<FString, FString> ReadConfigurationFile(const TCHAR* Filename)
+	{
+		TMap<FString, FString> Contents;
+		TArray<FString> ConfigLines;
+		if (FFileHelper::LoadANSITextFileToStrings(Filename, &IFileManager::Get(), ConfigLines))
+		{
+			for (const FString& Line : ConfigLines)
+			{
+				FString CleanLine = Line.TrimStartAndEnd();
+				FString KeyString;
+				FString ValueString;
+				// Skip Comments
+				if (!CleanLine.StartsWith("#"))
+				{
+					if (CleanLine.Split(TEXT("="), &KeyString, &ValueString, ESearchCase::CaseSensitive))
+					{
+						KeyString.TrimStartAndEndInline();
+						ValueString.TrimStartAndEndInline();
+						if (ValueString.Left(1) == TEXT("\""))
+						{
+							ValueString.RightChopInline(1, false);
+						}
+						if (ValueString.Right(1) == TEXT("\""))
+						{
+							ValueString.LeftChopInline(1, false);
+						}
+						Contents.Add(KeyString, ValueString);
+					}
+				}
+			}
+		}
+		return Contents;
+	}
+}
+
+void FUnixPlatformMisc::GetOSVersions(FString& out_OSVersionLabel, FString& out_OSSubVersionLabel)
+{
+	// Set up Fallback Details
+	out_OSVersionLabel = FString(TEXT("GenericLinuxVersion"));
+	// Get Kernel Version
+	out_OSSubVersionLabel = GetKernelVersion();
+	// Get PRETTY_NAME/NAME or redhat-release line
+	TMap<FString, FString> OsInfo = ReadConfigurationFile(TEXT("/etc/os-release"));
+	if (OsInfo.Num() > 0)
+	{
+		FString* VersionAddress = OsInfo.Find(TEXT("PRETTY_NAME"));
+		if (VersionAddress)
+		{
+			FString* VersionNameAddress = nullptr;
+			if (VersionAddress->Equals(TEXT("Linux")))
+			{
+				VersionNameAddress = OsInfo.Find(TEXT("NAME"));
+				if (VersionNameAddress != nullptr)
+				{
+					VersionAddress = VersionNameAddress;
+				}
+			}
+
+			out_OSVersionLabel = FString(*VersionAddress);
+		}
+	}
+	else
+	{
+		TArray <FString> RedHatRelease;
+		if (FFileHelper::LoadANSITextFileToStrings(TEXT("/etc/redhat-release"), &IFileManager::Get(), RedHatRelease))
+		{
+			if (RedHatRelease.IsValidIndex(0))
+			{
+				out_OSVersionLabel = RedHatRelease[0];
+			}
+		}
+	}
+}
+
 FString FUnixPlatformMisc::GetOSVersion()
 {
-	// TODO [RCL] 2015-07-15: check if /etc/os-release or /etc/redhat-release exist and parse it
-	// See void FUnixPlatformSurvey::GetOSName(FHardwareSurveyResults& OutResults)
-	return FString();
+	TMap<FString, FString> OsInfo = ReadConfigurationFile(TEXT("/etc/os-release"));
+	if (OsInfo.Num() > 0)
+	{
+		FString Version;
+		FString* IDAddress = OsInfo.Find(TEXT("ID"));
+		if (IDAddress != nullptr)
+		{
+			Version = FString(*IDAddress);
+		}
+
+		IDAddress = OsInfo.Find(TEXT("VERSION_ID"));
+		if (IDAddress != nullptr)
+		{
+			Version += FString(*IDAddress);
+		}
+		else
+		{
+			Version += GetKernelVersion();
+		}
+
+		if (!Version.IsEmpty())
+		{
+			return Version;
+		}
+	}
+
+	TArray<FString> RedHatRelease;
+	if (FFileHelper::LoadANSITextFileToStrings(TEXT("/etc/redhat-release"), &IFileManager::Get(), RedHatRelease))
+	{
+		if (RedHatRelease.IsValidIndex(0))
+		{
+			return RedHatRelease[0];
+		}
+	}
+
+	return GetKernelVersion();
 }
 
 const TCHAR* FUnixPlatformMisc::GetSystemErrorMessage(TCHAR* OutBuffer, int32 BufferCount, int32 Error)
diff --git a/Engine/Source/Runtime/Core/Public/Unix/UnixPlatformMisc.h b/Engine/Source/Runtime/Core/Public/Unix/UnixPlatformMisc.h
--- a/Engine/Source/Runtime/Core/Public/Unix/UnixPlatformMisc.h
+++ b/Engine/Source/Runtime/Core/Public/Unix/UnixPlatformMisc.h
@@ -131,6 +131,7 @@ struct CORE_API FUnixPlatformMisc : public FGenericPlatformMisc
 	 * @return true if the error code has been overriden, false if not
 	 */
 	static bool HasOverriddenReturnCode(uint8 * OverriddenReturnCodeToUsePtr);
+	static void GetOSVersions(FString& out_OSVersionLabel, FString& out_OSSubVersionLabel);
 	static FString GetOSVersion();
 	static FString GetLoginId();
 
