Optimized NEON double version of VectorDot3\4 and added Vector3Scalar version

diff --git a/Engine/Source/Runtime/Core/Public/Math/UnrealMathNeon.h b/Engine/Source/Runtime/Core/Public/Math/UnrealMathNeon.h
--- a/Engine/Source/Runtime/Core/Public/Math/UnrealMathNeon.h
+++ b/Engine/Source/Runtime/Core/Public/Math/UnrealMathNeon.h
@@ -857,15 +857,31 @@ FORCEINLINE VectorRegister4Float VectorDot3( const VectorRegister4Float& Vec1, c
 
 FORCEINLINE VectorRegister4Double VectorDot3(const VectorRegister4Double& Vec1, const VectorRegister4Double& Vec2)
 {
-	VectorRegister4Double Temp = VectorMultiply(Vec1, Vec2);
-	float64x1_t Sum = vadd_f64(vget_low_f64(Temp.XY), vget_high_f64(Temp.XY));
-	Sum = vadd_f64(vget_low_f64(Temp.ZW), Sum);
-
+	VectorRegister2Double A, B;
+	A = vmulq_f64(Vec1.XY, Vec2.XY);
+	B = vfmaq_f64(A, Vec1.ZW, Vec2.ZW);
+	float64x1_t Sum = vadd_f64(vget_low_f64(B), vget_high_f64(A));
+	VectorRegister4Double Temp;
 	Temp.XY = vdupq_lane_f64(Sum, 0);
 	Temp.ZW = Temp.XY;
 	return Temp;
 }
 
+FORCEINLINE float VectorDot3Scalar(const VectorRegister4Float& Vec1, const VectorRegister4Float& Vec2)
+{
+	return vgetq_lane_f32(VectorDot3(Vec1, Vec2), 0);
+}
+
+FORCEINLINE double VectorDot3Scalar(const VectorRegister4Double& Vec1, const VectorRegister4Double& Vec2)
+{
+	VectorRegister2Double A, B;
+	A = vmulq_f64(Vec1.XY, Vec2.XY);
+	B = vfmaq_f64(A, Vec1.ZW, Vec2.ZW);
+	float64x1_t Sum = vadd_f64(vget_low_f64(B), vget_high_f64(A));
+	return *(double*)&Sum;
+}
+
+
 
 /**
  * Calculates the dot4 product of two vectors and returns a vector with the result in all 4 components.
@@ -885,10 +901,12 @@ FORCEINLINE VectorRegister4Float VectorDot4(VectorRegister4Float Vec1, VectorReg
 
 FORCEINLINE VectorRegister4Double VectorDot4(VectorRegister4Double Vec1, VectorRegister4Double Vec2)
 {
-	VectorRegister4Double Temp = VectorMultiply(Vec1, Vec2);
-	VectorRegister2Double Sum = vaddq_f64(Temp.XY, Temp.ZW);
-	Temp.ZW = vextq_f64(Sum, Sum, 1);
-	Temp.XY = vaddq_f64(Sum, Temp.ZW);
+	VectorRegister2Double A, B;
+	A = vmulq_f64(Vec1.XY, Vec2.XY);
+	B = vfmaq_f64(A, Vec1.ZW, Vec2.ZW);
+	A = vextq_f64(B, B, 1);
+	VectorRegister4Double Temp;
+	Temp.XY = vaddq_f64(A, B);
 	Temp.ZW = Temp.XY;
 	return Temp;
 }
