#ue5 - LWC: Avoid using SVML on clang on windows.

diff --git a/Engine/Source/Runtime/Core/Public/Math/UnrealMathSSE.h b/Engine/Source/Runtime/Core/Public/Math/UnrealMathSSE.h
--- a/Engine/Source/Runtime/Core/Public/Math/UnrealMathSSE.h
+++ b/Engine/Source/Runtime/Core/Public/Math/UnrealMathSSE.h
@@ -22,7 +22,7 @@
 #endif
 
 #ifndef UE_PLATFORM_MATH_USE_SVML
-	#if defined(_MSC_VER)
+	#if defined(_MSC_VER) && !defined(__clang__)
 		#define UE_PLATFORM_MATH_USE_SVML			(_MSC_VER >= 1920) // Support added to MSVC 2019 16.0+
 	#else
 		#define UE_PLATFORM_MATH_USE_SVML			0
