UE_CONSTEVAL macro.

diff --git a/Engine/Source/Runtime/Core/Public/HAL/Platform.h b/Engine/Source/Runtime/Core/Public/HAL/Platform.h
--- a/Engine/Source/Runtime/Core/Public/HAL/Platform.h
+++ b/Engine/Source/Runtime/Core/Public/HAL/Platform.h
@@ -685,6 +685,14 @@
 	#define UE_NORETURN
 #endif
 
+/* Macro wrapper for the consteval keyword which isn't yet present on all compilers - constexpr
+   can be used as a workaround but is less strict and so may let some non-consteval code pass */
+#if defined(__cpp_consteval)
+	#define UE_CONSTEVAL consteval
+#else
+	#define UE_CONSTEVAL constexpr
+#endif
+
 /* Wrap a function signature in these to indicate that the function never returns nullptr */
 #ifndef FUNCTION_NON_NULL_RETURN_START
 	#define FUNCTION_NON_NULL_RETURN_START
@@ -1077,7 +1085,6 @@ namespace TypeTests
 	#else
 		#define TEXT_PASTE(x) L ## x
 	#endif
-		#define TEXT(x) TEXT_PASTE(x)
+	#define TEXT(x) TEXT_PASTE(x)
 #endif
 
-
