Some explicit casts to UTF8CHAR, for upcoming changes to that type.

diff --git a/Engine/Plugins/Developer/Concert/ConcertMain/Source/Concert/Private/ConcertLogger.cpp b/Engine/Plugins/Developer/Concert/ConcertMain/Source/Concert/Private/ConcertLogger.cpp
--- a/Engine/Plugins/Developer/Concert/ConcertMain/Source/Concert/Private/ConcertLogger.cpp
+++ b/Engine/Plugins/Developer/Concert/ConcertMain/Source/Concert/Private/ConcertLogger.cpp
@@ -517,7 +517,7 @@ void FConcertLogger::InternalStartLogging()
 	CSVArchive = TUniquePtr<FArchive>(IFileManager::Get().CreateFileWriter(*CSVFilename, EFileWrite::FILEWRITE_AllowRead));
 	if (CSVArchive)
 	{
-		UTF8CHAR UTF8BOM[] = { 0xEF, 0xBB, 0xBF };
+		UTF8CHAR UTF8BOM[] = { (UTF8CHAR)0xEF, (UTF8CHAR)0xBB, (UTF8CHAR)0xBF };
 		CSVArchive->Serialize(&UTF8BOM, UE_ARRAY_COUNT(UTF8BOM) * sizeof(UTF8CHAR));
 	}
 	LogHeader();
diff --git a/Engine/Source/Runtime/Core/Private/Misc/FileHelper.cpp b/Engine/Source/Runtime/Core/Private/Misc/FileHelper.cpp
--- a/Engine/Source/Runtime/Core/Private/Misc/FileHelper.cpp
+++ b/Engine/Source/Runtime/Core/Private/Misc/FileHelper.cpp
@@ -384,7 +384,7 @@ bool FFileHelper::SaveStringToFile( FStringView String, const TCHAR* Filename,
 	bool SaveAsUnicode = EncodingOptions == EEncodingOptions::ForceUnicode || ( EncodingOptions == EEncodingOptions::AutoDetect && !FCString::IsPureAnsi(String.GetData(), String.Len()) );
 	if( EncodingOptions == EEncodingOptions::ForceUTF8 )
 	{
-		UTF8CHAR UTF8BOM[] = { 0xEF, 0xBB, 0xBF };
+		UTF8CHAR UTF8BOM[] = { (UTF8CHAR)0xEF, (UTF8CHAR)0xBB, (UTF8CHAR)0xBF };
 		Ar->Serialize( &UTF8BOM, UE_ARRAY_COUNT(UTF8BOM) * sizeof(UTF8CHAR) );
 
 		FTCHARToUTF8 UTF8String(String.GetData(), String.Len());
diff --git a/Engine/Source/Runtime/Core/Private/Windows/WindowsPlatformProcess.cpp b/Engine/Source/Runtime/Core/Private/Windows/WindowsPlatformProcess.cpp
--- a/Engine/Source/Runtime/Core/Private/Windows/WindowsPlatformProcess.cpp
+++ b/Engine/Source/Runtime/Core/Private/Windows/WindowsPlatformProcess.cpp
@@ -1468,7 +1468,7 @@ FString FWindowsPlatformProcess::ReadPipe( void* ReadPipe )
 		{
 			if (BytesRead > 0)
 			{
-				Buffer[BytesRead] = '\0';
+				Buffer[BytesRead] = (UTF8CHAR)'\0';
 				Output += FUTF8ToTCHAR((const ANSICHAR*)Buffer).Get();
 			}
 		}
@@ -1516,9 +1516,9 @@ bool FWindowsPlatformProcess::WritePipe(void* WritePipe, const FString& Message,
 	UTF8CHAR * Buffer = new UTF8CHAR[BytesAvailable + 2];
 	for (uint32 i = 0; i < BytesAvailable; i++)
 	{
-		Buffer[i] = Message[i];
+		Buffer[i] = (UTF8CHAR)Message[i];
 	}
-	Buffer[BytesAvailable] = '\n';
+	Buffer[BytesAvailable] = (UTF8CHAR)'\n';
 
 	// Write to pipe
 	uint32 BytesWritten = 0;
@@ -1527,7 +1527,7 @@ bool FWindowsPlatformProcess::WritePipe(void* WritePipe, const FString& Message,
 	// Get written message
 	if (OutWritten)
 	{
-		Buffer[BytesWritten] = '\0';
+		Buffer[BytesWritten] = (UTF8CHAR)'\0';
 		*OutWritten = FUTF8ToTCHAR((const ANSICHAR*)Buffer).Get();
 	}
 
