Prevent the engine from providing all threads to CRC on stall and ensure. Only transmit the responsible thread, so that CRC doesn't need to walk all threads before resuming the engine.
  - Analytics shows that CRC may takes up to 400 seconds to walk all the threads and create a minidump before responding back to the engine on the pipe to resume the crashing thread.

diff --git a/Engine/Source/Programs/CrashReportClient/Private/CrashReportAnalyticsSessionSummary.cpp b/Engine/Source/Programs/CrashReportClient/Private/CrashReportAnalyticsSessionSummary.cpp
--- a/Engine/Source/Programs/CrashReportClient/Private/CrashReportAnalyticsSessionSummary.cpp
+++ b/Engine/Source/Programs/CrashReportClient/Private/CrashReportAnalyticsSessionSummary.cpp
@@ -46,8 +46,9 @@ LRESULT CALLBACK CrashReportAnalyticsSessionSummaryWindowProc(HWND Hwnd, UINT Ms
 
 namespace CrcAnalyticsProperties
 {
-	//NOTE: Update this when you add/remove/change key behavior. That's useful to track how one changes affects metrics in-dev where users don't always have an engine versions.
-	constexpr uint32 CrcAnalyticsSummaryVersion = 2;
+	// NOTE: Update this when you add/remove/change key behavior. That's useful to track how one changes affects metrics in-dev where users don't always have an engine versions.
+	//     - V3 -> Windows optimization for stall/ensure -> The engine only captures the responsible thread so CRC walks 1 thread rather than all threads.
+	constexpr uint32 CrcAnalyticsSummaryVersion = 3;
 
 	/** The exit code of the monitored application. */
 	static const TAnalyticsProperty<int32> MonitoredAppExitCode(TEXT("ExitCode"));
diff --git a/Engine/Source/Runtime/Core/Private/Windows/WindowsPlatformCrashContext.cpp b/Engine/Source/Runtime/Core/Private/Windows/WindowsPlatformCrashContext.cpp
--- a/Engine/Source/Runtime/Core/Private/Windows/WindowsPlatformCrashContext.cpp
+++ b/Engine/Source/Runtime/Core/Private/Windows/WindowsPlatformCrashContext.cpp
@@ -746,6 +746,20 @@ int32 ReportCrashForMonitor(
 	bool bCapturedCrashingThreadId = false;
 	if (ThreadSnapshot != INVALID_HANDLE_VALUE)
 	{
+		// Helper function to capture a thread and store its name and id in the context.
+		auto CaptureThread = [&SharedContext, &ThreadIdx](const THREADENTRY32& ThreadEntry)
+		{
+			SharedContext->ThreadIds[ThreadIdx] = ThreadEntry.th32ThreadID;
+			const FString& TmThreadName = FThreadManager::GetThreadName(ThreadEntry.th32ThreadID);
+			const TCHAR* ThreadName = TmThreadName.IsEmpty() ? TEXT("Unknown") : *TmThreadName;
+			FCString::Strcpy(
+				&SharedContext->ThreadNames[ThreadIdx*CR_MAX_THREAD_NAME_CHARS],
+				CR_MAX_THREAD_NAME_CHARS - 1,
+				ThreadName
+			);
+			ThreadIdx++;
+		};
+
 		THREADENTRY32 ThreadEntry;
 		ThreadEntry.dwSize = sizeof(THREADENTRY32);
 		if (Thread32First(ThreadSnapshot, &ThreadEntry))
@@ -754,23 +768,25 @@ int32 ReportCrashForMonitor(
 			{
 				if (ThreadEntry.th32OwnerProcessID == CurrentProcessId)
 				{
-					if (CrashingThreadId == ThreadEntry.th32ThreadID)
+					if (ThreadEntry.th32ThreadID == CrashingThreadId)
 					{
+						// Always capture the crashing thread.
 						bCapturedCrashingThreadId = true;
-					}
+						CaptureThread(ThreadEntry);
 
-					// Always keep one spot for the crashing thread in case the number of captured threads is about to reach our limit.
-					if (bCapturedCrashingThreadId || ThreadIdx < CR_MAX_THREADS - 1)
+						if (FPlatformCrashContext::IsTypeContinuable(Type))
+						{
+							// Early out for continuable types (ensure/stall etc), only capture the responsible thread for performance reasons. (CRC processing 1 thread vs 256 thread is significant).
+							break;
+						}
+					}
+					else if (!FPlatformCrashContext::IsTypeContinuable(Type))
 					{
-						SharedContext->ThreadIds[ThreadIdx] = ThreadEntry.th32ThreadID;
-						const FString& TmThreadName = FThreadManager::GetThreadName(ThreadEntry.th32ThreadID);
-						const TCHAR* ThreadName = TmThreadName.IsEmpty() ? TEXT("Unknown") : *TmThreadName;
-						FCString::Strcpy(
-							&SharedContext->ThreadNames[ThreadIdx*CR_MAX_THREAD_NAME_CHARS],
-							CR_MAX_THREAD_NAME_CHARS - 1,
-							ThreadName
-						);
-						ThreadIdx++;
+						// Capture the thread if enough entries are left (always keep one entry for the crashing thread).
+						if (bCapturedCrashingThreadId || ThreadIdx < CR_MAX_THREADS - 1)
+						{
+							CaptureThread(ThreadEntry);
+						}
 					}
 				}
 			} while (Thread32Next(ThreadSnapshot, &ThreadEntry) && (ThreadIdx < CR_MAX_THREADS));
diff --git a/Engine/Source/Runtime/Engine/Private/Analytics/EditorAnalyticsSessionSummary.cpp b/Engine/Source/Runtime/Engine/Private/Analytics/EditorAnalyticsSessionSummary.cpp
--- a/Engine/Source/Runtime/Engine/Private/Analytics/EditorAnalyticsSessionSummary.cpp
+++ b/Engine/Source/Runtime/Engine/Private/Analytics/EditorAnalyticsSessionSummary.cpp
@@ -53,8 +53,9 @@ FEditorAnalyticsSessionSummary::FEditorAnalyticsSessionSummary(TSharedPtr<IAnaly
 	EditorAnalyticsProperties::ProcessDiagnostics.Set(GetStore(), static_cast<uint32>(FPlatformMisc::GetProcessDiagnostics()));
 
 	// The current summary revision number. Identify the key set used and its behavior. If we add/remove keys, we should increment the number. The number can
-	// also be incremented when we change some behaviors to be able to compare between versions.
-	EditorAnalyticsProperties::SummaryEventVersion.Set(GetStore(), 1);
+	// also be incremented when we change some behaviors to be able to compare between versions. Mostly useful for 'dev' branches where users bypasses UGS to update/build.
+	//    - V3 -> Windows optimization for stall/ensure -> The engine only captures the responsible thread so CRC walks 1 thread rather than all threads.
+	EditorAnalyticsProperties::SummaryEventVersion.Set(GetStore(), 3);
 
 	// Persist the session to disk.
 	GetStore()->Flush();
