Some explicit casts to UTF8CHAR, for upcoming changes to that type.

diff --git a/Engine/Source/Runtime/Core/Private/Mac/MacPlatformProcess.cpp b/Engine/Source/Runtime/Core/Private/Mac/MacPlatformProcess.cpp
--- a/Engine/Source/Runtime/Core/Private/Mac/MacPlatformProcess.cpp
+++ b/Engine/Source/Runtime/Core/Private/Mac/MacPlatformProcess.cpp
@@ -1381,9 +1381,9 @@ bool FMacPlatformProcess::WritePipe(void* WritePipe, const FString& Message, FSt
 	UTF8CHAR * Buffer = new UTF8CHAR[BytesAvailable + 2];
 	for (uint32 i = 0; i < BytesAvailable; i++)
 	{
-		Buffer[i] = Message[i];
+		Buffer[i] = (UTF8CHAR)Message[i];
 	}
-	Buffer[BytesAvailable] = '\n';
+	Buffer[BytesAvailable] = (UTF8CHAR)'\n';
 
 	// Write to pipe
 	uint32 BytesWritten = write([(NSFileHandle*)WritePipe fileDescriptor], Buffer, BytesAvailable + 1);
@@ -1391,7 +1391,7 @@ bool FMacPlatformProcess::WritePipe(void* WritePipe, const FString& Message, FSt
 	// Get written message
 	if (OutWritten)
 	{
-		Buffer[BytesWritten] = '\0';
+		Buffer[BytesWritten] = (UTF8CHAR)'\0';
 		*OutWritten = FUTF8ToTCHAR((const ANSICHAR*)Buffer).Get();
 	}
 
diff --git a/Engine/Source/Runtime/Core/Private/Unix/UnixPlatformProcess.cpp b/Engine/Source/Runtime/Core/Private/Unix/UnixPlatformProcess.cpp
--- a/Engine/Source/Runtime/Core/Private/Unix/UnixPlatformProcess.cpp
+++ b/Engine/Source/Runtime/Core/Private/Unix/UnixPlatformProcess.cpp
@@ -565,9 +565,9 @@ bool FUnixPlatformProcess::WritePipe(void* WritePipe, const FString& Message, FS
 	UTF8CHAR * Buffer = new UTF8CHAR[BytesAvailable + 2];
 	for (uint32 i = 0; i < BytesAvailable; i++)
 	{
-		Buffer[i] = Message[i];
+		Buffer[i] = (UTF8CHAR)Message[i];
 	}
-	Buffer[BytesAvailable] = '\n';
+	Buffer[BytesAvailable] = (UTF8CHAR)'\n';
 
 	// write to pipe
 	uint32 BytesWritten = write(*(int*)WritePipe, Buffer, BytesAvailable + 1);
@@ -575,7 +575,7 @@ bool FUnixPlatformProcess::WritePipe(void* WritePipe, const FString& Message, FS
 	// Get written message
 	if (OutWritten)
 	{
-		Buffer[BytesWritten] = '\0';
+		Buffer[BytesWritten] = (UTF8CHAR)'\0';
 		*OutWritten = FUTF8ToTCHAR((const ANSICHAR*)Buffer).Get();
 	}
 
