UTF8CHAR is now a char8_t on compatible compilers and a unique type in other cases.
UTF8TEXT() macro.

diff --git a/Engine/Source/Runtime/Core/Public/GenericPlatform/GenericPlatform.h b/Engine/Source/Runtime/Core/Public/GenericPlatform/GenericPlatform.h
--- a/Engine/Source/Runtime/Core/Public/GenericPlatform/GenericPlatform.h
+++ b/Engine/Source/Runtime/Core/Public/GenericPlatform/GenericPlatform.h
@@ -67,7 +67,15 @@ struct FGenericPlatformTypes
 	// A wide character. In-memory only. ?-bit fixed-width representation of the platform's natural wide character set. Could be different sizes on different platforms.
 	typedef wchar_t				WIDECHAR;
 	
-	// An 8-bit character type. In-memory only. 8-bit representation. Should really be char8_t but making this the generic option is easier for compilers which don't fully support C++11 yet (i.e. MSVC).
+	// An 8-bit character type. In-memory only. 8-bit representation. Should really be char8_t but making this the generic option is easier for compilers which don't fully support C++20 yet.
+#if defined(__cpp_char8_t)
+	typedef char8_t				UTF8CHAR;
+#else
+	enum UTF8CHAR : unsigned char {};
+#endif
+	
+	// An 8-bit character type. In-memory only. 8-bit representation.
+	/* UE_DEPRECATED(5.0) */[[deprecated("FPlatformTypes::CHAR8 is deprecated, please use FPlatformTypes::UTF8CHAR instead.")]]
 	typedef uint8				CHAR8;
 	
 	// A 16-bit character type. In-memory only.  16-bit representation. Should really be char16_t but making this the generic option is easier for compilers which don't fully support C++11 yet (i.e. MSVC).
@@ -95,3 +103,31 @@ struct FGenericPlatformTypes
 	typedef decltype(nullptr)		TYPE_OF_NULLPTR;
 };
 
+#if !defined(__cpp_char8_t)
+	constexpr FGenericPlatformTypes::UTF8CHAR& operator++(FGenericPlatformTypes::UTF8CHAR& Ch     ) {                                              ++*(unsigned char*)&Ch; return Ch;     }
+	constexpr FGenericPlatformTypes::UTF8CHAR& operator--(FGenericPlatformTypes::UTF8CHAR& Ch     ) {                                              --*(unsigned char*)&Ch; return Ch;     }
+	constexpr FGenericPlatformTypes::UTF8CHAR  operator++(FGenericPlatformTypes::UTF8CHAR& Ch, int) { FGenericPlatformTypes::UTF8CHAR Result = Ch; ++*(unsigned char*)&Ch; return Result; }
+	constexpr FGenericPlatformTypes::UTF8CHAR  operator--(FGenericPlatformTypes::UTF8CHAR& Ch, int) { FGenericPlatformTypes::UTF8CHAR Result = Ch; --*(unsigned char*)&Ch; return Result; }
+
+	template <typename T> inline auto operator+= (FGenericPlatformTypes::UTF8CHAR& Ch, T&& InValue) -> decltype((FGenericPlatformTypes::UTF8CHAR&)(*(unsigned char*)&Ch +=  (T&&)InValue)) { return (FGenericPlatformTypes::UTF8CHAR&)(*(unsigned char*)&Ch +=  (T&&)InValue); }
+	template <typename T> inline auto operator-= (FGenericPlatformTypes::UTF8CHAR& Ch, T&& InValue) -> decltype((FGenericPlatformTypes::UTF8CHAR&)(*(unsigned char*)&Ch -=  (T&&)InValue)) { return (FGenericPlatformTypes::UTF8CHAR&)(*(unsigned char*)&Ch -=  (T&&)InValue); }
+	template <typename T> inline auto operator/= (FGenericPlatformTypes::UTF8CHAR& Ch, T&& InValue) -> decltype((FGenericPlatformTypes::UTF8CHAR&)(*(unsigned char*)&Ch /=  (T&&)InValue)) { return (FGenericPlatformTypes::UTF8CHAR&)(*(unsigned char*)&Ch /=  (T&&)InValue); }
+	template <typename T> inline auto operator*= (FGenericPlatformTypes::UTF8CHAR& Ch, T&& InValue) -> decltype((FGenericPlatformTypes::UTF8CHAR&)(*(unsigned char*)&Ch *=  (T&&)InValue)) { return (FGenericPlatformTypes::UTF8CHAR&)(*(unsigned char*)&Ch *=  (T&&)InValue); }
+	template <typename T> inline auto operator%= (FGenericPlatformTypes::UTF8CHAR& Ch, T&& InValue) -> decltype((FGenericPlatformTypes::UTF8CHAR&)(*(unsigned char*)&Ch %=  (T&&)InValue)) { return (FGenericPlatformTypes::UTF8CHAR&)(*(unsigned char*)&Ch %=  (T&&)InValue); }
+	template <typename T> inline auto operator&= (FGenericPlatformTypes::UTF8CHAR& Ch, T&& InValue) -> decltype((FGenericPlatformTypes::UTF8CHAR&)(*(unsigned char*)&Ch &=  (T&&)InValue)) { return (FGenericPlatformTypes::UTF8CHAR&)(*(unsigned char*)&Ch &=  (T&&)InValue); }
+	template <typename T> inline auto operator|= (FGenericPlatformTypes::UTF8CHAR& Ch, T&& InValue) -> decltype((FGenericPlatformTypes::UTF8CHAR&)(*(unsigned char*)&Ch |=  (T&&)InValue)) { return (FGenericPlatformTypes::UTF8CHAR&)(*(unsigned char*)&Ch |=  (T&&)InValue); }
+	template <typename T> inline auto operator^= (FGenericPlatformTypes::UTF8CHAR& Ch, T&& InValue) -> decltype((FGenericPlatformTypes::UTF8CHAR&)(*(unsigned char*)&Ch ^=  (T&&)InValue)) { return (FGenericPlatformTypes::UTF8CHAR&)(*(unsigned char*)&Ch ^=  (T&&)InValue); }
+	template <typename T> inline auto operator<<=(FGenericPlatformTypes::UTF8CHAR& Ch, T&& InValue) -> decltype((FGenericPlatformTypes::UTF8CHAR&)(*(unsigned char*)&Ch <<= (T&&)InValue)) { return (FGenericPlatformTypes::UTF8CHAR&)(*(unsigned char*)&Ch <<= (T&&)InValue); }
+	template <typename T> inline auto operator>>=(FGenericPlatformTypes::UTF8CHAR& Ch, T&& InValue) -> decltype((FGenericPlatformTypes::UTF8CHAR&)(*(unsigned char*)&Ch >>= (T&&)InValue)) { return (FGenericPlatformTypes::UTF8CHAR&)(*(unsigned char*)&Ch >>= (T&&)InValue); }
+
+	inline FGenericPlatformTypes::UTF8CHAR& operator+= (FGenericPlatformTypes::UTF8CHAR& Ch, FGenericPlatformTypes::UTF8CHAR InValue) { return (FGenericPlatformTypes::UTF8CHAR&)(*(unsigned char*)&Ch +=  *(unsigned char*)&InValue); }
+	inline FGenericPlatformTypes::UTF8CHAR& operator-= (FGenericPlatformTypes::UTF8CHAR& Ch, FGenericPlatformTypes::UTF8CHAR InValue) { return (FGenericPlatformTypes::UTF8CHAR&)(*(unsigned char*)&Ch -=  *(unsigned char*)&InValue); }
+	inline FGenericPlatformTypes::UTF8CHAR& operator/= (FGenericPlatformTypes::UTF8CHAR& Ch, FGenericPlatformTypes::UTF8CHAR InValue) { return (FGenericPlatformTypes::UTF8CHAR&)(*(unsigned char*)&Ch /=  *(unsigned char*)&InValue); }
+	inline FGenericPlatformTypes::UTF8CHAR& operator*= (FGenericPlatformTypes::UTF8CHAR& Ch, FGenericPlatformTypes::UTF8CHAR InValue) { return (FGenericPlatformTypes::UTF8CHAR&)(*(unsigned char*)&Ch *=  *(unsigned char*)&InValue); }
+	inline FGenericPlatformTypes::UTF8CHAR& operator%= (FGenericPlatformTypes::UTF8CHAR& Ch, FGenericPlatformTypes::UTF8CHAR InValue) { return (FGenericPlatformTypes::UTF8CHAR&)(*(unsigned char*)&Ch %=  *(unsigned char*)&InValue); }
+	inline FGenericPlatformTypes::UTF8CHAR& operator&= (FGenericPlatformTypes::UTF8CHAR& Ch, FGenericPlatformTypes::UTF8CHAR InValue) { return (FGenericPlatformTypes::UTF8CHAR&)(*(unsigned char*)&Ch &=  *(unsigned char*)&InValue); }
+	inline FGenericPlatformTypes::UTF8CHAR& operator|= (FGenericPlatformTypes::UTF8CHAR& Ch, FGenericPlatformTypes::UTF8CHAR InValue) { return (FGenericPlatformTypes::UTF8CHAR&)(*(unsigned char*)&Ch |=  *(unsigned char*)&InValue); }
+	inline FGenericPlatformTypes::UTF8CHAR& operator^= (FGenericPlatformTypes::UTF8CHAR& Ch, FGenericPlatformTypes::UTF8CHAR InValue) { return (FGenericPlatformTypes::UTF8CHAR&)(*(unsigned char*)&Ch ^=  *(unsigned char*)&InValue); }
+	inline FGenericPlatformTypes::UTF8CHAR& operator<<=(FGenericPlatformTypes::UTF8CHAR& Ch, FGenericPlatformTypes::UTF8CHAR InValue) { return (FGenericPlatformTypes::UTF8CHAR&)(*(unsigned char*)&Ch <<= *(unsigned char*)&InValue); }
+	inline FGenericPlatformTypes::UTF8CHAR& operator>>=(FGenericPlatformTypes::UTF8CHAR& Ch, FGenericPlatformTypes::UTF8CHAR InValue) { return (FGenericPlatformTypes::UTF8CHAR&)(*(unsigned char*)&Ch >>= *(unsigned char*)&InValue); }
+#endif
diff --git a/Engine/Source/Runtime/Core/Public/HAL/Platform.h b/Engine/Source/Runtime/Core/Public/HAL/Platform.h
--- a/Engine/Source/Runtime/Core/Public/HAL/Platform.h
+++ b/Engine/Source/Runtime/Core/Public/HAL/Platform.h
@@ -970,7 +970,7 @@ typedef FPlatformTypes::WIDECHAR	WIDECHAR;
 /// Either ANSICHAR or WIDECHAR, depending on whether the platform supports wide characters or the requirements of the licensee.
 typedef FPlatformTypes::TCHAR		TCHAR;
 /// An 8-bit character containing a UTF8 (Unicode, 8-bit, variable-width) code unit.
-typedef FPlatformTypes::CHAR8		UTF8CHAR;
+typedef FPlatformTypes::UTF8CHAR	UTF8CHAR;
 /// A 16-bit character containing a UCS2 (Unicode, 16-bit, fixed-width) code unit, used for compatibility with 'Windows TCHAR' across multiple platforms.
 typedef FPlatformTypes::CHAR16		UCS2CHAR;
 /// A 16-bit character containing a UTF16 (Unicode, 16-bit, variable-width) code unit.
@@ -1085,6 +1085,27 @@ namespace TypeTests
 	#else
 		#define TEXT_PASTE(x) L ## x
 	#endif
-	#define TEXT(x) TEXT_PASTE(x)
+		#define TEXT(x) TEXT_PASTE(x)
 #endif
 
+#define UTF8TEXT_PASTE(x) u8 ## x
+
+#if defined(__cpp_char8_t)
+	#define UTF8TEXT(x) UTF8TEXT_PASTE(x)
+#else
+	namespace UE::Core::Private
+	{
+		template <SIZE_T N>
+		FORCEINLINE auto ToUTF8Literal(const char(&Array)[N]) -> const UTF8CHAR(&)[N]
+		{
+			return (const UTF8CHAR(&)[N])Array;
+		}
+
+		FORCEINLINE UTF8CHAR ToUTF8Literal(unsigned long long Ch)
+		{
+			return (UTF8CHAR)Ch;
+		}
+	}
+
+	#define UTF8TEXT(x) (UE::Core::Private::ToUTF8Literal(UTF8TEXT_PASTE(x)))
+#endif
