UTF-8 string views.

diff --git a/Engine/Source/Runtime/Core/Public/Containers/StringFwd.h b/Engine/Source/Runtime/Core/Public/Containers/StringFwd.h
--- a/Engine/Source/Runtime/Core/Public/Containers/StringFwd.h
+++ b/Engine/Source/Runtime/Core/Public/Containers/StringFwd.h
@@ -19,6 +19,8 @@ using FStringBuilderBase = TStringBuilderBase<TCHAR>;
 using FAnsiStringBuilderBase = TStringBuilderBase<ANSICHAR>;
 /** The base string builder type for WIDECHAR. */
 using FWideStringBuilderBase = TStringBuilderBase<WIDECHAR>;
+/** The base string builder type for UTF8CHAR. */
+using FUtf8StringBuilderBase = TStringBuilderBase<UTF8CHAR>;
 
 /** An extendable string builder for TCHAR. */
 template <int32 BufferSize> using TStringBuilder = TStringBuilderWithBuffer<TCHAR, BufferSize>;
@@ -26,6 +28,8 @@ template <int32 BufferSize> using TStringBuilder = TStringBuilderWithBuffer<TCHA
 template <int32 BufferSize> using TAnsiStringBuilder = TStringBuilderWithBuffer<ANSICHAR, BufferSize>;
 /** An extendable string builder for WIDECHAR. */
 template <int32 BufferSize> using TWideStringBuilder = TStringBuilderWithBuffer<WIDECHAR, BufferSize>;
+/** An extendable string builder for UTF8CHAR. */
+template <int32 BufferSize> using TUtf8StringBuilder = TStringBuilderWithBuffer<UTF8CHAR, BufferSize>;
 
 /** A fixed-size string builder for TCHAR. */
 template <int32 BufferSize> using TFixedStringBuilder UE_DEPRECATED(4.25, "'TFixedStringBuilder' is deprecated. Please use 'TStringBuilder' instead!") = TStringBuilderWithBuffer<TCHAR, BufferSize>;
@@ -41,3 +45,4 @@ template <typename CharType> struct TIsContiguousContainer<TStringView<CharType>
 using FStringView = TStringView<TCHAR>;
 using FAnsiStringView = TStringView<ANSICHAR>;
 using FWideStringView = TStringView<WIDECHAR>;
+using FUtf8StringView = TStringView<UTF8CHAR>;
diff --git a/Engine/Source/Runtime/Core/Public/Misc/StringBuilder.h b/Engine/Source/Runtime/Core/Public/Misc/StringBuilder.h
--- a/Engine/Source/Runtime/Core/Public/Misc/StringBuilder.h
+++ b/Engine/Source/Runtime/Core/Public/Misc/StringBuilder.h
@@ -179,6 +179,31 @@ public:
 		return *this;
 	}
 
+	inline BuilderType& AppendUtf8(const UTF8CHAR* NulTerminatedString)
+	{
+		if (!NulTerminatedString)
+		{
+			return *this;
+		}
+
+		return AppendUtf8(NulTerminatedString, TCString<UTF8CHAR>::Strlen(NulTerminatedString));
+	}
+
+	inline BuilderType& AppendUtf8(const FUtf8StringView& AnsiString)
+	{
+		return AppendUtf8(AnsiString.GetData(), AnsiString.Len());
+	}
+
+	inline BuilderType& AppendUtf8(const UTF8CHAR* String, const int32 Length)
+	{
+		int32 ConvertedBufferLength = FPlatformString::ConvertedLength<ElementType>(String, Length);
+		EnsureAdditionalCapacity(ConvertedBufferLength);
+
+		CurPos = FPlatformString::Convert(CurPos, ConvertedBufferLength, String, Length);
+
+		return *this;
+	}
+
 	inline BuilderType& Append(const CharType* NulTerminatedString)
 	{
 		if (!NulTerminatedString)
@@ -408,6 +433,7 @@ private:
 inline FAnsiStringBuilderBase&		operator<<(FAnsiStringBuilderBase& Builder, ANSICHAR Char)							{ return Builder.Append(Char); }
 inline FWideStringBuilderBase&		operator<<(FWideStringBuilderBase& Builder, ANSICHAR Char)							{ return Builder.Append(Char); }
 inline FWideStringBuilderBase&		operator<<(FWideStringBuilderBase& Builder, WIDECHAR Char)							{ return Builder.Append(Char); }
+inline FUtf8StringBuilderBase&		operator<<(FUtf8StringBuilderBase& Builder, UTF8CHAR Char)							{ return Builder.Append(Char); }
 
 template <typename T>
 inline auto operator<<(FAnsiStringBuilderBase& Builder, T&& Str) -> decltype(Builder.Append(ImplicitConv<FAnsiStringView>(Forward<T>(Str))))
@@ -427,11 +453,18 @@ inline auto operator<<(FWideStringBuilderBase& Builder, T&& Str) -> decltype(Bui
 	return Builder.Append(ImplicitConv<FWideStringView>(Forward<T>(Str)));
 }
 
+template <typename T>
+inline auto operator<<(FUtf8StringBuilderBase& Builder, T&& Str) -> decltype(Builder.Append(ImplicitConv<FUtf8StringView>(Forward<T>(Str))))
+{
+	return Builder.AppendUtf8(ImplicitConv<FUtf8StringView>(Forward<T>(Str)));
+}
+
 // Prefer using << instead of += as operator+= is only intended for mechanical FString -> FStringView replacement.
 inline FStringBuilderBase&			operator+=(FStringBuilderBase& Builder, ANSICHAR Char)								{ return Builder.Append(Char); }
 inline FStringBuilderBase&			operator+=(FStringBuilderBase& Builder, WIDECHAR Char)								{ return Builder.Append(Char); }
 inline FStringBuilderBase&			operator+=(FStringBuilderBase& Builder, FAnsiStringView Str)						{ return Builder.AppendAnsi(Str); }
 inline FStringBuilderBase&			operator+=(FStringBuilderBase& Builder, FWideStringView Str)						{ return Builder.Append(Str); }
+inline FStringBuilderBase&			operator+=(FStringBuilderBase& Builder, FUtf8StringView Str)						{ return Builder.AppendUtf8(Str); }
 
 // Integer Append Operators
 
@@ -439,21 +472,29 @@ inline FAnsiStringBuilderBase&		operator<<(FAnsiStringBuilderBase& Builder, int3
 inline FAnsiStringBuilderBase&		operator<<(FAnsiStringBuilderBase& Builder, uint32 Value)							{ return Builder.Appendf("%u", Value); }
 inline FWideStringBuilderBase&		operator<<(FWideStringBuilderBase& Builder, int32 Value)							{ return Builder.Appendf(TEXT("%d"), Value); }
 inline FWideStringBuilderBase&		operator<<(FWideStringBuilderBase& Builder, uint32 Value)							{ return Builder.Appendf(TEXT("%u"), Value); }
+inline FUtf8StringBuilderBase&		operator<<(FUtf8StringBuilderBase& Builder, int32 Value)							{ return Builder.Appendf(UTF8TEXT("%d"), Value); }
+inline FUtf8StringBuilderBase&		operator<<(FUtf8StringBuilderBase& Builder, uint32 Value)							{ return Builder.Appendf(UTF8TEXT("%u"), Value); }
 
 inline FAnsiStringBuilderBase&		operator<<(FAnsiStringBuilderBase& Builder, int64 Value)							{ return Builder.Appendf("%" INT64_FMT, Value); }
 inline FAnsiStringBuilderBase&		operator<<(FAnsiStringBuilderBase& Builder, uint64 Value)							{ return Builder.Appendf("%" UINT64_FMT, Value); }
 inline FWideStringBuilderBase&		operator<<(FWideStringBuilderBase& Builder, int64 Value)							{ return Builder.Appendf(TEXT("%" INT64_FMT), Value); }
 inline FWideStringBuilderBase&		operator<<(FWideStringBuilderBase& Builder, uint64 Value)							{ return Builder.Appendf(TEXT("%" UINT64_FMT), Value); }
+inline FUtf8StringBuilderBase&		operator<<(FUtf8StringBuilderBase& Builder, int64 Value)							{ return Builder.Appendf(UTF8TEXT("%" INT64_FMT), Value); }
+inline FUtf8StringBuilderBase&		operator<<(FUtf8StringBuilderBase& Builder, uint64 Value)							{ return Builder.Appendf(UTF8TEXT("%" UINT64_FMT), Value); }
 
 inline FAnsiStringBuilderBase&		operator<<(FAnsiStringBuilderBase& Builder, int8 Value)								{ return Builder << int32(Value); }
 inline FAnsiStringBuilderBase&		operator<<(FAnsiStringBuilderBase& Builder, uint8 Value)							{ return Builder << uint32(Value); }
 inline FWideStringBuilderBase&		operator<<(FWideStringBuilderBase& Builder, int8 Value)								{ return Builder << int32(Value); }
 inline FWideStringBuilderBase&		operator<<(FWideStringBuilderBase& Builder, uint8 Value)							{ return Builder << uint32(Value); }
+inline FUtf8StringBuilderBase&		operator<<(FUtf8StringBuilderBase& Builder, int8 Value)								{ return Builder << int32(Value); }
+inline FUtf8StringBuilderBase&		operator<<(FUtf8StringBuilderBase& Builder, uint8 Value)							{ return Builder << uint32(Value); }
 
 inline FAnsiStringBuilderBase&		operator<<(FAnsiStringBuilderBase& Builder, int16 Value)							{ return Builder << int32(Value); }
 inline FAnsiStringBuilderBase&		operator<<(FAnsiStringBuilderBase& Builder, uint16 Value)							{ return Builder << uint32(Value); }
 inline FWideStringBuilderBase&		operator<<(FWideStringBuilderBase& Builder, int16 Value)							{ return Builder << int32(Value); }
 inline FWideStringBuilderBase&		operator<<(FWideStringBuilderBase& Builder, uint16 Value)							{ return Builder << uint32(Value); }
+inline FUtf8StringBuilderBase&		operator<<(FUtf8StringBuilderBase& Builder, int16 Value)							{ return Builder << int32(Value); }
+inline FUtf8StringBuilderBase&		operator<<(FUtf8StringBuilderBase& Builder, uint16 Value)							{ return Builder << uint32(Value); }
 
 /**
  * A function-like type that creates a TStringBuilder by appending its arguments.
