#ue5 - LWC: SSE double alignment changed to always be 16-byte (instead of 32 depending on build flags). We're not using any AVX instructions which require 32 byte alignment.

diff --git a/Engine/Source/Runtime/Core/Public/Math/UnrealMathSSE.h b/Engine/Source/Runtime/Core/Public/Math/UnrealMathSSE.h
--- a/Engine/Source/Runtime/Core/Public/Math/UnrealMathSSE.h
+++ b/Engine/Source/Runtime/Core/Public/Math/UnrealMathSSE.h
@@ -41,11 +41,11 @@
 // If AVX is enabled, need additional defines.
 #if UE_PLATFORM_MATH_USE_AVX || UE_PLATFORM_MATH_USE_SVML
 #include <immintrin.h>
-#define UE_SSE_DOUBLE_ALIGNMENT 32
-#else
-#define UE_SSE_DOUBLE_ALIGNMENT 16
 #endif
 
+#define UE_SSE_FLOAT_ALIGNMENT	16
+#define UE_SSE_DOUBLE_ALIGNMENT 16
+
 #include "Math/sse_mathfun.h"
 
 // We suppress static analysis warnings for the cast from (double*) to (float*) in VectorLoadFloat2
@@ -170,7 +170,7 @@ void VectorStoreAligned(const VectorRegister4Double& Vec, double* Dst);
 
 
 // Helper for conveniently aligning a float array for extraction from VectorRegister4Float
-struct alignas(16) AlignedFloat4
+struct alignas(UE_SSE_FLOAT_ALIGNMENT) AlignedFloat4
 {
 	float V[4];
 
