#ue5 - LWC: AVX requires 32-byte alignment for __m256d, but at least now we don't use the SVML compile flag erroneously.

diff --git a/Engine/Source/Runtime/Core/Public/Math/UnrealMathSSE.h b/Engine/Source/Runtime/Core/Public/Math/UnrealMathSSE.h
--- a/Engine/Source/Runtime/Core/Public/Math/UnrealMathSSE.h
+++ b/Engine/Source/Runtime/Core/Public/Math/UnrealMathSSE.h
@@ -44,7 +44,12 @@
 #endif
 
 #define UE_SSE_FLOAT_ALIGNMENT	16
+
+#if UE_PLATFORM_MATH_USE_AVX
+#define UE_SSE_DOUBLE_ALIGNMENT 32 // required for __m256d
+#else
 #define UE_SSE_DOUBLE_ALIGNMENT 16
+#endif
 
 #include "Math/sse_mathfun.h"
 
