#ue5 - LWC: Various improvements especially for AVX/AVX2:

- Improved AVX loading of two separate m128d values using native instruction.
- Improved VectorGetComponent() in MSVC for some specializations because it generates suboptimal code with the generic approach. Biggest benefit to VectorGetComponent(V, 0) which is the most common.
- Implemented VectorDot3Scalar() for FPU and SSE for when only the scalar value is needed, which can avoid some permute operations. Added tests to UnrealMathTests.
- Added a few more specializations for AVX1 for VectorSwizzleTemplate<>.

diff --git a/Engine/Source/Runtime/Core/Private/Tests/Math/UnrealMathTest.cpp b/Engine/Source/Runtime/Core/Private/Tests/Math/UnrealMathTest.cpp
--- a/Engine/Source/Runtime/Core/Private/Tests/Math/UnrealMathTest.cpp
+++ b/Engine/Source/Runtime/Core/Private/Tests/Math/UnrealMathTest.cpp
@@ -1265,6 +1265,7 @@ bool RunDoubleVectorTest()
 	uint64 U1 = *(uint64*)&F1;
 	VectorRegister4Double V0, V1;
 	VectorRegister4Double V2, V3;
+	double Double0, Double1, Double2, Double3;
 
 	V0 = MakeVectorRegisterDouble(U1, U1, U1, U1);
 	V1 = MakeVectorRegisterDouble(F1, F1, F1, F1);
@@ -1362,14 +1363,15 @@ bool RunDoubleVectorTest()
 	// Shuffle
 	TestVectorShuffle<double, VectorRegister4Double>();
 
+	// VectorGetComponent
 	V0 = MakeVectorRegisterDouble(0., 0., 2., 0.);
 	V1 = MakeVectorRegisterDouble(0., 1., 0., 3.);
-	double D0 = VectorGetComponent(V0, 0);
-	double D1 = VectorGetComponent(V1, 1);
-	double D2 = VectorGetComponent(V0, 2);
-	double D3 = VectorGetComponent(V1, 3);
+	Double0 = VectorGetComponent(V0, 0);
+	Double1 = VectorGetComponent(V1, 1);
+	Double2 = VectorGetComponent(V0, 2);
+	Double3 = VectorGetComponent(V1, 3);
 	V0 = MakeVectorRegisterDouble(0., 1., 2., 3.);
-	V1 = MakeVectorRegisterDouble(D0, D1, D2, D3);
+	V1 = MakeVectorRegisterDouble(Double0, Double1, Double2, Double3);
 	LogTest<double>(TEXT("VectorGetComponent"), TestVectorsEqual(V0, V1));
 
 	// Abs
@@ -1419,6 +1421,20 @@ bool RunDoubleVectorTest()
 	V0 = MakeVectorRegisterDouble(68.0f, 68.0f, 68.0f, 68.0f);
 	LogTest<double>(TEXT("VectorDot3"), TestVectorsEqual(V0, V1));
 
+	V0 = MakeVectorRegisterDouble(2.0f, 4.0f, 6.0f, 8.0f);
+	V1 = MakeVectorRegisterDouble(1.0f, 2.0f, 3.0f, 4.0f);
+	Double0 = VectorDot3Scalar(V0, V1);
+	V1 = VectorSetDouble1(Double0);
+	V0 = MakeVectorRegisterDouble(28.0f, 28.0f, 28.0f, 28.0f);
+	LogTest<double>(TEXT("VectorDot3Scalar"), TestVectorsEqual(V0, V1));
+
+	V0 = MakeVectorRegisterDouble(1.0f, 2.0f, 3.0f, 5.0f);
+	V1 = MakeVectorRegisterDouble(7.0f, 11.0f, 13.0f, 17.0f);
+	Double0 = VectorDot3Scalar(V0, V1);
+	V1 = VectorSetDouble1(Double0);
+	V0 = MakeVectorRegisterDouble(68.0f, 68.0f, 68.0f, 68.0f);
+	LogTest<double>(TEXT("VectorDot3Scalar"), TestVectorsEqual(V0, V1));
+
 	V0 = MakeVectorRegisterDouble(2.0f, 4.0f, 6.0f, 8.0f);
 	V1 = MakeVectorRegisterDouble(1.0f, 2.0f, 3.0f, 4.0f);
 	V1 = VectorDot4(V0, V1);
@@ -1748,6 +1764,7 @@ bool FVectorRegisterAbstractionTest::RunTest(const FString& Parameters)
 	float F1 = 1.f;
 	uint32 D1 = *(uint32 *)&F1;
 	VectorRegister4Float V0, V1, V2, V3;
+	float Float0, Float1, Float2, Float3;
 
 	ResetPassing();
 
@@ -1843,10 +1860,10 @@ bool FVectorRegisterAbstractionTest::RunTest(const FString& Parameters)
 
 	V0 = MakeVectorRegisterFloat(0.f, 0.f, 2.f, 0.f);
 	V1 = MakeVectorRegisterFloat(0.f, 1.f, 0.f, 3.f);
-	float Float0 = VectorGetComponent(V0, 0);
-	float Float1 = VectorGetComponent(V1, 1);
-	float Float2 = VectorGetComponent(V0, 2);
-	float Float3 = VectorGetComponent(V1, 3);
+	Float0 = VectorGetComponent(V0, 0);
+	Float1 = VectorGetComponent(V1, 1);
+	Float2 = VectorGetComponent(V0, 2);
+	Float3 = VectorGetComponent(V1, 3);
 	V0 = MakeVectorRegisterFloat(0.f, 1.f, 2.f, 3.f);
 	V1 = MakeVectorRegisterFloat(Float0, Float1, Float2, Float3);
 	LogTest<float>(TEXT("VectorGetComponent<float>"), TestVectorsEqual(V0, V1));
@@ -1899,6 +1916,20 @@ bool FVectorRegisterAbstractionTest::RunTest(const FString& Parameters)
 	V0 = MakeVectorRegister(68.0f, 68.0f, 68.0f, 68.0f);
 	LogTest<float>(TEXT("VectorDot3"), TestVectorsEqual(V0, V1));
 
+	V0 = MakeVectorRegister(2.0f, 4.0f, 6.0f, 8.0f);
+	V1 = MakeVectorRegister(1.0f, 2.0f, 3.0f, 4.0f);
+	Float0 = VectorDot3Scalar(V0, V1);
+	V1 = VectorSetFloat1(Float0);
+	V0 = MakeVectorRegister(28.0f, 28.0f, 28.0f, 28.0f);
+	LogTest<float>(TEXT("VectorDot3Scalar"), TestVectorsEqual(V0, V1));
+
+	V0 = MakeVectorRegister(1.0f, 2.0f, 3.0f, 5.0f);
+	V1 = MakeVectorRegister(7.0f, 11.0f, 13.0f, 17.0f);
+	Float0 = VectorDot3Scalar(V0, V1);
+	V1 = VectorSetFloat1(Float0);
+	V0 = MakeVectorRegister(68.0f, 68.0f, 68.0f, 68.0f);
+	LogTest<float>(TEXT("VectorDot3Scalar"), TestVectorsEqual(V0, V1));
+
 	V0 = MakeVectorRegister(2.0f, 4.0f, 6.0f, 8.0f);
 	V1 = MakeVectorRegister(1.0f, 2.0f, 3.0f, 4.0f);
 	V1 = VectorDot4(V0, V1);
diff --git a/Engine/Source/Runtime/Core/Public/Math/UnrealMathFPU.h b/Engine/Source/Runtime/Core/Public/Math/UnrealMathFPU.h
--- a/Engine/Source/Runtime/Core/Public/Math/UnrealMathFPU.h
+++ b/Engine/Source/Runtime/Core/Public/Math/UnrealMathFPU.h
@@ -26,6 +26,11 @@ struct VectorRegister4Int
 };
 
 
+struct VectorRegister2Double
+{
+	double V[2];
+};
+
 /**
 *	double[4] vector register type, where the first double (X) is stored in the lowest 64 bits, and so on.
 */
@@ -35,8 +40,16 @@ struct alignas(16) VectorRegister4Double
 
 	VectorRegister4Double() = default;
 
+	FORCEINLINE VectorRegister4Double(const VectorRegister2Double& InXY, const VectorRegister2Double& InZW)
+	{
+		V[0] = InXY.V[0];
+		V[1] = InXY.V[1];
+		V[2] = InZW.V[0];
+		V[3] = InZW.V[1];
+	}
+
 	// Construct from a vector of 4 floats
-	VectorRegister4Double(const VectorRegister4Float& FloatVector)
+	FORCEINLINE VectorRegister4Double(const VectorRegister4Float& FloatVector)
 	{
 		V[0] = FloatVector.V[0];
 		V[1] = FloatVector.V[1];
@@ -45,7 +58,7 @@ struct alignas(16) VectorRegister4Double
 	}
 
 	// Assign from a vector of 4 floats
-	VectorRegister4Double& operator=(const VectorRegister4Float& From)
+	FORCEINLINE VectorRegister4Double& operator=(const VectorRegister4Float& From)
 	{
 		V[0] = From.V[0];
 		V[1] = From.V[1];
@@ -55,11 +68,6 @@ struct alignas(16) VectorRegister4Double
 	}
 };
 
-struct VectorRegister2Double
-{
-	double V[2];
-};
-
 // Aliases
 typedef VectorRegister4Int VectorRegister4i;
 typedef VectorRegister4Float VectorRegister4f;
@@ -217,6 +225,11 @@ FORCEINLINE VectorRegister4Double MakeVectorRegisterDouble(double X, double Y, d
 	return Vec;
 }
 
+FORCEINLINE VectorRegister4Double MakeVectorRegisterDouble(const VectorRegister2Double& XY, const VectorRegister2Double& ZW)
+{
+	return VectorRegister4Double(XY, ZW);
+}
+
 FORCEINLINE VectorRegister4Float MakeVectorRegister(float X, float Y, float Z, float W)
 {
 	return MakeVectorRegisterFloat(X, Y, Z, W);
@@ -457,11 +470,6 @@ FORCEINLINE VectorRegister4Double VectorSetFloat1(double D)
 	return MakeVectorRegisterDouble(D, D, D, D);
 }
 
-FORCEINLINE VectorRegister4Double VectorSetDouble1(double D)
-{
-	return MakeVectorRegisterDouble(D, D, D, D);
-}
-
 /**
  * Stores a vector to aligned memory.
  *
@@ -777,6 +785,25 @@ FORCEINLINE VectorRegister4Double VectorDivide(const VectorRegister4Double& Vec1
 	return Vec;
 }
 
+/**
+ * Calculates the dot3 product of two vectors and returns a scalar value.
+ *
+ * @param Vec1	1st vector
+ * @param Vec2	2nd vector
+ * @return		d = dot3(Vec1.xyz, Vec2.xyz)
+ */
+FORCEINLINE float VectorDot3Scalar(const VectorRegister4Float& Vec1, const VectorRegister4Float& Vec2)
+{
+	float D = Vec1.V[0] * Vec2.V[0] + Vec1.V[1] * Vec2.V[1] + Vec1.V[2] * Vec2.V[2];
+	return D;
+}
+
+FORCEINLINE double VectorDot3Scalar(const VectorRegister4Double& Vec1, const VectorRegister4Double& Vec2)
+{
+	double D = Vec1.V[0] * Vec2.V[0] + Vec1.V[1] * Vec2.V[1] + Vec1.V[2] * Vec2.V[2];
+	return D;
+}
+
 /**
  * Calculates the dot3 product of two vectors and returns a vector with the result in all 4 components.
  *
@@ -786,19 +813,18 @@ FORCEINLINE VectorRegister4Double VectorDivide(const VectorRegister4Double& Vec1
  */
 FORCEINLINE VectorRegister4Float VectorDot3(const VectorRegister4Float& Vec1, const VectorRegister4Float& Vec2)
 {
-	float D = Vec1.V[0] * Vec2.V[0] + Vec1.V[1] * Vec2.V[1] + Vec1.V[2] * Vec2.V[2];
+	float D = VectorDot3Scalar(Vec1, Vec2);
 	VectorRegister4Float Vec = MakeVectorRegisterFloat(D, D, D, D);
 	return Vec;
 }
 
 FORCEINLINE VectorRegister4Double VectorDot3(const VectorRegister4Double& Vec1, const VectorRegister4Double& Vec2)
 {
-	double D = Vec1.V[0] * Vec2.V[0] + Vec1.V[1] * Vec2.V[1] + Vec1.V[2] * Vec2.V[2];
+	double D = VectorDot3Scalar(Vec1, Vec2);
 	VectorRegister4Double Vec = MakeVectorRegisterDouble(D, D, D, D);
 	return Vec;
 }
 
-
 /**
  * Calculates the dot4 product of two vectors and returns a vector with the result in all 4 components.
  *
diff --git a/Engine/Source/Runtime/Core/Public/Math/UnrealMathNeon.h b/Engine/Source/Runtime/Core/Public/Math/UnrealMathNeon.h
--- a/Engine/Source/Runtime/Core/Public/Math/UnrealMathNeon.h
+++ b/Engine/Source/Runtime/Core/Public/Math/UnrealMathNeon.h
@@ -240,6 +240,11 @@ FORCEINLINE VectorRegister4Double MakeVectorRegisterDouble(double X, double Y, d
 	return MakeVectorRegister(X, Y, Z, W);
 }
 
+FORCEINLINE VectorRegister4Double MakeVectorRegisterDouble(const VectorRegister2Double& XY, const VectorRegister2Double& ZW)
+{
+	return VectorRegister4Double(XY, ZW);
+}
+
 FORCEINLINE VectorRegister4Double MakeVectorRegisterDouble(uint64 X, uint64 Y, uint64 Z, uint64 W)
 {
 	union U
diff --git a/Engine/Source/Runtime/Core/Public/Math/UnrealMathSSE.h b/Engine/Source/Runtime/Core/Public/Math/UnrealMathSSE.h
--- a/Engine/Source/Runtime/Core/Public/Math/UnrealMathSSE.h
+++ b/Engine/Source/Runtime/Core/Public/Math/UnrealMathSSE.h
@@ -97,9 +97,13 @@ struct alignas(UE_SSE_DOUBLE_ALIGNMENT) VectorRegister4Double
 	FORCEINLINE VectorRegister4Double() = default;
 
 	FORCEINLINE VectorRegister4Double(const VectorRegister2Double& InXY, const VectorRegister2Double& InZW)
-		: XY(InXY)
-		, ZW(InZW)
 	{
+#if UE_PLATFORM_MATH_USE_AVX
+		XYZW = _mm256_setr_m128d(InXY, InZW);
+#else
+		XY = InXY;
+		ZW = InZW;
+#endif
 	}
 
 	// Construct from a vector of 4 floats
@@ -266,10 +270,7 @@ FORCEINLINE VectorRegister4Float MakeVectorRegisterFloat( uint32 X, uint32 Y, ui
 
 FORCEINLINE VectorRegister4Double MakeVectorRegisterDouble(uint64 X, uint64 Y, uint64 Z, uint64 W)
 {
-	VectorRegister4Double Result;
-	Result.XY = MakeVectorRegister2DoubleMask(X, Y);
-	Result.ZW = MakeVectorRegister2DoubleMask(Z, W);
-	return Result;
+	return VectorRegister4Double(MakeVectorRegister2DoubleMask(X, Y), MakeVectorRegister2DoubleMask(Z, W));
 }
 
 FORCEINLINE VectorRegister4Float MakeVectorRegister(uint32 X, uint32 Y, uint32 Z, uint32 W)
@@ -324,6 +325,11 @@ FORCEINLINE VectorRegister4Double MakeVectorRegister(double X, double Y, double
 	return MakeVectorRegisterDouble(X, Y, Z, W);
 }
 
+FORCEINLINE VectorRegister4Double MakeVectorRegisterDouble(const VectorRegister2Double& XY, const VectorRegister2Double& ZW)
+{
+	return VectorRegister4Double(XY, ZW);
+}
+
 // Make double register from float register
 FORCEINLINE VectorRegister4Double MakeVectorRegisterDouble(const VectorRegister4Float& From)
 {
@@ -415,6 +421,9 @@ FORCEINLINE float VectorGetComponentImpl(const VectorRegister4Float& Vec)
 	return (((float*)&(Vec))[ComponentIndex]);
 }
 
+// Specializations
+template <> FORCEINLINE float VectorGetComponentImpl<0>(const VectorRegister4Float& Vec) { return _mm_cvtss_f32(Vec); }
+
 template <uint32 ComponentIndex>
 FORCEINLINE double VectorGetComponentImpl(const VectorRegister4Double& Vec)
 {
@@ -425,6 +434,15 @@ FORCEINLINE double VectorGetComponentImpl(const VectorRegister4Double& Vec)
 #endif
 }
 
+// Specializations
+#if UE_PLATFORM_MATH_USE_AVX
+// Lower latency than `vmovsd`, required since MSVC doesn't optimize the above impl well compared to clang/gcc. The latter basically generates something like this below (checked in godbolt).
+template <> FORCEINLINE double VectorGetComponentImpl<0>(const VectorRegister4Double& Vec) { return _mm256_cvtsd_f64(Vec); }
+template <> FORCEINLINE double VectorGetComponentImpl<1>(const VectorRegister4Double& Vec) { return _mm256_cvtsd_f64(_mm256_permute_pd(Vec, 1)); }
+template <> FORCEINLINE double VectorGetComponentImpl<2>(const VectorRegister4Double& Vec) { return _mm_cvtsd_f64(_mm256_extractf128_pd(Vec, 1)); }
+template <> FORCEINLINE double VectorGetComponentImpl<3>(const VectorRegister4Double& Vec) { return _mm_cvtsd_f64(_mm_permute_pd(_mm256_extractf128_pd(Vec, 1), 1)); }
+#endif
+
 #define VectorGetComponent(Vec, ComponentIndex) VectorGetComponentImpl<ComponentIndex>(Vec)
 
 FORCEINLINE float VectorGetComponentDynamic(const VectorRegister4Float& Vec, uint32 ComponentIndex)
@@ -434,7 +452,11 @@ FORCEINLINE float VectorGetComponentDynamic(const VectorRegister4Float& Vec, uin
 
 FORCEINLINE double VectorGetComponentDynamic(const VectorRegister4Double& Vec, uint32 ComponentIndex)
 {
-	return (((double*)&(Vec))[ComponentIndex]);
+#if !UE_PLATFORM_MATH_USE_AVX
+	return (((double*)&(Vec.XY))[ComponentIndex]);
+#else
+	return (((double*)&(Vec.XYZW))[ComponentIndex]);
+#endif
 }
 
 /**
@@ -518,7 +540,7 @@ FORCEINLINE VectorRegister4Double VectorLoadAligned(const double* Ptr)
 	Result.XY = _mm_load_pd((const double*)(Ptr));
 	Result.ZW = _mm_load_pd((const double*)(Ptr + 2));
 #else
-	// TODO: AVX using unaligned here, until we ensure 32-byte alignment (was previously 16 across most types)
+	// AVX using unaligned here, since we don't ensure 32-byte alignment (not significant on most modern processors)
 	Result = _mm256_loadu_pd(Ptr);
 #endif
 	return Result;
@@ -644,7 +666,7 @@ FORCEINLINE void VectorStoreAligned(const VectorRegister4Double& Vec, double* Ds
 	_mm_store_pd(Dst, Vec.XY);
 	_mm_store_pd(Dst + 2, Vec.ZW);
 #else
-	// TODO: AVX using unaligned store until we verify 32-byte alignment requirement (was 16 for most types)
+	// AVX using unaligned here, since we don't ensure 32-byte alignment (not significant on most modern processors)
 	_mm256_storeu_pd(Dst, Vec);
 #endif
 }
@@ -667,7 +689,7 @@ FORCEINLINE void VectorStoreAlignedStreamed(const VectorRegister4Double& Vec, do
 	_mm_stream_pd(Dst, Vec.XY);
 	_mm_stream_pd(Dst + 2, Vec.ZW);
 #else
-	// TODO: AVX using two 128-bit stores until we verify 32-byte alignment requirement for our data with _mm256_stream_pd (was 16 for most types)
+	// AVX using two 128-bit stores since we don't require 32-byte alignment requirement for our data (so don't use _mm256_stream_pd)
 	_mm_stream_pd(Dst, Vec.XY);
 	_mm_stream_pd(Dst + 2, Vec.ZW);
 #endif
@@ -785,13 +807,13 @@ FORCEINLINE VectorRegister4Double VectorSwizzleTemplate(const VectorRegister4Dou
 	static_assert(Index0 >= 0 && Index0 <= 3 && Index1 >= 0 && Index1 <= 3 && Index2 >= 0 && Index2 <= 3 && Index3 >= 0 && Index3 <= 3);
 
 #if UE_PLATFORM_MATH_USE_AVX_2
-	VectorRegister4Double Result = _mm256_permute4x64_pd(Vec, SHUFFLEMASK(Index0, Index1, Index2, Index3));
+	return _mm256_permute4x64_pd(Vec, SHUFFLEMASK(Index0, Index1, Index2, Index3));
 #else
-	VectorRegister4Double Result;
-	Result.XY = SelectVectorSwizzle2<Index0, Index1>(Vec);
-	Result.ZW = SelectVectorSwizzle2<Index2, Index3>(Vec);
+	return VectorRegister4Double(
+		SelectVectorSwizzle2<Index0, Index1>(Vec),	
+		SelectVectorSwizzle2<Index2, Index3>(Vec)
+	);
 #endif
-	return Result;
 }
 
 // Double Swizzle specializations
@@ -800,7 +822,8 @@ template<> FORCEINLINE VectorRegister4Double VectorSwizzleTemplate<0, 1, 2, 3>(c
 #if UE_PLATFORM_MATH_USE_AVX
 constexpr int PERMUTE_MASK(int A, int B, int C, int D) { return ((A == 1 ? (1 << 0) : 0) | (B == 1 ? (1 << 1) : 0) | (C == 3 ? (1 << 2) : 0) | (D == 3 ? (1 << 3) : 0)); }
 
-// _mm256_permute4x64_pd has a latency of 3, so here are some specializations using instructions which have a latency of 1 but are restricted to in-lane (128 bit) permutes.
+// AVX2 _mm256_permute4x64_pd has a latency of 3, but there are some specializations using instructions which have a latency of 1 but are restricted to in-lane (128 bit) permutes.
+// AVX1 benefits from lower latency instructions here than the shuffling between 128-bit and 256-bit operations of the generic implementation.
 template<> FORCEINLINE VectorRegister4Double VectorSwizzleTemplate<0, 1, 2, 2>(const VectorRegister4Double& Vec) { return _mm256_permute_pd(Vec, PERMUTE_MASK(0, 1, 2, 2)); }
 template<> FORCEINLINE VectorRegister4Double VectorSwizzleTemplate<0, 1, 3, 3>(const VectorRegister4Double& Vec) { return _mm256_permute_pd(Vec, PERMUTE_MASK(0, 1, 3, 3)); }
 template<> FORCEINLINE VectorRegister4Double VectorSwizzleTemplate<0, 1, 3, 2>(const VectorRegister4Double& Vec) { return _mm256_permute_pd(Vec, PERMUTE_MASK(0, 1, 3, 2)); }
@@ -819,9 +842,17 @@ template<> FORCEINLINE VectorRegister4Double VectorSwizzleTemplate<1, 1, 2, 3>(c
 template<> FORCEINLINE VectorRegister4Double VectorSwizzleTemplate<1, 1, 2, 2>(const VectorRegister4Double& Vec) { return _mm256_permute_pd(Vec, PERMUTE_MASK(1, 1, 2, 2)); }
 template<> FORCEINLINE VectorRegister4Double VectorSwizzleTemplate<1, 1, 3, 3>(const VectorRegister4Double& Vec) { return _mm256_permute_pd(Vec, PERMUTE_MASK(1, 1, 3, 3)); }
 template<> FORCEINLINE VectorRegister4Double VectorSwizzleTemplate<1, 1, 3, 2>(const VectorRegister4Double& Vec) { return _mm256_permute_pd(Vec, PERMUTE_MASK(1, 1, 3, 2)); }
+#endif // UE_PLATFORM_MATH_USE_AVX
+
+#if (UE_PLATFORM_MATH_USE_AVX && !UE_PLATFORM_MATH_USE_AVX_2)
+// These don't really bring any benefit on AVX2 compared to _mm256_permute4x64_pd, but avoid (extract + move + set) on AVX1.
+template<> FORCEINLINE VectorRegister4Double VectorSwizzleTemplate<2, 3, 0, 1>(const VectorRegister4Double& Vec) { return _mm256_permute2f128_pd(Vec, Vec, (0x01) | (0x02 << 4)); }
+template<> FORCEINLINE VectorRegister4Double VectorSwizzleTemplate<0, 1, 0, 1>(const VectorRegister4Double& Vec) { return _mm256_permute2f128_pd(Vec, Vec, (0x00) | (0x02 << 4)); }
+template<> FORCEINLINE VectorRegister4Double VectorSwizzleTemplate<2, 3, 2, 3>(const VectorRegister4Double& Vec) { return _mm256_permute2f128_pd(Vec, Vec, (0x01) | (0x03 << 4)); }
 #endif
 
 
+
 // Double replicate (2 doubles)
 template <int Index>
 FORCEINLINE VectorRegister2Double VectorReplicateImpl2(const VectorRegister2Double& Vec)
@@ -833,19 +864,15 @@ FORCEINLINE VectorRegister2Double VectorReplicateImpl2(const VectorRegister2Doub
 template <int Index, typename std::enable_if < (Index <= 1), bool >::type = true >
 FORCEINLINE VectorRegister4Double VectorReplicateImpl(const VectorRegister4Double& Vec)
 {
-	VectorRegister4Double Result;
-	Result.XY = VectorReplicateImpl2<Index>(Vec.XY);
-	Result.ZW = Result.XY;
-	return Result;
+	VectorRegister2Double Temp = VectorReplicateImpl2<Index>(Vec.XY);
+	return VectorRegister4Double(Temp, Temp);
 }
 
 template <int Index, typename std::enable_if < (Index > 1), bool >::type = true >
 FORCEINLINE VectorRegister4Double VectorReplicateImpl(const VectorRegister4Double& Vec)
 {
-	VectorRegister4Double Result;
-	Result.XY = VectorReplicateImpl2<Index - 2>(Vec.ZW);
-	Result.ZW = Result.XY;
-	return Result;
+	VectorRegister2Double Temp = VectorReplicateImpl2<Index - 2>(Vec.ZW);
+	return VectorRegister4Double(Temp, Temp);
 }
 
 template<int Index>
@@ -854,11 +881,10 @@ FORCEINLINE VectorRegister4Double VectorReplicateTemplate(const VectorRegister4D
 	static_assert(Index >= 0 && Index <= 3);
 
 #if UE_PLATFORM_MATH_USE_AVX_2
-	VectorRegister4Double Result = _mm256_permute4x64_pd(Vec, SHUFFLEMASK(Index, Index, Index, Index));
+	return _mm256_permute4x64_pd(Vec, SHUFFLEMASK(Index, Index, Index, Index));
 #else
-	VectorRegister4Double Result = VectorReplicateImpl<Index>(Vec);
+	return VectorReplicateImpl<Index>(Vec);
 #endif
-	return Result;
 }
 
 // Double shuffle
@@ -866,10 +892,10 @@ template<int Index0, int Index1, int Index2, int Index3>
 FORCEINLINE VectorRegister4Double VectorShuffleTemplate(const VectorRegister4Double& Vec1, const VectorRegister4Double& Vec2)
 {
 	static_assert(Index0 >= 0 && Index0 <= 3 && Index1 >= 0 && Index1 <= 3 && Index2 >= 0 && Index2 <= 3 && Index3 >= 0 && Index3 <= 3);
-	VectorRegister4Double Result;
-	Result.XY = SelectVectorSwizzle2<Index0, Index1>(Vec1);
-	Result.ZW = SelectVectorSwizzle2<Index2, Index3>(Vec2);
-	return Result;
+	return VectorRegister4Double(
+		SelectVectorSwizzle2<Index0, Index1>(Vec1),
+		SelectVectorSwizzle2<Index2, Index3>(Vec2)
+	);
 }
 
 // Double Shuffle specializations 
@@ -1185,6 +1211,91 @@ FORCEINLINE VectorRegister4Double VectorDivide(const VectorRegister4Double& Vec1
 	return Result;
 }
 
+namespace SSEVectorHelperFuncs
+{
+	// Computes VectorDot3 but only with the result in the first (X) element of a VectorRegister4Float
+	FORCEINLINE VectorRegister4Float InternalVectorDot3X(const VectorRegister4Float& Vec1, const VectorRegister4Float& Vec2)
+	{
+		// (X, Y, Z, W)
+		VectorRegister4Float Prod = VectorMultiply(Vec1, Vec2);
+		// (Y, Y, W, W)
+		VectorRegister4Float Shuf = VectorSwizzle(Prod, 1, 1, 3, 3); // _mm_movehdup_ps on SSE4.1, shuffle otherwise
+		// (X+Y, ???, ???, ???)
+		VectorRegister4Float Sum = VectorAdd(Prod, Shuf);
+		// (Z, W, Z, W)
+		Shuf = VectorSwizzle(Prod, 2, 3, 2, 3); // _mm_movehl_ps
+		// (X+Y+Z, ???, ???, ???)
+		Sum = VectorAdd(Sum, Shuf);
+		return Sum;
+	}
+
+	// Computes VectorDot3 but only with the result in the first (X) element of a VectorRegister4Double
+	FORCEINLINE VectorRegister4Double InternalVectorDot3X_Full(const VectorRegister4Double& Vec1, const VectorRegister4Double& Vec2)
+	{
+		// (X, Y, Z, W)
+		VectorRegister4Double Prod = VectorMultiply(Vec1, Vec2);
+		// (Y, Y, W, W)
+		VectorRegister4Double Shuf = VectorSwizzle(Prod, 1, 1, 3, 3); // fast in-lane permute on AVX (_mm256_permute_pd)
+		// (X+Y, ???, ???, ???)
+		VectorRegister4Double Sum = VectorAdd(Prod, Shuf);
+		// (Z, W, Z, W)
+		Shuf = VectorSwizzle(Prod, 2, 3, 2, 3); // various specializations exist for this depending on platform
+		// (X+Y+Z, ???, ???, ???)
+		Sum = VectorAdd(Sum, Shuf);
+
+		return Sum;
+	}
+
+	// Computes VectorDot3 but only with the result in the first (X) element of a VectorRegister2Double (half of VectorRegister4Double)
+	FORCEINLINE VectorRegister2Double InternalVectorDot3X_Half(const VectorRegister4Double& Vec1, const VectorRegister4Double& Vec2)
+	{
+		VectorRegister2Double T, A;
+
+		// (X, Y)
+		T = _mm_mul_pd(Vec1.XY, Vec2.XY);
+
+		// (X + Z, Y + W)
+#if UE_PLATFORM_MATH_USE_FMA3
+		A = _mm_fmadd_pd(Vec1.ZW, Vec2.ZW, T);
+#else
+		A = _mm_add_pd(_mm_mul_pd(Vec1.ZW, Vec2.ZW), T);
+#endif UE_PLATFORM_MATH_USE_FMA3
+
+		// (Y, X)  // Reverse of T
+		T = _mm_shuffle_pd(T, T, SHUFFLEMASK2(1, 0));
+
+		// (X + Z + Y, Y + W + X)
+		T = _mm_add_pd(A, T);
+
+		return T;
+	}
+
+} // namespace SSEVectorHelperFuncs
+
+
+/**
+ * Calculates the dot3 product of two vectors and returns a scalar value.
+ *
+ * @param Vec1	1st vector
+ * @param Vec2	2nd vector
+ * @return		d = dot3(Vec1.xyz, Vec2.xyz)
+ */
+FORCEINLINE float VectorDot3Scalar(const VectorRegister4Float& Vec1, const VectorRegister4Float& Vec2)
+{
+	return VectorGetComponent(SSEVectorHelperFuncs::InternalVectorDot3X(Vec1, Vec2), 0);
+}
+
+FORCEINLINE double VectorDot3Scalar(const VectorRegister4Double& Vec1, const VectorRegister4Double& Vec2)
+{
+#if UE_PLATFORM_MATH_USE_AVX
+	return VectorGetComponent(SSEVectorHelperFuncs::InternalVectorDot3X_Full(Vec1, Vec2), 0);
+#else
+	VectorRegister2Double T = SSEVectorHelperFuncs::InternalVectorDot3X_Half(Vec1, Vec2);
+	// Extract first component
+	return _mm_cvtsd_f64(T);
+#endif
+}
+
 /**
  * Calculates the dot3 product of two vectors and returns a vector with the result in all 4 components.
  *
@@ -1192,17 +1303,22 @@ FORCEINLINE VectorRegister4Double VectorDivide(const VectorRegister4Double& Vec1
  * @param Vec2	2nd vector
  * @return		d = dot3(Vec1.xyz, Vec2.xyz), VectorRegister4Float( d, d, d, d )
  */
-FORCEINLINE VectorRegister4Float VectorDot3( const VectorRegister4Float& Vec1, const VectorRegister4Float& Vec2 )
+FORCEINLINE VectorRegister4Float VectorDot3(const VectorRegister4Float& Vec1, const VectorRegister4Float& Vec2)
 {
-	VectorRegister4Float Temp = VectorMultiply( Vec1, Vec2 );
-	return VectorAdd( VectorReplicate(Temp,0), VectorAdd( VectorReplicate(Temp,1), VectorReplicate(Temp,2) ) );
+	return VectorReplicate(SSEVectorHelperFuncs::InternalVectorDot3X(Vec1, Vec2), 0);
 }
 
 FORCEINLINE VectorRegister4Double VectorDot3(const VectorRegister4Double& Vec1, const VectorRegister4Double& Vec2)
 {
-	// TODO: AVX version?
-	VectorRegister4Double Temp = VectorMultiply(Vec1, Vec2);
-	return VectorAdd(VectorReplicate(Temp,0), VectorAdd(VectorReplicate(Temp,1), VectorReplicate(Temp,2)));
+#if UE_PLATFORM_MATH_USE_AVX
+	return VectorReplicate(SSEVectorHelperFuncs::InternalVectorDot3X_Full(Vec1, Vec2), 0);
+#else
+	VectorRegister2Double T = SSEVectorHelperFuncs::InternalVectorDot3X_Half(Vec1, Vec2);
+	// Replicate in half (X,X)
+	T = _mm_shuffle_pd(T, T, SHUFFLEMASK2(0, 0));
+	// Replicate in full (X,X,X,X)
+	return VectorRegister4Double(T, T);
+#endif	
 }
 
 /**
@@ -1215,35 +1331,49 @@ FORCEINLINE VectorRegister4Double VectorDot3(const VectorRegister4Double& Vec1,
 FORCEINLINE VectorRegister4Float VectorDot4( const VectorRegister4Float& Vec1, const VectorRegister4Float& Vec2 )
 {
 	VectorRegister4Float Temp1, Temp2;
-	Temp1 = VectorMultiply( Vec1, Vec2 );
-	Temp2 = VectorSwizzle( Temp1, 2, 3, 0, 1 );	// (Z,W,X,Y).
-	Temp1 = VectorAdd( Temp1, Temp2 );			// (X*X + Z*Z, Y*Y + W*W, Z*Z + X*X, W*W + Y*Y)
-	Temp2 = VectorSwizzle( Temp1, 1, 2, 3, 0 );	// Rotate left 4 bytes (Y,Z,W,X).
-	return VectorAdd( Temp1, Temp2 );			// (X*X + Z*Z + Y*Y + W*W, Y*Y + W*W + Z*Z + X*X, Z*Z + X*X + W*W + Y*Y, W*W + Y*Y + X*X + Z*Z)
+	Temp1 = VectorMultiply(Vec1, Vec2);
+	Temp2 = VectorSwizzle( Temp1, 2, 3, 0, 1);	// (Z,W,X,Y).
+	Temp1 = VectorAdd(Temp1, Temp2);			// (X*X + Z*Z, Y*Y + W*W, Z*Z + X*X, W*W + Y*Y)
+	Temp2 = VectorSwizzle(Temp1, 1, 2, 3, 0);	// Rotate left 1 element (Y,Z,W,X).
+	return VectorAdd(Temp1, Temp2);				// (X*X + Z*Z + Y*Y + W*W, Y*Y + W*W + Z*Z + X*X, Z*Z + X*X + W*W + Y*Y, W*W + Y*Y + X*X + Z*Z)
 }
 
 FORCEINLINE VectorRegister4Double VectorDot4(const VectorRegister4Double& Vec1, const VectorRegister4Double& Vec2)
 {
-	// TODO: AVX version
-	VectorRegister4Double R;
-	VectorRegister2Double T;
-
-	// (X*X, Y*Y)
-	R.XY = _mm_mul_pd(Vec1.XY, Vec2.XY);
-	// (Z*Z, W*W)
-	R.ZW = _mm_mul_pd(Vec1.ZW, Vec2.ZW);
-
-	// (X*X + Z*Z, Y*Y + W*W)
-	T = _mm_add_pd(R.XY, R.ZW);
-
-	// (Y*Y + W*W, X*X + Z*Z)  // Reverse of T
-	R.ZW = _mm_shuffle_pd(T, T, SHUFFLEMASK2(1, 0));
+#if UE_PLATFORM_MATH_USE_AVX
+	// AVX implementation works only in 256-bit registers, avoiding extract/combine ops with 128d intermediates.
+	VectorRegister4Double R, T;
+
+	// (XX, YY, ZZ, WW)
+	R = VectorMultiply(Vec1, Vec2);
+	// (YY, XX, WW, ZZ)
+	T = VectorSwizzle(R, 1, 0, 3, 2); // fast in-lane permute (latency 1)
+	// (XX + YY, YY + XX, ZZ + WW, WW + ZZ)
+	R = VectorAdd(R, T);
+	// (ZZ + WW, WW + ZZ, XX + YY, YY + XX)
+	T = VectorSwizzle(R, 2, 3, 0, 1); // cross-lane permute (latency 3)
+	// (XX + YY + ZZ + WW, YY + XX + WW + ZZ, ZZ + WW + XX + YY, WW + ZZ + YY + XX)
+	return VectorAdd(R, T);
+#else
+	VectorRegister2Double T, A;
+
+	// (X, Y)
+	T = _mm_mul_pd(Vec1.XY, Vec2.XY);
+
+	// (X + Z, Y + W)
+#if UE_PLATFORM_MATH_USE_FMA3
+	A = _mm_fmadd_pd(Vec1.ZW, Vec2.ZW, T);
+#else
+	A = _mm_add_pd(_mm_mul_pd(Vec1.ZW, Vec2.ZW), T);
+#endif UE_PLATFORM_MATH_USE_FMA3
 
-	// (X*X + Z*Z + Y*Y + W*W, Y*Y + W*W + X*X + Z*Z)
-	R.XY = _mm_add_pd(T, R.ZW);
-	R.ZW = R.XY;
+	// (Y + W, X + Z)  // Reverse of A
+	T = _mm_shuffle_pd(A, A, SHUFFLEMASK2(1, 0));
 
-	return R;
+	// (X + Z + Y + W, Y + W + X + Z)
+	T = _mm_add_pd(A, T);
+	return VectorRegister4Double(T, T);
+#endif
 }
 
 /**
@@ -2253,10 +2383,7 @@ FORCEINLINE VectorRegister4Float VectorMergeVecXYZ_VecW(const VectorRegister4Flo
 FORCEINLINE VectorRegister4Double VectorMergeVecXYZ_VecW(const VectorRegister4Double& VecXYZ, const VectorRegister4Double& VecW)
 {
 	//return VectorSelect(GlobalVectorConstants::DoubleXYZMask, VecXYZ, VecW);
-	VectorRegister4Double Result;
-	Result.XY = VecXYZ.XY;
-	Result.ZW = _mm_move_sd(VecW.ZW, VecXYZ.ZW);
-	return Result;
+	return VectorRegister4Double(VecXYZ.XY, _mm_move_sd(VecW.ZW, VecXYZ.ZW));
 }
 
 /**
@@ -2483,9 +2610,9 @@ FORCEINLINE int VectorAnyGreaterThan(const VectorRegister4Double& Vec1, const Ve
 FORCEINLINE VectorRegister4Float VectorQuaternionMultiply2( const VectorRegister4Float& Quat1, const VectorRegister4Float& Quat2 )
 {
 	VectorRegister4Float Result = VectorMultiply(VectorReplicate(Quat1, 3), Quat2);
-	Result = VectorMultiplyAdd( VectorMultiply(VectorReplicate(Quat1, 0), VectorSwizzle(Quat2, 3,2,1,0)), GlobalVectorConstants::QMULTI_SIGN_MASK0, Result);
-	Result = VectorMultiplyAdd( VectorMultiply(VectorReplicate(Quat1, 1), VectorSwizzle(Quat2, 2,3,0,1)), GlobalVectorConstants::QMULTI_SIGN_MASK1, Result);
-	Result = VectorMultiplyAdd( VectorMultiply(VectorReplicate(Quat1, 2), VectorSwizzle(Quat2, 1,0,3,2)), GlobalVectorConstants::QMULTI_SIGN_MASK2, Result);
+	Result = VectorMultiplyAdd( VectorMultiply(VectorReplicate(Quat1, 0), VectorSwizzle(Quat2, 3, 2, 1, 0)), GlobalVectorConstants::QMULTI_SIGN_MASK0, Result);
+	Result = VectorMultiplyAdd( VectorMultiply(VectorReplicate(Quat1, 1), VectorSwizzle(Quat2, 2, 3, 0, 1)), GlobalVectorConstants::QMULTI_SIGN_MASK1, Result);
+	Result = VectorMultiplyAdd( VectorMultiply(VectorReplicate(Quat1, 2), VectorSwizzle(Quat2, 1, 0, 3, 2)), GlobalVectorConstants::QMULTI_SIGN_MASK2, Result);
 
 	return Result;
 }
diff --git a/Engine/Source/Runtime/Core/Public/Math/UnrealMathVectorCommon.h b/Engine/Source/Runtime/Core/Public/Math/UnrealMathVectorCommon.h
--- a/Engine/Source/Runtime/Core/Public/Math/UnrealMathVectorCommon.h
+++ b/Engine/Source/Runtime/Core/Public/Math/UnrealMathVectorCommon.h
@@ -5,11 +5,6 @@
 struct FGenericPlatformMath;
 struct FMath;
 
-FORCEINLINE bool VectorIsAligned(const void *Ptr, uint32 Alignment = SIMD_ALIGNMENT)
-{
-	return !(PTRINT(Ptr) & (Alignment - 1));
-}
-
 // Overload to resolve compiler ambiguity for things like MakeVectorRegister(V.X, V.Y, V.Z, 0.f) when V is a double type.
 FORCEINLINE VectorRegister4Double MakeVectorRegister(double X, double Y, double Z, float W)
 {
@@ -788,6 +783,11 @@ FORCEINLINE VectorRegister4Double VectorSetDouble(double X, double Y, double Z,
 	return MakeVectorRegisterDouble(X, Y, Z, W);
 }
 
+FORCEINLINE VectorRegister4Double VectorSetDouble1(double D)
+{
+	return VectorSetFloat1(D);
+}
+
 FORCEINLINE void VectorStoreDouble3(const VectorRegister4Double& Vec, double* Ptr)
 {
 	VectorStoreFloat3(Vec, Ptr);
