#ue5 - LWC: FBoxSphereBounds::TransformBy() uses new more efficient VectorDot3Scalar() to avoid shuffle + extract.

diff --git a/Engine/Source/Runtime/Core/Private/Math/BoxSphereBounds.cpp b/Engine/Source/Runtime/Core/Private/Math/BoxSphereBounds.cpp
--- a/Engine/Source/Runtime/Core/Private/Math/BoxSphereBounds.cpp
+++ b/Engine/Source/Runtime/Core/Private/Math/BoxSphereBounds.cpp
@@ -51,7 +51,7 @@ FBoxSphereBounds FBoxSphereBounds::TransformBy(const FMatrix& M) const
 	Result.SphereRadius = FMath::Sqrt(VectorGetComponent( MaxRadius, 0) ) * SphereRadius;
 
 	// For non-uniform scaling, computing sphere radius from a box results in a smaller sphere.
-	FVector::FReal const BoxExtentMagnitude = FMath::Sqrt(VectorGetComponent(VectorDot3(NewExtent, NewExtent), 0));
+	FVector::FReal const BoxExtentMagnitude = FMath::Sqrt(VectorDot3Scalar(NewExtent, NewExtent));
 	Result.SphereRadius = FMath::Min<FVector::FReal>(Result.SphereRadius, BoxExtentMagnitude);
 
 	Result.DiagnosticCheckNaN();
