U5 - partially revert ISInParallelRenderingThreadChanges.

diff --git a/Engine/Source/Runtime/Core/Private/HAL/ThreadingBase.cpp b/Engine/Source/Runtime/Core/Private/HAL/ThreadingBase.cpp
--- a/Engine/Source/Runtime/Core/Private/HAL/ThreadingBase.cpp
+++ b/Engine/Source/Runtime/Core/Private/HAL/ThreadingBase.cpp
@@ -302,7 +302,7 @@ PRAGMA_DISABLE_DEPRECATION_WARNINGS
 	if (!LowLevelTasks::FScheduler::IsBusyWaiting())
 	{
 		const uint32 CurrentThreadId = FPlatformTLS::GetCurrentThreadId();
-		bool oldValue = ((GRenderThreadId == 0) || bLocalIsLoadingThreadSuspended) ?  CurrentThreadId == GGameThreadId : CurrentThreadId != GGameThreadId;
+		bool oldValue = ((GRenderThreadId == 0) || bLocalIsLoadingThreadSuspended) ?  true : CurrentThreadId != GGameThreadId;
 		ensureMsgf(oldValue == newValue, TEXT("oldValue(%i) newValue(%i) If this check fails make sure that there is a FTaskTagScope(ETaskTag::EParallelRenderingThread) as deep as possible on the current callstack, you can see the current value in ActiveNamedThreads(%x), GRenderingThread(%x), GIsRenderingThreadSuspended(%d)"), oldValue, newValue, FTaskTagScope::GetCurrentTag(), GRenderingThread, bLocalIsLoadingThreadSuspended);
 		newValue = oldValue;
 	}
