#ue5 - LWC: Fix build break (missing comment //)

diff --git a/Engine/Source/Runtime/Core/Public/Math/UnrealMathSSE.h b/Engine/Source/Runtime/Core/Public/Math/UnrealMathSSE.h
--- a/Engine/Source/Runtime/Core/Public/Math/UnrealMathSSE.h
+++ b/Engine/Source/Runtime/Core/Public/Math/UnrealMathSSE.h
@@ -1259,7 +1259,7 @@ namespace SSEVectorHelperFuncs
 		A = _mm_fmadd_pd(Vec1.ZW, Vec2.ZW, T);
 #else
 		A = _mm_add_pd(_mm_mul_pd(Vec1.ZW, Vec2.ZW), T);
-#endif UE_PLATFORM_MATH_USE_FMA3
+#endif // UE_PLATFORM_MATH_USE_FMA3
 
 		// (Y, X)  // Reverse of T
 		T = _mm_shuffle_pd(T, T, SHUFFLEMASK2(1, 0));
@@ -1365,7 +1365,7 @@ FORCEINLINE VectorRegister4Double VectorDot4(const VectorRegister4Double& Vec1,
 	A = _mm_fmadd_pd(Vec1.ZW, Vec2.ZW, T);
 #else
 	A = _mm_add_pd(_mm_mul_pd(Vec1.ZW, Vec2.ZW), T);
-#endif UE_PLATFORM_MATH_USE_FMA3
+#endif // UE_PLATFORM_MATH_USE_FMA3
 
 	// (Y + W, X + Z)  // Reverse of A
 	T = _mm_shuffle_pd(A, A, SHUFFLEMASK2(1, 0));
