Added some new "trace." console commands;

trace.pause/resume - pauses and resumes the tracing of events
trace.send <host> - send a trace over the wire to "host".

diff --git a/Engine/Source/Runtime/Core/Private/ProfilingDebugging/TraceAuxiliary.cpp b/Engine/Source/Runtime/Core/Private/ProfilingDebugging/TraceAuxiliary.cpp
--- a/Engine/Source/Runtime/Core/Private/ProfilingDebugging/TraceAuxiliary.cpp
+++ b/Engine/Source/Runtime/Core/Private/ProfilingDebugging/TraceAuxiliary.cpp
@@ -348,19 +348,8 @@ void FTraceAuxiliaryImpl::ReadChannels(T&& Callback) const
 
 
 ////////////////////////////////////////////////////////////////////////////////
-static void TraceAuxiliaryStart(const TArray<FString>& Args)
+static void TraceAuxiliaryConnectEpilogue()
 {
-	if (Args.Num() > 0)
-	{
-		GTraceAuxiliary.AddChannels(*(Args[0]));
-	}
-
-	if (!GTraceAuxiliary.Connect(ETraceConnectType::File, nullptr))
-	{
-		UE_LOG(LogConsoleResponse, Warning, TEXT("Failed to start tracing to a file"));
-		return;
-	}
-
 	// It is possible that something outside of TraceAux's world view has called
 	// UE::Trace::SendTo/WriteTo(). A plugin that has created its own store for
 	// example. There's not really much that can be done about that here (tracing
@@ -388,6 +377,46 @@ static void TraceAuxiliaryStart(const TArray<FString>& Args)
 	UE_LOG(LogConsoleResponse, Log, TEXT("Trace channels: %s"), *Channels);
 }
 
+////////////////////////////////////////////////////////////////////////////////
+static void TraceAuxiliarySend(const TArray<FString>& Args)
+{
+	if (Args.Num() < 1)
+	{
+		UE_LOG(LogConsoleResponse, Warning, TEXT("No host name given; Trace.Send <Host> [ChannelSet]"));
+		return;
+	}
+
+	if (Args.Num() > 1)
+	{
+		GTraceAuxiliary.AddChannels(*(Args[1]));
+	}
+
+	if (!GTraceAuxiliary.Connect(ETraceConnectType::Network, *Args[0]))
+	{
+		UE_LOG(LogConsoleResponse, Warning, TEXT("Failed to start tracing to '%s'"), *Args[0]);
+		return;
+	}
+
+	TraceAuxiliaryConnectEpilogue();
+}
+
+////////////////////////////////////////////////////////////////////////////////
+static void TraceAuxiliaryStart(const TArray<FString>& Args)
+{
+	if (Args.Num() > 0)
+	{
+		GTraceAuxiliary.AddChannels(*(Args[0]));
+	}
+
+	if (!GTraceAuxiliary.Connect(ETraceConnectType::File, nullptr))
+	{
+		UE_LOG(LogConsoleResponse, Warning, TEXT("Failed to start tracing to a file"));
+		return;
+	}
+
+	TraceAuxiliaryConnectEpilogue();
+}
+
 ////////////////////////////////////////////////////////////////////////////////
 static void TraceAuxiliaryStop()
 {
@@ -395,6 +424,27 @@ static void TraceAuxiliaryStop()
 	GTraceAuxiliary.Stop();
 }
 
+////////////////////////////////////////////////////////////////////////////////
+static void TraceAuxiliaryPause()
+{
+	UE_LOG(LogConsoleResponse, Log, TEXT("Tracing paused"));
+	GTraceAuxiliary.DisableChannels();
+}
+
+////////////////////////////////////////////////////////////////////////////////
+static void TraceAuxiliaryResume()
+{
+	UE_LOG(LogConsoleResponse, Log, TEXT("Tracing resumed"));
+	GTraceAuxiliary.EnableChannels();
+}
+
+////////////////////////////////////////////////////////////////////////////////
+static FAutoConsoleCommand TraceAuxiliarySendCmd(
+	TEXT("Trace.Send"),
+	TEXT("Send trace data to the trace store; Trace.Send <Host> [ChannelSet]"),
+	FConsoleCommandWithArgsDelegate::CreateStatic(TraceAuxiliarySend)
+);
+
 ////////////////////////////////////////////////////////////////////////////////
 static FAutoConsoleCommand TraceAuxiliaryStartCmd(
 	TEXT("Trace.Start"),
@@ -412,6 +462,20 @@ static FAutoConsoleCommand TraceAuxiliaryStopCmd(
 	FConsoleCommandDelegate::CreateStatic(TraceAuxiliaryStop)
 );
 
+////////////////////////////////////////////////////////////////////////////////
+static FAutoConsoleCommand TraceAuxiliaryPauseCmd(
+	TEXT("Trace.Pause"),
+	TEXT("Pauses all trace channels currently sending events"),
+	FConsoleCommandDelegate::CreateStatic(TraceAuxiliaryPause)
+);
+
+////////////////////////////////////////////////////////////////////////////////
+static FAutoConsoleCommand TraceAuxiliaryResumeCmd(
+	TEXT("Trace.Resume"),
+	TEXT("Resume tracing that was previously paused"),
+	FConsoleCommandDelegate::CreateStatic(TraceAuxiliaryResume)
+);
+
 #endif // UE_TRACE_ENABLED
 
 
