Do not allocate space for an attachment on the trace event as this event does not attach anything.

diff --git a/Engine/Source/Runtime/Core/Private/ProfilingDebugging/TagTrace.cpp b/Engine/Source/Runtime/Core/Private/ProfilingDebugging/TagTrace.cpp
--- a/Engine/Source/Runtime/Core/Private/ProfilingDebugging/TagTrace.cpp
+++ b/Engine/Source/Runtime/Core/Private/ProfilingDebugging/TagTrace.cpp
@@ -204,7 +204,7 @@ void FTagTrace::AnnounceGenericTags()
 	#define TRACE_TAG_SPEC(Enum,Str,Stat,Group,ParentTag)\
 	{\
 		const uint32 DisplayLen = FCStringAnsi::Strlen(Str);\
-		UE_TRACE_LOG(Memory, TagSpec, MemAllocChannel, DisplayLen * sizeof(ANSICHAR))\
+		UE_TRACE_LOG(Memory, TagSpec, MemAllocChannel)\
 			<< TagSpec.Tag((int32) ELLMTag::Enum)\
 			<< TagSpec.Parent((int32) ParentTag)\
 			<< TagSpec.Display(Str, DisplayLen);\
