Add SetAndHold command to csv profiler custom stats to allow infrequently updated statistics to have their value persisted.

diff --git a/Engine/Source/Runtime/Core/Private/ProfilingDebugging/CsvProfiler.cpp b/Engine/Source/Runtime/Core/Private/ProfilingDebugging/CsvProfiler.cpp
--- a/Engine/Source/Runtime/Core/Private/ProfilingDebugging/CsvProfiler.cpp
+++ b/Engine/Source/Runtime/Core/Private/ProfilingDebugging/CsvProfiler.cpp
@@ -1397,7 +1397,7 @@ struct FCsvStatSeries
 		CustomStatFloat
 	};
 
-	FCsvStatSeries(EType InSeriesType, const FCsvStatID& InStatID, FCsvStreamWriter* InWriter, FCsvStatRegister& StatRegister, const FString& ThreadName);
+	FCsvStatSeries(EType InSeriesType, const FCsvStatID& InStatID, FCsvStreamWriter* InWriter, FCsvStatRegister& StatRegister, const FString& ThreadName, uint32 ThreadId);
 	void FlushIfDirty();
 
 	void SetTimerValue(uint32 DataFrameNumber, uint64 ElapsedCycles)
@@ -1432,6 +1432,7 @@ struct FCsvStatSeries
 				FlushIfDirty();
 			}
 
+			bHold = Op == ECsvCustomStatOp::SetAndHold;
 			// The first op in a frame is always a set. Otherwise min/max don't work
 			Op = ECsvCustomStatOp::Set;
 			CurrentWriteFrameNumber = DataFrameNumber;
@@ -1471,6 +1472,7 @@ struct FCsvStatSeries
 				FlushIfDirty();
 			}
 
+			bHold = Op == ECsvCustomStatOp::SetAndHold;
 			// The first op in a frame is always a set. Otherwise min/max don't work
 			Op = ECsvCustomStatOp::Set;
 			CurrentWriteFrameNumber = DataFrameNumber;
@@ -1519,8 +1521,10 @@ struct FCsvStatSeries
 	FCsvStreamWriter* Writer;
 
 	int32 ColumnIndex;
+	uint32 ThreadId;
 
 	bool bDirty;
+	bool bHold;
 };
 
 struct FCsvProcessThreadDataStats
@@ -1540,7 +1544,8 @@ class FCsvStreamWriter
 {
 	struct FCsvRow
 	{
-		TArray<FCsvStatSeriesValue> Values;
+		TSparseArray<FCsvStatSeriesValue> Values;
+		TBitArray<> ColumnsToHold;
 		TArray<FCsvProcessedEvent> Events;
 
 		inline uint64 GetAllocatedSize() const
@@ -1575,13 +1580,14 @@ class FCsvStreamWriter
 	uint32 RenderThreadId;
 	uint32 RHIThreadId;
 
+	TMap<FString, FCsvStatSeries*> SeriesByName;
 public:
 	FCsvStreamWriter(const TSharedRef<FArchive>& InOutputFile, bool bInContinuousWrites, int32 InBufferSize, bool bInCompressOutput, uint32 RenderThreadId, uint32 RHIThreadId);
 	~FCsvStreamWriter();
 
 	void AddSeries(FCsvStatSeries* Series);
 
-	void PushValue(FCsvStatSeries* Series, int64 FrameNumber, const FCsvStatSeriesValue& Value);
+	void PushValue(FCsvStatSeries* Series, int64 FrameNumber, const FCsvStatSeriesValue& Value, bool bHold=false, bool bAdvanceWriteFrame=true);
 	void PushEvent(const FCsvProcessedEvent& Event);
 
 	void FinalizeNextRow();
@@ -1592,13 +1598,15 @@ public:
 	inline uint64 GetAllocatedSize() const;
 };
 
-FCsvStatSeries::FCsvStatSeries(EType InSeriesType, const FCsvStatID& InStatID, FCsvStreamWriter* InWriter, FCsvStatRegister& StatRegister, const FString& ThreadName)
+FCsvStatSeries::FCsvStatSeries(EType InSeriesType, const FCsvStatID& InStatID, FCsvStreamWriter* InWriter, FCsvStatRegister& StatRegister, const FString& ThreadName, uint32 InThreadId)
 	: StatID(InStatID)
 	, SeriesType(InSeriesType)
 	, CurrentWriteFrameNumber(-1)
 	, Writer(InWriter)
 	, ColumnIndex(-1)
+	, ThreadId(InThreadId)
 	, bDirty(false)
+	, bHold(false)
 {
 	CurrentValue.AsTimerCycles = 0;
 
@@ -1645,7 +1653,7 @@ void FCsvStatSeries::FlushIfDirty()
 			Value.Value.AsFloat = CurrentValue.AsFloatValue;
 			break;
 		}
-		Writer->PushValue(this, CurrentWriteFrameNumber, Value);
+		Writer->PushValue(this, CurrentWriteFrameNumber, Value, bHold);
 		CurrentValue.AsTimerCycles = 0;
 		bDirty = false;
 	}
@@ -1972,7 +1980,7 @@ private:
 		}
 		if (StatSeriesArray[StatIndex] == nullptr)
 		{
-			Series = new FCsvStatSeries(SeriesType, StatIndex, Writer, StatRegister, ThreadData->ThreadName);
+			Series = new FCsvStatSeries(SeriesType, StatIndex, Writer, StatRegister, ThreadData->ThreadName, ThreadData->ThreadId);
 			StatSeriesArray[StatIndex] = Series;
 		}
 		else
@@ -2013,21 +2021,30 @@ void FCsvStreamWriter::AddSeries(FCsvStatSeries* Series)
 	AllSeries.Add(Series);
 }
 
-void FCsvStreamWriter::PushValue(FCsvStatSeries* Series, int64 FrameNumber, const FCsvStatSeriesValue& Value)
+void FCsvStreamWriter::PushValue(FCsvStatSeries* Series, int64 FrameNumber, const FCsvStatSeriesValue& Value, bool bHold, bool bAdvanceWriteFrame)
 {
 	check(Series->ColumnIndex != -1);
 
-	WriteFrameIndex = FMath::Max(FrameNumber, WriteFrameIndex);
+	if (bAdvanceWriteFrame)
+	{
+		WriteFrameIndex = FMath::Max(FrameNumber, WriteFrameIndex);
+	}
 
 	FCsvRow& Row = Rows.FindOrAdd(FrameNumber);
-
-	// Ensure the row is large enough to hold every series
-	if (Row.Values.Num() < AllSeries.Num())
+	if (!Row.Values.IsValidIndex(Series->ColumnIndex))
+	{
+		Row.Values.Insert(Series->ColumnIndex, Value);
+	}
+	else
 	{
-		Row.Values.SetNumZeroed(AllSeries.Num(), false);
+		Row.Values[Series->ColumnIndex] = Value;
 	}
 
-	Row.Values[Series->ColumnIndex] = Value;
+	if (Row.ColumnsToHold.Num() < AllSeries.Num())
+	{
+		Row.ColumnsToHold.Add(false, AllSeries.Num() - Row.ColumnsToHold.Num());
+	}
+	Row.ColumnsToHold[Series->ColumnIndex] = bHold;
 }
 
 void FCsvStreamWriter::PushEvent(const FCsvProcessedEvent& Event)
@@ -2055,6 +2072,7 @@ void FCsvStreamWriter::FinalizeNextRow()
 
 	// Don't remove yet. Flushing series may modify this row
 	FCsvRow* Row = Rows.Find(ReadFrameIndex);
+	FCsvRow* NextRow = Rows.Find(ReadFrameIndex + 1);
 	if (Row)
 	{
 		if (Row->Events.Num() > 0)
@@ -2075,6 +2093,12 @@ void FCsvStreamWriter::FinalizeNextRow()
 			Stream.WriteEmptyString();
 		}
 
+		// Simplify conditional for ColumnsToHold 
+		if (Row->ColumnsToHold.Num() < AllSeries.Num())
+		{
+			Row->ColumnsToHold.Add(false, AllSeries.Num() - Row->ColumnsToHold.Num());
+		}
+
 		for (FCsvStatSeries* Series : AllSeries)
 		{
 			// Stat values are held in the series until a new value arrives.
@@ -2094,6 +2118,15 @@ void FCsvStreamWriter::FinalizeNextRow()
 				{
 					Stream.WriteValue(Value.Value.AsFloat);
 				}
+
+				if (Row->ColumnsToHold[Series->ColumnIndex])
+				{
+					// Propagate this value to the next row if it's not set already by data from a newer frame 
+					if (NextRow == nullptr || !NextRow->Values.IsValidIndex(Series->ColumnIndex))
+					{						
+						PushValue(Series, ReadFrameIndex + 1, Value, true, false);
+					}
+				}
 			}
 			else
 			{
diff --git a/Engine/Source/Runtime/Core/Public/ProfilingDebugging/CsvProfiler.h b/Engine/Source/Runtime/Core/Public/ProfilingDebugging/CsvProfiler.h
--- a/Engine/Source/Runtime/Core/Public/ProfilingDebugging/CsvProfiler.h
+++ b/Engine/Source/Runtime/Core/Public/ProfilingDebugging/CsvProfiler.h
@@ -146,6 +146,7 @@ class FName;
 enum class ECsvCustomStatOp : uint8
 {
 	Set,
+	SetAndHold,
 	Min,
 	Max,
 	Accumulate,
