CompactBinary: Renamed CompactBinaryAttachment to ObjectAttachment

This is a purely mechanical change to update names and comments. There is no behavioral difference in this change. A subsequent change will modify the behavior to match the new name.

#rnx

diff --git a/Engine/Source/Runtime/Core/Private/Serialization/CompactBinary.cpp b/Engine/Source/Runtime/Core/Private/Serialization/CompactBinary.cpp
--- a/Engine/Source/Runtime/Core/Private/Serialization/CompactBinary.cpp
+++ b/Engine/Source/Runtime/Core/Private/Serialization/CompactBinary.cpp
@@ -66,8 +66,8 @@ void FCbFieldType::StaticAssertTypeConstants()
 	static_assert(!(BoolMask & (BoolBase ^ ECbFieldType::BoolTrue)), "BoolMask or BoolBase is invalid!");
 	static_assert(TypeMask == (BoolMask | (BoolBase ^ ECbFieldType::BoolTrue)), "BoolMask or BoolBase is invalid!");
 
-	static_assert(AttachmentBase == ECbFieldType::CompactBinaryAttachment, "AttachmentBase is invalid!");
-	static_assert((AttachmentMask & (AllFlags | ECbFieldType::BinaryAttachment)) == ECbFieldType::CompactBinaryAttachment, "AttachmentMask is invalid!");
+	static_assert(AttachmentBase == ECbFieldType::ObjectAttachment, "AttachmentBase is invalid!");
+	static_assert((AttachmentMask & (AllFlags | ECbFieldType::BinaryAttachment)) == ECbFieldType::ObjectAttachment, "AttachmentMask is invalid!");
 	static_assert(!(AttachmentMask & (AttachmentBase ^ ECbFieldType::BinaryAttachment)), "AttachmentMask or AttachmentBase is invalid!");
 	static_assert(TypeMask == (AttachmentMask | (AttachmentBase ^ ECbFieldType::BinaryAttachment)), "AttachmentMask or AttachmentBase is invalid!");
 }
@@ -274,9 +274,9 @@ bool FCbFieldView::AsBool(const bool bDefault)
 	return (uint8(bIsBool) & uint8(LocalType) & 1) | ((!bIsBool) & bDefault);
 }
 
-FIoHash FCbFieldView::AsCompactBinaryAttachment(const FIoHash& Default)
+FIoHash FCbFieldView::AsObjectAttachment(const FIoHash& Default)
 {
-	if (FCbFieldType::IsCompactBinaryAttachment(Type))
+	if (FCbFieldType::IsObjectAttachment(Type))
 	{
 		Error = ECbFieldError::None;
 		return FIoHash(*static_cast<const FIoHash::ByteArray*>(Payload));
@@ -532,7 +532,7 @@ uint64 FCbFieldView::GetPayloadSize() const
 	case ECbFieldType::BoolFalse:
 	case ECbFieldType::BoolTrue:
 		return 0;
-	case ECbFieldType::CompactBinaryAttachment:
+	case ECbFieldType::ObjectAttachment:
 	case ECbFieldType::BinaryAttachment:
 	case ECbFieldType::Hash:
 		return 20;
@@ -597,7 +597,7 @@ void FCbFieldView::IterateAttachments(FCbFieldVisitor Visitor) const
 	case ECbFieldType::Array:
 	case ECbFieldType::UniformArray:
 		return FCbArrayView::FromFieldNoCheck(*this).IterateAttachments(Visitor);
-	case ECbFieldType::CompactBinaryAttachment:
+	case ECbFieldType::ObjectAttachment:
 	case ECbFieldType::BinaryAttachment:
 		return Visitor(*this);
 	default:
diff --git a/Engine/Source/Runtime/Core/Private/Serialization/CompactBinaryPackage.cpp b/Engine/Source/Runtime/Core/Private/Serialization/CompactBinaryPackage.cpp
--- a/Engine/Source/Runtime/Core/Private/Serialization/CompactBinaryPackage.cpp
+++ b/Engine/Source/Runtime/Core/Private/Serialization/CompactBinaryPackage.cpp
@@ -95,7 +95,7 @@ void FCbAttachment::Load(FCbFieldIterator& Fields)
 		++Fields;
 		Hash = Fields.AsAttachment();
 		checkf(!Fields.HasError(), TEXT("Attachments must be a non-empty binary value with a content hash."));
-		if (Fields.IsCompactBinaryAttachment())
+		if (Fields.IsObjectAttachment())
 		{
 			CompactBinary = FCbFieldViewIterator::MakeRange(Buffer);
 		}
@@ -129,7 +129,7 @@ void FCbAttachment::Load(FArchive& Ar, FCbBufferAllocator Allocator)
 			});
 		Hash = HashField.AsAttachment();
 		checkf(!HashField.HasError(), TEXT("Attachments must be a non-empty binary value with a content hash."));
-		if (HashField.IsCompactBinaryAttachment())
+		if (HashField.IsObjectAttachment())
 		{
 			CompactBinary = FCbFieldViewIterator::MakeRange(Buffer);
 		}
@@ -155,7 +155,7 @@ void FCbAttachment::Save(FCbWriter& Writer) const
 		{
 			Writer.AddBinary(AsBinaryView());
 		}
-		Writer.AddCompactBinaryAttachment(Hash);
+		Writer.AddObjectAttachment(Hash);
 	}
 	else if (Buffer.GetSize())
 	{
@@ -258,7 +258,7 @@ void FCbPackage::GatherAttachments(const FCbFieldViewIterator& Fields, FAttachme
 			const FIoHash& Hash = Field.AsAttachment();
 			if (FSharedBuffer Buffer = Resolver(Hash))
 			{
-				if (Field.IsCompactBinaryAttachment())
+				if (Field.IsObjectAttachment())
 				{
 					AddAttachment(FCbAttachment(FCbFieldIterator::MakeRange(MoveTemp(Buffer)), Hash), &Resolver);
 				}
@@ -294,8 +294,8 @@ void FCbPackage::Load(FCbFieldIterator& Fields)
 			++Fields;
 			if (Object)
 			{
-				ObjectHash = Fields.AsCompactBinaryAttachment();
-				checkf(!Fields.HasError(), TEXT("Object must be followed by a CompactBinaryAttachment with the object hash."));
+				ObjectHash = Fields.AsObjectAttachment();
+				checkf(!Fields.HasError(), TEXT("Object must be followed by an ObjectAttachment with the object hash."));
 				++Fields;
 			}
 			else
@@ -331,7 +331,7 @@ void FCbPackage::Load(FArchive& Ar, FCbBufferAllocator Allocator)
 				FCbField HashField = LoadCompactBinary(Ar, StackAllocator);
 				const FIoHash& Hash = HashField.AsAttachment();
 				checkf(!HashField.HasError(), TEXT("Attachments must be a non-empty binary value with a content hash."));
-				if (HashField.IsCompactBinaryAttachment())
+				if (HashField.IsObjectAttachment())
 				{
 					AddAttachment(FCbAttachment(FCbFieldIterator::MakeRange(MoveTemp(Buffer)), Hash));
 				}
@@ -349,8 +349,8 @@ void FCbPackage::Load(FArchive& Ar, FCbBufferAllocator Allocator)
 			if (Object)
 			{
 				FCbField HashField = LoadCompactBinary(Ar, StackAllocator);
-				ObjectHash = HashField.AsCompactBinaryAttachment();
-				checkf(!HashField.HasError(), TEXT("Object must be followed by a CompactBinaryAttachment with the object hash."));
+				ObjectHash = HashField.AsObjectAttachment();
+				checkf(!HashField.HasError(), TEXT("Object must be followed by an ObjectAttachment with the object hash."));
 			}
 			else
 			{
@@ -365,7 +365,7 @@ void FCbPackage::Save(FCbWriter& Writer) const
 	if (Object)
 	{
 		Writer.AddObject(Object);
-		Writer.AddCompactBinaryAttachment(ObjectHash);
+		Writer.AddObjectAttachment(ObjectHash);
 	}
 	for (const FCbAttachment& Attachment : Attachments)
 	{
diff --git a/Engine/Source/Runtime/Core/Private/Serialization/CompactBinarySerialization.cpp b/Engine/Source/Runtime/Core/Private/Serialization/CompactBinarySerialization.cpp
--- a/Engine/Source/Runtime/Core/Private/Serialization/CompactBinarySerialization.cpp
+++ b/Engine/Source/Runtime/Core/Private/Serialization/CompactBinarySerialization.cpp
@@ -61,7 +61,7 @@ bool TryMeasureCompactBinary(FMemoryView View, ECbFieldType& OutType, uint64& Ou
 	case ECbFieldType::BoolFalse:
 	case ECbFieldType::BoolTrue:
 		break;
-	case ECbFieldType::CompactBinaryAttachment:
+	case ECbFieldType::ObjectAttachment:
 	case ECbFieldType::BinaryAttachment:
 	case ECbFieldType::Hash:
 		FixedSize = 20;
diff --git a/Engine/Source/Runtime/Core/Private/Serialization/CompactBinaryValidation.cpp b/Engine/Source/Runtime/Core/Private/Serialization/CompactBinaryValidation.cpp
--- a/Engine/Source/Runtime/Core/Private/Serialization/CompactBinaryValidation.cpp
+++ b/Engine/Source/Runtime/Core/Private/Serialization/CompactBinaryValidation.cpp
@@ -346,7 +346,7 @@ static FCbFieldView ValidateCbField(FMemoryView& View, ECbValidateMode Mode, ECb
 	case ECbFieldType::Float64:
 		ValidateCbFloat64(View, Mode, Error);
 		break;
-	case ECbFieldType::CompactBinaryAttachment:
+	case ECbFieldType::ObjectAttachment:
 	case ECbFieldType::BinaryAttachment:
 	case ECbFieldType::Hash:
 		ValidateCbFixedPayload(View, Mode, Error, 20);
diff --git a/Engine/Source/Runtime/Core/Private/Serialization/CompactBinaryWriter.cpp b/Engine/Source/Runtime/Core/Private/Serialization/CompactBinaryWriter.cpp
--- a/Engine/Source/Runtime/Core/Private/Serialization/CompactBinaryWriter.cpp
+++ b/Engine/Source/Runtime/Core/Private/Serialization/CompactBinaryWriter.cpp
@@ -485,11 +485,11 @@ void FCbWriter::AddBool(const bool bValue)
 	EndField(bValue ? ECbFieldType::BoolTrue : ECbFieldType::BoolFalse);
 }
 
-void FCbWriter::AddCompactBinaryAttachment(const FIoHash& Value)
+void FCbWriter::AddObjectAttachment(const FIoHash& Value)
 {
 	BeginField();
 	Data.Append(Value.GetBytes(), sizeof(decltype(Value.GetBytes())));
-	EndField(ECbFieldType::CompactBinaryAttachment);
+	EndField(ECbFieldType::ObjectAttachment);
 }
 
 void FCbWriter::AddBinaryAttachment(const FIoHash& Value)
diff --git a/Engine/Source/Runtime/Core/Private/Tests/Serialization/CompactBinaryPackageTest.cpp b/Engine/Source/Runtime/Core/Private/Tests/Serialization/CompactBinaryPackageTest.cpp
--- a/Engine/Source/Runtime/Core/Private/Tests/Serialization/CompactBinaryPackageTest.cpp
+++ b/Engine/Source/Runtime/Core/Private/Tests/Serialization/CompactBinaryPackageTest.cpp
@@ -372,7 +372,7 @@ bool FCbPackageTest::RunTest(const FString& Parameters)
 		TCbWriter<256> Writer;
 		Writer.SetName("Level3");
 		Writer.BeginArray();
-		Writer.AddCompactBinaryAttachment(Level3Hash);
+		Writer.AddObjectAttachment(Level3Hash);
 		Writer.EndArray();
 		Level2 = Writer.Save().AsArray();
 	}
@@ -382,7 +382,7 @@ bool FCbPackageTest::RunTest(const FString& Parameters)
 	{
 		TCbWriter<256> Writer;
 		Writer.BeginObject();
-		Writer.SetName("Level2").AddCompactBinaryAttachment(Level2Hash);
+		Writer.SetName("Level2").AddObjectAttachment(Level2Hash);
 		Writer.EndObject();
 		Level1 = Writer.Save().AsObject();
 	}
@@ -453,13 +453,13 @@ bool FCbPackageTest::RunTest(const FString& Parameters)
 	{
 		TCbWriter<384> Writer;
 		Writer.AddBinary(Level2.GetBuffer());
-		Writer.AddCompactBinaryAttachment(Level2Hash);
+		Writer.AddObjectAttachment(Level2Hash);
 		Writer.AddBinary(Level4);
 		Writer.AddBinaryAttachment(Level4Hash);
 		Writer.AddObject(Level1);
-		Writer.AddCompactBinaryAttachment(Level1Hash);
+		Writer.AddObjectAttachment(Level1Hash);
 		Writer.AddBinary(Level3.GetBuffer());
-		Writer.AddCompactBinaryAttachment(Level3Hash);
+		Writer.AddObjectAttachment(Level3Hash);
 		Writer.AddNull();
 
 		FCbFieldIterator Fields = Writer.Save();
@@ -507,15 +507,15 @@ bool FCbPackageTest::RunTest(const FString& Parameters)
 		FCbFieldIterator Saved = Writer.Save();
 		TestTrue(TEXT("FCbPackage(OutOfOrder).Save()->Level1"), Saved.AsObject().Equals(Level1));
 		++Saved;
-		TestEqual(TEXT("FCbPackage(OutOfOrder).Save()->Level1Hash"), Saved.AsCompactBinaryAttachment(), Level1Hash);
+		TestEqual(TEXT("FCbPackage(OutOfOrder).Save()->Level1Hash"), Saved.AsObjectAttachment(), Level1Hash);
 		++Saved;
 		TestTrue(TEXT("FCbPackage(OutOfOrder).Save()->Level2"), Saved.AsBinaryView().EqualBytes(Level2.GetView()));
 		++Saved;
-		TestEqual(TEXT("FCbPackage(OutOfOrder).Save()->Level2Hash"), Saved.AsCompactBinaryAttachment(), Level2Hash);
+		TestEqual(TEXT("FCbPackage(OutOfOrder).Save()->Level2Hash"), Saved.AsObjectAttachment(), Level2Hash);
 		++Saved;
 		TestTrue(TEXT("FCbPackage(OutOfOrder).Save()->Level3"), Saved.AsBinaryView().EqualBytes(Level3.GetView()));
 		++Saved;
-		TestEqual(TEXT("FCbPackage(OutOfOrder).Save()->Level3Hash"), Saved.AsCompactBinaryAttachment(), Level3Hash);
+		TestEqual(TEXT("FCbPackage(OutOfOrder).Save()->Level3Hash"), Saved.AsObjectAttachment(), Level3Hash);
 		++Saved;
 		TestTrue(TEXT("FCbPackage(OutOfOrder).Save()->Level4"), Saved.AsBinaryView().EqualBytes(Level4.GetView()));
 		++Saved;
diff --git a/Engine/Source/Runtime/Core/Private/Tests/Serialization/CompactBinarySerializationTest.cpp b/Engine/Source/Runtime/Core/Private/Tests/Serialization/CompactBinarySerializationTest.cpp
--- a/Engine/Source/Runtime/Core/Private/Tests/Serialization/CompactBinarySerializationTest.cpp
+++ b/Engine/Source/Runtime/Core/Private/Tests/Serialization/CompactBinarySerializationTest.cpp
@@ -115,7 +115,7 @@ bool FCbMeasureTest::RunTest(const FString& Parameters)
 	TestMeasure(TEXT("BoolFalse"), {uint8(EType::BoolFalse)}, true, 1, EType::BoolFalse);
 	TestMeasure(TEXT("BoolTrue"), {uint8(EType::BoolTrue)}, true, 1, EType::BoolTrue);
 
-	TestMeasure(TEXT("CompactBinaryAttachment"), {uint8(EType::CompactBinaryAttachment)}, true, 21, EType::CompactBinaryAttachment);
+	TestMeasure(TEXT("ObjectAttachment"), {uint8(EType::ObjectAttachment)}, true, 21, EType::ObjectAttachment);
 	TestMeasure(TEXT("BinaryAttachment"), {uint8(EType::BinaryAttachment)}, true, 21, EType::BinaryAttachment);
 
 	TestMeasure(TEXT("Hash"), {uint8(EType::Hash)}, true, 21, EType::Hash);
diff --git a/Engine/Source/Runtime/Core/Private/Tests/Serialization/CompactBinaryTest.cpp b/Engine/Source/Runtime/Core/Private/Tests/Serialization/CompactBinaryTest.cpp
--- a/Engine/Source/Runtime/Core/Private/Tests/Serialization/CompactBinaryTest.cpp
+++ b/Engine/Source/Runtime/Core/Private/Tests/Serialization/CompactBinaryTest.cpp
@@ -59,7 +59,7 @@ UE_CBFIELD_TYPE_ACCESSOR(Float32, IsFloat, AsFloat, float);
 UE_CBFIELD_TYPE_ACCESSOR(Float64, IsFloat, AsDouble, double);
 UE_CBFIELD_TYPE_ACCESSOR(BoolFalse, IsBool, AsBool, bool);
 UE_CBFIELD_TYPE_ACCESSOR(BoolTrue, IsBool, AsBool, bool);
-UE_CBFIELD_TYPE_ACCESSOR_EX(CompactBinaryAttachment, IsCompactBinaryAttachment, AsCompactBinaryAttachment, FIoHash, const FIoHash&);
+UE_CBFIELD_TYPE_ACCESSOR_EX(ObjectAttachment, IsObjectAttachment, AsObjectAttachment, FIoHash, const FIoHash&);
 UE_CBFIELD_TYPE_ACCESSOR_EX(BinaryAttachment, IsBinaryAttachment, AsBinaryAttachment, FIoHash, const FIoHash&);
 UE_CBFIELD_TYPE_ACCESSOR_EX(Hash, IsHash, AsHash, FIoHash, const FIoHash&);
 UE_CBFIELD_TYPE_ACCESSOR_EX(Uuid, IsUuid, AsUuid, FGuid, const FGuid&);
@@ -978,28 +978,28 @@ bool FCbFieldBoolTest::RunTest(const FString& Parameters)
 	return true;
 }
 
-IMPLEMENT_CUSTOM_SIMPLE_AUTOMATION_TEST(FCbFieldCompactBinaryAttachmentTest, FCbFieldTestBase, "System.Core.Serialization.CbField.CompactBinaryAttachment", CompactBinaryTestFlags)
-bool FCbFieldCompactBinaryAttachmentTest::RunTest(const FString& Parameters)
+IMPLEMENT_CUSTOM_SIMPLE_AUTOMATION_TEST(FCbFieldObjectAttachmentTest, FCbFieldTestBase, "System.Core.Serialization.CbField.ObjectAttachment", CompactBinaryTestFlags)
+bool FCbFieldObjectAttachmentTest::RunTest(const FString& Parameters)
 {
 	const FIoHash::ByteArray ZeroBytes{};
 	const FIoHash::ByteArray SequentialBytes{1, 2, 3, 4, 5, 6, 7, 8, 9, 10, 11, 12, 13, 14, 15, 16, 17, 18, 19, 20};
 
-	// Test FCbFieldView(CompactBinaryAttachment, Zero)
-	TestField<ECbFieldType::CompactBinaryAttachment>(TEXT("CompactBinaryAttachment, Zero"), ZeroBytes);
+	// Test FCbFieldView(ObjectAttachment, Zero)
+	TestField<ECbFieldType::ObjectAttachment>(TEXT("ObjectAttachment, Zero"), ZeroBytes);
 
-	// Test FCbFieldView(CompactBinaryAttachment, NonZero)
-	TestField<ECbFieldType::CompactBinaryAttachment>(TEXT("CompactBinaryAttachment, NonZero"), SequentialBytes, FIoHash(SequentialBytes));
+	// Test FCbFieldView(ObjectAttachment, NonZero)
+	TestField<ECbFieldType::ObjectAttachment>(TEXT("ObjectAttachment, NonZero"), SequentialBytes, FIoHash(SequentialBytes));
 
-	// Test FCbFieldView(CompactBinaryAttachment, NonZero) AsAttachment
+	// Test FCbFieldView(ObjectAttachment, NonZero) AsAttachment
 	{
-		FCbFieldView Field(SequentialBytes, ECbFieldType::CompactBinaryAttachment);
-		TestField<ECbFieldType::CompactBinaryAttachment>(TEXT("CompactBinaryAttachment, NonZero, AsAttachment"), Field, FIoHash(SequentialBytes), FIoHash(), ECbFieldError::None, FCbAttachmentAccessors());
+		FCbFieldView Field(SequentialBytes, ECbFieldType::ObjectAttachment);
+		TestField<ECbFieldType::ObjectAttachment>(TEXT("ObjectAttachment, NonZero, AsAttachment"), Field, FIoHash(SequentialBytes), FIoHash(), ECbFieldError::None, FCbAttachmentAccessors());
 	}
 
-	// Test FCbFieldView(None) as CompactBinaryAttachment
+	// Test FCbFieldView(None) as ObjectAttachment
 	{
 		FCbFieldView DefaultField;
-		TestFieldError<ECbFieldType::CompactBinaryAttachment>(TEXT("CompactBinaryAttachment, None"), DefaultField, ECbFieldError::TypeError, FIoHash(SequentialBytes));
+		TestFieldError<ECbFieldType::ObjectAttachment>(TEXT("ObjectAttachment, None"), DefaultField, ECbFieldError::TypeError, FIoHash(SequentialBytes));
 	}
 
 	return true;
@@ -1050,10 +1050,10 @@ bool FCbFieldHashTest::RunTest(const FString& Parameters)
 		TestFieldError<ECbFieldType::Hash>(TEXT("Hash, None"), DefaultField, ECbFieldError::TypeError, FIoHash(SequentialBytes));
 	}
 
-	// Test FCbFieldView(CompactBinaryAttachment) as Hash
+	// Test FCbFieldView(ObjectAttachment) as Hash
 	{
-		FCbFieldView Field(SequentialBytes, ECbFieldType::CompactBinaryAttachment);
-		TestField<ECbFieldType::Hash>(TEXT("CompactBinaryAttachment, NonZero, AsHash"), Field, FIoHash(SequentialBytes));
+		FCbFieldView Field(SequentialBytes, ECbFieldType::ObjectAttachment);
+		TestField<ECbFieldType::Hash>(TEXT("ObjectAttachment, NonZero, AsHash"), Field, FIoHash(SequentialBytes));
 	}
 
 	// Test FCbFieldView(BinaryAttachment) as Hash
@@ -1279,9 +1279,9 @@ bool FCbFieldIterateAttachmentsTest::RunTest(const FString& Parameters)
 		FCbWriter Writer;
 
 		Writer.SetName("IgnoredTypeInRoot").AddHash(MakeTestHash(100));
-		Writer.AddCompactBinaryAttachment(MakeTestHash(0));
+		Writer.AddObjectAttachment(MakeTestHash(0));
 		Writer.AddBinaryAttachment(MakeTestHash(1));
-		Writer.SetName("CbAttachmentInRoot").AddCompactBinaryAttachment(MakeTestHash(2));
+		Writer.SetName("ObjAttachmentInRoot").AddObjectAttachment(MakeTestHash(2));
 		Writer.SetName("BinAttachmentInRoot").AddBinaryAttachment(MakeTestHash(3));
 
 		// Uniform array of type to ignore.
@@ -1319,7 +1319,7 @@ bool FCbFieldIterateAttachmentsTest::RunTest(const FString& Parameters)
 			Writer << false;
 			Writer.EndArray();
 			Writer.BeginArray();
-			Writer.AddCompactBinaryAttachment(MakeTestHash(10));
+			Writer.AddObjectAttachment(MakeTestHash(10));
 			Writer << false;
 			Writer.EndArray();
 		}
@@ -1328,12 +1328,12 @@ bool FCbFieldIterateAttachmentsTest::RunTest(const FString& Parameters)
 		Writer.BeginArray();
 		{
 			Writer.BeginObject();
-			Writer.SetName("CbAttachmentInUniObjInUniObj1").AddCompactBinaryAttachment(MakeTestHash(11));
-			Writer.SetName("CbAttachmentInUniObjInUniObj2").AddCompactBinaryAttachment(MakeTestHash(12));
+			Writer.SetName("ObjAttachmentInUniObjInUniObj1").AddObjectAttachment(MakeTestHash(11));
+			Writer.SetName("ObjAttachmentInUniObjInUniObj2").AddObjectAttachment(MakeTestHash(12));
 			Writer.EndObject();
 			Writer.BeginObject();
-			Writer.SetName("CbAttachmentInUniObjInUniObj3").AddCompactBinaryAttachment(MakeTestHash(13));
-			Writer.SetName("CbAttachmentInUniObjInUniObj4").AddCompactBinaryAttachment(MakeTestHash(14));
+			Writer.SetName("ObjAttachmentInUniObjInUniObj3").AddObjectAttachment(MakeTestHash(13));
+			Writer.SetName("ObjAttachmentInUniObjInUniObj4").AddObjectAttachment(MakeTestHash(14));
 			Writer.EndObject();
 		}
 		Writer.EndArray();
@@ -1345,7 +1345,7 @@ bool FCbFieldIterateAttachmentsTest::RunTest(const FString& Parameters)
 			Writer << "Bool" << false;
 			Writer.EndObject();
 			Writer.BeginObject();
-			Writer.SetName("CbAttachmentInNonUniObjInUniObj").AddCompactBinaryAttachment(MakeTestHash(15));
+			Writer.SetName("ObjAttachmentInNonUniObjInUniObj").AddObjectAttachment(MakeTestHash(15));
 			Writer << "Bool" << false;
 			Writer.EndObject();
 		}
@@ -1390,7 +1390,7 @@ bool FCbFieldIterateAttachmentsTest::RunTest(const FString& Parameters)
 			Writer.EndArray();
 			Writer.SetName("Array2");
 			Writer.BeginArray();
-			Writer.AddCompactBinaryAttachment(MakeTestHash(22));
+			Writer.AddObjectAttachment(MakeTestHash(22));
 			Writer << false;
 			Writer.EndArray();
 		}
@@ -1400,13 +1400,13 @@ bool FCbFieldIterateAttachmentsTest::RunTest(const FString& Parameters)
 		{
 			Writer.SetName("Object1");
 			Writer.BeginObject();
-			Writer.SetName("CbAttachmentInUniObjInUniObj1").AddCompactBinaryAttachment(MakeTestHash(23));
-			Writer.SetName("CbAttachmentInUniObjInUniObj2").AddCompactBinaryAttachment(MakeTestHash(24));
+			Writer.SetName("ObjAttachmentInUniObjInUniObj1").AddObjectAttachment(MakeTestHash(23));
+			Writer.SetName("ObjAttachmentInUniObjInUniObj2").AddObjectAttachment(MakeTestHash(24));
 			Writer.EndObject();
 			Writer.SetName("Object2");
 			Writer.BeginObject();
-			Writer.SetName("CbAttachmentInUniObjInUniObj3").AddCompactBinaryAttachment(MakeTestHash(25));
-			Writer.SetName("CbAttachmentInUniObjInUniObj4").AddCompactBinaryAttachment(MakeTestHash(26));
+			Writer.SetName("ObjAttachmentInUniObjInUniObj3").AddObjectAttachment(MakeTestHash(25));
+			Writer.SetName("ObjAttachmentInUniObjInUniObj4").AddObjectAttachment(MakeTestHash(26));
 			Writer.EndObject();
 		}
 		Writer.EndObject();
@@ -1420,7 +1420,7 @@ bool FCbFieldIterateAttachmentsTest::RunTest(const FString& Parameters)
 			Writer.EndObject();
 			Writer.SetName("Object2");
 			Writer.BeginObject();
-			Writer.SetName("CbAttachmentInNonUniObjInUniObj").AddCompactBinaryAttachment(MakeTestHash(27));
+			Writer.SetName("ObjAttachmentInNonUniObjInUniObj").AddObjectAttachment(MakeTestHash(27));
 			Writer << "Bool" << false;
 			Writer.EndObject();
 		}
diff --git a/Engine/Source/Runtime/Core/Private/Tests/Serialization/CompactBinaryValidationTest.cpp b/Engine/Source/Runtime/Core/Private/Tests/Serialization/CompactBinaryValidationTest.cpp
--- a/Engine/Source/Runtime/Core/Private/Tests/Serialization/CompactBinaryValidationTest.cpp
+++ b/Engine/Source/Runtime/Core/Private/Tests/Serialization/CompactBinaryValidationTest.cpp
@@ -140,10 +140,10 @@ bool FCbValidateTest::RunTest(const FString& Parameters)
 	TestEqual(TEXT("ValidateCompactBinary(Valid, BoolFalse)"), Validate({uint8(ECbFieldType::BoolFalse)}), ECbValidateError::None);
 	TestEqual(TEXT("ValidateCompactBinary(Valid, BoolTrue)"), Validate({uint8(ECbFieldType::BoolTrue)}), ECbValidateError::None);
 
-	TestEqual(TEXT("ValidateCompactBinary(Valid, CompactBinaryAttachment)"), Validate({uint8(ECbFieldType::CompactBinaryAttachment), 0, 0, 0, 0, 0, 0, 0, 0, 0, 0, 0, 0, 0, 0, 0, 0, 0, 0, 0, 0}), ECbValidateError::None);
-	TestEqual(TEXT("ValidateCompactBinary(Valid, CompactBinaryAttachment, NoType)"), Validate({0, 0, 0, 0, 0, 0, 0, 0, 0, 0, 0, 0, 0, 0, 0, 0, 0, 0, 0, 0}, ECbFieldType::CompactBinaryAttachment), ECbValidateError::None);
-	TestEqual(TEXT("ValidateCompactBinary(OutOfBounds, CompactBinaryAttachment)"), Validate({uint8(ECbFieldType::CompactBinaryAttachment), 0, 0, 0, 0, 0, 0, 0, 0, 0, 0, 0, 0, 0, 0, 0, 0, 0, 0, 0}), ECbValidateError::OutOfBounds);
-	TestEqual(TEXT("ValidateCompactBinary(OutOfBounds, CompactBinaryAttachment, NoType)"), Validate({0, 0, 0, 0, 0, 0, 0, 0, 0, 0, 0, 0, 0, 0, 0, 0, 0, 0, 0}, ECbFieldType::CompactBinaryAttachment), ECbValidateError::OutOfBounds);
+	TestEqual(TEXT("ValidateCompactBinary(Valid, ObjectAttachment)"), Validate({uint8(ECbFieldType::ObjectAttachment), 0, 0, 0, 0, 0, 0, 0, 0, 0, 0, 0, 0, 0, 0, 0, 0, 0, 0, 0, 0}), ECbValidateError::None);
+	TestEqual(TEXT("ValidateCompactBinary(Valid, ObjectAttachment, NoType)"), Validate({0, 0, 0, 0, 0, 0, 0, 0, 0, 0, 0, 0, 0, 0, 0, 0, 0, 0, 0, 0}, ECbFieldType::ObjectAttachment), ECbValidateError::None);
+	TestEqual(TEXT("ValidateCompactBinary(OutOfBounds, ObjectAttachment)"), Validate({uint8(ECbFieldType::ObjectAttachment), 0, 0, 0, 0, 0, 0, 0, 0, 0, 0, 0, 0, 0, 0, 0, 0, 0, 0, 0}), ECbValidateError::OutOfBounds);
+	TestEqual(TEXT("ValidateCompactBinary(OutOfBounds, ObjectAttachment, NoType)"), Validate({0, 0, 0, 0, 0, 0, 0, 0, 0, 0, 0, 0, 0, 0, 0, 0, 0, 0, 0}, ECbFieldType::ObjectAttachment), ECbValidateError::OutOfBounds);
 
 	TestEqual(TEXT("ValidateCompactBinary(Valid, BinaryAttachment)"), Validate({uint8(ECbFieldType::BinaryAttachment), 0, 0, 0, 0, 0, 0, 0, 0, 0, 0, 0, 0, 0, 0, 0, 0, 0, 0, 0, 0}), ECbValidateError::None);
 	TestEqual(TEXT("ValidateCompactBinary(Valid, BinaryAttachment, NoType)"), Validate({0, 0, 0, 0, 0, 0, 0, 0, 0, 0, 0, 0, 0, 0, 0, 0, 0, 0, 0, 0}, ECbFieldType::BinaryAttachment), ECbValidateError::None);
@@ -399,7 +399,7 @@ bool FCbValidateAttachmentTest::RunTest(const FString& Parameters)
 	{
 		FBufferArchive Buffer;
 		FCbWriter Writer;
-		Writer.AddCompactBinaryAttachment(FIoHash::HashBuffer(BinaryView));
+		Writer.AddObjectAttachment(FIoHash::HashBuffer(BinaryView));
 		Writer.AddBinary(BinaryView);
 		Writer.Save(Buffer);
 		TestEqual(TEXT("ValidateCompactBinaryAttachment(HashBeforeValue)"), ValidateCompactBinaryAttachment(MakeMemoryView(Buffer), ECbValidateMode::All), ECbValidateError::InvalidPackageFormat | ECbValidateError::Padding);
@@ -409,7 +409,7 @@ bool FCbValidateAttachmentTest::RunTest(const FString& Parameters)
 		FBufferArchive Buffer;
 		FCbWriter Writer;
 		Writer.SetName("Name"_ASV).AddBinary(BinaryView);
-		Writer.AddCompactBinaryAttachment(FIoHash::HashBuffer(BinaryView));
+		Writer.AddObjectAttachment(FIoHash::HashBuffer(BinaryView));
 		Writer.Save(Buffer);
 		TestEqual(TEXT("ValidateCompactBinaryAttachment(NameOnValue)"), ValidateCompactBinaryAttachment(MakeMemoryView(Buffer), ECbValidateMode::All), ECbValidateError::InvalidPackageFormat);
 		TestEqual(TEXT("ValidateCompactBinaryAttachment(NameOnValue, Mode)"), ValidateCompactBinaryAttachment(MakeMemoryView(Buffer), ECbValidateMode::All & ~ECbValidateMode::Package), ECbValidateError::None);
@@ -418,7 +418,7 @@ bool FCbValidateAttachmentTest::RunTest(const FString& Parameters)
 		FBufferArchive Buffer;
 		FCbWriter Writer;
 		Writer.AddBinary(BinaryView);
-		Writer.SetName("Name"_ASV).AddCompactBinaryAttachment(FIoHash::HashBuffer(BinaryView));
+		Writer.SetName("Name"_ASV).AddObjectAttachment(FIoHash::HashBuffer(BinaryView));
 		Writer.Save(Buffer);
 		TestEqual(TEXT("ValidateCompactBinaryAttachment(NameOnHash)"), ValidateCompactBinaryAttachment(MakeMemoryView(Buffer), ECbValidateMode::All), ECbValidateError::InvalidPackageFormat);
 		TestEqual(TEXT("ValidateCompactBinaryAttachment(NameOnHash, Mode)"), ValidateCompactBinaryAttachment(MakeMemoryView(Buffer), ECbValidateMode::All & ~ECbValidateMode::Package), ECbValidateError::None);
@@ -449,7 +449,7 @@ bool FCbValidatePackageTest::RunTest(const FString& Parameters)
 		Writer.BeginObject();
 		Writer << "Field" << 42;
 		Writer.EndObject();
-		Writer.AddCompactBinaryAttachment(Writer.Save().GetHash());
+		Writer.AddObjectAttachment(Writer.Save().GetHash());
 		Writer.AddNull();
 		Writer.Save(Buffer);
 		TestEqual(TEXT("ValidateCompactBinaryPackage(Object)"), ValidateCompactBinaryPackage(MakeMemoryView(Buffer), ECbValidateMode::All), ECbValidateError::None);
@@ -477,10 +477,10 @@ bool FCbValidatePackageTest::RunTest(const FString& Parameters)
 		Writer.BeginObject();
 		Writer << "Field" << 42;
 		Writer.EndObject();
-		Writer.AddCompactBinaryAttachment(Writer.Save().GetHash());
+		Writer.AddObjectAttachment(Writer.Save().GetHash());
 		const uint64 AttachmentOffset = Writer.GetSaveSize();
 		Writer.AddBinary(BinaryView);
-		Writer.AddCompactBinaryAttachment(FIoHash::HashBuffer(BinaryView));
+		Writer.AddObjectAttachment(FIoHash::HashBuffer(BinaryView));
 		Writer.AddNull();
 		Writer.Save(Buffer);
 		TestEqual(TEXT("ValidateCompactBinaryPackage(Object+Attachment)"), ValidateCompactBinaryPackage(MakeMemoryView(Buffer), ECbValidateMode::All), ECbValidateError::None);
@@ -513,7 +513,7 @@ bool FCbValidatePackageTest::RunTest(const FString& Parameters)
 		Writer.BeginObject();
 		Writer << "Field" << 42;
 		Writer.EndObject();
-		Writer.AddCompactBinaryAttachment(Writer.Save().GetHash());
+		Writer.AddObjectAttachment(Writer.Save().GetHash());
 		Writer.AddNull();
 		Writer.Save(Buffer);
 		TestEqual(TEXT("ValidateCompactBinaryPackage(NameOnObject)"), ValidateCompactBinaryPackage(MakeMemoryView(Buffer), ECbValidateMode::All), ECbValidateError::InvalidPackageFormat);
@@ -527,7 +527,7 @@ bool FCbValidatePackageTest::RunTest(const FString& Parameters)
 		Writer.EndObject();
 		const FIoHash ObjectHash = Writer.Save().GetHash();
 		Writer.SetName("Name"_ASV);
-		Writer.AddCompactBinaryAttachment(ObjectHash);
+		Writer.AddObjectAttachment(ObjectHash);
 		Writer.AddNull();
 		Writer.Save(Buffer);
 		TestEqual(TEXT("ValidateCompactBinaryPackage(NameOnHash)"), ValidateCompactBinaryPackage(MakeMemoryView(Buffer), ECbValidateMode::All), ECbValidateError::InvalidPackageFormat);
@@ -539,9 +539,9 @@ bool FCbValidatePackageTest::RunTest(const FString& Parameters)
 		Writer.BeginObject();
 		Writer << "Field" << 42;
 		Writer.EndObject();
-		Writer.AddCompactBinaryAttachment(Writer.Save().GetHash());
+		Writer.AddObjectAttachment(Writer.Save().GetHash());
 		Writer.SetName("Name"_ASV).AddBinary(BinaryView);
-		Writer.AddCompactBinaryAttachment(FIoHash::HashBuffer(BinaryView));
+		Writer.AddObjectAttachment(FIoHash::HashBuffer(BinaryView));
 		Writer.AddNull();
 		Writer.Save(Buffer);
 		TestEqual(TEXT("ValidateCompactBinaryPackage(NameOnAttachment)"), ValidateCompactBinaryPackage(MakeMemoryView(Buffer), ECbValidateMode::All), ECbValidateError::InvalidPackageFormat);
@@ -554,11 +554,11 @@ bool FCbValidatePackageTest::RunTest(const FString& Parameters)
 		Writer << "Field" << 42;
 		Writer.EndObject();
 		const FIoHash ObjectHash = Writer.Save().GetHash();
-		Writer.AddCompactBinaryAttachment(ObjectHash);
+		Writer.AddObjectAttachment(ObjectHash);
 		Writer.BeginObject();
 		Writer << "Field" << 42;
 		Writer.EndObject();
-		Writer.AddCompactBinaryAttachment(ObjectHash);
+		Writer.AddObjectAttachment(ObjectHash);
 		Writer.AddNull();
 		Writer.Save(Buffer);
 		TestEqual(TEXT("ValidateCompactBinaryPackage(MultipleObjects)"), ValidateCompactBinaryPackage(MakeMemoryView(Buffer), ECbValidateMode::All), ECbValidateError::MultiplePackageObjects);
@@ -568,9 +568,9 @@ bool FCbValidatePackageTest::RunTest(const FString& Parameters)
 		FBufferArchive Buffer;
 		TCbWriter<256> Writer;
 		Writer.AddBinary(BinaryView);
-		Writer.AddCompactBinaryAttachment(FIoHash::HashBuffer(BinaryView));
+		Writer.AddObjectAttachment(FIoHash::HashBuffer(BinaryView));
 		Writer.AddBinary(BinaryView);
-		Writer.AddCompactBinaryAttachment(FIoHash::HashBuffer(BinaryView));
+		Writer.AddObjectAttachment(FIoHash::HashBuffer(BinaryView));
 		Writer.AddNull();
 		Writer.Save(Buffer);
 		TestEqual(TEXT("ValidateCompactBinaryPackage(DuplicateAttachments)"), ValidateCompactBinaryPackage(MakeMemoryView(Buffer), ECbValidateMode::All), ECbValidateError::DuplicateAttachments);
@@ -581,7 +581,7 @@ bool FCbValidatePackageTest::RunTest(const FString& Parameters)
 		TCbWriter<256> Writer;
 		Writer.BeginObject();
 		Writer.EndObject();
-		Writer.AddCompactBinaryAttachment(FCbObjectView().GetHash());
+		Writer.AddObjectAttachment(FCbObjectView().GetHash());
 		Writer.AddNull();
 		Writer.Save(Buffer);
 		TestEqual(TEXT("ValidateCompactBinaryPackage(NullObject)"), ValidateCompactBinaryPackage(MakeMemoryView(Buffer), ECbValidateMode::All), ECbValidateError::NullPackageObject);
diff --git a/Engine/Source/Runtime/Core/Private/Tests/Serialization/CompactBinaryWriterTest.cpp b/Engine/Source/Runtime/Core/Private/Tests/Serialization/CompactBinaryWriterTest.cpp
--- a/Engine/Source/Runtime/Core/Private/Tests/Serialization/CompactBinaryWriterTest.cpp
+++ b/Engine/Source/Runtime/Core/Private/Tests/Serialization/CompactBinaryWriterTest.cpp
@@ -564,8 +564,8 @@ bool FCbWriterBoolTest::RunTest(const FString& Parameters)
 
 ///////////////////////////////////////////////////////////////////////////////////////////////////////////////////////
 
-IMPLEMENT_SIMPLE_AUTOMATION_TEST(FCbWriterCompactBinaryAttachmentTest, "System.Core.Serialization.CbWriter.CompactBinaryAttachment", CompactBinaryWriterTestFlags)
-bool FCbWriterCompactBinaryAttachmentTest::RunTest(const FString& Parameters)
+IMPLEMENT_SIMPLE_AUTOMATION_TEST(FCbWriterObjectAttachmentTest, "System.Core.Serialization.CbWriter.ObjectAttachment", CompactBinaryWriterTestFlags)
+bool FCbWriterObjectAttachmentTest::RunTest(const FString& Parameters)
 {
 	TCbWriter<256> Writer;
 
@@ -575,17 +575,17 @@ bool FCbWriterCompactBinaryAttachmentTest::RunTest(const FString& Parameters)
 	const FIoHash Values[] = { FIoHash(ZeroBytes), FIoHash(SequentialBytes) };
 	for (const FIoHash& Value : Values)
 	{
-		Writer.AddCompactBinaryAttachment(Value);
+		Writer.AddObjectAttachment(Value);
 	}
 
 	FCbFieldIterator Fields = Writer.Save();
-	if (TestEqual(TEXT("FCbWriter(CompactBinaryAttachment) Validate"), ValidateCompactBinaryRange(Fields.GetRangeBuffer(), ECbValidateMode::All), ECbValidateError::None))
+	if (TestEqual(TEXT("FCbWriter(ObjectAttachment) Validate"), ValidateCompactBinaryRange(Fields.GetRangeBuffer(), ECbValidateMode::All), ECbValidateError::None))
 	{
 		const FIoHash* CheckValue = Values;
 		for (FCbFieldView Field : Fields)
 		{
-			TestEqual(TEXT("FCbWriter(CompactBinaryAttachment).AsCompactBinaryAttachment()"), Field.AsCompactBinaryAttachment(), *CheckValue++);
-			TestFalse(TEXT("FCbWriter(CompactBinaryAttachment) Error"), Field.HasError());
+			TestEqual(TEXT("FCbWriter(ObjectAttachment).AsObjectAttachment()"), Field.AsObjectAttachment(), *CheckValue++);
+			TestFalse(TEXT("FCbWriter(ObjectAttachment) Error"), Field.HasError());
 		}
 	}
 
@@ -916,7 +916,7 @@ bool FCbWriterComplexTest::RunTest(const FString& Parameters)
 		Writer.AddBool("False"_ASV, false);
 		Writer.AddBool("True"_ASV, true);
 
-		Writer.AddCompactBinaryAttachment("CompactBinaryAttachment"_ASV, FIoHash());
+		Writer.AddObjectAttachment("ObjectAttachment"_ASV, FIoHash());
 		Writer.AddBinaryAttachment("BinaryAttachment"_ASV, FIoHash());
 		Writer.AddAttachment("Attachment"_ASV, FCbAttachment());
 
diff --git a/Engine/Source/Runtime/Core/Public/Serialization/CompactBinary.h b/Engine/Source/Runtime/Core/Public/Serialization/CompactBinary.h
--- a/Engine/Source/Runtime/Core/Public/Serialization/CompactBinary.h
+++ b/Engine/Source/Runtime/Core/Public/Serialization/CompactBinary.h
@@ -147,11 +147,11 @@ enum class ECbFieldType : uint8
 	BoolTrue                         = 0x0d,
 
 	/**
-	 * CompactBinaryAttachment is a reference to a compact binary attachment stored externally.
+	 * ObjectAttachment is a reference to a compact binary object attachment stored externally.
 	 *
-	 * Payload is a 160-bit hash digest of the referenced compact binary data.
+	 * Payload is a 160-bit hash digest of the referenced compact binary object data.
 	 */
-	CompactBinaryAttachment          = 0x0e,
+	ObjectAttachment          = 0x0e,
 	/**
 	 * BinaryAttachment is a reference to a binary attachment stored externally.
 	 *
@@ -268,9 +268,9 @@ public:
 	static constexpr inline bool IsFloat(ECbFieldType Type)      { return (Type & FloatMask) == FloatBase; }
 	static constexpr inline bool IsBool(ECbFieldType Type)       { return (Type & BoolMask) == BoolBase; }
 
-	static constexpr inline bool IsCompactBinaryAttachment(ECbFieldType Type) { return GetType(Type) == ECbFieldType::CompactBinaryAttachment; }
-	static constexpr inline bool IsBinaryAttachment(ECbFieldType Type)        { return GetType(Type) == ECbFieldType::BinaryAttachment; }
-	static constexpr inline bool IsAttachment(ECbFieldType Type)              { return (Type & AttachmentMask) == AttachmentBase; }
+	static constexpr inline bool IsObjectAttachment(ECbFieldType Type) { return GetType(Type) == ECbFieldType::ObjectAttachment; }
+	static constexpr inline bool IsBinaryAttachment(ECbFieldType Type) { return GetType(Type) == ECbFieldType::BinaryAttachment; }
+	static constexpr inline bool IsAttachment(ECbFieldType Type)       { return (Type & AttachmentMask) == AttachmentBase; }
 
 	static constexpr inline bool IsHash(ECbFieldType Type)       { return GetType(Type) == ECbFieldType::Hash || IsAttachment(Type); }
 	static constexpr inline bool IsUuid(ECbFieldType Type)       { return GetType(Type) == ECbFieldType::Uuid; }
@@ -610,8 +610,8 @@ public:
 	/** Access the field as a bool. Returns the provided default on error. */
 	CORE_API bool AsBool(bool bDefault = false);
 
-	/** Access the field as a hash referencing a compact binary attachment. Returns the provided default on error. */
-	CORE_API FIoHash AsCompactBinaryAttachment(const FIoHash& Default = FIoHash());
+	/** Access the field as a hash referencing an object attachment. Returns the provided default on error. */
+	CORE_API FIoHash AsObjectAttachment(const FIoHash& Default = FIoHash());
 	/** Access the field as a hash referencing a binary attachment. Returns the provided default on error. */
 	CORE_API FIoHash AsBinaryAttachment(const FIoHash& Default = FIoHash());
 	/** Access the field as a hash referencing an attachment. Returns the provided default on error. */
@@ -671,9 +671,9 @@ public:
 	constexpr inline bool IsFloat() const           { return FCbFieldType::IsFloat(Type); }
 	constexpr inline bool IsBool() const            { return FCbFieldType::IsBool(Type); }
 
-	constexpr inline bool IsCompactBinaryAttachment() const { return FCbFieldType::IsCompactBinaryAttachment(Type); }
-	constexpr inline bool IsBinaryAttachment() const        { return FCbFieldType::IsBinaryAttachment(Type); }
-	constexpr inline bool IsAttachment() const              { return FCbFieldType::IsAttachment(Type); }
+	constexpr inline bool IsObjectAttachment() const { return FCbFieldType::IsObjectAttachment(Type); }
+	constexpr inline bool IsBinaryAttachment() const { return FCbFieldType::IsBinaryAttachment(Type); }
+	constexpr inline bool IsAttachment() const       { return FCbFieldType::IsAttachment(Type); }
 
 	constexpr inline bool IsHash() const            { return FCbFieldType::IsHash(Type); }
 	constexpr inline bool IsUuid() const            { return FCbFieldType::IsUuid(Type); }
diff --git a/Engine/Source/Runtime/Core/Public/Serialization/CompactBinaryPackage.h b/Engine/Source/Runtime/Core/Public/Serialization/CompactBinaryPackage.h
--- a/Engine/Source/Runtime/Core/Public/Serialization/CompactBinaryPackage.h
+++ b/Engine/Source/Runtime/Core/Public/Serialization/CompactBinaryPackage.h
@@ -21,7 +21,7 @@ template <typename FuncType> class TFunctionRef;
  *
  * Attachments are serialized as one or two compact binary fields with no name. A Binary field is
  * written first with its content. The content hash is omitted when the content size is zero, and
- * is otherwise written as a BinaryAttachment or CompactBinaryAttachment depending on the type.
+ * is otherwise written as a BinaryAttachment or ObjectAttachment depending on the type.
  */
 class FCbAttachment
 {
@@ -144,7 +144,7 @@ inline uint32 GetTypeHash(const FCbAttachment& Attachment)
  *
  * A package is serialized as a sequence of compact binary fields with no name. The object may be
  * both preceded and followed by attachments. The object itself is written as an Object field and
- * followed by its hash in a CompactBinaryAttachment field, if the object is non-empty. A package
+ * followed by its hash in a ObjectAttachment field, if the object is non-empty. A package
  * ends with a Null field. The canonical order of components is the object and its hash, followed
  * by the attachments ordered by hash, followed by a Null field. It is valid for the a package to
  * have its components serialized in any order, provided there is at most one object and the null
diff --git a/Engine/Source/Runtime/Core/Public/Serialization/CompactBinaryWriter.h b/Engine/Source/Runtime/Core/Public/Serialization/CompactBinaryWriter.h
--- a/Engine/Source/Runtime/Core/Public/Serialization/CompactBinaryWriter.h
+++ b/Engine/Source/Runtime/Core/Public/Serialization/CompactBinaryWriter.h
@@ -189,9 +189,9 @@ public:
 	inline void AddBool(FAnsiStringView Name, bool bValue) { SetName(Name); AddBool(bValue); }
 	CORE_API void AddBool(bool bValue);
 
-	/** Write a field referencing a compact binary attachment by its hash. */
-	inline void AddCompactBinaryAttachment(FAnsiStringView Name, const FIoHash& Value) { SetName(Name); AddCompactBinaryAttachment(Value); }
-	CORE_API void AddCompactBinaryAttachment(const FIoHash& Value);
+	/** Write a field referencing an object attachment by its hash. */
+	inline void AddObjectAttachment(FAnsiStringView Name, const FIoHash& Value) { SetName(Name); AddObjectAttachment(Value); }
+	CORE_API void AddObjectAttachment(const FIoHash& Value);
 	/** Write a field referencing a binary attachment by its hash. */
 	inline void AddBinaryAttachment(FAnsiStringView Name, const FIoHash& Value) { SetName(Name); AddBinaryAttachment(Value); }
 	CORE_API void AddBinaryAttachment(const FIoHash& Value);
