CompactBinary: Remove an assert when loading from an archive

This is the only data-dependent assert that crept into compact binary serialization, and needs to be removed to allow for graceful error handling.

#rnx

diff --git a/Engine/Source/Runtime/Core/Private/Serialization/CompactBinarySerialization.cpp b/Engine/Source/Runtime/Core/Private/Serialization/CompactBinarySerialization.cpp
--- a/Engine/Source/Runtime/Core/Private/Serialization/CompactBinarySerialization.cpp
+++ b/Engine/Source/Runtime/Core/Private/Serialization/CompactBinarySerialization.cpp
@@ -173,7 +173,11 @@ FCbField LoadCompactBinary(FArchive& Ar, FCbBufferAllocator Allocator)
 		{
 			break;
 		}
-		checkf(FieldSize > 0, TEXT("Failed to load from invalid compact binary data."), FieldSize);
+		if (FieldSize == 0)
+		{
+			Ar.SetError();
+			return FCbField();
+		}
 	}
 
 	// Allocate the buffer, copy the header, and read the remainder of the field.
