Expose IoDispatcher and AsyncLoading stats to csv profiler.

diff --git a/Engine/Source/Runtime/Core/Private/GenericPlatform/GenericPlatformIoDispatcher.cpp b/Engine/Source/Runtime/Core/Private/GenericPlatform/GenericPlatformIoDispatcher.cpp
--- a/Engine/Source/Runtime/Core/Private/GenericPlatform/GenericPlatformIoDispatcher.cpp
+++ b/Engine/Source/Runtime/Core/Private/GenericPlatform/GenericPlatformIoDispatcher.cpp
@@ -12,10 +12,6 @@
 
 //PRAGMA_DISABLE_OPTIMIZATION
 
-TRACE_DECLARE_INT_COUNTER(IoDispatcherSequentialReads, TEXT("IoDispatcher/SequentialReads"));
-TRACE_DECLARE_INT_COUNTER(IoDispatcherForwardSeeks, TEXT("IoDispatcher/ForwardSeeks"));
-TRACE_DECLARE_INT_COUNTER(IoDispatcherBackwardSeeks, TEXT("IoDispatcher/BackwardSeeks"));
-TRACE_DECLARE_MEMORY_COUNTER(IoDispatcherTotalSeekDistance, TEXT("IoDispatcher/TotalSeekDistance"));
 
 FGenericFileIoStoreEventQueue::FGenericFileIoStoreEventQueue()
 	: ServiceEvent(FPlatformProcess::GetSynchEventFromPool())
@@ -69,6 +65,7 @@ bool FGenericFileIoStoreImpl::OpenContainer(const TCHAR* ContainerFilePath, uint
 void FGenericFileIoStoreImpl::CloseContainer(uint64 ContainerFileHandle)
 {
 	check(ContainerFileHandle);
+	FFileIoStats::OnCloseHandle(ContainerFileHandle);
 	IFileHandle* FileHandle = reinterpret_cast<IFileHandle*>(ContainerFileHandle);
 	delete FileHandle;
 }
@@ -110,22 +107,8 @@ bool FGenericFileIoStoreImpl::StartRequests(FFileIoStoreRequestQueue& RequestQue
 	if (!BlockCache.Read(NextRequest))
 	{
 		IFileHandle* FileHandle = reinterpret_cast<IFileHandle*>(static_cast<UPTRINT>(NextRequest->FileHandle));
-		if (FileHandle->Tell() != NextRequest->Offset)
-		{
-			if (uint64(FileHandle->Tell()) > NextRequest->Offset)
-			{
-				TRACE_COUNTER_INCREMENT(IoDispatcherBackwardSeeks);
-			}
-			else
-			{
-				TRACE_COUNTER_INCREMENT(IoDispatcherForwardSeeks);
-			}
-			TRACE_COUNTER_ADD(IoDispatcherTotalSeekDistance, FMath::Abs(FileHandle->Tell() - int64(NextRequest->Offset)));
-		}
-		else
-		{
-			TRACE_COUNTER_INCREMENT(IoDispatcherSequentialReads);
-		}
+		 
+		FFileIoStats::OnFilesystemReadStarted(NextRequest->FileHandle, NextRequest->Offset, NextRequest->Size);
 		{
 			TRACE_CPUPROFILER_EVENT_SCOPE(ReadBlockFromFile);
 			NextRequest->bFailed = true;
@@ -147,6 +130,7 @@ bool FGenericFileIoStoreImpl::StartRequests(FFileIoStoreRequestQueue& RequestQue
 				break;
 			}
 		}
+		FFileIoStats::OnFilesystemReadsComplete(NextRequest->Size);
 	}
 	{
 		FScopeLock _(&CompletedRequestsCritical);
diff --git a/Engine/Source/Runtime/Core/Private/IO/IoDispatcher.cpp b/Engine/Source/Runtime/Core/Private/IO/IoDispatcher.cpp
--- a/Engine/Source/Runtime/Core/Private/IO/IoDispatcher.cpp
+++ b/Engine/Source/Runtime/Core/Private/IO/IoDispatcher.cpp
@@ -29,6 +29,9 @@ TUniquePtr<FIoDispatcher> GIoDispatcher;
 //PRAGMA_DISABLE_OPTIMIZATION
 #endif
 
+CSV_DEFINE_CATEGORY(IoDispatcher, true);
+CSV_DEFINE_STAT(IoDispatcher, PendingIoRequests);
+
 TRACE_DECLARE_INT_COUNTER(PendingIoRequests, TEXT("IoDispatcher/PendingIoRequests"));
 
 template <typename T, uint32 BlockSize = 128>
@@ -153,6 +156,7 @@ public:
 
 	~FIoDispatcherImpl()
 	{
+
 		delete Thread;
 		FPlatformProcess::ReturnSynchEventToPool(DispatcherEvent);
 	}
@@ -170,6 +174,7 @@ public:
 			Backend->Initialize(BackendContext);
 		}
 		Thread = FRunnableThread::Create(this, TEXT("IoDispatcher"), 0, TPri_AboveNormal, FPlatformAffinity::GetIoDispatcherThreadMask());
+		FFileIoStats::SetThreadIds(FileIoStore.GetThreadId(), Thread ? Thread->GetThreadID() : 0);
 		return true;
 	}
 
@@ -445,6 +450,7 @@ private:
 				--PendingIoRequestsCount;
 				TRACE_COUNTER_SET(PendingIoRequests, PendingIoRequestsCount);
 			}
+			CSV_CUSTOM_STAT_DEFINED(PendingIoRequests, (int32)PendingIoRequestsCount, ECsvCustomStatOp::SetAndHold);
 		}
 	}
 
@@ -603,6 +609,7 @@ private:
 			
 			++PendingIoRequestsCount;
 			TRACE_COUNTER_SET(PendingIoRequests, PendingIoRequestsCount);
+			CSV_CUSTOM_STAT_DEFINED(PendingIoRequests, (int32)PendingIoRequestsCount, ECsvCustomStatOp::SetAndHold);
 			
 			ProcessCompletedRequests();
 		}
diff --git a/Engine/Source/Runtime/Core/Private/IO/IoDispatcherFileBackend.cpp b/Engine/Source/Runtime/Core/Private/IO/IoDispatcherFileBackend.cpp
--- a/Engine/Source/Runtime/Core/Private/IO/IoDispatcherFileBackend.cpp
+++ b/Engine/Source/Runtime/Core/Private/IO/IoDispatcherFileBackend.cpp
@@ -17,12 +17,6 @@
 #include "Algo/MinElement.h"
 #include "Templates/Greater.h"
 
-TRACE_DECLARE_MEMORY_COUNTER(IoDispatcherTotalBytesRead, TEXT("IoDispatcher/TotalBytesRead"));
-TRACE_DECLARE_MEMORY_COUNTER(IoDispatcherTotalBytesScattered, TEXT("IoDispatcher/TotalBytesScattered"));
-TRACE_DECLARE_INT_COUNTER(IoDispatcherCacheHits, TEXT("IoDispatcher/CacheHits"));
-TRACE_DECLARE_INT_COUNTER(IoDispatcherCacheMisses, TEXT("IoDispatcher/CacheMisses"));
-TRACE_DECLARE_INT_COUNTER(IoDispatcherOutstandingReads, TEXT("IoDispatcher/OutstandingReads"));
-TRACE_DECLARE_MEMORY_COUNTER(IoDispatcherOutstandingBytesToRead, TEXT("IoDispatcher/OutstandingBytesToRead"));
 TRACE_DECLARE_INT_COUNTER(IoDispatcherLatencyCircuitBreaks, TEXT("IoDispatcher/LatencyCircuitBreaks"));
 TRACE_DECLARE_INT_COUNTER(IoDispatcherSeekDistanceCircuitBreaks, TEXT("IoDispatcher/SeekDistanceCircuitBreaks"));
 TRACE_DECLARE_INT_COUNTER(IoDispatcherNumPriorityQueues, TEXT("IoDispatcher/NumPriorityQueues"));
@@ -223,16 +217,16 @@ bool FFileIoStoreBlockCache::Read(FFileIoStoreReadRequest* Block)
 
 	if (!CachedBlock)
 	{
-		TRACE_COUNTER_INCREMENT(IoDispatcherCacheMisses);
+		FFileIoStats::OnBlockCacheMiss(ReadBufferSize);
 		return false;
 	}
 	check(CachedBlock->Buffer);
+	FFileIoStats::OnBlockCacheHit(ReadBufferSize);
 	FMemory::Memcpy(Block->Buffer->Memory, CachedBlock->Buffer, ReadBufferSize);
 	{
 		FScopeLock Lock(&CriticalSection);
 		CachedBlock->bLocked = false;
 	}
-	TRACE_COUNTER_INCREMENT(IoDispatcherCacheHits);
 	return true;
 }
 
@@ -273,6 +267,7 @@ void FFileIoStoreBlockCache::Store(const FFileIoStoreReadRequest* Block)
 	check(BlockToReplace);
 	check(BlockToReplace->Buffer);
 	FMemory::Memcpy(BlockToReplace->Buffer, Block->Buffer->Memory, ReadBufferSize);
+	FFileIoStats::OnBlockCacheStore(ReadBufferSize);
 	{
 		FScopeLock Lock(&CriticalSection);
 		BlockToReplace->bLocked = false;
@@ -542,11 +537,10 @@ void FFileIoStoreRequestQueue::Push(FFileIoStoreReadRequest& Request)
 	TRACE_CPUPROFILER_EVENT_SCOPE(RequestQueuePush);
 	FScopeLock _(&CriticalSection);
 	UpdateSortRequestsByOffset();
-	TRACE_COUNTER_INCREMENT(IoDispatcherOutstandingReads);
-	TRACE_COUNTER_ADD(IoDispatcherOutstandingBytesToRead, Request.Size);
 	
 	check(Request.QueueStatus != FFileIoStoreReadRequest::QueueStatus_InQueue);
 	Request.QueueStatus = FFileIoStoreReadRequest::QueueStatus_InQueue;
+
 	if (bSortRequestsByOffset)
 	{
 		PushToPriorityQueues(&Request);
@@ -565,9 +559,6 @@ void FFileIoStoreRequestQueue::Push(FFileIoStoreReadRequestList& Requests)
 
 	for (auto It = Requests.Steal(); It; ++It)
 	{
-		TRACE_COUNTER_INCREMENT(IoDispatcherOutstandingReads);
-		TRACE_COUNTER_ADD(IoDispatcherOutstandingBytesToRead, It->Size);
-
 		check(It->QueueStatus != FFileIoStoreReadRequest::QueueStatus_InQueue);
 		It->QueueStatus = FFileIoStoreReadRequest::QueueStatus_InQueue;
 
@@ -580,6 +571,7 @@ void FFileIoStoreRequestQueue::Push(FFileIoStoreReadRequestList& Requests)
 			Heap.HeapPush(*It, QueueSortFunc);
 		}
 	}
+
 }
 
 void FFileIoStoreRequestQueue::UpdateOrder()
@@ -1157,6 +1149,7 @@ bool FFileIoStore::Resolve(FIoRequestImpl* Request)
 				if (PlatformImpl.CreateCustomRequests(RequestAllocator, *ResolvedRequest, CustomRequests))
 				{
 					RequestTracker.AddReadRequestsToResolvedRequest(CustomRequests, *ResolvedRequest);
+					FFileIoStats::OnFilesystemReadsQueued(CustomRequests);
 					RequestQueue.Push(CustomRequests);
 					OnNewPendingRequestsAdded();
 				}
@@ -1367,6 +1360,8 @@ void FFileIoStore::CompleteDispatcherRequest(FFileIoStoreResolvedRequest* Resolv
 
 void FFileIoStore::FinalizeCompressedBlock(FFileIoStoreCompressedBlock* CompressedBlock)
 {
+	FFileIoStats::OnDecompressComplete(CompressedBlock->CompressedSize, CompressedBlock->UncompressedSize); 
+
 	if (CompressedBlock->RawBlocks.Num() > 1)
 	{
 		check(CompressedBlock->CompressedDataBuffer || CompressedBlock->bCancelled || CompressedBlock->bFailed);
@@ -1397,7 +1392,7 @@ void FFileIoStore::FinalizeCompressedBlock(FFileIoStoreCompressedBlock* Compress
 	for (int32 ScatterIndex = 0, ScatterCount = CompressedBlock->ScatterList.Num(); ScatterIndex < ScatterCount; ++ScatterIndex)
 	{
 		FFileIoStoreBlockScatter& Scatter = CompressedBlock->ScatterList[ScatterIndex];
-		TRACE_COUNTER_ADD(IoDispatcherTotalBytesScattered, Scatter.Size);
+		FFileIoStats::OnBytesScattered(Scatter.Size);
 		Scatter.Request->bFailed |= CompressedBlock->bFailed;
 		check(!CompressedBlock->bCancelled || !Scatter.Request->DispatcherRequest || Scatter.Request->DispatcherRequest->IsCancelled());
 		check(Scatter.Request->UnfinishedReadsCount > 0);
@@ -1428,9 +1423,7 @@ FIoRequestImpl* FFileIoStore::GetCompletedRequests()
 	{
 		FFileIoStoreReadRequest* CompletedRequest = *It;
 
-		TRACE_COUNTER_ADD(IoDispatcherTotalBytesRead, CompletedRequest->Size);
-		TRACE_COUNTER_DECREMENT(IoDispatcherOutstandingReads);
-		TRACE_COUNTER_SUBTRACT(IoDispatcherOutstandingBytesToRead, CompletedRequest->Size);
+		FFileIoStats::OnReadComplete(CompletedRequest->Size);
 
 		if (!CompletedRequest->ImmediateScatter.Request)
 		{
@@ -1488,6 +1481,7 @@ FIoRequestImpl* FFileIoStore::GetCompletedRequests()
 				check(CompressedBlock->UnfinishedRawBlocksCount > 0);
 				if (--CompressedBlock->UnfinishedRawBlocksCount == 0)
 				{
+					FFileIoStats::OnDecompressQueued(CompressedBlock->CompressedSize, CompressedBlock->UncompressedSize);
 					RequestTracker.RemoveCompressedBlock(CompressedBlock);
 					if (!ReadyForDecompressionTail)
 					{
@@ -1504,7 +1498,7 @@ FIoRequestImpl* FFileIoStore::GetCompletedRequests()
 		}
 		else
 		{
-			TRACE_COUNTER_ADD(IoDispatcherTotalBytesScattered, CompletedRequest->ImmediateScatter.Size);
+			FFileIoStats::OnBytesScattered(CompletedRequest->ImmediateScatter.Size);
 
 			check(!CompletedRequest->Buffer);
 			FFileIoStoreResolvedRequest* CompletedResolvedRequest = CompletedRequest->ImmediateScatter.Request;
@@ -1733,6 +1727,7 @@ void FFileIoStore::ReadBlocks(FFileIoStoreResolvedRequest& ResolvedRequest)
 
 	if (!NewBlocks.IsEmpty())
 	{
+		FFileIoStats::OnFilesystemReadsQueued(NewBlocks);
 		RequestQueue.Push(NewBlocks);
 		OnNewPendingRequestsAdded();
 	}
@@ -1810,4 +1805,369 @@ uint32 FFileIoStore::Run()
 TSharedRef<IIoDispatcherFileBackend> CreateIoDispatcherFileBackend()
 {
 	return MakeShared<FFileIoStore>();
-}
\ No newline at end of file
+}
+
+CSV_DEFINE_CATEGORY(IoDispatcherFileBackend, true);
+
+// These stats go to both insights and csv by default
+// TODO: Ideally these should go to insights even if CSV is not capturing, but not be doubled-up if both CSV and Insights are capturing
+// TODO: It would also be nice to send these to insights as int64 without unit conversion where appropriate
+// IoDispatcher thread
+CSV_DEFINE_STAT(IoDispatcherFileBackend, QueuedFilesystemReads);	
+CSV_DEFINE_STAT(IoDispatcherFileBackend, QueuedFilesystemReadMB);	
+CSV_DEFINE_STAT(IoDispatcherFileBackend, QueuedUncompressBlocks);	
+CSV_DEFINE_STAT(IoDispatcherFileBackend, QueuedUncompressInMB);		
+CSV_DEFINE_STAT(IoDispatcherFileBackend, QueuedUncompressOutMB);	
+CSV_DEFINE_STAT(IoDispatcherFileBackend, FrameBytesReadKB);				
+CSV_DEFINE_STAT(IoDispatcherFileBackend, FrameBytesUncompressedInKB);	
+CSV_DEFINE_STAT(IoDispatcherFileBackend, FrameBytesUncompressedOutKB);	
+CSV_DEFINE_STAT(IoDispatcherFileBackend, FrameBytesScatteredKB);		
+
+// FileIoStore thread
+CSV_DEFINE_STAT(IoDispatcherFileBackend, FrameFilesystemReads);			
+CSV_DEFINE_STAT(IoDispatcherFileBackend, FrameFilesystemBytesReadKB);	
+CSV_DEFINE_STAT(IoDispatcherFileBackend, FrameSequentialReads);			
+CSV_DEFINE_STAT(IoDispatcherFileBackend, FrameForwardSeeks);			
+CSV_DEFINE_STAT(IoDispatcherFileBackend, FrameBackwardSeeks);			
+CSV_DEFINE_STAT(IoDispatcherFileBackend, FrameHandleChangeSeeks);		
+CSV_DEFINE_STAT(IoDispatcherFileBackend, FrameSeeks);					
+CSV_DEFINE_STAT(IoDispatcherFileBackend, FrameSeekDistanceMB);			
+CSV_DEFINE_STAT(IoDispatcherFileBackend, FrameBlockCacheStores);		
+CSV_DEFINE_STAT(IoDispatcherFileBackend, FrameBlockCacheStoresKB);		
+CSV_DEFINE_STAT(IoDispatcherFileBackend, FrameBlockCacheHits);			
+CSV_DEFINE_STAT(IoDispatcherFileBackend, FrameBlockCacheHitKB);			
+CSV_DEFINE_STAT(IoDispatcherFileBackend, FrameBlockCacheMisses);		
+CSV_DEFINE_STAT(IoDispatcherFileBackend, FrameBlockCacheMissKB);		
+
+// These stats only go to Insights because they get very large and insights will automatically convert to friendly units
+// IoDispatcher thread
+// Outstanding bytes uncompressed (or decompressed by filesystem api) queued for reading
+TRACE_DECLARE_MEMORY_COUNTER(IoDispatcherFileBackendQueuedFilesystemReadBytes,	TEXT("IoDispatcherFileBackend/QueuedFilesystemReadBytes")); 
+
+// These are long term totals that grow forever and may be less useful?
+
+// IoDispatcher thread 
+// Total bytes read returned from backend whether read from disk or returned from cache 
+TRACE_DECLARE_MEMORY_COUNTER(IoDispatcherFileBackendTotalBytesRead,				TEXT("IoDispatcherFileBackend/TotalBytesRead"));			
+// Total bytes passing through the decompression stage, even if they were not compressed
+TRACE_DECLARE_MEMORY_COUNTER(IoDispatcherFileBackendTotalBytesUncompressedIn,	TEXT("IoDispatcherFileBackend/TotalBytesUncompressedIn"));
+// Total bytes passing through the decompression stage, even if they were not compressed
+TRACE_DECLARE_MEMORY_COUNTER(IoDispatcherFileBackendTotalBytesUncompressedOut,	TEXT("IoDispatcherFileBackend/TotalBytesUncompressedOut"));
+// Total bytes provided into IoBuffers for users 
+TRACE_DECLARE_MEMORY_COUNTER(IoDispatcherFileBackendTotalBytesScattered,		TEXT("IoDispatcherFileBackend/TotalBytesScattered"));
+
+// FileIoStore thread
+// Total number of reads executed through filesystem API layer (useful for calculating ratios of seeks vs sequential reads)
+TRACE_DECLARE_MEMORY_COUNTER(IoDispatcherFileBackendTotalFilesystemReads,		TEXT("IoDispatcherFileBackend/TotalFilesystemReads"));		
+// Total bytes read from platform uncompressed or decompressed for us below the API layer
+TRACE_DECLARE_MEMORY_COUNTER(IoDispatcherFileBackendTotalFilesystemBytesRead,	TEXT("IoDispatcherFileBackend/TotalFilesystemBytesRead"));
+// Total distance in bytes from seeks backwards and forwards
+TRACE_DECLARE_MEMORY_COUNTER(IoDispatcherFileBackendTotalSeekDistance,			TEXT("IoDispatcherFileBackend/TotalSeekDistance"));
+// Total bytes stored in block cache including replacements
+TRACE_DECLARE_MEMORY_COUNTER(IoDispatcherFileBackendTotalBlockCacheBytesStored,	TEXT("IoDispatcherFileBackend/TotalBlockCacheBytesStored"));
+// Total bytes retrieved from block cache
+TRACE_DECLARE_MEMORY_COUNTER(IoDispatcherFileBackendTotalBlockCacheBytesHit,	TEXT("IoDispatcherFileBackend/TotalBlockCacheBytesHit"));
+// Total bytes we failed to retrieve from block cache & had to read from the filesystem
+TRACE_DECLARE_MEMORY_COUNTER(IoDispatcherFileBackendTotalBlockCacheBytesMissed,	TEXT("IoDispatcherFileBackend/TotalBlockCacheBytesMissed"));
+// Total reads that the platform backend considered sequential 
+TRACE_DECLARE_INT_COUNTER(IoDispatcherFileBackendTotalSequentialReads,		TEXT("IoDispatcherFileBackend/TotalSequentialReads"));	
+TRACE_DECLARE_INT_COUNTER(IoDispatcherFileBackendTotalForwardSeeks,			TEXT("IoDispatcherFileBackend/TotalForwardSeeks"));		
+TRACE_DECLARE_INT_COUNTER(IoDispatcherFileBackendTotalBackwardSeeks,		TEXT("IoDispatcherFileBackend/TotalBackwardSeeks"));	
+TRACE_DECLARE_INT_COUNTER(IoDispatcherFileBackendTotalHandleChangeSeeks,	TEXT("IoDispatcherFileBackend/TotalHandleChangeSeeks"));
+// Total seeks including backwards and forwards
+TRACE_DECLARE_INT_COUNTER(IoDispatcherFileBackendTotalSeeks,			TEXT("IoDispatcherFileBackend/TotalSeeks"));
+TRACE_DECLARE_INT_COUNTER(IoDispatcherFileBackendTotalBlockCacheStores, TEXT("IoDispatcherFileBackend/TotalBlockCacheStores"));
+TRACE_DECLARE_INT_COUNTER(IoDispatcherFileBackendTotalBlockCacheHits,	TEXT("IoDispatcherFileBackend/TotalBlockCacheHits"));
+TRACE_DECLARE_INT_COUNTER(IoDispatcherFileBackendTotalBlockCacheMisses, TEXT("IoDispatcherFileBackend/TotalBlockCacheMisses"));
+
+
+#if CSV_PROFILER
+// IoDispatcher thread
+uint64 FFileIoStats::QueuedFilesystemReadBytes{ 0 };	
+uint64 FFileIoStats::QueuedFilesystemReads{ 0 };
+uint64 FFileIoStats::QueuedUncompressBytesIn{ 0 };
+uint64 FFileIoStats::QueuedUncompressBytesOut{ 0 };
+uint64 FFileIoStats::QueuedUncompressBlocks{ 0 };
+#endif
+
+#if IO_DISPATCHER_FILE_STATS
+// IoDispatcher thread
+uint64 FFileIoStats::LastHandle = 0;
+uint64 FFileIoStats::LastOffset = 0;
+
+uint32 FFileIoStats::FileIoStoreThreadId = 0;
+uint32 FFileIoStats::IoDispatcherThreadId = 0;
+#endif
+
+bool FFileIoStats::IsInIoDispatcherThread()
+{
+#if IO_DISPATCHER_FILE_STATS
+	if (FPlatformProcess::SupportsMultithreading())
+	{
+		return FPlatformTLS::GetCurrentThreadId() == IoDispatcherThreadId;
+	}
+#endif
+	return true;
+}
+
+bool FFileIoStats::IsInFileIoStoreThread()
+{
+#if IO_DISPATCHER_FILE_STATS
+	if (FPlatformProcess::SupportsMultithreading())
+	{
+		return FPlatformTLS::GetCurrentThreadId() == FileIoStoreThreadId;
+	}
+#endif
+	return true;
+}
+
+void FFileIoStats::SetThreadIds(uint32 InFileIoStoreThreadId, uint32 InIoDispatcherThreadId)
+{
+#if IO_DISPATCHER_FILE_STATS
+	FileIoStoreThreadId = InFileIoStoreThreadId;
+	IoDispatcherThreadId = InIoDispatcherThreadId;
+#endif
+}
+
+void FFileIoStats::OnCloseHandle(uint64 Handle)
+{
+#if IO_DISPATCHER_FILE_STATS
+	checkSlow(IsInIoDispatcherThread());
+	if( LastHandle == Handle )
+	{
+		LastHandle = 0;
+	}
+#endif
+}
+
+void FFileIoStats::OnFilesystemReadsQueued(const FFileIoStoreReadRequestList& Requests)
+{
+	uint64 TotalBytes = 0;
+	int32 NumReads = 0;
+	for (const FFileIoStoreReadRequest* Request : Requests)
+	{
+		++NumReads;
+		TotalBytes += Request->Size;
+	}
+
+	OnFilesystemReadsQueued(TotalBytes, NumReads);
+}
+
+void FFileIoStats::OnFilesystemReadsQueued(uint64 BytesToRead, int32 NumReads)
+{
+#if IO_DISPATCHER_FILE_STATS
+	checkSlow(IsInIoDispatcherThread());
+
+#if CSV_PROFILER
+	uint64 LocalQueuedFilesystemReadBytes = QueuedFilesystemReadBytes += BytesToRead;
+	CSV_CUSTOM_STAT_DEFINED(QueuedFilesystemReadMB, BytesToApproxMB(LocalQueuedFilesystemReadBytes), ECsvCustomStatOp::SetAndHold);
+
+	uint64 LocalQueuedFilesystemReads = QueuedFilesystemReads += NumReads;
+	CSV_CUSTOM_STAT_DEFINED(QueuedFilesystemReads, (int32)LocalQueuedFilesystemReads, ECsvCustomStatOp::SetAndHold);
+#endif
+
+	TRACE_COUNTER_ADD(IoDispatcherFileBackendQueuedFilesystemReadBytes, BytesToRead);
+#endif
+}
+
+void FFileIoStats::OnFilesystemReadStarted(uint64 Handle, uint64 Offset, uint64 NumBytes)
+{
+#if IO_DISPATCHER_FILE_STATS
+	checkSlow(IsInFileIoStoreThread());
+
+	CSV_CUSTOM_STAT_DEFINED(FrameFilesystemReads, 1, ECsvCustomStatOp::Accumulate);
+	TRACE_COUNTER_INCREMENT(IoDispatcherFileBackendTotalFilesystemReads);
+
+	if (LastHandle != Handle)
+	{
+		OnHandleChangeSeek();
+	}
+	else if (LastOffset == Offset)
+	{
+		OnSequentialRead();
+	}
+	else
+	{
+		OnSeek(LastOffset, Offset);
+	}
+	LastOffset = Offset + NumBytes;
+	LastHandle = Handle;
+#endif
+}
+
+void FFileIoStats::OnFilesystemReadsComplete(const FFileIoStoreReadRequestList& CompletedRequests)
+{
+#if IO_DISPATCHER_FILE_STATS
+	int64 TotalBytes = 0;
+	int32 NumReads = 0;
+	for (const FFileIoStoreReadRequest* Request : CompletedRequests)
+	{
+		++NumReads;
+		TotalBytes += Request->Size;
+	}
+
+	FFileIoStats::OnFilesystemReadsComplete(TotalBytes, NumReads);
+#endif
+}
+
+void FFileIoStats::OnFilesystemReadsComplete(int64 BytesRead, int32 NumReads)
+{
+#if IO_DISPATCHER_FILE_STATS
+	checkSlow(IsInFileIoStoreThread());
+	CSV_CUSTOM_STAT_DEFINED(FrameFilesystemBytesReadKB, BytesToApproxKB(BytesRead), ECsvCustomStatOp::Accumulate);
+
+	TRACE_COUNTER_ADD(IoDispatcherFileBackendTotalFilesystemBytesRead, BytesRead);
+#endif
+}
+
+void FFileIoStats::OnReadComplete(int64 BytesRead)
+{
+#if IO_DISPATCHER_FILE_STATS
+	checkSlow(IsInIoDispatcherThread());
+#if CSV_PROFILER
+	uint64 LocalQueuedFilesystemReadBytes = QueuedFilesystemReadBytes -= BytesRead;
+	CSV_CUSTOM_STAT_DEFINED(QueuedFilesystemReadMB, BytesToApproxMB(LocalQueuedFilesystemReadBytes), ECsvCustomStatOp::SetAndHold);
+
+	uint64 LocalQueuedFilesystemReads = --QueuedFilesystemReads;
+	CSV_CUSTOM_STAT_DEFINED(QueuedFilesystemReads, (int32)LocalQueuedFilesystemReads, ECsvCustomStatOp::SetAndHold);
+#endif
+
+	CSV_CUSTOM_STAT_DEFINED(FrameBytesReadKB, BytesToApproxKB(BytesRead), ECsvCustomStatOp::Accumulate);
+
+	TRACE_COUNTER_SUBTRACT(IoDispatcherFileBackendQueuedFilesystemReadBytes, BytesRead);
+	TRACE_COUNTER_ADD(IoDispatcherFileBackendTotalBytesRead, BytesRead);
+#endif
+}
+
+void FFileIoStats::OnDecompressQueued(int64 CompressedBytes, int64 UncompressedBytes)
+{
+#if IO_DISPATCHER_FILE_STATS
+	checkSlow(IsInIoDispatcherThread());
+#if CSV_PROFILER
+	uint64 QueuedBlocks = ++QueuedUncompressBlocks;
+	CSV_CUSTOM_STAT_DEFINED(QueuedUncompressBlocks, int32(QueuedBlocks), ECsvCustomStatOp::SetAndHold);
+
+	uint64 QueuedBytesIn = QueuedUncompressBytesIn += CompressedBytes;
+	uint64 QueuedBytesOut = QueuedUncompressBytesOut += UncompressedBytes;
+	CSV_CUSTOM_STAT_DEFINED(QueuedUncompressInMB, BytesToApproxMB(QueuedBytesIn), ECsvCustomStatOp::SetAndHold);
+	CSV_CUSTOM_STAT_DEFINED(QueuedUncompressOutMB, BytesToApproxMB(QueuedBytesOut), ECsvCustomStatOp::SetAndHold);
+#endif
+#endif
+}
+
+void FFileIoStats::OnDecompressComplete(int64 CompressedBytes, int64 UncompressedBytes)
+{
+#if IO_DISPATCHER_FILE_STATS
+	checkSlow(IsInIoDispatcherThread());
+#if CSV_PROFILER
+	uint64 QueuedBlocks = --QueuedUncompressBlocks;
+	CSV_CUSTOM_STAT_DEFINED(QueuedUncompressBlocks, int32(QueuedBlocks), ECsvCustomStatOp::SetAndHold);
+
+	uint64 QueuedBytesIn = QueuedUncompressBytesIn -= CompressedBytes;
+	uint64 QueuedBytesOut = QueuedUncompressBytesOut -= UncompressedBytes;
+	CSV_CUSTOM_STAT_DEFINED(QueuedUncompressInMB, BytesToApproxMB(QueuedBytesIn), ECsvCustomStatOp::SetAndHold);
+	CSV_CUSTOM_STAT_DEFINED(QueuedUncompressOutMB, BytesToApproxMB(QueuedBytesOut), ECsvCustomStatOp::SetAndHold);
+#endif
+
+	CSV_CUSTOM_STAT_DEFINED(FrameBytesUncompressedInKB, BytesToApproxKB(CompressedBytes), ECsvCustomStatOp::Accumulate);
+	CSV_CUSTOM_STAT_DEFINED(FrameBytesUncompressedOutKB, BytesToApproxKB(UncompressedBytes), ECsvCustomStatOp::Accumulate);
+
+	TRACE_COUNTER_ADD(IoDispatcherFileBackendTotalBytesUncompressedIn, CompressedBytes);
+	TRACE_COUNTER_ADD(IoDispatcherFileBackendTotalBytesUncompressedOut, UncompressedBytes);
+#endif
+}
+
+void FFileIoStats::OnBytesScattered(int64 NumBytes)
+{
+#if IO_DISPATCHER_FILE_STATS
+	checkSlow(IsInIoDispatcherThread());
+	CSV_CUSTOM_STAT_DEFINED(FrameBytesScatteredKB, BytesToApproxKB(NumBytes), ECsvCustomStatOp::Accumulate);
+	TRACE_COUNTER_ADD(IoDispatcherFileBackendTotalBytesScattered, NumBytes);
+#endif
+}
+
+void FFileIoStats::OnSequentialRead()
+{
+#if IO_DISPATCHER_FILE_STATS
+	checkSlow(IsInFileIoStoreThread());
+	CSV_CUSTOM_STAT_DEFINED(FrameSequentialReads, 1, ECsvCustomStatOp::Accumulate);
+	TRACE_COUNTER_INCREMENT(IoDispatcherFileBackendTotalSequentialReads);
+#endif
+}
+
+void FFileIoStats::OnSeek(uint64 PrevOffset, uint64 NewOffset)
+{
+#if IO_DISPATCHER_FILE_STATS
+	checkSlow(IsInFileIoStoreThread());
+	if (NewOffset > PrevOffset)
+	{
+		int64 Delta = NewOffset - PrevOffset;
+
+		CSV_CUSTOM_STAT_DEFINED(FrameForwardSeeks, 1, ECsvCustomStatOp::Accumulate);
+		CSV_CUSTOM_STAT_DEFINED(FrameSeekDistanceMB, BytesToApproxMB(Delta), ECsvCustomStatOp::Accumulate);
+
+		TRACE_COUNTER_ADD(IoDispatcherFileBackendTotalSeekDistance, Delta);
+		TRACE_COUNTER_INCREMENT(IoDispatcherFileBackendTotalForwardSeeks);
+	}
+	else
+	{
+		int64 Delta = PrevOffset - NewOffset;
+		CSV_CUSTOM_STAT_DEFINED(FrameBackwardSeeks, 1, ECsvCustomStatOp::Accumulate);
+		CSV_CUSTOM_STAT_DEFINED(FrameSeekDistanceMB, BytesToApproxMB(Delta), ECsvCustomStatOp::Accumulate);
+
+		TRACE_COUNTER_ADD(IoDispatcherFileBackendTotalSeekDistance, Delta);
+		TRACE_COUNTER_INCREMENT(IoDispatcherFileBackendTotalBackwardSeeks);
+	}
+
+	CSV_CUSTOM_STAT_DEFINED(FrameSeeks, 1, ECsvCustomStatOp::Accumulate);
+	TRACE_COUNTER_INCREMENT(IoDispatcherFileBackendTotalSeeks);
+#endif
+}
+
+void FFileIoStats::OnHandleChangeSeek()
+{
+#if IO_DISPATCHER_FILE_STATS
+	checkSlow(IsInFileIoStoreThread());
+	CSV_CUSTOM_STAT_DEFINED(FrameHandleChangeSeeks, 1, ECsvCustomStatOp::Accumulate);
+	TRACE_COUNTER_INCREMENT(IoDispatcherFileBackendTotalHandleChangeSeeks);
+
+	CSV_CUSTOM_STAT_DEFINED(FrameSeeks, 1, ECsvCustomStatOp::Accumulate);
+	TRACE_COUNTER_INCREMENT(IoDispatcherFileBackendTotalSeeks);
+#endif
+}
+
+void FFileIoStats::OnBlockCacheStore(uint64 NumBytes)
+{
+#if IO_DISPATCHER_FILE_STATS
+	checkSlow(IsInFileIoStoreThread());
+	CSV_CUSTOM_STAT_DEFINED(FrameBlockCacheStores, 1, ECsvCustomStatOp::Accumulate);
+	CSV_CUSTOM_STAT_DEFINED(FrameBlockCacheStoresKB, BytesToApproxKB(NumBytes), ECsvCustomStatOp::Accumulate);
+
+	TRACE_COUNTER_INCREMENT(IoDispatcherFileBackendTotalBlockCacheStores);
+	TRACE_COUNTER_ADD(IoDispatcherFileBackendTotalBlockCacheBytesStored, NumBytes);
+#endif
+}
+
+void FFileIoStats::OnBlockCacheHit(uint64 NumBytes)
+{
+#if IO_DISPATCHER_FILE_STATS
+	checkSlow(IsInFileIoStoreThread());
+	CSV_CUSTOM_STAT_DEFINED(FrameBlockCacheHits, 1, ECsvCustomStatOp::Accumulate);
+	CSV_CUSTOM_STAT_DEFINED(FrameBlockCacheHitKB, BytesToApproxKB(NumBytes), ECsvCustomStatOp::Accumulate);
+
+	TRACE_COUNTER_INCREMENT(IoDispatcherFileBackendTotalBlockCacheHits);
+	TRACE_COUNTER_ADD(IoDispatcherFileBackendTotalBlockCacheBytesHit, NumBytes);
+#endif
+}
+
+void FFileIoStats::OnBlockCacheMiss(uint64 NumBytes)
+{
+#if IO_DISPATCHER_FILE_STATS
+	checkSlow(IsInFileIoStoreThread());
+	CSV_CUSTOM_STAT_DEFINED(FrameBlockCacheMisses, 1, ECsvCustomStatOp::Accumulate);
+	CSV_CUSTOM_STAT_DEFINED(FrameBlockCacheMissKB, BytesToApproxKB(NumBytes), ECsvCustomStatOp::Accumulate);
+
+	TRACE_COUNTER_INCREMENT(IoDispatcherFileBackendTotalBlockCacheMisses);
+	TRACE_COUNTER_ADD(IoDispatcherFileBackendTotalBlockCacheBytesMissed, NumBytes);
+#endif
+}
diff --git a/Engine/Source/Runtime/Core/Private/IO/IoDispatcherFileBackend.h b/Engine/Source/Runtime/Core/Private/IO/IoDispatcherFileBackend.h
--- a/Engine/Source/Runtime/Core/Private/IO/IoDispatcherFileBackend.h
+++ b/Engine/Source/Runtime/Core/Private/IO/IoDispatcherFileBackend.h
@@ -118,6 +118,8 @@ public:
 	virtual uint32 Run() override;
 	virtual void Stop() override;
 
+	uint32 GetThreadId() const { return Thread ? Thread->GetThreadID() : 0; }
+
 private:
 	class FDecompressAsyncTask
 	{
diff --git a/Engine/Source/Runtime/Core/Private/IO/IoDispatcherFileBackendTypes.h b/Engine/Source/Runtime/Core/Private/IO/IoDispatcherFileBackendTypes.h
--- a/Engine/Source/Runtime/Core/Private/IO/IoDispatcherFileBackendTypes.h
+++ b/Engine/Source/Runtime/Core/Private/IO/IoDispatcherFileBackendTypes.h
@@ -4,6 +4,10 @@
 
 #include "IO/IoStore.h"
 #include "IO/IoDispatcherBackend.h"
+#include "ProfilingDebugging/CsvProfiler.h"
+#include "ProfilingDebugging/CountersTrace.h"
+
+#define IO_DISPATCHER_FILE_STATS (COUNTERSTRACE_ENABLED || CSV_PROFILER)
 
 struct FFileIoStoreCompressionContext;
 
@@ -837,3 +841,69 @@ private:
 	friend class FFileIoStore;
 	friend class FFileIoStoreRequestTracker;
 };
+
+
+// Wrapper for sending stats to both insights and csv profiling from any platform-specific dispatcher.
+struct FFileIoStats
+{
+#if CSV_PROFILER
+	static uint64 QueuedFilesystemReadBytes;
+	static uint64 QueuedFilesystemReads;
+
+	static uint64 QueuedUncompressBytesIn;
+	static uint64 QueuedUncompressBytesOut;
+	static uint64 QueuedUncompressBlocks;
+#endif
+
+#if IO_DISPATCHER_FILE_STATS
+	// Used for seek tracking 
+	static uint64 LastHandle;
+	static uint64 LastOffset;
+
+	static uint32 FileIoStoreThreadId;
+	static uint32 IoDispatcherThreadId;
+#endif
+
+	static float BytesToApproxMB(uint64 Bytes) {
+		return float(double(Bytes) / 1024.0 / 1024.0);
+	}
+	static float BytesToApproxKB(uint64 Bytes) {
+		return float(double(Bytes) / 1024.0);
+	}
+
+	static bool IsInIoDispatcherThread();
+	static bool IsInFileIoStoreThread();
+	static void SetThreadIds(uint32 InIoDispatcherThreadId, uint32 InFileIoStoreThreadId);
+
+	// Adds an underlying filesystem read to the queue to track num reads/bytes in the queue
+	// There is currently no separate stat for in-flight reads, those are still counted in the queued stat
+	static void OnFilesystemReadsQueued(const FFileIoStoreReadRequestList& Requests);
+	static void OnFilesystemReadsQueued(uint64 NumBytes, int32 NumReads);
+	// Called when the backend commits to the next read to check for estimated seek tracking
+	static void OnFilesystemReadStarted(uint64 Handle, uint64 Offset, uint64 NumBytes);
+	// Underlying filesystem read complete, BytesRead is bytes returned from filesystem, possibly already decompressed on some systems
+	static void OnFilesystemReadsComplete(const FFileIoStoreReadRequestList& CompletedRequests);
+	static void OnFilesystemReadsComplete(int64 BytesRead, int32 NumReads=1);
+	// Backend read complete, BytesRead is bytes returned from either block cache or filesystem
+	static void OnReadComplete(int64 BytesRead);
+	// NumBytes of data passed the decompression stage, even if the compression method was none
+	static void OnDecompressQueued(int64 CompressedBytes, int64 UncompressedBytes);
+	static void OnDecompressComplete(int64 CompressedBytes, int64 UncompressedBytes);
+
+	// Bytes were copied into the buffer provided to users of the io system
+	static void OnBytesScattered(int64 BytesScattered);
+	// Record stats for block cache hit/miss rate & data throughput
+	static void OnBlockCacheStore(uint64 NumBytes);
+	static void OnBlockCacheHit(uint64 NumBytes);
+	static void OnBlockCacheMiss(uint64 NumBytes);
+
+	static void OnCloseHandle(uint64 Handle);
+
+private:
+	// A read was started which the backend considers sequential
+	static void OnSequentialRead();
+	// A read was started which the backend considers a forward or reverse seek
+	static void OnSeek(uint64 LastOffset, uint64 NewOffset);
+	// A read was started from a different file handle from the last one we used
+	static void OnHandleChangeSeek();
+};
diff --git a/Engine/Source/Runtime/CoreUObject/Private/Serialization/AsyncLoading2.cpp b/Engine/Source/Runtime/CoreUObject/Private/Serialization/AsyncLoading2.cpp
--- a/Engine/Source/Runtime/CoreUObject/Private/Serialization/AsyncLoading2.cpp
+++ b/Engine/Source/Runtime/CoreUObject/Private/Serialization/AsyncLoading2.cpp
@@ -66,6 +66,11 @@
 #if UE_BUILD_DEVELOPMENT || UE_BUILD_DEBUG
 PRAGMA_DISABLE_OPTIMIZATION
 #endif
+TRACE_DECLARE_MEMORY_COUNTER(AsyncLoadingPendingIoRequestsSize, TEXT("AsyncLoading/PendingIoRequestsSize"));
+CSV_DEFINE_CATEGORY(AsyncLoading, true);
+CSV_DEFINE_STAT(AsyncLoading, PendingIoRequestsSizeMB);
+CSV_DEFINE_STAT(AsyncLoading, FrameCompletedBundleLoadsKB);
+
 int32 GAsyncLoadingMaxPendingRequestsSizeMB = 256;
 static FAutoConsoleVariableRef CVar_AsyncLoadingMaxPendingRequestsSizeMB(
 	TEXT("s.AsyncLoadingMaxPendingRequestsSizeMB"),
@@ -2593,6 +2598,9 @@ void FAsyncLoadingThread2::BundleIoRequestCompleted(FAsyncPackage2* Package)
 {
 	check(PendingBundleIoRequestsTotalSize >= Package->ExportBundlesSize)
 	PendingBundleIoRequestsTotalSize -= Package->ExportBundlesSize;
+	TRACE_COUNTER_SET(AsyncLoadingPendingIoRequestsSize, PendingBundleIoRequestsTotalSize);
+	CSV_CUSTOM_STAT_DEFINED(PendingIoRequestsSizeMB, float((double)PendingBundleIoRequestsTotalSize / 1024.0 / 1024.0), ECsvCustomStatOp::SetAndHold);
+	CSV_CUSTOM_STAT_DEFINED(FrameCompletedBundleLoadsKB, float((double)Package->ExportBundlesSize / 1024.0), ECsvCustomStatOp::Accumulate);
 	if (WaitingIoRequests.Num())
 	{
 		StartBundleIoRequests();
@@ -2638,6 +2646,9 @@ void FAsyncLoadingThread2::StartBundleIoRequests()
 		});
 	}
 	IoBatch.Issue();
+
+	TRACE_COUNTER_SET(AsyncLoadingPendingIoRequestsSize, PendingBundleIoRequestsTotalSize);
+	CSV_CUSTOM_STAT_DEFINED(PendingIoRequestsSizeMB, float((double)PendingBundleIoRequestsTotalSize / 1024.0 / 1024.0), ECsvCustomStatOp::SetAndHold);
 }
 
 FEventLoadNode2::FEventLoadNode2(const FAsyncLoadEventSpec* InSpec, FAsyncPackage2* InPackage, int32 InImportOrExportIndex, int32 InBarrierCount)
