Fix integration from UE4 to track thread ids for stats.

diff --git a/Engine/Source/Runtime/Core/Private/IO/IoDispatcher.cpp b/Engine/Source/Runtime/Core/Private/IO/IoDispatcher.cpp
--- a/Engine/Source/Runtime/Core/Private/IO/IoDispatcher.cpp
+++ b/Engine/Source/Runtime/Core/Private/IO/IoDispatcher.cpp
@@ -17,6 +17,7 @@
 #include "ProfilingDebugging/CountersTrace.h"
 #include "ProfilingDebugging/CsvProfiler.h"
 #include "IO/IoDispatcherBackend.h"
+#include "IO/IoDispatcherFileBackendTypes.h"
 
 DEFINE_LOG_CATEGORY(LogIoDispatcher);
 
@@ -174,7 +175,7 @@ public:
 			Backend->Initialize(BackendContext);
 		}
 		Thread = FRunnableThread::Create(this, TEXT("IoDispatcher"), 0, TPri_AboveNormal, FPlatformAffinity::GetIoDispatcherThreadMask());
-		FFileIoStats::SetThreadIds(FileIoStore.GetThreadId(), Thread ? Thread->GetThreadID() : 0);
+		FFileIoStats::SetDispatcherThreadId(Thread ? Thread->GetThreadID() : 0);
 		return true;
 	}
 
diff --git a/Engine/Source/Runtime/Core/Private/IO/IoDispatcherFileBackend.cpp b/Engine/Source/Runtime/Core/Private/IO/IoDispatcherFileBackend.cpp
--- a/Engine/Source/Runtime/Core/Private/IO/IoDispatcherFileBackend.cpp
+++ b/Engine/Source/Runtime/Core/Private/IO/IoDispatcherFileBackend.cpp
@@ -1034,6 +1034,7 @@ void FFileIoStore::Initialize(TSharedRef<const FIoDispatcherBackendContext> InCo
 	}
 
 	Thread = FRunnableThread::Create(this, TEXT("IoService"), 0, TPri_AboveNormal);
+	FFileIoStats::SetFileIoStoreThreadId(Thread ? Thread->GetThreadID() : 0);
 }
 
 TIoStatusOr<FIoContainerId> FFileIoStore::Mount(const TCHAR* ContainerPath, int32 Order, const FGuid& EncryptionKeyGuid, const FAES::FAESKey& EncryptionKey)
@@ -1921,14 +1922,20 @@ bool FFileIoStats::IsInFileIoStoreThread()
 	return true;
 }
 
-void FFileIoStats::SetThreadIds(uint32 InFileIoStoreThreadId, uint32 InIoDispatcherThreadId)
+void FFileIoStats::SetDispatcherThreadId(uint32 InIoDispatcherThreadId)
 {
 #if IO_DISPATCHER_FILE_STATS
-	FileIoStoreThreadId = InFileIoStoreThreadId;
 	IoDispatcherThreadId = InIoDispatcherThreadId;
 #endif
 }
 
+void FFileIoStats::SetFileIoStoreThreadId(uint32 InFileIoStoreThreadId)
+{
+#if IO_DISPATCHER_FILE_STATS
+	FileIoStoreThreadId = InFileIoStoreThreadId;
+#endif
+}
+
 void FFileIoStats::OnCloseHandle(uint64 Handle)
 {
 #if IO_DISPATCHER_FILE_STATS
diff --git a/Engine/Source/Runtime/Core/Private/IO/IoDispatcherFileBackend.h b/Engine/Source/Runtime/Core/Private/IO/IoDispatcherFileBackend.h
--- a/Engine/Source/Runtime/Core/Private/IO/IoDispatcherFileBackend.h
+++ b/Engine/Source/Runtime/Core/Private/IO/IoDispatcherFileBackend.h
@@ -118,8 +118,6 @@ public:
 	virtual uint32 Run() override;
 	virtual void Stop() override;
 
-	uint32 GetThreadId() const { return Thread ? Thread->GetThreadID() : 0; }
-
 private:
 	class FDecompressAsyncTask
 	{
diff --git a/Engine/Source/Runtime/Core/Private/IO/IoDispatcherFileBackendTypes.h b/Engine/Source/Runtime/Core/Private/IO/IoDispatcherFileBackendTypes.h
--- a/Engine/Source/Runtime/Core/Private/IO/IoDispatcherFileBackendTypes.h
+++ b/Engine/Source/Runtime/Core/Private/IO/IoDispatcherFileBackendTypes.h
@@ -873,7 +873,8 @@ struct FFileIoStats
 
 	static bool IsInIoDispatcherThread();
 	static bool IsInFileIoStoreThread();
-	static void SetThreadIds(uint32 InIoDispatcherThreadId, uint32 InFileIoStoreThreadId);
+	static void SetDispatcherThreadId(uint32 InIoDispatcherThreadId);
+	static void SetFileIoStoreThreadId(uint32 InFileIoStoreThreadId);
 
 	// Adds an underlying filesystem read to the queue to track num reads/bytes in the queue
 	// There is currently no separate stat for in-flight reads, those are still counted in the queued stat
