removed changes that slipped into 16481905 by mistake

diff --git a/Engine/Source/Runtime/Core/Private/Async/TaskGraph.cpp b/Engine/Source/Runtime/Core/Private/Async/TaskGraph.cpp
--- a/Engine/Source/Runtime/Core/Private/Async/TaskGraph.cpp
+++ b/Engine/Source/Runtime/Core/Private/Async/TaskGraph.cpp
@@ -57,10 +57,6 @@ static int32 GNumWorkerThreadsToIgnore = 0;
 	#define YIELD_BETWEEN_TASKS 0
 #endif
 
-#if !defined(UE_TASKGRAPH_BUSYWAIT)
-	#define UE_TASKGRAPH_BUSYWAIT 0
-#endif
-
 namespace ENamedThreads
 {
 	CORE_API TAtomic<Type> FRenderThreadStatics::RenderThread(ENamedThreads::GameThread); // defaults to game and is set and reset by the render thread itself
@@ -2110,7 +2106,6 @@ private:
 			TGraphTask<FReturnGraphTask>::CreateTask(&Tasks, CurrentThread).ConstructAndDispatchWhenReady(CurrentThread);
 			ProcessThreadUntilRequestReturn(CurrentThread);
 		}
-#if UE_TASKGRAPH_BUSYWAIT
 		else if (LowLevelTasks::FScheduler::Get().IsWorkerThread())
 		{
 			LowLevelTasks::BusyWaitUntil([Index(0), &Tasks]() mutable
@@ -2126,7 +2121,6 @@ private:
 				return true;
 			});
 		}
-#endif
 		else
 		{
 			if (!FTaskGraphInterface::IsMultithread())
