[CVARS] - allow scalability set by to avoid the override on double set console variables.

this avoids a bug where UE uses a default scalability to set a non engine CVAR(does not exist yet) but the local project uses a local setting that is not default. In this scenario we want the local setting to win however without the change in this CL, the default setting persists.

diff --git a/Engine/Source/Runtime/Core/Private/HAL/ConsoleManager.cpp b/Engine/Source/Runtime/Core/Private/HAL/ConsoleManager.cpp
--- a/Engine/Source/Runtime/Core/Private/HAL/ConsoleManager.cpp
+++ b/Engine/Source/Runtime/Core/Private/HAL/ConsoleManager.cpp
@@ -1544,12 +1544,54 @@ IConsoleObject* FConsoleManager::AddConsoleObject(const TCHAR* Name, IConsoleObj
 		{
 			if(ExistingVar->TestFlags(ECVF_CreatedFromIni))
 			{
+				// when a cvar is not yet created, and scalability sets it, the cvar will be setup deferred. (ECVF_CreatedFromIni)
+				// the cheat protection code below ensures that a deferred variable cannot be later overridden again before the real cvar is created
+				// the fallout here however is that when the Engine initialization for scalability sets EPIC, but the project sets LOW 
+				// when the cvar is finally created, the project will receive the EPIC setting but the UI within the editor will show LOW.
+				// setting [ScalabilitySettings] AllowScalabilityOverrides = true will allow this to succeed in !SHIPPING builds and warn
+				// that is it occurring otherwise.
+				// shipping builds remain unaffected
+				bool bAllowScalabilityOverrides = false;
+
+#if !UE_BUILD_SHIPPING
+				// avoid logging about console variable overrides in retail (it will help cheaters)
+				#define OUTPUT_OVERRIDES_TO_LOG_ENABLED
+
+				// only allow the overrides to succeed in NON shipping builds
+				GConfig->GetBool(TEXT("ScalabilitySettings"), TEXT("AllowScalabilityOverrides"), bAllowScalabilityOverrides, GEditorIni);
+#endif //!UE_BUILD_SHIPPING
+				
+				bool bScalabilityOverride = bAllowScalabilityOverrides && ((Var->GetFlags() & ECVF_SetByMask) == ECVF_SetByScalability) ;
+
 				// This is to prevent cheaters to set a value from an ini of a cvar that is created later
 				// TODO: This is not ideal as it also prevents consolevariables.ini to set the value where we allow that. We could fix that.
-				if(!Var->TestFlags(ECVF_Cheat))
+				if(!Var->TestFlags(ECVF_Cheat) && !bScalabilityOverride)
 				{
 					// The existing one came from the ini, get the value
 					Var->Set(*ExistingVar->GetString(), (EConsoleVariableFlags)((uint32)ExistingVar->GetFlags() & ECVF_SetByMask));
+#ifdef OUTPUT_OVERRIDES_TO_LOG_ENABLED
+					UE_LOG(LogConsoleManager, Warning, TEXT("Console object named [[%s:%s]] already exists but is being registered again - keeping the old value"), Name, *ExistingVar->GetString());
+#endif //OUTPUT_OVERRIDES_TO_LOG_ENABLED
+				}
+				else
+				{
+#ifdef OUTPUT_OVERRIDES_TO_LOG_ENABLED
+					if (Var->GetString().Equals(ExistingVar->GetString()) == false)
+					{
+						if (bScalabilityOverride)
+						{
+
+							UE_LOG(LogConsoleManager, Error,
+								TEXT("Console object [[%s:%s]] is being re-registered through scalability - overloading with new value [[%s]] this will NOT work in a SHIPPING build"),
+								Name, *ExistingVar->GetString(), *Var->GetString());
+
+						}
+						else
+						{
+							UE_LOG(LogConsoleManager, Warning, TEXT("Console object [[%s:%s]] is being re-registered - overloading with new value [[%s]]"), Name, *ExistingVar->GetString(), *Var->GetString());
+						}
+					}
+#endif //OUTPUT_OVERRIDES_TO_LOG_ENABLED
 				}
 
 				// destroy the existing one (no need to call sink because that will happen after all ini setting have been loaded)
diff --git a/Engine/Source/Runtime/Core/Private/Misc/ConfigCacheIni.cpp b/Engine/Source/Runtime/Core/Private/Misc/ConfigCacheIni.cpp
--- a/Engine/Source/Runtime/Core/Private/Misc/ConfigCacheIni.cpp
+++ b/Engine/Source/Runtime/Core/Private/Misc/ConfigCacheIni.cpp
@@ -4949,6 +4949,9 @@ CORE_API void OnSetCVarFromIniEntry(const TCHAR *IniFile, const TCHAR *Key, cons
 		// this is important for variables created in external modules, such as the game module
 		IConsoleManager::Get().RegisterConsoleVariable(Key, Value, TEXT("IAmNoRealVariable"),
 			(uint32)ECVF_Unregistered | (uint32)ECVF_CreatedFromIni | SetBy);
+#if !UE_BUILD_SHIPPING
+		UE_LOG(LogConfig, Log, TEXT("CVar [[%s:%s]] deferred - dummy variable created"), Key, Value);
+#endif //!UE_BUILD_SHIPPING
 	}
 }
 
