[CVARS] - disable warning on duplicate ini entry to fix builds

diff --git a/Engine/Source/Runtime/Core/Private/HAL/ConsoleManager.cpp b/Engine/Source/Runtime/Core/Private/HAL/ConsoleManager.cpp
--- a/Engine/Source/Runtime/Core/Private/HAL/ConsoleManager.cpp
+++ b/Engine/Source/Runtime/Core/Private/HAL/ConsoleManager.cpp
@@ -1555,7 +1555,7 @@ IConsoleObject* FConsoleManager::AddConsoleObject(const TCHAR* Name, IConsoleObj
 
 #if !UE_BUILD_SHIPPING
 				// avoid logging about console variable overrides in retail (it will help cheaters)
-				#define OUTPUT_OVERRIDES_TO_LOG_ENABLED
+				// #define OUTPUT_OVERRIDES_TO_LOG_ENABLED
 
 				// only allow the overrides to succeed in NON shipping builds
 				GConfig->GetBool(TEXT("ScalabilitySettings"), TEXT("AllowScalabilityOverrides"), bAllowScalabilityOverrides, GEditorIni);
