Expose IoDispatcher and AsyncLoading stats to csv profiler.

diff --git a/Engine/Source/Runtime/Core/Private/IO/IoDispatcherFileBackend.cpp b/Engine/Source/Runtime/Core/Private/IO/IoDispatcherFileBackend.cpp
--- a/Engine/Source/Runtime/Core/Private/IO/IoDispatcherFileBackend.cpp
+++ b/Engine/Source/Runtime/Core/Private/IO/IoDispatcherFileBackend.cpp
@@ -1882,6 +1882,114 @@ TRACE_DECLARE_INT_COUNTER(IoDispatcherFileBackendTotalBlockCacheHits,	TEXT("IoDi
 TRACE_DECLARE_INT_COUNTER(IoDispatcherFileBackendTotalBlockCacheMisses, TEXT("IoDispatcherFileBackend/TotalBlockCacheMisses"));
 
 
+#if CSV_PROFILER
+// IoDispatcher thread
+uint64 FFileIoStats::QueuedFilesystemReadBytes{ 0 };	
+uint64 FFileIoStats::QueuedFilesystemReads{ 0 };
+uint64 FFileIoStats::QueuedUncompressBytesIn{ 0 };
+uint64 FFileIoStats::QueuedUncompressBytesOut{ 0 };
+uint64 FFileIoStats::QueuedUncompressBlocks{ 0 };
+#endif
+
+#if IO_DISPATCHER_FILE_STATS
+// IoDispatcher thread
+uint64 FFileIoStats::LastHandle = 0;
+uint64 FFileIoStats::LastOffset = 0;
+
+uint32 FFileIoStats::FileIoStoreThreadId = 0;
+uint32 FFileIoStats::IoDispatcherThreadId = 0;
+#endif
+
+bool FFileIoStats::IsInIoDispatcherThread()
+{
+#if IO_DISPATCHER_FILE_STATS
+	if (FPlatformProcess::SupportsMultithreading())
+	{
+		return FPlatformTLS::GetCurrentThreadId() == IoDispatcherThreadId;
+	}
+#endif
+	return true;
+}
+
+uint32 FFileIoStore::GetThreadId() const
+{
+	return Thread ? Thread->GetThreadID() : 0; 
+}
+
+CSV_DEFINE_CATEGORY(IoDispatcherFileBackend, true);
+
+// These stats go to both insights and csv by default
+// TODO: Ideally these should go to insights even if CSV is not capturing, but not be doubled-up if both CSV and Insights are capturing
+// TODO: It would also be nice to send these to insights as int64 without unit conversion where appropriate
+// IoDispatcher thread
+CSV_DEFINE_STAT(IoDispatcherFileBackend, QueuedFilesystemReads);	
+CSV_DEFINE_STAT(IoDispatcherFileBackend, QueuedFilesystemReadMB);	
+CSV_DEFINE_STAT(IoDispatcherFileBackend, QueuedUncompressBlocks);	
+CSV_DEFINE_STAT(IoDispatcherFileBackend, QueuedUncompressInMB);		
+CSV_DEFINE_STAT(IoDispatcherFileBackend, QueuedUncompressOutMB);	
+CSV_DEFINE_STAT(IoDispatcherFileBackend, FrameBytesReadKB);				
+CSV_DEFINE_STAT(IoDispatcherFileBackend, FrameBytesUncompressedInKB);	
+CSV_DEFINE_STAT(IoDispatcherFileBackend, FrameBytesUncompressedOutKB);	
+CSV_DEFINE_STAT(IoDispatcherFileBackend, FrameBytesScatteredKB);		
+
+// FileIoStore thread
+CSV_DEFINE_STAT(IoDispatcherFileBackend, FrameFilesystemReads);			
+CSV_DEFINE_STAT(IoDispatcherFileBackend, FrameFilesystemBytesReadKB);	
+CSV_DEFINE_STAT(IoDispatcherFileBackend, FrameSequentialReads);			
+CSV_DEFINE_STAT(IoDispatcherFileBackend, FrameForwardSeeks);			
+CSV_DEFINE_STAT(IoDispatcherFileBackend, FrameBackwardSeeks);			
+CSV_DEFINE_STAT(IoDispatcherFileBackend, FrameHandleChangeSeeks);		
+CSV_DEFINE_STAT(IoDispatcherFileBackend, FrameSeeks);					
+CSV_DEFINE_STAT(IoDispatcherFileBackend, FrameSeekDistanceMB);			
+CSV_DEFINE_STAT(IoDispatcherFileBackend, FrameBlockCacheStores);		
+CSV_DEFINE_STAT(IoDispatcherFileBackend, FrameBlockCacheStoresKB);		
+CSV_DEFINE_STAT(IoDispatcherFileBackend, FrameBlockCacheHits);			
+CSV_DEFINE_STAT(IoDispatcherFileBackend, FrameBlockCacheHitKB);			
+CSV_DEFINE_STAT(IoDispatcherFileBackend, FrameBlockCacheMisses);		
+CSV_DEFINE_STAT(IoDispatcherFileBackend, FrameBlockCacheMissKB);		
+
+// These stats only go to Insights because they get very large and insights will automatically convert to friendly units
+// IoDispatcher thread
+// Outstanding bytes uncompressed (or decompressed by filesystem api) queued for reading
+TRACE_DECLARE_MEMORY_COUNTER(IoDispatcherFileBackendQueuedFilesystemReadBytes,	TEXT("IoDispatcherFileBackend/QueuedFilesystemReadBytes")); 
+
+// These are long term totals that grow forever and may be less useful?
+
+// IoDispatcher thread 
+// Total bytes read returned from backend whether read from disk or returned from cache 
+TRACE_DECLARE_MEMORY_COUNTER(IoDispatcherFileBackendTotalBytesRead,				TEXT("IoDispatcherFileBackend/TotalBytesRead"));			
+// Total bytes passing through the decompression stage, even if they were not compressed
+TRACE_DECLARE_MEMORY_COUNTER(IoDispatcherFileBackendTotalBytesUncompressedIn,	TEXT("IoDispatcherFileBackend/TotalBytesUncompressedIn"));
+// Total bytes passing through the decompression stage, even if they were not compressed
+TRACE_DECLARE_MEMORY_COUNTER(IoDispatcherFileBackendTotalBytesUncompressedOut,	TEXT("IoDispatcherFileBackend/TotalBytesUncompressedOut"));
+// Total bytes provided into IoBuffers for users 
+TRACE_DECLARE_MEMORY_COUNTER(IoDispatcherFileBackendTotalBytesScattered,		TEXT("IoDispatcherFileBackend/TotalBytesScattered"));
+
+// FileIoStore thread
+// Total number of reads executed through filesystem API layer (useful for calculating ratios of seeks vs sequential reads)
+TRACE_DECLARE_MEMORY_COUNTER(IoDispatcherFileBackendTotalFilesystemReads,		TEXT("IoDispatcherFileBackend/TotalFilesystemReads"));		
+// Total bytes read from platform uncompressed or decompressed for us below the API layer
+TRACE_DECLARE_MEMORY_COUNTER(IoDispatcherFileBackendTotalFilesystemBytesRead,	TEXT("IoDispatcherFileBackend/TotalFilesystemBytesRead"));
+// Total distance in bytes from seeks backwards and forwards
+TRACE_DECLARE_MEMORY_COUNTER(IoDispatcherFileBackendTotalSeekDistance,			TEXT("IoDispatcherFileBackend/TotalSeekDistance"));
+// Total bytes stored in block cache including replacements
+TRACE_DECLARE_MEMORY_COUNTER(IoDispatcherFileBackendTotalBlockCacheBytesStored,	TEXT("IoDispatcherFileBackend/TotalBlockCacheBytesStored"));
+// Total bytes retrieved from block cache
+TRACE_DECLARE_MEMORY_COUNTER(IoDispatcherFileBackendTotalBlockCacheBytesHit,	TEXT("IoDispatcherFileBackend/TotalBlockCacheBytesHit"));
+// Total bytes we failed to retrieve from block cache & had to read from the filesystem
+TRACE_DECLARE_MEMORY_COUNTER(IoDispatcherFileBackendTotalBlockCacheBytesMissed,	TEXT("IoDispatcherFileBackend/TotalBlockCacheBytesMissed"));
+// Total reads that the platform backend considered sequential 
+TRACE_DECLARE_INT_COUNTER(IoDispatcherFileBackendTotalSequentialReads,		TEXT("IoDispatcherFileBackend/TotalSequentialReads"));	
+TRACE_DECLARE_INT_COUNTER(IoDispatcherFileBackendTotalForwardSeeks,			TEXT("IoDispatcherFileBackend/TotalForwardSeeks"));		
+TRACE_DECLARE_INT_COUNTER(IoDispatcherFileBackendTotalBackwardSeeks,		TEXT("IoDispatcherFileBackend/TotalBackwardSeeks"));	
+TRACE_DECLARE_INT_COUNTER(IoDispatcherFileBackendTotalHandleChangeSeeks,	TEXT("IoDispatcherFileBackend/TotalHandleChangeSeeks"));
+// Total seeks including backwards and forwards
+TRACE_DECLARE_INT_COUNTER(IoDispatcherFileBackendTotalSeeks,			TEXT("IoDispatcherFileBackend/TotalSeeks"));
+TRACE_DECLARE_INT_COUNTER(IoDispatcherFileBackendTotalBlockCacheStores, TEXT("IoDispatcherFileBackend/TotalBlockCacheStores"));
+TRACE_DECLARE_INT_COUNTER(IoDispatcherFileBackendTotalBlockCacheHits,	TEXT("IoDispatcherFileBackend/TotalBlockCacheHits"));
+TRACE_DECLARE_INT_COUNTER(IoDispatcherFileBackendTotalBlockCacheMisses, TEXT("IoDispatcherFileBackend/TotalBlockCacheMisses"));
+
+
 #if CSV_PROFILER
 // IoDispatcher thread
 uint64 FFileIoStats::QueuedFilesystemReadBytes{ 0 };	
