UETOOL-3650 - Delete expired UECrashContext-pid.xml that could be left over by crashed/killed CRC
  - Added code to run a clean up on UECrashContext-{pid}.xml files that 30 days old where the process ID (pid in the name) is not running anymore.

diff --git a/Engine/Source/Programs/CrashReportClient/Private/CrashReportClientApp.cpp b/Engine/Source/Programs/CrashReportClient/Private/CrashReportClientApp.cpp
--- a/Engine/Source/Programs/CrashReportClient/Private/CrashReportClientApp.cpp
+++ b/Engine/Source/Programs/CrashReportClient/Private/CrashReportClientApp.cpp
@@ -779,6 +779,15 @@ static void DeleteTempCrashContextFile(uint64 ProcessID)
 	FPlatformFileManager::Get().GetPlatformFile().DeleteFile(*SessionContextFile);
 }
 
+static void DeleteExpiredTempCrashContext(const FTimespan& ExpirationAge)
+{
+	// Clean up old temp context that were likely left over by crashed/killed CRC (unless the process that wrote it is still alive).
+	// If by any chances a CRC instance deletes the context that was produced for another CRC instance that is still running after the
+	// expiration delay, then we may lose the analytics session and not spoof an abnormal shutdown. In general, the consequences are none
+	// to meaningless, so we are better keeping the user disk clean.
+	FGenericCrashContext::CleanupTempSessionContextFiles(ExpirationAge);
+}
+
 #if CRASH_REPORT_WITH_MTBF
 
 template <typename Type>
@@ -1122,6 +1131,9 @@ void RunCrashReportClient(const TCHAR* CommandLine)
 		// Clean up the context file
 		DeleteTempCrashContextFile(MonitorPid);
 
+		// Clean up left-over context files that weren't cleaned up properly by previous instance(s) (because it was killed, it crashed or user logged out).
+		DeleteExpiredTempCrashContext(FTimespan::FromDays(30));
+
 		FPlatformProcess::CloseProc(MonitoredProcess);
 	}
 
diff --git a/Engine/Source/Runtime/Core/Private/GenericPlatform/GenericPlatformCrashContext.cpp b/Engine/Source/Runtime/Core/Private/GenericPlatform/GenericPlatformCrashContext.cpp
--- a/Engine/Source/Runtime/Core/Private/GenericPlatform/GenericPlatformCrashContext.cpp
+++ b/Engine/Source/Runtime/Core/Private/GenericPlatform/GenericPlatformCrashContext.cpp
@@ -13,6 +13,7 @@
 #include "Misc/Guid.h"
 #include "Misc/SecureHash.h"
 #include "Containers/Ticker.h"
+#include "Containers/StringFwd.h"
 #include "Misc/ConfigCacheIni.h"
 #include "Misc/CoreDelegates.h"
 #include "Misc/App.h"
@@ -24,6 +25,7 @@
 #include "Internationalization/Regex.h"
 #include "Internationalization/TextLocalizationManager.h"
 #include "Logging/LogScopedCategoryAndVerbosityOverride.h"
+#include "HAL/FileManager.h"
 #include "HAL/PlatformOutputDevices.h"
 #include "HAL/PlatformFileManager.h"
 #include "Misc/OutputDeviceArchiveWrapper.h"
@@ -473,6 +475,56 @@ FString FGenericCrashContext::GetTempSessionContextFilePath(uint64 ProcessID)
 	return FPlatformProcess::UserTempDir() / FString::Printf(TEXT("UECrashContext-%u.xml"), ProcessID);
 }
 
+void FGenericCrashContext::CleanupTempSessionContextFiles(const FTimespan& ExpirationAge)
+{
+	const FString BasePathname = FPlatformProcess::UserTempDir() / FString::Printf(TEXT("UECrashContext-"));
+	IFileManager::Get().IterateDirectory(FPlatformProcess::UserTempDir(), [&BasePathname, &ExpirationAge](const TCHAR* Pathname, bool bIsDirectory) -> bool
+	{
+		if (bIsDirectory)
+		{
+			return true; // Looking for files, continue.
+		}
+
+		FStringView PathnameView(Pathname);
+		if (PathnameView.EndsWith(TEXT(".xml")) && PathnameView.StartsWith(BasePathname))
+		{
+			if (IFileManager::Get().GetFileAgeSeconds(Pathname) > ExpirationAge.GetTotalSeconds())
+			{
+				// Extract the process ID from the pathname view.
+				static_assert(sizeof(decltype(FPlatformProcess::GetCurrentProcessId())) == 4, "The code below assumes the process ID is 4 bytes");
+				constexpr uint32 MaxPidStrLen = 10; // std::numeric_limit<uint32>::max() is 4294967295 -> 10 characters.
+				TCHAR PidStr[MaxPidStrLen + 1]; // +1 for null terminator.
+				int32 CharIndex = 0;
+				PathnameView.FindLastChar(TEXT('-'), CharIndex);
+				PathnameView.RemovePrefix(CharIndex + 1); // Remove everything up to the last '-'.
+				PathnameView.FindLastChar(TEXT('.'), CharIndex);
+				PathnameView.RemoveSuffix(PathnameView.Len() - CharIndex); // Remove the trailing '.xml'
+
+				if (PathnameView.Len() <= MaxPidStrLen)
+				{
+					// Put back what we expect to be the PID into a null terminated string.
+					PathnameView.CopyString(PidStr, PathnameView.Len());
+					PidStr[PathnameView.Len()] = TEXT('\0');
+
+					// Converts the PID string, validating it only contained digits.
+					TCHAR* End = nullptr;
+					int64 ProcessId = FCString::Strtoi64(PidStr, &End, 10);
+					if (End == PidStr + PathnameView.Len())
+					{
+						// Ensure the process that created this context file is not running anymore before deleting it.
+						if (!FPlatformProcess::IsApplicationRunning(static_cast<uint32>(ProcessId)))
+						{
+							IFileManager::Get().Delete(Pathname);
+						}
+					}
+				}
+			}
+		}
+
+		return true; // Iterate to next file.
+	});
+}
+
 TOptional<int32> FGenericCrashContext::GetOutOfProcessCrashReporterExitCode()
 {
 	TOptional<int32> ExitCode;
diff --git a/Engine/Source/Runtime/Core/Public/GenericPlatform/GenericPlatformCrashContext.h b/Engine/Source/Runtime/Core/Public/GenericPlatform/GenericPlatformCrashContext.h
--- a/Engine/Source/Runtime/Core/Public/GenericPlatform/GenericPlatformCrashContext.h
+++ b/Engine/Source/Runtime/Core/Public/GenericPlatform/GenericPlatformCrashContext.h
@@ -333,6 +333,9 @@ public:
 	/** Get the file path to the temporary session context file that we create for the given process. */
 	static FString GetTempSessionContextFilePath(uint64 ProcessID);
 
+	/** Clean up expired context files that were left-over on the user disks (because the consumer crashed and/or failed to delete it). */
+	static void CleanupTempSessionContextFiles(const FTimespan& ExpirationAge);
+
 	/** Serializes all data to the buffer. */
 	void SerializeContentToBuffer() const;
 
