CompactBinary: Renamed Type to TypeWithFlags and clarified its protected accessors

#rnx

diff --git a/Engine/Source/Runtime/Core/Private/Serialization/CompactBinary.cpp b/Engine/Source/Runtime/Core/Private/Serialization/CompactBinary.cpp
--- a/Engine/Source/Runtime/Core/Private/Serialization/CompactBinary.cpp
+++ b/Engine/Source/Runtime/Core/Private/Serialization/CompactBinary.cpp
@@ -93,22 +93,22 @@ void FCbObjectId::ToString(FWideStringBuilderBase& Builder) const
 
 ///////////////////////////////////////////////////////////////////////////////////////////////////////////////////////
 
-FCbFieldView::FCbFieldView(const void* const InData, const ECbFieldType InType)
+FCbFieldView::FCbFieldView(const void* const Data, ECbFieldType Type)
 {
-	const uint8* Bytes = static_cast<const uint8*>(InData);
-	const ECbFieldType LocalType = FCbFieldType::HasFieldType(InType) ? (ECbFieldType(*Bytes++) | ECbFieldType::HasFieldType) : InType;
+	const uint8* Bytes = static_cast<const uint8*>(Data);
+	Type = FCbFieldType::HasFieldType(Type) ? (ECbFieldType(*Bytes++) | ECbFieldType::HasFieldType) : Type;
 	uint32 NameLenByteCount = 0;
-	const uint64 NameLen64 = FCbFieldType::HasFieldName(LocalType) ? ReadVarUInt(Bytes, NameLenByteCount) : 0;
+	const uint64 NameLen64 = FCbFieldType::HasFieldName(Type) ? ReadVarUInt(Bytes, NameLenByteCount) : 0;
 	Bytes += NameLen64 + NameLenByteCount;
 
-	Type = LocalType;
+	TypeWithFlags = Type;
 	NameLen = uint32(FMath::Clamp<uint64>(NameLen64, 0, ~uint32(0)));
 	Payload = Bytes;
 }
 
 FCbObjectView FCbFieldView::AsObjectView()
 {
-	if (FCbFieldType::IsObject(Type))
+	if (FCbFieldType::IsObject(TypeWithFlags))
 	{
 		Error = ECbFieldError::None;
 		return FCbObjectView::FromFieldNoCheck(*this);
@@ -122,7 +122,7 @@ FCbObjectView FCbFieldView::AsObjectView()
 
 FCbArrayView FCbFieldView::AsArrayView()
 {
-	if (FCbFieldType::IsArray(Type))
+	if (FCbFieldType::IsArray(TypeWithFlags))
 	{
 		Error = ECbFieldError::None;
 		return FCbArrayView::FromFieldNoCheck(*this);
@@ -136,7 +136,7 @@ FCbArrayView FCbFieldView::AsArrayView()
 
 FMemoryView FCbFieldView::AsBinaryView(const FMemoryView Default)
 {
-	if (FCbFieldType::IsBinary(Type))
+	if (FCbFieldType::IsBinary(TypeWithFlags))
 	{
 		const uint8* const PayloadBytes = static_cast<const uint8*>(Payload);
 		uint32 ValueSizeByteCount;
@@ -154,7 +154,7 @@ FMemoryView FCbFieldView::AsBinaryView(const FMemoryView Default)
 
 FAnsiStringView FCbFieldView::AsString(const FAnsiStringView Default)
 {
-	if (FCbFieldType::IsString(Type))
+	if (FCbFieldType::IsString(TypeWithFlags))
 	{
 		const ANSICHAR* const PayloadChars = static_cast<const ANSICHAR*>(Payload);
 		uint32 ValueSizeByteCount;
@@ -178,11 +178,11 @@ FAnsiStringView FCbFieldView::AsString(const FAnsiStringView Default)
 
 uint64 FCbFieldView::AsInteger(const uint64 Default, const FIntegerParams Params)
 {
-	if (FCbFieldType::IsInteger(Type))
+	if (FCbFieldType::IsInteger(TypeWithFlags))
 	{
 		// A shift of a 64-bit value by 64 is undefined so shift by one less because magnitude is never zero.
 		const uint64 OutOfRangeMask = uint64(-2) << (Params.MagnitudeBits - 1);
-		const uint64 IsNegative = uint8(Type) & 1;
+		const uint64 IsNegative = uint8(TypeWithFlags) & 1;
 
 		uint32 MagnitudeByteCount;
 		const uint64 Magnitude = ReadVarUInt(Payload, MagnitudeByteCount);
@@ -203,12 +203,12 @@ uint64 FCbFieldView::AsInteger(const uint64 Default, const FIntegerParams Params
 
 float FCbFieldView::AsFloat(const float Default)
 {
-	switch (FCbFieldType::GetType(Type))
+	switch (GetType())
 	{
 	case ECbFieldType::IntegerPositive:
 	case ECbFieldType::IntegerNegative:
 	{
-		const uint64 IsNegative = uint8(Type) & 1;
+		const uint64 IsNegative = uint8(TypeWithFlags) & 1;
 		constexpr uint64 OutOfRangeMask = ~((uint64(1) << /*FLT_MANT_DIG*/ 24) - 1);
 
 		uint32 MagnitudeByteCount;
@@ -234,12 +234,12 @@ float FCbFieldView::AsFloat(const float Default)
 
 double FCbFieldView::AsDouble(const double Default)
 {
-	switch (FCbFieldType::GetType(Type))
+	switch (GetType())
 	{
 	case ECbFieldType::IntegerPositive:
 	case ECbFieldType::IntegerNegative:
 	{
-		const uint64 IsNegative = uint8(Type) & 1;
+		const uint64 IsNegative = uint8(TypeWithFlags) & 1;
 		constexpr uint64 OutOfRangeMask = ~((uint64(1) << /*DBL_MANT_DIG*/ 53) - 1);
 
 		uint32 MagnitudeByteCount;
@@ -268,15 +268,15 @@ double FCbFieldView::AsDouble(const double Default)
 
 bool FCbFieldView::AsBool(const bool bDefault)
 {
-	const ECbFieldType LocalType = Type;
-	const bool bIsBool = FCbFieldType::IsBool(LocalType);
+	const ECbFieldType LocalTypeWithFlags = TypeWithFlags;
+	const bool bIsBool = FCbFieldType::IsBool(LocalTypeWithFlags);
 	Error = bIsBool ? ECbFieldError::None : ECbFieldError::TypeError;
-	return (uint8(bIsBool) & uint8(LocalType) & 1) | ((!bIsBool) & bDefault);
+	return (uint8(bIsBool) & uint8(LocalTypeWithFlags) & 1) | ((!bIsBool) & bDefault);
 }
 
 FIoHash FCbFieldView::AsObjectAttachment(const FIoHash& Default)
 {
-	if (FCbFieldType::IsObjectAttachment(Type))
+	if (FCbFieldType::IsObjectAttachment(TypeWithFlags))
 	{
 		Error = ECbFieldError::None;
 		return FIoHash(*static_cast<const FIoHash::ByteArray*>(Payload));
@@ -290,7 +290,7 @@ FIoHash FCbFieldView::AsObjectAttachment(const FIoHash& Default)
 
 FIoHash FCbFieldView::AsBinaryAttachment(const FIoHash& Default)
 {
-	if (FCbFieldType::IsBinaryAttachment(Type))
+	if (FCbFieldType::IsBinaryAttachment(TypeWithFlags))
 	{
 		Error = ECbFieldError::None;
 		return FIoHash(*static_cast<const FIoHash::ByteArray*>(Payload));
@@ -304,7 +304,7 @@ FIoHash FCbFieldView::AsBinaryAttachment(const FIoHash& Default)
 
 FIoHash FCbFieldView::AsAttachment(const FIoHash& Default)
 {
-	if (FCbFieldType::IsAttachment(Type))
+	if (FCbFieldType::IsAttachment(TypeWithFlags))
 	{
 		Error = ECbFieldError::None;
 		return FIoHash(*static_cast<const FIoHash::ByteArray*>(Payload));
@@ -318,7 +318,7 @@ FIoHash FCbFieldView::AsAttachment(const FIoHash& Default)
 
 FIoHash FCbFieldView::AsHash(const FIoHash& Default)
 {
-	if (FCbFieldType::IsHash(Type))
+	if (FCbFieldType::IsHash(TypeWithFlags))
 	{
 		Error = ECbFieldError::None;
 		return FIoHash(*static_cast<const FIoHash::ByteArray*>(Payload));
@@ -337,7 +337,7 @@ FGuid FCbFieldView::AsUuid()
 
 FGuid FCbFieldView::AsUuid(const FGuid& Default)
 {
-	if (FCbFieldType::IsUuid(Type))
+	if (FCbFieldType::IsUuid(TypeWithFlags))
 	{
 		Error = ECbFieldError::None;
 		FGuid Value;
@@ -357,7 +357,7 @@ FGuid FCbFieldView::AsUuid(const FGuid& Default)
 
 int64 FCbFieldView::AsDateTimeTicks(const int64 Default)
 {
-	if (FCbFieldType::IsDateTime(Type))
+	if (FCbFieldType::IsDateTime(TypeWithFlags))
 	{
 		Error = ECbFieldError::None;
 		return NETWORK_ORDER64(CompactBinaryPrivate::ReadUnaligned<int64>(Payload));
@@ -381,7 +381,7 @@ FDateTime FCbFieldView::AsDateTime(FDateTime Default)
 
 int64 FCbFieldView::AsTimeSpanTicks(const int64 Default)
 {
-	if (FCbFieldType::IsTimeSpan(Type))
+	if (FCbFieldType::IsTimeSpan(TypeWithFlags))
 	{
 		Error = ECbFieldError::None;
 		return NETWORK_ORDER64(CompactBinaryPrivate::ReadUnaligned<int64>(Payload));
@@ -406,7 +406,7 @@ FTimespan FCbFieldView::AsTimeSpan(FTimespan Default)
 FCbObjectId FCbFieldView::AsObjectId(const FCbObjectId& Default)
 {
 	static_assert(sizeof(FCbObjectId) == 12, "FCbObjectId is expected to be 12 bytes.");
-	if (FCbFieldType::IsObjectId(Type))
+	if (FCbFieldType::IsObjectId(TypeWithFlags))
 	{
 		Error = ECbFieldError::None;
 		return FCbObjectId(MakeMemoryView(Payload, sizeof(FCbObjectId)));
@@ -420,7 +420,7 @@ FCbObjectId FCbFieldView::AsObjectId(const FCbObjectId& Default)
 
 FCbCustomById FCbFieldView::AsCustomById(FCbCustomById Default)
 {
-	if (FCbFieldType::IsCustomById(Type))
+	if (FCbFieldType::IsCustomById(TypeWithFlags))
 	{
 		const uint8* PayloadBytes = static_cast<const uint8*>(Payload);
 		uint32 PayloadSizeByteCount;
@@ -443,7 +443,7 @@ FCbCustomById FCbFieldView::AsCustomById(FCbCustomById Default)
 
 FCbCustomByName FCbFieldView::AsCustomByName(FCbCustomByName Default)
 {
-	if (FCbFieldType::IsCustomByName(Type))
+	if (FCbFieldType::IsCustomByName(TypeWithFlags))
 	{
 		const uint8* PayloadBytes = static_cast<const uint8*>(Payload);
 		uint32 PayloadSizeByteCount;
@@ -504,7 +504,7 @@ uint64 FCbFieldView::GetSize() const
 
 uint64 FCbFieldView::GetPayloadSize() const
 {
-	switch (FCbFieldType::GetType(Type))
+	switch (GetType())
 	{
 	case ECbFieldType::None:
 	case ECbFieldType::Null:
@@ -557,14 +557,14 @@ FBlake3Hash FCbFieldView::GetHash() const
 
 void FCbFieldView::AppendHash(FBlake3& Builder) const
 {
-	const ECbFieldType SerializedType = FCbFieldType::GetSerializedType(Type);
+	const ECbFieldType SerializedType = FCbFieldType::GetSerializedType(TypeWithFlags);
 	Builder.Update(&SerializedType, sizeof(SerializedType));
 	Builder.Update(GetViewNoType());
 }
 
 bool FCbFieldView::Equals(const FCbFieldView& Other) const
 {
-	return FCbFieldType::GetSerializedType(Type) == FCbFieldType::GetSerializedType(Other.Type) &&
+	return FCbFieldType::GetSerializedType(TypeWithFlags) == FCbFieldType::GetSerializedType(Other.TypeWithFlags) &&
 		GetViewNoType().EqualBytes(Other.GetViewNoType());
 }
 
@@ -574,7 +574,7 @@ void FCbFieldView::CopyTo(FMutableMemoryView Buffer) const
 	checkf(Buffer.GetSize() == sizeof(ECbFieldType) + Source.GetSize(),
 		TEXT("Buffer is %" UINT64_FMT " bytes but %" UINT64_FMT " is required."),
 		Buffer.GetSize(), sizeof(ECbFieldType) + Source.GetSize());
-	const ECbFieldType SerializedType = FCbFieldType::GetSerializedType(Type);
+	const ECbFieldType SerializedType = FCbFieldType::GetSerializedType(TypeWithFlags);
 	Buffer.CopyFrom(MakeMemoryView(&SerializedType, 1)).CopyFrom(Source);
 }
 
@@ -582,14 +582,14 @@ void FCbFieldView::CopyTo(FArchive& Ar) const
 {
 	check(Ar.IsSaving());
 	const FMemoryView Source = GetViewNoType();
-	ECbFieldType SerializedType = FCbFieldType::GetSerializedType(Type);
+	ECbFieldType SerializedType = FCbFieldType::GetSerializedType(TypeWithFlags);
 	Ar.Serialize(&SerializedType, sizeof(SerializedType));
 	Ar.Serialize(const_cast<void*>(Source.GetData()), static_cast<int64>(Source.GetSize()));
 }
 
 void FCbFieldView::IterateAttachments(FCbFieldVisitor Visitor) const
 {
-	switch (FCbFieldType::GetType(Type))
+	switch (GetType())
 	{
 	case ECbFieldType::Object:
 	case ECbFieldType::UniformObject:
@@ -607,34 +607,34 @@ void FCbFieldView::IterateAttachments(FCbFieldVisitor Visitor) const
 
 FMemoryView FCbFieldView::GetView() const
 {
-	const uint32 TypeSize = FCbFieldType::HasFieldType(Type) ? sizeof(ECbFieldType) : 0;
-	const uint32 NameSize = FCbFieldType::HasFieldName(Type) ? NameLen + MeasureVarUInt(NameLen) : 0;
+	const uint32 TypeSize = FCbFieldType::HasFieldType(TypeWithFlags) ? sizeof(ECbFieldType) : 0;
+	const uint32 NameSize = FCbFieldType::HasFieldName(TypeWithFlags) ? NameLen + MeasureVarUInt(NameLen) : 0;
 	const uint64 PayloadSize = GetPayloadSize();
 	return MakeMemoryView(static_cast<const uint8*>(Payload) - TypeSize - NameSize, TypeSize + NameSize + PayloadSize);
 }
 
 FMemoryView FCbFieldView::GetViewNoType() const
 {
-	const uint32 NameSize = FCbFieldType::HasFieldName(Type) ? NameLen + MeasureVarUInt(NameLen) : 0;
+	const uint32 NameSize = FCbFieldType::HasFieldName(TypeWithFlags) ? NameLen + MeasureVarUInt(NameLen) : 0;
 	const uint64 PayloadSize = GetPayloadSize();
 	return MakeMemoryView(static_cast<const uint8*>(Payload) - NameSize, NameSize + PayloadSize);
 }
 
 FCbFieldViewIterator FCbFieldView::CreateViewIterator() const
 {
-	const ECbFieldType LocalType = Type;
-	if (FCbFieldType::HasFields(LocalType))
+	const ECbFieldType LocalTypeWithFlags = TypeWithFlags;
+	if (FCbFieldType::HasFields(LocalTypeWithFlags))
 	{
 		const uint8* PayloadBytes = static_cast<const uint8*>(GetPayload());
 		uint32 PayloadSizeByteCount;
 		const uint64 PayloadSize = ReadVarUInt(PayloadBytes, PayloadSizeByteCount);
 		PayloadBytes += PayloadSizeByteCount;
-		const uint64 NumByteCount = FCbFieldType::IsArray(LocalType) ? MeasureVarUInt(PayloadBytes) : 0;
+		const uint64 NumByteCount = FCbFieldType::IsArray(LocalTypeWithFlags) ? MeasureVarUInt(PayloadBytes) : 0;
 		if (PayloadSize > NumByteCount)
 		{
 			const void* const PayloadEnd = PayloadBytes + PayloadSize;
 			PayloadBytes += NumByteCount;
-			const ECbFieldType UniformType = FCbFieldType::HasUniformFields(LocalType) ?
+			const ECbFieldType UniformType = FCbFieldType::HasUniformFields(LocalTypeWithFlags) ?
 				ECbFieldType(*PayloadBytes++) : ECbFieldType::HasFieldType;
 			return FCbFieldViewIterator::MakeRange(MakeMemoryView(PayloadBytes, PayloadEnd), UniformType);
 		}
@@ -671,14 +671,14 @@ FBlake3Hash FCbArrayView::GetHash() const
 
 void FCbArrayView::AppendHash(FBlake3& Builder) const
 {
-	const ECbFieldType SerializedType = FCbFieldType::GetType(GetType());
+	const ECbFieldType SerializedType = GetType();
 	Builder.Update(&SerializedType, sizeof(SerializedType));
 	Builder.Update(GetPayloadView());
 }
 
 bool FCbArrayView::Equals(const FCbArrayView& Other) const
 {
-	return FCbFieldType::GetType(GetType()) == FCbFieldType::GetType(Other.GetType()) &&
+	return GetType() == Other.GetType() &&
 		GetPayloadView().EqualBytes(Other.GetPayloadView());
 }
 
@@ -688,7 +688,7 @@ void FCbArrayView::CopyTo(const FMutableMemoryView Buffer) const
 	checkf(Buffer.GetSize() == sizeof(ECbFieldType) + Source.GetSize(),
 		TEXT("Buffer is %" UINT64_FMT " bytes but %" UINT64_FMT " is required."),
 		Buffer.GetSize(), sizeof(ECbFieldType) + Source.GetSize());
-	const ECbFieldType SerializedType = FCbFieldType::GetType(GetType());
+	const ECbFieldType SerializedType = GetType();
 	Buffer.CopyFrom(MakeMemoryView(&SerializedType, 1)).CopyFrom(Source);
 }
 
@@ -696,7 +696,7 @@ void FCbArrayView::CopyTo(FArchive& Ar) const
 {
 	check(Ar.IsSaving());
 	const FMemoryView Source = GetPayloadView();
-	ECbFieldType SerializedType = FCbFieldType::GetType(GetType());
+	ECbFieldType SerializedType = GetType();
 	Ar.Serialize(&SerializedType, sizeof(SerializedType));
 	Ar.Serialize(const_cast<void*>(Source.GetData()), static_cast<int64>(Source.GetSize()));
 }
@@ -751,14 +751,14 @@ FBlake3Hash FCbObjectView::GetHash() const
 
 void FCbObjectView::AppendHash(FBlake3& Builder) const
 {
-	const ECbFieldType SerializedType = FCbFieldType::GetType(GetType());
+	const ECbFieldType SerializedType = GetType();
 	Builder.Update(&SerializedType, sizeof(SerializedType));
 	Builder.Update(GetPayloadView());
 }
 
 bool FCbObjectView::Equals(const FCbObjectView& Other) const
 {
-	return FCbFieldType::GetType(GetType()) == FCbFieldType::GetType(Other.GetType()) &&
+	return GetType() == Other.GetType() &&
 		GetPayloadView().EqualBytes(Other.GetPayloadView());
 }
 
@@ -768,7 +768,7 @@ void FCbObjectView::CopyTo(const FMutableMemoryView Buffer) const
 	checkf(Buffer.GetSize() == sizeof(ECbFieldType) + Source.GetSize(),
 		TEXT("Buffer is %" UINT64_FMT " bytes but %" UINT64_FMT " is required."),
 		Buffer.GetSize(), sizeof(ECbFieldType) + Source.GetSize());
-	const ECbFieldType SerializedType = FCbFieldType::GetType(GetType());
+	const ECbFieldType SerializedType = GetType();
 	Buffer.CopyFrom(MakeMemoryView(&SerializedType, 1)).CopyFrom(Source);
 }
 
@@ -776,7 +776,7 @@ void FCbObjectView::CopyTo(FArchive& Ar) const
 {
 	check(Ar.IsSaving());
 	const FMemoryView Source = GetPayloadView();
-	ECbFieldType SerializedType = FCbFieldType::GetType(GetType());
+	ECbFieldType SerializedType = GetType();
 	Ar.Serialize(&SerializedType, sizeof(SerializedType));
 	Ar.Serialize(const_cast<void*>(Source.GetData()), static_cast<int64>(Source.GetSize()));
 }
@@ -788,7 +788,7 @@ TCbFieldIterator<FieldType>& TCbFieldIterator<FieldType>::operator++()
 {
 	const void* const PayloadEnd = FieldType::GetPayloadEnd();
 	const int64 AtEndMask = int64(PayloadEnd == FieldsEnd) - 1;
-	const ECbFieldType NextType = ECbFieldType(int64(FieldType::GetType()) & AtEndMask);
+	const ECbFieldType NextType = ECbFieldType(int64(FieldType::GetTypeWithFlags()) & AtEndMask);
 	const void* const NextField = reinterpret_cast<const void*>(int64(PayloadEnd) & AtEndMask);
 	const void* const NextFieldsEnd = reinterpret_cast<const void*>(int64(FieldsEnd) & AtEndMask);
 	FieldType::Assign(NextField, NextType);
@@ -884,11 +884,11 @@ template <typename FieldType>
 void TCbFieldIterator<FieldType>::IterateRangeAttachments(FCbFieldVisitor Visitor) const
 {
 	// Always iterate over non-uniform ranges because we do not know if they contain an attachment.
-	if (FCbFieldType::HasFieldType(FieldType::GetType()))
+	if (FCbFieldType::HasFieldType(FieldType::GetTypeWithFlags()))
 	{
 		for (TCbFieldIterator It(*this); It; ++It)
 		{
-			if (FCbFieldType::MayContainAttachments(It.GetType()))
+			if (FCbFieldType::MayContainAttachments(It.GetTypeWithFlags()))
 			{
 				It.IterateAttachments(Visitor);
 			}
@@ -897,7 +897,7 @@ void TCbFieldIterator<FieldType>::IterateRangeAttachments(FCbFieldVisitor Visito
 	// Only iterate over uniform ranges if the uniform type may contain an attachment.
 	else
 	{
-		if (FCbFieldType::MayContainAttachments(FieldType::GetType()))
+		if (FCbFieldType::MayContainAttachments(FieldType::GetTypeWithFlags()))
 		{
 			for (TCbFieldIterator It(*this); It; ++It)
 			{
diff --git a/Engine/Source/Runtime/Core/Private/Serialization/CompactBinaryWriter.cpp b/Engine/Source/Runtime/Core/Private/Serialization/CompactBinaryWriter.cpp
--- a/Engine/Source/Runtime/Core/Private/Serialization/CompactBinaryWriter.cpp
+++ b/Engine/Source/Runtime/Core/Private/Serialization/CompactBinaryWriter.cpp
@@ -62,7 +62,7 @@ static inline ECbFieldType AppendCompactBinary(const FCbFieldView& Value, TArray
 	const FMemoryView SourceView = ValueCopy.GetPayloadView();
 	const int64 TargetOffset = OutData.AddUninitialized(SourceView.GetSize());
 	MakeMemoryView(OutData).RightChop(TargetOffset).CopyFrom(SourceView);
-	return FCbFieldType::GetType(ValueCopy.GetType());
+	return ValueCopy.GetType();
 }
 
 FCbWriter::FCbWriter()
diff --git a/Engine/Source/Runtime/Core/Public/Serialization/CompactBinary.h b/Engine/Source/Runtime/Core/Public/Serialization/CompactBinary.h
--- a/Engine/Source/Runtime/Core/Public/Serialization/CompactBinary.h
+++ b/Engine/Source/Runtime/Core/Public/Serialization/CompactBinary.h
@@ -397,7 +397,7 @@ public:
 	 */
 	inline bool TryGetSerializedRangeView(FMemoryView& OutView) const
 	{
-		if (FCbFieldType::HasFieldType(FieldType::GetType()))
+		if (FCbFieldType::HasFieldType(FieldType::GetTypeWithFlags()))
 		{
 			OutView = GetRangeView();
 			return true;
@@ -655,36 +655,36 @@ public:
 	CORE_API FMemoryView AsCustom(FAnsiStringView Name, FMemoryView Default = FMemoryView());
 
 	/** True if the field has a name. */
-	constexpr inline bool HasName() const           { return FCbFieldType::HasFieldName(Type); }
+	constexpr inline bool HasName() const           { return FCbFieldType::HasFieldName(TypeWithFlags); }
 
-	constexpr inline bool IsNull() const            { return FCbFieldType::IsNull(Type); }
+	constexpr inline bool IsNull() const            { return FCbFieldType::IsNull(TypeWithFlags); }
 
-	constexpr inline bool IsObject() const          { return FCbFieldType::IsObject(Type); }
-	constexpr inline bool IsArray() const           { return FCbFieldType::IsArray(Type); }
+	constexpr inline bool IsObject() const          { return FCbFieldType::IsObject(TypeWithFlags); }
+	constexpr inline bool IsArray() const           { return FCbFieldType::IsArray(TypeWithFlags); }
 
-	constexpr inline bool IsBinary() const          { return FCbFieldType::IsBinary(Type); }
-	constexpr inline bool IsString() const          { return FCbFieldType::IsString(Type); }
+	constexpr inline bool IsBinary() const          { return FCbFieldType::IsBinary(TypeWithFlags); }
+	constexpr inline bool IsString() const          { return FCbFieldType::IsString(TypeWithFlags); }
 
 	/** Whether the field is an integer of unspecified range and sign. */
-	constexpr inline bool IsInteger() const         { return FCbFieldType::IsInteger(Type); }
+	constexpr inline bool IsInteger() const         { return FCbFieldType::IsInteger(TypeWithFlags); }
 	/** Whether the field is a float, or integer that supports implicit conversion. */
-	constexpr inline bool IsFloat() const           { return FCbFieldType::IsFloat(Type); }
-	constexpr inline bool IsBool() const            { return FCbFieldType::IsBool(Type); }
+	constexpr inline bool IsFloat() const           { return FCbFieldType::IsFloat(TypeWithFlags); }
+	constexpr inline bool IsBool() const            { return FCbFieldType::IsBool(TypeWithFlags); }
 
-	constexpr inline bool IsObjectAttachment() const { return FCbFieldType::IsObjectAttachment(Type); }
-	constexpr inline bool IsBinaryAttachment() const { return FCbFieldType::IsBinaryAttachment(Type); }
-	constexpr inline bool IsAttachment() const       { return FCbFieldType::IsAttachment(Type); }
+	constexpr inline bool IsObjectAttachment() const { return FCbFieldType::IsObjectAttachment(TypeWithFlags); }
+	constexpr inline bool IsBinaryAttachment() const { return FCbFieldType::IsBinaryAttachment(TypeWithFlags); }
+	constexpr inline bool IsAttachment() const       { return FCbFieldType::IsAttachment(TypeWithFlags); }
 
-	constexpr inline bool IsHash() const            { return FCbFieldType::IsHash(Type); }
-	constexpr inline bool IsUuid() const            { return FCbFieldType::IsUuid(Type); }
+	constexpr inline bool IsHash() const            { return FCbFieldType::IsHash(TypeWithFlags); }
+	constexpr inline bool IsUuid() const            { return FCbFieldType::IsUuid(TypeWithFlags); }
 
-	constexpr inline bool IsDateTime() const        { return FCbFieldType::IsDateTime(Type); }
-	constexpr inline bool IsTimeSpan() const        { return FCbFieldType::IsTimeSpan(Type); }
+	constexpr inline bool IsDateTime() const        { return FCbFieldType::IsDateTime(TypeWithFlags); }
+	constexpr inline bool IsTimeSpan() const        { return FCbFieldType::IsTimeSpan(TypeWithFlags); }
 
-	constexpr inline bool IsObjectId() const        { return FCbFieldType::IsObjectId(Type); }
+	constexpr inline bool IsObjectId() const        { return FCbFieldType::IsObjectId(TypeWithFlags); }
 
-	constexpr inline bool IsCustomById() const      { return FCbFieldType::IsCustomById(Type); }
-	constexpr inline bool IsCustomByName() const    { return FCbFieldType::IsCustomByName(Type); }
+	constexpr inline bool IsCustomById() const      { return FCbFieldType::IsCustomById(TypeWithFlags); }
+	constexpr inline bool IsCustomByName() const    { return FCbFieldType::IsCustomByName(TypeWithFlags); }
 
 	/** Whether the field has a value. */
 	constexpr inline explicit operator bool() const { return HasValue(); }
@@ -695,7 +695,7 @@ public:
 	 * All fields in a valid object or array have a value. A field with no value is returned when
 	 * finding a field by name fails or when accessing an iterator past the end.
 	 */
-	constexpr inline bool HasValue() const          { return !FCbFieldType::IsNone(Type); };
+	constexpr inline bool HasValue() const          { return !FCbFieldType::IsNone(TypeWithFlags); };
 
 	/** Whether the last field access encountered an error. */
 	constexpr inline bool HasError() const          { return Error != ECbFieldError::None; }
@@ -742,7 +742,7 @@ public:
 	 */
 	inline bool TryGetSerializedView(FMemoryView& OutView) const
 	{
-		if (FCbFieldType::HasFieldType(Type))
+		if (FCbFieldType::HasFieldType(TypeWithFlags))
 		{
 			OutView = GetView();
 			return true;
@@ -764,8 +764,11 @@ protected:
 	/** Returns a view of the value payload, which excludes the type and name. */
 	inline FMemoryView GetPayloadView() const { return MakeMemoryView(Payload, GetPayloadSize()); }
 
+	/** Returns the type of the field excluding flags. */
+	constexpr inline ECbFieldType GetType() const { return FCbFieldType::GetType(TypeWithFlags); }
+
 	/** Returns the type of the field including flags. */
-	constexpr inline ECbFieldType GetType() const { return Type; }
+	constexpr inline ECbFieldType GetTypeWithFlags() const { return TypeWithFlags; }
 
 	/** Returns the start of the value payload. */
 	constexpr inline const void* GetPayload() const { return Payload; }
@@ -819,7 +822,7 @@ private:
 
 private:
 	/** The field type, with the transient HasFieldType flag if the field contains its type. */
-	ECbFieldType Type = ECbFieldType::None;
+	ECbFieldType TypeWithFlags = ECbFieldType::None;
 	/** The error (if any) that occurred on the last field access. */
 	ECbFieldError Error = ECbFieldError::None;
 	/** The number of bytes for the name stored before the payload. */
