U5 - Move EventStack out of the FScheduler class.

diff --git a/Engine/Source/Runtime/Core/Public/Async/Fundamental/Scheduler.h b/Engine/Source/Runtime/Core/Public/Async/Fundamental/Scheduler.h
--- a/Engine/Source/Runtime/Core/Public/Async/Fundamental/Scheduler.h
+++ b/Engine/Source/Runtime/Core/Public/Async/Fundamental/Scheduler.h
@@ -19,6 +19,74 @@ namespace LowLevelTasks
 		DefaultPreference = LocalQueuePreference,
 	};
 
+
+	//implementation of a treiber stack
+	//(https://en.wikipedia.org/wiki/Treiber_stack)
+	template<typename NodeType>
+	class TEventStack
+	{
+	private:
+		struct FTopNode
+		{
+			uintptr_t Address  : 45; //all CPUs we care about use less than 48 bits for their addressing, and the lower 3 bits are unused due to alignment
+			uintptr_t Revision : 19; //Tagging is used to avoid ABA (the wrap around is several minutes for this use-case (https://en.wikipedia.org/wiki/ABA_problem#Tagged_state_reference))
+		};
+		std::atomic<FTopNode> Top { FTopNode{0, 0} };
+
+	public:
+		NodeType* Pop()
+		{
+			FTopNode LocalTop = Top.load(std::memory_order_relaxed);
+			while (true) 
+			{			
+				if (LocalTop.Address == 0)
+				{
+					return nullptr;
+				}
+#if DO_CHECK
+				int64 LastRevision = int64(LocalTop.Revision); 
+#endif
+				NodeType* Item = reinterpret_cast<NodeType*>(LocalTop.Address << 3);
+				if (Top.compare_exchange_weak(LocalTop, FTopNode { reinterpret_cast<uintptr_t>(Item->Next) >> 3, uintptr_t(LocalTop.Revision + 1) }, std::memory_order_acquire, std::memory_order_relaxed))
+				{
+					Item->Next = nullptr;
+					return Item;
+				}
+#if DO_CHECK
+				int64 NewRevision = int64(LocalTop.Revision) < LastRevision ? ((1ll << 19) + int64(LocalTop.Revision)) : int64(LocalTop.Revision);
+				ensureMsgf((NewRevision - LastRevision) < (1ll << 18), TEXT("Dangerously close to the wraparound: %d, %d"), LastRevision, NewRevision);
+#endif
+			}
+		}
+
+		void Push(NodeType* Item)
+		{
+			checkSlow(Item != nullptr);
+#if !USING_CODE_ANALYSIS //MS SA thowing warning C6011 on Item->Next access, even when it is validated or branched over
+			checkSlow(reinterpret_cast<uintptr_t>(Item) < (1ull << 48));
+			checkSlow((reinterpret_cast<uintptr_t>(Item) & 0x7) == 0);
+#endif
+			checkSlow(Item->Next == nullptr);
+
+			FTopNode LocalTop = Top.load(std::memory_order_relaxed);
+			while (true) 
+			{
+#if DO_CHECK
+				int64 LastRevision = int64(LocalTop.Revision); 
+#endif
+				Item->Next = reinterpret_cast<NodeType*>(LocalTop.Address << 3);
+				if (Top.compare_exchange_weak(LocalTop, FTopNode { reinterpret_cast<uintptr_t>(Item) >> 3, uintptr_t(LocalTop.Revision + 1) }, std::memory_order_release, std::memory_order_acquire))  
+				{
+					return;
+				}
+#if DO_CHECK
+				int64 NewRevision = int64(LocalTop.Revision) < LastRevision ? ((1ll << 19) + int64(LocalTop.Revision)) : int64(LocalTop.Revision);
+				ensureMsgf((NewRevision - LastRevision) < (1ll << 18), TEXT("Dangerously close to the wraparound: %d, %d"), LastRevision, NewRevision);
+#endif
+			}
+		}
+	};
+
 	class FScheduler
 	{
 		UE_NONCOPYABLE(FScheduler);
@@ -107,72 +175,6 @@ namespace LowLevelTasks
 			FSleepEvent* Next = nullptr;
 		};
 
-		//implementation of a treiber stack
-		//(https://en.wikipedia.org/wiki/Treiber_stack)
-		class FEventStack
-		{
-		private:
-			struct FTopNode
-			{
-				uintptr_t Address  : 45; //all CPUs we care about use less than 48 bits for their addressing, and the lower 3 bits are unused due to alignment
-				uintptr_t Revision : 19; //Tagging is used to avoid ABA (the wrap around is several minutes for this use-case (https://en.wikipedia.org/wiki/ABA_problem#Tagged_state_reference))
-			};
-			std::atomic<FTopNode> Top { FTopNode{0, 0} };
-
-		public:
-			FSleepEvent* Pop()
-			{
-				FTopNode LocalTop = Top.load(std::memory_order_relaxed);
-				while (true) 
-				{			
-					if (LocalTop.Address == 0)
-					{
-						return nullptr;
-					}
-#if DO_CHECK
-					int64 LastRevision = int64(LocalTop.Revision); 
-#endif
-					FSleepEvent* Item = reinterpret_cast<FSleepEvent*>(LocalTop.Address << 3);
-					if (Top.compare_exchange_weak(LocalTop, FTopNode { reinterpret_cast<uintptr_t>(Item->Next) >> 3, uintptr_t(LocalTop.Revision + 1) }, std::memory_order_acquire, std::memory_order_relaxed))
-					{
-						Item->Next = nullptr;
-						return Item;
-					}
-#if DO_CHECK
-					int64 NewRevision = int64(LocalTop.Revision) < LastRevision ? ((1ll << 19) + int64(LocalTop.Revision)) : int64(LocalTop.Revision);
-					ensureMsgf((NewRevision - LastRevision) < (1ll << 18), TEXT("Dangerously close to the wraparound: %d, %d"), LastRevision, NewRevision);
-#endif
-				}
-			}
-
-			void Push(FSleepEvent* Item)
-			{
-				checkSlow(Item != nullptr);
-#if !USING_CODE_ANALYSIS //MS SA thowing warning C6011 on Item->Next access, even when it is validated or branched over
-				checkSlow(reinterpret_cast<uintptr_t>(Item) < (1ull << 48));
-				checkSlow((reinterpret_cast<uintptr_t>(Item) & 0x7) == 0);
-#endif
-				checkSlow(Item->Next == nullptr);
-			
-				FTopNode LocalTop = Top.load(std::memory_order_relaxed);
-				while (true) 
-				{
-#if DO_CHECK
-					int64 LastRevision = int64(LocalTop.Revision); 
-#endif
-					Item->Next = reinterpret_cast<FSleepEvent*>(LocalTop.Address << 3);
-					if (Top.compare_exchange_weak(LocalTop, FTopNode { reinterpret_cast<uintptr_t>(Item) >> 3, uintptr_t(LocalTop.Revision + 1) }, std::memory_order_release, std::memory_order_acquire))  
-					{
-						return;
-					}
-#if DO_CHECK
-					int64 NewRevision = int64(LocalTop.Revision) < LastRevision ? ((1ll << 19) + int64(LocalTop.Revision)) : int64(LocalTop.Revision);
-					ensureMsgf((NewRevision - LastRevision) < (1ll << 18), TEXT("Dangerously close to the wraparound: %d, %d"), LastRevision, NewRevision);
-#endif
-				}
-			}
-		};
-
 	public:
 		FScheduler() = default;
 		~FScheduler();
@@ -189,7 +191,7 @@ namespace LowLevelTasks
 		bool TryExecuteTaskFrom(FLocalQueueType* Queue, FQueueRegistry::FOutOfWork& OutOfWork, bool bPermitBackgroundWork);
 
 	private:
-		FEventStack 				SleepEventStack[2];
+		TEventStack<FSleepEvent> 	SleepEventStack[2];
 		FQueueRegistry 				QueueRegistry;
 		FCriticalSection 			WorkerThreadsCS;
 		TArray<TUniquePtr<FThread>> WorkerThreads;
