[ColorList] fixed some typos and unified method descriptions

diff --git a/Engine/Source/Runtime/Core/Public/Math/ColorList.h b/Engine/Source/Runtime/Core/Public/Math/ColorList.h
--- a/Engine/Source/Runtime/Core/Public/Math/ColorList.h
+++ b/Engine/Source/Runtime/Core/Public/Math/ColorList.h
@@ -119,22 +119,22 @@ public:
 	/** Initializes list of common colors. */
 	void CreateColorMap();
 
-	/** Returns a color based on ColorName if not found returs White. */
+	/** Returns a color based on ColorName. If not found, returns White. */
 	const FColor& GetFColorByName( const TCHAR* ColorName ) const;
 
-	/** Returns a linear color based on ColorName if not found returs White. */
+	/** Returns a linear color based on ColorName. If not found, returns White. */
 	const FLinearColor GetFLinearColorByName( const TCHAR* ColorName ) const;
 
-	/** Returns true if color is valid common colors, returns false otherwise. */
+	/** Returns true if color is valid common colors. If not found returns false. */
 	bool IsValidColorName( const TCHAR* ColorName ) const;
 
-	/** Returns index of color. */
+	/** Returns index of color based on ColorName. If not found returns 0. */
 	int32 GetColorIndex( const TCHAR* ColorName ) const;
 
 	/** Returns a color based on index. If index is invalid, returns White. */
 	const FColor& GetFColorByIndex( int32 ColorIndex ) const;
 
-	/** Resturn color's name based on index. If index is invalid, returns BadIndex. */
+	/** Returns color's name based on index. If index is invalid, returns BadIndex. */
 	const FString& GetColorNameByIndex( int32 ColorIndex ) const;
 
 	/** Returns the number of colors. */
