U5 - Hoist out the shared Sheduler TLS sate into it's own class and allow forcing of background work for busy waiting.

diff --git a/Engine/Source/Runtime/Core/Private/Async/Fundamental/Scheduler.cpp b/Engine/Source/Runtime/Core/Private/Async/Fundamental/Scheduler.cpp
--- a/Engine/Source/Runtime/Core/Private/Async/Fundamental/Scheduler.cpp
+++ b/Engine/Source/Runtime/Core/Private/Async/Fundamental/Scheduler.cpp
@@ -12,9 +12,9 @@ namespace LowLevelTasks
 	DEFINE_LOG_CATEGORY(LowLevelTasks);
 
 	thread_local FScheduler::FLocalQueueType* FScheduler::LocalQueue = nullptr;
-	thread_local FTask* FScheduler::ActiveTask = nullptr;
-	thread_local FScheduler* FScheduler::ActiveScheduler = nullptr;
-	thread_local FScheduler::EWorkerType FScheduler::WorkerType = FScheduler::EWorkerType::None;
+	thread_local FTask* FSchedulerTls::ActiveTask = nullptr;
+	thread_local FSchedulerTls* FSchedulerTls::ActiveScheduler = nullptr;
+	thread_local FSchedulerTls::EWorkerType FSchedulerTls::WorkerType = FSchedulerTls::EWorkerType::None;
 	thread_local uint32 FScheduler::BusyWaitingDepth = 0;
 
 	FScheduler FScheduler::Singleton;
@@ -24,7 +24,7 @@ namespace LowLevelTasks
 		RegisteredLocalQueue = LocalQueue == nullptr;
 		if (RegisteredLocalQueue)
 		{
-			bool bPermitBackgroundWork = ActiveTask && ActiveTask->IsBackgroundTask();
+			bool bPermitBackgroundWork = FSchedulerTls::PermitBackgroundWork();
 			LocalQueue = FLocalQueueType::AllocateLocalQueue(Scheduler.QueueRegistry, bPermitBackgroundWork);
 		}
 	}
@@ -33,7 +33,7 @@ namespace LowLevelTasks
 	{
 		if (RegisteredLocalQueue)
 		{
-			bool bPermitBackgroundWork = ActiveTask && ActiveTask->IsBackgroundTask();
+			bool bPermitBackgroundWork = FSchedulerTls::PermitBackgroundWork();
 			FLocalQueueType::DeleteLocalQueue(LocalQueue, bPermitBackgroundWork);
 			LocalQueue = nullptr;
 		}
@@ -144,7 +144,7 @@ namespace LowLevelTasks
 		if (ActiveWorkers.load(std::memory_order_relaxed))
 		{			
 			const bool bIsBackgroundTask = Task.IsBackgroundTask();
-			const bool bIsBackgroundWorker = WorkerType == EWorkerType::Background;
+			const bool bIsBackgroundWorker = FSchedulerTls::IsBackgroundWorker();
 			if (bIsBackgroundTask && !bIsBackgroundWorker)
 			{
 				QueuePreference = EQueuePreference::GlobalQueuePreference;
@@ -179,19 +179,19 @@ namespace LowLevelTasks
 		}
 	}
 
-	const FTask* FScheduler::GetActiveTask() 
+	const FTask* FSchedulerTls::GetActiveTask() 
 	{
 		return ActiveTask;
 	}
 
-	bool FScheduler::IsWorkerThread() const
+	bool FSchedulerTls::IsWorkerThread() const
 	{
-		return WorkerType != EWorkerType::None && ActiveScheduler == this;
+		return WorkerType != FSchedulerTls::EWorkerType::None && ActiveScheduler == this;
 	}
 
 	void FTask::PropagateUserData()
 	{
-		const FTask* ActiveTask = FScheduler::GetActiveTask();
+		const FTask* ActiveTask = FSchedulerTls::GetActiveTask();
 		if (ActiveTask != nullptr)
 		{
 			UserData = ActiveTask->GetUserData();
@@ -216,13 +216,13 @@ namespace LowLevelTasks
 					continue;
 				}
 				OutOfWork.Stop();		
-				FTask* OldTask = ActiveTask;
-				ActiveTask = Task;
+				FTask* OldTask = FSchedulerTls::ActiveTask;
+				FSchedulerTls::ActiveTask = Task;
 				{
 					TRACE_CPUPROFILER_EVENT_SCOPE(ExecuteTask);
 					Task->ExecuteTask();
 				}
-				ActiveTask = OldTask;
+				FSchedulerTls::ActiveTask = OldTask;
 				return true;
 			}
 			return false;
@@ -233,10 +233,10 @@ namespace LowLevelTasks
 	void FScheduler::WorkerMain(FSleepEvent* WorkerEvent, FLocalQueueType* ExternalWorkerLocalQueue, uint32 WaitCycles, bool bPermitBackgroundWork)
 	{
 		FTaskTagScope WorkerScope(ETaskTag::EWorkerThread);
-		ActiveScheduler = this;
+		FSchedulerTls::ActiveScheduler = this;
 
 		FMemory::SetupTLSCachesOnCurrentThread();
-		WorkerType = bPermitBackgroundWork ? EWorkerType::Background : EWorkerType::Foreground;
+		FSchedulerTls::WorkerType = bPermitBackgroundWork ? FSchedulerTls::EWorkerType::Background : FSchedulerTls::EWorkerType::Foreground;
 
 		checkSlow(LocalQueue == nullptr);
 		if(ExternalWorkerLocalQueue)
@@ -289,12 +289,12 @@ namespace LowLevelTasks
 		FLocalQueueType::DeleteLocalQueue(WorkerLocalQueue, bPermitBackgroundWork, ExternalWorkerLocalQueue != nullptr);
 		LocalQueue = nullptr;
 
-		ActiveScheduler = nullptr;
-		WorkerType = EWorkerType::None;
+		FSchedulerTls::ActiveScheduler = nullptr;
+		FSchedulerTls::WorkerType = FSchedulerTls::EWorkerType::None;
 		FMemory::ClearAndDisableTLSCachesOnCurrentThread();
 	}
 
-	void FScheduler::BusyWaitInternal(const FConditional& Conditional)
+	void FScheduler::BusyWaitInternal(const FConditional& Conditional, bool ForceAllowBackgroundWork)
 	{
 		TRACE_CPUPROFILER_EVENT_SCOPE(FScheduler::BusyWaitInternal);
 		FTaskTagScope WorkerScope(ETaskTag::EWorkerThread);
@@ -308,8 +308,8 @@ namespace LowLevelTasks
 
 		uint32 WaitCount = 0;
 		bool HasWokenEmergencyWorker = false;
-		const bool bIsBackgroundWorker = WorkerType == EWorkerType::Background;
-		bool bPermitBackgroundWork = ActiveTask && ActiveTask->IsBackgroundTask();
+		const bool bIsBackgroundWorker = FSchedulerTls::IsBackgroundWorker();
+		bool bPermitBackgroundWork = FSchedulerTls::PermitBackgroundWork() || ForceAllowBackgroundWork;
 		FQueueRegistry::FOutOfWork OutOfWork = QueueRegistry.GetOutOfWorkScope(bIsBackgroundWorker);
 		while (true)
 		{
diff --git a/Engine/Source/Runtime/Core/Public/Async/Fundamental/Scheduler.h b/Engine/Source/Runtime/Core/Public/Async/Fundamental/Scheduler.h
--- a/Engine/Source/Runtime/Core/Public/Async/Fundamental/Scheduler.h
+++ b/Engine/Source/Runtime/Core/Public/Async/Fundamental/Scheduler.h
@@ -87,7 +87,39 @@ namespace LowLevelTasks
 		}
 	};
 
-	class FScheduler
+	class FSchedulerTls
+	{
+	protected: 
+		enum class EWorkerType
+		{
+			None,
+			Background,
+			Foreground,
+		};
+
+		static thread_local FTask* ActiveTask;
+		static thread_local FSchedulerTls* ActiveScheduler;
+		static thread_local EWorkerType WorkerType;
+
+	public:
+		CORE_API bool IsWorkerThread() const;
+
+		//get the active task if any
+		CORE_API static const FTask* GetActiveTask();
+
+	protected:
+		inline static bool PermitBackgroundWork()
+		{
+			return ActiveTask && ActiveTask->IsBackgroundTask();
+		}
+
+		inline static bool IsBackgroundWorker()
+		{
+			return WorkerType == EWorkerType::Background;
+		}
+	};
+
+	class FScheduler final : public FSchedulerTls
 	{
 		UE_NONCOPYABLE(FScheduler);
 		static constexpr uint32 WorkerSpinCycles = 53;
@@ -136,16 +168,16 @@ namespace LowLevelTasks
 
 		//tries to do some work until the Task is completed
 		template<typename TaskType>
-		inline void BusyWait(const TaskType& Task);
+		inline void BusyWait(const TaskType& Task, bool ForceAllowBackgroundWork = false);
 
 		//tries to do some work until the Conditional return true
 		template<typename Conditional>
-		inline void BusyWaitUntil(Conditional&& Cond);
+		inline void BusyWaitUntil(Conditional&& Cond, bool ForceAllowBackgroundWork = false);
 
 		//tries to do some work until all the Tasks are completed
 		//the template parameter can be any Type that has a const conversion operator to FTask
 		template<typename TaskType>
-		inline void BusyWait(const TArrayView<const TaskType>& Tasks);
+		inline void BusyWait(const TArrayView<const TaskType>& Tasks, bool ForceAllowBackgroundWork = false);
 
 		// returns true if the current thread execution is in the context of busy-waiting
 		inline static bool IsBusyWaiting();
@@ -153,11 +185,6 @@ namespace LowLevelTasks
 		//number of instantiated workers
 		inline uint32 GetNumWorkers() const;
 
-		//get the active task if any
-		CORE_API static const FTask* GetActiveTask();
-
-		CORE_API bool IsWorkerThread() const;
-
 	private: //Private Interface of the Scheduler	
 		enum class ESleepState
 		{
@@ -183,7 +210,7 @@ namespace LowLevelTasks
 		TUniquePtr<FThread> CreateWorker(FLocalQueueType* ExternalWorkerLocalQueue = nullptr, EThreadPriority Priority = EThreadPriority::TPri_Normal, bool bPermitBackgroundWork = false, bool bIsForkable = false);
 		void WorkerMain(struct FSleepEvent* WorkerEvent, FLocalQueueType* ExternalWorkerLocalQueue, uint32 WaitCycles, bool bPermitBackgroundWork);
 		CORE_API void LaunchInternal(FTask& Task, EQueuePreference QueuePreference, bool bWakeUpWorker);
-		CORE_API void BusyWaitInternal(const FConditional& Conditional);
+		CORE_API void BusyWaitInternal(const FConditional& Conditional, bool ForceAllowBackgroundWork);
 		FORCENOINLINE void TrySleeping(FSleepEvent* WorkerEvent, FQueueRegistry::FOutOfWork& OutOfWork, bool& Drowsing, bool bBackgroundWorker);
 		inline bool WakeUpWorker(bool bBackgroundWorker);
 
@@ -205,21 +232,21 @@ namespace LowLevelTasks
 		return FScheduler::Get().TryLaunch(Task, QueuePreference, bWakeUpWorker);
 	}
 
-	FORCEINLINE_DEBUGGABLE void BusyWaitForTask(const FTask& Task)
+	FORCEINLINE_DEBUGGABLE void BusyWaitForTask(const FTask& Task, bool ForceAllowBackgroundWork = false)
 	{
-		FScheduler::Get().BusyWait(Task);
+		FScheduler::Get().BusyWait(Task, ForceAllowBackgroundWork);
 	}
 
 	template<typename Conditional>
-	FORCEINLINE_DEBUGGABLE void BusyWaitUntil(Conditional&& Cond)
+	FORCEINLINE_DEBUGGABLE void BusyWaitUntil(Conditional&& Cond, bool ForceAllowBackgroundWork = false)
 	{
-		FScheduler::Get().BusyWaitUntil<Conditional>(Forward<Conditional>(Cond));
+		FScheduler::Get().BusyWaitUntil<Conditional>(Forward<Conditional>(Cond), ForceAllowBackgroundWork);
 	}
 
 	template<typename TaskType>
-	FORCEINLINE_DEBUGGABLE void BusyWaitForTasks(const TArrayView<const TaskType>& Tasks)
+	FORCEINLINE_DEBUGGABLE void BusyWaitForTasks(const TArrayView<const TaskType>& Tasks, bool ForceAllowBackgroundWork = false)
 	{
-		FScheduler::Get().BusyWait<TaskType>(Tasks);
+		FScheduler::Get().BusyWait<TaskType>(Tasks, ForceAllowBackgroundWork);
 	}
 
    /******************
@@ -241,17 +268,17 @@ namespace LowLevelTasks
 	}
 
 	template<typename TaskType>
-	inline void FScheduler::BusyWait(const TaskType& Task)
+	inline void FScheduler::BusyWait(const TaskType& Task, bool ForceAllowBackgroundWork)
 	{
 		if(!Task.IsCompleted())
 		{
 			FLocalQueueInstaller Installer(*this);
-			FScheduler::BusyWaitInternal([&Task](){ return Task.IsCompleted(); });
+			FScheduler::BusyWaitInternal([&Task](){ return Task.IsCompleted(); }, ForceAllowBackgroundWork);
 		}
 	}
 
 	template<typename Conditional>
-	inline void FScheduler::BusyWaitUntil(Conditional&& Cond)
+	inline void FScheduler::BusyWaitUntil(Conditional&& Cond, bool ForceAllowBackgroundWork)
 	{
 		static_assert(TIsInvocable<Conditional>::Value, "Conditional is not invocable");
 		static_assert(TIsSame<decltype(Cond()), bool>::Value, "Conditional must return a boolean");
@@ -259,12 +286,12 @@ namespace LowLevelTasks
 		if(!Cond())
 		{
 			FLocalQueueInstaller Installer(*this);
-			FScheduler::BusyWaitInternal(Forward<Conditional>(Cond));
+			FScheduler::BusyWaitInternal(Forward<Conditional>(Cond), ForceAllowBackgroundWork);
 		}
 	}
 
 	template<typename TaskType>
-	inline void FScheduler::BusyWait(const TArrayView<const TaskType>& Tasks)
+	inline void FScheduler::BusyWait(const TArrayView<const TaskType>& Tasks, bool ForceAllowBackgroundWork)
 	{
 		auto AllTasksCompleted = [Index(0), &Tasks]() mutable
 		{
@@ -282,7 +309,7 @@ namespace LowLevelTasks
 		if (!AllTasksCompleted())
 		{
 			FLocalQueueInstaller Installer(*this);
-			FScheduler::BusyWaitInternal([&AllTasksCompleted](){ return AllTasksCompleted(); });
+			FScheduler::BusyWaitInternal([&AllTasksCompleted](){ return AllTasksCompleted(); }, ForceAllowBackgroundWork);
 		}
 	}
 
