U5 - remove invalid assert

diff --git a/Engine/Source/Runtime/Core/Public/Async/Fundamental/LocalQueue.h b/Engine/Source/Runtime/Core/Public/Async/Fundamental/LocalQueue.h
--- a/Engine/Source/Runtime/Core/Public/Async/Fundamental/LocalQueue.h
+++ b/Engine/Source/Runtime/Core/Public/Async/Fundamental/LocalQueue.h
@@ -367,7 +367,6 @@ private:
 		while(true)
 		{
 			FLocalQueueCollection* Previous = Hazard.Get();
-			checkSlow(Previous->RemovedQueue == nullptr);
 			FLocalQueueCollection* Copy = new FLocalQueueCollection(Previous);
 			Copy->LocalQueues.Add(QueueToAdd);
 			if (!QueueCollection.compare_exchange_strong(Previous, Copy, std::memory_order_release, std::memory_order_relaxed))
@@ -388,7 +387,6 @@ private:
 		while(true)
 		{
 			FLocalQueueCollection* Previous = Hazard.Get();
-			checkSlow(Previous->RemovedQueue == nullptr);
 			FLocalQueueCollection* Copy = new FLocalQueueCollection(Previous);
 			int NumRemoved = Copy->LocalQueues.Remove(QueueToRemove);
 			checkSlow(NumRemoved == 1);
