CompactBinary: Restricted compact binary attachments to objects to simplify package management

#rnx

diff --git a/Engine/Source/Developer/DerivedDataCache/Private/DerivedDataCacheRecord.cpp b/Engine/Source/Developer/DerivedDataCache/Private/DerivedDataCacheRecord.cpp
--- a/Engine/Source/Developer/DerivedDataCache/Private/DerivedDataCacheRecord.cpp
+++ b/Engine/Source/Developer/DerivedDataCache/Private/DerivedDataCacheRecord.cpp
@@ -354,7 +354,7 @@ FOptionalCacheRecord LoadCacheRecord(const FCbPackage& Package)
 
 		if (const FCbAttachment* Attachment = Package.FindAttachment(PayloadObject["CompressedHash"].AsHash()))
 		{
-			if (FCompressedBuffer Compressed = FCompressedBuffer::FromCompressed(Attachment->AsBinaryView()))
+			if (FCompressedBuffer Compressed = FCompressedBuffer::FromCompressed(Attachment->AsBinary()))
 			{
 				return FPayload(PayloadId, Compressed);
 			}
diff --git a/Engine/Source/Runtime/Core/Private/Serialization/CompactBinaryPackage.cpp b/Engine/Source/Runtime/Core/Private/Serialization/CompactBinaryPackage.cpp
--- a/Engine/Source/Runtime/Core/Private/Serialization/CompactBinaryPackage.cpp
+++ b/Engine/Source/Runtime/Core/Private/Serialization/CompactBinaryPackage.cpp
@@ -8,34 +8,25 @@
 
 ///////////////////////////////////////////////////////////////////////////////////////////////////////////////////////
 
-FCbAttachment::FCbAttachment(FCbFieldIterator InValue, const FIoHash* const InHash)
+FCbAttachment::FCbAttachment(FCbObject InValue, const FIoHash* const InHash)
 {
-	if (InValue)
+	FMemoryView View;
+	if (!InValue.IsOwned() || !InValue.TryGetSerializedView(View))
 	{
-		if (!InValue.IsOwned())
-		{
-			InValue = FCbFieldIterator::CloneRange(InValue);
-		}
-
-		CompactBinary = FCbFieldViewIterator(InValue);
-		Buffer = MoveTemp(InValue).GetOuterBuffer();
+		InValue = FCbObject::Clone(InValue);
 	}
 
+	Object = InValue.AsFieldView();
+	Buffer = MoveTemp(InValue).GetBuffer();
+
 	if (InHash)
 	{
 		Hash = *InHash;
-		if (CompactBinary)
-		{
-			checkSlow(Hash == CompactBinary.GetRangeHash());
-		}
-		else
-		{
-			checkfSlow(Hash.IsZero(), TEXT("A null or empty field range must use a hash of zero."));
-		}
+		checkSlow(Hash == FIoHash::HashBuffer(Buffer));
 	}
-	else if (CompactBinary)
+	else
 	{
-		Hash = CompactBinary.GetRangeHash();
+		Hash = FIoHash::HashBuffer(Buffer);
 	}
 }
 
@@ -45,46 +36,14 @@ FCbAttachment::FCbAttachment(FSharedBuffer InBuffer, const FIoHash* const InHash
 	if (InHash)
 	{
 		Hash = *InHash;
-		if (Buffer.GetSize())
-		{
-			checkSlow(Hash == FIoHash::HashBuffer(Buffer));
-		}
-		else
-		{
-			checkfSlow(Hash.IsZero(), TEXT("A null or empty buffer must use a hash of zero."));
-		}
-	}
-	else if (Buffer.GetSize())
-	{
-		Hash = FIoHash::HashBuffer(Buffer);
+		checkSlow(Hash == FIoHash::HashBuffer(Buffer));
 	}
 	else
 	{
-		Buffer.Reset();
+		Hash = FIoHash::HashBuffer(Buffer);
 	}
 }
 
-FSharedBuffer FCbAttachment::AsBinaryView() const
-{
-	if (!CompactBinary)
-	{
-		return Buffer;
-	}
-
-	FMemoryView SerializedView;
-	if (CompactBinary.TryGetSerializedRangeView(SerializedView))
-	{
-		return SerializedView == Buffer.GetView() ? Buffer : FSharedBuffer::MakeView(SerializedView, Buffer);
-	}
-
-	return FCbFieldIterator::CloneRange(CompactBinary).GetRangeBuffer();
-}
-
-FCbFieldIterator FCbAttachment::AsCompactBinary() const
-{
-	return CompactBinary ? FCbFieldIterator::MakeRangeView(CompactBinary, Buffer) : FCbFieldIterator();
-}
-
 void FCbAttachment::Load(FCbFieldIterator& Fields)
 {
 	checkf(Fields.IsBinary(), TEXT("Attachments must start with a binary field."));
@@ -92,12 +51,14 @@ void FCbAttachment::Load(FCbFieldIterator& Fields)
 	if (View.GetSize() > 0)
 	{
 		Buffer = FSharedBuffer::MakeView(View, Fields.GetOuterBuffer()).MakeOwned();
+		Object = FCbFieldView();
 		++Fields;
+
 		Hash = Fields.AsAttachment();
 		checkf(!Fields.HasError(), TEXT("Attachments must be a non-empty binary value with a content hash."));
 		if (Fields.IsObjectAttachment())
 		{
-			CompactBinary = FCbFieldViewIterator::MakeRange(Buffer);
+			Object = FCbFieldView(Buffer.GetData());
 		}
 		++Fields;
 	}
@@ -105,7 +66,7 @@ void FCbAttachment::Load(FCbFieldIterator& Fields)
 	{
 		++Fields;
 		Buffer.Reset();
-		CompactBinary.Reset();
+		Object = FCbFieldView();
 		Hash.Reset();
 	}
 }
@@ -118,7 +79,7 @@ void FCbAttachment::Load(FArchive& Ar, FCbBufferAllocator Allocator)
 	if (View.GetSize() > 0)
 	{
 		Buffer = FSharedBuffer::MakeView(View, BufferField.GetOuterBuffer()).MakeOwned();
-		CompactBinary = FCbFieldViewIterator();
+		Object = FCbFieldView();
 
 		TArray<uint8, TInlineAllocator<64>> HashBuffer;
 		FCbField HashField = LoadCompactBinary(Ar,
@@ -131,31 +92,26 @@ void FCbAttachment::Load(FArchive& Ar, FCbBufferAllocator Allocator)
 		checkf(!HashField.HasError(), TEXT("Attachments must be a non-empty binary value with a content hash."));
 		if (HashField.IsObjectAttachment())
 		{
-			CompactBinary = FCbFieldViewIterator::MakeRange(Buffer);
+			Object = FCbFieldView(Buffer.GetData());
 		}
 	}
 	else
 	{
 		Buffer.Reset();
-		CompactBinary.Reset();
+		Object = FCbFieldView();
 		Hash.Reset();
 	}
 }
 
 void FCbAttachment::Save(FCbWriter& Writer) const
 {
-	if (CompactBinary)
+	if (Object.IsObject())
 	{
-		FMemoryView SerializedView;
-		if (CompactBinary.TryGetSerializedRangeView(SerializedView))
-		{
-			Writer.AddBinary(SerializedView);
-		}
-		else
+		Writer.AddBinary(Buffer);
+		if (FCbFieldView(Object).AsObjectView())
 		{
-			Writer.AddBinary(AsBinaryView());
+			Writer.AddObjectAttachment(Hash);
 		}
-		Writer.AddObjectAttachment(Hash);
 	}
 	else if (Buffer.GetSize())
 	{
@@ -193,7 +149,7 @@ void FCbPackage::SetObject(FCbObject InObject, const FIoHash* InObjectHash, FAtt
 		}
 		if (InResolver)
 		{
-			GatherAttachments(Object.CreateViewIterator(), *InResolver);
+			GatherAttachments(Object, *InResolver);
 		}
 	}
 	else
@@ -211,18 +167,18 @@ void FCbPackage::AddAttachment(const FCbAttachment& Attachment, FAttachmentResol
 		if (Attachments.IsValidIndex(Index) && Attachments[Index] == Attachment)
 		{
 			FCbAttachment& Existing = Attachments[Index];
-			if (Attachment.IsCompactBinary() && !Existing.IsCompactBinary())
+			if (Attachment.IsObject() && !Existing.IsObject())
 			{
-				Existing = FCbAttachment(FCbFieldIterator::MakeRange(Existing.AsBinaryView()));
+				Existing = FCbAttachment(FCbObject(Existing.AsBinary()), Existing.GetHash());
 			}
 		}
 		else
 		{
 			Attachments.Insert(Attachment, Index);
 		}
-		if (Attachment.IsCompactBinary() && Resolver)
+		if (Attachment.IsObject() && Resolver)
 		{
-			GatherAttachments(Attachment.AsCompactBinary(), *Resolver);
+			GatherAttachments(Attachment.AsObject(), *Resolver);
 		}
 	}
 }
@@ -251,16 +207,16 @@ const FCbAttachment* FCbPackage::FindAttachment(const FIoHash& Hash) const
 	return Attachments.IsValidIndex(Index) ? &Attachments[Index] : nullptr;
 }
 
-void FCbPackage::GatherAttachments(const FCbFieldViewIterator& Fields, FAttachmentResolver Resolver)
+void FCbPackage::GatherAttachments(const FCbObject& Value, FAttachmentResolver Resolver)
 {
-	Fields.IterateRangeAttachments([this, &Resolver](FCbFieldView Field)
+	Value.IterateAttachments([this, &Resolver](FCbFieldView Field)
 		{
 			const FIoHash& Hash = Field.AsAttachment();
 			if (FSharedBuffer Buffer = Resolver(Hash))
 			{
 				if (Field.IsObjectAttachment())
 				{
-					AddAttachment(FCbAttachment(FCbFieldIterator::MakeRange(MoveTemp(Buffer)), Hash), &Resolver);
+					AddAttachment(FCbAttachment(FCbObject(MoveTemp(Buffer)), Hash), &Resolver);
 				}
 				else
 				{
@@ -333,7 +289,7 @@ void FCbPackage::Load(FArchive& Ar, FCbBufferAllocator Allocator)
 				checkf(!HashField.HasError(), TEXT("Attachments must be a non-empty binary value with a content hash."));
 				if (HashField.IsObjectAttachment())
 				{
-					AddAttachment(FCbAttachment(FCbFieldIterator::MakeRange(MoveTemp(Buffer)), Hash));
+					AddAttachment(FCbAttachment(FCbObject(MoveTemp(Buffer)), Hash));
 				}
 				else
 				{
diff --git a/Engine/Source/Runtime/Core/Private/Tests/Serialization/CompactBinaryPackageTest.cpp b/Engine/Source/Runtime/Core/Private/Tests/Serialization/CompactBinaryPackageTest.cpp
--- a/Engine/Source/Runtime/Core/Private/Tests/Serialization/CompactBinaryPackageTest.cpp
+++ b/Engine/Source/Runtime/Core/Private/Tests/Serialization/CompactBinaryPackageTest.cpp
@@ -53,10 +53,10 @@ bool FCbAttachmentTest::RunTest(const FString& Parameters)
 		FCbAttachment Attachment;
 		TestTrue(TEXT("FCbAttachment(Null).IsNull()"), Attachment.IsNull());
 		TestFalse(TEXT("FCbAttachment(Null) as bool"), bool(Attachment));
-		TestFalse(TEXT("FCbAttachment(Null).AsBinaryView()"), bool(Attachment.AsBinaryView()));
-		TestFalse(TEXT("FCbAttachment(Null).AsCompactBinary()"), Attachment.AsCompactBinary().HasValue());
+		TestFalse(TEXT("FCbAttachment(Null).AsBinary()"), bool(Attachment.AsBinary()));
+		TestFalse(TEXT("FCbAttachment(Null).AsObject()"), bool(Attachment.AsObject()));
 		TestFalse(TEXT("FCbAttachment(Null).IsBinary()"), Attachment.IsBinary());
-		TestFalse(TEXT("FCbAttachment(Null).IsCompactBinary()"), Attachment.IsCompactBinary());
+		TestFalse(TEXT("FCbAttachment(Null).IsObject()"), Attachment.IsObject());
 		TestEqual(TEXT("FCbAttachment(Null).GetHash()"), Attachment.GetHash(), FIoHash());
 		TestSaveLoadValidate(TEXT("Null"), Attachment);
 	}
@@ -67,28 +67,30 @@ bool FCbAttachmentTest::RunTest(const FString& Parameters)
 		FCbAttachment Attachment(Buffer);
 		TestFalse(TEXT("FCbAttachment(Binary).IsNull()"), Attachment.IsNull());
 		TestTrue(TEXT("FCbAttachment(Binary) as bool"), bool(Attachment));
-		TestEqual(TEXT("FCbAttachment(Binary).AsBinaryView()"), Attachment.AsBinaryView(), Buffer);
-		TestFalse(TEXT("FCbAttachment(Binary).AsCompactBinary()"), Attachment.AsCompactBinary().HasValue());
+		TestEqual(TEXT("FCbAttachment(Binary).AsBinary()"), Attachment.AsBinary(), Buffer);
+		TestFalse(TEXT("FCbAttachment(Binary).AsObject()"), bool(Attachment.AsObject()));
 		TestTrue(TEXT("FCbAttachment(Binary).IsBinary()"), Attachment.IsBinary());
-		TestFalse(TEXT("FCbAttachment(Binary).IsCompactBinary()"), Attachment.IsCompactBinary());
+		TestFalse(TEXT("FCbAttachment(Binary).IsObject()"), Attachment.IsObject());
 		TestEqual(TEXT("FCbAttachment(Binary).GetHash()"), Attachment.GetHash(), FIoHash::HashBuffer(Buffer));
 		TestSaveLoadValidate(TEXT("Binary"), Attachment);
 	}
 
-	// Compact Binary Attachment
+	// Object Attachment
 	{
 		FCbWriter Writer;
+		Writer.BeginObject();
 		Writer << "Name"_ASV << 42;
-		FCbFieldIterator Fields = Writer.Save();
-		FCbAttachment Attachment(Fields);
-		TestFalse(TEXT("FCbAttachment(CompactBinary).IsNull()"), Attachment.IsNull());
-		TestTrue(TEXT("FCbAttachment(CompactBinary) as bool"), bool(Attachment));
-		TestEqual(TEXT("FCbAttachment(CompactBinary).AsBinaryView()"), Attachment.AsBinaryView(), Fields.GetRangeBuffer());
-		TestEqual(TEXT("FCbAttachment(CompactBinary).AsCompactBinary()"), Attachment.AsCompactBinary(), Fields);
-		TestTrue(TEXT("FCbAttachment(CompactBinary).IsBinary()"), Attachment.IsBinary());
-		TestTrue(TEXT("FCbAttachment(CompactBinary).IsCompactBinary()"), Attachment.IsCompactBinary());
-		TestEqual(TEXT("FCbAttachment(CompactBinary).GetHash()"), Attachment.GetHash(), FIoHash(Fields.GetRangeHash()));
-		TestSaveLoadValidate(TEXT("CompactBinary"), Attachment);
+		Writer.EndObject();
+		FCbObject Object = Writer.Save().AsObject();
+		FCbAttachment Attachment(Object);
+		TestFalse(TEXT("FCbAttachment(Object).IsNull()"), Attachment.IsNull());
+		TestTrue(TEXT("FCbAttachment(Object) as bool"), bool(Attachment));
+		TestEqual(TEXT("FCbAttachment(Object).AsBinary()"), Attachment.AsBinary(), Object.GetBuffer());
+		TestTrue(TEXT("FCbAttachment(Object).AsObject()"), Attachment.AsObject().Equals(Object));
+		TestTrue(TEXT("FCbAttachment(Object).IsBinary()"), Attachment.IsBinary());
+		TestTrue(TEXT("FCbAttachment(Object).IsObject()"), Attachment.IsObject());
+		TestEqual(TEXT("FCbAttachment(Object).GetHash()"), Attachment.GetHash(), FIoHash(Object.GetHash()));
+		TestSaveLoadValidate(TEXT("Object"), Attachment);
 	}
 
 	// Binary View
@@ -98,32 +100,30 @@ bool FCbAttachmentTest::RunTest(const FString& Parameters)
 		FCbAttachment Attachment(Buffer);
 		TestFalse(TEXT("FCbAttachment(BinaryView).IsNull()"), Attachment.IsNull());
 		TestTrue(TEXT("FCbAttachment(BinaryView) as bool"), bool(Attachment));
-		TestNotEqual(TEXT("FCbAttachment(BinaryView).AsBinaryView()"), Attachment.AsBinaryView(), Buffer);
-		TestTrue(TEXT("FCbAttachment(BinaryView).AsBinaryView()"),
-			Attachment.AsBinaryView().GetView().EqualBytes(Buffer.GetView()));
-		TestFalse(TEXT("FCbAttachment(BinaryView).AsCompactBinary()"), Attachment.AsCompactBinary().HasValue());
+		TestNotEqual(TEXT("FCbAttachment(BinaryView).AsBinary()"), Attachment.AsBinary(), Buffer);
+		TestTrue(TEXT("FCbAttachment(BinaryView).AsBinary()"), Attachment.AsBinary().GetView().EqualBytes(Buffer.GetView()));
+		TestFalse(TEXT("FCbAttachment(BinaryView).AsObject()"), bool(Attachment.AsObject()));
 		TestTrue(TEXT("FCbAttachment(BinaryView).IsBinary()"), Attachment.IsBinary());
-		TestFalse(TEXT("FCbAttachment(BinaryView).IsCompactBinary()"), Attachment.IsCompactBinary());
-		TestEqual(TEXT("FCbAttachment(BinaryView).GetHash()"),
-			Attachment.GetHash(), FIoHash::HashBuffer(Buffer));
+		TestFalse(TEXT("FCbAttachment(BinaryView).IsObject()"), Attachment.IsObject());
+		TestEqual(TEXT("FCbAttachment(BinaryView).GetHash()"), Attachment.GetHash(), FIoHash::HashBuffer(Buffer));
 	}
 
-	// Compact Binary View
+	// Object View
 	{
 		FCbWriter Writer;
+		Writer.BeginObject();
 		Writer << "Name"_ASV << 42;
-		FCbFieldIterator Fields = Writer.Save();
-		FCbFieldIterator FieldsView = FCbFieldIterator::MakeRangeView(FCbFieldViewIterator(Fields));
-		FCbAttachment Attachment(FieldsView);
-		TestFalse(TEXT("FCbAttachment(CompactBinaryView).IsNull()"), Attachment.IsNull());
-		TestTrue(TEXT("FCbAttachment(CompactBinaryView) as bool"), bool(Attachment));
-		TestNotEqual(TEXT("FCbAttachment(CompactBinaryView).AsBinaryView()"),
-			Attachment.AsBinaryView(), FieldsView.GetRangeBuffer());
-		TestTrue(TEXT("FCbAttachment(CompactBinaryView).AsCompactBinary()"),
-			Attachment.AsCompactBinary().GetRangeView().EqualBytes(Fields.GetRangeView()));
-		TestTrue(TEXT("FCbAttachment(CompactBinaryView).IsBinary()"), Attachment.IsBinary());
-		TestTrue(TEXT("FCbAttachment(CompactBinaryView).IsCompactBinary()"), Attachment.IsCompactBinary());
-		TestEqual(TEXT("FCbAttachment(CompactBinaryView).GetHash()"), Attachment.GetHash(), FIoHash(Fields.GetRangeHash()));
+		Writer.EndObject();
+		FCbObject Object = Writer.Save().AsObject();
+		FCbObject ObjectView = FCbObject::MakeView(Object);
+		FCbAttachment Attachment(ObjectView);
+		TestFalse(TEXT("FCbAttachment(ObjectView).IsNull()"), Attachment.IsNull());
+		TestTrue(TEXT("FCbAttachment(ObjectView) as bool"), bool(Attachment));
+		TestNotEqual(TEXT("FCbAttachment(ObjectView).AsBinary()"), Attachment.AsBinary(), ObjectView.GetBuffer());
+		TestTrue(TEXT("FCbAttachment(ObjectView).AsObject()"), Attachment.AsObject().Equals(Object));
+		TestTrue(TEXT("FCbAttachment(ObjectView).IsBinary()"), Attachment.IsBinary());
+		TestTrue(TEXT("FCbAttachment(ObjectView).IsObject()"), Attachment.IsObject());
+		TestEqual(TEXT("FCbAttachment(ObjectView).GetHash()"), Attachment.GetHash(), FIoHash(Object.GetHash()));
 	}
 
 	// Binary Load from View
@@ -140,24 +140,26 @@ bool FCbAttachmentTest::RunTest(const FString& Parameters)
 		Attachment.Load(FieldsView);
 		TestFalse(TEXT("FCbAttachment(LoadBinaryView).IsNull()"), Attachment.IsNull());
 		TestTrue(TEXT("FCbAttachment(LoadBinaryView) as bool"), bool(Attachment));
-		TestFalse(TEXT("FCbAttachment(LoadBinaryView).AsBinaryView()->!InView"),
-			FieldsView.GetRangeBuffer().GetView().Contains(Attachment.AsBinaryView().GetView()));
-		TestTrue(TEXT("FCbAttachment(LoadBinaryView).AsBinaryView()->EqualBytes"),
-			Attachment.AsBinaryView().GetView().EqualBytes(Buffer.GetView()));
-		TestFalse(TEXT("FCbAttachment(LoadBinaryView).AsCompactBinary()"), Attachment.AsCompactBinary().HasValue());
+		TestFalse(TEXT("FCbAttachment(LoadBinaryView).AsBinary()->!InView"),
+			FieldsView.GetRangeBuffer().GetView().Contains(Attachment.AsBinary().GetView()));
+		TestTrue(TEXT("FCbAttachment(LoadBinaryView).AsBinary()->EqualBytes"),
+			Attachment.AsBinary().GetView().EqualBytes(Buffer.GetView()));
+		TestFalse(TEXT("FCbAttachment(LoadBinaryView).AsObject()"), bool(Attachment.AsObject()));
 		TestTrue(TEXT("FCbAttachment(LoadBinaryView).IsBinary()"), Attachment.IsBinary());
-		TestFalse(TEXT("FCbAttachment(LoadBinaryView).IsCompactBinary()"), Attachment.IsCompactBinary());
+		TestFalse(TEXT("FCbAttachment(LoadBinaryView).IsObject()"), Attachment.IsObject());
 		TestEqual(TEXT("FCbAttachment(LoadBinaryView).GetHash()"),
 			Attachment.GetHash(), FIoHash::HashBuffer(MakeMemoryView(Value)));
 	}
 
-	// Compact Binary Load from View
+	// Object Load from View
 	{
 		FCbWriter ValueWriter;
+		ValueWriter.BeginObject();
 		ValueWriter << "Name"_ASV << 42;
-		const FCbFieldIterator Value = ValueWriter.Save();
-		TestEqual(TEXT("FCbAttachment(LoadCompactBinaryView).Validate"),
-			ValidateCompactBinaryRange(Value.GetRangeView(), ECbValidateMode::All), ECbValidateError::None);
+		ValueWriter.EndObject();
+		const FCbObject Value = ValueWriter.Save().AsObject();
+		TestEqual(TEXT("FCbAttachment(LoadObjectView).Validate"),
+			ValidateCompactBinaryRange(Value.GetView(), ECbValidateMode::All), ECbValidateError::None);
 		FCbAttachment Attachment(Value);
 
 		FCbWriter Writer;
@@ -166,34 +168,15 @@ bool FCbAttachmentTest::RunTest(const FString& Parameters)
 		FCbFieldIterator FieldsView = FCbFieldIterator::MakeRangeView(FCbFieldViewIterator(Fields));
 
 		Attachment.Load(FieldsView);
-		TestFalse(TEXT("FCbAttachment(LoadCompactBinaryView).IsNull()"), Attachment.IsNull());
-		TestTrue(TEXT("FCbAttachment(LoadCompactBinaryView) as bool"), bool(Attachment));
-		TestTrue(TEXT("FCbAttachment(LoadCompactBinaryView).AsBinaryView()->EqualBytes"),
-			Attachment.AsBinaryView().GetView().EqualBytes(Value.GetRangeView()));
-		TestFalse(TEXT("FCbAttachment(LoadCompactBinaryView).AsCompactBinary()->!InView"),
-			FieldsView.GetRangeBuffer().GetView().Contains(Attachment.AsCompactBinary().GetRangeBuffer().GetView()));
-		TestTrue(TEXT("FCbAttachment(LoadCompactBinaryView).IsBinary()"), Attachment.IsBinary());
-		TestTrue(TEXT("FCbAttachment(LoadCompactBinaryView).IsCompactBinary()"), Attachment.IsCompactBinary());
-		TestEqual(TEXT("FCbAttachment(LoadCompactBinaryView).GetHash()"), Attachment.GetHash(), FIoHash(Value.GetRangeHash()));
-	}
-
-	// Compact Binary Uniform Sub-View
-	{
-		const FSharedBuffer Buffer = FSharedBuffer::Clone(MakeMemoryView<uint8>({0, 1, 2, 3}));
-		const FCbFieldViewIterator FieldViews = FCbFieldViewIterator::MakeRange(Buffer.GetView().RightChop(2), ECbFieldType::IntegerPositive);
-		const FCbFieldIterator SavedFields = FCbFieldIterator::CloneRange(FieldViews);
-		FCbFieldIterator Fields = FCbFieldIterator::MakeRangeView(FieldViews, Buffer);
-		FCbAttachment Attachment(Fields);
-		const FSharedBuffer Binary = Attachment.AsBinaryView();
-		TestEqual(TEXT("FCbAttachment(CompactBinaryUniformSubView).AsCompactBinary()->Equals()"),
-			Attachment.AsCompactBinary(), Fields);
-		TestEqual(TEXT("FCbAttachment(CompactBinaryUniformSubView).AsBinaryView()->GetSize()"),
-			Binary.GetSize(), SavedFields.GetRangeSize());
-		TestTrue(TEXT("FCbAttachment(CompactBinaryUniformSubView).AsBinaryView()->EqualBytes()"),
-			Binary.GetView().EqualBytes(SavedFields.GetRangeView()));
-		TestEqual(TEXT("FCbAttachment(CompactBinaryUniformSubView).GetHash()"),
-			Attachment.GetHash(), FIoHash(SavedFields.GetRangeHash()));
-		TestSaveLoadValidate(TEXT("CompactBinaryUniformSubView"), Attachment);
+		TestFalse(TEXT("FCbAttachment(LoadObjectView).IsNull()"), Attachment.IsNull());
+		TestTrue(TEXT("FCbAttachment(LoadObjectView) as bool"), bool(Attachment));
+		TestTrue(TEXT("FCbAttachment(LoadObjectView).AsBinary()->EqualBytes"),
+			Attachment.AsBinary().GetView().EqualBytes(Value.GetView()));
+		TestFalse(TEXT("FCbAttachment(LoadObjectView).AsObject()->!InView"),
+			FieldsView.GetRangeBuffer().GetView().Contains(Attachment.AsObject().GetView()));
+		TestTrue(TEXT("FCbAttachment(LoadObjectView).IsBinary()"), Attachment.IsBinary());
+		TestTrue(TEXT("FCbAttachment(LoadObjectView).IsObject()"), Attachment.IsObject());
+		TestEqual(TEXT("FCbAttachment(LoadObjectView).GetHash()"), Attachment.GetHash(), FIoHash(Value.GetHash()));
 	}
 
 	// Binary Null
@@ -201,26 +184,26 @@ bool FCbAttachmentTest::RunTest(const FString& Parameters)
 		const FCbAttachment Attachment(FSharedBuffer{});
 		TestTrue(TEXT("FCbAttachment(BinaryNull).IsNull()"), Attachment.IsNull());
 		TestFalse(TEXT("FCbAttachment(BinaryNull).IsBinary()"), Attachment.IsBinary());
-		TestFalse(TEXT("FCbAttachment(BinaryNull).IsCompactBinary()"), Attachment.IsCompactBinary());
-		TestEqual(TEXT("FCbAttachment(BinaryNull).GetHash()"), Attachment.GetHash(), FIoHash());
+		TestFalse(TEXT("FCbAttachment(BinaryNull).IsObject()"), Attachment.IsObject());
+		TestEqual(TEXT("FCbAttachment(BinaryNull).GetHash()"), Attachment.GetHash(), FIoHash::HashBuffer(FSharedBuffer{}));
 	}
 
 	// Binary Empty
 	{
 		const FCbAttachment Attachment(FUniqueBuffer::Alloc(0).MoveToShared());
-		TestTrue(TEXT("FCbAttachment(BinaryEmpty).IsNull()"), Attachment.IsNull());
-		TestFalse(TEXT("FCbAttachment(BinaryEmpty).IsBinary()"), Attachment.IsBinary());
-		TestFalse(TEXT("FCbAttachment(BinaryEmpty).IsCompactBinary()"), Attachment.IsCompactBinary());
-		TestEqual(TEXT("FCbAttachment(BinaryEmpty).GetHash()"), Attachment.GetHash(), FIoHash());
+		TestFalse(TEXT("FCbAttachment(BinaryEmpty).IsNull()"), Attachment.IsNull());
+		TestTrue(TEXT("FCbAttachment(BinaryEmpty).IsBinary()"), Attachment.IsBinary());
+		TestFalse(TEXT("FCbAttachment(BinaryEmpty).IsObject()"), Attachment.IsObject());
+		TestEqual(TEXT("FCbAttachment(BinaryEmpty).GetHash()"), Attachment.GetHash(), FIoHash::HashBuffer(FSharedBuffer{}));
 	}
 
-	// Compact Binary Empty
+	// Object Empty
 	{
-		const FCbAttachment Attachment(FCbFieldIterator{});
-		TestTrue(TEXT("FCbAttachment(CompactBinaryEmpty).IsNull()"), Attachment.IsNull());
-		TestFalse(TEXT("FCbAttachment(CompactBinaryEmpty).IsBinary()"), Attachment.IsBinary());
-		TestFalse(TEXT("FCbAttachment(CompactBinaryEmpty).IsCompactBinary()"), Attachment.IsCompactBinary());
-		TestEqual(TEXT("FCbAttachment(CompactBinaryEmpty).GetHash()"), Attachment.GetHash(), FIoHash());
+		const FCbAttachment Attachment(FCbObject{});
+		TestFalse(TEXT("FCbAttachment(ObjectEmpty).IsNull()"), Attachment.IsNull());
+		TestTrue(TEXT("FCbAttachment(ObjectEmpty).IsBinary()"), Attachment.IsBinary());
+		TestTrue(TEXT("FCbAttachment(ObjectEmpty).IsObject()"), Attachment.IsObject());
+		TestEqual(TEXT("FCbAttachment(ObjectEmpty).GetHash()"), Attachment.GetHash(), FIoHash(FCbObject().GetHash()));
 	}
 
 	return true;
@@ -303,19 +286,26 @@ bool FCbPackageTest::RunTest(const FString& Parameters)
 
 	// Attachment Only
 	{
-		FCbObject Object;
+		FCbObject Object1;
 		{
 			TCbWriter<256> Writer;
 			Writer.BeginObject();
-			Writer << "Field" << 42;
+			Writer << "Field1" << 42;
 			Writer.EndObject();
-			Object = Writer.Save().AsObject();
+			Object1 = Writer.Save().AsObject();
+		}
+		FCbObject Object2;
+		{
+			TCbWriter<256> Writer;
+			Writer.BeginObject();
+			Writer << "Field2" << 42;
+			Writer.EndObject();
+			Object2 = Writer.Save().AsObject();
 		}
-		FCbField Field = FCbField::Clone(Object["Field"]);
 
 		FCbPackage Package;
-		Package.AddAttachment(FCbAttachment(FCbFieldIterator::MakeSingle(Object.AsField())));
-		Package.AddAttachment(FCbAttachment(Field.GetBuffer()));
+		Package.AddAttachment(FCbAttachment(Object1));
+		Package.AddAttachment(FCbAttachment(Object2.GetBuffer()));
 
 		TestFalse(TEXT("FCbPackage(Attachments).IsNull()"), Package.IsNull());
 		TestTrue(TEXT("FCbPackage(Attachments) as bool"), bool(Package));
@@ -324,31 +314,31 @@ bool FCbPackageTest::RunTest(const FString& Parameters)
 		TestEqual(TEXT("FCbPackage(Attachments).GetObjectHash()"), Package.GetObjectHash(), FIoHash());
 		TestSaveLoadValidate(TEXT("Attachments"), Package);
 
-		const FCbAttachment* const ObjectAttachment = Package.FindAttachment(Object.GetHash());
-		const FCbAttachment* const FieldAttachment = Package.FindAttachment(Field.GetHash());
+		const FCbAttachment* const Object1Attachment = Package.FindAttachment(Object1.GetHash());
+		const FCbAttachment* const Object2Attachment = Package.FindAttachment(Object2.GetHash());
 
-		TestTrue(TEXT("FCbPackage(Attachments).FindAttachment(Object)"),
-			ObjectAttachment && ObjectAttachment->AsCompactBinary().AsObject().Equals(Object));
-		TestTrue(TEXT("FCbPackage(Attachments).FindAttachment(Field)"),
-			FieldAttachment && FieldAttachment->AsBinaryView() == Field.GetBuffer());
+		TestTrue(TEXT("FCbPackage(Attachments).FindAttachment(Object1)"),
+			Object1Attachment && Object1Attachment->AsObject().Equals(Object1));
+		TestTrue(TEXT("FCbPackage(Attachments).FindAttachment(Object2)"),
+			Object2Attachment && Object2Attachment->AsBinary() == Object2.GetBuffer());
 
-		Package.AddAttachment(FCbAttachment(FSharedBuffer::Clone(Object.GetView())));
-		Package.AddAttachment(FCbAttachment(FCbFieldIterator::CloneRange(FCbFieldViewIterator::MakeSingle(Field))));
+		Package.AddAttachment(FCbAttachment(FSharedBuffer::Clone(Object1.GetView())));
+		Package.AddAttachment(FCbAttachment(FCbObject::Clone(Object2)));
 
 		TestEqual(TEXT("FCbPackage(Attachments).GetAttachments()"), Package.GetAttachments().Num(), 2);
-		TestEqual(TEXT("FCbPackage(Attachments).FindAttachment(Object, Re-Add)"),
-			Package.FindAttachment(Object.GetHash()), ObjectAttachment);
-		TestEqual(TEXT("FCbPackage(Attachments).FindAttachment(Field, Re-Add)"),
-			Package.FindAttachment(Field.GetHash()), FieldAttachment);
+		TestEqual(TEXT("FCbPackage(Attachments).FindAttachment(Object1, Re-Add)"),
+			Package.FindAttachment(Object1.GetHash()), Object1Attachment);
+		TestEqual(TEXT("FCbPackage(Attachments).FindAttachment(Object2, Re-Add)"),
+			Package.FindAttachment(Object2.GetHash()), Object2Attachment);
 
 		TestTrue(TEXT("FCbPackage(Attachments).FindAttachment(ObjectAsObject)"),
-			ObjectAttachment && ObjectAttachment->AsCompactBinary().AsObject().Equals(Object));
+			Object1Attachment && Object1Attachment->AsObject().Equals(Object1));
 		TestTrue(TEXT("FCbPackage(Attachments).FindAttachment(ObjectAsBinary)"),
-			ObjectAttachment && ObjectAttachment->AsBinaryView() == Object.GetBuffer());
+			Object1Attachment && Object1Attachment->AsBinary() == Object1.GetBuffer());
 		TestTrue(TEXT("FCbPackage(Attachments).FindAttachment(FieldAsField)"),
-			FieldAttachment && FieldAttachment->AsCompactBinary().Equals(Field));
+			Object2Attachment && Object2Attachment->AsObject().Equals(Object2));
 		TestTrue(TEXT("FCbPackage(Attachments).FindAttachment(FieldAsBinary)"),
-			FieldAttachment && FieldAttachment->AsBinaryView() == Field.GetBuffer());
+			Object2Attachment && Object2Attachment->AsBinary() == Object2.GetBuffer());
 
 		TestTrue(TEXT("FCbPackage(Attachments).GetAttachments()->Sorted"),
 			Algo::IsSorted(Package.GetAttachments()));
@@ -359,34 +349,35 @@ bool FCbPackageTest::RunTest(const FString& Parameters)
 	FSharedBuffer Level4 = FSharedBuffer::MakeView(MakeMemoryView(Level4Values));
 	const FIoHash Level4Hash = FIoHash::HashBuffer(Level4);
 
-	FCbField Level3;
+	FCbObject Level3;
 	{
 		TCbWriter<256> Writer;
-		Writer.SetName("Level4").AddBinaryAttachment(Level4Hash);
-		Level3 = Writer.Save();
+		Writer.BeginObject();
+		Writer.AddBinaryAttachment("Level4", Level4Hash);
+		Writer.EndObject();
+		Level3 = Writer.Save().AsObject();
 	}
 	const FIoHash Level3Hash = Level3.GetHash();
 
-	FCbArray Level2;
+	FCbObject Level2;
 	{
 		TCbWriter<256> Writer;
-		Writer.SetName("Level3");
-		Writer.BeginArray();
-		Writer.AddObjectAttachment(Level3Hash);
-		Writer.EndArray();
-		Level2 = Writer.Save().AsArray();
+		Writer.BeginObject();
+		Writer.AddObjectAttachment("Level3", Level3Hash);
+		Writer.EndObject();
+		Level2 = Writer.Save().AsObject();
 	}
-	const FIoHash Level2Hash = Level2.AsFieldView().GetHash();
+	const FIoHash Level2Hash = Level2.GetHash();
 
 	FCbObject Level1;
 	{
 		TCbWriter<256> Writer;
 		Writer.BeginObject();
-		Writer.SetName("Level2").AddObjectAttachment(Level2Hash);
+		Writer.AddObjectAttachment("Level2", Level2Hash);
 		Writer.EndObject();
 		Level1 = Writer.Save().AsObject();
 	}
-	const FIoHash Level1Hash = Level1.AsFieldView().GetHash();
+	const FIoHash Level1Hash = Level1.GetHash();
 
 	const auto Resolver = [&Level2, &Level2Hash, &Level3, &Level3Hash, &Level4, &Level4Hash]
 		(const FIoHash& Hash) -> FSharedBuffer
@@ -415,13 +406,13 @@ bool FCbPackageTest::RunTest(const FString& Parameters)
 		const FCbAttachment* const Level3Attachment = Package.FindAttachment(Level3Hash);
 		const FCbAttachment* const Level4Attachment = Package.FindAttachment(Level4Hash);
 		TestTrue(TEXT("FCbPackage(Object+Attachments).FindAttachment(Level2)"),
-			Level2Attachment && Level2Attachment->AsCompactBinary().AsArray().Equals(Level2));
+			Level2Attachment && Level2Attachment->AsObject().Equals(Level2));
 		TestTrue(TEXT("FCbPackage(Object+Attachments).FindAttachment(Level3)"),
-			Level3Attachment && Level3Attachment->AsCompactBinary().Equals(Level3));
+			Level3Attachment && Level3Attachment->AsObject().Equals(Level3));
 		TestTrue(TEXT("FCbPackage(Object+Attachments).FindAttachment(Level4)"),
 			Level4Attachment &&
-			Level4Attachment->AsBinaryView() != Level4 &&
-			Level4Attachment->AsBinaryView().GetView().EqualBytes(Level4.GetView()));
+			Level4Attachment->AsBinary() != Level4 &&
+			Level4Attachment->AsBinary().GetView().EqualBytes(Level4.GetView()));
 
 		TestTrue(TEXT("FCbPackage(Object+Attachments).GetAttachments()->Sorted"),
 			Algo::IsSorted(Package.GetAttachments()));
@@ -476,23 +467,23 @@ bool FCbPackageTest::RunTest(const FString& Parameters)
 		TestEqual(TEXT("FCbPackage(OutOfOrder).Load()->Level1Hash"), FromFields.GetObjectHash(), Level1Hash);
 
 		TestTrue(TEXT("FCbPackage(OutOfOrder).Load()->Level2"),
-			Level2Attachment && Level2Attachment->AsCompactBinary().AsArray().Equals(Level2));
+			Level2Attachment && Level2Attachment->AsObject().Equals(Level2));
 		TestTrue(TEXT("FCbPackage(OutOfOrder).Load()->Level2Buffer"),
-			Level2Attachment && Fields.GetOuterBuffer().GetView().Contains(Level2Attachment->AsBinaryView().GetView()));
+			Level2Attachment && Fields.GetOuterBuffer().GetView().Contains(Level2Attachment->AsBinary().GetView()));
 		TestTrue(TEXT("FCbPackage(OutOfOrder).Load()->Level2Hash"),
 			Level2Attachment && Level2Attachment->GetHash() == Level2Hash);
 
 		TestTrue(TEXT("FCbPackage(OutOfOrder).Load()->Level3"),
-			Level3Attachment && Level3Attachment->AsCompactBinary().Equals(Level3));
+			Level3Attachment && Level3Attachment->AsObject().Equals(Level3));
 		TestTrue(TEXT("FCbPackage(OutOfOrder).Load()->Level3Buffer"),
-			Level3Attachment && Fields.GetOuterBuffer().GetView().Contains(Level3Attachment->AsBinaryView().GetView()));
+			Level3Attachment && Fields.GetOuterBuffer().GetView().Contains(Level3Attachment->AsBinary().GetView()));
 		TestTrue(TEXT("FCbPackage(OutOfOrder).Load()->Level3Hash"),
 			Level3Attachment && Level3Attachment->GetHash() == Level3Hash);
 
 		TestTrue(TEXT("FCbPackage(OutOfOrder).Load()->Level4"),
-			Level4Attachment && Level4Attachment->AsBinaryView().GetView().EqualBytes(Level4.GetView()));
+			Level4Attachment && Level4Attachment->AsBinary().GetView().EqualBytes(Level4.GetView()));
 		TestTrue(TEXT("FCbPackage(OutOfOrder).Load()->Level4Buffer"),
-			Level4Attachment && Fields.GetOuterBuffer().GetView().Contains(Level4Attachment->AsBinaryView().GetView()));
+			Level4Attachment && Fields.GetOuterBuffer().GetView().Contains(Level4Attachment->AsBinary().GetView()));
 		TestTrue(TEXT("FCbPackage(OutOfOrder).Load()->Level4Hash"),
 			Level4Attachment && Level4Attachment->GetHash() == Level4Hash);
 
@@ -542,7 +533,7 @@ bool FCbPackageTest::RunTest(const FString& Parameters)
 		bool bResolved = false;
 		FCbPackage Package;
 		Package.AddAttachment(FCbAttachment(Level3.GetBuffer()));
-		Package.AddAttachment(FCbAttachment(FCbFieldIterator::MakeSingle(Level3)),
+		Package.AddAttachment(FCbAttachment(Level3),
 			[&bResolved](const FIoHash& Hash) -> FSharedBuffer
 			{
 				bResolved = true;
diff --git a/Engine/Source/Runtime/Core/Private/Tests/Serialization/CompactBinaryValidationTest.cpp b/Engine/Source/Runtime/Core/Private/Tests/Serialization/CompactBinaryValidationTest.cpp
--- a/Engine/Source/Runtime/Core/Private/Tests/Serialization/CompactBinaryValidationTest.cpp
+++ b/Engine/Source/Runtime/Core/Private/Tests/Serialization/CompactBinaryValidationTest.cpp
@@ -365,24 +365,28 @@ bool FCbValidateAttachmentTest::RunTest(const FString& Parameters)
 		TestEqual(TEXT("ValidateCompactBinaryAttachment(Binary, OutOfBoundsHash, Mode)"), ValidateCompactBinaryAttachment(MakeMemoryView(Buffer).Left(4), ECbValidateMode::None), ECbValidateError::None);
 	}
 
-	// Test CompactBinary
+	// Test Object
 	{
+		TCbWriter<256> Writer;
+		Writer.BeginObject();
+		Writer.AddNull("N"_ASV);
+		Writer.EndObject();
 		FBufferArchive Buffer;
-		FCbAttachment(FCbFieldIterator::MakeSingle(FCbField::MakeView(FCbFieldView(nullptr, ECbFieldType::Null)))).Save(Buffer);
-		TestEqual(TEXT("ValidateCompactBinaryAttachment(CompactBinary)"), ValidateCompactBinaryAttachment(MakeMemoryView(Buffer), ECbValidateMode::All), ECbValidateError::None);
-		TestEqual(TEXT("ValidateCompactBinaryAttachment(CompactBinary, MissingHash)"), ValidateCompactBinaryAttachment(MakeMemoryView(Buffer).Left(3), ECbValidateMode::All), ECbValidateError::InvalidPackageFormat);
-		TestEqual(TEXT("ValidateCompactBinaryAttachment(CompactBinary, MissingHash, Mode)"), ValidateCompactBinaryAttachment(MakeMemoryView(Buffer).Left(3), ECbValidateMode::All & ~ECbValidateMode::Package), ECbValidateError::None);
+		FCbAttachment(Writer.Save().AsObject()).Save(Buffer);
+		TestEqual(TEXT("ValidateCompactBinaryAttachment(Object)"), ValidateCompactBinaryAttachment(MakeMemoryView(Buffer), ECbValidateMode::All), ECbValidateError::None);
+		TestEqual(TEXT("ValidateCompactBinaryAttachment(Object, MissingHash)"), ValidateCompactBinaryAttachment(MakeMemoryView(Buffer).Left(7), ECbValidateMode::All), ECbValidateError::InvalidPackageFormat);
+		TestEqual(TEXT("ValidateCompactBinaryAttachment(Object, MissingHash, Mode)"), ValidateCompactBinaryAttachment(MakeMemoryView(Buffer).Left(7), ECbValidateMode::All & ~ECbValidateMode::Package), ECbValidateError::None);
 		++Buffer.Last();
-		TestEqual(TEXT("ValidateCompactBinaryAttachment(CompactBinary, InvalidHash)"), ValidateCompactBinaryAttachment(MakeMemoryView(Buffer), ECbValidateMode::All), ECbValidateError::InvalidPackageHash);
-		TestEqual(TEXT("ValidateCompactBinaryAttachment(CompactBinary, InvalidHash, Mode)"), ValidateCompactBinaryAttachment(MakeMemoryView(Buffer), ECbValidateMode::All & ~ECbValidateMode::Package), ECbValidateError::None);
+		TestEqual(TEXT("ValidateCompactBinaryAttachment(Object, InvalidHash)"), ValidateCompactBinaryAttachment(MakeMemoryView(Buffer), ECbValidateMode::All), ECbValidateError::InvalidPackageHash);
+		TestEqual(TEXT("ValidateCompactBinaryAttachment(Object, InvalidHash, Mode)"), ValidateCompactBinaryAttachment(MakeMemoryView(Buffer), ECbValidateMode::All & ~ECbValidateMode::Package), ECbValidateError::None);
 		--Buffer.Last();
 		Buffer.AddZeroed(1);
-		TestEqual(TEXT("ValidateCompactBinaryAttachment(CompactBinary, Padding)"), ValidateCompactBinaryAttachment(MakeMemoryView(Buffer), ECbValidateMode::All), ECbValidateError::Padding);
-		TestEqual(TEXT("ValidateCompactBinaryAttachment(CompactBinary, Padding, Mode)"), ValidateCompactBinaryAttachment(MakeMemoryView(Buffer), ECbValidateMode::All & ~ECbValidateMode::Padding), ECbValidateError::None);
+		TestEqual(TEXT("ValidateCompactBinaryAttachment(Object, Padding)"), ValidateCompactBinaryAttachment(MakeMemoryView(Buffer), ECbValidateMode::All), ECbValidateError::Padding);
+		TestEqual(TEXT("ValidateCompactBinaryAttachment(Object, Padding, Mode)"), ValidateCompactBinaryAttachment(MakeMemoryView(Buffer), ECbValidateMode::All & ~ECbValidateMode::Padding), ECbValidateError::None);
 		Buffer.Pop(/*bAllowShrinking*/ false);
-		TestEqual(TEXT("ValidateCompactBinaryAttachment(CompactBinary, OutOfBoundsValue)"), ValidateCompactBinaryAttachment(MakeMemoryView(Buffer).Left(1), ECbValidateMode::All), ECbValidateError::OutOfBounds);
-		TestEqual(TEXT("ValidateCompactBinaryAttachment(CompactBinary, OutOfBoundsHash)"), ValidateCompactBinaryAttachment(MakeMemoryView(Buffer).Left(4), ECbValidateMode::All), ECbValidateError::OutOfBounds);
-		TestEqual(TEXT("ValidateCompactBinaryAttachment(CompactBinary, OutOfBoundsHash, Mode)"), ValidateCompactBinaryAttachment(MakeMemoryView(Buffer).Left(4), ECbValidateMode::None), ECbValidateError::None);
+		TestEqual(TEXT("ValidateCompactBinaryAttachment(Object, OutOfBoundsValue)"), ValidateCompactBinaryAttachment(MakeMemoryView(Buffer).Left(1), ECbValidateMode::All), ECbValidateError::OutOfBounds);
+		TestEqual(TEXT("ValidateCompactBinaryAttachment(Object, OutOfBoundsHash)"), ValidateCompactBinaryAttachment(MakeMemoryView(Buffer).Left(8), ECbValidateMode::All), ECbValidateError::OutOfBounds);
+		TestEqual(TEXT("ValidateCompactBinaryAttachment(Object, OutOfBoundsHash, Mode)"), ValidateCompactBinaryAttachment(MakeMemoryView(Buffer).Left(8), ECbValidateMode::None), ECbValidateError::None);
 	}
 
 	// Test InvalidPackageFormat
diff --git a/Engine/Source/Runtime/Core/Public/Serialization/CompactBinaryPackage.h b/Engine/Source/Runtime/Core/Public/Serialization/CompactBinaryPackage.h
--- a/Engine/Source/Runtime/Core/Public/Serialization/CompactBinaryPackage.h
+++ b/Engine/Source/Runtime/Core/Public/Serialization/CompactBinaryPackage.h
@@ -15,9 +15,9 @@ template <typename FuncType> class TFunctionRef;
 ///////////////////////////////////////////////////////////////////////////////////////////////////////////////////////
 
 /**
- * An attachment is either binary or compact binary and is identified by its hash.
+ * An attachment is either binary or object and is identified by its hash.
  *
- * A compact binary attachment is also a valid binary attachment and may be accessed as binary.
+ * An object attachment is also a valid binary attachment and may be accessed as binary.
  *
  * Attachments are serialized as one or two compact binary fields with no name. A Binary field is
  * written first with its content. The content hash is omitted when the content size is zero, and
@@ -29,14 +29,14 @@ public:
 	/** Construct a null attachment. */
 	FCbAttachment() = default;
 
-	/** Construct a compact binary attachment. Value is cloned if not owned. */
-	inline explicit FCbAttachment(FCbFieldIterator Value)
+	/** Construct an object attachment. Value is cloned if not owned. */
+	inline explicit FCbAttachment(FCbObject Value)
 		: FCbAttachment(MoveTemp(Value), nullptr)
 	{
 	}
 
-	/** Construct a compact binary attachment. Value is cloned if not owned. Hash must match Value. */
-	inline explicit FCbAttachment(FCbFieldIterator Value, const FIoHash& Hash)
+	/** Construct an object attachment. Value is cloned if not owned. Hash must match Value. */
+	inline explicit FCbAttachment(FCbObject Value, const FIoHash& Hash)
 		: FCbAttachment(MoveTemp(Value), &Hash)
 	{
 	}
@@ -62,17 +62,17 @@ public:
 	/** Whether the attachment has a value. */
 	inline bool IsNull() const { return !Buffer; }
 
-	/** Access the attachment as binary. Defaults to a null buffer on error. */
-	CORE_API FSharedBuffer AsBinaryView() const;
+	/** Access the attachment as binary. Defaults to a null buffer if the attachment is null. */
+	inline FSharedBuffer AsBinary() const { return Buffer; }
 
-	/** Access the attachment as compact binary. Defaults to a field iterator with no value on error. */
-	CORE_API FCbFieldIterator AsCompactBinary() const;
+	/** Access the attachment as an object. Defaults to an empty object on error. */
+	inline FCbObject AsObject() const { return Object ? FCbObject(FCbFieldView(Object).AsObjectView(), Buffer) : FCbObject(); }
 
-	/** Returns whether the attachment is binary or compact binary. */
+	/** Returns whether the attachment is binary or an object. */
 	inline bool IsBinary() const { return !Buffer.IsNull(); }
 
-	/** Returns whether the attachment is compact binary. */
-	inline bool IsCompactBinary() const { return CompactBinary.HasValue(); }
+	/** Returns whether the attachment is an object. */
+	inline bool IsObject() const { return Object.IsObject(); }
 
 	/** Returns the hash of the attachment value. */
 	inline const FIoHash& GetHash() const { return Hash; }
@@ -108,13 +108,13 @@ public:
 	CORE_API void Save(FArchive& Ar) const;
 
 private:
-	CORE_API FCbAttachment(FCbFieldIterator Value, const FIoHash* Hash);
+	CORE_API FCbAttachment(FCbObject Value, const FIoHash* Hash);
 	CORE_API FCbAttachment(FSharedBuffer Value, const FIoHash* Hash);
 
-	/** An owned buffer containing the binary or compact binary data. */
+	/** An owned buffer containing the binary or object data. */
 	FSharedBuffer Buffer;
-	/** A field iterator that is valid only for compact binary attachments. */
-	FCbFieldViewIterator CompactBinary;
+	/** A field that is valid only for object attachments. */
+	FCbFieldView Object;
 	/** A hash of the attachment value. */
 	FIoHash Hash;
 };
@@ -128,16 +128,16 @@ inline uint32 GetTypeHash(const FCbAttachment& Attachment)
 ///////////////////////////////////////////////////////////////////////////////////////////////////////////////////////
 
 /**
- * A package is a compact binary object with attachments for its external references.
+ * A package is an object with a tree of referenced attachments.
  *
- * A package is basically a Merkle tree with compact binary as its root and other non-leaf nodes,
- * and either binary or compact binary as its leaf nodes. A node references its child nodes using
- * BinaryHash or FieldHash fields in its compact binary representation.
+ * A package is a Merkle tree with an object as its root and non-leaf nodes, and either binary or
+ * objects as its leaf nodes. Nodes reference their children from fields of type BinaryAttachment
+ * or ObjectAttachment, which store the hash of the referenced data.
  *
  * It is invalid for a package to include attachments that are not referenced by its object or by
- * one of its referenced compact binary attachments. When attachments are added explicitly, it is
- * the responsibility of the package creator to follow this requirement. Attachments that are not
- * referenced may not survive a round-trip through certain storage systems.
+ * one of its referenced object attachments. This invariant needs to be maintained if attachments
+ * are added explicitly instead of being discovered by the attachment resolver. If any attachment
+ * is not referenced, it may not survive a round-trip through certain storage systems.
  *
  * It is valid for a package to exclude referenced attachments, but then it is the responsibility
  * of the package consumer to have a mechanism for resolving those references when necessary.
@@ -216,10 +216,10 @@ public:
 	/** Whether the package has an empty object and no attachments. */
 	inline bool IsNull() const { return !Object && Attachments.IsEmpty(); }
 
-	/** Returns the compact binary object for the package. */
+	/** Returns the object for the package. */
 	inline const FCbObject& GetObject() const { return Object; }
 
-	/** Returns the has of the compact binary object for the package. */
+	/** Returns the hash of the object for the package. */
 	inline const FIoHash& GetObjectHash() const { return ObjectHash; }
 
 	/**
@@ -338,7 +338,7 @@ private:
 	CORE_API void SetObject(FCbObject Object, const FIoHash* Hash, FAttachmentResolver* Resolver);
 	CORE_API void AddAttachment(const FCbAttachment& Attachment, FAttachmentResolver* Resolver);
 
-	void GatherAttachments(const FCbFieldViewIterator& Fields, FAttachmentResolver Resolver);
+	void GatherAttachments(const FCbObject& Object, FAttachmentResolver Resolver);
 
 	/** Attachments ordered by their hash. */
 	TArray<FCbAttachment> Attachments;
