CompactBinary: Made buffers and views easier to use by eliminating the complex options

- Renamed TryGetSerialized[Range]View to TryGet[Range]View.
- Changed GetBuffer to return a composite buffer with the type.
- Changed GetView to be protected instead of public.
- Removed GetRangeBuffer and GetRangeView.

#rnx

diff --git a/Engine/Source/Developer/DerivedDataCache/Private/DerivedDataBuildJobContext.cpp b/Engine/Source/Developer/DerivedDataCache/Private/DerivedDataBuildJobContext.cpp
--- a/Engine/Source/Developer/DerivedDataCache/Private/DerivedDataBuildJobContext.cpp
+++ b/Engine/Source/Developer/DerivedDataCache/Private/DerivedDataBuildJobContext.cpp
@@ -99,15 +99,7 @@ void FBuildJobContext::AddPayload(const FPayloadId& Id, const FSharedBuffer& Buf
 
 void FBuildJobContext::AddPayload(const FPayloadId& Id, const FCbObject& Object)
 {
-	FMemoryView SerializedView;
-	if (Object.TryGetSerializedView(SerializedView))
-	{
-		AddPayload(FPayload(Id, FCompressedBuffer::Compress(NAME_Default, Object.GetBuffer())));
-	}
-	else
-	{
-		AddPayload(FPayload(Id, FCompressedBuffer::Compress(NAME_Default, FCbObject::Clone(Object).GetBuffer())));
-	}
+	AddPayload(FPayload(Id, FCompressedBuffer::Compress(NAME_Default, Object.GetBuffer())));
 }
 
 void FBuildJobContext::BeginAsyncBuild()
diff --git a/Engine/Source/Runtime/Core/Private/Serialization/CompactBinary.cpp b/Engine/Source/Runtime/Core/Private/Serialization/CompactBinary.cpp
--- a/Engine/Source/Runtime/Core/Private/Serialization/CompactBinary.cpp
+++ b/Engine/Source/Runtime/Core/Private/Serialization/CompactBinary.cpp
@@ -4,6 +4,7 @@
 
 #include "Hash/Blake3.h"
 #include "Math/UnrealMathUtility.h"
+#include "Memory/CompositeBuffer.h"
 #include "Misc/ByteSwap.h"
 #include "Misc/DateTime.h"
 #include "Misc/Guid.h"
@@ -800,7 +801,7 @@ template <typename FieldType>
 uint64 TCbFieldIterator<FieldType>::GetRangeSize() const
 {
 	FMemoryView View;
-	if (TryGetSerializedRangeView(View))
+	if (TryGetRangeView(View))
 	{
 		return View.GetSize();
 	}
@@ -827,7 +828,7 @@ template <typename FieldType>
 void TCbFieldIterator<FieldType>::AppendRangeHash(FBlake3& Builder) const
 {
 	FMemoryView View;
-	if (TryGetSerializedRangeView(View))
+	if (TryGetRangeView(View))
 	{
 		Builder.Update(View);
 	}
@@ -844,7 +845,7 @@ template <typename FieldType>
 void TCbFieldIterator<FieldType>::CopyRangeTo(FMutableMemoryView InBuffer) const
 {
 	FMemoryView Source;
-	if (TryGetSerializedRangeView(Source))
+	if (TryGetRangeView(Source))
 	{
 		checkf(InBuffer.GetSize() == Source.GetSize(),
 			TEXT("Buffer is %" UINT64_FMT " bytes but %" UINT64_FMT " is required."),
@@ -867,7 +868,7 @@ void TCbFieldIterator<FieldType>::CopyRangeTo(FArchive& Ar) const
 {
 	check(Ar.IsSaving());
 	FMemoryView Source;
-	if (TryGetSerializedRangeView(Source))
+	if (TryGetRangeView(Source))
 	{
 		Ar.Serialize(const_cast<void*>(Source.GetData()), static_cast<int64>(Source.GetSize()));
 	}
@@ -915,7 +916,7 @@ template class TCbFieldIterator<FCbField>;
 FCbFieldIterator FCbFieldIterator::CloneRange(const FCbFieldViewIterator& It)
 {
 	FMemoryView View;
-	if (It.TryGetSerializedRangeView(View))
+	if (It.TryGetRangeView(View))
 	{
 		return MakeRange(FSharedBuffer::Clone(View));
 	}
@@ -927,11 +928,49 @@ FCbFieldIterator FCbFieldIterator::CloneRange(const FCbFieldViewIterator& It)
 	}
 }
 
-FSharedBuffer FCbFieldIterator::GetRangeBuffer() const
+///////////////////////////////////////////////////////////////////////////////////////////////////////////////////////
+
+FCompositeBuffer FCbField::GetBuffer() const
+{
+	FMemoryView View;
+	if (TryGetView(View))
+	{
+		return FCompositeBuffer(FSharedBuffer::MakeView(View, GetOuterBuffer()));
+	}
+	else
+	{
+		const ECbFieldType SerializedType = FCbFieldType::GetSerializedType(GetTypeWithFlags());
+		return FCompositeBuffer(MakeSharedBufferFromArray(TArray<ECbFieldType, TInlineAllocator<1>>{SerializedType}),
+			FSharedBuffer::MakeView(GetViewNoType(), GetOuterBuffer()));
+	}
+}
+
+FCompositeBuffer FCbArray::GetBuffer() const
+{
+	FMemoryView View;
+	if (TryGetView(View))
+	{
+		return FCompositeBuffer(FSharedBuffer::MakeView(View, GetOuterBuffer()));
+	}
+	else
+	{
+		return FCompositeBuffer(MakeSharedBufferFromArray(TArray<ECbFieldType, TInlineAllocator<1>>{GetType()}),
+			FSharedBuffer::MakeView(GetPayloadView(), GetOuterBuffer()));
+	}
+}
+
+FCompositeBuffer FCbObject::GetBuffer() const
 {
-	const FMemoryView RangeView = GetRangeView();
-	const FSharedBuffer& OuterBuffer = GetOuterBuffer();
-	return OuterBuffer.GetView() == RangeView ? OuterBuffer : FSharedBuffer::MakeView(RangeView, OuterBuffer);
+	FMemoryView View;
+	if (TryGetView(View))
+	{
+		return FCompositeBuffer(FSharedBuffer::MakeView(View, GetOuterBuffer()));
+	}
+	else
+	{
+		return FCompositeBuffer(MakeSharedBufferFromArray(TArray<ECbFieldType, TInlineAllocator<1>>{GetType()}),
+			FSharedBuffer::MakeView(GetPayloadView(), GetOuterBuffer()));
+	}
 }
 
 ///////////////////////////////////////////////////////////////////////////////////////////////////////////////////////
diff --git a/Engine/Source/Runtime/Core/Private/Serialization/CompactBinaryPackage.cpp b/Engine/Source/Runtime/Core/Private/Serialization/CompactBinaryPackage.cpp
--- a/Engine/Source/Runtime/Core/Private/Serialization/CompactBinaryPackage.cpp
+++ b/Engine/Source/Runtime/Core/Private/Serialization/CompactBinaryPackage.cpp
@@ -3,6 +3,7 @@
 #include "Serialization/CompactBinaryPackage.h"
 
 #include "Algo/BinarySearch.h"
+#include "Memory/CompositeBuffer.h"
 #include "Serialization/CompactBinarySerialization.h"
 #include "Serialization/CompactBinaryWriter.h"
 
@@ -11,13 +12,13 @@
 FCbAttachment::FCbAttachment(FCbObject InValue, const FIoHash* const InHash)
 {
 	FMemoryView View;
-	if (!InValue.IsOwned() || !InValue.TryGetSerializedView(View))
+	if (!InValue.IsOwned() || !InValue.TryGetView(View))
 	{
 		InValue = FCbObject::Clone(InValue);
 	}
 
 	Object = InValue.AsFieldView();
-	Buffer = MoveTemp(InValue).GetBuffer();
+	Buffer = MoveTemp(InValue).GetBuffer().Flatten();
 
 	if (InHash)
 	{
diff --git a/Engine/Source/Runtime/Core/Private/Tests/Serialization/CompactBinaryPackageTest.cpp b/Engine/Source/Runtime/Core/Private/Tests/Serialization/CompactBinaryPackageTest.cpp
--- a/Engine/Source/Runtime/Core/Private/Tests/Serialization/CompactBinaryPackageTest.cpp
+++ b/Engine/Source/Runtime/Core/Private/Tests/Serialization/CompactBinaryPackageTest.cpp
@@ -30,7 +30,7 @@ bool FCbAttachmentTest::RunTest(const FString& Parameters)
 		FCbFieldIterator Fields = Writer.Save();
 
 		TestTrue(FString::Printf(TEXT("FCbAttachment(%s).Save()->Equals"), Test),
-			MakeMemoryView(WriteAr).EqualBytes(Fields.GetRangeBuffer().GetView()));
+			MakeMemoryView(WriteAr).EqualBytes(Fields.GetOuterBuffer().GetView()));
 		TestEqual(FString::Printf(TEXT("FCbAttachment(%s).Save()->ValidateRange"), Test),
 			ValidateCompactBinaryRange(MakeMemoryView(WriteAr), ECbValidateMode::All), ECbValidateError::None);
 		TestEqual(FString::Printf(TEXT("FCbAttachment(%s).Save()->ValidateAttachment"), Test),
@@ -85,7 +85,7 @@ bool FCbAttachmentTest::RunTest(const FString& Parameters)
 		FCbAttachment Attachment(Object);
 		TestFalse(TEXT("FCbAttachment(Object).IsNull()"), Attachment.IsNull());
 		TestTrue(TEXT("FCbAttachment(Object) as bool"), bool(Attachment));
-		TestEqual(TEXT("FCbAttachment(Object).AsBinary()"), Attachment.AsBinary(), Object.GetBuffer());
+		TestEqual(TEXT("FCbAttachment(Object).AsBinary()"), Attachment.AsBinary(), Object.GetOuterBuffer());
 		TestTrue(TEXT("FCbAttachment(Object).AsObject()"), Attachment.AsObject().Equals(Object));
 		TestTrue(TEXT("FCbAttachment(Object).IsBinary()"), Attachment.IsBinary());
 		TestTrue(TEXT("FCbAttachment(Object).IsObject()"), Attachment.IsObject());
@@ -119,7 +119,7 @@ bool FCbAttachmentTest::RunTest(const FString& Parameters)
 		FCbAttachment Attachment(ObjectView);
 		TestFalse(TEXT("FCbAttachment(ObjectView).IsNull()"), Attachment.IsNull());
 		TestTrue(TEXT("FCbAttachment(ObjectView) as bool"), bool(Attachment));
-		TestNotEqual(TEXT("FCbAttachment(ObjectView).AsBinary()"), Attachment.AsBinary(), ObjectView.GetBuffer());
+		TestNotEqual(TEXT("FCbAttachment(ObjectView).AsBinary()"), Attachment.AsBinary(), ObjectView.GetOuterBuffer());
 		TestTrue(TEXT("FCbAttachment(ObjectView).AsObject()"), Attachment.AsObject().Equals(Object));
 		TestTrue(TEXT("FCbAttachment(ObjectView).IsBinary()"), Attachment.IsBinary());
 		TestTrue(TEXT("FCbAttachment(ObjectView).IsObject()"), Attachment.IsObject());
@@ -141,7 +141,7 @@ bool FCbAttachmentTest::RunTest(const FString& Parameters)
 		TestFalse(TEXT("FCbAttachment(LoadBinaryView).IsNull()"), Attachment.IsNull());
 		TestTrue(TEXT("FCbAttachment(LoadBinaryView) as bool"), bool(Attachment));
 		TestFalse(TEXT("FCbAttachment(LoadBinaryView).AsBinary()->!InView"),
-			FieldsView.GetRangeBuffer().GetView().Contains(Attachment.AsBinary().GetView()));
+			FieldsView.GetOuterBuffer().GetView().Contains(Attachment.AsBinary().GetView()));
 		TestTrue(TEXT("FCbAttachment(LoadBinaryView).AsBinary()->EqualBytes"),
 			Attachment.AsBinary().GetView().EqualBytes(Buffer.GetView()));
 		TestFalse(TEXT("FCbAttachment(LoadBinaryView).AsObject()"), bool(Attachment.AsObject()));
@@ -159,7 +159,7 @@ bool FCbAttachmentTest::RunTest(const FString& Parameters)
 		ValueWriter.EndObject();
 		const FCbObject Value = ValueWriter.Save().AsObject();
 		TestEqual(TEXT("FCbAttachment(LoadObjectView).Validate"),
-			ValidateCompactBinaryRange(Value.GetView(), ECbValidateMode::All), ECbValidateError::None);
+			ValidateCompactBinaryRange(Value.GetOuterBuffer(), ECbValidateMode::All), ECbValidateError::None);
 		FCbAttachment Attachment(Value);
 
 		FCbWriter Writer;
@@ -168,12 +168,13 @@ bool FCbAttachmentTest::RunTest(const FString& Parameters)
 		FCbFieldIterator FieldsView = FCbFieldIterator::MakeRangeView(FCbFieldViewIterator(Fields));
 
 		Attachment.Load(FieldsView);
+		FMemoryView View;
 		TestFalse(TEXT("FCbAttachment(LoadObjectView).IsNull()"), Attachment.IsNull());
 		TestTrue(TEXT("FCbAttachment(LoadObjectView) as bool"), bool(Attachment));
 		TestTrue(TEXT("FCbAttachment(LoadObjectView).AsBinary()->EqualBytes"),
-			Attachment.AsBinary().GetView().EqualBytes(Value.GetView()));
+			Attachment.AsBinary().GetView().EqualBytes(Value.GetOuterBuffer().GetView()));
 		TestFalse(TEXT("FCbAttachment(LoadObjectView).AsObject()->!InView"),
-			FieldsView.GetRangeBuffer().GetView().Contains(Attachment.AsObject().GetView()));
+			!Attachment.AsObject().TryGetView(View) || FieldsView.GetOuterBuffer().GetView().Contains(View));
 		TestTrue(TEXT("FCbAttachment(LoadObjectView).IsBinary()"), Attachment.IsBinary());
 		TestTrue(TEXT("FCbAttachment(LoadObjectView).IsObject()"), Attachment.IsObject());
 		TestEqual(TEXT("FCbAttachment(LoadObjectView).GetHash()"), Attachment.GetHash(), FIoHash(Value.GetHash()));
@@ -221,7 +222,7 @@ bool FCbPackageTest::RunTest(const FString& Parameters)
 		FCbFieldIterator Fields = Writer.Save();
 
 		TestTrue(FString::Printf(TEXT("FCbPackage(%s).Save()->Equals"), Test),
-			MakeMemoryView(WriteAr).EqualBytes(Fields.GetRangeBuffer().GetView()));
+			MakeMemoryView(WriteAr).EqualBytes(Fields.GetOuterBuffer().GetView()));
 		TestEqual(FString::Printf(TEXT("FCbPackage(%s).Save()->ValidateRange"), Test),
 			ValidateCompactBinaryRange(MakeMemoryView(WriteAr), ECbValidateMode::All), ECbValidateError::None);
 		TestEqual(FString::Printf(TEXT("FCbPackage(%s).Save()->ValidatePackage"), Test),
@@ -305,7 +306,7 @@ bool FCbPackageTest::RunTest(const FString& Parameters)
 
 		FCbPackage Package;
 		Package.AddAttachment(FCbAttachment(Object1));
-		Package.AddAttachment(FCbAttachment(Object2.GetBuffer()));
+		Package.AddAttachment(FCbAttachment(Object2.GetOuterBuffer()));
 
 		TestFalse(TEXT("FCbPackage(Attachments).IsNull()"), Package.IsNull());
 		TestTrue(TEXT("FCbPackage(Attachments) as bool"), bool(Package));
@@ -320,9 +321,9 @@ bool FCbPackageTest::RunTest(const FString& Parameters)
 		TestTrue(TEXT("FCbPackage(Attachments).FindAttachment(Object1)"),
 			Object1Attachment && Object1Attachment->AsObject().Equals(Object1));
 		TestTrue(TEXT("FCbPackage(Attachments).FindAttachment(Object2)"),
-			Object2Attachment && Object2Attachment->AsBinary() == Object2.GetBuffer());
+			Object2Attachment && Object2Attachment->AsBinary() == Object2.GetOuterBuffer());
 
-		Package.AddAttachment(FCbAttachment(FSharedBuffer::Clone(Object1.GetView())));
+		Package.AddAttachment(FCbAttachment(FSharedBuffer::Clone(Object1.GetOuterBuffer())));
 		Package.AddAttachment(FCbAttachment(FCbObject::Clone(Object2)));
 
 		TestEqual(TEXT("FCbPackage(Attachments).GetAttachments()"), Package.GetAttachments().Num(), 2);
@@ -334,11 +335,11 @@ bool FCbPackageTest::RunTest(const FString& Parameters)
 		TestTrue(TEXT("FCbPackage(Attachments).FindAttachment(ObjectAsObject)"),
 			Object1Attachment && Object1Attachment->AsObject().Equals(Object1));
 		TestTrue(TEXT("FCbPackage(Attachments).FindAttachment(ObjectAsBinary)"),
-			Object1Attachment && Object1Attachment->AsBinary() == Object1.GetBuffer());
+			Object1Attachment && Object1Attachment->AsBinary() == Object1.GetOuterBuffer());
 		TestTrue(TEXT("FCbPackage(Attachments).FindAttachment(FieldAsField)"),
 			Object2Attachment && Object2Attachment->AsObject().Equals(Object2));
 		TestTrue(TEXT("FCbPackage(Attachments).FindAttachment(FieldAsBinary)"),
-			Object2Attachment && Object2Attachment->AsBinary() == Object2.GetBuffer());
+			Object2Attachment && Object2Attachment->AsBinary() == Object2.GetOuterBuffer());
 
 		TestTrue(TEXT("FCbPackage(Attachments).GetAttachments()->Sorted"),
 			Algo::IsSorted(Package.GetAttachments()));
@@ -383,8 +384,8 @@ bool FCbPackageTest::RunTest(const FString& Parameters)
 		(const FIoHash& Hash) -> FSharedBuffer
 		{
 			return
-				Hash == Level2Hash ? Level2.GetBuffer() :
-				Hash == Level3Hash ? Level3.GetBuffer() :
+				Hash == Level2Hash ? Level2.GetOuterBuffer() :
+				Hash == Level3Hash ? Level3.GetOuterBuffer() :
 				Hash == Level4Hash ? Level4 :
 				FSharedBuffer();
 		};
@@ -398,7 +399,7 @@ bool FCbPackageTest::RunTest(const FString& Parameters)
 		TestTrue(TEXT("FCbPackage(Object+Attachments) as bool"), bool(Package));
 		TestEqual(TEXT("FCbPackage(Object+Attachments).GetAttachments()"), Package.GetAttachments().Num(), 3);
 		TestTrue(TEXT("FCbPackage(Object+Attachments).GetObject()"),
-			Package.GetObject().GetBuffer() == Level1.GetBuffer());
+			Package.GetObject().GetOuterBuffer() == Level1.GetOuterBuffer());
 		TestEqual(TEXT("FCbPackage(Object+Attachments).GetObjectHash()"), Package.GetObjectHash(), Level1Hash);
 		TestSaveLoadValidate(TEXT("Object+Attachments"), Package);
 
@@ -443,13 +444,13 @@ bool FCbPackageTest::RunTest(const FString& Parameters)
 	// Out of Order
 	{
 		TCbWriter<384> Writer;
-		Writer.AddBinary(Level2.GetBuffer());
+		Writer.AddBinary(Level2.GetOuterBuffer());
 		Writer.AddObjectAttachment(Level2Hash);
 		Writer.AddBinary(Level4);
 		Writer.AddBinaryAttachment(Level4Hash);
 		Writer.AddObject(Level1);
 		Writer.AddObjectAttachment(Level1Hash);
-		Writer.AddBinary(Level3.GetBuffer());
+		Writer.AddBinary(Level3.GetOuterBuffer());
 		Writer.AddObjectAttachment(Level3Hash);
 		Writer.AddNull();
 
@@ -496,15 +497,16 @@ bool FCbPackageTest::RunTest(const FString& Parameters)
 		Writer.Reset();
 		FromArchive.Save(Writer);
 		FCbFieldIterator Saved = Writer.Save();
+		FMemoryView View;
 		TestTrue(TEXT("FCbPackage(OutOfOrder).Save()->Level1"), Saved.AsObject().Equals(Level1));
 		++Saved;
 		TestEqual(TEXT("FCbPackage(OutOfOrder).Save()->Level1Hash"), Saved.AsObjectAttachment(), Level1Hash);
 		++Saved;
-		TestTrue(TEXT("FCbPackage(OutOfOrder).Save()->Level2"), Saved.AsBinaryView().EqualBytes(Level2.GetView()));
+		TestTrue(TEXT("FCbPackage(OutOfOrder).Save()->Level2"), Level2.TryGetView(View) && Saved.AsBinaryView().EqualBytes(View));
 		++Saved;
 		TestEqual(TEXT("FCbPackage(OutOfOrder).Save()->Level2Hash"), Saved.AsObjectAttachment(), Level2Hash);
 		++Saved;
-		TestTrue(TEXT("FCbPackage(OutOfOrder).Save()->Level3"), Saved.AsBinaryView().EqualBytes(Level3.GetView()));
+		TestTrue(TEXT("FCbPackage(OutOfOrder).Save()->Level3"), Level3.TryGetView(View) && Saved.AsBinaryView().EqualBytes(View));
 		++Saved;
 		TestEqual(TEXT("FCbPackage(OutOfOrder).Save()->Level3Hash"), Saved.AsObjectAttachment(), Level3Hash);
 		++Saved;
@@ -532,7 +534,7 @@ bool FCbPackageTest::RunTest(const FString& Parameters)
 	{
 		bool bResolved = false;
 		FCbPackage Package;
-		Package.AddAttachment(FCbAttachment(Level3.GetBuffer()));
+		Package.AddAttachment(FCbAttachment(Level3.GetOuterBuffer()));
 		Package.AddAttachment(FCbAttachment(Level3),
 			[&bResolved](const FIoHash& Hash) -> FSharedBuffer
 			{
diff --git a/Engine/Source/Runtime/Core/Private/Tests/Serialization/CompactBinarySerializationTest.cpp b/Engine/Source/Runtime/Core/Private/Tests/Serialization/CompactBinarySerializationTest.cpp
--- a/Engine/Source/Runtime/Core/Private/Tests/Serialization/CompactBinarySerializationTest.cpp
+++ b/Engine/Source/Runtime/Core/Private/Tests/Serialization/CompactBinarySerializationTest.cpp
@@ -156,7 +156,7 @@ bool FCbSaveTest::RunTest(const FString& Parameters)
 			{
 				FMemoryReader ReadAr(WriteAr);
 				FCbField Field = LoadCompactBinary(ReadAr);
-				TestTrue(FString::Printf(TEXT("LoadCompactBinary(%s)->EqualBytes"), Test), ExpectedData.EqualBytes(Field.GetView()));
+				TestTrue(FString::Printf(TEXT("LoadCompactBinary(%s)->EqualBytes"), Test), ExpectedData.EqualBytes(Field.GetOuterBuffer().GetView()));
 			}
 		}
 		{
@@ -167,7 +167,7 @@ bool FCbSaveTest::RunTest(const FString& Parameters)
 			{
 				FMemoryReader ReadAr(WriteAr);
 				ReadAr << Value;
-				TestTrue(FString::Printf(TEXT("Ar << CompactBinary Load(%s)->EqualBytes"), Test), ExpectedData.EqualBytes(Value.GetView()));
+				TestTrue(FString::Printf(TEXT("Ar << CompactBinary Load(%s)->EqualBytes"), Test), ExpectedData.EqualBytes(Value.GetOuterBuffer().GetView()));
 			}
 		}
 	};
diff --git a/Engine/Source/Runtime/Core/Private/Tests/Serialization/CompactBinaryTest.cpp b/Engine/Source/Runtime/Core/Private/Tests/Serialization/CompactBinaryTest.cpp
--- a/Engine/Source/Runtime/Core/Private/Tests/Serialization/CompactBinaryTest.cpp
+++ b/Engine/Source/Runtime/Core/Private/Tests/Serialization/CompactBinaryTest.cpp
@@ -3,6 +3,7 @@
 #include "Serialization/CompactBinary.h"
 
 #include "IO/IoHash.h"
+#include "Memory/CompositeBuffer.h"
 #include "Misc/AutomationTest.h"
 #include "Misc/DateTime.h"
 #include "Misc/Guid.h"
@@ -183,9 +184,8 @@ bool FCbFieldNoneTest::RunTest(const FString& Parameters)
 		TestFalse(TEXT("!FCbFieldView()::HasError()"), DefaultField.HasError());
 		TestEqual(TEXT("FCbFieldView()::GetError() == None"), DefaultField.GetError(), ECbFieldError::None);
 		TestEqual(TEXT("FCbFieldView()::GetHash()"), DefaultField.GetHash(), FBlake3::HashBuffer(MakeMemoryView<uint8>({uint8(ECbFieldType::None)})));
-		TestEqual(TEXT("FCbFieldView()::GetView()"), DefaultField.GetView(), FMemoryView());
-		FMemoryView SerializedView;
-		TestFalse(TEXT("FCbFieldView()::TryGetSerializedView()"), DefaultField.TryGetSerializedView(SerializedView));
+		FMemoryView View;
+		TestFalse(TEXT("FCbFieldView()::TryGetView()"), DefaultField.TryGetView(View));
 	}
 
 	// Test FCbFieldView(None)
@@ -198,9 +198,8 @@ bool FCbFieldNoneTest::RunTest(const FString& Parameters)
 		TestFalse(TEXT("!FCbFieldView(None)::HasError()"), NoneField.HasError());
 		TestEqual(TEXT("FCbFieldView(None)::GetError() == None"), NoneField.GetError(), ECbFieldError::None);
 		TestEqual(TEXT("FCbFieldView(None)::GetHash()"), NoneField.GetHash(), FCbFieldView().GetHash());
-		TestEqual(TEXT("FCbFieldView(None)::GetView()"), NoneField.GetView(), FMemoryView());
-		FMemoryView SerializedView;
-		TestFalse(TEXT("FCbFieldView(None)::TryGetSerializedView()"), NoneField.TryGetSerializedView(SerializedView));
+		FMemoryView View;
+		TestFalse(TEXT("FCbFieldView(None)::TryGetView()"), NoneField.TryGetView(View));
 	}
 
 	// Test FCbFieldView(None|Type|Name)
@@ -213,9 +212,8 @@ bool FCbFieldNoneTest::RunTest(const FString& Parameters)
 		TestTrue(TEXT("FCbFieldView(None|Type|Name)::HasName()"), NoneField.HasName());
 		TestFalse(TEXT("!FCbFieldView(None|Type|Name)::HasValue()"), NoneField.HasValue());
 		TestEqual(TEXT("FCbFieldView(None|Type|Name)::GetHash()"), NoneField.GetHash(), FBlake3::HashBuffer(MakeMemoryView(NoneBytes)));
-		TestEqual(TEXT("FCbFieldView(None|Type|Name)::GetView()"), NoneField.GetView(), MakeMemoryView(NoneBytes));
-		FMemoryView SerializedView;
-		TestTrue(TEXT("FCbFieldView(None|Type|Name)::TryGetSerializedView()"), NoneField.TryGetSerializedView(SerializedView) && SerializedView == MakeMemoryView(NoneBytes));
+		FMemoryView View;
+		TestTrue(TEXT("FCbFieldView(None|Type|Name)::TryGetView()"), NoneField.TryGetView(View) && View == MakeMemoryView(NoneBytes));
 
 		uint8 CopyBytes[sizeof(NoneBytes)];
 		NoneField.CopyTo(MakeMemoryView(CopyBytes));
@@ -232,9 +230,8 @@ bool FCbFieldNoneTest::RunTest(const FString& Parameters)
 		TestFalse(TEXT("FCbFieldView(None|Type)::HasName()"), NoneField.HasName());
 		TestFalse(TEXT("!FCbFieldView(None|Type)::HasValue()"), NoneField.HasValue());
 		TestEqual(TEXT("FCbFieldView(None|Type)::GetHash()"), NoneField.GetHash(), FCbFieldView().GetHash());
-		TestEqual(TEXT("FCbFieldView(None|Type)::GetView()"), NoneField.GetView(), MakeMemoryView(NoneBytes));
-		FMemoryView SerializedView;
-		TestTrue(TEXT("FCbFieldView(None|Type)::TryGetSerializedView()"), NoneField.TryGetSerializedView(SerializedView) && SerializedView == MakeMemoryView(NoneBytes));
+		FMemoryView View;
+		TestTrue(TEXT("FCbFieldView(None|Type)::TryGetView()"), NoneField.TryGetView(View) && View == MakeMemoryView(NoneBytes));
 	}
 
 	// Test FCbFieldView(None|Name)
@@ -247,9 +244,8 @@ bool FCbFieldNoneTest::RunTest(const FString& Parameters)
 		TestTrue(TEXT("FCbFieldView(None|Name)::HasName()"), NoneField.HasName());
 		TestFalse(TEXT("!FCbFieldView(None|Name)::HasValue()"), NoneField.HasValue());
 		TestEqual(TEXT("FCbFieldView(None|Name)::GetHash()"), NoneField.GetHash(), FBlake3::HashBuffer(MakeMemoryView(NoneBytes)));
-		TestEqual(TEXT("FCbFieldView(None|Name)::GetView()"), NoneField.GetView(), MakeMemoryView(NoneBytes) + 1);
-		FMemoryView SerializedView;
-		TestFalse(TEXT("FCbFieldView(None|Name)::TryGetSerializedView()"), NoneField.TryGetSerializedView(SerializedView));
+		FMemoryView View;
+		TestFalse(TEXT("FCbFieldView(None|Name)::TryGetView()"), NoneField.TryGetView(View));
 
 		uint8 CopyBytes[sizeof(NoneBytes)];
 		NoneField.CopyTo(MakeMemoryView(CopyBytes));
@@ -266,9 +262,8 @@ bool FCbFieldNoneTest::RunTest(const FString& Parameters)
 		TestTrue(TEXT("FCbFieldView(None|EmptyName)::HasName()"), NoneField.HasName());
 		TestFalse(TEXT("!FCbFieldView(None|EmptyName)::HasValue()"), NoneField.HasValue());
 		TestEqual(TEXT("FCbFieldView(None|EmptyName)::GetHash()"), NoneField.GetHash(), FBlake3::HashBuffer(MakeMemoryView(NoneBytes)));
-		TestEqual(TEXT("FCbFieldView(None|EmptyName)::GetView()"), NoneField.GetView(), MakeMemoryView(NoneBytes) + 1);
-		FMemoryView SerializedView;
-		TestFalse(TEXT("FCbFieldView(None|EmptyName)::TryGetSerializedView()"), NoneField.TryGetSerializedView(SerializedView));
+		FMemoryView View;
+		TestFalse(TEXT("FCbFieldView(None|EmptyName)::TryGetView()"), NoneField.TryGetView(View));
 	}
 
 	return true;
@@ -379,7 +374,6 @@ bool FCbFieldObjectTest::RunTest(const FString& Parameters)
 		TestEqual(TEXT("FCbObjectView::FindViewIgnoreCase()"), Object.FindViewIgnoreCase("b"_ASV).AsInt32(), 2);
 		TestEqual(TEXT("FCbObjectView::operator[]"), Object["B"_ASV].AsInt32(), 2);
 		TestEqual(TEXT("FCbObjectView::operator[]"), Object["b"_ASV].AsInt32(4), 4);
-		TestEqual(TEXT("FCbObjectView::GetView()"), Field.AsObjectView().GetView(), Field.GetView());
 	}
 
 	// Test FCbFieldView(UniformObject, NotEmpty)
@@ -402,9 +396,8 @@ bool FCbFieldObjectTest::RunTest(const FString& Parameters)
 		TestEqual(TEXT("FCbObjectView(Uniform)::FindViewIgnoreCase()"), Object.FindIgnoreCase("b"_ASV).AsInt32(), 2);
 		TestEqual(TEXT("FCbObjectView(Uniform)::operator[]"), Object["B"_ASV].AsInt32(), 2);
 		TestEqual(TEXT("FCbObjectView(Uniform)::operator[]"), Object["b"_ASV].AsInt32(4), 4);
-		TestEqual(TEXT("FCbObjectView(Uniform)::GetView()"), Field.AsObjectView().GetView(), Field.GetView());
 
-		TestTrue(TEXT("FCbObject::AsField()"), Object.GetBuffer() == Object.AsField().AsObject().GetBuffer());
+		TestTrue(TEXT("FCbObject::AsField()"), Object.GetOuterBuffer() == Object.AsField().AsObject().GetOuterBuffer());
 
 		// Equals
 		const uint8 NamedPayload[] = { 1, 'O', 10, IntType, 1, 'A', 1, 1, 'B', 2, 1, 'C', 3 };
@@ -412,7 +405,6 @@ bool FCbFieldObjectTest::RunTest(const FString& Parameters)
 		TestTrue(TEXT("FCbObjectView::Equals()"), Field.AsObjectView().Equals(NamedField.AsObjectView()));
 		TestTrue(TEXT("FCbObjectView::AsFieldView().Equals()"), Field.Equals(Field.AsObjectView().AsFieldView()));
 		TestTrue(TEXT("FCbObjectView::AsFieldView().Equals()"), NamedField.Equals(NamedField.AsObjectView().AsFieldView()));
-		TestEqual(TEXT("FCbObjectView(Name)::GetView()"), NamedField.AsObjectView().GetView(), NamedField.GetView());
 
 		// CopyTo
 		uint8 CopyBytes[sizeof(Payload) + 1];
@@ -421,11 +413,16 @@ bool FCbFieldObjectTest::RunTest(const FString& Parameters)
 		NamedField.AsObjectView().CopyTo(MakeMemoryView(CopyBytes));
 		TestTrue(TEXT("FCbObjectView(NoType, Name)::CopyTo()"), MakeMemoryView(Payload).EqualBytes(MakeMemoryView(CopyBytes) + 1));
 
-		// GetSerializedView
-		FMemoryView SerializedView;
-		TestTrue(TEXT("FCbObjectView(Clone)::TryGetSerializedView()"), Object.TryGetSerializedView(SerializedView) && SerializedView == Object.AsFieldView().GetView());
-		TestFalse(TEXT("FCbObjectView(NoType)::TryGetSerializedView()"), Field.AsObjectView().TryGetSerializedView(SerializedView));
-		TestFalse(TEXT("FCbObjectView(Name)::TryGetSerializedView()"), NamedField.AsObjectView().TryGetSerializedView(SerializedView));
+		// TryGetView
+		FMemoryView View;
+		TestTrue(TEXT("FCbObjectView(Clone)::TryGetView()"), Object.TryGetView(View) && View == Object.GetOuterBuffer().GetView());
+		TestFalse(TEXT("FCbObjectView(NoType)::TryGetView()"), Field.AsObjectView().TryGetView(View));
+		TestFalse(TEXT("FCbObjectView(Name)::TryGetView()"), NamedField.AsObjectView().TryGetView(View));
+
+		// GetBuffer
+		TestTrue(TEXT("FCbObjectView(Clone)::GetBuffer()"), Object.GetBuffer().Flatten().GetView() == Object.GetOuterBuffer().GetView());
+		TestTrue(TEXT("FCbObjectView(NoType)::GetBuffer()"), FCbField::MakeView(Field).AsObject().GetBuffer().Flatten().GetView().EqualBytes(Object.GetOuterBuffer().GetView()));
+		TestTrue(TEXT("FCbObjectView(Name)::GetBuffer()"), FCbField::MakeView(NamedField).AsObject().GetBuffer().Flatten().GetView().EqualBytes(Object.GetOuterBuffer().GetView()));
 	}
 
 	// Test FCbFieldView(None) as Object
@@ -451,7 +448,6 @@ bool FCbFieldObjectTest::RunTest(const FString& Parameters)
 			TestEqual(TEXT("FCbObject::CreateIterator().GetName()"), Field.GetName(), "F"_ASV);
 			TestEqual(TEXT("FCbObject::CreateIterator().AsInt32()"), Field.AsInt32(), 8);
 			TestTrue(TEXT("FCbObject::CreateIterator().IsOwned()"), Field.IsOwned());
-			TestEqual(TEXT("FCbObject::CreateIterator().GetBuffer()"), Field.GetBuffer().GetView(), Field.GetView());
 		}
 		for (FCbFieldIterator It = ObjectClone.CreateIterator(), End; It != End; ++It)
 		{
@@ -463,9 +459,14 @@ bool FCbFieldObjectTest::RunTest(const FString& Parameters)
 		// CopyTo
 		uint8 CopyBytes[6];
 		Object.CopyTo(MakeMemoryView(CopyBytes));
-		TestTrue(TEXT("FCbObjectView(Name)::CopyTo()"), ObjectClone.GetView().EqualBytes(MakeMemoryView(CopyBytes)));
+		TestTrue(TEXT("FCbObjectView(Name)::CopyTo()"), ObjectClone.GetOuterBuffer().GetView().EqualBytes(MakeMemoryView(CopyBytes)));
 		ObjectClone.CopyTo(MakeMemoryView(CopyBytes));
-		TestTrue(TEXT("FCbObjectView()::CopyTo()"), ObjectClone.GetView().EqualBytes(MakeMemoryView(CopyBytes)));
+		TestTrue(TEXT("FCbObjectView()::CopyTo()"), ObjectClone.GetOuterBuffer().GetView().EqualBytes(MakeMemoryView(CopyBytes)));
+
+		// GetBuffer
+		TestTrue(TEXT("FCbField(ObjectWithName)::GetBuffer()"), FCbField(FSharedBuffer::MakeView(MakeMemoryView(Buffer))).GetBuffer().Flatten().GetView().EqualBytes(MakeMemoryView(Buffer)));
+		TestTrue(TEXT("FCbField(ObjectWithNameNoType)::GetBuffer()"),
+			FCbField(FCbFieldView(Buffer + 1, ECbFieldType(ObjectType)), FSharedBuffer::MakeView(MakeMemoryView(Buffer))).GetBuffer().Flatten().GetView().EqualBytes(MakeMemoryView(Buffer)));
 	}
 
 	// Test FCbObjectView as FCbFieldViewIterator
@@ -565,7 +566,6 @@ bool FCbFieldArrayTest::RunTest(const FString& Parameters)
 		TestIntArray(Array, 3, sizeof(Payload));
 		TestIntArray(Field.AsArrayView(), 3, sizeof(Payload));
 		TestTrue(TEXT("FCbArrayView::Equals()"), Array.Equals(Field.AsArrayView()));
-		TestEqual(TEXT("FCbArrayView::GetView()"), Field.AsArrayView().GetView(), Field.GetView());
 	}
 
 	// Test FCbFieldView(UniformArray)
@@ -578,9 +578,8 @@ bool FCbFieldArrayTest::RunTest(const FString& Parameters)
 		TestIntArray(Array, 3, sizeof(Payload));
 		TestIntArray(Field.AsArrayView(), 3, sizeof(Payload));
 		TestTrue(TEXT("FCbArrayView(Uniform)::Equals()"), Array.Equals(Field.AsArrayView()));
-		TestEqual(TEXT("FCbArrayView(Uniform)::GetView()"), Field.AsArrayView().GetView(), Field.GetView());
 
-		TestTrue(TEXT("FCbArray::AsField()"), Array.GetBuffer() == Array.AsField().AsArray().GetBuffer());
+		TestTrue(TEXT("FCbArray::AsField()"), Array.GetOuterBuffer() == Array.AsField().AsArray().GetOuterBuffer());
 
 		// Equals
 		const uint8 NamedPayload[] = { 1, 'A', 5, 3, IntType, 1, 2, 3 };
@@ -588,7 +587,6 @@ bool FCbFieldArrayTest::RunTest(const FString& Parameters)
 		TestTrue(TEXT("FCbArrayView::Equals()"), Field.AsArrayView().Equals(NamedField.AsArrayView()));
 		TestTrue(TEXT("FCbArrayView::AsFieldView().Equals()"), Field.Equals(Field.AsArrayView().AsFieldView()));
 		TestTrue(TEXT("FCbArrayView::AsFieldView().Equals()"), NamedField.Equals(NamedField.AsArrayView().AsFieldView()));
-		TestEqual(TEXT("FCbArrayView(Name)::GetView()"), NamedField.AsArrayView().GetView(), NamedField.GetView());
 
 		// CopyTo
 		uint8 CopyBytes[sizeof(Payload) + 1];
@@ -597,11 +595,16 @@ bool FCbFieldArrayTest::RunTest(const FString& Parameters)
 		NamedField.AsArrayView().CopyTo(MakeMemoryView(CopyBytes));
 		TestTrue(TEXT("FCbArrayView(NoType, Name)::CopyTo()"), MakeMemoryView(Payload).EqualBytes(MakeMemoryView(CopyBytes) + 1));
 
-		// GetSerializedView
-		FMemoryView SerializedView;
-		TestTrue(TEXT("FCbArrayView(Clone)::TryGetSerializedView()"), Array.TryGetSerializedView(SerializedView) && SerializedView == Array.AsFieldView().GetView());
-		TestFalse(TEXT("FCbArrayView(NoType)::TryGetSerializedView()"), Field.AsArrayView().TryGetSerializedView(SerializedView));
-		TestFalse(TEXT("FCbArrayView(Name)::TryGetSerializedView()"), NamedField.AsArrayView().TryGetSerializedView(SerializedView));
+		// TryGetView
+		FMemoryView View;
+		TestTrue(TEXT("FCbArrayView(Clone)::TryGetView()"), Array.TryGetView(View) && View == Array.GetOuterBuffer().GetView());
+		TestFalse(TEXT("FCbArrayView(NoType)::TryGetView()"), Field.AsArrayView().TryGetView(View));
+		TestFalse(TEXT("FCbArrayView(Name)::TryGetView()"), NamedField.AsArrayView().TryGetView(View));
+
+		// GetBuffer
+		TestTrue(TEXT("FCbArrayView(Clone)::GetBuffer()"), Array.GetBuffer().Flatten().GetView() == Array.GetOuterBuffer().GetView());
+		TestTrue(TEXT("FCbArrayView(NoType)::GetBuffer()"), FCbField::MakeView(Field).AsArray().GetBuffer().Flatten().GetView().EqualBytes(Array.GetOuterBuffer().GetView()));
+		TestTrue(TEXT("FCbArrayView(Name)::GetBuffer()"), FCbField::MakeView(NamedField).AsArray().GetBuffer().Flatten().GetView().EqualBytes(Array.GetOuterBuffer().GetView()));
 	}
 
 	// Test FCbFieldView(None) as Array
@@ -626,7 +629,6 @@ bool FCbFieldArrayTest::RunTest(const FString& Parameters)
 			FCbField Field = *It;
 			TestEqual(TEXT("FCbArray::CreateIterator().AsInt32()"), Field.AsInt32(), 8);
 			TestTrue(TEXT("FCbArray::CreateIterator().IsOwned()"), Field.IsOwned());
-			TestEqual(TEXT("FCbArray::CreateIterator().GetBuffer()"), Field.GetBuffer().GetView(), Field.GetView());
 		}
 		for (FCbFieldIterator It = ArrayClone.CreateIterator(), End; It != End; ++It)
 		{
@@ -638,9 +640,14 @@ bool FCbFieldArrayTest::RunTest(const FString& Parameters)
 		// CopyTo
 		uint8 CopyBytes[5];
 		Array.CopyTo(MakeMemoryView(CopyBytes));
-		TestTrue(TEXT("FCbArrayView(Name)::CopyTo()"), ArrayClone.GetView().EqualBytes(MakeMemoryView(CopyBytes)));
+		TestTrue(TEXT("FCbArrayView(Name)::CopyTo()"), ArrayClone.GetOuterBuffer().GetView().EqualBytes(MakeMemoryView(CopyBytes)));
 		ArrayClone.CopyTo(MakeMemoryView(CopyBytes));
-		TestTrue(TEXT("FCbArrayView()::CopyTo()"), ArrayClone.GetView().EqualBytes(MakeMemoryView(CopyBytes)));
+		TestTrue(TEXT("FCbArrayView()::CopyTo()"), ArrayClone.GetOuterBuffer().GetView().EqualBytes(MakeMemoryView(CopyBytes)));
+
+		// GetBuffer
+		TestTrue(TEXT("FCbField(ArrayWithName)::GetBuffer()"), FCbField(FSharedBuffer::MakeView(MakeMemoryView(Buffer))).GetBuffer().Flatten().GetView().EqualBytes(MakeMemoryView(Buffer)));
+		TestTrue(TEXT("FCbField(ArrayWithNameNoType)::GetBuffer()"),
+			FCbField(FCbFieldView(Buffer + 1, ECbFieldType(ArrayType)), FSharedBuffer::MakeView(MakeMemoryView(Buffer))).GetBuffer().Flatten().GetView().EqualBytes(MakeMemoryView(Buffer)));
 	}
 
 	// Test FCbArrayView as FCbFieldViewIterator
@@ -1429,7 +1436,7 @@ bool FCbFieldIterateAttachmentsTest::RunTest(const FString& Parameters)
 		Fields = Writer.Save();
 	}
 
-	TestEqual(TEXT("FCbFieldView::IterateAttachments Validate"), ValidateCompactBinaryRange(Fields.GetBuffer(), ECbValidateMode::All), ECbValidateError::None);
+	TestEqual(TEXT("FCbFieldView::IterateAttachments Validate"), ValidateCompactBinaryRange(Fields.GetOuterBuffer(), ECbValidateMode::All), ECbValidateError::None);
 
 	uint32 AttachmentIndex = 0;
 	Fields.IterateRangeAttachments([this, &AttachmentIndex, &MakeTestHash](FCbFieldView Field)
@@ -1681,11 +1688,9 @@ bool FCbFieldBufferIteratorTest::RunTest(const FString& Parameters)
 		TestEqual(TEXT("FCbFieldViewIterator::GetRangeHash()"), FieldViewIt.GetRangeHash(), FBlake3::HashBuffer(View));
 		TestEqual(TEXT("FCbFieldIterator::GetRangeHash()"), FieldIt.GetRangeHash(), FBlake3::HashBuffer(View));
 
-		TestEqual(TEXT("FCbFieldViewIterator::GetRangeView()"), FieldViewIt.GetRangeView(), MakeMemoryView(Payload));
-		TestEqual(TEXT("FCbFieldIterator::GetRangeView()"), FieldIt.GetRangeView(), MakeMemoryView(Payload));
-		FMemoryView SerializedView;
-		TestTrue(TEXT("FCbFieldViewIterator::TryGetSerializedRangeView()"), FieldViewIt.TryGetSerializedRangeView(SerializedView) && SerializedView == MakeMemoryView(Payload));
-		TestTrue(TEXT("FCbFieldIterator::TryGetSerializedRangeView()"), FieldIt.TryGetSerializedRangeView(SerializedView) && SerializedView == MakeMemoryView(Payload));
+		FMemoryView RangeView;
+		TestTrue(TEXT("FCbFieldViewIterator::TryGetRangeView()"), FieldViewIt.TryGetRangeView(RangeView) && RangeView == MakeMemoryView(Payload));
+		TestTrue(TEXT("FCbFieldIterator::TryGetRangeView()"), FieldIt.TryGetRangeView(RangeView) && RangeView == MakeMemoryView(Payload));
 
 		TestEqual(TEXT("FCbFieldIterator::CloneRange(EmptyViewIt)"), GetCount(FCbFieldIterator::CloneRange(FCbFieldViewIterator())), 0);
 		TestEqual(TEXT("FCbFieldIterator::CloneRange(EmptyIt)"), GetCount(FCbFieldIterator::CloneRange(FCbFieldIterator())), 0);
@@ -1718,13 +1723,11 @@ bool FCbFieldBufferIteratorTest::RunTest(const FString& Parameters)
 		const uint8 UniformPayload[] = { 0, 1, 2, 3 };
 		const FCbFieldViewIterator UniformFieldViewIt = FCbFieldViewIterator::MakeRange(MakeMemoryView(UniformPayload), ECbFieldType::IntegerPositive);
 		TestEqual(TEXT("FCbFieldViewIterator::MakeRange(Uniform).GetRangeHash()"), UniformFieldViewIt.GetRangeHash(), FieldViewIt.GetRangeHash());
-		TestEqual(TEXT("FCbFieldViewIterator::MakeRange(Uniform).GetRangeView()"), UniformFieldViewIt.GetRangeView(), MakeMemoryView(UniformPayload));
-		TestFalse(TEXT("FCbFieldViewIterator::MakeRange(Uniform).TryGetSerializedRangeView()"), UniformFieldViewIt.TryGetSerializedRangeView(SerializedView));
+		TestFalse(TEXT("FCbFieldViewIterator::MakeRange(Uniform).TryGetRangeView()"), UniformFieldViewIt.TryGetRangeView(RangeView));
 		const FSharedBuffer UniformView = FSharedBuffer::MakeView(MakeMemoryView(UniformPayload));
 		const FCbFieldIterator UniformFieldIt = FCbFieldIterator::MakeRange(UniformView, ECbFieldType::IntegerPositive);
 		TestEqual(TEXT("FCbFieldIterator::MakeRange(Uniform).GetRangeHash()"), UniformFieldIt.GetRangeHash(), FieldViewIt.GetRangeHash());
-		TestEqual(TEXT("FCbFieldIterator::MakeRange(Uniform).GetRangeView()"), UniformFieldIt.GetRangeView(), MakeMemoryView(UniformPayload));
-		TestFalse(TEXT("FCbFieldIterator::MakeRange(Uniform).TryGetSerializedRangeView()"), UniformFieldIt.TryGetSerializedRangeView(SerializedView));
+		TestFalse(TEXT("FCbFieldIterator::MakeRange(Uniform).TryGetRangeView()"), UniformFieldIt.TryGetRangeView(RangeView));
 
 		// Equals
 		TestTrue(TEXT("FCbFieldViewIterator::Equals(Self)"), FieldViewIt.Equals(FieldViewIt));
@@ -1756,16 +1759,11 @@ bool FCbFieldBufferIteratorTest::RunTest(const FString& Parameters)
 		// MakeRangeOwned
 		FCbFieldIterator OwnedFromView = UniformFieldIt;
 		OwnedFromView.MakeRangeOwned();
-		TestTrue(TEXT("FCbFieldIterator::MakeRangeOwned(View)"), OwnedFromView.GetRangeView().EqualBytes(MakeMemoryView(Payload)));
+		TestTrue(TEXT("FCbFieldIterator::MakeRangeOwned(View)"), OwnedFromView.TryGetRangeView(RangeView) && RangeView.EqualBytes(MakeMemoryView(Payload)));
 		FCbFieldIterator OwnedFromOwned = OwnedFromView;
 		OwnedFromOwned.MakeRangeOwned();
 		TestEqual(TEXT("FCbFieldIterator::MakeRangeOwned(Owned)"), OwnedFromOwned, OwnedFromView);
 
-		// GetRangeBuffer
-		TestEqual(TEXT("FCbFieldIterator::GetRangeBuffer()"), UniformFieldIt.GetRangeBuffer(), UniformView);
-		const FCbFieldIterator UniformShortFieldIt = FCbFieldIterator::MakeRangeView(FCbFieldViewIterator::MakeRange(MakeMemoryView(UniformPayload).LeftChop(1), ECbFieldType::IntegerPositive), UniformView);
-		TestNotEqual(TEXT("FCbFieldIterator::GetRangeBuffer(Short)"), UniformShortFieldIt.GetRangeBuffer(), UniformView);
-
 		// These lines are expected to assert when uncommented.
 		//const FSharedBuffer ShortView = FSharedBuffer::MakeView(MakeMemoryView(Payload).LeftChop(2));
 		//TestEqual(TEXT("FCbFieldIterator::MakeRangeView(FieldIt, InvalidBufferL)"), GetCount(FCbFieldIterator::MakeRangeView(FCbFieldViewIterator::MakeRange(*View), ShortView)), 4);
diff --git a/Engine/Source/Runtime/Core/Private/Tests/Serialization/CompactBinaryWriterTest.cpp b/Engine/Source/Runtime/Core/Private/Tests/Serialization/CompactBinaryWriterTest.cpp
--- a/Engine/Source/Runtime/Core/Private/Tests/Serialization/CompactBinaryWriterTest.cpp
+++ b/Engine/Source/Runtime/Core/Private/Tests/Serialization/CompactBinaryWriterTest.cpp
@@ -29,7 +29,7 @@ bool FCbWriterObjectTest::RunTest(const FString& Parameters)
 		Writer.BeginObject();
 		Writer.EndObject();
 		FCbField Field = Writer.Save();
-		if (TestEqual(TEXT("FCbWriter(Object, Empty) Validate"), ValidateCompactBinary(Field.GetBuffer(), ECbValidateMode::All), ECbValidateError::None))
+		if (TestEqual(TEXT("FCbWriter(Object, Empty) Validate"), ValidateCompactBinary(Field.GetOuterBuffer(), ECbValidateMode::All), ECbValidateError::None))
 		{
 			TestTrue(TEXT("FCbWriter(Object, Empty).IsObject()"), Field.IsObject());
 			TestFalse(TEXT("FCbWriter(Object, Empty).AsObjectView()"), Field.AsObjectView().CreateViewIterator().HasValue());
@@ -43,7 +43,7 @@ bool FCbWriterObjectTest::RunTest(const FString& Parameters)
 		Writer.BeginObject();
 		Writer.EndObject();
 		FCbField Field = Writer.Save();
-		if (TestEqual(TEXT("FCbWriter(Object, Empty, Name) Validate"), ValidateCompactBinary(Field.GetBuffer(), ECbValidateMode::All), ECbValidateError::None))
+		if (TestEqual(TEXT("FCbWriter(Object, Empty, Name) Validate"), ValidateCompactBinary(Field.GetOuterBuffer(), ECbValidateMode::All), ECbValidateError::None))
 		{
 			TestTrue(TEXT("FCbWriter(Object, Empty, Name).IsObject()"), Field.IsObject());
 			TestFalse(TEXT("FCbWriter(Object, Empty, Name).AsObjectView()"), Field.AsObjectView().CreateViewIterator().HasValue());
@@ -58,7 +58,7 @@ bool FCbWriterObjectTest::RunTest(const FString& Parameters)
 		Writer.SetName("Float"_ASV).AddFloat(0.0f);
 		Writer.EndObject();
 		FCbField Field = Writer.Save();
-		if (TestEqual(TEXT("FCbWriter(Object, Basic) Validate"), ValidateCompactBinary(Field.GetBuffer(), ECbValidateMode::All), ECbValidateError::None))
+		if (TestEqual(TEXT("FCbWriter(Object, Basic) Validate"), ValidateCompactBinary(Field.GetOuterBuffer(), ECbValidateMode::All), ECbValidateError::None))
 		{
 			TestTrue(TEXT("FCbWriter(Object, Basic).IsObject()"), Field.IsObject());
 			FCbObjectView Object = Field.AsObjectView();
@@ -75,7 +75,7 @@ bool FCbWriterObjectTest::RunTest(const FString& Parameters)
 		Writer.SetName("Field2"_ASV).AddInteger(1);
 		Writer.EndObject();
 		FCbField Field = Writer.Save();
-		if (TestEqual(TEXT("FCbWriter(Object, Uniform) Validate"), ValidateCompactBinary(Field.GetBuffer(), ECbValidateMode::All), ECbValidateError::None))
+		if (TestEqual(TEXT("FCbWriter(Object, Uniform) Validate"), ValidateCompactBinary(Field.GetOuterBuffer(), ECbValidateMode::All), ECbValidateError::None))
 		{
 			TestTrue(TEXT("FCbWriter(Object, Uniform).IsObject()"), Field.IsObject());
 			FCbObjectView Object = Field.AsObjectView();
@@ -100,7 +100,7 @@ bool FCbWriterArrayTest::RunTest(const FString& Parameters)
 		Writer.BeginArray();
 		Writer.EndArray();
 		FCbField Field = Writer.Save();
-		if (TestEqual(TEXT("FCbWriter(Array, Empty) Validate"), ValidateCompactBinary(Field.GetBuffer(), ECbValidateMode::All), ECbValidateError::None))
+		if (TestEqual(TEXT("FCbWriter(Array, Empty) Validate"), ValidateCompactBinary(Field.GetOuterBuffer(), ECbValidateMode::All), ECbValidateError::None))
 		{
 			TestTrue(TEXT("FCbWriter(Array, Empty).IsArray()"), Field.IsArray());
 			TestEqual(TEXT("FCbWriter(Array, Empty).AsArrayView()"), Field.AsArrayView().Num(), uint64(0));
@@ -114,7 +114,7 @@ bool FCbWriterArrayTest::RunTest(const FString& Parameters)
 		Writer.BeginArray();
 		Writer.EndArray();
 		FCbField Field = Writer.Save();
-		if (TestEqual(TEXT("FCbWriter(Array, Empty, Name) Validate"), ValidateCompactBinary(Field.GetBuffer(), ECbValidateMode::All), ECbValidateError::None))
+		if (TestEqual(TEXT("FCbWriter(Array, Empty, Name) Validate"), ValidateCompactBinary(Field.GetOuterBuffer(), ECbValidateMode::All), ECbValidateError::None))
 		{
 			TestTrue(TEXT("FCbWriter(Array, Empty, Name).IsArray()"), Field.IsArray());
 			TestEqual(TEXT("FCbWriter(Array, Empty, Name).AsArrayView()"), Field.AsArrayView().Num(), uint64(0));
@@ -129,7 +129,7 @@ bool FCbWriterArrayTest::RunTest(const FString& Parameters)
 		Writer.AddFloat(0.0f);
 		Writer.EndArray();
 		FCbField Field = Writer.Save();
-		if (TestEqual(TEXT("FCbWriter(Array, Basic) Validate"), ValidateCompactBinary(Field.GetBuffer(), ECbValidateMode::All), ECbValidateError::None))
+		if (TestEqual(TEXT("FCbWriter(Array, Basic) Validate"), ValidateCompactBinary(Field.GetOuterBuffer(), ECbValidateMode::All), ECbValidateError::None))
 		{
 			TestTrue(TEXT("FCbWriter(Array, Basic).IsArray()"), Field.IsArray());
 			FCbFieldViewIterator Iterator = Field.AsArrayView().CreateViewIterator();
@@ -149,7 +149,7 @@ bool FCbWriterArrayTest::RunTest(const FString& Parameters)
 		Writer.AddInteger(1);
 		Writer.EndArray();
 		FCbField Field = Writer.Save();
-		if (TestEqual(TEXT("FCbWriter(Array, Uniform) Validate"), ValidateCompactBinary(Field.GetBuffer(), ECbValidateMode::All), ECbValidateError::None))
+		if (TestEqual(TEXT("FCbWriter(Array, Uniform) Validate"), ValidateCompactBinary(Field.GetOuterBuffer(), ECbValidateMode::All), ECbValidateError::None))
 		{
 			TestTrue(TEXT("FCbWriter(Array, Uniform).IsArray()"), Field.IsArray());
 			FCbFieldViewIterator Iterator = Field.AsArrayView().CreateViewIterator();
@@ -176,7 +176,7 @@ bool FCbWriterNullTest::RunTest(const FString& Parameters)
 		Writer.Reset();
 		Writer.AddNull();
 		FCbField Field = Writer.Save();
-		if (TestEqual(TEXT("FCbWriter(Null) Validate"), ValidateCompactBinary(Field.GetBuffer(), ECbValidateMode::All), ECbValidateError::None))
+		if (TestEqual(TEXT("FCbWriter(Null) Validate"), ValidateCompactBinary(Field.GetOuterBuffer(), ECbValidateMode::All), ECbValidateError::None))
 		{
 			TestFalse(TEXT("FCbWriter(Null).HasName()"), Field.HasName());
 			TestTrue(TEXT("FCbWriter(Null).IsNull()"), Field.IsNull());
@@ -189,7 +189,7 @@ bool FCbWriterNullTest::RunTest(const FString& Parameters)
 		Writer.SetName("Null"_ASV);
 		Writer.AddNull();
 		FCbField Field = Writer.Save();
-		if (TestEqual(TEXT("FCbWriter(Null, Name) Validate"), ValidateCompactBinary(Field.GetBuffer(), ECbValidateMode::All), ECbValidateError::None))
+		if (TestEqual(TEXT("FCbWriter(Null, Name) Validate"), ValidateCompactBinary(Field.GetOuterBuffer(), ECbValidateMode::All), ECbValidateError::None))
 		{
 			TestEqual(TEXT("FCbWriter(Null, Name).GetName()"), Field.GetName(), "Null"_ASV);
 			TestTrue(TEXT("FCbWriter(Null, Name).HasName()"), Field.HasName());
@@ -214,7 +214,7 @@ bool FCbWriterNullTest::RunTest(const FString& Parameters)
 		Writer.EndObject();
 
 		FCbFieldIterator Fields = Writer.Save();
-		TestEqual(TEXT("FCbWriter(Null, Uniform) Validate"), ValidateCompactBinaryRange(Fields.GetRangeBuffer(), ECbValidateMode::All), ECbValidateError::None);
+		TestEqual(TEXT("FCbWriter(Null, Uniform) Validate"), ValidateCompactBinaryRange(Fields.GetOuterBuffer(), ECbValidateMode::All), ECbValidateError::None);
 	}
 
 	// Test Null with Save(Buffer)
@@ -266,7 +266,7 @@ bool FCbWriterBinaryTest::RunTest(const FString& Parameters)
 		Writer.Reset();
 		Writer.AddBinary(nullptr, 0);
 		FCbField Field = Writer.Save();
-		if (TestEqual(TEXT("FCbWriter(Binary, Empty) Validate"), ValidateCompactBinary(Field.GetBuffer(), ECbValidateMode::All), ECbValidateError::None))
+		if (TestEqual(TEXT("FCbWriter(Binary, Empty) Validate"), ValidateCompactBinary(Field.GetOuterBuffer(), ECbValidateMode::All), ECbValidateError::None))
 		{
 			TestFalse(TEXT("FCbWriter(Binary, Empty).HasName()"), Field.HasName());
 			TestTrue(TEXT("FCbWriter(Binary, Empty).IsBinary()"), Field.IsBinary());
@@ -281,7 +281,7 @@ bool FCbWriterBinaryTest::RunTest(const FString& Parameters)
 		Writer.SetName("Binary"_ASV);
 		Writer.AddBinary(BinaryValue, sizeof(BinaryValue));
 		FCbField Field = Writer.Save();
-		if (TestEqual(TEXT("FCbWriter(Binary, Array) Validate"), ValidateCompactBinary(Field.GetBuffer(), ECbValidateMode::All), ECbValidateError::None))
+		if (TestEqual(TEXT("FCbWriter(Binary, Array) Validate"), ValidateCompactBinary(Field.GetOuterBuffer(), ECbValidateMode::All), ECbValidateError::None))
 		{
 			TestEqual(TEXT("FCbWriter(Binary, Array).GetName()"), Field.GetName(), "Binary"_ASV);
 			TestTrue(TEXT("FCbWriter(Binary, Array).HasName()"), Field.HasName());
@@ -306,7 +306,7 @@ bool FCbWriterStringTest::RunTest(const FString& Parameters)
 		Writer.AddString(FAnsiStringView());
 		Writer.AddString(FWideStringView());
 		FCbFieldIterator Fields = Writer.Save();
-		if (TestEqual(TEXT("FCbWriter(String, Empty) Validate"), ValidateCompactBinaryRange(Fields.GetRangeBuffer(), ECbValidateMode::All), ECbValidateError::None))
+		if (TestEqual(TEXT("FCbWriter(String, Empty) Validate"), ValidateCompactBinaryRange(Fields.GetOuterBuffer(), ECbValidateMode::All), ECbValidateError::None))
 		{
 			for (FCbFieldView Field : Fields)
 			{
@@ -323,7 +323,7 @@ bool FCbWriterStringTest::RunTest(const FString& Parameters)
 		Writer.SetName("String"_ASV).AddString("Value"_ASV);
 		Writer.SetName("String"_ASV).AddString(TEXT("Value"_SV));
 		FCbFieldIterator Fields = Writer.Save();
-		if (TestEqual(TEXT("FCbWriter(String, Basic) Validate"), ValidateCompactBinaryRange(Fields.GetRangeBuffer(), ECbValidateMode::All), ECbValidateError::None))
+		if (TestEqual(TEXT("FCbWriter(String, Basic) Validate"), ValidateCompactBinaryRange(Fields.GetOuterBuffer(), ECbValidateMode::All), ECbValidateError::None))
 		{
 			for (FCbFieldView Field : Fields)
 			{
@@ -347,7 +347,7 @@ bool FCbWriterStringTest::RunTest(const FString& Parameters)
 		Writer.AddString(Dots);
 		Writer.AddString(FString::ChrN(DotCount, TEXT('.')));
 		FCbFieldIterator Fields = Writer.Save();
-		if (TestEqual(TEXT("FCbWriter(String, Long) Validate"), ValidateCompactBinaryRange(Fields.GetRangeBuffer(), ECbValidateMode::All), ECbValidateError::None))
+		if (TestEqual(TEXT("FCbWriter(String, Long) Validate"), ValidateCompactBinaryRange(Fields.GetOuterBuffer(), ECbValidateMode::All), ECbValidateError::None))
 		{
 			for (FCbFieldView Field : Fields)
 			{
@@ -363,7 +363,7 @@ bool FCbWriterStringTest::RunTest(const FString& Parameters)
 		Writer.AddString("\xf0\x9f\x98\x80"_ASV);
 		Writer.AddString(FWideStringView(Value, UE_ARRAY_COUNT(Value)));
 		FCbFieldIterator Fields = Writer.Save();
-		if (TestEqual(TEXT("FCbWriter(String, Unicode) Validate"), ValidateCompactBinaryRange(Fields.GetRangeBuffer(), ECbValidateMode::All), ECbValidateError::None))
+		if (TestEqual(TEXT("FCbWriter(String, Unicode) Validate"), ValidateCompactBinaryRange(Fields.GetOuterBuffer(), ECbValidateMode::All), ECbValidateError::None))
 		{
 			for (FCbFieldView Field : Fields)
 			{
@@ -387,7 +387,7 @@ bool FCbWriterIntegerTest::RunTest(const FString& Parameters)
 		Writer.Reset();
 		Writer.AddInteger(Value);
 		FCbField Field = Writer.Save();
-		if (TestEqual(TEXT("FCbWriter(Integer, Int32) Validate"), ValidateCompactBinary(Field.GetBuffer(), ECbValidateMode::All), ECbValidateError::None))
+		if (TestEqual(TEXT("FCbWriter(Integer, Int32) Validate"), ValidateCompactBinary(Field.GetOuterBuffer(), ECbValidateMode::All), ECbValidateError::None))
 		{
 			TestEqual(TEXT("FCbWriter(Integer, Int32) Value"), Field.AsInt32(), Value);
 			TestFalse(TEXT("FCbWriter(Integer, Int32) Error"), Field.HasError());
@@ -399,7 +399,7 @@ bool FCbWriterIntegerTest::RunTest(const FString& Parameters)
 		Writer.Reset();
 		Writer.AddInteger(Value);
 		FCbField Field = Writer.Save();
-		if (TestEqual(TEXT("FCbWriter(Integer, UInt32) Validate"), ValidateCompactBinary(Field.GetBuffer(), ECbValidateMode::All), ECbValidateError::None))
+		if (TestEqual(TEXT("FCbWriter(Integer, UInt32) Validate"), ValidateCompactBinary(Field.GetOuterBuffer(), ECbValidateMode::All), ECbValidateError::None))
 		{
 			TestEqual(TEXT("FCbWriter(Integer, UInt32) Value"), Field.AsUInt32(), Value);
 			TestFalse(TEXT("FCbWriter(Integer, UInt32) Error"), Field.HasError());
@@ -411,7 +411,7 @@ bool FCbWriterIntegerTest::RunTest(const FString& Parameters)
 		Writer.Reset();
 		Writer.AddInteger(Value);
 		FCbField Field = Writer.Save();
-		if (TestEqual(TEXT("FCbWriter(Integer, Int64) Validate"), ValidateCompactBinary(Field.GetBuffer(), ECbValidateMode::All), ECbValidateError::None))
+		if (TestEqual(TEXT("FCbWriter(Integer, Int64) Validate"), ValidateCompactBinary(Field.GetOuterBuffer(), ECbValidateMode::All), ECbValidateError::None))
 		{
 			TestEqual(TEXT("FCbWriter(Integer, Int64) Value"), Field.AsInt64(), Value);
 			TestFalse(TEXT("FCbWriter(Integer, Int64) Error"), Field.HasError());
@@ -423,7 +423,7 @@ bool FCbWriterIntegerTest::RunTest(const FString& Parameters)
 		Writer.Reset();
 		Writer.AddInteger(Value);
 		FCbField Field = Writer.Save();
-		if (TestEqual(TEXT("FCbWriter(Integer, UInt64) Validate"), ValidateCompactBinary(Field.GetBuffer(), ECbValidateMode::All), ECbValidateError::None))
+		if (TestEqual(TEXT("FCbWriter(Integer, UInt64) Validate"), ValidateCompactBinary(Field.GetOuterBuffer(), ECbValidateMode::All), ECbValidateError::None))
 		{
 			TestEqual(TEXT("FCbWriter(Integer, UInt64) Value"), Field.AsUInt64(), Value);
 			TestFalse(TEXT("FCbWriter(Integer, UInt64) Error"), Field.HasError());
@@ -482,7 +482,7 @@ bool FCbWriterFloatTest::RunTest(const FString& Parameters)
 			Writer.AddFloat(Value);
 		}
 		FCbFieldIterator Fields = Writer.Save();
-		if (TestEqual(TEXT("FCbWriter(Float, Single) Validate"), ValidateCompactBinaryRange(Fields.GetRangeBuffer(), ECbValidateMode::All), ECbValidateError::None))
+		if (TestEqual(TEXT("FCbWriter(Float, Single) Validate"), ValidateCompactBinaryRange(Fields.GetOuterBuffer(), ECbValidateMode::All), ECbValidateError::None))
 		{
 			const float* CheckValue = Values;
 			for (FCbFieldView Field : Fields)
@@ -502,7 +502,7 @@ bool FCbWriterFloatTest::RunTest(const FString& Parameters)
 			Writer.AddFloat(Value);
 		}
 		FCbFieldIterator Fields = Writer.Save();
-		if (TestEqual(TEXT("FCbWriter(Float, Double) Validate"), ValidateCompactBinaryRange(Fields.GetRangeBuffer(), ECbValidateMode::All), ECbValidateError::None))
+		if (TestEqual(TEXT("FCbWriter(Float, Double) Validate"), ValidateCompactBinaryRange(Fields.GetOuterBuffer(), ECbValidateMode::All), ECbValidateError::None))
 		{
 			const double* CheckValue = Values;
 			for (FCbFieldView Field : Fields)
@@ -529,7 +529,7 @@ bool FCbWriterBoolTest::RunTest(const FString& Parameters)
 		Writer.AddBool(false);
 
 		FCbFieldIterator Fields = Writer.Save();
-		if (TestEqual(TEXT("FCbWriter(Bool) Validate"), ValidateCompactBinaryRange(Fields.GetRangeBuffer(), ECbValidateMode::All), ECbValidateError::None))
+		if (TestEqual(TEXT("FCbWriter(Bool) Validate"), ValidateCompactBinaryRange(Fields.GetOuterBuffer(), ECbValidateMode::All), ECbValidateError::None))
 		{
 			TestTrue(TEXT("FCbWriter(Bool).AsBool()"), Fields.AsBool());
 			TestFalse(TEXT("FCbWriter(Bool) Error"), Fields.HasError());
@@ -556,7 +556,7 @@ bool FCbWriterBoolTest::RunTest(const FString& Parameters)
 		Writer.EndObject();
 
 		FCbFieldIterator Fields = Writer.Save();
-		TestEqual(TEXT("FCbWriter(Bool, Uniform) Validate"), ValidateCompactBinaryRange(Fields.GetRangeBuffer(), ECbValidateMode::All), ECbValidateError::None);
+		TestEqual(TEXT("FCbWriter(Bool, Uniform) Validate"), ValidateCompactBinaryRange(Fields.GetOuterBuffer(), ECbValidateMode::All), ECbValidateError::None);
 	}
 
 	return true;
@@ -579,7 +579,7 @@ bool FCbWriterObjectAttachmentTest::RunTest(const FString& Parameters)
 	}
 
 	FCbFieldIterator Fields = Writer.Save();
-	if (TestEqual(TEXT("FCbWriter(ObjectAttachment) Validate"), ValidateCompactBinaryRange(Fields.GetRangeBuffer(), ECbValidateMode::All), ECbValidateError::None))
+	if (TestEqual(TEXT("FCbWriter(ObjectAttachment) Validate"), ValidateCompactBinaryRange(Fields.GetOuterBuffer(), ECbValidateMode::All), ECbValidateError::None))
 	{
 		const FIoHash* CheckValue = Values;
 		for (FCbFieldView Field : Fields)
@@ -609,7 +609,7 @@ bool FCbWriterBinaryAttachmentTest::RunTest(const FString& Parameters)
 	}
 
 	FCbFieldIterator Fields = Writer.Save();
-	if (TestEqual(TEXT("FCbWriter(BinaryAttachment) Validate"), ValidateCompactBinaryRange(Fields.GetRangeBuffer(), ECbValidateMode::All), ECbValidateError::None))
+	if (TestEqual(TEXT("FCbWriter(BinaryAttachment) Validate"), ValidateCompactBinaryRange(Fields.GetOuterBuffer(), ECbValidateMode::All), ECbValidateError::None))
 	{
 		const FIoHash* CheckValue = Values;
 		for (FCbFieldView Field : Fields)
@@ -639,7 +639,7 @@ bool FCbWriterHashTest::RunTest(const FString& Parameters)
 	}
 
 	FCbFieldIterator Fields = Writer.Save();
-	if (TestEqual(TEXT("FCbWriter(Hash) Validate"), ValidateCompactBinaryRange(Fields.GetRangeBuffer(), ECbValidateMode::All), ECbValidateError::None))
+	if (TestEqual(TEXT("FCbWriter(Hash) Validate"), ValidateCompactBinaryRange(Fields.GetOuterBuffer(), ECbValidateMode::All), ECbValidateError::None))
 	{
 		const FIoHash* CheckValue = Values;
 		for (FCbFieldView Field : Fields)
@@ -666,7 +666,7 @@ bool FCbWriterUuidTest::RunTest(const FString& Parameters)
 	}
 
 	FCbFieldIterator Fields = Writer.Save();
-	if (TestEqual(TEXT("FCbWriter(Uuid) Validate"), ValidateCompactBinaryRange(Fields.GetRangeBuffer(), ECbValidateMode::All), ECbValidateError::None))
+	if (TestEqual(TEXT("FCbWriter(Uuid) Validate"), ValidateCompactBinaryRange(Fields.GetOuterBuffer(), ECbValidateMode::All), ECbValidateError::None))
 	{
 		const FGuid* CheckValue = Values;
 		for (FCbFieldView Field : Fields)
@@ -693,7 +693,7 @@ bool FCbWriterDateTimeTest::RunTest(const FString& Parameters)
 	}
 
 	FCbFieldIterator Fields = Writer.Save();
-	if (TestEqual(TEXT("FCbWriter(DateTime) Validate"), ValidateCompactBinaryRange(Fields.GetRangeBuffer(), ECbValidateMode::All), ECbValidateError::None))
+	if (TestEqual(TEXT("FCbWriter(DateTime) Validate"), ValidateCompactBinaryRange(Fields.GetOuterBuffer(), ECbValidateMode::All), ECbValidateError::None))
 	{
 		const FDateTime* CheckValue = Values;
 		for (FCbFieldView Field : Fields)
@@ -720,7 +720,7 @@ bool FCbWriterTimeSpanTest::RunTest(const FString& Parameters)
 	}
 
 	FCbFieldIterator Fields = Writer.Save();
-	if (TestEqual(TEXT("FCbWriter(TimeSpan) Validate"), ValidateCompactBinaryRange(Fields.GetRangeBuffer(), ECbValidateMode::All), ECbValidateError::None))
+	if (TestEqual(TEXT("FCbWriter(TimeSpan) Validate"), ValidateCompactBinaryRange(Fields.GetOuterBuffer(), ECbValidateMode::All), ECbValidateError::None))
 	{
 		const FTimespan* CheckValue = Values;
 		for (FCbFieldView Field : Fields)
@@ -747,7 +747,7 @@ bool FCbWriterObjectIdTest::RunTest(const FString& Parameters)
 	}
 
 	FCbFieldIterator Fields = Writer.Save();
-	if (TestEqual(TEXT("FCbWriter(ObjectId) Validate"), ValidateCompactBinaryRange(Fields.GetRangeBuffer(), ECbValidateMode::All), ECbValidateError::None))
+	if (TestEqual(TEXT("FCbWriter(ObjectId) Validate"), ValidateCompactBinaryRange(Fields.GetOuterBuffer(), ECbValidateMode::All), ECbValidateError::None))
 	{
 		const FCbObjectId* CheckValue = Values;
 		for (FCbFieldView Field : Fields)
@@ -784,7 +784,7 @@ bool FCbWriterCustomByIdTest::RunTest(const FString& Parameters)
 	}
 
 	FCbFieldIterator Fields = Writer.Save();
-	if (TestEqual(TEXT("FCbWriter(CustomById) Validate"), ValidateCompactBinaryRange(Fields.GetRangeBuffer(), ECbValidateMode::All), ECbValidateError::None))
+	if (TestEqual(TEXT("FCbWriter(CustomById) Validate"), ValidateCompactBinaryRange(Fields.GetOuterBuffer(), ECbValidateMode::All), ECbValidateError::None))
 	{
 		const FCustomValue* CheckValue = Values;
 		for (FCbFieldView Field : Fields)
@@ -822,7 +822,7 @@ bool FCbWriterCustomByNameTest::RunTest(const FString& Parameters)
 	}
 
 	FCbFieldIterator Fields = Writer.Save();
-	if (TestEqual(TEXT("FCbWriter(CustomByName) Validate"), ValidateCompactBinaryRange(Fields.GetRangeBuffer(), ECbValidateMode::All), ECbValidateError::None))
+	if (TestEqual(TEXT("FCbWriter(CustomByName) Validate"), ValidateCompactBinaryRange(Fields.GetOuterBuffer(), ECbValidateMode::All), ECbValidateError::None))
 	{
 		const FCustomValue* CheckValue = Values;
 		for (FCbFieldView Field : Fields)
@@ -1013,7 +1013,7 @@ bool FCbWriterComplexTest::RunTest(const FString& Parameters)
 	}
 
 	TestEqual(TEXT("FCbWriter(Complex).Save()->Validate"),
-		ValidateCompactBinary(Object.GetBuffer(), ECbValidateMode::All), ECbValidateError::None);
+		ValidateCompactBinary(Object.GetOuterBuffer(), ECbValidateMode::All), ECbValidateError::None);
 
 	TestEqual(TEXT("FCbWriter(Complex).Save(Ar)->Validate"),
 		ValidateCompactBinary(MakeMemoryView(Archive), ECbValidateMode::All), ECbValidateError::None);
@@ -1095,7 +1095,7 @@ bool FCbWriterStreamTest::RunTest(const FString& Parameters)
 		Object = Writer.Save().AsObject();
 	}
 
-	TestEqual(TEXT("FCbWriter(Stream) Validate"), ValidateCompactBinary(Object.GetBuffer(), ECbValidateMode::All), ECbValidateError::None);
+	TestEqual(TEXT("FCbWriter(Stream) Validate"), ValidateCompactBinary(Object.GetOuterBuffer(), ECbValidateMode::All), ECbValidateError::None);
 
 	return true;
 }
diff --git a/Engine/Source/Runtime/Core/Public/Serialization/CompactBinary.h b/Engine/Source/Runtime/Core/Public/Serialization/CompactBinary.h
--- a/Engine/Source/Runtime/Core/Public/Serialization/CompactBinary.h
+++ b/Engine/Source/Runtime/Core/Public/Serialization/CompactBinary.h
@@ -6,6 +6,7 @@
 #include "Containers/StringFwd.h"
 #include "Containers/StringView.h"
 #include "IO/IoHash.h"
+#include "Memory/MemoryFwd.h"
 #include "Memory/MemoryView.h"
 #include "Memory/SharedBuffer.h"
 #include "Misc/EnumClassFlags.h"
@@ -383,23 +384,18 @@ public:
 	/** Invoke the visitor for every attachment in the field range. */
 	CORE_API void IterateRangeAttachments(FCbFieldVisitor Visitor) const;
 
-	/** Create a view of every field in the range. */
-	inline FMemoryView GetRangeView() const
-	{
-		return MakeMemoryView(FieldType::GetView().GetData(), FieldsEnd);
-	}
-
 	/**
 	 * Try to get a view of every field in the range as they would be serialized.
 	 *
-	 * A serialized view is not available if the underlying fields have an externally-provided type.
-	 * Access the serialized form of such ranges using FCbFieldIterator::CloneRange.
+	 * A view is available if each field contains its type. Access the equivalent for other field
+	 * ranges through FCbFieldIterator::CloneRange or CopyRangeTo.
 	 */
-	inline bool TryGetSerializedRangeView(FMemoryView& OutView) const
+	inline bool TryGetRangeView(FMemoryView& OutView) const
 	{
-		if (FCbFieldType::HasFieldType(FieldType::GetTypeWithFlags()))
+		FMemoryView View;
+		if (FieldType::TryGetView(View))
 		{
-			OutView = GetRangeView();
+			OutView = MakeMemoryView(View.GetData(), FieldsEnd);
 			return true;
 		}
 		return false;
@@ -731,16 +727,13 @@ public:
 	/** Invoke the visitor for every attachment in the field. */
 	CORE_API void IterateAttachments(FCbFieldVisitor Visitor) const;
 
-	/** Returns a view of the field, including the type and name when present. */
-	CORE_API FMemoryView GetView() const;
-
 	/**
 	 * Try to get a view of the field as it would be serialized, such as by CopyTo.
 	 *
-	 * A serialized view is not available if the field has an externally-provided type.
-	 * Access the serialized form of such fields using CopyTo or FCbField::Clone.
+	 * A view is available if the field contains its type. Access the equivalent for other fields
+	 * through FCbField::GetBuffer, FCbField::Clone, or CopyTo.
 	 */
-	inline bool TryGetSerializedView(FMemoryView& OutView) const
+	inline bool TryGetView(FMemoryView& OutView) const
 	{
 		if (FCbFieldType::HasFieldType(TypeWithFlags))
 		{
@@ -758,6 +751,9 @@ public:
 	constexpr inline FCbIteratorSentinel end() const { return FCbIteratorSentinel(); }
 
 protected:
+	/** Returns a view of the field, including the type and name when present. */
+	CORE_API FMemoryView GetView() const;
+
 	/** Returns a view of the name and value payload, which excludes the type. */
 	CORE_API FMemoryView GetViewNoType() const;
 
@@ -935,18 +931,15 @@ public:
 	/** Invoke the visitor for every attachment in the array. */
 	inline void IterateAttachments(FCbFieldVisitor Visitor) const { CreateViewIterator().IterateRangeAttachments(Visitor); }
 
-	/** Returns a view of the array, including the type and name when present. */
-	using FCbFieldView::GetView;
-
 	/**
 	 * Try to get a view of the array as it would be serialized, such as by CopyTo.
 	 *
-	 * A serialized view is not available if the underlying field has an externally-provided type or
-	 * a name. Access the serialized form of such arrays using CopyTo or FCbArray::Clone.
+	 * A view is available if the array contains its type and has no name. Access the equivalent
+	 * for other arrays through FCbArray::GetBuffer, FCbArray::Clone, or CopyTo.
 	 */
-	inline bool TryGetSerializedView(FMemoryView& OutView) const
+	inline bool TryGetView(FMemoryView& OutView) const
 	{
-		return !FCbFieldView::HasName() && FCbFieldView::TryGetSerializedView(OutView);
+		return !FCbFieldView::HasName() && FCbFieldView::TryGetView(OutView);
 	}
 
 	/** @see FCbFieldView::CreateViewIterator */
@@ -1035,18 +1028,15 @@ public:
 	/** Invoke the visitor for every attachment in the object. */
 	inline void IterateAttachments(FCbFieldVisitor Visitor) const { CreateViewIterator().IterateRangeAttachments(Visitor); }
 
-	/** Returns a view of the object, including the type and name when present. */
-	using FCbFieldView::GetView;
-
 	/**
 	 * Try to get a view of the object as it would be serialized, such as by CopyTo.
 	 *
-	 * A serialized view is not available if the underlying field has an externally-provided type or
-	 * a name. Access the serialized form of such objects using CopyTo or FCbObject::Clone.
+	 * A view is available if the object contains its type and has no name. Access the equivalent
+	 * for other objects through FCbObject::GetBuffer, FCbObject::Clone, or CopyTo.
 	 */
-	inline bool TryGetSerializedView(FMemoryView& OutView) const
+	inline bool TryGetView(FMemoryView& OutView) const
 	{
-		return !FCbFieldView::HasName() && FCbFieldView::TryGetSerializedView(OutView);
+		return !FCbFieldView::HasName() && FCbFieldView::TryGetView(OutView);
 	}
 
 	/** @see FCbFieldView::CreateViewIterator */
@@ -1128,15 +1118,12 @@ public:
 		}
 	}
 
-	/** Returns a buffer that exactly contains this value. */
-	inline FSharedBuffer GetBuffer() const
-	{
-		const FMemoryView View = ViewType::GetView();
-		const FSharedBuffer& OuterBuffer = GetOuterBuffer();
-		return View == OuterBuffer.GetView() ? OuterBuffer : FSharedBuffer::MakeView(View, OuterBuffer);
-	}
-
-	/** Returns the outer buffer (if any) that contains this value. */
+	/**
+	 * Returns the outer buffer (if any) that contains this value.
+	 *
+	 * The outer buffer might contain other data before and/or after this value. Use GetBuffer to
+	 * request a buffer that exactly contains this value, or TryGetView for a contiguous view.
+	 */
 	inline const FSharedBuffer& GetOuterBuffer() const & { return Buffer; }
 	inline FSharedBuffer GetOuterBuffer() && { return MoveTemp(Buffer); }
 
@@ -1222,6 +1209,9 @@ public:
 	/** Access the field as binary. Returns the provided default on error. */
 	inline FSharedBuffer AsBinary(const FSharedBuffer& Default = FSharedBuffer()) &;
 	inline FSharedBuffer AsBinary(const FSharedBuffer& Default = FSharedBuffer()) &&;
+
+	/** Returns a buffer that contains the field as it would be serialized by CopyTo. */
+	CORE_API FCompositeBuffer GetBuffer() const;
 };
 
 /**
@@ -1281,9 +1271,6 @@ public:
 		}
 	}
 
-	/** Returns a buffer that exactly contains the field range. */
-	CORE_API FSharedBuffer GetRangeBuffer() const;
-
 private:
 	using TCbFieldIterator::TCbFieldIterator;
 };
@@ -1314,16 +1301,11 @@ public:
 	using TCbBuffer::TCbBuffer;
 
 	/** Access the array as an array field. */
-	inline FCbField AsField() const &
-	{
-		return FCbField(FCbArrayView::AsFieldView(), *this);
-	}
+	inline FCbField AsField() const & { return FCbField(FCbArrayView::AsFieldView(), *this); }
+	inline FCbField AsField() && { return FCbField(FCbArrayView::AsFieldView(), MoveTemp(*this)); }
 
-	/** Access the array as an array field. */
-	inline FCbField AsField() &&
-	{
-		return FCbField(FCbArrayView::AsFieldView(), MoveTemp(*this));
-	}
+	/** Returns a buffer that contains the array as it would be serialized by CopyTo. */
+	CORE_API FCompositeBuffer GetBuffer() const;
 };
 
 ///////////////////////////////////////////////////////////////////////////////////////////////////////////////////////
@@ -1363,16 +1345,11 @@ public:
 	inline FCbField operator[](FAnsiStringView Name) const { return Find(Name); }
 
 	/** Access the object as an object field. */
-	inline FCbField AsField() const &
-	{
-		return FCbField(FCbObjectView::AsFieldView(), *this);
-	}
+	inline FCbField AsField() const & { return FCbField(FCbObjectView::AsFieldView(), *this); }
+	inline FCbField AsField() && { return FCbField(FCbObjectView::AsFieldView(), MoveTemp(*this)); }
 
-	/** Access the object as an object field. */
-	inline FCbField AsField() &&
-	{
-		return FCbField(FCbObjectView::AsFieldView(), MoveTemp(*this));
-	}
+	/** Returns a buffer that contains the object as it would be serialized by CopyTo. */
+	CORE_API FCompositeBuffer GetBuffer() const;
 };
 
 ///////////////////////////////////////////////////////////////////////////////////////////////////////////////////////
