Copying the error message can overflow the allocated buffer when reporting an ensure.
Replaced usage of misleading FCString::Strcpy() with FCString::Strncpy()

diff --git a/Engine/Source/Runtime/Core/Private/Windows/WindowsPlatformCrashContext.cpp b/Engine/Source/Runtime/Core/Private/Windows/WindowsPlatformCrashContext.cpp
--- a/Engine/Source/Runtime/Core/Private/Windows/WindowsPlatformCrashContext.cpp
+++ b/Engine/Source/Runtime/Core/Private/Windows/WindowsPlatformCrashContext.cpp
@@ -735,7 +735,7 @@ int32 ReportCrashForMonitor(
 	SharedContext->UserSettings.bImplicitSend = bImplicitSend;
 
 	SharedContext->SessionContext.bIsExitRequested = IsEngineExitRequested();
-	FCString::Strcpy(SharedContext->ErrorMessage, CR_MAX_ERROR_MESSAGE_CHARS-1, ErrorMessage);
+	FCString::Strncpy(SharedContext->ErrorMessage, ErrorMessage, CR_MAX_ERROR_MESSAGE_CHARS);
 
 	// Setup all the thread ids and names using snapshot dbghelp. Since it's not possible to 
 	// query thread names from an external process.
@@ -751,10 +751,10 @@ int32 ReportCrashForMonitor(
 			SharedContext->ThreadIds[ThreadIdx] = ThreadEntry.th32ThreadID;
 			const FString& TmThreadName = FThreadManager::GetThreadName(ThreadEntry.th32ThreadID);
 			const TCHAR* ThreadName = TmThreadName.IsEmpty() ? TEXT("Unknown") : *TmThreadName;
-			FCString::Strcpy(
+			FCString::Strncpy(
 				&SharedContext->ThreadNames[ThreadIdx*CR_MAX_THREAD_NAME_CHARS],
-				CR_MAX_THREAD_NAME_CHARS - 1,
-				ThreadName
+				ThreadName,
+				CR_MAX_THREAD_NAME_CHARS
 			);
 			ThreadIdx++;
 		};
@@ -797,7 +797,7 @@ int32 ReportCrashForMonitor(
 	FString CrashDirectoryAbsolute;
 	if (FGenericCrashContext::CreateCrashReportDirectory(SharedContext->SessionContext.CrashGUIDRoot, ReportCallCount++, CrashDirectoryAbsolute))
 	{
-		FCString::Strcpy(SharedContext->CrashFilesDirectory, *CrashDirectoryAbsolute);
+		FCString::Strncpy(SharedContext->CrashFilesDirectory, *CrashDirectoryAbsolute, CR_MAX_DIRECTORY_CHARS);
 		// Copy the log file to output
 		FGenericCrashContext::DumpLog(CrashDirectoryAbsolute);
 
