U5 - fix function after recent changes.

diff --git a/Engine/Source/Runtime/Core/Private/HAL/ThreadingBase.cpp b/Engine/Source/Runtime/Core/Private/HAL/ThreadingBase.cpp
--- a/Engine/Source/Runtime/Core/Private/HAL/ThreadingBase.cpp
+++ b/Engine/Source/Runtime/Core/Private/HAL/ThreadingBase.cpp
@@ -665,19 +665,22 @@ FScopedEvent::FScopedEvent()
 
 bool FScopedEvent::IsReady()
 {
-	if ( Event->Wait(1) )
+	if ( Event && Event->Wait(1) )
 	{
 		TLazySingleton<TEventPool<EEventMode::AutoReset>>::Get().ReturnRawEvent(Event);
 		Event = nullptr;
 		return true;
 	}
-	return false;
+	return Event == nullptr;
 }
 
 FScopedEvent::~FScopedEvent()
 {
-	Event->Wait();
-	TLazySingleton<TEventPool<EEventMode::AutoReset>>::Get().ReturnRawEvent(Event);
+	if(Event)
+	{
+		Event->Wait();
+		TLazySingleton<TEventPool<EEventMode::AutoReset>>::Get().ReturnRawEvent(Event);
+	}
 }
 
 
