U5 - Always wake workers during busy waiting as we might tap out any moment after finishing a task dependency.

diff --git a/Engine/Source/Runtime/Core/Private/Async/Fundamental/Scheduler.cpp b/Engine/Source/Runtime/Core/Private/Async/Fundamental/Scheduler.cpp
--- a/Engine/Source/Runtime/Core/Private/Async/Fundamental/Scheduler.cpp
+++ b/Engine/Source/Runtime/Core/Private/Async/Fundamental/Scheduler.cpp
@@ -15,7 +15,7 @@ namespace LowLevelTasks
 	thread_local FTask* FSchedulerTls::ActiveTask = nullptr;
 	thread_local FSchedulerTls* FSchedulerTls::ActiveScheduler = nullptr;
 	thread_local FSchedulerTls::EWorkerType FSchedulerTls::WorkerType = FSchedulerTls::EWorkerType::None;
-	thread_local uint32 FScheduler::BusyWaitingDepth = 0;
+	thread_local uint32 FSchedulerTls::BusyWaitingDepth = 0;
 
 	FScheduler FScheduler::Singleton;
 
@@ -189,6 +189,11 @@ namespace LowLevelTasks
 		return WorkerType != FSchedulerTls::EWorkerType::None && ActiveScheduler == this;
 	}
 
+	bool FSchedulerTls::IsBusyWaiting()
+	{
+		return BusyWaitingDepth != 0;
+	}
+
 	void FTask::PropagateUserData()
 	{
 		const FTask* ActiveTask = FSchedulerTls::GetActiveTask();
@@ -299,8 +304,8 @@ namespace LowLevelTasks
 		TRACE_CPUPROFILER_EVENT_SCOPE(FScheduler::BusyWaitInternal);
 		FTaskTagScope WorkerScope(ETaskTag::EWorkerThread);
 
-		++BusyWaitingDepth;
-		ON_SCOPE_EXIT{ --BusyWaitingDepth; };
+		++FSchedulerTls::BusyWaitingDepth;
+		ON_SCOPE_EXIT{ --FSchedulerTls::BusyWaitingDepth; };
 
 		checkSlow(LocalQueue != nullptr);
 		check(ActiveWorkers.load(std::memory_order_relaxed));
diff --git a/Engine/Source/Runtime/Core/Private/Async/TaskGraph.cpp b/Engine/Source/Runtime/Core/Private/Async/TaskGraph.cpp
--- a/Engine/Source/Runtime/Core/Private/Async/TaskGraph.cpp
+++ b/Engine/Source/Runtime/Core/Private/Async/TaskGraph.cpp
@@ -1967,6 +1967,7 @@ private:
 				Task->Execute(NewTasks, InThreadToExecuteOn, false);
 			});
 	
+			bWakeUpWorker |= LowLevelTasks::FSchedulerTls::IsBusyWaiting();
 			verifySlow(LowLevelTasks::TryLaunch(Task->TaskHandle, bWakeUpWorker ? LowLevelTasks::EQueuePreference::GlobalQueuePreference : LowLevelTasks::EQueuePreference::LocalQueuePreference, bWakeUpWorker));
 			return;
 		}
diff --git a/Engine/Source/Runtime/Core/Public/Async/Fundamental/Scheduler.h b/Engine/Source/Runtime/Core/Public/Async/Fundamental/Scheduler.h
--- a/Engine/Source/Runtime/Core/Public/Async/Fundamental/Scheduler.h
+++ b/Engine/Source/Runtime/Core/Public/Async/Fundamental/Scheduler.h
@@ -100,6 +100,8 @@ namespace LowLevelTasks
 		static thread_local FTask* ActiveTask;
 		static thread_local FSchedulerTls* ActiveScheduler;
 		static thread_local EWorkerType WorkerType;
+		// number of busy-waiting calls in the call-stack
+		static thread_local uint32 BusyWaitingDepth;
 
 	public:
 		CORE_API bool IsWorkerThread() const;
@@ -107,6 +109,9 @@ namespace LowLevelTasks
 		//get the active task if any
 		CORE_API static const FTask* GetActiveTask();
 
+		// returns true if the current thread execution is in the context of busy-waiting
+		CORE_API static bool IsBusyWaiting();
+
 	protected:
 		inline static bool PermitBackgroundWork()
 		{
@@ -131,19 +136,8 @@ namespace LowLevelTasks
 		static thread_local FTask* ActiveTask;
 		static thread_local FScheduler* ActiveScheduler;
 
-		enum class EWorkerType
-		{
-			None,
-			Background,
-			Foreground,
-		};
-
-		static thread_local EWorkerType WorkerType;
 		static CORE_API FScheduler Singleton;
 
-		// number of busy-waiting calls in the call-stack
-		static thread_local uint32 BusyWaitingDepth;
-
 		// using 16 bytes here because it fits the vtable and one additional pointer
 		using FConditional = TTaskDelegate<16, bool>;
 
@@ -179,9 +173,6 @@ namespace LowLevelTasks
 		template<typename TaskType>
 		inline void BusyWait(const TArrayView<const TaskType>& Tasks, bool ForceAllowBackgroundWork = false);
 
-		// returns true if the current thread execution is in the context of busy-waiting
-		inline static bool IsBusyWaiting();
-
 		//number of instantiated workers
 		inline uint32 GetNumWorkers() const;
 
@@ -313,11 +304,6 @@ namespace LowLevelTasks
 		}
 	}
 
-	inline bool FScheduler::IsBusyWaiting()
-	{
-		return BusyWaitingDepth != 0;
-	}
-
 	inline void FScheduler::TrySleeping(FSleepEvent* WorkerEvent, FQueueRegistry::FOutOfWork& OutOfWork, bool& Drowsing, bool bBackgroundWorker)
 	{
 		ESleepState DrowsingState1 = ESleepState::Drowsing;
diff --git a/Engine/Source/Runtime/Core/Public/Misc/QueuedThreadPoolWrapper.h b/Engine/Source/Runtime/Core/Public/Misc/QueuedThreadPoolWrapper.h
--- a/Engine/Source/Runtime/Core/Public/Misc/QueuedThreadPoolWrapper.h
+++ b/Engine/Source/Runtime/Core/Public/Misc/QueuedThreadPoolWrapper.h
@@ -306,6 +306,7 @@ private:
 			FQueuedWorkInternalData* QueuedWork = Dequeue();
 			if (QueuedWork)
 			{
+				bWakeUpWorker |= LowLevelTasks::FSchedulerTls::IsBusyWaiting();
 				verifySlow(Scheduler->TryLaunch(QueuedWork->Task, bWakeUpWorker ? LowLevelTasks::EQueuePreference::GlobalQueuePreference : LowLevelTasks::EQueuePreference::LocalQueuePreference, bWakeUpWorker));
 				TaskCount.fetch_add(1, std::memory_order_acquire);
 				bWakeUpWorker = true;
