CompactBinary: Added operator[] to field types and AsView() to buffer types

#rnx

diff --git a/Engine/Source/Runtime/Core/Private/Serialization/CompactBinary.cpp b/Engine/Source/Runtime/Core/Private/Serialization/CompactBinary.cpp
--- a/Engine/Source/Runtime/Core/Private/Serialization/CompactBinary.cpp
+++ b/Engine/Source/Runtime/Core/Private/Serialization/CompactBinary.cpp
@@ -621,6 +621,18 @@ FMemoryView FCbFieldView::GetViewNoType() const
 	return MakeMemoryView(static_cast<const uint8*>(Payload) - NameSize, NameSize + PayloadSize);
 }
 
+FCbFieldView FCbFieldView::operator[](FAnsiStringView Name) const
+{
+	switch (GetType())
+	{
+	case ECbFieldType::Object:
+	case ECbFieldType::UniformObject:
+		return FCbObjectView::FromFieldNoCheck(*this).FindView(Name);
+	default:
+		return FCbFieldView();
+	}
+}
+
 FCbFieldViewIterator FCbFieldView::CreateViewIterator() const
 {
 	const ECbFieldType LocalTypeWithFlags = TypeWithFlags;
diff --git a/Engine/Source/Runtime/Core/Public/Serialization/CompactBinary.h b/Engine/Source/Runtime/Core/Public/Serialization/CompactBinary.h
--- a/Engine/Source/Runtime/Core/Public/Serialization/CompactBinary.h
+++ b/Engine/Source/Runtime/Core/Public/Serialization/CompactBinary.h
@@ -65,6 +65,8 @@
 class FArchive;
 class FBlake3;
 class FCbArrayView;
+class FCbField;
+class FCbFieldIterator;
 class FCbFieldView;
 class FCbFieldViewIterator;
 class FCbObjectId;
@@ -74,6 +76,9 @@ struct FGuid;
 struct FTimespan;
 template <typename FuncType> class TFunctionRef;
 
+/** A reference to a function that is used to allocate buffers for compact binary data. */
+using FCbBufferAllocator = TFunctionRef<FUniqueBuffer (uint64 Size)>;
+
 ///////////////////////////////////////////////////////////////////////////////////////////////////////////////////////
 
 /**
@@ -757,6 +762,9 @@ public:
 		return false;
 	}
 
+	/** Find a field of an object by case-sensitive name comparison, otherwise a field with no value. */
+	CORE_API FCbFieldView operator[](FAnsiStringView Name) const;
+
 	/** Create an iterator for the fields of an array or object, otherwise an empty iterator. */
 	CORE_API FCbFieldViewIterator CreateViewIterator() const;
 
@@ -1065,13 +1073,6 @@ private:
 
 ///////////////////////////////////////////////////////////////////////////////////////////////////////////////////////
 
-class FCbFieldIterator;
-
-/** A reference to a function that is used to allocate buffers for compact binary data. */
-using FCbBufferAllocator = TFunctionRef<FUniqueBuffer (uint64 Size)>;
-
-///////////////////////////////////////////////////////////////////////////////////////////////////////////////////////
-
 /** A wrapper that holds a reference to the buffer that contains its compact binary value. */
 template <typename ViewType>
 class TCbBuffer : public ViewType
@@ -1132,6 +1133,9 @@ public:
 		}
 	}
 
+	/** Returns the value as a view. */
+	inline const ViewType& AsView() const { return *this; }
+
 	/**
 	 * Returns the outer buffer (if any) that contains this value.
 	 *
@@ -1141,6 +1145,9 @@ public:
 	inline const FSharedBuffer& GetOuterBuffer() const & { return Buffer; }
 	inline FSharedBuffer GetOuterBuffer() && { return MoveTemp(Buffer); }
 
+	/** Find a field of an object by case-sensitive name comparison, otherwise a field with no value. */
+	inline FCbField operator[](FAnsiStringView Name) const;
+
 	/** Create an iterator for the fields of an array or object, otherwise an empty iterator. */
 	inline FCbFieldIterator CreateIterator() const;
 
@@ -1228,6 +1235,12 @@ public:
 	CORE_API FCompositeBuffer GetBuffer() const;
 };
 
+template <typename ViewType>
+inline FCbField TCbBuffer<ViewType>::operator[](FAnsiStringView Name) const
+{
+	return FCbField::MakeView(ViewType::operator[](Name), GetOuterBuffer());
+}
+
 /**
  * Iterator for FCbField.
  *
