#ue5 - LWC: AVX/AVX2 optimizations for VectorShuffle(), VectorSwizzle(), and VectorReplicate(). Instruction latencies on AVX or AVX2 dropped from ~6 to ~1-2 in many cases, and are hand-tuned between whether AVX1 or AVX2 are available.

Instruction latencies and uop counts from documentation manuals and community testing (referenced in the source), plus comparing assembly output on godbolt led to many shuffle and swizzle template overrides and specializations.

Also optimized access to the XY and ZW portions of the AVX m256 union to sub-optimal MSVC compiler output that turns them in to memory ops rather than simple casts or 1-latency extracts.

Added more shuffle/swizzle validation tests to UnrealMathTests to stress specific in-lane and cross-lane permutes.

diff --git a/Engine/Source/Runtime/Core/Private/Tests/Math/UnrealMathTest.cpp b/Engine/Source/Runtime/Core/Private/Tests/Math/UnrealMathTest.cpp
--- a/Engine/Source/Runtime/Core/Private/Tests/Math/UnrealMathTest.cpp
+++ b/Engine/Source/Runtime/Core/Private/Tests/Math/UnrealMathTest.cpp
@@ -970,64 +970,81 @@ FORCENOINLINE void TestVectorSwizzle()
 	// Unfortunately it can't be done in a loop because it uses a #define and compile-time constants for the VectorSwizzle() 'function'.
 	// Many of these were selected to also stress the specializations in certain implementations.
 
+	SwizzleTest(V0, 0, 1, 2, 3); // Identity
+
 	SwizzleTest(V0, 0, 0, 0, 0); // Replicate 0
 	SwizzleTest(V0, 1, 1, 1, 1); // Replicate 1
 	SwizzleTest(V0, 2, 2, 2, 2); // Replicate 2
 	SwizzleTest(V0, 3, 3, 3, 3); // Replicate 3
 
-	SwizzleTest(V0, 0, 1, 2, 3); // Identity
 	SwizzleTest(V0, 1, 2, 3, 0); // Rotate << 1
 	SwizzleTest(V0, 2, 3, 0, 1); // Rotate << 2
 	SwizzleTest(V0, 3, 0, 1, 2); // Rotate << 3
 	
-	SwizzleTest(V0, 0, 2, 0, 2);
-
-	SwizzleTest(V0, 1, 0, 1, 0);
-	SwizzleTest(V0, 1, 2, 0, 1);
-
-	SwizzleTest(V0, 0, 3, 1, 2);
-	SwizzleTest(V0, 0, 3, 2, 1);
-	SwizzleTest(V0, 1, 0, 3, 2);
-	SwizzleTest(V0, 1, 0, 2, 3);
-	SwizzleTest(V0, 2, 1, 3, 1);
-	SwizzleTest(V0, 2, 1, 1, 3);
-	SwizzleTest(V0, 3, 2, 0, 1);
-	SwizzleTest(V0, 3, 2, 1, 0);
-
-	SwizzleTest(V0, 3, 0, 3, 0);
-	SwizzleTest(V0, 0, 3, 0, 3);
-
-	SwizzleTest(V0, 2, 2, 0, 1);
-	SwizzleTest(V0, 3, 3, 1, 0);
-	SwizzleTest(V0, 2, 2, 0, 2);
-	SwizzleTest(V0, 3, 3, 1, 3);
-
-	// Common specializations
-	SwizzleTest(V0, 0, 1, 2, 2);
-	SwizzleTest(V0, 0, 1, 3, 3);
-	SwizzleTest(V0, 0, 1, 3, 2);
-
-	SwizzleTest(V0, 0, 0, 2, 3);
-	SwizzleTest(V0, 0, 0, 2, 2);
-	SwizzleTest(V0, 0, 0, 3, 3);
-	SwizzleTest(V0, 0, 0, 3, 2);
-
-	SwizzleTest(V0, 1, 0, 2, 3);
-	SwizzleTest(V0, 1, 0, 2, 2);
-	SwizzleTest(V0, 1, 0, 3, 3);
-	SwizzleTest(V0, 1, 0, 3, 2);
-
-	SwizzleTest(V0, 1, 1, 2, 3);
-	SwizzleTest(V0, 1, 1, 2, 2);
-	SwizzleTest(V0, 1, 1, 3, 3);
-	SwizzleTest(V0, 1, 1, 3, 2);
-
-	SwizzleTest(V0, 0, 1, 0, 1);
-	SwizzleTest(V0, 2, 3, 2, 3);
+	// Lane swaps
+	SwizzleTest(V0, 0, 1, 2, 3); // lane 0, lane 1
+	SwizzleTest(V0, 2, 3, 0, 1); // lane 1, lane 0
+	SwizzleTest(V0, 0, 1, 0, 1); // lane 0, lane 0
+	SwizzleTest(V0, 2, 3, 2, 3); // lane 1, lane 1
+
+	// Lane swaps with two in-lane permutes
+	SwizzleTest(V0, 1, 0, 3, 2); // lane 0, lane 1
+	SwizzleTest(V0, 3, 2, 1, 0); // lane 1, lane 0
+	SwizzleTest(V0, 1, 0, 1, 0); // lane 0, lane 0
+	SwizzleTest(V0, 3, 2, 3, 2); // lane 1, lane 1
+
+	// Lane swaps with one in-lane permute (first lane)
+	SwizzleTest(V0, 1, 0, 2, 3); // lane 0, lane 1
+	SwizzleTest(V0, 3, 2, 0, 1); // lane 1, lane 0
+	SwizzleTest(V0, 1, 0, 0, 1); // lane 0, lane 0
+	SwizzleTest(V0, 3, 2, 2, 3); // lane 1, lane 1
+
+	// Lane swaps with one in-lane permute (second lane)
+	SwizzleTest(V0, 0, 1, 3, 2); // lane 0, lane 1
+	SwizzleTest(V0, 2, 3, 1, 0); // lane 1, lane 0
+	SwizzleTest(V0, 0, 1, 1, 0); // lane 0, lane 0
+	SwizzleTest(V0, 2, 3, 3, 2); // lane 1, lane 1
+
+	// Lane swaps with a cross-lane permute
+	SwizzleTest(V0, 0, 1, 0, 3); // lane 0, lane X
+	SwizzleTest(V0, 2, 3, 0, 3); // lane 1, lane X
+	SwizzleTest(V0, 0, 1, 2, 1); // lane 0, lane X
+	SwizzleTest(V0, 2, 3, 1, 2); // lane 1, lane X
+
+	SwizzleTest(V0, 3, 0, 3, 2); // lane X, lane 1
+	SwizzleTest(V0, 2, 0, 1, 0); // lane X, lane 0
+	SwizzleTest(V0, 1, 2, 1, 0); // lane X, lane 0
+	SwizzleTest(V0, 2, 1, 3, 2); // lane X, lane 1
+
+	// In-lane permute with a cross-lane permute
+	SwizzleTest(V0, 1, 0, 3, 0); // lane 0, lane X
+	SwizzleTest(V0, 3, 2, 0, 3); // lane 1, lane X
+	SwizzleTest(V0, 1, 0, 2, 1); // lane 0, lane X
+	SwizzleTest(V0, 3, 2, 1, 2); // lane 1, lane X
+
+	SwizzleTest(V0, 3, 0, 3, 2); // lane X, lane 1
+	SwizzleTest(V0, 0, 3, 1, 0); // lane X, lane 0
+	SwizzleTest(V0, 1, 2, 1, 0); // lane X, lane 0
+	SwizzleTest(V0, 2, 1, 3, 2); // lane X, lane 1
+
+	// Two cross-lane permutes
+	SwizzleTest(V0, 3, 0, 0, 3); // lane X, lane X
+	SwizzleTest(V0, 1, 2, 2, 1); // lane X, lane X
+	SwizzleTest(V0, 3, 0, 2, 1); // lane X, lane X
+	SwizzleTest(V0, 1, 2, 0, 3); // lane X, lane X
+	SwizzleTest(V0, 0, 2, 0, 2); // lane X, lane X
+	SwizzleTest(V0, 2, 0, 0, 2); // lane X, lane X
+
+	// Common specializations or uses
 	SwizzleTest(V0, 0, 0, 1, 1);
-	SwizzleTest(V0, 2, 2, 3, 3);
 	SwizzleTest(V0, 0, 0, 2, 2);
+	SwizzleTest(V0, 0, 1, 2, 2);
+	SwizzleTest(V0, 0, 2, 2, 3);
 	SwizzleTest(V0, 1, 1, 3, 3);
+	SwizzleTest(V0, 1, 3, 1, 3);
+	SwizzleTest(V0, 1, 0, 3, 0);
+	SwizzleTest(V0, 2, 2, 3, 3);
+	SwizzleTest(V0, 2, 0, 3, 0);
 
 #undef SwizzleTest
 }
@@ -1049,37 +1066,82 @@ FORCENOINLINE void TestVectorShuffle()
 	// This is not an exhaustive list because it would be 4*4*4*4 = 256 entries, but it tries to test a lot of common permutations.
 	// Unfortunately it can't be done in a loop because it uses a #define and compile-time constants for the VectorShuffle() 'function'.
 	// Many of these were selected to also stress the specializations in certain implementations.
-	ShuffleTest(V0, V1, 0, 0, 0, 0);
-	ShuffleTest(V0, V1, 1, 1, 1, 1);
-	ShuffleTest(V0, V1, 2, 2, 2, 2);
-	ShuffleTest(V0, V1, 3, 3, 3, 3);
-
-	ShuffleTest(V0, V1, 0, 1, 2, 3);
-	ShuffleTest(V0, V1, 1, 2, 3, 0);
-	ShuffleTest(V0, V1, 2, 3, 0, 1);
-	ShuffleTest(V0, V1, 3, 0, 1, 2);
-
-	ShuffleTest(V0, V1, 3, 2, 1, 0);
-
-	ShuffleTest(V0, V1, 0, 1, 0, 1);
-	ShuffleTest(V0, V1, 0, 2, 0, 2);
-	ShuffleTest(V0, V1, 0, 3, 0, 3);
-
-	ShuffleTest(V0, V1, 1, 0, 1, 0);
-	ShuffleTest(V0, V1, 2, 0, 2, 0);
-	ShuffleTest(V0, V1, 3, 0, 3, 0);
-
+	
+	ShuffleTest(V0, V1, 0, 1, 2, 3); // Identity
+
+	ShuffleTest(V0, V1, 0, 0, 0, 0); // Replicate 0
+	ShuffleTest(V0, V1, 1, 1, 1, 1); // Replicate 1
+	ShuffleTest(V0, V1, 2, 2, 2, 2); // Replicate 2
+	ShuffleTest(V0, V1, 3, 3, 3, 3); // Replicate 3
+
+	ShuffleTest(V0, V1, 1, 2, 3, 0); // Rotate << 1
+	ShuffleTest(V0, V1, 2, 3, 0, 1); // Rotate << 2
+	ShuffleTest(V0, V1, 3, 0, 1, 2); // Rotate << 3
+
+	// Lane swaps
+	ShuffleTest(V0, V1, 0, 1, 2, 3); // lane 0, lane 1
+	ShuffleTest(V0, V1, 2, 3, 0, 1); // lane 1, lane 0
+	ShuffleTest(V0, V1, 0, 1, 0, 1); // lane 0, lane 0
+	ShuffleTest(V0, V1, 2, 3, 2, 3); // lane 1, lane 1
+
+	// Lane swaps with two in-lane permutes
+	ShuffleTest(V0, V1, 1, 0, 3, 2); // lane 0, lane 1
+	ShuffleTest(V0, V1, 3, 2, 1, 0); // lane 1, lane 0
+	ShuffleTest(V0, V1, 1, 0, 1, 0); // lane 0, lane 0
+	ShuffleTest(V0, V1, 3, 2, 3, 2); // lane 1, lane 1
+
+	// Lane swaps with one in-lane permute (first lane)
+	ShuffleTest(V0, V1, 1, 0, 2, 3); // lane 0, lane 1
+	ShuffleTest(V0, V1, 3, 2, 0, 1); // lane 1, lane 0
+	ShuffleTest(V0, V1, 1, 0, 0, 1); // lane 0, lane 0
+	ShuffleTest(V0, V1, 3, 2, 2, 3); // lane 1, lane 1
+
+	// Lane swaps with one in-lane permute (second lane)
+	ShuffleTest(V0, V1, 0, 1, 3, 2); // lane 0, lane 1
+	ShuffleTest(V0, V1, 2, 3, 1, 0); // lane 1, lane 0
+	ShuffleTest(V0, V1, 0, 1, 1, 0); // lane 0, lane 0
+	ShuffleTest(V0, V1, 2, 3, 3, 2); // lane 1, lane 1
+
+	// Lane swaps with a cross-lane permute
+	ShuffleTest(V0, V1, 0, 1, 0, 3); // lane 0, lane X
+	ShuffleTest(V0, V1, 2, 3, 0, 3); // lane 1, lane X
+	ShuffleTest(V0, V1, 0, 1, 2, 1); // lane 0, lane X
+	ShuffleTest(V0, V1, 2, 3, 1, 2); // lane 1, lane X
+
+	ShuffleTest(V0, V1, 3, 0, 3, 2); // lane X, lane 1
+	ShuffleTest(V0, V1, 2, 0, 1, 0); // lane X, lane 0
+	ShuffleTest(V0, V1, 1, 2, 1, 0); // lane X, lane 0
+	ShuffleTest(V0, V1, 2, 1, 3, 2); // lane X, lane 1
+
+	// In-lane permute with a cross-lane permute
+	ShuffleTest(V0, V1, 1, 0, 3, 0); // lane 0, lane X
+	ShuffleTest(V0, V1, 3, 2, 0, 3); // lane 1, lane X
+	ShuffleTest(V0, V1, 1, 0, 2, 1); // lane 0, lane X
+	ShuffleTest(V0, V1, 3, 2, 1, 2); // lane 1, lane X
+
+	ShuffleTest(V0, V1, 3, 0, 3, 2); // lane X, lane 1
+	ShuffleTest(V0, V1, 0, 3, 1, 0); // lane X, lane 0
+	ShuffleTest(V0, V1, 1, 2, 1, 0); // lane X, lane 0
+	ShuffleTest(V0, V1, 2, 1, 3, 2); // lane X, lane 1
+
+	// Two cross-lane permutes
+	ShuffleTest(V0, V1, 3, 0, 0, 3); // lane X, lane X
+	ShuffleTest(V0, V1, 1, 2, 2, 1); // lane X, lane X
+	ShuffleTest(V0, V1, 3, 0, 2, 1); // lane X, lane X
+	ShuffleTest(V0, V1, 1, 2, 0, 3); // lane X, lane X
+	ShuffleTest(V0, V1, 0, 2, 0, 2); // lane X, lane X
+	ShuffleTest(V0, V1, 2, 0, 0, 2); // lane X, lane X
+
+	// Common specializations or uses
 	ShuffleTest(V0, V1, 0, 0, 1, 1);
 	ShuffleTest(V0, V1, 0, 0, 2, 2);
+	ShuffleTest(V0, V1, 0, 1, 2, 2);
+	ShuffleTest(V0, V1, 0, 2, 2, 3);
 	ShuffleTest(V0, V1, 1, 1, 3, 3);
+	ShuffleTest(V0, V1, 1, 3, 1, 3);
+	ShuffleTest(V0, V1, 1, 0, 3, 0);
 	ShuffleTest(V0, V1, 2, 2, 3, 3);
-	ShuffleTest(V0, V1, 2, 3, 0, 1);
-	ShuffleTest(V0, V1, 2, 3, 2, 3);
-
-	ShuffleTest(V0, V1, 3, 1, 3, 0);
-	ShuffleTest(V0, V1, 2, 2, 0, 1);
-	ShuffleTest(V0, V1, 0, 1, 3, 2);
-	ShuffleTest(V0, V1, 1, 3, 0, 3);
+	ShuffleTest(V0, V1, 2, 0, 3, 0);
 
 #undef ShuffleTest
 }
diff --git a/Engine/Source/Runtime/Core/Public/Math/UnrealMathSSE.h b/Engine/Source/Runtime/Core/Public/Math/UnrealMathSSE.h
--- a/Engine/Source/Runtime/Core/Public/Math/UnrealMathSSE.h
+++ b/Engine/Source/Runtime/Core/Public/Math/UnrealMathSSE.h
@@ -82,6 +82,9 @@ struct alignas(UE_SSE_DOUBLE_ALIGNMENT) VectorRegister4Double
 #if !UE_PLATFORM_MATH_USE_AVX
 	VectorRegister2Double XY;
 	VectorRegister2Double ZW;
+
+	FORCEINLINE VectorRegister2Double GetXY() const { return XY; }
+	FORCEINLINE VectorRegister2Double GetZW() const { return ZW; }
 #else
 	union
 	{
@@ -92,6 +95,10 @@ struct alignas(UE_SSE_DOUBLE_ALIGNMENT) VectorRegister4Double
 		};
 		__m256d XYZW;
 	};
+
+	// Use in preference when reading XY or ZW to extract values, it's better on MSVC than the generated memory reads.
+	FORCEINLINE VectorRegister2Double GetXY() const { return _mm256_castpd256_pd128(XYZW); }
+	FORCEINLINE VectorRegister2Double GetZW() const { return _mm256_extractf128_pd(XYZW, 1); }
 #endif
 
 	FORCEINLINE VectorRegister4Double() = default;
@@ -738,10 +745,8 @@ FORCEINLINE void VectorStoreFloat3(const VectorRegister4Double& Vec, double* Dst
 	_mm_storeu_pd(Dst, Vec.XY);
 	_mm_store_sd(Dst + 2, Vec.ZW);
 #else
-	// TODO: AVX: compare perf. MaskStore is supposedly not so great on AMD Zen1 and Zen2 or older Intel.
 	_mm_storeu_pd(Dst, Vec.XY);
 	_mm_store_sd(Dst + 2, Vec.ZW);
-	//_mm256_maskstore_pd(Dst, _mm256_castpd_si256(GlobalVectorConstants::DoubleXYZMask), Vec);
 #endif
 }
 
@@ -761,204 +766,380 @@ FORCEINLINE void VectorStoreFloat1(const VectorRegister4Double& Vec, double* Dst
 	_mm_store_sd(Dst, Vec.XY);
 }
 
-// Templated shuffles required for double shuffles when using __m128d, since we have to break it down in to two separate operations.
-
-// [0,1]:[0,1]
-template<int Index0, int Index1, typename std::enable_if< (Index0 <= 1) && (Index1 <= 1), bool >::type = true>
-FORCEINLINE VectorRegister2Double SelectVectorSwizzle2(const VectorRegister4Double& Vec)
+namespace SSEPermuteHelpers
 {
-	return _mm_shuffle_pd(Vec.XY, Vec.XY, SHUFFLEMASK2(Index0, Index1));
-}
+	// These need to be #define and not constexpr to make them work with enable_if
+#define InLane0(Index0, Index1)		((Index0) <= 1 && (Index1) <= 1)
+#define InLane1(Index0, Index1)		((Index0) >= 2 && (Index1) >= 2)
+#define InSameLane(Index0, Index1)	(InLane0(Index0, Index1) || InLane1(Index0, Index1))
+#define OutOfLane(Index0, Index1)	(!InSameLane(Index0, Index1))
 
-// [0,1]:[2,3]
-template<int Index0, int Index1, typename std::enable_if< (Index0 <= 1) && (Index1 > 1), bool >::type = true>
-FORCEINLINE VectorRegister2Double SelectVectorSwizzle2(const VectorRegister4Double& Vec)
-{
-	return _mm_shuffle_pd(Vec.XY, Vec.ZW, SHUFFLEMASK2(Index0, Index1 - 2));
-}
+	/////////////////////////////////////////////////////////////////////////////////////////////////////////////////////////////////
+	// Double swizzle
+	/////////////////////////////////////////////////////////////////////////////////////////////////////////////////////////////////
 
-// [2,3]:[0,1]
-template<int Index0, int Index1, typename std::enable_if< (Index0 > 1) && (Index1 <= 1), bool >::type = true>
-FORCEINLINE VectorRegister2Double SelectVectorSwizzle2(const VectorRegister4Double& Vec)
-{
-	return _mm_shuffle_pd(Vec.ZW, Vec.XY, SHUFFLEMASK2(Index0 - 2, Index1));
-}
+	// Double Swizzle helpers
+	// Templated swizzles required for double shuffles when using __m128d, since we have to break it down in to two separate operations.
 
-// [2,3]:[2,3]
-template<int Index0, int Index1, typename std::enable_if< (Index0 > 1) && (Index1 > 1), bool >::type = true>
-FORCEINLINE VectorRegister2Double SelectVectorSwizzle2(const VectorRegister4Double& Vec)
-{
-	return _mm_shuffle_pd(Vec.ZW, Vec.ZW, SHUFFLEMASK2(Index0 - 2, Index1 - 2));
-}
+	// [0,1]:[0,1]
+	template<int Index0, int Index1, typename std::enable_if< (Index0 <= 1) && (Index1 <= 1), bool >::type = true>
+	FORCEINLINE VectorRegister2Double SelectVectorSwizzle2(const VectorRegister4Double& Vec)
+	{
+		return _mm_shuffle_pd(Vec.GetXY(), Vec.GetXY(), SHUFFLEMASK2(Index0, Index1));
+	}
 
-template<> FORCEINLINE VectorRegister2Double SelectVectorSwizzle2<0, 1>(const VectorRegister4Double& Vec) { return Vec.XY; }
-template<> FORCEINLINE VectorRegister2Double SelectVectorSwizzle2<2, 3>(const VectorRegister4Double& Vec) { return Vec.ZW; }
+	// [0,1]:[2,3]
+	template<int Index0, int Index1, typename std::enable_if< (Index0 <= 1) && (Index1 >= 2), bool >::type = true>
+	FORCEINLINE VectorRegister2Double SelectVectorSwizzle2(const VectorRegister4Double& Vec)
+	{
+		return _mm_shuffle_pd(Vec.GetXY(), Vec.GetZW(), SHUFFLEMASK2(Index0, Index1 - 2));
+	}
+
+	// [2,3]:[0,1]
+	template<int Index0, int Index1, typename std::enable_if< (Index0 >= 2) && (Index1 <= 1), bool >::type = true>
+	FORCEINLINE VectorRegister2Double SelectVectorSwizzle2(const VectorRegister4Double& Vec)
+	{
+		return _mm_shuffle_pd(Vec.GetZW(), Vec.GetXY(), SHUFFLEMASK2(Index0 - 2, Index1));
+	}
+
+	// [2,3]:[2,3]
+	template<int Index0, int Index1, typename std::enable_if< (Index0 >= 2) && (Index1 >= 2), bool >::type = true>
+	FORCEINLINE VectorRegister2Double SelectVectorSwizzle2(const VectorRegister4Double& Vec)
+	{
+		return _mm_shuffle_pd(Vec.GetZW(), Vec.GetZW(), SHUFFLEMASK2(Index0 - 2, Index1 - 2));
+	}
+
+	template<> FORCEINLINE VectorRegister2Double SelectVectorSwizzle2<0, 1>(const VectorRegister4Double& Vec) { return Vec.GetXY(); }
+	template<> FORCEINLINE VectorRegister2Double SelectVectorSwizzle2<2, 3>(const VectorRegister4Double& Vec) { return Vec.GetZW(); }
 
 #if UE_PLATFORM_MATH_USE_SSE4_1
-// blend can run on more ports than shuffle, so are preferable even if latency is claimed to be the same.
-template<> FORCEINLINE VectorRegister2Double SelectVectorSwizzle2<0, 3>(const VectorRegister4Double& Vec) { return _mm_blend_pd(Vec.XY, Vec.ZW, SHUFFLEMASK2(0, 1)); }
-template<> FORCEINLINE VectorRegister2Double SelectVectorSwizzle2<2, 1>(const VectorRegister4Double& Vec) { return _mm_blend_pd(Vec.ZW, Vec.XY, SHUFFLEMASK2(0, 1)); }
-#endif
+	// blend can run on more ports than shuffle, so are preferable even if latency is claimed to be the same.
+	template<> FORCEINLINE VectorRegister2Double SelectVectorSwizzle2<0, 3>(const VectorRegister4Double& Vec) { return _mm_blend_pd(Vec.GetXY(), Vec.GetZW(), SHUFFLEMASK2(0, 1)); }
+	template<> FORCEINLINE VectorRegister2Double SelectVectorSwizzle2<2, 1>(const VectorRegister4Double& Vec) { return _mm_blend_pd(Vec.GetZW(), Vec.GetXY(), SHUFFLEMASK2(0, 1)); }
+#endif // UE_PLATFORM_MATH_USE_SSE4_1
 
-// Double swizzle
-template<int Index0, int Index1, int Index2, int Index3>
-FORCEINLINE VectorRegister4Double VectorSwizzleTemplate(const VectorRegister4Double& Vec)
-{
-	static_assert(Index0 >= 0 && Index0 <= 3 && Index1 >= 0 && Index1 <= 3 && Index2 >= 0 && Index2 <= 3 && Index3 >= 0 && Index3 <= 3);
 
+#if UE_PLATFORM_MATH_USE_AVX
+
+	// Helper to swap lanes (128-bit pairs)
+	constexpr int PERMUTE_LANE_MASK(int A, int B) { return (A == 0 ? 0x00 : 0x01) | (B == 0 ? (0x02 << 4) : (0x03 << 4)); }
+
+	template<int Lane0, int Lane1>
+	FORCEINLINE VectorRegister4Double PermuteLanes(const VectorRegister4Double& Vec)
+	{
+		static_assert(Lane0 >= 0 && Lane0 <= 1 && Lane1 >= 0 && Lane1 <= 1, "Invalid Index");
+		return _mm256_permute2f128_pd(Vec, Vec, PERMUTE_LANE_MASK(Lane0, Lane1));
+	}
+
+	// Identity
+	template<> FORCEINLINE VectorRegister4Double PermuteLanes<0, 1>(const VectorRegister4Double& Vec) { return Vec; }
+	template<> FORCEINLINE VectorRegister4Double PermuteLanes<0, 0>(const VectorRegister4Double& Vec) { return _mm256_insertf128_pd(Vec, Vec.GetXY(), 1); } // copy XY to lane 1
+#if !UE_PLATFORM_MATH_USE_AVX_2
+	// On AVX1, permute2f128 can be quite slow, so look for alternatives (extract + insert). On AVX2, permute2f128 is more efficient and should equal or beat (extract + insert).
+	// Sources: https://www.agner.org/optimize/instruction_tables.pdf, https://uops.info/table.html
+	template<> FORCEINLINE VectorRegister4Double PermuteLanes<1, 0>(const VectorRegister4Double& Vec) { return _mm256_insertf128_pd(_mm256_castpd128_pd256(Vec.GetZW()), Vec.GetXY(), 1); } // copy XY to lane 1
+	template<> FORCEINLINE VectorRegister4Double PermuteLanes<1, 1>(const VectorRegister4Double& Vec) { return _mm256_insertf128_pd(Vec, Vec.GetZW(), 0); } // copy ZW to lane 0
+#endif // !AVX2
+
+	//
+	// AVX2 _mm256_permute4x64_pd has a latency of 3-6, but there are some specializations using instructions which have a latency of 1 but are restricted to in-lane (128 bit) permutes.
+	// AVX1 benefits from lower latency instructions here than the toggling between 128-bit and 256-bit operations of the generic implementation.
+
+	constexpr int PERMUTE_MASK(int A, int B, int C, int D) { return ((A == 1 ? (1 << 0) : 0) | (B == 1 ? (1 << 1) : 0) | (C == 3 ? (1 << 2) : 0) | (D == 3 ? (1 << 3) : 0)); }
+
+	// [0..1][0..1][2..3][2..3]
+	template<int Index0, int Index1, int Index2, int Index3, typename std::enable_if< InLane0(Index0, Index1) && InLane1(Index2, Index3), bool >::type = true>
+	FORCEINLINE VectorRegister4Double SelectVectorSwizzle(const VectorRegister4Double& Vec)
+	{
+		return _mm256_permute_pd(Vec, PERMUTE_MASK(Index0, Index1, Index2, Index3));
+	}
+
+	// [2..3][2..3][0..1][0..1]
+	template <int Index0, int Index1, int Index2, int Index3, typename std::enable_if< InLane1(Index0, Index1) && InLane0(Index2, Index3), bool >::type = true>
+	FORCEINLINE VectorRegister4Double SelectVectorSwizzle(const VectorRegister4Double& Vec)
+	{
+		// Permute lanes then use [lane0][lane1] swizzle
+		return SelectVectorSwizzle<Index0 - 2, Index1 - 2, Index2 + 2, Index3 + 2>(PermuteLanes<1, 0>(Vec));
+	}
+
+	// [0..1][0..1][0..1][0..1]
+	template <int Index0, int Index1, int Index2, int Index3, typename std::enable_if< InLane0(Index0, Index1) && InLane0(Index2, Index3), bool >::type = true>
+	FORCEINLINE VectorRegister4Double SelectVectorSwizzle(const VectorRegister4Double& Vec)
+	{
+		// Permute lanes then use [lane0][lane1] swizzle
+		return SelectVectorSwizzle<Index0, Index1, Index2 + 2, Index3 + 2>(PermuteLanes<0, 0>(Vec));
+	}
+
+	// [2..3][2..3][2..3][2..3]
+	template <int Index0, int Index1, int Index2, int Index3, typename std::enable_if< InLane1(Index0, Index1) && InLane1(Index2, Index3), bool >::type = true>
+	FORCEINLINE VectorRegister4Double SelectVectorSwizzle(const VectorRegister4Double& Vec)
+	{
+		// Permute lanes then use [lane0][lane1] swizzle
+		return SelectVectorSwizzle<Index0 - 2, Index1 - 2, Index2, Index3>(PermuteLanes<1, 1>(Vec));
+	}
+
+	// Anything with out-of-lane pairs
+	template<int Index0, int Index1, int Index2, int Index3, typename std::enable_if< OutOfLane(Index0, Index1) || OutOfLane(Index2, Index3), bool >::type = true>
+	FORCEINLINE VectorRegister4Double SelectVectorSwizzle(const VectorRegister4Double& Vec)
+	{
 #if UE_PLATFORM_MATH_USE_AVX_2
-	return _mm256_permute4x64_pd(Vec, SHUFFLEMASK(Index0, Index1, Index2, Index3));
+		return _mm256_permute4x64_pd(Vec, SHUFFLEMASK(Index0, Index1, Index2, Index3));
 #else
-	return VectorRegister4Double(
-		SelectVectorSwizzle2<Index0, Index1>(Vec),	
-		SelectVectorSwizzle2<Index2, Index3>(Vec)
-	);
+		return VectorRegister4Double(
+			SelectVectorSwizzle2<Index0, Index1>(Vec),
+			SelectVectorSwizzle2<Index2, Index3>(Vec)
+		);
 #endif
-}
+	}
 
-// Double Swizzle specializations
-template<> FORCEINLINE VectorRegister4Double VectorSwizzleTemplate<0, 1, 2, 3>(const VectorRegister4Double& Vec) { return Vec; }
+	//
+	// Specializations
+	//
+	template<> FORCEINLINE VectorRegister4Double SelectVectorSwizzle<0, 0, 2, 2>(const VectorRegister4Double& Vec) { return _mm256_movedup_pd(Vec); } // special instruction exists for this.
+	template<> FORCEINLINE VectorRegister4Double SelectVectorSwizzle<0, 1, 2, 3>(const VectorRegister4Double& Vec) { return Vec; } // Identity
+	template<> FORCEINLINE VectorRegister4Double SelectVectorSwizzle<2, 3, 0, 1>(const VectorRegister4Double& Vec) { return PermuteLanes<1, 0>(Vec); }
+	template<> FORCEINLINE VectorRegister4Double SelectVectorSwizzle<0, 1, 0, 1>(const VectorRegister4Double& Vec) { return PermuteLanes<0, 0>(Vec); }
+	template<> FORCEINLINE VectorRegister4Double SelectVectorSwizzle<2, 3, 2, 3>(const VectorRegister4Double& Vec) { return PermuteLanes<1, 1>(Vec); }
 
-#if UE_PLATFORM_MATH_USE_AVX
-constexpr int PERMUTE_MASK(int A, int B, int C, int D) { return ((A == 1 ? (1 << 0) : 0) | (B == 1 ? (1 << 1) : 0) | (C == 3 ? (1 << 2) : 0) | (D == 3 ? (1 << 3) : 0)); }
+#endif // AVX
 
-// AVX2 _mm256_permute4x64_pd has a latency of 3, but there are some specializations using instructions which have a latency of 1 but are restricted to in-lane (128 bit) permutes.
-// AVX1 benefits from lower latency instructions here than the shuffling between 128-bit and 256-bit operations of the generic implementation.
-template<> FORCEINLINE VectorRegister4Double VectorSwizzleTemplate<0, 1, 2, 2>(const VectorRegister4Double& Vec) { return _mm256_permute_pd(Vec, PERMUTE_MASK(0, 1, 2, 2)); }
-template<> FORCEINLINE VectorRegister4Double VectorSwizzleTemplate<0, 1, 3, 3>(const VectorRegister4Double& Vec) { return _mm256_permute_pd(Vec, PERMUTE_MASK(0, 1, 3, 3)); }
-template<> FORCEINLINE VectorRegister4Double VectorSwizzleTemplate<0, 1, 3, 2>(const VectorRegister4Double& Vec) { return _mm256_permute_pd(Vec, PERMUTE_MASK(0, 1, 3, 2)); }
-																																											 
-template<> FORCEINLINE VectorRegister4Double VectorSwizzleTemplate<0, 0, 2, 3>(const VectorRegister4Double& Vec) { return _mm256_permute_pd(Vec, PERMUTE_MASK(0, 0, 2, 3)); }
-template<> FORCEINLINE VectorRegister4Double VectorSwizzleTemplate<0, 0, 2, 2>(const VectorRegister4Double& Vec) { return _mm256_movedup_pd(Vec); }							 
-template<> FORCEINLINE VectorRegister4Double VectorSwizzleTemplate<0, 0, 3, 3>(const VectorRegister4Double& Vec) { return _mm256_permute_pd(Vec, PERMUTE_MASK(0, 0, 3, 3)); }
-template<> FORCEINLINE VectorRegister4Double VectorSwizzleTemplate<0, 0, 3, 2>(const VectorRegister4Double& Vec) { return _mm256_permute_pd(Vec, PERMUTE_MASK(0, 0, 3, 2)); }
-																																											 
-template<> FORCEINLINE VectorRegister4Double VectorSwizzleTemplate<1, 0, 2, 3>(const VectorRegister4Double& Vec) { return _mm256_permute_pd(Vec, PERMUTE_MASK(1, 0, 2, 3)); }
-template<> FORCEINLINE VectorRegister4Double VectorSwizzleTemplate<1, 0, 2, 2>(const VectorRegister4Double& Vec) { return _mm256_permute_pd(Vec, PERMUTE_MASK(1, 0, 2, 2)); }
-template<> FORCEINLINE VectorRegister4Double VectorSwizzleTemplate<1, 0, 3, 3>(const VectorRegister4Double& Vec) { return _mm256_permute_pd(Vec, PERMUTE_MASK(1, 0, 3, 3)); }
-template<> FORCEINLINE VectorRegister4Double VectorSwizzleTemplate<1, 0, 3, 2>(const VectorRegister4Double& Vec) { return _mm256_permute_pd(Vec, PERMUTE_MASK(1, 0, 3, 2)); }
-																																											 
-template<> FORCEINLINE VectorRegister4Double VectorSwizzleTemplate<1, 1, 2, 3>(const VectorRegister4Double& Vec) { return _mm256_permute_pd(Vec, PERMUTE_MASK(1, 1, 2, 3)); }
-template<> FORCEINLINE VectorRegister4Double VectorSwizzleTemplate<1, 1, 2, 2>(const VectorRegister4Double& Vec) { return _mm256_permute_pd(Vec, PERMUTE_MASK(1, 1, 2, 2)); }
-template<> FORCEINLINE VectorRegister4Double VectorSwizzleTemplate<1, 1, 3, 3>(const VectorRegister4Double& Vec) { return _mm256_permute_pd(Vec, PERMUTE_MASK(1, 1, 3, 3)); }
-template<> FORCEINLINE VectorRegister4Double VectorSwizzleTemplate<1, 1, 3, 2>(const VectorRegister4Double& Vec) { return _mm256_permute_pd(Vec, PERMUTE_MASK(1, 1, 3, 2)); }
-#endif // UE_PLATFORM_MATH_USE_AVX
+	// Double swizzle wrapper
+	template<int Index0, int Index1, int Index2, int Index3>
+	FORCEINLINE VectorRegister4Double VectorSwizzleTemplate(const VectorRegister4Double& Vec)
+	{
+		static_assert(Index0 >= 0 && Index0 <= 3 && Index1 >= 0 && Index1 <= 3 && Index2 >= 0 && Index2 <= 3 && Index3 >= 0 && Index3 <= 3, "Invalid Index");
 
-#if (UE_PLATFORM_MATH_USE_AVX && !UE_PLATFORM_MATH_USE_AVX_2)
-// These don't really bring any benefit on AVX2 compared to _mm256_permute4x64_pd, but avoid (extract + move + set) on AVX1.
-template<> FORCEINLINE VectorRegister4Double VectorSwizzleTemplate<2, 3, 0, 1>(const VectorRegister4Double& Vec) { return _mm256_permute2f128_pd(Vec, Vec, (0x01) | (0x02 << 4)); }
-template<> FORCEINLINE VectorRegister4Double VectorSwizzleTemplate<0, 1, 0, 1>(const VectorRegister4Double& Vec) { return _mm256_permute2f128_pd(Vec, Vec, (0x00) | (0x02 << 4)); }
-template<> FORCEINLINE VectorRegister4Double VectorSwizzleTemplate<2, 3, 2, 3>(const VectorRegister4Double& Vec) { return _mm256_permute2f128_pd(Vec, Vec, (0x01) | (0x03 << 4)); }
+#if UE_PLATFORM_MATH_USE_AVX
+		return SelectVectorSwizzle<Index0, Index1, Index2, Index3>(Vec);
+#else
+		return VectorRegister4Double(
+			SelectVectorSwizzle2<Index0, Index1>(Vec),
+			SelectVectorSwizzle2<Index2, Index3>(Vec)
+		);
 #endif
+	}
 
+	// Specializations
+	template<> FORCEINLINE VectorRegister4Double VectorSwizzleTemplate<0, 1, 2, 3>(const VectorRegister4Double& Vec) { return Vec; } // Identity
 
+	/////////////////////////////////////////////////////////////////////////////////////////////////////////////////////////////////
+	// Double replicate
+	/////////////////////////////////////////////////////////////////////////////////////////////////////////////////////////////////
 
-// Double replicate (2 doubles)
-template <int Index>
-FORCEINLINE VectorRegister2Double VectorReplicateImpl2(const VectorRegister2Double& Vec)
-{
-	return _mm_shuffle_pd(Vec, Vec, SHUFFLEMASK2(Index, Index));
-}
+	template <int Index>
+	FORCEINLINE VectorRegister2Double VectorReplicateImpl2(const VectorRegister2Double& Vec)
+	{
+		// Note: 2 doubles (VectorRegister2Double / m128d)
+		return _mm_shuffle_pd(Vec, Vec, SHUFFLEMASK2(Index, Index));
+	}
 
-// Double replicate (4 doubles)
-template <int Index, typename std::enable_if < (Index <= 1), bool >::type = true >
-FORCEINLINE VectorRegister4Double VectorReplicateImpl(const VectorRegister4Double& Vec)
-{
-	VectorRegister2Double Temp = VectorReplicateImpl2<Index>(Vec.XY);
-	return VectorRegister4Double(Temp, Temp);
-}
+	// Double replicate (4 doubles)
+	template <int Index, typename std::enable_if< (Index <= 1), bool >::type = true >
+	FORCEINLINE VectorRegister4Double VectorReplicateImpl4(const VectorRegister4Double& Vec)
+	{
+		VectorRegister2Double Temp = VectorReplicateImpl2<Index>(Vec.GetXY());
+		return VectorRegister4Double(Temp, Temp);
+	}
 
-template <int Index, typename std::enable_if < (Index > 1), bool >::type = true >
-FORCEINLINE VectorRegister4Double VectorReplicateImpl(const VectorRegister4Double& Vec)
-{
-	VectorRegister2Double Temp = VectorReplicateImpl2<Index - 2>(Vec.ZW);
-	return VectorRegister4Double(Temp, Temp);
-}
+	template <int Index, typename std::enable_if< (Index >= 2), bool >::type = true >
+	FORCEINLINE VectorRegister4Double VectorReplicateImpl4(const VectorRegister4Double& Vec)
+	{
+		VectorRegister2Double Temp = VectorReplicateImpl2<Index - 2>(Vec.GetZW());
+		return VectorRegister4Double(Temp, Temp);
+	}
 
-template<int Index>
-FORCEINLINE VectorRegister4Double VectorReplicateTemplate(const VectorRegister4Double& Vec)
-{
-	static_assert(Index >= 0 && Index <= 3);
+	//
+	// Double replicate wrapper
+	//
+	template<int Index>
+	FORCEINLINE VectorRegister4Double VectorReplicateTemplate(const VectorRegister4Double& Vec)
+	{
+		static_assert(Index >= 0 && Index <= 3, "Invalid Index");
 
 #if UE_PLATFORM_MATH_USE_AVX_2
-	return _mm256_permute4x64_pd(Vec, SHUFFLEMASK(Index, Index, Index, Index));
+		return VectorSwizzleTemplate<Index, Index, Index, Index>(Vec);
 #else
-	return VectorReplicateImpl<Index>(Vec);
+		return VectorReplicateImpl4<Index>(Vec);
 #endif
-}
+	}
 
-// Double shuffle
-template<int Index0, int Index1, int Index2, int Index3>
-FORCEINLINE VectorRegister4Double VectorShuffleTemplate(const VectorRegister4Double& Vec1, const VectorRegister4Double& Vec2)
-{
-	static_assert(Index0 >= 0 && Index0 <= 3 && Index1 >= 0 && Index1 <= 3 && Index2 >= 0 && Index2 <= 3 && Index3 >= 0 && Index3 <= 3);
-	return VectorRegister4Double(
-		SelectVectorSwizzle2<Index0, Index1>(Vec1),
-		SelectVectorSwizzle2<Index2, Index3>(Vec2)
-	);
-}
+	/////////////////////////////////////////////////////////////////////////////////////////////////////////////////////////////////
+	// Double shuffle
+	/////////////////////////////////////////////////////////////////////////////////////////////////////////////////////////////////
 
-// Double Shuffle specializations 
 #if UE_PLATFORM_MATH_USE_AVX
-template<> FORCEINLINE VectorRegister4Double VectorShuffleTemplate<0, 1, 0, 1>(const VectorRegister4Double& Vec1, const VectorRegister4Double& Vec2) { return _mm256_permute2f128_pd(Vec1, Vec2, (0x00) | (0x02 << 4)); }
-template<> FORCEINLINE VectorRegister4Double VectorShuffleTemplate<0, 1, 2, 3>(const VectorRegister4Double& Vec1, const VectorRegister4Double& Vec2) { return _mm256_blend_pd(Vec1, Vec2, 0b1100); }
-template<> FORCEINLINE VectorRegister4Double VectorShuffleTemplate<2, 3, 0, 1>(const VectorRegister4Double& Vec1, const VectorRegister4Double& Vec2) { return _mm256_permute2f128_pd(Vec1, Vec2, (0x01) | (0x02 << 4)); }
-template<> FORCEINLINE VectorRegister4Double VectorShuffleTemplate<2, 3, 2, 3>(const VectorRegister4Double& Vec1, const VectorRegister4Double& Vec2) { return _mm256_permute2f128_pd(Vec1, Vec2, (0x01) | (0x03 << 4)); }
+
+	//
+	// Lane shuffle helper
+	//
+	template<int Lane0, int Lane1>
+	FORCEINLINE VectorRegister4Double ShuffleLanes(const VectorRegister4Double& Vec1, const VectorRegister4Double& Vec2)
+	{
+		static_assert(Lane0 >= 0 && Lane0 <= 1 && Lane1 >= 0 && Lane1 <= 1, "Invalid Index");
+		return _mm256_permute2f128_pd(Vec1, Vec2, PERMUTE_LANE_MASK(Lane0, Lane1));
+	}
+
+	// Lane shuffle helper specialization
+	template<> FORCEINLINE VectorRegister4Double ShuffleLanes<0, 1>(const VectorRegister4Double& Vec1, const VectorRegister4Double& Vec2) { return _mm256_blend_pd(Vec1, Vec2, 0b1100); }
+	template<> FORCEINLINE VectorRegister4Double ShuffleLanes<0, 0>(const VectorRegister4Double& Vec1, const VectorRegister4Double& Vec2) { return _mm256_insertf128_pd(Vec1, Vec2.GetXY(), 1); } // copy XY to lane 1
+#if !UE_PLATFORM_MATH_USE_AVX_2
+	// On AVX1, permute2f128 can be quite slow, so look for alternatives (extract + insert). On AVX2, permute2f128 is more efficient and should equal or beat (extract + insert).
+	// Sources: https://www.agner.org/optimize/instruction_tables.pdf, https://uops.info/table.html
+	template<> FORCEINLINE VectorRegister4Double ShuffleLanes<1, 0>(const VectorRegister4Double& Vec1, const VectorRegister4Double& Vec2) { return _mm256_insertf128_pd(_mm256_castpd128_pd256(Vec1.GetZW()), Vec2.GetXY(), 1); } // copy XY to lane 1
+	template<> FORCEINLINE VectorRegister4Double ShuffleLanes<1, 1>(const VectorRegister4Double& Vec1, const VectorRegister4Double& Vec2) { return _mm256_insertf128_pd(Vec2, Vec1.GetZW(), 0); } // copy ZW to lane 0
+#endif // !AVX2
+
+	//
+	// Double shuffle helpers
+	//
+
+	// When index pairs are within the same lane, SelectVectorShuffle first efficiently blends elements from the two vectors,
+	// then efficiently swizzles within 128-bit lanes using specializations for indices [0..1][0..1][2..3][2..3]
+	// 
+	// [0..1][0..1][2..3][2..3]
+	template <int Index0, int Index1, int Index2, int Index3, typename std::enable_if< InLane0(Index0, Index1) && InLane1(Index2, Index3), bool >::type = true>
+	FORCEINLINE VectorRegister4Double SelectVectorShuffle(const VectorRegister4Double& Vec1, const VectorRegister4Double& Vec2)
+	{
+		const VectorRegister4Double Blended = ShuffleLanes<0, 1>(Vec1, Vec2);
+		return VectorSwizzleTemplate<Index0, Index1, Index2, Index3>(Blended);
+	}
+
+	// [2..3][2..3][0..1][0..1]
+	template <int Index0, int Index1, int Index2, int Index3, typename std::enable_if< InLane1(Index0, Index1) && InLane0(Index2, Index3), bool >::type = true>
+	FORCEINLINE VectorRegister4Double SelectVectorShuffle(const VectorRegister4Double& Vec1, const VectorRegister4Double& Vec2)
+	{
+		const VectorRegister4Double Blended = ShuffleLanes<1, 0>(Vec1, Vec2);
+		return VectorSwizzleTemplate<Index0 - 2, Index1 - 2, Index2 + 2, Index3 + 2>(Blended);
+	}
+
+	// [0..1][0..1][0..1][0..1]
+	template <int Index0, int Index1, int Index2, int Index3, typename std::enable_if< InLane0(Index0, Index1) && InLane0(Index2, Index3), bool >::type = true>
+	FORCEINLINE VectorRegister4Double SelectVectorShuffle(const VectorRegister4Double& Vec1, const VectorRegister4Double& Vec2)
+	{
+		const VectorRegister4Double Blended = ShuffleLanes<0, 0>(Vec1, Vec2);
+		return VectorSwizzleTemplate<Index0, Index1, Index2 + 2, Index3 + 2>(Blended);
+	}
+
+	// [2..3][2..3][2..3][2..3]
+	template <int Index0, int Index1, int Index2, int Index3, typename std::enable_if< InLane1(Index0, Index1) && InLane1(Index2, Index3), bool >::type = true>
+	FORCEINLINE VectorRegister4Double SelectVectorShuffle(const VectorRegister4Double& Vec1, const VectorRegister4Double& Vec2)
+	{
+		const VectorRegister4Double Blended = ShuffleLanes<1, 1>(Vec1, Vec2);
+		return VectorSwizzleTemplate<Index0 - 2, Index1 - 2, Index2, Index3>(Blended);
+	}
+
+	//
+	// Remaining cases have at least one pair doing a cross-lane swizzle.
+	//
+
+	template <int Index0, int Index1, int Index2, int Index3, typename std::enable_if< InSameLane(Index0, Index1) && OutOfLane(Index2, Index3), bool >::type = true>
+	FORCEINLINE VectorRegister4Double SelectVectorShuffle(const VectorRegister4Double& Vec1, const VectorRegister4Double& Vec2)
+	{
+		const VectorRegister4Double Vec1_XY = VectorSwizzleTemplate<Index0, Index1, 2, 3>(Vec1);
+		const VectorRegister2Double Vec2_ZW = SelectVectorSwizzle2<Index2, Index3>(Vec2);
+		return _mm256_insertf128_pd(Vec1_XY, Vec2_ZW, 0x1);
+	}
+
+	template <int Index0, int Index1, int Index2, int Index3, typename std::enable_if< OutOfLane(Index0, Index1) && InSameLane(Index2, Index3), bool >::type = true>
+	FORCEINLINE VectorRegister4Double SelectVectorShuffle(const VectorRegister4Double& Vec1, const VectorRegister4Double& Vec2)
+	{
+		const VectorRegister2Double Vec1_XY = SelectVectorSwizzle2<Index0, Index1>(Vec1);
+		const VectorRegister4Double Vec2_ZW = VectorSwizzleTemplate<0, 1, Index2, Index3>(Vec2);
+		return _mm256_insertf128_pd(Vec2_ZW, Vec1_XY, 0x0);
+	}
+
+	template <int Index0, int Index1, int Index2, int Index3, typename std::enable_if< OutOfLane(Index0, Index1) && OutOfLane(Index2, Index3), bool >::type = true>
+	FORCEINLINE VectorRegister4Double SelectVectorShuffle(const VectorRegister4Double& Vec1, const VectorRegister4Double& Vec2)
+	{
+		return VectorRegister4Double(
+			SelectVectorSwizzle2<Index0, Index1>(Vec1),
+			SelectVectorSwizzle2<Index2, Index3>(Vec2)
+		);
+	}
+
+	// AVX Double Shuffle specializations 
+	// Shuffles of 128-bit pairs, ie combinations of [0,1][2,3].
+	template<> FORCEINLINE VectorRegister4Double SelectVectorShuffle<0, 1, 0, 1>(const VectorRegister4Double& Vec1, const VectorRegister4Double& Vec2) { return ShuffleLanes<0, 0>(Vec1, Vec2); }
+	template<> FORCEINLINE VectorRegister4Double SelectVectorShuffle<0, 1, 2, 3>(const VectorRegister4Double& Vec1, const VectorRegister4Double& Vec2) { return ShuffleLanes<0, 1>(Vec1, Vec2); }
+	template<> FORCEINLINE VectorRegister4Double SelectVectorShuffle<2, 3, 0, 1>(const VectorRegister4Double& Vec1, const VectorRegister4Double& Vec2) { return ShuffleLanes<1, 0>(Vec1, Vec2); }
+	template<> FORCEINLINE VectorRegister4Double SelectVectorShuffle<2, 3, 2, 3>(const VectorRegister4Double& Vec1, const VectorRegister4Double& Vec2) { return ShuffleLanes<1, 1>(Vec1, Vec2); }
+
 #else
-template<> FORCEINLINE VectorRegister4Double VectorShuffleTemplate<0, 1, 0, 1>(const VectorRegister4Double& Vec1, const VectorRegister4Double& Vec2) { return VectorRegister4Double(Vec1.XY, Vec2.XY); }
-template<> FORCEINLINE VectorRegister4Double VectorShuffleTemplate<0, 1, 2, 3>(const VectorRegister4Double& Vec1, const VectorRegister4Double& Vec2) { return VectorRegister4Double(Vec1.XY, Vec2.ZW); }
-template<> FORCEINLINE VectorRegister4Double VectorShuffleTemplate<2, 3, 0, 1>(const VectorRegister4Double& Vec1, const VectorRegister4Double& Vec2) { return VectorRegister4Double(Vec1.ZW, Vec2.XY); }
-template<> FORCEINLINE VectorRegister4Double VectorShuffleTemplate<2, 3, 2, 3>(const VectorRegister4Double& Vec1, const VectorRegister4Double& Vec2) { return VectorRegister4Double(Vec1.ZW, Vec2.ZW); }
-#endif
 
+	// Non-AVX implementation
+	template<int Index0, int Index1, int Index2, int Index3>
+	FORCEINLINE VectorRegister4Double SelectVectorShuffle(const VectorRegister4Double& Vec1, const VectorRegister4Double& Vec2)
+	{
+		return VectorRegister4Double(
+			SelectVectorSwizzle2<Index0, Index1>(Vec1),
+			SelectVectorSwizzle2<Index2, Index3>(Vec2)
+		);
+	}
 
-// Float swizzle
-template<int Index0, int Index1, int Index2, int Index3>
-FORCEINLINE VectorRegister4Float VectorSwizzleTemplate(const VectorRegister4Float& Vec)
-{
-	VectorRegister4Float Result = _mm_shuffle_ps(Vec, Vec, SHUFFLEMASK(Index0, Index1, Index2, Index3));
-	return Result;
-}
+#endif // AVX
 
-// Float Swizzle specializations.
-// These can result in no-ops or simpler ops than shuffle which don't compete with the shuffle unit, or which can copy directly to the destination and avoid an intermediate mov.
-// See: https://stackoverflow.com/questions/56238197/what-is-the-difference-between-mm-movehdup-ps-and-mm-shuffle-ps-in-this-case
-template<> FORCEINLINE VectorRegister4Float VectorSwizzleTemplate<0, 1, 2, 3>(const VectorRegister4Float& Vec) { return Vec; }
-template<> FORCEINLINE VectorRegister4Float VectorSwizzleTemplate<0, 1, 0, 1>(const VectorRegister4Float& Vec) { return _mm_movelh_ps(Vec, Vec); }
-template<> FORCEINLINE VectorRegister4Float VectorSwizzleTemplate<2, 3, 2, 3>(const VectorRegister4Float& Vec) { return _mm_movehl_ps(Vec, Vec); }
-template<> FORCEINLINE VectorRegister4Float VectorSwizzleTemplate<0, 0, 1, 1>(const VectorRegister4Float& Vec) { return _mm_unpacklo_ps(Vec, Vec); }
-template<> FORCEINLINE VectorRegister4Float VectorSwizzleTemplate<2, 2, 3, 3>(const VectorRegister4Float& Vec) { return _mm_unpackhi_ps(Vec, Vec); }
+	//
+	// Double shuffle wrapper
+	//
+	template<int Index0, int Index1, int Index2, int Index3>
+	FORCEINLINE VectorRegister4Double VectorShuffleTemplate(const VectorRegister4Double& Vec1, const VectorRegister4Double& Vec2)
+	{
+		static_assert(Index0 >= 0 && Index0 <= 3 && Index1 >= 0 && Index1 <= 3 && Index2 >= 0 && Index2 <= 3 && Index3 >= 0 && Index3 <= 3, "Invalid Index");
+		return SelectVectorShuffle<Index0, Index1, Index2, Index3>(Vec1, Vec2);
+	}
+
+	/////////////////////////////////////////////////////////////////////////////////////////////////////////////////////////////////
+	// Float swizzle
+	/////////////////////////////////////////////////////////////////////////////////////////////////////////////////////////////////
+
+	template<int Index0, int Index1, int Index2, int Index3>
+	FORCEINLINE VectorRegister4Float VectorSwizzleTemplate(const VectorRegister4Float& Vec)
+	{
+		VectorRegister4Float Result = _mm_shuffle_ps(Vec, Vec, SHUFFLEMASK(Index0, Index1, Index2, Index3));
+		return Result;
+	}
+
+	// Float Swizzle specializations.
+	// These can result in no-ops or simpler ops than shuffle which don't compete with the shuffle unit, or which can copy directly to the destination and avoid an intermediate mov.
+	// See: https://stackoverflow.com/questions/56238197/what-is-the-difference-between-mm-movehdup-ps-and-mm-shuffle-ps-in-this-case
+	template<> FORCEINLINE VectorRegister4Float VectorSwizzleTemplate<0, 1, 2, 3>(const VectorRegister4Float& Vec) { return Vec; }
+	template<> FORCEINLINE VectorRegister4Float VectorSwizzleTemplate<0, 1, 0, 1>(const VectorRegister4Float& Vec) { return _mm_movelh_ps(Vec, Vec); }
+	template<> FORCEINLINE VectorRegister4Float VectorSwizzleTemplate<2, 3, 2, 3>(const VectorRegister4Float& Vec) { return _mm_movehl_ps(Vec, Vec); }
+	template<> FORCEINLINE VectorRegister4Float VectorSwizzleTemplate<0, 0, 1, 1>(const VectorRegister4Float& Vec) { return _mm_unpacklo_ps(Vec, Vec); }
+	template<> FORCEINLINE VectorRegister4Float VectorSwizzleTemplate<2, 2, 3, 3>(const VectorRegister4Float& Vec) { return _mm_unpackhi_ps(Vec, Vec); }
 
 #if UE_PLATFORM_MATH_USE_SSE4_1
-template<> FORCEINLINE VectorRegister4Float VectorSwizzleTemplate<0, 0, 2, 2>(const VectorRegister4Float& Vec) { return _mm_moveldup_ps(Vec); }
-template<> FORCEINLINE VectorRegister4Float VectorSwizzleTemplate<1, 1, 3, 3>(const VectorRegister4Float& Vec) { return _mm_movehdup_ps(Vec); }
+	template<> FORCEINLINE VectorRegister4Float VectorSwizzleTemplate<0, 0, 2, 2>(const VectorRegister4Float& Vec) { return _mm_moveldup_ps(Vec); }
+	template<> FORCEINLINE VectorRegister4Float VectorSwizzleTemplate<1, 1, 3, 3>(const VectorRegister4Float& Vec) { return _mm_movehdup_ps(Vec); }
 #endif
 
 #if UE_PLATFORM_MATH_USE_AVX_2
-template<> FORCEINLINE VectorRegister4Float VectorSwizzleTemplate<0, 0, 0, 0>(const VectorRegister4Float& Vec) { return _mm_broadcastss_ps(Vec); }
+	template<> FORCEINLINE VectorRegister4Float VectorSwizzleTemplate<0, 0, 0, 0>(const VectorRegister4Float& Vec) { return _mm_broadcastss_ps(Vec); }
 #endif
 
+	/////////////////////////////////////////////////////////////////////////////////////////////////////////////////////////////////
+	// Float replicate
+	template<int Index>
+	FORCEINLINE VectorRegister4Float VectorReplicateTemplate(const VectorRegister4Float& Vec)
+	{
+		static_assert(Index >= 0 && Index <= 3, "Invalid Index");
+		return VectorSwizzleTemplate<Index, Index, Index, Index>(Vec);
+	}
 
-// Float replicate
-template<int Index>
-FORCEINLINE VectorRegister4Float VectorReplicateTemplate(const VectorRegister4Float& Vec)
-{
-	static_assert(Index >= 0 && Index <= 3);
-	return VectorSwizzleTemplate<Index, Index, Index, Index>(Vec);
-}
-
-// Float shuffle
-template<int Index0, int Index1, int Index2, int Index3>
-FORCEINLINE VectorRegister4Float VectorShuffleTemplate(const VectorRegister4Float& Vec1, const VectorRegister4Float& Vec2)
-{
-	static_assert(Index0 >= 0 && Index0 <= 3 && Index1 >= 0 && Index1 <= 3 && Index2 >= 0 && Index2 <= 3 && Index3 >= 0 && Index3 <= 3);
-	return _mm_shuffle_ps(Vec1, Vec2, SHUFFLEMASK(Index0, Index1, Index2, Index3));
-}
+	/////////////////////////////////////////////////////////////////////////////////////////////////////////////////////////////////
+	// Float shuffle
+	template<int Index0, int Index1, int Index2, int Index3>
+	FORCEINLINE VectorRegister4Float VectorShuffleTemplate(const VectorRegister4Float& Vec1, const VectorRegister4Float& Vec2)
+	{
+		static_assert(Index0 >= 0 && Index0 <= 3 && Index1 >= 0 && Index1 <= 3 && Index2 >= 0 && Index2 <= 3 && Index3 >= 0 && Index3 <= 3, "Invalid Index");
+		return _mm_shuffle_ps(Vec1, Vec2, SHUFFLEMASK(Index0, Index1, Index2, Index3));
+	}
 
-// Float Shuffle specializations
-template<> FORCEINLINE VectorRegister4Float VectorShuffleTemplate<0, 1, 0, 1>(const VectorRegister4Float& Vec1, const VectorRegister4Float& Vec2) { return _mm_movelh_ps(Vec1, Vec2); }
-template<> FORCEINLINE VectorRegister4Float VectorShuffleTemplate<2, 3, 2, 3>(const VectorRegister4Float& Vec1, const VectorRegister4Float& Vec2) { return _mm_movehl_ps(Vec2, Vec1); } // Note: movehl copies first from the 2nd argument
+	// Float Shuffle specializations
+	template<> FORCEINLINE VectorRegister4Float VectorShuffleTemplate<0, 1, 0, 1>(const VectorRegister4Float& Vec1, const VectorRegister4Float& Vec2) { return _mm_movelh_ps(Vec1, Vec2); }
+	template<> FORCEINLINE VectorRegister4Float VectorShuffleTemplate<2, 3, 2, 3>(const VectorRegister4Float& Vec1, const VectorRegister4Float& Vec2) { return _mm_movehl_ps(Vec2, Vec1); } // Note: movehl copies first from the 2nd argument
 
+}; // namespace SSEPermuteHelpers
 
 /**
  * Replicates one element into all four elements and returns the new vector.
@@ -968,7 +1149,7 @@ template<> FORCEINLINE VectorRegister4Float VectorShuffleTemplate<2, 3, 2, 3>(co
  * @return				VectorRegister4Float( Vec[ElementIndex], Vec[ElementIndex], Vec[ElementIndex], Vec[ElementIndex] )
  */
 
-#define VectorReplicate( Vec, ElementIndex )	VectorReplicateTemplate<ElementIndex>(Vec)
+#define VectorReplicate( Vec, ElementIndex )	SSEPermuteHelpers::VectorReplicateTemplate<ElementIndex>(Vec)
 
  /**
   * Swizzles the 4 components of a vector and returns the result.
@@ -980,7 +1161,7 @@ template<> FORCEINLINE VectorRegister4Float VectorShuffleTemplate<2, 3, 2, 3>(co
   * @param W			Index for which component to use for W (literal 0-3)
   * @return			The swizzled vector
   */
-#define VectorSwizzle( Vec, X, Y, Z, W )		VectorSwizzleTemplate<X,Y,Z,W>( Vec )
+#define VectorSwizzle( Vec, X, Y, Z, W )		SSEPermuteHelpers::VectorSwizzleTemplate<X,Y,Z,W>( Vec )
 
   /**
    * Creates a vector through selecting two components from each vector via a shuffle mask.
@@ -993,7 +1174,7 @@ template<> FORCEINLINE VectorRegister4Float VectorShuffleTemplate<2, 3, 2, 3>(co
    * @param W			Index for which component of Vector2 to use for W (literal 0-3)
    * @return			The swizzled vector
    */
-#define VectorShuffle( Vec1, Vec2, X, Y, Z, W )		VectorShuffleTemplate<X,Y,Z,W>( Vec1, Vec2 )
+#define VectorShuffle( Vec1, Vec2, X, Y, Z, W )		SSEPermuteHelpers::VectorShuffleTemplate<X,Y,Z,W>( Vec1, Vec2 )
 
 
 /**
@@ -1330,30 +1511,24 @@ FORCEINLINE VectorRegister4Double VectorDot3(const VectorRegister4Double& Vec1,
  */
 FORCEINLINE VectorRegister4Float VectorDot4( const VectorRegister4Float& Vec1, const VectorRegister4Float& Vec2 )
 {
-	VectorRegister4Float Temp1, Temp2;
-	Temp1 = VectorMultiply(Vec1, Vec2);
-	Temp2 = VectorSwizzle( Temp1, 2, 3, 0, 1);	// (Z,W,X,Y).
-	Temp1 = VectorAdd(Temp1, Temp2);			// (X*X + Z*Z, Y*Y + W*W, Z*Z + X*X, W*W + Y*Y)
-	Temp2 = VectorSwizzle(Temp1, 1, 2, 3, 0);	// Rotate left 1 element (Y,Z,W,X).
-	return VectorAdd(Temp1, Temp2);				// (X*X + Z*Z + Y*Y + W*W, Y*Y + W*W + Z*Z + X*X, Z*Z + X*X + W*W + Y*Y, W*W + Y*Y + X*X + Z*Z)
+	VectorRegister4Float R, T;
+	R = VectorMultiply(Vec1, Vec2);		// (XX, YY, ZZ, WW)
+	T = VectorSwizzle(R, 1, 0, 3, 2);	// (YY, XX, WW, ZZ)
+	R = VectorAdd(R, T);				// (XX + YY, YY + XX, ZZ + WW, WW + ZZ)
+	T = VectorSwizzle(R, 2, 3, 0, 1);	// (ZZ + WW, WW + ZZ, XX + YY, YY + XX)
+	return VectorAdd(R, T);				// (XX + YY + ZZ + WW, YY + XX + WW + ZZ, ZZ + WW + XX + YY, WW + ZZ + YY + XX)
 }
 
 FORCEINLINE VectorRegister4Double VectorDot4(const VectorRegister4Double& Vec1, const VectorRegister4Double& Vec2)
 {
 #if UE_PLATFORM_MATH_USE_AVX
-	// AVX implementation works only in 256-bit registers, avoiding extract/combine ops with 128d intermediates.
+	// AVX implementation uses fast permutes
 	VectorRegister4Double R, T;
-
-	// (XX, YY, ZZ, WW)
-	R = VectorMultiply(Vec1, Vec2);
-	// (YY, XX, WW, ZZ)
-	T = VectorSwizzle(R, 1, 0, 3, 2); // fast in-lane permute (latency 1)
-	// (XX + YY, YY + XX, ZZ + WW, WW + ZZ)
-	R = VectorAdd(R, T);
-	// (ZZ + WW, WW + ZZ, XX + YY, YY + XX)
-	T = VectorSwizzle(R, 2, 3, 0, 1); // cross-lane permute (latency 3)
-	// (XX + YY + ZZ + WW, YY + XX + WW + ZZ, ZZ + WW + XX + YY, WW + ZZ + YY + XX)
-	return VectorAdd(R, T);
+	R = VectorMultiply(Vec1, Vec2);		// (XX, YY, ZZ, WW)
+	T = VectorSwizzle(R, 1, 0, 3, 2);	// (YY, XX, WW, ZZ) // fast in-lane permute
+	R = VectorAdd(R, T);				// (XX + YY, YY + XX, ZZ + WW, WW + ZZ)
+	T = VectorSwizzle(R, 2, 3, 0, 1);	// (ZZ + WW, WW + ZZ, XX + YY, YY + XX) // lane-swap permute
+	return VectorAdd(R, T);				// (XX + YY + ZZ + WW, YY + XX + WW + ZZ, ZZ + WW + XX + YY, WW + ZZ + YY + XX)
 #else
 	VectorRegister2Double T, A;
 
@@ -1646,8 +1821,6 @@ FORCEINLINE VectorRegister4Float VectorCross( const VectorRegister4Float& Vec1,
 
 FORCEINLINE VectorRegister4Double VectorCross(const VectorRegister4Double& Vec1, const VectorRegister4Double& Vec2)
 {
-	// TODO: AVX?
-
 	// YZX
 	VectorRegister4Double A = VectorSwizzle(Vec2, 1, 2, 0, 3);
 	VectorRegister4Double B = VectorSwizzle(Vec1, 1, 2, 0, 3);
