Back out CsvProfiler changes: Add SetAndHold command to csv profiler custom stats to allow infrequently updated statistics to have their value persisted.

diff --git a/Engine/Source/Runtime/Core/Private/ProfilingDebugging/CsvProfiler.cpp b/Engine/Source/Runtime/Core/Private/ProfilingDebugging/CsvProfiler.cpp
--- a/Engine/Source/Runtime/Core/Private/ProfilingDebugging/CsvProfiler.cpp
+++ b/Engine/Source/Runtime/Core/Private/ProfilingDebugging/CsvProfiler.cpp
@@ -1397,7 +1397,7 @@ struct FCsvStatSeries
 		CustomStatFloat
 	};
 
-	FCsvStatSeries(EType InSeriesType, const FCsvStatID& InStatID, FCsvStreamWriter* InWriter, FCsvStatRegister& StatRegister, const FString& ThreadName, uint32 ThreadId);
+	FCsvStatSeries(EType InSeriesType, const FCsvStatID& InStatID, FCsvStreamWriter* InWriter, FCsvStatRegister& StatRegister, const FString& ThreadName);
 	void FlushIfDirty();
 
 	void SetTimerValue(uint32 DataFrameNumber, uint64 ElapsedCycles)
@@ -1432,7 +1432,6 @@ struct FCsvStatSeries
 				FlushIfDirty();
 			}
 
-			bHold = Op == ECsvCustomStatOp::SetAndHold;
 			// The first op in a frame is always a set. Otherwise min/max don't work
 			Op = ECsvCustomStatOp::Set;
 			CurrentWriteFrameNumber = DataFrameNumber;
@@ -1472,7 +1471,6 @@ struct FCsvStatSeries
 				FlushIfDirty();
 			}
 
-			bHold = Op == ECsvCustomStatOp::SetAndHold;
 			// The first op in a frame is always a set. Otherwise min/max don't work
 			Op = ECsvCustomStatOp::Set;
 			CurrentWriteFrameNumber = DataFrameNumber;
@@ -1521,10 +1519,8 @@ struct FCsvStatSeries
 	FCsvStreamWriter* Writer;
 
 	int32 ColumnIndex;
-	uint32 ThreadId;
 
 	bool bDirty;
-	bool bHold;
 };
 
 struct FCsvProcessThreadDataStats
@@ -1544,8 +1540,7 @@ class FCsvStreamWriter
 {
 	struct FCsvRow
 	{
-		TSparseArray<FCsvStatSeriesValue> Values;
-		TBitArray<> ColumnsToHold;
+		TArray<FCsvStatSeriesValue> Values;
 		TArray<FCsvProcessedEvent> Events;
 
 		inline uint64 GetAllocatedSize() const
@@ -1580,14 +1575,13 @@ class FCsvStreamWriter
 	uint32 RenderThreadId;
 	uint32 RHIThreadId;
 
-	TMap<FString, FCsvStatSeries*> SeriesByName;
 public:
 	FCsvStreamWriter(const TSharedRef<FArchive>& InOutputFile, bool bInContinuousWrites, int32 InBufferSize, bool bInCompressOutput, uint32 RenderThreadId, uint32 RHIThreadId);
 	~FCsvStreamWriter();
 
 	void AddSeries(FCsvStatSeries* Series);
 
-	void PushValue(FCsvStatSeries* Series, int64 FrameNumber, const FCsvStatSeriesValue& Value, bool bHold=false, bool bAdvanceWriteFrame=true);
+	void PushValue(FCsvStatSeries* Series, int64 FrameNumber, const FCsvStatSeriesValue& Value);
 	void PushEvent(const FCsvProcessedEvent& Event);
 
 	void FinalizeNextRow();
@@ -1598,15 +1592,13 @@ public:
 	inline uint64 GetAllocatedSize() const;
 };
 
-FCsvStatSeries::FCsvStatSeries(EType InSeriesType, const FCsvStatID& InStatID, FCsvStreamWriter* InWriter, FCsvStatRegister& StatRegister, const FString& ThreadName, uint32 InThreadId)
+FCsvStatSeries::FCsvStatSeries(EType InSeriesType, const FCsvStatID& InStatID, FCsvStreamWriter* InWriter, FCsvStatRegister& StatRegister, const FString& ThreadName)
 	: StatID(InStatID)
 	, SeriesType(InSeriesType)
 	, CurrentWriteFrameNumber(-1)
 	, Writer(InWriter)
 	, ColumnIndex(-1)
-	, ThreadId(InThreadId)
 	, bDirty(false)
-	, bHold(false)
 {
 	CurrentValue.AsTimerCycles = 0;
 
@@ -1653,7 +1645,7 @@ void FCsvStatSeries::FlushIfDirty()
 			Value.Value.AsFloat = CurrentValue.AsFloatValue;
 			break;
 		}
-		Writer->PushValue(this, CurrentWriteFrameNumber, Value, bHold);
+		Writer->PushValue(this, CurrentWriteFrameNumber, Value);
 		CurrentValue.AsTimerCycles = 0;
 		bDirty = false;
 	}
@@ -1980,7 +1972,7 @@ private:
 		}
 		if (StatSeriesArray[StatIndex] == nullptr)
 		{
-			Series = new FCsvStatSeries(SeriesType, StatIndex, Writer, StatRegister, ThreadData->ThreadName, ThreadData->ThreadId);
+			Series = new FCsvStatSeries(SeriesType, StatIndex, Writer, StatRegister, ThreadData->ThreadName);
 			StatSeriesArray[StatIndex] = Series;
 		}
 		else
@@ -2021,30 +2013,21 @@ void FCsvStreamWriter::AddSeries(FCsvStatSeries* Series)
 	AllSeries.Add(Series);
 }
 
-void FCsvStreamWriter::PushValue(FCsvStatSeries* Series, int64 FrameNumber, const FCsvStatSeriesValue& Value, bool bHold, bool bAdvanceWriteFrame)
+void FCsvStreamWriter::PushValue(FCsvStatSeries* Series, int64 FrameNumber, const FCsvStatSeriesValue& Value)
 {
 	check(Series->ColumnIndex != -1);
 
-	if (bAdvanceWriteFrame)
-	{
-		WriteFrameIndex = FMath::Max(FrameNumber, WriteFrameIndex);
-	}
+	WriteFrameIndex = FMath::Max(FrameNumber, WriteFrameIndex);
 
 	FCsvRow& Row = Rows.FindOrAdd(FrameNumber);
-	if (!Row.Values.IsValidIndex(Series->ColumnIndex))
-	{
-		Row.Values.Insert(Series->ColumnIndex, Value);
-	}
-	else
-	{
-		Row.Values[Series->ColumnIndex] = Value;
-	}
 
-	if (Row.ColumnsToHold.Num() < AllSeries.Num())
+	// Ensure the row is large enough to hold every series
+	if (Row.Values.Num() < AllSeries.Num())
 	{
-		Row.ColumnsToHold.Add(false, AllSeries.Num() - Row.ColumnsToHold.Num());
+		Row.Values.SetNumZeroed(AllSeries.Num(), false);
 	}
-	Row.ColumnsToHold[Series->ColumnIndex] = bHold;
+
+	Row.Values[Series->ColumnIndex] = Value;
 }
 
 void FCsvStreamWriter::PushEvent(const FCsvProcessedEvent& Event)
@@ -2072,7 +2055,6 @@ void FCsvStreamWriter::FinalizeNextRow()
 
 	// Don't remove yet. Flushing series may modify this row
 	FCsvRow* Row = Rows.Find(ReadFrameIndex);
-	FCsvRow* NextRow = Rows.Find(ReadFrameIndex + 1);
 	if (Row)
 	{
 		if (Row->Events.Num() > 0)
@@ -2093,12 +2075,6 @@ void FCsvStreamWriter::FinalizeNextRow()
 			Stream.WriteEmptyString();
 		}
 
-		// Simplify conditional for ColumnsToHold 
-		if (Row->ColumnsToHold.Num() < AllSeries.Num())
-		{
-			Row->ColumnsToHold.Add(false, AllSeries.Num() - Row->ColumnsToHold.Num());
-		}
-
 		for (FCsvStatSeries* Series : AllSeries)
 		{
 			// Stat values are held in the series until a new value arrives.
@@ -2118,15 +2094,6 @@ void FCsvStreamWriter::FinalizeNextRow()
 				{
 					Stream.WriteValue(Value.Value.AsFloat);
 				}
-
-				if (Row->ColumnsToHold[Series->ColumnIndex])
-				{
-					// Propagate this value to the next row if it's not set already by data from a newer frame 
-					if (NextRow == nullptr || !NextRow->Values.IsValidIndex(Series->ColumnIndex))
-					{						
-						PushValue(Series, ReadFrameIndex + 1, Value, true, false);
-					}
-				}
 			}
 			else
 			{
diff --git a/Engine/Source/Runtime/Core/Public/ProfilingDebugging/CsvProfiler.h b/Engine/Source/Runtime/Core/Public/ProfilingDebugging/CsvProfiler.h
--- a/Engine/Source/Runtime/Core/Public/ProfilingDebugging/CsvProfiler.h
+++ b/Engine/Source/Runtime/Core/Public/ProfilingDebugging/CsvProfiler.h
@@ -146,7 +146,6 @@ class FName;
 enum class ECsvCustomStatOp : uint8
 {
 	Set,
-	SetAndHold,
 	Min,
 	Max,
 	Accumulate,
