Don't use SetAndHold in csv profiler which will be backed out, register game thread ticks to push data to the csv profiler each frame.

diff --git a/Engine/Source/Runtime/Core/Private/IO/IoDispatcher.cpp b/Engine/Source/Runtime/Core/Private/IO/IoDispatcher.cpp
--- a/Engine/Source/Runtime/Core/Private/IO/IoDispatcher.cpp
+++ b/Engine/Source/Runtime/Core/Private/IO/IoDispatcher.cpp
@@ -16,6 +16,7 @@
 #include "Async/MappedFileHandle.h"
 #include "ProfilingDebugging/CountersTrace.h"
 #include "ProfilingDebugging/CsvProfiler.h"
+#include "Containers/Ticker.h"
 #include "IO/IoDispatcherBackend.h"
 #include "IO/IoDispatcherFileBackendTypes.h"
 
@@ -153,11 +154,16 @@ public:
 			RequestAllocator.Trim();
 			BatchAllocator.Trim();
 		});
+#if CSV_PROFILER
+		TickerHandle = FTicker::GetCoreTicker().AddTicker(FTickerDelegate::CreateRaw(this, &FIoDispatcherImpl::TickCsv));
+#endif
 	}
 
 	~FIoDispatcherImpl()
 	{
-
+#if CSV_PROFILER
+		FTicker::GetCoreTicker().RemoveTicker(TickerHandle);
+#endif
 		delete Thread;
 		FPlatformProcess::ReturnSynchEventToPool(DispatcherEvent);
 	}
@@ -451,7 +457,6 @@ private:
 				--PendingIoRequestsCount;
 				TRACE_COUNTER_SET(PendingIoRequests, PendingIoRequestsCount);
 			}
-			CSV_CUSTOM_STAT_DEFINED(PendingIoRequests, (int32)PendingIoRequestsCount, ECsvCustomStatOp::SetAndHold);
 		}
 	}
 
@@ -609,9 +614,7 @@ private:
 			}
 			
 			++PendingIoRequestsCount;
-			TRACE_COUNTER_SET(PendingIoRequests, PendingIoRequestsCount);
-			CSV_CUSTOM_STAT_DEFINED(PendingIoRequests, (int32)PendingIoRequestsCount, ECsvCustomStatOp::SetAndHold);
-			
+			TRACE_COUNTER_SET(PendingIoRequests, PendingIoRequestsCount);	
 			ProcessCompletedRequests();
 		}
 	}
@@ -647,6 +650,14 @@ private:
 		DispatcherEvent->Trigger();
 	}
 
+#if CSV_PROFILER
+	bool TickCsv(float DeltaTime)
+	{
+		CSV_CUSTOM_STAT_DEFINED(PendingIoRequests, (int32)PendingIoRequestsCount, ECsvCustomStatOp::Set);
+		return true; // Keep ticking
+	}
+#endif
+
 	using FRequestAllocator = TBlockAllocator<FIoRequestImpl, 4096>;
 	using FBatchAllocator = TBlockAllocator<FIoBatchImpl, 4096>;
 
@@ -666,6 +677,9 @@ private:
 	FIoDispatcher::FIoContainerUnmountedEvent ContainerUnmountedEvent;
 	uint64 PendingIoRequestsCount = 0;
 	int64 TotalLoaded = 0;
+#if CSV_PROFILER
+	FDelegateHandle TickerHandle;
+#endif
 };
 
 FIoDispatcher::FIoDispatcher()
diff --git a/Engine/Source/Runtime/Core/Private/IO/IoDispatcherFileBackend.cpp b/Engine/Source/Runtime/Core/Private/IO/IoDispatcherFileBackend.cpp
--- a/Engine/Source/Runtime/Core/Private/IO/IoDispatcherFileBackend.cpp
+++ b/Engine/Source/Runtime/Core/Private/IO/IoDispatcherFileBackend.cpp
@@ -16,6 +16,7 @@
 #include "Algo/IsSorted.h"
 #include "Algo/MinElement.h"
 #include "Templates/Greater.h"
+#include "Containers/Ticker.h"
 
 TRACE_DECLARE_INT_COUNTER(IoDispatcherLatencyCircuitBreaks, TEXT("IoDispatcher/LatencyCircuitBreaks"));
 TRACE_DECLARE_INT_COUNTER(IoDispatcherSeekDistanceCircuitBreaks, TEXT("IoDispatcher/SeekDistanceCircuitBreaks"));
@@ -1780,11 +1781,13 @@ void FFileIoStore::UpdateAsyncIOMinimumPriority()
 
 bool FFileIoStore::Init()
 {
+	FFileIoStats::Init();
 	return true;
 }
 
 void FFileIoStore::Stop()
 {
+	FFileIoStats::Shutdown();
 	bStopRequested = true;
 	EventQueue.ServiceNotify();
 }
@@ -1898,7 +1901,37 @@ uint64 FFileIoStats::LastOffset = 0;
 
 uint32 FFileIoStats::FileIoStoreThreadId = 0;
 uint32 FFileIoStats::IoDispatcherThreadId = 0;
+
+FDelegateHandle FFileIoStats::TickerHandle;
+#endif
+
+void FFileIoStats::Init()
+{
+#if CSV_PROFILER
+	TickerHandle = FTicker::GetCoreTicker().AddTicker(FTickerDelegate::CreateStatic(&FFileIoStats::CsvTick));
+#endif
+}
+
+void FFileIoStats::Shutdown()
+{
+#if CSV_PROFILER
+	FTicker::GetCoreTicker().RemoveTicker(TickerHandle);
 #endif
+}
+
+bool FFileIoStats::CsvTick(float DeltaTime)
+{
+#if CSV_PROFILER
+	CSV_CUSTOM_STAT_DEFINED(QueuedFilesystemReadMB, BytesToApproxMB(QueuedFilesystemReadBytes), ECsvCustomStatOp::Set);
+	CSV_CUSTOM_STAT_DEFINED(QueuedFilesystemReads, (int32)QueuedFilesystemReads, ECsvCustomStatOp::Set);
+
+	CSV_CUSTOM_STAT_DEFINED(QueuedUncompressBlocks, int32(QueuedUncompressBlocks), ECsvCustomStatOp::Set);
+	CSV_CUSTOM_STAT_DEFINED(QueuedUncompressInMB, BytesToApproxMB(QueuedUncompressBytesIn), ECsvCustomStatOp::Set);
+	CSV_CUSTOM_STAT_DEFINED(QueuedUncompressOutMB, BytesToApproxMB(QueuedUncompressBytesOut), ECsvCustomStatOp::Set);
+#endif
+
+	return true;
+}
 
 bool FFileIoStats::IsInIoDispatcherThread()
 {
@@ -1966,11 +1999,8 @@ void FFileIoStats::OnFilesystemReadsQueued(uint64 BytesToRead, int32 NumReads)
 	checkSlow(IsInIoDispatcherThread());
 
 #if CSV_PROFILER
-	uint64 LocalQueuedFilesystemReadBytes = QueuedFilesystemReadBytes += BytesToRead;
-	CSV_CUSTOM_STAT_DEFINED(QueuedFilesystemReadMB, BytesToApproxMB(LocalQueuedFilesystemReadBytes), ECsvCustomStatOp::SetAndHold);
-
-	uint64 LocalQueuedFilesystemReads = QueuedFilesystemReads += NumReads;
-	CSV_CUSTOM_STAT_DEFINED(QueuedFilesystemReads, (int32)LocalQueuedFilesystemReads, ECsvCustomStatOp::SetAndHold);
+	QueuedFilesystemReadBytes += BytesToRead;
+	QueuedFilesystemReads += NumReads;
 #endif
 
 	TRACE_COUNTER_ADD(IoDispatcherFileBackendQueuedFilesystemReadBytes, BytesToRead);
@@ -2032,11 +2062,8 @@ void FFileIoStats::OnReadComplete(int64 BytesRead)
 #if IO_DISPATCHER_FILE_STATS
 	checkSlow(IsInIoDispatcherThread());
 #if CSV_PROFILER
-	uint64 LocalQueuedFilesystemReadBytes = QueuedFilesystemReadBytes -= BytesRead;
-	CSV_CUSTOM_STAT_DEFINED(QueuedFilesystemReadMB, BytesToApproxMB(LocalQueuedFilesystemReadBytes), ECsvCustomStatOp::SetAndHold);
-
-	uint64 LocalQueuedFilesystemReads = --QueuedFilesystemReads;
-	CSV_CUSTOM_STAT_DEFINED(QueuedFilesystemReads, (int32)LocalQueuedFilesystemReads, ECsvCustomStatOp::SetAndHold);
+	QueuedFilesystemReadBytes -= BytesRead;
+	--QueuedFilesystemReads;
 #endif
 
 	CSV_CUSTOM_STAT_DEFINED(FrameBytesReadKB, BytesToApproxKB(BytesRead), ECsvCustomStatOp::Accumulate);
@@ -2051,13 +2078,9 @@ void FFileIoStats::OnDecompressQueued(int64 CompressedBytes, int64 UncompressedB
 #if IO_DISPATCHER_FILE_STATS
 	checkSlow(IsInIoDispatcherThread());
 #if CSV_PROFILER
-	uint64 QueuedBlocks = ++QueuedUncompressBlocks;
-	CSV_CUSTOM_STAT_DEFINED(QueuedUncompressBlocks, int32(QueuedBlocks), ECsvCustomStatOp::SetAndHold);
-
-	uint64 QueuedBytesIn = QueuedUncompressBytesIn += CompressedBytes;
-	uint64 QueuedBytesOut = QueuedUncompressBytesOut += UncompressedBytes;
-	CSV_CUSTOM_STAT_DEFINED(QueuedUncompressInMB, BytesToApproxMB(QueuedBytesIn), ECsvCustomStatOp::SetAndHold);
-	CSV_CUSTOM_STAT_DEFINED(QueuedUncompressOutMB, BytesToApproxMB(QueuedBytesOut), ECsvCustomStatOp::SetAndHold);
+	++QueuedUncompressBlocks;
+	QueuedUncompressBytesIn += CompressedBytes;
+	QueuedUncompressBytesOut += UncompressedBytes;
 #endif
 #endif
 }
@@ -2067,13 +2090,9 @@ void FFileIoStats::OnDecompressComplete(int64 CompressedBytes, int64 Uncompresse
 #if IO_DISPATCHER_FILE_STATS
 	checkSlow(IsInIoDispatcherThread());
 #if CSV_PROFILER
-	uint64 QueuedBlocks = --QueuedUncompressBlocks;
-	CSV_CUSTOM_STAT_DEFINED(QueuedUncompressBlocks, int32(QueuedBlocks), ECsvCustomStatOp::SetAndHold);
-
-	uint64 QueuedBytesIn = QueuedUncompressBytesIn -= CompressedBytes;
-	uint64 QueuedBytesOut = QueuedUncompressBytesOut -= UncompressedBytes;
-	CSV_CUSTOM_STAT_DEFINED(QueuedUncompressInMB, BytesToApproxMB(QueuedBytesIn), ECsvCustomStatOp::SetAndHold);
-	CSV_CUSTOM_STAT_DEFINED(QueuedUncompressOutMB, BytesToApproxMB(QueuedBytesOut), ECsvCustomStatOp::SetAndHold);
+	--QueuedUncompressBlocks;
+	QueuedUncompressBytesIn -= CompressedBytes;
+	QueuedUncompressBytesOut -= UncompressedBytes;
 #endif
 
 	CSV_CUSTOM_STAT_DEFINED(FrameBytesUncompressedInKB, BytesToApproxKB(CompressedBytes), ECsvCustomStatOp::Accumulate);
diff --git a/Engine/Source/Runtime/Core/Private/IO/IoDispatcherFileBackendTypes.h b/Engine/Source/Runtime/Core/Private/IO/IoDispatcherFileBackendTypes.h
--- a/Engine/Source/Runtime/Core/Private/IO/IoDispatcherFileBackendTypes.h
+++ b/Engine/Source/Runtime/Core/Private/IO/IoDispatcherFileBackendTypes.h
@@ -862,6 +862,8 @@ struct FFileIoStats
 
 	static uint32 FileIoStoreThreadId;
 	static uint32 IoDispatcherThreadId;
+
+	static FDelegateHandle TickerHandle;
 #endif
 
 	static float BytesToApproxMB(uint64 Bytes) {
@@ -871,6 +873,10 @@ struct FFileIoStats
 		return float(double(Bytes) / 1024.0);
 	}
 
+	static void Init();
+	static void Shutdown();
+	static bool CsvTick(float DeltaTime);
+
 	static bool IsInIoDispatcherThread();
 	static bool IsInFileIoStoreThread();
 	static void SetDispatcherThreadId(uint32 InIoDispatcherThreadId);
diff --git a/Engine/Source/Runtime/CoreUObject/Private/Serialization/AsyncLoading2.cpp b/Engine/Source/Runtime/CoreUObject/Private/Serialization/AsyncLoading2.cpp
--- a/Engine/Source/Runtime/CoreUObject/Private/Serialization/AsyncLoading2.cpp
+++ b/Engine/Source/Runtime/CoreUObject/Private/Serialization/AsyncLoading2.cpp
@@ -2600,7 +2600,6 @@ void FAsyncLoadingThread2::BundleIoRequestCompleted(FAsyncPackage2* Package)
 	check(PendingBundleIoRequestsTotalSize >= Package->ExportBundlesSize)
 	PendingBundleIoRequestsTotalSize -= Package->ExportBundlesSize;
 	TRACE_COUNTER_SET(AsyncLoadingPendingIoRequestsSize, PendingBundleIoRequestsTotalSize);
-	CSV_CUSTOM_STAT_DEFINED(PendingIoRequestsSizeMB, float((double)PendingBundleIoRequestsTotalSize / 1024.0 / 1024.0), ECsvCustomStatOp::SetAndHold);
 	CSV_CUSTOM_STAT_DEFINED(FrameCompletedBundleLoadsKB, float((double)Package->ExportBundlesSize / 1024.0), ECsvCustomStatOp::Accumulate);
 	if (WaitingIoRequests.Num())
 	{
@@ -2649,7 +2648,6 @@ void FAsyncLoadingThread2::StartBundleIoRequests()
 	IoBatch.Issue();
 
 	TRACE_COUNTER_SET(AsyncLoadingPendingIoRequestsSize, PendingBundleIoRequestsTotalSize);
-	CSV_CUSTOM_STAT_DEFINED(PendingIoRequestsSizeMB, float((double)PendingBundleIoRequestsTotalSize / 1024.0 / 1024.0), ECsvCustomStatOp::SetAndHold);
 }
 
 FEventLoadNode2::FEventLoadNode2(const FAsyncLoadEventSpec* InSpec, FAsyncPackage2* InPackage, int32 InImportOrExportIndex, int32 InBarrierCount)
@@ -5642,6 +5640,7 @@ EAsyncPackageState::Type FAsyncLoadingThread2::ProcessLoadingFromGameThread(FAsy
 	// CSV_CUSTOM_STAT(FileIO, EDLEventQueueDepth, (int32)GraphAllocator.TotalNodeCount, ECsvCustomStatOp::Set);
 	CSV_CUSTOM_STAT(FileIO, QueuedPackagesQueueDepth, GetNumQueuedPackages(), ECsvCustomStatOp::Set);
 	CSV_CUSTOM_STAT(FileIO, ExistingQueuedPackagesQueueDepth, GetNumAsyncPackages(), ECsvCustomStatOp::Set);
+	CSV_CUSTOM_STAT_DEFINED(PendingIoRequestsSizeMB, float((double)PendingBundleIoRequestsTotalSize / 1024.0 / 1024.0), ECsvCustomStatOp::Set);
 
 	TickAsyncLoadingFromGameThread(ThreadState, bUseTimeLimit, bUseFullTimeLimit, TimeLimit);
 	return IsAsyncLoading() ? EAsyncPackageState::TimeOut : EAsyncPackageState::Complete;
