Add Csv metadata for DeviceID, LoginID. Cleanup metadata setters in CsvProfiler - handle permanent metadata once only; always use internal method

diff --git a/Engine/Source/Runtime/Core/Private/ProfilingDebugging/CsvProfiler.cpp b/Engine/Source/Runtime/Core/Private/ProfilingDebugging/CsvProfiler.cpp
--- a/Engine/Source/Runtime/Core/Private/ProfilingDebugging/CsvProfiler.cpp
+++ b/Engine/Source/Runtime/Core/Private/ProfilingDebugging/CsvProfiler.cpp
@@ -2518,6 +2518,20 @@ FCsvProfiler::FCsvProfiler()
 	SetMetadataInternal(TEXT("Commandline"), *CommandlineStr, false);
 	SetMetadataInternal(TEXT("OS"), *OSString);
 	SetMetadataInternal(TEXT("CPU"), *FPlatformMisc::GetDeviceMakeAndModel());
+	SetMetadataInternal(TEXT("PGOEnabled"), FPlatformMisc::IsPGOEnabled() ? TEXT("1") : TEXT("0"));
+	SetMetadataInternal(TEXT("LoginID"), *FPlatformMisc::GetLoginId());
+
+	// Set the device ID if the platform supports it
+	FString DeviceID = FPlatformMisc::GetDeviceId();
+	if (!DeviceID.IsEmpty())
+	{
+		SetMetadataInternal(TEXT("DeviceID"), *DeviceID);
+	}
+#if !UE_BUILD_SHIPPING
+	int32 ExtraDevelopmentMemoryMB = (int32)(FPlatformMemory::GetExtraDevelopmentMemorySize() / 1024ull / 1024ull);
+	SetMetadataInternal(TEXT("ExtraDevelopmentMemoryMB"), *FString::FromInt(ExtraDevelopmentMemoryMB));
+#endif
+
 }
 
 FCsvProfiler::~FCsvProfiler()
@@ -2644,7 +2658,7 @@ void FCsvProfiler::BeginFrame()
 					 
 					// Set the CSV ID and mirror it to the log
 					FString CsvId = FGuid::NewGuid().ToString();
-					SetMetadata(TEXT("CsvID"), *CsvId);
+					SetMetadataInternal(TEXT("CsvID"), *CsvId);
 					UE_LOG(LogCsvProfiler, Display, TEXT("Capture started. CSV ID: %s"), *CsvId);
 
 					// Figure out the target framerate
@@ -2659,19 +2673,9 @@ void FCsvProfiler::BeginFrame()
 					{
 						TargetFPS = FMath::Min(TargetFPS, FPlatformMisc::GetMaxRefreshRate() / SyncIntervalCVar->GetInt());
 					}
-					SetMetadata(TEXT("TargetFramerate"), *FString::FromInt(TargetFPS));
-
-#if !UE_BUILD_SHIPPING
-					int32 ExtraDevelopmentMemoryMB = (int32)(FPlatformMemory::GetExtraDevelopmentMemorySize()/1024ull/1024ull);
-					SetMetadata(TEXT("ExtraDevelopmentMemoryMB"), *FString::FromInt(ExtraDevelopmentMemoryMB)); 
-#endif
-
-					SetMetadata(TEXT("PGOEnabled"), FPlatformMisc::IsPGOEnabled() ? TEXT("1") : TEXT("0"));
-
-					// Output the current time in seconds seconds since the Unix Epoch
-					SetMetadata(TEXT("StartTimestamp"), *FString::Printf(TEXT("%lld"), FDateTime::UtcNow().ToUnixTimestamp()));
-
-					SetMetadata(TEXT("NamedEvents"), (GCycleStatsShouldEmitNamedEvents > 0) ? TEXT("1") : TEXT("0"));
+					SetMetadataInternal(TEXT("TargetFramerate"), *FString::FromInt(TargetFPS));
+					SetMetadataInternal(TEXT("StartTimestamp"), *FString::Printf(TEXT("%lld"), FDateTime::UtcNow().ToUnixTimestamp()));
+					SetMetadataInternal(TEXT("NamedEvents"), (GCycleStatsShouldEmitNamedEvents > 0) ? TEXT("1") : TEXT("0"));
 
 					GCsvStatCounts = !!CVarCsvStatCounts.GetValueOnGameThread();
 
