CompactBinary: Added validation to LoadCompactBinary and removed asserts from the other load functions

#rnx

diff --git a/Engine/Source/Developer/TraceServices/Private/AnalysisCache.cpp b/Engine/Source/Developer/TraceServices/Private/AnalysisCache.cpp
--- a/Engine/Source/Developer/TraceServices/Private/AnalysisCache.cpp
+++ b/Engine/Source/Developer/TraceServices/Private/AnalysisCache.cpp
@@ -174,7 +174,10 @@ bool FAnalysisCache::FFileContents::Load()
 	FMemoryReaderView Ar(MakeArrayView<uint8>((uint8*)Buffer.GetData(), Buffer.GetSize()));
 	
 	FCbPackage Package;
-	Package.Load(Ar);
+	if (!Package.TryLoad(Ar))
+	{
+		return false;
+	}
 
 	uint32 PackageVersion = Package.GetObject().Find("Version").AsUInt32();
 	UE_LOG(LogAnalysisCache, Display, TEXT("Cache file (version %u) loaded."), PackageVersion);
diff --git a/Engine/Source/Runtime/Core/Private/Serialization/CompactBinaryPackage.cpp b/Engine/Source/Runtime/Core/Private/Serialization/CompactBinaryPackage.cpp
--- a/Engine/Source/Runtime/Core/Private/Serialization/CompactBinaryPackage.cpp
+++ b/Engine/Source/Runtime/Core/Private/Serialization/CompactBinaryPackage.cpp
@@ -45,10 +45,13 @@ FCbAttachment::FCbAttachment(FSharedBuffer InBuffer, const FIoHash* const InHash
 	}
 }
 
-void FCbAttachment::Load(FCbFieldIterator& Fields)
+bool FCbAttachment::TryLoad(FCbFieldIterator& Fields)
 {
-	checkf(Fields.IsBinary(), TEXT("Attachments must start with a binary field."));
 	const FMemoryView View = Fields.AsBinaryView();
+	if (Fields.HasError())
+	{
+		return false;
+	}
 	if (View.GetSize() > 0)
 	{
 		Buffer = FSharedBuffer::MakeView(View, Fields.GetOuterBuffer()).MakeOwned();
@@ -56,7 +59,10 @@ void FCbAttachment::Load(FCbFieldIterator& Fields)
 		++Fields;
 
 		Hash = Fields.AsAttachment();
-		checkf(!Fields.HasError(), TEXT("Attachments must be a non-empty binary value with a content hash."));
+		if (Fields.HasError())
+		{
+			return false;
+		}
 		if (Fields.IsObjectAttachment())
 		{
 			Object = FCbFieldView(Buffer.GetData());
@@ -70,13 +76,18 @@ void FCbAttachment::Load(FCbFieldIterator& Fields)
 		Object = FCbFieldView();
 		Hash.Reset();
 	}
+	return true;
 }
 
-void FCbAttachment::Load(FArchive& Ar, FCbBufferAllocator Allocator)
+bool FCbAttachment::TryLoad(FArchive& Ar, FCbBufferAllocator Allocator)
 {
 	FCbField BufferField = LoadCompactBinary(Ar, Allocator);
-	checkf(BufferField.IsBinary(), TEXT("Attachments must start with a binary field."));
 	const FMemoryView View = BufferField.AsBinaryView();
+	if (BufferField.HasError())
+	{
+		Ar.SetError();
+		return false;
+	}
 	if (View.GetSize() > 0)
 	{
 		Buffer = FSharedBuffer::MakeView(View, BufferField.GetOuterBuffer()).MakeOwned();
@@ -90,7 +101,11 @@ void FCbAttachment::Load(FArchive& Ar, FCbBufferAllocator Allocator)
 				return FUniqueBuffer::MakeView(HashBuffer.GetData(), Size);
 			});
 		Hash = HashField.AsAttachment();
-		checkf(!HashField.HasError(), TEXT("Attachments must be a non-empty binary value with a content hash."));
+		if (HashField.HasError() || FIoHash::HashBuffer(Buffer) != Hash)
+		{
+			Ar.SetError();
+			return false;
+		}
 		if (HashField.IsObjectAttachment())
 		{
 			Object = FCbFieldView(Buffer.GetData());
@@ -102,6 +117,7 @@ void FCbAttachment::Load(FArchive& Ar, FCbBufferAllocator Allocator)
 		Object = FCbFieldView();
 		Hash.Reset();
 	}
+	return true;
 }
 
 void FCbAttachment::Save(FCbWriter& Writer) const
@@ -227,7 +243,7 @@ void FCbPackage::GatherAttachments(const FCbObject& Value, FAttachmentResolver R
 		});
 }
 
-void FCbPackage::Load(FCbFieldIterator& Fields)
+bool FCbPackage::TryLoad(FCbFieldIterator& Fields)
 {
 	*this = FCbPackage();
 	while (Fields)
@@ -240,19 +256,25 @@ void FCbPackage::Load(FCbFieldIterator& Fields)
 		else if (Fields.IsBinary())
 		{
 			FCbAttachment Attachment;
-			Attachment.Load(Fields);
+			Attachment.TryLoad(Fields);
 			AddAttachment(Attachment);
 		}
 		else
 		{
-			checkf(Fields.IsObject(), TEXT("Expected Object, Binary, or Null field when loading a package."));
 			Object = Fields.AsObject();
+			if (Fields.HasError())
+			{
+				return false;
+			}
 			Object.MakeOwned();
 			++Fields;
 			if (Object)
 			{
 				ObjectHash = Fields.AsObjectAttachment();
-				checkf(!Fields.HasError(), TEXT("Object must be followed by an ObjectAttachment with the object hash."));
+				if (Fields.HasError())
+				{
+					return false;
+				}
 				++Fields;
 			}
 			else
@@ -261,9 +283,10 @@ void FCbPackage::Load(FCbFieldIterator& Fields)
 			}
 		}
 	}
+	return true;
 }
 
-void FCbPackage::Load(FArchive& Ar, FCbBufferAllocator Allocator)
+bool FCbPackage::TryLoad(FArchive& Ar, FCbBufferAllocator Allocator)
 {
 	uint8 StackBuffer[64];
 	const auto StackAllocator = [&Allocator, &StackBuffer](uint64 Size) -> FUniqueBuffer
@@ -275,9 +298,14 @@ void FCbPackage::Load(FArchive& Ar, FCbBufferAllocator Allocator)
 	for (;;)
 	{
 		FCbField ValueField = LoadCompactBinary(Ar, StackAllocator);
+		if (!ValueField)
+		{
+			Ar.SetError();
+			return false;
+		}
 		if (ValueField.IsNull())
 		{
-			break;
+			return true;
 		}
 		else if (ValueField.IsBinary())
 		{
@@ -287,7 +315,11 @@ void FCbPackage::Load(FArchive& Ar, FCbBufferAllocator Allocator)
 				FSharedBuffer Buffer = FSharedBuffer::MakeView(View, ValueField.GetOuterBuffer()).MakeOwned();
 				FCbField HashField = LoadCompactBinary(Ar, StackAllocator);
 				const FIoHash& Hash = HashField.AsAttachment();
-				checkf(!HashField.HasError(), TEXT("Attachments must be a non-empty binary value with a content hash."));
+				if (HashField.HasError() || FIoHash::HashBuffer(Buffer) != Hash)
+				{
+					Ar.SetError();
+					return false;
+				}
 				if (HashField.IsObjectAttachment())
 				{
 					AddAttachment(FCbAttachment(FCbObject(MoveTemp(Buffer)), Hash));
@@ -300,14 +332,22 @@ void FCbPackage::Load(FArchive& Ar, FCbBufferAllocator Allocator)
 		}
 		else
 		{
-			checkf(ValueField.IsObject(), TEXT("Expected Object, Binary, or Null field when loading a package."));
 			Object = ValueField.AsObject();
+			if (ValueField.HasError())
+			{
+				Ar.SetError();
+				return false;
+			}
 			Object.MakeOwned();
 			if (Object)
 			{
 				FCbField HashField = LoadCompactBinary(Ar, StackAllocator);
 				ObjectHash = HashField.AsObjectAttachment();
-				checkf(!HashField.HasError(), TEXT("Object must be followed by an ObjectAttachment with the object hash."));
+				if (HashField.HasError() || Object.GetHash() != ObjectHash)
+				{
+					Ar.SetError();
+					return false;
+				}
 			}
 			else
 			{
diff --git a/Engine/Source/Runtime/Core/Private/Serialization/CompactBinarySerialization.cpp b/Engine/Source/Runtime/Core/Private/Serialization/CompactBinarySerialization.cpp
--- a/Engine/Source/Runtime/Core/Private/Serialization/CompactBinarySerialization.cpp
+++ b/Engine/Source/Runtime/Core/Private/Serialization/CompactBinarySerialization.cpp
@@ -2,6 +2,7 @@
 
 #include "Serialization/CompactBinarySerialization.h"
 
+#include "Serialization/CompactBinaryValidation.h"
 #include "Serialization/VarInt.h"
 #include "Templates/IdentityFunctor.h"
 #include "Templates/Invoke.h"
@@ -190,6 +191,11 @@ FCbField LoadCompactBinary(FArchive& Ar, FCbBufferAllocator Allocator)
 	{
 		Ar.Serialize(View.GetData(), static_cast<int64>(View.GetSize()));
 	}
+	if (ValidateCompactBinary(Buffer, ECbValidateMode::Default) != ECbValidateError::None)
+	{
+		Ar.SetError();
+		return FCbField();
+	}
 	return FCbField(Buffer.MoveToShared());
 }
 
diff --git a/Engine/Source/Runtime/Core/Private/Tests/Serialization/CompactBinaryPackageTest.cpp b/Engine/Source/Runtime/Core/Private/Tests/Serialization/CompactBinaryPackageTest.cpp
--- a/Engine/Source/Runtime/Core/Private/Tests/Serialization/CompactBinaryPackageTest.cpp
+++ b/Engine/Source/Runtime/Core/Private/Tests/Serialization/CompactBinaryPackageTest.cpp
@@ -37,15 +37,15 @@ bool FCbAttachmentTest::RunTest(const FString& Parameters)
 			ValidateCompactBinaryAttachment(MakeMemoryView(WriteAr), ECbValidateMode::All), ECbValidateError::None);
 
 		FCbAttachment FromFields;
-		FromFields.Load(Fields);
-		TestTrue(FString::Printf(TEXT("FCbAttachment(%s).Load(Iterator)->AtEnd"), Test), !bool(Fields));
-		TestEqual(FString::Printf(TEXT("FCbAttachment(%s).Load(Iterator)->Equals"), Test), FromFields, Attachment);
+		FromFields.TryLoad(Fields);
+		TestTrue(FString::Printf(TEXT("FCbAttachment(%s).TryLoad(Iterator)->AtEnd"), Test), !bool(Fields));
+		TestEqual(FString::Printf(TEXT("FCbAttachment(%s).TryLoad(Iterator)->Equals"), Test), FromFields, Attachment);
 
 		FCbAttachment FromArchive;
 		FMemoryReader ReadAr(WriteAr);
-		FromArchive.Load(ReadAr);
-		TestTrue(FString::Printf(TEXT("FCbAttachment(%s).Load(Archive)->AtEnd"), Test), ReadAr.AtEnd());
-		TestEqual(FString::Printf(TEXT("FCbAttachment(%s).Load(Archive)->Equals"), Test), FromArchive, Attachment);
+		FromArchive.TryLoad(ReadAr);
+		TestTrue(FString::Printf(TEXT("FCbAttachment(%s).TryLoad(Archive)->AtEnd"), Test), ReadAr.AtEnd());
+		TestEqual(FString::Printf(TEXT("FCbAttachment(%s).TryLoad(Archive)->Equals"), Test), FromArchive, Attachment);
 	};
 
 	// Empty Attachment
@@ -137,7 +137,7 @@ bool FCbAttachmentTest::RunTest(const FString& Parameters)
 		FCbFieldIterator Fields = Writer.Save();
 		FCbFieldIterator FieldsView = FCbFieldIterator::MakeRangeView(FCbFieldViewIterator(Fields));
 
-		Attachment.Load(FieldsView);
+		Attachment.TryLoad(FieldsView);
 		TestFalse(TEXT("FCbAttachment(LoadBinaryView).IsNull()"), Attachment.IsNull());
 		TestTrue(TEXT("FCbAttachment(LoadBinaryView) as bool"), bool(Attachment));
 		TestFalse(TEXT("FCbAttachment(LoadBinaryView).AsBinary()->!InView"),
@@ -167,7 +167,7 @@ bool FCbAttachmentTest::RunTest(const FString& Parameters)
 		FCbFieldIterator Fields = Writer.Save();
 		FCbFieldIterator FieldsView = FCbFieldIterator::MakeRangeView(FCbFieldViewIterator(Fields));
 
-		Attachment.Load(FieldsView);
+		Attachment.TryLoad(FieldsView);
 		FMemoryView View;
 		TestFalse(TEXT("FCbAttachment(LoadObjectView).IsNull()"), Attachment.IsNull());
 		TestTrue(TEXT("FCbAttachment(LoadObjectView) as bool"), bool(Attachment));
@@ -229,15 +229,15 @@ bool FCbPackageTest::RunTest(const FString& Parameters)
 			ValidateCompactBinaryPackage(MakeMemoryView(WriteAr), ECbValidateMode::All), ECbValidateError::None);
 
 		FCbPackage FromFields;
-		FromFields.Load(Fields);
-		TestFalse(FString::Printf(TEXT("FCbPackage(%s).Load(Iterator)->AtEnd"), Test), bool(Fields));
-		TestEqual(FString::Printf(TEXT("FCbPackage(%s).Load(Iterator)->Equals"), Test), FromFields, Package);
+		FromFields.TryLoad(Fields);
+		TestFalse(FString::Printf(TEXT("FCbPackage(%s).TryLoad(Iterator)->AtEnd"), Test), bool(Fields));
+		TestEqual(FString::Printf(TEXT("FCbPackage(%s).TryLoad(Iterator)->Equals"), Test), FromFields, Package);
 
 		FCbPackage FromArchive;
 		FMemoryReader ReadAr(WriteAr);
-		FromArchive.Load(ReadAr);
-		TestTrue(FString::Printf(TEXT("FCbPackage(%s).Load(Archive)->AtEnd"), Test), ReadAr.AtEnd());
-		TestEqual(FString::Printf(TEXT("FCbPackage(%s).Load(Archive)->Equals"), Test), FromArchive, Package);
+		FromArchive.TryLoad(ReadAr);
+		TestTrue(FString::Printf(TEXT("FCbPackage(%s).TryLoad(Archive)->AtEnd"), Test), ReadAr.AtEnd());
+		TestEqual(FString::Printf(TEXT("FCbPackage(%s).TryLoad(Archive)->Equals"), Test), FromArchive, Package);
 	};
 
 	// Empty
@@ -456,43 +456,43 @@ bool FCbPackageTest::RunTest(const FString& Parameters)
 
 		FCbFieldIterator Fields = Writer.Save();
 		FCbPackage FromFields;
-		FromFields.Load(Fields);
+		FromFields.TryLoad(Fields);
 
 		const FCbAttachment* const Level2Attachment = FromFields.FindAttachment(Level2Hash);
 		const FCbAttachment* const Level3Attachment = FromFields.FindAttachment(Level3Hash);
 		const FCbAttachment* const Level4Attachment = FromFields.FindAttachment(Level4Hash);
 
-		TestTrue(TEXT("FCbPackage(OutOfOrder).Load()->Level1"), FromFields.GetObject().Equals(Level1));
-		TestEqual(TEXT("FCbPackage(OutOfOrder).Load()->Level1Buffer"),
+		TestTrue(TEXT("FCbPackage(OutOfOrder).TryLoad()->Level1"), FromFields.GetObject().Equals(Level1));
+		TestEqual(TEXT("FCbPackage(OutOfOrder).TryLoad()->Level1Buffer"),
 			FromFields.GetObject().GetOuterBuffer(), Fields.GetOuterBuffer());
-		TestEqual(TEXT("FCbPackage(OutOfOrder).Load()->Level1Hash"), FromFields.GetObjectHash(), Level1Hash);
+		TestEqual(TEXT("FCbPackage(OutOfOrder).TryLoad()->Level1Hash"), FromFields.GetObjectHash(), Level1Hash);
 
-		TestTrue(TEXT("FCbPackage(OutOfOrder).Load()->Level2"),
+		TestTrue(TEXT("FCbPackage(OutOfOrder).TryLoad()->Level2"),
 			Level2Attachment && Level2Attachment->AsObject().Equals(Level2));
-		TestTrue(TEXT("FCbPackage(OutOfOrder).Load()->Level2Buffer"),
+		TestTrue(TEXT("FCbPackage(OutOfOrder).TryLoad()->Level2Buffer"),
 			Level2Attachment && Fields.GetOuterBuffer().GetView().Contains(Level2Attachment->AsBinary().GetView()));
-		TestTrue(TEXT("FCbPackage(OutOfOrder).Load()->Level2Hash"),
+		TestTrue(TEXT("FCbPackage(OutOfOrder).TryLoad()->Level2Hash"),
 			Level2Attachment && Level2Attachment->GetHash() == Level2Hash);
 
-		TestTrue(TEXT("FCbPackage(OutOfOrder).Load()->Level3"),
+		TestTrue(TEXT("FCbPackage(OutOfOrder).TryLoad()->Level3"),
 			Level3Attachment && Level3Attachment->AsObject().Equals(Level3));
-		TestTrue(TEXT("FCbPackage(OutOfOrder).Load()->Level3Buffer"),
+		TestTrue(TEXT("FCbPackage(OutOfOrder).TryLoad()->Level3Buffer"),
 			Level3Attachment && Fields.GetOuterBuffer().GetView().Contains(Level3Attachment->AsBinary().GetView()));
-		TestTrue(TEXT("FCbPackage(OutOfOrder).Load()->Level3Hash"),
+		TestTrue(TEXT("FCbPackage(OutOfOrder).TryLoad()->Level3Hash"),
 			Level3Attachment && Level3Attachment->GetHash() == Level3Hash);
 
-		TestTrue(TEXT("FCbPackage(OutOfOrder).Load()->Level4"),
+		TestTrue(TEXT("FCbPackage(OutOfOrder).TryLoad()->Level4"),
 			Level4Attachment && Level4Attachment->AsBinary().GetView().EqualBytes(Level4.GetView()));
-		TestTrue(TEXT("FCbPackage(OutOfOrder).Load()->Level4Buffer"),
+		TestTrue(TEXT("FCbPackage(OutOfOrder).TryLoad()->Level4Buffer"),
 			Level4Attachment && Fields.GetOuterBuffer().GetView().Contains(Level4Attachment->AsBinary().GetView()));
-		TestTrue(TEXT("FCbPackage(OutOfOrder).Load()->Level4Hash"),
+		TestTrue(TEXT("FCbPackage(OutOfOrder).TryLoad()->Level4Hash"),
 			Level4Attachment && Level4Attachment->GetHash() == Level4Hash);
 
 		FBufferArchive WriteAr;
 		Writer.Save(WriteAr);
 		FCbPackage FromArchive;
 		FMemoryReader ReadAr(WriteAr);
-		FromArchive.Load(ReadAr);
+		FromArchive.TryLoad(ReadAr);
 
 		Writer.Reset();
 		FromArchive.Save(Writer);
diff --git a/Engine/Source/Runtime/Core/Public/Serialization/CompactBinaryPackage.h b/Engine/Source/Runtime/Core/Public/Serialization/CompactBinaryPackage.h
--- a/Engine/Source/Runtime/Core/Public/Serialization/CompactBinaryPackage.h
+++ b/Engine/Source/Runtime/Core/Public/Serialization/CompactBinaryPackage.h
@@ -89,17 +89,18 @@ public:
 	 *
 	 * The iterator is advanced as attachment fields are consumed from it.
 	 */
-	CORE_API void Load(FCbFieldIterator& Fields);
+	CORE_API bool TryLoad(FCbFieldIterator& Fields);
 
 	/**
 	 * Load the attachment from compact binary as written by Save.
 	 *
 	 * The attachments value will be loaded into an owned buffer.
 	 *
+	 * @param Ar Archive to read the attachment from. An error state is set on failure.
 	 * @param Allocator Allocator for the attachment value buffer.
 	 * @note Allocated buffers will be cloned if they are not owned.
 	 */
-	CORE_API void Load(FArchive& Ar, FCbBufferAllocator Allocator = FUniqueBuffer::Alloc);
+	CORE_API bool TryLoad(FArchive& Ar, FCbBufferAllocator Allocator = FUniqueBuffer::Alloc);
 
 	/** Save the attachment into the writer as a stream of compact binary fields. */
 	CORE_API void Save(FCbWriter& Writer) const;
@@ -316,17 +317,18 @@ public:
 	 *
 	 * The iterator is advanced as object and attachment fields are consumed from it.
 	 */
-	CORE_API void Load(FCbFieldIterator& Fields);
+	CORE_API bool TryLoad(FCbFieldIterator& Fields);
 
 	/**
 	 * Load the object and attachments from compact binary as written by Save.
 	 *
 	 * The object and attachments will be individually loaded into owned buffers.
 	 *
+	 * @param Ar Archive to read the package from. An error state is set on failure.
 	 * @param Allocator Allocator for object and attachment buffers.
 	 * @note Allocated buffers will be cloned if they are not owned.
 	 */
-	CORE_API void Load(FArchive& Ar, FCbBufferAllocator Allocator = FUniqueBuffer::Alloc);
+	CORE_API bool TryLoad(FArchive& Ar, FCbBufferAllocator Allocator = FUniqueBuffer::Alloc);
 
 	/** Save the object and attachments into the writer as a stream of compact binary fields. */
 	CORE_API void Save(FCbWriter& Writer) const;
diff --git a/Engine/Source/Runtime/Core/Public/Serialization/CompactBinarySerialization.h b/Engine/Source/Runtime/Core/Public/Serialization/CompactBinarySerialization.h
--- a/Engine/Source/Runtime/Core/Public/Serialization/CompactBinarySerialization.h
+++ b/Engine/Source/Runtime/Core/Public/Serialization/CompactBinarySerialization.h
@@ -58,9 +58,9 @@ CORE_API bool TryMeasureCompactBinary(
  * AsObject as appropriate. The buffer allocator is called to provide the buffer for the field
  * to load into once its size has been determined.
  *
- * @param Ar Archive to read the field from.
+ * @param Ar Archive to read the field from. An error state is set on failure.
  * @param Allocator Allocator for the buffer that the field is loaded into.
- * @return A field with a reference to the provided buffer if it is owned.
+ * @return A field with a reference to the allocated buffer, or a default field on failure.
  */
 CORE_API FCbField LoadCompactBinary(FArchive& Ar, FCbBufferAllocator Allocator = FUniqueBuffer::Alloc);
 
