hooks to allow crash handling to be extended by platform extensions.

#rnx

diff --git a/Engine/Source/Runtime/Core/Private/GenericPlatform/GenericPlatformCrashContext.cpp b/Engine/Source/Runtime/Core/Private/GenericPlatform/GenericPlatformCrashContext.cpp
--- a/Engine/Source/Runtime/Core/Private/GenericPlatform/GenericPlatformCrashContext.cpp
+++ b/Engine/Source/Runtime/Core/Private/GenericPlatform/GenericPlatformCrashContext.cpp
@@ -1131,7 +1131,7 @@ void FGenericCrashContext::AddPlugin(const FString& PluginDesc)
 	NCached::EnabledPluginsList.Add(PluginDesc);
 }
 
-void FGenericCrashContext::DumpLog(const FString& CrashFolderAbsolute)
+FString FGenericCrashContext::DumpLog(const FString& CrashFolderAbsolute)
 {
 	// Copy log
 	const FString LogSrcAbsolute = FPlatformOutputDevices::GetAbsoluteLogFilename();
@@ -1171,7 +1171,7 @@ void FGenericCrashContext::DumpLog(const FString& CrashFolderAbsolute)
 	}
 #endif // !NO_LOGGING
 
-	
+	return LogDstAbsolute;
 }
 
 void FGenericCrashContext::CapturePortableCallStack(void* ErrorProgramCounter, void* Context)
diff --git a/Engine/Source/Runtime/Core/Private/Windows/WindowsPlatformCrashContext.cpp b/Engine/Source/Runtime/Core/Private/Windows/WindowsPlatformCrashContext.cpp
--- a/Engine/Source/Runtime/Core/Private/Windows/WindowsPlatformCrashContext.cpp
+++ b/Engine/Source/Runtime/Core/Private/Windows/WindowsPlatformCrashContext.cpp
@@ -1532,7 +1532,11 @@ private:
 TOptional<FCrashReportingThread> GCrashReportingThread(InPlace);
 #endif
 
+#if WINDOWS_CRASHCONTEXT_WITH_CUSTOM_HANDLERS
+LONG WINAPI DefaultUnhandledStaticInitException(LPEXCEPTION_POINTERS ExceptionInfo)
+#else
 LONG WINAPI UnhandledStaticInitException(LPEXCEPTION_POINTERS ExceptionInfo)
+#endif
 {
 #if !NOINITCRASHREPORTER
 	// If we get an exception during static init we hope that the crash reporting thread
@@ -1559,7 +1563,7 @@ LONG WINAPI UnhandledStaticInitException(LPEXCEPTION_POINTERS ExceptionInfo)
  *   - Any unhandled exception is going to terminate the program whether it is a benign exception or a fatal one.
  *   - Vectored exception handlers, Vectored continue handlers and the unhandled exception filter are global to the process.
  *   - Exceptions occurring in a thread doesn't automatically halt other threads. Exception handling executes in thread where the exception fired. The other threads continue to run.
- *   - Several threads can crash concurrently­.
+ *   - Several threads can crash concurrently.
  *   - Not all exceptions are equal. Some exceptions can be handled doing nothing more than catching them and telling the code to continue (like some user defined exception), some
  *     needs to be handled in a __except() clause to allow the program to continue (like access violation) and others are fatal and can only be reported but not continued (like stack overflow).
  *   - Not all machines are equal. Different exceptions may be fired on different machines for the same usage of the program. This seems especially true when
@@ -1609,7 +1613,11 @@ LONG WINAPI UnhandledStaticInitException(LPEXCEPTION_POINTERS ExceptionInfo)
  * The engine hooks itself in the unhandled exception filter. This is the best place to be as it runs after structured exception handlers and
  * it can be easily overriden externally (because there can only be one) to do something else.
  */
+#if WINDOWS_CRASHCONTEXT_WITH_CUSTOM_HANDLERS
+LONG WINAPI DefaultEngineUnhandledExceptionFilter(LPEXCEPTION_POINTERS ExceptionInfo)
+#else
 LONG WINAPI EngineUnhandledExceptionFilter(LPEXCEPTION_POINTERS ExceptionInfo)
+#endif
 {
 	ReportCrash(ExceptionInfo);
 	GIsCriticalError = true;
@@ -1619,7 +1627,11 @@ LONG WINAPI EngineUnhandledExceptionFilter(LPEXCEPTION_POINTERS ExceptionInfo)
 }
 
 // #CrashReport: 2015-05-28 This should be named EngineCrashHandler
+#if WINDOWS_CRASHCONTEXT_WITH_CUSTOM_HANDLERS
+int32 DefaultReportCrash( LPEXCEPTION_POINTERS ExceptionInfo )
+#else
 int32 ReportCrash( LPEXCEPTION_POINTERS ExceptionInfo )
+#endif
 {
 #if !NOINITCRASHREPORTER
 	// Only create a minidump the first time this function is called.
diff --git a/Engine/Source/Runtime/Core/Public/GenericPlatform/GenericPlatformCrashContext.h b/Engine/Source/Runtime/Core/Public/GenericPlatform/GenericPlatformCrashContext.h
--- a/Engine/Source/Runtime/Core/Public/GenericPlatform/GenericPlatformCrashContext.h
+++ b/Engine/Source/Runtime/Core/Public/GenericPlatform/GenericPlatformCrashContext.h
@@ -408,8 +408,8 @@ public:
 	/** Adds a plugin descriptor string to the enabled plugins list in the crash context */
 	static void AddPlugin(const FString& PluginDesc);
 
-	/** Flushes the logs. In the case of in memory logs is used on this configuration, dumps them to file. */
-	static void DumpLog(const FString& CrashFolderAbsolute);
+	/** Flushes the logs. In the case of in memory logs is used on this configuration, dumps them to file. Returns the name of the file */
+	static FString DumpLog(const FString& CrashFolderAbsolute);
 
 	/** Collects additional crash context providers. See FAdditionalCrashContextStack. */
 	static void DumpAdditionalContext(const TCHAR* CrashFolderAbsolute);
diff --git a/Engine/Source/Runtime/Core/Public/Windows/WindowsPlatformCrashContext.h b/Engine/Source/Runtime/Core/Public/Windows/WindowsPlatformCrashContext.h
--- a/Engine/Source/Runtime/Core/Public/Windows/WindowsPlatformCrashContext.h
+++ b/Engine/Source/Runtime/Core/Public/Windows/WindowsPlatformCrashContext.h
@@ -5,6 +5,11 @@
 #include "CoreTypes.h"
 #include "GenericPlatform/GenericPlatformCrashContext.h"
 
+#if !defined(WINDOWS_CRASHCONTEXT_WITH_CUSTOM_HANDLERS)
+	#define WINDOWS_CRASHCONTEXT_WITH_CUSTOM_HANDLERS 0
+#endif
+
+
 struct CORE_API FWindowsPlatformCrashContext : public FGenericCrashContext
 {
 	static const TCHAR* const UEGPUAftermathMinidumpName;
