#ue5 - LWC: add RESTRICT keyword to VectorDeinterleave, to protect against duplicated inputs for the two out vars.

diff --git a/Engine/Source/Runtime/Core/Public/Math/UnrealMathFPU.h b/Engine/Source/Runtime/Core/Public/Math/UnrealMathFPU.h
--- a/Engine/Source/Runtime/Core/Public/Math/UnrealMathFPU.h
+++ b/Engine/Source/Runtime/Core/Public/Math/UnrealMathFPU.h
@@ -1808,13 +1808,13 @@ FORCEINLINE VectorRegister4Double VectorCombineLow(const VectorRegister4Double&
  * @param OutEvens [Even0, Even1, Even2, Even3]
  * @param OutOdds [Odd0, Odd1, Odd2, Odd3]
 */
-FORCEINLINE void VectorDeinterleave(VectorRegister4Float& OutEvens, VectorRegister4Float& OutOdds, const VectorRegister4Float& Lo, const VectorRegister4Float& Hi)
+FORCEINLINE void VectorDeinterleave(VectorRegister4Float& RESTRICT OutEvens, VectorRegister4Float& RESTRICT OutOdds, const VectorRegister4Float& RESTRICT Lo, const VectorRegister4Float& RESTRICT Hi)
 {
 	OutEvens = MakeVectorRegister(Lo.V[0], Lo.V[2], Hi.V[0], Hi.V[2]);
 	OutOdds = MakeVectorRegister(Lo.V[1], Lo.V[3], Hi.V[1], Hi.V[3]);
 }
 
-FORCEINLINE void VectorDeinterleave(VectorRegister4Double& OutEvens, VectorRegister4Double& OutOdds, const VectorRegister4Double& Lo, const VectorRegister4Double& Hi)
+FORCEINLINE void VectorDeinterleave(VectorRegister4Double& RESTRICT OutEvens, VectorRegister4Double& RESTRICT OutOdds, const VectorRegister4Double& RESTRICT Lo, const VectorRegister4Double& RESTRICT Hi)
 {
 	OutEvens = MakeVectorRegister(Lo.V[0], Lo.V[2], Hi.V[0], Hi.V[2]);
 	OutOdds = MakeVectorRegister(Lo.V[1], Lo.V[3], Hi.V[1], Hi.V[3]);
diff --git a/Engine/Source/Runtime/Core/Public/Math/UnrealMathSSE.h b/Engine/Source/Runtime/Core/Public/Math/UnrealMathSSE.h
--- a/Engine/Source/Runtime/Core/Public/Math/UnrealMathSSE.h
+++ b/Engine/Source/Runtime/Core/Public/Math/UnrealMathSSE.h
@@ -2506,13 +2506,13 @@ FORCEINLINE VectorRegister4Double VectorCombineLow(const VectorRegister4Double&
  * @param OutEvens [Even0, Even1, Even2, Even3]
  * @param OutOdds [Odd0, Odd1, Odd2, Odd3]
 */
-FORCEINLINE void VectorDeinterleave(VectorRegister4Float& OutEvens, VectorRegister4Float& OutOdds, const VectorRegister4Float& Lo, const VectorRegister4Float& Hi)
+FORCEINLINE void VectorDeinterleave(VectorRegister4Float& RESTRICT OutEvens, VectorRegister4Float& RESTRICT OutOdds, const VectorRegister4Float& RESTRICT Lo, const VectorRegister4Float& RESTRICT Hi)
 {
 	OutEvens = VectorShuffle(Lo, Hi, 0, 2, 0, 2);
 	OutOdds = VectorShuffle(Lo, Hi, 1, 3, 1, 3);
 }
 
-FORCEINLINE void VectorDeinterleave(VectorRegister4Double& OutEvens, VectorRegister4Double& OutOdds, const VectorRegister4Double& Lo, const VectorRegister4Double& Hi)
+FORCEINLINE void VectorDeinterleave(VectorRegister4Double& RESTRICT OutEvens, VectorRegister4Double& RESTRICT OutOdds, const VectorRegister4Double& RESTRICT Lo, const VectorRegister4Double& RESTRICT Hi)
 {
 	OutEvens = VectorShuffle(Lo, Hi, 0, 2, 0, 2);
 	OutOdds = VectorShuffle(Lo, Hi, 1, 3, 1, 3);
