#ue5 - LWC: add platform flag for AVX2 feature level, and enable AVX2 on next-gen consoles.

diff --git a/Engine/Source/Runtime/Core/Public/HAL/Platform.h b/Engine/Source/Runtime/Core/Public/HAL/Platform.h
--- a/Engine/Source/Runtime/Core/Public/HAL/Platform.h
+++ b/Engine/Source/Runtime/Core/Public/HAL/Platform.h
@@ -188,6 +188,9 @@
 #ifndef PLATFORM_ALWAYS_HAS_AVX
 	#define PLATFORM_ALWAYS_HAS_AVX				0
 #endif
+#ifndef PLATFORM_ALWAYS_HAS_AVX_2
+#define PLATFORM_ALWAYS_HAS_AVX_2				0
+#endif
 #ifndef PLATFORM_ALWAYS_HAS_FMA3
 	#define PLATFORM_ALWAYS_HAS_FMA3			0
 #endif
diff --git a/Engine/Source/Runtime/Core/Public/Math/UnrealMathSSE.h b/Engine/Source/Runtime/Core/Public/Math/UnrealMathSSE.h
--- a/Engine/Source/Runtime/Core/Public/Math/UnrealMathSSE.h
+++ b/Engine/Source/Runtime/Core/Public/Math/UnrealMathSSE.h
@@ -14,7 +14,7 @@
 #endif
 
 #ifndef UE_PLATFORM_MATH_USE_AVX_2
-#define UE_PLATFORM_MATH_USE_AVX_2			0 //PLATFORM_ALWAYS_HAS_AVX_2
+#define UE_PLATFORM_MATH_USE_AVX_2			(PLATFORM_ALWAYS_HAS_AVX_2 && UE_PLATFORM_MATH_USE_AVX)
 #endif
 
 #ifndef UE_PLATFORM_MATH_USE_FMA3
