Added missing NEON double vector functions and fixed VectorMaskBits

diff --git a/Engine/Source/Runtime/Core/Private/Tests/Math/UnrealMathTest.cpp b/Engine/Source/Runtime/Core/Private/Tests/Math/UnrealMathTest.cpp
--- a/Engine/Source/Runtime/Core/Private/Tests/Math/UnrealMathTest.cpp
+++ b/Engine/Source/Runtime/Core/Private/Tests/Math/UnrealMathTest.cpp
@@ -1143,6 +1143,11 @@ FORCENOINLINE void TestVectorShuffle()
 	ShuffleTest(V0, V1, 2, 2, 3, 3);
 	ShuffleTest(V0, V1, 2, 0, 3, 0);
 
+	ShuffleTest(V0, V1, 3, 1, 3, 0);
+	ShuffleTest(V0, V1, 2, 2, 0, 1);
+	ShuffleTest(V0, V1, 0, 1, 3, 2);
+	ShuffleTest(V0, V1, 1, 3, 0, 3);
+
 #undef ShuffleTest
 }
 
@@ -1782,6 +1787,26 @@ bool RunDoubleVectorTest()
 	V3 = VectorFloor(V0);
 	LogTest<double>(TEXT("VectorFloor"), TestVectorsEqual(V2, V3));
 
+	// VectorDeinterleave
+	V0 = MakeVectorRegister(0.0, 1.0, 2.0, 3.0);
+	V1 = MakeVectorRegister(4.0, 5.0, 6.0, 7.0);
+	VectorDeinterleave(V2, V3, V0, V1);
+	V0 = MakeVectorRegister(0.0, 2.0, 4.0, 6.0);
+	V1 = MakeVectorRegister(1.0, 3.0, 5.0, 7.0);
+	LogTest<double>(TEXT("VectorDeinterleave"), TestVectorsEqual(V2, V0) && TestVectorsEqual(V3, V1));
+
+	// VectorMaskBits
+	V0 = MakeVectorRegister(0.0, 1.0, 2.0, 3.0);
+	uint32 MaskBits = VectorMaskBits(V0);
+	LogTest<double>(TEXT("VectorMaskBits"), MaskBits == 0);
+	V0 = MakeVectorRegister(-0.0, -1.0, -2.0, -3.0);
+	MaskBits = VectorMaskBits(V0);
+	LogTest<double>(TEXT("VectorMaskBits"), MaskBits == 0xf);
+	V0 = MakeVectorRegister(-0.0, 1.0, -2.0, 3.0);
+	MaskBits = VectorMaskBits(V0);
+	LogTest<double>(TEXT("VectorMaskBits"), MaskBits == 5);
+
+
 	// NaN / Inf tests
 	// Using a union as we need to do a bitwise cast of 0xFFFFFFFF into a float.
 	typedef union
@@ -2312,6 +2337,17 @@ bool FVectorRegisterAbstractionTest::RunTest(const FString& Parameters)
 	V3 = VectorFloor(V0);
 	LogTest<float>(TEXT("VectorFloor"), TestVectorsEqual(V2, V3));
 
+	// VectorMaskBits
+	V0 = MakeVectorRegister(0.0f, 1.0f, 2.0f, 3.0f);
+	uint32 MaskBits = VectorMaskBits(V0);
+	LogTest<double>(TEXT("VectorMaskBits"), MaskBits == 0);
+	V0 = MakeVectorRegister(-0.0f, -1.0f, -2.0f, -3.0f);
+	MaskBits = VectorMaskBits(V0);
+	LogTest<double>(TEXT("VectorMaskBits"), MaskBits == 0xf);
+	V0 = MakeVectorRegister(-0.0f, 1.0f, -2.0f, 3.0f);
+	MaskBits = VectorMaskBits(V0);
+	LogTest<double>(TEXT("VectorMaskBits"), MaskBits == 5);
+
 	FMatrix44f	M0, M1, M2, M3;		// LWC_TODO: Need double version tests!
 	FVector3f Eye, LookAt, Up;	
 	// Create Look at Matrix
diff --git a/Engine/Source/Runtime/Core/Public/Math/UnrealMathNeon.h b/Engine/Source/Runtime/Core/Public/Math/UnrealMathNeon.h
--- a/Engine/Source/Runtime/Core/Public/Math/UnrealMathNeon.h
+++ b/Engine/Source/Runtime/Core/Public/Math/UnrealMathNeon.h
@@ -46,20 +46,18 @@ struct alignas(alignof(T)) VectorRegisterWrapper
 typedef VectorRegisterWrapper<float32x4_t, false> VectorRegister4Float;
 typedef VectorRegisterWrapper<float64x2_t, false> VectorRegister2Double;
 typedef VectorRegisterWrapper<int32x4_t, true> VectorRegister4Int;
-typedef VectorRegisterWrapper<uint64x2_t, true> VectorRegister2Int64;
-
-#define DECLARE_VECTOR_REGISTER(X, Y, Z, W) MakeVectorRegister( X, Y, Z, W )
+typedef VectorRegisterWrapper<int64x2_t, true> VectorRegister2Int64;
 #else
 
 /** 16-byte vector register type */
 typedef float32x4_t GCC_ALIGN(16) VectorRegister4Float;
 typedef float64x2_t GCC_ALIGN(16) VectorRegister2Double;
 typedef int32x4_t  GCC_ALIGN(16) VectorRegister4Int;
-typedef uint64x2_t GCC_ALIGN(16) VectorRegister2Int64;
-
-#define DECLARE_VECTOR_REGISTER(X, Y, Z, W) { X, Y, Z, W }
+typedef int64x2_t GCC_ALIGN(16) VectorRegister2Int64;
 #endif
 
+#define DECLARE_VECTOR_REGISTER(X, Y, Z, W) MakeVectorRegister( X, Y, Z, W )
+
 struct alignas(16) VectorRegister4Double
 {
 	union
@@ -92,9 +90,9 @@ struct alignas(16) VectorRegister4Double
 
 // LWC: Alias VectorRegister to correct precision based on LWC
 #if !UE_LARGE_WORLD_COORDINATES_DISABLED
-typedef VectorRegister4Double VectorRegister;
+	typedef VectorRegister4Double VectorRegister;
 #else
-typedef VectorRegister4Float VectorRegister;
+	typedef VectorRegister4Float VectorRegister;
 #endif
 
 // Forward declarations
@@ -311,6 +309,18 @@ FORCEINLINE VectorRegister4Int MakeVectorRegisterInt(int32 X, int32 Y, int32 Z,
 	return Tmp.V;
 }
 
+FORCEINLINE VectorRegister2Int64 MakeVectorRegisterInt64(int64 X, int64 Y)
+{
+	union U
+	{
+		VectorRegister2Int64 V; int64 I[2];
+		FORCEINLINE U() : V() {}
+	} Tmp;
+	Tmp.I[0] = X;
+	Tmp.I[1] = Y;
+	return Tmp.V;
+}
+
 // Lossy conversion: double->float vector
 FORCEINLINE VectorRegister4Float MakeVectorRegisterFloatFromDouble(const VectorRegister4Double& Vec)
 {
@@ -539,6 +549,13 @@ FORCEINLINE VectorRegister4Float VectorSetFloat1(float X)
 	return vdupq_n_f32(X);
 }
 
+FORCEINLINE VectorRegister4Double VectorSetFloat1(double X)
+{
+	VectorRegister4Double Result;
+	Result.XY = vdupq_n_f64(X);
+	Result.ZW = Result.XY;
+	return Result;
+}
 
 /**
  * Stores a vector to aligned memory.
@@ -1288,14 +1305,21 @@ FORCEINLINE VectorRegister4Double VectorShuffleImpl(VectorRegister4Double Vec1,
  */
 FORCEINLINE uint32 VectorMaskBits(VectorRegister4Float VecMask)
 {
-	uint32x4_t mmA = vandq_u32(vreinterpretq_u32_f32(VecMask), MakeVectorRegisterInt( 0x1, 0x2, 0x4, 0x8 )); // [0 1 2 3]
-	uint32x4_t mmB = vextq_u32(mmA, mmA, 2);                        // [2 3 0 1]
-	uint32x4_t mmC = vorrq_u32(mmA, mmB);                           // [0+2 1+3 0+2 1+3]
-	uint32x4_t mmD = vextq_u32(mmC, mmC, 3);                        // [1+3 0+2 1+3 0+2]
-	uint32x4_t mmE = vorrq_u32(mmC, mmD);                           // [0+1+2+3 ...]
-	return vgetq_lane_u32(mmE, 0);
+	uint32x4_t mmA = vtstq_u32(vreinterpretq_u32_f32(VecMask), GlobalVectorConstants::SignBit); // mask with 1s every bit for vector element if it's sign is negative
+	uint32x4_t mmB = vandq_u32(mmA, MakeVectorRegisterInt(0x1, 0x2, 0x4, 0x8)); // pick only one bit on it's corresponding position
+	uint32x2_t mmC = vorr_u32(vget_low_u32(mmB), vget_high_u32(mmB));           // now combine the result
+	return vget_lane_u32(mmC, 0) | vget_lane_u32(mmC, 1);                       // reduce the result from 2 elements to one
 }
 
+FORCEINLINE uint32 VectorMaskBits(VectorRegister4Double VecMask)
+{
+	uint64x2_t mmA = vtstq_u64(vreinterpretq_u64_f64(VecMask.XY), GlobalVectorConstants::DoubleSignBit.XY); // mask with 1s every bit for vector element if it's sign is negative
+	uint64x2_t mmA1 = vtstq_u64(vreinterpretq_u64_f64(VecMask.ZW), GlobalVectorConstants::DoubleSignBit.XY);
+	uint64x2_t mmB = vandq_u64(mmA, MakeVectorRegisterInt64(0x1, 0x2)); // pick only one bit on it's corresponding position
+	uint64x2_t mmB1 = vandq_u64(mmA, MakeVectorRegisterInt64(0x4, 0x8));
+	uint64x2_t mmC = vorrq_u64(mmB, mmB1);								// now combine the result
+	return (uint32)(vgetq_lane_u64(mmC, 0) | vgetq_lane_u64(mmC, 1));     // reduce the result from 2 elements to one
+}
 
 /**
 * Creates a vector by combining two high components from each vector
@@ -1353,6 +1377,12 @@ FORCEINLINE void VectorDeinterleave(VectorRegister4Float& OutEvens, VectorRegist
 	OutOdds = deinterleaved.val[1];
 }
 
+FORCEINLINE void VectorDeinterleave(VectorRegister4Double& RESTRICT OutEvens, VectorRegister4Double& RESTRICT OutOdds, const VectorRegister4Double& Lo, const VectorRegister4Double& Hi)
+{
+	OutEvens = VectorShuffle(Lo, Hi, 0, 2, 0, 2);
+	OutOdds = VectorShuffle(Lo, Hi, 1, 3, 1, 3);
+}
+
 /**
  * Calculates the cross product of two vectors (XYZ components). W is set to 0.
  *
@@ -1395,6 +1425,19 @@ FORCEINLINE VectorRegister4Float VectorPow( const VectorRegister4Float& Base, co
 	return MakeVectorRegister( powf(B.F[0], E.F[0]), powf(B.F[1], E.F[1]), powf(B.F[2], E.F[2]), powf(B.F[3], E.F[3]) );
 }
 
+FORCEINLINE VectorRegister4Double VectorPow(const VectorRegister4Double& Base, const VectorRegister4Double& Exponent)
+{
+	//@TODO: Optimize this
+	AlignedDouble4 Values(Base);
+	AlignedDouble4 Exponents(Exponent);
+
+	Values[0] = FMath::Pow(Values[0], Exponents[0]);
+	Values[1] = FMath::Pow(Values[1], Exponents[1]);
+	Values[2] = FMath::Pow(Values[2], Exponents[2]);
+	Values[3] = FMath::Pow(Values[3], Exponents[3]);
+	return Values.ToVectorRegister();
+}
+
 /**
 * Returns an estimate of 1/sqrt(c) for each component of the vector
 *
@@ -2744,7 +2787,16 @@ FORCEINLINE VectorRegister4Int VectorIntSelect(const VectorRegister4Int& Mask, c
 #define VectorIntSign(A) VectorIntSelect( VectorIntCompareGE(A, GlobalVectorConstants::IntZero), GlobalVectorConstants::IntOne, GlobalVectorConstants::IntMinusOne )
 
 #define VectorIntToFloat(A) vcvtq_f32_s32(A)
-#define VectorFloatToInt(A) vcvtq_s32_f32(A)
+
+FORCEINLINE VectorRegister4Int VectorFloatToInt(const VectorRegister4Float& A)
+{
+	return vcvtq_s32_f32(A);
+}
+
+FORCEINLINE VectorRegister4Int VectorFloatToInt(const VectorRegister4Double& A)
+{
+	return VectorFloatToInt(MakeVectorRegisterFloatFromDouble(A));
+}
 
 //Loads and stores
 
