Revert VectorMaskBits test due to a bug in CL

diff --git a/Engine/Source/Runtime/Core/Private/Tests/Math/UnrealMathTest.cpp b/Engine/Source/Runtime/Core/Private/Tests/Math/UnrealMathTest.cpp
--- a/Engine/Source/Runtime/Core/Private/Tests/Math/UnrealMathTest.cpp
+++ b/Engine/Source/Runtime/Core/Private/Tests/Math/UnrealMathTest.cpp
@@ -1795,17 +1795,6 @@ bool RunDoubleVectorTest()
 	V1 = MakeVectorRegister(1.0, 3.0, 5.0, 7.0);
 	LogTest<double>(TEXT("VectorDeinterleave"), TestVectorsEqual(V2, V0) && TestVectorsEqual(V3, V1));
 
-	// VectorMaskBits
-	V0 = MakeVectorRegister(0.0, 1.0, 2.0, 3.0);
-	uint32 MaskBits = VectorMaskBits(V0);
-	LogTest<double>(TEXT("VectorMaskBits"), MaskBits == 0);
-	V0 = MakeVectorRegister(-0.0, -1.0, -2.0, -3.0);
-	MaskBits = VectorMaskBits(V0);
-	LogTest<double>(TEXT("VectorMaskBits"), MaskBits == 0xf);
-	V0 = MakeVectorRegister(-0.0, 1.0, -2.0, 3.0);
-	MaskBits = VectorMaskBits(V0);
-	LogTest<double>(TEXT("VectorMaskBits"), MaskBits == 5);
-
 
 	// NaN / Inf tests
 	// Using a union as we need to do a bitwise cast of 0xFFFFFFFF into a float.
@@ -2337,17 +2326,6 @@ bool FVectorRegisterAbstractionTest::RunTest(const FString& Parameters)
 	V3 = VectorFloor(V0);
 	LogTest<float>(TEXT("VectorFloor"), TestVectorsEqual(V2, V3));
 
-	// VectorMaskBits
-	V0 = MakeVectorRegister(0.0f, 1.0f, 2.0f, 3.0f);
-	uint32 MaskBits = VectorMaskBits(V0);
-	LogTest<double>(TEXT("VectorMaskBits"), MaskBits == 0);
-	V0 = MakeVectorRegister(-0.0f, -1.0f, -2.0f, -3.0f);
-	MaskBits = VectorMaskBits(V0);
-	LogTest<double>(TEXT("VectorMaskBits"), MaskBits == 0xf);
-	V0 = MakeVectorRegister(-0.0f, 1.0f, -2.0f, 3.0f);
-	MaskBits = VectorMaskBits(V0);
-	LogTest<double>(TEXT("VectorMaskBits"), MaskBits == 5);
-
 	FMatrix44f	M0, M1, M2, M3;		// LWC_TODO: Need double version tests!
 	FVector3f Eye, LookAt, Up;	
 	// Create Look at Matrix
