Fix for FPaths::Combine using freed memory when being passed a ANSICHAR*, or any other type which results in an FString being implicitly constructed.

diff --git a/Engine/Source/Runtime/Core/Public/Misc/Paths.h b/Engine/Source/Runtime/Core/Public/Misc/Paths.h
--- a/Engine/Source/Runtime/Core/Public/Misc/Paths.h
+++ b/Engine/Source/Runtime/Core/Public/Misc/Paths.h
@@ -6,6 +6,33 @@
 #include "Containers/UnrealString.h"
 #include "HAL/CriticalSection.h"
 
+namespace UE::Core::Private
+{
+	// This is used to force the arguments of FString::Combine to implicitly convert (if necessary)
+	// to FString when calling CombineImpl(), allowing them to remain as temporaries on the stack
+	// so that they stay allocated during the combining process.
+	//
+	// Pointer arguments are passed without causing FString temporaries to be created,
+	// and FString arguments are referenced directly without creating extra copies.
+	template <typename T>
+	struct TToStringType
+	{
+		using Type = const FString&;
+	};
+
+	template <>
+	struct TToStringType<TCHAR*>
+	{
+		using Type = const TCHAR*;
+	};
+
+	template <>
+	struct TToStringType<const TCHAR*>
+	{
+		using Type = const TCHAR*;
+	};
+}
+
 /**
  * Path helpers for retrieving game dir, engine dir, etc.
  */
@@ -603,11 +630,7 @@ public:
 	template <typename... PathTypes>
 	FORCEINLINE static FString Combine(PathTypes&&... InPaths)
 	{
-		const TCHAR* Paths[] = { GetTCharPtr(Forward<PathTypes>(InPaths))... };
-		FString Out;
-		
-		CombineInternal(Out, Paths, UE_ARRAY_COUNT(Paths));
-		return Out;
+		return CombineImpl<PathTypes...>(Forward<PathTypes>(InPaths)...);
 	}
 
 	/**
@@ -622,6 +645,16 @@ protected:
 private:
 	struct FStaticData;
 
+	template <typename... PathTypes>
+	FORCEINLINE static FString CombineImpl(typename UE::Core::Private::TToStringType<std::decay_t<PathTypes>>::Type... InPaths)
+	{
+		const TCHAR* Paths[] = { GetTCharPtr(InPaths)... };
+		FString Out;
+		
+		CombineInternal(Out, Paths, UE_ARRAY_COUNT(Paths));
+		return Out;
+	}
+
 	FORCEINLINE static const TCHAR* GetTCharPtr(const TCHAR* Ptr)
 	{
 		return Ptr;
