Fix for variable-precision floating point types in FGenericWidePlatformString::GetVarArgs() and Android's vswprintf.

diff --git a/Engine/Source/Runtime/Core/Private/Android/AndroidPlatformString.cpp b/Engine/Source/Runtime/Core/Private/Android/AndroidPlatformString.cpp
--- a/Engine/Source/Runtime/Core/Private/Android/AndroidPlatformString.cpp
+++ b/Engine/Source/Runtime/Core/Private/Android/AndroidPlatformString.cpp
@@ -55,6 +55,7 @@ int vswprintf( TCHAR *buf, int max, const TCHAR *fmt, va_list args )
 
 		TCHAR *percent_ptr = src;
 		int fieldlen = 0;
+		bool bVariablePrecision = false;
 		int precisionlen = -1;
 
 		src++; // skip the '%' char...
@@ -91,13 +92,22 @@ int vswprintf( TCHAR *buf, int max, const TCHAR *fmt, va_list args )
 		if (*src == '.')
 		{
 			TCHAR *ptr = src + 1;
-			while ((*ptr >= '0') && (*ptr <= '9'))
+			if (*ptr == '*')
+			{
+				bVariablePrecision = true;
+				precisionlen = va_arg(args, int);
 				ptr++;
+			}
+			else
+			{
+				while ((*ptr >= '0') && (*ptr <= '9'))
+					ptr++;
 
-			TCHAR ch = *ptr;
-			*ptr = '\0';
-			precisionlen = atoi(TCHAR_TO_ANSI(src + 1));
-			*ptr = ch;
+				TCHAR ch = *ptr;
+				*ptr = '\0';
+				precisionlen = atoi(TCHAR_TO_ANSI(src + 1));
+				*ptr = ch;
+			}
 			src = ptr;
 		}
 
@@ -252,7 +262,15 @@ int vswprintf( TCHAR *buf, int max, const TCHAR *fmt, va_list args )
 				}
 				fmtbuf[cpyidx] = 0;
 
-				int rc = snprintf(ansinum, sizeof (ansinum), fmtbuf, val);
+				int rc;
+				if (bVariablePrecision)
+				{
+					rc = snprintf(ansinum, sizeof (ansinum), fmtbuf, precisionlen, val);
+				}
+				else
+				{
+					rc = snprintf(ansinum, sizeof (ansinum), fmtbuf, val);
+				}
 				if ((dst + rc) > enddst)
 					return -1;	// Fail - the app needs to create a larger buffer and try again
 				for (int i = 0; i < rc; i++)
diff --git a/Engine/Source/Runtime/Core/Private/GenericPlatform/GenericWidePlatformString.cpp b/Engine/Source/Runtime/Core/Private/GenericPlatform/GenericWidePlatformString.cpp
--- a/Engine/Source/Runtime/Core/Private/GenericPlatform/GenericWidePlatformString.cpp
+++ b/Engine/Source/Runtime/Core/Private/GenericPlatform/GenericWidePlatformString.cpp
@@ -409,6 +409,7 @@ int32 FGenericWidePlatformString::GetVarArgs( WIDECHAR* Dest, SIZE_T DestSize, c
 
 		const TCHAR *Percent = Src;
 		int FieldLen = 0;
+		bool bVariablePrecision = false;
 		int PrecisionLen = -1;
 
 		Src++; // skip the '%' char...
@@ -454,6 +455,7 @@ int32 FGenericWidePlatformString::GetVarArgs( WIDECHAR* Dest, SIZE_T DestSize, c
 			const TCHAR *Cur = Src + 1;
 			if (*Cur == '*')
 			{
+				bVariablePrecision = true;
 				PrecisionLen = va_arg(ArgPtr, int32);
 				Cur++;
 			}
@@ -783,7 +785,15 @@ int32 FGenericWidePlatformString::GetVarArgs( WIDECHAR* Dest, SIZE_T DestSize, c
 				}
 				FmtBuf[CpyIdx] = 0;
 
-				int RetCnt = snprintf(AnsiNum, sizeof (AnsiNum), FmtBuf, Val);
+				int RetCnt;
+				if (bVariablePrecision)
+				{
+					RetCnt = snprintf(AnsiNum, sizeof (AnsiNum), FmtBuf, PrecisionLen, Val);
+				}
+				else
+				{
+					RetCnt = snprintf(AnsiNum, sizeof (AnsiNum), FmtBuf, Val);
+				}
 				if (RetCnt >= UE_ARRAY_COUNT(AnsiNum))
 				{
 					// We should print what we have written into AnsiNum but ensure we null terminate before printing
