Fixed Editor taking several second to resolve and log the ensure callstack.
  - Reenabled the optimization that load the symbols on demand only rather than loading all the symbols.
  - Detected when a module loaded in the debug engine (dbghelp) misses the debug symbols, likely because it was implicitly loaded with an incomplete symbol search path and unloaded the module to reload it with the proper symbol search path set so that debug symbols could be found.
  - For Editor, 'debug ensure' is significantly faste. The time to stack walk and dump went from ~20s to ~1s.

diff --git a/Engine/Source/Runtime/Core/Private/Windows/WindowsPlatformStackWalk.cpp b/Engine/Source/Runtime/Core/Private/Windows/WindowsPlatformStackWalk.cpp
--- a/Engine/Source/Runtime/Core/Private/Windows/WindowsPlatformStackWalk.cpp
+++ b/Engine/Source/Runtime/Core/Private/Windows/WindowsPlatformStackWalk.cpp
@@ -55,8 +55,7 @@ static const TCHAR* CrashReporterSettings = TEXT("/Script/UnrealEd.CrashReporter
 #endif
 
 // Optimization that only loads symbols on demand rather than always loading all symbols for loaded modules.
-// Warning: Enabling this currently has a bug where symbols are not loaded for modules outside of the main binary directory (i.e. plugins).
-#if 0
+#if 1
 #define ON_DEMAND_SYMBOL_LOADING 1
 #else
 #define ON_DEMAND_SYMBOL_LOADING 0
@@ -79,7 +78,7 @@ FString GetSymbolSearchPath();
 HMODULE* GetProcessModules(HANDLE ProcessHandle);
 void LoadSymbolsForModule(HMODULE ModuleHandle, const FString& RemoteStorage);
 void LoadSymbolsForProcessModules(const FString &RemoteStorage);
-void LoadSymbolsForModuleByAddress(uint64 Address, const FString& RemoteStorage);
+void LoadSymbolsForModuleByAddress(uint64 Address, const FString& RemoteStorage, bool bShouldReloadModuleMissingSymbols);
 
 struct FWindowsThreadContextWrapper
 {
@@ -450,7 +449,8 @@ void FWindowsPlatformStackWalk::ProgramCounterToSymbolInfo( uint64 ProgramCounte
 
 #if ON_DEMAND_SYMBOL_LOADING
 	// Load symbols for the module
-	LoadSymbolsForModuleByAddress(ProgramCounter, GetSymbolSearchPath());
+	bool bShouldReloadModuleMissingDebugSymbols = !FPlatformProperties::IsMonolithicBuild() && FPlatformStackWalk::WantsDetailedCallstacksInNonMonolithicBuilds();
+	LoadSymbolsForModuleByAddress(ProgramCounter, GetSymbolSearchPath(), bShouldReloadModuleMissingDebugSymbols);
 #endif
 
 	// Set the program counter.
@@ -516,7 +516,8 @@ void FWindowsPlatformStackWalk::ProgramCounterToSymbolInfoEx(uint64 ProgramCount
 {
 #if ON_DEMAND_SYMBOL_LOADING
 	// Load symbols for the module
-	LoadSymbolsForModuleByAddress(ProgramCounter, GetSymbolSearchPath());
+	bool bShouldReloadModuleMissingDebugSymbols = !FPlatformProperties::IsMonolithicBuild() && FPlatformStackWalk::WantsDetailedCallstacksInNonMonolithicBuilds();
+	LoadSymbolsForModuleByAddress(ProgramCounter, GetSymbolSearchPath(), bShouldReloadModuleMissingDebugSymbols);
 #endif
 
 	// Set the program counter.
@@ -793,12 +794,28 @@ void LoadSymbolsForProcessModules(const FString &RemoteStorage)
 	FMemory::Free(ModuleHandlePointer);
 }
 
-void LoadSymbolsForModuleByAddress(uint64 Address, const FString& RemoteStorage)
+void LoadSymbolsForModuleByAddress(uint64 Address, const FString& RemoteStorage, bool bShouldReloadModuleMissingDebugSymbols)
 {
 	HMODULE ModuleHandle = NULL;
 
 	if (GetModuleHandleEx(GET_MODULE_HANDLE_EX_FLAG_FROM_ADDRESS | GET_MODULE_HANDLE_EX_FLAG_UNCHANGED_REFCOUNT, (LPCTSTR)Address, &ModuleHandle))
 	{
+		// Check if the module was already loaded, but failed to locate the the debug symbols.
+		if (bShouldReloadModuleMissingDebugSymbols)
+		{
+			IMAGEHLP_MODULE64 ImageHelpModule = {0};
+			ImageHelpModule.SizeOfStruct = sizeof( ImageHelpModule );
+			if (SymGetModuleInfo64(GProcessHandle, Address, &ImageHelpModule) && ImageHelpModule.SymType == SymNone)
+			{
+				MODULEINFO ModuleInfo = { 0 };
+				FGetModuleInformation(GProcessHandle, ModuleHandle, &ModuleInfo, sizeof(ModuleInfo));
+
+				// The module is already loaded but 'SymNone' means that we are missing debug symbols. The module was likely loaded implicitly while the symbol search path wasn't properly set, so the debug engine did not find the .pdb and
+				// now that 'bad' state is cached. Unloading the module will clear the entry in the debug engine cache and loading it again with the proper symbol search path should pick up the .pdb this time.
+				SymUnloadModule(GProcessHandle, (DWORD64)ModuleInfo.lpBaseOfDll);
+			}
+		}
+
 		LoadSymbolsForModule(ModuleHandle, RemoteStorage);
 	}
 }
