Start the UnrealTrace server implicitly when non-programs launch, and when external platforms are launched through UAT.

#rnx

diff --git a/Engine/Source/Programs/AutomationTool/Scripts/RunProjectCommand.Automation.cs b/Engine/Source/Programs/AutomationTool/Scripts/RunProjectCommand.Automation.cs
--- a/Engine/Source/Programs/AutomationTool/Scripts/RunProjectCommand.Automation.cs
+++ b/Engine/Source/Programs/AutomationTool/Scripts/RunProjectCommand.Automation.cs
@@ -1,6 +1,7 @@
 // Copyright Epic Games, Inc. All Rights Reserved.
 using System;
 using System.Collections.Generic;
+using System.Diagnostics;
 using System.IO;
 using System.Text.RegularExpressions;
 using System.Threading;
@@ -118,6 +119,9 @@ public partial class Project : CommandUtils
 
 	private static void RunInternal(ProjectParams Params, string ServerLogFile, string ClientLogFile)
 	{
+		// Start the UnrealTrace
+		StartUnrealTrace();
+
 		// Setup server process if required.
 		if (Params.DedicatedServer && !Params.SkipServer)
 		{
@@ -1048,4 +1052,35 @@ public partial class Project : CommandUtils
 		PopDir();
 		return Result;
 	}
+
+	private static bool StartUnrealTrace()
+	{
+		// [TEMPORARY] - UnrealTrace server is currently only available on Windows
+		if (!Utils.IsRunningOnWindows)
+		{
+			return true;
+		}
+		// [/TEMPORARY]
+
+		LogInformation("UnrealTrace: Starting server");
+
+		// Locate the UnrealTrace binary
+		var UnrealTracePath = HostPlatform.Current.GetUE4ExePath("UnrealTrace.exe");
+		if (!File.Exists(UnrealTracePath))
+		{
+			LogWarning("UnrealTrace: Unable to locate binary at " + UnrealTracePath);
+			return false;
+		}
+
+		// Launch UnrealTrace and wait for it to fork and return
+		Process Proc = Process.Start(UnrealTracePath);
+		Proc.WaitForExit();
+		if (Proc.ExitCode != 0)
+		{
+			LogWarning("UnrealTrace: Failed to start server; ExitCode=" + Proc.ExitCode);
+			return false;
+		}
+
+		return true;
+	}
 }
diff --git a/Engine/Source/Programs/UnrealInsights/UnrealInsights.Target.cs b/Engine/Source/Programs/UnrealInsights/UnrealInsights.Target.cs
--- a/Engine/Source/Programs/UnrealInsights/UnrealInsights.Target.cs
+++ b/Engine/Source/Programs/UnrealInsights/UnrealInsights.Target.cs
@@ -28,5 +28,8 @@ public class UnrealInsightsTarget : TargetRules
 		bBuildDeveloperTools = true;
 
 		bHasExports = false;
+
+		// Have UnrealInsights implicitly launch the trace store.
+		GlobalDefinitions.Add("WITH_UNREAL_TRACE_LAUNCH=1");
 	}
 }
diff --git a/Engine/Source/Runtime/Core/Private/ProfilingDebugging/TraceAuxiliary.cpp b/Engine/Source/Runtime/Core/Private/ProfilingDebugging/TraceAuxiliary.cpp
--- a/Engine/Source/Runtime/Core/Private/ProfilingDebugging/TraceAuxiliary.cpp
+++ b/Engine/Source/Runtime/Core/Private/ProfilingDebugging/TraceAuxiliary.cpp
@@ -3,6 +3,16 @@
 #include "ProfilingDebugging/TraceAuxiliary.h"
 #include "Trace/Trace.h"
 
+#if PLATFORM_WINDOWS
+#	include "Windows/AllowWindowsPlatformTypes.h"
+#	include <Windows.h>
+#	include "Windows/HideWindowsPlatformTypes.h"
+#endif
+
+#if !defined(WITH_UNREAL_TRACE_LAUNCH)
+#	define WITH_UNREAL_TRACE_LAUNCH (PLATFORM_DESKTOP && !UE_BUILD_SHIPPING && !IS_PROGRAM)
+#endif
+
 #if UE_TRACE_ENABLED
 
 #include "BuildSettings.h"
@@ -29,12 +39,6 @@
 #include "Templates/UnrealTemplate.h"
 #include "Trace/Trace.inl"
 
-#if PLATFORM_WINDOWS
-#include "Windows/AllowWindowsPlatformTypes.h"
-#include <Windows.h>
-#include "Windows/HideWindowsPlatformTypes.h"
-#endif
-
 ////////////////////////////////////////////////////////////////////////////////
 const TCHAR* GDefaultChannels = TEXT("cpu,gpu,frame,log,bookmark");
 const TCHAR* GMemoryChannels = TEXT("memtag,memalloc,callstack,module");
@@ -480,6 +484,61 @@ static FAutoConsoleCommand TraceAuxiliaryResumeCmd(
 
 
 
+#if WITH_UNREAL_TRACE_LAUNCH
+////////////////////////////////////////////////////////////////////////////////
+#if PLATFORM_WINDOWS
+static void StartUnrealTrace()
+{
+	FString BinPath = "\"";
+	BinPath += FPaths::EngineDir();
+	BinPath += "/Binaries/Win64/UnrealTrace.exe\"";
+
+	uint32 Flags = 0x0100'0000; // CREATE_BREAKAWAY_FROM_JOB
+	STARTUPINFOW StartupInfo = { sizeof(STARTUPINFOW) };
+	PROCESS_INFORMATION ProcessInfo = {};
+	BOOL bOk = CreateProcessW(nullptr, LPWSTR(*BinPath), nullptr, nullptr,
+		false, Flags, nullptr, nullptr, &StartupInfo, &ProcessInfo);
+
+	if (!bOk)
+	{
+		UE_LOG(LogCore, Warning, TEXT("Unable to launch the trace store from '%s' (%08x)"), *BinPath, GetLastError());
+		return;
+	}
+
+	if (WaitForSingleObject(ProcessInfo.hProcess, 5000) == WAIT_TIMEOUT)
+	{
+		UE_LOG(LogCore, Warning, TEXT("Timed out waiting for the trace store to start"));
+	}
+	else
+	{
+		DWORD ExitCode = 0x0000'a9e0;
+		GetExitCodeProcess(ProcessInfo.hProcess, &ExitCode);
+		if (ExitCode)
+		{
+			UE_LOG(LogCore, Warning, TEXT("Trace store returned an error (%08x)"), ExitCode);
+		}
+		else
+		{
+			UE_LOG(LogCore, Log, TEXT("Trace store launch successful"));
+		}
+	}
+
+	CloseHandle(ProcessInfo.hProcess);
+	CloseHandle(ProcessInfo.hThread);
+}
+#endif // PLATFORM_WINDOWS
+
+////////////////////////////////////////////////////////////////////////////////
+#if PLATFORM_UNIX || PLATFORM_MAC
+static void StartUnrealTrace()
+{
+	/* nop */
+}
+#endif // PLATFORM_UNIX/MAC
+#endif // WITH_UNREAL_TRACE_LAUNCH
+
+
+
 ////////////////////////////////////////////////////////////////////////////////
 UE_TRACE_EVENT_BEGIN(Diagnostics, Session2, NoSync|Important)
 	UE_TRACE_EVENT_FIELD(UE::Trace::AnsiString, Platform)
@@ -495,6 +554,15 @@ UE_TRACE_EVENT_END()
 ////////////////////////////////////////////////////////////////////////////////
 void FTraceAuxiliary::Initialize(const TCHAR* CommandLine)
 {
+	TRACE_CPUPROFILER_EVENT_SCOPE(FTraceAux_Init);
+
+#if WITH_UNREAL_TRACE_LAUNCH
+	{
+		TRACE_CPUPROFILER_EVENT_SCOPE(FTraceAux_StartUnrealTrace);
+		StartUnrealTrace();
+	}
+#endif
+
 #if UE_TRACE_ENABLED
 	// Trace out information about this session. This is done before initialisation
 	// so that it is always sent (all channels are enabled prior to initialisation)
