Fix to allow FPaths::Combine to work for arguments implicitly convertible to const TCHAR*.

diff --git a/Engine/Source/Runtime/Core/Public/Misc/Paths.h b/Engine/Source/Runtime/Core/Public/Misc/Paths.h
--- a/Engine/Source/Runtime/Core/Public/Misc/Paths.h
+++ b/Engine/Source/Runtime/Core/Public/Misc/Paths.h
@@ -8,6 +8,9 @@
 
 namespace UE::Core::Private
 {
+	const TCHAR*   GetBindingType(const TCHAR* Ptr);
+	const FString& GetBindingType(const FString& Str);
+
 	// This is used to force the arguments of FString::Combine to implicitly convert (if necessary)
 	// to FString when calling CombineImpl(), allowing them to remain as temporaries on the stack
 	// so that they stay allocated during the combining process.
@@ -15,22 +18,7 @@ namespace UE::Core::Private
 	// Pointer arguments are passed without causing FString temporaries to be created,
 	// and FString arguments are referenced directly without creating extra copies.
 	template <typename T>
-	struct TToStringType
-	{
-		using Type = const FString&;
-	};
-
-	template <>
-	struct TToStringType<TCHAR*>
-	{
-		using Type = const TCHAR*;
-	};
-
-	template <>
-	struct TToStringType<const TCHAR*>
-	{
-		using Type = const TCHAR*;
-	};
+	using TToStringType_T = decltype(GetBindingType(std::declval<T>()));
 }
 
 /**
@@ -646,7 +634,7 @@ private:
 	struct FStaticData;
 
 	template <typename... PathTypes>
-	FORCEINLINE static FString CombineImpl(typename UE::Core::Private::TToStringType<std::decay_t<PathTypes>>::Type... InPaths)
+	FORCEINLINE static FString CombineImpl(UE::Core::Private::TToStringType_T<std::decay_t<PathTypes>>... InPaths)
 	{
 		const TCHAR* Paths[] = { GetTCharPtr(InPaths)... };
 		FString Out;
