Remove extra logging

diff --git a/Engine/Source/Runtime/Core/Private/Misc/ConfigCacheIni.cpp b/Engine/Source/Runtime/Core/Private/Misc/ConfigCacheIni.cpp
--- a/Engine/Source/Runtime/Core/Private/Misc/ConfigCacheIni.cpp
+++ b/Engine/Source/Runtime/Core/Private/Misc/ConfigCacheIni.cpp
@@ -4217,8 +4217,6 @@ void FConfigCacheIni::InitializeConfigSystem()
 		}
 	}
 
-	FPlatformMisc::LowLevelOutputDebugStringf(TEXT("Loading text-based GConfig...\n"));
-
 	// Perform any upgrade we need before we load any configuration files
 	FConfigManifest::UpgradeFromPreviousVersions();
 
