Added custom object version file for private stream.

diff --git a/Engine/Source/Runtime/Core/Private/UObject/DevObjectVersion.cpp b/Engine/Source/Runtime/Core/Private/UObject/DevObjectVersion.cpp
--- a/Engine/Source/Runtime/Core/Private/UObject/DevObjectVersion.cpp
+++ b/Engine/Source/Runtime/Core/Private/UObject/DevObjectVersion.cpp
@@ -31,6 +31,7 @@
 #include "UObject/VirtualProductionObjectVersion.h"
 #include "UObject/UE5MainStreamObjectVersion.h"
 #include "UObject/UE5ReleaseStreamObjectVersion.h"
+#include "UObject/UE5PrivateFrostyStreamObjectVersion.h"
 #include "UObject/UE5CookerObjectVersion.h"
 #include "UObject/MediaFrameWorkObjectVersion.h"
 
@@ -209,6 +210,10 @@ const FGuid FUE5ReleaseStreamObjectVersion::GUID(0xD89B5E42, 0x24BD4D46, 0x8412A
 // Register UE5 release stream custom version with Core
 FDevVersionRegistration GRegisterUE5ReleaseStreamObjectVersion(FUE5ReleaseStreamObjectVersion::GUID, FUE5ReleaseStreamObjectVersion::LatestVersion, TEXT("UE5-Release"));
 
+// Unique UE5 private frosty version id
+const FGuid FUE5PrivateFrostyStreamObjectVersion::GUID(0x59DA5D52, 0x12324948, 0xB8785978, 0x70B8E98B);
+// Register UE5 private frosty stream custom version with Core
+FDevVersionRegistration GRegisterUE5PrivateFrostyStreamObjectVersion(FUE5PrivateFrostyStreamObjectVersion::GUID, FUE5PrivateFrostyStreamObjectVersion::LatestVersion, TEXT("UE5-PrivateFrosty"));
 
 // Unique Cooker version id
 const FGuid FUE5CookerObjectVersion::GUID(0x26075A32, 0x730F4708, 0x88E98C32, 0xF1599D05);
diff --git a/Engine/Source/Runtime/Core/Public/UObject/UE5PrivateFrostyStreamObjectVersion.h b/Engine/Source/Runtime/Core/Public/UObject/UE5PrivateFrostyStreamObjectVersion.h
new file mode 100644
--- /dev/null
+++ b/Engine/Source/Runtime/Core/Public/UObject/UE5PrivateFrostyStreamObjectVersion.h
@@ -0,0 +1,24 @@
+// Copyright Epic Games, Inc. All Rights Reserved.
+#pragma once
+
+#include "CoreTypes.h"
+#include "Misc/Guid.h"
+
+// Custom serialization version for changes made in //UE5/Private-Frosty stream
+struct CORE_API FUE5PrivateFrostyStreamObjectVersion
+{
+	enum Type
+	{
+		// Before any version changes were made
+		BeforeCustomVersionWasAdded = 0,
+
+		// -----<new versions can be added above this line>-------------------------------------------------
+		VersionPlusOne,
+		LatestVersion = VersionPlusOne - 1
+	};
+
+	// The GUID for this custom version number
+	const static FGuid GUID;
+
+	FUE5PrivateFrostyStreamObjectVersion() = delete;
+};
