Replaced DWORD with Word in BitArray.h (except for the globally defined symbols) and made FConstWordIterator public

diff --git a/Engine/Source/Runtime/Core/Public/Containers/BitArray.h b/Engine/Source/Runtime/Core/Public/Containers/BitArray.h
--- a/Engine/Source/Runtime/Core/Public/Containers/BitArray.h
+++ b/Engine/Source/Runtime/Core/Public/Containers/BitArray.h
@@ -184,12 +184,12 @@ class FRelativeBitReference
 {
 public:
 	FORCEINLINE explicit FRelativeBitReference(int32 BitIndex)
-		: DWORDIndex(BitIndex >> NumBitsPerDWORDLogTwo)
+		: WordIndex(BitIndex >> NumBitsPerDWORDLogTwo)
 		, Mask(1 << (BitIndex & (NumBitsPerDWORD - 1)))
 	{
 	}
 
-	int32  DWORDIndex;
+	int32  WordIndex;
 	uint32 Mask;
 };
 
@@ -456,11 +456,11 @@ public:
 		const int32 UsedBits = (NumBits % NumBitsPerDWORD);
 		if (UsedBits != 0)
 		{
-			const int32 LastDWORDIndex = NumBits / NumBitsPerDWORD;
+			const int32 LastWordIndex = NumBits / NumBitsPerDWORD;
 			const uint32 SlackMask = FullWordMask << UsedBits;
 
-			const uint32 LastDWORD = *(GetData() + LastDWORDIndex);
-			checkf((LastDWORD & SlackMask) == 0, TEXT("TBitArray slack bits are non-zero, this will result in undefined behavior."));
+			const uint32 LastWord = *(GetData() + LastWordIndex);
+			checkf((LastWord & SlackMask) == 0, TEXT("TBitArray slack bits are non-zero, this will result in undefined behavior."));
 		}
 #endif
 	}
@@ -686,12 +686,12 @@ public:
 	{
 		if (Number > MaxBits)
 		{
-			const uint32 MaxDWORDs = AllocatorInstance.CalculateSlackGrow(
+			const uint32 MaxWords = AllocatorInstance.CalculateSlackGrow(
 				FBitSet::CalculateNumWords(Number),
 				GetMaxWords(),
 				sizeof(uint32)
 				);
-			MaxBits = MaxDWORDs * NumBitsPerDWORD;
+			MaxBits = MaxWords * NumBitsPerDWORD;
 			Realloc(NumBits);
 		}
 	}
@@ -738,13 +738,13 @@ public:
 
 		if (InNumBits > MaxBits)
 		{
-			const int32 PreviousNumDWORDs = FBitSet::CalculateNumWords(PreviousNumBits);
-			const uint32 MaxDWORDs = AllocatorInstance.CalculateSlackReserve(
+			const int32 PreviousNumWords = FBitSet::CalculateNumWords(PreviousNumBits);
+			const uint32 MaxWords = AllocatorInstance.CalculateSlackReserve(
 				FBitSet::CalculateNumWords(InNumBits), sizeof(uint32));
 			
-			AllocatorInstance.ResizeAllocation(PreviousNumDWORDs, MaxDWORDs, sizeof(uint32));	
+			AllocatorInstance.ResizeAllocation(PreviousNumWords, MaxWords, sizeof(uint32));	
 
-			MaxBits = MaxDWORDs * NumBitsPerDWORD;
+			MaxBits = MaxWords * NumBitsPerDWORD;
 		}
 
 		ClearPartialSlackBits();
@@ -1162,9 +1162,9 @@ public:
 	 */
 	void BitwiseNOT()
 	{
-		for (FDWORDIterator It(*this); It; ++It)
+		for (FWordIterator It(*this); It; ++It)
 		{
-			It.SetDWORD(~It.GetDWORD());
+			It.SetWord(~It.GetWord());
 		}
 	}
 
@@ -1182,9 +1182,9 @@ public:
 		checkSlow(ToIndex >= FromIndex && ToIndex <= NumBits);
 
 		int32 NumSetBits = 0;
-		for (FConstDWORDIterator It(*this, FromIndex, ToIndex); It; ++It)
+		for (FConstWordIterator It(*this, FromIndex, ToIndex); It; ++It)
 		{
-			NumSetBits += FMath::CountBits(It.GetDWORD());
+			NumSetBits += FMath::CountBits(It.GetWord());
 		}
 		return NumSetBits;
 	}
@@ -1201,16 +1201,16 @@ public:
 	{
 		const uint32 MissingBitsFill = bMissingBitValue ? ~0u : 0;
 
-		FConstDWORDIterator ThisIterator(*this);
-		FConstDWORDIterator OtherIterator(Other);
+		FConstWordIterator ThisIterator(*this);
+		FConstWordIterator OtherIterator(Other);
 
 		ThisIterator.FillMissingBits(MissingBitsFill);
 		OtherIterator.FillMissingBits(MissingBitsFill);
 
 		while (ThisIterator || OtherIterator)
 		{
-			const uint32 A = ThisIterator  ? ThisIterator.GetDWORD()  : MissingBitsFill;
-			const uint32 B = OtherIterator ? OtherIterator.GetDWORD() : MissingBitsFill;
+			const uint32 A = ThisIterator  ? ThisIterator.GetWord()  : MissingBitsFill;
+			const uint32 B = OtherIterator ? OtherIterator.GetWord() : MissingBitsFill;
 			if (A != B)
 			{
 				return false;
@@ -1279,20 +1279,20 @@ public:
 	FORCEINLINE FBitReference AccessCorrespondingBit(const FRelativeBitReference& RelativeReference)
 	{
 		checkSlow(RelativeReference.Mask);
-		checkSlow(RelativeReference.DWORDIndex >= 0);
-		checkSlow(((uint32)RelativeReference.DWORDIndex + 1) * NumBitsPerDWORD - 1 - FMath::CountLeadingZeros(RelativeReference.Mask) < (uint32)NumBits);
+		checkSlow(RelativeReference.WordIndex >= 0);
+		checkSlow(((uint32)RelativeReference.WordIndex + 1) * NumBitsPerDWORD - 1 - FMath::CountLeadingZeros(RelativeReference.Mask) < (uint32)NumBits);
 		return FBitReference(
-			GetData()[RelativeReference.DWORDIndex],
+			GetData()[RelativeReference.WordIndex],
 			RelativeReference.Mask
 			);
 	}
 	FORCEINLINE const FConstBitReference AccessCorrespondingBit(const FRelativeBitReference& RelativeReference) const
 	{
 		checkSlow(RelativeReference.Mask);
-		checkSlow(RelativeReference.DWORDIndex >= 0);
-		checkSlow(((uint32)RelativeReference.DWORDIndex + 1) * NumBitsPerDWORD - 1 - FMath::CountLeadingZeros(RelativeReference.Mask) < (uint32)NumBits);
+		checkSlow(RelativeReference.WordIndex >= 0);
+		checkSlow(((uint32)RelativeReference.WordIndex + 1) * NumBitsPerDWORD - 1 - FMath::CountLeadingZeros(RelativeReference.Mask) < (uint32)NumBits);
 		return FConstBitReference(
-			GetData()[RelativeReference.DWORDIndex],
+			GetData()[RelativeReference.WordIndex],
 			RelativeReference.Mask
 			);
 	}
@@ -1315,7 +1315,7 @@ public:
 			{
 				// Advance to the next uint32.
 				this->Mask = 1;
-				++this->DWORDIndex;
+				++this->WordIndex;
 			}
 			return *this;
 		}
@@ -1330,7 +1330,7 @@ public:
 			return !(bool)*this;
 		}
 
-		FORCEINLINE FBitReference GetValue() const { return FBitReference(Array.GetData()[this->DWORDIndex],this->Mask); }
+		FORCEINLINE FBitReference GetValue() const { return FBitReference(Array.GetData()[this->WordIndex],this->Mask); }
 		FORCEINLINE int32 GetIndex() const { return Index; }
 	private:
 		TBitArray<Allocator>& Array;
@@ -1355,7 +1355,7 @@ public:
 			{
 				// Advance to the next uint32.
 				this->Mask = 1;
-				++this->DWORDIndex;
+				++this->WordIndex;
 			}
 			return *this;
 		}
@@ -1371,7 +1371,7 @@ public:
 			return !(bool)*this;
 		}
 
-		FORCEINLINE FConstBitReference GetValue() const { return FConstBitReference(Array.GetData()[this->DWORDIndex],this->Mask); }
+		FORCEINLINE FConstBitReference GetValue() const { return FConstBitReference(Array.GetData()[this->WordIndex],this->Mask); }
 		FORCEINLINE int32 GetIndex() const { return Index; }
 	private:
 		const TBitArray<Allocator>& Array;
@@ -1396,7 +1396,7 @@ public:
 			{
 				// Advance to the next uint32.
 				this->Mask = (1 << (NumBitsPerDWORD-1));
-				--this->DWORDIndex;
+				--this->WordIndex;
 			}
 			return *this;
 		}
@@ -1412,7 +1412,7 @@ public:
 			return !(bool)*this;
 		}
 
-		FORCEINLINE FConstBitReference GetValue() const { return FConstBitReference(Array.GetData()[this->DWORDIndex],this->Mask); }
+		FORCEINLINE FConstBitReference GetValue() const { return FConstBitReference(Array.GetData()[this->WordIndex],this->Mask); }
 		FORCEINLINE int32 GetIndex() const { return Index; }
 	private:
 		const TBitArray<Allocator>& Array;
@@ -1445,15 +1445,15 @@ private:
 				OutResult.Reserve(MinNumBits);
 				OutResult.NumBits = MinNumBits;
 
-				FConstDWORDIterator IteratorA(InA);
-				FConstDWORDIterator IteratorB(InB);
+				FConstWordIterator IteratorA(InA);
+				FConstWordIterator IteratorB(InB);
 
-				FDWORDIterator IteratorResult(OutResult);
+				FWordIterator IteratorResult(OutResult);
 
 				for ( ; IteratorResult; ++IteratorResult, ++IteratorA, ++IteratorB)
 				{
-					const uint32 NewValue = Invoke(InProjection, IteratorA.GetDWORD(), IteratorB.GetDWORD());
-					IteratorResult.SetDWORD(NewValue);
+					const uint32 NewValue = Invoke(InProjection, IteratorA.GetWord(), IteratorB.GetWord());
+					IteratorResult.SetWord(NewValue);
 				}
 			}
 
@@ -1469,20 +1469,20 @@ private:
 				OutResult.Reserve(MaxNumBits);
 				OutResult.NumBits = MaxNumBits;
 
-				FConstDWORDIterator IteratorA(InA);
-				FConstDWORDIterator IteratorB(InB);
+				FConstWordIterator IteratorA(InA);
+				FConstWordIterator IteratorB(InB);
 
 				IteratorA.FillMissingBits(MissingBitsFill);
 				IteratorB.FillMissingBits(MissingBitsFill);
 
-				FDWORDIterator IteratorResult(OutResult);
+				FWordIterator IteratorResult(OutResult);
 
 				for ( ; IteratorResult; ++IteratorResult, ++IteratorA, ++IteratorB)
 				{
-					uint32 A = IteratorA ? IteratorA.GetDWORD() : MissingBitsFill;
-					uint32 B = IteratorB ? IteratorB.GetDWORD() : MissingBitsFill;
+					uint32 A = IteratorA ? IteratorA.GetWord() : MissingBitsFill;
+					uint32 B = IteratorB ? IteratorB.GetWord() : MissingBitsFill;
 
-					IteratorResult.SetDWORD(Invoke(InProjection, A, B));
+					IteratorResult.SetWord(Invoke(InProjection, A, B));
 				}
 			}
 
@@ -1526,15 +1526,15 @@ private:
 		const uint32 MissingBitsFill = EnumHasAnyFlags(InFlags, EBitwiseOperatorFlags::OneFillMissingBits) ? ~0u : 0;
 		if (OutResult.NumBits != 0)
 		{
-			FConstDWORDIterator IteratorOther(InOther);
+			FConstWordIterator IteratorOther(InOther);
 			IteratorOther.FillMissingBits(MissingBitsFill);
 
-			FDWORDIterator IteratorResult(OutResult);
+			FWordIterator IteratorResult(OutResult);
 
 			for ( ; IteratorResult; ++IteratorResult, ++IteratorOther)
 			{
-				const uint32 OtherValue = IteratorOther ? IteratorOther.GetDWORD() : MissingBitsFill;
-				IteratorResult.SetDWORD(Invoke(InProjection, IteratorResult.GetDWORD(), OtherValue));
+				const uint32 OtherValue = IteratorOther ? IteratorOther.GetWord() : MissingBitsFill;
+				IteratorResult.SetWord(Invoke(InProjection, IteratorResult.GetWord(), OtherValue));
 			}
 		}
 
@@ -1542,12 +1542,12 @@ private:
 	}
 
 
-	template<typename DWORDType>
-	struct TDWORDIteratorBase
+	template<typename WordType>
+	struct TWordIteratorBase
 	{
 		explicit operator bool() const
 		{
-			return CurrentIndex < NumDWORDs;
+			return CurrentIndex < NumWords;
 		}
 
 		int32 GetIndex() const
@@ -1555,9 +1555,9 @@ private:
 			return CurrentIndex;
 		}
 
-		uint32 GetDWORD() const
+		uint32 GetWord() const
 		{
-			checkSlow(CurrentIndex < NumDWORDs);
+			checkSlow(CurrentIndex < NumWords);
 
 			if (CurrentMask == ~0u)
 			{
@@ -1576,7 +1576,7 @@ private:
 		void operator++()
 		{
 			++this->CurrentIndex;
-			if (this->CurrentIndex == NumDWORDs-1)
+			if (this->CurrentIndex == NumWords-1)
 			{
 				CurrentMask = FinalMask;
 			}
@@ -1593,10 +1593,10 @@ private:
 
 	protected:
 
-		explicit TDWORDIteratorBase(DWORDType* InData, int32 InStartBitIndex, int32 InEndBitIndex)
+		explicit TWordIteratorBase(WordType* InData, int32 InStartBitIndex, int32 InEndBitIndex)
 			: Data(InData)
 			, CurrentIndex(InStartBitIndex / NumBitsPerDWORD)
-			, NumDWORDs(FMath::DivideAndRoundUp(InEndBitIndex, NumBitsPerDWORD))
+			, NumWords(FMath::DivideAndRoundUp(InEndBitIndex, NumBitsPerDWORD))
 			, CurrentMask(~0u << (InStartBitIndex % NumBitsPerDWORD))
 			, FinalMask(~0u)
 			, MissingBitsFill(0)
@@ -1607,54 +1607,56 @@ private:
 				FinalMask = ~0u >> Shift;
 			}
 
-			if (CurrentIndex == NumDWORDs - 1)
+			if (CurrentIndex == NumWords - 1)
 			{
 				CurrentMask &= FinalMask;
 				FinalMask = CurrentMask;
 			}
 		}
 
-		DWORDType* RESTRICT Data;
+		WordType* RESTRICT Data;
 
 		int32 CurrentIndex;
-		int32 NumDWORDs;
+		int32 NumWords;
 
 		uint32 CurrentMask;
 		uint32 FinalMask;
 		uint32 MissingBitsFill;
 	};
 
-	struct FConstDWORDIterator : TDWORDIteratorBase<const uint32>
+public:
+	struct FConstWordIterator : TWordIteratorBase<const uint32>
 	{
-		explicit FConstDWORDIterator(const TBitArray<Allocator>& InArray)
-			: TDWORDIteratorBase<const uint32>(InArray.GetData(), 0, InArray.Num())
+		explicit FConstWordIterator(const TBitArray<Allocator>& InArray)
+			: TWordIteratorBase<const uint32>(InArray.GetData(), 0, InArray.Num())
 		{}
 
-		explicit FConstDWORDIterator(const TBitArray<Allocator>& InArray, int32 InStartBitIndex, int32 InEndBitIndex)
-			: TDWORDIteratorBase<const uint32>(InArray.GetData(), InStartBitIndex, InEndBitIndex)
+		explicit FConstWordIterator(const TBitArray<Allocator>& InArray, int32 InStartBitIndex, int32 InEndBitIndex)
+			: TWordIteratorBase<const uint32>(InArray.GetData(), InStartBitIndex, InEndBitIndex)
 		{
 			checkSlow(InStartBitIndex <= InEndBitIndex && InStartBitIndex <= InArray.Num() && InEndBitIndex <= InArray.Num());
 			checkSlow(InStartBitIndex >= 0 && InEndBitIndex >= 0);
 		}
 	};
 
-	struct FDWORDIterator : TDWORDIteratorBase<uint32>
+
+	struct FWordIterator : TWordIteratorBase<uint32>
 	{
-		explicit FDWORDIterator(TBitArray<Allocator>& InArray)
-			: TDWORDIteratorBase<uint32>(InArray.GetData(), 0, InArray.Num())
+		explicit FWordIterator(TBitArray<Allocator>& InArray)
+			: TWordIteratorBase<uint32>(InArray.GetData(), 0, InArray.Num())
 		{}
 
-		void SetDWORD(uint32 InDWORD)
+		void SetWord(uint32 InWord)
 		{
-			checkSlow(this->CurrentIndex < this->NumDWORDs);
+			checkSlow(this->CurrentIndex < this->NumWords);
 
-			if (this->CurrentIndex == this->NumDWORDs-1)
+			if (this->CurrentIndex == this->NumWords-1)
 			{
-				this->Data[this->CurrentIndex] = InDWORD & this->FinalMask;
+				this->Data[this->CurrentIndex] = InWord & this->FinalMask;
 			}
 			else
 			{
-				this->Data[this->CurrentIndex] = InDWORD;
+				this->Data[this->CurrentIndex] = InWord;
 			}
 		}
 	};
@@ -1666,10 +1668,10 @@ private:
 
 	FORCENOINLINE void Realloc(int32 PreviousNumBits)
 	{
-		const uint32 PreviousNumDWORDs = FBitSet::CalculateNumWords(PreviousNumBits);
-		const uint32 MaxDWORDs = FBitSet::CalculateNumWords(MaxBits);
+		const uint32 PreviousNumWords = FBitSet::CalculateNumWords(PreviousNumBits);
+		const uint32 MaxWords = FBitSet::CalculateNumWords(MaxBits);
 
-		AllocatorInstance.ResizeAllocation(PreviousNumDWORDs,MaxDWORDs,sizeof(uint32));
+		AllocatorInstance.ResizeAllocation(PreviousNumWords,MaxWords,sizeof(uint32));
 		ClearPartialSlackBits(); // Implement class invariant
 	}
 
@@ -1681,7 +1683,7 @@ private:
 	}
 
 	/**
-	 * Clears the slack bits within the final partially relevant DWORD
+	 * Clears the slack bits within the final partially relevant Word
 	 */
 	void ClearPartialSlackBits()
 	{
@@ -1691,11 +1693,11 @@ private:
 		const int32 UsedBits = NumBits % NumBitsPerDWORD;
 		if (UsedBits != 0)
 		{
-			const int32  LastDWORDIndex = NumBits / NumBitsPerDWORD;
+			const int32  LastWordIndex = NumBits / NumBitsPerDWORD;
 			const uint32 SlackMask = FullWordMask >> (NumBitsPerDWORD - UsedBits);
 
-			uint32* LastDWORD = (GetData() + LastDWORDIndex);
-			*LastDWORD = *LastDWORD & SlackMask;
+			uint32* LastWord = (GetData() + LastWordIndex);
+			*LastWord = *LastWord & SlackMask;
 		}
 	}
 
@@ -1710,8 +1712,8 @@ private:
 	{
 		static void WriteMemoryImage(FMemoryImageWriter& Writer, const TBitArray& Object)
 		{
-			const int32 NumDWORDs = FMath::DivideAndRoundUp(Object.NumBits, NumBitsPerDWORD);
-			Object.AllocatorInstance.WriteMemoryImage(Writer, StaticGetTypeLayoutDesc<uint32>(), NumDWORDs);
+			const int32 NumWords = FMath::DivideAndRoundUp(Object.NumBits, NumBitsPerDWORD);
+			Object.AllocatorInstance.WriteMemoryImage(Writer, StaticGetTypeLayoutDesc<uint32>(), NumWords);
 			Writer.WriteBytes(Object.NumBits);
 			Writer.WriteBytes(Object.NumBits);
 		}
@@ -1832,22 +1834,22 @@ private:
 	{
 		const uint32* ArrayData      = Array.GetData();
 		const int32   ArrayNum       = Array.Num();
-		const int32   LastDWORDIndex = (ArrayNum - 1) / NumBitsPerDWORD;
+		const int32   LastWordIndex = (ArrayNum - 1) / NumBitsPerDWORD;
 
 		// Advance to the next non-zero uint32.
-		uint32 RemainingBitMask = ArrayData[this->DWORDIndex] & UnvisitedBitMask;
+		uint32 RemainingBitMask = ArrayData[this->WordIndex] & UnvisitedBitMask;
 		while (!RemainingBitMask)
 		{
-			++this->DWORDIndex;
+			++this->WordIndex;
 			BaseBitIndex += NumBitsPerDWORD;
-			if (this->DWORDIndex > LastDWORDIndex)
+			if (this->WordIndex > LastWordIndex)
 			{
 				// We've advanced past the end of the array.
 				CurrentBitIndex = ArrayNum;
 				return;
 			}
 
-			RemainingBitMask = ArrayData[this->DWORDIndex];
+			RemainingBitMask = ArrayData[this->WordIndex];
 			UnvisitedBitMask = ~0;
 		}
 
@@ -1861,7 +1863,7 @@ private:
 		// If the Nth bit was the lowest set bit of BitMask, then this gives us N
 		CurrentBitIndex = BaseBitIndex + NumBitsPerDWORD - 1 - FMath::CountLeadingZeros(this->Mask);
 
-		// If we've accidentally iterated off the end of an array but still within the same DWORD
+		// If we've accidentally iterated off the end of an array but still within the same Word
 		// then set the index to the last index of the array
 		if (CurrentBitIndex > ArrayNum)
 		{
@@ -1949,27 +1951,27 @@ private:
 		
 		if (Both)
 		{
-			RemainingBitMask = ArrayDataA[this->DWORDIndex] & ArrayDataB[this->DWORDIndex] & UnvisitedBitMask;
+			RemainingBitMask = ArrayDataA[this->WordIndex] & ArrayDataB[this->WordIndex] & UnvisitedBitMask;
 		}
 		else
 		{
-			RemainingBitMask = (ArrayDataA[this->DWORDIndex] | ArrayDataB[this->DWORDIndex]) & UnvisitedBitMask;
+			RemainingBitMask = (ArrayDataA[this->WordIndex] | ArrayDataB[this->WordIndex]) & UnvisitedBitMask;
 		}
 
 		while(!RemainingBitMask)
 		{
-			this->DWORDIndex++;
+			this->WordIndex++;
 			BaseBitIndex += NumBitsPerDWORD;
-			const int32 LastDWORDIndex = (ArrayA.Num() - 1) / NumBitsPerDWORD;
-			if (this->DWORDIndex <= LastDWORDIndex)
+			const int32 LastWordIndex = (ArrayA.Num() - 1) / NumBitsPerDWORD;
+			if (this->WordIndex <= LastWordIndex)
 			{
 				if (Both)
 				{
-					RemainingBitMask = ArrayDataA[this->DWORDIndex] & ArrayDataB[this->DWORDIndex];
+					RemainingBitMask = ArrayDataA[this->WordIndex] & ArrayDataB[this->WordIndex];
 				}
 				else
 				{
-					RemainingBitMask = ArrayDataA[this->DWORDIndex] | ArrayDataB[this->DWORDIndex];
+					RemainingBitMask = ArrayDataA[this->WordIndex] | ArrayDataB[this->WordIndex];
 				}
 
 				UnvisitedBitMask = ~0;
@@ -2117,36 +2119,36 @@ private:
 
 	FORCENOINLINE void Realloc(int32 PreviousNumBits)
 	{
-		const uint32 MaxDWORDs = AllocatorInstance.CalculateSlackReserve(
+		const uint32 MaxWords = AllocatorInstance.CalculateSlackReserve(
 			FBitSet::CalculateNumWords(MaxBits),
 			sizeof(uint32)
 			);
-		MaxBits = MaxDWORDs * NumBitsPerDWORD;
-		const uint32 PreviousNumDWORDs = FBitSet::CalculateNumWords(PreviousNumBits);
+		MaxBits = MaxWords * NumBitsPerDWORD;
+		const uint32 PreviousNumWords = FBitSet::CalculateNumWords(PreviousNumBits);
 
-		AllocatorInstance.ResizeAllocation(PreviousNumDWORDs, MaxDWORDs, sizeof(uint32));
+		AllocatorInstance.ResizeAllocation(PreviousNumWords, MaxWords, sizeof(uint32));
 
-		if (MaxDWORDs && MaxDWORDs > PreviousNumDWORDs)
+		if (MaxWords && MaxWords > PreviousNumWords)
 		{
-			// Reset the newly allocated slack DWORDs.
-			FMemory::Memzero((uint32*)AllocatorInstance.GetAllocation() + PreviousNumDWORDs, (MaxDWORDs - PreviousNumDWORDs) * sizeof(uint32));
+			// Reset the newly allocated slack Words.
+			FMemory::Memzero((uint32*)AllocatorInstance.GetAllocation() + PreviousNumWords, (MaxWords - PreviousNumWords) * sizeof(uint32));
 		}
 	}
 	FORCENOINLINE void ReallocGrow(int32 PreviousNumBits)
 	{
 		// Allocate memory for the new bits.
-		const uint32 MaxDWORDs = AllocatorInstance.CalculateSlackGrow(
+		const uint32 MaxWords = AllocatorInstance.CalculateSlackGrow(
 			FBitSet::CalculateNumWords(NumBits),
 			FBitSet::CalculateNumWords(MaxBits),
 			sizeof(uint32)
 			);
-		MaxBits = MaxDWORDs * NumBitsPerDWORD;
-		const uint32 PreviousNumDWORDs = FBitSet::CalculateNumWords(PreviousNumBits);
-		AllocatorInstance.ResizeAllocation(PreviousNumDWORDs, MaxDWORDs, sizeof(uint32));
-		if (MaxDWORDs && MaxDWORDs > PreviousNumDWORDs)
+		MaxBits = MaxWords * NumBitsPerDWORD;
+		const uint32 PreviousNumWords = FBitSet::CalculateNumWords(PreviousNumBits);
+		AllocatorInstance.ResizeAllocation(PreviousNumWords, MaxWords, sizeof(uint32));
+		if (MaxWords && MaxWords > PreviousNumWords)
 		{
-			// Reset the newly allocated slack DWORDs.
-			FMemory::Memzero((uint32*)AllocatorInstance.GetAllocation() + PreviousNumDWORDs, (MaxDWORDs - PreviousNumDWORDs) * sizeof(uint32));
+			// Reset the newly allocated slack Words.
+			FMemory::Memzero((uint32*)AllocatorInstance.GetAllocation() + PreviousNumWords, (MaxWords - PreviousNumWords) * sizeof(uint32));
 		}
 	}
 
