U5 - refactored local queue behaviour so that they can be shared between differnt schedulers.

diff --git a/Engine/Source/Runtime/Core/Private/Async/Fundamental/Scheduler.cpp b/Engine/Source/Runtime/Core/Private/Async/Fundamental/Scheduler.cpp
--- a/Engine/Source/Runtime/Core/Private/Async/Fundamental/Scheduler.cpp
+++ b/Engine/Source/Runtime/Core/Private/Async/Fundamental/Scheduler.cpp
@@ -11,7 +11,7 @@ namespace LowLevelTasks
 {
 	DEFINE_LOG_CATEGORY(LowLevelTasks);
 
-	thread_local FScheduler::FLocalQueueType* FScheduler::LocalQueue = nullptr;
+	thread_local FSchedulerTls::FLocalQueueType* FSchedulerTls::LocalQueue = nullptr;
 	thread_local FTask* FSchedulerTls::ActiveTask = nullptr;
 	thread_local FSchedulerTls* FSchedulerTls::ActiveScheduler = nullptr;
 	thread_local FSchedulerTls::EWorkerType FSchedulerTls::WorkerType = FSchedulerTls::EWorkerType::None;
@@ -21,11 +21,11 @@ namespace LowLevelTasks
 
 	FScheduler::FLocalQueueInstaller::FLocalQueueInstaller(FScheduler& Scheduler)
 	{
-		RegisteredLocalQueue = LocalQueue == nullptr;
+		RegisteredLocalQueue = FSchedulerTls::LocalQueue == nullptr;
 		if (RegisteredLocalQueue)
 		{
 			bool bPermitBackgroundWork = FSchedulerTls::PermitBackgroundWork();
-			LocalQueue = FLocalQueueType::AllocateLocalQueue(Scheduler.QueueRegistry, bPermitBackgroundWork);
+			FSchedulerTls::LocalQueue = FSchedulerTls::FLocalQueueType::AllocateLocalQueue(Scheduler.QueueRegistry, bPermitBackgroundWork ? ELocalQueueType::EBackground : ELocalQueueType::EForeground);
 		}
 	}
 
@@ -34,12 +34,12 @@ namespace LowLevelTasks
 		if (RegisteredLocalQueue)
 		{
 			bool bPermitBackgroundWork = FSchedulerTls::PermitBackgroundWork();
-			FLocalQueueType::DeleteLocalQueue(LocalQueue, bPermitBackgroundWork);
-			LocalQueue = nullptr;
+			FSchedulerTls::FLocalQueueType::DeleteLocalQueue(FSchedulerTls::LocalQueue, bPermitBackgroundWork ? ELocalQueueType::EBackground : ELocalQueueType::EForeground);
+			FSchedulerTls::LocalQueue = nullptr;
 		}
 	}
 
-	TUniquePtr<FThread> FScheduler::CreateWorker(FLocalQueueType* ExternalWorkerLocalQueue, EThreadPriority Priority, bool bPermitBackgroundWork, bool bIsForkable)
+	TUniquePtr<FThread> FScheduler::CreateWorker(FSchedulerTls::FLocalQueueType* ExternalWorkerLocalQueue, EThreadPriority Priority, bool bPermitBackgroundWork, bool bIsForkable)
 	{
 		uint32 WorkerId = NextWorkerId++;
 		const uint32 WaitTimes[8] = { 719, 991, 1361, 1237, 1597, 953, 587, 1439 };
@@ -102,14 +102,14 @@ namespace LowLevelTasks
 			UE::Trace::ThreadGroupBegin(TEXT("Foreground Workers"));
 			for (uint32 WorkerId = 0; WorkerId < NumForegroundWorkers; ++WorkerId)
 			{
-				WorkerLocalQueues.Emplace(QueueRegistry, false);
+				WorkerLocalQueues.Emplace(QueueRegistry, ELocalQueueType::EForeground);
 				WorkerThreads.Add(CreateWorker(&WorkerLocalQueues.Last(), WorkerPriority, false, bIsForkable));
 			}
 			UE::Trace::ThreadGroupEnd();
 			UE::Trace::ThreadGroupBegin(TEXT("Background Workers"));
 			for (uint32 WorkerId = 0; WorkerId < NumBackgroundWorkers; ++WorkerId)
 			{
-				WorkerLocalQueues.Emplace(QueueRegistry, true);
+				WorkerLocalQueues.Emplace(QueueRegistry, ELocalQueueType::EBackground);
 				WorkerThreads.Add(CreateWorker(&WorkerLocalQueues.Last(), BackgroundPriority, true, bIsForkable));
 			}
 			UE::Trace::ThreadGroupEnd();
@@ -150,11 +150,11 @@ namespace LowLevelTasks
 				QueuePreference = EQueuePreference::GlobalQueuePreference;
 			}
 
-			bWakeUpWorker |= LocalQueue == nullptr;
+			bWakeUpWorker |= FSchedulerTls::LocalQueue == nullptr;
 
-			if (LocalQueue && QueuePreference != EQueuePreference::GlobalQueuePreference)
+			if (FSchedulerTls::LocalQueue && QueuePreference != EQueuePreference::GlobalQueuePreference)
 			{
-				if (LocalQueue->Enqueue(&Task, uint32(Task.GetPriority())))
+				if (FSchedulerTls::LocalQueue->Enqueue(&Task, uint32(Task.GetPriority())))
 				{
 					if(bWakeUpWorker && !WakeUpWorker(bIsBackgroundTask) && !bIsBackgroundTask)
 					{
@@ -207,8 +207,8 @@ namespace LowLevelTasks
 		}
 	}
 
-	template<FTask* (FScheduler::FLocalQueueType::*DequeueFunction)(bool), bool bIsBusyWaiting>
-	bool FScheduler::TryExecuteTaskFrom(FLocalQueueType* Queue, FQueueRegistry::FOutOfWork& OutOfWork, bool bPermitBackgroundWork)
+	template<FTask* (FSchedulerTls::FLocalQueueType::*DequeueFunction)(bool), bool bIsBusyWaiting>
+	bool FScheduler::TryExecuteTaskFrom(FSchedulerTls::FLocalQueueType* Queue, FSchedulerTls::FQueueRegistry::FOutOfWork& OutOfWork, bool bPermitBackgroundWork)
 	{
 		for(int i = 0; i < 2; i++) // one retry if we pick up a task that cannot used during busy waiting.
 		{
@@ -235,7 +235,7 @@ namespace LowLevelTasks
 		return false;
 	}
 
-	void FScheduler::WorkerMain(FSleepEvent* WorkerEvent, FLocalQueueType* ExternalWorkerLocalQueue, uint32 WaitCycles, bool bPermitBackgroundWork)
+	void FScheduler::WorkerMain(FSleepEvent* WorkerEvent, FSchedulerTls::FLocalQueueType* ExternalWorkerLocalQueue, uint32 WaitCycles, bool bPermitBackgroundWork)
 	{
 		FTaskTagScope WorkerScope(ETaskTag::EWorkerThread);
 		FSchedulerTls::ActiveScheduler = this;
@@ -243,31 +243,31 @@ namespace LowLevelTasks
 		FMemory::SetupTLSCachesOnCurrentThread();
 		FSchedulerTls::WorkerType = bPermitBackgroundWork ? FSchedulerTls::EWorkerType::Background : FSchedulerTls::EWorkerType::Foreground;
 
-		checkSlow(LocalQueue == nullptr);
+		checkSlow(FSchedulerTls::LocalQueue == nullptr);
 		if(ExternalWorkerLocalQueue)
 		{
-			LocalQueue = ExternalWorkerLocalQueue;
+			FSchedulerTls::LocalQueue = ExternalWorkerLocalQueue;
 		}
 		else
 		{
-			LocalQueue = FLocalQueueType::AllocateLocalQueue(QueueRegistry, bPermitBackgroundWork);
+			FSchedulerTls::LocalQueue = FSchedulerTls::FLocalQueueType::AllocateLocalQueue(QueueRegistry, bPermitBackgroundWork ? ELocalQueueType::EBackground : ELocalQueueType::EForeground);
 		}
-		FLocalQueueType* WorkerLocalQueue = LocalQueue;
+		FSchedulerTls::FLocalQueueType* WorkerLocalQueue = FSchedulerTls::LocalQueue;
 
 		bool Drowsing = false;
 		uint32 WaitCount = 0;
-		FQueueRegistry::FOutOfWork OutOfWork = QueueRegistry.GetOutOfWorkScope(bPermitBackgroundWork);
+		FSchedulerTls::FQueueRegistry::FOutOfWork OutOfWork = QueueRegistry.GetOutOfWorkScope(bPermitBackgroundWork ? ELocalQueueType::EBackground : ELocalQueueType::EForeground);
 		while (true)
 		{
-			while(TryExecuteTaskFrom<&FLocalQueueType::DequeueLocal, false>(WorkerLocalQueue, OutOfWork, bPermitBackgroundWork)
-			   || TryExecuteTaskFrom<&FLocalQueueType::DequeueGlobal, false>(WorkerLocalQueue, OutOfWork, bPermitBackgroundWork))
+			while(TryExecuteTaskFrom<&FSchedulerTls::FLocalQueueType::DequeueLocal, false>(WorkerLocalQueue, OutOfWork, bPermitBackgroundWork)
+			   || TryExecuteTaskFrom<&FSchedulerTls::FLocalQueueType::DequeueGlobal, false>(WorkerLocalQueue, OutOfWork, bPermitBackgroundWork))
 			{		
 				Drowsing = false;
 				WaitCount = 0;
 			}
 
-			while(TryExecuteTaskFrom<&FLocalQueueType::DequeueLocal, false>(WorkerLocalQueue, OutOfWork, bPermitBackgroundWork)
-			   || TryExecuteTaskFrom<&FLocalQueueType::DequeueSteal, false>(WorkerLocalQueue, OutOfWork, bPermitBackgroundWork))
+			while(TryExecuteTaskFrom<&FSchedulerTls::FLocalQueueType::DequeueLocal, false>(WorkerLocalQueue, OutOfWork, bPermitBackgroundWork)
+			   || TryExecuteTaskFrom<&FSchedulerTls::FLocalQueueType::DequeueSteal, false>(WorkerLocalQueue, OutOfWork, bPermitBackgroundWork))
 			{
 				Drowsing = false;
 				WaitCount = 0;
@@ -291,8 +291,8 @@ namespace LowLevelTasks
 
 		while (WakeUpWorker(bPermitBackgroundWork)) {}
 
-		FLocalQueueType::DeleteLocalQueue(WorkerLocalQueue, bPermitBackgroundWork, ExternalWorkerLocalQueue != nullptr);
-		LocalQueue = nullptr;
+		FSchedulerTls::FLocalQueueType::DeleteLocalQueue(WorkerLocalQueue, bPermitBackgroundWork ? ELocalQueueType::EBackground : ELocalQueueType::EForeground, ExternalWorkerLocalQueue != nullptr);
+		FSchedulerTls::LocalQueue = nullptr;
 
 		FSchedulerTls::ActiveScheduler = nullptr;
 		FSchedulerTls::WorkerType = FSchedulerTls::EWorkerType::None;
@@ -307,19 +307,19 @@ namespace LowLevelTasks
 		++FSchedulerTls::BusyWaitingDepth;
 		ON_SCOPE_EXIT{ --FSchedulerTls::BusyWaitingDepth; };
 
-		checkSlow(LocalQueue != nullptr);
+		checkSlow(FSchedulerTls::LocalQueue != nullptr);
 		check(ActiveWorkers.load(std::memory_order_relaxed));
-		FLocalQueueType* WorkerLocalQueue = LocalQueue;
+		FSchedulerTls::FLocalQueueType* WorkerLocalQueue = FSchedulerTls::LocalQueue;
 
 		uint32 WaitCount = 0;
 		bool HasWokenEmergencyWorker = false;
 		const bool bIsBackgroundWorker = FSchedulerTls::IsBackgroundWorker();
 		bool bPermitBackgroundWork = FSchedulerTls::PermitBackgroundWork() || ForceAllowBackgroundWork;
-		FQueueRegistry::FOutOfWork OutOfWork = QueueRegistry.GetOutOfWorkScope(bIsBackgroundWorker);
+		FSchedulerTls::FQueueRegistry::FOutOfWork OutOfWork = QueueRegistry.GetOutOfWorkScope(bPermitBackgroundWork ? ELocalQueueType::EBackground : ELocalQueueType::EForeground);
 		while (true)
 		{
-			while(TryExecuteTaskFrom<&FLocalQueueType::DequeueLocal, true>(WorkerLocalQueue, OutOfWork, bPermitBackgroundWork)
-			   || TryExecuteTaskFrom<&FLocalQueueType::DequeueGlobal, true>(WorkerLocalQueue, OutOfWork, bPermitBackgroundWork))
+			while(TryExecuteTaskFrom<&FSchedulerTls::FLocalQueueType::DequeueLocal, true>(WorkerLocalQueue, OutOfWork, bPermitBackgroundWork)
+			   || TryExecuteTaskFrom<&FSchedulerTls::FLocalQueueType::DequeueGlobal, true>(WorkerLocalQueue, OutOfWork, bPermitBackgroundWork))
 			{
 				if (Conditional())
 				{
@@ -328,8 +328,8 @@ namespace LowLevelTasks
 				WaitCount = 0;
 			}
 
-			while(TryExecuteTaskFrom<&FLocalQueueType::DequeueLocal, true>(WorkerLocalQueue, OutOfWork, bPermitBackgroundWork)
-			   || TryExecuteTaskFrom<&FLocalQueueType::DequeueSteal, true>(WorkerLocalQueue, OutOfWork, bPermitBackgroundWork))
+			while(TryExecuteTaskFrom<&FSchedulerTls::FLocalQueueType::DequeueLocal, true>(WorkerLocalQueue, OutOfWork, bPermitBackgroundWork)
+			   || TryExecuteTaskFrom<&FSchedulerTls::FLocalQueueType::DequeueSteal, true>(WorkerLocalQueue, OutOfWork, bPermitBackgroundWork))
 			{
 				if (Conditional())
 				{
diff --git a/Engine/Source/Runtime/Core/Private/HAL/ThreadingBase.cpp b/Engine/Source/Runtime/Core/Private/HAL/ThreadingBase.cpp
--- a/Engine/Source/Runtime/Core/Private/HAL/ThreadingBase.cpp
+++ b/Engine/Source/Runtime/Core/Private/HAL/ThreadingBase.cpp
@@ -160,7 +160,7 @@ CORE_API bool IsInGameThread()
 	{
 		bool newValue = FTaskTagScope::IsCurrentTag(ETaskTag::EGameThread) || FTaskTagScope::IsRunningDuringStaticInit();
 #if !UE_BUILD_SHIPPING && !UE_BUILD_TEST
-		if (!LowLevelTasks::FScheduler::IsBusyWaiting())
+		if (!LowLevelTasks::FSchedulerTls::IsBusyWaiting())
 		{
 			const uint32 CurrentThreadId = FPlatformTLS::GetCurrentThreadId();
 			bool oldValue = CurrentThreadId == GGameThreadId;
@@ -184,7 +184,7 @@ CORE_API bool IsInSlateThread()
 	// If this explicitly is a slate thread, not just the main thread running slate
 	bool newValue = FTaskTagScope::IsCurrentTag(ETaskTag::ESlateThread);
 #if !UE_BUILD_SHIPPING && !UE_BUILD_TEST
-	if (!LowLevelTasks::FScheduler::IsBusyWaiting())
+	if (!LowLevelTasks::FSchedulerTls::IsBusyWaiting())
 	{
 		bool oldValue = GSlateLoadingThreadId != 0 && FPlatformTLS::GetCurrentThreadId() == GSlateLoadingThreadId;
 		ensureMsgf(oldValue == newValue, TEXT("oldValue(%i) newValue(%i) If this check fails make sure that there is a FTaskTagScope(ETaskTag::ESlateThread) as as deep as possible on the current callstack, you can see the current value in ActiveNamedThreads(%x)"), oldValue, newValue, FTaskTagScope::GetCurrentTag());
@@ -225,7 +225,7 @@ PRAGMA_DISABLE_DEPRECATION_WARNINGS
 		? FTaskTagScope::IsCurrentTag(ETaskTag::EGameThread)
 		: FTaskTagScope::IsCurrentTag(ETaskTag::EAudioThread);
 #if !UE_BUILD_SHIPPING && !UE_BUILD_TEST
-	if (!LowLevelTasks::FScheduler::IsBusyWaiting())
+	if (!LowLevelTasks::FSchedulerTls::IsBusyWaiting())
 	{
 		bool oldValue = FPlatformTLS::GetCurrentThreadId() == ((nullptr == GAudioThread || GIsAudioThreadSuspended.load(std::memory_order_relaxed)) ? GGameThreadId : GAudioThread->GetThreadID());
 		ensureMsgf(oldValue == newValue, TEXT("oldValue(%i) newValue(%i) If this check fails make sure that there is a FTaskTagScope(ETaskTag::EAudioThread) as deep as possible on the current callstack, you can see the current value in ActiveNamedThreads(%x)"), oldValue, newValue, FTaskTagScope::GetCurrentTag());
@@ -247,7 +247,7 @@ CORE_API bool IsInActualRenderingThread()
 PRAGMA_DISABLE_DEPRECATION_WARNINGS
 	bool newValue = FTaskTagScope::IsCurrentTag(ETaskTag::ERenderingThread);
 #if !UE_BUILD_SHIPPING && !UE_BUILD_TEST
-	if (!LowLevelTasks::FScheduler::IsBusyWaiting())
+	if (!LowLevelTasks::FSchedulerTls::IsBusyWaiting())
 	{
 		bool oldValue = FPlatformTLS::GetCurrentThreadId() == GRenderThreadId;
 		ensureMsgf(oldValue == newValue, TEXT("oldValue(%i) newValue(%i) If this check fails make sure that there is a FTaskTagScope(ETaskTag::ERenderingThread) as deep as possible on the current callstack, you can see the current value in ActiveNamedThreads(%x)"), oldValue, newValue, FTaskTagScope::GetCurrentTag());
@@ -268,7 +268,7 @@ PRAGMA_DISABLE_DEPRECATION_WARNINGS
 		: FTaskTagScope::IsCurrentTag(ETaskTag::ERenderingThread);
 
 #if !UE_BUILD_SHIPPING && !UE_BUILD_TEST
-	if (!LowLevelTasks::FScheduler::IsBusyWaiting())
+	if (!LowLevelTasks::FSchedulerTls::IsBusyWaiting())
 	{
 		const uint32 CurrentThreadId = FPlatformTLS::GetCurrentThreadId();
 		bool oldValue = ((GRenderThreadId == 0) || bLocalIsLoadingThreadSuspended) ? (CurrentThreadId == GGameThreadId) || FTaskTagScope::IsRunningDuringStaticInit() : (CurrentThreadId == GRenderThreadId);
@@ -299,7 +299,7 @@ PRAGMA_DISABLE_DEPRECATION_WARNINGS
 	}
 
 #if !UE_BUILD_SHIPPING && !UE_BUILD_TEST
-	if (!LowLevelTasks::FScheduler::IsBusyWaiting())
+	if (!LowLevelTasks::FSchedulerTls::IsBusyWaiting())
 	{
 		const uint32 CurrentThreadId = FPlatformTLS::GetCurrentThreadId();
 		bool oldValue = ((GRenderThreadId == 0) || bLocalIsLoadingThreadSuspended) ?  true : CurrentThreadId != GGameThreadId;
diff --git a/Engine/Source/Runtime/Core/Public/Async/Fundamental/LocalQueue.h b/Engine/Source/Runtime/Core/Public/Async/Fundamental/LocalQueue.h
--- a/Engine/Source/Runtime/Core/Public/Async/Fundamental/LocalQueue.h
+++ b/Engine/Source/Runtime/Core/Public/Async/Fundamental/LocalQueue.h
@@ -120,6 +120,13 @@ public:
 };
 }
 
+enum class ELocalQueueType
+{
+	EBackground,
+	EForeground,
+	EBusyWait,
+};
+
 /********************************************************************************************************************************************
  * A LocalQueueRegistry is a collection of LockFree queues that store pointers to Items, there are ThreadLocal LocalQueues with LocalItems. *
  * LocalQueues can only be Enqueued and Dequeued by the current Thread they were installed on. But Items can be stolen from any Thread      *
@@ -238,27 +245,27 @@ public:
 		friend class TLocalQueueRegistry;
 
 	public:
-		TLocalQueue(TLocalQueueRegistry& InRegistry, bool bBackgroundWorker) : Registry(&InRegistry)
+		TLocalQueue(TLocalQueueRegistry& InRegistry, ELocalQueueType QueueType) : Registry(&InRegistry)
 		{
 			checkSlow(Registry);
 			StealHazard = FStealHazard(Registry->QueueCollection, Registry->HazardsCollection);
-			Registry->AddLocalQueue(StealHazard, this, bBackgroundWorker);
+			Registry->AddLocalQueue(StealHazard, this, QueueType);
 			for (int32 PriorityIndex = 0; PriorityIndex < int32(ETaskPriority::Count); PriorityIndex++)
 			{
 				DequeueHazards[PriorityIndex] = Registry->OverflowQueues[PriorityIndex].getHeadHazard();
 			}
 		}
 
-		static TLocalQueue* AllocateLocalQueue(TLocalQueueRegistry& InRegistry, bool bBackgroundWorker)
+		static TLocalQueue* AllocateLocalQueue(TLocalQueueRegistry& InRegistry, ELocalQueueType QueueType)
 		{
 			void* Memory = FMemory::Malloc(sizeof(TLocalQueue), 128u);
-			return new (Memory) TLocalQueue(InRegistry, bBackgroundWorker);
+			return new (Memory) TLocalQueue(InRegistry, QueueType);
 		}
 
 		//delete a queue
 		//WorkeOwned means that the queue will not be automatically deleted when removal succeded.
 		//It is a special case where the memory for the local queues is alloated linearly by the Scheduler for improved stealing performance.
-		static void DeleteLocalQueue(TLocalQueue* Queue, bool bBackgroundWorker, bool WorkerOwned = false)
+		static void DeleteLocalQueue(TLocalQueue* Queue, ELocalQueueType QueueType, bool WorkerOwned = false)
 		{
 			TLocalQueueRegistry* Registry = Queue->Registry;
 			Queue->Registry = nullptr;
@@ -275,7 +282,7 @@ public:
 					Registry->OverflowQueues[PriorityIndex].enqueue(Item);
 				}
 			}
-			Registry->DeleteLocalQueue(Queue->StealHazard, Queue, bBackgroundWorker, WorkerOwned);
+			Registry->DeleteLocalQueue(Queue->StealHazard, Queue, QueueType, WorkerOwned);
 		}
 
 		//add an item to the local queue and overflow into the global queue if full
@@ -361,9 +368,13 @@ public:
 
 private:
 	// add a queue to the Registry
-	void AddLocalQueue(FStealHazard& Hazard, TLocalQueue* QueueToAdd, bool bBackgroundWorker)
+	void AddLocalQueue(FStealHazard& Hazard, TLocalQueue* QueueToAdd, ELocalQueueType QueueType)
 	{
-		NumActiveWorkers[bBackgroundWorker].fetch_add(1, std::memory_order_relaxed);
+		if (QueueType != ELocalQueueType::EBusyWait)
+		{
+			NumActiveWorkers[QueueType == ELocalQueueType::EBackground].fetch_add(1, std::memory_order_relaxed);
+		}
+
 		while(true)
 		{
 			FLocalQueueCollection* Previous = Hazard.Get();
@@ -381,9 +392,13 @@ private:
 	}
 
 	// remove a queue from the Registry
-	void DeleteLocalQueue(FStealHazard& Hazard, TLocalQueue* QueueToRemove, bool bBackgroundWorker, bool WorkerOwned)
+	void DeleteLocalQueue(FStealHazard& Hazard, TLocalQueue* QueueToRemove, ELocalQueueType QueueType, bool WorkerOwned)
 	{
-		NumActiveWorkers[bBackgroundWorker].fetch_sub(1, std::memory_order_relaxed);
+		if (QueueType != ELocalQueueType::EBusyWait)
+		{
+			NumActiveWorkers[QueueType == ELocalQueueType::EBackground].fetch_sub(1, std::memory_order_relaxed);
+		}
+
 		while(true)
 		{
 			FLocalQueueCollection* Previous = Hazard.Get();
@@ -463,15 +478,15 @@ public:
 		return nullptr;
 	}
 
-	inline FOutOfWork GetOutOfWorkScope(bool bBackgroundWorker)
+	inline FOutOfWork GetOutOfWorkScope(ELocalQueueType QueueType)
 	{
-		return FOutOfWork(NumWorkersLookingForWork[bBackgroundWorker]);
+		return FOutOfWork(NumWorkersLookingForWork[QueueType == ELocalQueueType::EBackground]);
 	}
 
 private:
-	inline bool AnyWorkerLookingForWork(bool bBackgroundWorker) const
+	inline bool AnyWorkerLookingForWork(bool bBackgroundTask) const
 	{
-		return (NumWorkersLookingForWork[bBackgroundWorker].load(std::memory_order_acquire) == 0) && (bBackgroundWorker || (NumWorkersLookingForWork[true].load(std::memory_order_acquire) == 0));
+		return (NumWorkersLookingForWork[bBackgroundTask].load(std::memory_order_acquire) == 0) && (bBackgroundTask || (NumWorkersLookingForWork[true].load(std::memory_order_acquire) == 0));
 	}
 
 	FOverflowQueueType	  OverflowQueues[uint32(ETaskPriority::Count)];
diff --git a/Engine/Source/Runtime/Core/Public/Async/Fundamental/Scheduler.h b/Engine/Source/Runtime/Core/Public/Async/Fundamental/Scheduler.h
--- a/Engine/Source/Runtime/Core/Public/Async/Fundamental/Scheduler.h
+++ b/Engine/Source/Runtime/Core/Public/Async/Fundamental/Scheduler.h
@@ -89,7 +89,10 @@ namespace LowLevelTasks
 
 	class FSchedulerTls
 	{
-	protected: 
+	protected:
+		using FQueueRegistry	= TLocalQueueRegistry<>;
+		using FLocalQueueType	= FQueueRegistry::TLocalQueue;
+
 		enum class EWorkerType
 		{
 			None,
@@ -99,6 +102,7 @@ namespace LowLevelTasks
 
 		static thread_local FTask* ActiveTask;
 		static thread_local FSchedulerTls* ActiveScheduler;
+		static thread_local FLocalQueueType* LocalQueue;
 		static thread_local EWorkerType WorkerType;
 		// number of busy-waiting calls in the call-stack
 		static thread_local uint32 BusyWaitingDepth;
@@ -129,12 +133,7 @@ namespace LowLevelTasks
 		UE_NONCOPYABLE(FScheduler);
 		static constexpr uint32 WorkerSpinCycles = 53;
 
-		using FQueueRegistry	= TLocalQueueRegistry<>;
-		using FLocalQueueType	= FQueueRegistry::TLocalQueue;
-
-		static thread_local FLocalQueueType* LocalQueue;
 		static thread_local FTask* ActiveTask;
-		static thread_local FScheduler* ActiveScheduler;
 
 		static CORE_API FScheduler Singleton;
 
@@ -176,6 +175,9 @@ namespace LowLevelTasks
 		//number of instantiated workers
 		inline uint32 GetNumWorkers() const;
 
+		//get the Queue registry to register additonal WorkerQueues for this Scheduler
+		inline FSchedulerTls::FQueueRegistry& GetQueueRegistry();
+
 	private: //Private Interface of the Scheduler	
 		enum class ESleepState
 		{
@@ -198,24 +200,24 @@ namespace LowLevelTasks
 		~FScheduler();
 
 	private: 
-		TUniquePtr<FThread> CreateWorker(FLocalQueueType* ExternalWorkerLocalQueue = nullptr, EThreadPriority Priority = EThreadPriority::TPri_Normal, bool bPermitBackgroundWork = false, bool bIsForkable = false);
-		void WorkerMain(struct FSleepEvent* WorkerEvent, FLocalQueueType* ExternalWorkerLocalQueue, uint32 WaitCycles, bool bPermitBackgroundWork);
+		TUniquePtr<FThread> CreateWorker(FSchedulerTls::FLocalQueueType* ExternalWorkerLocalQueue = nullptr, EThreadPriority Priority = EThreadPriority::TPri_Normal, bool bPermitBackgroundWork = false, bool bIsForkable = false);
+		void WorkerMain(struct FSleepEvent* WorkerEvent, FSchedulerTls::FLocalQueueType* ExternalWorkerLocalQueue, uint32 WaitCycles, bool bPermitBackgroundWork);
 		CORE_API void LaunchInternal(FTask& Task, EQueuePreference QueuePreference, bool bWakeUpWorker);
 		CORE_API void BusyWaitInternal(const FConditional& Conditional, bool ForceAllowBackgroundWork);
-		FORCENOINLINE void TrySleeping(FSleepEvent* WorkerEvent, FQueueRegistry::FOutOfWork& OutOfWork, bool& Drowsing, bool bBackgroundWorker);
+		FORCENOINLINE void TrySleeping(FSleepEvent* WorkerEvent, FSchedulerTls::FQueueRegistry::FOutOfWork& OutOfWork, bool& Drowsing, bool bBackgroundWorker);
 		inline bool WakeUpWorker(bool bBackgroundWorker);
 
-		template<FTask* (FLocalQueueType::*DequeueFunction)(bool), bool bIsBusyWaiting>
-		bool TryExecuteTaskFrom(FLocalQueueType* Queue, FQueueRegistry::FOutOfWork& OutOfWork, bool bPermitBackgroundWork);
+		template<FTask* (FSchedulerTls::FLocalQueueType::*DequeueFunction)(bool), bool bIsBusyWaiting>
+		bool TryExecuteTaskFrom(FSchedulerTls::FLocalQueueType* Queue, FSchedulerTls::FQueueRegistry::FOutOfWork& OutOfWork, bool bPermitBackgroundWork);
 
 	private:
-		TEventStack<FSleepEvent> 	SleepEventStack[2];
-		FQueueRegistry 				QueueRegistry;
-		FCriticalSection 			WorkerThreadsCS;
-		TArray<TUniquePtr<FThread>> WorkerThreads;
-		TArray<FLocalQueueType>		WorkerLocalQueues;
-		std::atomic_uint			ActiveWorkers { 0 };
-		std::atomic_uint			NextWorkerId { 0 };
+		TEventStack<FSleepEvent> 				SleepEventStack[2];
+		FSchedulerTls::FQueueRegistry 			QueueRegistry;
+		FCriticalSection 						WorkerThreadsCS;
+		TArray<TUniquePtr<FThread>>				WorkerThreads;
+		TArray<FSchedulerTls::FLocalQueueType>	WorkerLocalQueues;
+		std::atomic_uint						ActiveWorkers { 0 };
+		std::atomic_uint						NextWorkerId { 0 };
 	};
 
 	FORCEINLINE_DEBUGGABLE bool TryLaunch(FTask& Task, EQueuePreference QueuePreference = EQueuePreference::DefaultPreference, bool bWakeUpWorker = true)
@@ -304,7 +306,7 @@ namespace LowLevelTasks
 		}
 	}
 
-	inline void FScheduler::TrySleeping(FSleepEvent* WorkerEvent, FQueueRegistry::FOutOfWork& OutOfWork, bool& Drowsing, bool bBackgroundWorker)
+	inline void FScheduler::TrySleeping(FSleepEvent* WorkerEvent, FSchedulerTls::FQueueRegistry::FOutOfWork& OutOfWork, bool& Drowsing, bool bBackgroundWorker)
 	{
 		ESleepState DrowsingState1 = ESleepState::Drowsing;
 		ESleepState DrowsingState2 = ESleepState::Drowsing;
@@ -349,6 +351,11 @@ namespace LowLevelTasks
 		return false;
 	}
 
+	inline FSchedulerTls::FQueueRegistry& FScheduler::GetQueueRegistry()
+	{
+		return QueueRegistry;
+	}
+
 	inline FScheduler& FScheduler::Get()
 	{
 		return Singleton;
