Properly handle Windows paths of MAX_PATH length.

diff --git a/Engine/Source/Runtime/Core/Private/Windows/WindowsPlatformFile.cpp b/Engine/Source/Runtime/Core/Private/Windows/WindowsPlatformFile.cpp
--- a/Engine/Source/Runtime/Core/Private/Windows/WindowsPlatformFile.cpp
+++ b/Engine/Source/Runtime/Core/Private/Windows/WindowsPlatformFile.cpp
@@ -925,8 +925,8 @@ class CORE_API FWindowsPlatformFile : public IPhysicalPlatformFile
 			}
 		}
 
-		// Handle Windows Path length over MAX_PATH
-		if (Path.Len() > MAX_PATH)
+		// Handle Windows paths with null-terminated length over MAX_PATH
+		if (Path.Len() >= MAX_PATH)
 		{
 			if (bIsUNCPath)
 			{
