Fix-up deprecation warnings introduced with previous CL to fix event stack underflow when using Apple.InstrumentsEvents

#rnx

diff --git a/Engine/Source/Runtime/Core/Private/Apple/ApplePlatformDebugEvents.cpp b/Engine/Source/Runtime/Core/Private/Apple/ApplePlatformDebugEvents.cpp
--- a/Engine/Source/Runtime/Core/Private/Apple/ApplePlatformDebugEvents.cpp
+++ b/Engine/Source/Runtime/Core/Private/Apple/ApplePlatformDebugEvents.cpp
@@ -129,6 +129,8 @@ FORCEINLINE_DEBUGGABLE os_log_t GetLog()
 	return Log;
 }
 
+PRAGMA_DISABLE_DEPRECATION_WARNINGS
+
 void FApplePlatformDebugEvents::DebugSignPost(uint16 Code, uintptr_t Arg1, uintptr_t Arg2, uintptr_t Arg3, uintptr_t Arg4)
 {
 	if (GAppleInstrumentsEvents)
@@ -153,6 +155,8 @@ void FApplePlatformDebugEvents::DebugSignPostEnd(uint16 Code, uintptr_t Arg1, ui
 	}
 }
 
+PRAGMA_ENABLE_DEPRECATION_WARNINGS
+
 TArray<FApplePlatformDebugEvents::FEvent>* FApplePlatformDebugEvents::GetEventStack()
 {
 	TArray<FEvent>* Current = (TArray<FEvent>*)FPlatformTLS::GetTlsValue(TLSSlot);
