U5 - disable tasktag scope checks when the engine is already shutting down and in a bad state because TLS will not be valid anymore.

diff --git a/Engine/Source/Runtime/Core/Private/HAL/ThreadingBase.cpp b/Engine/Source/Runtime/Core/Private/HAL/ThreadingBase.cpp
--- a/Engine/Source/Runtime/Core/Private/HAL/ThreadingBase.cpp
+++ b/Engine/Source/Runtime/Core/Private/HAL/ThreadingBase.cpp
@@ -94,7 +94,7 @@ FTaskTagScope::FTaskTagScope(bool InTagOnlyIfNone, ETaskTag InTag) : Tag(InTag),
 			//Try to catch other Threads that already opened a non parallel scope
 			ActiveNamedThreads.store(int32(ETaskTag::ENone));
 		}
-		checkf(IsOK || !IsEngineExitRequested(), TEXT("Only Scopes tagged with ETaskTag::EParallelThread can be tagged multiple times. ActiveNamedThreads(%x) cannot be tagged multiple times in the same callstack you can use FOptionalTaskTagScope to avoid retagging check the ActiveNamedThreads(%x) with the current Tag(%x)"), OldTag, FTaskTagScope::GetCurrentTag(), Tag);
+		checkf(IsOK || IsEngineExitRequested(), TEXT("Only Scopes tagged with ETaskTag::EParallelThread can be tagged multiple times. ActiveNamedThreads(%x) cannot be tagged multiple times in the same callstack you can use FOptionalTaskTagScope to avoid retagging check the ActiveNamedThreads(%x) with the current Tag(%x)"), OldTag, FTaskTagScope::GetCurrentTag(), Tag);
 	}
 	
 	ParentTag = ActiveTaskTag;
@@ -129,7 +129,7 @@ FTaskTagScope::~FTaskTagScope()
 		ETaskTag NamedThreadBits = (Tag & ETaskTag::ENamedThreadBits);
 		static_assert(sizeof(ETaskTag) == sizeof(int32), "EnumSize must match interlockedAnd");
 		ETaskTag OldTag = ETaskTag(ActiveNamedThreads.fetch_and(int32(~NamedThreadBits)));
-		checkf((OldTag & NamedThreadBits) == (NamedThreadBits), TEXT("Currently active Threads(%x) got corrupted check the ActiveNamedThreads(%x)"), OldTag, FTaskTagScope::GetCurrentTag());
+		checkf((OldTag & NamedThreadBits) == (NamedThreadBits) || IsEngineExitRequested(), TEXT("Currently active Threads(%x) got corrupted check the ActiveNamedThreads(%x)"), OldTag, FTaskTagScope::GetCurrentTag());
 	}
 
 	//prolong the scope of the GT for static variable destructors
