- Expose starting UnrealTraceServer to give dependents peace of mind. This is not normally required as the server should get launched when trace is initialised. However this initialisation happens so early that CIS tests/automation doesn't capture the TTY output making it difficult to investigate CIS-only issues. This is a strategic submit to push the TTY output later.
- Renamed StartUnrealTrace() to LaunchUnrealTrace() to be more indicative that it runs something.

#rnx

diff --git a/Engine/Plugins/Animation/GameplayInsights/Source/GameplayInsights/Private/GameplayInsightsModule.cpp b/Engine/Plugins/Animation/GameplayInsights/Source/GameplayInsights/Private/GameplayInsightsModule.cpp
--- a/Engine/Plugins/Animation/GameplayInsights/Source/GameplayInsights/Private/GameplayInsightsModule.cpp
+++ b/Engine/Plugins/Animation/GameplayInsights/Source/GameplayInsights/Private/GameplayInsightsModule.cpp
@@ -5,6 +5,7 @@
 #include "Insights/ITimingViewExtender.h"
 #include "Containers/Ticker.h"
 #include "Modules/ModuleManager.h"
+#include "ProfilingDebugging/TraceAuxiliary.h"
 #include "Framework/Docking/LayoutExtender.h"
 #include "SAnimGraphSchematicView.h"
 #include "Insights/IUnrealInsightsModule.h"
@@ -184,6 +185,7 @@ void FGameplayInsightsModule::StartupModule()
 				UnrealInsightsModule.ConnectToStore(TEXT("localhost"), StoreService->GetPort());
 				UE::Trace::SendTo(TEXT("localhost"), StoreService->GetRecorderPort());
 #else
+				FTraceAuxiliary::LaunchUnrealTrace();
 				ensure(UnrealInsightsModule.ConnectToStore(TEXT("127.0.0.1")));
 				UE::Trace::SendTo(TEXT("127.0.0.1"));
 #endif // WITH_TRACE_STORE
diff --git a/Engine/Plugins/Runtime/NetworkPredictionInsights/Source/NetworkPredictionInsights/Private/NetworkPredictionInsightsModule.cpp b/Engine/Plugins/Runtime/NetworkPredictionInsights/Source/NetworkPredictionInsights/Private/NetworkPredictionInsightsModule.cpp
--- a/Engine/Plugins/Runtime/NetworkPredictionInsights/Source/NetworkPredictionInsights/Private/NetworkPredictionInsightsModule.cpp
+++ b/Engine/Plugins/Runtime/NetworkPredictionInsights/Source/NetworkPredictionInsights/Private/NetworkPredictionInsightsModule.cpp
@@ -5,6 +5,7 @@
 #include "Insights/ITimingViewExtender.h"
 #include "Containers/Ticker.h"
 #include "Modules/ModuleManager.h"
+#include "ProfilingDebugging/TraceAuxiliary.h"
 #include "Framework/Docking/LayoutExtender.h"
 #include "Insights/IUnrealInsightsModule.h"
 #include "HAL/PlatformApplicationMisc.h"
@@ -117,6 +118,7 @@ void FNetworkPredictionInsightsModule::StartupModule()
 				ensure(UnrealInsightsModule.ConnectToStore(TEXT("localhost"), StoreService->GetPort()));
 				UE::Trace::SendTo(TEXT("localhost"), StoreService->GetRecorderPort());
 #else
+				FTraceAuxiliary::LaunchUnrealTrace();
 				ensure(UnrealInsightsModule.ConnectToStore(TEXT("127.0.0.1")));
 				UE::Trace::SendTo(TEXT("127.0.0.1"));
 #endif // WITH_TRACE_STORE
diff --git a/Engine/Source/Runtime/Core/Private/ProfilingDebugging/TraceAuxiliary.cpp b/Engine/Source/Runtime/Core/Private/ProfilingDebugging/TraceAuxiliary.cpp
--- a/Engine/Source/Runtime/Core/Private/ProfilingDebugging/TraceAuxiliary.cpp
+++ b/Engine/Source/Runtime/Core/Private/ProfilingDebugging/TraceAuxiliary.cpp
@@ -486,10 +486,19 @@ static FAutoConsoleCommand TraceAuxiliaryResumeCmd(
 
 
 #if WITH_UNREAL_TRACE_LAUNCH
+////////////////////////////////////////////////////////////////////////////////
+static std::atomic<int32> GUnrealTraceLanuched; // = 0;
+
 ////////////////////////////////////////////////////////////////////////////////
 #if PLATFORM_WINDOWS
-static void StartUnrealTrace()
+static void LaunchUnrealTraceInternal()
 {
+	if (GUnrealTraceLanuched.load(std::memory_order_relaxed))
+	{
+		UE_LOG(LogCore, Log, TEXT("UnrealTraceServer: Trace store already started"));
+		return;
+	}
+
 	FString BinPath = "\"";
 	BinPath += FPaths::EngineDir();
 	BinPath += "/Binaries/Win64/UnrealTraceServer.exe\"";
@@ -502,13 +511,13 @@ static void StartUnrealTrace()
 
 	if (!bOk)
 	{
-		UE_LOG(LogCore, Warning, TEXT("Unable to launch the trace store from '%s' (%08x)"), *BinPath, GetLastError());
+		UE_LOG(LogCore, Warning, TEXT("UnrealTraceServer: Unable to launch the trace store from '%s' (%08x)"), *BinPath, GetLastError());
 		return;
 	}
 
 	if (WaitForSingleObject(ProcessInfo.hProcess, 5000) == WAIT_TIMEOUT)
 	{
-		UE_LOG(LogCore, Warning, TEXT("Timed out waiting for the trace store to start"));
+		UE_LOG(LogCore, Warning, TEXT("UnrealTraceServer: Timed out waiting for the trace store to start"));
 	}
 	else
 	{
@@ -516,11 +525,12 @@ static void StartUnrealTrace()
 		GetExitCodeProcess(ProcessInfo.hProcess, &ExitCode);
 		if (ExitCode)
 		{
-			UE_LOG(LogCore, Warning, TEXT("Trace store returned an error (%08x)"), ExitCode);
+			UE_LOG(LogCore, Warning, TEXT("UnrealTraceServer: Trace store returned an error (%08x)"), ExitCode);
 		}
 		else
 		{
-			UE_LOG(LogCore, Log, TEXT("Trace store launch successful"));
+			UE_LOG(LogCore, Log, TEXT("UnrealTraceServer: Trace store launch successful"));
+			GUnrealTraceLanuched.fetch_add(1, std::memory_order_relaxed);
 		}
 	}
 
@@ -531,7 +541,7 @@ static void StartUnrealTrace()
 
 ////////////////////////////////////////////////////////////////////////////////
 #if PLATFORM_UNIX || PLATFORM_MAC
-static void StartUnrealTrace()
+static void LaunchUnrealTraceInternal()
 {
 	/* nop */
 }
@@ -559,8 +569,8 @@ void FTraceAuxiliary::Initialize(const TCHAR* CommandLine)
 
 #if WITH_UNREAL_TRACE_LAUNCH
 	{
-		TRACE_CPUPROFILER_EVENT_SCOPE(FTraceAux_StartUnrealTrace);
-		StartUnrealTrace();
+		TRACE_CPUPROFILER_EVENT_SCOPE(FTraceAux_LaunchUnrealTrace);
+		LaunchUnrealTraceInternal();
 	}
 #endif
 
@@ -688,3 +698,11 @@ void FTraceAuxiliary::TryAutoConnect()
 	#endif
 #endif
 }
+
+////////////////////////////////////////////////////////////////////////////////
+void FTraceAuxiliary::LaunchUnrealTrace()
+{
+#if WITH_UNREAL_TRACE_LAUNCH
+	LaunchUnrealTraceInternal();
+#endif
+}
diff --git a/Engine/Source/Runtime/Core/Public/ProfilingDebugging/TraceAuxiliary.h b/Engine/Source/Runtime/Core/Public/ProfilingDebugging/TraceAuxiliary.h
--- a/Engine/Source/Runtime/Core/Public/ProfilingDebugging/TraceAuxiliary.h
+++ b/Engine/Source/Runtime/Core/Public/ProfilingDebugging/TraceAuxiliary.h
@@ -13,4 +13,5 @@ public:
 	static void Shutdown();
 	static void TryAutoConnect();
 	static void EnableChannels();
+	static void LaunchUnrealTrace();
 };
