implemented sleeping waiting for Tasks (the new system), made it the default option while BusyWait is still available as an opt-in. this doesn't touch TaskGraph or ThreadPool. also happens to fix UE-117431 as it was caused by busy waiting on Audio pipe

diff --git a/Engine/Source/Runtime/Core/Private/Tests/Tasks/TasksTest.cpp b/Engine/Source/Runtime/Core/Private/Tests/Tasks/TasksTest.cpp
--- a/Engine/Source/Runtime/Core/Private/Tests/Tasks/TasksTest.cpp
+++ b/Engine/Source/Runtime/Core/Private/Tests/Tasks/TasksTest.cpp
@@ -91,6 +91,7 @@ namespace UE { namespace TasksTests
 
 		{	// launch a task and wait till it's executed
 			Launch(UE_SOURCE_LOCATION, [] {}).Wait();
+			Launch(UE_SOURCE_LOCATION, [] {}).BusyWait();
 		}
 
 		{	// FTaskEvent
@@ -105,6 +106,7 @@ namespace UE { namespace TasksTests
 			Event.Trigger();
 			check(Event.IsCompleted());
 			verify(Event.Wait(FTimespan::Zero()));
+			verify(Event.BusyWait(FTimespan::Zero()));
 		}
 
 		{	// postpone execution so waiting kicks in first
@@ -116,6 +118,15 @@ namespace UE { namespace TasksTests
 			check(Counter == 1);
 		}
 
+		{	// postpone execution so busy waiting kicks in first
+			std::atomic<int32> Counter{ 0 };
+			FTask Task = Launch(UE_SOURCE_LOCATION, [&Counter] { ++Counter; FPlatformProcess::Sleep(0.1f); });
+
+			ensure(!Task.BusyWait(FTimespan::Zero()));
+			Task.BusyWait();
+			check(Counter == 1);
+		}
+
 		{	// same but using `FTaskEvent`
 			FTaskEvent Event{ UE_SOURCE_LOCATION };
 			FTask Task = Launch(UE_SOURCE_LOCATION, [&Event] { Event.Wait(); });
@@ -124,6 +135,14 @@ namespace UE { namespace TasksTests
 			Task.Wait();
 		}
 
+		{	// same but using busy-wait and `FTaskEvent`
+			FTaskEvent Event{ UE_SOURCE_LOCATION };
+			FTask Task = Launch(UE_SOURCE_LOCATION, [&Event] { Event.BusyWait(); });
+			ensure(!Task.BusyWait(FTimespan::FromMilliseconds(100)));
+			Event.Trigger();
+			Task.BusyWait();
+		}
+
 		{	// basic use-case, postpone waiting so the task is executed first
 			std::atomic<bool> Done{ false };
 			FTask Task = Launch(UE_SOURCE_LOCATION, [&Done] { Done = true; });
@@ -135,6 +154,17 @@ namespace UE { namespace TasksTests
 			check(Done);
 		}
 
+		{	// basic use-case, postpone busy-waiting so the task is executed first
+			std::atomic<bool> Done{ false };
+			FTask Task = Launch(UE_SOURCE_LOCATION, [&Done] { Done = true; });
+			while (!Task.IsCompleted())
+			{
+				FPlatformProcess::Yield();
+			}
+			Task.BusyWait();
+			check(Done);
+		}
+
 		{	// basic use-case with result, postpone execution so waiting kicks in first
 			TTask<int> Task = Launch(UE_SOURCE_LOCATION, [] { FPlatformProcess::Sleep(0.1f);  return 42; });
 			verify(Task.GetResult() == 42);
diff --git a/Engine/Source/Runtime/Core/Public/Tasks/Task.h b/Engine/Source/Runtime/Core/Public/Tasks/Task.h
--- a/Engine/Source/Runtime/Core/Public/Tasks/Task.h
+++ b/Engine/Source/Runtime/Core/Public/Tasks/Task.h
@@ -6,10 +6,20 @@
 #include "Async/Fundamental/Task.h"
 #include "Templates/RefCounting.h"
 #include "Containers/StaticArray.h"
+#include "HAL/Event.h"
 #include "CoreTypes.h"
 
 namespace UE { namespace Tasks
 {
+	template<typename ResultType>
+	class TTask;
+
+	template<typename TaskBodyType>
+	TTask<TInvokeResult_T<TaskBodyType>> Launch(const TCHAR* DebugName, TaskBodyType&& TaskBody, LowLevelTasks::ETaskPriority Priority = LowLevelTasks::ETaskPriority::Normal);
+
+	template<typename TaskBodyType, typename PrerequisitesCollectionType>
+	TTask<TInvokeResult_T<TaskBodyType>> Launch(const TCHAR* DebugName, TaskBodyType&& TaskBody, PrerequisitesCollectionType&& Prerequisites, LowLevelTasks::ETaskPriority Priority = LowLevelTasks::ETaskPriority::Normal);
+
 	// a common part of the generic `TTask<ResultType>` and its `TTask<void>` specialisation
 	template<typename ResultType>
 	class TTaskBase
@@ -19,6 +29,9 @@ namespace UE { namespace Tasks
 		template<typename... TaskTypes>
 		friend class TPrerequisites;
 
+		template<typename TaskCollectionType>
+		friend bool Private::TryRetractAndExecute(TaskCollectionType&& Tasks);
+
 	protected:
 		TTaskBase() = default;
 
@@ -38,30 +51,68 @@ namespace UE { namespace Tasks
 			return !IsValid() || Pimpl->IsCompleted();
 		}
 
-		// waits for task's completion
-		void Wait()
+		// waits for task's completion, with optional timeout
+		// @return true if the task is completed
+		bool Wait(FTimespan Timeout = FTimespan::MaxValue())
+		{
+			if (!IsValid() || IsCompleted())
+			{
+				return true;
+			}
+
+			// "zero timeout" - just check if the task is completed
+			if (Timeout == FTimespan::Zero())
+			{
+				return false;
+			}
+
+			// waiting with timeout can't do retraction, otherwise timeout value is unreliable
+			if (Timeout == FTimespan::MaxValue() && Pimpl->TryRetractAndExecute())
+			{
+				return true;
+			}
+
+			// the event must be alive for the task and this function lifetime, we don't know which one will be finished first as waiting can time out
+			// before the waiting task is completed
+			FSharedEventRef CompletionEvent;
+
+			TRefCountPtr<Private::TTaskWithResult<void>> WaitingTask{ new Private::TTaskWithResult<void>, /*bAddRef=*/ false };
+			WaitingTask->Init(TEXT("Waiting Task"), [CompletionEvent] { CompletionEvent->Trigger(); }, Private::FTaskBase::InlineTaskPriority);
+			WaitingTask->AddPrerequisites(*this);
+
+			if (WaitingTask->TryLaunch())
+			{	// was executed inline
+				check(WaitingTask->IsCompleted());
+				return true;
+			}
+
+			return CompletionEvent->Wait(Timeout);
+		}
+
+		// waits for task's completion while executing other tasks. Shouldn't be used inside a latency-sensitive task
+		void BusyWait()
 		{
 			if (IsValid())
 			{
-				Pimpl->Wait();
+				Pimpl->BusyWait();
 			}
 		}
 
-		// waits for task's completion at least specified amount of time.
+		// waits for task's completion at least specified amount of time, while executing other tasks.
 		// the call can return much later than the given timeout
 		// @return true if the task is completed
-		bool Wait(FTimespan Timeout)
+		bool BusyWait(FTimespan Timeout)
 		{
-			return !IsValid() || Pimpl->Wait(Timeout);
+			return !IsValid() || Pimpl->BusyWait(Timeout);
 		}
 
-		// waits for task's completion or the given condition becomes true
+		// waits for task's completion or the given condition becomes true, while executing other tasks.
 		// the call can return much later than the given condition became true
 		// @return true if the task is completed
 		template<typename ConditionType>
-		bool Wait(ConditionType&& Condition)
+		bool BusyWait(ConditionType&& Condition)
 		{
-			return !IsValid() || Pimpl->Wait(Forward<ConditionType>(Condition));
+			return !IsValid() || Pimpl->BusyWait(Forward<ConditionType>(Condition));
 		}
 
 	protected:
@@ -82,6 +133,7 @@ namespace UE { namespace Tasks
 		ResultType& GetResult()
 		{
 			check(FSuper::IsValid());
+			FSuper::Wait();
 			return FSuper::Pimpl->GetResult();
 		}
 
@@ -188,12 +240,40 @@ namespace UE { namespace Tasks
 		return TTask<FResult>{ Task };
 	}
 
+	// wait for multiple tasks, with optional timeout
+	template<typename TaskCollectionType>
+	bool Wait(TaskCollectionType&& Tasks, FTimespan Timeout = FTimespan::MaxValue())
+	{
+		// waiting with timeout can't do retraction, otherwise timeout value is unreliable
+		if (Timeout == FTimespan::MaxValue() && Private::TryRetractAndExecute(Tasks))
+		{
+			return true;
+		}
+
+		// the event must be alive for the task and this function lifetime, we don't know which one will be finished first as waiting can time out
+		// before the waiting task is completed
+		FSharedEventRef CompletionEvent;
+
+		TRefCountPtr<Private::TTaskWithResult<void>> WaitingTask{ new Private::TTaskWithResult<void>, /*bAddRef=*/ false };
+		WaitingTask->Init(TEXT("Waiting Task"), [CompletionEvent] { CompletionEvent->Trigger(); }, Private::FTaskBase::InlineTaskPriority);
+		WaitingTask->AddPrerequisites(Tasks);
+
+		if (WaitingTask->TryLaunch())
+		{	// was executed inline
+			check(WaitingTask->IsCompleted());
+			return true;
+		}
+
+		return CompletionEvent->Wait(Timeout);
+	}
+
+	// wait for multiple tasks while executing other tasks
 	template<typename TaskCollectionType>
-	void Wait(TaskCollectionType&& Tasks)
+	void BusyWait(TaskCollectionType&& Tasks)
 	{
 		for (auto& Task : Tasks)
 		{
-			Task.Wait();
+			Task.BusyWait();
 		}
 	}
 
diff --git a/Engine/Source/Runtime/Core/Public/Tasks/TaskPrivate.h b/Engine/Source/Runtime/Core/Public/Tasks/TaskPrivate.h
--- a/Engine/Source/Runtime/Core/Public/Tasks/TaskPrivate.h
+++ b/Engine/Source/Runtime/Core/Public/Tasks/TaskPrivate.h
@@ -21,32 +21,10 @@ namespace UE { namespace Tasks
 
 	// BEGIN: forward declarations
 
-	class FTaskEvent;
-	class FPipe;
-
-	template<typename T>
-	class TTask;
-
-	template<typename TaskBodyType>
-	TTask<TInvokeResult_T<TaskBodyType>> Launch(
-		const TCHAR* DebugName, 
-		TaskBodyType&& TaskBody, 
-		LowLevelTasks::ETaskPriority Priority = LowLevelTasks::ETaskPriority::Normal
-	);
-
-	template<typename TaskBodyType, typename PrerequisitesCollectionType>
-	TTask<TInvokeResult_T<TaskBodyType>> Launch(
-		const TCHAR* DebugName, 
-		TaskBodyType&& TaskBody, 
-		PrerequisitesCollectionType&& Prerequisites, 
-		LowLevelTasks::ETaskPriority Priority = LowLevelTasks::ETaskPriority::Normal
-	);
-
 	template<typename ResultType>
 	class TTaskBase;
 
-	template<typename... TaskTypes>
-	class TPrerequisites;
+	class FPipe;
 
 	// END: forward declarations
 
@@ -184,8 +162,15 @@ namespace UE { namespace Tasks
 				return Subsequents.IsClosed();
 			}
 
-			// waits until the task is completed
-			void Wait()
+			FORCENOINLINE bool TryRetractAndExecute()
+			{
+				bool bUnlocked = NumLocks.load(std::memory_order_acquire) == (GetPipe() == nullptr ? 1 : 0);
+				// TryCancel() will try to cancel the task (if it's not being executed yet) and execute it
+				return bUnlocked && LowLevelTask.TryCancel() && IsCompleted();
+			}
+
+			// waits until the task is completed while executing other tasks
+			void BusyWait()
 			{
 				// we try to retract the task from the scheduler, if it's execution hasn't started yet, to execute it immediately. This makes
 				// the task skip waiting in the scheduler's queue.
@@ -197,9 +182,8 @@ namespace UE { namespace Tasks
 				LowLevelTasks::BusyWaitUntil([this] { return IsCompleted(); });
 			}
 
-			// waits until the task is completed or waiting timed out
-			// @see `void Wait()`
-			bool Wait(FTimespan InTimeout)
+			// waits until the task is completed or waiting timed out, while executing other tasks
+			bool BusyWait(FTimespan InTimeout)
 			{
 				FTimeout Timeout{ InTimeout };
 				LowLevelTasks::BusyWaitUntil([this, Timeout] { return !LowLevelTask.IsReady() || Timeout; });
@@ -214,10 +198,9 @@ namespace UE { namespace Tasks
 				return IsCompleted();
 			}
 
-			// waits until the task is completed or the condition returns true
-			// @see `void Wait()`
+			// waits until the task is completed or the condition returns true, while executing other tasks
 			template<typename ConditionType>
-			bool Wait(ConditionType&& Condition)
+			bool BusyWait(ConditionType&& Condition)
 			{
 				LowLevelTasks::BusyWaitUntil(
 					[this, Condition = Forward<ConditionType>(Condition)] { return !LowLevelTask.IsReady() || Condition(); }
@@ -301,9 +284,6 @@ namespace UE { namespace Tasks
 				{
 					// the task is locked for a pipe initially even if eventually there's no pipe
 					check(LocalNumLocks == 1);
-					// there's no need to set NumLocks to 0 in non-shipping builds as it's not used anymore, but we do this for debugging
-					// as it's confusing to see scheduled tasks that are still locked
-					check(NumLocks.exchange(0, std::memory_order_relaxed) == 1);
 
 					return true;
 				}
@@ -350,6 +330,20 @@ namespace UE { namespace Tasks
 #endif
 		};
 
+		template<typename TaskCollectionType>
+		bool TryRetractAndExecute(TaskCollectionType&& Tasks)
+		{
+			bool bResult = true;
+			for (auto& Task : Tasks)
+			{
+				if (!Task.Pimpl->TryRetractAndExecute())
+				{
+					bResult = false;
+				}
+			}
+			return bResult;
+		}
+
 		// Extends FTaskBase by supporting execution result.
 		template<typename ResultType>
 		class TTaskWithResultBase : public FTaskBase
@@ -364,8 +358,7 @@ namespace UE { namespace Tasks
 
 			~TTaskWithResultBase()
 			{
-				// Every task instance must be launched. This is quaranteed by the the high level task API (`UE::Tasks::TTask`)
-				checkf(bLaunched, TEXT("Every task instance must be launched"));
+				checkf(IsCompleted(), TEXT("Every task instance must be completed before it's destroyed"));
 				DestructItem(ResultStorage.GetTypedPtr());
 			}
 
@@ -386,26 +379,17 @@ namespace UE { namespace Tasks
 
 			bool TryLaunch()
 			{
-#if DO_CHECK
-				bLaunched = true;
-#endif
 				return FTaskBase::TryLaunch();
 			}
 
 			ResultType& GetResult()
 			{
-				check(bLaunched);
-				Wait();
-
+				checkf(IsCompleted(), TEXT("The task must be completed to obtain its result"));
 				return *ResultStorage.GetTypedPtr();
 			}
 
 		private:
 			TTypeCompatibleBytes<ResultType> ResultStorage;
-
-#if DO_CHECK
-			bool bLaunched = false;
-#endif
 		};
 
 		template<>
@@ -417,7 +401,7 @@ namespace UE { namespace Tasks
 		public:
 			void GetResult()
 			{
-				Wait();
+				checkf(IsCompleted(), TEXT("The task must be completed to obtain its result"));
 			}
 		};
 
@@ -452,7 +436,7 @@ namespace UE { namespace Tasks
 			std::atomic<uint32> RefCount;
 		};
 
-		// A final task class that must be used only with `TRefCountPtr`, thus only friends are allowed to construct instances.
+		// A final task class that must be used only with `TRefCountPtr`
 		// @see UE::Tasks::Task<T>
 		template<typename ResultType>
 		class TTaskWithResult final : public TTaskWithResultBase<ResultType>, public TRefCountedBase<TTaskWithResult<ResultType>>
@@ -461,23 +445,13 @@ namespace UE { namespace Tasks
 
 			using FSelf = TTaskWithResult<ResultType>;
 
-		private:
-			template<typename TaskBodyType>
-			friend TTask<TInvokeResult_T<TaskBodyType>> UE::Tasks::Launch(const TCHAR* DebugName, TaskBodyType&& TaskBody, LowLevelTasks::ETaskPriority Priority);
-
-			template<typename TaskBodyType, typename PrerequisitesCollectionType>
-			friend TTask<TInvokeResult_T<TaskBodyType>> UE::Tasks::Launch(const TCHAR* DebugName, TaskBodyType&& TaskBody, PrerequisitesCollectionType&& Prerequisites, LowLevelTasks::ETaskPriority Priority);
-
-			friend FTaskEvent;
-			friend FPipe;
-
+		public:
 			TTaskWithResult()
 				: TRefCountedBase<FSelf>(2) // one for the initial reference (we don't increment it when passing to `TRefCountPtr`) and 
 				// one for scheduler's reference which will be released on task completion
 			{
 			}
 
-		public:
 			template<typename TaskBodyType>
 			void Init(const TCHAR* DebugName, TaskBodyType&& TaskBody, LowLevelTasks::ETaskPriority Priority)
 			{
