TAtomic deprecation preparation: TAtomic implementation switched to std::atomic for testing. disabled until QA did a pass

diff --git a/Engine/Plugins/Enterprise/DatasmithImporter/Source/DatasmithImporter/Private/DatasmithImporter.cpp b/Engine/Plugins/Enterprise/DatasmithImporter/Source/DatasmithImporter/Private/DatasmithImporter.cpp
--- a/Engine/Plugins/Enterprise/DatasmithImporter/Source/DatasmithImporter/Private/DatasmithImporter.cpp
+++ b/Engine/Plugins/Enterprise/DatasmithImporter/Source/DatasmithImporter/Private/DatasmithImporter.cpp
@@ -275,8 +275,10 @@ void FDatasmithImporter::ImportStaticMeshes( FDatasmithImportContext& ImportCont
 	{
 		for (int32 MeshIndex = 0; MeshIndex < StaticMeshesCount && !ImportContext.bUserCancelled; ++MeshIndex)
 		{
-			ImportContext.bUserCancelled |= FDatasmithImporterImpl::HasUserCancelledTask( ImportContext.FeedbackContext );
-
+			if (FDatasmithImporterImpl::HasUserCancelledTask(ImportContext.FeedbackContext))
+			{
+				ImportContext.bUserCancelled = true;
+			}
 
 			if (!ImportContext.AssetsContext.StaticMeshesFinalPackage || ImportContext.AssetsContext.StaticMeshesFinalPackage->GetFName() == NAME_None || ImportContext.SceneTranslator == nullptr)
 			{
@@ -317,7 +319,10 @@ void FDatasmithImporter::ImportStaticMeshes( FDatasmithImportContext& ImportCont
 	// This pass will wait on the futures we got from the first pass async tasks
 	for ( int32 MeshIndex = 0; MeshIndex < StaticMeshesCount && !ImportContext.bUserCancelled; ++MeshIndex )
 	{
-		ImportContext.bUserCancelled |= FDatasmithImporterImpl::HasUserCancelledTask( ImportContext.FeedbackContext );
+		if (FDatasmithImporterImpl::HasUserCancelledTask(ImportContext.FeedbackContext))
+		{
+			ImportContext.bUserCancelled = true;
+		}
 
 		TSharedRef< IDatasmithMeshElement > MeshElement = ImportContext.FilteredScene->GetMesh( MeshIndex ).ToSharedRef();
 
@@ -504,7 +509,10 @@ void FDatasmithImporter::ImportTextures( FDatasmithImportContext& ImportContext
 		{
 			for ( int32 TextureIndex = 0; TextureIndex < FilteredTextureElements.Num(); TextureIndex++ )
 			{
-				ImportContext.bUserCancelled |= FDatasmithImporterImpl::HasUserCancelledTask( ImportContext.FeedbackContext );
+				if (FDatasmithImporterImpl::HasUserCancelledTask(ImportContext.FeedbackContext))
+				{
+					ImportContext.bUserCancelled = true;
+				}
 
 				if (ImportContext.bUserCancelled)
 				{
@@ -550,7 +558,10 @@ void FDatasmithImporter::ImportTextures( FDatasmithImportContext& ImportContext
 				continue;
 			}
 
-			ImportContext.bUserCancelled |= FDatasmithImporterImpl::HasUserCancelledTask( ImportContext.FeedbackContext );
+			if (FDatasmithImporterImpl::HasUserCancelledTask(ImportContext.FeedbackContext))
+			{
+				ImportContext.bUserCancelled = true;
+			}
 
 			AsyncData[TextureIndex].Result =
 				Async(
@@ -582,7 +593,10 @@ void FDatasmithImporter::ImportTextures( FDatasmithImportContext& ImportContext
 				continue;
 			}
 
-			ImportContext.bUserCancelled |= FDatasmithImporterImpl::HasUserCancelledTask( ImportContext.FeedbackContext );
+			if (FDatasmithImporterImpl::HasUserCancelledTask(ImportContext.FeedbackContext))
+			{
+				ImportContext.bUserCancelled = true;
+			}
 
 			if ( ImportContext.bUserCancelled )
 			{
@@ -872,7 +886,10 @@ void FDatasmithImporter::ImportActors( FDatasmithImportContext& ImportContext )
 
 		for (int32 i = 0; i < ActorsCount && !ImportContext.bUserCancelled; ++i)
 		{
-			ImportContext.bUserCancelled |= FDatasmithImporterImpl::HasUserCancelledTask( ImportContext.FeedbackContext );
+			if (FDatasmithImporterImpl::HasUserCancelledTask(ImportContext.FeedbackContext))
+			{
+				ImportContext.bUserCancelled = true;
+			}
 
 			TSharedPtr< IDatasmithActorElement > ActorElement = ImportContext.Scene->GetActor(i);
 
@@ -894,7 +911,10 @@ void FDatasmithImporter::ImportActors( FDatasmithImportContext& ImportContext )
 		// After all actors were imported, perform a post import step so that any dependencies can be resolved
 		for (int32 i = 0; i < ActorsCount && !ImportContext.bUserCancelled; ++i)
 		{
-			ImportContext.bUserCancelled |= FDatasmithImporterImpl::HasUserCancelledTask( ImportContext.FeedbackContext );
+			if (FDatasmithImporterImpl::HasUserCancelledTask(ImportContext.FeedbackContext))
+			{
+				ImportContext.bUserCancelled = true;
+			}
 
 			TSharedPtr< IDatasmithActorElement > ActorElement = ImportContext.Scene->GetActor(i);
 
@@ -984,7 +1004,10 @@ AActor* FDatasmithImporter::ImportActor( FDatasmithImportContext& ImportContext,
 
 	for (int32 i = 0; i < ActorElement->GetChildrenCount() && !ImportContext.bUserCancelled; ++i)
 	{
-		ImportContext.bUserCancelled |= FDatasmithImporterImpl::HasUserCancelledTask( ImportContext.FeedbackContext );
+		if (FDatasmithImporterImpl::HasUserCancelledTask(ImportContext.FeedbackContext))
+		{
+			ImportContext.bUserCancelled = true;
+		}
 
 		const TSharedPtr< IDatasmithActorElement >& ChildActorElement = ActorElement->GetChild(i);
 
@@ -1353,7 +1376,10 @@ void FDatasmithImporter::ImportLevelSequences( FDatasmithImportContext& ImportCo
 	SequencesToImport.Reserve(SequencesCount);
 	for ( int32 SequenceIndex = 0; SequenceIndex < SequencesCount && !ImportContext.bUserCancelled; ++SequenceIndex )
 	{
-		ImportContext.bUserCancelled |= FDatasmithImporterImpl::HasUserCancelledTask( ImportContext.FeedbackContext );
+		if (FDatasmithImporterImpl::HasUserCancelledTask(ImportContext.FeedbackContext))
+		{
+			ImportContext.bUserCancelled = true;
+		}
 
 		TSharedPtr< IDatasmithLevelSequenceElement > SequenceElement = ImportContext.FilteredScene->GetLevelSequence( SequenceIndex );
 		if ( !SequenceElement )
@@ -1376,7 +1402,10 @@ void FDatasmithImporter::ImportLevelSequences( FDatasmithImportContext& ImportCo
 		// Scan remaining sequences and import the ones we can, removing from this array
 		for ( int32 SequenceIndex = SequencesToImport.Num() - 1; SequenceIndex >= 0 && !ImportContext.bUserCancelled; --SequenceIndex )
 		{
-			ImportContext.bUserCancelled |= FDatasmithImporterImpl::HasUserCancelledTask( ImportContext.FeedbackContext );
+			if (FDatasmithImporterImpl::HasUserCancelledTask(ImportContext.FeedbackContext))
+			{
+				ImportContext.bUserCancelled = true;
+			}
 
 			TSharedPtr<IDatasmithLevelSequenceElement>& SequenceElement = SequencesToImport[SequenceIndex];
 
@@ -1460,7 +1489,10 @@ void FDatasmithImporter::ImportLevelVariantSets( FDatasmithImportContext& Import
 
 	for ( int32 LevelVariantSetIndex = 0; LevelVariantSetIndex < LevelVariantSetsCount && !ImportContext.bUserCancelled; ++LevelVariantSetIndex )
 	{
-		ImportContext.bUserCancelled |= FDatasmithImporterImpl::HasUserCancelledTask( ImportContext.FeedbackContext );
+		if (FDatasmithImporterImpl::HasUserCancelledTask(ImportContext.FeedbackContext))
+		{
+			ImportContext.bUserCancelled = true;
+		}
 
 		TSharedPtr< IDatasmithLevelVariantSetsElement > LevelVariantSetsElement = ImportContext.FilteredScene->GetLevelVariantSets( LevelVariantSetIndex );
 		if ( !LevelVariantSetsElement )
diff --git a/Engine/Plugins/Enterprise/DatasmithImporter/Source/DatasmithImporter/Private/Utility/DatasmithImporterImpl.cpp b/Engine/Plugins/Enterprise/DatasmithImporter/Source/DatasmithImporter/Private/Utility/DatasmithImporterImpl.cpp
--- a/Engine/Plugins/Enterprise/DatasmithImporter/Source/DatasmithImporter/Private/Utility/DatasmithImporterImpl.cpp
+++ b/Engine/Plugins/Enterprise/DatasmithImporter/Source/DatasmithImporter/Private/Utility/DatasmithImporterImpl.cpp
@@ -247,7 +247,10 @@ void FDatasmithImporterImpl::SetTexturesMode( FDatasmithImportContext& ImportCon
 	FFeedbackContext* FeedbackContext = ImportContext.FeedbackContext;
 	for ( int32 TextureIndex = 0; TextureIndex < TexturesCount && !ImportContext.bUserCancelled; ++TextureIndex )
 	{
-		ImportContext.bUserCancelled |= FDatasmithImporterImpl::HasUserCancelledTask( ImportContext.FeedbackContext );
+		if (FDatasmithImporterImpl::HasUserCancelledTask(ImportContext.FeedbackContext))
+		{
+			ImportContext.bUserCancelled = true;
+		}
 
 		TSharedPtr< IDatasmithTextureElement > TextureElement = ImportContext.FilteredScene->GetTexture( TextureIndex );
 		if (TextureElement->GetTextureMode() == EDatasmithTextureMode::Ies)
diff --git a/Engine/Source/Runtime/Core/Public/Templates/Atomic.h b/Engine/Source/Runtime/Core/Public/Templates/Atomic.h
--- a/Engine/Source/Runtime/Core/Public/Templates/Atomic.h
+++ b/Engine/Source/Runtime/Core/Public/Templates/Atomic.h
@@ -7,11 +7,18 @@
 #include "Templates/IsIntegral.h"
 #include "Templates/IsTrivial.h"
 #include "Traits/IntType.h"
+#include <atomic>
 
 ////////////////////////////////////////////////////////////////////////////////////////////////////////////////////////////
 // `TAtomic` is planned for deprecation. Please use `std::atomic`
 ////////////////////////////////////////////////////////////////////////////////////////////////////////////////////////////
 
+#define USE_DEPRECATED_TATOMIC 1
+
+#if !defined(USE_DEPRECATED_TATOMIC)
+#define USE_DEPRECATED_TATOMIC 0
+#endif
+
 template <typename T>
 struct TAtomicBase_Basic;
 
@@ -33,7 +40,9 @@ enum class EMemoryOrder
 	Relaxed,
 
 	// Establishes a single total order of all other atomic operations marked with this.
-	SequentiallyConsistent
+	SequentiallyConsistent,
+
+	Count
 };
 
 namespace UE4Atomic_Private
@@ -210,7 +219,11 @@ namespace UE4Atomic_Private
 
 // Basic storage and implementation - only allows getting and setting via platform atomics.
 template <typename T>
-struct alignas((UE4Atomic_Private::TIsSupportedSize<sizeof(T)>::Value) ? alignof(UE4Atomic_Private::TUnderlyingIntegerType_T<T>) : alignof(T)) TAtomicBase_Basic
+struct 
+#if USE_DEPRECATED_TATOMIC
+	alignas((UE4Atomic_Private::TIsSupportedSize<sizeof(T)>::Value) ? alignof(UE4Atomic_Private::TUnderlyingIntegerType_T<T>) : alignof(T))
+#endif
+	TAtomicBase_Basic
 {
 public:
 	/**
@@ -222,6 +235,7 @@ public:
 	 */
 	FORCEINLINE T Load(EMemoryOrder Order = EMemoryOrder::SequentiallyConsistent) const
 	{
+#if USE_DEPRECATED_TATOMIC
 		switch (Order)
 		{
 			case EMemoryOrder::Relaxed:
@@ -229,6 +243,9 @@ public:
 		}
 
 		return UE4Atomic_Private::Load(&Element);
+#else
+		return Element.load(ToStd(Order));
+#endif
 	}
 
 	/**
@@ -239,6 +256,7 @@ public:
 	 */
 	FORCEINLINE void Store(T Value, EMemoryOrder Order = EMemoryOrder::SequentiallyConsistent)
 	{
+#if USE_DEPRECATED_TATOMIC
 		switch (Order)
 		{
 			case EMemoryOrder::Relaxed:
@@ -246,6 +264,9 @@ public:
 		}
 
 		return UE4Atomic_Private::Store(&Element, Value);
+#else
+		Element.store(Value, ToStd(Order));
+#endif
 	}
 
 	/**
@@ -257,7 +278,11 @@ public:
 	 */
 	FORCEINLINE T Exchange(T Value)
 	{
+#if USE_DEPRECATED_TATOMIC
 		return UE4Atomic_Private::Exchange(&Element, Value);
+#else
+		return Element.exchange(Value);
+#endif
 	}
 
 	/**
@@ -284,10 +309,14 @@ public:
 	 */
 	FORCEINLINE bool CompareExchange(T& Expected, T Value)
 	{
+#if USE_DEPRECATED_TATOMIC
 		T PrevValue = UE4Atomic_Private::CompareExchange(&this->Element, Expected, Value);
 		bool bResult = PrevValue == Expected;
 		Expected = PrevValue;
 		return bResult;
+#else
+		return Element.compare_exchange_strong(Expected, Value);
+#endif
 	}
 
 protected:
@@ -298,7 +327,18 @@ protected:
 	{
 	}
 
+#if USE_DEPRECATED_TATOMIC
 	volatile T Element;
+#else
+	std::atomic<T> Element;
+#endif
+
+private:
+	static std::memory_order ToStd(EMemoryOrder Order)
+	{
+		std::memory_order Map[(int)EMemoryOrder::Count] = { std::memory_order_relaxed, std::memory_order_seq_cst };
+		return Map[(int)Order];
+	}
 };
 
 // TODO : basic storage and getting and setting, but the element is protected by a mutex instead of using platform atomics.
@@ -318,7 +358,11 @@ public:
 	 */
 	FORCEINLINE T operator++()
 	{
+#if USE_DEPRECATED_TATOMIC
 		return UE4Atomic_Private::IncrementExchange(&this->Element) + 1;
+#else
+		return this->Element.fetch_add(1) + 1;
+#endif
 	}
 
 	/**
@@ -328,7 +372,11 @@ public:
 	 */
 	FORCEINLINE T operator++(int)
 	{
+#if USE_DEPRECATED_TATOMIC
 		return UE4Atomic_Private::IncrementExchange(&this->Element);
+#else
+		return this->Element.fetch_add(1);
+#endif
 	}
 
 	/**
@@ -340,7 +388,11 @@ public:
 	 */
 	FORCEINLINE T operator+=(DiffType Value)
 	{
+#if USE_DEPRECATED_TATOMIC
 		return UE4Atomic_Private::AddExchange(&this->Element, Value) + Value;
+#else
+		return this->Element.fetch_add(Value) + Value;
+#endif
 	}
 
 	/**
@@ -350,7 +402,11 @@ public:
 	 */
 	FORCEINLINE T operator--()
 	{
+#if USE_DEPRECATED_TATOMIC
 		return UE4Atomic_Private::DecrementExchange(&this->Element) - 1;
+#else
+		return this->Element.fetch_sub(1) - 1;
+#endif
 	}
 
 	/**
@@ -360,7 +416,11 @@ public:
 	 */
 	FORCEINLINE T operator--(int)
 	{
+#if USE_DEPRECATED_TATOMIC
 		return UE4Atomic_Private::DecrementExchange(&this->Element);
+#else
+		return this->Element.fetch_sub(1);
+#endif
 	}
 
 	/**
@@ -372,7 +432,11 @@ public:
 	 */
 	FORCEINLINE T operator-=(DiffType Value)
 	{
+#if USE_DEPRECATED_TATOMIC
 		return UE4Atomic_Private::SubExchange(&this->Element, Value) - Value;
+#else
+		return this->Element.fetch_sub(Value) - Value;
+#endif
 	}
 
 	/**
@@ -382,7 +446,11 @@ public:
 	 */
 	FORCEINLINE T IncrementExchange()
 	{
+#if USE_DEPRECATED_TATOMIC
 		return UE4Atomic_Private::IncrementExchange(&this->Element);
+#else
+		return this->Element.fetch_add(1);
+#endif
 	}
 
 	/**
@@ -392,7 +460,11 @@ public:
 	 */
 	FORCEINLINE T DecrementExchange()
 	{
+#if USE_DEPRECATED_TATOMIC
 		return UE4Atomic_Private::DecrementExchange(&this->Element);
+#else
+		return this->Element.fetch_sub(1);
+#endif
 	}
 
 	/**
@@ -404,7 +476,11 @@ public:
 	 */
 	FORCEINLINE T AddExchange(DiffType Value)
 	{
+#if USE_DEPRECATED_TATOMIC
 		return UE4Atomic_Private::AddExchange(&this->Element, Value);
+#else
+		return this->Element.fetch_add(Value);
+#endif
 	}
 
 	/**
@@ -416,7 +492,11 @@ public:
 	 */
 	FORCEINLINE T SubExchange(DiffType Value)
 	{
+#if USE_DEPRECATED_TATOMIC
 		return UE4Atomic_Private::SubExchange(&this->Element, Value);
+#else
+		return this->Element.fetch_sub(Value);
+#endif
 	}
 
 protected:
@@ -455,7 +535,11 @@ public:
 	 */
 	FORCEINLINE T operator&=(const T Value)
 	{
+#if USE_DEPRECATED_TATOMIC
 		return UE4Atomic_Private::AndExchange(&this->Element, Value) & Value;
+#else
+		return this->Element.fetch_and(Value) & Value;
+#endif
 	}
 
 	/**
@@ -467,7 +551,11 @@ public:
 	 */
 	FORCEINLINE T operator|=(const T Value)
 	{
+#if USE_DEPRECATED_TATOMIC
 		return UE4Atomic_Private::OrExchange(&this->Element, Value) | Value;
+#else
+		return this->Element.fetch_or(Value) | Value;
+#endif
 	}
 
 	/**
@@ -479,7 +567,11 @@ public:
 	 */
 	FORCEINLINE T operator^=(const T Value)
 	{
+#if USE_DEPRECATED_TATOMIC
 		return UE4Atomic_Private::XorExchange(&this->Element, Value) ^ Value;
+#else
+		return this->Element.fetch_xor(Value) ^ Value;
+#endif
 	}
 
 	/**
@@ -491,7 +583,11 @@ public:
 	 */
 	FORCEINLINE T AndExchange(const T Value)
 	{
+#if USE_DEPRECATED_TATOMIC
 		return UE4Atomic_Private::AndExchange(&this->Element, Value);
+#else
+		return this->Element.fetch_and(Value);
+#endif
 	}
 
 	/**
@@ -503,7 +599,11 @@ public:
 	 */
 	FORCEINLINE T OrExchange(const T Value)
 	{
+#if USE_DEPRECATED_TATOMIC
 		return UE4Atomic_Private::OrExchange(&this->Element, Value);
+#else
+		return this->Element.fetch_or(Value);
+#endif
 	}
 
 	/**
@@ -515,7 +615,11 @@ public:
 	 */
 	FORCEINLINE T XorExchange(const T Value)
 	{
+#if USE_DEPRECATED_TATOMIC
 		return UE4Atomic_Private::XorExchange(&this->Element, Value);
+#else
+		return this->Element.fetch_xor(Value);
+#endif
 	}
 
 
