Non Unity include fixes #rb trivial #rnx

diff --git a/Engine/Source/Runtime/Core/Private/ProfilingDebugging/TraceAuxiliary.cpp b/Engine/Source/Runtime/Core/Private/ProfilingDebugging/TraceAuxiliary.cpp
--- a/Engine/Source/Runtime/Core/Private/ProfilingDebugging/TraceAuxiliary.cpp
+++ b/Engine/Source/Runtime/Core/Private/ProfilingDebugging/TraceAuxiliary.cpp
@@ -14,6 +14,11 @@
 #	define WITH_UNREAL_TRACE_LAUNCH (PLATFORM_DESKTOP && !UE_BUILD_SHIPPING && !IS_PROGRAM)
 #endif
 
+#include "CoreGlobals.h"
+#include "Misc/Paths.h"
+
+#include <atomic>
+
 #if UE_TRACE_ENABLED
 
 #include "BuildSettings.h"
@@ -21,7 +26,6 @@
 #include "Containers/Map.h"
 #include "Containers/StringView.h"
 #include "Containers/UnrealString.h"
-#include "CoreGlobals.h"
 #include "HAL/FileManager.h"
 #include "HAL/IConsoleManager.h"
 #include "HAL/PlatformMisc.h"
@@ -30,7 +34,6 @@
 #include "Misc/CoreDelegates.h"
 #include "Misc/CString.h"
 #include "Misc/DateTime.h"
-#include "Misc/Paths.h"
 #include "Modules/ModuleManager.h"
 #include "ProfilingDebugging/CountersTrace.h"
 #include "ProfilingDebugging/MiscTrace.h"
