Add SpringInterpQuaternion script function and ControlRig node. Also add an initialization option to the spring interp functions. COntrol rig nodes support the current input being optional, and expose strength rather than stiffness.

diff --git a/Engine/Plugins/Experimental/ControlRig/Source/ControlRig/Private/Units/Math/RigUnit_MathQuaternion.cpp b/Engine/Plugins/Experimental/ControlRig/Source/ControlRig/Private/Units/Math/RigUnit_MathQuaternion.cpp
--- a/Engine/Plugins/Experimental/ControlRig/Source/ControlRig/Private/Units/Math/RigUnit_MathQuaternion.cpp
+++ b/Engine/Plugins/Experimental/ControlRig/Source/ControlRig/Private/Units/Math/RigUnit_MathQuaternion.cpp
@@ -52,7 +52,6 @@ FRigUnit_MathQuaternionToAxisAndAngle_Execute()
 	float AngleSign = Angle < 0.f ? -1.f : 1.f;
 	Angle = FMath::Abs(Angle);
 
-	static const float TWO_PI = PI * 2.f;
 	if (Angle > TWO_PI)
 	{
 		Angle = FMath::Fmod(Angle, TWO_PI);
diff --git a/Engine/Plugins/Experimental/ControlRig/Source/ControlRig/Private/Units/Simulation/RigUnit_SpringInterp.cpp b/Engine/Plugins/Experimental/ControlRig/Source/ControlRig/Private/Units/Simulation/RigUnit_SpringInterp.cpp
--- a/Engine/Plugins/Experimental/ControlRig/Source/ControlRig/Private/Units/Simulation/RigUnit_SpringInterp.cpp
+++ b/Engine/Plugins/Experimental/ControlRig/Source/ControlRig/Private/Units/Simulation/RigUnit_SpringInterp.cpp
@@ -3,7 +3,6 @@
 #include "Units/Simulation/RigUnit_SpringInterp.h"
 #include "Units/RigUnitContext.h"
 
-
 namespace RigUnitSpringInterpConstants
 {
 	static const float Mass = 1.0f;
@@ -19,8 +18,13 @@ FRigUnit_SpringInterp_Execute()
 	}
 	else
 	{
-		Result = UKismetMathLibrary::FloatSpringInterp(Current, Target, SpringState, Stiffness, CriticalDamping,
-			Context.DeltaTime, RigUnitSpringInterpConstants::Mass, TargetVelocityAmount, false, 0.0f, 0.0f, true);
+		// Treat the input as a frequency in Hz
+		float AngularFrequency = Strength * 2.0f * PI;
+		float Stiffness = AngularFrequency * AngularFrequency;
+		Result = UKismetMathLibrary::FloatSpringInterp(
+			bUseCurrentInput ? Current : Result, Target, SpringState, Stiffness, CriticalDamping,
+			Context.DeltaTime, RigUnitSpringInterpConstants::Mass, TargetVelocityAmount, 
+			false, 0.0f, 0.0f, true, !bUseCurrentInput || bInitializeFromTarget);
 	}
 }
 
@@ -34,7 +38,33 @@ FRigUnit_SpringInterpVector_Execute()
 	}
 	else
 	{
-		Result = UKismetMathLibrary::VectorSpringInterp(Current, Target, SpringState, Stiffness, CriticalDamping,
-			Context.DeltaTime, RigUnitSpringInterpConstants::Mass, TargetVelocityAmount, false, FVector(), FVector(), true);
+		// Treat the input as a frequency in Hz
+		float AngularFrequency = Strength * 2.0f * PI;
+		float Stiffness = AngularFrequency * AngularFrequency;
+		Result = UKismetMathLibrary::VectorSpringInterp(
+			bUseCurrentInput ? Current : Result, Target, SpringState, Stiffness, CriticalDamping,
+			Context.DeltaTime, RigUnitSpringInterpConstants::Mass, TargetVelocityAmount, 
+			false, FVector(), FVector(), true, !bUseCurrentInput || bInitializeFromTarget);
+	}
+}
+
+FRigUnit_SpringInterpQuaternion_Execute()
+{
+    DECLARE_SCOPE_HIERARCHICAL_COUNTER_RIGUNIT()
+
+	if (Context.State == EControlRigState::Init)
+	{
+		SpringState.Reset();
+	}
+	else
+	{
+		// Treat the input as a frequency in Hz
+		float AngularFrequency = Strength * 2.0f * PI;
+		float Stiffness = AngularFrequency * AngularFrequency;
+		Result = UKismetMathLibrary::QuaternionSpringInterp(
+			bUseCurrentInput ? Current : Result, Target, SpringState, Stiffness, CriticalDamping,
+			Context.DeltaTime, RigUnitSpringInterpConstants::Mass, TargetVelocityAmount, 
+			true, !bUseCurrentInput || bInitializeFromTarget);
 	}
 }
+
diff --git a/Engine/Plugins/Experimental/ControlRig/Source/ControlRig/Private/Units/Simulation/RigUnit_SpringInterp.h b/Engine/Plugins/Experimental/ControlRig/Source/ControlRig/Private/Units/Simulation/RigUnit_SpringInterp.h
--- a/Engine/Plugins/Experimental/ControlRig/Source/ControlRig/Private/Units/Simulation/RigUnit_SpringInterp.h
+++ b/Engine/Plugins/Experimental/ControlRig/Source/ControlRig/Private/Units/Simulation/RigUnit_SpringInterp.h
@@ -16,43 +16,64 @@ struct CONTROLRIG_API FRigUnit_SpringInterp : public FRigUnit_SimBase
 
 	FRigUnit_SpringInterp()
 	{
+		bUseCurrentInput = true;
 		Current = Target = Result = 0.0f;
-		Stiffness = 10.0f;
-		CriticalDamping = 2.0f;
+		Strength = 4.0f;
+		CriticalDamping = 1.0f;
 		TargetVelocityAmount = 0.0f;
+		bInitializeFromTarget = false;
 		SpringState = FFloatSpringState();
 	}
 
 	RIGVM_METHOD()
 	virtual void Execute(const FRigUnitContext& Context) override;
 
-	/** Current position of the spring. */
-	UPROPERTY(meta = (Input))
-	float Current;
-
-	/** Rest position of the spring. */
+	/** Rest/target position of the spring. */
 	UPROPERTY(meta=(Input))
 	float Target;
 
-	/** The spring stiffness determines how hard it will pull towards the target. */
-	UPROPERTY(meta=(Input))
-	float Stiffness;
+	/**
+	 * If true, then the Current input will be used to initialize the state, and is required to be a variable that 
+	 * holds the current state. If false then the Target value will be used to initialize the state and the Current 
+	 * input will be ignored/unnecessary as a state will be maintained by this node.
+	 */
+	UPROPERTY(meta=(Input, Constant))
+	bool bUseCurrentInput;
+
+	/** Current position of the spring. */
+	UPROPERTY(meta = (Input, EditCondition = "bUseCurrentInput"))
+	float Current;
+
+	/**
+	 * The spring strength determines how hard it will pull towards the target. The value is the frequency
+	 * at which it will oscillate when there is no damping.
+	 */
+	UPROPERTY(meta=(Input, ClampMin="0"))
+	float Strength;
 
 	/** 
+	 * The amount of damping in the spring.
 	 * Set it smaller than 1 to make the spring oscillate before stabilizing on the target. 
 	 * Set it equal to 1 to reach the target without overshooting. 
 	 * Set it higher than one to make the spring take longer to reach the target.
 	 */
-	UPROPERTY(meta=(Input))
+	UPROPERTY(meta=(Input, ClampMin="0"))
 	float CriticalDamping;
 
 	/**
-	 * If 1, target changes turn into position offsets and the target is followed with less delay. As it gets close to
-	 * zero, the effect is disabled and only spring pull will affect the output.
+	 * The amount that the velocity should be passed through to the spring. A value of 1 will result in more
+	 * responsive output, but if the input is noisy or has step changes, these discontinuities will be passed 
+	 * through to the output much more than if a smaller value such as 0 is used.
 	 */
-	UPROPERTY(meta = (Input))
+	UPROPERTY(meta = (Input, ClampMin="0", ClampMax="1"))
 	float TargetVelocityAmount;
 
+	/**
+	 * If true, then the initial value will be taken from the target value, and not from the current value.
+	 */
+	UPROPERTY(meta = (Input, Constant, EditCondition = "bUseCurrentInput"))
+	bool bInitializeFromTarget;
+
 	/** New position of the spring after delta time. */
 	UPROPERTY(meta=(Output))
 	float Result;
@@ -71,47 +92,145 @@ struct CONTROLRIG_API FRigUnit_SpringInterpVector : public FRigUnit_SimBase
 	
 	FRigUnit_SpringInterpVector()
 	{
+		bUseCurrentInput = true;
 		Current = Target = Result = FVector::ZeroVector;
-		Stiffness = 10.0f;
-		CriticalDamping = 2.0f;
+		Strength = 4.0f;
+		CriticalDamping = 1.0f;
 		TargetVelocityAmount = 0.0f;
+		bInitializeFromTarget = false;
 		SpringState = FVectorSpringState();
 	}
 
 	RIGVM_METHOD()
 	virtual void Execute(const FRigUnitContext& Context) override;
 
-	/** Current position of the spring. */
-	UPROPERTY(meta = (Input))
-	FVector Current;
-
-	/** Rest position of the spring. */
+	/** Rest/target position of the spring. */
 	UPROPERTY(meta=(Input))
 	FVector Target;
 
-	/** The spring stiffness determines how hard it will pull towards the target. */
-	UPROPERTY(meta=(Input))
-	float Stiffness;
+	/**
+	 * If true, then the Current input will be used to initialize the state, and is required to be a variable that 
+	 * holds the current state. If false then the Target value will be used to initialize the state and the Current 
+	 * input will be ignored/unnecessary as a state will be maintained by this node.
+	 */
+	UPROPERTY(meta=(Input, Constant))
+	bool bUseCurrentInput;
+
+	/** Current position of the spring. */
+	UPROPERTY(meta = (Input, EditCondition = "bUseCurrentInput"))
+	FVector Current;
 
 	/**
-	 * Set it smaller than 1 to make the spring oscillate before stabilizing on the target.
-	 * Set it equal to 1 to reach the target without overshooting.
+	 * The spring strength determines how hard it will pull towards the target. The value is the frequency
+	 * at which it will oscillate when there is no damping.
+	 */
+	UPROPERTY(meta=(Input, ClampMin="0"))
+	float Strength;
+
+	/** 
+	 * The amount of damping in the spring.
+	 * Set it smaller than 1 to make the spring oscillate before stabilizing on the target. 
+	 * Set it equal to 1 to reach the target without overshooting. 
 	 * Set it higher than one to make the spring take longer to reach the target.
 	 */
-	UPROPERTY(meta=(Input))
+	UPROPERTY(meta=(Input, ClampMin="0"))
 	float CriticalDamping;
 
-	/** 
-	 * If 1, target changes turn into position offsets and the target is followed with less delay. As it gets close to 
-	 * zero, the effect is disabled and only spring pull will affect the output. 
+	/**
+	 * The amount that the velocity should be passed through to the spring. A value of 1 will result in more
+	 * responsive output, but if the input is noisy or has step changes, these discontinuities will be passed 
+	 * through to the output much more than if a smaller value such as 0 is used.
 	 */
-	UPROPERTY(meta = (Input))
+	UPROPERTY(meta = (Input, ClampMin="0", ClampMax="1"))
 	float TargetVelocityAmount;
 
+	/**
+	 * If true, then the initial value will be taken from the target value, and not from the current value. 
+	 */
+	UPROPERTY(meta = (Input, Constant, EditCondition = "bUseCurrentInput"))
+	bool bInitializeFromTarget;
+
 	/** New position of the spring after delta time. */
 	UPROPERTY(meta=(Output))
 	FVector Result;
 
 	UPROPERTY()
 	FVectorSpringState SpringState;
-};
\ No newline at end of file
+};
+
+/**
+ * Uses a simple spring model to interpolate a quaternion from Current to Target.
+ */
+USTRUCT(meta=(DisplayName="Spring Interpolate", Keywords="Alpha,SpringInterpolate", PrototypeName = "SpringInterp", MenuDescSuffix = "(Quaternion)"))
+struct CONTROLRIG_API FRigUnit_SpringInterpQuaternion : public FRigUnit_SimBase
+{
+	GENERATED_BODY()
+	
+	FRigUnit_SpringInterpQuaternion()
+	{
+		bUseCurrentInput = true;
+		Current = Target = Result = FQuat::Identity;
+		Strength = 4.0f;
+		CriticalDamping = 1.0f;
+		TargetVelocityAmount = 0.0f;
+		bInitializeFromTarget = false;
+		SpringState = FQuaternionSpringState();
+	}
+
+	RIGVM_METHOD()
+	virtual void Execute(const FRigUnitContext& Context) override;
+
+	/** Rest/target position of the spring. */
+	UPROPERTY(meta=(Input))
+	FQuat Target;
+
+	/**
+	 * If true, then the Current input will be used to initialize the state, and is required to be a variable that 
+	 * holds the current state. If false then the Target value will be used to initialize the state and the Current 
+	 * input will be ignored/unnecessary as a state will be maintained by this node.
+	 */
+	UPROPERTY(meta=(Input, Constant))
+	bool bUseCurrentInput;
+
+	/** Current position of the spring. */
+	UPROPERTY(meta = (Input, EditCondition = "bUseCurrentInput"))
+	FQuat Current;
+
+	/**
+	 * The spring strength determines how hard it will pull towards the target. The value is the frequency
+	 * at which it will oscillate when there is no damping.
+	 */
+	UPROPERTY(meta=(Input, ClampMin="0"))
+	float Strength;
+
+	/** 
+	 * The amount of damping in the spring.
+	 * Set it smaller than 1 to make the spring oscillate before stabilizing on the target. 
+	 * Set it equal to 1 to reach the target without overshooting. 
+	 * Set it higher than one to make the spring take longer to reach the target.
+	 */
+	UPROPERTY(meta=(Input, ClampMin="0"))
+	float CriticalDamping;
+
+	/**
+	 * The amount that the velocity should be passed through to the spring. A value of 1 will result in more
+	 * responsive output, but if the input is noisy or has step changes, these discontinuities will be passed 
+	 * through to the output much more than if a smaller value such as 0 is used.
+	 */
+	UPROPERTY(meta=(Input, ClampMin="0", ClampMax="1"))
+	float TargetVelocityAmount;
+
+	/**
+	 * If true, then the initial value will be taken from the target value, and not from the current value. 
+	 */
+	UPROPERTY(meta = (Input, Constant, EditCondition = "bUseCurrentInput"))
+	bool bInitializeFromTarget;
+
+	/** New position of the spring after delta time. */
+	UPROPERTY(meta=(Output))
+	FQuat Result;
+
+	UPROPERTY()
+	FQuaternionSpringState SpringState;
+};
+
diff --git a/Engine/Source/Runtime/Core/Public/Math/Quat.h b/Engine/Source/Runtime/Core/Public/Math/Quat.h
--- a/Engine/Source/Runtime/Core/Public/Math/Quat.h
+++ b/Engine/Source/Runtime/Core/Public/Math/Quat.h
@@ -70,6 +70,11 @@ public:
 	 */
 	FORCEINLINE FQuat(float InX, float InY, float InZ, float InW);
 
+	/**
+	 * Initializes all elements to V
+	 */
+	explicit FORCEINLINE FQuat(float V);
+
 	/**
 	 * Creates and initializes a new quaternion from the XYZW values in the given VectorRegister4Float.
 	 *
@@ -138,6 +143,13 @@ public:
 	 */
 	FORCEINLINE FQuat operator-(const FQuat& Q) const;
 
+	/**
+	* Negates the quaternion. Note that this represents the same rotation.
+	*
+	* @return The result of negation.
+	*/
+	FORCEINLINE FQuat operator-() const;
+
 	/**
 	 * Checks whether another Quaternion is equal to this, within specified tolerance.
 	 *
@@ -714,7 +726,6 @@ FORCEINLINE FQuat::FQuat(EForceInit ZeroOrNot)
 	:	X(0), Y(0), Z(0), W(ZeroOrNot == ForceInitToZero ? 0.0f : 1.0f)
 { }
 
-
 FORCEINLINE FQuat::FQuat(float InX, float InY, float InZ, float InW)
 	: X(InX)
 	, Y(InY)
@@ -724,6 +735,11 @@ FORCEINLINE FQuat::FQuat(float InX, float InY, float InZ, float InW)
 	DiagnosticCheckNaN();
 }
 
+FORCEINLINE FQuat::FQuat(float V)
+	: X(V), Y(V), Z(V), W(V)
+{
+	DiagnosticCheckNaN();
+}
 
 FORCEINLINE FQuat::FQuat(const VectorRegister4Float& V)
 {
@@ -737,7 +753,6 @@ FORCEINLINE FQuat::FQuat(const VectorRegister4Double& V)
 	DiagnosticCheckNaN();
 }
 
-
 FORCEINLINE FString FQuat::ToString() const
 {
 	return FString::Printf(TEXT("X=%.9f Y=%.9f Z=%.9f W=%.9f"), X, Y, Z, W);
@@ -827,6 +842,10 @@ FORCEINLINE FQuat FQuat::operator-(const FQuat& Q) const
 #endif // PLATFORM_ENABLE_VECTORINTRINSICS
 }
 
+FORCEINLINE FQuat FQuat::operator-() const
+{
+	return FQuat(-X, -Y, -Z, -W);
+}
 
 FORCEINLINE bool FQuat::Equals(const FQuat& Q, float Tolerance) const
 {
@@ -923,6 +942,10 @@ FORCEINLINE FQuat FQuat::operator*(const float Scale) const
 #endif // PLATFORM_ENABLE_VECTORINTRINSICS
 }
 
+FORCEINLINE FQuat operator*(const float Scale, const FQuat& Q)
+{
+	return Q.operator*(Scale);
+}
 
 FORCEINLINE FQuat FQuat::operator/=(const float Scale)
 {
diff --git a/Engine/Source/Runtime/Core/Public/Math/UnrealMathUtility.h b/Engine/Source/Runtime/Core/Public/Math/UnrealMathUtility.h
--- a/Engine/Source/Runtime/Core/Public/Math/UnrealMathUtility.h
+++ b/Engine/Source/Runtime/Core/Public/Math/UnrealMathUtility.h
@@ -68,9 +68,13 @@ class TRange;
 // Aux constants.
 #define INV_PI			(0.31830988618f)
 #define HALF_PI			(1.57079632679f)
+#define TWO_PI			(6.28318530717f)
+#define PI_SQUARED		(9.86960440108f)
 
-#define DOUBLE_INV_PI	(0.31830988618379067154)
-#define DOUBLE_HALF_PI	(1.57079632679489661923)
+#define DOUBLE_INV_PI		(0.31830988618379067154)
+#define DOUBLE_HALF_PI		(1.57079632679489661923)
+#define DOUBLE_TWO_PI		(6.28318530717958647692)
+#define DOUBLE_PI_SQUARED	(9.86960440108935861883)
 
 // Common square roots
 #define UE_SQRT_2		(1.4142135623730950488016887242097f)
@@ -567,7 +571,7 @@ public:
 		{
 			quotient = (T)((int64)(quotient - 0.5f));
 		}
-		T y = Value - (2.0f*PI)*quotient;
+		T y = Value - TWO_PI * quotient;
 
 		// Map y to [-pi/2,pi/2] with sin(y) = sin(Value).
 		T sign;
@@ -695,13 +699,13 @@ public:
 		if (Delta > PI)
 		{
 			// Flip to negative equivalent
-			Delta = Delta - (PI * 2.0f);
+			Delta = Delta - TWO_PI;
 		}
 		else if (Delta < -PI)
 		{
 			// Otherwise, if change is smaller than -PI
 			// Flip to positive equivalent
-			Delta = Delta + (PI * 2.0f);
+			Delta = Delta + TWO_PI;
 		}
 
 		// Return delta in [-PI,PI] range
@@ -719,12 +723,12 @@ public:
 	{
 		while(A > PI)
 		{
-			A -= ((float)PI * 2.0f);
+			A -= TWO_PI;
 		}
 
 		while(A < -PI)
 		{
-			A += ((float)PI * 2.0f);
+			A += TWO_PI;
 		}
 
 		return A;
@@ -1324,38 +1328,51 @@ public:
 	* 
 	* The implementation is exact, so stable and accurate for all values of DeltaTime etc
 	* 
-	* @param  InOutValue        The value to be smoothed
-	* @param  InOutValueRate    The rate of change of the value
-	* @param  InTargetValue     The target to smooth towards
-	* @param  InTargetValueRate The target rate of change smooth towards. Note that if this is discontinuous, then the output will have discontinuous velocity too.
-	* @param  InDeltaTime       Time interval
-	* @param  InSmoothingTime   Timescale over which to smooth. Larger values result in more smoothed behaviour. Can be zero.
-	* @param  InDampingRatio    1 is critical damping. <1 results in under-damped motion (i.e. with overshoot), and >1 results in over-damped motion. 
+	* @param  InOutValue          The value to be smoothed
+	* @param  InOutValueRate      The rate of change of the value
+	* @param  InTargetValue       The target to smooth towards
+	* @param  InTargetValueRate   The target rate of change smooth towards. Note that if this is discontinuous, then the output will have discontinuous velocity too.
+	* @param  InDeltaTime         Time interval
+	* @param  InUndampedFrequency Oscillation frequency when there is no damping. Proportional to the square root of the spring stiffness.
+	* @param  InDampingRatio      1 is critical damping. <1 results in under-damped motion (i.e. with overshoot), and >1 results in over-damped motion. 
 	*/
 	template< class T >
-	static void SpringDamperSmoothing(
+	static void SpringDamper(
 	    T&          InOutValue,
 	    T&          InOutValueRate,
 	    const T&    InTargetValue,
 	    const T&    InTargetValueRate,
 	    const float InDeltaTime,
-	    const float InSmoothingTime,
+	    const float InUndampedFrequency,
 	    const float InDampingRatio)
 	{
 		if (InDeltaTime <= 0.0f)
 		{
 			return;
 		}
-		if (InSmoothingTime < KINDA_SMALL_NUMBER)
+
+		float W = InUndampedFrequency * TWO_PI;
+		// Handle special cases
+		if (W < SMALL_NUMBER) // no strength which means no damping either
 		{
-			InOutValueRate = (InTargetValue - InOutValue) / InDeltaTime;
-			InOutValue = InTargetValue;
+			InOutValue += InOutValueRate * InDeltaTime;
+			return;
+		}
+		else if (InDampingRatio < SMALL_NUMBER) // No damping at all
+		{
+			T Err = InOutValue - InTargetValue;
+			const T A = Err;
+			const T B = InOutValueRate / W;
+			const float C = FMath::Cos(W * InDeltaTime);
+			const float S = FMath::Sin(W * InDeltaTime);
+			InOutValue = InTargetValue + (A * C + B * S);
+			InOutValueRate = InOutValueRate * C - Err * (W * S);
 			return;
 		}
-		// Target velocity turns into an offset to the position
-		T AdjustedTarget = InTargetValue + InTargetValueRate * (InDampingRatio * InSmoothingTime);
 
-		float W = 2.0f / InSmoothingTime;
+		// Target velocity turns into an offset to the position
+		float SmoothingTime = 2.0f / W;
+		T AdjustedTarget = InTargetValue + InTargetValueRate * (InDampingRatio * SmoothingTime);
 		T Err = InOutValue - AdjustedTarget;
 
 		// Handle the three cases separately
@@ -1368,9 +1385,8 @@ public:
 			const float A2 = -(WD + InDampingRatio * W);
 			const float E1 = FMath::Exp(A1 * InDeltaTime);
 			const float E2 = FMath::Exp(A2 * InDeltaTime);
-			InOutValue = E1 * C1 + E2 * C2;
+			InOutValue = AdjustedTarget + E1 * C1 + E2 * C2;
 			InOutValueRate = E1 * C1 * A1 + E2 * C2 * A2;
-			InOutValue += AdjustedTarget;
 		}
 		else if (InDampingRatio < 1.0f) // Underdamped
 		{
@@ -1390,18 +1406,15 @@ public:
 			const T& C1 = Err;
 			T C2 = InOutValueRate + Err * W;
 			float E = FMath::Exp(-W * InDeltaTime);
-			InOutValue = (C1 + C2 * InDeltaTime) * E;
+			InOutValue = AdjustedTarget + (C1 + C2 * InDeltaTime) * E;
 			InOutValueRate = (C2 - C1 * W - C2 * (W * InDeltaTime)) * E;
-			InOutValue += AdjustedTarget;
 		}
 	}
 
 	/**
 	* Smooths a value using a spring damper towards a target.
 	* 
-	* The implementation uses approximations for Exp/Sin/Cos. These are good for most sensible values of DampingRatio so
-	* long as InDeltaTime < 0.5 * InSmoothingTime. If you are passing in a very large delta time, or want accurate
-	* results for highly overdamped motion, you might need to use the exact version of this function. 
+	* The implementation is exact, so stable and accurate for all values of DeltaTime etc
 	* 
 	* @param  InOutValue        The value to be smoothed
 	* @param  InOutValueRate    The rate of change of the value
@@ -1412,32 +1425,88 @@ public:
 	* @param  InDampingRatio    1 is critical damping. <1 results in under-damped motion (i.e. with overshoot), and >1 results in over-damped motion. 
 	*/
 	template< class T >
-	static void SpringDamperSmoothingApprox(
+	static void SpringDamperSmoothing(
+		T&          InOutValue,
+		T&          InOutValueRate,
+		const T&    InTargetValue,
+		const T&    InTargetValueRate,
+		const float InDeltaTime,
+		const float InSmoothingTime,
+		const float InDampingRatio)
+	{
+		if (InSmoothingTime < SMALL_NUMBER)
+		{
+			if (InDeltaTime <= 0.0f)
+			{
+				return;
+			}
+			InOutValueRate = (InTargetValue - InOutValue) / InDeltaTime;
+			InOutValue = InTargetValue;
+			return;
+		}
+
+		// Undamped frequency
+		float UndampedFrequency = UndampedFrequency = 1.0f / (PI * InSmoothingTime);
+		SpringDamper(InOutValue, InOutValueRate, InTargetValue, InTargetValueRate, InDeltaTime, UndampedFrequency, InDampingRatio);
+	}
+
+	/**
+	* Smooths a value using a spring damper towards a target.
+	* 
+	* The implementation uses approximations for Exp/Sin/Cos. These are good for most sensible values of DampingRatio so
+	* long as InDeltaTime < 0.5 * InSmoothingTime. If you are passing in a very large delta time, or want accurate
+	* results for highly overdamped motion, you might need to use the exact version of this function. 
+	* 
+	* @param  InOutValue          The value to be smoothed
+	* @param  InOutValueRate      The rate of change of the value
+	* @param  InTargetValue       The target to smooth towards
+	* @param  InTargetValueRate   The target rate of change smooth towards. Note that if this is discontinuous, then the output will have discontinuous velocity too.
+	* @param  InDeltaTime         Time interval
+	* @param  InUndampedFrequency Oscillation frequency when there is no damping. Proportional to the square root of the spring stiffness.
+	* @param  InDampingRatio      1 is critical damping. <1 results in under-damped motion (i.e. with overshoot), and >1 results in over-damped motion. 
+	*/
+	template< class T >
+	static void SpringDamperApprox(
 	    T&          InOutValue,
 	    T&          InOutValueRate,
 	    const T&    InTargetValue,
 	    const T&    InTargetValueRate,
 	    const float InDeltaTime,
-	    const float InSmoothingTime,
+	    const float InUndampedFrequency,
 	    const float InDampingRatio)
 	{
 		if (InDeltaTime <= 0.0f)
 		{
 			return;
 		}
-		if (InSmoothingTime < KINDA_SMALL_NUMBER)
+
+		float W = InUndampedFrequency * TWO_PI;
+		// Handle special cases
+		if (W < SMALL_NUMBER) // no strength which means no damping either
 		{
-			InOutValueRate = (InTargetValue - InOutValue) / InDeltaTime;
-			InOutValue = InTargetValue;
+			InOutValue += InOutValueRate * InDeltaTime;
+			return;
+		}
+		else if (InDampingRatio < SMALL_NUMBER) // No damping at all
+		{
+			T Err = InOutValue - InTargetValue;
+			const T B = InOutValueRate / W;
+			// Bhaskara approximation for Sin(Angle) and Cos.
+			// Sin is really accurate for Angle < PI, and Cos is really accurate for Angle < PI/2
+			const float Angle = FMath::Min(W * InDeltaTime, HALF_PI);
+			const float S = 16.f * Angle * (PI - Angle) / (5.f * PI_SQUARED - 4.f * Angle * (PI - Angle));
+			const float C = (PI_SQUARED - 4.0f * FMath::Square(Angle)) / (PI_SQUARED + FMath::Square(Angle));
+			InOutValue = InTargetValue + Err * C + B * S;
+			InOutValueRate = InOutValueRate * C - Err * (W * S);
 			return;
 		}
-		// Target velocity turns into an offset to the position
-		T AdjustedTarget = InTargetValue + InTargetValueRate * (InDampingRatio * InSmoothingTime);
 
-		float W = 2.0f / InSmoothingTime;
-		T Err = InOutValue - AdjustedTarget;
+		// Target velocity turns into an offset to the position
+		float SmoothingTime = 2.0f / W;
+		T AdjustedTarget = InTargetValue + InTargetValueRate * (InDampingRatio * SmoothingTime);
+ 		T Err = InOutValue - AdjustedTarget;
 
-		// Handle the three cases separately
+		// Handle the cases separately
 		if (InDampingRatio > 1.0f) // Overdamped
 		{
 			const float WD = W * FMath::Sqrt(FMath::Square(InDampingRatio) - 1.0f);
@@ -1453,9 +1522,8 @@ public:
 			// As DampingRatio gets big, this approximation gets worse, but mere inaccuracy for overdamped motion is
 			// not likely to be important, since we end up with 1 / BigNumber
 			const float E2 = InvExpApprox(A2_DT);
-			InOutValue = E1 * C1 + E2 * C2;
+			InOutValue = AdjustedTarget + E1 * C1 + E2 * C2;
 			InOutValueRate = E1 * C1 * A1 + E2 * C2 * A2;
-			InOutValue += AdjustedTarget;
 		}
 		else if (InDampingRatio < 1.0f) // Underdamped
 		{
@@ -1463,10 +1531,10 @@ public:
 			const T A = Err;
 			const T B = (InOutValueRate + Err * (InDampingRatio * W)) / WD;
 			// Bhaskara approximation for Sin(Angle) and Cos.
-			// Needs Angle < 1 so DeltaTime < SmoothingTime / (2 * DampingRatio * Sqrt(1 - DampingRatio^2)) 
-			const float Angle = WD * InDeltaTime;
-			const float S = 16.f * Angle * (PI - Angle) / (5.f * FMath::Square(PI) - 4.f * Angle * (PI - Angle));
-			const float C = (FMath::Square(PI) - 4.0f * FMath::Square(Angle)) / (FMath::Square(PI) + FMath::Square(Angle));
+			// Sin is really accurate for Angle < PI, and Cos is really accurate for Angle < PI/2
+			const float Angle = FMath::Min(WD * InDeltaTime, HALF_PI);
+			const float S = 16.f * Angle * (PI - Angle) / (5.f * PI_SQUARED - 4.f * Angle * (PI - Angle));
+			const float C = (PI_SQUARED - 4.0f * FMath::Square(Angle)) / (PI_SQUARED + FMath::Square(Angle));
 			const float E0 = InDampingRatio * WD * InDeltaTime;
 			// Needs E0 < 1 so DeltaTime < SmoothingTime / (2 * DampingRatio * Sqrt(1 - DampingRatio^2))
 			const float E = InvExpApprox(E0);
@@ -1480,14 +1548,54 @@ public:
 			const T& C1 = Err;
 			T C2 = InOutValueRate + Err * W;
 			const float E0 = W * InDeltaTime;
-			// Needs E0 < 1 so InDeltaTime < InSmoothingTime / 2 
+			// Needs E0 < 1 so InDeltaTime < SmoothingTime / 2 
 			float E = InvExpApprox(E0);
-			InOutValue = (C1 + C2 * InDeltaTime) * E;
+			InOutValue = AdjustedTarget + (C1 + C2 * InDeltaTime) * E;
 			InOutValueRate = (C2 - C1 * W - C2 * (W * InDeltaTime)) * E;
-			InOutValue += AdjustedTarget;
 		}
 	}
 
+	/**
+	* Smooths a value using a spring damper towards a target.
+	* 
+	* The implementation uses approximations for Exp/Sin/Cos. These are good for most sensible values of DampingRatio so
+	* long as InDeltaTime < 0.5 * InSmoothingTime. If you are passing in a very large delta time, or want accurate
+	* results for highly overdamped motion, you might need to use the exact version of this function. 
+	* 
+	* @param  InOutValue        The value to be smoothed
+	* @param  InOutValueRate    The rate of change of the value
+	* @param  InTargetValue     The target to smooth towards
+	* @param  InTargetValueRate The target rate of change smooth towards. Note that if this is discontinuous, then the output will have discontinuous velocity too.
+	* @param  InDeltaTime       Time interval
+	* @param  InSmoothingTime   Timescale over which to smooth. Larger values result in more smoothed behaviour. Can be zero.
+	* @param  InDampingRatio    1 is critical damping. <1 results in under-damped motion (i.e. with overshoot), and >1 results in over-damped motion. 
+	*/
+	template< class T >
+	static void SpringDamperSmoothingApprox(
+		T&          InOutValue,
+		T&          InOutValueRate,
+		const T&    InTargetValue,
+		const T&    InTargetValueRate,
+		const float InDeltaTime,
+		const float InSmoothingTime,
+		const float InDampingRatio)
+	{
+		if (InSmoothingTime < SMALL_NUMBER)
+		{
+			if (InDeltaTime <= 0.0f)
+			{
+				return;
+			}
+			InOutValueRate = (InTargetValue - InOutValue) / InDeltaTime;
+			InOutValue = InTargetValue;
+			return;
+		}
+
+		// Undamped frequency
+		float UndampedFrequency = UndampedFrequency = 1.0f / (PI * InSmoothingTime);
+		SpringDamperApprox(InOutValue, InOutValueRate, InTargetValue, InTargetValueRate, InDeltaTime, UndampedFrequency, InDampingRatio);
+	}
+
 	/**
 	 * Simple function to create a pulsating scalar value
 	 *
@@ -1499,7 +1607,7 @@ public:
 	 */
 	static float MakePulsatingValue( const double InCurrentTime, const float InPulsesPerSecond, const float InPhase = 0.0f )
 	{
-		return 0.5f + 0.5f * FMath::Sin( ( ( 0.25f + InPhase ) * (float)PI * 2.0f ) + ( (float)InCurrentTime * (float)PI * 2.0f ) * InPulsesPerSecond );
+		return 0.5f + 0.5f * FMath::Sin( ( ( 0.25f + InPhase ) * TWO_PI ) + ( (float)InCurrentTime * TWO_PI ) * InPulsesPerSecond );
 	}
 
 	// Geometry intersection 
diff --git a/Engine/Source/Runtime/Engine/Classes/Kismet/KismetMathLibrary.h b/Engine/Source/Runtime/Engine/Classes/Kismet/KismetMathLibrary.h
--- a/Engine/Source/Runtime/Engine/Classes/Kismet/KismetMathLibrary.h
+++ b/Engine/Source/Runtime/Engine/Classes/Kismet/KismetMathLibrary.h
@@ -153,6 +153,30 @@ struct ENGINE_API FVectorSpringState
 	}
 };
 
+USTRUCT(BlueprintType)
+struct ENGINE_API FQuaternionSpringState
+{
+	GENERATED_BODY()
+
+	FQuat   PrevTarget;
+	FVector Velocity; // Angular velocity
+	bool    bPrevTargetValid;
+
+	FQuaternionSpringState()
+	: PrevTarget(FQuat::Identity)
+	, Velocity(FVector::ZeroVector)
+	, bPrevTargetValid(false)
+	{
+	}
+
+	void Reset()
+	{
+		PrevTarget = FQuat::Identity;
+		Velocity = FVector::ZeroVector;
+		bPrevTargetValid = false;
+	}
+};
+
 UCLASS(meta=(BlueprintThreadSafe, ScriptName = "MathLibrary"))
 class ENGINE_API UKismetMathLibrary : public UBlueprintFunctionLibrary
 {
@@ -1754,13 +1778,34 @@ class ENGINE_API UKismetMathLibrary : public UBlueprintFunctionLibrary
 	 * @param MinValue              Clamps the minimum output value and cancels the velocity if it reaches this limit
 	 * @param MaxValue              Clamps the maximum output value and cancels the velocity if it reaches this limit
 	 * @param bApproximate          Use an approximate/fast implementation. Will be suitable for all cases unless the delta time is very large or high accuracy is needed.
+	 * @param bInitializeFromTarget If set then the current value will be set from the target on the first update
 	 */
 	UFUNCTION(BlueprintCallable,  meta = (ScriptMethod = "InterpSpringTo", Keywords = "position", AdvancedDisplay = "8"), Category = "Math|Interpolation")
 	static FVector VectorSpringInterp(FVector Current, FVector Target, UPARAM(ref) FVectorSpringState& SpringState,
 	                                  float Stiffness, float CriticalDampingFactor, float DeltaTime,
 	                                  float Mass = 1.f, float TargetVelocityAmount = 1.f, 
 	                                  bool bClamp = false, FVector MinValue = FVector(-1.f), FVector MaxValue = FVector(1.f),
-	                                  bool bApproximate = true);
+	                                  bool bApproximate = true, bool bInitializeFromTarget = false);
+
+	/**
+	* Uses a simple spring model to interpolate a quaternion from Current to Target.
+	*
+	* @param Current               Current value
+	* @param Target                Target value
+	* @param SpringState           Data related to spring model (velocity, error, etc..) - Create a unique variable per spring
+	* @param Stiffness             How stiff the spring model is (more stiffness means more oscillation around the target value)
+	* @param CriticalDampingFactor How much damping to apply to the spring (0 means no damping, 1 means critically damped which means no oscillation)
+	* @param DeltaTime             Time difference since the last update
+	* @param Mass                  Multiplier that acts like mass on a spring
+	* @param TargetVelocityAmount  If 1 then the target velocity will be calculated and used, which results following the target more closely/without lag. Values down to zero (recommended when using this to smooth data) will progressively disable this effect.
+	* @param bApproximate          Use an approximate/fast implementation. Will be suitable for all cases unless the delta time is very large or high accuracy is needed.
+	* @param bInitializeFromTarget If set then the current value will be set from the target on the first update
+	*/
+	UFUNCTION(BlueprintCallable,  meta = (ScriptMethod = "InterpSpringTo", Keywords = "quaternion", AdvancedDisplay = "8"), Category = "Math|Interpolation")
+	static FQuat QuaternionSpringInterp(FQuat Current, FQuat Target, UPARAM(ref) FQuaternionSpringState& SpringState,
+	                                    float Stiffness, float CriticalDampingFactor, float DeltaTime,
+	                                    float Mass = 1.f, float TargetVelocityAmount = 1.f, bool bApproximate = true, 
+                                        bool bInitializeFromTarget = false);
 
 	/**
 	 * Gets the reciprocal of this vector, avoiding division by zero.
@@ -3858,22 +3903,38 @@ class ENGINE_API UKismetMathLibrary : public UBlueprintFunctionLibrary
 	 * @param MinValue              Clamps the minimum output value and cancels the velocity if it reaches this limit
 	 * @param MaxValue              Clamps the maximum output value and cancels the velocity if it reaches this limit
 	 * @param bApproximate          Use an approximate/fast implementation. Will be suitable for all cases unless the delta time is very large or high accuracy is needed.
+	 * @param bInitializeFromTarget If set then the current value will be set from the target on the first update
 	 */
 	UFUNCTION(BlueprintCallable, Category = "Math|Interpolation", meta=(AdvancedDisplay = "8"))
 	static float FloatSpringInterp(float Current, float Target, UPARAM(ref) FFloatSpringState& SpringState,
 	                               float Stiffness, float CriticalDampingFactor, float DeltaTime,
 	                               float Mass = 1.f, float TargetVelocityAmount = 1.f, 
 	                               bool bClamp = false, float MinValue = -1.f, float MaxValue = 1.f,
-	                               bool bApproximate = true);
+	                               bool bApproximate = true, bool bInitializeFromTarget = false);
 
-	/** Resets the state of a given spring */
+	/** Resets the state of a float spring */
 	UFUNCTION(BlueprintCallable, Category = "Math|Interpolation")
 	static void ResetFloatSpringState(UPARAM(ref) FFloatSpringState& SpringState);
 
-	/** Resets the state of a given spring */
+	/** Resets the state of a vector spring */
 	UFUNCTION(BlueprintCallable, Category = "Math|Interpolation")
 	static void ResetVectorSpringState(UPARAM(ref) FVectorSpringState& SpringState);
 
+	/** Resets the state of a quaternion spring */
+	UFUNCTION(BlueprintCallable, Category = "Math|Interpolation")
+	static void ResetQuaternionSpringState(UPARAM(ref) FQuaternionSpringState& SpringState);
+
+	/** Sets the state velocity of a float spring */
+	UFUNCTION(BlueprintCallable, Category = "Math|Interpolation")
+	static void SetFloatSpringStateVelocity(UPARAM(ref) FFloatSpringState& SpringState, float Velocity);
+
+	/** Sets the state velocity of a vector spring */
+	UFUNCTION(BlueprintCallable, Category = "Math|Interpolation")
+	static void SetVectorSpringStateVelocity(UPARAM(ref) FVectorSpringState& SpringState, FVector Velocity);
+
+	/** Sets the state angular velocity of a quaternion spring */
+	UFUNCTION(BlueprintCallable, Category = "Math|Interpolation")
+	static void SetQuaternionSpringStateVelocity(UPARAM(ref) FQuaternionSpringState& SpringState, FVector Velocity);
 
 	//
 	// Random stream functions
diff --git a/Engine/Source/Runtime/Engine/Private/KismetMathLibrary.cpp b/Engine/Source/Runtime/Engine/Private/KismetMathLibrary.cpp
--- a/Engine/Source/Runtime/Engine/Private/KismetMathLibrary.cpp
+++ b/Engine/Source/Runtime/Engine/Private/KismetMathLibrary.cpp
@@ -348,24 +348,28 @@ float UKismetMathLibrary::Ease(float A, float B, float Alpha, TEnumAsByte<EEasin
 template <typename T>
 void GenericSpringInterp(T& Current, const T Target, T& PrevTarget, bool& bPrevTargetValid, T& Velocity,
                          float Stiffness, float CriticalDamping, float DeltaTime,
-                         float TargetVelocityAmount, float Mass, bool bApproximate)
+                         float TargetVelocityAmount, float Mass, bool bApproximate, bool bInitializeFromTarget)
 {
+	if (bInitializeFromTarget && !bPrevTargetValid)
+	{
+		Current = Target;
+	}
 	if (DeltaTime > SMALL_NUMBER)
 	{
 		if (!FMath::IsNearlyZero(Mass))
 		{
 			// Note that old target was stored in PrevTarget
 			T TargetVel = bPrevTargetValid ? (Target - PrevTarget) * (TargetVelocityAmount / DeltaTime) : T(0.f);
-			const float Omega = FMath::Sqrt(Stiffness / Mass);
-			const float SmoothingTime = 2.0f / Omega;
+			const float Omega = FMath::Sqrt(Stiffness / Mass); // angular frequency
+			const float Frequency = Omega / (2.0f * PI);
 			T NewValue = Current;
 			if (bApproximate)
 			{
-				FMath::SpringDamperSmoothingApprox(Current, Velocity, Target, TargetVel, DeltaTime, SmoothingTime, CriticalDamping);
+				FMath::SpringDamperApprox(Current, Velocity, Target, TargetVel, DeltaTime, Frequency, CriticalDamping);
 			}
 			else
 			{
-				FMath::SpringDamperSmoothing(Current, Velocity, Target, TargetVel, DeltaTime, SmoothingTime, CriticalDamping);
+				FMath::SpringDamper(Current, Velocity, Target, TargetVel, DeltaTime, Frequency, CriticalDamping);
 			}
 			PrevTarget = Target;
 			bPrevTargetValid = true;
@@ -377,10 +381,10 @@ float UKismetMathLibrary::FloatSpringInterp(float Current, float Target, FFloatS
                                             float Stiffness, float CriticalDamping, float DeltaTime,
                                             float Mass, float TargetVelocityAmount, 
                                             bool bClamp, float MinValue, float MaxValue,
-                                            bool bApproximate)
+                                            bool bApproximate, bool bInitializeFromTarget)
 {
 	GenericSpringInterp(Current, Target, SpringState.PrevTarget, SpringState.bPrevTargetValid, SpringState.Velocity, Stiffness, CriticalDamping,
-	                    DeltaTime, TargetVelocityAmount, Mass, bApproximate);
+	                    DeltaTime, TargetVelocityAmount, Mass, bApproximate, bInitializeFromTarget);
 	if (bClamp)
 	{
 		if (Current < MinValue)
@@ -417,10 +421,10 @@ FVector UKismetMathLibrary::VectorSpringInterp(FVector Current, FVector Target,
                                                float Stiffness, float CriticalDamping, float DeltaTime,
                                                float Mass, float TargetVelocityAmount, 
                                                bool bClamp, FVector MinValue, FVector MaxValue,
-                                               bool bApproximate)
+                                               bool bApproximate, bool bInitializeFromTarget)
 {
 	GenericSpringInterp(Current, Target, SpringState.PrevTarget, SpringState.bPrevTargetValid, SpringState.Velocity, Stiffness,
-	                    CriticalDamping, DeltaTime, TargetVelocityAmount, Mass, bApproximate);
+	                    CriticalDamping, DeltaTime, TargetVelocityAmount, Mass, bApproximate, bInitializeFromTarget);
 	if (bClamp)
 	{
 		for (int Index = 0 ; Index != 3 ; ++Index)
@@ -446,6 +450,29 @@ FVector UKismetMathLibrary::VectorSpringInterp(FVector Current, FVector Target,
 	return Current;
 }
 
+FQuat UKismetMathLibrary::QuaternionSpringInterp(FQuat Current, FQuat Target, UPARAM(ref) FQuaternionSpringState& SpringState,
+                                                 float Stiffness, float CriticalDamping, float DeltaTime,
+                                                 float Mass, float TargetVelocityAmount, bool bApproximate, 
+                                                 bool bInitializeFromTarget)
+{
+	if ((Target | Current) < 0.0f)
+	{
+		Target = -Target;
+	}
+	FQuat VelocityQuat(SpringState.Velocity.X, SpringState.Velocity.Y, SpringState.Velocity.Z, 0.0f);
+	FQuat Velocity = 0.5f * Current * VelocityQuat;
+
+  	GenericSpringInterp(Current, Target, SpringState.PrevTarget, SpringState.bPrevTargetValid, Velocity, Stiffness,
+	                    CriticalDamping, DeltaTime, TargetVelocityAmount, Mass, bApproximate, bInitializeFromTarget);
+	Current.Normalize();
+
+	VelocityQuat = Current.Inverse() * 2.0f * Velocity;
+	SpringState.Velocity.Set(VelocityQuat.X, VelocityQuat.Y, VelocityQuat.Z);
+
+	return Current;
+}
+
+
 void UKismetMathLibrary::ResetFloatSpringState(FFloatSpringState& SpringState)
 {
 	SpringState.Reset();
@@ -456,6 +483,26 @@ void UKismetMathLibrary::ResetVectorSpringState(FVectorSpringState& SpringState)
 	SpringState.Reset();
 }
 
+void UKismetMathLibrary::ResetQuaternionSpringState(FQuaternionSpringState& SpringState)
+{
+	SpringState.Reset();
+}
+
+void UKismetMathLibrary::SetFloatSpringStateVelocity(FFloatSpringState& SpringState, float Velocity)
+{
+	SpringState.Velocity = Velocity;
+}
+
+void UKismetMathLibrary::SetVectorSpringStateVelocity(FVectorSpringState& SpringState, FVector Velocity)
+{
+	SpringState.Velocity = Velocity;
+}
+
+void UKismetMathLibrary::SetQuaternionSpringStateVelocity(FQuaternionSpringState& SpringState, FVector Velocity)
+{
+	SpringState.Velocity = Velocity;
+}
+
 FVector UKismetMathLibrary::RandomUnitVector()
 {
 	return FMath::VRand();
