Trivial/compiler warning fixes

diff --git a/Engine/Source/Runtime/Core/Public/Math/UnrealMathUtility.h b/Engine/Source/Runtime/Core/Public/Math/UnrealMathUtility.h
--- a/Engine/Source/Runtime/Core/Public/Math/UnrealMathUtility.h
+++ b/Engine/Source/Runtime/Core/Public/Math/UnrealMathUtility.h
@@ -1446,7 +1446,7 @@ public:
 		}
 
 		// Undamped frequency
-		float UndampedFrequency = UndampedFrequency = 1.0f / (PI * InSmoothingTime);
+		float UndampedFrequency = 1.0f / (PI * InSmoothingTime);
 		SpringDamper(InOutValue, InOutValueRate, InTargetValue, InTargetValueRate, InDeltaTime, UndampedFrequency, InDampingRatio);
 	}
 
@@ -1592,7 +1592,7 @@ public:
 		}
 
 		// Undamped frequency
-		float UndampedFrequency = UndampedFrequency = 1.0f / (PI * InSmoothingTime);
+		float UndampedFrequency = 1.0f / (PI * InSmoothingTime);
 		SpringDamperApprox(InOutValue, InOutValueRate, InTargetValue, InTargetValueRate, InDeltaTime, UndampedFrequency, InDampingRatio);
 	}
 
