[Insights]
  - Fixed tracing of UObject names and function names (UObject::CallFunction, FScopeCycleCounterUObject) when "stat named events" toggle is enabled (-statnamedevents, "stat NamedEvents" console command).
  - Added FCpuProfilerTrace::OutputBeginDynamicEvent(const FName), in addition to existing API that receives const ANSICHAR*  or const TCHAR*. It uses a separate DynamicFNameScopeNamesMap.

diff --git a/Engine/Source/Runtime/Core/Private/ProfilingDebugging/CpuProfilerTrace.cpp b/Engine/Source/Runtime/Core/Private/ProfilingDebugging/CpuProfilerTrace.cpp
--- a/Engine/Source/Runtime/Core/Private/ProfilingDebugging/CpuProfilerTrace.cpp
+++ b/Engine/Source/Runtime/Core/Private/ProfilingDebugging/CpuProfilerTrace.cpp
@@ -80,6 +80,7 @@ struct FCpuProfilerTraceInternal
 		FMemStackBase DynamicScopeNamesMemory;
 		TMap<const ANSICHAR*, uint32, FDefaultSetAllocator, FDynamicScopeNameMapKeyFuncs<ANSICHAR>> DynamicAnsiScopeNamesMap;
 		TMap<const TCHAR*, uint32, FDefaultSetAllocator, FDynamicScopeNameMapKeyFuncs<TCHAR>> DynamicTCharScopeNamesMap;
+		TMap<FNameEntryId, uint32> DynamicFNameScopeNamesMap;
 	};
 
 	uint32 static GetNextSpecId();
@@ -174,6 +175,30 @@ void FCpuProfilerTrace::OutputBeginDynamicEvent(const TCHAR* Name)
 	CPUPROFILERTRACE_OUTPUTBEGINEVENT_EPILOGUE();
 }
 
+void FCpuProfilerTrace::OutputBeginDynamicEvent(const FName Name)
+{
+	CPUPROFILERTRACE_OUTPUTBEGINEVENT_PROLOGUE();
+	uint32 SpecId = ThreadBuffer->DynamicFNameScopeNamesMap.FindRef(Name.GetComparisonIndex());
+	if (!SpecId)
+	{
+		const FNameEntry* NameEntry = Name.GetDisplayNameEntry();
+		if (NameEntry->IsWide())
+		{
+			static WIDECHAR WideName[NAME_SIZE];
+			NameEntry->GetWideName(WideName);
+			SpecId = OutputEventType(WideName);
+		}
+		else
+		{
+			static ANSICHAR AnsiName[NAME_SIZE];
+			NameEntry->GetAnsiName(AnsiName);
+			SpecId = OutputEventType(AnsiName);
+		}
+		ThreadBuffer->DynamicFNameScopeNamesMap.Add(Name.GetComparisonIndex(), SpecId);
+	}
+	CPUPROFILERTRACE_OUTPUTBEGINEVENT_EPILOGUE();
+}
+
 #undef CPUPROFILERTRACE_OUTPUTBEGINEVENT_PROLOGUE
 #undef CPUPROFILERTRACE_OUTPUTBEGINEVENT_EPILOGUE
 
diff --git a/Engine/Source/Runtime/Core/Public/ProfilingDebugging/CpuProfilerTrace.h b/Engine/Source/Runtime/Core/Public/ProfilingDebugging/CpuProfilerTrace.h
--- a/Engine/Source/Runtime/Core/Public/ProfilingDebugging/CpuProfilerTrace.h
+++ b/Engine/Source/Runtime/Core/Public/ProfilingDebugging/CpuProfilerTrace.h
@@ -60,6 +60,13 @@ struct FCpuProfilerTrace
 	 * @param Name Name of event
 	 */
 	CORE_API static void OutputBeginDynamicEvent(const TCHAR* Name);
+	/*
+	 * Output begin event marker for a dynamic event identified by an FName. This is more expensive than
+	 * statically known event names using \ref OutputBeginEvent, but it is faster than \ref OutputBeginDynamicEvent
+	 * that receives ANSICHAR* / TCHAR* name. Must always be matched with an end event.
+	 * @param Name Name of event
+	 */
+	CORE_API static void OutputBeginDynamicEvent(const FName Name);
 	/*
 	 * Output end event marker for static or dynamic event for the currently open scope.
 	 */
diff --git a/Engine/Source/Runtime/Core/Public/Stats/Stats2.h b/Engine/Source/Runtime/Core/Public/Stats/Stats2.h
--- a/Engine/Source/Runtime/Core/Public/Stats/Stats2.h
+++ b/Engine/Source/Runtime/Core/Public/Stats/Stats2.h
@@ -1601,6 +1601,17 @@ public:
 		Start(InStatId, EStatFlags::None, bAlways);
 	}
 
+	FORCEINLINE_STATS void StartTrace(const FName Name)
+	{
+#if CPUPROFILERTRACE_ENABLED
+		if (UE_TRACE_CHANNELEXPR_IS_ENABLED(CpuChannel))
+		{
+			FCpuProfilerTrace::OutputBeginDynamicEvent(Name);
+			EmittedEvent |= TraceEvent;
+		}
+#endif
+	}
+
 	/**
 	 * Stops the capturing and stores the result
 	 */
diff --git a/Engine/Source/Runtime/CoreUObject/Private/UObject/ScriptCore.cpp b/Engine/Source/Runtime/CoreUObject/Private/UObject/ScriptCore.cpp
--- a/Engine/Source/Runtime/CoreUObject/Private/UObject/ScriptCore.cpp
+++ b/Engine/Source/Runtime/CoreUObject/Private/UObject/ScriptCore.cpp
@@ -957,12 +957,12 @@ IMPLEMENT_VM_FUNCTION(EX_CallMath, execCallMathFunction);
 void UObject::CallFunction( FFrame& Stack, RESULT_DECL, UFunction* Function )
 {
 #if PER_FUNCTION_SCRIPT_STATS
-	const bool bShouldTrackFunction = (Stack.DepthCounter <= GMaxFunctionStatDepth) && Stats::IsThreadCollectingData();
+	const bool bShouldTrackFunction = (Stack.DepthCounter <= GMaxFunctionStatDepth) && (Stats::IsThreadCollectingData() || UE_TRACE_CHANNELEXPR_IS_ENABLED(CpuChannel));
 	FScopeCycleCounterUObject FunctionScope(bShouldTrackFunction ? Function : nullptr);
 #endif // PER_FUNCTION_SCRIPT_STATS
 
 #if STATS || ENABLE_STATNAMEDEVENTS
-	const bool bShouldTrackObject = GVerboseScriptStats && Stats::IsThreadCollectingData();
+	const bool bShouldTrackObject = (GVerboseScriptStats && Stats::IsThreadCollectingData()) || UE_TRACE_CHANNELEXPR_IS_ENABLED(CpuChannel);
 	FScopeCycleCounterUObject ContextScope(bShouldTrackObject ? this : nullptr);
 #endif
 
diff --git a/Engine/Source/Runtime/CoreUObject/Public/UObject/UObjectBaseUtility.h b/Engine/Source/Runtime/CoreUObject/Public/UObject/UObjectBaseUtility.h
--- a/Engine/Source/Runtime/CoreUObject/Public/UObject/UObjectBaseUtility.h
+++ b/Engine/Source/Runtime/CoreUObject/Public/UObject/UObjectBaseUtility.h
@@ -820,10 +820,17 @@ public:
 	{
 		if (Object)
 		{
+			bool bStarted = false;
 			TStatId ObjectStatId = Object->GetStatID();
 			if (FThreadStats::IsCollectingData(ObjectStatId))
 			{
 				Start(ObjectStatId);
+				bStarted = true;
+			}
+
+			if (!bStarted && UE_TRACE_CHANNELEXPR_IS_ENABLED(CpuChannel))
+			{
+				StartTrace(Object->GetFName());
 			}
 		}
 
@@ -839,12 +846,22 @@ public:
 	 */
 	FORCEINLINE_STATS FScopeCycleCounterUObject(const UObjectBaseUtility *Object, TStatId OtherStat)
 	{
-		if (FThreadStats::IsCollectingData(OtherStat) && Object)
+		if (Object)
 		{
-			TStatId ObjectStatId = Object->GetStatID();
-			if (!ObjectStatId.IsNone())
+			bool bStarted = false;
+			if (FThreadStats::IsCollectingData(OtherStat))
 			{
-				Start(ObjectStatId);
+				TStatId ObjectStatId = Object->GetStatID();
+				if (!ObjectStatId.IsNone())
+				{
+					Start(ObjectStatId);
+					bStarted = true;
+				}
+			}
+
+			if (!bStarted && UE_TRACE_CHANNELEXPR_IS_ENABLED(CpuChannel))
+			{
+				StartTrace(Object->GetFName());
 			}
 		}
 
