Remove the approx version of spring damper. Just deleting some very newly introduced options so hopefully not problematic.

diff --git a/Engine/Plugins/Animation/ControlRig/Source/ControlRig/Private/Units/Simulation/RigUnit_SpringInterp.cpp b/Engine/Plugins/Animation/ControlRig/Source/ControlRig/Private/Units/Simulation/RigUnit_SpringInterp.cpp
--- a/Engine/Plugins/Animation/ControlRig/Source/ControlRig/Private/Units/Simulation/RigUnit_SpringInterp.cpp
+++ b/Engine/Plugins/Animation/ControlRig/Source/ControlRig/Private/Units/Simulation/RigUnit_SpringInterp.cpp
@@ -24,7 +24,7 @@ FRigUnit_SpringInterp_Execute()
 		Result = UKismetMathLibrary::FloatSpringInterp(
 			bUseCurrentInput ? Current : Result, Target, SpringState, Stiffness, CriticalDamping,
 			Context.DeltaTime, RigUnitSpringInterpConstants::Mass, TargetVelocityAmount, 
-			false, 0.0f, 0.0f, true, !bUseCurrentInput || bInitializeFromTarget);
+			false, 0.0f, 0.0f, !bUseCurrentInput || bInitializeFromTarget);
 	}
 }
 
@@ -44,7 +44,7 @@ FRigUnit_SpringInterpVector_Execute()
 		Result = UKismetMathLibrary::VectorSpringInterp(
 			bUseCurrentInput ? Current : Result, Target, SpringState, Stiffness, CriticalDamping,
 			Context.DeltaTime, RigUnitSpringInterpConstants::Mass, TargetVelocityAmount, 
-			false, FVector(), FVector(), true, !bUseCurrentInput || bInitializeFromTarget);
+			false, FVector(), FVector(), !bUseCurrentInput || bInitializeFromTarget);
 	}
 }
 
@@ -64,7 +64,7 @@ FRigUnit_SpringInterpQuaternion_Execute()
 		Result = UKismetMathLibrary::QuaternionSpringInterp(
 			bUseCurrentInput ? Current : Result, Target, SpringState, Stiffness, CriticalDamping,
 			Context.DeltaTime, RigUnitSpringInterpConstants::Mass, TargetVelocityAmount, 
-			true, !bUseCurrentInput || bInitializeFromTarget);
+			!bUseCurrentInput || bInitializeFromTarget);
 	}
 }
 
diff --git a/Engine/Source/Runtime/Core/Public/Math/UnrealMathUtility.h b/Engine/Source/Runtime/Core/Public/Math/UnrealMathUtility.h
--- a/Engine/Source/Runtime/Core/Public/Math/UnrealMathUtility.h
+++ b/Engine/Source/Runtime/Core/Public/Math/UnrealMathUtility.h
@@ -1326,7 +1326,9 @@ public:
 	/**
 	* Smooths a value using a spring damper towards a target.
 	* 
-	* The implementation is exact, so stable and accurate for all values of DeltaTime etc
+	* The implementation uses approximations for Exp/Sin/Cos. These are accurate for all sensible values of 
+	* InUndampedFrequency and DampingRatio so long as InDeltaTime < 1 / InUndampedFrequency (approximately), but
+	* are generally well behaved even for larger timesteps etc. 
 	* 
 	* @param  InOutValue          The value to be smoothed
 	* @param  InOutValueRate      The rate of change of the value
@@ -1361,141 +1363,9 @@ public:
 		else if (InDampingRatio < SMALL_NUMBER) // No damping at all
 		{
 			T Err = InOutValue - InTargetValue;
-			const T A = Err;
 			const T B = InOutValueRate / W;
-			const float C = FMath::Cos(W * InDeltaTime);
-			const float S = FMath::Sin(W * InDeltaTime);
-			InOutValue = InTargetValue + (A * C + B * S);
-			InOutValueRate = InOutValueRate * C - Err * (W * S);
-			return;
-		}
-
-		// Target velocity turns into an offset to the position
-		float SmoothingTime = 2.0f / W;
-		T AdjustedTarget = InTargetValue + InTargetValueRate * (InDampingRatio * SmoothingTime);
-		T Err = InOutValue - AdjustedTarget;
-
-		// Handle the three cases separately
-		if (InDampingRatio > 1.0f) // Overdamped
-		{
-			const float WD = W * FMath::Sqrt(FMath::Square(InDampingRatio) - 1.0f);
-			const T C2 = -(InOutValueRate + (W * InDampingRatio - WD) * Err) / (2.0f * WD);
-			const T C1 = Err - C2;
-			const float A1 = (WD - InDampingRatio * W);
-			const float A2 = -(WD + InDampingRatio * W);
-			const float E1 = FMath::Exp(A1 * InDeltaTime);
-			const float E2 = FMath::Exp(A2 * InDeltaTime);
-			InOutValue = AdjustedTarget + E1 * C1 + E2 * C2;
-			InOutValueRate = E1 * C1 * A1 + E2 * C2 * A2;
-		}
-		else if (InDampingRatio < 1.0f) // Underdamped
-		{
-			const float WD = W * FMath::Sqrt(1.0f - FMath::Square(InDampingRatio));
-			const T A = Err;
-			const T B = (InOutValueRate + Err * (InDampingRatio * W)) / WD;
-			const float C = FMath::Cos(WD * InDeltaTime);
-			const float S = FMath::Sin(WD * InDeltaTime);
-			const float E = FMath::Exp(-InDampingRatio * W * InDeltaTime);
-			InOutValue = E * (A * C + B * S);
-			InOutValueRate = -InOutValue * InDampingRatio * W;
-			InOutValueRate += E * (B * (WD * C) - A * (WD * S));
-			InOutValue += AdjustedTarget;
-		}
-		else // Critical damping
-		{
-			const T& C1 = Err;
-			T C2 = InOutValueRate + Err * W;
-			float E = FMath::Exp(-W * InDeltaTime);
-			InOutValue = AdjustedTarget + (C1 + C2 * InDeltaTime) * E;
-			InOutValueRate = (C2 - C1 * W - C2 * (W * InDeltaTime)) * E;
-		}
-	}
-
-	/**
-	* Smooths a value using a spring damper towards a target.
-	* 
-	* The implementation is exact, so stable and accurate for all values of DeltaTime etc
-	* 
-	* @param  InOutValue        The value to be smoothed
-	* @param  InOutValueRate    The rate of change of the value
-	* @param  InTargetValue     The target to smooth towards
-	* @param  InTargetValueRate The target rate of change smooth towards. Note that if this is discontinuous, then the output will have discontinuous velocity too.
-	* @param  InDeltaTime       Time interval
-	* @param  InSmoothingTime   Timescale over which to smooth. Larger values result in more smoothed behaviour. Can be zero.
-	* @param  InDampingRatio    1 is critical damping. <1 results in under-damped motion (i.e. with overshoot), and >1 results in over-damped motion. 
-	*/
-	template< class T >
-	static void SpringDamperSmoothing(
-		T&          InOutValue,
-		T&          InOutValueRate,
-		const T&    InTargetValue,
-		const T&    InTargetValueRate,
-		const float InDeltaTime,
-		const float InSmoothingTime,
-		const float InDampingRatio)
-	{
-		if (InSmoothingTime < SMALL_NUMBER)
-		{
-			if (InDeltaTime <= 0.0f)
-			{
-				return;
-			}
-			InOutValueRate = (InTargetValue - InOutValue) / InDeltaTime;
-			InOutValue = InTargetValue;
-			return;
-		}
-
-		// Undamped frequency
-		float UndampedFrequency = 1.0f / (PI * InSmoothingTime);
-		SpringDamper(InOutValue, InOutValueRate, InTargetValue, InTargetValueRate, InDeltaTime, UndampedFrequency, InDampingRatio);
-	}
-
-	/**
-	* Smooths a value using a spring damper towards a target.
-	* 
-	* The implementation uses approximations for Exp/Sin/Cos. These are good for most sensible values of DampingRatio so
-	* long as InDeltaTime < 0.5 * InSmoothingTime. If you are passing in a very large delta time, or want accurate
-	* results for highly overdamped motion, you might need to use the exact version of this function. 
-	* 
-	* @param  InOutValue          The value to be smoothed
-	* @param  InOutValueRate      The rate of change of the value
-	* @param  InTargetValue       The target to smooth towards
-	* @param  InTargetValueRate   The target rate of change smooth towards. Note that if this is discontinuous, then the output will have discontinuous velocity too.
-	* @param  InDeltaTime         Time interval
-	* @param  InUndampedFrequency Oscillation frequency when there is no damping. Proportional to the square root of the spring stiffness.
-	* @param  InDampingRatio      1 is critical damping. <1 results in under-damped motion (i.e. with overshoot), and >1 results in over-damped motion. 
-	*/
-	template< class T >
-	static void SpringDamperApprox(
-	    T&          InOutValue,
-	    T&          InOutValueRate,
-	    const T&    InTargetValue,
-	    const T&    InTargetValueRate,
-	    const float InDeltaTime,
-	    const float InUndampedFrequency,
-	    const float InDampingRatio)
-	{
-		if (InDeltaTime <= 0.0f)
-		{
-			return;
-		}
-
-		float W = InUndampedFrequency * TWO_PI;
-		// Handle special cases
-		if (W < SMALL_NUMBER) // no strength which means no damping either
-		{
-			InOutValue += InOutValueRate * InDeltaTime;
-			return;
-		}
-		else if (InDampingRatio < SMALL_NUMBER) // No damping at all
-		{
-			T Err = InOutValue - InTargetValue;
-			const T B = InOutValueRate / W;
-			// Bhaskara approximation for Sin(Angle) and Cos.
-			// Sin is really accurate for Angle < PI, and Cos is really accurate for Angle < PI/2
-			const float Angle = FMath::Min(W * InDeltaTime, HALF_PI);
-			const float S = 16.f * Angle * (PI - Angle) / (5.f * PI_SQUARED - 4.f * Angle * (PI - Angle));
-			const float C = (PI_SQUARED - 4.0f * FMath::Square(Angle)) / (PI_SQUARED + FMath::Square(Angle));
+			float S, C;
+			FMath::SinCos(&S, &C, W * InDeltaTime);
 			InOutValue = InTargetValue + Err * C + B * S;
 			InOutValueRate = InOutValueRate * C - Err * (W * S);
 			return;
@@ -1530,11 +1400,8 @@ public:
 			const float WD = W * FMath::Sqrt(1.0f - FMath::Square(InDampingRatio));
 			const T A = Err;
 			const T B = (InOutValueRate + Err * (InDampingRatio * W)) / WD;
-			// Bhaskara approximation for Sin(Angle) and Cos.
-			// Sin is really accurate for Angle < PI, and Cos is really accurate for Angle < PI/2
-			const float Angle = FMath::Min(WD * InDeltaTime, HALF_PI);
-			const float S = 16.f * Angle * (PI - Angle) / (5.f * PI_SQUARED - 4.f * Angle * (PI - Angle));
-			const float C = (PI_SQUARED - 4.0f * FMath::Square(Angle)) / (PI_SQUARED + FMath::Square(Angle));
+			float S, C;
+			FMath::SinCos(&S, &C, W * InDeltaTime);
 			const float E0 = InDampingRatio * WD * InDeltaTime;
 			// Needs E0 < 1 so DeltaTime < SmoothingTime / (2 * DampingRatio * Sqrt(1 - DampingRatio^2))
 			const float E = InvExpApprox(E0);
@@ -1558,9 +1425,9 @@ public:
 	/**
 	* Smooths a value using a spring damper towards a target.
 	* 
-	* The implementation uses approximations for Exp/Sin/Cos. These are good for most sensible values of DampingRatio so
-	* long as InDeltaTime < 0.5 * InSmoothingTime. If you are passing in a very large delta time, or want accurate
-	* results for highly overdamped motion, you might need to use the exact version of this function. 
+	* The implementation uses approximations for Exp/Sin/Cos. These are accurate for all sensible values of 
+	* DampingRatio and InSmoothingTime so long as InDeltaTime < 0.5 * InSmoothingTime, but are generally well behaved 
+	* even for larger timesteps etc. 
 	* 
 	* @param  InOutValue        The value to be smoothed
 	* @param  InOutValueRate    The rate of change of the value
@@ -1571,7 +1438,7 @@ public:
 	* @param  InDampingRatio    1 is critical damping. <1 results in under-damped motion (i.e. with overshoot), and >1 results in over-damped motion. 
 	*/
 	template< class T >
-	static void SpringDamperSmoothingApprox(
+	static void SpringDamperSmoothing(
 		T&          InOutValue,
 		T&          InOutValueRate,
 		const T&    InTargetValue,
@@ -1593,7 +1460,7 @@ public:
 
 		// Undamped frequency
 		float UndampedFrequency = 1.0f / (PI * InSmoothingTime);
-		SpringDamperApprox(InOutValue, InOutValueRate, InTargetValue, InTargetValueRate, InDeltaTime, UndampedFrequency, InDampingRatio);
+		SpringDamper(InOutValue, InOutValueRate, InTargetValue, InTargetValueRate, InDeltaTime, UndampedFrequency, InDampingRatio);
 	}
 
 	/**
diff --git a/Engine/Source/Runtime/Engine/Classes/Kismet/KismetMathLibrary.h b/Engine/Source/Runtime/Engine/Classes/Kismet/KismetMathLibrary.h
--- a/Engine/Source/Runtime/Engine/Classes/Kismet/KismetMathLibrary.h
+++ b/Engine/Source/Runtime/Engine/Classes/Kismet/KismetMathLibrary.h
@@ -1777,7 +1777,6 @@ class ENGINE_API UKismetMathLibrary : public UBlueprintFunctionLibrary
 	 * @param bClamp                Whether to use the Min/Max values to clamp the motion
 	 * @param MinValue              Clamps the minimum output value and cancels the velocity if it reaches this limit
 	 * @param MaxValue              Clamps the maximum output value and cancels the velocity if it reaches this limit
-	 * @param bApproximate          Use an approximate/fast implementation. Will be suitable for all cases unless the delta time is very large or high accuracy is needed.
 	 * @param bInitializeFromTarget If set then the current value will be set from the target on the first update
 	 */
 	UFUNCTION(BlueprintCallable,  meta = (ScriptMethod = "InterpSpringTo", Keywords = "position", AdvancedDisplay = "8"), Category = "Math|Interpolation")
@@ -1785,7 +1784,7 @@ class ENGINE_API UKismetMathLibrary : public UBlueprintFunctionLibrary
 	                                  float Stiffness, float CriticalDampingFactor, float DeltaTime,
 	                                  float Mass = 1.f, float TargetVelocityAmount = 1.f, 
 	                                  bool bClamp = false, FVector MinValue = FVector(-1.f), FVector MaxValue = FVector(1.f),
-	                                  bool bApproximate = true, bool bInitializeFromTarget = false);
+	                                  bool bInitializeFromTarget = false);
 
 	/**
 	* Uses a simple spring model to interpolate a quaternion from Current to Target.
@@ -1798,13 +1797,12 @@ class ENGINE_API UKismetMathLibrary : public UBlueprintFunctionLibrary
 	* @param DeltaTime             Time difference since the last update
 	* @param Mass                  Multiplier that acts like mass on a spring
 	* @param TargetVelocityAmount  If 1 then the target velocity will be calculated and used, which results following the target more closely/without lag. Values down to zero (recommended when using this to smooth data) will progressively disable this effect.
-	* @param bApproximate          Use an approximate/fast implementation. Will be suitable for all cases unless the delta time is very large or high accuracy is needed.
 	* @param bInitializeFromTarget If set then the current value will be set from the target on the first update
 	*/
 	UFUNCTION(BlueprintCallable,  meta = (ScriptMethod = "InterpSpringTo", Keywords = "quaternion", AdvancedDisplay = "8"), Category = "Math|Interpolation")
 	static FQuat QuaternionSpringInterp(FQuat Current, FQuat Target, UPARAM(ref) FQuaternionSpringState& SpringState,
 	                                    float Stiffness, float CriticalDampingFactor, float DeltaTime,
-	                                    float Mass = 1.f, float TargetVelocityAmount = 1.f, bool bApproximate = true, 
+	                                    float Mass = 1.f, float TargetVelocityAmount = 1.f, 
                                         bool bInitializeFromTarget = false);
 
 	/**
@@ -3902,7 +3900,6 @@ class ENGINE_API UKismetMathLibrary : public UBlueprintFunctionLibrary
 	 * @param bClamp                Whether to use the Min/Max values to clamp the motion
 	 * @param MinValue              Clamps the minimum output value and cancels the velocity if it reaches this limit
 	 * @param MaxValue              Clamps the maximum output value and cancels the velocity if it reaches this limit
-	 * @param bApproximate          Use an approximate/fast implementation. Will be suitable for all cases unless the delta time is very large or high accuracy is needed.
 	 * @param bInitializeFromTarget If set then the current value will be set from the target on the first update
 	 */
 	UFUNCTION(BlueprintCallable, Category = "Math|Interpolation", meta=(AdvancedDisplay = "8"))
@@ -3910,7 +3907,7 @@ class ENGINE_API UKismetMathLibrary : public UBlueprintFunctionLibrary
 	                               float Stiffness, float CriticalDampingFactor, float DeltaTime,
 	                               float Mass = 1.f, float TargetVelocityAmount = 1.f, 
 	                               bool bClamp = false, float MinValue = -1.f, float MaxValue = 1.f,
-	                               bool bApproximate = true, bool bInitializeFromTarget = false);
+	                               bool bInitializeFromTarget = false);
 
 	/** Resets the state of a float spring */
 	UFUNCTION(BlueprintCallable, Category = "Math|Interpolation")
diff --git a/Engine/Source/Runtime/Engine/Private/Animation/AnimInterpFilter.cpp b/Engine/Source/Runtime/Engine/Private/Animation/AnimInterpFilter.cpp
--- a/Engine/Source/Runtime/Engine/Private/Animation/AnimInterpFilter.cpp
+++ b/Engine/Source/Runtime/Engine/Private/Animation/AnimInterpFilter.cpp
@@ -219,8 +219,8 @@ float FFIRFilterTimeBased::UpdateAndGetFilteredData(float Input, float DeltaTime
 					FilterWindow.Push(FFilterData(0.0f, 0.0f));
 				}
 				const float OrigValue = FilterWindow[0].Input;
-				FMath::SpringDamperSmoothingApprox(FilterWindow[0].Input, FilterWindow[1].Input, Input, 0.0f, DeltaTime,
-				                                   WindowDuration / EULERS_NUMBER, DampingRatio);
+				FMath::SpringDamperSmoothing(FilterWindow[0].Input, FilterWindow[1].Input, Input, 0.0f, DeltaTime,
+				                             WindowDuration / EULERS_NUMBER, DampingRatio);
 				if (MaxSpeed > 0.0f)
 				{
 					// Clamp the speed
diff --git a/Engine/Source/Runtime/Engine/Private/KismetMathLibrary.cpp b/Engine/Source/Runtime/Engine/Private/KismetMathLibrary.cpp
--- a/Engine/Source/Runtime/Engine/Private/KismetMathLibrary.cpp
+++ b/Engine/Source/Runtime/Engine/Private/KismetMathLibrary.cpp
@@ -348,7 +348,7 @@ float UKismetMathLibrary::Ease(float A, float B, float Alpha, TEnumAsByte<EEasin
 template <typename T>
 void GenericSpringInterp(T& Current, const T Target, T& PrevTarget, bool& bPrevTargetValid, T& Velocity,
                          float Stiffness, float CriticalDamping, float DeltaTime,
-                         float TargetVelocityAmount, float Mass, bool bApproximate, bool bInitializeFromTarget)
+                         float TargetVelocityAmount, float Mass, bool bInitializeFromTarget)
 {
 	if (bInitializeFromTarget && !bPrevTargetValid)
 	{
@@ -363,14 +363,7 @@ void GenericSpringInterp(T& Current, const T Target, T& PrevTarget, bool& bPrevT
 			const float Omega = FMath::Sqrt(Stiffness / Mass); // angular frequency
 			const float Frequency = Omega / (2.0f * PI);
 			T NewValue = Current;
-			if (bApproximate)
-			{
-				FMath::SpringDamperApprox(Current, Velocity, Target, TargetVel, DeltaTime, Frequency, CriticalDamping);
-			}
-			else
-			{
-				FMath::SpringDamper(Current, Velocity, Target, TargetVel, DeltaTime, Frequency, CriticalDamping);
-			}
+			FMath::SpringDamper(Current, Velocity, Target, TargetVel, DeltaTime, Frequency, CriticalDamping);
 			PrevTarget = Target;
 			bPrevTargetValid = true;
 		}
@@ -381,10 +374,10 @@ float UKismetMathLibrary::FloatSpringInterp(float Current, float Target, FFloatS
                                             float Stiffness, float CriticalDamping, float DeltaTime,
                                             float Mass, float TargetVelocityAmount, 
                                             bool bClamp, float MinValue, float MaxValue,
-                                            bool bApproximate, bool bInitializeFromTarget)
+                                            bool bInitializeFromTarget)
 {
 	GenericSpringInterp(Current, Target, SpringState.PrevTarget, SpringState.bPrevTargetValid, SpringState.Velocity, Stiffness, CriticalDamping,
-	                    DeltaTime, TargetVelocityAmount, Mass, bApproximate, bInitializeFromTarget);
+	                    DeltaTime, TargetVelocityAmount, Mass, bInitializeFromTarget);
 	if (bClamp)
 	{
 		if (Current < MinValue)
@@ -421,10 +414,10 @@ FVector UKismetMathLibrary::VectorSpringInterp(FVector Current, FVector Target,
                                                float Stiffness, float CriticalDamping, float DeltaTime,
                                                float Mass, float TargetVelocityAmount, 
                                                bool bClamp, FVector MinValue, FVector MaxValue,
-                                               bool bApproximate, bool bInitializeFromTarget)
+                                               bool bInitializeFromTarget)
 {
 	GenericSpringInterp(Current, Target, SpringState.PrevTarget, SpringState.bPrevTargetValid, SpringState.Velocity, Stiffness,
-	                    CriticalDamping, DeltaTime, TargetVelocityAmount, Mass, bApproximate, bInitializeFromTarget);
+	                    CriticalDamping, DeltaTime, TargetVelocityAmount, Mass, bInitializeFromTarget);
 	if (bClamp)
 	{
 		for (int Index = 0 ; Index != 3 ; ++Index)
@@ -452,8 +445,7 @@ FVector UKismetMathLibrary::VectorSpringInterp(FVector Current, FVector Target,
 
 FQuat UKismetMathLibrary::QuaternionSpringInterp(FQuat Current, FQuat Target, UPARAM(ref) FQuaternionSpringState& SpringState,
                                                  float Stiffness, float CriticalDamping, float DeltaTime,
-                                                 float Mass, float TargetVelocityAmount, bool bApproximate, 
-                                                 bool bInitializeFromTarget)
+                                                 float Mass, float TargetVelocityAmount, bool bInitializeFromTarget)
 {
 	if ((Target | Current) < 0.0f)
 	{
@@ -463,7 +455,7 @@ FQuat UKismetMathLibrary::QuaternionSpringInterp(FQuat Current, FQuat Target, UP
 	FQuat Velocity = 0.5f * Current * VelocityQuat;
 
   	GenericSpringInterp(Current, Target, SpringState.PrevTarget, SpringState.bPrevTargetValid, Velocity, Stiffness,
-	                    CriticalDamping, DeltaTime, TargetVelocityAmount, Mass, bApproximate, bInitializeFromTarget);
+	                    CriticalDamping, DeltaTime, TargetVelocityAmount, Mass, bInitializeFromTarget);
 	Current.Normalize();
 
 	VelocityQuat = Current.Inverse() * 2.0f * Velocity;
