Replaced the single ensure crash reporting process handle with a pool of process handles

diff --git a/Engine/Source/Runtime/Core/Private/Unix/UnixPlatformCrashContext.cpp b/Engine/Source/Runtime/Core/Private/Unix/UnixPlatformCrashContext.cpp
--- a/Engine/Source/Runtime/Core/Private/Unix/UnixPlatformCrashContext.cpp
+++ b/Engine/Source/Runtime/Core/Private/Unix/UnixPlatformCrashContext.cpp
@@ -336,18 +336,42 @@ void FUnixCrashContext::GetPortableCallStack(const uint64* StackFrames, int32 Nu
 
 namespace UnixCrashReporterTracker
 {
-	FProcHandle CurrentlyRunningCrashReporter;
+	struct CrashReporterProcess
+	{
+		FProcHandle Process;
+		bool bIsSpawned = false;
+	};
+
+	// Matching MaxPreviousErrorsToTrack = 4 in AssertionMacros.cpp
+	CrashReporterProcess Processes[4];
 	FDelegateHandle CurrentTicker;
 
 	bool Tick(float DeltaTime)
 	{
 		QUICK_SCOPE_CYCLE_COUNTER(STAT_UnixCrashReporterTracker_Tick);
 
-		if (!FPlatformProcess::IsProcRunning(CurrentlyRunningCrashReporter))
+		int ProcessLeft = 0;
+		for (int ProcessNumber = 0; ProcessNumber < UE_ARRAY_COUNT(UnixCrashReporterTracker::Processes); ProcessNumber++)
 		{
-			FPlatformProcess::CloseProc(CurrentlyRunningCrashReporter);
-			CurrentlyRunningCrashReporter = FProcHandle();
+			if (UnixCrashReporterTracker::Processes[ProcessNumber].bIsSpawned)
+			{
+				if (!FPlatformProcess::IsProcRunning(UnixCrashReporterTracker::Processes[ProcessNumber].Process))
+				{
+					UnixCrashReporterTracker::Processes[ProcessNumber].bIsSpawned = false;
+
+					FPlatformProcess::CloseProc(UnixCrashReporterTracker::Processes[ProcessNumber].Process);
+					UnixCrashReporterTracker::Processes[ProcessNumber].Process = FProcHandle();
+				}
+				else
+				{
+					ProcessLeft++;
+				}
+			}
+		}
 
+		if (ProcessLeft <= 0)
+		{
+			// clean up ticker
 			FTicker::GetCoreTicker().RemoveTicker(CurrentTicker);
 			CurrentTicker.Reset();
 
@@ -395,7 +419,11 @@ namespace UnixCrashReporterTracker
 		{
 			FTicker::GetCoreTicker().RemoveTicker(CurrentTicker);
 			CurrentTicker.Reset();
-			CurrentlyRunningCrashReporter = FProcHandle();
+
+			for (int ProcessNumber = 0; ProcessNumber < UE_ARRAY_COUNT(UnixCrashReporterTracker::Processes); ProcessNumber++)
+			{
+				UnixCrashReporterTracker::Processes[ProcessNumber].Process = FProcHandle();
+			}
 		}
 	}
 }
@@ -639,33 +667,61 @@ void FUnixCrashContext::GenerateCrashInfoAndLaunchReporter(bool bReportingNonCra
 
 			CrashReportClientArguments += TEXT("\"\"") + CrashInfoAbsolute + TEXT("/\"\"");
 
+			// FIXME: this code is not attempting to be thread safe
 			if (bReportingNonCrash)
 			{
-				// When running a dedicated server and we are reporting a non-crash and we are already in the process of uploading an ensure
-				// we will skip the upload to avoid hitching.
-				if (UnixCrashReporterTracker::CurrentTicker.IsValid() && IsRunningDedicatedServer())
 				{
-					UE_LOG(LogCore, Warning, TEXT("An ensure is already in the process of being uploaded, skipping upload."));
-				}
-				else
-				{
-					// If we're reporting non-crash, try to avoid spinning here and instead do that in the tick.
-					// However, if there was already a crash reporter running (i.e. we hit ensure() too quickly), take a hitch here
-					if (UnixCrashReporterTracker::CurrentTicker.IsValid())
+					bool bFoundEmptySlot = false;
+
+					const double kEnsureTimeOut = 45.0;
+					const double kEnsureSleepInterval = 0.1;
+					double kTimeOutTimer = 0.0;
+
+					while (!bFoundEmptySlot && !IsRunningDedicatedServer())
 					{
-						// do not wait indefinitely, allow 45 second hitch (anticipating callstack parsing)
-						const double kEnsureTimeOut = 45.0;
-						const double kEnsureSleepInterval = 0.1;
-						if (!UnixCrashReporterTracker::WaitForProcWithTimeout(UnixCrashReporterTracker::CurrentlyRunningCrashReporter, kEnsureTimeOut, kEnsureSleepInterval))
+						// case 1 - no process handle in use: spawn ticker and a process
+						if (!UnixCrashReporterTracker::CurrentTicker.IsValid())
 						{
-							FPlatformProcess::TerminateProc(UnixCrashReporterTracker::CurrentlyRunningCrashReporter);
-						}
+							UnixCrashReporterTracker::CurrentTicker = FTicker::GetCoreTicker().AddTicker(FTickerDelegate::CreateStatic(&UnixCrashReporterTracker::Tick), 1.f);
 
-						UnixCrashReporterTracker::Tick(0.001f);	// tick one more time to make it clean up after itself
-					}
+							// assume if no ticker all slots are open, pick the first one
+							UnixCrashReporterTracker::Processes[0].bIsSpawned = true;
+							UnixCrashReporterTracker::Processes[0].Process = FPlatformProcess::CreateProc(
+									RelativePathToCrashReporter, *CrashReportClientArguments, true, false, false, NULL, 0, NULL, NULL);
 
-					UnixCrashReporterTracker::CurrentlyRunningCrashReporter = FPlatformProcess::CreateProc(RelativePathToCrashReporter, *CrashReportClientArguments, true, false, false, NULL, 0, NULL, NULL);
-					UnixCrashReporterTracker::CurrentTicker = FTicker::GetCoreTicker().AddTicker(FTickerDelegate::CreateStatic(&UnixCrashReporterTracker::Tick), 1.f);
+							bFoundEmptySlot = true;
+						}
+						// case 2 - some process handles in use: spawn a process if available 
+						else
+						{
+							for (int ProcessNumber = 0; ProcessNumber < UE_ARRAY_COUNT(UnixCrashReporterTracker::Processes); ProcessNumber++)
+							{
+								if (!UnixCrashReporterTracker::Processes[ProcessNumber].bIsSpawned)
+								{
+									bFoundEmptySlot = true;
+									UnixCrashReporterTracker::Processes[ProcessNumber].bIsSpawned = true;
+									UnixCrashReporterTracker::Processes[ProcessNumber].Process = FPlatformProcess::CreateProc(
+											RelativePathToCrashReporter, *CrashReportClientArguments, true, false, false, NULL, 0, NULL, NULL);
+									break;
+								}
+							}
+
+							// case 3 - all process handles in use: wait for up to 45 seconds on the next available process
+							if (!bFoundEmptySlot)
+							{
+								UnixCrashReporterTracker::Tick(0.001f);
+								FPlatformProcess::Sleep(kEnsureSleepInterval);	
+								kTimeOutTimer += kEnsureSleepInterval;
+
+								// Could lose information by killing an unfinished CrashReportClient instance
+								if (kTimeOutTimer >= kEnsureTimeOut)
+								{
+									UnixCrashReporterTracker::Processes[0].bIsSpawned = false;
+									FPlatformProcess::TerminateProc(UnixCrashReporterTracker::Processes[0].Process);
+								}
+							}
+						}
+					}
 				}
 			}
 			else
