Hack to allow marketplace build of Bridge plugin to load in EA2.

diff --git a/Engine/Source/Runtime/Core/Private/Modules/ModuleManager.cpp b/Engine/Source/Runtime/Core/Private/Modules/ModuleManager.cpp
--- a/Engine/Source/Runtime/Core/Private/Modules/ModuleManager.cpp
+++ b/Engine/Source/Runtime/Core/Private/Modules/ModuleManager.cpp
@@ -1164,7 +1164,7 @@ void FModuleManager::FindModulePathsInDirectory(const FString& InDirectoryName,
 	for(const FString& SearchDirectoryName: SearchDirectoryNames)
 	{
 		FModuleManifest Manifest;
-		if (FModuleManifest::TryRead(FModuleManifest::GetFileName(SearchDirectoryName, bIsGameDirectory), Manifest) && Manifest.BuildId == BuildId.GetValue())
+		if (FModuleManifest::TryRead(FModuleManifest::GetFileName(SearchDirectoryName, bIsGameDirectory), Manifest) && (Manifest.BuildId == BuildId.GetValue() || SearchDirectoryName.Contains(TEXT("/Engine/Plugins/Bridge/"))))
 		{
 			for (const TPair<FString, FString>& Pair : Manifest.ModuleNameToFileName)
 			{
