FileCache: Restore the UE4 default value fc.NumFileCacheBlocks=256.

#rnx

diff --git a/Engine/Source/Runtime/Core/Private/FileCache/FileCache.cpp b/Engine/Source/Runtime/Core/Private/FileCache/FileCache.cpp
--- a/Engine/Source/Runtime/Core/Private/FileCache/FileCache.cpp
+++ b/Engine/Source/Runtime/Core/Private/FileCache/FileCache.cpp
@@ -26,7 +26,7 @@ DEFINE_LOG_CATEGORY_STATIC(LogStreamingFileCache, Log, All);
 
 static const int CacheLineSize = 64 * 1024;
 
-static int32 GNumFileCacheBlocks = 2048;
+static int32 GNumFileCacheBlocks = 256;
 static FAutoConsoleVariableRef CVarNumFileCacheBlocks(
 	TEXT("fc.NumFileCacheBlocks"),
 	GNumFileCacheBlocks,
