Ensure that our compilers support if constexpr and fold expressions.

diff --git a/Engine/Source/Runtime/Core/Public/Clang/ClangPlatform.h b/Engine/Source/Runtime/Core/Public/Clang/ClangPlatform.h
--- a/Engine/Source/Runtime/Core/Public/Clang/ClangPlatform.h
+++ b/Engine/Source/Runtime/Core/Public/Clang/ClangPlatform.h
@@ -9,13 +9,13 @@
 #if defined(__cpp_if_constexpr)
 	#define PLATFORM_COMPILER_HAS_IF_CONSTEXPR 1
 #else
-	#define PLATFORM_COMPILER_HAS_IF_CONSTEXPR 0
+	#error "Compiler is expected to support if constexpr"
 #endif
 
 #if defined(__cpp_fold_expressions)
 	#define PLATFORM_COMPILER_HAS_FOLD_EXPRESSIONS 1
 #else
-	#define PLATFORM_COMPILER_HAS_FOLD_EXPRESSIONS 0
+	#error "Compiler is expected to support fold expressions"
 #endif
 
 #define PLATFORM_RETURN_ADDRESS()			__builtin_return_address(0)
diff --git a/Engine/Source/Runtime/Core/Public/MSVC/MSVCPlatform.h b/Engine/Source/Runtime/Core/Public/MSVC/MSVCPlatform.h
--- a/Engine/Source/Runtime/Core/Public/MSVC/MSVCPlatform.h
+++ b/Engine/Source/Runtime/Core/Public/MSVC/MSVCPlatform.h
@@ -9,13 +9,13 @@
 #if _MSC_VER >= 1920
 	#define PLATFORM_COMPILER_HAS_IF_CONSTEXPR 1
 #else
-	#define PLATFORM_COMPILER_HAS_IF_CONSTEXPR 0
+	#error "Compiler is expected to support if constexpr"
 #endif
 
 #if defined(__cpp_fold_expressions)
 	#define PLATFORM_COMPILER_HAS_FOLD_EXPRESSIONS 1
 #else
-	#define PLATFORM_COMPILER_HAS_FOLD_EXPRESSIONS 0
+	#error "Compiler is expected to support fold expressions"
 #endif
 
 #define PLATFORM_RETURN_ADDRESS()	        _ReturnAddress()
