Default Value Helper: Support parsing FRotator serialized text (Pitch=...,Yaw=...,Roll=...)

diff --git a/Engine/Source/Runtime/Core/Private/Misc/DefaultValueHelper.cpp b/Engine/Source/Runtime/Core/Private/Misc/DefaultValueHelper.cpp
--- a/Engine/Source/Runtime/Core/Private/Misc/DefaultValueHelper.cpp
+++ b/Engine/Source/Runtime/Core/Private/Misc/DefaultValueHelper.cpp
@@ -540,8 +540,23 @@ bool FDefaultValueHelper::ParseRotator(const FString& Source, FRotator& OutVal)
 		return true;
 	}
 
-	// Fallback to x= format
-	return OutVal.InitFromString(Source);
+	// Fallback to p= format
+	if(OutVal.InitFromString(Source))
+	{
+		return true;
+	}
+	
+	// try Pitch= format
+	FRotator Rotator = FRotator::ZeroRotator;
+	if(FParse::Value( *Source, TEXT("Pitch="), Rotator.Pitch ) &&
+		FParse::Value( *Source, TEXT("Yaw="), Rotator.Yaw ) &&
+		FParse::Value( *Source, TEXT("Roll="), Rotator.Roll ))
+	{
+		OutVal = Rotator;
+		return true;
+	}
+
+	return false;
 }
 
 
