Fix for CIS compile error (-DisableUnity).

diff --git a/Engine/Source/Runtime/Core/Private/ProfilingDebugging/CpuProfilerTrace.cpp b/Engine/Source/Runtime/Core/Private/ProfilingDebugging/CpuProfilerTrace.cpp
--- a/Engine/Source/Runtime/Core/Private/ProfilingDebugging/CpuProfilerTrace.cpp
+++ b/Engine/Source/Runtime/Core/Private/ProfilingDebugging/CpuProfilerTrace.cpp
@@ -9,6 +9,7 @@
 #include "Containers/Map.h"
 #include "Misc/MemStack.h"
 #include "Misc/Crc.h"
+#include "UObject/NameTypes.h"
 
 #if CPUPROFILERTRACE_ENABLED
 
@@ -175,7 +176,7 @@ void FCpuProfilerTrace::OutputBeginDynamicEvent(const TCHAR* Name)
 	CPUPROFILERTRACE_OUTPUTBEGINEVENT_EPILOGUE();
 }
 
-void FCpuProfilerTrace::OutputBeginDynamicEvent(const FName Name)
+void FCpuProfilerTrace::OutputBeginDynamicEvent(const FName& Name)
 {
 	CPUPROFILERTRACE_OUTPUTBEGINEVENT_PROLOGUE();
 	uint32 SpecId = ThreadBuffer->DynamicFNameScopeNamesMap.FindRef(Name.GetComparisonIndex());
diff --git a/Engine/Source/Runtime/Core/Public/ProfilingDebugging/CpuProfilerTrace.h b/Engine/Source/Runtime/Core/Public/ProfilingDebugging/CpuProfilerTrace.h
--- a/Engine/Source/Runtime/Core/Public/ProfilingDebugging/CpuProfilerTrace.h
+++ b/Engine/Source/Runtime/Core/Public/ProfilingDebugging/CpuProfilerTrace.h
@@ -19,6 +19,8 @@
 
 UE_TRACE_CHANNEL_EXTERN(CpuChannel, CORE_API);
 
+class FName;
+
 /*
  * Facilities for tracing timed cpu events. Two types of events are supported, static events where the identifier is
  * known at compile time, and dynamic event were identifiers can be constructed in runtime. Static events have lower overhead
@@ -66,7 +68,7 @@ struct FCpuProfilerTrace
 	 * that receives ANSICHAR* / TCHAR* name. Must always be matched with an end event.
 	 * @param Name Name of event
 	 */
-	CORE_API static void OutputBeginDynamicEvent(const FName Name);
+	CORE_API static void OutputBeginDynamicEvent(const FName& Name);
 	/*
 	 * Output end event marker for static or dynamic event for the currently open scope.
 	 */
