U5 - change throttle stealing behaviour so that we only start throllting when more than half of the workers are idling. Also have an idling worker that just found new work to wake yet another worker. THis way we should spin up to full capacity even if the producer missed to wake the other half of the sleeping workers

diff --git a/Engine/Source/Runtime/Core/Private/Async/Fundamental/Scheduler.cpp b/Engine/Source/Runtime/Core/Private/Async/Fundamental/Scheduler.cpp
--- a/Engine/Source/Runtime/Core/Private/Async/Fundamental/Scheduler.cpp
+++ b/Engine/Source/Runtime/Core/Private/Async/Fundamental/Scheduler.cpp
@@ -220,7 +220,11 @@ namespace LowLevelTasks
 					QueueRegistry.Enqueue(Task, uint32(Task->GetPriority()));
 					continue;
 				}
-				OutOfWork.Stop();		
+				if (OutOfWork.Stop())
+				{
+					//if we are ramping up work again we start waking up workers that might not have been woken while a lot of work might have been queued
+					WakeUpWorker(bPermitBackgroundWork);
+				}
 				FTask* OldTask = FSchedulerTls::ActiveTask;
 				FSchedulerTls::ActiveTask = Task;
 				{
diff --git a/Engine/Source/Runtime/Core/Public/Async/Fundamental/LocalQueue.h b/Engine/Source/Runtime/Core/Public/Async/Fundamental/LocalQueue.h
--- a/Engine/Source/Runtime/Core/Public/Async/Fundamental/LocalQueue.h
+++ b/Engine/Source/Runtime/Core/Public/Async/Fundamental/LocalQueue.h
@@ -298,7 +298,7 @@ public:
 			{
 				Registry->OverflowQueues[PriorityIndex].enqueue(Item);
 			}
-			return Registry->AnyWorkerLookingForWork(bBackgroundTask);
+			return Registry->LessThanHalfWorkersLookingForWork(bBackgroundTask);
 		}
 
 		inline FTask* DequeueLocal(bool GetBackGroundTasks)
@@ -461,7 +461,7 @@ public:
 		bool bBackgroundTask = Item->IsBackgroundTask();
 		OverflowQueues[PriorityIndex].enqueue(Item);
 
-		return AnyWorkerLookingForWork(bBackgroundTask);
+		return LessThanHalfWorkersLookingForWork(bBackgroundTask);
 	}
 
 	// grab an Item directy from the Global OverflowQueue
@@ -484,9 +484,10 @@ public:
 	}
 
 private:
-	inline bool AnyWorkerLookingForWork(bool bBackgroundTask) const
+	inline bool LessThanHalfWorkersLookingForWork(bool bBackgroundTask) const
 	{
-		return (NumWorkersLookingForWork[bBackgroundTask].load(std::memory_order_acquire) == 0) && (bBackgroundTask || (NumWorkersLookingForWork[true].load(std::memory_order_acquire) == 0));
+		return (NumWorkersLookingForWork[bBackgroundTask].load(std::memory_order_acquire) * 2 < NumActiveWorkers[bBackgroundTask].load(std::memory_order_acquire)) 
+			&& (bBackgroundTask || (NumWorkersLookingForWork[true].load(std::memory_order_acquire) * 2 < NumActiveWorkers[true].load(std::memory_order_acquire)));
 	}
 
 	FOverflowQueueType	  OverflowQueues[uint32(ETaskPriority::Count)];
