Use a map to store the transient per line info in the FileCache.
Saves memory compared to allocating arrays for all lines.

diff --git a/Engine/Source/Runtime/Core/Private/FileCache/FileCache.cpp b/Engine/Source/Runtime/Core/Private/FileCache/FileCache.cpp
--- a/Engine/Source/Runtime/Core/Private/FileCache/FileCache.cpp
+++ b/Engine/Source/Runtime/Core/Private/FileCache/FileCache.cpp
@@ -254,21 +254,23 @@ public:
 	void Evict(CacheLineID Line);
 
 private:
-	struct FPendingRequest
-	{
-		FGraphEventRef Event;
-	};
-
 	void CheckForSizeRequestComplete();
 
 	CacheSlotID AcquireSlotAndReadLine(FFileCache& Cache, CacheLineID LineID, EAsyncIOPriorityAndFlags Priority);
 	void ReadLine(FFileCache& Cache, CacheSlotID SlotID, CacheLineID LineID, EAsyncIOPriorityAndFlags Priority, const FGraphEventRef& CompletionEvent);
 
-	TArray<CacheSlotID> LineToSlot;
-	TArray<FPendingRequest> LineToRequest;
+	struct FPendingRequest
+	{
+		FGraphEventRef Event;
+	};
+	struct FLineInfo
+	{
+		CacheSlotID SlotID;
+		FPendingRequest PendingRequest;
+	};
+	TMap<int32, FLineInfo> LineInfos;
 
 	int64 BaseOffset;
-	int64 NumSlots;
 	int64 FileSize;
 	IAsyncReadFileHandle* InnerHandle;
 	FGraphEventRef SizeRequestEvent;
@@ -426,7 +428,6 @@ FFileCacheHandle::~FFileCacheHandle()
 
 FFileCacheHandle::FFileCacheHandle(IAsyncReadFileHandle* InHandle, int64 InBaseOffset)
 	: BaseOffset(InBaseOffset)
-	, NumSlots(0)
 	, FileSize(-1)
 	, InnerHandle(InHandle)
 {
@@ -565,14 +566,9 @@ void FFileCacheHandle::CheckForSizeRequestComplete()
 
 		check(FileSize > 0);
 
-		// Make sure we haven't lazily allocated more slots than are in the file, then allocate the final number of slots
+		// LineInfos key is int32
 		const int64 TotalNumSlots = FMath::DivideAndRoundUp(FileSize, (int64)CacheLineSize);
-		check(NumSlots <= TotalNumSlots);
-		NumSlots = TotalNumSlots;
-		// TArray is max signed int
 		check(TotalNumSlots < MAX_int32);
-		LineToSlot.SetNum((int32)TotalNumSlots, false);
-		LineToRequest.SetNum((int32)TotalNumSlots, false);
 	}
 }
 
@@ -601,15 +597,15 @@ CacheSlotID FFileCacheHandle::AcquireSlotAndReadLine(FFileCache& Cache, CacheLin
 	// no valid slot for this line, grab a new slot from cache and start a read request
 	CacheSlotID SlotID = Cache.AcquireAndLockSlot(this, LineID);
 
-	FPendingRequest& PendingRequest = LineToRequest[LineID.Get()];
-	if (PendingRequest.Event)
+	FLineInfo& LineInfo = LineInfos.FindOrAdd(LineID.Get());
+	if (LineInfo.PendingRequest.Event)
 	{
 		// previous async request/event (if any) should be completed, if this is back in the free list
-		check(PendingRequest.Event->IsComplete());
+		check(LineInfo.PendingRequest.Event->IsComplete());
 	}
 
 	FGraphEventRef CompletionEvent = FGraphEvent::CreateGraphEvent();
-	PendingRequest.Event = CompletionEvent;
+	LineInfo.PendingRequest.Event = CompletionEvent;
 	if (FileSize >= 0)
 	{
 		// If FileSize >= 0, that means the async file size request has completed, we can perform the read immediately
@@ -637,6 +633,7 @@ IMemoryReadStreamRef FFileCacheHandle::ReadData(FGraphEventArray& OutCompletionE
 	const int64 Offset = BaseOffset + InOffset;
 	const CacheLineID StartLine = GetBlock<CacheLineID>(Offset);
 	const CacheLineID EndLine = GetBlock<CacheLineID>(Offset + BytesToRead - 1);
+	check(EndLine.Get() + 1 < MAX_int32);
 
 	const int32 NumSlotsNeeded = EndLine.Get() + 1 - StartLine.Get();
 
@@ -655,18 +652,6 @@ IMemoryReadStreamRef FFileCacheHandle::ReadData(FGraphEventArray& OutCompletionE
 		return ReadDataUncached(OutCompletionEvents, Offset, BytesToRead, Priority);
 	}
 
-	if (EndLine.Get() >= NumSlots)
-	{
-		// If we're still waiting on SizeRequest, may need to lazily allocate some slots to service this request
-		// If this happens after SizeRequest has completed, that means something must have gone wrong
-		check(SizeRequestEvent);
-		NumSlots = EndLine.Get() + 1;
-		// TArray is max signed int
-		check(NumSlots < MAX_int32);
-		LineToSlot.SetNum((int32)NumSlots, false);
-		LineToRequest.SetNum((int32)NumSlots, false);
-	}
-
 	const int32 NumCacheSlots = EndLine.Get() + 1 - StartLine.Get();
 	check(NumCacheSlots > 0);
 	const uint32 AllocSize = sizeof(FMemoryReadStreamCache) + sizeof(CacheSlotID) * (NumCacheSlots - 1);
@@ -679,31 +664,30 @@ IMemoryReadStreamRef FFileCacheHandle::ReadData(FGraphEventArray& OutCompletionE
 	bool bHasPendingSlot = false;
 	for (CacheLineID LineID = StartLine; LineID.Get() <= EndLine.Get(); ++LineID)
 	{
-		CacheSlotID& SlotID = LineToSlot[LineID.Get()];
-		if (!SlotID.IsValid())
+		FLineInfo& LineInfo = LineInfos.FindOrAdd(LineID.Get());
+		if (!LineInfo.SlotID.IsValid())
 		{
 			// no valid slot for this line, grab a new slot from cache and start a read request
-			SlotID = AcquireSlotAndReadLine(Cache, LineID, Priority);
+			LineInfo.SlotID = AcquireSlotAndReadLine(Cache, LineID, Priority);
 		}
 		else
 		{
-			Cache.LockSlot(SlotID);
+			Cache.LockSlot(LineInfo.SlotID);
 		}
 
-		check(SlotID.IsValid());
-		Result->CacheSlots[LineID.Get() - StartLine.Get()] = SlotID;
+		check(LineInfo.SlotID.IsValid());
+		Result->CacheSlots[LineID.Get() - StartLine.Get()] = LineInfo.SlotID;
 
-		FPendingRequest& PendingRequest = LineToRequest[LineID.Get()];
-		if (PendingRequest.Event && !PendingRequest.Event->IsComplete())
+		if (LineInfo.PendingRequest.Event && !LineInfo.PendingRequest.Event->IsComplete())
 		{
 			// this line has a pending async request to read data
 			// will need to wait for this request to complete before data is valid
-			OutCompletionEvents.Add(PendingRequest.Event);
+			OutCompletionEvents.Add(LineInfo.PendingRequest.Event);
 			bHasPendingSlot = true;
 		}
 		else
 		{
-			PendingRequest.Event.SafeRelease();
+			LineInfo.PendingRequest.Event.SafeRelease();
 		}
 	}
 
@@ -744,22 +728,6 @@ FGraphEventRef FFileCacheHandle::PreloadData(const FFileCachePreloadEntry* Prelo
 
 	CheckForSizeRequestComplete();
 
-	{
-		const FFileCachePreloadEntry& LastEntry = PreloadEntries[NumEntries - 1];
-		const CacheLineID EndLine = GetBlock<CacheLineID>(BaseOffset + LastEntry.Offset + LastEntry.Size - 1);
-		if (EndLine.Get() >= NumSlots)
-		{
-			// If we're still waiting on SizeRequest, may need to lazily allocate some slots to service this request
-			// If this happens after SizeRequest has completed, that means something must have gone wrong
-			check(SizeRequestEvent);
-			NumSlots = EndLine.Get() + 1;
-			// TArray is max signed int
-			check(NumSlots < MAX_int32);
-			LineToSlot.SetNum((int32)NumSlots, false);
-			LineToRequest.SetNum((int32)NumSlots, false);
-		}
-	}
-
 	FGraphEventArray CompletionEvents;
 	TArray<CacheSlotID> LockedSlots;
 	LockedSlots.Empty(NumEntries);
@@ -773,6 +741,7 @@ FGraphEventRef FFileCacheHandle::PreloadData(const FFileCachePreloadEntry* Prelo
 		const CacheLineID StartLine = GetBlock<CacheLineID>(EntryOffset);
 		const CacheLineID EndLine = GetBlock<CacheLineID>(EntryOffset + Entry.Size - 1);
 
+		check(EndLine.Get() + 1 < MAX_int32);
 		checkf(Entry.Offset > PrevOffset, TEXT("Preload entries must be sorted by Offset [%lld, %lld), %lld"),
 			Entry.Offset, Entry.Offset + Entry.Size, PrevOffset);
 		PrevOffset = Entry.Offset;
@@ -780,24 +749,24 @@ FGraphEventRef FFileCacheHandle::PreloadData(const FFileCachePreloadEntry* Prelo
 		CurrentLine = CacheLineID(FMath::Max(CurrentLine.Get(), StartLine.Get()));
 		while (CurrentLine.Get() <= EndLine.Get() && Cache.NumFreeSlots > 0)
 		{
-			CacheSlotID& SlotID = LineToSlot[CurrentLine.Get()];
-			if (!SlotID.IsValid())
+			FLineInfo& LineInfo = LineInfos.FindOrAdd(CurrentLine.Get());
+			
+			if (!LineInfo.SlotID.IsValid())
 			{
 				// no valid slot for this line, grab a new slot from cache and start a read request
-				SlotID = AcquireSlotAndReadLine(Cache, CurrentLine, Priority);
-				LockedSlots.Add(SlotID);
+				LineInfo.SlotID = AcquireSlotAndReadLine(Cache, CurrentLine, Priority);
+				LockedSlots.Add(LineInfo.SlotID);
 			}
 
-			FPendingRequest& PendingRequest = LineToRequest[CurrentLine.Get()];
-			if (PendingRequest.Event && !PendingRequest.Event->IsComplete())
+			if (LineInfo.PendingRequest.Event && !LineInfo.PendingRequest.Event->IsComplete())
 			{
 				// this line has a pending async request to read data
 				// will need to wait for this request to complete before data is valid
-				CompletionEvents.Add(PendingRequest.Event);
+				CompletionEvents.Add(LineInfo.PendingRequest.Event);
 			}
 			else
 			{
-				PendingRequest.Event.SafeRelease();
+				LineInfo.PendingRequest.Event.SafeRelease();
 			}
 
 			++CurrentLine;
@@ -823,25 +792,28 @@ FGraphEventRef FFileCacheHandle::PreloadData(const FFileCachePreloadEntry* Prelo
 
 void FFileCacheHandle::Evict(CacheLineID LineID)
 {
-	LineToSlot[LineID.Get()] = CacheSlotID();
+	FLineInfo* LineInfo = LineInfos.Find(LineID.Get());
+	if (LineInfo != nullptr)
+	{ 
+		if (LineInfo->PendingRequest.Event)
+		{
+			check(LineInfo->PendingRequest.Event->IsComplete());
+			LineInfo->PendingRequest.Event.SafeRelease();
+		}
 
-	FPendingRequest& PendingRequest = LineToRequest[LineID.Get()];
-	if (PendingRequest.Event)
-	{
-		check(PendingRequest.Event->IsComplete());
-		PendingRequest.Event.SafeRelease();
+		LineInfos.Remove(LineID.Get());
 	}
 }
 
 void FFileCacheHandle::WaitAll()
 {
-	for (int i = 0; i < LineToRequest.Num(); ++i)
+	for (TPair<int32, FLineInfo>& Pair : LineInfos)
 	{
-		FPendingRequest& PendingRequest = LineToRequest[i];
-		if (PendingRequest.Event)
+		FLineInfo& LineInfo = Pair.Value;
+		if (LineInfo.PendingRequest.Event)
 		{
-			check(PendingRequest.Event->IsComplete());
-			PendingRequest.Event.SafeRelease();
+			check(LineInfo.PendingRequest.Event->IsComplete());
+			LineInfo.PendingRequest.Event.SafeRelease();
 		}
 	}
 }
