conditional compilation for "StatsThread as a pipe" with an ability to turn back to "StatsThread as a named thread"

diff --git a/Engine/Plugins/Messaging/MessagingDebugger/Source/MessagingDebugger/Private/Widgets/MessageDetails/SMessagingDispatchStateTableRow.h b/Engine/Plugins/Messaging/MessagingDebugger/Source/MessagingDebugger/Private/Widgets/MessageDetails/SMessagingDispatchStateTableRow.h
--- a/Engine/Plugins/Messaging/MessagingDebugger/Source/MessagingDebugger/Private/Widgets/MessageDetails/SMessagingDispatchStateTableRow.h
+++ b/Engine/Plugins/Messaging/MessagingDebugger/Source/MessagingDebugger/Private/Widgets/MessageDetails/SMessagingDispatchStateTableRow.h
@@ -189,6 +189,12 @@ protected:
 			return LOCTEXT("ActualRenderingThread", "ActualRenderingThread");
 			break;
 
+#if STATS && !UE_STATS_THREAD_AS_PIPE
+		case ENamedThreads::StatsThread:
+			return LOCTEXT("StatsThread", "StatsThread");
+			break;
+#endif
+
 		default:
 			return LOCTEXT("UnknownThread", "Unknown");
 		}
diff --git a/Engine/Plugins/Messaging/MessagingDebugger/Source/MessagingDebugger/Private/Widgets/MessageDetails/SMessagingMessageDetails.cpp b/Engine/Plugins/Messaging/MessagingDebugger/Source/MessagingDebugger/Private/Widgets/MessageDetails/SMessagingMessageDetails.cpp
--- a/Engine/Plugins/Messaging/MessagingDebugger/Source/MessagingDebugger/Private/Widgets/MessageDetails/SMessagingMessageDetails.cpp
+++ b/Engine/Plugins/Messaging/MessagingDebugger/Source/MessagingDebugger/Private/Widgets/MessageDetails/SMessagingMessageDetails.cpp
@@ -233,6 +233,16 @@ FText SMessagingMessageDetails::HandleSenderThreadText() const
 			return LOCTEXT("ActualRenderingThread_Local", "ActualRenderingThread_Local");
 			break;
 
+#if STATS && !UE_STATS_THREAD_AS_PIPE
+		case ENamedThreads::StatsThread:
+			return LOCTEXT("StatsThread", "StatsThread");
+			break;
+
+		case ENamedThreads::StatsThread_Local:
+			return LOCTEXT("StatsThread_Local", "StatsThread_Local");
+			break;
+#endif
+
 		default:
 			return LOCTEXT("UnknownThread", "Unknown");
 		}
diff --git a/Engine/Source/Developer/ProfilerService/Private/ProfilerServiceManager.cpp b/Engine/Source/Developer/ProfilerService/Private/ProfilerServiceManager.cpp
--- a/Engine/Source/Developer/ProfilerService/Private/ProfilerServiceManager.cpp
+++ b/Engine/Source/Developer/ProfilerService/Private/ProfilerServiceManager.cpp
@@ -164,11 +164,13 @@ void FProfilerServiceManager::SetPreviewState( const FMessageAddress& ClientAddr
 				// Enable stat capture.
 				if (PreviewClients.Num() == 0)
 				{
+#if UE_STATS_THREAD_AS_PIPE
 					if (FPlatformProcess::SupportsMultithreading())
 					{
 						GStatsPipe.Launch(UE_SOURCE_LOCATION, [this] { AddNewFrameHandleStatsPipe(); });
 					}
 					else
+#endif
 					{
 						FSimpleDelegateGraphTask::CreateAndDispatchWhenReady
 						(
@@ -190,11 +192,13 @@ void FProfilerServiceManager::SetPreviewState( const FMessageAddress& ClientAddr
 				// Disable stat capture.
 				if (PreviewClients.Num() == 0)
 				{
+#if UE_STATS_THREAD_AS_PIPE
 					if (FPlatformProcess::SupportsMultithreading())
 					{
 						GStatsPipe.Launch(UE_SOURCE_LOCATION, [this] { RemoveNewFrameHandleStatsPipe(); });
 					}
 					else
+#endif
 					{
 						FSimpleDelegateGraphTask::CreateAndDispatchWhenReady
 						(
diff --git a/Engine/Source/Runtime/Core/Private/Async/TaskGraph.cpp b/Engine/Source/Runtime/Core/Private/Async/TaskGraph.cpp
--- a/Engine/Source/Runtime/Core/Private/Async/TaskGraph.cpp
+++ b/Engine/Source/Runtime/Core/Private/Async/TaskGraph.cpp
@@ -1323,13 +1323,15 @@ public:
 
 			int32 StackSize;
 
-#if WITH_EDITOR
+//#if WITH_EDITOR
+//			StackSize = 1024 * 1024;
+//#elif (!UE_BUILD_SHIPPING && !UE_BUILD_TEST)
+//			StackSize = 512 * 1024;
+//#else
+//			StackSize = 384 * 1024;
+//#endif
+
 			StackSize = 1024 * 1024;
-#elif (!UE_BUILD_SHIPPING && !UE_BUILD_TEST)
-			StackSize = 512 * 1024;
-#else
-			StackSize = 384 * 1024;
-#endif
 
 			// GConfig is not initialized yet, the only solution for now is to hardcode desired values
 			//GConfig->GetInt(TEXT("Core.System"), TEXT("TaskThreadStackSize"), StackSize, GEngineIni);
@@ -1445,7 +1447,9 @@ PRAGMA_ENABLE_DEPRECATION_WARNINGS
 	{
 		TASKGRAPH_SCOPE_CYCLE_COUNTER(2, STAT_TaskGraph_QueueTask);
 
+#if UE_STATS_THREAD_AS_PIPE
 		RedirectStatsTasksToPipe(Task);
+#endif
 #if UE_AUDIO_THREAD_AS_PIPE
 		RedirectAudioTasksToPipe(Task);
 #endif
@@ -1948,7 +1952,12 @@ public:
 private:
 	void QueueTask(class FBaseGraphTask* Task, bool bWakeUpWorker, ENamedThreads::Type InThreadToExecuteOn, ENamedThreads::Type InCurrentThreadIfKnown) override
 	{
+#if UE_STATS_THREAD_AS_PIPE
 		FTaskGraphImplementation::RedirectStatsTasksToPipe(Task);
+#endif
+#if UE_AUDIO_THREAD_AS_PIPE
+		FTaskGraphImplementation::RedirectAudioTasksToPipe(Task);
+#endif
 
 		if (ENamedThreads::GetThreadIndex(InThreadToExecuteOn) == ENamedThreads::AnyThread)
 		{
diff --git a/Engine/Source/Runtime/Core/Private/Stats/Stats2.cpp b/Engine/Source/Runtime/Core/Private/Stats/Stats2.cpp
--- a/Engine/Source/Runtime/Core/Private/Stats/Stats2.cpp
+++ b/Engine/Source/Runtime/Core/Private/Stats/Stats2.cpp
@@ -870,14 +870,15 @@ bool FStatNameAndInfo::GetSortByNameFrom(FName InLongName)
 }
 
 /*-----------------------------------------------------------------------------
-	FStatsProcessor
+	FStatsThread
 -----------------------------------------------------------------------------*/
 
 static TAutoConsoleVariable<int32> CVarDumpStatPackets(	TEXT("DumpStatPackets"),0,	TEXT("If true, dump stat packets."));
 
+#if UE_STATS_THREAD_AS_PIPE
 
 /** The rendering thread runnable object. */
-class FStatsProcessor
+class FStatsThread
 {
 	/** Array of stat packets, queued data to be processed on this thread. */
 	FStatPacketArray IncomingData;
@@ -889,14 +890,14 @@ class FStatsProcessor
 	bool bReadyToProcess;
 
 public:
-	static FStatsProcessor& Get()
+	static FStatsThread& Get()
 	{
-		static FStatsProcessor Singleton;
+		static FStatsThread Singleton;
 		return Singleton;
 	}
 
 	/** Default constructor. */
-	FStatsProcessor()
+	FStatsThread()
 		: State(FStatsThreadState::GetLocalState())
 		, bReadyToProcess(false)
 	{
@@ -994,6 +995,208 @@ private:
 	}
 };
 
+#else // UE_STATS_THREAD_AS_PIPE
+
+/** The rendering thread runnable object. */
+class FStatsThread : public FRunnable, FSingleThreadRunnable
+{
+	/** Array of stat packets, queued data to be processed on this thread. */
+	FStatPacketArray IncomingData;
+
+	/** Stats state. */
+	FStatsThreadState& State;
+
+	/** Whether we are ready to process the packets, sets by game or render packets. */
+	bool bReadyToProcess;
+public:
+
+	/** Default constructor. */
+	FStatsThread()
+		: State(FStatsThreadState::GetLocalState())
+		, bReadyToProcess(false)
+	{
+		check(IsInGameThread());
+	}
+
+	/**
+	 * Returns a pointer to the single threaded interface when multithreading is disabled.
+	 */
+	virtual FSingleThreadRunnable* GetSingleThreadInterface() override
+	{
+		return this;
+	}
+
+	/** Attaches to the task graph stats thread, all processing will be handled by the task graph. */
+	virtual uint32 Run() override
+	{
+		FThreadStats::GetThreadStats()->bIsStatsThread = true;
+		FMemory::SetupTLSCachesOnCurrentThread();
+		FTaskGraphInterface::Get().AttachToThread(ENamedThreads::StatsThread);
+		FTaskGraphInterface::Get().ProcessThreadUntilRequestReturn(ENamedThreads::StatsThread);
+		FMemory::ClearAndDisableTLSCachesOnCurrentThread();
+		return 0;
+	}
+
+	/** Tick function. */
+	virtual void Tick() override
+	{
+		LLM_SCOPE(ELLMTag::Stats);
+
+		static double LastTime = -1.0;
+		bool bShouldProcess = false;
+
+		const int32 MaxIncomingPackets = 16;
+		if (FThreadStats::bIsRawStatsActive)
+		{
+			// For raw stats we process every 24MB of packet data to minimize the stats messages memory usage.
+			//const bool bShouldProcessRawStats = IncomingData.Packets.Num() > 10;
+			const int32 MaxIncomingMessages = 24 * 1024 * 1024 / sizeof(FStatMessage);
+
+			int32 IncomingDataMessages = 0;
+			for (FStatPacket* Packet : IncomingData.Packets)
+			{
+				IncomingDataMessages += Packet->StatMessages.Num();
+			}
+
+			bShouldProcess = IncomingDataMessages > MaxIncomingMessages || IncomingData.Packets.Num() > MaxIncomingPackets;
+		}
+		else
+		{
+			// For regular stats we won't process more than every 5ms or every 16 packets.
+			// Commandlet stats are flushed as soon as.
+			bShouldProcess = bReadyToProcess && (FPlatformTime::Seconds() - LastTime > 0.005f || IncomingData.Packets.Num() > MaxIncomingPackets || FStats::EnabledForCommandlet());
+		}
+
+		if (bShouldProcess)
+		{
+			SCOPE_CYCLE_COUNTER(STAT_StatsNewTick);
+
+			IStatGroupEnableManager::Get().UpdateMemoryUsage();
+			State.UpdateStatMessagesMemoryUsage();
+
+			bReadyToProcess = false;
+			FStatPacketArray NowData;
+			Exchange(NowData.Packets, IncomingData.Packets);
+			INC_DWORD_STAT_BY(STAT_StatFramePacketsRecv, NowData.Packets.Num());
+			{
+				SCOPE_CYCLE_COUNTER(STAT_StatsNewParseMeta);
+				TArray<FStatMessage> MetaMessages;
+				{
+					FScopeLock Lock(&FStartupMessages::Get().CriticalSection);
+					Exchange(FStartupMessages::Get().DelayedMessages, MetaMessages);
+				}
+				if (MetaMessages.Num())
+				{
+					State.ProcessMetaDataOnly(MetaMessages);
+				}
+			}
+			{
+				SCOPE_CYCLE_COUNTER(STAT_ScanForAdvance);
+				State.ScanForAdvance(NowData);
+			}
+
+			if (FThreadStats::bIsRawStatsActive)
+			{
+				// Process raw stats.
+				State.ProcessRawStats(NowData);
+				State.ResetRegularStats();
+			}
+			else
+			{
+				// Process regular stats.
+				SCOPE_CYCLE_COUNTER(STAT_StatsNewAddToHistory);
+				State.ResetRawStats();
+				State.AddToHistoryAndEmpty(NowData);
+			}
+			check(!NowData.Packets.Num());
+			LastTime = FPlatformTime::Seconds();
+		}
+	}
+
+	/** Accesses singleton. */
+	static FStatsThread& Get()
+	{
+		static FStatsThread Singleton;
+		return Singleton;
+	}
+
+	/** Received a stat packet from other thread and add to the processing queue. */
+	void StatMessage(FStatPacket* Packet)
+	{
+		LLM_SCOPE(ELLMTag::Stats);
+
+		if (CVarDumpStatPackets.GetValueOnAnyThread())
+		{
+			UE_LOG(LogStats, Log, TEXT("Packet from %x with %d messages"), Packet->ThreadId, Packet->StatMessages.Num());
+		}
+
+		bReadyToProcess = Packet->ThreadType != EThreadType::Other;
+		IncomingData.Packets.Add(Packet);
+		State.NumStatMessages.Add(Packet->StatMessages.Num());
+
+		Tick();
+	}
+
+	void SelfStatMessage(FStatPacket* Packet)
+	{
+		if (CVarDumpStatPackets.GetValueOnAnyThread())
+		{
+			UE_LOG(LogStats, Log, TEXT("Self Packet from %x with %d messages"), Packet->ThreadId, Packet->StatMessages.Num());
+		}
+
+		IncomingData.Packets.Add(Packet);
+		State.NumStatMessages.Add(Packet->StatMessages.Num());
+	}
+
+	/** Start a stats runnable thread. */
+	void Start()
+	{
+		Thread = FRunnableThread::Create(this, TEXT("StatsThread"), 512 * 1024, TPri_BelowNormal, FPlatformAffinity::GetStatsThreadMask());
+		check(Thread);
+	}
+
+	/** Ends the stats runnable thread. */
+	void End()
+	{
+		delete Thread;
+		Thread = nullptr;
+	}
+
+private:
+	FRunnableThread* Thread;
+};
+
+/*-----------------------------------------------------------------------------
+	FStatMessagesTask
+-----------------------------------------------------------------------------*/
+
+// not using a delegate here to allow higher performance since we may end up sending a lot of small message arrays to the thread.
+class FStatMessagesTask
+{
+	FStatPacket* Packet;
+public:
+	FStatMessagesTask(FStatPacket* InPacket)
+		: Packet(InPacket)
+	{
+	}
+	FORCEINLINE TStatId GetStatId() const
+	{
+		return TStatId(); // we don't want to record this or it spams the stat system; we cover this time when we tick the stats system
+	}
+	static ENamedThreads::Type GetDesiredThread()
+	{
+		return FPlatformProcess::SupportsMultithreading() ? ENamedThreads::StatsThread : ENamedThreads::GameThread;
+	}
+	static ESubsequentsMode::Type GetSubsequentsMode() { return ESubsequentsMode::TrackSubsequents; }
+
+	void DoTask(ENamedThreads::Type CurrentThread, const FGraphEventRef& MyCompletionGraphEvent)
+	{
+		FStatsThread::Get().StatMessage(Packet);
+		Packet = NULL;
+	}
+};
+
+#endif // UE_STATS_THREAD_AS_PIPE
 /*-----------------------------------------------------------------------------
 	FThreadStatsPool
 -----------------------------------------------------------------------------*/
@@ -1046,6 +1249,9 @@ FThreadStats::FThreadStats():
 	MemoryMessageScope(0),
 	bReentranceGuard(false),
 	bSawExplicitFlush(false)
+#if !UE_STATS_THREAD_AS_PIPE
+	, bIsStatsThread(false)
+#endif
 {
 	Packet.SetThreadProperties();
 
@@ -1060,6 +1266,9 @@ FThreadStats::FThreadStats( EConstructor ):
 	MemoryMessageScope(0),
 	bReentranceGuard(false),
 	bSawExplicitFlush(false)
+#if !UE_STATS_THREAD_AS_PIPE
+	, bIsStatsThread(false)
+#endif
 {}
 
 void FThreadStats::CheckEnable()
@@ -1096,11 +1305,11 @@ void FThreadStats::SendMessage_Async(FStatPacket* ToSend)
 {
 	if (FPlatformProcess::SupportsMultithreading())
 	{
-		GStatsPipe.Launch(UE_SOURCE_LOCATION, [ToSend] { FStatsProcessor::Get().StatMessage(ToSend); });
+		GStatsPipe.Launch(UE_SOURCE_LOCATION, [ToSend] { FStatsThread::Get().StatMessage(ToSend); });
 	}
 	else
 	{
-		FFunctionGraphTask::CreateAndDispatchWhenReady([ToSend] { FStatsProcessor::Get().StatMessage(ToSend); }, TStatId{}, nullptr, ENamedThreads::GameThread);
+		FFunctionGraphTask::CreateAndDispatchWhenReady([ToSend] { FStatsThread::Get().StatMessage(ToSend); }, TStatId{}, nullptr, ENamedThreads::GameThread);
 	}
 }
 
@@ -1163,7 +1372,20 @@ void FThreadStats::FlushRegularStats( bool bHasBrokenCallstacks, bool bForceFlus
 			}
 			Packet.StatMessages.Empty(MaxPresize);
 		}
+
+#if UE_STATS_THREAD_AS_PIPE
 		SendMessage_Async(ToSend);
+#else
+		if (bIsStatsThread)
+		{
+			FStatsThread::Get().SelfStatMessage(ToSend);
+		}
+		else
+		{
+			TGraphTask<FStatMessagesTask>::CreateTask().ConstructAndDispatchWhenReady(ToSend);
+		}
+#endif
+
 		UpdateExplicitFlush();
 	}
 }
@@ -1201,7 +1423,19 @@ void FThreadStats::FlushRawStats( bool bHasBrokenCallstacks /*= false*/, bool bF
 
 		check(!Packet.StatMessages.Num());
 
+#if UE_STATS_THREAD_AS_PIPE
 		SendMessage_Async(ToSend);
+#else
+		if (bIsStatsThread)
+		{
+			FStatsThread::Get().SelfStatMessage(ToSend);
+		}
+		else
+		{
+			TGraphTask<FStatMessagesTask>::CreateTask().ConstructAndDispatchWhenReady(ToSend);
+		}
+#endif
+
 		UpdateExplicitFlush();
 
 		const float NumMessagesAsMB = float(NumMessages * sizeof(FStatMessage)) / 1024.0f / 1024.0f;
@@ -1311,6 +1545,11 @@ void FThreadStats::StartThread()
 		check(TlsSlot);
 	}
 
+#if !UE_STATS_THREAD_AS_PIPE
+	FStatsThread::Get();
+	FStatsThread::Get().Start();
+#endif
+
 	check(IsThreadingReady());
 	CheckEnable();
 
@@ -1327,6 +1566,8 @@ void FThreadStats::StartThread()
 
 static int32 CurrentEventIndex = 0;
 
+static FGraphEventRef LastFramesEvents[MAX_STAT_LAG];
+
 void FThreadStats::StopThread()
 {
 	// Nothing to stop if it was never started
@@ -1346,11 +1587,22 @@ void FThreadStats::StopThread()
 		WaitForStats();
 	}
 
+#if UE_STATS_THREAD_AS_PIPE
 	// wait for the pipe to complete all currently piped tasks
 	while (GStatsPipe.HasWork())
 	{
 		FPlatformProcess::Yield();
 	}
+#else
+	for (int32 Index = 0; Index < MAX_STAT_LAG; Index++)
+	{
+		LastFramesEvents[Index] = NULL;
+	}
+	FGraphEventRef QuitTask = TGraphTask<FReturnGraphTask>::CreateTask(NULL, ENamedThreads::GameThread).ConstructAndDispatchWhenReady(FPlatformProcess::SupportsMultithreading() ? ENamedThreads::StatsThread : ENamedThreads::GameThread);
+	FTaskGraphInterface::Get().WaitUntilTaskCompletes(QuitTask, ENamedThreads::GameThread_Local);
+
+	FStatsThread::Get().End();
+#endif
 }
 
 void FThreadStats::WaitForStats()
@@ -1367,21 +1619,20 @@ void FThreadStats::WaitForStats()
 
 		int32 EventIndex = (CurrentEventIndex + MAX_STAT_LAG - 1) % MAX_STAT_LAG;
 
+#if UE_STATS_THREAD_AS_PIPE
 		if (FPlatformProcess::SupportsMultithreading())
 		{
-			static UE::Tasks::FTask LastFramesEvents[MAX_STAT_LAG];
-		
+			static UE::Tasks::FTask LocalLastFramesEvents[MAX_STAT_LAG];
 			{
 				SCOPE_CYCLE_COUNTER(STAT_WaitForStats);
-				LastFramesEvents[EventIndex].Wait();
+				LocalLastFramesEvents[EventIndex].Wait();
 			}
 
-			LastFramesEvents[EventIndex] = GStatsPipe.Launch(UE_SOURCE_LOCATION, [] {});
+			LocalLastFramesEvents[EventIndex] = GStatsPipe.Launch(UE_SOURCE_LOCATION, [] {});
 		}
 		else
+#endif
 		{
-			static FGraphEventRef LastFramesEvents[MAX_STAT_LAG];
-
 			{
 				SCOPE_CYCLE_COUNTER(STAT_WaitForStats);
 				if (LastFramesEvents[EventIndex].GetReference())
diff --git a/Engine/Source/Runtime/Core/Private/Stats/StatsCommand.cpp b/Engine/Source/Runtime/Core/Private/Stats/StatsCommand.cpp
--- a/Engine/Source/Runtime/Core/Private/Stats/StatsCommand.cpp
+++ b/Engine/Source/Runtime/Core/Private/Stats/StatsCommand.cpp
@@ -2351,16 +2351,34 @@ bool DirectStatsCommand(const TCHAR* Cmd, bool bBlockForCompletion /*= false*/,
 				STAT_FSimpleDelegateGraphTask_StatCmd,
 				STATGROUP_TaskGraphTasks);
 
+#if UE_STATS_THREAD_AS_PIPE
 			UE::Tasks::FTask Task = GStatsPipe.Launch(UE_SOURCE_LOCATION, [FullCmd, bStatCommand, Ar] { StatCmd(FullCmd, bStatCommand, Ar); });
 			if (bBlockForCompletion)
 			{
 				Task.Wait();
 				GLog->FlushThreadedLogs();
 			}
-#else
+#else // UE_STATS_THREAD_AS_PIPE
+			ENamedThreads::Type ThreadType = ENamedThreads::GameThread;
+			if (FPlatformProcess::SupportsMultithreading())
+			{
+				ThreadType = ENamedThreads::StatsThread;
+			}
+
+			FGraphEventRef CompleteHandle = FSimpleDelegateGraphTask::CreateAndDispatchWhenReady(
+				FSimpleDelegateGraphTask::FDelegate::CreateStatic(&StatCmd, FullCmd, bStatCommand, Ar),
+				GET_STATID(STAT_FSimpleDelegateGraphTask_StatCmd), NULL, ThreadType
+			);
+			if (bBlockForCompletion && FPlatformProcess::SupportsMultithreading())
+			{
+				FTaskGraphInterface::Get().WaitUntilTaskCompletes(CompleteHandle);
+				GLog->FlushThreadedLogs();
+			}
+#endif // UE_STATS_THREAD_AS_PIPE
+#else // STATS
 			// If stats aren't enabled, broadcast so engine stats can still be triggered
 			StatCmd(FullCmd, bStatCommand, Ar);
-#endif
+#endif // STATS
 		}
 	}
 	return bResult;
@@ -2368,12 +2386,12 @@ bool DirectStatsCommand(const TCHAR* Cmd, bool bBlockForCompletion /*= false*/,
 
 #if STATS
 
-static void GetPermanentStats_StatsPipe(TArray<FStatMessage>* OutStats)
+static void GetPermanentStats_StatsThread(TArray<FStatMessage>* OutStats)
 {
-	DECLARE_CYCLE_STAT(TEXT("FSimpleDelegateGraphTask.GetPermanentStatsString_StatsPipe"),
-		STAT_FSimpleDelegateGraphTask_GetPermanentStatsString_StatsPipe,
+	DECLARE_CYCLE_STAT(TEXT("FSimpleDelegateGraphTask.GetPermanentStatsString_StatsThread"),
+		STAT_FSimpleDelegateGraphTask_GetPermanentStatsString_StatsThread,
 		STATGROUP_TaskGraphTasks);
-	SCOPE_CYCLE_COUNTER(STAT_FSimpleDelegateGraphTask_GetPermanentStatsString_StatsPipe);
+	SCOPE_CYCLE_COUNTER(STAT_FSimpleDelegateGraphTask_GetPermanentStatsString_StatsThread);
 
 	FStatsThreadState& StatsData = FStatsThreadState::GetLocalState();
 	TArray<FStatMessage>& Stats = *OutStats;
@@ -2386,9 +2404,18 @@ static void GetPermanentStats_StatsPipe(TArray<FStatMessage>* OutStats)
 
 void GetPermanentStats(TArray<FStatMessage>& OutStats)
 {
+#if UE_STATS_THREAD_AS_PIPE
 	GStatsPipe
-		.Launch(UE_SOURCE_LOCATION, [&OutStats] { GetPermanentStats_StatsPipe(&OutStats); })
+		.Launch(UE_SOURCE_LOCATION, [&OutStats] { GetPermanentStats_StatsThread(&OutStats); })
 		.Wait();
+#else
+	FGraphEventRef CompleteHandle = FSimpleDelegateGraphTask::CreateAndDispatchWhenReady(
+		FSimpleDelegateGraphTask::FDelegate::CreateStatic(&GetPermanentStats_StatsThread, &OutStats),
+		TStatId{}, NULL,
+		FPlatformProcess::SupportsMultithreading() ? ENamedThreads::StatsThread : ENamedThreads::GameThread
+	);
+	FTaskGraphInterface::Get().WaitUntilTaskCompletes(CompleteHandle);
+#endif
 }
 
 #endif
diff --git a/Engine/Source/Runtime/Core/Public/Async/TaskGraphInterfaces.h b/Engine/Source/Runtime/Core/Public/Async/TaskGraphInterfaces.h
--- a/Engine/Source/Runtime/Core/Public/Async/TaskGraphInterfaces.h
+++ b/Engine/Source/Runtime/Core/Public/Async/TaskGraphInterfaces.h
@@ -57,15 +57,18 @@ namespace ENamedThreads
 		UnusedAnchor = -1,
 		/** The always-present, named threads are listed next **/
 #if STATS
+#if UE_STATS_THREAD_AS_PIPE
 		StatsThread UE_DEPRECATED(5.0, "`StatsThread` has been removed. Stats system should be used by Stats public API"),
+#else
+		StatsThread,
+#endif
 #endif
 		RHIThread,
 #if UE_AUDIO_THREAD_AS_PIPE
-		AudioThread UE_DEPRECATED(5.0, "`AudioThread` has been removed. Please use `FAudioThread` API")
+		AudioThread UE_DEPRECATED(5.0, "`AudioThread` has been removed. Please use `FAudioThread` API"),
 #else
-		AudioThread
+		AudioThread,
 #endif
-		,
 		GameThread,
 		// The render thread is sometimes the game thread and is sometimes the actual rendering thread
 		ActualRenderingThread = GameThread + 1,
@@ -103,9 +106,13 @@ namespace ENamedThreads
 
 		/** Combinations **/
 #if STATS
+#if UE_STATS_THREAD_AS_PIPE
 PRAGMA_DISABLE_DEPRECATION_WARNINGS
 		StatsThread_Local UE_DEPRECATED(5.0, "`StatsThread_Local` has been removed (and has never been supported). Stats system should be used by its public API in `Stats2.h`") = StatsThread | LocalQueue,
 PRAGMA_ENABLE_DEPRECATION_WARNINGS
+#else
+		StatsThread_Local = StatsThread | LocalQueue,
+#endif
 #endif
 		GameThread_Local = GameThread | LocalQueue,
 		ActualRenderingThread_Local = ActualRenderingThread | LocalQueue,
diff --git a/Engine/Source/Runtime/Core/Public/GenericPlatform/GenericPlatformAffinity.h b/Engine/Source/Runtime/Core/Public/GenericPlatform/GenericPlatformAffinity.h
--- a/Engine/Source/Runtime/Core/Public/GenericPlatform/GenericPlatformAffinity.h
+++ b/Engine/Source/Runtime/Core/Public/GenericPlatform/GenericPlatformAffinity.h
@@ -80,7 +80,9 @@ public:
 		return 0xFFFFFFFFFFFFFFFF;
 	}
 
+#if UE_STATS_THREAD_AS_PIPE
 	UE_DEPRECATED(5.0, "This parameter is obsolete as StatsThread was removed and the affinity mask has no effect anymore. Please remove the call.")
+#endif
 	static const CORE_API uint64 GetStatsThreadMask()
 	{
 		return 0xFFFFFFFFFFFFFFFF;
diff --git a/Engine/Source/Runtime/Core/Public/Misc/CoreDefines.h b/Engine/Source/Runtime/Core/Public/Misc/CoreDefines.h
--- a/Engine/Source/Runtime/Core/Public/Misc/CoreDefines.h
+++ b/Engine/Source/Runtime/Core/Public/Misc/CoreDefines.h
@@ -70,3 +70,7 @@ Can be overridden by setting to 1 or 0 in the project's .Target.cs files.
 	#define UE_AUDIO_THREAD_AS_PIPE 1
 #endif
 
+#if !defined(UE_STATS_THREAD_AS_PIPE)
+	#define UE_STATS_THREAD_AS_PIPE 1
+#endif
+
diff --git a/Engine/Source/Runtime/Core/Public/Stats/Stats2.h b/Engine/Source/Runtime/Core/Public/Stats/Stats2.h
--- a/Engine/Source/Runtime/Core/Public/Stats/Stats2.h
+++ b/Engine/Source/Runtime/Core/Public/Stats/Stats2.h
@@ -1232,7 +1232,7 @@ class FThreadStats : FNoncopyable
 {
 	friend class FStatsMallocProfilerProxy;
 	friend class FStatsThreadState;
-	friend class FStatsProcessor;
+	friend class FStatsThread;
 	friend struct FThreadStatsPool;
 
 	/** Used to control when we are collecting stats. User of the stats system increment and decrement this counter as they need data. **/
@@ -1275,6 +1275,11 @@ class FThreadStats : FNoncopyable
 	/** Tracks current stack depth for cycle counters. **/
 	bool bSawExplicitFlush;
 
+#if !UE_STATS_THREAD_AS_PIPE
+	/** True if this is the stats thread, which needs special handling. **/
+	bool bIsStatsThread;
+#endif
+
 	/** Gathers information about the current thread and sets up the TLS value. **/
 	CORE_API FThreadStats();
 
@@ -1675,7 +1680,7 @@ private:
 /** Manages startup messages, usually to update the metadata. */
 class FStartupMessages
 {
-	friend class FStatsProcessor;
+	friend class FStatsThread;
 
 	TArray<FStatMessage> DelayedMessages;
 	FCriticalSection CriticalSection;
diff --git a/Engine/Source/Runtime/Core/Public/Stats/StatsData.h b/Engine/Source/Runtime/Core/Public/Stats/StatsData.h
--- a/Engine/Source/Runtime/Core/Public/Stats/StatsData.h
+++ b/Engine/Source/Runtime/Core/Public/Stats/StatsData.h
@@ -488,7 +488,7 @@ struct IItemFilter
  */
 class CORE_API FStatsThreadState
 {
-	friend class FStatsProcessor;
+	friend class FStatsThread;
 	friend struct FStatPacketArray;
 	friend struct FStatsReadFile;
 
