New HashCombineFast function which is simpler than HashCombine.  We can't change HashCombine because existing code may depend on those values
Some commenting to formalize that HashCombineFast and GetTypeHash should not be stored or used outside of the current running process.
PointerHash and GetTypeHash(TTuple) updated to use HashCombineFast.

diff --git a/Engine/Source/Runtime/Core/Public/Templates/Tuple.h b/Engine/Source/Runtime/Core/Public/Templates/Tuple.h
--- a/Engine/Source/Runtime/Core/Public/Templates/Tuple.h
+++ b/Engine/Source/Runtime/Core/Public/Templates/Tuple.h
@@ -647,7 +647,7 @@ namespace UE4Tuple_Private
 		template <typename TupleType>
 		FORCEINLINE static uint32 Do(uint32 Hash, const TupleType& Tuple)
 		{
-			return TGetTupleHashHelper<ArgToCombine + 1, ArgCount>::Do(HashCombine(Hash, GetTypeHash(Tuple.template Get<ArgToCombine>())), Tuple);
+			return TGetTupleHashHelper<ArgToCombine + 1, ArgCount>::Do(HashCombineFast(Hash, GetTypeHash(Tuple.template Get<ArgToCombine>())), Tuple);
 		}
 	};
 
diff --git a/Engine/Source/Runtime/Core/Public/Templates/TypeHash.h b/Engine/Source/Runtime/Core/Public/Templates/TypeHash.h
--- a/Engine/Source/Runtime/Core/Public/Templates/TypeHash.h
+++ b/Engine/Source/Runtime/Core/Public/Templates/TypeHash.h
@@ -9,6 +9,9 @@
 /**
  * Combines two hash values to get a third.
  * Note - this function is not commutative.
+ *
+ * This function cannot change for backward compatibility reasons.
+ * You may want to choose HashCombineFast for a better in-memory hash combining function.
  */
 inline uint32 HashCombine(uint32 A, uint32 C)
 {
@@ -28,10 +31,22 @@ inline uint32 HashCombine(uint32 A, uint32 C)
 	return C;
 }
 
+/**
+ * Combines two hash values to get a third.
+ * Note - this function is not commutative.
+ *
+ * WARNING!  This function is subject to change and should only be used for creating
+ *           combined hash values which don't leave the running process,
+ *           e.g. GetTypeHash() overloads.
+ */
+inline uint32 HashCombineFast(uint32 A, uint32 B)
+{
+	// From http://www.open-std.org/jtc1/sc22/wg21/docs/papers/2018/p0814r2.pdf
+	return A ^ (B + 0x9e3779b9 + (A << 6) + (A >> 2));
+}
 
 inline uint32 PointerHash(const void* Key,uint32 C = 0)
 {
-	// Avoid LHS stalls on PS3 and Xbox 360
 #if PLATFORM_64BITS
 	// Ignoring the lower 4 bits since they are likely zero anyway.
 	// Higher bits are more significant in 64 bit builds.
@@ -40,12 +55,17 @@ inline uint32 PointerHash(const void* Key,uint32 C = 0)
 	auto PtrInt = reinterpret_cast<UPTRINT>(Key);
 #endif
 
-	return HashCombine((uint32)PtrInt, C);
+	// we can use HashCombineFast here because pointers are non-persistent
+	return HashCombineFast((uint32)PtrInt, C);
 }
 
 
 //
 // Hash functions for common types.
+// 
+// WARNING!  GetTypeHash result values are not expected to leave the running process.
+//           Do not persist them to disk, send them to another running process or
+//           expect them to be consistent across multiple runs.
 //
 
 inline uint32 GetTypeHash( const uint8 A )
