LWC: Avoid replacing double types with float precision versions on save.

diff --git a/Engine/Source/Runtime/Core/Public/Math/BoxSphereBounds.h b/Engine/Source/Runtime/Core/Public/Math/BoxSphereBounds.h
--- a/Engine/Source/Runtime/Core/Public/Math/BoxSphereBounds.h
+++ b/Engine/Source/Runtime/Core/Public/Math/BoxSphereBounds.h
@@ -307,7 +307,10 @@ public:
 		{
 			float Radius = (float)Bounds.SphereRadius;
 			Ar << Radius;
-			Bounds.SphereRadius = Radius;
+			if(Ar.IsLoading())
+			{
+				Bounds.SphereRadius = Radius;
+			}
 		}
 
 		return Ar;
diff --git a/Engine/Source/Runtime/Core/Public/Math/Matrix.h b/Engine/Source/Runtime/Core/Public/Math/Matrix.h
--- a/Engine/Source/Runtime/Core/Public/Math/Matrix.h
+++ b/Engine/Source/Runtime/Core/Public/Math/Matrix.h
@@ -500,10 +500,13 @@ inline FArchive& operator<<(FArchive& Ar, TMatrix<double>& M)
 		{
 			float Col0 = (float)M.M[Row][0], Col1 = (float)M.M[Row][1], Col2 = (float)M.M[Row][2], Col3 = (float)M.M[Row][3];
 			Ar << Col0 << Col1 << Col2 << Col3;
-			M.M[Row][0] = Col0;
-			M.M[Row][1] = Col1;
-			M.M[Row][2] = Col2;
-			M.M[Row][3] = Col3;
+			if(Ar.IsLoading())
+			{
+				M.M[Row][0] = Col0;
+				M.M[Row][1] = Col1;
+				M.M[Row][2] = Col2;
+				M.M[Row][3] = Col3;
+			}
 		}
 	}
 	M.DiagnosticCheckNaN();
diff --git a/Engine/Source/Runtime/Core/Public/Math/Plane.h b/Engine/Source/Runtime/Core/Public/Math/Plane.h
--- a/Engine/Source/Runtime/Core/Public/Math/Plane.h
+++ b/Engine/Source/Runtime/Core/Public/Math/Plane.h
@@ -366,7 +366,10 @@ inline FArchive& operator<<(FArchive& Ar, TPlane<double>& P)
 		// Stored as floats, so serialize float and copy.
 		float SW = (float)P.W;
 		Ar << SW;
-		P.W = SW;
+		if(Ar.IsLoading())
+		{
+			P.W = SW;
+		}
 	}
 
 	P.DiagnosticCheckNaN();
diff --git a/Engine/Source/Runtime/Core/Public/Math/Sphere.h b/Engine/Source/Runtime/Core/Public/Math/Sphere.h
--- a/Engine/Source/Runtime/Core/Public/Math/Sphere.h
+++ b/Engine/Source/Runtime/Core/Public/Math/Sphere.h
@@ -190,7 +190,10 @@ public:
 		{
 			float SW = (float)Sphere.W;
 			Ar << SW;
-			Sphere.W = SW;
+			if (Ar.IsLoading())
+			{
+				Sphere.W = SW;
+			}
 		}
 
 		return Ar;
diff --git a/Engine/Source/Runtime/Core/Public/Math/Vector.h b/Engine/Source/Runtime/Core/Public/Math/Vector.h
--- a/Engine/Source/Runtime/Core/Public/Math/Vector.h
+++ b/Engine/Source/Runtime/Core/Public/Math/Vector.h
@@ -1161,7 +1161,10 @@ inline void operator<<(FStructuredArchive::FSlot Slot, TVector<double>& V)
 		Record << SA_VALUE(TEXT("X"), SX);
 		Record << SA_VALUE(TEXT("Y"), SY);
 		Record << SA_VALUE(TEXT("Z"), SZ);
-		V = TVector<double>(SX, SY, SZ);
+		if(Slot.GetUnderlyingArchive().IsLoading())
+		{
+			V = TVector<double>(SX, SY, SZ);
+		}
 	}
 	V.DiagnosticCheckNaN();
 }
diff --git a/Engine/Source/Runtime/Experimental/Chaos/Public/Chaos/ChaosArchive.h b/Engine/Source/Runtime/Experimental/Chaos/Public/Chaos/ChaosArchive.h
--- a/Engine/Source/Runtime/Experimental/Chaos/Public/Chaos/ChaosArchive.h
+++ b/Engine/Source/Runtime/Experimental/Chaos/Public/Chaos/ChaosArchive.h
@@ -366,7 +366,10 @@ FORCEINLINE FChaosArchive& operator<<(FChaosArchive& Ar, Chaos::FReal& Real)
 		//ensure(Ar.IsLoading()); // this case should normally only happening when reading 
 		FRealSingle RealSingle = (FRealSingle)Real; 
 		operator<<((FArchive&)Ar, RealSingle);
-		Real = (FReal)RealSingle;
+		if(Ar.IsLoading())
+		{
+			Real = (FReal)RealSingle;
+		}
 	}
 	return Ar;
 }
