Speculative fix/workaround for issue where metadata queries sometimes fail on a newly opened file handle

When this happens the file handle becomes invalid and cannot be used. This change ensures a null handle is returned instead of an invalid handle

It's not clear why a GetFileSizeEx operation would fail on a recently created file handle but it sometimes does for files on file shares (such as for derived data cache files)

diff --git a/Engine/Source/Runtime/Core/Private/Windows/WindowsPlatformFile.cpp b/Engine/Source/Runtime/Core/Private/Windows/WindowsPlatformFile.cpp
--- a/Engine/Source/Runtime/Core/Private/Windows/WindowsPlatformFile.cpp
+++ b/Engine/Source/Runtime/Core/Private/Windows/WindowsPlatformFile.cpp
@@ -1238,7 +1238,21 @@ public:
 		if (Handle != INVALID_HANDLE_VALUE)
 		{
 			TRACE_PLATFORMFILE_END_OPEN(Handle);
-			return new FFileHandleWindows(Handle, Access, WinFlags, FILE_FLAG_OVERLAPPED);
+
+			FFileHandleWindows* FileHandle = new FFileHandleWindows(Handle, Access, WinFlags, FILE_FLAG_OVERLAPPED);
+
+			// Some operations can fail during the handle initialization, so we
+			// double check that the handle is valid before returning it
+			if (FileHandle->IsValid())
+			{
+				return FileHandle;
+			}
+			else
+			{
+				delete FileHandle;
+
+				return nullptr;
+			}
 		}
 		else
 		{
@@ -1257,12 +1271,24 @@ public:
 		if(Handle != INVALID_HANDLE_VALUE)
 		{
 			TRACE_PLATFORMFILE_END_OPEN(Handle);
-			FFileHandleWindows *PlatformFileHandle = new FFileHandleWindows(Handle, Access, WinFlags, 0);
-			if (bAppend)
+			FFileHandleWindows* PlatformFileHandle = new FFileHandleWindows(Handle, Access, WinFlags, 0);
+
+			// Some operations can fail during the handle initialization, so we
+			// double check that the handle is valid before returning it
+			if (PlatformFileHandle->IsValid())
+			{
+				if (bAppend)
+				{
+					PlatformFileHandle->SeekFromEnd(0);
+				}
+				return PlatformFileHandle;
+			}
+			else
 			{
-				PlatformFileHandle->SeekFromEnd(0);
+				delete PlatformFileHandle;
+
+				return nullptr;
 			}
-			return PlatformFileHandle;
 		}
 		else
 		{
