Fix FastDecimalFormat not handling inf.

diff --git a/Engine/Source/Runtime/Core/Private/Internationalization/FastDecimalFormat.cpp b/Engine/Source/Runtime/Core/Private/Internationalization/FastDecimalFormat.cpp
--- a/Engine/Source/Runtime/Core/Private/Internationalization/FastDecimalFormat.cpp
+++ b/Engine/Source/Runtime/Core/Private/Internationalization/FastDecimalFormat.cpp
@@ -462,7 +462,7 @@ FString CultureInvariantDecimalToString(const double InVal, const TCHAR*& InBuff
 
 		// 48 for raw ascii -> int
 		int32 CharacterIndex = (int32)InBuffer[0] - 48;
-		if (ensure(CharacterIndex >= 0 && CharacterIndex < sizeof(InFormattingRules.DigitCharacters) / sizeof(InFormattingRules.DigitCharacters[0])))
+		if (ensure(CharacterIndex >= 0 && CharacterIndex < UE_ARRAY_COUNT(InFormattingRules.DigitCharacters)))
 		{
 			OutStr += InFormattingRules.DigitCharacters[CharacterIndex];
 			FractionalDigitsPrinted += bParsedFractional ? 1 : 0;
@@ -533,8 +533,12 @@ void FractionalToString(const double InVal, const FDecimalNumberFormattingRules&
 	{
 		OutString = LexToSanitizedString(InVal);
 
-		const TCHAR* CultureInvariantDecimalBuffer = *OutString;
-		OutString = CultureInvariantDecimalToString(InVal, CultureInvariantDecimalBuffer, OutString.Len(), InFormattingRules, InFormattingOptions);
+		// Guard against 'inf', other non-decimal floats.
+		if (OutString.IsNumeric())
+		{
+			const TCHAR* CultureInvariantDecimalBuffer = *OutString;
+			OutString = CultureInvariantDecimalToString(InVal, CultureInvariantDecimalBuffer, OutString.Len(), InFormattingRules, InFormattingOptions);
+		}
 		return;
 	}
 
