Added missing double overload for VectorGetComponentDynamic

diff --git a/Engine/Source/Runtime/Core/Public/Math/UnrealMathNeon.h b/Engine/Source/Runtime/Core/Public/Math/UnrealMathNeon.h
--- a/Engine/Source/Runtime/Core/Public/Math/UnrealMathNeon.h
+++ b/Engine/Source/Runtime/Core/Public/Math/UnrealMathNeon.h
@@ -1629,6 +1629,12 @@ FORCEINLINE float VectorGetComponentDynamic(VectorRegister4Float Vec, uint32 Ele
 	return Floats[ElementIndex];
 }
 
+FORCEINLINE double VectorGetComponentDynamic(VectorRegister4Double Vec, uint32 ElementIndex)
+{
+	AlignedDouble4 Doubles(Vec);
+	return Doubles[ElementIndex];
+}
+
 /**
  * Multiplies two 4x4 matrices.
  *
