TaskGraph: disabled "busy waiting by default" and replaced by reserve workers, based on Arne Schober implementation. to turn back to original implementation use `-TaskGraphDisableReserveWorkers`

diff --git a/Engine/Source/Runtime/Core/Private/Async/Fundamental/ReserveScheduler.cpp b/Engine/Source/Runtime/Core/Private/Async/Fundamental/ReserveScheduler.cpp
new file mode 100644
--- /dev/null
+++ b/Engine/Source/Runtime/Core/Private/Async/Fundamental/ReserveScheduler.cpp
@@ -0,0 +1,106 @@
+// Copyright Epic Games, Inc. All Rights Reserved.
+
+#include "Async/Fundamental/ReserveScheduler.h"
+
+namespace LowLevelTasks
+{
+
+FReserveScheduler FReserveScheduler::Singleton;
+
+TUniquePtr<FThread> FReserveScheduler::CreateWorker(FSchedulerTls::FLocalQueueType* WorkerLocalQueue, EThreadPriority Priority, bool bIsForkable)
+{
+	uint32 WorkerId = NextWorkerId++;
+	return MakeUnique<FThread>
+	(
+		*FString::Printf(TEXT("Reserve Worker #%d"), WorkerId),
+		[this, WorkerLocalQueue]
+		{
+			FSchedulerTls::ActiveScheduler = this;
+			FSchedulerTls::LocalQueue = WorkerLocalQueue;
+
+			FYieldedWork ReserveEvent;
+			while (true)
+			{
+				EventStack.Push(&ReserveEvent);
+				ReserveEvent.SleepEvent->Wait();
+				{
+					TRACE_CPUPROFILER_EVENT_SCOPE(FReserveScheduler::BusyWaitUntil);
+					FSchedulerTls::WorkerType = ReserveEvent.bPermitBackgroundWork ? FSchedulerTls::EWorkerType::Background : FSchedulerTls::EWorkerType::Foreground;
+
+					if (ActiveWorkers.load(std::memory_order_relaxed) == 0)
+					{
+						break;
+					}
+
+					BusyWaitUntil(MoveTemp(ReserveEvent.CompletedDelegate), ReserveEvent.bPermitBackgroundWork);
+				}
+			}
+			FSchedulerTls::WorkerType = EWorkerType::None;
+			FSchedulerTls::ActiveScheduler = nullptr;
+			FSchedulerTls::LocalQueue = nullptr;
+		}, 0, Priority, FThreadAffinity{ FPlatformAffinity::GetTaskGraphThreadMask(), 0 }, bIsForkable
+	);
+}
+
+bool FReserveScheduler::DoReserveWorkUntil(FConditional&& Condition)
+{
+	if (FYieldedWork* WorkerEvent = EventStack.Pop())
+	{
+		WorkerEvent->CompletedDelegate = MoveTemp(Condition);
+		WorkerEvent->bPermitBackgroundWork = FSchedulerTls::PermitBackgroundWork();
+		WorkerEvent->SleepEvent->Trigger();
+		return true;
+	}
+	return false;
+}
+
+void FReserveScheduler::StartWorkers(FScheduler& MainScheduler, uint32 NumWorkers, EThreadPriority WorkerPriority, bool bIsForkable)
+{
+	TaskTrace::Init();
+
+	if (NumWorkers == 0)
+	{
+		NumWorkers = FMath::Min(FPlatformMisc::NumberOfWorkerThreadsToSpawn(), 64);
+	}
+
+	uint32 OldActiveWorkers = ActiveWorkers.load(std::memory_order_relaxed);
+	if(OldActiveWorkers == 0 && FPlatformProcess::SupportsMultithreading() && ActiveWorkers.compare_exchange_strong(OldActiveWorkers, NumWorkers, std::memory_order_relaxed))
+	{
+		FScopeLock Lock(&WorkerThreadsCS);
+		check(!WorkerThreads.Num());
+		check(!WorkerLocalQueues.Num());
+		check(NextWorkerId == 0);
+
+		WorkerLocalQueues.Reserve(NumWorkers);
+		UE::Trace::ThreadGroupBegin(TEXT("Reserve Workers"));
+		for (uint32 WorkerId = 0; WorkerId < NumWorkers; ++WorkerId)
+		{
+			WorkerLocalQueues.Emplace(MainScheduler.GetQueueRegistry(), ELocalQueueType::EBusyWait);
+			WorkerThreads.Add(CreateWorker(&WorkerLocalQueues.Last(), WorkerPriority, bIsForkable));
+		}
+		UE::Trace::ThreadGroupEnd();
+	}
+}
+
+void FReserveScheduler::StopWorkers()
+{
+	uint32 OldActiveWorkers = ActiveWorkers.load(std::memory_order_relaxed);
+	if(OldActiveWorkers != 0 && ActiveWorkers.compare_exchange_strong(OldActiveWorkers, 0, std::memory_order_relaxed))
+	{
+		FScopeLock Lock(&WorkerThreadsCS);
+		while (FYieldedWork* Event = EventStack.Pop()) 
+		{
+			Event->SleepEvent->Trigger();
+		}
+
+		for (TUniquePtr<FThread>& Thread : WorkerThreads)
+		{
+			Thread->Join();
+		}
+		NextWorkerId = 0;
+		WorkerThreads.Reset();
+		WorkerLocalQueues.Reset();
+	}
+}
+
+}
\ No newline at end of file
diff --git a/Engine/Source/Runtime/Core/Private/Async/TaskGraph.cpp b/Engine/Source/Runtime/Core/Private/Async/TaskGraph.cpp
--- a/Engine/Source/Runtime/Core/Private/Async/TaskGraph.cpp
+++ b/Engine/Source/Runtime/Core/Private/Async/TaskGraph.cpp
@@ -34,6 +34,7 @@
 #include "Misc/ConfigCacheIni.h"
 
 #include "Async/Fundamental/Scheduler.h"
+#include "Async/Fundamental/ReserveScheduler.h"
 #include "Tasks/Pipe.h"
 
 DEFINE_LOG_CATEGORY_STATIC(LogTaskGraph, Log, All);
@@ -101,6 +102,7 @@ static FAutoConsoleVariableRef CVarForkedProcessMaxWorkerThreads(
 
 CORE_API int32 GUseNewTaskBackend = 1;
 uint32 GNumForegroundWorkers = 2;
+bool GDisableReserveWorkers = 0;
 
 #if CREATE_HIPRI_TASK_THREADS || CREATE_BACKGROUND_TASK_THREADS
 	static void ThreadSwitchForABTest(const TArray<FString>& Args)
@@ -1876,6 +1878,11 @@ public:
 			int32 NumForegroundWorkers =  FMath::Max(1, NumWorkerThreads - NumBackgroundWorkers);
 
 			LowLevelTasks::FScheduler::Get().StartWorkers(NumForegroundWorkers, NumBackgroundWorkers, EThreadPriority::TPri_Normal, EThreadPriority::TPri_BelowNormal, FForkProcessHelper::IsForkedMultithreadInstance());
+			if (!GDisableReserveWorkers)
+			{
+				LowLevelTasks::FReserveScheduler::Get().StartWorkers(LowLevelTasks::FScheduler::Get(), NumForegroundWorkers + NumBackgroundWorkers, EThreadPriority::TPri_BelowNormal, FForkProcessHelper::IsForkedMultithreadInstance());
+			}
+
 			NumNamedThreads = ENamedThreads::ActualRenderingThread + 1;
 			ENamedThreads::bHasBackgroundThreads = 1;
 			ENamedThreads::bHasHighPriorityThreads = 1;
@@ -1883,6 +1890,10 @@ public:
 		else
 		{
 			LowLevelTasks::FScheduler::Get().StopWorkers();
+			if (!GDisableReserveWorkers)
+			{
+				LowLevelTasks::FReserveScheduler::Get().StopWorkers();
+			}
 			NumNamedThreads = ENamedThreads::ActualRenderingThread;
 			ENamedThreads::bHasBackgroundThreads = 0;
 			ENamedThreads::bHasHighPriorityThreads = 0;
@@ -1913,6 +1924,10 @@ public:
 			NamedThreads[ThreadIndex].bAttached = false;
 		}
 		LowLevelTasks::FScheduler::Get().StopWorkers();
+		if (!GDisableReserveWorkers)
+		{
+			LowLevelTasks::FReserveScheduler::Get().StopWorkers();
+		}
 		FPlatformTLS::FreeTlsSlot(PerThreadIDTLSSlot);
 	}
 
@@ -1946,6 +1961,12 @@ public:
 
 			LowLevelTasks::FScheduler::Get().StopWorkers();
 			LowLevelTasks::FScheduler::Get().StartWorkers(NumWorkers, NumBackgroundWorkers, Pri, EThreadPriority::TPri_BelowNormal, FForkProcessHelper::IsForkedMultithreadInstance());
+
+			if (!GDisableReserveWorkers)
+			{
+				LowLevelTasks::FReserveScheduler::Get().StopWorkers();
+				LowLevelTasks::FReserveScheduler::Get().StartWorkers(LowLevelTasks::FScheduler::Get(), NumWorkers + NumBackgroundWorkers, EThreadPriority::TPri_BelowNormal, FForkProcessHelper::IsForkedMultithreadInstance());
+			}
 		}
 	}
 
@@ -2116,20 +2137,42 @@ private:
 			TGraphTask<FReturnGraphTask>::CreateTask(&Tasks, CurrentThread).ConstructAndDispatchWhenReady(CurrentThread);
 			ProcessThreadUntilRequestReturn(CurrentThread);
 		}
-		else if (LowLevelTasks::FScheduler::Get().IsWorkerThread())
+		else if (LowLevelTasks::FScheduler::Get().IsWorkerThread() || LowLevelTasks::FReserveScheduler::Get().IsWorkerThread())
 		{
-			LowLevelTasks::BusyWaitUntil([Index(0), &Tasks]() mutable
+			if (GDisableReserveWorkers)
 			{
-				while (Index < Tasks.Num())
+				LowLevelTasks::BusyWaitUntil([Index(0), &Tasks]() mutable
 				{
-					if (!Tasks[Index]->IsComplete())
+					while (Index < Tasks.Num())
 					{
-						return false;
+						if (!Tasks[Index]->IsComplete())
+						{
+							return false;
+						}
+						Index++;
 					}
-					Index++;
-				}
-				return true;
-			});
+					return true;
+				});
+			}
+			else
+			{
+				LowLevelTasks::DoReserveWorkUntil([Index(0), Tasks]() mutable
+				{
+					while (Index < Tasks.Num())
+					{
+						if (!Tasks[Index]->IsComplete())
+						{
+							return false;
+						}
+						Index++;
+					}
+					return true;
+				});
+
+				// We will just stall this thread on an event while we wait
+				FScopedEvent Event;
+				TriggerEventWhenTasksComplete(Event.Get(), Tasks, CurrentThreadIfKnown);
+			}
 		}
 		else
 		{
@@ -2273,6 +2316,15 @@ void FTaskGraphInterface::Startup(int32 NumThreads)
 		GUseNewTaskBackend = 1;
 	}
 
+	if (FParse::Param(FCommandLine::Get(), TEXT("TaskGraphDisableReserveWorkers")))
+	{
+		GDisableReserveWorkers = 1;
+	}
+	else if (FParse::Param(FCommandLine::Get(), TEXT("TaskGraphEnableReserveWorkers")))
+	{
+		GDisableReserveWorkers = 0;
+	}
+
 	if (GUseNewTaskBackend)
 	{
 		//We want to reduce the number of overall threads that UE uses so that there is are some 
diff --git a/Engine/Source/Runtime/Core/Private/Tests/Async/TaskGraphTest.cpp b/Engine/Source/Runtime/Core/Private/Tests/Async/TaskGraphTest.cpp
--- a/Engine/Source/Runtime/Core/Private/Tests/Async/TaskGraphTest.cpp
+++ b/Engine/Source/Runtime/Core/Private/Tests/Async/TaskGraphTest.cpp
@@ -1344,6 +1344,80 @@ PRAGMA_ENABLE_DEPRECATION_WARNINGS
 
 		return true;
 	}
+
+	template<uint32 Num>
+	void OversubscriptionStressTest()
+	{
+		// launch a number of tasks >= num of cores, each of them launches a task and waits for its completion. do this inside a task so 
+		// we can detect deadlock.
+		// outer tasks occupy all worker threads. when they launch inner tasks and wait for them, the inner tasks can't be 
+		// executed because all workers are blocked -> deadlock.
+		// can be solved by busy waiting
+
+		for (int i = 0; i != Num; ++i)
+		{
+			FSharedEventRef Event;
+			FFunctionGraphTask::CreateAndDispatchWhenReady(
+				[Event]
+				{
+					ParallelFor(200,
+						[](int32)
+						{
+							FPlatformProcess::Sleep(0.01f); // simulate some work and let all workers to pick up ParallelFor tasks
+							FFunctionGraphTask::CreateAndDispatchWhenReady([] {})->Wait();
+						}
+					);
+					Event->Trigger();
+				}
+			);
+			verify(Event->Wait(FTimespan::FromSeconds(5.f)));
+		}
+	}
+
+	IMPLEMENT_SIMPLE_AUTOMATION_TEST(FTaskGraphOversubscriptionTest, "System.Core.Async.TaskGraph.Oversubscription", EAutomationTestFlags::ApplicationContextMask | EAutomationTestFlags::EngineFilter);
+
+	bool FTaskGraphOversubscriptionTest::RunTest(const FString& Parameters)
+	{
+		UE_BENCHMARK(5, OversubscriptionStressTest<10>);
+
+		return true;
+	}
+
+	template<uint32 Nujm>
+	void SquaredOversubscriptionStressTest()
+	{
+		// same as before but using two ParallelFor nested one into another to simulate oversubscription in square
+
+		FSharedEventRef Event;
+		FFunctionGraphTask::CreateAndDispatchWhenReady(
+			[Event]
+			{
+				ParallelFor(200,
+					[](int32)
+					{
+						ParallelFor(200,
+							[](int32)
+							{
+								FPlatformProcess::Sleep(0.01f); // simulate some work and let all workers to pick up ParallelFor tasks
+								FFunctionGraphTask::CreateAndDispatchWhenReady([] {})->Wait();
+							}
+						);
+					}
+				);
+				Event->Trigger();
+			}
+		);
+		verify(Event->Wait(FTimespan::FromSeconds(30.f)));
+	}
+
+	IMPLEMENT_SIMPLE_AUTOMATION_TEST(FTaskGraphSquaredOversubscriptionTest, "System.Core.Async.TaskGraph.SquaredOversubscription", EAutomationTestFlags::ApplicationContextMask | EAutomationTestFlags::EngineFilter);
+
+	bool FTaskGraphSquaredOversubscriptionTest::RunTest(const FString& Parameters)
+	{
+		UE_BENCHMARK(5, SquaredOversubscriptionStressTest<10>);
+
+		return true;
+	}
 }
 
 #endif //WITH_DEV_AUTOMATION_TESTS
diff --git a/Engine/Source/Runtime/Core/Public/Async/Fundamental/ReserveScheduler.h b/Engine/Source/Runtime/Core/Public/Async/Fundamental/ReserveScheduler.h
new file mode 100644
--- /dev/null
+++ b/Engine/Source/Runtime/Core/Public/Async/Fundamental/ReserveScheduler.h
@@ -0,0 +1,65 @@
+// Copyright Epic Games, Inc. All Rights Reserved.
+
+#pragma once 
+#include "Scheduler.h"
+
+namespace LowLevelTasks
+{
+	class FReserveScheduler final : public FSchedulerTls
+	{
+		UE_NONCOPYABLE(FReserveScheduler);
+		static CORE_API FReserveScheduler Singleton;
+
+	public:
+		FReserveScheduler() = default;
+		~FReserveScheduler();
+
+	public:
+		using FConditional = TTaskDelegate<16, bool>;
+
+	private:
+		struct alignas(8) FYieldedWork
+		{
+			FEventRef		SleepEvent;
+			FConditional	CompletedDelegate = []() { return true; };
+			FYieldedWork*	Next = nullptr;
+			bool			bPermitBackgroundWork = false;
+		};
+	
+	public:
+		FORCEINLINE_DEBUGGABLE static FReserveScheduler& Get();
+
+		//start number of reserve workers where 0 is the system default
+		CORE_API void StartWorkers(FScheduler& MainScheduler, uint32 ReserveWorkers = 0, EThreadPriority WorkerPriority = EThreadPriority::TPri_Normal, bool bIsForkable = false);
+		CORE_API void StopWorkers();
+
+		//tries to yield this thread using the YieldEvent and do busywork on a reserve worker.
+		CORE_API bool DoReserveWorkUntil(FConditional&& Condition);
+
+	private: 
+		TUniquePtr<FThread> CreateWorker(FSchedulerTls::FLocalQueueType* WorkerLocalQueue, EThreadPriority Priority = EThreadPriority::TPri_Normal, bool bIsForkable = false);
+
+	private:
+		TEventStack<FYieldedWork> 				EventStack;
+		FCriticalSection 						WorkerThreadsCS;
+		TArray<FSchedulerTls::FLocalQueueType>	WorkerLocalQueues;
+		TArray<TUniquePtr<FThread>>				WorkerThreads;
+		std::atomic_uint						ActiveWorkers { 0 };
+		std::atomic_uint						NextWorkerId { 0 };
+	};
+
+	inline FReserveScheduler& FReserveScheduler::Get()
+	{
+		return Singleton;
+	}
+
+	FORCEINLINE_DEBUGGABLE bool DoReserveWorkUntil(FReserveScheduler::FConditional&& Condition)
+	{
+		return FReserveScheduler::Get().DoReserveWorkUntil(MoveTemp(Condition));
+	}
+
+	inline FReserveScheduler::~FReserveScheduler()
+	{
+		StopWorkers();
+	}
+}
\ No newline at end of file
