StringBuilder: Fixed compile errors for FUtf8StringBuilder

#rnx

diff --git a/Engine/Source/Runtime/Core/Private/Misc/StringBuilder.cpp b/Engine/Source/Runtime/Core/Private/Misc/StringBuilder.cpp
--- a/Engine/Source/Runtime/Core/Private/Misc/StringBuilder.cpp
+++ b/Engine/Source/Runtime/Core/Private/Misc/StringBuilder.cpp
@@ -86,4 +86,5 @@ TStringBuilderBase<C>& TStringBuilderBase<C>::AppendfImpl(BuilderType& Self, con
 // Instantiate templates once
 
 template class TStringBuilderBase<ANSICHAR>;
+template class TStringBuilderBase<UTF8CHAR>;
 template class TStringBuilderBase<WIDECHAR>;
diff --git a/Engine/Source/Runtime/Core/Public/Misc/StringBuilder.h b/Engine/Source/Runtime/Core/Public/Misc/StringBuilder.h
--- a/Engine/Source/Runtime/Core/Public/Misc/StringBuilder.h
+++ b/Engine/Source/Runtime/Core/Public/Misc/StringBuilder.h
@@ -173,7 +173,7 @@ public:
 
 		for (int32 i = 0; i < Length; ++i)
 		{
-			Dest[i] = String[i];
+			Dest[i] = CharType(String[i]);
 		}
 
 		return *this;
@@ -371,7 +371,7 @@ protected:
 	{
 		if (*CurPos)
 		{
-			*CurPos = 0;
+			*CurPos = CharType(0);
 		}
 	}
 
