Guid: Added functions for writing to string builders

#rnx

diff --git a/Engine/Source/Runtime/Core/Private/Misc/Guid.cpp b/Engine/Source/Runtime/Core/Private/Misc/Guid.cpp
--- a/Engine/Source/Runtime/Core/Private/Misc/Guid.cpp
+++ b/Engine/Source/Runtime/Core/Private/Misc/Guid.cpp
@@ -5,6 +5,7 @@
 #include "Math/Int128.h"
 #include "UObject/PropertyPortFlags.h"
 #include "Misc/Base64.h"
+#include "Misc/StringBuilder.h"
 
 
 /* FGuid interface
@@ -48,6 +49,9 @@ FString FGuid::ToString(EGuidFormats Format) const
 	case EGuidFormats::DigitsWithHyphens:
 		return FString::Printf(TEXT("%08X-%04X-%04X-%04X-%04X%08X"), A, B >> 16, B & 0xFFFF, C >> 16, C & 0xFFFF, D);
 
+	case EGuidFormats::DigitsWithHyphensLower:
+		return FString::Printf(TEXT("%08x-%04x-%04x-%04x-%04x%08x"), A, B >> 16, B & 0xFFFF, C >> 16, C & 0xFFFF, D);
+
 	case EGuidFormats::DigitsWithHyphensInBraces:
 		return FString::Printf(TEXT("{%08X-%04X-%04X-%04X-%04X%08X}"), A, B >> 16, B & 0xFFFF, C >> 16, C & 0xFFFF, D);
 
@@ -113,6 +117,21 @@ FString FGuid::ToString(EGuidFormats Format) const
 	}
 }
 
+void FGuid::AppendString(FAnsiStringBuilderBase& Builder) const
+{
+	Builder.Appendf("%08x-%04x-%04x-%04x-%04x%08x", A, B >> 16, B & 0xFFFF, C >> 16, C & 0xFFFF, D);
+}
+
+void FGuid::AppendString(FUtf8StringBuilderBase& Builder) const
+{
+	Builder.Appendf(UTF8TEXT("%08x-%04x-%04x-%04x-%04x%08x"), A, B >> 16, B & 0xFFFF, C >> 16, C & 0xFFFF, D);
+}
+
+void FGuid::AppendString(FWideStringBuilderBase& Builder) const
+{
+	Builder.Appendf(TEXT("%08x-%04x-%04x-%04x-%04x%08x"), A, B >> 16, B & 0xFFFF, C >> 16, C & 0xFFFF, D);
+}
+
 
 /* FGuid static interface
  *****************************************************************************/
diff --git a/Engine/Source/Runtime/Core/Public/Misc/Guid.h b/Engine/Source/Runtime/Core/Public/Misc/Guid.h
--- a/Engine/Source/Runtime/Core/Public/Misc/Guid.h
+++ b/Engine/Source/Runtime/Core/Public/Misc/Guid.h
@@ -5,6 +5,7 @@
 #include "CoreTypes.h"
 #include "Misc/AssertionMacros.h"
 #include "Misc/Crc.h"
+#include "Containers/StringFwd.h"
 #include "Containers/UnrealString.h"
 #include "Serialization/StructuredArchive.h"
 #include "Serialization/MemoryLayout.h"
@@ -34,6 +35,13 @@ enum class EGuidFormats
 	 */
 	DigitsWithHyphens,
 
+	/**
+	 * 32 digits separated by hyphens, in lowercase as described by RFC 4122.
+	 *
+	 * For example: bd048ce3-358b-46c5-8cee-627c719418f8
+	 */
+	DigitsWithHyphensLower,
+
 	/**
 	 * 32 digits separated by hyphens and enclosed in braces.
 	 *
@@ -306,6 +314,13 @@ public:
 	 */
 	CORE_API FString ToString(EGuidFormats Format) const;
 
+	/**
+	 * Appends this GUID to the string builder in DigitsWithHyphensLower format.
+	 */
+	CORE_API void AppendString(FAnsiStringBuilderBase& Builder) const;
+	CORE_API void AppendString(FUtf8StringBuilderBase& Builder) const;
+	CORE_API void AppendString(FWideStringBuilderBase& Builder) const;
+
 public:
 
 	/**
@@ -368,3 +383,10 @@ template<> struct TCanBulkSerialize<FGuid> { enum { Value = true }; };
 DECLARE_INTRINSIC_TYPE_LAYOUT(FGuid);
 
 template <> struct TIsPODType<FGuid> { enum { Value = true }; };
+
+template <typename CharType>
+inline TStringBuilderBase<CharType>& operator<<(TStringBuilderBase<CharType>& Builder, const FGuid& Value)
+{
+	Value.AppendString(Builder);
+	return Builder;
+}
