Bug fix, make the amount of address space reserved by the CachedOSVeryLargeAllocator configurable

diff --git a/Engine/Source/Runtime/Core/Private/HAL/Allocators/CachedOSVeryLargePageAllocator.cpp b/Engine/Source/Runtime/Core/Private/HAL/Allocators/CachedOSVeryLargePageAllocator.cpp
--- a/Engine/Source/Runtime/Core/Private/HAL/Allocators/CachedOSVeryLargePageAllocator.cpp
+++ b/Engine/Source/Runtime/Core/Private/HAL/Allocators/CachedOSVeryLargePageAllocator.cpp
@@ -93,6 +93,13 @@ void* FCachedOSVeryLargePageAllocator::Allocate(SIZE_T Size, uint32 AllocationHi
 					}
 					CachedFree -= SizeOfSubPage;
 				}
+				else
+				{
+					if (AllocationHint == FMemory::AllocationHints::SmallPool)
+					{
+						UE_CLOG(!ret, LogMemory, Fatal, TEXT("The FCachedOSVeryLargePageAllocator has run out of address space for SmallPool allocations, increase UE_VERYLARGEPAGEALLOCATOR_RESERVEDSIZEINGB for your platform!"));
+					}
+				}
 			}
 		}
 	}
diff --git a/Engine/Source/Runtime/Core/Public/HAL/Allocators/CachedOSVeryLargePageAllocator.h b/Engine/Source/Runtime/Core/Public/HAL/Allocators/CachedOSVeryLargePageAllocator.h
--- a/Engine/Source/Runtime/Core/Public/HAL/Allocators/CachedOSVeryLargePageAllocator.h
+++ b/Engine/Source/Runtime/Core/Public/HAL/Allocators/CachedOSVeryLargePageAllocator.h
@@ -24,6 +24,11 @@
 #define UE_VERYLARGEPAGEALLOCATOR_TAKEONALL64KBALLOCATIONS 0
 #endif
 
+#ifndef UE_VERYLARGEPAGEALLOCATOR_RESERVED_SIZE_IN_GB
+#define UE_VERYLARGEPAGEALLOCATOR_RESERVED_SIZE_IN_GB 2	//default to 2GB
+#endif
+
+
 #ifndef UE_VERYLARGEPAGEALLOCATOR_PAGESIZE_KB
 #define UE_VERYLARGEPAGEALLOCATOR_PAGESIZE_KB 4096	//default to 4MB
 #endif
@@ -31,9 +36,11 @@ class FCachedOSVeryLargePageAllocator
 {
 	// we make the address space twice as big as we need and use the 1st have for small pool allocations, the 2nd half is used for other allocations that are still == SizeOfSubPage
 #if UE_VERYLARGEPAGEALLOCATOR_TAKEONALL64KBALLOCATIONS
-	static const uint64 AddressSpaceToReserve = ((1024 * 1024) * 2048LL * 2LL);
+	static const uint64 AddressSpaceToReserve = ((1024LL * 1024LL * 1024LL) * UE_VERYLARGEPAGEALLOCATOR_RESERVED_SIZE_IN_GB * 2LL);
+	static const uint64 AddressSpaceToReserveSmall = AddressSpaceToReserve / 2;
 #else
-	static const uint64 AddressSpaceToReserve = ((1024 * 1024) * 2048LL);
+	static const uint64 AddressSpaceToReserve = ((1024 * 1024 * 1024LL) * UE_VERYLARGEPAGEALLOCATOR_RESERVED_SIZE_IN_GB);
+	static const uint64 AddressSpaceToReserveSmall = AddressSpaceToReserve;
 #endif
 	static const uint64 SizeOfLargePage = (UE_VERYLARGEPAGEALLOCATOR_PAGESIZE_KB * 1024);
 	static const uint64 SizeOfSubPage = (1024 * 64);
@@ -66,7 +73,7 @@ public:
 
 	FORCEINLINE bool IsPartOf(const void* Ptr)
 	{
-		if ((uintptr_t)Ptr >= AddressSpaceReserved && (uintptr_t)Ptr < AddressSpaceReservedEndSmallPool)
+		if (((uintptr_t)Ptr - AddressSpaceReserved) < AddressSpaceToReserveSmall)
 		{
 			return true;
 		}
