Guid: Added EGuidFormats param to the string builder functions

#rnx

diff --git a/Engine/Source/Runtime/Core/Private/Misc/Guid.cpp b/Engine/Source/Runtime/Core/Private/Misc/Guid.cpp
--- a/Engine/Source/Runtime/Core/Private/Misc/Guid.cpp
+++ b/Engine/Source/Runtime/Core/Private/Misc/Guid.cpp
@@ -1,6 +1,8 @@
 // Copyright Epic Games, Inc. All Rights Reserved.
 
 #include "Misc/Guid.h"
+#include "Algo/Replace.h"
+#include "Containers/ArrayView.h"
 #include "Misc/Parse.h"
 #include "Math/Int128.h"
 #include "UObject/PropertyPortFlags.h"
@@ -117,19 +119,93 @@ FString FGuid::ToString(EGuidFormats Format) const
 	}
 }
 
-void FGuid::AppendString(FAnsiStringBuilderBase& Builder) const
+void FGuid::AppendString(FAnsiStringBuilderBase& Builder, EGuidFormats Format) const
 {
-	Builder.Appendf("%08x-%04x-%04x-%04x-%04x%08x", A, B >> 16, B & 0xFFFF, C >> 16, C & 0xFFFF, D);
+	switch (Format)
+	{
+	case EGuidFormats::DigitsWithHyphens:
+		Builder.Appendf("%08X-%04X-%04X-%04X-%04X%08X", A, B >> 16, B & 0xFFFF, C >> 16, C & 0xFFFF, D);
+		return;
+
+	case EGuidFormats::DigitsWithHyphensLower:
+		Builder.Appendf("%08x-%04x-%04x-%04x-%04x%08x", A, B >> 16, B & 0xFFFF, C >> 16, C & 0xFFFF, D);
+		return;
+
+	case EGuidFormats::DigitsWithHyphensInBraces:
+		Builder.Appendf("{%08X-%04X-%04X-%04X-%04X%08X}", A, B >> 16, B & 0xFFFF, C >> 16, C & 0xFFFF, D);
+		return;
+
+	case EGuidFormats::DigitsWithHyphensInParentheses:
+		Builder.Appendf("(%08X-%04X-%04X-%04X-%04X%08X)", A, B >> 16, B & 0xFFFF, C >> 16, C & 0xFFFF, D);
+		return;
+
+	case EGuidFormats::HexValuesInBraces:
+		Builder.Appendf("{0x%08X,0x%04X,0x%04X,{0x%02X,0x%02X,0x%02X,0x%02X,0x%02X,0x%02X,0x%02X,0x%02X}}",
+			A, B >> 16, B & 0xFFFF, C >> 24, (C >> 16) & 0xFF, (C >> 8) & 0xFF, C & 0XFF, D >> 24, (D >> 16) & 0XFF, (D >> 8) & 0XFF, D & 0XFF);
+		return;
+
+	case EGuidFormats::UniqueObjectGuid:
+		Builder.Appendf("%08X-%08X-%08X-%08X", A, B, C, D);
+		return;
+
+	case EGuidFormats::Short:
+	{
+		const uint32 Data[] = {A,B,C,D};
+		int32 Index = Builder.AddUninitialized(24);
+		TArrayView<ANSICHAR> Result(Builder.GetData() + Index, 24);
+		int32 Len = FBase64::Encode(reinterpret_cast<const uint8*>(&Data), sizeof(Data), Result.GetData());
+
+		Algo::Replace(Result, '+', '-');
+		Algo::Replace(Result, '/', '_');
+
+		// Remove trailing '=' base64 padding
+		check(Len == 24);
+		Builder.RemoveSuffix(2);
+		return;
+	}
+
+	case EGuidFormats::Base36Encoded:
+	{
+		static const uint8 Alphabet[36] =
+		{
+			'0', '1', '2', '3', '4', '5', '6', '7', '8', '9', 'A', 'B', 'C', 'D', 'E', 'F',
+			'G', 'H', 'I', 'J', 'K', 'L', 'M', 'N', 'O', 'P', 'Q', 'R', 'S', 'T', 'U', 'V',
+			'W', 'X', 'Y', 'Z'
+		};
+
+		FUInt128 Value(A, B, C, D);
+
+		int32 Index = Builder.AddUninitialized(25);
+		TArrayView<ANSICHAR> Result(Builder.GetData() + Index, 25);
+
+		for (int32 I = 24; I >= 0; --I)
+		{
+			uint32 Remainder;
+			Value = Value.Divide(36, Remainder);
+			Result[I] = Alphabet[Remainder];
+		}
+
+		return;
+	}
+
+	default:
+		Builder.Appendf("%08X%08X%08X%08X", A, B, C, D);
+		return;
+	}
 }
 
-void FGuid::AppendString(FUtf8StringBuilderBase& Builder) const
+void FGuid::AppendString(FUtf8StringBuilderBase& Builder, EGuidFormats Format) const
 {
-	Builder.Appendf(UTF8TEXT("%08x-%04x-%04x-%04x-%04x%08x"), A, B >> 16, B & 0xFFFF, C >> 16, C & 0xFFFF, D);
+	TAnsiStringBuilder<70> AnsiBuilder;
+	AppendString(AnsiBuilder, Format);
+	Builder.AppendAnsi(AnsiBuilder);
 }
 
-void FGuid::AppendString(FWideStringBuilderBase& Builder) const
+void FGuid::AppendString(FWideStringBuilderBase& Builder, EGuidFormats Format) const
 {
-	Builder.Appendf(TEXT("%08x-%04x-%04x-%04x-%04x%08x"), A, B >> 16, B & 0xFFFF, C >> 16, C & 0xFFFF, D);
+	TAnsiStringBuilder<70> AnsiBuilder;
+	AppendString(AnsiBuilder, Format);
+	Builder.AppendAnsi(AnsiBuilder);
 }
 
 
@@ -265,7 +341,7 @@ bool FGuid::ParseExact(const FString& GuidString, EGuidFormats Format, FGuid& Ou
 	{
 		NormalizedGuidString = GuidString;
 	}
-	else if (Format == EGuidFormats::DigitsWithHyphens)
+	else if (Format == EGuidFormats::DigitsWithHyphens || Format == EGuidFormats::DigitsWithHyphensLower)
 	{
 		if ((GuidString[8] != TCHAR('-')) ||
 			(GuidString[13] != TCHAR('-')) ||
diff --git a/Engine/Source/Runtime/Core/Public/Misc/Guid.h b/Engine/Source/Runtime/Core/Public/Misc/Guid.h
--- a/Engine/Source/Runtime/Core/Public/Misc/Guid.h
+++ b/Engine/Source/Runtime/Core/Public/Misc/Guid.h
@@ -315,11 +315,11 @@ public:
 	CORE_API FString ToString(EGuidFormats Format) const;
 
 	/**
-	 * Appends this GUID to the string builder in DigitsWithHyphensLower format.
+	 * Appends this GUID to the string builder using the specified format.
 	 */
-	CORE_API void AppendString(FAnsiStringBuilderBase& Builder) const;
-	CORE_API void AppendString(FUtf8StringBuilderBase& Builder) const;
-	CORE_API void AppendString(FWideStringBuilderBase& Builder) const;
+	CORE_API void AppendString(FAnsiStringBuilderBase& Builder, EGuidFormats Format = EGuidFormats::DigitsWithHyphensLower) const;
+	CORE_API void AppendString(FUtf8StringBuilderBase& Builder, EGuidFormats Format = EGuidFormats::DigitsWithHyphensLower) const;
+	CORE_API void AppendString(FWideStringBuilderBase& Builder, EGuidFormats Format = EGuidFormats::DigitsWithHyphensLower) const;
 
 public:
 
