fixed non-unity compilation: https://horde.devtools.epicgames.com/job/60d11e0478c3b00001ed36fd?step=b476

diff --git a/Engine/Source/Runtime/Core/Private/Async/Fundamental/ReserveScheduler.cpp b/Engine/Source/Runtime/Core/Private/Async/Fundamental/ReserveScheduler.cpp
--- a/Engine/Source/Runtime/Core/Private/Async/Fundamental/ReserveScheduler.cpp
+++ b/Engine/Source/Runtime/Core/Private/Async/Fundamental/ReserveScheduler.cpp
@@ -1,6 +1,7 @@
 // Copyright Epic Games, Inc. All Rights Reserved.
 
 #include "Async/Fundamental/ReserveScheduler.h"
+#include "Misc/ScopeLock.h"
 
 namespace LowLevelTasks
 {
@@ -56,8 +57,6 @@ bool FReserveScheduler::DoReserveWorkUntil(FConditional&& Condition)
 
 void FReserveScheduler::StartWorkers(FScheduler& MainScheduler, uint32 NumWorkers, EThreadPriority WorkerPriority, bool bIsForkable)
 {
-	TaskTrace::Init();
-
 	if (NumWorkers == 0)
 	{
 		NumWorkers = FMath::Min(FPlatformMisc::NumberOfWorkerThreadsToSpawn(), 64);
diff --git a/Engine/Source/Runtime/Core/Private/Async/Fundamental/Scheduler.cpp b/Engine/Source/Runtime/Core/Private/Async/Fundamental/Scheduler.cpp
--- a/Engine/Source/Runtime/Core/Private/Async/Fundamental/Scheduler.cpp
+++ b/Engine/Source/Runtime/Core/Private/Async/Fundamental/Scheduler.cpp
@@ -81,8 +81,6 @@ namespace LowLevelTasks
 
 	void FScheduler::StartWorkers(uint32 NumForegroundWorkers, uint32 NumBackgroundWorkers, EThreadPriority WorkerPriority,  EThreadPriority BackgroundPriority, bool bIsForkable)
 	{
-		TaskTrace::Init();
-
 		if (NumForegroundWorkers == 0 && NumBackgroundWorkers == 0)
 		{
 			NumForegroundWorkers = FMath::Max<int32>(1, FMath::Min<int32>(2, FPlatformMisc::NumberOfWorkerThreadsToSpawn() - 1));
diff --git a/Engine/Source/Runtime/Core/Private/Async/TaskGraph.cpp b/Engine/Source/Runtime/Core/Private/Async/TaskGraph.cpp
--- a/Engine/Source/Runtime/Core/Private/Async/TaskGraph.cpp
+++ b/Engine/Source/Runtime/Core/Private/Async/TaskGraph.cpp
@@ -1867,6 +1867,8 @@ public:
 	FTaskGraphCompatibilityImplementation(int32 InNumWorkerThreads) 
 		: NumWorkerThreads(FForkProcessHelper::IsForkedMultithreadInstance() ? CVar_ForkedProcess_MaxWorkerThreads : InNumWorkerThreads)
 	{
+		TaskTrace::Init();
+
 		if (FTaskGraphInterface::IsMultithread())
 		{
 			if (NumWorkerThreads <= 3)
diff --git a/Engine/Source/Runtime/Core/Public/Async/Fundamental/Task.h b/Engine/Source/Runtime/Core/Public/Async/Fundamental/Task.h
--- a/Engine/Source/Runtime/Core/Public/Async/Fundamental/Task.h
+++ b/Engine/Source/Runtime/Core/Public/Async/Fundamental/Task.h
@@ -259,7 +259,7 @@ namespace LowLevelTasks
 	{
 		checkf(IsCompleted(), TEXT("State: %d"), PackedData.load(std::memory_order_relaxed).GetState());
 		checkSlow(!Runnable.IsSet());
-		Runnable = [this, LocalRunnable = Forward<TRunnable>(InRunnable)]() -> bool
+		Runnable = [this, LocalRunnable = Forward<TRunnable>(InRunnable)]() mutable -> bool
 		{
 			checkSlow(IsScheduled() || IsCanceled());
 			FPackedData LocalPackedData = PackedData.load(std::memory_order_relaxed);
