UE_APP_NAME is not unique if the build environment is "shared" (such is the case for installed builds). This makes it difficult to identify traces from monolithic builds. Use the project name set by UBT instead if available.

#rnx

diff --git a/Engine/Source/Runtime/Core/Private/ProfilingDebugging/TraceAuxiliary.cpp b/Engine/Source/Runtime/Core/Private/ProfilingDebugging/TraceAuxiliary.cpp
--- a/Engine/Source/Runtime/Core/Private/ProfilingDebugging/TraceAuxiliary.cpp
+++ b/Engine/Source/Runtime/Core/Private/ProfilingDebugging/TraceAuxiliary.cpp
@@ -578,12 +578,21 @@ void FTraceAuxiliary::Initialize(const TCHAR* CommandLine)
 #endif
 
 #if UE_TRACE_ENABLED
+	const TCHAR* AppName = TEXT(UE_APP_NAME);
+#if IS_MONOLITHIC && !IS_PROGRAM
+	extern TCHAR GInternalProjectName[];
+	if (GInternalProjectName[0] != '\0')
+	{
+		AppName = GInternalProjectName;
+	}
+#endif
+
 	// Trace out information about this session. This is done before initialisation
 	// so that it is always sent (all channels are enabled prior to initialisation)
 	const TCHAR* BranchName = BuildSettings::GetBranchName();
 	const TCHAR* BuildVersion = BuildSettings::GetBuildVersion();
 	constexpr uint32 PlatformLen = UE_ARRAY_COUNT(PREPROCESSOR_TO_STRING(UBT_COMPILED_PLATFORM)) - 1;
-	constexpr uint32 AppNameLen = UE_ARRAY_COUNT(UE_APP_NAME) - 1;
+	const uint32 AppNameLen = FCString::Strlen(AppName);
 	const uint32 CommandLineLen = FCString::Strlen(CommandLine);
 	const uint32 BranchNameLen = FCString::Strlen(BranchName);
 	const uint32 BuildVersionLen = FCString::Strlen(BuildVersion);
@@ -595,7 +604,7 @@ void FTraceAuxiliary::Initialize(const TCHAR* CommandLine)
 		(BuildVersionLen * sizeof(TCHAR));
 	UE_TRACE_LOG(Diagnostics, Session2, UE::Trace::TraceLogChannel, DataSize)
 		<< Session2.Platform(PREPROCESSOR_TO_STRING(UBT_COMPILED_PLATFORM), PlatformLen)
-		<< Session2.AppName(UE_APP_NAME, AppNameLen)
+		<< Session2.AppName(AppName, AppNameLen)
 		<< Session2.CommandLine(CommandLine, CommandLineLen)
 		<< Session2.Branch(BranchName, BranchNameLen)
 		<< Session2.BuildVersion(BuildVersion, BuildVersionLen)
