Exclude mimalloc.c from static analysis
Fix typo where double slash were used

#rnx

diff --git a/Engine/Source/Runtime/Core/Core.Build.cs b/Engine/Source/Runtime/Core/Core.Build.cs
--- a/Engine/Source/Runtime/Core/Core.Build.cs
+++ b/Engine/Source/Runtime/Core/Core.Build.cs
@@ -63,8 +63,12 @@ public class Core : ModuleRules
 				"IntelVTune"
 				);
 
-			PublicSystemIncludePaths.Add(Path.Combine(Target.UEThirdPartySourceDirectory, "mimalloc//include"));
-			PrivateDefinitions.Add("PLATFORM_BUILDS_MIMALLOC=1");
+			// We do not want the static analyzer to run on thirdparty code
+			if (Target.WindowsPlatform.StaticAnalyzer == WindowsStaticAnalyzer.None) 
+			{
+				PublicSystemIncludePaths.Add(Path.Combine(Target.UEThirdPartySourceDirectory, "mimalloc/include"));
+				PrivateDefinitions.Add("PLATFORM_BUILDS_MIMALLOC=1");
+			}
 
 			if (Target.WindowsPlatform.bUseBundledDbgHelp)
 			{
@@ -104,7 +108,7 @@ public class Core : ModuleRules
 				);
 			PublicFrameworks.AddRange(new string[] { "Cocoa", "Carbon", "IOKit", "Security" });
 
-			PublicSystemIncludePaths.Add(Path.Combine(Target.UEThirdPartySourceDirectory, "mimalloc//include"));
+			PublicSystemIncludePaths.Add(Path.Combine(Target.UEThirdPartySourceDirectory, "mimalloc/include"));
 			PrivateDefinitions.Add("PLATFORM_BUILDS_MIMALLOC=1");
 
 			if (Target.bBuildEditor == true)
@@ -156,7 +160,7 @@ public class Core : ModuleRules
 			// Core uses dlopen()
 			PublicSystemLibraries.Add("dl");
 
-			PublicSystemIncludePaths.Add(Path.Combine(Target.UEThirdPartySourceDirectory, "mimalloc//include"));
+			PublicSystemIncludePaths.Add(Path.Combine(Target.UEThirdPartySourceDirectory, "mimalloc/include"));
 			PrivateDefinitions.Add("PLATFORM_BUILDS_MIMALLOC=1");
 		}
 
diff --git a/Engine/Source/Runtime/Core/Private/HAL/MallocMimalloc.cpp b/Engine/Source/Runtime/Core/Private/HAL/MallocMimalloc.cpp
--- a/Engine/Source/Runtime/Core/Private/HAL/MallocMimalloc.cpp
+++ b/Engine/Source/Runtime/Core/Private/HAL/MallocMimalloc.cpp
@@ -9,7 +9,7 @@
 #include "HAL/UnrealMemory.h"
 
 // Only use for supported platforms
-#if PLATFORM_SUPPORTS_MIMALLOC && MIMALLOC_ALLOCATOR_ALLOWED
+#if PLATFORM_SUPPORTS_MIMALLOC && MIMALLOC_ALLOCATOR_ALLOWED && PLATFORM_BUILDS_MIMALLOC
 
 /** Value we fill a memory block with after it is free, in UE_BUILD_DEBUG **/
 #define DEBUG_FILL_FREED (0xdd)
diff --git a/Engine/Source/Runtime/Core/Private/Mac/MacPlatformMemory.cpp b/Engine/Source/Runtime/Core/Private/Mac/MacPlatformMemory.cpp
--- a/Engine/Source/Runtime/Core/Private/Mac/MacPlatformMemory.cpp
+++ b/Engine/Source/Runtime/Core/Private/Mac/MacPlatformMemory.cpp
@@ -136,7 +136,7 @@ FMalloc* FMacPlatformMemory::BaseAllocator()
 	case EMemoryAllocatorToUse::TBB:
 		return new FMallocTBB();
 #endif
-#if PLATFORM_SUPPORTS_MIMALLOC && MIMALLOC_ALLOCATOR_ALLOWED
+#if PLATFORM_SUPPORTS_MIMALLOC && MIMALLOC_ALLOCATOR_ALLOWED && PLATFORM_BUILDS_MIMALLOC
 	case EMemoryAllocatorToUse::Mimalloc:
 		return new FMallocMimalloc();
 #endif
diff --git a/Engine/Source/Runtime/Core/Private/Unix/UnixPlatformMemory.cpp b/Engine/Source/Runtime/Core/Private/Unix/UnixPlatformMemory.cpp
--- a/Engine/Source/Runtime/Core/Private/Unix/UnixPlatformMemory.cpp
+++ b/Engine/Source/Runtime/Core/Private/Unix/UnixPlatformMemory.cpp
@@ -295,7 +295,7 @@ class FMalloc* FUnixPlatformMemory::BaseAllocator()
 		break;
 #endif // PLATFORM_SUPPORTS_JEMALLOC
 
-#if PLATFORM_SUPPORTS_MIMALLOC && MIMALLOC_ALLOCATOR_ALLOWED
+#if PLATFORM_SUPPORTS_MIMALLOC && MIMALLOC_ALLOCATOR_ALLOWED && PLATFORM_BUILDS_MIMALLOC
 	case EMemoryAllocatorToUse::Mimalloc:
 		Allocator = new FMallocMimalloc();
 		break;
diff --git a/Engine/Source/Runtime/Core/Private/Windows/WindowsPlatformMemory.cpp b/Engine/Source/Runtime/Core/Private/Windows/WindowsPlatformMemory.cpp
--- a/Engine/Source/Runtime/Core/Private/Windows/WindowsPlatformMemory.cpp
+++ b/Engine/Source/Runtime/Core/Private/Windows/WindowsPlatformMemory.cpp
@@ -187,7 +187,7 @@ FMalloc* FWindowsPlatformMemory::BaseAllocator()
 	case EMemoryAllocatorToUse::TBB:
 		return new FMallocTBB();
 #endif
-#if MIMALLOC_ALLOCATOR_ALLOWED && PLATFORM_SUPPORTS_MIMALLOC
+#if MIMALLOC_ALLOCATOR_ALLOWED && PLATFORM_SUPPORTS_MIMALLOC && PLATFORM_BUILDS_MIMALLOC
 	case EMemoryAllocatorToUse::Mimalloc:
 		return new FMallocMimalloc();
 #endif
diff --git a/Engine/Source/Runtime/Core/Public/HAL/MallocMimalloc.h b/Engine/Source/Runtime/Core/Public/HAL/MallocMimalloc.h
--- a/Engine/Source/Runtime/Core/Public/HAL/MallocMimalloc.h
+++ b/Engine/Source/Runtime/Core/Public/HAL/MallocMimalloc.h
@@ -5,11 +5,15 @@
 #include "HAL/PlatformMemory.h"
 #include "HAL/MemoryBase.h"
 
+#if !defined(PLATFORM_BUILDS_MIMALLOC)
+#	define PLATFORM_BUILDS_MIMALLOC 0
+#endif
+
 #if !defined(PLATFORM_SUPPORTS_MIMALLOC)
 #	define PLATFORM_SUPPORTS_MIMALLOC 0
 #endif
 
-#if PLATFORM_SUPPORTS_MIMALLOC && MIMALLOC_ALLOCATOR_ALLOWED
+#if PLATFORM_SUPPORTS_MIMALLOC && MIMALLOC_ALLOCATOR_ALLOWED && PLATFORM_BUILDS_MIMALLOC
 
 /**
  * TBB 64-bit scalable memory allocator.
