ESPMode underlying type added to allow proper forward declaration.

diff --git a/Engine/Source/Runtime/Core/Public/Delegates/DelegateInstancesImpl.h b/Engine/Source/Runtime/Core/Public/Delegates/DelegateInstancesImpl.h
--- a/Engine/Source/Runtime/Core/Public/Delegates/DelegateInstancesImpl.h
+++ b/Engine/Source/Runtime/Core/Public/Delegates/DelegateInstancesImpl.h
@@ -18,7 +18,7 @@
 
 class FDelegateBase;
 class FDelegateHandle;
-enum class ESPMode;
+enum class ESPMode : uint8;
 
 namespace UE4Delegates_Private
 {
diff --git a/Engine/Source/Runtime/Core/Public/Templates/SharedPointerInternals.h b/Engine/Source/Runtime/Core/Public/Templates/SharedPointerInternals.h
--- a/Engine/Source/Runtime/Core/Public/Templates/SharedPointerInternals.h
+++ b/Engine/Source/Runtime/Core/Public/Templates/SharedPointerInternals.h
@@ -24,7 +24,7 @@
  * ESPMode is used select between either 'fast' or 'thread safe' shared pointer types.
  * This is only used by templates at compile time to generate one code path or another.
  */
-enum class ESPMode
+enum class ESPMode : uint8
 {
 	/** Forced to be not thread-safe. */
 	NotThreadSafe = 0,
diff --git a/Engine/Source/Runtime/Core/Public/Templates/SharedPointerTesting.inl b/Engine/Source/Runtime/Core/Public/Templates/SharedPointerTesting.inl
--- a/Engine/Source/Runtime/Core/Public/Templates/SharedPointerTesting.inl
+++ b/Engine/Source/Runtime/Core/Public/Templates/SharedPointerTesting.inl
@@ -6,7 +6,7 @@
 /** Toggle this define to enable shared pointer testing features */
 #define WITH_SHARED_POINTER_TESTS 0 && !(UE_BUILD_SHIPPING || UE_BUILD_TEST)
 
-enum class ESPMode;
+enum class ESPMode : uint8;
 
 template<typename KeyType,typename ValueType,typename SetAllocator ,typename KeyFuncs > class TMap;
 
