Don't trace log messages unless the Log trace channel is enabled. This was incorrectly marking specs as sent even though they weren't causing some log trace events to arrive analysis side without a format string (aka spec).

#rnx

diff --git a/Engine/Source/Runtime/Core/Private/Logging/LogTrace.cpp b/Engine/Source/Runtime/Core/Private/Logging/LogTrace.cpp
--- a/Engine/Source/Runtime/Core/Private/Logging/LogTrace.cpp
+++ b/Engine/Source/Runtime/Core/Private/Logging/LogTrace.cpp
@@ -8,7 +8,7 @@
 #include "HAL/PlatformTime.h"
 #include "HAL/PlatformTLS.h"
 
-UE_TRACE_CHANNEL(LogChannel)
+UE_TRACE_CHANNEL_DEFINE(LogChannel)
 
 UE_TRACE_EVENT_BEGIN(Logging, LogCategory, NoSync|Important)
 	UE_TRACE_EVENT_FIELD(const void*, CategoryPointer)
diff --git a/Engine/Source/Runtime/Core/Public/Logging/LogTrace.h b/Engine/Source/Runtime/Core/Public/Logging/LogTrace.h
--- a/Engine/Source/Runtime/Core/Public/Logging/LogTrace.h
+++ b/Engine/Source/Runtime/Core/Public/Logging/LogTrace.h
@@ -16,6 +16,8 @@
 #if LOGTRACE_ENABLED
 #include "ProfilingDebugging/FormatArgsTrace.h"
 
+UE_TRACE_CHANNEL_EXTERN(LogChannel, CORE_API)
+
 struct FLogCategoryBase;
 
 struct FLogTrace
@@ -42,13 +44,16 @@ private:
 	FLogTrace::OutputLogCategory(Category, Name, DefaultVerbosity);
 
 #define TRACE_LOG_MESSAGE(Category, Verbosity, Format, ...) \
-	static bool PREPROCESSOR_JOIN(__LogPoint, __LINE__); \
-	if (!PREPROCESSOR_JOIN(__LogPoint, __LINE__)) \
+	if (UE_TRACE_CHANNELEXPR_IS_ENABLED(LogChannel)) \
 	{ \
-		FLogTrace::OutputLogMessageSpec(&PREPROCESSOR_JOIN(__LogPoint, __LINE__), &Category, ELogVerbosity::Verbosity, __FILE__, __LINE__, Format); \
-		PREPROCESSOR_JOIN(__LogPoint, __LINE__) = true; \
-	} \
-	FLogTrace::OutputLogMessage(&PREPROCESSOR_JOIN(__LogPoint, __LINE__), ##__VA_ARGS__);
+		static bool PREPROCESSOR_JOIN(__LogPoint, __LINE__); \
+		if (!PREPROCESSOR_JOIN(__LogPoint, __LINE__)) \
+		{ \
+			FLogTrace::OutputLogMessageSpec(&PREPROCESSOR_JOIN(__LogPoint, __LINE__), &Category, ELogVerbosity::Verbosity, __FILE__, __LINE__, Format); \
+			PREPROCESSOR_JOIN(__LogPoint, __LINE__) = true; \
+		} \
+		FLogTrace::OutputLogMessage(&PREPROCESSOR_JOIN(__LogPoint, __LINE__), ##__VA_ARGS__); \
+	}
 
 #else
 #define TRACE_LOG_CATEGORY(Category, Name, DefaultVerbosity)
