TSharedPtr now uses std::atomic and more appropriate (more optimal) memory ordering in ESPMode::ThreadSafe mode.

diff --git a/Engine/Source/Runtime/Core/Public/Templates/SharedPointer.h b/Engine/Source/Runtime/Core/Public/Templates/SharedPointer.h
--- a/Engine/Source/Runtime/Core/Public/Templates/SharedPointer.h
+++ b/Engine/Source/Runtime/Core/Public/Templates/SharedPointer.h
@@ -130,7 +130,7 @@ namespace UE4SharedPointer_Private
 	// Needed to work around an Android compiler bug - we need to construct a TSharedRef
 	// from MakeShared without making MakeShared a friend in order to access the private constructor.
 	template <typename ObjectType, ESPMode Mode>
-	FORCEINLINE TSharedRef<ObjectType, Mode> MakeSharedRef(ObjectType* InObject, SharedPointerInternals::FReferenceControllerBase* InSharedReferenceCount)
+	FORCEINLINE TSharedRef<ObjectType, Mode> MakeSharedRef(ObjectType* InObject, SharedPointerInternals::TReferenceControllerBase<Mode>* InSharedReferenceCount)
 	{
 		return TSharedRef<ObjectType, Mode>(InObject, InSharedReferenceCount);
 	}
@@ -175,7 +175,7 @@ public:
 	>
 	FORCEINLINE explicit TSharedRef( OtherType* InObject )
 		: Object( InObject )
-		, SharedReferenceCount( SharedPointerInternals::NewDefaultReferenceController( InObject ) )
+		, SharedReferenceCount( SharedPointerInternals::NewDefaultReferenceController< Mode >( InObject ) )
 	{
 		UE_TSHAREDPTR_STATIC_ASSERT_VALID_MODE(ObjectType, Mode)
 
@@ -195,7 +195,7 @@ public:
 	>
 	FORCEINLINE TSharedRef( OtherType* InObject, DeleterType&& InDeleter )
 		: Object( InObject )
-		, SharedReferenceCount( SharedPointerInternals::NewCustomReferenceController( InObject, Forward< DeleterType >( InDeleter ) ) )
+		, SharedReferenceCount( SharedPointerInternals::NewCustomReferenceController< Mode >( InObject, Forward< DeleterType >( InDeleter ) ) )
 	{
 		UE_TSHAREDPTR_STATIC_ASSERT_VALID_MODE(ObjectType, Mode)
 
@@ -209,7 +209,7 @@ public:
 	 */
 	TSharedRef()
 		: Object(new ObjectType())
-		, SharedReferenceCount(SharedPointerInternals::NewDefaultReferenceController(Object))
+		, SharedReferenceCount(SharedPointerInternals::NewDefaultReferenceController<Mode>(Object))
 	{
 		EnsureRetrievingVTablePtrDuringCtor(TEXT("TSharedRef()"));
 		Init(Object);
@@ -226,9 +226,63 @@ public:
 		typename OtherType,
 		typename = decltype(ImplicitConv<ObjectType*>((OtherType*)nullptr))
 	>
-	FORCEINLINE TSharedRef( SharedPointerInternals::FRawPtrProxy< OtherType > const& InRawPtrProxy )
+	FORCEINLINE TSharedRef( SharedPointerInternals::TRawPtrProxy< OtherType > const& InRawPtrProxy )
 		: Object( InRawPtrProxy.Object )
-		, SharedReferenceCount( InRawPtrProxy.ReferenceController )
+		, SharedReferenceCount( SharedPointerInternals::NewDefaultReferenceController< Mode >( InRawPtrProxy.Object ) )
+	{
+		UE_TSHAREDPTR_STATIC_ASSERT_VALID_MODE(ObjectType, Mode)
+
+		// If the following assert goes off, it means a TSharedRef was initialized from a nullptr object pointer.
+		// Shared references must never be nullptr, so either pass a valid object or consider using TSharedPtr instead.
+		check( InRawPtrProxy.Object != nullptr );
+
+		// If the object happens to be derived from TSharedFromThis, the following method
+		// will prime the object with a weak pointer to itself.
+		SharedPointerInternals::EnableSharedFromThis( this, InRawPtrProxy.Object, InRawPtrProxy.Object );
+	}
+
+	/**
+	 * Constructs a shared reference using a proxy reference to a raw pointer. (See MakeShareable())
+	 * Must not be nullptr.
+	 *
+	 * @param  InRawPtrProxy  Proxy raw pointer that contains the object that the new shared reference will reference
+	 */
+	// NOTE: The following is an Unreal extension to standard shared_ptr behavior
+	template <
+		typename OtherType,
+		typename DeleterType,
+		typename = decltype(ImplicitConv<ObjectType*>((OtherType*)nullptr))
+	>
+	FORCEINLINE TSharedRef( SharedPointerInternals::TRawPtrProxyWithDeleter< OtherType, DeleterType > const& InRawPtrProxy )
+		: Object( InRawPtrProxy.Object )
+		, SharedReferenceCount( SharedPointerInternals::NewCustomReferenceController< Mode >( InRawPtrProxy.Object, InRawPtrProxy.Deleter ) )
+	{
+		UE_TSHAREDPTR_STATIC_ASSERT_VALID_MODE(ObjectType, Mode)
+
+		// If the following assert goes off, it means a TSharedRef was initialized from a nullptr object pointer.
+		// Shared references must never be nullptr, so either pass a valid object or consider using TSharedPtr instead.
+		check( InRawPtrProxy.Object != nullptr );
+
+		// If the object happens to be derived from TSharedFromThis, the following method
+		// will prime the object with a weak pointer to itself.
+		SharedPointerInternals::EnableSharedFromThis( this, InRawPtrProxy.Object, InRawPtrProxy.Object );
+	}
+
+	/**
+	 * Constructs a shared reference using a proxy reference to a raw pointer. (See MakeShareable())
+	 * Must not be nullptr.
+	 *
+	 * @param  InRawPtrProxy  Proxy raw pointer that contains the object that the new shared reference will reference
+	 */
+	// NOTE: The following is an Unreal extension to standard shared_ptr behavior
+	template <
+		typename OtherType,
+		typename DeleterType,
+		typename = decltype(ImplicitConv<ObjectType*>((OtherType*)nullptr))
+	>
+	FORCEINLINE TSharedRef( SharedPointerInternals::TRawPtrProxyWithDeleter< OtherType, DeleterType >&& InRawPtrProxy )
+		: Object( InRawPtrProxy.Object )
+		, SharedReferenceCount( SharedPointerInternals::NewCustomReferenceController< Mode >( InRawPtrProxy.Object, MoveTemp( InRawPtrProxy.Deleter ) ) )
 	{
 		UE_TSHAREDPTR_STATIC_ASSERT_VALID_MODE(ObjectType, Mode)
 
@@ -350,7 +404,30 @@ public:
 		typename OtherType,
 		typename = decltype(ImplicitConv<ObjectType*>((OtherType*)nullptr))
 	>
-	FORCEINLINE TSharedRef& operator=( SharedPointerInternals::FRawPtrProxy< OtherType > const& InRawPtrProxy )
+	FORCEINLINE TSharedRef& operator=( SharedPointerInternals::TRawPtrProxy< OtherType > const& InRawPtrProxy )
+	{
+		// If the following assert goes off, it means a TSharedRef was initialized from a nullptr object pointer.
+		// Shared references must never be nullptr, so either pass a valid object or consider using TSharedPtr instead.
+		check( InRawPtrProxy.Object != nullptr );
+
+		*this = TSharedRef< ObjectType, Mode >( InRawPtrProxy );
+		return *this;
+	}
+
+	/**
+	 * Assignment operator replaces this shared reference with the specified shared reference.  The object
+	 * currently referenced by this shared reference will no longer be referenced and will be deleted if
+	 * there are no other referencers.  Must not be nullptr.
+	 *
+	 * @param  InRawPtrProxy  Proxy object used to assign the object (see MakeShareable helper function)
+	 */
+	// NOTE: The following is an Unreal extension to standard shared_ptr behavior
+	template <
+		typename OtherType,
+		typename DeleterType,
+		typename = decltype(ImplicitConv<ObjectType*>((OtherType*)nullptr))
+	>
+	FORCEINLINE TSharedRef& operator=( SharedPointerInternals::TRawPtrProxyWithDeleter< OtherType, DeleterType > const& InRawPtrProxy )
 	{
 		// If the following assert goes off, it means a TSharedRef was initialized from a nullptr object pointer.
 		// Shared references must never be nullptr, so either pass a valid object or consider using TSharedPtr instead.
@@ -360,6 +437,29 @@ public:
 		return *this;
 	}
 
+	/**
+	 * Assignment operator replaces this shared reference with the specified shared reference.  The object
+	 * currently referenced by this shared reference will no longer be referenced and will be deleted if
+	 * there are no other referencers.  Must not be nullptr.
+	 *
+	 * @param  InRawPtrProxy  Proxy object used to assign the object (see MakeShareable helper function)
+	 */
+	// NOTE: The following is an Unreal extension to standard shared_ptr behavior
+	template <
+		typename OtherType,
+		typename DeleterType,
+		typename = decltype(ImplicitConv<ObjectType*>((OtherType*)nullptr))
+	>
+	FORCEINLINE TSharedRef& operator=( SharedPointerInternals::TRawPtrProxyWithDeleter< OtherType, DeleterType >&& InRawPtrProxy )
+	{
+		// If the following assert goes off, it means a TSharedRef was initialized from a nullptr object pointer.
+		// Shared references must never be nullptr, so either pass a valid object or consider using TSharedPtr instead.
+		check( InRawPtrProxy.Object != nullptr );
+
+		*this = TSharedRef< ObjectType, Mode >( MoveTemp( InRawPtrProxy ) );
+		return *this;
+	}
+
 	/**
 	 * Returns a C++ reference to the object this shared reference is referencing
 	 *
@@ -511,13 +611,13 @@ private:
 		#pragma warning(disable : 4396) // warning: the inline specifier cannot be used when a friend declaration refers to a specialization of a function template
 	#endif
 
-	friend TSharedRef UE4SharedPointer_Private::MakeSharedRef<ObjectType, Mode>(ObjectType* InObject, SharedPointerInternals::FReferenceControllerBase* InSharedReferenceCount);
+	friend TSharedRef UE4SharedPointer_Private::MakeSharedRef<ObjectType, Mode>(ObjectType* InObject, SharedPointerInternals::TReferenceControllerBase<Mode>* InSharedReferenceCount);
 
 	#ifdef _MSC_VER
 		#pragma warning(pop)
 	#endif
 
-	FORCEINLINE explicit TSharedRef(ObjectType* InObject, SharedPointerInternals::FReferenceControllerBase* InSharedReferenceCount)
+	FORCEINLINE explicit TSharedRef(ObjectType* InObject, SharedPointerInternals::TReferenceControllerBase<Mode>* InSharedReferenceCount)
 		: Object(InObject)
 		, SharedReferenceCount(InSharedReferenceCount)
 	{
@@ -585,7 +685,7 @@ public:
 	>
 	FORCEINLINE explicit TSharedPtr( OtherType* InObject )
 		: Object( InObject )
-		, SharedReferenceCount( SharedPointerInternals::NewDefaultReferenceController( InObject ) )
+		, SharedReferenceCount( SharedPointerInternals::NewDefaultReferenceController< Mode >( InObject ) )
 	{
 		UE_TSHAREDPTR_STATIC_ASSERT_VALID_MODE(ObjectType, Mode)
 
@@ -608,7 +708,7 @@ public:
 	>
 	FORCEINLINE TSharedPtr( OtherType* InObject, DeleterType&& InDeleter )
 		: Object( InObject )
-		, SharedReferenceCount( SharedPointerInternals::NewCustomReferenceController( InObject, Forward< DeleterType >( InDeleter ) ) )
+		, SharedReferenceCount( SharedPointerInternals::NewCustomReferenceController< Mode >( InObject, Forward< DeleterType >( InDeleter ) ) )
 	{
 		UE_TSHAREDPTR_STATIC_ASSERT_VALID_MODE(ObjectType, Mode)
 
@@ -627,9 +727,53 @@ public:
 		typename OtherType,
 		typename = decltype(ImplicitConv<ObjectType*>((OtherType*)nullptr))
 	>
-	FORCEINLINE TSharedPtr( SharedPointerInternals::FRawPtrProxy< OtherType > const& InRawPtrProxy )
+	FORCEINLINE TSharedPtr( SharedPointerInternals::TRawPtrProxy< OtherType > const& InRawPtrProxy )
+		: Object( InRawPtrProxy.Object )
+		, SharedReferenceCount( SharedPointerInternals::NewDefaultReferenceController< Mode >( InRawPtrProxy.Object ) )
+	{
+		UE_TSHAREDPTR_STATIC_ASSERT_VALID_MODE(ObjectType, Mode)
+
+		// If the object happens to be derived from TSharedFromThis, the following method
+		// will prime the object with a weak pointer to itself.
+		SharedPointerInternals::EnableSharedFromThis( this, InRawPtrProxy.Object, InRawPtrProxy.Object );
+	}
+
+	/**
+	 * Constructs a shared pointer using a proxy reference to a raw pointer. (See MakeShareable())
+	 *
+	 * @param  InRawPtrProxy  Proxy raw pointer that contains the object that the new shared pointer will reference
+	 */
+	// NOTE: The following is an Unreal extension to standard shared_ptr behavior
+	template <
+		typename OtherType,
+		typename DeleterType,
+		typename = decltype(ImplicitConv<ObjectType*>((OtherType*)nullptr))
+	>
+	FORCEINLINE TSharedPtr( SharedPointerInternals::TRawPtrProxyWithDeleter< OtherType, DeleterType > const& InRawPtrProxy )
+		: Object( InRawPtrProxy.Object )
+		, SharedReferenceCount( SharedPointerInternals::NewCustomReferenceController< Mode >( InRawPtrProxy.Object, InRawPtrProxy.Deleter ) )
+	{
+		UE_TSHAREDPTR_STATIC_ASSERT_VALID_MODE(ObjectType, Mode)
+
+		// If the object happens to be derived from TSharedFromThis, the following method
+		// will prime the object with a weak pointer to itself.
+		SharedPointerInternals::EnableSharedFromThis( this, InRawPtrProxy.Object, InRawPtrProxy.Object );
+	}
+
+	/**
+	 * Constructs a shared pointer using a proxy reference to a raw pointer. (See MakeShareable())
+	 *
+	 * @param  InRawPtrProxy  Proxy raw pointer that contains the object that the new shared pointer will reference
+	 */
+	// NOTE: The following is an Unreal extension to standard shared_ptr behavior
+	template <
+		typename OtherType,
+		typename DeleterType,
+		typename = decltype(ImplicitConv<ObjectType*>((OtherType*)nullptr))
+	>
+	FORCEINLINE TSharedPtr( SharedPointerInternals::TRawPtrProxyWithDeleter< OtherType, DeleterType >&& InRawPtrProxy )
 		: Object( InRawPtrProxy.Object )
-		, SharedReferenceCount( InRawPtrProxy.ReferenceController )
+		, SharedReferenceCount( SharedPointerInternals::NewCustomReferenceController< Mode >( InRawPtrProxy.Object, MoveTemp( InRawPtrProxy.Deleter ) ) )
 	{
 		UE_TSHAREDPTR_STATIC_ASSERT_VALID_MODE(ObjectType, Mode)
 
@@ -807,12 +951,50 @@ public:
 		typename OtherType,
 		typename = decltype(ImplicitConv<ObjectType*>((OtherType*)nullptr))
 	>
-	FORCEINLINE TSharedPtr& operator=( SharedPointerInternals::FRawPtrProxy< OtherType > const& InRawPtrProxy )
+	FORCEINLINE TSharedPtr& operator=( SharedPointerInternals::TRawPtrProxy< OtherType > const& InRawPtrProxy )
+	{
+		*this = TSharedPtr< ObjectType, Mode >( InRawPtrProxy );
+		return *this;
+	}
+
+	/**
+	 * Assignment operator replaces this shared pointer with the specified shared pointer.  The object
+	 * currently referenced by this shared pointer will no longer be referenced and will be deleted if
+	 * there are no other referencers.
+	 *
+	 * @param  InRawPtrProxy  Proxy object used to assign the object (see MakeShareable helper function)
+	 */
+	// NOTE: The following is an Unreal extension to standard shared_ptr behavior
+	template <
+		typename OtherType,
+		typename DeleterType,
+		typename = decltype(ImplicitConv<ObjectType*>((OtherType*)nullptr))
+	>
+	FORCEINLINE TSharedPtr& operator=( SharedPointerInternals::TRawPtrProxyWithDeleter< OtherType, DeleterType > const& InRawPtrProxy )
 	{
 		*this = TSharedPtr< ObjectType, Mode >( InRawPtrProxy );
 		return *this;
 	}
 
+	/**
+	 * Assignment operator replaces this shared pointer with the specified shared pointer.  The object
+	 * currently referenced by this shared pointer will no longer be referenced and will be deleted if
+	 * there are no other referencers.
+	 *
+	 * @param  InRawPtrProxy  Proxy object used to assign the object (see MakeShareable helper function)
+	 */
+	// NOTE: The following is an Unreal extension to standard shared_ptr behavior
+	template <
+		typename OtherType,
+		typename DeleterType,
+		typename = decltype(ImplicitConv<ObjectType*>((OtherType*)nullptr))
+	>
+	FORCEINLINE TSharedPtr& operator=( SharedPointerInternals::TRawPtrProxyWithDeleter< OtherType, DeleterType >&& InRawPtrProxy )
+	{
+		*this = TSharedPtr< ObjectType, Mode >( MoveTemp( InRawPtrProxy ) );
+		return *this;
+	}
+
 	/**
 	 * Converts a shared pointer to a shared reference.  The pointer *must* be valid or an assertion will trigger.
 	 *
@@ -1815,9 +1997,9 @@ FORCEINLINE TSharedPtr< CastToType, Mode > ConstCastSharedPtr( TSharedPtr< CastF
  */
 // NOTE: The following is an Unreal extension to standard shared_ptr behavior
 template< class ObjectType >
-FORCEINLINE SharedPointerInternals::FRawPtrProxy< ObjectType > MakeShareable( ObjectType* InObject )
+FORCEINLINE SharedPointerInternals::TRawPtrProxy< ObjectType > MakeShareable( ObjectType* InObject )
 {
-	return SharedPointerInternals::FRawPtrProxy< ObjectType >( InObject );
+	return SharedPointerInternals::TRawPtrProxy< ObjectType >( InObject );
 }
 
 
@@ -1828,9 +2010,9 @@ FORCEINLINE SharedPointerInternals::FRawPtrProxy< ObjectType > MakeShareable( Ob
  */
 // NOTE: The following is an Unreal extension to standard shared_ptr behavior
 template< class ObjectType, class DeleterType >
-FORCEINLINE SharedPointerInternals::FRawPtrProxy< ObjectType > MakeShareable( ObjectType* InObject, DeleterType&& InDeleter )
+FORCEINLINE SharedPointerInternals::TRawPtrProxyWithDeleter< ObjectType, DeleterType > MakeShareable( ObjectType* InObject, DeleterType&& InDeleter )
 {
-	return SharedPointerInternals::FRawPtrProxy< ObjectType >( InObject, Forward< DeleterType >( InDeleter ) );
+	return SharedPointerInternals::TRawPtrProxyWithDeleter< ObjectType, DeleterType >( InObject, Forward< DeleterType >( InDeleter ) );
 }
 
 /**
@@ -1844,8 +2026,8 @@ FORCEINLINE SharedPointerInternals::FRawPtrProxy< ObjectType > MakeShareable( Ob
 template <typename InObjectType, ESPMode InMode = ESPMode::ThreadSafe, typename... InArgTypes>
 FORCEINLINE TSharedRef<InObjectType, InMode> MakeShared(InArgTypes&&... Args)
 {
-	SharedPointerInternals::TIntrusiveReferenceController<InObjectType>* Controller = SharedPointerInternals::NewIntrusiveReferenceController<InObjectType>(Forward<InArgTypes>(Args)...);
-	return UE4SharedPointer_Private::MakeSharedRef<InObjectType, InMode>(Controller->GetObjectPtr(), (SharedPointerInternals::FReferenceControllerBase*)Controller);
+	SharedPointerInternals::TIntrusiveReferenceController<InObjectType, InMode>* Controller = SharedPointerInternals::NewIntrusiveReferenceController<InMode, InObjectType>(Forward<InArgTypes>(Args)...);
+	return UE4SharedPointer_Private::MakeSharedRef<InObjectType, InMode>(Controller->GetObjectPtr(), (SharedPointerInternals::TReferenceControllerBase<InMode>*)Controller);
 }
 
 
diff --git a/Engine/Source/Runtime/Core/Public/Templates/SharedPointerInternals.h b/Engine/Source/Runtime/Core/Public/Templates/SharedPointerInternals.h
--- a/Engine/Source/Runtime/Core/Public/Templates/SharedPointerInternals.h
+++ b/Engine/Source/Runtime/Core/Public/Templates/SharedPointerInternals.h
@@ -3,12 +3,12 @@
 #pragma once
 
 #include "CoreTypes.h"
-#include "Templates/EnableIf.h"
 #include "Misc/AssertionMacros.h"
 #include "HAL/UnrealMemory.h"
-#include "Templates/AndOrNot.h"
 #include "Templates/RemoveReference.h"
 #include "Templates/TypeCompatibleBytes.h"
+#include <atomic>
+#include <type_traits>
 
 
 /** Default behavior. */
@@ -64,38 +64,38 @@ namespace SharedPointerInternals
 	// NOTE: The following is an Unreal extension to standard shared_ptr behavior
 	struct FNullTag {};
 
-
-	class FReferenceControllerBase
+	template <ESPMode Mode>
+	class TReferenceControllerBase
 	{
-	public:
-		/** Constructor */
-		FORCEINLINE explicit FReferenceControllerBase()
-			: SharedReferenceCount(1)
-			, WeakReferenceCount(1)
-		{
-		}
-
-		// NOTE: The primary reason these reference counters are 32-bit values (and not 16-bit to save
-		//       memory), is that atomic operations require at least 32-bit objects.
+		using RefCountType = std::conditional_t<Mode == ESPMode::ThreadSafe, std::atomic<int32>, int32>;
 
-		/** Number of shared references to this object.  When this count reaches zero, the associated object
-		    will be destroyed (even if there are still weak references!) */
-		int32 SharedReferenceCount;
-
-		/** Number of weak references to this object.  If there are any shared references, that counts as one
-		   weak reference too. */
-		int32 WeakReferenceCount;
+	public:
+		FORCEINLINE explicit TReferenceControllerBase() = default;
+
+		// Number of shared references to this object.  When this count reaches zero, the associated object
+		// will be destroyed (even if there are still weak references!), but not the reference controller.
+		//
+		// This starts at 1 because we create reference controllers via the construction of a TSharedPtr,
+		// and that is the first reference.  There is no point in starting at 0 and then incrementing it.
+		RefCountType SharedReferenceCount{1};
+
+		// Number of weak references to this object.  If there are any shared references, that counts as one
+		// weak reference too.  When this count reaches zero, the reference controller will be deleted.
+		//
+		// This starts at 1 because it represents the shared reference that we are also initializing
+		// SharedReferenceCount with.
+		RefCountType WeakReferenceCount{1};
 
 		/** Destroys the object associated with this reference counter.  */
 		virtual void DestroyObject() = 0;
 
-		virtual ~FReferenceControllerBase()
+		virtual ~TReferenceControllerBase()
 		{
 		}
 
-	private:
-		FReferenceControllerBase( FReferenceControllerBase const& );
-		FReferenceControllerBase& operator=( FReferenceControllerBase const& );
+		// Non-copyable
+		TReferenceControllerBase(const TReferenceControllerBase&) = delete;
+		TReferenceControllerBase& operator=(const TReferenceControllerBase&) = delete;
 	};
 
 	// A helper class that efficiently stores a custom deleter and is intended to be derived from. If the custom deleter is an empty class,
@@ -135,8 +135,8 @@ namespace SharedPointerInternals
 		DeleterType Deleter;
 	};
 
-	template <typename ObjectType, typename DeleterType>
-	class TReferenceControllerWithDeleter : private TDeleterHolder<DeleterType>, public FReferenceControllerBase
+	template <typename ObjectType, typename DeleterType, ESPMode Mode>
+	class TReferenceControllerWithDeleter : private TDeleterHolder<DeleterType>, public TReferenceControllerBase<Mode>
 	{
 	public:
 		explicit TReferenceControllerWithDeleter(ObjectType* InObject, DeleterType&& Deleter)
@@ -159,8 +159,8 @@ namespace SharedPointerInternals
 		ObjectType* Object;
 	};
 
-	template <typename ObjectType>
-	class TIntrusiveReferenceController : public FReferenceControllerBase
+	template <typename ObjectType, ESPMode Mode>
+	class TIntrusiveReferenceController : public TReferenceControllerBase<Mode>
 	{
 	public:
 		template <typename... ArgTypes>
@@ -200,50 +200,71 @@ namespace SharedPointerInternals
 	};
 
 	/** Creates a reference controller which just calls delete */
-	template <typename ObjectType>
-	inline FReferenceControllerBase* NewDefaultReferenceController(ObjectType* Object)
+	template <ESPMode Mode, typename ObjectType>
+	inline TReferenceControllerBase<Mode>* NewDefaultReferenceController(ObjectType* Object)
 	{
-		return new TReferenceControllerWithDeleter<ObjectType, DefaultDeleter<ObjectType>>(Object, DefaultDeleter<ObjectType>());
+		return new TReferenceControllerWithDeleter<ObjectType, DefaultDeleter<ObjectType>, Mode>(Object, DefaultDeleter<ObjectType>());
 	}
 
 	/** Creates a custom reference controller with a specified deleter */
-	template <typename ObjectType, typename DeleterType>
-	inline FReferenceControllerBase* NewCustomReferenceController(ObjectType* Object, DeleterType&& Deleter)
+	template <ESPMode Mode, typename ObjectType, typename DeleterType>
+	inline TReferenceControllerBase<Mode>* NewCustomReferenceController(ObjectType* Object, DeleterType&& Deleter)
 	{
-		return new TReferenceControllerWithDeleter<ObjectType, typename TRemoveReference<DeleterType>::Type>(Object, Forward<DeleterType>(Deleter));
+		return new TReferenceControllerWithDeleter<ObjectType, typename TRemoveReference<DeleterType>::Type, Mode>(Object, Forward<DeleterType>(Deleter));
 	}
 
 	/** Creates an intrusive reference controller */
-	template <typename ObjectType, typename... ArgTypes>
-	inline TIntrusiveReferenceController<ObjectType>* NewIntrusiveReferenceController(ArgTypes&&... Args)
+	template <ESPMode Mode, typename ObjectType, typename... ArgTypes>
+	inline TIntrusiveReferenceController<ObjectType, Mode>* NewIntrusiveReferenceController(ArgTypes&&... Args)
 	{
-		return new TIntrusiveReferenceController<ObjectType>(Forward<ArgTypes>(Args)...);
+		return new TIntrusiveReferenceController<ObjectType, Mode>(Forward<ArgTypes>(Args)...);
 	}
 
 
 	/** Proxy structure for implicitly converting raw pointers to shared/weak pointers */
 	// NOTE: The following is an Unreal extension to standard shared_ptr behavior
-	template< class ObjectType >
-	struct FRawPtrProxy
+	template <class ObjectType>
+	struct TRawPtrProxy
 	{
 		/** The object pointer */
 		ObjectType* Object;
 
-		/** Reference controller used to destroy the object */
-		FReferenceControllerBase* ReferenceController;
+		/** Construct implicitly from a nullptr */
+		FORCEINLINE TRawPtrProxy( TYPE_OF_NULLPTR )
+			: Object( nullptr )
+		{
+		}
 
-		/** Construct implicitly from an object */
-		FORCEINLINE FRawPtrProxy( ObjectType* InObject )
-			: Object             ( InObject )
-			, ReferenceController( NewDefaultReferenceController( InObject ) )
+		/** Construct explicitly from an object */
+		explicit FORCEINLINE TRawPtrProxy( ObjectType* InObject )
+			: Object( InObject )
 		{
 		}
+	};
+
+
+	/** Proxy structure for implicitly converting raw pointers to shared/weak pointers */
+	// NOTE: The following is an Unreal extension to standard shared_ptr behavior
+	template <class ObjectType, typename DeleterType>
+	struct TRawPtrProxyWithDeleter
+	{
+		/** The object pointer */
+		ObjectType* Object;
+
+		/** The deleter object */
+		DeleterType Deleter;
 
 		/** Construct implicitly from an object and a custom deleter */
-		template< class Deleter >
-		FORCEINLINE FRawPtrProxy( ObjectType* InObject, Deleter&& InDeleter )
-			: Object             ( InObject )
-			, ReferenceController( NewCustomReferenceController( InObject, Forward< Deleter >( InDeleter ) ) )
+		FORCEINLINE TRawPtrProxyWithDeleter( ObjectType* InObject, const DeleterType& InDeleter )
+			: Object ( InObject )
+			, Deleter( InObject, InDeleter )
+		{
+		}
+
+		/** Construct implicitly from an object and a custom deleter */
+		FORCEINLINE TRawPtrProxyWithDeleter( ObjectType* InObject, DeleterType&& InDeleter )
+			: Object ( InObject )
+			, Deleter( MoveTemp( InDeleter ) )
 		{
 		}
 	};
@@ -262,16 +283,24 @@ namespace SharedPointerInternals
 	struct FReferenceControllerOps<ESPMode::ThreadSafe>
 	{
 		/** Returns the shared reference count */
-		static FORCEINLINE const int32 GetSharedReferenceCount(const FReferenceControllerBase* ReferenceController)
+		static FORCEINLINE const int32 GetSharedReferenceCount(const TReferenceControllerBase<ESPMode::ThreadSafe>* ReferenceController)
 		{
+			// A 'live' shared reference count is unstable by nature and so there's no benefit
+			// to try and enforce memory ordering around the reading of it.
+			//
+			// This is equivalent to https://en.cppreference.com/w/cpp/memory/shared_ptr/use_count
+
 			// This reference count may be accessed by multiple threads
-			return FPlatformAtomics::AtomicRead( (int32 volatile*)&ReferenceController->SharedReferenceCount );
+			return ReferenceController->SharedReferenceCount.load(std::memory_order_relaxed);
 		}
 
 		/** Adds a shared reference to this counter */
-		static FORCEINLINE void AddSharedReference(FReferenceControllerBase* ReferenceController)
+		static FORCEINLINE void AddSharedReference(TReferenceControllerBase<ESPMode::ThreadSafe>* ReferenceController)
 		{
-			FPlatformAtomics::InterlockedIncrement( &ReferenceController->SharedReferenceCount );
+			// Incrementing a reference count with relaxed ordering is always safe because no other action is taken
+			// in response to the increment, so there's nothing to order with.
+
+			ReferenceController->SharedReferenceCount.fetch_add(1, std::memory_order_relaxed);
 		}
 
 		/**
@@ -279,13 +308,16 @@ namespace SharedPointerInternals
 		 *
 		 * @return  True if the shared reference was added successfully
 		 */
-		static bool ConditionallyAddSharedReference(FReferenceControllerBase* ReferenceController)
+		static bool ConditionallyAddSharedReference(TReferenceControllerBase<ESPMode::ThreadSafe>* ReferenceController)
 		{
-			for( ; ; )
+			// See AddSharedReference for the same reasons that std::memory_order_relaxed is used in this function.
+
+			// Peek at the current shared reference count.  Remember, this value may be updated by
+			// multiple threads.
+			int32 OriginalCount = ReferenceController->SharedReferenceCount.load(std::memory_order_relaxed);
+
+			for ( ; ; )
 			{
-				// Peek at the current shared reference count.  Remember, this value may be updated by
-				// multiple threads.
-				const int32 OriginalCount = FPlatformAtomics::AtomicRead( (int32 volatile*)&ReferenceController->SharedReferenceCount );
 				if( OriginalCount == 0 )
 				{
 					// Never add a shared reference if the pointer has already expired
@@ -293,12 +325,18 @@ namespace SharedPointerInternals
 				}
 
 				// Attempt to increment the reference count.
-				const int32 ActualOriginalCount = FPlatformAtomics::InterlockedCompareExchange( &ReferenceController->SharedReferenceCount, OriginalCount + 1, OriginalCount );
-
+				//
 				// We need to make sure that we never revive a counter that has already expired, so if the
 				// actual value what we expected (because it was touched by another thread), then we'll try
 				// again.  Note that only in very unusual cases will this actually have to loop.
-				if( ActualOriginalCount == OriginalCount )
+				//
+				// We do a weak read here because we require a loop and this is the recommendation:
+				//
+				// https://en.cppreference.com/w/cpp/atomic/atomic/compare_exchange
+				//
+				// > When a compare-and-exchange is in a loop, the weak version will yield better performance on some platforms.
+				// > When a weak compare-and-exchange would require a loop and a strong one would not, the strong one is preferable
+				if (ReferenceController->SharedReferenceCount.compare_exchange_weak(OriginalCount, OriginalCount + 1, std::memory_order_relaxed))
 				{
 					return true;
 				}
@@ -306,12 +344,20 @@ namespace SharedPointerInternals
 		}
 
 		/** Releases a shared reference to this counter */
-		static FORCEINLINE void ReleaseSharedReference(FReferenceControllerBase* ReferenceController)
+		static FORCEINLINE void ReleaseSharedReference(TReferenceControllerBase<ESPMode::ThreadSafe>* ReferenceController)
 		{
-			checkSlow( ReferenceController->SharedReferenceCount > 0 );
+			// std::memory_order_release is used here so that, if we do end up executing the destructor, it's not possible
+			// for side effects from executing the destructor end up being visible before we've determined that the shared
+			// reference count is actually zero.
 
-			if( FPlatformAtomics::InterlockedDecrement( &ReferenceController->SharedReferenceCount ) == 0 )
+			int32 OldSharedCount = ReferenceController->SharedReferenceCount.fetch_sub(1, std::memory_order_release);
+			checkSlow(OldSharedCount > 0);
+			if (OldSharedCount == 1)
 			{
+				// Ensure that all other threads' accesses to the object are visible to this thread before we call the
+				// destructor.
+				std::atomic_thread_fence(std::memory_order_acquire);
+
 				// Last shared reference was released!  Destroy the referenced object.
 				ReferenceController->DestroyObject();
 
@@ -323,18 +369,24 @@ namespace SharedPointerInternals
 
 
 		/** Adds a weak reference to this counter */
-		static FORCEINLINE void AddWeakReference(FReferenceControllerBase* ReferenceController)
+		static FORCEINLINE void AddWeakReference(TReferenceControllerBase<ESPMode::ThreadSafe>* ReferenceController)
 		{
-			FPlatformAtomics::InterlockedIncrement( &ReferenceController->WeakReferenceCount );
+			// See AddSharedReference for the same reasons that std::memory_order_relaxed is used in this function.
+
+			ReferenceController->WeakReferenceCount.fetch_add(1, std::memory_order_relaxed);
 		}
 
 		/** Releases a weak reference to this counter */
-		static void ReleaseWeakReference(FReferenceControllerBase* ReferenceController)
+		static void ReleaseWeakReference(TReferenceControllerBase<ESPMode::ThreadSafe>* ReferenceController)
 		{
-			checkSlow( ReferenceController->WeakReferenceCount > 0 );
+			// See ReleaseSharedReference for the same reasons that std::memory_order_release and std::memory_order_acquire are used in this function.
 
-			if( FPlatformAtomics::InterlockedDecrement( &ReferenceController->WeakReferenceCount ) == 0 )
+			int32 OldWeakCount = ReferenceController->WeakReferenceCount.fetch_sub(1, std::memory_order_release);
+			checkSlow(OldWeakCount > 0);
+			if (OldWeakCount == 1)
 			{
+				std::atomic_thread_fence(std::memory_order_acquire);
+
 				// No more references to this reference count.  Destroy it!
 				delete ReferenceController;
 			}
@@ -346,13 +398,13 @@ namespace SharedPointerInternals
 	struct FReferenceControllerOps<ESPMode::NotThreadSafe>
 	{
 		/** Returns the shared reference count */
-		static FORCEINLINE const int32 GetSharedReferenceCount(const FReferenceControllerBase* ReferenceController) TSAN_SAFE_UNSAFEPTR
+		static FORCEINLINE const int32 GetSharedReferenceCount(const TReferenceControllerBase<ESPMode::NotThreadSafe>* ReferenceController) TSAN_SAFE_UNSAFEPTR
 		{
 			return ReferenceController->SharedReferenceCount;
 		}
 
 		/** Adds a shared reference to this counter */
-		static FORCEINLINE void AddSharedReference(FReferenceControllerBase* ReferenceController) TSAN_SAFE_UNSAFEPTR
+		static FORCEINLINE void AddSharedReference(TReferenceControllerBase<ESPMode::NotThreadSafe>* ReferenceController) TSAN_SAFE_UNSAFEPTR
 		{
 			++ReferenceController->SharedReferenceCount;
 		}
@@ -362,7 +414,7 @@ namespace SharedPointerInternals
 		 *
 		 * @return  True if the shared reference was added successfully
 		 */
-		static bool ConditionallyAddSharedReference(FReferenceControllerBase* ReferenceController) TSAN_SAFE_UNSAFEPTR
+		static bool ConditionallyAddSharedReference(TReferenceControllerBase<ESPMode::NotThreadSafe>* ReferenceController) TSAN_SAFE_UNSAFEPTR
 		{
 			if( ReferenceController->SharedReferenceCount == 0 )
 			{
@@ -375,7 +427,7 @@ namespace SharedPointerInternals
 		}
 
 		/** Releases a shared reference to this counter */
-		static FORCEINLINE void ReleaseSharedReference(FReferenceControllerBase* ReferenceController) TSAN_SAFE_UNSAFEPTR
+		static FORCEINLINE void ReleaseSharedReference(TReferenceControllerBase<ESPMode::NotThreadSafe>* ReferenceController) TSAN_SAFE_UNSAFEPTR
 		{
 			checkSlow( ReferenceController->SharedReferenceCount > 0 );
 
@@ -391,13 +443,13 @@ namespace SharedPointerInternals
 		}
 
 		/** Adds a weak reference to this counter */
-		static FORCEINLINE void AddWeakReference(FReferenceControllerBase* ReferenceController) TSAN_SAFE_UNSAFEPTR
+		static FORCEINLINE void AddWeakReference(TReferenceControllerBase<ESPMode::NotThreadSafe>* ReferenceController) TSAN_SAFE_UNSAFEPTR
 		{
 			++ReferenceController->WeakReferenceCount;
 		}
 
 		/** Releases a weak reference to this counter */
-		static void ReleaseWeakReference(FReferenceControllerBase* ReferenceController) TSAN_SAFE_UNSAFEPTR
+		static void ReleaseWeakReference(TReferenceControllerBase<ESPMode::NotThreadSafe>* ReferenceController) TSAN_SAFE_UNSAFEPTR
 		{
 			checkSlow( ReferenceController->WeakReferenceCount > 0 );
 
@@ -427,7 +479,7 @@ namespace SharedPointerInternals
 		{ }
 
 		/** Constructor that counts a single reference to the specified object */
-		inline explicit FSharedReferencer( FReferenceControllerBase* InReferenceController )
+		inline explicit FSharedReferencer( TReferenceControllerBase<Mode>* InReferenceController )
 			: ReferenceController( InReferenceController )
 		{ }
 		
@@ -593,7 +645,7 @@ namespace SharedPointerInternals
 	private:
 
 		/** Pointer to the reference controller for the object a shared reference/pointer is referencing */
-		FReferenceControllerBase* ReferenceController;
+		TReferenceControllerBase<Mode>* ReferenceController;
 	};
 
 
@@ -700,7 +752,7 @@ namespace SharedPointerInternals
 
 		/** Assigns a new reference controller to this counter object, first adding a reference to it, then
 		    releasing the previous object. */
-		inline void AssignReferenceController( FReferenceControllerBase* NewReferenceController )
+		inline void AssignReferenceController( TReferenceControllerBase<Mode>* NewReferenceController )
 		{
 			// Only proceed if the new reference counter is different than our current
 			if( NewReferenceController != ReferenceController )
@@ -730,7 +782,7 @@ namespace SharedPointerInternals
 	private:
 
 		/** Pointer to the reference controller for the object a TWeakPtr is referencing */
-		FReferenceControllerBase* ReferenceController;
+		TReferenceControllerBase<Mode>* ReferenceController;
 	};
 
 
diff --git a/Engine/Source/Runtime/SlateCore/Public/Widgets/SlateControlledConstruction.h b/Engine/Source/Runtime/SlateCore/Public/Widgets/SlateControlledConstruction.h
--- a/Engine/Source/Runtime/SlateCore/Public/Widgets/SlateControlledConstruction.h
+++ b/Engine/Source/Runtime/SlateCore/Public/Widgets/SlateControlledConstruction.h
@@ -7,10 +7,11 @@
 #include "Types/SlateAttributeDescriptor.h"
 #include "Templates/Identity.h"
 
+enum class ESPMode : uint8;
 class FSlateControlledConstruction;
 namespace SharedPointerInternals
 {
-	template <typename ObjectType>
+	template <typename ObjectType, ESPMode Mode>
 	class TIntrusiveReferenceController;
 }
 
@@ -23,8 +24,8 @@ namespace SharedPointerInternals
 		using PrivateParentType = ParentType; \
 		template<class WidgetType, bool bIsUserWidget> \
 		friend struct TWidgetAllocator; \
-		template <typename ObjectType> \
-		friend class SharedPointerInternals::TIntrusiveReferenceController; \
+		template <typename ObjectType, ESPMode Mode> \
+		friend class SharedPointerInternals::TIntrusiveReferenceController; /* Shouldn't really forward-declare this to allow MakeShared to work, but is for legacy reasons */ \
 		static const FSlateWidgetClassData& GetPrivateWidgetClass() \
 		{ \
 			static FSlateWidgetClassData WidgetClassDataInstance = FSlateWidgetClassData(TIdentity<ParentType>(), #WidgetType, &WidgetType::PrivateRegisterAttributes); \
@@ -113,7 +114,8 @@ private:
 	template<class WidgetType, bool bIsUserWidget>
 	friend struct TWidgetAllocator;
 
-	template <typename ObjectType>
+	/* Shouldn't really forward-declare this to allow MakeShared to work, but is for legacy reasons */
+	template <typename ObjectType, ESPMode Mode>
 	friend class SharedPointerInternals::TIntrusiveReferenceController;
 
 public:
