Fix: RandomPointInBoundingBox is available globally but not available on the random stream

diff --git a/Engine/Source/Runtime/Core/Public/Math/RandomStream.h b/Engine/Source/Runtime/Core/Public/Math/RandomStream.h
--- a/Engine/Source/Runtime/Core/Public/Math/RandomStream.h
+++ b/Engine/Source/Runtime/Core/Public/Math/RandomStream.h
@@ -210,6 +210,16 @@ public:
 		return InMin + (InMax - InMin) * FRand();
 	}
 
+	/**
+	* Helper function for rand implementations.
+	*
+	* @return A random number >= Min and <= Max
+	*/
+	FORCEINLINE double FRandRange(double InMin, double InMax) const
+	{
+		return InMin + (InMax - InMin) * FRand();	// LWC_TODO: Implement FRandDbl() for increased precision
+	}
+
 	/**
 	 * Returns a random vector of unit size.
 	 *
@@ -220,6 +230,11 @@ public:
 		return GetUnitVector();
 	}
 
+	/**
+	* Returns a random point contained inside an axis aligned bounding box
+	*
+	* @return Random point within Box.
+	*/
 	FORCEINLINE FVector RandPointInBox(const FBox& Box) const
 	{
 		return FVector(	FRandRange(Box.Min.X, Box.Max.X),
