Fix for LWC compile error in RandomStream
#rnx

diff --git a/Engine/Source/Runtime/Core/Public/Math/RandomStream.h b/Engine/Source/Runtime/Core/Public/Math/RandomStream.h
--- a/Engine/Source/Runtime/Core/Public/Math/RandomStream.h
+++ b/Engine/Source/Runtime/Core/Public/Math/RandomStream.h
@@ -3,6 +3,7 @@
 #pragma once
 
 #include "CoreTypes.h"
+#include "Math/Box.h"
 #include "Math/UnrealMathUtility.h"
 #include "Math/Vector.h"
 #include "Math/Matrix.h"
@@ -205,7 +206,7 @@ public:
 	 *
 	 * @return A random number >= Min and <= Max
 	 */
-	FORCEINLINE float FRandRange( float InMin, float InMax ) const
+	FORCEINLINE FVector::FReal FRandRange( FVector::FReal InMin, FVector::FReal InMax ) const
 	{
 		return InMin + (InMax - InMin) * FRand();
 	}
