Allow counters-trace's on/off define to be defined in a Target.cs.

#rnx

diff --git a/Engine/Source/Runtime/Core/Public/ProfilingDebugging/CountersTrace.h b/Engine/Source/Runtime/Core/Public/ProfilingDebugging/CountersTrace.h
--- a/Engine/Source/Runtime/Core/Public/ProfilingDebugging/CountersTrace.h
+++ b/Engine/Source/Runtime/Core/Public/ProfilingDebugging/CountersTrace.h
@@ -6,11 +6,13 @@
 #include "Trace/Config.h"
 #include "Trace/Trace.h"
 
+#if !defined(COUNTERSTRACE_ENABLED)
 #if UE_TRACE_ENABLED && !UE_BUILD_SHIPPING
 #define COUNTERSTRACE_ENABLED 1
 #else
 #define COUNTERSTRACE_ENABLED 0
 #endif
+#endif
 
 enum ETraceCounterType
 {
