Allow forked child process task graph to start worker threads

diff --git a/Engine/Source/Runtime/Core/Private/Async/Fundamental/Scheduler.cpp b/Engine/Source/Runtime/Core/Private/Async/Fundamental/Scheduler.cpp
--- a/Engine/Source/Runtime/Core/Private/Async/Fundamental/Scheduler.cpp
+++ b/Engine/Source/Runtime/Core/Private/Async/Fundamental/Scheduler.cpp
@@ -5,6 +5,7 @@
 #include "Async/TaskTrace.h"
 #include "Logging/LogMacros.h"
 #include "Misc/ScopeLock.h"
+#include "Misc/Fork.h"
 #include "CoreGlobals.h"
 
 namespace LowLevelTasks
@@ -87,8 +88,10 @@ namespace LowLevelTasks
 			NumBackgroundWorkers = FMath::Max<int32>(1, FPlatformMisc::NumberOfWorkerThreadsToSpawn() - NumForegroundWorkers);
 		}
 
+		const bool bSupportsMultithreading = FPlatformProcess::SupportsMultithreading() || FForkProcessHelper::IsForkedMultithreadInstance();
+
 		uint32 OldActiveWorkers = ActiveWorkers.load(std::memory_order_relaxed);
-		if(OldActiveWorkers == 0 && FPlatformProcess::SupportsMultithreading() && ActiveWorkers.compare_exchange_strong(OldActiveWorkers, NumForegroundWorkers + NumBackgroundWorkers, std::memory_order_relaxed))
+		if(OldActiveWorkers == 0 && bSupportsMultithreading && ActiveWorkers.compare_exchange_strong(OldActiveWorkers, NumForegroundWorkers + NumBackgroundWorkers, std::memory_order_relaxed))
 		{
 			FScopeLock Lock(&WorkerThreadsCS);
 			check(!WorkerThreads.Num());
