Small typo/fix for SpringDamper

diff --git a/Engine/Source/Runtime/Core/Public/Math/UnrealMathUtility.h b/Engine/Source/Runtime/Core/Public/Math/UnrealMathUtility.h
--- a/Engine/Source/Runtime/Core/Public/Math/UnrealMathUtility.h
+++ b/Engine/Source/Runtime/Core/Public/Math/UnrealMathUtility.h
@@ -1401,8 +1401,8 @@ public:
 			const T A = Err;
 			const T B = (InOutValueRate + Err * (InDampingRatio * W)) / WD;
 			float S, C;
-			FMath::SinCos(&S, &C, W * InDeltaTime);
-			const float E0 = InDampingRatio * WD * InDeltaTime;
+			FMath::SinCos(&S, &C, WD * InDeltaTime);
+			const float E0 = InDampingRatio * W * InDeltaTime;
 			// Needs E0 < 1 so DeltaTime < SmoothingTime / (2 * DampingRatio * Sqrt(1 - DampingRatio^2))
 			const float E = InvExpApprox(E0);
 			InOutValue = E * (A * C + B * S);
