TSharedPtr's FReferenceControllerOps moved to being direct members of TReferenceControllerBase for clarity, now that TReferenceControllerBase knows about the threading mode.

diff --git a/Engine/Source/Runtime/Core/Public/Templates/SharedPointerInternals.h b/Engine/Source/Runtime/Core/Public/Templates/SharedPointerInternals.h
--- a/Engine/Source/Runtime/Core/Public/Templates/SharedPointerInternals.h
+++ b/Engine/Source/Runtime/Core/Public/Templates/SharedPointerInternals.h
@@ -93,6 +93,180 @@ namespace SharedPointerInternals
 		{
 		}
 
+		/** Returns the shared reference count */
+		FORCEINLINE int32 GetSharedReferenceCount() const
+		{
+			if constexpr (Mode == ESPMode::ThreadSafe)
+			{
+				// A 'live' shared reference count is unstable by nature and so there's no benefit
+				// to try and enforce memory ordering around the reading of it.
+				//
+				// This is equivalent to https://en.cppreference.com/w/cpp/memory/shared_ptr/use_count
+
+				// This reference count may be accessed by multiple threads
+				return SharedReferenceCount.load(std::memory_order_relaxed);
+			}
+			else
+			{
+				return SharedReferenceCount;
+			}
+		}
+
+		/** Adds a shared reference to this counter */
+		FORCEINLINE void AddSharedReference()
+		{
+			if constexpr (Mode == ESPMode::ThreadSafe)
+			{
+				// Incrementing a reference count with relaxed ordering is always safe because no other action is taken
+				// in response to the increment, so there's nothing to order with.
+
+				SharedReferenceCount.fetch_add(1, std::memory_order_relaxed);
+			}
+			else
+			{
+				++SharedReferenceCount;
+			}
+		}
+
+		/**
+		 * Adds a shared reference to this counter ONLY if there is already at least one reference
+		 *
+		 * @return  True if the shared reference was added successfully
+		 */
+		bool ConditionallyAddSharedReference()
+		{
+			if constexpr (Mode == ESPMode::ThreadSafe)
+			{
+				// See AddSharedReference for the same reasons that std::memory_order_relaxed is used in this function.
+
+				// Peek at the current shared reference count.  Remember, this value may be updated by
+				// multiple threads.
+				int32 OriginalCount = SharedReferenceCount.load(std::memory_order_relaxed);
+
+				for ( ; ; )
+				{
+					if( OriginalCount == 0 )
+					{
+						// Never add a shared reference if the pointer has already expired
+						return false;
+					}
+
+					// Attempt to increment the reference count.
+					//
+					// We need to make sure that we never revive a counter that has already expired, so if the
+					// actual value what we expected (because it was touched by another thread), then we'll try
+					// again.  Note that only in very unusual cases will this actually have to loop.
+					//
+					// We do a weak read here because we require a loop and this is the recommendation:
+					//
+					// https://en.cppreference.com/w/cpp/atomic/atomic/compare_exchange
+					//
+					// > When a compare-and-exchange is in a loop, the weak version will yield better performance on some platforms.
+					// > When a weak compare-and-exchange would require a loop and a strong one would not, the strong one is preferable
+					if (SharedReferenceCount.compare_exchange_weak(OriginalCount, OriginalCount + 1, std::memory_order_relaxed))
+					{
+						return true;
+					}
+				}
+			}
+			else
+			{
+				if( SharedReferenceCount == 0 )
+				{
+					// Never add a shared reference if the pointer has already expired
+					return false;
+				}
+
+				++SharedReferenceCount;
+				return true;
+			}
+		}
+
+		/** Releases a shared reference to this counter */
+		FORCEINLINE void ReleaseSharedReference()
+		{
+			if constexpr (Mode == ESPMode::ThreadSafe)
+			{
+				// std::memory_order_release is used here so that, if we do end up executing the destructor, it's not possible
+				// for side effects from executing the destructor end up being visible before we've determined that the shared
+				// reference count is actually zero.
+
+				int32 OldSharedCount = SharedReferenceCount.fetch_sub(1, std::memory_order_release);
+				checkSlow(OldSharedCount > 0);
+				if (OldSharedCount == 1)
+				{
+					// Ensure that all other threads' accesses to the object are visible to this thread before we call the
+					// destructor.
+					std::atomic_thread_fence(std::memory_order_acquire);
+
+					// Last shared reference was released!  Destroy the referenced object.
+					DestroyObject();
+
+					// No more shared referencers, so decrement the weak reference count by one.  When the weak
+					// reference count reaches zero, this object will be deleted.
+					ReleaseWeakReference();
+				}
+			}
+			else
+			{
+				checkSlow( SharedReferenceCount > 0 );
+
+				if( --SharedReferenceCount == 0 )
+				{
+					// Last shared reference was released!  Destroy the referenced object.
+					DestroyObject();
+
+					// No more shared referencers, so decrement the weak reference count by one.  When the weak
+					// reference count reaches zero, this object will be deleted.
+					ReleaseWeakReference();
+				}
+			}
+		}
+
+		/** Adds a weak reference to this counter */
+		FORCEINLINE void AddWeakReference()
+		{
+			if constexpr (Mode == ESPMode::ThreadSafe)
+			{
+				// See AddSharedReference for the same reasons that std::memory_order_relaxed is used in this function.
+
+				WeakReferenceCount.fetch_add(1, std::memory_order_relaxed);
+			}
+			else
+			{
+				++WeakReferenceCount;
+			}
+		}
+
+		/** Releases a weak reference to this counter */
+		void ReleaseWeakReference()
+		{
+			if constexpr (Mode == ESPMode::ThreadSafe)
+			{
+				// See ReleaseSharedReference for the same reasons that std::memory_order_release and std::memory_order_acquire are used in this function.
+
+				int32 OldWeakCount = WeakReferenceCount.fetch_sub(1, std::memory_order_release);
+				checkSlow(OldWeakCount > 0);
+				if (OldWeakCount == 1)
+				{
+					std::atomic_thread_fence(std::memory_order_acquire);
+
+					// No more references to this reference count.  Destroy it!
+					delete this;
+				}
+			}
+			else
+			{
+				checkSlow( WeakReferenceCount > 0 );
+
+				if( --WeakReferenceCount == 0 )
+				{
+					// No more references to this reference count.  Destroy it!
+					delete this;
+				}
+			}
+		}
+
 		// Non-copyable
 		TReferenceControllerBase(const TReferenceControllerBase&) = delete;
 		TReferenceControllerBase& operator=(const TReferenceControllerBase&) = delete;
@@ -270,198 +444,6 @@ namespace SharedPointerInternals
 	};
 
 
-	/**
-	 * FReferenceController is a standalone heap-allocated object that tracks the number of references
-	 * to an object referenced by TSharedRef, TSharedPtr or TWeakPtr objects.
-	 *
-	 * It is specialized for different threading modes.
-	 */
-	template< ESPMode Mode >
-	struct FReferenceControllerOps;
-
-	template<>
-	struct FReferenceControllerOps<ESPMode::ThreadSafe>
-	{
-		/** Returns the shared reference count */
-		static FORCEINLINE const int32 GetSharedReferenceCount(const TReferenceControllerBase<ESPMode::ThreadSafe>* ReferenceController)
-		{
-			// A 'live' shared reference count is unstable by nature and so there's no benefit
-			// to try and enforce memory ordering around the reading of it.
-			//
-			// This is equivalent to https://en.cppreference.com/w/cpp/memory/shared_ptr/use_count
-
-			// This reference count may be accessed by multiple threads
-			return ReferenceController->SharedReferenceCount.load(std::memory_order_relaxed);
-		}
-
-		/** Adds a shared reference to this counter */
-		static FORCEINLINE void AddSharedReference(TReferenceControllerBase<ESPMode::ThreadSafe>* ReferenceController)
-		{
-			// Incrementing a reference count with relaxed ordering is always safe because no other action is taken
-			// in response to the increment, so there's nothing to order with.
-
-			ReferenceController->SharedReferenceCount.fetch_add(1, std::memory_order_relaxed);
-		}
-
-		/**
-		 * Adds a shared reference to this counter ONLY if there is already at least one reference
-		 *
-		 * @return  True if the shared reference was added successfully
-		 */
-		static bool ConditionallyAddSharedReference(TReferenceControllerBase<ESPMode::ThreadSafe>* ReferenceController)
-		{
-			// See AddSharedReference for the same reasons that std::memory_order_relaxed is used in this function.
-
-			// Peek at the current shared reference count.  Remember, this value may be updated by
-			// multiple threads.
-			int32 OriginalCount = ReferenceController->SharedReferenceCount.load(std::memory_order_relaxed);
-
-			for ( ; ; )
-			{
-				if( OriginalCount == 0 )
-				{
-					// Never add a shared reference if the pointer has already expired
-					return false;
-				}
-
-				// Attempt to increment the reference count.
-				//
-				// We need to make sure that we never revive a counter that has already expired, so if the
-				// actual value what we expected (because it was touched by another thread), then we'll try
-				// again.  Note that only in very unusual cases will this actually have to loop.
-				//
-				// We do a weak read here because we require a loop and this is the recommendation:
-				//
-				// https://en.cppreference.com/w/cpp/atomic/atomic/compare_exchange
-				//
-				// > When a compare-and-exchange is in a loop, the weak version will yield better performance on some platforms.
-				// > When a weak compare-and-exchange would require a loop and a strong one would not, the strong one is preferable
-				if (ReferenceController->SharedReferenceCount.compare_exchange_weak(OriginalCount, OriginalCount + 1, std::memory_order_relaxed))
-				{
-					return true;
-				}
-			}
-		}
-
-		/** Releases a shared reference to this counter */
-		static FORCEINLINE void ReleaseSharedReference(TReferenceControllerBase<ESPMode::ThreadSafe>* ReferenceController)
-		{
-			// std::memory_order_release is used here so that, if we do end up executing the destructor, it's not possible
-			// for side effects from executing the destructor end up being visible before we've determined that the shared
-			// reference count is actually zero.
-
-			int32 OldSharedCount = ReferenceController->SharedReferenceCount.fetch_sub(1, std::memory_order_release);
-			checkSlow(OldSharedCount > 0);
-			if (OldSharedCount == 1)
-			{
-				// Ensure that all other threads' accesses to the object are visible to this thread before we call the
-				// destructor.
-				std::atomic_thread_fence(std::memory_order_acquire);
-
-				// Last shared reference was released!  Destroy the referenced object.
-				ReferenceController->DestroyObject();
-
-				// No more shared referencers, so decrement the weak reference count by one.  When the weak
-				// reference count reaches zero, this object will be deleted.
-				ReleaseWeakReference(ReferenceController);
-			}
-		}
-
-
-		/** Adds a weak reference to this counter */
-		static FORCEINLINE void AddWeakReference(TReferenceControllerBase<ESPMode::ThreadSafe>* ReferenceController)
-		{
-			// See AddSharedReference for the same reasons that std::memory_order_relaxed is used in this function.
-
-			ReferenceController->WeakReferenceCount.fetch_add(1, std::memory_order_relaxed);
-		}
-
-		/** Releases a weak reference to this counter */
-		static void ReleaseWeakReference(TReferenceControllerBase<ESPMode::ThreadSafe>* ReferenceController)
-		{
-			// See ReleaseSharedReference for the same reasons that std::memory_order_release and std::memory_order_acquire are used in this function.
-
-			int32 OldWeakCount = ReferenceController->WeakReferenceCount.fetch_sub(1, std::memory_order_release);
-			checkSlow(OldWeakCount > 0);
-			if (OldWeakCount == 1)
-			{
-				std::atomic_thread_fence(std::memory_order_acquire);
-
-				// No more references to this reference count.  Destroy it!
-				delete ReferenceController;
-			}
-		}
-	};
-
-
-	template<>
-	struct FReferenceControllerOps<ESPMode::NotThreadSafe>
-	{
-		/** Returns the shared reference count */
-		static FORCEINLINE const int32 GetSharedReferenceCount(const TReferenceControllerBase<ESPMode::NotThreadSafe>* ReferenceController) TSAN_SAFE_UNSAFEPTR
-		{
-			return ReferenceController->SharedReferenceCount;
-		}
-
-		/** Adds a shared reference to this counter */
-		static FORCEINLINE void AddSharedReference(TReferenceControllerBase<ESPMode::NotThreadSafe>* ReferenceController) TSAN_SAFE_UNSAFEPTR
-		{
-			++ReferenceController->SharedReferenceCount;
-		}
-
-		/**
-		 * Adds a shared reference to this counter ONLY if there is already at least one reference
-		 *
-		 * @return  True if the shared reference was added successfully
-		 */
-		static bool ConditionallyAddSharedReference(TReferenceControllerBase<ESPMode::NotThreadSafe>* ReferenceController) TSAN_SAFE_UNSAFEPTR
-		{
-			if( ReferenceController->SharedReferenceCount == 0 )
-			{
-				// Never add a shared reference if the pointer has already expired
-				return false;
-			}
-
-			++ReferenceController->SharedReferenceCount;
-			return true;
-		}
-
-		/** Releases a shared reference to this counter */
-		static FORCEINLINE void ReleaseSharedReference(TReferenceControllerBase<ESPMode::NotThreadSafe>* ReferenceController) TSAN_SAFE_UNSAFEPTR
-		{
-			checkSlow( ReferenceController->SharedReferenceCount > 0 );
-
-			if( --ReferenceController->SharedReferenceCount == 0 )
-			{
-				// Last shared reference was released!  Destroy the referenced object.
-				ReferenceController->DestroyObject();
-
-				// No more shared referencers, so decrement the weak reference count by one.  When the weak
-				// reference count reaches zero, this object will be deleted.
-				ReleaseWeakReference(ReferenceController);
-			}
-		}
-
-		/** Adds a weak reference to this counter */
-		static FORCEINLINE void AddWeakReference(TReferenceControllerBase<ESPMode::NotThreadSafe>* ReferenceController) TSAN_SAFE_UNSAFEPTR
-		{
-			++ReferenceController->WeakReferenceCount;
-		}
-
-		/** Releases a weak reference to this counter */
-		static void ReleaseWeakReference(TReferenceControllerBase<ESPMode::NotThreadSafe>* ReferenceController) TSAN_SAFE_UNSAFEPTR
-		{
-			checkSlow( ReferenceController->WeakReferenceCount > 0 );
-
-			if( --ReferenceController->WeakReferenceCount == 0 )
-			{
-				// No more references to this reference count.  Destroy it!
-				delete ReferenceController;
-			}
-		}
-	};
-
-
 	/**
 	 * FSharedReferencer is a wrapper around a pointer to a reference controller that is used by either a
 	 * TSharedRef or a TSharedPtr to keep track of a referenced object's lifetime
@@ -469,8 +451,6 @@ namespace SharedPointerInternals
 	template< ESPMode Mode >
 	class FSharedReferencer
 	{
-		typedef FReferenceControllerOps<Mode> TOps;
-
 	public:
 
 		/** Constructor for an empty shared referencer object */
@@ -491,7 +471,7 @@ namespace SharedPointerInternals
 			// shared reference count
 			if( ReferenceController != nullptr )
 			{
-				TOps::AddSharedReference(ReferenceController);
+				ReferenceController->AddSharedReference();
 			}
 		}
 
@@ -514,7 +494,7 @@ namespace SharedPointerInternals
 				// Attempt to elevate a weak reference to a shared one.  For this to work, the object this
 				// weak counter is associated with must already have at least one shared reference.  We'll
 				// never revive a pointer that has already expired!
-				if( !TOps::ConditionallyAddSharedReference(ReferenceController) )
+				if( !ReferenceController->ConditionallyAddSharedReference() )
 				{
 					ReferenceController = nullptr;
 				}
@@ -533,14 +513,14 @@ namespace SharedPointerInternals
 				// Attempt to elevate a weak reference to a shared one.  For this to work, the object this
 				// weak counter is associated with must already have at least one shared reference.  We'll
 				// never revive a pointer that has already expired!
-				if( !TOps::ConditionallyAddSharedReference(ReferenceController) )
+				if( !ReferenceController->ConditionallyAddSharedReference() )
 				{
 					ReferenceController = nullptr;
 				}
 
 				// Tell the reference counter object that we're no longer referencing the object with
 				// this weak pointer
-				TOps::ReleaseWeakReference(InWeakReference.ReferenceController);
+				InWeakReference.ReferenceController->ReleaseWeakReference();
 				InWeakReference.ReferenceController = nullptr;
 			}
 		}
@@ -552,7 +532,7 @@ namespace SharedPointerInternals
 			{
 				// Tell the reference counter object that we're no longer referencing the object with
 				// this shared pointer
-				TOps::ReleaseSharedReference(ReferenceController);
+				ReferenceController->ReleaseSharedReference();
 			}
 		}
 
@@ -567,13 +547,13 @@ namespace SharedPointerInternals
 				// First, add a shared reference to the new object
 				if( NewReferenceController != nullptr )
 				{
-					TOps::AddSharedReference(NewReferenceController);
+					NewReferenceController->AddSharedReference();
 				}
 
 				// Release shared reference to the old object
 				if( ReferenceController != nullptr )
 				{
-					TOps::ReleaseSharedReference(ReferenceController);
+					ReferenceController->ReleaseSharedReference();
 				}
 
 				// Assume ownership of the assigned reference counter
@@ -599,7 +579,7 @@ namespace SharedPointerInternals
 				// Release shared reference to the old object
 				if( OldReferenceController != nullptr )
 				{
-					TOps::ReleaseSharedReference(OldReferenceController);
+					OldReferenceController->ReleaseSharedReference();
 				}
 			}
 
@@ -623,7 +603,7 @@ namespace SharedPointerInternals
 		 */
 		FORCEINLINE const int32 GetSharedReferenceCount() const
 		{
-			return ReferenceController != nullptr ? TOps::GetSharedReferenceCount(ReferenceController) : 0;
+			return ReferenceController != nullptr ? ReferenceController->GetSharedReferenceCount() : 0;
 		}
 
 		/**
@@ -656,8 +636,6 @@ namespace SharedPointerInternals
 	template< ESPMode Mode >
 	class FWeakReferencer
 	{
-		typedef FReferenceControllerOps<Mode> TOps;
-
 	public:
 
 		/** Default constructor with empty counter */
@@ -672,7 +650,7 @@ namespace SharedPointerInternals
 			// If the weak referencer has a valid controller, then go ahead and add a weak reference to it!
 			if( ReferenceController != nullptr )
 			{
-				TOps::AddWeakReference(ReferenceController);
+				ReferenceController->AddWeakReference();
 			}
 		}
 
@@ -690,7 +668,7 @@ namespace SharedPointerInternals
 			// If the shared referencer had a valid controller, then go ahead and add a weak reference to it!
 			if( ReferenceController != nullptr )
 			{
-				TOps::AddWeakReference(ReferenceController);
+				ReferenceController->AddWeakReference();
 			}
 		}
 
@@ -701,7 +679,7 @@ namespace SharedPointerInternals
 			{
 				// Tell the reference counter object that we're no longer referencing the object with
 				// this weak pointer
-				TOps::ReleaseWeakReference(ReferenceController);
+				ReferenceController->ReleaseWeakReference();
 			}
 		}
 		
@@ -723,7 +701,7 @@ namespace SharedPointerInternals
 			InWeakReference.ReferenceController = nullptr;
 			if( OldReferenceController != nullptr )
 			{
-				TOps::ReleaseWeakReference(OldReferenceController);
+				OldReferenceController->ReleaseWeakReference();
 			}
 
 			return *this;
@@ -745,7 +723,7 @@ namespace SharedPointerInternals
 		 */
 		FORCEINLINE const bool IsValid() const
 		{
-			return ReferenceController != nullptr && TOps::GetSharedReferenceCount(ReferenceController) > 0;
+			return ReferenceController != nullptr && ReferenceController->GetSharedReferenceCount() > 0;
 		}
 
 	private:
@@ -760,13 +738,13 @@ namespace SharedPointerInternals
 				// First, add a weak reference to the new object
 				if( NewReferenceController != nullptr )
 				{
-					TOps::AddWeakReference(NewReferenceController);
+					NewReferenceController->AddWeakReference();
 				}
 
 				// Release weak reference to the old object
 				if( ReferenceController != nullptr )
 				{
-					TOps::ReleaseWeakReference(ReferenceController);
+					ReferenceController->ReleaseWeakReference();
 				}
 
 				// Assume ownership of the assigned reference counter
