New emplacement constructor for TInlineValue.

diff --git a/Engine/Source/Runtime/Core/Public/Misc/InlineValue.h b/Engine/Source/Runtime/Core/Public/Misc/InlineValue.h
--- a/Engine/Source/Runtime/Core/Public/Misc/InlineValue.h
+++ b/Engine/Source/Runtime/Core/Public/Misc/InlineValue.h
@@ -41,6 +41,16 @@ public:
 		InitializeFrom<typename TDecay<T>::Type>(Forward<T>(In));
 	}
 
+	/**
+	 * In-place construction of BaseType from a set of arguments.
+	 */
+	template<typename... ArgTypes>
+	TInlineValue(EInPlace, ArgTypes&&... Args)
+		: bIsValid(false)
+	{
+		InitializeFrom<BaseType>(Forward<ArgTypes>(Args)...);
+	}
+
 	/**
 	 * Destructor
 	 */
