Revert TUniquePtr's Deleter back to private, as it seems unnecessary.

diff --git a/Engine/Source/Runtime/Core/Public/Templates/UniquePtr.h b/Engine/Source/Runtime/Core/Public/Templates/UniquePtr.h
--- a/Engine/Source/Runtime/Core/Public/Templates/UniquePtr.h
+++ b/Engine/Source/Runtime/Core/Public/Templates/UniquePtr.h
@@ -101,7 +101,7 @@ struct TDefaultDelete<T[]>
 };
 
 template <typename T, typename Deleter = TDefaultDelete<T>>
-class TUniquePtr : public /*private*/ Deleter // @todo loadtime: can we go back to private? I get this: error C2243: 'static_cast': conversion from 'T *' to 'Base *' exists, but is inaccessible
+class TUniquePtr : private Deleter
 {
 	DECLARE_INLINE_TYPE_LAYOUT_EXPLICIT_BASES(TUniquePtr, NonVirtual, Deleter);
 
