Log the hitch threshold when the hitch detector is enabled or the hitch threshold changes at runtime

diff --git a/Engine/Source/Runtime/Core/Private/HAL/ThreadHeartBeat.cpp b/Engine/Source/Runtime/Core/Private/HAL/ThreadHeartBeat.cpp
--- a/Engine/Source/Runtime/Core/Private/HAL/ThreadHeartBeat.cpp
+++ b/Engine/Source/Runtime/Core/Private/HAL/ThreadHeartBeat.cpp
@@ -953,6 +953,8 @@ void FGameThreadHitchHeartBeatThreaded::InitSettings()
 	static float CmdLine_HangDuration = 0.0f;
 	static bool CmdLine_StackWalk = false;
 
+	float OldHangDuration = HangDuration;
+
 	if (bFirst)
 	{
 		bHasCmdLine = FParse::Value(FCommandLine::Get(), TEXT("hitchdetection="), CmdLine_HangDuration);
@@ -1011,7 +1013,12 @@ void FGameThreadHitchHeartBeatThreaded::InitSettings()
 			bWalkStackOnHitch = false;
 		}
 	}
-	
+
+	if (OldHangDuration != HangDuration)
+	{
+		UE_LOG(LogCore, Display, TEXT("Hitch detector threshold: %dms"), int32(HangDuration * 1000.0f));
+	}
+
 	// Start the heart beat thread if it hasn't already been started.
 	if (Thread == nullptr && (FPlatformProcess::SupportsMultithreading() || FForkProcessHelper::SupportsMultithreadingPostFork()) && HangDuration > 0)
 	{
