Fixing non-unity issue with missing include/forward decl for FString in LogVerbosity
(horde# 85342)
#rnx

diff --git a/Engine/Source/Runtime/Core/Private/Logging/LogVerbosity.cpp b/Engine/Source/Runtime/Core/Private/Logging/LogVerbosity.cpp
--- a/Engine/Source/Runtime/Core/Private/Logging/LogVerbosity.cpp
+++ b/Engine/Source/Runtime/Core/Private/Logging/LogVerbosity.cpp
@@ -1,6 +1,7 @@
 // Copyright Epic Games, Inc. All Rights Reserved.
 
 #include "Logging/LogVerbosity.h"
+#include "Containers/UnrealString.h"
 
 const TCHAR* ToString(ELogVerbosity::Type Verbosity)
 {
diff --git a/Engine/Source/Runtime/Core/Public/Logging/LogVerbosity.h b/Engine/Source/Runtime/Core/Public/Logging/LogVerbosity.h
--- a/Engine/Source/Runtime/Core/Public/Logging/LogVerbosity.h
+++ b/Engine/Source/Runtime/Core/Public/Logging/LogVerbosity.h
@@ -4,6 +4,8 @@
 
 #include "CoreTypes.h"
 
+class FString;
+
 /** 
  * Enum that defines the verbosity levels of the logging system.
  * Also defines some non-verbosity levels that are hacks that allow
