move MinimalWindowsApi header/source from Windows folder to Microsoft folder, allowing it to be used on all Microsoft platforms.

#rnx

diff --git a/Engine/Source/Runtime/Core/Private/Microsoft/MinimalWindowsApi.cpp b/Engine/Source/Runtime/Core/Private/Microsoft/MinimalWindowsApi.cpp
new file mode 100644
--- /dev/null
+++ b/Engine/Source/Runtime/Core/Private/Microsoft/MinimalWindowsApi.cpp
@@ -0,0 +1,164 @@
+// Copyright Epic Games, Inc. All Rights Reserved.
+
+#include "Microsoft/MinimalWindowsApi.h"
+#include "Microsoft/WindowsHWrapper.h"
+
+// Check that constants are what we expect them to be
+static_assert(WINDOWS_MAX_PATH == MAX_PATH, "Value of WINDOWSAPI_MAX_PATH is not correct");
+static_assert(WINDOWS_PF_COMPARE_EXCHANGE128 == PF_COMPARE_EXCHANGE128, "Value of WINDOWS_PF_COMPARE_EXCHANGE128 is not correct");
+
+// Check that AllocTlsSlot() returns INDEX_NONE on failure
+static_assert(static_cast<uint32>(INDEX_NONE) == TLS_OUT_OF_INDEXES, "TLS_OUT_OF_INDEXES is different from INDEX_NONE, change FWindowsPlatformTLS::AllocTlsSlot() implementation.");
+
+// Check the size and alignment of the OVERLAPPED mirror
+static_assert(sizeof(Windows::OVERLAPPED) == sizeof(OVERLAPPED), "Size of Windows::OVERLAPPED must match definition in Windows.h");
+static_assert(__alignof(Windows::OVERLAPPED) == __alignof(OVERLAPPED), "Alignment of Windows::OVERLAPPED must match definition in Windows.h");
+
+// Check the size and alignment of the CRITICAL_SECTION mirror
+static_assert(sizeof(Windows::CRITICAL_SECTION) == sizeof(CRITICAL_SECTION), "Size of Windows::CRITICAL_SECTION must match definition in Windows.h");
+static_assert(__alignof(Windows::CRITICAL_SECTION) == __alignof(CRITICAL_SECTION), "Alignment of Windows::CRITICAL_SECTION must match definition in Windows.h");
+
+// Check the size and alignment of the LARGE_INTEGER mirror
+static_assert(sizeof(Windows::LARGE_INTEGER) == sizeof(LARGE_INTEGER), "Size of Windows::LARGE_INTEGER must match definition in Windows.h");
+static_assert(__alignof(Windows::LARGE_INTEGER) == __alignof(LARGE_INTEGER), "Alignment of Windows::LARGE_INTEGER must match definition in Windows.h");
+
+#ifdef __clang__
+namespace Windows
+{
+	MINIMAL_WINDOWS_API HMODULE WINAPI LoadLibraryW(LPCTSTR lpFileName)
+	{
+		return ::LoadLibraryW(lpFileName);
+	}
+
+	MINIMAL_WINDOWS_API BOOL WINAPI FreeLibrary(HMODULE hModule)
+	{
+		return ::FreeLibrary(hModule);
+	}
+
+	MINIMAL_WINDOWS_API void WINAPI InitializeCriticalSection(LPCRITICAL_SECTION lpCriticalSection)
+	{
+		::InitializeCriticalSection(lpCriticalSection);
+	}
+
+	MINIMAL_WINDOWS_API BOOL WINAPI InitializeCriticalSectionAndSpinCount(LPCRITICAL_SECTION lpCriticalSection, DWORD dwSpinCount)
+	{
+		return ::InitializeCriticalSectionAndSpinCount(lpCriticalSection, dwSpinCount);
+	}
+
+	MINIMAL_WINDOWS_API DWORD WINAPI SetCriticalSectionSpinCount(LPCRITICAL_SECTION lpCriticalSection, DWORD dwSpinCount)
+	{
+		return ::SetCriticalSectionSpinCount(lpCriticalSection, dwSpinCount);
+	}
+
+	MINIMAL_WINDOWS_API BOOL WINAPI TryEnterCriticalSection(LPCRITICAL_SECTION lpCriticalSection)
+	{
+		return ::TryEnterCriticalSection(lpCriticalSection);
+	}
+
+	MINIMAL_WINDOWS_API void WINAPI EnterCriticalSection(LPCRITICAL_SECTION lpCriticalSection)
+	{
+		::EnterCriticalSection(lpCriticalSection);
+	}
+
+	MINIMAL_WINDOWS_API void WINAPI LeaveCriticalSection(LPCRITICAL_SECTION lpCriticalSection)
+	{
+		::LeaveCriticalSection(lpCriticalSection);
+	}
+
+	MINIMAL_WINDOWS_API void WINAPI DeleteCriticalSection(LPCRITICAL_SECTION lpCriticalSection)
+	{
+		::DeleteCriticalSection(lpCriticalSection);
+	}
+
+	MINIMAL_WINDOWS_API void WINAPI InitializeSRWLock(PSRWLOCK SRWLock)
+	{
+		::InitializeSRWLock((::PSRWLOCK)SRWLock);
+	}
+
+	MINIMAL_WINDOWS_API void WINAPI AcquireSRWLockShared(PSRWLOCK SRWLock)
+	{
+		::AcquireSRWLockShared((::PSRWLOCK)SRWLock);
+	}
+
+	MINIMAL_WINDOWS_API BOOLEAN WINAPI TryAcquireSRWLockShared(PSRWLOCK SRWLock)
+	{
+		return ::TryAcquireSRWLockShared((::PSRWLOCK)SRWLock);
+	}
+
+	MINIMAL_WINDOWS_API BOOLEAN WINAPI TryAcquireSRWLockExclusive(PSRWLOCK SRWLock)
+	{
+		return ::TryAcquireSRWLockExclusive((::PSRWLOCK)SRWLock);
+	}
+
+	MINIMAL_WINDOWS_API void WINAPI ReleaseSRWLockShared(PSRWLOCK SRWLock)
+	{
+		::ReleaseSRWLockShared((::PSRWLOCK)SRWLock);
+	}
+
+	MINIMAL_WINDOWS_API void WINAPI AcquireSRWLockExclusive(PSRWLOCK SRWLock)
+	{
+		::AcquireSRWLockExclusive((::PSRWLOCK)SRWLock);
+	}
+
+	MINIMAL_WINDOWS_API void WINAPI ReleaseSRWLockExclusive(PSRWLOCK SRWLock)
+	{
+		::ReleaseSRWLockExclusive((::PSRWLOCK)SRWLock);
+	}
+
+	MINIMAL_WINDOWS_API BOOL WINAPI ConnectNamedPipe(HANDLE hNamedPipe, LPOVERLAPPED lpOverlapped)
+	{
+		return ::ConnectNamedPipe(hNamedPipe, lpOverlapped);
+	}
+
+	MINIMAL_WINDOWS_API BOOL WINAPI GetOverlappedResult(HANDLE hFile, LPOVERLAPPED lpOverlapped, LPDWORD lpNumberOfBytesTransferred, BOOL bWait)
+	{
+		return ::GetOverlappedResult(hFile, lpOverlapped, lpNumberOfBytesTransferred, bWait);
+	}
+
+	MINIMAL_WINDOWS_API BOOL WINAPI WriteFile(HANDLE hFile, LPCVOID lpBuffer, DWORD nNumberOfBytesToWrite, LPDWORD lpNumberOfBytesWritten, LPOVERLAPPED lpOverlapped)
+	{
+		return ::WriteFile(hFile, lpBuffer, nNumberOfBytesToWrite, lpNumberOfBytesWritten, lpOverlapped);
+	}
+
+	MINIMAL_WINDOWS_API BOOL WINAPI ReadFile(HANDLE hFile, LPVOID lpBuffer, DWORD nNumberOfBytesToRead, LPDWORD lpNumberOfBytesRead, LPOVERLAPPED lpOverlapped)
+	{
+		return ::ReadFile(hFile, lpBuffer, nNumberOfBytesToRead, lpNumberOfBytesRead, lpOverlapped);
+	}
+
+	MINIMAL_WINDOWS_API BOOL WINAPI QueryPerformanceCounter(LPLARGE_INTEGER Cycles)
+	{
+		return ::QueryPerformanceCounter(Cycles);
+	}
+
+	MINIMAL_WINDOWS_API DWORD WINAPI GetCurrentThreadId()
+	{
+		return ::GetCurrentThreadId();
+	}
+
+	MINIMAL_WINDOWS_API DWORD WINAPI TlsAlloc()
+	{
+		return ::TlsAlloc();
+	}
+
+	MINIMAL_WINDOWS_API LPVOID WINAPI TlsGetValue(DWORD dwTlsIndex)
+	{
+		return ::TlsGetValue(dwTlsIndex);
+	}
+
+	MINIMAL_WINDOWS_API BOOL WINAPI TlsSetValue(DWORD dwTlsIndex, LPVOID lpTlsValue)
+	{
+		return ::TlsSetValue(dwTlsIndex, lpTlsValue);
+	}
+
+	MINIMAL_WINDOWS_API BOOL WINAPI TlsFree(DWORD dwTlsIndex)
+	{
+		return ::TlsFree(dwTlsIndex);
+	}
+
+	MINIMAL_WINDOWS_API BOOL WINAPI IsProcessorFeaturePresent(DWORD ProcessorFeature)
+	{
+		return ::IsProcessorFeaturePresent(ProcessorFeature);
+	}
+}
+#endif
+
diff --git a/Engine/Source/Runtime/Core/Private/Windows/MinimalWindowsApi.cpp b/Engine/Source/Runtime/Core/Private/Windows/MinimalWindowsApi.cpp
deleted file mode 100644
diff --git a/Engine/Source/Runtime/Core/Public/Microsoft/MinimalWindowsApi.h b/Engine/Source/Runtime/Core/Public/Microsoft/MinimalWindowsApi.h
new file mode 100644
--- /dev/null
+++ b/Engine/Source/Runtime/Core/Public/Microsoft/MinimalWindowsApi.h
@@ -0,0 +1,253 @@
+// Copyright Epic Games, Inc. All Rights Reserved.
+
+//=====================================================================================================================
+// Implementation of a minimal subset of the Windows API required for inline function definitions and platform-specific
+// interfaces in Core. Use of this file allows avoiding including large platform headers in the public engine API.
+//
+// Win32 API functions are declared in the "Windows" namespace to avoid conflicts if the real Windows.h is included 
+// later, but are mapped to the same imported symbols by the linker due to C-style linkage.
+//=====================================================================================================================
+
+#pragma once
+
+#include "CoreTypes.h"
+
+#ifdef __clang__
+	#define MINIMAL_WINDOWS_API CORE_API
+#else
+	#define MINIMAL_WINDOWS_API extern "C" __declspec(dllimport)
+#endif
+
+// The 10.0.18362.0 SDK introduces an error if the packing isn't the default for the platform.
+PRAGMA_PUSH_PLATFORM_DEFAULT_PACKING
+
+// Undefine Windows types that we're going to redefine. This should be done by HideWindowsPlatformTypes.h after including any system headers, 
+// but it's less friction to just undefine them here.
+#pragma push_macro("TRUE")
+#pragma push_macro("FALSE")
+#undef TRUE
+#undef FALSE
+
+// Use strongly typed handles
+#ifndef STRICT
+#define STRICT
+#endif
+
+// With STRICT enabled, handles are implemented as opaque struct pointers. We can prototype these structs here and 
+// typedef them under the Windows namespace below. Since typedefs are only aliases rather than types in their own 
+// right, this allows using handles from the Windows:: namespace interchangably with their native definition.
+struct HINSTANCE__;
+struct HWND__;
+struct HKEY__;
+struct HDC__;
+struct HICON__;
+struct _RTL_SRWLOCK;
+
+// Other types
+struct tagPROCESSENTRY32W;
+struct _EXCEPTION_POINTERS;
+struct _OVERLAPPED;
+struct _RTL_CRITICAL_SECTION;
+union _LARGE_INTEGER;
+
+// Global constants
+#define WINDOWS_MAX_PATH 260
+#define WINDOWS_PF_COMPARE_EXCHANGE128	14
+
+// Standard calling convention for Win32 functions
+#define WINAPI __stdcall
+
+// Minimal subset of the Windows API required for interfaces and inline functions
+namespace Windows
+{
+	// Typedefs for basic Windows types
+	typedef int32 BOOL;
+	typedef unsigned char BYTE;
+	typedef BYTE  BOOLEAN;
+	typedef unsigned long DWORD;
+	typedef DWORD* LPDWORD;
+	typedef long LONG;
+	typedef long* LPLONG;
+	typedef int64 LONGLONG;
+	typedef LONGLONG* LPLONGLONG;
+	typedef void* LPVOID;
+	typedef const void* LPCVOID;
+	typedef const wchar_t* LPCTSTR;
+
+	// Typedefs for standard handles
+	typedef void* HANDLE;
+	typedef HINSTANCE__* HINSTANCE;
+	typedef HINSTANCE HMODULE;
+	typedef HWND__* HWND;
+	typedef HKEY__* HKEY;
+	typedef HDC__* HDC;
+	typedef HICON__* HICON;
+	typedef HICON__* HCURSOR;
+
+	// Other typedefs
+	typedef tagPROCESSENTRY32W PROCESSENTRY32;
+	typedef _EXCEPTION_POINTERS* LPEXCEPTION_POINTERS;
+	typedef _RTL_CRITICAL_SECTION* LPCRITICAL_SECTION;
+	typedef _OVERLAPPED* LPOVERLAPPED;
+	typedef _LARGE_INTEGER* LPLARGE_INTEGER;
+
+	typedef _RTL_SRWLOCK RTL_SRWLOCK, *PRTL_SRWLOCK;
+	typedef RTL_SRWLOCK *PSRWLOCK;
+
+	// Opaque SRWLOCK structure
+	struct SRWLOCK
+	{
+		void* Ptr;
+	};
+
+	// Constants
+	static const BOOL TRUE = 1;
+	static const BOOL FALSE = 0;
+
+	// Modules
+	MINIMAL_WINDOWS_API HMODULE WINAPI LoadLibraryW(LPCTSTR lpFileName);
+	MINIMAL_WINDOWS_API BOOL WINAPI FreeLibrary(HMODULE hModule);
+
+	// Critical sections
+	MINIMAL_WINDOWS_API void WINAPI InitializeCriticalSection(LPCRITICAL_SECTION lpCriticalSection);
+	MINIMAL_WINDOWS_API BOOL WINAPI InitializeCriticalSectionAndSpinCount(LPCRITICAL_SECTION lpCriticalSection, DWORD dwSpinCount);
+	MINIMAL_WINDOWS_API DWORD WINAPI SetCriticalSectionSpinCount(LPCRITICAL_SECTION lpCriticalSection, DWORD dwSpinCount);
+	MINIMAL_WINDOWS_API BOOL WINAPI TryEnterCriticalSection(LPCRITICAL_SECTION lpCriticalSection);
+	MINIMAL_WINDOWS_API void WINAPI EnterCriticalSection(LPCRITICAL_SECTION lpCriticalSection);
+	MINIMAL_WINDOWS_API void WINAPI LeaveCriticalSection(LPCRITICAL_SECTION lpCriticalSection);
+	MINIMAL_WINDOWS_API void WINAPI DeleteCriticalSection(LPCRITICAL_SECTION lpCriticalSection);
+
+	MINIMAL_WINDOWS_API void WINAPI InitializeSRWLock(PSRWLOCK SRWLock);
+	MINIMAL_WINDOWS_API void WINAPI AcquireSRWLockShared(PSRWLOCK SRWLock);
+	MINIMAL_WINDOWS_API void WINAPI ReleaseSRWLockShared(PSRWLOCK SRWLock);
+	MINIMAL_WINDOWS_API void WINAPI AcquireSRWLockExclusive(PSRWLOCK SRWLock);
+	MINIMAL_WINDOWS_API void WINAPI ReleaseSRWLockExclusive(PSRWLOCK SRWLock);
+	MINIMAL_WINDOWS_API BOOLEAN WINAPI TryAcquireSRWLockExclusive(PSRWLOCK SRWLock);
+	MINIMAL_WINDOWS_API BOOLEAN WINAPI TryAcquireSRWLockShared(PSRWLOCK SRWLock);
+
+	FORCEINLINE void WINAPI InitializeSRWLock(SRWLOCK* SRWLock)
+	{
+		InitializeSRWLock((PSRWLOCK)SRWLock);
+	}
+
+	FORCEINLINE BOOLEAN WINAPI TryAcquireSRWLockExclusive(SRWLOCK* SRWLock)
+	{
+		return TryAcquireSRWLockExclusive((PSRWLOCK)SRWLock);
+	}
+
+	FORCEINLINE BOOLEAN WINAPI TryAcquireSRWLockShared(SRWLOCK* SRWLock)
+	{
+		return TryAcquireSRWLockShared((PSRWLOCK)SRWLock);
+	}
+	
+	FORCEINLINE void WINAPI AcquireSRWLockShared(SRWLOCK* SRWLock)
+	{
+		AcquireSRWLockShared((PSRWLOCK)SRWLock);
+	}
+
+	FORCEINLINE void WINAPI ReleaseSRWLockShared(SRWLOCK* SRWLock)
+	{
+		ReleaseSRWLockShared((PSRWLOCK)SRWLock);
+	}
+
+	FORCEINLINE void WINAPI AcquireSRWLockExclusive(SRWLOCK* SRWLock)
+	{
+		AcquireSRWLockExclusive((PSRWLOCK)SRWLock);
+	}
+
+	FORCEINLINE void WINAPI ReleaseSRWLockExclusive(SRWLOCK* SRWLock)
+	{
+		ReleaseSRWLockExclusive((PSRWLOCK)SRWLock);
+	}
+
+	// I/O
+	MINIMAL_WINDOWS_API BOOL WINAPI ConnectNamedPipe(HANDLE hNamedPipe, LPOVERLAPPED lpOverlapped);
+	MINIMAL_WINDOWS_API BOOL WINAPI GetOverlappedResult(HANDLE hFile, LPOVERLAPPED lpOverlapped, LPDWORD lpNumberOfBytesTransferred, BOOL bWait);
+	MINIMAL_WINDOWS_API BOOL WINAPI WriteFile(HANDLE hFile, LPCVOID lpBuffer, DWORD nNumberOfBytesToWrite, LPDWORD lpNumberOfBytesWritten, LPOVERLAPPED lpOverlapped);
+	MINIMAL_WINDOWS_API BOOL WINAPI ReadFile(HANDLE hFile, LPVOID lpBuffer, DWORD nNumberOfBytesToRead, LPDWORD lpNumberOfBytesRead, LPOVERLAPPED lpOverlapped);
+
+	// Timing
+	MINIMAL_WINDOWS_API BOOL WINAPI QueryPerformanceCounter(LPLARGE_INTEGER Cycles);
+
+	// Thread-local storage functions
+	MINIMAL_WINDOWS_API DWORD WINAPI GetCurrentThreadId();
+	MINIMAL_WINDOWS_API DWORD WINAPI TlsAlloc();
+	MINIMAL_WINDOWS_API LPVOID WINAPI TlsGetValue(DWORD dwTlsIndex);
+	MINIMAL_WINDOWS_API BOOL WINAPI TlsSetValue(DWORD dwTlsIndex, LPVOID lpTlsValue);
+	MINIMAL_WINDOWS_API BOOL WINAPI TlsFree(DWORD dwTlsIndex);
+
+	// System
+	MINIMAL_WINDOWS_API BOOL WINAPI IsProcessorFeaturePresent(DWORD ProcessorFeature);
+
+	// For structures which are opaque
+	struct CRITICAL_SECTION { void* Opaque1[1]; long Opaque2[2]; void* Opaque3[3]; };
+	struct OVERLAPPED { void *Opaque[3]; unsigned long Opaque2[2]; };
+	union LARGE_INTEGER { struct { DWORD LowPart; LONG HighPart; }; LONGLONG QuadPart; };
+
+	FORCEINLINE BOOL WINAPI ConnectNamedPipe(HANDLE hNamedPipe, OVERLAPPED* lpOverlapped)
+	{
+		return ConnectNamedPipe(hNamedPipe, (LPOVERLAPPED)lpOverlapped);
+	}
+
+	FORCEINLINE BOOL WINAPI GetOverlappedResult(HANDLE hFile, OVERLAPPED* lpOverlapped, LPDWORD lpNumberOfBytesTransferred, BOOL bWait)
+	{
+		return GetOverlappedResult(hFile, (LPOVERLAPPED)lpOverlapped, lpNumberOfBytesTransferred, bWait);
+	}
+
+	FORCEINLINE BOOL WINAPI WriteFile(HANDLE hFile, LPCVOID lpBuffer, DWORD nNumberOfBytesToWrite, LPDWORD lpNumberOfBytesWritten, OVERLAPPED* lpOverlapped)
+	{
+		return WriteFile(hFile, lpBuffer, nNumberOfBytesToWrite, lpNumberOfBytesWritten, (LPOVERLAPPED)lpOverlapped);
+	}
+
+	FORCEINLINE BOOL WINAPI ReadFile(HANDLE hFile, LPVOID lpBuffer, DWORD nNumberOfBytesToRead, LPDWORD lpNumberOfBytesRead, OVERLAPPED* lpOverlapped)
+	{
+		return ReadFile(hFile, lpBuffer, nNumberOfBytesToRead, lpNumberOfBytesRead, (LPOVERLAPPED)lpOverlapped);
+	}
+
+	FORCEINLINE void WINAPI InitializeCriticalSection(CRITICAL_SECTION* lpCriticalSection)
+	{
+		InitializeCriticalSection((LPCRITICAL_SECTION)lpCriticalSection);
+	}
+
+	FORCEINLINE BOOL WINAPI InitializeCriticalSectionAndSpinCount(CRITICAL_SECTION* lpCriticalSection, DWORD dwSpinCount)
+	{
+		return InitializeCriticalSectionAndSpinCount((LPCRITICAL_SECTION)lpCriticalSection, dwSpinCount);
+	}
+
+	FORCEINLINE DWORD WINAPI SetCriticalSectionSpinCount(CRITICAL_SECTION* lpCriticalSection, DWORD dwSpinCount)
+	{
+		return SetCriticalSectionSpinCount((LPCRITICAL_SECTION)lpCriticalSection, dwSpinCount);
+	}
+
+	FORCEINLINE BOOL WINAPI TryEnterCriticalSection(CRITICAL_SECTION* lpCriticalSection)
+	{
+		return TryEnterCriticalSection((LPCRITICAL_SECTION)lpCriticalSection);
+	}
+
+	FORCEINLINE void WINAPI EnterCriticalSection(CRITICAL_SECTION* lpCriticalSection)
+	{
+		EnterCriticalSection((LPCRITICAL_SECTION)lpCriticalSection);
+	}
+
+	FORCEINLINE void WINAPI LeaveCriticalSection(CRITICAL_SECTION* lpCriticalSection)
+	{
+		LeaveCriticalSection((LPCRITICAL_SECTION)lpCriticalSection);
+	}
+
+	FORCEINLINE void WINAPI DeleteCriticalSection(CRITICAL_SECTION* lpCriticalSection)
+	{
+		DeleteCriticalSection((LPCRITICAL_SECTION)lpCriticalSection);
+	}
+
+	FORCEINLINE BOOL WINAPI QueryPerformanceCounter(LARGE_INTEGER* Cycles)
+	{
+		return QueryPerformanceCounter((LPLARGE_INTEGER)Cycles);
+	}
+}
+
+// Restore the definitions for TRUE and FALSE
+#pragma pop_macro("FALSE")
+#pragma pop_macro("TRUE")
+
+// Restore the packing setting
+PRAGMA_POP_PLATFORM_DEFAULT_PACKING
diff --git a/Engine/Source/Runtime/Core/Public/Windows/MinimalWindowsApi.h b/Engine/Source/Runtime/Core/Public/Windows/MinimalWindowsApi.h
--- a/Engine/Source/Runtime/Core/Public/Windows/MinimalWindowsApi.h
+++ b/Engine/Source/Runtime/Core/Public/Windows/MinimalWindowsApi.h
@@ -1,258 +1,10 @@
 // Copyright Epic Games, Inc. All Rights Reserved.
 
-//=====================================================================================================================
-// Implementation of a minimal subset of the Windows API required for inline function definitions and platform-specific
-// interfaces in Core. Use of this file allows avoiding including large platform headers in the public engine API.
-//
-// Win32 API functions are declared in the "Windows" namespace to avoid conflicts if the real Windows.h is included 
-// later, but are mapped to the same imported symbols by the linker due to C-style linkage.
-//=====================================================================================================================
-
 #pragma once
 
 #if !PLATFORM_WINDOWS && !PLATFORM_HOLOLENS
 	#error this file is for PLATFORM_WINDOWS or PLATFORM_HOLOLENS only
 #endif
 
-
-#include "CoreTypes.h"
-
-#ifdef __clang__
-	#define MINIMAL_WINDOWS_API CORE_API
-#else
-	#define MINIMAL_WINDOWS_API extern "C" __declspec(dllimport)
-#endif
-
-// The 10.0.18362.0 SDK introduces an error if the packing isn't the default for the platform.
-PRAGMA_PUSH_PLATFORM_DEFAULT_PACKING
-
-// Undefine Windows types that we're going to redefine. This should be done by HideWindowsPlatformTypes.h after including any system headers, 
-// but it's less friction to just undefine them here.
-#pragma push_macro("TRUE")
-#pragma push_macro("FALSE")
-#undef TRUE
-#undef FALSE
-
-// Use strongly typed handles
-#ifndef STRICT
-#define STRICT
-#endif
-
-// With STRICT enabled, handles are implemented as opaque struct pointers. We can prototype these structs here and 
-// typedef them under the Windows namespace below. Since typedefs are only aliases rather than types in their own 
-// right, this allows using handles from the Windows:: namespace interchangably with their native definition.
-struct HINSTANCE__;
-struct HWND__;
-struct HKEY__;
-struct HDC__;
-struct HICON__;
-struct _RTL_SRWLOCK;
-
-// Other types
-struct tagPROCESSENTRY32W;
-struct _EXCEPTION_POINTERS;
-struct _OVERLAPPED;
-struct _RTL_CRITICAL_SECTION;
-union _LARGE_INTEGER;
-
-// Global constants
-#define WINDOWS_MAX_PATH 260
-#define WINDOWS_PF_COMPARE_EXCHANGE128	14
-
-// Standard calling convention for Win32 functions
-#define WINAPI __stdcall
-
-// Minimal subset of the Windows API required for interfaces and inline functions
-namespace Windows
-{
-	// Typedefs for basic Windows types
-	typedef int32 BOOL;
-	typedef unsigned char BYTE;
-	typedef BYTE  BOOLEAN;
-	typedef unsigned long DWORD;
-	typedef DWORD* LPDWORD;
-	typedef long LONG;
-	typedef long* LPLONG;
-	typedef int64 LONGLONG;
-	typedef LONGLONG* LPLONGLONG;
-	typedef void* LPVOID;
-	typedef const void* LPCVOID;
-	typedef const wchar_t* LPCTSTR;
-
-	// Typedefs for standard handles
-	typedef void* HANDLE;
-	typedef HINSTANCE__* HINSTANCE;
-	typedef HINSTANCE HMODULE;
-	typedef HWND__* HWND;
-	typedef HKEY__* HKEY;
-	typedef HDC__* HDC;
-	typedef HICON__* HICON;
-	typedef HICON__* HCURSOR;
-
-	// Other typedefs
-	typedef tagPROCESSENTRY32W PROCESSENTRY32;
-	typedef _EXCEPTION_POINTERS* LPEXCEPTION_POINTERS;
-	typedef _RTL_CRITICAL_SECTION* LPCRITICAL_SECTION;
-	typedef _OVERLAPPED* LPOVERLAPPED;
-	typedef _LARGE_INTEGER* LPLARGE_INTEGER;
-
-	typedef _RTL_SRWLOCK RTL_SRWLOCK, *PRTL_SRWLOCK;
-	typedef RTL_SRWLOCK *PSRWLOCK;
-
-	// Opaque SRWLOCK structure
-	struct SRWLOCK
-	{
-		void* Ptr;
-	};
-
-	// Constants
-	static const BOOL TRUE = 1;
-	static const BOOL FALSE = 0;
-
-	// Modules
-	MINIMAL_WINDOWS_API HMODULE WINAPI LoadLibraryW(LPCTSTR lpFileName);
-	MINIMAL_WINDOWS_API BOOL WINAPI FreeLibrary(HMODULE hModule);
-
-	// Critical sections
-	MINIMAL_WINDOWS_API void WINAPI InitializeCriticalSection(LPCRITICAL_SECTION lpCriticalSection);
-	MINIMAL_WINDOWS_API BOOL WINAPI InitializeCriticalSectionAndSpinCount(LPCRITICAL_SECTION lpCriticalSection, DWORD dwSpinCount);
-	MINIMAL_WINDOWS_API DWORD WINAPI SetCriticalSectionSpinCount(LPCRITICAL_SECTION lpCriticalSection, DWORD dwSpinCount);
-	MINIMAL_WINDOWS_API BOOL WINAPI TryEnterCriticalSection(LPCRITICAL_SECTION lpCriticalSection);
-	MINIMAL_WINDOWS_API void WINAPI EnterCriticalSection(LPCRITICAL_SECTION lpCriticalSection);
-	MINIMAL_WINDOWS_API void WINAPI LeaveCriticalSection(LPCRITICAL_SECTION lpCriticalSection);
-	MINIMAL_WINDOWS_API void WINAPI DeleteCriticalSection(LPCRITICAL_SECTION lpCriticalSection);
-
-	MINIMAL_WINDOWS_API void WINAPI InitializeSRWLock(PSRWLOCK SRWLock);
-	MINIMAL_WINDOWS_API void WINAPI AcquireSRWLockShared(PSRWLOCK SRWLock);
-	MINIMAL_WINDOWS_API void WINAPI ReleaseSRWLockShared(PSRWLOCK SRWLock);
-	MINIMAL_WINDOWS_API void WINAPI AcquireSRWLockExclusive(PSRWLOCK SRWLock);
-	MINIMAL_WINDOWS_API void WINAPI ReleaseSRWLockExclusive(PSRWLOCK SRWLock);
-	MINIMAL_WINDOWS_API BOOLEAN WINAPI TryAcquireSRWLockExclusive(PSRWLOCK SRWLock);
-	MINIMAL_WINDOWS_API BOOLEAN WINAPI TryAcquireSRWLockShared(PSRWLOCK SRWLock);
-
-	FORCEINLINE void WINAPI InitializeSRWLock(SRWLOCK* SRWLock)
-	{
-		InitializeSRWLock((PSRWLOCK)SRWLock);
-	}
-
-	FORCEINLINE BOOLEAN WINAPI TryAcquireSRWLockExclusive(SRWLOCK* SRWLock)
-	{
-		return TryAcquireSRWLockExclusive((PSRWLOCK)SRWLock);
-	}
-
-	FORCEINLINE BOOLEAN WINAPI TryAcquireSRWLockShared(SRWLOCK* SRWLock)
-	{
-		return TryAcquireSRWLockShared((PSRWLOCK)SRWLock);
-	}
-	
-	FORCEINLINE void WINAPI AcquireSRWLockShared(SRWLOCK* SRWLock)
-	{
-		AcquireSRWLockShared((PSRWLOCK)SRWLock);
-	}
-
-	FORCEINLINE void WINAPI ReleaseSRWLockShared(SRWLOCK* SRWLock)
-	{
-		ReleaseSRWLockShared((PSRWLOCK)SRWLock);
-	}
-
-	FORCEINLINE void WINAPI AcquireSRWLockExclusive(SRWLOCK* SRWLock)
-	{
-		AcquireSRWLockExclusive((PSRWLOCK)SRWLock);
-	}
-
-	FORCEINLINE void WINAPI ReleaseSRWLockExclusive(SRWLOCK* SRWLock)
-	{
-		ReleaseSRWLockExclusive((PSRWLOCK)SRWLock);
-	}
-
-	// I/O
-	MINIMAL_WINDOWS_API BOOL WINAPI ConnectNamedPipe(HANDLE hNamedPipe, LPOVERLAPPED lpOverlapped);
-	MINIMAL_WINDOWS_API BOOL WINAPI GetOverlappedResult(HANDLE hFile, LPOVERLAPPED lpOverlapped, LPDWORD lpNumberOfBytesTransferred, BOOL bWait);
-	MINIMAL_WINDOWS_API BOOL WINAPI WriteFile(HANDLE hFile, LPCVOID lpBuffer, DWORD nNumberOfBytesToWrite, LPDWORD lpNumberOfBytesWritten, LPOVERLAPPED lpOverlapped);
-	MINIMAL_WINDOWS_API BOOL WINAPI ReadFile(HANDLE hFile, LPVOID lpBuffer, DWORD nNumberOfBytesToRead, LPDWORD lpNumberOfBytesRead, LPOVERLAPPED lpOverlapped);
-
-	// Timing
-	MINIMAL_WINDOWS_API BOOL WINAPI QueryPerformanceCounter(LPLARGE_INTEGER Cycles);
-
-	// Thread-local storage functions
-	MINIMAL_WINDOWS_API DWORD WINAPI GetCurrentThreadId();
-	MINIMAL_WINDOWS_API DWORD WINAPI TlsAlloc();
-	MINIMAL_WINDOWS_API LPVOID WINAPI TlsGetValue(DWORD dwTlsIndex);
-	MINIMAL_WINDOWS_API BOOL WINAPI TlsSetValue(DWORD dwTlsIndex, LPVOID lpTlsValue);
-	MINIMAL_WINDOWS_API BOOL WINAPI TlsFree(DWORD dwTlsIndex);
-
-	// System
-	MINIMAL_WINDOWS_API BOOL WINAPI IsProcessorFeaturePresent(DWORD ProcessorFeature);
-
-	// For structures which are opaque
-	struct CRITICAL_SECTION { void* Opaque1[1]; long Opaque2[2]; void* Opaque3[3]; };
-	struct OVERLAPPED { void *Opaque[3]; unsigned long Opaque2[2]; };
-	union LARGE_INTEGER { struct { DWORD LowPart; LONG HighPart; }; LONGLONG QuadPart; };
-
-	FORCEINLINE BOOL WINAPI ConnectNamedPipe(HANDLE hNamedPipe, OVERLAPPED* lpOverlapped)
-	{
-		return ConnectNamedPipe(hNamedPipe, (LPOVERLAPPED)lpOverlapped);
-	}
-
-	FORCEINLINE BOOL WINAPI GetOverlappedResult(HANDLE hFile, OVERLAPPED* lpOverlapped, LPDWORD lpNumberOfBytesTransferred, BOOL bWait)
-	{
-		return GetOverlappedResult(hFile, (LPOVERLAPPED)lpOverlapped, lpNumberOfBytesTransferred, bWait);
-	}
-
-	FORCEINLINE BOOL WINAPI WriteFile(HANDLE hFile, LPCVOID lpBuffer, DWORD nNumberOfBytesToWrite, LPDWORD lpNumberOfBytesWritten, OVERLAPPED* lpOverlapped)
-	{
-		return WriteFile(hFile, lpBuffer, nNumberOfBytesToWrite, lpNumberOfBytesWritten, (LPOVERLAPPED)lpOverlapped);
-	}
-
-	FORCEINLINE BOOL WINAPI ReadFile(HANDLE hFile, LPVOID lpBuffer, DWORD nNumberOfBytesToRead, LPDWORD lpNumberOfBytesRead, OVERLAPPED* lpOverlapped)
-	{
-		return ReadFile(hFile, lpBuffer, nNumberOfBytesToRead, lpNumberOfBytesRead, (LPOVERLAPPED)lpOverlapped);
-	}
-
-	FORCEINLINE void WINAPI InitializeCriticalSection(CRITICAL_SECTION* lpCriticalSection)
-	{
-		InitializeCriticalSection((LPCRITICAL_SECTION)lpCriticalSection);
-	}
-
-	FORCEINLINE BOOL WINAPI InitializeCriticalSectionAndSpinCount(CRITICAL_SECTION* lpCriticalSection, DWORD dwSpinCount)
-	{
-		return InitializeCriticalSectionAndSpinCount((LPCRITICAL_SECTION)lpCriticalSection, dwSpinCount);
-	}
-
-	FORCEINLINE DWORD WINAPI SetCriticalSectionSpinCount(CRITICAL_SECTION* lpCriticalSection, DWORD dwSpinCount)
-	{
-		return SetCriticalSectionSpinCount((LPCRITICAL_SECTION)lpCriticalSection, dwSpinCount);
-	}
-
-	FORCEINLINE BOOL WINAPI TryEnterCriticalSection(CRITICAL_SECTION* lpCriticalSection)
-	{
-		return TryEnterCriticalSection((LPCRITICAL_SECTION)lpCriticalSection);
-	}
-
-	FORCEINLINE void WINAPI EnterCriticalSection(CRITICAL_SECTION* lpCriticalSection)
-	{
-		EnterCriticalSection((LPCRITICAL_SECTION)lpCriticalSection);
-	}
-
-	FORCEINLINE void WINAPI LeaveCriticalSection(CRITICAL_SECTION* lpCriticalSection)
-	{
-		LeaveCriticalSection((LPCRITICAL_SECTION)lpCriticalSection);
-	}
-
-	FORCEINLINE void WINAPI DeleteCriticalSection(CRITICAL_SECTION* lpCriticalSection)
-	{
-		DeleteCriticalSection((LPCRITICAL_SECTION)lpCriticalSection);
-	}
-
-	FORCEINLINE BOOL WINAPI QueryPerformanceCounter(LARGE_INTEGER* Cycles)
-	{
-		return QueryPerformanceCounter((LPLARGE_INTEGER)Cycles);
-	}
-}
-
-// Restore the definitions for TRUE and FALSE
-#pragma pop_macro("FALSE")
-#pragma pop_macro("TRUE")
-
-// Restore the packing setting
-PRAGMA_POP_PLATFORM_DEFAULT_PACKING
+// redirect to new file
+#include "Microsoft/MinimalWindowsApi.h"
