Added a FMemoryView constructor to FMemoryReaderView

#rnx

diff --git a/Engine/Source/Runtime/Core/Public/Serialization/MemoryReader.h b/Engine/Source/Runtime/Core/Public/Serialization/MemoryReader.h
--- a/Engine/Source/Runtime/Core/Public/Serialization/MemoryReader.h
+++ b/Engine/Source/Runtime/Core/Public/Serialization/MemoryReader.h
@@ -5,6 +5,7 @@
 #include "CoreTypes.h"
 #include "HAL/UnrealMemory.h"
 #include "Math/UnrealMathUtility.h"
+#include "Memory/MemoryView.h"
 #include "Containers/UnrealString.h"
 #include "Serialization/MemoryArchive.h"
 #include "Containers/ArrayView.h"
@@ -68,7 +69,7 @@ private:
 };
 
 /**
- * Archive for reading arbitrary data from the specified array view
+ * Archive for reading arbitrary data from the specified memory view
  */
 class FMemoryReaderView : public FMemoryArchive
 {
@@ -86,17 +87,17 @@ public:
 
 	virtual int64 TotalSize() override
 	{
-		return FMath::Min((int64)Bytes.Num(), LimitSize);
+		return FMath::Min(static_cast<int64>(Bytes.GetSize()), LimitSize);
 	}
 
-	void Serialize( void* Data, int64 Num )
+	void Serialize(void* Data, int64 Num)
 	{
 		if (Num && !IsError())
 		{
 			// Only serialize if we have the requested amount of data
 			if (Offset + Num <= TotalSize())
 			{
-				FMemory::Memcpy( Data, &Bytes[(int32)Offset], Num );
+				FMutableMemoryView(Data, static_cast<uint64>(Num)).CopyFrom(Bytes.Mid(Offset, Num));
 				Offset += Num;
 			}
 			else
@@ -106,8 +107,13 @@ public:
 		}
 	}
 
-	explicit FMemoryReaderView( TArrayView<const uint8> InBytes, bool bIsPersistent = false )
-		: Bytes    (InBytes)
+	explicit FMemoryReaderView(TArrayView<const uint8> InBytes, bool bIsPersistent = false)
+		: FMemoryReaderView(MakeMemoryView(InBytes), bIsPersistent)
+	{
+	}
+
+	explicit FMemoryReaderView(FMemoryView InBytes, bool bIsPersistent = false)
+		: Bytes(InBytes)
 		, LimitSize(INT64_MAX)
 	{
 		this->SetIsLoading(true);
@@ -121,7 +127,6 @@ public:
 	}
 
 private:
-	TArrayView<const uint8> Bytes;
+	FMemoryView          Bytes;
 	int64                LimitSize;
 };
-
