PR #8108: Replay compatible change list override (Contributed by [REDACTED_PERSON])

diff --git a/Engine/Source/Runtime/Core/Private/Misc/NetworkVersion.cpp b/Engine/Source/Runtime/Core/Private/Misc/NetworkVersion.cpp
--- a/Engine/Source/Runtime/Core/Private/Misc/NetworkVersion.cpp
+++ b/Engine/Source/Runtime/Core/Private/Misc/NetworkVersion.cpp
@@ -12,6 +12,7 @@ DEFINE_LOG_CATEGORY( LogNetVersion );
 
 FNetworkVersion::FGetLocalNetworkVersionOverride FNetworkVersion::GetLocalNetworkVersionOverride;
 FNetworkVersion::FIsNetworkCompatibleOverride FNetworkVersion::IsNetworkCompatibleOverride;
+FNetworkVersion::FGetReplayCompatibleChangeListOverride FNetworkVersion::GetReplayCompatibleChangeListOverride;
 
 FString& FNetworkVersion::GetProjectVersion_Internal()
 {
@@ -85,6 +86,15 @@ uint32 FNetworkVersion::GetNetworkCompatibleChangelist()
 
 uint32 FNetworkVersion::GetReplayCompatibleChangelist()
 {
+	// Use the override if it's bound
+	if (GetReplayCompatibleChangeListOverride.IsBound())
+	{
+		const uint32 Changelist = GetReplayCompatibleChangeListOverride.Execute();
+
+		UE_LOG(LogNetVersion, Log, TEXT("Replay changelist override: %d"), Changelist);
+		return Changelist;
+	}
+
 	return FEngineVersion::CompatibleWith().GetChangelist();
 }
 
diff --git a/Engine/Source/Runtime/Core/Public/Misc/NetworkVersion.h b/Engine/Source/Runtime/Core/Public/Misc/NetworkVersion.h
--- a/Engine/Source/Runtime/Core/Public/Misc/NetworkVersion.h
+++ b/Engine/Source/Runtime/Core/Public/Misc/NetworkVersion.h
@@ -66,6 +66,10 @@ struct CORE_API FNetworkVersion
 	DECLARE_DELEGATE_RetVal_TwoParams( bool, FIsNetworkCompatibleOverride, uint32, uint32 );
 	static FIsNetworkCompatibleOverride IsNetworkCompatibleOverride;
 
+	/** Called in GetReplayCompatibleChangelist if bound */
+	DECLARE_DELEGATE_RetVal(uint32, FGetReplayCompatibleChangeListOverride);
+	static FGetReplayCompatibleChangeListOverride GetReplayCompatibleChangeListOverride;
+
 	static uint32 GetNetworkCompatibleChangelist();
 	static uint32 GetReplayCompatibleChangelist();
 	static uint32 GetEngineNetworkProtocolVersion();
