#ue5 - Disable AVX in UnrealMathSSE to unblock builds. 32-byte alignment causes FTransform size to mismatch the script-struct / noexporttypes size, and blows up. Changes from LWC branch will fix this.

diff --git a/Engine/Source/Runtime/Core/Public/Math/UnrealMathSSE.h b/Engine/Source/Runtime/Core/Public/Math/UnrealMathSSE.h
--- a/Engine/Source/Runtime/Core/Public/Math/UnrealMathSSE.h
+++ b/Engine/Source/Runtime/Core/Public/Math/UnrealMathSSE.h
@@ -10,7 +10,7 @@
 #endif
 
 #ifndef UE_PLATFORM_MATH_USE_AVX
-#define UE_PLATFORM_MATH_USE_AVX			PLATFORM_ALWAYS_HAS_AVX
+#define UE_PLATFORM_MATH_USE_AVX			0 //PLATFORM_ALWAYS_HAS_AVX
 #endif
 
 #ifndef UE_PLATFORM_MATH_USE_AVX_2
