LWC: Temporarily enforce float FVector::NetSerialize-ation allowing replays to work across LWC enabled/disabled builds.

diff --git a/Engine/Source/Runtime/Core/Public/Math/Vector.h b/Engine/Source/Runtime/Core/Public/Math/Vector.h
--- a/Engine/Source/Runtime/Core/Public/Math/Vector.h
+++ b/Engine/Source/Runtime/Core/Public/Math/Vector.h
@@ -1105,9 +1105,17 @@ public:
      */
 	bool NetSerialize(FArchive& Ar, class UPackageMap* Map, bool& bOutSuccess)
 	{
-		Ar << X;
-		Ar << Y;
-		Ar << Z;
+		// LWC_TODO: Serializer - NetSerialize forced to float for now to fix replays recorded with LWC disabled.
+		float SX = (float)X, SY = (float)Y, SZ = (float)Z;
+		Ar << SX;
+		Ar << SY;
+		Ar << SZ;
+		if(Ar.IsLoading())
+		{
+			X = SX;
+			Y = SY;
+			Z = SZ;
+		}
 		return true;
 	}
 
