FMemoryWriter64 for writing to TArray64.

diff --git a/Engine/Source/Runtime/Core/Public/Serialization/MemoryWriter.h b/Engine/Source/Runtime/Core/Public/Serialization/MemoryWriter.h
--- a/Engine/Source/Runtime/Core/Public/Serialization/MemoryWriter.h
+++ b/Engine/Source/Runtime/Core/Public/Serialization/MemoryWriter.h
@@ -15,10 +15,13 @@
 /**
  * Archive for storing arbitrary data to the specified memory location
  */
-class FMemoryWriter : public FMemoryArchive
+template <int IndexSize>
+class TMemoryWriter : public FMemoryArchive
 {
+	static_assert(IndexSize == 32 || IndexSize == 64, "Only 32-bit and 64-bit index sizes supported");
+
 public:
-	FMemoryWriter( TArray<uint8>& InBytes, bool bIsPersistent = false, bool bSetOffset = false, const FName InArchiveName = NAME_None )
+	TMemoryWriter( TArray<uint8, TSizedDefaultAllocator<IndexSize>>& InBytes, bool bIsPersistent = false, bool bSetOffset = false, const FName InArchiveName = NAME_None )
 	: FMemoryArchive()
 	, Bytes(InBytes)
 	, ArchiveName(InArchiveName)
@@ -37,9 +40,12 @@ public:
 		if( NumBytesToAdd > 0 )
 		{
 			const int64 NewArrayCount = Bytes.Num() + NumBytesToAdd;
-			if( NewArrayCount >= MAX_int32 )
+			if constexpr (IndexSize == 32)
 			{
-				UE_LOG( LogSerialization, Fatal, TEXT( "FMemoryWriter does not support data larger than 2GB. Archive name: %s." ), *ArchiveName.ToString() );
+				if( NewArrayCount >= MAX_int32 )
+				{
+					UE_LOG( LogSerialization, Fatal, TEXT( "FMemoryWriter does not support data larger than 2GB. Archive name: %s." ), *ArchiveName.ToString() );
+				}
 			}
 
 			Bytes.AddUninitialized( (int32)NumBytesToAdd );
@@ -59,7 +65,17 @@ public:
 	 *
 	 * This is overridden for the specific Archive Types
 	 **/
-	virtual FString GetArchiveName() const override { return TEXT("FMemoryWriter"); }
+	virtual FString GetArchiveName() const override
+	{
+		if constexpr (IndexSize == 64)
+		{
+			return TEXT("FMemoryWriter64");
+		}
+		else
+		{
+			return TEXT("FMemoryWriter");
+		}
+	}
 
 	int64 TotalSize() override
 	{
@@ -68,9 +84,28 @@ public:
 
 protected:
 
-	TArray<uint8>&	Bytes;
+	TArray<uint8, TSizedDefaultAllocator<IndexSize>>&	Bytes;
 
 	/** Archive name, used to debugging, by default set to NAME_None. */
 	const FName ArchiveName;
 };
 
+
+// FMemoryWriter and FMemoryWriter64 are implemented as derived classes rather than aliases
+// so that forward declarations will work.
+
+class FMemoryWriter : public TMemoryWriter<32>
+{
+	using Super = TMemoryWriter<32>;
+
+public:
+	using Super::Super;
+};
+
+class FMemoryWriter64 : public TMemoryWriter<64>
+{
+	using Super = TMemoryWriter<64>;
+
+public:
+	using Super::Super;
+};
