CompactBinary: Fixed crashes in TCbBuffer and FCbField

- Fixed a crash when TCbBuffer::CreateIterator() is called on an empty array or object.
- Fixed a crash when FCbField::operator[] is called for field name that does not exist in the array or object.

#rnx

diff --git a/Engine/Source/Runtime/Core/Private/Tests/Serialization/CompactBinaryTest.cpp b/Engine/Source/Runtime/Core/Private/Tests/Serialization/CompactBinaryTest.cpp
--- a/Engine/Source/Runtime/Core/Private/Tests/Serialization/CompactBinaryTest.cpp
+++ b/Engine/Source/Runtime/Core/Private/Tests/Serialization/CompactBinaryTest.cpp
@@ -467,6 +467,10 @@ bool FCbFieldObjectTest::RunTest(const FString& Parameters)
 		TestTrue(TEXT("FCbField(ObjectWithName)::GetBuffer()"), FCbField(FSharedBuffer::MakeView(MakeMemoryView(Buffer))).GetBuffer().Flatten().GetView().EqualBytes(MakeMemoryView(Buffer)));
 		TestTrue(TEXT("FCbField(ObjectWithNameNoType)::GetBuffer()"),
 			FCbField(FCbFieldView(Buffer + 1, ECbFieldType(ObjectType)), FSharedBuffer::MakeView(MakeMemoryView(Buffer))).GetBuffer().Flatten().GetView().EqualBytes(MakeMemoryView(Buffer)));
+
+		// Access Missing Field
+		ObjectClone["M"_ASV];
+		ObjectClone.AsField()["M"_ASV];
 	}
 
 	// Test FCbObjectView as FCbFieldViewIterator
@@ -494,6 +498,15 @@ bool FCbFieldObjectTest::RunTest(const FString& Parameters)
 		TestEqual(TEXT("FCbObject::AsFieldView() as Iterator Count"), Count, 1u);
 	}
 
+	// Test FCbObject(Empty) as FCbFieldIterator
+	{
+		const uint8 Buffer[] = { uint8(ECbFieldType::Object), 0 };
+		const FCbObject Object = FCbObject::Clone(Buffer);
+		for (FCbField& Field : Object)
+		{
+		}
+	}
+
 	return true;
 }
 
@@ -675,6 +688,15 @@ bool FCbFieldArrayTest::RunTest(const FString& Parameters)
 		TestEqual(TEXT("FCbArray::AsFieldView() as Iterator Count"), Count, 1u);
 	}
 
+	// Test FCbArray(Empty) as FCbFieldIterator
+	{
+		const uint8 Buffer[] = { uint8(ECbFieldType::Array), 1, 0 };
+		const FCbArray Array = FCbArray::Clone(Buffer);
+		for (FCbField& Field : Array)
+		{
+		}
+	}
+
 	return true;
 }
 
diff --git a/Engine/Source/Runtime/Core/Public/Serialization/CompactBinary.h b/Engine/Source/Runtime/Core/Public/Serialization/CompactBinary.h
--- a/Engine/Source/Runtime/Core/Public/Serialization/CompactBinary.h
+++ b/Engine/Source/Runtime/Core/Public/Serialization/CompactBinary.h
@@ -1238,7 +1238,11 @@ public:
 template <typename ViewType>
 inline FCbField TCbBuffer<ViewType>::operator[](FAnsiStringView Name) const
 {
-	return FCbField::MakeView(ViewType::operator[](Name), GetOuterBuffer());
+	if (FCbFieldView Field = ViewType::operator[](Name))
+	{
+		return FCbField::MakeView(Field, GetOuterBuffer());
+	}
+	return FCbField();
 }
 
 /**
@@ -1305,7 +1309,11 @@ private:
 template <typename ViewType>
 inline FCbFieldIterator TCbBuffer<ViewType>::CreateIterator() const
 {
-	return FCbFieldIterator::MakeRangeView(ViewType::CreateViewIterator(), GetOuterBuffer());
+	if (FCbFieldViewIterator It = ViewType::CreateViewIterator())
+	{
+		return FCbFieldIterator::MakeRangeView(It, GetOuterBuffer());
+	}
+	return FCbFieldIterator();
 }
 
 template <typename ViewType>
