xxHash: Added FXxHash64[Builder] and FXxHash128[Builder] wrappers for XXH3

#rnx

diff --git a/Engine/Source/Runtime/Core/Private/Hash/xxhash.cpp b/Engine/Source/Runtime/Core/Private/Hash/xxhash.cpp
--- a/Engine/Source/Runtime/Core/Private/Hash/xxhash.cpp
+++ b/Engine/Source/Runtime/Core/Private/Hash/xxhash.cpp
@@ -1,59 +1,116 @@
 // Copyright Epic Games, Inc. All Rights Reserved.
 
 #include "Hash/xxhash.h"
-#include "Misc/AssertionMacros.h"
+#include "Memory/CompositeBuffer.h"
+#include "Memory/MemoryView.h"
 
-#if 0	// Temporary workaround for Mac build
+THIRD_PARTY_INCLUDES_START
+#define XXH_INLINE_ALL
+#include "ThirdParty/xxhash/xxhash.h"
+THIRD_PARTY_INCLUDES_END
 
-//////////////////////////////////////////////////////////////////////////
+FXxHash64 FXxHash64::HashBuffer(FMemoryView View)
+{
+	return {XXH3_64bits(View.GetData(), View.GetSize())};
+}
 
-THIRD_PARTY_INCLUDES_START
+FXxHash64 FXxHash64::HashBuffer(const void* Data, uint64 Size)
+{
+	return {XXH3_64bits(Data, Size)};
+}
 
-#define XXH_NAMESPACE UE4XXH
+FXxHash64 FXxHash64::HashBuffer(const FCompositeBuffer& Buffer)
+{
+	FXxHash64Builder Builder;
+	Builder.Update(Buffer);
+	return Builder.Finalize();
+}
 
-#include "ThirdParty/xxhash/xxhash.c"
-#include "ThirdParty/xxhash/xxh3.h"
+FXxHash128 FXxHash128::HashBuffer(FMemoryView View)
+{
+	const XXH128_hash_t Value = XXH3_128bits(View.GetData(), View.GetSize());
+	return {Value.low64, Value.high64};
+}
 
-THIRD_PARTY_INCLUDES_END
+FXxHash128 FXxHash128::HashBuffer(const void* Data, uint64 Size)
+{
+	const XXH128_hash_t Value = XXH3_128bits(Data, Size);
+	return {Value.low64, Value.high64};
+}
 
-//////////////////////////////////////////////////////////////////////////
+FXxHash128 FXxHash128::HashBuffer(const FCompositeBuffer& Buffer)
+{
+	FXxHash128Builder Builder;
+	Builder.Update(Buffer);
+	return Builder.Finalize();
+}
 
-uint32 
-FXxHash::HashBuffer32(const void* BlockPtr, SIZE_T BlockSize)
+void FXxHash64Builder::Reset()
 {
-	XXH64_state_t State;
-	XXH64_reset(&State, 0);
+	static_assert(sizeof(StateBytes) == sizeof(XXH3_state_t), "Adjust the allocation in FXxHash64Builder to match XXH3_state_t");
+	XXH3_state_t& State = reinterpret_cast<XXH3_state_t&>(StateBytes);
+	XXH3_64bits_reset(&State);
+}
 
-	XXH64_update(&State, BlockPtr, BlockSize);
-	return uint32(XXH64_digest(&State));
+void FXxHash64Builder::Update(FMemoryView View)
+{
+	XXH3_state_t& State = reinterpret_cast<XXH3_state_t&>(StateBytes);
+	XXH3_64bits_update(&State, View.GetData(), View.GetSize());
 }
 
-uint64 
-FXxHash::HashBuffer64(const void* BlockPtr, SIZE_T BlockSize)
+void FXxHash64Builder::Update(const void* Data, uint64 Size)
 {
-	XXH64_state_t State;
-	XXH64_reset(&State, 0);
+	XXH3_state_t& State = reinterpret_cast<XXH3_state_t&>(StateBytes);
+	XXH3_64bits_update(&State, Data, Size);
+}
 
-	XXH64_update(&State, BlockPtr, BlockSize);
-	return XXH64_digest(&State);
+void FXxHash64Builder::Update(const FCompositeBuffer& Buffer)
+{
+	XXH3_state_t& State = reinterpret_cast<XXH3_state_t&>(StateBytes);
+	for (const FSharedBuffer& Segment : Buffer.GetSegments())
+	{
+		XXH3_64bits_update(&State, Segment.GetData(), Segment.GetSize());
+	}
 }
 
-#else
+FXxHash64 FXxHash64Builder::Finalize() const
+{
+	const XXH3_state_t& State = reinterpret_cast<const XXH3_state_t&>(StateBytes);
+	const XXH64_hash_t Hash = XXH3_64bits_digest(&State);
+	return {Hash};
+}
 
-uint32
-FXxHash::HashBuffer32(const void* BlockPtr, SIZE_T BlockSize)
+void FXxHash128Builder::Reset()
 {
-	unimplemented();
+	static_assert(sizeof(StateBytes) == sizeof(XXH3_state_t), "Adjust the allocation in FXxHash128Builder to match XXH3_state_t");
+	XXH3_state_t& State = reinterpret_cast<XXH3_state_t&>(StateBytes);
+	XXH3_128bits_reset(&State);
+}
 
-	return 0;
+void FXxHash128Builder::Update(FMemoryView View)
+{
+	XXH3_state_t& State = reinterpret_cast<XXH3_state_t&>(StateBytes);
+	XXH3_128bits_update(&State, View.GetData(), View.GetSize());
 }
 
-uint64
-FXxHash::HashBuffer64(const void* BlockPtr, SIZE_T BlockSize)
+void FXxHash128Builder::Update(const void* Data, uint64 Size)
 {
-	unimplemented();
+	XXH3_state_t& State = reinterpret_cast<XXH3_state_t&>(StateBytes);
+	XXH3_128bits_update(&State, Data, Size);
+}
 
-	return 0;
+void FXxHash128Builder::Update(const FCompositeBuffer& Buffer)
+{
+	XXH3_state_t& State = reinterpret_cast<XXH3_state_t&>(StateBytes);
+	for (const FSharedBuffer& Segment : Buffer.GetSegments())
+	{
+		XXH3_128bits_update(&State, Segment.GetData(), Segment.GetSize());
+	}
 }
 
-#endif
+FXxHash128 FXxHash128Builder::Finalize() const
+{
+	const XXH3_state_t& State = reinterpret_cast<const XXH3_state_t&>(StateBytes);
+	const XXH128_hash_t Hash = XXH3_128bits_digest(&State);
+	return {Hash.low64, Hash.high64};
+}
diff --git a/Engine/Source/Runtime/Core/Public/Hash/xxhash.h b/Engine/Source/Runtime/Core/Public/Hash/xxhash.h
--- a/Engine/Source/Runtime/Core/Public/Hash/xxhash.h
+++ b/Engine/Source/Runtime/Core/Public/Hash/xxhash.h
@@ -3,9 +3,186 @@
 #pragma once
 
 #include "CoreTypes.h"
+#include "Containers/StringFwd.h"
+#include "HAL/UnrealMemory.h"
+#include "Memory/MemoryFwd.h"
+#include "Misc/ByteSwap.h"
+#include "Serialization/Archive.h"
+#include "String/BytesToHex.h"
 
-struct FXxHash
+/** A 64-bit hash from XXH3. */
+struct FXxHash64
 {
-	static CORE_API uint32 HashBuffer32(const void* BlockPtr, SIZE_T BlockSize);
-	static CORE_API uint64 HashBuffer64(const void* BlockPtr, SIZE_T BlockSize);
+	[[nodiscard]] CORE_API static FXxHash64 HashBuffer(FMemoryView View);
+	[[nodiscard]] CORE_API static FXxHash64 HashBuffer(const void* Data, uint64 Size);
+	[[nodiscard]] CORE_API static FXxHash64 HashBuffer(const FCompositeBuffer& Buffer);
+
+	/** Load the hash from its canonical (big-endian) representation. */
+	static inline FXxHash64 FromByteArray(uint8 (&Bytes)[sizeof(uint64)])
+	{
+		uint64 HashBigEndian;
+		FMemory::Memcpy(&HashBigEndian, Bytes, sizeof(uint64));
+		return {NETWORK_ORDER64(HashBigEndian)};
+	}
+
+	/** Store the hash to its canonical (big-endian) representation. */
+	inline void ToByteArray(uint8 (&Bytes)[sizeof(uint64)]) const
+	{
+		const uint64 HashBigEndian = NETWORK_ORDER64(Hash);
+		FMemory::Memcpy(Bytes, &HashBigEndian, sizeof(uint64));
+	}
+
+	/**
+	 * The hash in its native representation.
+	 *
+	 * Use the canonical representation from ToByteArray to serialize or display the hash.
+	 */
+	uint64 Hash{};
+};
+
+/** A 128-bit hash from XXH3. */
+struct FXxHash128
+{
+public:
+	[[nodiscard]] CORE_API static FXxHash128 HashBuffer(FMemoryView View);
+	[[nodiscard]] CORE_API static FXxHash128 HashBuffer(const void* Data, uint64 Size);
+	[[nodiscard]] CORE_API static FXxHash128 HashBuffer(const FCompositeBuffer& Buffer);
+
+	/** Load the hash from its canonical (big-endian) representation. */
+	static inline FXxHash128 FromByteArray(uint8 (&Bytes)[sizeof(uint64[2])])
+	{
+		uint64 HashBigEndian[2];
+		FMemory::Memcpy(&HashBigEndian, Bytes, sizeof(uint64[2]));
+		return {NETWORK_ORDER64(HashBigEndian[1]), NETWORK_ORDER64(HashBigEndian[0])};
+	}
+
+	/** Store the hash to its canonical (big-endian) representation. */
+	inline void ToByteArray(uint8 (&Bytes)[sizeof(uint64[2])]) const
+	{
+		const uint64 HashBigEndian[2]{NETWORK_ORDER64(HashHigh), NETWORK_ORDER64(HashLow)};
+		FMemory::Memcpy(Bytes, &HashBigEndian, sizeof(uint64[2]));
+	}
+
+	/**
+	 * The low 64 bits of the hash in its native representation.
+	 *
+	 * Use the canonical representation from ToByteArray to serialize or display the hash.
+	 */
+	uint64 HashLow{};
+	/**
+	 * The high 64 bits of the hash in its native representation.
+	 *
+	 * Use the canonical representation from ToByteArray to serialize or display the hash.
+	 */
+	uint64 HashHigh{};
+};
+
+/** Calculates a 64-bit hash with XXH3. */
+class FXxHash64Builder
+{
+public:
+	inline FXxHash64Builder() { Reset(); }
+
+	FXxHash64Builder(const FXxHash64Builder&) = delete;
+	FXxHash64Builder& operator=(const FXxHash64Builder&) = delete;
+
+	CORE_API void Reset();
+
+	CORE_API void Update(FMemoryView View);
+	CORE_API void Update(const void* Data, uint64 Size);
+	CORE_API void Update(const FCompositeBuffer& Buffer);
+
+	[[nodiscard]] CORE_API FXxHash64 Finalize() const;
+
+private:
+	alignas(64) char StateBytes[576];
+};
+
+/** Calculates a 128-bit hash with XXH3. */
+class FXxHash128Builder
+{
+public:
+	inline FXxHash128Builder() { Reset(); }
+
+	FXxHash128Builder(const FXxHash128Builder&) = delete;
+	FXxHash128Builder& operator=(const FXxHash128Builder&) = delete;
+
+	CORE_API void Reset();
+
+	CORE_API void Update(FMemoryView View);
+	CORE_API void Update(const void* Data, uint64 Size);
+	CORE_API void Update(const FCompositeBuffer& Buffer);
+
+	[[nodiscard]] CORE_API FXxHash128 Finalize() const;
+
+private:
+	alignas(64) char StateBytes[576];
 };
+
+inline bool operator==(const FXxHash64& A, const FXxHash64& B)
+{
+	return A.Hash == B.Hash;
+}
+
+inline bool operator!=(const FXxHash64& A, const FXxHash64& B)
+{
+	return A.Hash != B.Hash;
+}
+
+inline bool operator<(const FXxHash64& A, const FXxHash64& B)
+{
+	return A.Hash < B.Hash;
+}
+
+inline uint32 GetTypeHash(const FXxHash64& Hash)
+{
+	return uint32(Hash.Hash);
+}
+
+inline FArchive& operator<<(FArchive& Ar, FXxHash64& Hash)
+{
+	return Ar << Hash.Hash;
+}
+
+template <typename CharType>
+inline TStringBuilderBase<CharType>& operator<<(TStringBuilderBase<CharType>& Builder, const FXxHash64& Hash)
+{
+	uint8 Bytes[8];
+	Hash.ToByteArray(Bytes);
+	UE::String::BytesToHexLower(MakeArrayView(Bytes), Builder);
+	return Builder;
+}
+
+inline bool operator==(const FXxHash128& A, const FXxHash128& B)
+{
+	return A.HashLow == B.HashLow && A.HashHigh == B.HashHigh;
+}
+
+inline bool operator!=(const FXxHash128& A, const FXxHash128& B)
+{
+	return A.HashLow != B.HashLow || A.HashHigh != B.HashHigh;
+}
+
+inline bool operator<(const FXxHash128& A, const FXxHash128& B)
+{
+	return A.HashHigh != B.HashHigh ? A.HashHigh < B.HashHigh : A.HashLow < A.HashLow;
+}
+
+inline uint32 GetTypeHash(const FXxHash128& Hash)
+{
+	return uint32(Hash.HashLow);
+}
+
+inline FArchive& operator<<(FArchive& Ar, FXxHash128& Hash)
+{
+	return Ar << Hash.HashLow << Hash.HashHigh;
+}
+
+template <typename CharType>
+inline TStringBuilderBase<CharType>& operator<<(TStringBuilderBase<CharType>& Builder, const FXxHash128& Hash)
+{
+	uint8 Bytes[16];
+	Hash.ToByteArray(Bytes);
+	UE::String::BytesToHexLower(MakeArrayView(Bytes), Builder);
+	return Builder;
+}
