Fixed warnings from static analysis

#rnx

diff --git a/Engine/Source/Runtime/Core/Private/Tests/Serialization/CompactBinaryTest.cpp b/Engine/Source/Runtime/Core/Private/Tests/Serialization/CompactBinaryTest.cpp
--- a/Engine/Source/Runtime/Core/Private/Tests/Serialization/CompactBinaryTest.cpp
+++ b/Engine/Source/Runtime/Core/Private/Tests/Serialization/CompactBinaryTest.cpp
@@ -469,8 +469,8 @@ bool FCbFieldObjectTest::RunTest(const FString& Parameters)
 			FCbField(FCbFieldView(Buffer + 1, ECbFieldType(ObjectType)), FSharedBuffer::MakeView(MakeMemoryView(Buffer))).GetBuffer().Flatten().GetView().EqualBytes(MakeMemoryView(Buffer)));
 
 		// Access Missing Field
-		ObjectClone["M"_ASV];
-		ObjectClone.AsField()["M"_ASV];
+		TestFalse(TEXT("FCbObject()[Missing]"), ObjectClone["M"_ASV].HasValue());
+		TestFalse(TEXT("FCbField(Object)[Missing]"), ObjectClone.AsField()["M"_ASV].HasValue());
 	}
 
 	// Test FCbObjectView as FCbFieldViewIterator
@@ -1752,15 +1752,15 @@ bool FCbFieldBufferIteratorTest::RunTest(const FString& Parameters)
 		TestFalse(TEXT("FCbFieldIterator::MakeRange(Uniform).TryGetRangeView()"), UniformFieldIt.TryGetRangeView(RangeView));
 
 		// Equals
-		TestTrue(TEXT("FCbFieldViewIterator::Equals(Self)"), FieldViewIt.Equals(FieldViewIt));
-		TestTrue(TEXT("FCbFieldViewIterator::Equals(OtherType)"), FieldViewIt.Equals(FieldIt));
-		TestTrue(TEXT("FCbFieldIterator::Equals(Self)"), FieldIt.Equals(FieldIt));
-		TestTrue(TEXT("FCbFieldIterator::Equals(OtherType)"), FieldIt.Equals(FieldViewIt));
+		TestTrue(TEXT("FCbFieldViewIterator::Equals(Self)"), FieldViewIt.Equals(AsConst(FieldViewIt)));
+		TestTrue(TEXT("FCbFieldViewIterator::Equals(OtherType)"), FieldViewIt.Equals(AsConst(FieldIt)));
+		TestTrue(TEXT("FCbFieldIterator::Equals(Self)"), FieldIt.Equals(AsConst(FieldIt)));
+		TestTrue(TEXT("FCbFieldIterator::Equals(OtherType)"), FieldIt.Equals(AsConst(FieldViewIt)));
 		TestFalse(TEXT("FCbFieldViewIterator::Equals(OtherRange)"), FieldViewIt.Equals(FieldViewItClone));
 		TestFalse(TEXT("FCbFieldIterator::Equals(OtherRange)"), FieldIt.Equals(FieldItClone));
-		TestTrue(TEXT("FCbFieldViewIterator::Equals(Uniform, Self)"), UniformFieldViewIt.Equals(UniformFieldViewIt));
+		TestTrue(TEXT("FCbFieldViewIterator::Equals(Uniform, Self)"), UniformFieldViewIt.Equals(AsConst(UniformFieldViewIt)));
 		TestTrue(TEXT("FCbFieldViewIterator::Equals(Uniform, OtherType)"), UniformFieldViewIt.Equals(UniformFieldIt));
-		TestTrue(TEXT("FCbFieldIterator::Equals(Uniform, Self)"), UniformFieldIt.Equals(UniformFieldIt));
+		TestTrue(TEXT("FCbFieldIterator::Equals(Uniform, Self)"), UniformFieldIt.Equals(AsConst(UniformFieldIt)));
 		TestTrue(TEXT("FCbFieldIterator::Equals(Uniform, OtherType)"), UniformFieldIt.Equals(UniformFieldViewIt));
 		TestFalse(TEXT("FCbFieldViewIterator::Equals(SamePayload, DifferentEnd)"),
 			FCbFieldViewIterator::MakeRange(MakeMemoryView(UniformPayload), ECbFieldType::IntegerPositive)
diff --git a/Engine/Source/Runtime/Core/Public/Hash/xxhash.h b/Engine/Source/Runtime/Core/Public/Hash/xxhash.h
--- a/Engine/Source/Runtime/Core/Public/Hash/xxhash.h
+++ b/Engine/Source/Runtime/Core/Public/Hash/xxhash.h
@@ -165,7 +165,7 @@ inline bool operator!=(const FXxHash128& A, const FXxHash128& B)
 
 inline bool operator<(const FXxHash128& A, const FXxHash128& B)
 {
-	return A.HashHigh != B.HashHigh ? A.HashHigh < B.HashHigh : A.HashLow < A.HashLow;
+	return A.HashHigh != B.HashHigh ? A.HashHigh < B.HashHigh : A.HashLow < B.HashLow;
 }
 
 inline uint32 GetTypeHash(const FXxHash128& Hash)
