Allow non-zero default to GetSafeNormal

diff --git a/Engine/Source/Runtime/Core/Public/Math/Vector.h b/Engine/Source/Runtime/Core/Public/Math/Vector.h
--- a/Engine/Source/Runtime/Core/Public/Math/Vector.h
+++ b/Engine/Source/Runtime/Core/Public/Math/Vector.h
@@ -620,21 +620,21 @@ public:
 
     /**
      * Gets a normalized copy of the vector, checking it is safe to do so based on the length.
-     * Returns zero vector if vector length is too small to safely normalize.
+     * Returns zero vector by default if vector length is too small to safely normalize.
      *
      * @param Tolerance Minimum squared vector length.
-     * @return A normalized copy if safe, (0,0,0) otherwise.
+     * @return A normalized copy if safe, ResultIfZero otherwise.
      */
-    TVector<T> GetSafeNormal(T Tolerance=SMALL_NUMBER) const;
+    TVector<T> GetSafeNormal(T Tolerance=SMALL_NUMBER, const TVector<T>& ResultIfZero = ZeroVector) const;
 
     /**
      * Gets a normalized copy of the 2D components of the vector, checking it is safe to do so. Z is set to zero. 
-     * Returns zero vector if vector length is too small to normalize.
+     * Returns zero vector by default if vector length is too small to normalize.
      *
      * @param Tolerance Minimum squared vector length.
-     * @return Normalized copy if safe, otherwise returns zero vector.
+     * @return Normalized copy if safe, otherwise returns ResultIfZero.
      */
-    TVector<T> GetSafeNormal2D(T Tolerance=SMALL_NUMBER) const;
+    TVector<T> GetSafeNormal2D(T Tolerance=SMALL_NUMBER, const TVector<T>& ResultIfZero = ZeroVector) const;
 
     /**
      * Util to convert this vector into a unit direction vector and its original length.
@@ -1908,7 +1908,7 @@ FORCEINLINE TVector<T> TVector<T>::MirrorByVector(const TVector<T>& MirrorNormal
 }
 
 template<typename T>
-FORCEINLINE TVector<T> TVector<T>::GetSafeNormal(T Tolerance) const
+FORCEINLINE TVector<T> TVector<T>::GetSafeNormal(T Tolerance, const TVector<T>& ResultIfZero) const
 {
     const T SquareSum = X*X + Y*Y + Z*Z;
 
@@ -1919,14 +1919,14 @@ FORCEINLINE TVector<T> TVector<T>::GetSafeNormal(T Tolerance) const
     }		
     else if(SquareSum < Tolerance)
     {
-        return ZeroVector;
+        return ResultIfZero;
     }
     const T Scale = (T)FMath::InvSqrt(SquareSum);
     return TVector<T>(X*Scale, Y*Scale, Z*Scale);
 }
 
 template<typename T>
-FORCEINLINE TVector<T> TVector<T>::GetSafeNormal2D(T Tolerance) const
+FORCEINLINE TVector<T> TVector<T>::GetSafeNormal2D(T Tolerance, const TVector<T>& ResultIfZero) const
 {
     const T SquareSum = X*X + Y*Y;
 
@@ -1944,7 +1944,7 @@ FORCEINLINE TVector<T> TVector<T>::GetSafeNormal2D(T Tolerance) const
     }
     else if(SquareSum < Tolerance)
     {
-        return ZeroVector;
+        return ResultIfZero;
     }
 
     const T Scale = FMath::InvSqrt(SquareSum);
