Fix for filesystem memory leak

diff --git a/Engine/Source/Runtime/Core/Private/IO/IoDispatcher.cpp b/Engine/Source/Runtime/Core/Private/IO/IoDispatcher.cpp
--- a/Engine/Source/Runtime/Core/Private/IO/IoDispatcher.cpp
+++ b/Engine/Source/Runtime/Core/Private/IO/IoDispatcher.cpp
@@ -955,6 +955,10 @@ FIoRequest& FIoRequest::operator=(const FIoRequest& Other)
 
 FIoRequest& FIoRequest::operator=(FIoRequest&& Other)
 {
+	if (Impl)
+	{
+		Impl->ReleaseRef();
+	}
 	Impl = Other.Impl;
 	Other.Impl = nullptr;
 	return *this;
