Fixed single-property config updates failing to remove entries that were no longer needed as they matched the CDO

The UObject would write the correct config state, but UpdateSinglePropertyInSection would fail to remove the data for the updated key as it (correctly) wasn't present in the config being applied. This change has it remove the key from the config section if it exists.

diff --git a/Engine/Source/Runtime/Core/Private/Misc/ConfigCacheIni.cpp b/Engine/Source/Runtime/Core/Private/Misc/ConfigCacheIni.cpp
--- a/Engine/Source/Runtime/Core/Private/Misc/ConfigCacheIni.cpp
+++ b/Engine/Source/Runtime/Core/Private/Misc/ConfigCacheIni.cpp
@@ -4665,9 +4665,9 @@ public:
 	 * @Param InIniFilename - The disk location of the file we wish to edit.
 	 * @Param InSectionName - The section the property belongs to.
 	 * @Param InPropertyName - The name of the property that has been edited.
-	 * @Param InPropertyValue - The new value of the property that has been edited.
+	 * @Param InPropertyValue - The new value of the property that has been edited, or unset to remove the property.
 	 */
-	FSinglePropertyConfigHelper(const FString& InIniFilename, const FString& InSectionName, const FString& InPropertyName, const FString& InPropertyValue)
+	FSinglePropertyConfigHelper(const FString& InIniFilename, const FString& InSectionName, const FString& InPropertyName, const TOptional<FString>& InPropertyValue)
 		: IniFilename(InIniFilename)
 		, SectionName(InSectionName)
 		, PropertyName(InPropertyName)
@@ -4716,25 +4716,31 @@ private:
 		FString UpdatedSection;
 		if (IniFileMakeup.Section.IsEmpty())
 		{
-			const FString DecoratedSectionName = FString::Printf(TEXT("[%s]"), *SectionName);
-			
-			ClearTrailingWhitespace(IniFileMakeup.BeforeSection);
-			UpdatedSection += LINE_TERMINATOR;
-			UpdatedSection += LINE_TERMINATOR;
-			UpdatedSection += DecoratedSectionName;
-			AppendPropertyLine(UpdatedSection);
+			if (PropertyValue.IsSet())
+			{
+				const FString DecoratedSectionName = FString::Printf(TEXT("[%s]"), *SectionName);
+
+				ClearTrailingWhitespace(IniFileMakeup.BeforeSection);
+				UpdatedSection += LINE_TERMINATOR;
+				UpdatedSection += LINE_TERMINATOR;
+				UpdatedSection += DecoratedSectionName;
+				AppendPropertyLine(UpdatedSection);
+			}
 		}
 		else
 		{
 			FString SectionLine;
 			const TCHAR* Ptr = *IniFileMakeup.Section;
-			bool bWrotePropertyOnPass = false;
+			bool bUpdatedPropertyOnPass = false;
 			while (Ptr != nullptr && FParse::Line(&Ptr, SectionLine, true))
 			{
 				if (SectionLine.StartsWith(FString::Printf(TEXT("%s="), *PropertyName)))
 				{
-					UpdatedSection += FConfigFile::GenerateExportedPropertyLine(PropertyName, PropertyValue);
-					bWrotePropertyOnPass = true;
+					if (PropertyValue.IsSet())
+					{
+						UpdatedSection += FConfigFile::GenerateExportedPropertyLine(PropertyName, PropertyValue.GetValue());
+					}
+					bUpdatedPropertyOnPass = true;
 				}
 				else
 				{
@@ -4745,7 +4751,7 @@ private:
 
 			// If the property wasnt found in the text of the existing section content,
 			// append it to the end of the section.
-			if (!bWrotePropertyOnPass)
+			if (!bUpdatedPropertyOnPass && PropertyValue.IsSet())
 			{
 				AppendPropertyLine(UpdatedSection);
 			}
@@ -4810,10 +4816,12 @@ private:
 	 */
 	void AppendPropertyLine(FString& PreText)
 	{
+		check(PropertyValue.IsSet());
+		
 		// Make sure we dont leave much whitespace, and append the property name/value entry
 		ClearTrailingWhitespace(PreText);
 		PreText += LINE_TERMINATOR;
-		PreText += FConfigFile::GenerateExportedPropertyLine(PropertyName, PropertyValue);
+		PreText += FConfigFile::GenerateExportedPropertyLine(PropertyName, PropertyValue.GetValue());
 		PreText += LINE_TERMINATOR;
 	}
 
@@ -4829,7 +4837,8 @@ private:
 	FString PropertyName;
 
 	// The new value, in string format, of the property that has been changed
-	FString PropertyValue;
+	// This will be unset if the property has been removed from the config
+	TOptional<FString> PropertyValue;
 
 	// Helper struct that holds the makeup of the ini file.
 	struct IniFileContent
@@ -4848,22 +4857,18 @@ private:
 
 bool FConfigFile::UpdateSinglePropertyInSection(const TCHAR* DiskFilename, const TCHAR* PropertyName, const TCHAR* SectionName)
 {
-	// Result of whether the file has been updated on disk.
-	bool bSuccessfullyUpdatedFile = false;
-
-	FString PropertyValue;
+	TOptional<FString> PropertyValue;
 	if (const FConfigSection* LocalSection = this->Find(SectionName))
 	{
 		if (const FConfigValue* ConfigValue = LocalSection->Find(PropertyName))
 		{
 			// Use GetSavedValueForWriting rather than GetSavedValue to avoid having this save operation mark the value as having been accessed for dependency tracking
 			PropertyValue = UE::ConfigCacheIni::Private::FAccessor::GetSavedValueForWriting(*ConfigValue);
-			FSinglePropertyConfigHelper SinglePropertyConfigHelper(DiskFilename, SectionName, PropertyName, PropertyValue);
-			bSuccessfullyUpdatedFile = SinglePropertyConfigHelper.UpdateConfigFile();
 		}
 	}
 
-	return bSuccessfullyUpdatedFile;
+	FSinglePropertyConfigHelper SinglePropertyConfigHelper(DiskFilename, SectionName, PropertyName, PropertyValue);
+	return SinglePropertyConfigHelper.UpdateConfigFile();
 }
 
 const TCHAR* ConvertValueFromHumanFriendlyValue( const TCHAR* Value )
diff --git a/Engine/Source/Runtime/CoreUObject/Private/UObject/Obj.cpp b/Engine/Source/Runtime/CoreUObject/Private/UObject/Obj.cpp
--- a/Engine/Source/Runtime/CoreUObject/Private/UObject/Obj.cpp
+++ b/Engine/Source/Runtime/CoreUObject/Private/UObject/Obj.cpp
@@ -2594,7 +2594,7 @@ void UObject::SaveConfig( uint64 Flags, const TCHAR* InFilename, FConfigCacheIni
 			FArrayProperty* Array   = CastField<FArrayProperty>( Property );
 			if( Array )
 			{
-				FConfigSection* Sec = Config->GetSectionPrivate(*Section, 1, 0, *PropFileName);
+				FConfigSection* Sec = Config->GetSectionPrivate(*Section, 1, 0, PropFileName);
 				// Default ini's require the array syntax to be applied to the property name
 				FString CompleteKey = FString::Printf(TEXT("%s%s"), bIsADefaultIniWrite ? TEXT("+") : TEXT(""), *Key);
 				if (Sec)
@@ -2630,16 +2630,12 @@ void UObject::SaveConfig( uint64 Flags, const TCHAR* InFilename, FConfigCacheIni
 					{
 						FString	Value;
 						Property->ExportText_InContainer( Index, Value, this, this, this, PortFlags );
-						Config->SetString( *Section, *Key, *Value, *PropFileName );
+						Config->SetString( *Section, *Key, *Value, PropFileName );
 					}
 					else
 					{
 						// If we are not writing it to config above, we should make sure that this property isn't stagnant in the cache.
-						FConfigSection* Sec = Config->GetSectionPrivate( *Section, 1, 0, *PropFileName );
-						if( Sec )
-						{
-							Sec->Remove( *Key );
-						}
+						Config->RemoveKey( *Section, *Key, PropFileName );
 					}
 				}
 			}
