fix check for compressionformats and compressionformat that could cause incorrect logging

diff --git a/Engine/Source/Runtime/Core/Private/Compression/OodleDataCompressionFormat.cpp b/Engine/Source/Runtime/Core/Private/Compression/OodleDataCompressionFormat.cpp
--- a/Engine/Source/Runtime/Core/Private/Compression/OodleDataCompressionFormat.cpp
+++ b/Engine/Source/Runtime/Core/Private/Compression/OodleDataCompressionFormat.cpp
@@ -390,7 +390,9 @@ ICompressionFormat * CreateOodleDataCompressionFormat()
 		}
 
 		// no init log line if we're not enabled for pak/iostore :
-		bool bUseCompressionFormatOodle = FCString::Strifind(FCommandLine::Get(), TEXT("-compressionformats=oodle")) != NULL;
+		bool bUseCompressionFormatOodle = 
+			( FCString::Strifind(FCommandLine::Get(), TEXT("-compressionformats=oodle")) != NULL ) ||
+			( FCString::Strifind(FCommandLine::Get(), TEXT("-compressionformat=oodle")) != NULL );
 		if ( bUseCompressionFormatOodle )			
 		{
 			UE_LOG(OodleDataCompression, Display, TEXT("Oodle v%s format for pak/iostore with method=%s, level=%d=%s"), TEXT(OodleVersion), **MethodMap.FindKey(UsedCompressor), (int)UsedLevel, **LevelMap.FindKey(UsedLevel) );
