OodleDataCompression.h function param docs

diff --git a/Engine/Source/Runtime/Core/Public/Compression/OodleDataCompression.h b/Engine/Source/Runtime/Core/Public/Compression/OodleDataCompression.h
--- a/Engine/Source/Runtime/Core/Public/Compression/OodleDataCompression.h
+++ b/Engine/Source/Runtime/Core/Public/Compression/OodleDataCompression.h
@@ -15,7 +15,7 @@
 // Let me emphasize again : for data to be stored in packages for shipping games, do NOT use this
 //	let the pak/iostore system choose the compression from the config options and platform settings
 //
-// This is for utility compression in non-shipping game package scenarios
+// OodleDataCompression.h is for utility compression in non-shipping game package scenarios
 //	eg. uassets, storage caches and back-ends, large network transfers
 
 DECLARE_LOG_CATEGORY_EXTERN(OodleDataCompression, Log, All);
@@ -23,11 +23,10 @@ DECLARE_LOG_CATEGORY_EXTERN(OodleDataCompression, Log, All);
 /**
  * EOodleDataCompressor : Choose the Oodle Compressor
  *  this mostly trades decompression speed vs compression ratio
- * 
- * From fastest to slowest (to decode) : Selkie, Mermaid, Kraken, Leviathan
- * 
  * encode speed is determined by EOodleDataCompressionLevel , not the compressor choice.
  * 
+ * From fastest to slowest (to decode) : Selkie, Mermaid, Kraken, Leviathan
+ *  
  * When in doubt, start with Kraken
  * Representative compression ratios and decode speeds :
  
@@ -52,12 +51,13 @@ enum class EOodleDataCompressor : uint8
 /**
  * EOodleDataCompressionLevel : Choose the Oodle Compression Level
  *  this mostly trades encode speed vs compression ratio
+ *  decode speed is determined by choice of compressor (EOodleDataCompressor)
  * 
  * If in doubt start with "Normal" (level 4) then move up or down from there
  * 
  * The standard range is in "SuperFast" - "Normal" (levels 1-4)
  * 
- * HyperFast is for real-time encoding with minimal compression
+ * "HyperFast" levels are for real-time encoding with minimal compression
  * 
  * The "Optimal" levels are much slower to encode, but provide more compression
  * they are intended for offline cooks
@@ -96,20 +96,54 @@ enum class EOodleDataCompressionLevel : int8
 	Optimal4 = 8,
 };
 
-int64 CORE_API OodleDataCompressedBufferSizeNeeded(int64 IncompressedSize);
+/**
+* What size of compressed output buffer is needed to encode into for OodleDataCompress()
+*
+* @param	UncompressedSize			Length of uncompressed data that will fit in this output buffer
+* @return								Minimum size to allocate compressed output buffer
+*/
+int64 CORE_API OodleDataCompressedBufferSizeNeeded(int64 UncompressedSize);
 
+/**
+* What is the maximum size of compressed data made after encoding
+*  For pre-allocating buffers at decode time or to check the length of compressed data at load time
+*  NOTE : OodleDataCompressedBufferSizeNeeded() >= OodleDataGetMaximumCompressedSize()
+* this is not the size to allocate compressed buffers at encode time; see OodleDataCompressedBufferSizeNeeded
+*
+* @param	UncompressedSize			Length of uncompressed data that could have been encoded
+* @return								Maximum size of compressed data that OodleDataCompress() could have made
+*/
 int64 CORE_API OodleDataGetMaximumCompressedSize(int64 UncompressedSize);
 
-// OutCompressedSize must be >= OodleDataCompressedBufferSizeNeeded(InUncompressedSize)
-// returns compressed size or zero for failure
+/**
+* Encode provided data with chosen Compressor and Level
+* CompressedBufferSize must be >= OodleDataCompressedBufferSizeNeeded(UncompressedSize)
+*
+* @param	OutCompressedData			output buffer where compressed data is written
+* @param	CompressedBufferSize		bytes available to write in OutCompressedData
+* @param	InUncompressedData			input buffer containing data to compress
+* @param	UncompressedSize			number of bytes in InUncompressedData to read
+* @param	Compressor					EOodleDataCompressor to encode with (this is saved in the stream)
+* @param	Level						EOodleDataCompressionLevel to encode with (this is not saved in the stream)
+* @return								Compressed size written or zero for failure
+*/
 int64 CORE_API OodleDataCompress(
 							void * OutCompressedData, int64 CompressedBufferSize,
 							const void * InUncompressedData, int64 UncompressedSize,
 							EOodleDataCompressor Compressor,
 							EOodleDataCompressionLevel Level);
 
-// OutUncompressedSize is the number of bytes to decompress, it must match the number that you encoded
-// InCompressedSize is the buffer size of compressed data, it must be >= the number of compressed bytes needed
+/**
+* Decode compressed data that was made by OodleDataCompress
+*
+* UncompressedSize must match exactly the uncompressed size that was used at encode time.  No partial decodes.
+*
+* @param	OutUncompressedData		output buffer where uncompressed data is written
+* @param	UncompressedSize		number of bytes to decompress
+* @param	InCompressedData		input buffer containing compressed data
+* @param	CompressedSize			size of the input buffer, must be greater or equal to the number of compressed bytes needed
+* @return							boolean success
+*/
 bool CORE_API OodleDataDecompress(
 						void * OutUncompressedData, int64 UncompressedSize,
 						const void * InCompressedData, int64 CompressedSize
