FCompression::GetMaximumCompressedSize to clean up differentiation from CompressMemoryBound
use GetMaximumCompressedSize for decoding
use CompressMemoryBound for encoding

diff --git a/Engine/Source/Editor/UnrealEd/Private/Factories/PackFactory.cpp b/Engine/Source/Editor/UnrealEd/Private/Factories/PackFactory.cpp
--- a/Engine/Source/Editor/UnrealEd/Private/Factories/PackFactory.cpp
+++ b/Engine/Source/Editor/UnrealEd/Private/Factories/PackFactory.cpp
@@ -92,7 +92,7 @@ namespace PackFactoryHelper
 		}
 
 		int64 WorkingSize = Entry.CompressionBlockSize;
-		int32 MaxCompressionBlockSize = FCompression::CompressMemoryBound(PakFile.GetInfo().GetCompressionMethod(Entry.CompressionMethodIndex), WorkingSize);
+		int32 MaxCompressionBlockSize = FCompression::GetMaximumCompressedSize(PakFile.GetInfo().GetCompressionMethod(Entry.CompressionMethodIndex), WorkingSize);
 		WorkingSize += MaxCompressionBlockSize;
 		if (PersistentBuffer.Num() < WorkingSize)
 		{
diff --git a/Engine/Source/Runtime/Core/Private/Misc/Compression.cpp b/Engine/Source/Runtime/Core/Private/Misc/Compression.cpp
--- a/Engine/Source/Runtime/Core/Private/Misc/Compression.cpp
+++ b/Engine/Source/Runtime/Core/Private/Misc/Compression.cpp
@@ -391,38 +391,42 @@ TAtomic<uint64> FCompression::CompressorDstBytes(0);
 
 static ECompressionFlags CheckGlobalCompressionFlags(ECompressionFlags Flags)
 {
-	static bool GAlwaysBiasCompressionForSize = false;
+	// COMPRESS_BiasSpeed / COMPRESS_BiasMemory is not actually checked anywhere
+	//	so no point doing this
+	#if 0
 	if(FPlatformProperties::HasEditorOnlyData())
 	{
+		static bool GAlwaysBiasCompressionForSize = false;
 		static bool GTestedCmdLine = false;
+
 		if(!GTestedCmdLine && FCommandLine::IsInitialized())
 		{
 			GTestedCmdLine = true;
 			// Override compression settings wrt size.
 			GAlwaysBiasCompressionForSize = FParse::Param(FCommandLine::Get(),TEXT("BIASCOMPRESSIONFORSIZE"));
 		}
-	}
 
-	// Always bias for speed if option is set.
-	if(GAlwaysBiasCompressionForSize)
-	{
-		int32 NewFlags = Flags;
-		NewFlags &= ~COMPRESS_BiasSpeed;
-		NewFlags |= COMPRESS_BiasMemory;
-		Flags = (ECompressionFlags)NewFlags;
+		// Always bias for speed if option is set.
+		if(GAlwaysBiasCompressionForSize)
+		{
+			int32 NewFlags = Flags;
+			NewFlags &= ~COMPRESS_BiasSpeed;
+			NewFlags |= COMPRESS_BiasMemory;
+			Flags = (ECompressionFlags)NewFlags;
+		}
 	}
+	#endif
 
 	return Flags;
 }
 
-
 uint32 FCompression::GetCompressorVersion(FName FormatName)
 {
-	if (FormatName == NAME_Zlib)
+	if (FormatName == NAME_None || FormatName == NAME_LZ4)
 	{
-		return appZLIBVersion();
+		return 0;
 	}
-	else if (FormatName == NAME_Gzip)
+	else if (FormatName == NAME_Zlib || FormatName == NAME_Gzip)
 	{
 		return appZLIBVersion();
 	}
@@ -504,11 +508,34 @@ FName FCompression::GetCompressionFormatFromDeprecatedFlags(ECompressionFlags Fl
 	return NAME_None;
 }
 
+int32 FCompression::GetMaximumCompressedSize(FName FormatName, int32 UncompressedSize, ECompressionFlags Flags, int32 CompressionData)
+{
+	if (FormatName == NAME_None)
+	{
+		return UncompressedSize;
+	}
+	else if ( FormatName == NAME_Oodle )
+	{
+		//	avoid calling CompressMemoryBound in the Decoder because it creates an ICompressionFormat for Oodle
+		//	and initializes encoders
+
+		return OodleDataGetMaximumCompressedSize(UncompressedSize);
+	}
+	else
+	{
+		return CompressMemoryBound(FormatName,UncompressedSize,Flags,CompressionData);
+	}
+}
+
 int32 FCompression::CompressMemoryBound(FName FormatName, int32 UncompressedSize, ECompressionFlags Flags, int32 CompressionData)
 {
 	int32 CompressionBound = UncompressedSize;
-
-	if (FormatName == NAME_Zlib)
+	
+	if (FormatName == NAME_None)
+	{
+		return UncompressedSize;
+	}
+	else if (FormatName == NAME_Zlib)
 	{
 		// Zlib's compressBounds gives a better (smaller) value, but only for the default bit window.
 		if (CompressionData == 0 || CompressionData == DEFAULT_ZLIB_BIT_WINDOW)
@@ -540,6 +567,9 @@ int32 FCompression::CompressMemoryBound(FName FormatName, int32 UncompressedSize
 		}
 	}
 
+	// this will fail if there was an int32 overflow (CompressionBound will be negative)
+	check( CompressionBound >= UncompressedSize );
+
 	return CompressionBound;
 }
 
@@ -643,6 +673,7 @@ bool FCompression::CompressMemory(FName FormatName, void* CompressedBuffer, int3
 	else
 	{
 		// let the format module compress it
+		// Oodle will make the OodleCompressionFormat here
 		ICompressionFormat* Format = GetCompressionFormat(FormatName);
 		if (Format)
 		{
@@ -669,13 +700,16 @@ FString FCompression::GetCompressorDDCSuffix(FName FormatName)
 {
 	FString DDCSuffix = FString::Printf(TEXT("%s_VER%D_"), *FormatName.ToString(), FCompression::GetCompressorVersion(FormatName));
 
-
-	if (FormatName == NAME_Zlib)
+	if (FormatName == NAME_None || FormatName == NAME_LZ4 )
+	{
+		// nothing
+	}
+	else if (FormatName == NAME_Zlib)
 	{
 		// hardcoded zlib
 		DDCSuffix += ZLIB_DERIVEDDATA_VER;
 	}
-	if (FormatName == NAME_Gzip)
+	else if (FormatName == NAME_Gzip)
 	{
 		DDCSuffix += GZIP_DERIVEDDATA_VER;
 	}
@@ -730,6 +764,8 @@ bool FCompression::UncompressMemory(FName FormatName, void* UncompressedBuffer,
 		ICompressionFormat* Format = GetCompressionFormat(FormatName);
 		if (Format)
 		{
+			// @todo Oodle : UncompressedSize is read-write to Format->Uncompress but that's not actually used
+			//	-> get rid of that
 			bUncompressSucceeded = Format->Uncompress(UncompressedBuffer, UncompressedSize, CompressedBuffer, CompressedSize, CompressionData);
 		}
 	}
diff --git a/Engine/Source/Runtime/Core/Public/Misc/Compression.h b/Engine/Source/Runtime/Core/Public/Misc/Compression.h
--- a/Engine/Source/Runtime/Core/Public/Misc/Compression.h
+++ b/Engine/Source/Runtime/Core/Public/Misc/Compression.h
@@ -42,6 +42,8 @@ struct FCompression
 
 	/**
 	 * Thread-safe abstract compression routine to query memory requirements for a compression operation.
+	 * This is the minimize size to allocate the buffer for CompressMemory (encoding).
+	 * Use GetMaximumCompressedSize at decode to know how large a compressed buffer may be.
 	 *
 	 * @param	FormatName					Name of the compression format
 	 * @param	UncompressedSize			Size of uncompressed data in bytes
@@ -50,6 +52,18 @@ struct FCompression
 	 * @return The maximum possible bytes needed for compression of data buffer of size UncompressedSize
 	 */
 	CORE_API static int32 CompressMemoryBound(FName FormatName, int32 UncompressedSize, ECompressionFlags Flags=COMPRESS_NoFlags, int32 CompressionData=0);
+	
+	/**
+	 * Thread-safe abstract compression routine to query maximum compressed size that could be made.
+	 * CompressMemoryBound is strictly greater equal GetMaximumCompressedSize.
+	 *
+	 * @param	FormatName					Name of the compression format
+	 * @param	UncompressedSize			Size of uncompressed data in bytes
+	 * @param	Flags						Flags to control what method to use and optionally control memory vs speed
+	 * @param	CompressionData				Additional compression parameter (specifies BitWindow value for ZLIB compression format)
+	 * @return The maximum possible size of valid compressed data made by this format
+	 */
+	CORE_API static int32 GetMaximumCompressedSize(FName FormatName, int32 UncompressedSize, ECompressionFlags Flags=COMPRESS_NoFlags, int32 CompressionData=0);
 
 	/**
 	 * Thread-safe abstract compression routine. Compresses memory from uncompressed buffer and writes it to compressed
diff --git a/Engine/Source/Runtime/PakFile/Private/IPlatformFilePak.cpp b/Engine/Source/Runtime/PakFile/Private/IPlatformFilePak.cpp
--- a/Engine/Source/Runtime/PakFile/Private/IPlatformFilePak.cpp
+++ b/Engine/Source/Runtime/PakFile/Private/IPlatformFilePak.cpp
@@ -4857,20 +4857,13 @@ public:
 			TEXT("See ExtensionsToNotUsePluginCompression in [Pak] section of Engine.ini to add more extensions."),
 			*CompressionMethod.ToString(), TEXT("Unknown"));
 
-		// @todo Oodle : make an FCompression::GetMaximumCompressedSize
-		//		avoid calling CompressMemoryBound in the Decoder because it creates an ICompressionFormat for Oodle
-		//		and initializes encoders
-		
-		int64 WorkingBufferRequiredSize;
-		if ( CompressionMethod == NAME_Oodle )
-		{
-			WorkingBufferRequiredSize = OodleDataGetMaximumCompressedSize(CompressionBlockSize);
-		}
-		else
+		int64 WorkingBufferRequiredSize = FCompression::GetMaximumCompressedSize(CompressionMethod,CompressionBlockSize);
+		if ( CompressionMethod != NAME_Oodle )
 		{
-			// an amount to extra allocate, in case one block's compressed size is bigger than CompressMemoryBound
+			// an amount to extra allocate, in case one block's compressed size is bigger than GetMaximumCompressedSize
+			// @todo this should not be needed, can it be removed?
 			float SlopMultiplier = 1.1f;
-			WorkingBufferRequiredSize = FCompression::CompressMemoryBound(CompressionMethod, CompressionBlockSize) * SlopMultiplier;
+			WorkingBufferRequiredSize = (int64)( WorkingBufferRequiredSize * SlopMultiplier );
 		}
 
 		WorkingBufferRequiredSize = EncryptionPolicy::AlignReadRequest(WorkingBufferRequiredSize);
