change OodleCompressDLL option to come from Engine.ini of the platform you are running on
instead of the game.ini of the platform you are packaging for

diff --git a/Engine/Source/Runtime/Core/Private/Compression/OodleDataCompressionFormat.cpp b/Engine/Source/Runtime/Core/Private/Compression/OodleDataCompressionFormat.cpp
--- a/Engine/Source/Runtime/Core/Private/Compression/OodleDataCompressionFormat.cpp
+++ b/Engine/Source/Runtime/Core/Private/Compression/OodleDataCompressionFormat.cpp
@@ -401,8 +401,12 @@ ICompressionFormat * CreateOodleDataCompressionFormat()
 		// OodleCompressDLL to use an earlier *encoder* to maintain identical bit streams
 		// mainly for shipped games to avoid patches
 		// do not specify this if you want to use the latest Engine version of Oodle for encoding
+
+		// NOTE : we get OodleCompressDLL from Engine.ini not from command line
+		//	because this is a function of the platform we are running *on* not the platform we are packaging *for*
+		// the ProjectPackaging settings we get on the command line from Game.ini come from the *target* platform
 		FString OodleDLL = "";
-		if ( FParse::Value(FCommandLine::Get(), TEXT("OodleCompressDLL="), OodleDLL) && ! OodleDLL.IsEmpty() )
+		if ( GConfig && GConfig->GetString(TEXT("OodleDataCompressionFormat"), TEXT("OodleCompressDLL"), OodleDLL, GEngineIni) && ! OodleDLL.IsEmpty() )
 		{
 			UE_LOG(OodleDataCompression, Display, TEXT("OodleCompressDLL=%s"), *OodleDLL);
 				
