ios and android SupportsVirtualTextureStreaming

diff --git a/Engine/Source/Runtime/Core/Public/Android/AndroidPlatformProperties.h b/Engine/Source/Runtime/Core/Public/Android/AndroidPlatformProperties.h
--- a/Engine/Source/Runtime/Core/Public/Android/AndroidPlatformProperties.h
+++ b/Engine/Source/Runtime/Core/Public/Android/AndroidPlatformProperties.h
@@ -122,6 +122,11 @@ struct FAndroidPlatformProperties
 	{
 		return true;
 	}
+
+	static FORCEINLINE bool SupportsVirtualTextureStreaming()
+	{
+		return true;
+	}
 };
 
 #ifdef PROPERTY_HEADER_SHOULD_DEFINE_TYPE
diff --git a/Engine/Source/Runtime/Core/Public/IOS/IOSPlatformProperties.h b/Engine/Source/Runtime/Core/Public/IOS/IOSPlatformProperties.h
--- a/Engine/Source/Runtime/Core/Public/IOS/IOSPlatformProperties.h
+++ b/Engine/Source/Runtime/Core/Public/IOS/IOSPlatformProperties.h
@@ -104,6 +104,11 @@ struct FIOSPlatformProperties
 	{
 		return true;
 	}
+
+	static FORCEINLINE bool SupportsVirtualTextureStreaming()
+	{
+		return true;
+	}
 };
 
 struct FTVOSPlatformProperties : public FIOSPlatformProperties
