Compression Flags add ForPackaging to differentiate packaging compression usage
add OodleDataCompression CommonUsage to translate Flags

diff --git a/Engine/Source/Developer/PakFileUtilities/Private/PakFileUtilities.cpp b/Engine/Source/Developer/PakFileUtilities/Private/PakFileUtilities.cpp
--- a/Engine/Source/Developer/PakFileUtilities/Private/PakFileUtilities.cpp
+++ b/Engine/Source/Developer/PakFileUtilities/Private/PakFileUtilities.cpp
@@ -96,7 +96,7 @@ public:
 	{
 		// Compress memory block.
 		// Actual size will be stored to CompressedSize.
-		Result = FCompression::CompressMemory(Format, CompressedBuffer, CompressedSize, UncompressedBuffer, UncompressedSize);
+		Result = FCompression::CompressMemory(Format, CompressedBuffer, CompressedSize, UncompressedBuffer, UncompressedSize, COMPRESS_ForPackaging);
 	}
 
 	
diff --git a/Engine/Source/Runtime/Core/Private/Compression/OodleDataCompression.cpp b/Engine/Source/Runtime/Core/Private/Compression/OodleDataCompression.cpp
--- a/Engine/Source/Runtime/Core/Private/Compression/OodleDataCompression.cpp
+++ b/Engine/Source/Runtime/Core/Private/Compression/OodleDataCompression.cpp
@@ -187,6 +187,54 @@ static OodleLZ_CompressionLevel OodleDataCompressionLevelToOodleLZ_CompressionLe
 	return (OodleLZ_CompressionLevel) IntLevel;
 }
 
+EOodleDataCompressionCommonUsage CORE_API GetCommonUsageFromLegacyCompressionFlags(ECompressionFlags Flags)
+{
+	switch(Flags)
+	{
+		case 0:
+			return EOodleDataCompressionCommonUsage::Default;
+		case COMPRESS_BiasSpeed:
+			return EOodleDataCompressionCommonUsage::FastRealtimeEncode;
+		case COMPRESS_BiasSize:
+			return EOodleDataCompressionCommonUsage::SlowerSmallerEncode;
+		case COMPRESS_ForPackaging:
+			return EOodleDataCompressionCommonUsage::SlowestOfflineDistributionEncode;
+
+		default:
+			UE_LOG(OodleDataCompression,Error,TEXT("Invalid ECompressionFlags : %04X\n"),Flags);
+			return EOodleDataCompressionCommonUsage::Default;
+	}
+}
+
+void CORE_API GetCompressorAndLevelForCommonUsage(EOodleDataCompressionCommonUsage Usage,EOodleDataCompressor & OutCompressor,EOodleDataCompressionLevel & OutLevel)
+{
+	switch(Usage)
+	{
+		case EOodleDataCompressionCommonUsage::Default:
+			OutCompressor = EOodleDataCompressor::Kraken;
+			OutLevel = EOodleDataCompressionLevel::Fast;
+			break;
+		case EOodleDataCompressionCommonUsage::FastRealtimeEncode:
+			OutCompressor = EOodleDataCompressor::Mermaid;
+			OutLevel = EOodleDataCompressionLevel::HyperFast2;
+			break;
+		case EOodleDataCompressionCommonUsage::SlowerSmallerEncode:
+			OutCompressor = EOodleDataCompressor::Kraken;
+			OutLevel = EOodleDataCompressionLevel::Normal;
+			break;
+		case EOodleDataCompressionCommonUsage::SlowestOfflineDistributionEncode:
+			OutCompressor = EOodleDataCompressor::Kraken;
+			OutLevel = EOodleDataCompressionLevel::Optimal2;
+			break;
+		default:
+			UE_LOG(OodleDataCompression,Error,TEXT("Invalid ECompressionFlags : %d\n"),(int)Usage);
+			OutCompressor = EOodleDataCompressor::Selkie;
+			OutLevel = EOodleDataCompressionLevel::None;
+			return;
+	}
+}
+
+
 int64 CORE_API OodleDataCompressedBufferSizeNeeded(int64 InUncompressedSize)
 {
 	// size needed is the same for all newlz's
diff --git a/Engine/Source/Runtime/Core/Private/Compression/OodleDataCompressionFormat.cpp b/Engine/Source/Runtime/Core/Private/Compression/OodleDataCompressionFormat.cpp
--- a/Engine/Source/Runtime/Core/Private/Compression/OodleDataCompressionFormat.cpp
+++ b/Engine/Source/Runtime/Core/Private/Compression/OodleDataCompressionFormat.cpp
@@ -185,63 +185,100 @@ struct FOodleDataCompressionFormat : ICompressionFormat
 		return FString::Printf(TEXT("C_%s_CL_%s_%s"), *GetCompressorString(), *GetCompressionLevelString(), OODLE_DERIVEDDATA_VER);
 	}
 
-	virtual bool Compress(void* OutCompressedBuffer, int32& OutCompressedSize, const void* InUncompressedBuffer, int32 InUncompressedSize, int32 InCompressionData) override
+	virtual bool Compress(void* OutCompressedBuffer, int32& OutCompressedSize, const void* InUncompressedBuffer, int32 InUncompressedSize, int32 InCompressionData, ECompressionFlags Flags) override
 	{
 		// OutCompressedSize is read-write
 		int32 CompressedBufferSize = OutCompressedSize;
 
 		// CompressedSize should be >= GetCompressedBufferSize(UncompressedSize, CompressionData)
 		check(CompressedBufferSize >= GetCompressedBufferSize(InUncompressedSize, InCompressionData));
-		
-		OO_SINTa Result = 0;
 
-		if ( OodleLZCompressFuncPtr != nullptr )
+		if ( Flags & COMPRESS_ForPackaging )
 		{
-			// encode using func pointer to loaded DLL :
-			Result = (*OodleLZCompressFuncPtr)(Compressor, InUncompressedBuffer, InUncompressedSize, OutCompressedBuffer, CompressionLevel, NULL, NULL, NULL);
+			// Compressing for pak/iostore
+			// use the options from the command line (ProjectPackagingSettings) to choose compression settings
 
-			// NOTE: not using &CompressOptions
-			//	because OodleLZ_CompressOptions_GetDefault has changed from v5
-			//	need to either call that in the DLL func ptr or not use it
-			//	(SSTB == 0 doesn't mean default in old Oodle)
+			OO_SINTa Result = 0;
 
-			#if 0 // DO_CHECK
-			if ( Result > 0 )
+			if ( OodleLZCompressFuncPtr != nullptr )
 			{
-				// verify we can decode data made with DLL-enocder using the current decoder :
+				// encode using func pointer to loaded DLL :
+				Result = (*OodleLZCompressFuncPtr)(Compressor, InUncompressedBuffer, InUncompressedSize, OutCompressedBuffer, CompressionLevel, NULL, NULL, NULL);
 
-				int32 DecodeSpaceSize = InUncompressedSize + 16; // FuzzSafe_Yes so shouldn't need big padding even on Oodle v5
-				void * DecodeSpace = FMemory::Malloc(DecodeSpaceSize);
-				check( DecodeSpace != nullptr );
+				// NOTE: not using &CompressOptions
+				//	because OodleLZ_CompressOptions_GetDefault has changed from v5
+				//	need to either call that in the DLL func ptr or not use it
+				//	(SSTB == 0 doesn't mean default in old Oodle)
+
+				#if 0 // DO_CHECK
+				if ( Result > 0 )
+				{
+					// verify we can decode data made with DLL-enocder using the current decoder :
 
-				OO_SINTa DecRet = OodleLZ_Decompress(OutCompressedBuffer,Result,DecodeSpace,InUncompressedSize,OodleLZ_FuzzSafe_Yes);
-				check( DecRet == InUncompressedSize );
+					int32 DecodeSpaceSize = InUncompressedSize + 16; // FuzzSafe_Yes so shouldn't need big padding even on Oodle v5
+					void * DecodeSpace = FMemory::Malloc(DecodeSpaceSize);
+					check( DecodeSpace != nullptr );
 
-				int cmp = memcmp(DecodeSpace,InUncompressedBuffer,InUncompressedSize);
-				check( cmp == 0 );
+					OO_SINTa DecRet = OodleLZ_Decompress(OutCompressedBuffer,Result,DecodeSpace,InUncompressedSize,OodleLZ_FuzzSafe_Yes);
+					check( DecRet == InUncompressedSize );
 
-				FMemory::Free(DecodeSpace);
+					int cmp = memcmp(DecodeSpace,InUncompressedBuffer,InUncompressedSize);
+					check( cmp == 0 );
+
+					FMemory::Free(DecodeSpace);
+				}
+				#endif
+			}
+			else
+			{
+				Result = OodleLZ_Compress(Compressor, InUncompressedBuffer, InUncompressedSize, OutCompressedBuffer, CompressionLevel, &CompressionOptions);
 			}
-			#endif
-		}
-		else
-		{
-			Result = OodleLZ_Compress(Compressor, InUncompressedBuffer, InUncompressedSize, OutCompressedBuffer, CompressionLevel, &CompressionOptions);
-		}
 
-		// verbose log all compresses :
-		//UE_LOG(OodleDataCompression, Display, TEXT("Oodle Compress : %d -> %d"), UncompressedSize, Result);
+			// verbose log all compresses :
+			//UE_LOG(OodleDataCompression, Display, TEXT("Oodle Compress : %d -> %d"), UncompressedSize, Result);
 		
-		if (Result <= 0)
-		{
-			OutCompressedSize = -1;
-			return false;
+			if (Result <= 0)
+			{
+				OutCompressedSize = -1;
+				return false;
+			}
+			else
+			{
+				OutCompressedSize = (int32) Result;
+				return true;
+			}
 		}
 		else
 		{
-			OutCompressedSize = (int32) Result;
-			return true;
+			// Not for Packaging
+			// in-game or tools compress
+			// interpret the legacy Flags if any
+			// most new users of Oodle should go through OodleDataCompress instead
+			// this supports legacy users
+
+			// Compressor and CompressionLevel are member vars, watch out
+			EOodleDataCompressor CurCompressor;
+			EOodleDataCompressionLevel CurCompressionLevel;
+			EOodleDataCompressionCommonUsage Usage = GetCommonUsageFromLegacyCompressionFlags(Flags);
+			GetCompressorAndLevelForCommonUsage(Usage,CurCompressor,CurCompressionLevel);
+
+			int64 Result = OodleDataCompress(
+							OutCompressedBuffer,CompressedBufferSize,
+							InUncompressedBuffer,InUncompressedSize,
+							CurCompressor,CurCompressionLevel);
+
+			if ( Result <= 0 )
+			{
+				OutCompressedSize = -1;
+				return false;
+			}
+			else
+			{
+				OutCompressedSize = (int32) Result;
+				return true;
+			}
 		}
+	
 	}
 
 	virtual bool Uncompress(void* OutUncompressedBuffer, int32& OutUncompressedSize, const void* InCompressedBuffer, int32 InCompressedSize, int32 CompressionData) override
diff --git a/Engine/Source/Runtime/Core/Private/IO/IoStore.cpp b/Engine/Source/Runtime/Core/Private/IO/IoStore.cpp
--- a/Engine/Source/Runtime/Core/Private/IO/IoStore.cpp
+++ b/Engine/Source/Runtime/Core/Private/IO/IoStore.cpp
@@ -1037,7 +1037,8 @@ private:
 					Block->IoBuffer->Data(),
 					CompressedBlockSize,
 					Block->UncompressedData,
-					static_cast<int32>(Block->UncompressedSize));
+					static_cast<int32>(Block->UncompressedSize),
+					COMPRESS_ForPackaging);
 			}
 			if ( !bCompressed )
 			{
diff --git a/Engine/Source/Runtime/Core/Private/Misc/Compression.cpp b/Engine/Source/Runtime/Core/Private/Misc/Compression.cpp
--- a/Engine/Source/Runtime/Core/Private/Misc/Compression.cpp
+++ b/Engine/Source/Runtime/Core/Private/Misc/Compression.cpp
@@ -389,37 +389,6 @@ TAtomic<uint64> FCompression::CompressorSrcBytes(0);
 /** Nubmer of bytes after compression.		*/
 TAtomic<uint64> FCompression::CompressorDstBytes(0);
 
-static ECompressionFlags CheckGlobalCompressionFlags(ECompressionFlags Flags)
-{
-	// COMPRESS_BiasSpeed / COMPRESS_BiasMemory is not actually checked anywhere
-	//	so no point doing this
-	#if 0
-	if(FPlatformProperties::HasEditorOnlyData())
-	{
-		static bool GAlwaysBiasCompressionForSize = false;
-		static bool GTestedCmdLine = false;
-
-		if(!GTestedCmdLine && FCommandLine::IsInitialized())
-		{
-			GTestedCmdLine = true;
-			// Override compression settings wrt size.
-			GAlwaysBiasCompressionForSize = FParse::Param(FCommandLine::Get(),TEXT("BIASCOMPRESSIONFORSIZE"));
-		}
-
-		// Always bias for speed if option is set.
-		if(GAlwaysBiasCompressionForSize)
-		{
-			int32 NewFlags = Flags;
-			NewFlags &= ~COMPRESS_BiasSpeed;
-			NewFlags |= COMPRESS_BiasMemory;
-			Flags = (ECompressionFlags)NewFlags;
-		}
-	}
-	#endif
-
-	return Flags;
-}
-
 uint32 FCompression::GetCompressorVersion(FName FormatName)
 {
 	if (FormatName == NAME_None || FormatName == NAME_LZ4)
@@ -652,8 +621,6 @@ bool FCompression::CompressMemory(FName FormatName, void* CompressedBuffer, int3
 
 	bool bCompressSucceeded = false;
 
-	Flags = CheckGlobalCompressionFlags(Flags);
-
 	if (FormatName == NAME_Zlib)
 	{
 		// hardcoded zlib
@@ -677,7 +644,7 @@ bool FCompression::CompressMemory(FName FormatName, void* CompressedBuffer, int3
 		ICompressionFormat* Format = GetCompressionFormat(FormatName);
 		if (Format)
 		{
-			bCompressSucceeded = Format->Compress(CompressedBuffer, CompressedSize, UncompressedBuffer, UncompressedSize, CompressionData);
+			bCompressSucceeded = Format->Compress(CompressedBuffer, CompressedSize, UncompressedBuffer, UncompressedSize, CompressionData, Flags);
 		}
 	}
 
@@ -1019,8 +986,11 @@ void* FCompressedGrowableBuffer::Access( int32 Offset )
 
 bool FCompression::IsFormatValid(FName FormatName)
 {
+	//@todo Oodle make NAME_None a valid compressor
+	// FormatName == NAME_None || 
+
 	// build in formats are always valid
-	if (FormatName == NAME_Zlib || FormatName == NAME_Gzip || FormatName == NAME_LZ4 || FormatName == NAME_Oodle)
+	if ( FormatName == NAME_Zlib || FormatName == NAME_Gzip || FormatName == NAME_LZ4 || FormatName == NAME_Oodle)
 	{
 		return true;
 	}
@@ -1031,7 +1001,7 @@ bool FCompression::IsFormatValid(FName FormatName)
 
 bool FCompression::VerifyCompressionFlagsValid(int32 InCompressionFlags)
 {
-	const int32 CompressionFlagsMask = COMPRESS_DeprecatedFormatFlagsMask | COMPRESS_OptionsFlagsMask;
+	const int32 CompressionFlagsMask = COMPRESS_DeprecatedFormatFlagsMask | COMPRESS_OptionsFlagsMask | COMPRESS_ForPurposeMask;
 	if (InCompressionFlags & (~CompressionFlagsMask))
 	{
 		return false;
diff --git a/Engine/Source/Runtime/Core/Public/Compression/OodleDataCompression.h b/Engine/Source/Runtime/Core/Public/Compression/OodleDataCompression.h
--- a/Engine/Source/Runtime/Core/Public/Compression/OodleDataCompression.h
+++ b/Engine/Source/Runtime/Core/Public/Compression/OodleDataCompression.h
@@ -5,6 +5,7 @@
 #include "CoreTypes.h"
 #include "Logging/LogCategory.h"
 #include "Logging/LogMacros.h"
+#include "Misc/CompressionFlags.h"
 
 // OodleDataCompression :
 //	Unreal API for direct access to Oodle Data lossless data compression
@@ -96,6 +97,35 @@ enum class EOodleDataCompressionLevel : int8
 	Optimal4 = 8,
 };
 
+enum class EOodleDataCompressionCommonUsage : uint8
+{
+	Default = 0,
+	FastRealtimeEncode = 1,
+	SlowerSmallerEncode = 2,
+	SlowestOfflineDistributionEncode = 3
+};
+
+/**
+* Translate legacy CompressionFlags to an EOodleDataCompressionCommonUsage
+*   CompressionFlags is not encouraged for new code; prefer directly calling to Oodle
+*
+* @param	Flags				ECompressionFlags (BiasSpeed,BiasSize,ForPackaging)
+* @return						EOodleDataCompressionCommonUsage
+*/
+EOodleDataCompressionCommonUsage CORE_API GetCommonUsageFromLegacyCompressionFlags(ECompressionFlags Flags);
+
+/**
+* Translate OodleDataCompressionCommonUsage to a Compressor & Level selection
+*   usually prefer the more expressive choice of {Compressor,Level}
+*
+* @param	Usage				Your intended compression use case
+* @param	OutCompressor		Output reference
+* @param	OutLevel			Output reference
+* @return						
+*/
+void CORE_API GetCompressorAndLevelForCommonUsage(EOodleDataCompressionCommonUsage Usage,EOodleDataCompressor & OutCompressor,EOodleDataCompressionLevel & OutLevel);
+
+
 /**
 * What size of compressed output buffer is needed to encode into for OodleDataCompress()
 *
diff --git a/Engine/Source/Runtime/Core/Public/Misc/CompressionFlags.h b/Engine/Source/Runtime/Core/Public/Misc/CompressionFlags.h
--- a/Engine/Source/Runtime/Core/Public/Misc/CompressionFlags.h
+++ b/Engine/Source/Runtime/Core/Public/Misc/CompressionFlags.h
@@ -27,6 +27,7 @@ enum ECompressionFlags
 	COMPRESS_NoFlags = 0x00,
 	/** Prefer compression that compresses smaller (ONLY VALID FOR COMPRESSION)		*/
 	COMPRESS_BiasMemory = 0x10,
+	COMPRESS_BiasSize = COMPRESS_BiasMemory,
 	/** Prefer compression that compresses faster (ONLY VALID FOR COMPRESSION)		*/
 	COMPRESS_BiasSpeed = 0x20,
 	/** Is the source buffer padded out	(ONLY VALID FOR UNCOMPRESS)					*/
@@ -34,5 +35,9 @@ enum ECompressionFlags
 
 	/** Set of flags that are options are still allowed								*/
 	COMPRESS_OptionsFlagsMask = 0xF0,
+
+	/** Indicate this compress call is for Packaging (pak/iostore) */
+	COMPRESS_ForPackaging = 0x100,
+	COMPRESS_ForPurposeMask = 0xF00,
 };
 
diff --git a/Engine/Source/Runtime/Core/Public/Misc/ICompressionFormat.h b/Engine/Source/Runtime/Core/Public/Misc/ICompressionFormat.h
--- a/Engine/Source/Runtime/Core/Public/Misc/ICompressionFormat.h
+++ b/Engine/Source/Runtime/Core/Public/Misc/ICompressionFormat.h
@@ -6,7 +6,7 @@
 #include "Features/IModularFeatures.h"
 #include "Modules/ModuleInterface.h"
 #include "Modules/ModuleManager.h"
-
+#include "Misc/CompressionFlags.h"
 
 
 #define COMPRESSION_FORMAT_FEATURE_NAME "CompressionFormat"
@@ -14,7 +14,7 @@
 struct ICompressionFormat : public IModularFeature, public IModuleInterface
 {
 	virtual FName GetCompressionFormatName() = 0;
-	virtual bool Compress(void* CompressedBuffer, int32& CompressedSize, const void* UncompressedBuffer, int32 UncompressedSize, int32 CompressionData) = 0;
+	virtual bool Compress(void* CompressedBuffer, int32& CompressedSize, const void* UncompressedBuffer, int32 UncompressedSize, int32 CompressionData, ECompressionFlags Flags) = 0;
 	virtual bool Uncompress(void* UncompressedBuffer, int32& UncompressedSize, const void* CompressedBuffer, int32 CompressedSize, int32 CompressionData) = 0;
 	virtual int32 GetCompressedBufferSize(int32 UncompressedSize, int32 CompressionData) = 0;
 	virtual uint32 GetVersion() = 0;
