OodleDataCompression fixing namespace convention

diff --git a/Engine/Plugins/Compression/OodleNetwork/Source/Private/OodleNetworkArchives.cpp b/Engine/Plugins/Compression/OodleNetwork/Source/Private/OodleNetworkArchives.cpp
--- a/Engine/Plugins/Compression/OodleNetwork/Source/Private/OodleNetworkArchives.cpp
+++ b/Engine/Plugins/Compression/OodleNetwork/Source/Private/OodleNetworkArchives.cpp
@@ -86,7 +86,7 @@ bool FOodleNetworkArchiveBase::SerializeOodleDecompressData(FOodleCompressedData
 
 			if (!InnerArchive.IsError())
 			{
-				bSuccess = OodleDataDecompress(
+				bSuccess = FOodleDataCompression::Decompress(
 					(void*)DecompressedData, DecompressedLength,
 					(void*)CompressedData, CompressedLength);
 			}
diff --git a/Engine/Source/Runtime/Core/Private/Compression/OodleDataCompression.cpp b/Engine/Source/Runtime/Core/Private/Compression/OodleDataCompression.cpp
--- a/Engine/Source/Runtime/Core/Private/Compression/OodleDataCompression.cpp
+++ b/Engine/Source/Runtime/Core/Private/Compression/OodleDataCompression.cpp
@@ -9,6 +9,11 @@
 
 DEFINE_LOG_CATEGORY(OodleDataCompression);
 
+extern ICompressionFormat * CreateOodleDataCompressionFormat();
+
+namespace FOodleDataCompression
+{
+
 struct OodleDataCompressionDecoders
 {
 	OO_SINTa OodleDecoderMemorySize = 0;
@@ -154,19 +159,19 @@ struct OodleDataCompressionDecoders
 
 
 
-static OodleLZ_Compressor OodleDataCompressorToOodleLZ_Compressor(EOodleDataCompressor Compressor)
+static OodleLZ_Compressor CompressorToOodleLZ_Compressor(ECompressor Compressor)
 {
 	switch(Compressor)
 	{
-	case EOodleDataCompressor::Selkie:
+	case ECompressor::Selkie:
 		return OodleLZ_Compressor_Selkie;
-	case EOodleDataCompressor::Mermaid:
+	case ECompressor::Mermaid:
 		return OodleLZ_Compressor_Mermaid;
-	case EOodleDataCompressor::Kraken:
+	case ECompressor::Kraken:
 		return OodleLZ_Compressor_Kraken;
-	case EOodleDataCompressor::Leviathan:
+	case ECompressor::Leviathan:
 		return OodleLZ_Compressor_Leviathan;
-	case EOodleDataCompressor::NotSet:
+	case ECompressor::NotSet:
 		return OodleLZ_Compressor_Invalid;
 	default:
 		UE_LOG(OodleDataCompression,Error,TEXT("Invalid Compressor: %d\n"),(int)Compressor);
@@ -174,7 +179,7 @@ static OodleLZ_Compressor OodleDataCompressorToOodleLZ_Compressor(EOodleDataComp
 	}
 }
 
-static OodleLZ_CompressionLevel OodleDataCompressionLevelToOodleLZ_CompressionLevel(EOodleDataCompressionLevel Level)
+static OodleLZ_CompressionLevel CompressionLevelToOodleLZ_CompressionLevel(ECompressionLevel Level)
 {
 	int IntLevel = (int)Level;
 	
@@ -187,76 +192,76 @@ static OodleLZ_CompressionLevel OodleDataCompressionLevelToOodleLZ_CompressionLe
 	return (OodleLZ_CompressionLevel) IntLevel;
 }
 
-EOodleDataCompressionCommonUsage CORE_API GetCommonUsageFromLegacyCompressionFlags(ECompressionFlags Flags)
+ECompressionCommonUsage CORE_API GetCommonUsageFromLegacyCompressionFlags(ECompressionFlags Flags)
 {
 	switch(Flags)
 	{
 		case 0:
-			return EOodleDataCompressionCommonUsage::Default;
+			return ECompressionCommonUsage::Default;
 		case COMPRESS_BiasSpeed:
-			return EOodleDataCompressionCommonUsage::FastRealtimeEncode;
+			return ECompressionCommonUsage::FastRealtimeEncode;
 		case COMPRESS_BiasSize:
-			return EOodleDataCompressionCommonUsage::SlowerSmallerEncode;
+			return ECompressionCommonUsage::SlowerSmallerEncode;
 		case COMPRESS_ForPackaging:
-			return EOodleDataCompressionCommonUsage::SlowestOfflineDistributionEncode;
+			return ECompressionCommonUsage::SlowestOfflineDistributionEncode;
 
 		default:
 			UE_LOG(OodleDataCompression,Error,TEXT("Invalid ECompressionFlags : %04X\n"),Flags);
-			return EOodleDataCompressionCommonUsage::Default;
+			return ECompressionCommonUsage::Default;
 	}
 }
 
-void CORE_API GetCompressorAndLevelForCommonUsage(EOodleDataCompressionCommonUsage Usage,EOodleDataCompressor & OutCompressor,EOodleDataCompressionLevel & OutLevel)
+void CORE_API GetCompressorAndLevelForCommonUsage(ECompressionCommonUsage Usage,ECompressor & OutCompressor,ECompressionLevel & OutLevel)
 {
 	switch(Usage)
 	{
-		case EOodleDataCompressionCommonUsage::Default:
-			OutCompressor = EOodleDataCompressor::Kraken;
-			OutLevel = EOodleDataCompressionLevel::Fast;
+		case ECompressionCommonUsage::Default:
+			OutCompressor = ECompressor::Kraken;
+			OutLevel = ECompressionLevel::Fast;
 			break;
-		case EOodleDataCompressionCommonUsage::FastRealtimeEncode:
-			OutCompressor = EOodleDataCompressor::Mermaid;
-			OutLevel = EOodleDataCompressionLevel::HyperFast2;
+		case ECompressionCommonUsage::FastRealtimeEncode:
+			OutCompressor = ECompressor::Mermaid;
+			OutLevel = ECompressionLevel::HyperFast2;
 			break;
-		case EOodleDataCompressionCommonUsage::SlowerSmallerEncode:
-			OutCompressor = EOodleDataCompressor::Kraken;
-			OutLevel = EOodleDataCompressionLevel::Normal;
+		case ECompressionCommonUsage::SlowerSmallerEncode:
+			OutCompressor = ECompressor::Kraken;
+			OutLevel = ECompressionLevel::Normal;
 			break;
-		case EOodleDataCompressionCommonUsage::SlowestOfflineDistributionEncode:
-			OutCompressor = EOodleDataCompressor::Kraken;
-			OutLevel = EOodleDataCompressionLevel::Optimal2;
+		case ECompressionCommonUsage::SlowestOfflineDistributionEncode:
+			OutCompressor = ECompressor::Kraken;
+			OutLevel = ECompressionLevel::Optimal2;
 			break;
 		default:
 			UE_LOG(OodleDataCompression,Error,TEXT("Invalid ECompressionFlags : %d\n"),(int)Usage);
-			OutCompressor = EOodleDataCompressor::Selkie;
-			OutLevel = EOodleDataCompressionLevel::None;
+			OutCompressor = ECompressor::Selkie;
+			OutLevel = ECompressionLevel::None;
 			return;
 	}
 }
 
 
-int64 CORE_API OodleDataCompressedBufferSizeNeeded(int64 InUncompressedSize)
+int64 CORE_API CompressedBufferSizeNeeded(int64 InUncompressedSize)
 {
 	// size needed is the same for all newlz's
 	//	so don't bother with a compressor arg here
 	return OodleLZ_GetCompressedBufferSizeNeeded(OodleLZ_Compressor_Kraken, TCheckValueCast<OO_SINTa>(InUncompressedSize));
 }
 
-int64 CORE_API OodleDataGetMaximumCompressedSize(int64 InUncompressedSize)
+int64 CORE_API GetMaximumCompressedSize(int64 InUncompressedSize)
 {
 	int64 NumBlocks = (InUncompressedSize+ OODLELZ_BLOCK_LEN-1)/OODLELZ_BLOCK_LEN;
 	int64 MaxCompressedSize = InUncompressedSize + NumBlocks * OODLELZ_BLOCK_MAXIMUM_EXPANSION;
 	return MaxCompressedSize;
 }
 
-int64 CORE_API OodleDataCompress(
+int64 CORE_API Compress(
 							void * OutCompressedData, int64 InCompressedBufferSize,
 							const void * InUncompressedData, int64 InUncompressedSize,
-							EOodleDataCompressor Compressor,
-							EOodleDataCompressionLevel Level)
+							ECompressor Compressor,
+							ECompressionLevel Level)
 {
-	OodleLZ_Compressor LZCompressor = OodleDataCompressorToOodleLZ_Compressor(Compressor);
-	OodleLZ_CompressionLevel LZLevel = OodleDataCompressionLevelToOodleLZ_CompressionLevel(Level);
+	OodleLZ_Compressor LZCompressor = CompressorToOodleLZ_Compressor(Compressor);
+	OodleLZ_CompressionLevel LZLevel = CompressionLevelToOodleLZ_CompressionLevel(Level);
 
 	if ( InCompressedBufferSize < (int64) OodleLZ_GetCompressedBufferSizeNeeded(LZCompressor,TCheckValueCast<OO_SINTa>(InUncompressedSize)) )
 	{
@@ -280,9 +285,8 @@ static OodleDataCompressionDecoders * GetGlobalOodleDataCompressionDecoders()
 
 
 static ICompressionFormat * GlobalOodleDataCompressionFormat = nullptr;
-extern ICompressionFormat * CreateOodleDataCompressionFormat();
 
-void CORE_API OodleDataCompressionFormatInitOnFirstUseFromLock()
+void CORE_API CompressionFormatInitOnFirstUseFromLock()
 {
 	// called from inside a critical section lock 
 	//	from Compression.cpp / GetCompressionFormat
@@ -293,7 +297,7 @@ void CORE_API OodleDataCompressionFormatInitOnFirstUseFromLock()
 }
 
 
-bool CORE_API OodleDataDecompress(
+bool CORE_API Decompress(
 						void * OutUncompressedData, int64 InUncompressedSize,
 						const void * InCompressedData, int64 InCompressedSize)
 {
@@ -310,7 +314,7 @@ bool CORE_API OodleDataDecompress(
 	return true;
 }
 					
-void CORE_API OodleDataCompressionStartupPreInit(void)
+void CORE_API StartupPreInit(void)
 {
 	// called from LaunchEngineLoop at "PreInit" time
 	// not all Engine services may be set up yet, be careful what you use
@@ -324,3 +328,5 @@ void CORE_API OodleDataCompressionStartupPreInit(void)
 	OodleConfig.m_OodleLZ_BackwardsCompatible_MajorVersion = 9;
 	Oodle_SetConfigValues(&OodleConfig);
 }
+
+};
diff --git a/Engine/Source/Runtime/Core/Private/Compression/OodleDataCompressionFormat.cpp b/Engine/Source/Runtime/Core/Private/Compression/OodleDataCompressionFormat.cpp
--- a/Engine/Source/Runtime/Core/Private/Compression/OodleDataCompressionFormat.cpp
+++ b/Engine/Source/Runtime/Core/Private/Compression/OodleDataCompressionFormat.cpp
@@ -257,12 +257,12 @@ struct FOodleDataCompressionFormat : ICompressionFormat
 			// this supports legacy users
 
 			// Compressor and CompressionLevel are member vars, watch out
-			EOodleDataCompressor CurCompressor;
-			EOodleDataCompressionLevel CurCompressionLevel;
-			EOodleDataCompressionCommonUsage Usage = GetCommonUsageFromLegacyCompressionFlags(Flags);
+			FOodleDataCompression::ECompressor CurCompressor;
+			FOodleDataCompression::ECompressionLevel CurCompressionLevel;
+			FOodleDataCompression::ECompressionCommonUsage Usage = FOodleDataCompression::GetCommonUsageFromLegacyCompressionFlags(Flags);
 			GetCompressorAndLevelForCommonUsage(Usage,CurCompressor,CurCompressionLevel);
 
-			int64 Result = OodleDataCompress(
+			int64 Result = FOodleDataCompression::Compress(
 							OutCompressedBuffer,CompressedBufferSize,
 							InUncompressedBuffer,InUncompressedSize,
 							CurCompressor,CurCompressionLevel);
@@ -284,7 +284,7 @@ struct FOodleDataCompressionFormat : ICompressionFormat
 	virtual bool Uncompress(void* OutUncompressedBuffer, int32& OutUncompressedSize, const void* InCompressedBuffer, int32 InCompressedSize, int32 CompressionData) override
 	{
 		// OutUncompressedSize is read-write
-		return OodleDataDecompress(
+		return FOodleDataCompression::Decompress(
 			OutUncompressedBuffer, OutUncompressedSize,
 			InCompressedBuffer,InCompressedSize);
 	}
diff --git a/Engine/Source/Runtime/Core/Private/Compression/OodleDataCompressionUtil.cpp b/Engine/Source/Runtime/Core/Private/Compression/OodleDataCompressionUtil.cpp
--- a/Engine/Source/Runtime/Core/Private/Compression/OodleDataCompressionUtil.cpp
+++ b/Engine/Source/Runtime/Core/Private/Compression/OodleDataCompressionUtil.cpp
@@ -6,7 +6,7 @@
 #include "Serialization/MemoryReader.h"
 #include "Serialization/MemoryWriter.h"
 
-namespace OodleDataCompressionUtil
+namespace FOodleDataCompressionUtil
 {
 	bool CORE_API DecompressReplayData(const TArray<uint8>& InCompressed, TArray< uint8 >& OutBuffer)
 	{
@@ -22,7 +22,7 @@ namespace OodleDataCompressionUtil
 		if ( Reader.Tell() + CompressedSize > InCompressed.Num() )
 			return false;
 
-		return OodleDataDecompress(OutBuffer.GetData(), OutBuffer.Num(), InCompressed.GetData() + Reader.Tell(), CompressedSize );
+		return FOodleDataCompression::Decompress(OutBuffer.GetData(), OutBuffer.Num(), InCompressed.GetData() + Reader.Tell(), CompressedSize );
 	}
 
 	bool CORE_API CompressReplayData(const TArray<uint8>& InBuffer, TArray< uint8 >& OutCompressed)
@@ -33,14 +33,14 @@ namespace OodleDataCompressionUtil
 
 		int32 ReservedBytes = sizeof(int32)*2;
 				
-		const EOodleDataCompressor			Compressor = EOodleDataCompressor::Selkie;
-		const EOodleDataCompressionLevel	Level = EOodleDataCompressionLevel::VeryFast;
+		const FOodleDataCompression::ECompressor		Compressor = FOodleDataCompression::ECompressor::Selkie;
+		const FOodleDataCompression::ECompressionLevel	Level = FOodleDataCompression::ECompressionLevel::VeryFast;
 		
-		int32 CompSizeNeeded = (int32) OodleDataCompressedBufferSizeNeeded(InSize);
+		int32 CompSizeNeeded = (int32) FOodleDataCompression::CompressedBufferSizeNeeded(InSize);
 		OutCompressed.SetNum( TCheckValueCast<int32>( ReservedBytes + CompSizeNeeded) );
 		void * CompPtr = OutCompressed.GetData() + ReservedBytes;
 
-		CompressedSize = (int32) OodleDataCompress(CompPtr,CompSizeNeeded,InPtr,InSize,Compressor,Level);
+		CompressedSize = (int32) FOodleDataCompression::Compress(CompPtr,CompSizeNeeded,InPtr,InSize,Compressor,Level);
 		if ( CompressedSize <= 0 )
 		{
 			return false;
diff --git a/Engine/Source/Runtime/Core/Private/Misc/Compression.cpp b/Engine/Source/Runtime/Core/Private/Misc/Compression.cpp
--- a/Engine/Source/Runtime/Core/Private/Misc/Compression.cpp
+++ b/Engine/Source/Runtime/Core/Private/Misc/Compression.cpp
@@ -423,7 +423,7 @@ ICompressionFormat* FCompression::GetCompressionFormat(FName FormatName, bool bE
 		{
 			// Oodle ICompressionFormat is created on first use :
 			// inside CompressionFormatsCriticalSection lock
-			OodleDataCompressionFormatInitOnFirstUseFromLock();
+			FOodleDataCompression::CompressionFormatInitOnFirstUseFromLock();
 
 			// OodleDataCompressionFormatInitOnFirstUseFromLock added it to the ModularFeatures list
 		}
@@ -488,7 +488,7 @@ int32 FCompression::GetMaximumCompressedSize(FName FormatName, int32 Uncompresse
 		//	avoid calling CompressMemoryBound in the Decoder because it creates an ICompressionFormat for Oodle
 		//	and initializes encoders
 
-		return OodleDataGetMaximumCompressedSize(UncompressedSize);
+		return FOodleDataCompression::GetMaximumCompressedSize(UncompressedSize);
 	}
 	else
 	{
@@ -723,7 +723,7 @@ bool FCompression::UncompressMemory(FName FormatName, void* UncompressedBuffer,
 	{
 		// hardcoded Oodle
 		// can decode Oodle data without creating Oodle ICompressionFormat
-		bUncompressSucceeded = OodleDataDecompress(UncompressedBuffer,UncompressedSize,CompressedBuffer,CompressedSize);
+		bUncompressSucceeded = FOodleDataCompression::Decompress(UncompressedBuffer,UncompressedSize,CompressedBuffer,CompressedSize);
 	}
 	else
 	{
diff --git a/Engine/Source/Runtime/Core/Public/Compression/OodleDataCompression.h b/Engine/Source/Runtime/Core/Public/Compression/OodleDataCompression.h
--- a/Engine/Source/Runtime/Core/Public/Compression/OodleDataCompression.h
+++ b/Engine/Source/Runtime/Core/Public/Compression/OodleDataCompression.h
@@ -21,10 +21,13 @@
 
 DECLARE_LOG_CATEGORY_EXTERN(OodleDataCompression, Log, All);
 
+namespace FOodleDataCompression
+{
+
 /**
- * EOodleDataCompressor : Choose the Oodle Compressor
+ * ECompressor : Choose the Oodle Compressor
  *  this mostly trades decompression speed vs compression ratio
- * encode speed is determined by EOodleDataCompressionLevel , not the compressor choice.
+ * encode speed is determined by ECompressionLevel , not the compressor choice.
  * 
  * From fastest to slowest (to decode) : Selkie, Mermaid, Kraken, Leviathan
  *  
@@ -39,7 +42,7 @@ Leviathan4    :  2.62:1 ,    961.8 dec MB/s
  */
 //UENUM() // @todo Oodle might be nice if these were UENUM but can't pull UObject from inside Core?
 // enum values should not change, they may be persisted
-enum class EOodleDataCompressor : uint8
+enum class ECompressor : uint8
 {
 	NotSet = 0,
 	Selkie = 1,
@@ -50,9 +53,9 @@ enum class EOodleDataCompressor : uint8
 
 
 /**
- * EOodleDataCompressionLevel : Choose the Oodle Compression Level
+ * ECompressionLevel : Choose the Oodle Compression Level
  *  this mostly trades encode speed vs compression ratio
- *  decode speed is determined by choice of compressor (EOodleDataCompressor)
+ *  decode speed is determined by choice of compressor (ECompressor)
  * 
  * If in doubt start with "Normal" (level 4) then move up or down from there
  * 
@@ -79,8 +82,8 @@ Kraken7       :  2.64:1 ,    3.7 enc MB/s
 
  */
 //UENUM() // @todo Oodle might be nice if these were UENUM but can't pull UObject from inside Core?
-// EOodleDataCompressionLevel must numerically match the Oodle internal enum values
-enum class EOodleDataCompressionLevel : int8
+// ECompressionLevel must numerically match the Oodle internal enum values
+enum class ECompressionLevel : int8
 {
 	HyperFast4 = -4,
 	HyperFast3 = -3,
@@ -97,7 +100,7 @@ enum class EOodleDataCompressionLevel : int8
 	Optimal4 = 8,
 };
 
-enum class EOodleDataCompressionCommonUsage : uint8
+enum class ECompressionCommonUsage : uint8
 {
 	Default = 0,
 	FastRealtimeEncode = 1,
@@ -106,16 +109,16 @@ enum class EOodleDataCompressionCommonUsage : uint8
 };
 
 /**
-* Translate legacy CompressionFlags to an EOodleDataCompressionCommonUsage
+* Translate legacy CompressionFlags to an ECompressionCommonUsage
 *   CompressionFlags is not encouraged for new code; prefer directly calling to Oodle
 *
 * @param	Flags				ECompressionFlags (BiasSpeed,BiasSize,ForPackaging)
-* @return						EOodleDataCompressionCommonUsage
+* @return						ECompressionCommonUsage
 */
-EOodleDataCompressionCommonUsage CORE_API GetCommonUsageFromLegacyCompressionFlags(ECompressionFlags Flags);
+ECompressionCommonUsage CORE_API GetCommonUsageFromLegacyCompressionFlags(ECompressionFlags Flags);
 
 /**
-* Translate OodleDataCompressionCommonUsage to a Compressor & Level selection
+* Translate CompressionCommonUsage to a Compressor & Level selection
 *   usually prefer the more expressive choice of {Compressor,Level}
 *
 * @param	Usage				Your intended compression use case
@@ -123,48 +126,48 @@ EOodleDataCompressionCommonUsage CORE_API GetCommonUsageFromLegacyCompressionFla
 * @param	OutLevel			Output reference
 * @return						
 */
-void CORE_API GetCompressorAndLevelForCommonUsage(EOodleDataCompressionCommonUsage Usage,EOodleDataCompressor & OutCompressor,EOodleDataCompressionLevel & OutLevel);
+void CORE_API GetCompressorAndLevelForCommonUsage(ECompressionCommonUsage Usage,ECompressor & OutCompressor,ECompressionLevel & OutLevel);
 
 
 /**
-* What size of compressed output buffer is needed to encode into for OodleDataCompress()
+* What size of compressed output buffer is needed to encode into for Compress()
 *
 * @param	UncompressedSize			Length of uncompressed data that will fit in this output buffer
 * @return								Minimum size to allocate compressed output buffer
 */
-int64 CORE_API OodleDataCompressedBufferSizeNeeded(int64 UncompressedSize);
+int64 CORE_API CompressedBufferSizeNeeded(int64 UncompressedSize);
 
 /**
 * What is the maximum size of compressed data made after encoding
 *  For pre-allocating buffers at decode time or to check the length of compressed data at load time
-*  NOTE : OodleDataCompressedBufferSizeNeeded() >= OodleDataGetMaximumCompressedSize()
-* this is not the size to allocate compressed buffers at encode time; see OodleDataCompressedBufferSizeNeeded
+*  NOTE : CompressedBufferSizeNeeded() >= GetMaximumCompressedSize()
+* this is not the size to allocate compressed buffers at encode time; see CompressedBufferSizeNeeded
 *
 * @param	UncompressedSize			Length of uncompressed data that could have been encoded
-* @return								Maximum size of compressed data that OodleDataCompress() could have made
+* @return								Maximum size of compressed data that Compress() could have made
 */
-int64 CORE_API OodleDataGetMaximumCompressedSize(int64 UncompressedSize);
+int64 CORE_API GetMaximumCompressedSize(int64 UncompressedSize);
 
 /**
 * Encode provided data with chosen Compressor and Level
-* CompressedBufferSize must be >= OodleDataCompressedBufferSizeNeeded(UncompressedSize)
+* CompressedBufferSize must be >= CompressedBufferSizeNeeded(UncompressedSize)
 *
 * @param	OutCompressedData			output buffer where compressed data is written
 * @param	CompressedBufferSize		bytes available to write in OutCompressedData
 * @param	InUncompressedData			input buffer containing data to compress
 * @param	UncompressedSize			number of bytes in InUncompressedData to read
-* @param	Compressor					EOodleDataCompressor to encode with (this is saved in the stream)
-* @param	Level						EOodleDataCompressionLevel to encode with (this is not saved in the stream)
+* @param	Compressor					ECompressor to encode with (this is saved in the stream)
+* @param	Level						ECompressionLevel to encode with (this is not saved in the stream)
 * @return								Compressed size written or zero for failure
 */
-int64 CORE_API OodleDataCompress(
+int64 CORE_API Compress(
 							void * OutCompressedData, int64 CompressedBufferSize,
 							const void * InUncompressedData, int64 UncompressedSize,
-							EOodleDataCompressor Compressor,
-							EOodleDataCompressionLevel Level);
+							ECompressor Compressor,
+							ECompressionLevel Level);
 
 /**
-* Decode compressed data that was made by OodleDataCompress
+* Decode compressed data that was made by Compress
 *
 * UncompressedSize must match exactly the uncompressed size that was used at encode time.  No partial decodes.
 *
@@ -174,7 +177,7 @@ int64 CORE_API OodleDataCompress(
 * @param	CompressedSize			size of the input buffer, must be greater or equal to the number of compressed bytes needed
 * @return							boolean success
 */
-bool CORE_API OodleDataDecompress(
+bool CORE_API Decompress(
 						void * OutUncompressedData, int64 UncompressedSize,
 						const void * InCompressedData, int64 CompressedSize
 						);
@@ -183,7 +186,9 @@ bool CORE_API OodleDataDecompress(
 // For Compression to/from TArray and such higher level actions use OodleDataCompressionUtil.h
 
 // from Compression.cpp :
-void CORE_API OodleDataCompressionFormatInitOnFirstUseFromLock();
+void CORE_API CompressionFormatInitOnFirstUseFromLock();
 
 // from LaunchEngineLoop :
-void CORE_API OodleDataCompressionStartupPreInit();
+void CORE_API StartupPreInit();
+
+};
diff --git a/Engine/Source/Runtime/Core/Public/Compression/OodleDataCompressionUtil.h b/Engine/Source/Runtime/Core/Public/Compression/OodleDataCompressionUtil.h
--- a/Engine/Source/Runtime/Core/Public/Compression/OodleDataCompressionUtil.h
+++ b/Engine/Source/Runtime/Core/Public/Compression/OodleDataCompressionUtil.h
@@ -6,7 +6,7 @@
 
 // OodleDataCompressionUtil : utilities for common Unreal actions built on top of OodleDataCompression
 
-namespace OodleDataCompressionUtil
+namespace FOodleDataCompressionUtil
 {
 	/**
 	* Decompress replay data using Oodle, for use by INetworkReplayStreamer streamers
diff --git a/Engine/Source/Runtime/Launch/Private/LaunchEngineLoop.cpp b/Engine/Source/Runtime/Launch/Private/LaunchEngineLoop.cpp
--- a/Engine/Source/Runtime/Launch/Private/LaunchEngineLoop.cpp
+++ b/Engine/Source/Runtime/Launch/Private/LaunchEngineLoop.cpp
@@ -2382,7 +2382,7 @@ int32 FEngineLoop::PreInitPreStartupScreen(const TCHAR* CmdLine)
 #endif
 
 	// init Oodle here
-	OodleDataCompressionStartupPreInit();
+	FOodleDataCompression::StartupPreInit();
 
 	{
 		SCOPED_BOOT_TIMING("LoadPreInitModules");
