CompressedBuffer: Added Oodle as a compression method

#rnx

diff --git a/Engine/Source/Runtime/Core/Private/Compression/CompressedBuffer.cpp b/Engine/Source/Runtime/Core/Private/Compression/CompressedBuffer.cpp
--- a/Engine/Source/Runtime/Core/Private/Compression/CompressedBuffer.cpp
+++ b/Engine/Source/Runtime/Core/Private/Compression/CompressedBuffer.cpp
@@ -3,6 +3,7 @@
 #include "Compression/CompressedBuffer.h"
 
 #include "Algo/ForEach.h"
+#include "Compression/OodleDataCompression.h"
 #include "Hash/Blake3.h"
 #include "Math/UnrealMathUtility.h"
 #include "Misc/ByteSwap.h"
@@ -30,6 +31,8 @@ enum class EMethod : uint8
 	/** Header is followed by one uncompressed block. */
 	None = 0,
 	/** Header is followed by an array of compressed block sizes then the compressed blocks. */
+	Oodle = 3,
+	/** Header is followed by an array of compressed block sizes then the compressed blocks. */
 	LZ4 = 4,
 };
 
@@ -46,8 +49,10 @@ struct FHeader
 	uint32 Crc32 = 0;
 	/** The method used to compress the buffer. Affects layout of data following the header. */
 	EMethod Method = EMethod::None;
-	/** The reserved bytes must be initialized to zero. */
-	uint8 Reserved[2]{};
+	/** The method-specific compressor used to compress the buffer. */
+	uint8 Compressor = 0;
+	/** The method-specific compression level used to compress the buffer. */
+	uint8 CompressionLevel = 0;
 	/** The power of two size of every uncompressed block except the last. Size is 1 << BlockSizeExponent. */
 	uint8 BlockSizeExponent = 0;
 	/** The number of blocks that follow the header. */
@@ -119,10 +124,15 @@ static_assert(sizeof(FHeader) == 64, "FHeader is the wrong size.");
 
 ///////////////////////////////////////////////////////////////////////////////////////////////////////////////////////
 
-class FMethod
+class FEncoder
+{
+public:
+	virtual FCompositeBuffer Compress(const FCompositeBuffer& RawData, uint64 BlockSize) const = 0;
+};
+
+class FDecoder
 {
 public:
-	virtual FCompositeBuffer Compress(const FCompositeBuffer& RawData, uint64 BlockSize = DefaultBlockSize) const = 0;
 	virtual FCompositeBuffer Decompress(const FHeader& Header, const FCompositeBuffer& CompressedData) const = 0;
 	virtual bool TryDecompressTo(const FHeader& Header, const FCompositeBuffer& CompressedData, FMutableMemoryView RawView) const = 0;
 	virtual uint64 GetHeaderSize(const FHeader& Header) const = 0;
@@ -130,10 +140,10 @@ public:
 
 ///////////////////////////////////////////////////////////////////////////////////////////////////////////////////////
 
-class FMethodNone final : public FMethod
+class FNoneEncoder final : public FEncoder
 {
 public:
-	FCompositeBuffer Compress(const FCompositeBuffer& RawData, uint64 BlockSize = DefaultBlockSize) const final
+	FCompositeBuffer Compress(const FCompositeBuffer& RawData, uint64 BlockSize) const final
 	{
 		FHeader Header;
 		Header.Method = EMethod::None;
@@ -146,7 +156,11 @@ public:
 		Header.Write(HeaderData);
 		return FCompositeBuffer(HeaderData.MoveToShared(), RawData.MakeOwned());
 	}
+};
 
+class FNoneDecoder final : public FDecoder
+{
+public:
 	FCompositeBuffer Decompress(const FHeader& Header, const FCompositeBuffer& CompressedData) const final
 	{
 		if (Header.Method == EMethod::None &&
@@ -179,23 +193,17 @@ public:
 
 ///////////////////////////////////////////////////////////////////////////////////////////////////////////////////////
 
-class FMethodBlock : public FMethod
+class FBlockEncoder : public FEncoder
 {
 public:
-	FCompositeBuffer Compress(const FCompositeBuffer& RawData, uint64 BlockSize = DefaultBlockSize) const final;
-	FCompositeBuffer Decompress(const FHeader& Header, const FCompositeBuffer& CompressedData) const final;
-	bool TryDecompressTo(const FHeader& Header, const FCompositeBuffer& CompressedData, FMutableMemoryView RawView) const final;
-
-	uint64 GetHeaderSize(const FHeader& Header) const final
-	{
-		return sizeof(FHeader) + sizeof(uint32) * uint64(Header.BlockCount);
-	}
+	FCompositeBuffer Compress(const FCompositeBuffer& RawData, uint64 BlockSize) const final;
 
 protected:
 	virtual EMethod GetMethod() const = 0;
+	virtual uint8 GetCompressor() const = 0;
+	virtual uint8 GetCompressionLevel() const = 0;
 	virtual uint64 CompressBlockBound(uint64 RawSize) const = 0;
 	virtual bool CompressBlock(FMutableMemoryView& CompressedData, FMemoryView RawData) const = 0;
-	virtual bool DecompressBlock(FMutableMemoryView RawData, FMemoryView CompressedData) const = 0;
 
 private:
 	uint64 GetCompressedBlocksBound(uint64 BlockCount, uint64 BlockSize, uint64 RawSize) const
@@ -209,7 +217,7 @@ private:
 	}
 };
 
-FCompositeBuffer FMethodBlock::Compress(const FCompositeBuffer& RawData, const uint64 BlockSize) const
+FCompositeBuffer FBlockEncoder::Compress(const FCompositeBuffer& RawData, const uint64 BlockSize) const
 {
 	checkf(FMath::IsPowerOfTwo(BlockSize) && BlockSize <= MAX_uint32,
 		TEXT("BlockSize must be a 32-bit power of two but was %" UINT64_FMT "."), BlockSize);
@@ -265,7 +273,7 @@ FCompositeBuffer FMethodBlock::Compress(const FCompositeBuffer& RawData, const u
 	if (RawSize <= MetaSize + CompressedSize)
 	{
 		CompressedData.Reset();
-		return FMethodNone().Compress(RawData, BlockSize);
+		return FNoneEncoder().Compress(RawData, BlockSize);
 	}
 
 	// Write the header and calculate the CRC-32.
@@ -274,6 +282,8 @@ FCompositeBuffer FMethodBlock::Compress(const FCompositeBuffer& RawData, const u
 
 	FHeader Header;
 	Header.Method = GetMethod();
+	Header.Compressor = GetCompressor();
+	Header.CompressionLevel = GetCompressionLevel();
 	Header.BlockSizeExponent = static_cast<uint8>(FMath::FloorLog2_64(BlockSize));
 	Header.BlockCount = static_cast<uint32>(BlockCount);
 	Header.TotalRawSize = RawSize;
@@ -285,7 +295,22 @@ FCompositeBuffer FMethodBlock::Compress(const FCompositeBuffer& RawData, const u
 	return FCompositeBuffer(FSharedBuffer::MakeView(CompositeView, CompressedData.MoveToShared()));
 }
 
-FCompositeBuffer FMethodBlock::Decompress(const FHeader& Header, const FCompositeBuffer& CompressedData) const
+class FBlockDecoder : public FDecoder
+{
+public:
+	FCompositeBuffer Decompress(const FHeader& Header, const FCompositeBuffer& CompressedData) const final;
+	bool TryDecompressTo(const FHeader& Header, const FCompositeBuffer& CompressedData, FMutableMemoryView RawView) const final;
+
+	uint64 GetHeaderSize(const FHeader& Header) const final
+	{
+		return sizeof(FHeader) + sizeof(uint32) * uint64(Header.BlockCount);
+	}
+
+protected:
+	virtual bool DecompressBlock(FMutableMemoryView RawData, FMemoryView CompressedData) const = 0;
+};
+
+FCompositeBuffer FBlockDecoder::Decompress(const FHeader& Header, const FCompositeBuffer& CompressedData) const
 {
 	if (Header.BlockCount == 0 ||
 		Header.TotalCompressedSize != CompressedData.GetSize())
@@ -398,7 +423,7 @@ FCompositeBuffer FMethodBlock::Decompress(const FHeader& Header, const FComposit
 	return FCompositeBuffer(MoveTemp(Segments));
 }
 
-bool FMethodBlock::TryDecompressTo(const FHeader& Header, const FCompositeBuffer& CompressedData, FMutableMemoryView RawView) const
+bool FBlockDecoder::TryDecompressTo(const FHeader& Header, const FCompositeBuffer& CompressedData, FMutableMemoryView RawView) const
 {
 	if (Header.TotalRawSize != RawView.GetSize() ||
 		Header.TotalCompressedSize != CompressedData.GetSize())
@@ -448,13 +473,59 @@ bool FMethodBlock::TryDecompressTo(const FHeader& Header, const FCompositeBuffer
 
 ///////////////////////////////////////////////////////////////////////////////////////////////////////////////////////
 
-class FMethodLZ4 final : public FMethodBlock
+class FOodleEncoder final : public FBlockEncoder
+{
+public:
+	FOodleEncoder(FOodleDataCompression::ECompressor InCompressor, FOodleDataCompression::ECompressionLevel InCompressionLevel)
+		: Compressor(InCompressor)
+		, CompressionLevel(InCompressionLevel)
+	{
+	}
+
+protected:
+	EMethod GetMethod() const final { return EMethod::Oodle; }
+	uint8 GetCompressor() const final { return static_cast<uint8>(Compressor); }
+	uint8 GetCompressionLevel() const final { return static_cast<uint8>(CompressionLevel); }
+
+	uint64 CompressBlockBound(uint64 RawSize) const final
+	{
+		return static_cast<uint64>(FOodleDataCompression::CompressedBufferSizeNeeded(static_cast<int64>(RawSize)));
+	}
+
+	bool CompressBlock(FMutableMemoryView& CompressedData, FMemoryView RawData) const final
+	{
+		const int64 Size = FOodleDataCompression::Compress(
+			CompressedData.GetData(), static_cast<uint64>(CompressedData.GetSize()),
+			RawData.GetData(), static_cast<int64>(RawData.GetSize()),
+			Compressor, CompressionLevel);
+		CompressedData.LeftInline(static_cast<uint64>(Size));
+		return Size > 0;
+	}
+
+private:
+	const FOodleDataCompression::ECompressor Compressor;
+	const FOodleDataCompression::ECompressionLevel CompressionLevel;
+};
+
+class FOodleDecoder final : public FBlockDecoder
 {
 protected:
-	EMethod GetMethod() const final
+	bool DecompressBlock(FMutableMemoryView RawData, FMemoryView CompressedData) const final
 	{
-		return EMethod::LZ4;
+		return FOodleDataCompression::Decompress(
+			RawData.GetData(), static_cast<int64>(RawData.GetSize()),
+			CompressedData.GetData(), static_cast<int64>(CompressedData.GetSize()));
 	}
+};
+
+///////////////////////////////////////////////////////////////////////////////////////////////////////////////////////
+
+class FLZ4Encoder final : public FBlockEncoder
+{
+protected:
+	EMethod GetMethod() const final { return EMethod::LZ4; }
+	uint8 GetCompressor() const final { return 0; }
+	uint8 GetCompressionLevel() const final { return 0; }
 
 	uint64 CompressBlockBound(uint64 RawSize) const final
 	{
@@ -477,7 +548,11 @@ protected:
 		}
 		return false;
 	}
+};
 
+class FLZ4Decoder final : public FBlockDecoder
+{
+protected:
 	bool DecompressBlock(FMutableMemoryView RawData, FMemoryView CompressedData) const final
 	{
 		if (CompressedData.GetSize() <= MAX_int32)
@@ -495,30 +570,40 @@ protected:
 
 ///////////////////////////////////////////////////////////////////////////////////////////////////////////////////////
 
-static const FMethod* GetMethod(EMethod Method)
+static const FDecoder* GetDecoder(EMethod Method)
 {
-	static FMethodNone MethodNone;
-	static FMethodLZ4 MethodLZ4;
+	static FNoneDecoder None;
+	static FOodleDecoder Oodle;
+	static FLZ4Decoder LZ4;
 	switch (Method)
 	{
 	default:
 		return nullptr;
 	case EMethod::None:
-		return &MethodNone;
+		return &None;
+	case EMethod::Oodle:
+		return &Oodle;
 	case EMethod::LZ4:
-		return &MethodLZ4;
+		return &LZ4;
 	}
 }
 
-static const FMethod* GetMethod(FName FormatName)
+static const FEncoder* GetEncoder(FName FormatName)
 {
+	static FNoneEncoder None;
+	static FOodleEncoder Oodle(FOodleDataCompression::ECompressor::Mermaid, FOodleDataCompression::ECompressionLevel::VeryFast);
+	static FLZ4Encoder LZ4;
 	if (FormatName.IsNone())
 	{
-		return GetMethod(EMethod::None);
+		return &None;
+	}
+	if (FormatName == NAME_Oodle)
+	{
+		return &Oodle;
 	}
 	if (FormatName == NAME_LZ4 || FormatName == NAME_Default)
 	{
-		return GetMethod(EMethod::LZ4);
+		return &LZ4;
 	}
 	return nullptr;
 }
@@ -531,6 +616,8 @@ static FName GetMethodName(EMethod Method)
 		return NAME_Error;
 	case EMethod::None:
 		return NAME_None;
+	case EMethod::Oodle:
+		return NAME_Oodle;
 	case EMethod::LZ4:
 		return NAME_LZ4;
 	}
@@ -551,10 +638,10 @@ bool FHeader::IsValid(const FCompositeBuffer& CompressedData)
 		const FHeader Header = Read(CompressedData);
 		if (Header.Magic == FHeader::ExpectedMagic)
 		{
-			if (const FMethod* const Method = GetMethod(Header.Method))
+			if (const FDecoder* const Decoder = GetDecoder(Header.Method))
 			{
 				FUniqueBuffer HeaderCopy;
-				const FMemoryView HeaderView = CompressedData.ViewOrCopyRange(0, Method->GetHeaderSize(Header), HeaderCopy);
+				const FMemoryView HeaderView = CompressedData.ViewOrCopyRange(0, Decoder->GetHeaderSize(Header), HeaderCopy);
 				if (Header.Crc32 == FHeader::CalculateCrc32(HeaderView))
 				{
 					return true;
@@ -569,13 +656,57 @@ bool FHeader::IsValid(const FCompositeBuffer& CompressedData)
 
 } // UE::CompressedBuffer
 
+FCompressedBuffer FCompressedBuffer::Compress(const FCompositeBuffer& RawData)
+{
+	return Compress(RawData, FOodleDataCompression::ECompressor::Mermaid, FOodleDataCompression::ECompressionLevel::VeryFast);
+}
+
+FCompressedBuffer FCompressedBuffer::Compress(const FSharedBuffer& RawData)
+{
+	return Compress(FCompositeBuffer(RawData));
+}
+
+FCompressedBuffer FCompressedBuffer::Compress(
+	const FCompositeBuffer& RawData,
+	FOodleDataCompression::ECompressor Compressor,
+	FOodleDataCompression::ECompressionLevel CompressionLevel,
+	uint64 BlockSize)
+{
+	using namespace UE::CompressedBuffer;
+
+	if (BlockSize == 0)
+	{
+		BlockSize = DefaultBlockSize;
+	}
+
+	FCompressedBuffer Local;
+	if (CompressionLevel == FOodleDataCompression::ECompressionLevel::None)
+	{
+		Local.CompressedData = FNoneEncoder().Compress(RawData, BlockSize);
+	}
+	else
+	{
+		Local.CompressedData = FOodleEncoder(Compressor, CompressionLevel).Compress(RawData, BlockSize);
+	}
+	return Local;
+}
+
+FCompressedBuffer FCompressedBuffer::Compress(
+	const FSharedBuffer& RawData,
+	FOodleDataCompression::ECompressor Compressor,
+	FOodleDataCompression::ECompressionLevel CompressionLevel,
+	uint64 BlockSize)
+{
+	return Compress(FCompositeBuffer(RawData), Compressor, CompressionLevel, BlockSize);
+}
+
 FCompressedBuffer FCompressedBuffer::Compress(FName FormatName, const FCompositeBuffer& RawData)
 {
 	using namespace UE::CompressedBuffer;
 	FCompressedBuffer Local;
-	if (const FMethod* const Method = GetMethod(FormatName))
+	if (const FEncoder* const Encoder = GetEncoder(FormatName))
 	{
-		Local.CompressedData = Method->Compress(RawData);
+		Local.CompressedData = Encoder->Compress(RawData, DefaultBlockSize);
 	}
 	return Local;
 }
@@ -656,15 +787,39 @@ FName FCompressedBuffer::GetFormatName() const
 	return CompressedData ? GetMethodName(FHeader::Read(CompressedData).Method) : NAME_None;
 }
 
+bool FCompressedBuffer::TryGetCompressParameters(
+	FOodleDataCompression::ECompressor& OutCompressor,
+	FOodleDataCompression::ECompressionLevel& OutCompressionLevel) const
+{
+	using namespace UE::CompressedBuffer;
+	if (CompressedData)
+	{
+		switch (const FHeader Header = FHeader::Read(CompressedData); Header.Method)
+		{
+		case EMethod::None:
+			OutCompressor = FOodleDataCompression::ECompressor::NotSet;
+			OutCompressionLevel = FOodleDataCompression::ECompressionLevel::None;
+			return true;
+		case EMethod::Oodle:
+			OutCompressor = FOodleDataCompression::ECompressor(Header.Compressor);
+			OutCompressionLevel = FOodleDataCompression::ECompressionLevel(Header.CompressionLevel);
+			return true;
+		default:
+			break;
+		}
+	}
+	return false;
+}
+
 bool FCompressedBuffer::TryDecompressTo(FMutableMemoryView RawView) const
 {
 	using namespace UE::CompressedBuffer;
 	if (CompressedData)
 	{
 		const FHeader Header = FHeader::Read(CompressedData);
-		if (const FMethod* const Method = GetMethod(Header.Method))
+		if (const FDecoder* const Decoder = GetDecoder(Header.Method))
 		{
-			return Method->TryDecompressTo(Header, CompressedData, RawView);
+			return Decoder->TryDecompressTo(Header, CompressedData, RawView);
 		}
 	}
 	return false;
@@ -676,14 +831,14 @@ FSharedBuffer FCompressedBuffer::Decompress() const
 	if (CompressedData)
 	{
 		const FHeader Header = FHeader::Read(CompressedData);
-		if (const FMethod* const Method = GetMethod(Header.Method))
+		if (const FDecoder* const Decoder = GetDecoder(Header.Method))
 		{
 			if (Header.Method == EMethod::None)
 			{
-				return Method->Decompress(Header, CompressedData).Flatten();
+				return Decoder->Decompress(Header, CompressedData).Flatten();
 			}
 			FUniqueBuffer RawData = FUniqueBuffer::Alloc(Header.TotalRawSize);
-			if (Method->TryDecompressTo(Header, CompressedData, RawData))
+			if (Decoder->TryDecompressTo(Header, CompressedData, RawData))
 			{
 				return RawData.MoveToShared();
 			}
@@ -698,9 +853,9 @@ FCompositeBuffer FCompressedBuffer::DecompressToComposite() const
 	if (CompressedData)
 	{
 		const FHeader Header = FHeader::Read(CompressedData);
-		if (const FMethod* const Method = GetMethod(Header.Method))
+		if (const FDecoder* const Decoder = GetDecoder(Header.Method))
 		{
-			return Method->Decompress(Header, CompressedData);
+			return Decoder->Decompress(Header, CompressedData);
 		}
 	}
 	return FCompositeBuffer();
diff --git a/Engine/Source/Runtime/Core/Public/Compression/CompressedBuffer.h b/Engine/Source/Runtime/Core/Public/Compression/CompressedBuffer.h
--- a/Engine/Source/Runtime/Core/Public/Compression/CompressedBuffer.h
+++ b/Engine/Source/Runtime/Core/Public/Compression/CompressedBuffer.h
@@ -9,6 +9,9 @@ class FArchive;
 class FName;
 struct FBlake3Hash;
 
+namespace FOodleDataCompression { enum class ECompressionLevel : int8; }
+namespace FOodleDataCompression { enum class ECompressor : uint8; }
+
 /**
  * A compressed buffer stores compressed data in a self-contained format.
  *
@@ -18,6 +21,38 @@ struct FBlake3Hash;
 class FCompressedBuffer
 {
 public:
+	/**
+	 * Compress the buffer using a balanced level of compression.
+	 *
+	 * @return An owned compressed buffer, or null on error.
+	 */
+	[[nodiscard]] CORE_API static FCompressedBuffer Compress(const FCompositeBuffer& RawData);
+	[[nodiscard]] CORE_API static FCompressedBuffer Compress(const FSharedBuffer& RawData);
+
+	/**
+	 * Compress the buffer using the specified compressor and compression level.
+	 *
+	 * Data that does not compress will be return uncompressed, as if with level None.
+	 *
+	 * @note Using a level of None will return a buffer that references owned raw data.
+	 *
+	 * @param RawData            The raw data to be compressed.
+	 * @param Compressor         The compressor to encode with. May use NotSet if level is None.
+	 * @param CompressionLevel   The compression level to encode with.
+	 * @param BlockSize          The power-of-two block size to encode raw data in. 0 is default.
+	 * @return An owned compressed buffer, or null on error.
+	 */
+	[[nodiscard]] CORE_API static FCompressedBuffer Compress(
+		const FCompositeBuffer& RawData,
+		FOodleDataCompression::ECompressor Compressor,
+		FOodleDataCompression::ECompressionLevel CompressionLevel,
+		uint64 BlockSize = 0);
+	[[nodiscard]] CORE_API static FCompressedBuffer Compress(
+		const FSharedBuffer& RawData,
+		FOodleDataCompression::ECompressor Compressor,
+		FOodleDataCompression::ECompressionLevel CompressionLevel,
+		uint64 BlockSize = 0);
+
 	/**
 	 * Compress the buffer using the requested compression format.
 	 *
@@ -79,6 +114,19 @@ public:
 	 */
 	[[nodiscard]] CORE_API FName GetFormatName() const;
 
+	/**
+	 * Returns the compressor and compression level used by this buffer.
+	 *
+	 * The compressor and compression level may differ from those specified when creating the buffer
+	 * because an incompressible buffer is stored with no compression. Parameters cannot be accessed
+	 * if this is null or uses a method other than Oodle, in which case this returns false.
+	 *
+	 * @return True if parameters were written, otherwise false.
+	 */
+	[[nodiscard]] CORE_API bool TryGetCompressParameters(
+		FOodleDataCompression::ECompressor& OutCompressor,
+		FOodleDataCompression::ECompressionLevel& OutCompressionLevel) const;
+
 	/**
 	 * Decompress into a memory view that is exactly GetRawSize() bytes.
 	 */
