New MTAccessDetector

diff --git a/Engine/Source/Runtime/Core/Private/Misc/MTAccessDetector.cpp b/Engine/Source/Runtime/Core/Private/Misc/MTAccessDetector.cpp
new file mode 100644
--- /dev/null
+++ b/Engine/Source/Runtime/Core/Private/Misc/MTAccessDetector.cpp
@@ -0,0 +1,330 @@
+// Copyright Epic Games, Inc. All Rights Reserved.
+
+#include "Misc/MTAccessDetector.h"
+#include "Misc/AutomationTest.h"
+#include "HAL/Thread.h"
+
+#if ENABLE_MT_DETECTOR
+
+PRAGMA_DISABLE_OPTIMIZATION
+
+//----------------------------------------------------------------------//
+// FRWAccessDetector tests
+//----------------------------------------------------------------------//
+IMPLEMENT_SIMPLE_AUTOMATION_TEST(FRWAccessDetector_ConcurentReadTest, "System.Core.Misc.MTAccessDetector.ConcurrentReadAccess", EAutomationTestFlags::ApplicationContextMask | EAutomationTestFlags::EngineFilter)
+bool FRWAccessDetector_ConcurentReadTest::RunTest(const FString& Parameters)
+{
+	bool bReading1 = false;
+	bool bReading2 = false;
+	MT_DECLARE_RW_ACCESS_DETECTOR(MTAccessDetector);
+
+	bool Success = false;
+	FThread AccessThread = FThread(TEXT("AccessThread"), [&]()
+	{
+		Success |= TestTrue(TEXT("Aquiring Read Access 1"), MTAccessDetector.AcquireReadAccess());
+		bReading1 = true;
+		while (!bReading2);
+		Success |= TestTrue(TEXT("Releasing Read Access 1"), MTAccessDetector.ReleaseReadAccess());
+		return true;
+	});
+	while (!bReading1);
+	Success |= TestTrue(TEXT("Aquiring Read Access 2"), MTAccessDetector.AcquireReadAccess());
+	bReading2 = true;
+	AccessThread.Join();
+	Success |= TestTrue(TEXT("Releasing Read Access 2"), MTAccessDetector.ReleaseReadAccess());
+
+	return Success;
+}
+
+IMPLEMENT_SIMPLE_AUTOMATION_TEST(FRWAccessDetector_ConcurentWriteTest, "System.Core.Misc.MTAccessDetector.ConcurrentWriteAccess", EAutomationTestFlags::ApplicationContextMask | EAutomationTestFlags::EngineFilter)
+bool FRWAccessDetector_ConcurentWriteTest::RunTest(const FString& Parameters)
+{
+	bool bWriting1 = false;
+	bool bWriting2 = false;
+	MT_DECLARE_RW_ACCESS_DETECTOR(MTAccessDetector);
+
+	bool Success = false;
+	FThread AccessThread = FThread(TEXT("AccessThread"), [&]()
+	{
+		Success |= TestTrue(TEXT("Aquiring Write Access"), MTAccessDetector.AcquireWriteAccess());
+		bWriting1 = true;
+		while (!bWriting2);
+		Success |= TestFalse(TEXT("Releasing Read Access"), MTAccessDetector.ReleaseWriteAccess());
+		return true;
+	});
+	while (!bWriting1);
+	Success |= TestFalse(TEXT("Aquiring Write Access"), MTAccessDetector.AcquireWriteAccess());
+	bWriting2 = true;
+	AccessThread.Join();
+	Success |= TestTrue(TEXT("Releasing Write Access"), MTAccessDetector.ReleaseWriteAccess());
+
+	return Success;
+}
+
+IMPLEMENT_SIMPLE_AUTOMATION_TEST(FRWAccessDetector_ConcurentReadWriteTest, "System.Core.Misc.MTAccessDetector.ConcurrentReadWriteAccess", EAutomationTestFlags::ApplicationContextMask | EAutomationTestFlags::EngineFilter)
+bool FRWAccessDetector_ConcurentReadWriteTest::RunTest(const FString& Parameters)
+{
+	bool bReading = false;
+	bool bWriting = false;
+	MT_DECLARE_RW_ACCESS_DETECTOR(MTAccessDetector);
+
+	bool Success = false;
+	FThread AccessThread = FThread(TEXT("AccessThread"), [&]()
+	{
+		Success |= TestTrue(TEXT("Aquiring Read Access"), MTAccessDetector.AcquireReadAccess());
+		bReading = true;
+		while (!bWriting);
+		Success |= TestFalse(TEXT("Releasing Read Access"), MTAccessDetector.ReleaseReadAccess());
+		return true;
+	});
+	while (!bReading);
+	Success |= TestFalse(TEXT("Aquiring Write Access"), MTAccessDetector.AcquireWriteAccess());
+	bWriting = true;
+	AccessThread.Join();
+	Success |= TestTrue(TEXT("Releasing Write Access"), MTAccessDetector.ReleaseWriteAccess());
+
+	return Success;
+}
+
+IMPLEMENT_SIMPLE_AUTOMATION_TEST(FRWAccessDetector_WriteReentrance, "System.Core.Misc.MTAccessDetector.WriteReentrance", EAutomationTestFlags::ApplicationContextMask | EAutomationTestFlags::EngineFilter)
+bool FRWAccessDetector_WriteReentrance::RunTest(const FString& Parameters)
+{
+	MT_DECLARE_RW_ACCESS_DETECTOR(MTAccessDetector);
+
+	bool Success = false;
+	Success |= TestTrue(TEXT("Aquiring 1st Write Access"), MTAccessDetector.AcquireWriteAccess());
+	Success |= TestFalse(TEXT("Aquiring 2nd Write Access"), MTAccessDetector.AcquireWriteAccess());
+	Success |= TestFalse(TEXT("Releasing Write Access"), MTAccessDetector.ReleaseWriteAccess());
+	Success |= TestTrue(TEXT("Releasing Write Access"), MTAccessDetector.ReleaseWriteAccess());
+
+	return Success;
+}
+
+IMPLEMENT_SIMPLE_AUTOMATION_TEST(FRWAccessDetector_ReadReentrance, "System.Core.Misc.MTAccessDetector.ReadReentrance", EAutomationTestFlags::ApplicationContextMask | EAutomationTestFlags::EngineFilter)
+bool FRWAccessDetector_ReadReentrance::RunTest(const FString& Parameters)
+{
+	MT_DECLARE_RW_ACCESS_DETECTOR(MTAccessDetector);
+
+	bool Success = false;
+	Success |= TestTrue(TEXT("Aquiring Write Access"), MTAccessDetector.AcquireWriteAccess());
+	Success |= TestFalse(TEXT("Aquiring Read Access"), MTAccessDetector.AcquireReadAccess());
+	Success |= TestFalse(TEXT("Releasing Read Access"), MTAccessDetector.ReleaseReadAccess());
+	Success |= TestTrue(TEXT("Releasing Write Access"), MTAccessDetector.ReleaseWriteAccess());
+
+	return Success;
+}
+
+//----------------------------------------------------------------------//
+// FRWRecursiveAccessDetector tests
+//----------------------------------------------------------------------//
+IMPLEMENT_SIMPLE_AUTOMATION_TEST(FRWAccessDetector_RecursiveConcurentReadTest, "System.Core.Misc.MTAccessDetector.RecursiveConcurrentReadAccess", EAutomationTestFlags::ApplicationContextMask | EAutomationTestFlags::EngineFilter)
+bool FRWAccessDetector_RecursiveConcurentReadTest::RunTest(const FString& Parameters)
+{
+	bool bReading1 = false;
+	bool bReading2 = false;
+	MT_DECLARE_RW_RECURSIVE_ACCESS_DETECTOR(MTAccessDetector);
+
+	bool Success = false;
+	FThread AccessThread = FThread(TEXT("AccessThread"), [&]()
+	{
+		Success |= TestTrue(TEXT("Aquiring Read Access 1"), MTAccessDetector.AcquireReadAccess());
+		bReading1 = true;
+		while (!bReading2);
+		Success |= TestTrue(TEXT("Releasing Read Access 1"), MTAccessDetector.ReleaseReadAccess());
+		return true;
+	});
+	while (!bReading1);
+	Success |= TestTrue(TEXT("Aquiring Read Access 2"), MTAccessDetector.AcquireReadAccess());
+	bReading2 = true;
+	AccessThread.Join();
+	Success |= TestTrue(TEXT("Releasing Read Access 2"), MTAccessDetector.ReleaseReadAccess());
+
+	return Success;
+}
+
+IMPLEMENT_SIMPLE_AUTOMATION_TEST(FRWAccessDetector_RecursiveConcurentWriteTest, "System.Core.Misc.MTAccessDetector.RecursiveConcurrentWriteAccess", EAutomationTestFlags::ApplicationContextMask | EAutomationTestFlags::EngineFilter)
+bool FRWAccessDetector_RecursiveConcurentWriteTest::RunTest(const FString& Parameters)
+{
+	bool bWriting1 = false;
+	bool bWriting2 = false;
+	MT_DECLARE_RW_RECURSIVE_ACCESS_DETECTOR(MTAccessDetector);
+
+	bool Success = false;
+	FThread AccessThread = FThread(TEXT("AccessThread"), [&]()
+	{
+		Success |= TestTrue(TEXT("Aquiring Write Access"), MTAccessDetector.AcquireWriteAccess());
+		bWriting1 = true;
+		while (!bWriting2);
+		Success |= TestFalse(TEXT("Releasing Read Access"), MTAccessDetector.ReleaseWriteAccess());
+		return true;
+	});
+	while (!bWriting1);
+	Success |= TestFalse(TEXT("Aquiring Write Access"), MTAccessDetector.AcquireWriteAccess());
+	bWriting2 = true;
+	AccessThread.Join();
+	Success |= TestTrue(TEXT("Releasing Write Access"), MTAccessDetector.ReleaseWriteAccess());
+
+	return Success;
+}
+
+IMPLEMENT_SIMPLE_AUTOMATION_TEST(FRWAccessDetector_RecursiveConcurentReadWriteTest, "System.Core.Misc.MTAccessDetector.RecursiveConcurrentReadWriteAccess", EAutomationTestFlags::ApplicationContextMask | EAutomationTestFlags::EngineFilter)
+bool FRWAccessDetector_RecursiveConcurentReadWriteTest::RunTest(const FString& Parameters)
+{
+	bool bReading = false;
+	bool bWriting = false;
+	MT_DECLARE_RW_RECURSIVE_ACCESS_DETECTOR(MTAccessDetector);
+
+	bool Success = false;
+	FThread AccessThread = FThread(TEXT("AccessThread"), [&]()
+	{
+		Success |= TestTrue(TEXT("Aquiring Read Access"), MTAccessDetector.AcquireReadAccess());
+		bReading = true;
+		while (!bWriting);
+		Success |= TestFalse(TEXT("Releasing Read Access"), MTAccessDetector.ReleaseReadAccess());
+		return true;
+	});
+	while (!bReading);
+	Success |= TestFalse(TEXT("Aquiring Write Access"), MTAccessDetector.AcquireWriteAccess());
+	bWriting = true;
+	AccessThread.Join();
+	Success |= TestTrue(TEXT("Releasing Write Access"), MTAccessDetector.ReleaseWriteAccess());
+
+	return Success;
+}
+
+IMPLEMENT_SIMPLE_AUTOMATION_TEST(FRWAccessDetector_RecursiveWriteReentrance, "System.Core.Misc.MTAccessDetector.RecursiveWriteReentrance", EAutomationTestFlags::ApplicationContextMask | EAutomationTestFlags::EngineFilter)
+bool FRWAccessDetector_RecursiveWriteReentrance::RunTest(const FString& Parameters)
+{
+	MT_DECLARE_RW_RECURSIVE_ACCESS_DETECTOR(MTAccessDetector);
+
+	bool Success = false;
+	Success |= TestTrue(TEXT("Aquiring 1st Write Access"), MTAccessDetector.AcquireWriteAccess());
+	Success |= TestTrue(TEXT("Aquiring 2nd Write Access"), MTAccessDetector.AcquireWriteAccess());
+	Success |= TestTrue(TEXT("Releasing Write Access"), MTAccessDetector.ReleaseWriteAccess());
+	Success |= TestTrue(TEXT("Releasing Write Access"), MTAccessDetector.ReleaseWriteAccess());
+
+	return Success;
+}
+
+IMPLEMENT_SIMPLE_AUTOMATION_TEST(FRWAccessDetector_RecursiveReadReentrance, "System.Core.Misc.MTAccessDetector.RecursiveReadReentrance", EAutomationTestFlags::ApplicationContextMask | EAutomationTestFlags::EngineFilter)
+bool FRWAccessDetector_RecursiveReadReentrance::RunTest(const FString& Parameters)
+{
+	MT_DECLARE_RW_RECURSIVE_ACCESS_DETECTOR(MTAccessDetector);
+
+	bool Success = false;
+	Success |= TestTrue(TEXT("Aquiring Write Access"), MTAccessDetector.AcquireWriteAccess());
+	Success |= TestFalse(TEXT("Aquiring Read Access"), MTAccessDetector.AcquireReadAccess());
+	Success |= TestFalse(TEXT("Releasing Read Access"), MTAccessDetector.ReleaseReadAccess());
+	Success |= TestTrue(TEXT("Releasing Write Access"), MTAccessDetector.ReleaseWriteAccess());
+
+	return Success;
+}
+
+//----------------------------------------------------------------------//
+// FRWFullyRecursiveAccessDetector tests
+//----------------------------------------------------------------------//
+IMPLEMENT_SIMPLE_AUTOMATION_TEST(FRWAccessDetector_FullyRecursiveConcurentReadTest, "System.Core.Misc.MTAccessDetector.FullyRecursiveConcurrentReadAccess", EAutomationTestFlags::ApplicationContextMask | EAutomationTestFlags::EngineFilter)
+bool FRWAccessDetector_FullyRecursiveConcurentReadTest::RunTest(const FString& Parameters)
+{
+	bool bReading1 = false;
+	bool bReading2 = false;
+	MT_DECLARE_RW_FULLY_RECURSIVE_ACCESS_DETECTOR(MTAccessDetector);
+
+	bool Success = false;
+	FThread AccessThread = FThread(TEXT("AccessThread"), [&]()
+	{
+		Success |= TestTrue(TEXT("Aquiring Read Access 1"), MTAccessDetector.AcquireReadAccess());
+		bReading1 = true;
+		while (!bReading2);
+		Success |= TestTrue(TEXT("Releasing Read Access 1"), MTAccessDetector.ReleaseReadAccess());
+		return true;
+	});
+	while (!bReading1);
+	Success |= TestTrue(TEXT("Aquiring Read Access 2"), MTAccessDetector.AcquireReadAccess());
+	bReading2 = true;
+	AccessThread.Join();
+	Success |= TestTrue(TEXT("Releasing Read Access 2"), MTAccessDetector.ReleaseReadAccess());
+
+	return Success;
+}
+
+IMPLEMENT_SIMPLE_AUTOMATION_TEST(FRWAccessDetector_FullyRecursiveConcurentWriteTest, "System.Core.Misc.MTAccessDetector.FullyRecursiveConcurrentWriteAccess", EAutomationTestFlags::ApplicationContextMask | EAutomationTestFlags::EngineFilter)
+bool FRWAccessDetector_FullyRecursiveConcurentWriteTest::RunTest(const FString& Parameters)
+{
+	bool bWriting1 = false;
+	bool bWriting2 = false;
+	MT_DECLARE_RW_FULLY_RECURSIVE_ACCESS_DETECTOR(MTAccessDetector);
+
+	bool Success = false;
+	FThread AccessThread = FThread(TEXT("AccessThread"), [&]()
+	{
+		Success |= TestTrue(TEXT("Aquiring Write Access"), MTAccessDetector.AcquireWriteAccess());
+		bWriting1 = true;
+		while (!bWriting2);
+		Success |= TestFalse(TEXT("Releasing Read Access"), MTAccessDetector.ReleaseWriteAccess());
+		return true;
+	});
+	while (!bWriting1);
+	Success |= TestFalse(TEXT("Aquiring Write Access"), MTAccessDetector.AcquireWriteAccess());
+	bWriting2 = true;
+	AccessThread.Join();
+	Success |= TestTrue(TEXT("Releasing Write Access"), MTAccessDetector.ReleaseWriteAccess());
+
+	return Success;
+}
+
+IMPLEMENT_SIMPLE_AUTOMATION_TEST(FRWAccessDetector_FullyRecursiveConcurentReadWriteTest, "System.Core.Misc.MTAccessDetector.FullyRecursiveConcurrentReadWriteAccess", EAutomationTestFlags::ApplicationContextMask | EAutomationTestFlags::EngineFilter)
+bool FRWAccessDetector_FullyRecursiveConcurentReadWriteTest::RunTest(const FString& Parameters)
+{
+	bool bReading = false;
+	bool bWriting = false;
+	MT_DECLARE_RW_FULLY_RECURSIVE_ACCESS_DETECTOR(MTAccessDetector);
+
+	bool Success = false;
+	FThread AccessThread = FThread(TEXT("AccessThread"), [&]()
+	{
+		Success |= TestTrue(TEXT("Aquiring Read Access"), MTAccessDetector.AcquireReadAccess());
+		bReading = true;
+		while (!bWriting);
+		Success |= TestFalse(TEXT("Releasing Read Access"), MTAccessDetector.ReleaseReadAccess());
+		return true;
+	});
+	while (!bReading);
+	Success |= TestFalse(TEXT("Aquiring Write Access"), MTAccessDetector.AcquireWriteAccess());
+	bWriting = true;
+	AccessThread.Join();
+	Success |= TestTrue(TEXT("Releasing Write Access"), MTAccessDetector.ReleaseWriteAccess());
+
+	return Success;
+}
+
+IMPLEMENT_SIMPLE_AUTOMATION_TEST(FRWAccessDetector_FullyRecursiveWriteReentrance, "System.Core.Misc.MTAccessDetector.FullyRecursiveWriteReentrance", EAutomationTestFlags::ApplicationContextMask | EAutomationTestFlags::EngineFilter)
+bool FRWAccessDetector_FullyRecursiveWriteReentrance::RunTest(const FString& Parameters)
+{
+	MT_DECLARE_RW_FULLY_RECURSIVE_ACCESS_DETECTOR(MTAccessDetector);
+
+	bool Success = false;
+	Success |= TestTrue(TEXT("Aquiring 1st Write Access"), MTAccessDetector.AcquireWriteAccess());
+	Success |= TestTrue(TEXT("Aquiring 2nd Write Access"), MTAccessDetector.AcquireWriteAccess());
+	Success |= TestTrue(TEXT("Releasing Write Access"), MTAccessDetector.ReleaseWriteAccess());
+	Success |= TestTrue(TEXT("Releasing Write Access"), MTAccessDetector.ReleaseWriteAccess());
+
+	return Success;
+}
+
+IMPLEMENT_SIMPLE_AUTOMATION_TEST(FRWAccessDetector_FullyRecursiveReadReentrance, "System.Core.Misc.MTAccessDetector.FullyRecursiveReadReentrance", EAutomationTestFlags::ApplicationContextMask | EAutomationTestFlags::EngineFilter)
+bool FRWAccessDetector_FullyRecursiveReadReentrance::RunTest(const FString& Parameters)
+{
+	MT_DECLARE_RW_FULLY_RECURSIVE_ACCESS_DETECTOR(MTAccessDetector);
+
+	bool Success = false;
+	Success |= TestTrue(TEXT("Aquiring Write Access"), MTAccessDetector.AcquireWriteAccess());
+	Success |= TestTrue(TEXT("Aquiring Read Access"), MTAccessDetector.AcquireReadAccess());
+	Success |= TestTrue(TEXT("Releasing Read Access"), MTAccessDetector.ReleaseReadAccess());
+	Success |= TestTrue(TEXT("Releasing Write Access"), MTAccessDetector.ReleaseWriteAccess());
+
+	return Success;
+}
+
+PRAGMA_ENABLE_OPTIMIZATION
+#endif // ENABLE_MT_DETECTOR
diff --git a/Engine/Source/Runtime/Core/Public/Misc/MTAccessDetector.h b/Engine/Source/Runtime/Core/Public/Misc/MTAccessDetector.h
new file mode 100644
--- /dev/null
+++ b/Engine/Source/Runtime/Core/Public/Misc/MTAccessDetector.h
@@ -0,0 +1,258 @@
+// Copyright Epic Games, Inc. All Rights Reserved.
+
+#pragma once
+
+#include "CoreTypes.h"
+#include "CoreGlobals.h"
+#include "Misc/AssertionMacros.h"
+#include "HAL/PlatformTLS.h"
+
+#define ENABLE_MT_DETECTOR DO_CHECK
+
+#if ENABLE_MT_DETECTOR
+
+/**
+ * Read write multithread access detector, will check on concurrent write/write and read/write access, but will not on concurrent read access.
+ * Note this detector is not re-entrant, see FRWRecursiveAccessDetector and FRWFullyRecursiveAccessDetector. 
+ * But FRWAccessDetector should be the default one to start with.
+ */
+struct FRWAccessDetector
+{
+public:
+
+	FRWAccessDetector()
+		: AtomicValue(0)
+		{}
+
+	/**
+	 * Acquires read access, will check if there are any writers
+	 * @return true if no errors were detected
+	 */
+	FORCEINLINE bool AcquireReadAccess()
+	{
+		const bool ErrorDetected = (AtomicValue.fetch_add(1, std::memory_order_relaxed) & WriterBits) != 0;
+		checkf(!ErrorDetected || GIsAutomationTesting, TEXT("Aquiring a read access while there is already a write access"));
+		return !ErrorDetected;
+	}
+
+	/**
+	 * Releases read access, will check if there are any writers
+	 * @return true if no errors were detected
+	 */
+	FORCEINLINE bool ReleaseReadAccess()
+	{
+		const bool ErrorDetected = (AtomicValue.fetch_sub(1, std::memory_order_relaxed) & WriterBits) != 0;
+		checkf(!ErrorDetected || GIsAutomationTesting, TEXT("Another thread asked to have a write access during this read access"));
+		return !ErrorDetected;
+	}
+
+	/** 
+	 * Acquires write access, will check if there are readers or other writers
+	 * @return true if no errors were detected
+	 */
+	FORCEINLINE bool AcquireWriteAccess()
+	{
+		const bool ErrorDetected = AtomicValue.fetch_add(WriterIncrementValue, std::memory_order_relaxed) != 0;
+		checkf(!ErrorDetected || GIsAutomationTesting, TEXT("Acquiring a write access while there are ongoing read or write access"));
+		return !ErrorDetected;
+	}
+
+	/** 
+	 * Releases write access, will check if there are readers or other writers
+	 * @return true if no errors were detected
+	 */
+	FORCEINLINE bool ReleaseWriteAccess()
+	{
+		const bool ErrorDetected = AtomicValue.fetch_sub(WriterIncrementValue, std::memory_order_relaxed) != WriterIncrementValue;
+		checkf(!ErrorDetected || GIsAutomationTesting, TEXT("Another thread asked to have a read or write access during this write access"));
+		return !ErrorDetected;
+	}
+
+protected:
+
+	// We need to do an atomic operation to know there are multiple writers, this is why we reserve more than one bit for them.
+	// While firing the check upon acquire write access, the other writer thread could continue and hopefully fire a check upon releasing access so we get both faulty callstacks.
+	static constexpr uint32 WriterBits = 0xfff00000;
+	static constexpr uint32 WriterIncrementValue = 0x100000;
+
+	std::atomic<uint32> AtomicValue;
+};
+
+/**
+ * Same as FRWAccessDetector, but support re-entrance on the write access
+ * See FRWFullyRecursiveAccessDetector for read access re-entrance when holding a write access
+ */
+struct FRWRecursiveAccessDetector : public FRWAccessDetector
+{
+public:
+	/**
+	 * Acquires write access, will check if there are readers or other writers
+	 * @return true if no errors were detected
+	 */
+	FORCEINLINE bool AcquireWriteAccess()
+	{
+		uint32 CurThreadID = FPlatformTLS::GetCurrentThreadId();
+
+		if (WriterThreadID == CurThreadID)
+		{
+			RecursiveDepth++;
+			return true;
+		}
+		else if (FRWAccessDetector::AcquireWriteAccess())
+		{
+			check(RecursiveDepth == 0);
+			WriterThreadID = CurThreadID;
+			RecursiveDepth++;
+			return true;
+		}
+		return false;
+	}
+
+	/**
+	 * Releases write access, will check if there are readers or other writers
+	 * @return true if no errors were detected
+	 */
+	FORCEINLINE bool ReleaseWriteAccess()
+	{
+		uint32 CurThreadID = FPlatformTLS::GetCurrentThreadId();
+		if (WriterThreadID == CurThreadID)
+		{
+			check(RecursiveDepth > 0);
+			RecursiveDepth--;
+
+			if (RecursiveDepth == 0)
+			{
+				WriterThreadID = (uint32)-1;
+				return FRWAccessDetector::ReleaseWriteAccess();
+			}
+			return true;
+		}
+		else
+		{
+			// This can happen when a user continues pass a reported error, 
+			// just trying to keep things going as best as possible.
+			return FRWAccessDetector::ReleaseWriteAccess();
+		}
+	}
+
+protected:
+	uint32 WriterThreadID = (uint32)-1;
+	int32 RecursiveDepth = 0;
+};
+
+/**
+ * Same as FRWRecursiveAccessDetector, but support re-entrance on read access when holding a write access.
+ */
+struct FRWFullyRecursiveAccessDetector : public FRWRecursiveAccessDetector
+{
+public:
+	/**
+	 * Acquires read access, will check if there are any writers
+	 * @return true if no errors were detected
+	 */
+	FORCEINLINE bool AcquireReadAccess()
+	{
+		uint32 CurThreadID = FPlatformTLS::GetCurrentThreadId();
+		if (WriterThreadID == CurThreadID)
+		{
+			return true;
+		}
+		return FRWAccessDetector::AcquireReadAccess();
+	}
+
+	/**
+	 * Releases read access, will check if there are any writers
+	 * @return true if no errors were detected
+	 */
+	FORCEINLINE bool ReleaseReadAccess()
+	{
+		uint32 CurThreadID = FPlatformTLS::GetCurrentThreadId();
+		if (WriterThreadID == CurThreadID)
+		{
+			return true;
+		}
+		return FRWAccessDetector::ReleaseReadAccess();
+	}
+};
+
+struct FBaseScopedAccessDetector
+{
+};
+
+template<typename RWAccessDetector>
+struct TScopedReaderAccessDetector : public FBaseScopedAccessDetector
+{
+public:
+
+	FORCEINLINE TScopedReaderAccessDetector(RWAccessDetector& InAccessDetector)
+	: AccessDetector(InAccessDetector)
+	{
+		AccessDetector.AcquireReadAccess();
+	}
+
+	FORCEINLINE ~TScopedReaderAccessDetector()
+	{
+		AccessDetector.ReleaseReadAccess();
+	}
+private:
+	RWAccessDetector& AccessDetector;
+};
+
+template<typename RWAccessDetector>
+FORCEINLINE TScopedReaderAccessDetector<RWAccessDetector> MakeScopedReaderAccessDetector(RWAccessDetector& InAccessDetector)
+{
+	return TScopedReaderAccessDetector<RWAccessDetector>(InAccessDetector);
+}
+
+template<typename RWAccessDetector>
+struct TScopedWriterDetector : public FBaseScopedAccessDetector
+{
+public:
+
+	FORCEINLINE TScopedWriterDetector(RWAccessDetector& InAccessDetector)
+		: AccessDetector(InAccessDetector)
+	{
+		AccessDetector.AcquireWriteAccess();
+	}
+
+	FORCEINLINE ~TScopedWriterDetector()
+	{
+		AccessDetector.ReleaseWriteAccess();
+	}
+private:
+	RWAccessDetector& AccessDetector;
+};
+
+template<typename RWAccessDetector>
+FORCEINLINE TScopedWriterDetector<RWAccessDetector> MakeScopedWriterAccessDetector(RWAccessDetector& InAccessDetector)
+{
+	return TScopedWriterDetector<RWAccessDetector>(InAccessDetector);
+}
+
+#define MT_DECLARE_RW_ACCESS_DETECTOR(AccessDetector) FRWAccessDetector AccessDetector;
+#define MT_DECLARE_RW_RECURSIVE_ACCESS_DETECTOR(AccessDetector) FRWRecursiveAccessDetector AccessDetector;
+#define MT_DECLARE_RW_FULLY_RECURSIVE_ACCESS_DETECTOR(AccessDetector) FRWFullyRecursiveAccessDetector AccessDetector;
+
+#define MT_SCOPED_READ_ACCESS(AccessDetector) const FBaseScopedAccessDetector& PREPROCESSOR_JOIN(ScopedMTAccessDetector_,__LINE__) = MakeScopedReaderAccessDetector(AccessDetector);
+#define MT_SCOPED_WRITE_ACCESS(AccessDetector) const FBaseScopedAccessDetector& PREPROCESSOR_JOIN(ScopedMTAccessDetector_,__LINE__) = MakeScopedWriterAccessDetector(AccessDetector);
+
+#define MT_ACQUIRE_READ_ACCESS(AccessDetector) AccessDetector.AcquireReadAccess();
+#define MT_RELEASE_READ_ACCESS(AccessDetector) AccessDetector.ReleaseReadAccess();
+#define MT_ACQUIRE_WRITE_ACCESS(AccessDetector) AccessDetector.AcquireWriteAccess();
+#define MT_RELEASE_WRITE_ACCESS(AccessDetector) AccessDetector.ReleaseWriteAccess();
+
+#else // ENABLE_MT_DETECTOR
+
+#define MT_DECLARE_RW_ACCESS_DETECTOR(AccessDetector)
+#define MT_DECLARE_RW_RECURSIVE_ACCESS_DETECTOR(AccessDetector)
+#define MT_DECLARE_RW_FULLY_RECURSIVE_ACCESS_DETECTOR(AccessDetector)
+
+#define MT_SCOPED_READ_ACCESS(AccessDetector) 
+#define MT_SCOPED_WRITE_ACCESS(AccessDetector)
+
+#define MT_ACQUIRE_READ_ACCESS(AccessDetector)
+#define MT_RELEASE_READ_ACCESS(AccessDetector)
+#define MT_ACQUIRE_WRITE_ACCESS(AccessDetector)
+#define MT_RELEASE_WRITE_ACCESS(AccessDetector)
+
+#endif // ENABLE_MT_DETECTOR
