TaskGraph: fixed Oodle texture compression on low core count machines when the new scheduler wasn't picking up all tasks (courtesy of [REDACTED_PERSON])

diff --git a/Engine/Source/Runtime/Core/Private/Async/TaskGraph.cpp b/Engine/Source/Runtime/Core/Private/Async/TaskGraph.cpp
--- a/Engine/Source/Runtime/Core/Private/Async/TaskGraph.cpp
+++ b/Engine/Source/Runtime/Core/Private/Async/TaskGraph.cpp
@@ -2141,11 +2141,11 @@ private:
 		}
 		else if (LowLevelTasks::FScheduler::Get().IsWorkerThread() || LowLevelTasks::FReserveScheduler::Get().IsWorkerThread())
 		{
-			// for latency reason, worker thread aren't tackling other tasks while waiting if we have reserve workers, 
-			// reserve worker don't have that liberty as they have no backup of their own
-			if (GDisableReserveWorkers || LowLevelTasks::FReserveScheduler::Get().IsWorkerThread())
+			// a worker thread gets blocked, involve a reserve worker to utilise an idle core
+			bool bSuccess = false; // has a reserve worker helped?
+			if (!GDisableReserveWorkers)
 			{
-				LowLevelTasks::BusyWaitUntil([Index(0), &Tasks]() mutable
+				if (LowLevelTasks::DoReserveWorkUntil([Index(0), Tasks]() mutable
 				{
 					while (Index < Tasks.Num())
 					{
@@ -2156,11 +2156,21 @@ private:
 						Index++;
 					}
 					return true;
-				});
+				}))
+				{
+					// a reserve worker was woken up to help with the tasks.
+					// stall this thread on an event while we wait
+					FScopedEvent Event;
+					TriggerEventWhenTasksComplete(Event.Get(), Tasks, CurrentThreadIfKnown);
+					bSuccess = true;
+				}
 			}
-			else
+			
+			// for latency reason, worker thread aren't tackling other tasks while waiting if we have reserve workers, 
+			// otherwise worker threads help with task execution while they wait
+			if(!bSuccess)
 			{
-				LowLevelTasks::DoReserveWorkUntil([Index(0), Tasks]() mutable
+				LowLevelTasks::BusyWaitUntil([Index(0), &Tasks]() mutable
 				{
 					while (Index < Tasks.Num())
 					{
@@ -2172,10 +2182,6 @@ private:
 					}
 					return true;
 				});
-
-				// We will just stall this thread on an event while we wait
-				FScopedEvent Event;
-				TriggerEventWhenTasksComplete(Event.Get(), Tasks, CurrentThreadIfKnown);
 			}
 		}
 		else
