Respect memory flag on stats so we can pretty print that stat as memory (in MBs).

- This makes the output of memreport much easier to read:

812.081MB  -  StaticMesh Total Memory - STAT_StaticMeshTotalMemory2 - STATGROUP_MemoryStaticMesh - STATCAT_Advanced

diff --git a/Engine/Source/Runtime/Core/Private/Stats/StatsData.cpp b/Engine/Source/Runtime/Core/Private/Stats/StatsData.cpp
--- a/Engine/Source/Runtime/Core/Private/Stats/StatsData.cpp
+++ b/Engine/Source/Runtime/Core/Private/Stats/StatsData.cpp
@@ -1844,6 +1844,10 @@ FString FStatsUtils::DebugPrint(FStatMessage const& Item)
 		{
 			Result = FString::Printf(TEXT("%.3fms"), FPlatformTime::ToMilliseconds64(Item.GetValue_int64()));
 		}
+		if (Item.NameAndInfo.GetFlag(EStatMetaFlags::IsMemory))
+		{
+			Result = FString::Printf(TEXT("%.3fMB"), Item.GetValue_int64() / 1024.f / 1024.f);
+		}
 		else
 		{
 			Result = FString::Printf(TEXT("%llu"), Item.GetValue_int64());
