Fix path to launch UnrealTraceServer

#rnx

diff --git a/Engine/Source/Runtime/Core/Private/ProfilingDebugging/TraceAuxiliary.cpp b/Engine/Source/Runtime/Core/Private/ProfilingDebugging/TraceAuxiliary.cpp
--- a/Engine/Source/Runtime/Core/Private/ProfilingDebugging/TraceAuxiliary.cpp
+++ b/Engine/Source/Runtime/Core/Private/ProfilingDebugging/TraceAuxiliary.cpp
@@ -504,7 +504,7 @@ static void LaunchUnrealTraceInternal()
 
 	FString BinPath = "\"";
 	BinPath += FPaths::EngineDir();
-	BinPath += "/Binaries/Win64/UnrealTraceServer.exe\"";
+	BinPath += "Binaries/Win64/UnrealTraceServer.exe\"";
 
 	uint32 Flags = 0x0100'0000; // CREATE_BREAKAWAY_FROM_JOB
 	STARTUPINFOW StartupInfo = { sizeof(STARTUPINFOW) };
