Fix build warning.

diff --git a/Engine/Source/Runtime/Core/Private/Stats/StatsData.cpp b/Engine/Source/Runtime/Core/Private/Stats/StatsData.cpp
--- a/Engine/Source/Runtime/Core/Private/Stats/StatsData.cpp
+++ b/Engine/Source/Runtime/Core/Private/Stats/StatsData.cpp
@@ -1846,7 +1846,7 @@ FString FStatsUtils::DebugPrint(FStatMessage const& Item)
 		}
 		if (Item.NameAndInfo.GetFlag(EStatMetaFlags::IsMemory))
 		{
-			Result = FString::Printf(TEXT("%.3fMB"), Item.GetValue_int64() / 1024.f / 1024.f);
+			Result = FString::Printf(TEXT("%.3fMB"), (double)Item.GetValue_int64() / 1024.0 / 1024.0);
 		}
 		else
 		{
