Add a network version entry and compatibility check for the conversion of FHitResult::Actor to HitObjectHandle. This was converted back in CL 14369221, but without a new network version.
Adding this now will cause serialization failures in replays recorded between 14369221 and this change, but it fixes replays that were recorded before 14369221, including in UE4, as hit results are commonly used in network serialization of the ability system and various RPCs.

diff --git a/Engine/Source/Runtime/Core/Public/Misc/NetworkVersion.h b/Engine/Source/Runtime/Core/Public/Misc/NetworkVersion.h
--- a/Engine/Source/Runtime/Core/Public/Misc/NetworkVersion.h
+++ b/Engine/Source/Runtime/Core/Public/Misc/NetworkVersion.h
@@ -50,6 +50,7 @@ enum EEngineNetworkVersionHistory
 	HISTORY_REPLAY_DORMANCY = 16,					// Bump version to support dormancy properly in replays
 	HISTORY_ENUM_SERIALIZATION_COMPAT = 17,			// Bump version to include enum bits required for serialization into compat checksums, as well as unify enum and byte property enum serialization
 	HISTORY_SUBOBJECT_OUTER_CHAIN = 18,				// Bump version to support subobject outer chains matching on client and server
+	HISTORY_HITRESULT_INSTANCEHANDLE = 19,			// Bump version to support FHitResult change of Actor to HitObjectHandle. This change was made in CL 14369221 but a net version wasn't added at the time.
 	// New history items go above here.
 
 	HISTORY_ENGINENETVERSION_PLUS_ONE,
diff --git a/Engine/Source/Runtime/Engine/Private/Collision/Collision.cpp b/Engine/Source/Runtime/Engine/Private/Collision/Collision.cpp
--- a/Engine/Source/Runtime/Engine/Private/Collision/Collision.cpp
+++ b/Engine/Source/Runtime/Engine/Private/Collision/Collision.cpp
@@ -18,6 +18,7 @@
 #include "Engine/StaticMesh.h"
 #include "Engine/SkeletalMesh.h"
 #include "Engine/CollisionProfile.h"
+#include "Misc/NetworkVersion.h"
 #include "BodySetupEnums.h"
 #include "PhysicsEngine/BodySetup.h"
 #include "PhysicsEngine/PhysicsSettings.h"
@@ -130,7 +131,17 @@ bool FHitResult::NetSerialize(FArchive& Ar, class UPackageMap* Map, bool& bOutSu
 	}
 
 	Ar << PhysMaterial;
-	Ar << HitObjectHandle;
+
+	if (Ar.IsLoading() && Ar.EngineNetVer() < HISTORY_HITRESULT_INSTANCEHANDLE)
+	{
+		AActor* HitActor = nullptr;
+		Ar << HitActor;
+		HitObjectHandle = HitActor;
+	}
+	else
+	{
+		Ar << HitObjectHandle;
+	}
 	Ar << Component;
 	Ar << BoneName;
 	if (!bInvalidFaceIndex)
