Fix CIS complaining about an empty while loop in my unittests

#rnx

diff --git a/Engine/Source/Runtime/Core/Private/Misc/MTAccessDetector.cpp b/Engine/Source/Runtime/Core/Private/Misc/MTAccessDetector.cpp
--- a/Engine/Source/Runtime/Core/Private/Misc/MTAccessDetector.cpp
+++ b/Engine/Source/Runtime/Core/Private/Misc/MTAccessDetector.cpp
@@ -3,6 +3,7 @@
 #include "Misc/MTAccessDetector.h"
 #include "Misc/AutomationTest.h"
 #include "HAL/Thread.h"
+#include "GenericPlatform/GenericPlatformProcess.h"
 
 #if ENABLE_MT_DETECTOR
 
@@ -23,11 +24,11 @@ bool FRWAccessDetector_ConcurentReadTest::RunTest(const FString& Parameters)
 	{
 		Success |= TestTrue(TEXT("Aquiring Read Access 1"), MTAccessDetector.AcquireReadAccess());
 		bReading1 = true;
-		while (!bReading2);
+		while (!bReading2) { FPlatformProcess::SleepNoStats(0.01f); }
 		Success |= TestTrue(TEXT("Releasing Read Access 1"), MTAccessDetector.ReleaseReadAccess());
 		return true;
 	});
-	while (!bReading1);
+	while (!bReading1) { FPlatformProcess::SleepNoStats(0.01f); }
 	Success |= TestTrue(TEXT("Aquiring Read Access 2"), MTAccessDetector.AcquireReadAccess());
 	bReading2 = true;
 	AccessThread.Join();
@@ -48,11 +49,11 @@ bool FRWAccessDetector_ConcurentWriteTest::RunTest(const FString& Parameters)
 	{
 		Success |= TestTrue(TEXT("Aquiring Write Access"), MTAccessDetector.AcquireWriteAccess());
 		bWriting1 = true;
-		while (!bWriting2);
+		while (!bWriting2) { FPlatformProcess::SleepNoStats(0.01f); }
 		Success |= TestFalse(TEXT("Releasing Read Access"), MTAccessDetector.ReleaseWriteAccess());
 		return true;
 	});
-	while (!bWriting1);
+	while (!bWriting1) { FPlatformProcess::SleepNoStats(0.01f); }
 	Success |= TestFalse(TEXT("Aquiring Write Access"), MTAccessDetector.AcquireWriteAccess());
 	bWriting2 = true;
 	AccessThread.Join();
@@ -73,11 +74,11 @@ bool FRWAccessDetector_ConcurentReadWriteTest::RunTest(const FString& Parameters
 	{
 		Success |= TestTrue(TEXT("Aquiring Read Access"), MTAccessDetector.AcquireReadAccess());
 		bReading = true;
-		while (!bWriting);
+		while (!bWriting) { FPlatformProcess::SleepNoStats(0.01f); }
 		Success |= TestFalse(TEXT("Releasing Read Access"), MTAccessDetector.ReleaseReadAccess());
 		return true;
 	});
-	while (!bReading);
+	while (!bReading) { FPlatformProcess::SleepNoStats(0.01f); }
 	Success |= TestFalse(TEXT("Aquiring Write Access"), MTAccessDetector.AcquireWriteAccess());
 	bWriting = true;
 	AccessThread.Join();
@@ -129,11 +130,11 @@ bool FRWAccessDetector_RecursiveConcurentReadTest::RunTest(const FString& Parame
 	{
 		Success |= TestTrue(TEXT("Aquiring Read Access 1"), MTAccessDetector.AcquireReadAccess());
 		bReading1 = true;
-		while (!bReading2);
+		while (!bReading2) { FPlatformProcess::SleepNoStats(0.01f); }
 		Success |= TestTrue(TEXT("Releasing Read Access 1"), MTAccessDetector.ReleaseReadAccess());
 		return true;
 	});
-	while (!bReading1);
+	while (!bReading1) { FPlatformProcess::SleepNoStats(0.01f); }
 	Success |= TestTrue(TEXT("Aquiring Read Access 2"), MTAccessDetector.AcquireReadAccess());
 	bReading2 = true;
 	AccessThread.Join();
@@ -154,11 +155,11 @@ bool FRWAccessDetector_RecursiveConcurentWriteTest::RunTest(const FString& Param
 	{
 		Success |= TestTrue(TEXT("Aquiring Write Access"), MTAccessDetector.AcquireWriteAccess());
 		bWriting1 = true;
-		while (!bWriting2);
+		while (!bWriting2) { FPlatformProcess::SleepNoStats(0.01f); }
 		Success |= TestFalse(TEXT("Releasing Read Access"), MTAccessDetector.ReleaseWriteAccess());
 		return true;
 	});
-	while (!bWriting1);
+	while (!bWriting1) { FPlatformProcess::SleepNoStats(0.01f); }
 	Success |= TestFalse(TEXT("Aquiring Write Access"), MTAccessDetector.AcquireWriteAccess());
 	bWriting2 = true;
 	AccessThread.Join();
@@ -179,11 +180,11 @@ bool FRWAccessDetector_RecursiveConcurentReadWriteTest::RunTest(const FString& P
 	{
 		Success |= TestTrue(TEXT("Aquiring Read Access"), MTAccessDetector.AcquireReadAccess());
 		bReading = true;
-		while (!bWriting);
+		while (!bWriting) { FPlatformProcess::SleepNoStats(0.01f); }
 		Success |= TestFalse(TEXT("Releasing Read Access"), MTAccessDetector.ReleaseReadAccess());
 		return true;
 	});
-	while (!bReading);
+	while (!bReading) { FPlatformProcess::SleepNoStats(0.01f); }
 	Success |= TestFalse(TEXT("Aquiring Write Access"), MTAccessDetector.AcquireWriteAccess());
 	bWriting = true;
 	AccessThread.Join();
@@ -235,11 +236,11 @@ bool FRWAccessDetector_FullyRecursiveConcurentReadTest::RunTest(const FString& P
 	{
 		Success |= TestTrue(TEXT("Aquiring Read Access 1"), MTAccessDetector.AcquireReadAccess());
 		bReading1 = true;
-		while (!bReading2);
+		while (!bReading2) { FPlatformProcess::SleepNoStats(0.01f); }
 		Success |= TestTrue(TEXT("Releasing Read Access 1"), MTAccessDetector.ReleaseReadAccess());
 		return true;
 	});
-	while (!bReading1);
+	while (!bReading1) { FPlatformProcess::SleepNoStats(0.01f); }
 	Success |= TestTrue(TEXT("Aquiring Read Access 2"), MTAccessDetector.AcquireReadAccess());
 	bReading2 = true;
 	AccessThread.Join();
@@ -260,11 +261,11 @@ bool FRWAccessDetector_FullyRecursiveConcurentWriteTest::RunTest(const FString&
 	{
 		Success |= TestTrue(TEXT("Aquiring Write Access"), MTAccessDetector.AcquireWriteAccess());
 		bWriting1 = true;
-		while (!bWriting2);
+		while (!bWriting2) { FPlatformProcess::SleepNoStats(0.01f); }
 		Success |= TestFalse(TEXT("Releasing Read Access"), MTAccessDetector.ReleaseWriteAccess());
 		return true;
 	});
-	while (!bWriting1);
+	while (!bWriting1) { FPlatformProcess::SleepNoStats(0.01f); }
 	Success |= TestFalse(TEXT("Aquiring Write Access"), MTAccessDetector.AcquireWriteAccess());
 	bWriting2 = true;
 	AccessThread.Join();
@@ -285,11 +286,11 @@ bool FRWAccessDetector_FullyRecursiveConcurentReadWriteTest::RunTest(const FStri
 	{
 		Success |= TestTrue(TEXT("Aquiring Read Access"), MTAccessDetector.AcquireReadAccess());
 		bReading = true;
-		while (!bWriting);
+		while (!bWriting) { FPlatformProcess::SleepNoStats(0.01f); }
 		Success |= TestFalse(TEXT("Releasing Read Access"), MTAccessDetector.ReleaseReadAccess());
 		return true;
 	});
-	while (!bReading);
+	while (!bReading) { FPlatformProcess::SleepNoStats(0.01f); }
 	Success |= TestFalse(TEXT("Aquiring Write Access"), MTAccessDetector.AcquireWriteAccess());
 	bWriting = true;
 	AccessThread.Join();
