Support MT detector inside const methods

#rnx

diff --git a/Engine/Source/Runtime/Core/Public/Misc/MTAccessDetector.h b/Engine/Source/Runtime/Core/Public/Misc/MTAccessDetector.h
--- a/Engine/Source/Runtime/Core/Public/Misc/MTAccessDetector.h
+++ b/Engine/Source/Runtime/Core/Public/Misc/MTAccessDetector.h
@@ -24,11 +24,16 @@ public:
 		: AtomicValue(0)
 		{}
 
+	~FRWAccessDetector()
+	{
+		checkf(AtomicValue == 0, TEXT("Detector cannot be destroyed while other threads has not release all access"))
+	}
+
 	/**
 	 * Acquires read access, will check if there are any writers
 	 * @return true if no errors were detected
 	 */
-	FORCEINLINE bool AcquireReadAccess()
+	FORCEINLINE bool AcquireReadAccess() const
 	{
 		const bool ErrorDetected = (AtomicValue.fetch_add(1, std::memory_order_relaxed) & WriterBits) != 0;
 		checkf(!ErrorDetected || GIsAutomationTesting, TEXT("Aquiring a read access while there is already a write access"));
@@ -39,7 +44,7 @@ public:
 	 * Releases read access, will check if there are any writers
 	 * @return true if no errors were detected
 	 */
-	FORCEINLINE bool ReleaseReadAccess()
+	FORCEINLINE bool ReleaseReadAccess() const
 	{
 		const bool ErrorDetected = (AtomicValue.fetch_sub(1, std::memory_order_relaxed) & WriterBits) != 0;
 		checkf(!ErrorDetected || GIsAutomationTesting, TEXT("Another thread asked to have a write access during this read access"));
@@ -50,7 +55,7 @@ public:
 	 * Acquires write access, will check if there are readers or other writers
 	 * @return true if no errors were detected
 	 */
-	FORCEINLINE bool AcquireWriteAccess()
+	FORCEINLINE bool AcquireWriteAccess() const
 	{
 		const bool ErrorDetected = AtomicValue.fetch_add(WriterIncrementValue, std::memory_order_relaxed) != 0;
 		checkf(!ErrorDetected || GIsAutomationTesting, TEXT("Acquiring a write access while there are ongoing read or write access"));
@@ -61,7 +66,7 @@ public:
 	 * Releases write access, will check if there are readers or other writers
 	 * @return true if no errors were detected
 	 */
-	FORCEINLINE bool ReleaseWriteAccess()
+	FORCEINLINE bool ReleaseWriteAccess() const
 	{
 		const bool ErrorDetected = AtomicValue.fetch_sub(WriterIncrementValue, std::memory_order_relaxed) != WriterIncrementValue;
 		checkf(!ErrorDetected || GIsAutomationTesting, TEXT("Another thread asked to have a read or write access during this write access"));
@@ -75,7 +80,7 @@ protected:
 	static constexpr uint32 WriterBits = 0xfff00000;
 	static constexpr uint32 WriterIncrementValue = 0x100000;
 
-	std::atomic<uint32> AtomicValue;
+	mutable std::atomic<uint32> AtomicValue;
 };
 
 /**
@@ -89,7 +94,7 @@ public:
 	 * Acquires write access, will check if there are readers or other writers
 	 * @return true if no errors were detected
 	 */
-	FORCEINLINE bool AcquireWriteAccess()
+	FORCEINLINE bool AcquireWriteAccess() const
 	{
 		uint32 CurThreadID = FPlatformTLS::GetCurrentThreadId();
 
@@ -112,7 +117,7 @@ public:
 	 * Releases write access, will check if there are readers or other writers
 	 * @return true if no errors were detected
 	 */
-	FORCEINLINE bool ReleaseWriteAccess()
+	FORCEINLINE bool ReleaseWriteAccess() const
 	{
 		uint32 CurThreadID = FPlatformTLS::GetCurrentThreadId();
 		if (WriterThreadID == CurThreadID)
@@ -136,8 +141,8 @@ public:
 	}
 
 protected:
-	uint32 WriterThreadID = (uint32)-1;
-	int32 RecursiveDepth = 0;
+	mutable uint32 WriterThreadID = (uint32)-1;
+	mutable int32 RecursiveDepth = 0;
 };
 
 /**
@@ -150,7 +155,7 @@ public:
 	 * Acquires read access, will check if there are any writers
 	 * @return true if no errors were detected
 	 */
-	FORCEINLINE bool AcquireReadAccess()
+	FORCEINLINE bool AcquireReadAccess() const
 	{
 		uint32 CurThreadID = FPlatformTLS::GetCurrentThreadId();
 		if (WriterThreadID == CurThreadID)
@@ -164,7 +169,7 @@ public:
 	 * Releases read access, will check if there are any writers
 	 * @return true if no errors were detected
 	 */
-	FORCEINLINE bool ReleaseReadAccess()
+	FORCEINLINE bool ReleaseReadAccess() const
 	{
 		uint32 CurThreadID = FPlatformTLS::GetCurrentThreadId();
 		if (WriterThreadID == CurThreadID)
@@ -184,7 +189,7 @@ struct TScopedReaderAccessDetector : public FBaseScopedAccessDetector
 {
 public:
 
-	FORCEINLINE TScopedReaderAccessDetector(RWAccessDetector& InAccessDetector)
+	FORCEINLINE TScopedReaderAccessDetector(const RWAccessDetector& InAccessDetector)
 	: AccessDetector(InAccessDetector)
 	{
 		AccessDetector.AcquireReadAccess();
@@ -195,7 +200,7 @@ public:
 		AccessDetector.ReleaseReadAccess();
 	}
 private:
-	RWAccessDetector& AccessDetector;
+	const RWAccessDetector& AccessDetector;
 };
 
 template<typename RWAccessDetector>
@@ -209,7 +214,7 @@ struct TScopedWriterDetector : public FBaseScopedAccessDetector
 {
 public:
 
-	FORCEINLINE TScopedWriterDetector(RWAccessDetector& InAccessDetector)
+	FORCEINLINE TScopedWriterDetector(const RWAccessDetector& InAccessDetector)
 		: AccessDetector(InAccessDetector)
 	{
 		AccessDetector.AcquireWriteAccess();
@@ -220,7 +225,7 @@ public:
 		AccessDetector.ReleaseWriteAccess();
 	}
 private:
-	RWAccessDetector& AccessDetector;
+	const RWAccessDetector& AccessDetector;
 };
 
 template<typename RWAccessDetector>
@@ -236,10 +241,10 @@ FORCEINLINE TScopedWriterDetector<RWAccessDetector> MakeScopedWriterAccessDetect
 #define MT_SCOPED_READ_ACCESS(AccessDetector) const FBaseScopedAccessDetector& PREPROCESSOR_JOIN(ScopedMTAccessDetector_,__LINE__) = MakeScopedReaderAccessDetector(AccessDetector);
 #define MT_SCOPED_WRITE_ACCESS(AccessDetector) const FBaseScopedAccessDetector& PREPROCESSOR_JOIN(ScopedMTAccessDetector_,__LINE__) = MakeScopedWriterAccessDetector(AccessDetector);
 
-#define MT_ACQUIRE_READ_ACCESS(AccessDetector) AccessDetector.AcquireReadAccess();
-#define MT_RELEASE_READ_ACCESS(AccessDetector) AccessDetector.ReleaseReadAccess();
-#define MT_ACQUIRE_WRITE_ACCESS(AccessDetector) AccessDetector.AcquireWriteAccess();
-#define MT_RELEASE_WRITE_ACCESS(AccessDetector) AccessDetector.ReleaseWriteAccess();
+#define MT_ACQUIRE_READ_ACCESS(AccessDetector) (AccessDetector).AcquireReadAccess();
+#define MT_RELEASE_READ_ACCESS(AccessDetector) (AccessDetector).ReleaseReadAccess();
+#define MT_ACQUIRE_WRITE_ACCESS(AccessDetector) (AccessDetector).AcquireWriteAccess();
+#define MT_RELEASE_WRITE_ACCESS(AccessDetector) (AccessDetector).ReleaseWriteAccess();
 
 #else // ENABLE_MT_DETECTOR
 
