Convenience class for compressing/decompressing with Oodle #rb charles.bloom #pf 60f9ddad391d460001d39660

diff --git a/Engine/Source/Runtime/Core/Private/Compression/OodleDataCompressionUtil.cpp b/Engine/Source/Runtime/Core/Private/Compression/OodleDataCompressionUtil.cpp
--- a/Engine/Source/Runtime/Core/Private/Compression/OodleDataCompressionUtil.cpp
+++ b/Engine/Source/Runtime/Core/Private/Compression/OodleDataCompressionUtil.cpp
@@ -8,6 +8,70 @@
 
 namespace FOodleDataCompressionUtil
 {
+	// Our header contains 2 int32s, the compressed and decompressed sizes.
+	static constexpr int32 CompressedTArrayHeaderSize = sizeof(int32) * 2;
+
+	bool CORE_API FCompressedArray::CompressData(FOodleDataCompression::ECompressor InCompressor, FOodleDataCompression::ECompressionLevel InLevel, const void* InData, int32 InDataSize)
+	{
+		// Size the array so that it fits our header and the memory oodle requires to do the work.
+		int32 CompSizeNeeded = TCheckValueCast<int32>(FOodleDataCompression::CompressedBufferSizeNeeded(InDataSize));
+		this->SetNum( TCheckValueCast<int32>( CompressedTArrayHeaderSize + CompSizeNeeded) );
+
+		// We compress in to the buffer past our header, then write the header below.
+		void* CompPtr = this->GetData() + CompressedTArrayHeaderSize;
+		int32 CompressedSize = TCheckValueCast<int32>(FOodleDataCompression::Compress(CompPtr, CompSizeNeeded, InData, InDataSize, InCompressor, InLevel));
+		if ( CompressedSize <= 0 )
+		{
+			// Probably a bad parameter.
+			this->Empty();
+			return false;
+		}
+		
+		int32* Sizes = (int32*)this->GetData();
+
+		Sizes[0] = InDataSize;
+		Sizes[1] = CompressedSize;
+
+#if !PLATFORM_LITTLE_ENDIAN
+		Sizes[0] = BYTESWAP_ORDER32(Sizes[0]);
+		Sizes[1] = BYTESWAP_ORDER32(Sizes[1]);
+#endif
+
+		// Trim off the end working space Oodle needed to do the compress work.
+		this->SetNum( TCheckValueCast<int32>(CompressedSize + CompressedTArrayHeaderSize) , false);
+		return true;
+	}
+
+	bool CORE_API FCompressedArray::DecompressToAllocatedBuffer(void*& OutDestinationBuffer, int32& OutDestinationBufferSize) const
+	{
+		int32 DecompressedSize, CompressedSize;
+		if (PeekSizes(CompressedSize, DecompressedSize) == false)
+		{
+			return false;
+		}
+
+		// If we have a valid header, then if we don't have the actual data, it's corrupted data.
+		check(CompressedTArrayHeaderSize + CompressedSize <= Num());
+
+		void* DestinationBuffer = FMemory::Malloc(DecompressedSize);
+		if (DestinationBuffer == 0)
+		{
+			return false;
+		}
+
+		OutDestinationBufferSize = DecompressedSize;
+		OutDestinationBuffer = DestinationBuffer;
+
+		if (FOodleDataCompression::Decompress(DestinationBuffer, DecompressedSize, GetData() + CompressedTArrayHeaderSize, CompressedSize) == false)
+		{
+			FMemory::Free(DestinationBuffer);
+			OutDestinationBuffer = 0;
+			return false;
+		}
+
+		return true;
+	}
+
 	bool CORE_API DecompressReplayData(const TArray<uint8>& InCompressed, TArray< uint8 >& OutBuffer)
 	{
 		int32 Size = 0;
diff --git a/Engine/Source/Runtime/Core/Public/Compression/OodleDataCompressionUtil.h b/Engine/Source/Runtime/Core/Public/Compression/OodleDataCompressionUtil.h
--- a/Engine/Source/Runtime/Core/Public/Compression/OodleDataCompressionUtil.h
+++ b/Engine/Source/Runtime/Core/Public/Compression/OodleDataCompressionUtil.h
@@ -8,6 +8,7 @@
 
 namespace FOodleDataCompressionUtil
 {
+
 	/**
 	* Decompress replay data using Oodle, for use by INetworkReplayStreamer streamers
 	*
@@ -25,6 +26,140 @@ namespace FOodleDataCompressionUtil
 	* @return				Whether or not compression succeeded
 	*/
 	bool CORE_API CompressReplayData(const TArray<uint8>& InBuffer, TArray< uint8 >& OutCompressed);
+
+	/**
+	* Represents a compressed buffer prefixed with the compressed and decompressed size.
+	* The compressed size is technically double-specified, but affords writing the buffer
+	* to disk without any further header.
+	* 
+	* This is distinct from FCompressedBuffer as it has less overhead, is not intended to
+	* seamlessly interface with IoStore, and only supports built-in Oodle.
+	* 
+	* **NOTE** any data that will be staged to a pak/iostore should NOT be compressed, as 
+	* that system will compress for you in a manner appropriate for the platform!
+	* 
+	* If you just want raw pointers, call OodleDataCompress/Decompress directly.
+	* 
+	*/
+	class FCompressedArray : public TArray<uint8>
+	{
+	public:
+		/**
+		* Compress an arbitrary data pointer, replacing existing data.
+		* 
+		* @param OutCompressed		The destination FCompressedArray
+		* @param InCompressor		The Oodle compressor to use. See discussion.
+		* @param InLevel			The compression level to use. See discussion.
+		* @param InData				The memory to compress.
+		* @param InDataSize			The number of bytes in the InData buffer.
+		* @return					Success or failure. Failure is usually bad parameter or out-of-memory.
+		*
+		* Oodle exposes two knobs for compression - the compression _type_ and the compression _level_. Type
+		* more or less trades _decompression_ speed for compression ratio, and level more or less trades
+		* _compression_ speed for ratio.
+		* 
+		* This makes selection very usage case specific, and it's tough to just recommend a one-size-fits-all
+		* parameter set. If you have access to Epic Slack, ask in #oodle-dev. Some more detailed explanation
+		* exists adjacent to the EOodleDataCompressor and EOodleDataCompressionLevel declarations.
+		* 
+		* For common use cases, reference GetCompressorAndLevelForCommonUsage()
+		* 
+		* **NOTE** any data that will be staged to a pak/iostore should NOT be compressed, as that system will
+		* compress for you in a manner appropriate for the platform!
+		* 
+		*/
+		bool CORE_API CompressData(FOodleDataCompression::ECompressor InCompressor, FOodleDataCompression::ECompressionLevel InLevel, const void* InData, int32 InDataSize);
+
+		/**
+		* Provides access to the compressed and decompressed sizes.
+		* 
+		* @param OutCompressedSize		The amount of compressed data (The array Num() is this plus the header size).
+		* @param OutDecompressedSize	The number of bytes decompressing the data will produce.
+		* @return						False if the FCompressedArray isn't valid.
+		*/
+		bool CORE_API PeekSizes(int32& OutCompressedSize, int32& OutDecompressedSize) const
+		{
+			if (Num() < 8)
+			{
+				return false;
+			}
+
+			const int32* Sizes = (const int32*)GetData();
+			OutDecompressedSize = Sizes[0];
+			OutCompressedSize = Sizes[1];
+
+#if !PLATFORM_LITTLE_ENDIAN
+			OutDecompressedSize = BYTESWAP_ORDER32(OutDecompressedSize);
+			OutCompressedSize = BYTESWAP_ORDER32(OutCompressedSize);
+#endif
+
+			return true;
+		}
+
+		/**
+		* Decompresses to a buffer that has already been allocated.
+		* 
+		* @param InDestinationBUffer		The buffer to contain the decompressed data. This buffer must be at
+		*									least as large as the decompressed size specified by PeekSizes. If
+		*									it's not, memory corruption will occur.
+		* @return							Success or failure. Failure is usually because the FCompressedArray doesn't
+		*									actually contain any data.
+		*/
+		bool CORE_API DecompressToExistingBuffer(void* InDestinationBuffer) const;
+
+		/**
+		* Decompresses a compressed TArray to a buffer that will be allocated by this function.
+		* 
+		* @param OutDestinationBuffer		The pointer to contain the allocated buffer that contains the decompressed
+		*									data.
+		* @return							Success or failure. Failure is usually because the FCompressedArray doesn't
+		*									actually contain any data.
+		*/
+		bool CORE_API DecompressToAllocatedBuffer(void*& OutDestinationBuffer, int32& OutDestinationBufferSize) const;
+
+
+		/**
+		* Compress an arbitrary data pointer, replacing existing data.
+		*
+		* This is just a thunk to CompressData, and exists just for type safety
+		* so everyone doesn't have to pull out the raw data from their TArrays.
+		* 
+		* See CompressData for parameter discussion.
+		*/
+		template <class T>
+		bool CompressTArray(FOodleDataCompression::ECompressor InCompressor, FOodleDataCompression::ECompressionLevel InLevel, const TArray<T>& InBuffer)
+		{
+			return CompressData(InCompressor, InLevel, (const void*)InBuffer.GetData(), InBuffer.Num() * sizeof(T));
+		}
+
+		/**
+		* Decompress to a TArray. The should be paired with CompressTArray.
+		* 
+		* @param OutDecompressed			The output TArray. Existing contents will be destroyed.
+		* @return							Success or failure. Note that if the compress/decompress TArray
+		*									types are mismatched (i.e. have different sizes) this could fail
+		*									due to granularity concerns.
+		*/
+		template <class T>
+		bool DecompressToTArray(TArray<T>& OutDecompressed) const
+		{
+			int32 DecompressedSize, CompressedSize;
+			if (PeekSizes(CompressedSize, DecompressedSize) == false)
+			{
+				return false;
+			}
+
+			if ((DecompressedSize % sizeof(T)) != 0)
+			{
+				// We must be able to evenly fit our decompressed data in to the desired output.
+				return false;
+			}
+
+			OutDecompressed.SetNum(DecompressedSize / sizeof(T));
+			return DecompressToExistingBuffer(OutDecompressed.GetData());
+		}
+
+	};
 };
 
 
