Fix for octree if LWC is enabled, fixes directional light shadows, probably lots of other misc stuff

diff --git a/Engine/Source/Runtime/Core/Public/Math/GenericOctree.h b/Engine/Source/Runtime/Core/Public/Math/GenericOctree.h
--- a/Engine/Source/Runtime/Core/Public/Math/GenericOctree.h
+++ b/Engine/Source/Runtime/Core/Public/Math/GenericOctree.h
@@ -284,7 +284,7 @@ public:
 
 	inline VectorRegister GetChildOffsetVec(int i) const
 	{
-		union MaskType { VectorRegister v;  VectorRegister4Int i; 
+		union MaskType { VectorRegister4Float v;  VectorRegister4Int i; 
 #if PLATFORM_HOLOLENS
 			MaskType() 
 				: v(MakeVectorRegister(0.0f, 0.0f, 0.0f, 0.0f))
