Add Insights bookmark when exit is requested

diff --git a/Engine/Source/Runtime/Core/Private/Misc/CoreGlobals.cpp b/Engine/Source/Runtime/Core/Private/Misc/CoreGlobals.cpp
--- a/Engine/Source/Runtime/Core/Private/Misc/CoreGlobals.cpp
+++ b/Engine/Source/Runtime/Core/Private/Misc/CoreGlobals.cpp
@@ -290,6 +290,8 @@ void CORE_API RequestEngineExit(const TCHAR* ReasonString)
 {
 	ensureMsgf(ReasonString && FCString::Strlen(ReasonString) > 4, TEXT("RequestEngineExit must be given a valid reason (reason \"%s\""), ReasonString);
 
+	TRACE_BOOKMARK(TEXT("Engine exit requested (reason: %s)"), ReasonString);
+
 	FGenericCrashContext::SetEngineExit(true);
 
 #if UE_SET_REQUEST_EXIT_ON_TICK_ONLY
