Fixed bugs on ARM with rotation and removed unnecessary implicit conversion from Float vec to Double vectors

diff --git a/Engine/Source/Runtime/Core/Private/Math/TransformVectorized.cpp b/Engine/Source/Runtime/Core/Private/Math/TransformVectorized.cpp
--- a/Engine/Source/Runtime/Core/Private/Math/TransformVectorized.cpp
+++ b/Engine/Source/Runtime/Core/Private/Math/TransformVectorized.cpp
@@ -113,8 +113,8 @@ FTransform FTransform::GetRelativeTransformReverse(const FTransform& Other) cons
 	VectorRegister VScale3D = VectorMultiply(Other.Scale3D, VSafeScale3D);
 	
 	// Rotation = Q(B) * Q(A)(-1)	
-	VectorRegister VInverseRot = VectorQuaternionInverse(Rotation);
-	VectorRegister VRotation = VectorQuaternionMultiply2(Other.Rotation, VInverseRot );
+	VectorRegister4Float VInverseRot = VectorQuaternionInverse(Rotation);
+	VectorRegister4Float VRotation = VectorQuaternionMultiply2(Other.Rotation, VInverseRot );
 	
 	// RotatedTranslation
 	VectorRegister VR = VectorQuaternionRotateVector(VRotation, Translation);
@@ -124,7 +124,7 @@ FTransform FTransform::GetRelativeTransformReverse(const FTransform& Other) cons
 
 	Result.Scale3D = VScale3D;	
 	Result.Translation = VTranslation;
-	Result.Rotation = MakeVectorRegisterFloatFromDouble(VRotation);
+	Result.Rotation = VRotation;
 		
 	Result.DiagnosticCheckNaN_All(); 
 
@@ -165,14 +165,14 @@ void FTransform::SetToRelativeTransform(const FTransform& ParentTransform)
 	VectorRegister VQTranslation = VectorSet_W0(VectorSubtract(Translation, ParentTransform.Translation));
 
 	// Inverse RotatedTranslation
-	VectorRegister VInverseParentRot = VectorQuaternionInverse(ParentTransform.Rotation);
+	VectorRegister4Float VInverseParentRot = VectorQuaternionInverse(ParentTransform.Rotation);
 	VectorRegister VR = VectorQuaternionRotateVector(VInverseParentRot, VQTranslation);
 
 	// Translation = 1/S(B)
 	Translation = VectorMultiply(VR, VSafeScale3D);
 
 	// Rotation = Q(B)(-1) * Q(A)	
-	Rotation = MakeVectorRegisterFloatFromDouble(VectorQuaternionMultiply2(VInverseParentRot, Rotation ));
+	Rotation = VectorQuaternionMultiply2(VInverseParentRot, Rotation);
 
 	DiagnosticCheckNaN_All(); 
 
@@ -227,18 +227,18 @@ FTransform FTransform::GetRelativeTransform(const FTransform& Other) const
 		VectorRegister VQTranslation = VectorSet_W0(VectorSubtract(Translation, Other.Translation));
 
 		// Inverse RotatedTranslation
-		VectorRegister VInverseRot = VectorQuaternionInverse(Other.Rotation);
+		VectorRegister4Float VInverseRot = VectorQuaternionInverse(Other.Rotation);
 		VectorRegister VR = VectorQuaternionRotateVector(VInverseRot, VQTranslation);
 
 		//Translation = 1/S(B)
 		VectorRegister VTranslation = VectorMultiply(VR, VSafeScale3D);
 
 		// Rotation = Q(B)(-1) * Q(A)	
-		VectorRegister VRotation = VectorQuaternionMultiply2(VInverseRot, Rotation);
+		VectorRegister4Float VRotation = VectorQuaternionMultiply2(VInverseRot, Rotation);
 
 		Result.Scale3D = VScale3D;
 		Result.Translation = VTranslation;
-		Result.Rotation = MakeVectorRegisterFloatFromDouble(VRotation);
+		Result.Rotation = VRotation;
 
 		Result.DiagnosticCheckNaN_All();
 #if DEBUG_INVERSE_TRANSFORM
diff --git a/Engine/Source/Runtime/Core/Public/Math/TransformVectorized.h b/Engine/Source/Runtime/Core/Public/Math/TransformVectorized.h
--- a/Engine/Source/Runtime/Core/Public/Math/TransformVectorized.h
+++ b/Engine/Source/Runtime/Core/Public/Math/TransformVectorized.h
@@ -56,7 +56,7 @@ public:
 		if (VectorContainsNaNOrInfinite(Scale3D))
 		{
 			logOrEnsureNanError(TEXT("FTransform Vectorized Scale3D contains NaN"));
-			const_cast<FTransform*>(this)->Scale3D = VectorSet_W0( VectorOne() );
+			const_cast<FTransform*>(this)->Scale3D = VectorSet_W0( VectorOneDouble() );
 		}
 	}
 
@@ -65,7 +65,7 @@ public:
 		if (VectorContainsNaNOrInfinite(Translation))
 		{
 			logOrEnsureNanError(TEXT("FTransform Vectorized Translation contains NaN"));
-			const_cast<FTransform*>(this)->Translation = VectorZero();
+			const_cast<FTransform*>(this)->Translation = VectorZeroDouble();
 		}
 	}
 
@@ -74,7 +74,7 @@ public:
 		if (VectorContainsNaNOrInfinite(Rotation))
 		{
 			logOrEnsureNanError(TEXT("FTransform Vectorized Rotation contains NaN"));
-			const_cast<FTransform*>(this)->Rotation = VectorSet_W1( VectorZero() );
+			const_cast<FTransform*>(this)->Rotation = VectorSet_W1( VectorZeroFloat() );
 		}
 	}
 
@@ -110,7 +110,7 @@ public:
 		// Rotation = {0,0,0,1)
 		Rotation = GlobalVectorConstants::Float0001; //VectorSet_W1( VectorZero() );
 		// Translation = {0,0,0,0)
-		Translation = VectorZero();
+		Translation = VectorZeroDouble();
 		// Scale3D = {1,1,1,0);
 		Scale3D = GlobalVectorConstants::Vector1110; //VectorSet_W0( VectorOne() );
 	}
@@ -142,7 +142,7 @@ public:
 		// Rotation = InRotation
 		Rotation =  VectorLoadAligned( &InRotation.X );
 		// Translation = {0,0,0,0)
-		Translation = VectorZero();
+		Translation = VectorZeroDouble();
 		// Scale3D = {1,1,1,0);
 		Scale3D = GlobalVectorConstants::Vector1110; //VectorSet_W0( VectorOne() );
 
@@ -160,7 +160,7 @@ public:
 		// Rotation = InRotation
 		Rotation =  VectorLoadAligned( &InQuatRotation.X );
 		// Translation = {0,0,0,0)
-		Translation = VectorZero();
+		Translation = VectorZeroDouble();
 		// Scale3D = {1,1,1,0);
 		Scale3D = GlobalVectorConstants::Vector1110; //VectorSet_W0( VectorOne() );
 
@@ -675,16 +675,17 @@ private:
 
 	FORCEINLINE static bool Private_AnyHasNegativeScale(const VectorRegister& InScale3D, const  VectorRegister& InOtherScale3D)
 	{
-		return !!VectorAnyLesserThan(VectorMin(InScale3D, InOtherScale3D), VectorZero());
+		return !!VectorAnyLesserThan(VectorMin(InScale3D, InOtherScale3D), VectorZeroDouble());
 	}
 
-	FORCEINLINE bool Private_RotationEquals( const VectorRegister& InRotation, const ScalarRegister& Tolerance = ScalarRegister(GlobalVectorConstants::KindaSmallNumber)) const
+	FORCEINLINE bool Private_RotationEquals( const VectorRegister4Float& InRotation, float Tolerance = KINDA_SMALL_NUMBER) const
 	{			
 		// !( (FMath::Abs(X-Q.X) > Tolerance) || (FMath::Abs(Y-Q.Y) > Tolerance) || (FMath::Abs(Z-Q.Z) > Tolerance) || (FMath::Abs(W-Q.W) > Tolerance) )
-		const VectorRegister RotationSub = VectorAbs(VectorSubtract(Rotation, InRotation));		
+		const VectorRegister4Float RotationSub = VectorAbs(VectorSubtract(Rotation, InRotation));
 		// !( (FMath::Abs(X+Q.X) > Tolerance) || (FMath::Abs(Y+Q.Y) > Tolerance) || (FMath::Abs(Z+Q.Z) > Tolerance) || (FMath::Abs(W+Q.W) > Tolerance) )
-		const VectorRegister RotationAdd = VectorAbs(VectorAdd(Rotation, InRotation));
-		return !VectorAnyGreaterThan(RotationSub, Tolerance.Value) || !VectorAnyGreaterThan(RotationAdd, Tolerance.Value);
+		const VectorRegister4Float RotationAdd = VectorAbs(VectorAdd(Rotation, InRotation));
+		const VectorRegister4Float ToleranceFloat = VectorLoadFloat1(&Tolerance);
+		return !VectorAnyGreaterThan(RotationSub, ToleranceFloat) || !VectorAnyGreaterThan(RotationAdd, ToleranceFloat);
 	}
 
 	FORCEINLINE bool Private_TranslationEquals( const VectorRegister& InTranslation, const ScalarRegister& Tolerance = ScalarRegister(GlobalVectorConstants::KindaSmallNumber)) const
@@ -706,7 +707,7 @@ public:
 	// Test if A's rotation equals B's rotation, within a tolerance. Preferred over "A.GetRotation().Equals(B.GetRotation())" because it is faster on some platforms.
 	FORCEINLINE static bool AreRotationsEqual(const FTransform& A, const FTransform& B, float Tolerance=KINDA_SMALL_NUMBER)
 	{
-		return A.Private_RotationEquals(B.Rotation, ScalarRegister(Tolerance));
+		return A.Private_RotationEquals(B.Rotation, Tolerance);
 	}
 
 	// Test if A's translation equals B's translation, within a tolerance. Preferred over "A.GetTranslation().Equals(B.GetTranslation())" because it avoids VectorRegister->FVector conversion.
@@ -744,7 +745,7 @@ public:
 	FORCEINLINE bool Equals(const FTransform& Other, float Tolerance=KINDA_SMALL_NUMBER) const
 	{
 		const ScalarRegister ToleranceRegister(Tolerance);
-		return Private_TranslationEquals(Other.Translation, ToleranceRegister) && Private_RotationEquals(Other.Rotation, ToleranceRegister) && Private_Scale3DEquals(Other.Scale3D, ToleranceRegister);
+		return Private_TranslationEquals(Other.Translation, ToleranceRegister) && Private_RotationEquals(Other.Rotation, Tolerance) && Private_Scale3DEquals(Other.Scale3D, ToleranceRegister);
 	}
 
 	// Test if all components of the transform property are equal.
@@ -757,7 +758,7 @@ public:
 	FORCEINLINE bool EqualsNoScale(const FTransform& Other, float Tolerance=KINDA_SMALL_NUMBER) const
 	{
 		const ScalarRegister ToleranceRegister(Tolerance);
-		return Private_TranslationEquals(Other.Translation, ToleranceRegister) && Private_RotationEquals(Other.Rotation, ToleranceRegister);
+		return Private_TranslationEquals(Other.Translation, ToleranceRegister) && Private_RotationEquals(Other.Rotation, Tolerance);
 	}
 
 	FORCEINLINE static void Multiply(FTransform* OutTransform, const FTransform* A, const FTransform* B);
@@ -787,7 +788,7 @@ public:
 		// Rotation = {0,0,0,1)
 		Rotation = GlobalVectorConstants::Float0001; //VectorSet_W1( VectorZero() );
 		// Translation = {0,0,0,0)
-		Translation = VectorZero();
+		Translation = VectorZeroDouble();
 		// Scale3D = {1,1,1,0);
 		Scale3D = GlobalVectorConstants::Vector1110; //VectorSet_W0( VectorOne() );
 	}
@@ -1025,14 +1026,14 @@ public:
 	*/
 	FORCEINLINE void AccumulateWithAdditiveScale(const FTransform& Atom, const ScalarRegister& BlendWeight)
 	{
-		const VectorRegister DefaultScale = MakeVectorRegister(1.f, 1.f, 1.f, 0.f);
+		const VectorRegister DefaultScale = GlobalVectorConstants::Vector1110;
 
 		// SourceAtom = Atom * BlendWeight;
 		const VectorRegister4Float BlendedRotation = VectorMultiply(Atom.Rotation, MakeVectorRegisterFloatFromDouble(BlendWeight.Value));
 		const VectorRegister BlendedScale = VectorMultiply(Atom.Scale3D, BlendWeight.Value);
 		const VectorRegister BlendedTranslation = VectorMultiply(Atom.Translation, BlendWeight.Value);
 
-		const VectorRegister RotationW = VectorReplicate(BlendedRotation, 3);
+		const VectorRegister4Float RotationW = VectorReplicate(BlendedRotation, 3);
 
 		// Add ref pose relative animation to base animation, only if rotation is significant.
 		// if( Square(SourceAtom.Rotation.W) < 1.f - DELTA * DELTA )
@@ -1085,7 +1086,7 @@ public:
 	 */
 	FORCEINLINE bool IsRotationNormalized() const
 	{		
-		const VectorRegister TestValue = VectorAbs(VectorSubtract(VectorOne(), VectorDot4(Rotation, Rotation)));
+		const VectorRegister4Float TestValue = VectorAbs(VectorSubtract(VectorOneFloat(), VectorDot4(Rotation, Rotation)));
 		return !VectorAnyGreaterThan(TestValue, GlobalVectorConstants::ThreshQuatNormalized);
 	}
 
@@ -1107,7 +1108,7 @@ public:
 		const VectorRegister4Float Const0001 = GlobalVectorConstants::Float0001;
 		const VectorRegister4Float ConstNegative0001 = VectorSubtract(VectorZeroFloat(), Const0001);
 		const VectorRegister4Float VOneMinusAlpha = VectorSubtract(VectorOneFloat(), BlendWeightFloat);
-		const VectorRegister DefaultScale = MakeVectorRegister(1.f, 1.f, 1.f, 0.f);
+		const VectorRegister DefaultScale = GlobalVectorConstants::Vector1110;
 
 		// Blend rotation
 		//     To ensure the 'shortest route', we make sure the dot product between the both rotations is positive.
@@ -1117,7 +1118,7 @@ public:
 		//  Note: A = (0,0,0,1), which simplifies things a lot; only care about sign of B.W now, instead of doing a dot product
 		const VectorRegister4Float RotationB = SourceAtom.Rotation;
 
-		const VectorRegister4Float QuatRotationDirMask = VectorCompareGE(RotationB, VectorZero());
+		const VectorRegister4Float QuatRotationDirMask = VectorCompareGE(RotationB, VectorZeroFloat());
 		const VectorRegister4Float BiasTimesA = VectorSelect(QuatRotationDirMask, Const0001, ConstNegative0001);
 		const VectorRegister4Float RotateBTimesWeight = VectorMultiply(RotationB, BlendWeightFloat);
 		const VectorRegister4Float UnnormalizedRotation = VectorMultiplyAdd(BiasTimesA, VOneMinusAlpha, RotateBTimesWeight);
@@ -1131,8 +1132,8 @@ public:
 		// Blend translation and scale
 		//    BlendedAtom.Translation = Lerp(Zero, SourceAtom.Translation, Alpha);
 		//    BlendedAtom.Scale = Lerp(0, SourceAtom.Scale, Alpha);
-		const VectorRegister BlendedTranslation	= FMath::Lerp<VectorRegister>(VectorZero(), SourceAtom.Translation, BlendWeight.Value);
-		const VectorRegister BlendedScale3D	= FMath::Lerp<VectorRegister>(VectorZero(), SourceAtom.Scale3D, BlendWeight.Value);
+		const VectorRegister BlendedTranslation	= FMath::Lerp<VectorRegister>(VectorZeroDouble(), SourceAtom.Translation, BlendWeight.Value);
+		const VectorRegister BlendedScale3D	= FMath::Lerp<VectorRegister>(VectorZeroDouble(), SourceAtom.Scale3D, BlendWeight.Value);
 
 		// Apply translation and scale to final atom
 		//     FinalAtom.Translation += BlendedAtom.Translation
@@ -1223,7 +1224,7 @@ public:
 		{
 			// Assume it is along X and modify transform accordingly. 
 			// It doesn't actually matter which axis we choose, the 'appearance' will be the same			
-			Scale3D = VectorMultiply(Scale3D, GlobalVectorConstants::FloatMinus1_111 );			
+			Scale3D = VectorMultiply(Scale3D, GlobalVectorConstants::DoubleMinus1_111 );			
 			M.SetAxis(0, -M.GetScaledAxis( EAxis::X ));
 		}
 
@@ -1244,17 +1245,17 @@ private:
 		check( IsRotationNormalized() );
 #endif		
 
-		const VectorRegister RotationX2Y2Z2 = VectorAdd(Rotation, Rotation);	// x2, y2, z2
-		const VectorRegister RotationXX2YY2ZZ2 = VectorMultiply(RotationX2Y2Z2, Rotation);	// xx2, yy2, zz2		
+		const VectorRegister4Float RotationX2Y2Z2 = VectorAdd(Rotation, Rotation);	// x2, y2, z2
+		const VectorRegister4Float RotationXX2YY2ZZ2 = VectorMultiply(RotationX2Y2Z2, Rotation);	// xx2, yy2, zz2		
 
 		// The diagonal terms of the rotation matrix are:
 		//   (1 - (yy2 + zz2)) * scale
 		//   (1 - (xx2 + zz2)) * scale
 		//   (1 - (xx2 + yy2)) * scale
-		const VectorRegister yy2_xx2_xx2 = VectorSwizzle(RotationXX2YY2ZZ2, 1, 0, 0, 0);
-		const VectorRegister zz2_zz2_yy2 = VectorSwizzle(RotationXX2YY2ZZ2, 2, 2, 1, 0);
-		const VectorRegister DiagonalSum = VectorAdd(yy2_xx2_xx2, zz2_zz2_yy2);
-		const VectorRegister Diagonals = VectorSubtract(VectorOne(), DiagonalSum);
+		const VectorRegister4Float yy2_xx2_xx2 = VectorSwizzle(RotationXX2YY2ZZ2, 1, 0, 0, 0);
+		const VectorRegister4Float zz2_zz2_yy2 = VectorSwizzle(RotationXX2YY2ZZ2, 2, 2, 1, 0);
+		const VectorRegister4Float DiagonalSum = VectorAdd(yy2_xx2_xx2, zz2_zz2_yy2);
+		const VectorRegister4Float Diagonals = VectorSubtract(VectorOneFloat(), DiagonalSum);
 		OutDiagonals = VectorMultiply(Diagonals, Scale3D);
 
 		// Grouping the non-diagonal elements in the rotation block by operations:
@@ -1265,21 +1266,21 @@ private:
 
 		// RotBase = x*y2, y*z2, x*z2
 		// RotOffset = w*z2, w*x2, w*y2
-		const VectorRegister x_y_x = VectorSwizzle(Rotation, 0, 1, 0, 0);
-		const VectorRegister y2_z2_z2 = VectorSwizzle(RotationX2Y2Z2, 1, 2, 2, 0);
-		const VectorRegister RotBase = VectorMultiply(x_y_x, y2_z2_z2);
+		const VectorRegister4Float x_y_x = VectorSwizzle(Rotation, 0, 1, 0, 0);
+		const VectorRegister4Float y2_z2_z2 = VectorSwizzle(RotationX2Y2Z2, 1, 2, 2, 0);
+		const VectorRegister4Float RotBase = VectorMultiply(x_y_x, y2_z2_z2);
 
-		const VectorRegister w_w_w = VectorReplicate(Rotation, 3);
-		const VectorRegister z2_x2_y2 = VectorSwizzle(RotationX2Y2Z2, 2, 0, 1, 0);
-		const VectorRegister RotOffset = VectorMultiply(w_w_w, z2_x2_y2);
+		const VectorRegister4Float w_w_w = VectorReplicate(Rotation, 3);
+		const VectorRegister4Float z2_x2_y2 = VectorSwizzle(RotationX2Y2Z2, 2, 0, 1, 0);
+		const VectorRegister4Float RotOffset = VectorMultiply(w_w_w, z2_x2_y2);
 
 		// Adds = (RotBase + RotOffset)*Scale3D :  (x*y2 + w*z2) * Scale3D.X , (y*z2 + w*x2) * Scale3D.Y, (x*z2 + w*y2) * Scale3D.Z
 		// Subtracts = (RotBase - RotOffset)*Scale3DYZX :  (x*y2 - w*z2) * Scale3D.Y , (y*z2 - w*x2) * Scale3D.Z, (x*z2 - w*y2) * Scale3D.X
-		const VectorRegister Adds = VectorAdd(RotBase, RotOffset);
+		const VectorRegister4Float Adds = VectorAdd(RotBase, RotOffset);
 		OutAdds = VectorMultiply(Adds, Scale3D);
 		const VectorRegister Scale3DYZXW = VectorSwizzle( Scale3D, 1, 2, 0, 3);
-		const VectorRegister Subtracts = VectorSubtract(RotBase, RotOffset);
-		OutSubtracts = VectorMultiply(Subtracts , Scale3DYZXW);
+		const VectorRegister4Float Subtracts = VectorSubtract(RotBase, RotOffset);
+		OutSubtracts = VectorMultiply(Subtracts, Scale3DYZXW);
 	}
 
 	FORCEINLINE void ToMatrixInternalNoScale( VectorRegister& OutDiagonals, VectorRegister& OutAdds, VectorRegister& OutSubtracts ) const
@@ -1288,17 +1289,17 @@ private:
 		// Make sure Rotation is normalized when we turn it into a matrix.
 		ensure( IsRotationNormalized() );
 #endif		
-		const VectorRegister RotationX2Y2Z2 = VectorAdd(Rotation, Rotation);	// x2, y2, z2
-		const VectorRegister RotationXX2YY2ZZ2 = VectorMultiply(RotationX2Y2Z2, Rotation);	// xx2, yy2, zz2		
+		const VectorRegister4Float RotationX2Y2Z2 = VectorAdd(Rotation, Rotation);	// x2, y2, z2
+		const VectorRegister4Float RotationXX2YY2ZZ2 = VectorMultiply(RotationX2Y2Z2, Rotation);	// xx2, yy2, zz2		
 
 		// The diagonal terms of the rotation matrix are:
 		//   (1 - (yy2 + zz2))
 		//   (1 - (xx2 + zz2))
 		//   (1 - (xx2 + yy2))
-		const VectorRegister yy2_xx2_xx2 = VectorSwizzle(RotationXX2YY2ZZ2, 1, 0, 0, 0);
-		const VectorRegister zz2_zz2_yy2 = VectorSwizzle(RotationXX2YY2ZZ2, 2, 2, 1, 0);
-		const VectorRegister DiagonalSum = VectorAdd(yy2_xx2_xx2, zz2_zz2_yy2);
-		OutDiagonals = VectorSubtract(VectorOne(), DiagonalSum);
+		const VectorRegister4Float yy2_xx2_xx2 = VectorSwizzle(RotationXX2YY2ZZ2, 1, 0, 0, 0);
+		const VectorRegister4Float zz2_zz2_yy2 = VectorSwizzle(RotationXX2YY2ZZ2, 2, 2, 1, 0);
+		const VectorRegister4Float DiagonalSum = VectorAdd(yy2_xx2_xx2, zz2_zz2_yy2);
+		OutDiagonals = VectorSubtract(VectorOneFloat(), DiagonalSum);
 
 		// Grouping the non-diagonal elements in the rotation block by operations:
 		//    ((x*y2,y*z2,x*z2) + (w*z2,w*x2,w*y2)) and
@@ -1308,13 +1309,13 @@ private:
 
 		// RotBase = x*y2, y*z2, x*z2
 		// RotOffset = w*z2, w*x2, w*y2
-		const VectorRegister x_y_x = VectorSwizzle(Rotation, 0, 1, 0, 0);
-		const VectorRegister y2_z2_z2 = VectorSwizzle(RotationX2Y2Z2, 1, 2, 2, 0);
-		const VectorRegister RotBase = VectorMultiply(x_y_x, y2_z2_z2);
+		const VectorRegister4Float x_y_x = VectorSwizzle(Rotation, 0, 1, 0, 0);
+		const VectorRegister4Float y2_z2_z2 = VectorSwizzle(RotationX2Y2Z2, 1, 2, 2, 0);
+		const VectorRegister4Float RotBase = VectorMultiply(x_y_x, y2_z2_z2);
 
-		const VectorRegister w_w_w = VectorReplicate(Rotation, 3);
-		const VectorRegister z2_x2_y2 = VectorSwizzle(RotationX2Y2Z2, 2, 0, 1, 0);
-		const VectorRegister RotOffset = VectorMultiply(w_w_w, z2_x2_y2);
+		const VectorRegister4Float w_w_w = VectorReplicate(Rotation, 3);
+		const VectorRegister4Float z2_x2_y2 = VectorSwizzle(RotationX2Y2Z2, 2, 0, 1, 0);
+		const VectorRegister4Float RotOffset = VectorMultiply(w_w_w, z2_x2_y2);
 
 		// Adds = (RotBase + RotOffset):  (x*y2 + w*z2) , (y*z2 + w*x2), (x*z2 + w*y2)
 		// Subtracts = (RotBase - RotOffset) :  (x*y2 - w*z2) , (y*z2 - w*x2), (x*z2 - w*y2)
@@ -1343,7 +1344,7 @@ private:
 		//const VectorRegister ScaleZeroMask = VectorCompareEQ(InScale, VectorZero());
 
 		// VectorRegister( for each bit i: Mask[i] ? Vec1[i] : Vec2[i] )
-		SafeReciprocalScale = VectorSelect(ScaleZeroMask, VectorZero(), ReciprocalScale);
+		SafeReciprocalScale = VectorSelect(ScaleZeroMask, VectorZeroDouble(), ReciprocalScale);
 
 		return SafeReciprocalScale;
 	}
@@ -1372,14 +1373,14 @@ private:
 		const VectorRegister InvScale = VectorSet_W0(GetSafeScaleReciprocal(VectorSet_W1(Scale3D), ScalarRegister(GlobalVectorConstants::SmallNumber)));
 
 		// Invert the rotation
-		const VectorRegister InvRotation = VectorQuaternionInverse(Rotation);
+		const VectorRegister4Float InvRotation = VectorQuaternionInverse(Rotation);
 
 		// Invert the translation
 		const VectorRegister ScaledTranslation = VectorMultiply(InvScale, Translation);
 		const VectorRegister t2 = VectorQuaternionRotateVector(InvRotation, ScaledTranslation);
 		const VectorRegister InvTranslation = VectorSet_W0(VectorNegate(t2));
 
-		return FTransform(MakeVectorRegisterFloatFromDouble(InvRotation), InvTranslation, InvScale);
+		return FTransform(InvRotation, InvTranslation, InvScale);
 	}
 
 	/**
@@ -1442,7 +1443,7 @@ FORCEINLINE void FTransform::ScaleTranslation(const float& InScale)
 // this function is from matrix, and all it does is to normalize rotation portion
 FORCEINLINE void FTransform::RemoveScaling(float Tolerance/*=SMALL_NUMBER*/)
 {
-	Scale3D = VectorSet_W0( VectorOne() );
+	Scale3D = VectorSet_W0( VectorOneDouble() );
 	NormalizeRotation();	
 
 	DiagnosticCheckNaN_Rotate();
@@ -1534,7 +1535,7 @@ FORCEINLINE void FTransform::Multiply(FTransform* OutTransform, const FTransform
 		OutTransform->Translation = VectorAdd(RotatedTranslate, TranslateB);
 
 		// ScaleResult = Scale.B * Scale.A
-		OutTransform->Scale3D = VectorMultiply(ScaleA, ScaleB);;
+		OutTransform->Scale3D = VectorMultiply(ScaleA, ScaleB);
 	}
 }
 /** 
