Fixed compilation on Windows with LWC off

diff --git a/Engine/Source/Runtime/Core/Public/Math/TransformVectorized.h b/Engine/Source/Runtime/Core/Public/Math/TransformVectorized.h
--- a/Engine/Source/Runtime/Core/Public/Math/TransformVectorized.h
+++ b/Engine/Source/Runtime/Core/Public/Math/TransformVectorized.h
@@ -56,7 +56,7 @@ public:
 		if (VectorContainsNaNOrInfinite(Scale3D))
 		{
 			logOrEnsureNanError(TEXT("FTransform Vectorized Scale3D contains NaN"));
-			const_cast<FTransform*>(this)->Scale3D = VectorSet_W0( VectorOneDouble() );
+			const_cast<FTransform*>(this)->Scale3D = VectorSet_W0( VectorOneVectorRegister() );
 		}
 	}
 
@@ -65,7 +65,7 @@ public:
 		if (VectorContainsNaNOrInfinite(Translation))
 		{
 			logOrEnsureNanError(TEXT("FTransform Vectorized Translation contains NaN"));
-			const_cast<FTransform*>(this)->Translation = VectorZeroDouble();
+			const_cast<FTransform*>(this)->Translation = VectorOneVectorRegister();
 		}
 	}
 
@@ -74,7 +74,7 @@ public:
 		if (VectorContainsNaNOrInfinite(Rotation))
 		{
 			logOrEnsureNanError(TEXT("FTransform Vectorized Rotation contains NaN"));
-			const_cast<FTransform*>(this)->Rotation = VectorSet_W1( VectorZeroFloat() );
+			const_cast<FTransform*>(this)->Rotation = VectorSet_W1( VectorZeroVectorRegister() );
 		}
 	}
 
@@ -110,7 +110,7 @@ public:
 		// Rotation = {0,0,0,1)
 		Rotation = GlobalVectorConstants::Float0001; //VectorSet_W1( VectorZero() );
 		// Translation = {0,0,0,0)
-		Translation = VectorZeroDouble();
+		Translation = VectorZeroVectorRegister();
 		// Scale3D = {1,1,1,0);
 		Scale3D = GlobalVectorConstants::Vector1110; //VectorSet_W0( VectorOne() );
 	}
@@ -142,7 +142,7 @@ public:
 		// Rotation = InRotation
 		Rotation =  VectorLoadAligned( &InRotation.X );
 		// Translation = {0,0,0,0)
-		Translation = VectorZeroDouble();
+		Translation = VectorZeroVectorRegister();
 		// Scale3D = {1,1,1,0);
 		Scale3D = GlobalVectorConstants::Vector1110; //VectorSet_W0( VectorOne() );
 
@@ -160,7 +160,7 @@ public:
 		// Rotation = InRotation
 		Rotation =  VectorLoadAligned( &InQuatRotation.X );
 		// Translation = {0,0,0,0)
-		Translation = VectorZeroDouble();
+		Translation = VectorZeroVectorRegister();
 		// Scale3D = {1,1,1,0);
 		Scale3D = GlobalVectorConstants::Vector1110; //VectorSet_W0( VectorOne() );
 
@@ -675,7 +675,7 @@ private:
 
 	FORCEINLINE static bool Private_AnyHasNegativeScale(const VectorRegister& InScale3D, const  VectorRegister& InOtherScale3D)
 	{
-		return !!VectorAnyLesserThan(VectorMin(InScale3D, InOtherScale3D), VectorZeroDouble());
+		return !!VectorAnyLesserThan(VectorMin(InScale3D, InOtherScale3D), VectorZeroVectorRegister());
 	}
 
 	FORCEINLINE bool Private_RotationEquals( const VectorRegister4Float& InRotation, float Tolerance = KINDA_SMALL_NUMBER) const
@@ -788,7 +788,7 @@ public:
 		// Rotation = {0,0,0,1)
 		Rotation = GlobalVectorConstants::Float0001; //VectorSet_W1( VectorZero() );
 		// Translation = {0,0,0,0)
-		Translation = VectorZeroDouble();
+		Translation = VectorZeroVectorRegister();
 		// Scale3D = {1,1,1,0);
 		Scale3D = GlobalVectorConstants::Vector1110; //VectorSet_W0( VectorOne() );
 	}
@@ -1132,8 +1132,8 @@ public:
 		// Blend translation and scale
 		//    BlendedAtom.Translation = Lerp(Zero, SourceAtom.Translation, Alpha);
 		//    BlendedAtom.Scale = Lerp(0, SourceAtom.Scale, Alpha);
-		const VectorRegister BlendedTranslation	= FMath::Lerp<VectorRegister>(VectorZeroDouble(), SourceAtom.Translation, BlendWeight.Value);
-		const VectorRegister BlendedScale3D	= FMath::Lerp<VectorRegister>(VectorZeroDouble(), SourceAtom.Scale3D, BlendWeight.Value);
+		const VectorRegister BlendedTranslation	= FMath::Lerp<VectorRegister>(VectorZeroVectorRegister(), SourceAtom.Translation, BlendWeight.Value);
+		const VectorRegister BlendedScale3D	= FMath::Lerp<VectorRegister>(VectorZeroVectorRegister(), SourceAtom.Scale3D, BlendWeight.Value);
 
 		// Apply translation and scale to final atom
 		//     FinalAtom.Translation += BlendedAtom.Translation
@@ -1223,8 +1223,12 @@ public:
 		if(InMatrix.Determinant() < 0.f)
 		{
 			// Assume it is along X and modify transform accordingly. 
-			// It doesn't actually matter which axis we choose, the 'appearance' will be the same			
+			// It doesn't actually matter which axis we choose, the 'appearance' will be the same
+#if !UE_LARGE_WORLD_COORDINATES_DISABLED
 			Scale3D = VectorMultiply(Scale3D, GlobalVectorConstants::DoubleMinus1_111 );			
+#else
+			Scale3D = VectorMultiply(Scale3D, GlobalVectorConstants::FloatMinus1_111);
+#endif
 			M.SetAxis(0, -M.GetScaledAxis( EAxis::X ));
 		}
 
@@ -1344,7 +1348,7 @@ private:
 		//const VectorRegister ScaleZeroMask = VectorCompareEQ(InScale, VectorZero());
 
 		// VectorRegister( for each bit i: Mask[i] ? Vec1[i] : Vec2[i] )
-		SafeReciprocalScale = VectorSelect(ScaleZeroMask, VectorZeroDouble(), ReciprocalScale);
+		SafeReciprocalScale = VectorSelect(ScaleZeroMask, VectorZeroVectorRegister(), ReciprocalScale);
 
 		return SafeReciprocalScale;
 	}
@@ -1443,7 +1447,7 @@ FORCEINLINE void FTransform::ScaleTranslation(const float& InScale)
 // this function is from matrix, and all it does is to normalize rotation portion
 FORCEINLINE void FTransform::RemoveScaling(float Tolerance/*=SMALL_NUMBER*/)
 {
-	Scale3D = VectorSet_W0( VectorOneDouble() );
+	Scale3D = VectorSet_W0( VectorOneVectorRegister() );
 	NormalizeRotation();	
 
 	DiagnosticCheckNaN_Rotate();
diff --git a/Engine/Source/Runtime/Core/Public/Math/UnrealMathFPU.h b/Engine/Source/Runtime/Core/Public/Math/UnrealMathFPU.h
--- a/Engine/Source/Runtime/Core/Public/Math/UnrealMathFPU.h
+++ b/Engine/Source/Runtime/Core/Public/Math/UnrealMathFPU.h
@@ -134,9 +134,13 @@ struct alignas(alignof(VectorRegister4Double)) AlignedDouble4
 
 // LWC: Alias AlignedRegister4 to correct precision based on LWC
 #if !UE_LARGE_WORLD_COORDINATES_DISABLED
-typedef AlignedDouble4 AlignedRegister4;
+	typedef AlignedDouble4 AlignedRegister4;
+	#define VectorZeroVectorRegister() VectorZeroDouble()
+	#define VectorOneVectorRegister() VectorOneDouble()
 #else
-typedef AlignedFloat4 AlignedRegister4;
+	typedef AlignedFloat4 AlignedRegister4;
+	#define VectorZeroVectorRegister() VectorZeroFloat()
+	#define VectorOneVectorRegister() VectorOneFloat()
 #endif
 
 
diff --git a/Engine/Source/Runtime/Core/Public/Math/UnrealMathNeon.h b/Engine/Source/Runtime/Core/Public/Math/UnrealMathNeon.h
--- a/Engine/Source/Runtime/Core/Public/Math/UnrealMathNeon.h
+++ b/Engine/Source/Runtime/Core/Public/Math/UnrealMathNeon.h
@@ -91,8 +91,12 @@ struct alignas(16) VectorRegister4Double
 // LWC: Alias VectorRegister to correct precision based on LWC
 #if !UE_LARGE_WORLD_COORDINATES_DISABLED
 	typedef VectorRegister4Double VectorRegister;
+	#define VectorZeroVectorRegister() VectorZeroDouble()
+	#define VectorOneVectorRegister() VectorOneDouble()
 #else
 	typedef VectorRegister4Float VectorRegister;
+	#define VectorZeroVectorRegister() VectorZeroFloat()
+	#define VectorOneVectorRegister() VectorOneFloat()
 #endif
 
 // Forward declarations
diff --git a/Engine/Source/Runtime/Core/Public/Math/UnrealMathSSE.h b/Engine/Source/Runtime/Core/Public/Math/UnrealMathSSE.h
--- a/Engine/Source/Runtime/Core/Public/Math/UnrealMathSSE.h
+++ b/Engine/Source/Runtime/Core/Public/Math/UnrealMathSSE.h
@@ -168,9 +168,13 @@ typedef VectorRegister2Double VectorRegister2d;
 
 // LWC: Alias VectorRegister to correct precision based on LWC
 #if !UE_LARGE_WORLD_COORDINATES_DISABLED
-typedef VectorRegister4Double VectorRegister4;
+	typedef VectorRegister4Double VectorRegister4;
+	#define VectorZeroVectorRegister() VectorZeroDouble()
+	#define VectorOneVectorRegister() VectorOneDouble()
 #else
-typedef VectorRegister4Float VectorRegister4;
+	typedef VectorRegister4Float VectorRegister4;
+	#define VectorZeroVectorRegister() VectorZeroFloat()
+	#define VectorOneVectorRegister() VectorOneFloat()
 #endif
 
 // Backwards compatibility
