Tasks:
* refactored integration with low-level tasks to fix task retraction: high-level tasks now store task body because there's no way to store it in low-level task and to implement retraction
* added timeout support for task retraction
* removed tests that deadlocks with task retraction

diff --git a/Engine/Source/Runtime/Core/Private/Tests/Tasks/TasksTest.cpp b/Engine/Source/Runtime/Core/Private/Tests/Tasks/TasksTest.cpp
--- a/Engine/Source/Runtime/Core/Private/Tests/Tasks/TasksTest.cpp
+++ b/Engine/Source/Runtime/Core/Private/Tests/Tasks/TasksTest.cpp
@@ -109,38 +109,19 @@ namespace UE { namespace TasksTests
 			verify(Event.BusyWait(FTimespan::Zero()));
 		}
 
-		{	// postpone execution so waiting kicks in first
-			std::atomic<int32> Counter{ 0 };
-			FTask Task = Launch(UE_SOURCE_LOCATION, [&Counter] { ++Counter; FPlatformProcess::Sleep(0.1f); });
-
-			ensure(!Task.Wait(FTimespan::Zero()));
-			Task.Wait();
-			check(Counter == 1);
-		}
-
-		{	// postpone execution so busy waiting kicks in first
-			std::atomic<int32> Counter{ 0 };
-			FTask Task = Launch(UE_SOURCE_LOCATION, [&Counter] { ++Counter; FPlatformProcess::Sleep(0.1f); });
-
-			ensure(!Task.BusyWait(FTimespan::Zero()));
-			Task.BusyWait();
-			check(Counter == 1);
-		}
-
-		{	// same but using `FTaskEvent`
+		{	// same but using busy-waiting
 			FTaskEvent Event{ UE_SOURCE_LOCATION };
-			FTask Task = Launch(UE_SOURCE_LOCATION, [&Event] { Event.Wait(); });
-			ensure(!Task.Wait(FTimespan::FromMilliseconds(100)));
-			Event.Trigger();
-			Task.Wait();
-		}
+			check(!Event.IsCompleted());
+
+			// check that waiting blocks
+			FTask Task = Launch(UE_SOURCE_LOCATION, [Event]() mutable { Event.BusyWait(); });
+			FPlatformProcess::Sleep(0.1f);
+			check(!Task.IsCompleted());
 
-		{	// same but using busy-wait and `FTaskEvent`
-			FTaskEvent Event{ UE_SOURCE_LOCATION };
-			FTask Task = Launch(UE_SOURCE_LOCATION, [&Event] { Event.BusyWait(); });
-			ensure(!Task.BusyWait(FTimespan::FromMilliseconds(100)));
 			Event.Trigger();
-			Task.BusyWait();
+			check(Event.IsCompleted());
+			verify(Event.Wait(FTimespan::Zero()));
+			verify(Event.BusyWait(FTimespan::Zero()));
 		}
 
 		{	// basic use-case, postpone waiting so the task is executed first
@@ -547,7 +528,8 @@ namespace UE { namespace TasksTests
 			FTaskEvent Prereq{ UE_SOURCE_LOCATION };
 
 			FTask Task{ Launch(UE_SOURCE_LOCATION, [] {}, Prereq) };
-			verify(!Task.Wait(FTimespan::FromMilliseconds(10)));
+			FPlatformProcess::Sleep(0.1f);
+			check(!Task.IsCompleted());
 
 			Prereq.Trigger();
 			Task.Wait();
@@ -557,7 +539,8 @@ namespace UE { namespace TasksTests
 			FTaskEvent Prereq{ UE_SOURCE_LOCATION };
 
 			FTask Task{ Launch(UE_SOURCE_LOCATION, [] {}, Prereq, ETaskPriority::Normal) };
-			verify(!Task.Wait(FTimespan::FromMilliseconds(10)));
+			FPlatformProcess::Sleep(0.1f);
+			check(!Task.IsCompleted());
 
 			Prereq.Trigger();
 			Task.Wait();
@@ -567,7 +550,8 @@ namespace UE { namespace TasksTests
 			FTaskEvent Event{ UE_SOURCE_LOCATION };
 			FTask Prereq{ Launch(UE_SOURCE_LOCATION, [&Event] { Event.Wait(); }) };
 			FTask Task{ Launch(UE_SOURCE_LOCATION, [] {}, Prereq) };
-			verify(!Task.Wait(FTimespan::FromMilliseconds(10)));
+			FPlatformProcess::Sleep(0.1f);
+			check(!Task.IsCompleted());
 
 			Event.Trigger();
 			Task.Wait();
@@ -594,10 +578,12 @@ namespace UE { namespace TasksTests
 			FTask Prereq2{ Launch(UE_SOURCE_LOCATION, [&Event] { Event.Wait(); }) };
 
 			TTask<void> Task{ Launch(UE_SOURCE_LOCATION, [] {}, Prerequisites(Prereq1, Prereq2)) };
-			verify(!Task.Wait(FTimespan::FromMilliseconds(10)));
+			FPlatformProcess::Sleep(0.1f);
+			check(!Task.IsCompleted());
 
 			Prereq1.Trigger();
-			verify(!Task.Wait(FTimespan::FromMilliseconds(10)));
+			FPlatformProcess::Sleep(0.1f);
+			check(!Task.IsCompleted());
 
 			Event.Trigger();
 			Task.Wait();
@@ -610,10 +596,12 @@ namespace UE { namespace TasksTests
 			TArray<TTask<void>> Prereqs{ Prereq1, Prereq2 }; // to check if a random iterable container works as a prerequisite collection
 
 			TTask<void> Task{ Launch(UE_SOURCE_LOCATION, [] {}, Prereqs, ETaskPriority::Normal) };
-			verify(!Task.Wait(FTimespan::FromMilliseconds(10)));
+			FPlatformProcess::Sleep(0.1f);
+			check(!Task.IsCompleted());
 
 			Prereq1.Trigger();
-			verify(!Task.Wait(FTimespan::FromMilliseconds(10)));
+			FPlatformProcess::Sleep(0.1f);
+			check(!Task.IsCompleted());
 
 			Event.Trigger();
 			Task.Wait();
@@ -624,7 +612,8 @@ namespace UE { namespace TasksTests
 			FTaskEvent Prereq{ UE_SOURCE_LOCATION };
 
 			FTask Task1{ Pipe.Launch(UE_SOURCE_LOCATION, [] {}, Prereq) };
-			verify(!Task1.Wait(FTimespan::FromMilliseconds(10)));
+			FPlatformProcess::Sleep(0.1f);
+			check(!Task1.IsCompleted());
 
 			FTask Task2{ Pipe.Launch(UE_SOURCE_LOCATION, [] {}) };
 			Task2.Wait();
@@ -640,10 +629,13 @@ namespace UE { namespace TasksTests
 			FTask Prereq2{ Launch(UE_SOURCE_LOCATION, [&Event] { Event.Wait(); }) };
 
 			FTask Task{ Pipe.Launch(UE_SOURCE_LOCATION, [] {}, Prerequisites(Prereq1, Prereq2)) };
-			verify(!Task.Wait(FTimespan::FromMilliseconds(10)));
+			FPlatformProcess::Sleep(0.1f);
+			check(!Task.IsCompleted());
 
 			Prereq1.Trigger();
-			verify(!Task.Wait(FTimespan::FromMilliseconds(10)));
+			FPlatformProcess::Sleep(0.1f);
+			check(!Task.IsCompleted());
+
 			Event.Trigger();
 			Task.Wait();
 		}
@@ -655,10 +647,13 @@ namespace UE { namespace TasksTests
 			FTask Prereq2{ Launch(UE_SOURCE_LOCATION, [&Event] { Event.Wait(); }) };
 
 			FTask Task{ Pipe.Launch(UE_SOURCE_LOCATION, [] {}, Prerequisites(Prereq1, Prereq2), ETaskPriority::Normal) };
-			verify(!Task.Wait(FTimespan::FromMilliseconds(10)));
+			FPlatformProcess::Sleep(0.1f);
+			check(!Task.IsCompleted());
 
 			Prereq1.Trigger();
-			verify(!Task.Wait(FTimespan::FromMilliseconds(10)));
+			FPlatformProcess::Sleep(0.1f);
+			check(!Task.IsCompleted());
+
 			Event.Trigger();
 			Task.Wait();
 		}
diff --git a/Engine/Source/Runtime/Core/Public/Tasks/Task.h b/Engine/Source/Runtime/Core/Public/Tasks/Task.h
--- a/Engine/Source/Runtime/Core/Public/Tasks/Task.h
+++ b/Engine/Source/Runtime/Core/Public/Tasks/Task.h
@@ -30,7 +30,7 @@ namespace UE { namespace Tasks
 		friend class TPrerequisites;
 
 		template<typename TaskCollectionType>
-		friend bool Private::TryRetractAndExecute(TaskCollectionType&& Tasks);
+		friend bool Private::TryRetractAndExecute(TaskCollectionType&& Tasks, FTimespan Timeout);
 
 	protected:
 		TTaskBase() = default;
@@ -51,23 +51,12 @@ namespace UE { namespace Tasks
 			return !IsValid() || Pimpl->IsCompleted();
 		}
 
-		// waits for task's completion, with optional timeout
+		// waits for task's completion, with optional timeout. Tries to retract the task and execute it in-place, if failed - blocks until the task 
+		// is completed by another thread. If timeout is zero, tries to retract the task and returns immedially after that.
 		// @return true if the task is completed
 		bool Wait(FTimespan Timeout = FTimespan::MaxValue())
 		{
-			if (!IsValid() || IsCompleted())
-			{
-				return true;
-			}
-
-			// "zero timeout" - just check if the task is completed
-			if (Timeout == FTimespan::Zero())
-			{
-				return false;
-			}
-
-			// waiting with timeout can't do retraction, otherwise timeout value is unreliable
-			if (Timeout == FTimespan::MaxValue() && Pimpl->TryRetractAndExecute())
+			if (!IsValid() || Pimpl->TryRetractAndExecute())
 			{
 				return true;
 			}
@@ -242,10 +231,11 @@ namespace UE { namespace Tasks
 
 	// wait for multiple tasks, with optional timeout
 	template<typename TaskCollectionType>
-	bool Wait(TaskCollectionType&& Tasks, FTimespan Timeout = FTimespan::MaxValue())
+	bool Wait(TaskCollectionType&& Tasks, FTimespan InTimeout = FTimespan::MaxValue())
 	{
-		// waiting with timeout can't do retraction, otherwise timeout value is unreliable
-		if (Timeout == FTimespan::MaxValue() && Private::TryRetractAndExecute(Tasks))
+		FTimeout Timeout{ InTimeout };
+
+		if (Private::TryRetractAndExecute(Tasks, InTimeout))
 		{
 			return true;
 		}
@@ -264,17 +254,29 @@ namespace UE { namespace Tasks
 			return true;
 		}
 
-		return CompletionEvent->Wait(Timeout);
+		return CompletionEvent->Wait(Timeout.GetRemainingTime());
 	}
 
 	// wait for multiple tasks while executing other tasks
 	template<typename TaskCollectionType>
-	void BusyWait(TaskCollectionType&& Tasks)
+	bool BusyWait(TaskCollectionType&& Tasks, FTimespan TimeoutValue = FTimespan::MaxValue())
 	{
+		FTimeout Timeout{ TimeoutValue };
+
+		if (Private::TryRetractAndExecute(Tasks, TimeoutValue))
+		{
+			return true;
+		}
+
 		for (auto& Task : Tasks)
 		{
-			Task.BusyWait();
+			if (Timeout || (!Task.BusyWait(Timeout.GetRemainingTime())))
+			{
+				return false;
+			}
 		}
+
+		return true;
 	}
 
 	using FTask = TTask<void>;
diff --git a/Engine/Source/Runtime/Core/Public/Tasks/TaskPrivate.h b/Engine/Source/Runtime/Core/Public/Tasks/TaskPrivate.h
--- a/Engine/Source/Runtime/Core/Public/Tasks/TaskPrivate.h
+++ b/Engine/Source/Runtime/Core/Public/Tasks/TaskPrivate.h
@@ -51,26 +51,21 @@ namespace UE { namespace Tasks
 			
 			// initialises the task but doesn't launches it
 			template<typename TaskBodyType, typename DeleterType>
-			void Init(const TCHAR* DebugName, TaskBodyType&& TaskBody, DeleterType&& Deleter, LowLevelTasks::ETaskPriority Priority)
+			void Init(const TCHAR* DebugName, TaskBodyType&& InTaskBody, DeleterType&& Deleter, LowLevelTasks::ETaskPriority Priority)
 			{
-				// to not store a copy of `TaskBody` that is stored inside `LowLevelTask` anyway, and to implement task retraction
-				// (see `Wait()` method), we have to use LowLevelTask's "Continuation" parameter instead of "Runnable". This way we can make
-				// the scheduler immediately execute retracted task.
-				LowLevelTask.Init(DebugName, Priority, 
-					[] {}, // Runnable
+				TaskBody = Forward<TaskBodyType>(InTaskBody);
+
+				LowLevelTask.Init(DebugName, Priority,
 					[
-						this, 
-						TaskBody = Forward<TaskBodyType>(TaskBody), 
-						// releasing scheduler's task reference can cause task's automatic destruction and so must be done after the task is 
-						// flagged as completed. The task is flagged as completed after the continuation is executed but before its destroyed.
-						// `Deleter` is captured by value and is destroyed along with the continuation, calling the given functor 
-						// on destruction
+						this,
+						// releasing scheduler's task reference can cause task's automatic destruction and so must be done after the low-level task
+						// task is flagged as completed. The task is flagged as completed after the continuation is executed but before its destroyed.
+						// `Deleter` is captured by value and is destroyed along with the continuation, calling the given functor on destruction
 						Deleter = Forward<DeleterType>(Deleter)
-					]() mutable
+					]
 					{
-						Execute(MoveTemp(TaskBody));
-
-					} // Continuation
+						TryExecute();
+					}
 				);
 			}
 
@@ -162,11 +157,17 @@ namespace UE { namespace Tasks
 				return Subsequents.IsClosed();
 			}
 
-			FORCENOINLINE bool TryRetractAndExecute()
+			bool TryRetractAndExecute()
 			{
-				bool bUnlocked = NumLocks.load(std::memory_order_acquire) == (GetPipe() == nullptr ? 1 : 0);
-				// TryCancel() will try to cancel the task (if it's not being executed yet) and execute it
-				return bUnlocked && LowLevelTask.TryCancel() && IsCompleted();
+				bool bIsLocked = NumLocks.load(std::memory_order_acquire) != (GetPipe() == nullptr ? 1 : 0);
+				// try to execute before cancelling the low-level task as successful cancellation can release the last reference and destroy the task
+				if (!bIsLocked && TryExecute())
+				{
+					LowLevelTask.TryCancel();
+					return true;
+				}
+
+				return false;
 			}
 
 			// waits until the task is completed while executing other tasks
@@ -175,10 +176,16 @@ namespace UE { namespace Tasks
 				// we try to retract the task from the scheduler, if it's execution hasn't started yet, to execute it immediately. This makes
 				// the task skip waiting in the scheduler's queue.
 				// we can't retract it right here because it can be not scheduled yet (e.g. when it's blocked by a pipe), so we need to wait for this
-				// `!IsReady()` here means "the task is scheduled". The task can be completed w/o being scheduled at all (e.g. if it has 
-				// an empty body) so we need to check for this too
+				// `!IsReady()` here means "the task is scheduled"
 				LowLevelTasks::BusyWaitUntil([this] { return !LowLevelTask.IsReady(); });
-				LowLevelTask.TryCancel(); //This will try to cancel the task and run it's continuation (thereby executing the workload) 
+
+				// try to execute before cancelling the low-level task as successful cancellation can release the last reference and destroy the task
+				if (TryExecute())
+				{
+					LowLevelTask.TryCancel();
+					return;
+				}
+
 				LowLevelTasks::BusyWaitUntil([this] { return IsCompleted(); });
 			}
 
@@ -193,7 +200,13 @@ namespace UE { namespace Tasks
 					return false;
 				}
 
-				LowLevelTask.TryCancel(); //This will try to cancel the task and run it's continuation (thereby executing the workload) 
+				// try to execute before cancelling the low-level task as successful cancellation can release the last reference and destroy the task
+				if (TryExecute())
+				{
+					LowLevelTask.TryCancel();
+					return true;
+				}
+
 				LowLevelTasks::BusyWaitUntil([this, Timeout] { return IsCompleted() || Timeout; });
 				return IsCompleted();
 			}
@@ -211,7 +224,13 @@ namespace UE { namespace Tasks
 					return false;
 				}
 
-				LowLevelTask.TryCancel(); //This will try to cancel the task and run it's continuation (thereby executing the workload) 
+				// try to execute before cancelling the low-level task as successful cancellation can release the last reference and destroy the task
+				if (TryExecute())
+				{
+					LowLevelTask.TryCancel();
+					return true;
+				}
+
 				LowLevelTasks::BusyWaitUntil(
 					[this, Condition = Forward<ConditionType>(Condition)] { return IsCompleted() || Condition(); }
 				);
@@ -250,9 +269,10 @@ namespace UE { namespace Tasks
 				// "inline" tasks are not scheduled but executed as soon as they are unlocked
 				if (LowLevelTask.GetPriority() == InlineTaskPriority)
 				{
-					// the task is not scheduled yet, so successful retraction is guaranted
-					verify(LowLevelTask.TryCancel()); // this will cancel the task for the scheduler and execute LowLevelTask's continuation, 
-					// which is actually task execution
+					// execute before cancelling the low-level task as successful cancellation can release the last reference and destroy the task
+					// the low-level task wasn't scheduled, so successful execution and low-level task cancellation is guaranted
+					verify(TryExecute());
+					verify(LowLevelTask.TryCancel());
 					return true;
 				}
 
@@ -260,19 +280,26 @@ namespace UE { namespace Tasks
 				return LowLevelTasks::TryLaunch(LowLevelTask);
 			}
 
-			template<typename TaskBodyType>
-			void Execute(TaskBodyType&& TaskBody)
+			// execute the task body if the execution wasn't started yet
+			bool TryExecute()
 			{
+				if (!bAvailableForExecution.exchange(false, std::memory_order_relaxed))
+				{
+					return false; // too late, execution already started, e.g. by task retraction
+				}
+
 				TaskTrace::Started(GetTraceId());
 				{
 					TRACE_CPUPROFILER_EVENT_SCOPE(ExecuteTask);
 					StartPipeExecution();
-					Invoke(TaskBody);
+					TaskBody();
 					FinishPipeExecution();
 				}
 				TaskTrace::Finished(GetTraceId());
 
 				Close();
+
+				return true;
 			}
 
 			// checks if the task is ready to be launched by trying to push it into the pipe.
@@ -316,6 +343,8 @@ namespace UE { namespace Tasks
 
 		private:
 			LowLevelTasks::FTask LowLevelTask;
+			TUniqueFunction<void()> TaskBody;
+			std::atomic<bool> bAvailableForExecution{ true };
 
 			// the number of times that the task should be unlocked before it can be scheduled
 			// initial count is 1 for launching the task (it can't be scheduled before it's launched) and 1 for a potential blocked pipe
@@ -331,16 +360,24 @@ namespace UE { namespace Tasks
 		};
 
 		template<typename TaskCollectionType>
-		bool TryRetractAndExecute(TaskCollectionType&& Tasks)
+		bool TryRetractAndExecute(TaskCollectionType&& Tasks, FTimespan InTimeout = FTimespan::MaxValue())
 		{
+			FTimeout Timeout{ InTimeout };
 			bool bResult = true;
+
 			for (auto& Task : Tasks)
 			{
 				if (!Task.Pimpl->TryRetractAndExecute())
 				{
 					bResult = false;
 				}
+
+				if (Timeout)
+				{
+					return false;
+				}
 			}
+
 			return bResult;
 		}
 
