readd decompress to existing buffer - dropped when removing test stuff. #pf 60fae49ed3905c0001f20dc7

diff --git a/Engine/Source/Runtime/Core/Private/Compression/OodleDataCompressionUtil.cpp b/Engine/Source/Runtime/Core/Private/Compression/OodleDataCompressionUtil.cpp
--- a/Engine/Source/Runtime/Core/Private/Compression/OodleDataCompressionUtil.cpp
+++ b/Engine/Source/Runtime/Core/Private/Compression/OodleDataCompressionUtil.cpp
@@ -42,6 +42,20 @@ namespace FOodleDataCompressionUtil
 		return true;
 	}
 
+	bool CORE_API FCompressedArray::DecompressToExistingBuffer(void* InDestinationBuffer) const
+	{
+		int32 DecompressedSize, CompressedSize;
+		if (PeekSizes(CompressedSize, DecompressedSize) == false)
+		{
+			return false;
+		}
+
+		// If we have a valid header, then if we don't have the actual data, it's corrupted data.
+		check(CompressedTArrayHeaderSize + CompressedSize <= Num());
+
+		return FOodleDataCompression::Decompress(InDestinationBuffer, DecompressedSize, GetData() + CompressedTArrayHeaderSize, CompressedSize);
+	}
+
 	bool CORE_API FCompressedArray::DecompressToAllocatedBuffer(void*& OutDestinationBuffer, int32& OutDestinationBufferSize) const
 	{
 		int32 DecompressedSize, CompressedSize;
