Setting correct background task affinities.

diff --git a/Engine/Source/Runtime/Core/Private/Async/Fundamental/Scheduler.cpp b/Engine/Source/Runtime/Core/Private/Async/Fundamental/Scheduler.cpp
--- a/Engine/Source/Runtime/Core/Private/Async/Fundamental/Scheduler.cpp
+++ b/Engine/Source/Runtime/Core/Private/Async/Fundamental/Scheduler.cpp
@@ -46,6 +46,10 @@ namespace LowLevelTasks
 		const uint32 WaitTimes[8] = { 719, 991, 1361, 1237, 1597, 953, 587, 1439 };
 		uint32 WaitTime = WaitTimes[WorkerId % 8];
 		uint64 ThreadAffinityMask = FPlatformAffinity::GetTaskGraphThreadMask();
+		if (bPermitBackgroundWork && FPlatformAffinity::GetTaskGraphBackgroundTaskMask() != 0xFFFFFFFFFFFFFFFF)
+		{
+			ThreadAffinityMask = FPlatformAffinity::GetTaskGraphBackgroundTaskMask();
+		}
 
 		const FProcessorGroupDesc& ProcessorGroups = FPlatformMisc::GetProcessorGroupDesc();
 		int32 CpuGroupCount = ProcessorGroups.NumProcessorGroups;
