Temp build fix

#rnx

diff --git a/Engine/Source/Runtime/Core/Public/Math/TransformVectorized.h b/Engine/Source/Runtime/Core/Public/Math/TransformVectorized.h
--- a/Engine/Source/Runtime/Core/Public/Math/TransformVectorized.h
+++ b/Engine/Source/Runtime/Core/Public/Math/TransformVectorized.h
@@ -74,7 +74,7 @@ public:
 		if (VectorContainsNaNOrInfinite(Rotation))
 		{
 			logOrEnsureNanError(TEXT("FTransform Vectorized Rotation contains NaN"));
-			const_cast<FTransform*>(this)->Rotation = VectorSet_W1( VectorZeroVectorRegister() );
+//			const_cast<FTransform*>(this)->Rotation = VectorSet_W1( VectorZeroVectorRegister() );
 		}
 	}
 
