Remove checks when adding ticker outside of game thread

Checks added in CL 16705335

diff --git a/Engine/Source/Runtime/Core/Private/Containers/Ticker.cpp b/Engine/Source/Runtime/Core/Private/Containers/Ticker.cpp
--- a/Engine/Source/Runtime/Core/Private/Containers/Ticker.cpp
+++ b/Engine/Source/Runtime/Core/Private/Containers/Ticker.cpp
@@ -17,9 +17,6 @@ FTicker::~FTicker()
 
 FDelegateHandle FTicker::AddTicker(const FTickerDelegate& InDelegate, float InDelay)
 {
-	// not currently threadsafe
-	check(IsInGameThread());
-
 	// We can add elements safely even during tick.
 	Elements.Emplace(CurrentTime + InDelay, InDelay, InDelegate);
 	// @todo this needs a unique handle for each add call to allow you to register the same delegate twice with a different delay safely.
@@ -30,9 +27,6 @@ FDelegateHandle FTicker::AddTicker(const FTickerDelegate& InDelegate, float InDe
 
 FDelegateHandle FTicker::AddTicker(const TCHAR* InName, float InDelay, TFunction<bool(float)> Function)
 {
-	// not currently threadsafe
-	check(IsInGameThread());
-
 	// todo - use InName for profiling. Added in sig to be forward looking..
 
 	FTickerDelegate Delegate = FTickerDelegate::CreateLambda(Function);
@@ -45,9 +39,6 @@ FDelegateHandle FTicker::AddTicker(const TCHAR* InName, float InDelay, TFunction
 
 void FTicker::RemoveTicker(FDelegateHandle Handle)
 {
-	// not currently threadsafe
-	check(IsInGameThread());
-
 	// must remove the handle from both arrays because we could be in the middle of a tick, 
 	// and may be removing ourselves or something else we've already considered this Tick
 	auto CompareHandle = [=](const FElement& Element){ return Element.Delegate.GetHandle() == Handle; };
@@ -163,4 +154,4 @@ FTickerObjectBase::~FTickerObjectBase()
 		Ticker.RemoveTicker(TickHandle);
 		TickHandle = FDelegateHandle();
 	}
-}
\ No newline at end of file
+}
