PR #8235: Fixed a crash on the second run of FArchiveFileReaderGenericTest (Contributed by [REDACTED_PERSON])
#rnx

diff --git a/Engine/Source/Runtime/Core/Private/Tests/HAL/FileManagerGenericTest.cpp b/Engine/Source/Runtime/Core/Private/Tests/HAL/FileManagerGenericTest.cpp
--- a/Engine/Source/Runtime/Core/Private/Tests/HAL/FileManagerGenericTest.cpp
+++ b/Engine/Source/Runtime/Core/Private/Tests/HAL/FileManagerGenericTest.cpp
@@ -51,7 +51,7 @@ void FArchiveFileReaderGenericTest::TestInternalPrecache()
 	uint32 BufferSize = 1024;
 	TArray<uint8> UnusedBytes;
 	UnusedBytes.SetNumUninitialized(BufferSize);
-	TUniquePtr<FArchiveFileReaderGeneric> Reader(new FArchiveFileReaderGeneric(ReadHandle.Get(), *TestFileName, FileSize, BufferSize));
+	TUniquePtr<FArchiveFileReaderGeneric> Reader(new FArchiveFileReaderGeneric(ReadHandle.Release(), *TestFileName, FileSize, BufferSize));
 	TestEqual(TEXT("Initial Pos should be 0"), Reader->Pos, 0LL);
 	TestEqual(TEXT("Size should be what was passed in"), Reader->Size, FileSize);
 	TestEqual(TEXT("BufferSize should be what was passed in"), Reader->BufferSize, BufferSize);
