Add Insight tracing around file buffer shrinking which is often slow

#rnx

diff --git a/Engine/Source/Runtime/Core/Private/Windows/WindowsPlatformFile.cpp b/Engine/Source/Runtime/Core/Private/Windows/WindowsPlatformFile.cpp
--- a/Engine/Source/Runtime/Core/Private/Windows/WindowsPlatformFile.cpp
+++ b/Engine/Source/Runtime/Core/Private/Windows/WindowsPlatformFile.cpp
@@ -498,6 +498,8 @@ public:
 	{
 		if (IsValid())
 		{
+			TRACE_CPUPROFILER_EVENT_SCOPE(ShrinkBuffers);
+
 			TRACE_PLATFORMFILE_BEGIN_REOPEN(Handle);
 			HANDLE NewFileHandle = ReOpenFile(Handle, DesiredAccess, ShareMode, Flags);
 			TRACE_PLATFORMFILE_END_REOPEN(NewFileHandle);
