Add AddressLimit for Linux

PR #8202 (thanks anon-oss!)

diff --git a/Engine/Source/Runtime/Core/Private/Unix/UnixPlatformMemory.cpp b/Engine/Source/Runtime/Core/Private/Unix/UnixPlatformMemory.cpp
--- a/Engine/Source/Runtime/Core/Private/Unix/UnixPlatformMemory.cpp
+++ b/Engine/Source/Runtime/Core/Private/Unix/UnixPlatformMemory.cpp
@@ -898,12 +898,13 @@ const FPlatformMemoryConstants& FUnixPlatformMemory::GetConstants()
 
 #endif // PLATFORM_FREEBSD
 
-		MemoryConstants.TotalPhysicalGB = (MemoryConstants.TotalPhysical + 1024ULL * 1024ULL * 1024ULL - 1) / 1024ULL / 1024ULL / 1024ULL;
-
 		MemoryConstants.PageSize = sysconf(_SC_PAGESIZE);
 		MemoryConstants.BinnedPageSize = FMath::Max((SIZE_T)65536, MemoryConstants.PageSize);
 		MemoryConstants.BinnedAllocationGranularity = MemoryConstants.PageSize;
 		MemoryConstants.OsAllocationGranularity = MemoryConstants.PageSize;
+
+		MemoryConstants.TotalPhysicalGB = (MemoryConstants.TotalPhysical + 1024ULL * 1024ULL * 1024ULL - 1) / 1024ULL / 1024ULL / 1024ULL;
+		MemoryConstants.AddressLimit = FPlatformMath::RoundUpToPowerOfTwo64(MemoryConstants.TotalPhysical);
 	}
 
 	return MemoryConstants;	
