IoDispatcher: Inherit LLM tags from the caller instead of putting all io requests and io buffers in the FileSystem category
#rnx

diff --git a/Engine/Source/Runtime/Core/Private/IO/IoDispatcher.cpp b/Engine/Source/Runtime/Core/Private/IO/IoDispatcher.cpp
--- a/Engine/Source/Runtime/Core/Private/IO/IoDispatcher.cpp
+++ b/Engine/Source/Runtime/Core/Private/IO/IoDispatcher.cpp
@@ -187,7 +187,6 @@ public:
 
 	FIoRequestImpl* AllocRequest(const FIoChunkId& ChunkId, FIoReadOptions Options)
 	{
-		LLM_SCOPE(ELLMTag::FileSystem);
 		FIoRequestImpl* Request = RequestAllocator.Construct(*this);
 
 		Request->ChunkId = ChunkId;
diff --git a/Engine/Source/Runtime/Core/Private/IO/IoDispatcherFileBackend.cpp b/Engine/Source/Runtime/Core/Private/IO/IoDispatcherFileBackend.cpp
--- a/Engine/Source/Runtime/Core/Private/IO/IoDispatcherFileBackend.cpp
+++ b/Engine/Source/Runtime/Core/Private/IO/IoDispatcherFileBackend.cpp
@@ -1143,7 +1143,7 @@ bool FFileIoStore::Resolve(FIoRequestImpl* Request)
 				}
 				else
 				{
-					LLM_SCOPE(ELLMTag::FileSystem);
+					LLM_SCOPE(Request->InheritedLLMTag);
 					TRACE_CPUPROFILER_EVENT_SCOPE(AllocMemoryForRequest);
 					Request->IoBuffer = FIoBuffer(ResolvedSize);
 				}
diff --git a/Engine/Source/Runtime/Core/Public/IO/IoDispatcherBackend.h b/Engine/Source/Runtime/Core/Public/IO/IoDispatcherBackend.h
--- a/Engine/Source/Runtime/Core/Public/IO/IoDispatcherBackend.h
+++ b/Engine/Source/Runtime/Core/Public/IO/IoDispatcherBackend.h
@@ -3,12 +3,14 @@
 #pragma once
 
 #include "IoDispatcher.h"
+#include "HAL/LowLevelMemTracker.h"
 
 class FIoRequestImpl
 {
 public:
 	FIoRequestImpl* NextRequest = nullptr;
 	void* BackendData = nullptr;
+	LLM(const UE::LLMPrivate::FTagData* InheritedLLMTag);
 	FIoChunkId ChunkId;
 	FIoReadOptions Options;
 	FIoBuffer IoBuffer;
@@ -17,7 +19,7 @@ public:
 	FIoRequestImpl(FIoDispatcherImpl& InDispatcher)
 		: Dispatcher(InDispatcher)
 	{
-
+		LLM(InheritedLLMTag = FLowLevelMemTracker::bIsDisabled ? nullptr : FLowLevelMemTracker::Get().GetActiveTagData(ELLMTracker::Default));
 	}
 
 	bool IsCancelled() const
