Changed bookmarks trace events with an array-type field instead of using attachments.

#rnx

diff --git a/Engine/Source/Developer/TraceServices/Private/Analyzers/MiscTraceAnalysis.cpp b/Engine/Source/Developer/TraceServices/Private/Analyzers/MiscTraceAnalysis.cpp
--- a/Engine/Source/Developer/TraceServices/Private/Analyzers/MiscTraceAnalysis.cpp
+++ b/Engine/Source/Developer/TraceServices/Private/Analyzers/MiscTraceAnalysis.cpp
@@ -109,8 +109,9 @@ bool FMiscTraceAnalyzer::OnEvent(uint16 RouteId, EStyle Style, const FOnEventCon
 		uint64 BookmarkPoint = EventData.GetValue<uint64>("BookmarkPoint");
 		uint64 Cycle = EventData.GetValue<uint64>("Cycle");
 		double Timestamp = Context.EventTime.AsSeconds(Cycle);
-		BookmarkProvider.AppendBookmark(Timestamp, BookmarkPoint, EventData.GetAttachment());
-		LogProvider.AppendMessage(BookmarkPoint, Timestamp, EventData.GetAttachment());
+		TArrayView<const uint8> FormatArgsView = FTraceAnalyzerUtils::LegacyAttachmentArray("FormatArgs", Context);
+		BookmarkProvider.AppendBookmark(Timestamp, BookmarkPoint, FormatArgsView.GetData());
+		LogProvider.AppendMessage(BookmarkPoint, Timestamp, FormatArgsView.GetData());
 		break;
 	}
 	case RouteId_BeginFrame:
diff --git a/Engine/Source/Runtime/Core/Private/ProfilingDebugging/MiscTrace.cpp b/Engine/Source/Runtime/Core/Private/ProfilingDebugging/MiscTrace.cpp
--- a/Engine/Source/Runtime/Core/Private/ProfilingDebugging/MiscTrace.cpp
+++ b/Engine/Source/Runtime/Core/Private/ProfilingDebugging/MiscTrace.cpp
@@ -21,6 +21,7 @@ UE_TRACE_EVENT_END()
 UE_TRACE_EVENT_BEGIN(Misc, Bookmark)
 	UE_TRACE_EVENT_FIELD(uint64, Cycle)
 	UE_TRACE_EVENT_FIELD(const void*, BookmarkPoint)
+	UE_TRACE_EVENT_FIELD(uint8[], FormatArgs)
 UE_TRACE_EVENT_END()
 
 UE_TRACE_EVENT_BEGIN(Misc, BeginFrame)
@@ -48,10 +49,10 @@ void FMiscTrace::OutputBookmarkSpec(const void* BookmarkPoint, const ANSICHAR* F
 
 void FMiscTrace::OutputBookmarkInternal(const void* BookmarkPoint, uint16 EncodedFormatArgsSize, uint8* EncodedFormatArgs)
 {
-	UE_TRACE_LOG(Misc, Bookmark, BookmarkChannel, EncodedFormatArgsSize)
+	UE_TRACE_LOG(Misc, Bookmark, BookmarkChannel)
 		<< Bookmark.Cycle(FPlatformTime::Cycles64())
 		<< Bookmark.BookmarkPoint(BookmarkPoint)
-		<< Bookmark.Attachment(EncodedFormatArgs, EncodedFormatArgsSize);
+		<< Bookmark.FormatArgs(EncodedFormatArgs, EncodedFormatArgsSize);
 }
 
 void FMiscTrace::OutputBeginFrame(ETraceFrameType FrameType)
