Updated UE_LOG()-traced events to emit the format arguments in an array-type field instead of as an attachment.

#rnx

diff --git a/Engine/Source/Developer/TraceServices/Private/Analyzers/LogTraceAnalysis.cpp b/Engine/Source/Developer/TraceServices/Private/Analyzers/LogTraceAnalysis.cpp
--- a/Engine/Source/Developer/TraceServices/Private/Analyzers/LogTraceAnalysis.cpp
+++ b/Engine/Source/Developer/TraceServices/Private/Analyzers/LogTraceAnalysis.cpp
@@ -74,7 +74,8 @@ bool FLogTraceAnalyzer::OnEvent(uint16 RouteId, EStyle Style, const FOnEventCont
 	{
 		uint64 LogPoint = EventData.GetValue<uint64>("LogPoint");
 		uint64 Cycle = EventData.GetValue<uint64>("Cycle");
-		LogProvider.AppendMessage(LogPoint, Context.EventTime.AsSeconds(Cycle), EventData.GetAttachment());
+		TArrayView<const uint8> FormatArgsView = FTraceAnalyzerUtils::LegacyAttachmentArray("FormatArgs", Context);
+		LogProvider.AppendMessage(LogPoint, Context.EventTime.AsSeconds(Cycle), FormatArgsView.GetData());
 		break;
 	}
 	}
diff --git a/Engine/Source/Runtime/Core/Private/Logging/LogTrace.cpp b/Engine/Source/Runtime/Core/Private/Logging/LogTrace.cpp
--- a/Engine/Source/Runtime/Core/Private/Logging/LogTrace.cpp
+++ b/Engine/Source/Runtime/Core/Private/Logging/LogTrace.cpp
@@ -28,6 +28,7 @@ UE_TRACE_EVENT_END()
 UE_TRACE_EVENT_BEGIN(Logging, LogMessage, NoSync)
 	UE_TRACE_EVENT_FIELD(const void*, LogPoint)
 	UE_TRACE_EVENT_FIELD(uint64, Cycle)
+	UE_TRACE_EVENT_FIELD(uint8[], FormatArgs)
 UE_TRACE_EVENT_END()
 
 void FLogTrace::OutputLogCategory(const FLogCategoryBase* Category, const TCHAR* Name, ELogVerbosity::Type DefaultVerbosity)
@@ -55,10 +56,10 @@ void FLogTrace::OutputLogMessageSpec(const void* LogPoint, const FLogCategoryBas
 
 void FLogTrace::OutputLogMessageInternal(const void* LogPoint, uint16 EncodedFormatArgsSize, uint8* EncodedFormatArgs)
 {
-	UE_TRACE_LOG(Logging, LogMessage, LogChannel, EncodedFormatArgsSize)
+	UE_TRACE_LOG(Logging, LogMessage, LogChannel)
 		<< LogMessage.LogPoint(LogPoint)
 		<< LogMessage.Cycle(FPlatformTime::Cycles64())
-		<< LogMessage.Attachment(EncodedFormatArgs, EncodedFormatArgsSize);
+		<< LogMessage.FormatArgs(EncodedFormatArgs, EncodedFormatArgsSize);
 }
 
 #endif
