Removed use of trace event attachments to send the packed CPU enter/exit timestamps in favour of an array-type field.

#rnx

diff --git a/Engine/Source/Developer/TraceServices/Private/Analyzers/CpuProfilerTraceAnalysis.cpp b/Engine/Source/Developer/TraceServices/Private/Analyzers/CpuProfilerTraceAnalysis.cpp
--- a/Engine/Source/Developer/TraceServices/Private/Analyzers/CpuProfilerTraceAnalysis.cpp
+++ b/Engine/Source/Developer/TraceServices/Private/Analyzers/CpuProfilerTraceAnalysis.cpp
@@ -86,9 +86,10 @@ bool FCpuProfilerAnalyzer::OnEvent(uint16 RouteId, EStyle Style, const FOnEventC
 	case RouteId_EventBatch:
 	case RouteId_EndCapture:
 	{
-		TotalEventSize += EventData.GetAttachmentSize();
-		uint32 BufferSize = EventData.GetAttachmentSize();
-		const uint8* BufferPtr = EventData.GetAttachment();
+		TArrayView<const uint8> DataView = FTraceAnalyzerUtils::LegacyAttachmentArray("Data", Context);
+		TotalEventSize += DataView.Num();
+		uint32 BufferSize = DataView.Num();
+		const uint8* BufferPtr = DataView.GetData();
 		uint32 ThreadId = FTraceAnalyzerUtils::GetThreadIdField(Context);
 		if (uint64 LastCycle = ProcessBuffer(Context.EventTime, ThreadId, BufferPtr, BufferSize))
 		{
diff --git a/Engine/Source/Runtime/Core/Private/ProfilingDebugging/CpuProfilerTrace.cpp b/Engine/Source/Runtime/Core/Private/ProfilingDebugging/CpuProfilerTrace.cpp
--- a/Engine/Source/Runtime/Core/Private/ProfilingDebugging/CpuProfilerTrace.cpp
+++ b/Engine/Source/Runtime/Core/Private/ProfilingDebugging/CpuProfilerTrace.cpp
@@ -21,9 +21,11 @@ UE_TRACE_EVENT_BEGIN(CpuProfiler, EventSpec, NoSync|Important)
 UE_TRACE_EVENT_END()
 
 UE_TRACE_EVENT_BEGIN(CpuProfiler, EventBatch, NoSync)
+	UE_TRACE_EVENT_FIELD(uint8[], Data)
 UE_TRACE_EVENT_END()
 
 UE_TRACE_EVENT_BEGIN(CpuProfiler, EndCapture)
+	UE_TRACE_EVENT_FIELD(uint8[], Data)
 UE_TRACE_EVENT_END()
 
 UE_TRACE_EVENT_BEGIN(CpuProfiler, EndThread, NoSync)
@@ -105,16 +107,16 @@ FCpuProfilerTraceInternal::FThreadBuffer* FCpuProfilerTraceInternal::CreateThrea
 
 void FCpuProfilerTraceInternal::FlushThreadBuffer(FThreadBuffer* InThreadBuffer)
 {
-	UE_TRACE_LOG(CpuProfiler, EventBatch, true, InThreadBuffer->BufferSize)
-		<< EventBatch.Attachment(InThreadBuffer->Buffer, InThreadBuffer->BufferSize);
+	UE_TRACE_LOG(CpuProfiler, EventBatch, true)
+		<< EventBatch.Data(InThreadBuffer->Buffer, InThreadBuffer->BufferSize);
 	InThreadBuffer->BufferSize = 0;
 	InThreadBuffer->LastCycle = 0;
 }
 
 void FCpuProfilerTraceInternal::EndCapture(FThreadBuffer* InThreadBuffer)
 {
-	UE_TRACE_LOG(CpuProfiler, EndCapture, true, InThreadBuffer->BufferSize)
-		<< EndCapture.Attachment(InThreadBuffer->Buffer, InThreadBuffer->BufferSize);
+	UE_TRACE_LOG(CpuProfiler, EndCapture, true)
+		<< EndCapture.Data(InThreadBuffer->Buffer, InThreadBuffer->BufferSize);
 	InThreadBuffer->BufferSize = 0;
 	InThreadBuffer->LastCycle = 0;
 }
