Made UnixCrashContext for ensures thread-safe
* Stop adding/removing FTicker delegates outside the game-thread
* Made the concurrently uploading processes thread-safe using atomic state machine
* Added option to limit the number of processes on dedicated game servers
* Fix ensures not getting sent on dedicated game servers

diff --git a/Engine/Source/Runtime/Core/Private/Unix/UnixPlatformCrashContext.cpp b/Engine/Source/Runtime/Core/Private/Unix/UnixPlatformCrashContext.cpp
--- a/Engine/Source/Runtime/Core/Private/Unix/UnixPlatformCrashContext.cpp
+++ b/Engine/Source/Runtime/Core/Private/Unix/UnixPlatformCrashContext.cpp
@@ -32,6 +32,8 @@
 #include "HAL/ThreadHeartBeat.h"
 #include "BuildSettings.h"
 
+#include <atomic>
+
 extern CORE_API bool GIsGPUCrashed;
 
 FString DescribeSignal(int32 Signal, siginfo_t* Info, ucontext_t *Context)
@@ -334,53 +336,94 @@ void FUnixCrashContext::GetPortableCallStack(const uint64* StackFrames, int32 Nu
 	}
 }
 
+#ifndef SERVER_MAX_CONCURRENT_REPORTS
+	#define SERVER_MAX_CONCURRENT_REPORTS 1
+#endif
+
 namespace UnixCrashReporterTracker
 {
+	enum class SlotStatus : uint32
+	{
+		Available,
+		Spawning,
+		Uploading,
+		Closing,
+		Killing
+	};
+
 	struct CrashReporterProcess
 	{
 		FProcHandle Process;
-		bool bIsSpawned = false;
+		std::atomic<UnixCrashReporterTracker::SlotStatus> Status;
 	};
 
 	// Matching MaxPreviousErrorsToTrack = 4 in AssertionMacros.cpp
 	CrashReporterProcess Processes[4];
-	FDelegateHandle CurrentTicker;
+
+	/** Maximum index in the process slot array */
+	uint32 MaxProcessSlots;
+
+	/** Number of active processes uploading their data at the moment */
+	std::atomic<uint32> NumUploadingProcesses(0);
 
 	bool Tick(float DeltaTime)
 	{
 		QUICK_SCOPE_CYCLE_COUNTER(STAT_UnixCrashReporterTracker_Tick);
 
-		int ProcessLeft = 0;
-		for (int ProcessNumber = 0; ProcessNumber < UE_ARRAY_COUNT(UnixCrashReporterTracker::Processes); ProcessNumber++)
+		uint32 NumActiveProcess = NumUploadingProcesses.load(std::memory_order_relaxed);
+
+		if (NumActiveProcess > 0)
 		{
-			if (UnixCrashReporterTracker::Processes[ProcessNumber].bIsSpawned)
+			for (uint32 ProcessNumber = 0; ProcessNumber < MaxProcessSlots; ++ProcessNumber)
 			{
-				if (!FPlatformProcess::IsProcRunning(UnixCrashReporterTracker::Processes[ProcessNumber].Process))
-				{
-					UnixCrashReporterTracker::Processes[ProcessNumber].bIsSpawned = false;
+				CrashReporterProcess& CurrentSlot = Processes[ProcessNumber];
 
-					FPlatformProcess::CloseProc(UnixCrashReporterTracker::Processes[ProcessNumber].Process);
-					UnixCrashReporterTracker::Processes[ProcessNumber].Process = FProcHandle();
-				}
-				else
+				SlotStatus IsUploading = SlotStatus::Uploading;
+				const SlotStatus IsClosing = SlotStatus::Closing;
+
+				// Test if an uploading process is finished
+				if (CurrentSlot.Status.compare_exchange_weak(IsUploading, IsClosing))
 				{
-					ProcessLeft++;
+					if (!FPlatformProcess::IsProcRunning(CurrentSlot.Process))
+					{
+						FPlatformProcess::CloseProc(CurrentSlot.Process);
+						CurrentSlot.Process = FProcHandle();
+
+						--NumUploadingProcesses;
+						CurrentSlot.Status = SlotStatus::Available;
+					}
+					else
+					{
+						CurrentSlot.Status = SlotStatus::Uploading;
+					}
 				}
 			}
 		}
 
-		if (ProcessLeft <= 0)
+		// tick again
+		return true;
+	}
+
+	void PreInit()
+	{
+		for (CrashReporterProcess& CurrentSlot : Processes)
 		{
-			// clean up ticker
-			FTicker::GetCoreTicker().RemoveTicker(CurrentTicker);
-			CurrentTicker.Reset();
+			CurrentSlot.Status = SlotStatus::Available;
+		}
+
+		uint32 ActiveProcessSlots = UE_ARRAY_COUNT(Processes);
 
-			UE_LOG(LogCore, Log, TEXT("Done sending crash report for ensure()."));
-			return false;
+		// Lower the amount of concurrent reports on servers to limit the spike in cpu/memory when sending those reports
+		if (IsRunningDedicatedServer())
+		{
+			ActiveProcessSlots = FMath::Min((uint32)SERVER_MAX_CONCURRENT_REPORTS, ActiveProcessSlots);
 		}
 
-		// tick again
-		return true;
+		// Set the valid max index to iterate over
+		UnixCrashReporterTracker::MaxProcessSlots = ActiveProcessSlots;
+
+        // Register our Tick function
+		FTicker::GetCoreTicker().AddTicker(FTickerDelegate::CreateStatic(&UnixCrashReporterTracker::Tick), 1.f);
 	}
 
 	/**
@@ -415,16 +458,13 @@ namespace UnixCrashReporterTracker
 
 	void RemoveValidCrashReportTickerForChildProcess()
 	{
-		if (CurrentTicker.IsValid())
+		for (uint32 ProcessNumber = 0; ProcessNumber < UnixCrashReporterTracker::MaxProcessSlots; ++ProcessNumber)
 		{
-			FTicker::GetCoreTicker().RemoveTicker(CurrentTicker);
-			CurrentTicker.Reset();
-
-			for (int ProcessNumber = 0; ProcessNumber < UE_ARRAY_COUNT(UnixCrashReporterTracker::Processes); ProcessNumber++)
-			{
-				UnixCrashReporterTracker::Processes[ProcessNumber].Process = FProcHandle();
-			}
+			UnixCrashReporterTracker::Processes[ProcessNumber].Process = FProcHandle();
+			UnixCrashReporterTracker::Processes[ProcessNumber].Status = UnixCrashReporterTracker::SlotStatus::Available;
 		}
+
+		UnixCrashReporterTracker::NumUploadingProcesses = 0;
 	}
 }
 
@@ -667,58 +707,70 @@ void FUnixCrashContext::GenerateCrashInfoAndLaunchReporter(bool bReportingNonCra
 
 			CrashReportClientArguments += TEXT("\"\"") + CrashInfoAbsolute + TEXT("/\"\"");
 
-			// FIXME: this code is not attempting to be thread safe
 			if (bReportingNonCrash)
 			{
-				{
-					bool bFoundEmptySlot = false;
+				bool bFoundEmptySlot = false;
 
-					const double kEnsureTimeOut = 45.0;
-					const double kEnsureSleepInterval = 0.1;
-					double kTimeOutTimer = 0.0;
+				constexpr double kEnsureTimeOut = 45.0;
+				constexpr double kEnsureSleepInterval = 0.1;
+				double kTimeOutTimer = 0.0;
 
-					while (!bFoundEmptySlot && !IsRunningDedicatedServer())
+				while (!bFoundEmptySlot)
+				{
+					// Find an empty slot for sending the report
+					for( uint32 ProcessIdx=0; ProcessIdx < UnixCrashReporterTracker::MaxProcessSlots; ++ProcessIdx )
 					{
-						// case 1 - no process handle in use: spawn ticker and a process
-						if (!UnixCrashReporterTracker::CurrentTicker.IsValid())
+						UnixCrashReporterTracker::SlotStatus IsAvailable = UnixCrashReporterTracker::SlotStatus::Available;
+						const UnixCrashReporterTracker::SlotStatus IsSpawning = UnixCrashReporterTracker::SlotStatus::Spawning;
+						
+						// If the slot is available, get exclusive rights by setting it to spawning state
+						if (UnixCrashReporterTracker::Processes[ProcessIdx].Status.compare_exchange_weak(IsAvailable, IsSpawning))
 						{
-							UnixCrashReporterTracker::CurrentTicker = FTicker::GetCoreTicker().AddTicker(FTickerDelegate::CreateStatic(&UnixCrashReporterTracker::Tick), 1.f);
-
-							// assume if no ticker all slots are open, pick the first one
-							UnixCrashReporterTracker::Processes[0].bIsSpawned = true;
-							UnixCrashReporterTracker::Processes[0].Process = FPlatformProcess::CreateProc(
-									RelativePathToCrashReporter, *CrashReportClientArguments, true, false, false, NULL, 0, NULL, NULL);
+							UnixCrashReporterTracker::Processes[ProcessIdx].Process = FPlatformProcess::CreateProc(
+								RelativePathToCrashReporter, *CrashReportClientArguments, true, false, false, NULL, 0, NULL, NULL);
+							
+							++UnixCrashReporterTracker::NumUploadingProcesses;
+							UnixCrashReporterTracker::Processes[ProcessIdx].Status = UnixCrashReporterTracker::SlotStatus::Uploading;
 
 							bFoundEmptySlot = true;
+							break;
+						}
+					}
+					
+					// All process handles in use: wait for up to 45 seconds on the next available process
+					if (!bFoundEmptySlot)
+					{
+						// If all slots are uploading on a server, skip the report instead of hitching
+						if (IsRunningDedicatedServer())
+						{
+							UE_LOG(LogCore, Warning, TEXT("Too many reports already in progress, skipping upload of this one."));
+							bFoundEmptySlot = true;
 						}
-						// case 2 - some process handles in use: spawn a process if available 
 						else
 						{
-							for (int ProcessNumber = 0; ProcessNumber < UE_ARRAY_COUNT(UnixCrashReporterTracker::Processes); ProcessNumber++)
-							{
-								if (!UnixCrashReporterTracker::Processes[ProcessNumber].bIsSpawned)
-								{
-									bFoundEmptySlot = true;
-									UnixCrashReporterTracker::Processes[ProcessNumber].bIsSpawned = true;
-									UnixCrashReporterTracker::Processes[ProcessNumber].Process = FPlatformProcess::CreateProc(
-											RelativePathToCrashReporter, *CrashReportClientArguments, true, false, false, NULL, 0, NULL, NULL);
-									break;
-								}
-							}
+							FPlatformProcess::Sleep(kEnsureSleepInterval);
+							UnixCrashReporterTracker::Tick(0.001f);
 
-							// case 3 - all process handles in use: wait for up to 45 seconds on the next available process
-							if (!bFoundEmptySlot)
+							kTimeOutTimer += kEnsureSleepInterval;
+
+							// After waiting 45seconds, kill the process at slot 0 and lose it's information.
+							if (kTimeOutTimer >= kEnsureTimeOut)
 							{
-								UnixCrashReporterTracker::Tick(0.001f);
-								FPlatformProcess::Sleep(kEnsureSleepInterval);	
-								kTimeOutTimer += kEnsureSleepInterval;
+								UnixCrashReporterTracker::SlotStatus IsUploading = UnixCrashReporterTracker::SlotStatus::Uploading;
+								const UnixCrashReporterTracker::SlotStatus IsKilling = UnixCrashReporterTracker::SlotStatus::Killing;
 
-								// Could lose information by killing an unfinished CrashReportClient instance
-								if (kTimeOutTimer >= kEnsureTimeOut)
+								// Take over this uploading slot and kill it
+								if (UnixCrashReporterTracker::Processes[0].Status.compare_exchange_weak(IsUploading, IsKilling))
 								{
-									UnixCrashReporterTracker::Processes[0].bIsSpawned = false;
+									UE_LOG(LogCore, Warning, TEXT("Terminated CrashReport process[0]"));
+
 									FPlatformProcess::TerminateProc(UnixCrashReporterTracker::Processes[0].Process);
+									UnixCrashReporterTracker::Processes[0].Process = FProcHandle();
+
+									--UnixCrashReporterTracker::NumUploadingProcesses;
+									UnixCrashReporterTracker::Processes[0].Status = UnixCrashReporterTracker::SlotStatus::Available;
 								}
+								
 							}
 						}
 					}
@@ -754,6 +806,10 @@ void FUnixCrashContext::GenerateCrashInfoAndLaunchReporter(bool bReportingNonCra
 				}
 			}
 		}
+		else
+		{
+			UE_LOG(LogCore, Warning, TEXT("MakeDirectory %s failed"), *CrashInfoAbsolute);
+		}
 	}
 
 #endif
diff --git a/Engine/Source/Runtime/Core/Private/Unix/UnixPlatformMisc.cpp b/Engine/Source/Runtime/Core/Private/Unix/UnixPlatformMisc.cpp
--- a/Engine/Source/Runtime/Core/Private/Unix/UnixPlatformMisc.cpp
+++ b/Engine/Source/Runtime/Core/Private/Unix/UnixPlatformMisc.cpp
@@ -180,6 +180,13 @@ extern bool CORE_API GUseNewCrashSymbolicator;
 // Helps greatly by reducing I/O during ensures/crashes. Only suggested when running a monolithic build
 extern void CORE_API UnixPlatformStackWalk_PreloadModuleSymbolFile();
 
+void FUnixPlatformMisc::PlatformPreInit()
+{
+	FGenericPlatformMisc::PlatformPreInit();
+
+	UnixCrashReporterTracker::PreInit();
+}
+
 void FUnixPlatformMisc::PlatformInit()
 {
 	// install a platform-specific signal handler
diff --git a/Engine/Source/Runtime/Core/Public/Unix/UnixPlatformCrashContext.h b/Engine/Source/Runtime/Core/Public/Unix/UnixPlatformCrashContext.h
--- a/Engine/Source/Runtime/Core/Public/Unix/UnixPlatformCrashContext.h
+++ b/Engine/Source/Runtime/Core/Public/Unix/UnixPlatformCrashContext.h
@@ -128,6 +128,11 @@ typedef FUnixCrashContext FPlatformCrashContext;
 
 namespace UnixCrashReporterTracker
 {
+	/**
+	 * Initialize persistent data and tickers
+	 */
+	void PreInit();
+
 	/**
 	 * Only call this function from a forked child process. The child process cannot be responsible for a sibling process.
 	 *
diff --git a/Engine/Source/Runtime/Core/Public/Unix/UnixPlatformMisc.h b/Engine/Source/Runtime/Core/Public/Unix/UnixPlatformMisc.h
--- a/Engine/Source/Runtime/Core/Public/Unix/UnixPlatformMisc.h
+++ b/Engine/Source/Runtime/Core/Public/Unix/UnixPlatformMisc.h
@@ -21,6 +21,7 @@ struct FGenericCrashContext;
  */
 struct CORE_API FUnixPlatformMisc : public FGenericPlatformMisc
 {
+	static void PlatformPreInit();
 	static void PlatformInit();
 	static void PlatformTearDown();
 	static void SetGracefulTerminationHandler();
