Fix non unity compile error

diff --git a/Engine/Source/Runtime/Core/Private/ProfilingDebugging/TraceAuxiliary.cpp b/Engine/Source/Runtime/Core/Private/ProfilingDebugging/TraceAuxiliary.cpp
--- a/Engine/Source/Runtime/Core/Private/ProfilingDebugging/TraceAuxiliary.cpp
+++ b/Engine/Source/Runtime/Core/Private/ProfilingDebugging/TraceAuxiliary.cpp
@@ -17,6 +17,10 @@
 #include "CoreGlobals.h"
 #include "Misc/Paths.h"
 
+#if WITH_UNREAL_TRACE_LAUNCH
+#	include "Misc/Parse.h"
+#endif
+
 #include <atomic>
 
 #if UE_TRACE_ENABLED
