Remove unused UE3 define

#rnx

diff --git a/Engine/Source/Runtime/Core/Public/Misc/CoreMiscDefines.h b/Engine/Source/Runtime/Core/Public/Misc/CoreMiscDefines.h
--- a/Engine/Source/Runtime/Core/Public/Misc/CoreMiscDefines.h
+++ b/Engine/Source/Runtime/Core/Public/Misc/CoreMiscDefines.h
@@ -14,7 +14,6 @@
 // rest of code out for resource compilation.
 #ifndef RC_INVOKED
 
-#define LOCALIZED_SEEKFREE_SUFFIX	TEXT("_LOC")
 #define PLAYWORLD_PACKAGE_PREFIX TEXT("UEDPIE")
 
 #ifndef WITH_EDITORONLY_DATA
