Fixed non vectorized LWC build

diff --git a/Engine/Source/Runtime/Core/Public/Math/TransformNonVectorized.h b/Engine/Source/Runtime/Core/Public/Math/TransformNonVectorized.h
--- a/Engine/Source/Runtime/Core/Public/Math/TransformNonVectorized.h
+++ b/Engine/Source/Runtime/Core/Public/Math/TransformNonVectorized.h
@@ -378,9 +378,9 @@ public:
 		// Make sure Rotation is normalized when we turn it into a matrix.
 		check(IsRotationNormalized());
 #endif
-		OutMatrix.M[3][0] = Translation.X;
-		OutMatrix.M[3][1] = Translation.Y;
-		OutMatrix.M[3][2] = Translation.Z;
+		OutMatrix.M[3][0] = (float)Translation.X;
+		OutMatrix.M[3][1] = (float)Translation.Y;
+		OutMatrix.M[3][2] = (float)Translation.Z;
 
 		const float x2 = Rotation.X + Rotation.X;
 		const float y2 = Rotation.Y + Rotation.Y;
@@ -634,7 +634,7 @@ public:
 	/** Calculate the determinant of this transformation */
 	FORCEINLINE float GetDeterminant() const
 	{
-		return Scale3D.X * Scale3D.Y * Scale3D.Z;
+		return (float)(Scale3D.X * Scale3D.Y * Scale3D.Z);
 	}
 
 	/** Set the translation of this transformation */
@@ -1561,7 +1561,7 @@ FORCEINLINE void FTransform::Mirror(EAxis::Type MirrorAxis, EAxis::Type FlipAxis
 inline float FTransform::GetMaximumAxisScale() const
 {
 	DiagnosticCheckNaN_Scale3D();
-	return Scale3D.GetAbsMax();
+	return (float)Scale3D.GetAbsMax();
 }
 
 
@@ -1569,7 +1569,7 @@ inline float FTransform::GetMaximumAxisScale() const
 inline float FTransform::GetMinimumAxisScale() const
 {
 	DiagnosticCheckNaN_Scale3D();
-	return Scale3D.GetAbsMin();
+	return (float)Scale3D.GetAbsMin();
 }
 
 
