Use a more reasonable buffer size to reduce number of I/O on Windows for the editor
  - Use 256K precache size instead of the default of 1024, which was way too low for editor workflows
  - Any reduction in IO translate to less kernel transitions, which is critical now more than ever because of meltdown mitigations and filter drivers for anti-malware intercepting every syscall
  - 6782676 reads 28.9s -> 105762 reads 5.7s when loading P_Construct_WP + PIE + Pressing 4 to start the level sequence on a special project

#rnx

diff --git a/Engine/Source/Runtime/Core/Public/Windows/WindowsPlatform.h b/Engine/Source/Runtime/Core/Public/Windows/WindowsPlatform.h
--- a/Engine/Source/Runtime/Core/Public/Windows/WindowsPlatform.h
+++ b/Engine/Source/Runtime/Core/Public/Windows/WindowsPlatform.h
@@ -85,6 +85,10 @@ typedef FWindowsPlatformTypes FPlatformTypes;
 #define PLATFORM_SUPPORTS_VARIABLE_RATE_SHADING				1
 #define PLATFORM_SUPPORTS_MESH_SHADERS						1
 
+#if WITH_EDITOR
+#define PLATFORM_FILE_READER_BUFFER_SIZE					(256*1024)
+#endif
+
 #define PLATFORM_SUPPORTS_STACK_SYMBOLS						1
 #define PLATFORM_COMPILER_HAS_DECLTYPE_AUTO					1
 
