FTSTicker - a thread-safe version of FTicker: allows adding and removing tick delegates concurrently. Removing a delegate can block until its execution is finishing in parallel.

diff --git a/Engine/Source/Runtime/Core/Private/Containers/BackgroundableTicker.cpp b/Engine/Source/Runtime/Core/Private/Containers/BackgroundableTicker.cpp
--- a/Engine/Source/Runtime/Core/Private/Containers/BackgroundableTicker.cpp
+++ b/Engine/Source/Runtime/Core/Private/Containers/BackgroundableTicker.cpp
@@ -4,6 +4,43 @@
 #include "Misc/CoreDelegates.h"
 #include "Stats/Stats.h"
 
+FTSBackgroundableTicker& FTSBackgroundableTicker::GetCoreTicker()
+{
+	static FTSBackgroundableTicker Singleton;
+	return Singleton;
+}
+
+FTSBackgroundableTicker::FTSBackgroundableTicker()
+{
+	CoreTickerHandle = FTSTicker::GetCoreTicker().AddTicker(TEXT("FBackgroundableTicker"), 0.0f, [this](float DeltaTime) -> bool
+		{
+			QUICK_SCOPE_CYCLE_COUNTER(STAT_FBackgroundableTicker_ForegroundTick);
+			if (bWasBackgrounded)
+			{
+				// When returning to the foreground, ensure we do not report enormous delta time values coming from the foreground ticker.
+				bWasBackgrounded = false;
+				DeltaTime = FMath::Clamp(DeltaTime, 0.0f, 1.0f / 60.0f);
+			}
+			Tick(DeltaTime);
+			return true;
+		});
+	BackgroundTickerHandle = FCoreDelegates::MobileBackgroundTickDelegate.AddLambda([this](float DeltaTime)
+		{
+			QUICK_SCOPE_CYCLE_COUNTER(STAT_FBackgroundableTicker_BackgroundTick);
+			bWasBackgrounded = true;
+			Tick(DeltaTime);
+		});
+}
+
+FTSBackgroundableTicker::~FTSBackgroundableTicker()
+{
+	FTSTicker::GetCoreTicker().RemoveTicker(CoreTickerHandle);
+	FCoreDelegates::MobileBackgroundTickDelegate.Remove(BackgroundTickerHandle);
+}
+
+//////////////////////////////////////////////////////////////////////////////////////////////////////////////////////////////////////
+// deprecated non thread-safe version
+
 FBackgroundableTicker& FBackgroundableTicker::GetCoreTicker()
 {
 	static FBackgroundableTicker Singleton;
diff --git a/Engine/Source/Runtime/Core/Private/Containers/Ticker.cpp b/Engine/Source/Runtime/Core/Private/Containers/Ticker.cpp
--- a/Engine/Source/Runtime/Core/Private/Containers/Ticker.cpp
+++ b/Engine/Source/Runtime/Core/Private/Containers/Ticker.cpp
@@ -4,6 +4,152 @@
 #include "Stats/Stats.h"
 #include "Misc/TimeGuard.h"
 
+FTSTicker& FTSTicker::GetCoreTicker()
+{
+	static FTSTicker CoreTicker;
+	return CoreTicker;
+}
+
+FTSTicker::FDelegateHandle FTSTicker::AddTicker(const FTickerDelegate& InDelegate, float InDelay)
+{
+	FElementPtr NewElement{ new FElement{ CurrentTime + InDelay, InDelay, InDelegate } };
+	AddedElements.Enqueue(NewElement);
+	return NewElement;
+}
+
+FTSTicker::FDelegateHandle FTSTicker::AddTicker(const TCHAR* InName, float InDelay, TFunction<bool(float)> Function)
+{
+	FElementPtr NewElement{ new FElement{ CurrentTime + InDelay, InDelay, FTickerDelegate::CreateLambda(Function) } };
+	AddedElements.Enqueue(NewElement);
+	return NewElement;
+}
+
+void FTSTicker::RemoveTicker(FDelegateHandle Handle)
+{
+	if (FElementPtr Element = Handle.Pin())
+	{
+		// mark the element as removed and if it's being ticked atm, spin-wait until its execution is finished
+		uint8 PrevState = Element->State.fetch_add(FElement::Removed, std::memory_order_acquire); // "acquire" to prevent potential 
+		// resource release after RemoveTicker() to be reordered before it
+		checkf(PrevState < FElement::Removed, TEXT("The delegate is already removed (%u)"), PrevState);
+		while (PrevState != FElement::Idle && PrevState != FElement::Removed) // is not being executed
+		{
+			FPlatformProcess::Yield();
+			PrevState = Element->State.load(std::memory_order_relaxed);
+		}
+	}
+}
+
+void FTSTicker::Tick(float DeltaTime)
+{
+	SCOPE_TIME_GUARD(TEXT("FTSTicker::Tick"));
+
+	QUICK_SCOPE_CYCLE_COUNTER(STAT_FTicker_Tick);
+
+	// take in all new elements
+	while (TOptional<FElementPtr> AddedElement = AddedElements.Dequeue())
+	{
+		uint8 State = AddedElement.GetValue()->State.load(std::memory_order_relaxed);
+		checkf(State == FElement::Idle || State == FElement::Removed, TEXT("Invalid state %u"), State);
+		if (State == FElement::Idle)
+		{
+			Elements.Add(MoveTemp(AddedElement.GetValue()));
+		}
+	}
+
+	if (!Elements.Num())
+	{
+		return;
+	}
+
+	CurrentTime += DeltaTime;
+
+	TArray<FElementPtr> TickedElements;
+	for (FElementPtr& Element : Elements)
+	{
+		// set the execution flag
+		uint8 PrevState = Element->State.fetch_add(FElement::Executing, std::memory_order_acquire); // "acquire" to prevent anything between this and `ClearExecutionFlag()` to be reordered outside of this scope, is coupled with "release" in `ClearExecutionFlag`
+		checkf(PrevState == FElement::Idle || PrevState == FElement::Removed, TEXT("Invalid state %u"), PrevState);
+		auto ClearExecutionFlag = [](FElementPtr& Element) { return (uint8)(Element->State.fetch_sub(FElement::Executing, std::memory_order_release) - FElement::Executing); };
+
+		if (PrevState == FElement::Removed)
+		{
+			PrevState = ClearExecutionFlag(Element);
+			continue;
+		}
+
+		checkf(PrevState == FElement::Idle, TEXT("Invalid state %u"), PrevState);
+
+		if (Element->FireTime > CurrentTime)
+		{
+			ClearExecutionFlag(Element);
+			TickedElements.Add(MoveTemp(Element));
+			continue;
+		}
+
+		if (Element->Fire(DeltaTime))
+		{
+			PrevState = ClearExecutionFlag(Element); // it can be removed during execution
+			if (PrevState != FElement::Removed)
+			{
+				checkf(PrevState == FElement::Idle, TEXT("Invalid state %u"), PrevState);
+				Element->FireTime = CurrentTime + Element->DelayTime;
+				TickedElements.Add(MoveTemp(Element));
+			}
+		}
+		else
+		{
+			PrevState = ClearExecutionFlag(Element);
+			checkf(PrevState == FElement::Idle || PrevState == FElement::Removed, TEXT("Invalid state %u"), PrevState);
+		}
+	}
+
+	Elements.Reset();
+	Exchange(TickedElements, Elements);
+
+	if (this == &GetCoreTicker())
+	{
+		// tick also deprecated FTicker for backward compatibility
+		PRAGMA_DISABLE_DEPRECATION_WARNINGS
+		FTicker::GetCoreTicker().Tick(DeltaTime);
+		PRAGMA_ENABLE_DEPRECATION_WARNINGS
+	}
+}
+
+FTSTicker::FElement::FElement()
+	: FireTime(0.0)
+	, DelayTime(0.0f)
+{}
+
+FTSTicker::FElement::FElement(double InFireTime, float InDelayTime, const FTickerDelegate& InDelegate)
+	: FireTime(InFireTime)
+	, DelayTime(InDelayTime)
+	, Delegate(InDelegate)
+{}
+
+bool FTSTicker::FElement::Fire(float DeltaTime)
+{
+	check(Delegate.IsBound());
+	return Delegate.Execute(DeltaTime);
+}
+
+FTSTickerObjectBase::FTSTickerObjectBase(float InDelay, FTSTicker& InTicker)
+{
+	// Register delegate for ticker callback
+	FTickerDelegate TickDelegate = FTickerDelegate::CreateRaw(this, &FTSTickerObjectBase::Tick);
+	TickHandle = InTicker.AddTicker(TickDelegate, InDelay);
+}
+
+FTSTickerObjectBase::~FTSTickerObjectBase()
+{
+	FTSTicker::RemoveTicker(TickHandle);
+}
+
+//////////////////////////////////////////////////////////////////////////////////////////////////////////////////////////////////////
+// deprecated non thread-safe version
+
+PRAGMA_DISABLE_DEPRECATION_WARNINGS
+
 FTicker::FTicker()
 	: CurrentTime(0.0)
 	, bInTick(false)
@@ -27,6 +173,8 @@ FDelegateHandle FTicker::AddTicker(const FTickerDelegate& InDelegate, float InDe
 
 FDelegateHandle FTicker::AddTicker(const TCHAR* InName, float InDelay, TFunction<bool(float)> Function)
 {
+	checkf(IsInGameThread(), TEXT("FTicker is not thread-safe and will be deprecated, please use FTSTicker"));
+
 	// todo - use InName for profiling. Added in sig to be forward looking..
 
 	FTickerDelegate Delegate = FTickerDelegate::CreateLambda(Function);
@@ -39,6 +187,8 @@ FDelegateHandle FTicker::AddTicker(const TCHAR* InName, float InDelay, TFunction
 
 void FTicker::RemoveTicker(FDelegateHandle Handle)
 {
+	checkf(IsInGameThread(), TEXT("FTicker is not thread-safe and will be deprecated, please use FTSTicker"));
+
 	// must remove the handle from both arrays because we could be in the middle of a tick, 
 	// and may be removing ourselves or something else we've already considered this Tick
 	auto CompareHandle = [=](const FElement& Element){ return Element.Delegate.GetHandle() == Handle; };
@@ -155,3 +305,5 @@ FTickerObjectBase::~FTickerObjectBase()
 		TickHandle = FDelegateHandle();
 	}
 }
+
+PRAGMA_ENABLE_DEPRECATION_WARNINGS
\ No newline at end of file
diff --git a/Engine/Source/Runtime/Core/Private/Tests/Containers/TickerTests.cpp b/Engine/Source/Runtime/Core/Private/Tests/Containers/TickerTests.cpp
new file mode 100644
--- /dev/null
+++ b/Engine/Source/Runtime/Core/Private/Tests/Containers/TickerTests.cpp
@@ -0,0 +1,192 @@
+// Copyright Epic Games, Inc. All Rights Reserved.
+
+#include "CoreTypes.h"
+#include "Misc/AutomationTest.h"
+#include "Tests/Benchmark.h"
+#include "Containers/Ticker.h"
+#include "Tasks/Task.h"
+#include "Tests/Benchmark.h"
+
+#if WITH_DEV_AUTOMATION_TESTS
+
+IMPLEMENT_SIMPLE_AUTOMATION_TEST(FTSTickerTest, "System.Core.Containers.TSTicker", EAutomationTestFlags::ApplicationContextMask | EAutomationTestFlags::EngineFilter);
+
+template<uint32 NumDelegates, uint32 NumTicks>
+void TickerPerfTest()
+{
+	FTicker Ticker;
+
+	TArray<FDelegateHandle> DelegateHandles;
+	DelegateHandles.Reserve(NumDelegates);
+	for (uint32 i = 0; i != NumDelegates; ++i)
+	{
+		DelegateHandles.Add(Ticker.AddTicker(UE_SOURCE_LOCATION, 0.0f, [](float DeltaTime) { return true; }));
+	}
+
+	for (uint32 i = 0; i != NumTicks; ++i)
+	{
+		Ticker.Tick(0.0f);
+	}
+
+	for (FDelegateHandle& DelegateHandle : DelegateHandles)
+	{
+		Ticker.RemoveTicker(DelegateHandle);
+	}
+}
+
+template<uint32 NumDelegates, uint32 NumTicks>
+void TSTickerPerfTest()
+{
+	FTSTicker Ticker;
+
+	TArray<FTSTicker::FDelegateHandle> DelegateHandles;
+	DelegateHandles.Reserve(NumDelegates);
+	for (uint32 i = 0; i != NumDelegates; ++i)
+	{
+		DelegateHandles.Add(Ticker.AddTicker(UE_SOURCE_LOCATION, 0.0f, [](float DeltaTime) { return true; }));
+	}
+
+	for (uint32 i = 0; i != NumTicks; ++i)
+	{
+		Ticker.Tick(0.0f);
+	}
+
+	for (FTSTicker::FDelegateHandle& DelegateHandle : DelegateHandles)
+	{
+		FTSTicker::RemoveTicker(DelegateHandle);
+	}
+}
+
+bool FTSTickerTest::RunTest(const FString& Parameters)
+{
+	{	// a delegate returning false is executed once
+		FTSTicker Ticker;
+		bool bExecuted = false;
+		FTSTicker::FDelegateHandle DelegateHandle = Ticker.AddTicker(UE_SOURCE_LOCATION, 0.0f,
+			[&bExecuted](float DeltaTime)
+			{ 
+				check(!bExecuted);
+				bExecuted = true; 
+				return false; 
+			}
+		);
+		Ticker.Tick(0.0f);
+		Ticker.Tick(0.0f);
+		check(bExecuted);
+		Ticker.RemoveTicker(DelegateHandle);
+	}
+
+	{	// a delegate returning true is executed multiple times
+		FTSTicker Ticker;
+		uint32 NumExecuted = 0;
+		FTSTicker::FDelegateHandle DelegateHandle = Ticker.AddTicker(UE_SOURCE_LOCATION, 0.0f,
+			[&NumExecuted](float DeltaTime)
+			{
+				++NumExecuted;
+				return true;
+			}
+		);
+		Ticker.Tick(0.0f);
+		Ticker.Tick(0.0f);
+		check(NumExecuted == 2);
+		Ticker.RemoveTicker(DelegateHandle);
+	}
+
+	{	// a delegate removal while it's being ticked doesn't return until its execution finished
+		using namespace UE::Tasks;
+		
+		FTSTicker Ticker;
+
+		FTaskEvent DelegateResumeEvent{ UE_SOURCE_LOCATION };
+		FTSTicker::FDelegateHandle DelegateHandle = Ticker.AddTicker(UE_SOURCE_LOCATION, 0.0f, 
+			[&DelegateResumeEvent](float DeltaTime)
+			{ 
+				DelegateResumeEvent.Wait();
+				return false;
+			}
+		);
+
+		FTask RemoveTickerTask = Launch(UE_SOURCE_LOCATION, 
+			[&Ticker, &DelegateHandle] 
+			{ 
+				FPlatformProcess::Sleep(0.1); // let the ticking start and the delegate block on the event
+				Ticker.RemoveTicker(DelegateHandle); 
+			}
+		);
+
+		FTask TickTask = Launch(UE_SOURCE_LOCATION, 
+			[&Ticker] 
+			{ 
+				Ticker.Tick(0.0); 
+			}
+		);
+
+		FPlatformProcess::Sleep(0.1); // let workers pick up the tasks and start execution
+
+		verify(!TickTask.Wait(FTimespan::FromSeconds(0.1))); // tick is blocked because the delegate is blocked
+		verify(!RemoveTickerTask.Wait(FTimespan::FromSeconds(0.1))); // removal is blocked because the delegate is blocked
+
+		DelegateResumeEvent.Trigger();
+
+		verify(TickTask.Wait(FTimespan::FromSeconds(0.1)));
+		verify(RemoveTickerTask.Wait(FTimespan::FromSeconds(0.1)));
+	}
+
+	// multithreaded stress test
+	{
+		using namespace UE::Tasks;
+
+		FTSTicker Ticker;
+
+		std::atomic<bool> bQuit{ false };
+		FTask TickTask = Launch(UE_SOURCE_LOCATION,
+			[&Ticker, &bQuit]
+			{
+				while (!bQuit)
+				{
+					Ticker.Tick(0.0f);
+				}
+			}
+		);
+
+		TArray<FTask> Tasks;
+		Tasks.Reserve(1000);
+		for (int i = 0; i != 10; ++i)
+		{
+			Tasks.Add(Launch(UE_SOURCE_LOCATION,
+				[&Ticker, &bQuit]
+				{
+					while (!bQuit)
+					{
+						FTSTicker::FDelegateHandle DelegateHandle = Ticker.AddTicker(UE_SOURCE_LOCATION, 0.0f,
+							[](float DeltaTime)
+							{
+								return true;
+							}
+						);
+
+						FTask RemoveTickerTask = Launch(UE_SOURCE_LOCATION,
+							[DelegateHandle = MoveTemp(DelegateHandle)]
+							{
+								FTSTicker::RemoveTicker(DelegateHandle);
+							}
+						);
+						RemoveTickerTask.Wait();
+					}
+				}
+			));
+		}
+
+		FPlatformProcess::Sleep(1.0); // let it run for a while
+		bQuit = true;
+		Tasks.Add(TickTask);
+		verify(Wait(Tasks, FTimespan::FromSeconds(3)));
+	}
+
+	UE_BENCHMARK(5, TickerPerfTest<1000, 1000>);
+	UE_BENCHMARK(5, TSTickerPerfTest<1000, 1000>);
+
+	return true;
+}
+
+#endif // WITH_DEV_AUTOMATION_TESTS
\ No newline at end of file
diff --git a/Engine/Source/Runtime/Core/Public/Containers/BackgroundableTicker.h b/Engine/Source/Runtime/Core/Public/Containers/BackgroundableTicker.h
--- a/Engine/Source/Runtime/Core/Public/Containers/BackgroundableTicker.h
+++ b/Engine/Source/Runtime/Core/Public/Containers/BackgroundableTicker.h
@@ -5,6 +5,28 @@
 #include "Ticker.h"
 
 /**
+ * This works the same as the core FTSTicker, but on supported mobile platforms
+ * it continues ticking while the app is running in the background.
+ */
+class CORE_API FTSBackgroundableTicker
+	: public FTSTicker
+{
+public:
+	static FTSBackgroundableTicker& GetCoreTicker();
+
+	FTSBackgroundableTicker();
+	~FTSBackgroundableTicker();
+
+private:
+	FDelegateHandle CoreTickerHandle;
+	::FDelegateHandle BackgroundTickerHandle;
+	bool bWasBackgrounded = false;
+};
+
+/**
+ * DEPRECATED not thread-safe version
+ * For migration guide see `FTSTicker`
+ * 
  * This works the same as the core FTicker, but on supported mobile platforms 
  * it continues ticking while the app is running in the background.
  */
diff --git a/Engine/Source/Runtime/Core/Public/Containers/Ticker.h b/Engine/Source/Runtime/Core/Public/Containers/Ticker.h
--- a/Engine/Source/Runtime/Core/Public/Containers/Ticker.h
+++ b/Engine/Source/Runtime/Core/Public/Containers/Ticker.h
@@ -6,7 +6,10 @@
 #include "Containers/Array.h"
 #include "Delegates/IDelegateInstance.h"
 #include "Delegates/Delegate.h"
+#include "Templates/SharedPointer.h"
+#include "Containers/MpscQueue.h"
 
+#include <atomic>
 
 /**
  * Ticker delegates return true to automatically reschedule at the same delay or false for one-shot.
@@ -17,9 +20,147 @@
 DECLARE_DELEGATE_RetVal_OneParam(bool, FTickerDelegate, float);
 
 /**
- * Ticker class. Fires delegates after a delay.
+ * Thread-safe ticker class. Fires delegates after a delay.
+ */
+class CORE_API FTSTicker
+{
+private:
+	struct FElement;
+
+public:
+	using FDelegateHandle = TWeakPtr<FElement>;
+
+	/** Singleton used for the ticker in Core / Launch. If you add a new ticker for a different subsystem, do not put that singleton here! **/
+	static FTSTicker& GetCoreTicker();
+
+	/**
+	 * Add a new ticker with a given delay / interval
+	 *
+	 * Can be called concurrently. Resources used by the delegate must be released on the ticking thread.
+	 *
+	 * @param InDelegate Delegate to fire after the delay
+	 * @param InDelay Delay until next fire; 0 means "next frame"
+	 */
+	FDelegateHandle AddTicker(const FTickerDelegate & InDelegate, float InDelay = 0.0f);
+
+	/**
+	* Add a new ticker with a given delay / interval.
+	* 
+	* Can be called concurrently.
+	*
+	* @param InName Name of this ticker for profiling
+	* @param InDelay Delay until next fire; 0 means "next frame"
+	* @param Function Function to execute. Should return true to fire after another InDelay time
+	*/
+	FDelegateHandle AddTicker(const TCHAR * InName, float InDelay, TFunction<bool(float)> Function);
+
+	/**
+	 * Removes a previously added ticker delegate.
+	 *
+	 * Can be called concurrently. If is caleld in the middle of the delegate execution, it blocks until the execution finishes, and thus guarantees
+	 * that the delegate won't be executed after the call.
+	 *
+	 * @param Handle The handle of the ticker to remove.
+	 */
+	static void RemoveTicker(FDelegateHandle Handle);
+
+	/**
+	 * Fire all tickers who have passed their delay and reschedule the ones that return true
+	 * 
+	 * Note: This reschedule has timer skew, meaning we always wait a full Delay period after
+	 * each invocation even if we missed the last interval by a bit. For instance, setting a
+	 * delay of 0.1 seconds will not necessarily fire the delegate 10 times per second, depending
+	 * on the Tick() rate. What this DOES guarantee is that a delegate will never be called MORE
+	 * FREQUENTLY than the Delay interval, regardless of how far we miss the scheduled interval.
+	 * 
+	 * Must not be called concurrently.
+	 *
+	 * @param DeltaTime	time that has passed since the last tick call
+	 */
+	void Tick(float DeltaTime);
+
+private:
+	/** Internal structure to store a ticker delegate and related data **/
+	struct FElement
+	{
+		/** Time that this delegate must not fire before **/
+		double FireTime;
+		/** Delay that this delegate was scheduled with. Kept here so that if the delegate returns true, we will reschedule it. **/
+		float DelayTime;
+		/** Delegate to fire **/
+		FTickerDelegate Delegate;
+
+		// The element can be in 4 states: "idle", "executing", "idle and removed" and "executing and removed"
+		// Often right after `RemoveTicket()` call resources that are used by the delegate are destroyed. We must ensure that 
+		// these resources are not accessed by the delegate after `RemoveTicket()` returns.
+		// `State` is used to remove the element safely, if removal happens in the middle of the delegate execution, removal will not return
+		// until the execution is finished.
+		enum EState : uint8 { Idle, Executing, Removed };
+		std::atomic<uint8> State{ Idle };
+
+		/** Default ctor is only required to implement CurrentElement handling without making it a pointer. */
+		CORE_API FElement();
+		/** This is the ctor that the code will generally use. */
+		CORE_API FElement(double InFireTime, float InDelayTime, const FTickerDelegate& InDelegate);
+
+		/** Fire the delegate if it is fireable **/
+		CORE_API bool Fire(float DeltaTime);
+	};
+
+	using FElementPtr = TSharedPtr<FElement>;
+
+	// all added delegates are initially stored in a separate thread-safe queue and then in the next Tick are moved to the main not thread-safe
+	// container
+	TMpscQueue<FElementPtr> AddedElements;
+
+	/** Current time of the ticker **/
+	double CurrentTime{ 0.0 };
+	/** Future delegates to fire **/
+	TArray<FElementPtr> Elements;
+};
+
+/**
+ * Base class for thread-safe ticker objects
+ */
+class CORE_API FTSTickerObjectBase
+{
+public:
+	UE_NONCOPYABLE(FTSTickerObjectBase);
+
+	/**
+	 * Constructor
+	 *
+	 * @param InDelay Delay until next fire; 0 means "next frame"
+	 * @param Ticker the ticker to register with. Defaults to FTicker::GetCoreTicker().
+	*/
+	FTSTickerObjectBase(float InDelay = 0.0f, FTSTicker& Ticker = FTSTicker::GetCoreTicker());
+
+	/** Virtual destructor. */
+	virtual ~FTSTickerObjectBase();
+
+	/**
+	 * Pure virtual that must be overloaded by the inheriting class.
+	 *
+	 * @param DeltaTime	time passed since the last call.
+	 * @return true if should continue ticking
+	 */
+	virtual bool Tick(float DeltaTime) = 0;
+
+private:
+	/** Delegate for callbacks to Tick */
+	FTSTicker::FDelegateHandle TickHandle;
+};
+
+/**
+ * DEPRECATED
+ * Not thread-safe ticker class. Fires delegates after a delay.
  * 
  * Note: Do not try to add the same delegate instance twice, as there is no way to remove only a single instance (see member RemoveTicker).
+ * 
+ * Migration guide:
+ * As FTicker is deprecated and replaced by thread-safe FTSTicker class, if you see a deprecation warning please migrate your code. Change
+ * `FTicker` to `FTSTicker` and change the type of the delegate handle returned by `AddTicker` methods from `FDelegateHandle` to
+ * `FTSTicker::DelegateHandle`
  */
 class FTicker
 {
@@ -114,6 +255,8 @@ protected:
  */
 class FTickerObjectBase
 {
+	PRAGMA_DISABLE_DEPRECATION_WARNINGS // FTicker
+
 public:
 
 	/**
@@ -158,4 +301,6 @@ private:
 	FTicker& Ticker;
 	/** Delegate for callbacks to Tick */
 	FDelegateHandle TickHandle;
+
+	PRAGMA_ENABLE_DEPRECATION_WARNINGS
 };
