fixed compilation warning

diff --git a/Engine/Source/Runtime/Core/Private/Tests/Containers/TickerTests.cpp b/Engine/Source/Runtime/Core/Private/Tests/Containers/TickerTests.cpp
--- a/Engine/Source/Runtime/Core/Private/Tests/Containers/TickerTests.cpp
+++ b/Engine/Source/Runtime/Core/Private/Tests/Containers/TickerTests.cpp
@@ -109,7 +109,7 @@ bool FTSTickerTest::RunTest(const FString& Parameters)
 		FTask RemoveTickerTask = Launch(UE_SOURCE_LOCATION, 
 			[&Ticker, &DelegateHandle] 
 			{ 
-				FPlatformProcess::Sleep(0.1); // let the ticking start and the delegate block on the event
+				FPlatformProcess::Sleep(0.1f); // let the ticking start and the delegate block on the event
 				Ticker.RemoveTicker(DelegateHandle); 
 			}
 		);
@@ -121,7 +121,7 @@ bool FTSTickerTest::RunTest(const FString& Parameters)
 			}
 		);
 
-		FPlatformProcess::Sleep(0.1); // let workers pick up the tasks and start execution
+		FPlatformProcess::Sleep(0.1f); // let workers pick up the tasks and start execution
 
 		verify(!TickTask.Wait(FTimespan::FromSeconds(0.1))); // tick is blocked because the delegate is blocked
 		verify(!RemoveTickerTask.Wait(FTimespan::FromSeconds(0.1))); // removal is blocked because the delegate is blocked
@@ -177,7 +177,7 @@ bool FTSTickerTest::RunTest(const FString& Parameters)
 			));
 		}
 
-		FPlatformProcess::Sleep(1.0); // let it run for a while
+		FPlatformProcess::Sleep(1.0f); // let it run for a while
 		bQuit = true;
 		Tasks.Add(TickTask);
 		verify(Wait(Tasks, FTimespan::FromSeconds(3)));
