Tasks System: nested tasks

A nested task is a task that must be completed before its parent task is completed, w/o blocking a worker that executes the parent task. It's an analog to `DontCompleteUntil` feature of the old TaskGraph

diff --git a/Engine/Source/Runtime/Core/Private/Tasks/TaskPrivate.cpp b/Engine/Source/Runtime/Core/Private/Tasks/TaskPrivate.cpp
--- a/Engine/Source/Runtime/Core/Private/Tasks/TaskPrivate.cpp
+++ b/Engine/Source/Runtime/Core/Private/Tasks/TaskPrivate.cpp
@@ -23,7 +23,37 @@ namespace UE { namespace Tasks { namespace Private
 		if (GetPipe() != nullptr)
 		{
 			GetPipe()->ExecutionFinished();
+		}
+	}
+
+	void FTaskBase::Close()
+	{
+		if (GetPipe() != nullptr)
+		{
 			GetPipe()->ClearTask(*this);
 		}
+
+		Subsequents.Close(
+			[this](FTaskBase* Subsequent)
+			{
+				Subsequent->TryUnlock();
+			}
+		);
+	
+		Release(); // release the reference accounted for nested tasks. the task can be destroyed as the result
+	}
+
+	thread_local FTaskBase* CurrentTask = nullptr;
+
+	FTaskBase* FTaskBase::GetCurrentTask()
+	{
+		return CurrentTask;
+	}
+
+	FTaskBase* FTaskBase::ExchangeCurrentTask(FTaskBase* Task)
+	{
+		FTaskBase* PrevTask = CurrentTask;
+		CurrentTask = Task;
+		return PrevTask;
 	}
 }}}
\ No newline at end of file
diff --git a/Engine/Source/Runtime/Core/Private/Tests/Tasks/TasksTest.cpp b/Engine/Source/Runtime/Core/Private/Tests/Tasks/TasksTest.cpp
--- a/Engine/Source/Runtime/Core/Private/Tests/Tasks/TasksTest.cpp
+++ b/Engine/Source/Runtime/Core/Private/Tests/Tasks/TasksTest.cpp
@@ -6,7 +6,6 @@
 #include "Tasks/Pipe.h"
 #include "HAL/Thread.h"
 #include "Async/ParallelFor.h"
-#include "Experimental/Async/AwaitableTask.h"
 
 #include <atomic>
 
@@ -109,6 +108,15 @@ namespace UE { namespace TasksTests
 			verify(Event.BusyWait(FTimespan::Zero()));
 		}
 
+		{	// FTaskEvent can be triggered multiple times
+			FTaskEvent Event{ UE_SOURCE_LOCATION };
+			Event.Trigger();
+			check(Event.IsCompleted());
+			Event.Trigger();
+			Event.Trigger();
+			check(Event.IsCompleted());
+		}
+
 		{	// same but using busy-waiting
 			FTaskEvent Event{ UE_SOURCE_LOCATION };
 			check(!Event.IsCompleted());
@@ -237,6 +245,144 @@ namespace UE { namespace TasksTests
 			Task = {};
 		}
 
+		///////////////////////////////////////////////////////////////////////////////
+		// nested tasks
+		///////////////////////////////////////////////////////////////////////////////
+
+		{	// one nested task
+			FTaskEvent FinishSignal{ UE_SOURCE_LOCATION };
+			FTaskEvent ExecutedSignal{ UE_SOURCE_LOCATION };
+			FTask Task{ Launch(UE_SOURCE_LOCATION,
+				[&FinishSignal, &ExecutedSignal]
+				{
+					AddNested(FinishSignal);
+					ExecutedSignal.Trigger();
+				}
+			) };
+
+			verify(!Task.Wait(FTimespan::FromMilliseconds(100)));
+			verify(ExecutedSignal.IsCompleted());
+			FinishSignal.Trigger();
+			verify(Task.Wait());
+		}
+
+		{	// nested task completed before the parent finishes its execution
+			Launch(UE_SOURCE_LOCATION,
+				[]
+				{
+					FTask NestedTask = Launch(UE_SOURCE_LOCATION, [] {});
+					AddNested(NestedTask);
+					NestedTask.Wait();
+				}
+			).Wait();
+		}
+
+		{	// multiple nested tasks
+			FTaskEvent Signal1{ UE_SOURCE_LOCATION };
+			FTaskEvent Signal2{ UE_SOURCE_LOCATION };
+			FTaskEvent Signal3{ UE_SOURCE_LOCATION };
+
+			FTask Task{ Launch(UE_SOURCE_LOCATION,
+				[&Signal1, &Signal2, &Signal3]
+				{
+					AddNested(Signal1);
+					AddNested(Signal2);
+					AddNested(Signal3);
+				}
+			) };
+
+			verify(!Task.Wait(FTimespan::FromMilliseconds(100)));
+			Signal1.Trigger();
+			verify(!Task.Wait(FTimespan::FromMilliseconds(100)));
+			Signal2.Trigger();
+			verify(!Task.Wait(FTimespan::FromMilliseconds(100)));
+			Signal3.Trigger();
+			verify(Task.Wait());
+		}
+
+		{	// nested in square
+			FTaskEvent Signal{ UE_SOURCE_LOCATION };
+			FTask Task{ Launch(UE_SOURCE_LOCATION,
+				[&Signal]
+				{
+					AddNested(Signal);
+
+					FTaskEvent NestedSignal{ UE_SOURCE_LOCATION };
+					FTask NestedTask{ Launch(UE_SOURCE_LOCATION,
+						[&NestedSignal]
+						{
+							AddNested(NestedSignal);
+						}
+					) };
+
+					verify(!NestedTask.Wait(FTimespan::FromMilliseconds(100)));
+					NestedSignal.Trigger();
+					verify(NestedTask.Wait());
+				}
+			)};
+
+			verify(!Task.Wait(FTimespan::FromMilliseconds(200)));
+			Signal.Trigger();
+			verify(Task.Wait());
+		}
+
+		{	// nested task should block pipe execution, also an example of pipe suspension
+
+			struct FPipeSuspensionScope
+			{
+				explicit FPipeSuspensionScope(FPipe& Pipe)
+				{
+					Pipe.Launch(UE_SOURCE_LOCATION,
+						[this]
+						{
+							SuspendSignal.Trigger();
+							AddNested(ResumeSignal);
+						}
+					);
+
+					SuspendSignal.Wait();
+				}
+
+				~FPipeSuspensionScope()
+				{
+					ResumeSignal.Trigger();
+				}
+
+				FTaskEvent SuspendSignal{ UE_SOURCE_LOCATION };
+				FTaskEvent ResumeSignal{ UE_SOURCE_LOCATION };
+			};
+
+			FPipe Pipe{ UE_SOURCE_LOCATION };
+			FTask Task;
+			{
+				FPipeSuspensionScope Suspension(Pipe);
+				FPlatformProcess::Sleep(0.1f); // let pipe suspension finish its business, which was a bug as pipe was cleared (and unblocked) before AddNested kicked in
+				Task = Pipe.Launch(UE_SOURCE_LOCATION, [] {});
+				verify(!Task.Wait(FTimespan::FromMilliseconds(100)));
+			}
+			verify(Task.Wait());
+		}
+
+		//{	// Cancelled task is completed only after nested
+		//	FTaskEvent ParentBlocker{ UE_SOURCE_LOCATION };
+		//	FTaskEvent NestedBlocker{ UE_SOURCE_LOCATION };
+		//	FTask Parent = Launch(UE_SOURCE_LOCATION,
+		//		[&NestedBlocker]
+		//		{
+		//			FTask Nested{ Launch(UE_SOURCE_LOCATION, [] {}, NestedBlocker) };
+		//			AddNested(Nested);
+		//			Nested.Wait();
+		//			// the only reference left is hold by Parent because of `AddNested` dependency
+		//		},
+		//		ParentBlocker
+		//	);
+		//	verify(Parent.TryCancel());
+		//	ParentBlocker.Trigger(); // unblock `Parent` so it can notice that it's been cancelled
+		//	check(Parent.Wait(FTimespan::FromMilliseconds(100))); // but it's still not complete because of the nested task is blocked
+		//	NestedBlocker.Trigger();
+		//	verify(Parent.Wait());
+		//}
+
 		UE_BENCHMARK(5, BasicStressTest<1000, 1000>);
 
 		return true;
diff --git a/Engine/Source/Runtime/Core/Public/Tasks/Task.h b/Engine/Source/Runtime/Core/Public/Tasks/Task.h
--- a/Engine/Source/Runtime/Core/Public/Tasks/Task.h
+++ b/Engine/Source/Runtime/Core/Public/Tasks/Task.h
@@ -20,10 +20,17 @@ namespace UE { namespace Tasks
 	template<typename TaskBodyType, typename PrerequisitesCollectionType>
 	TTask<TInvokeResult_T<TaskBodyType>> Launch(const TCHAR* DebugName, TaskBodyType&& TaskBody, PrerequisitesCollectionType&& Prerequisites, LowLevelTasks::ETaskPriority Priority = LowLevelTasks::ETaskPriority::Normal);
 
+	template<typename TaskCollectionType>
+	bool Wait(TaskCollectionType&& Tasks, FTimespan InTimeout = FTimespan::MaxValue());
+
+	template<typename TaskCollectionType>
+	bool BusyWait(TaskCollectionType&& Tasks, FTimespan InTimeout = FTimespan::MaxValue());
+
 	// a common part of the generic `TTask<ResultType>` and its `TTask<void>` specialisation
 	template<typename ResultType>
 	class TTaskBase
 	{
+		// friends to get access to `Pimpl`
 		friend Private::FTaskBase;
 
 		template<typename... TaskTypes>
@@ -32,6 +39,15 @@ namespace UE { namespace Tasks
 		template<typename TaskCollectionType>
 		friend bool Private::TryRetractAndExecute(TaskCollectionType&& Tasks, FTimespan Timeout);
 
+		template<typename TaskType>
+		friend void AddNested(TaskType&& Nested);
+
+		template<typename TaskCollectionType>
+		friend bool Wait(TaskCollectionType&& Tasks, FTimespan InTimeout/* = FTimespan::MaxValue()*/);
+
+		template<typename TaskCollectionType>
+		friend bool BusyWait(TaskCollectionType&& Tasks, FTimespan InTimeout/* = FTimespan::MaxValue()*/);
+
 	protected:
 		TTaskBase() = default;
 
@@ -148,7 +164,7 @@ namespace UE { namespace Tasks
 	{
 	public:
 		TTask() = default;
-	
+
 		void GetResult()
 		{
 			check(IsValid()); // to be consistent with a generic `TTask<ResultType>::GetResult()`
@@ -171,7 +187,8 @@ namespace UE { namespace Tasks
 		{}
 	};
 
-	// a synchronisation primitive similar to `FEvent` that uses "busy waiting" - executing tasks while waiting
+	// A synchronisation primitive, a recommended substitution of `FEvent` for signalling between tasks. If used as a task prerequisite or 
+	// a nested task, it doesn't block a worker thread. Optionally can use "busy waiting" - executing tasks while waiting.
 	class FTaskEvent : public TTask<void>
 	{
 	public:
@@ -181,9 +198,19 @@ namespace UE { namespace Tasks
 			Pimpl->Init(DebugName, [] {}, Private::FTaskBase::InlineTaskPriority);
 		}
 
+		// all prerequisites must be added before triggering the event
+		template<typename PrerequisitesType>
+		void AddPrerequisites(const PrerequisitesType& Prerequisites)
+		{
+			Pimpl->AddPrerequisites(Prerequisites);
+		}
+
 		void Trigger()
 		{
-			verify(Pimpl->TryLaunch());
+			if (!IsCompleted()) // event can be triggered multiple times
+			{
+				Pimpl->TryLaunch();
+			}
 		}
 	};
 
@@ -231,7 +258,7 @@ namespace UE { namespace Tasks
 
 	// wait for multiple tasks, with optional timeout
 	template<typename TaskCollectionType>
-	bool Wait(TaskCollectionType&& Tasks, FTimespan InTimeout = FTimespan::MaxValue())
+	bool Wait(TaskCollectionType&& Tasks, FTimespan InTimeout/* = FTimespan::MaxValue()*/)
 	{
 		FTimeout Timeout{ InTimeout };
 
@@ -259,7 +286,7 @@ namespace UE { namespace Tasks
 
 	// wait for multiple tasks while executing other tasks
 	template<typename TaskCollectionType>
-	bool BusyWait(TaskCollectionType&& Tasks, FTimespan TimeoutValue = FTimespan::MaxValue())
+	bool BusyWait(TaskCollectionType&& Tasks, FTimespan TimeoutValue/* = FTimespan::MaxValue()*/)
 	{
 		FTimeout Timeout{ TimeoutValue };
 
@@ -270,7 +297,7 @@ namespace UE { namespace Tasks
 
 		for (auto& Task : Tasks)
 		{
-			if (Timeout || (!Task.BusyWait(Timeout.GetRemainingTime())))
+			if (Timeout || !Task.BusyWait(Timeout.GetRemainingTime()))
 			{
 				return false;
 			}
@@ -312,6 +339,17 @@ namespace UE { namespace Tasks
 	{
 		return TPrerequisites<TaskTypes...>{ Forward<TaskTypes>(Tasks)... };
 	}
+
+	// Adds the nested task to the task that is being currently executed by the current thread. A parent task is not flagged completed
+	// until all nested tasks are completed. It's similar to explicitly waiting for a sub-task at the end of its parent task, except explicit waiting
+	// blocks the worker executing the parent task until the sub-task is completed. With nested tasks, the worker won't be blocked.
+	template<typename TaskType>
+	void AddNested(TaskType&& Nested)
+	{
+		Private::FTaskBase* Parent = Private::FTaskBase::GetCurrentTask();
+		check(Parent != nullptr);
+		Parent->AddNested(*Nested.Pimpl);
+	}
 }}
 
 template <typename... TaskTypes>
diff --git a/Engine/Source/Runtime/Core/Public/Tasks/TaskPrivate.h b/Engine/Source/Runtime/Core/Public/Tasks/TaskPrivate.h
--- a/Engine/Source/Runtime/Core/Public/Tasks/TaskPrivate.h
+++ b/Engine/Source/Runtime/Core/Public/Tasks/TaskPrivate.h
@@ -11,6 +11,7 @@
 #include "Containers/ClosableMpscQueue.h"
 #include "Templates/UnrealTemplate.h"
 #include "CoreTypes.h"
+#include "Math/NumericLimits.h"
 
 #include <atomic>
 #include <type_traits>
@@ -19,39 +20,69 @@ namespace UE { namespace Tasks
 {
 	using LowLevelTasks::ETaskPriority;
 
-	// BEGIN: forward declarations
-
-	template<typename ResultType>
-	class TTaskBase;
-
 	class FPipe;
 
-	// END: forward declarations
-
 	namespace Private
 	{
-		// Task class implementation with piping support.
-		class FTaskBase
+		// intrusive atomic reference-counting base class.
+		class FRefCountedBase
+		{
+		public:
+			explicit FRefCountedBase(uint32 InitRefCount = 0)
+				: RefCount{ InitRefCount }
+			{}
+
+			virtual ~FRefCountedBase() = default;
+
+			void AddRef()
+			{
+				RefCount.fetch_add(1, std::memory_order_relaxed);
+			}
+
+			void Release()
+			{
+				uint32 LocalRefCount = RefCount.fetch_sub(1, std::memory_order_release) - 1;
+				if (LocalRefCount == 0)
+				{
+					std::atomic_thread_fence(std::memory_order_acquire);
+					delete this;
+				}
+			}
+
+		private:
+			std::atomic<uint32> RefCount;
+		};
+
+		// A base class for high-level task implementation.
+		// Must be dynamically allocated and used with TRefCountPtr.
+		// @see Tasks::FTaskBase
+		class CORE_API FTaskBase : public FRefCountedBase
 		{
 			UE_NONCOPYABLE(FTaskBase);
 
-			// Can be used only as a base class.
-		protected:
-			FTaskBase() = default;
+			// `ExecutionFlag` is set at the beginning of execution as the most significant bit of `NumLocks` and indicates a switch 
+			// of `NumLocks` from "execution prerequisites" (a number of uncompleted prerequisites that block task execution) to 
+			// "completion prerequisites" (a number of nested uncompleted tasks that block task completion)
+			static constexpr uint32 ExecutionFlag = 0x80000000;
 
-			~FTaskBase()
+		public:
+			FTaskBase()
+				: FRefCountedBase(2) // for the initial reference (we don't increment it on passing to `TRefCountPtr`),
+				// and for the reference passed to the scheduler - is released when scheduler's "runnable" lambda is destroyed (see `Init`)
+			{}
+
+			virtual ~FTaskBase() override
 			{
 				check(IsCompleted());
 			}
 
-		public:
 			// a special internal task priority for "inline" task execution - a task is executed as soon as it's launched and has no 
 			// pending dependencies, "inline", w/o scheduling
 			static const ETaskPriority InlineTaskPriority{ ETaskPriority::Count };
-			
+
 			// initialises the task but doesn't launches it
-			template<typename TaskBodyType, typename DeleterType>
-			void Init(const TCHAR* DebugName, TaskBodyType&& InTaskBody, DeleterType&& Deleter, LowLevelTasks::ETaskPriority Priority)
+			template<typename TaskBodyType>
+			void Init(const TCHAR* DebugName, TaskBodyType&& InTaskBody, LowLevelTasks::ETaskPriority Priority)
 			{
 				TaskBody = Forward<TaskBodyType>(InTaskBody);
 
@@ -61,24 +92,26 @@ namespace UE { namespace Tasks
 						// releasing scheduler's task reference can cause task's automatic destruction and so must be done after the low-level task
 						// task is flagged as completed. The task is flagged as completed after the continuation is executed but before its destroyed.
 						// `Deleter` is captured by value and is destroyed along with the continuation, calling the given functor on destruction
-						Deleter = Forward<DeleterType>(Deleter)
-					]
+						Deleter = LowLevelTasks::TDeleter<FRefCountedBase, &FRefCountedBase::Release>{ this }
+					] 
 					{
 						TryExecute();
 					}
 				);
 			}
 
-			// the task will be executed only when all prerequisites are completed
-			template<typename T>
-			void AddPrerequisites(const TTaskBase<T>& Prerequisite)
+			// the task will be executed only when all prerequisites are completed. The task type must be a task handle that holds a pointer to
+			//  FTaskBase as its `Pimpl` member (see Tasks::TTaskBase)
+			template<typename HigherLevelTaskType, decltype(std::declval<HigherLevelTaskType>().Pimpl)* = nullptr>
+			void AddPrerequisites(const HigherLevelTaskType& Prerequisite)
 			{
-				checkf(NumLocks.load(std::memory_order_relaxed) >= NumInitialLocks, TEXT("Prerequisites can be added only before the task is launched"));
+				checkf(NumLocks.load(std::memory_order_relaxed) >= NumInitialLocks && NumLocks.load(std::memory_order_relaxed) < ExecutionFlag, TEXT("Prerequisites can be added only before the task is launched"));
 
 				// registering the task as a subsequent of the given prerequisite can cause its immediate launch by the prerequisite
 				// (if the prerequisite has been completed on another thread), so we need to keep the task locked by assuming that the 
 				// prerequisite can be added successfully, and release the lock if it wasn't
-				NumLocks.fetch_add(1, std::memory_order_acquire); // `acquire` to make it happen before the task is registered as a subsequent
+				uint32 PrevNumLocks = NumLocks.fetch_add(1, std::memory_order_acquire); // `acquire` to make it happen before the task is registered as a subsequent
+				checkf(PrevNumLocks + 1 < ExecutionFlag, TEXT("Max number of nested tasks reached: %d"), ExecutionFlag);
 
 				if (!Prerequisite.Pimpl->AddSubsequent(*this))
 				{
@@ -89,22 +122,23 @@ namespace UE { namespace Tasks
 
 			// the task will be executed only when all prerequisites are completed
 			// @param Prerequisites - an iterable collection of tasks
-			template<typename PrerequisiteCollectionType, decltype(std::declval<PrerequisiteCollectionType&>().begin())* = nullptr>
-			void AddPrerequisites(const PrerequisiteCollectionType& Prerequisites)
+			template<typename PrerequisiteCollectionType, decltype(std::declval<PrerequisiteCollectionType>().begin())* = nullptr>
+			void AddPrerequisites(const PrerequisiteCollectionType& InPrerequisites)
 			{
-				checkf(NumLocks.load(std::memory_order_relaxed) >= NumInitialLocks, TEXT("Prerequisites can be added only before the task is launched"));
+				checkf(NumLocks.load(std::memory_order_relaxed) >= NumInitialLocks && NumLocks.load(std::memory_order_relaxed) < ExecutionFlag, TEXT("Prerequisites can be added only before the task is launched"));
 
 				// registering the task as a subsequent of the given prerequisite can cause its immediate launch by the prerequisite
 				// (if the prerequisite has been completed on another thread), so we need to keep the task locked by assuming that the 
 				// prerequisite can be added successfully, and release the lock if it wasn't
-				NumLocks.fetch_add(GetNum(Prerequisites), std::memory_order_acquire); // `acquire` to make it happen before the task is registered as a subsequent
+				uint32 PrevNumLocks = NumLocks.fetch_add(GetNum(InPrerequisites), std::memory_order_acquire); // `acquire` to make it happen before the task is registered as a subsequent
+				checkf(PrevNumLocks + GetNum(InPrerequisites) < ExecutionFlag, TEXT("Max number of nested tasks reached: %d"), ExecutionFlag);
 
 				uint32 NumCompletedPrerequisites = 0;
-				for (auto& Prereq : Prerequisites)
+				for (auto& Prereq : InPrerequisites)
 				{
 					// prerequisites can be either `FTaskBase*` or its Pimpl handle
 					FTaskBase* Prerequisite;
-					using FPrerequisiteType = std::decay_t<decltype(*std::declval<PrerequisiteCollectionType&>().begin())>;
+					using FPrerequisiteType = std::decay_t<decltype(*std::declval<PrerequisiteCollectionType>().begin())>;
 					if constexpr (std::is_same_v<FPrerequisiteType, FTaskBase*>)
 					{
 						Prerequisite = Prereq;
@@ -125,13 +159,13 @@ namespace UE { namespace Tasks
 			}
 
 			// the task unlocks all its subsequents on completion.
-			// return false if the task is already completed
+			// returns false if the task is already completed and the subsequent wasn't added
 			bool AddSubsequent(FTaskBase& Subsequent)
 			{
 				return Subsequents.Enqueue(&Subsequent);
 			}
 
-			// A piped task will be executed after the previous task from this pipe is completed. Tasks from the same pipe are not executed
+			// A piped task is executed after the previous task from this pipe is completed. Tasks from the same pipe are not executed
 			// concurrently (so don't require synchronization), but not necessarily on the same thread.
 			// @See FPipe
 			void SetPipe(FPipe& InPipe)
@@ -144,66 +178,67 @@ namespace UE { namespace Tasks
 				return Pipe;
 			}
 
-			// Tasks w/o incomplete prerequisites and not blocked by a pipe are scheduled immediately
+			// Tries to schedule task execution. Returns false if the task has incomplete dependencies (prerequisites or is blocked by a pipe). 
+			// In this case the task will be automatically scheduled when all dependencies are completed.
 			bool TryLaunch()
 			{
 				TaskTrace::Launched(GetTraceId(), LowLevelTask.GetDebugName(), true, (ENamedThreads::Type)0xff);
 				return TryUnlock();
 			}
 
-			// if the task execution is already completed
 			bool IsCompleted() const
 			{
 				return Subsequents.IsClosed();
 			}
 
+			// Tries to pull out the task from the scheduler and execute it. Returns false if task execution is already started.
 			bool TryRetractAndExecute()
 			{
-				bool bIsLocked = NumLocks.load(std::memory_order_acquire) != (GetPipe() == nullptr ? 1 : 0);
-				// try to execute before cancelling the low-level task as successful cancellation can release the last reference and destroy the task
-				if (!bIsLocked && TryExecute())
+				if (IsCompleted())
 				{
-					LowLevelTask.TryCancel();
 					return true;
 				}
 
+				uint32 LocalNumLocks = NumLocks.load(std::memory_order_acquire);
+				bool bExecutionUnlocked = LocalNumLocks < ExecutionFlag && LocalNumLocks == (GetPipe() == nullptr ? 1 : 0);
+				if (bExecutionUnlocked && TryExecute())
+				{
+					LowLevelTask.TryCancel();
+					return IsCompleted(); // still can be hold back by nested tasks
+				}
+
 				return false;
 			}
 
+			// adds a nested task that must be completed before the parent (this) is completed
+			void AddNested(FTaskBase& Nested)
+			{
+				uint32 PrevNumLocks = NumLocks.fetch_add(1, std::memory_order_relaxed); // in case we'll succeed in adding subsequent
+				checkf(PrevNumLocks + 1 < TNumericLimits<uint32>::Max(), TEXT("Max number of nested tasks reached: %d"), TNumericLimits<uint32>::Max() - ExecutionFlag);
+				checkf(PrevNumLocks >= ExecutionFlag, TEXT("Internal error: nested tasks can be added only during parent's execution (%u)"), PrevNumLocks);
+
+				if (!Nested.AddSubsequent(*this))
+				{
+					NumLocks.fetch_sub(1, std::memory_order_relaxed);
+				}
+			}
+
 			// waits until the task is completed while executing other tasks
 			void BusyWait()
 			{
-				// we try to retract the task from the scheduler, if it's execution hasn't started yet, to execute it immediately. This makes
-				// the task skip waiting in the scheduler's queue.
-				// we can't retract it right here because it can be not scheduled yet (e.g. when it's blocked by a pipe), so we need to wait for this
-				// `!IsReady()` here means "the task is scheduled"
-				LowLevelTasks::BusyWaitUntil([this] { return !LowLevelTask.IsReady(); });
-
-				// try to execute before cancelling the low-level task as successful cancellation can release the last reference and destroy the task
-				if (TryExecute())
+				if (!TryRetractAndExecute())
 				{
-					LowLevelTask.TryCancel();
-					return;
+					LowLevelTasks::BusyWaitUntil([this] { return IsCompleted(); });
 				}
-
-				LowLevelTasks::BusyWaitUntil([this] { return IsCompleted(); });
 			}
 
 			// waits until the task is completed or waiting timed out, while executing other tasks
 			bool BusyWait(FTimespan InTimeout)
 			{
 				FTimeout Timeout{ InTimeout };
-				LowLevelTasks::BusyWaitUntil([this, Timeout] { return !LowLevelTask.IsReady() || Timeout; });
-
-				if (LowLevelTask.IsReady()) // timeout
+				
+				if (TryRetractAndExecute())
 				{
-					return false;
-				}
-
-				// try to execute before cancelling the low-level task as successful cancellation can release the last reference and destroy the task
-				if (TryExecute())
-				{
-					LowLevelTask.TryCancel();
 					return true;
 				}
 
@@ -215,24 +250,13 @@ namespace UE { namespace Tasks
 			template<typename ConditionType>
 			bool BusyWait(ConditionType&& Condition)
 			{
-				LowLevelTasks::BusyWaitUntil(
-					[this, Condition = Forward<ConditionType>(Condition)] { return !LowLevelTask.IsReady() || Condition(); }
-				);
-
-				if (LowLevelTask.IsReady()) // `Condition` returned true
-				{
-					return false;
-				}
-
-				// try to execute before cancelling the low-level task as successful cancellation can release the last reference and destroy the task
-				if (TryExecute())
+				if (TryRetractAndExecute())
 				{
-					LowLevelTask.TryCancel();
 					return true;
 				}
 
 				LowLevelTasks::BusyWaitUntil(
-					[this, Condition = Forward<ConditionType>(Condition)] { return IsCompleted() || Condition(); }
+					[this, Condition = Forward<ConditionType>(Condition)]{ return IsCompleted() || Condition(); }
 				);
 				return IsCompleted();
 			}
@@ -246,13 +270,33 @@ namespace UE { namespace Tasks
 #endif
 			}
 
+			// returns the task that is being currently exeucuted by this thread, if any
+			static FTaskBase* GetCurrentTask();
+
 		private:
-			// allows the task to be scheduled and executed, once for each lock. it's called on actual launching, on every prerequisite 
-			// completion and when an associated pipe (if any) is unblocked
+			// sets the current task and returns the previous current task
+			static FTaskBase* ExchangeCurrentTask(FTaskBase* Task);
+
+			// A task can be locked for execution (by prerequisites or if it's not launched yet) or for completion (by nested tasks).
+			// This method is called to unlock the task and so can result in its scheduling (and execution) or completion
 			bool TryUnlock()
 			{
-				uint32 LocalNumLocks = NumLocks.fetch_sub(1, std::memory_order_acq_rel) - 1; // `acq_rel` to make it happen after task 
+				uint32 PrevNumLocks = NumLocks.fetch_sub(1, std::memory_order_acq_rel); // `acq_rel` to make it happen after task 
 				// preparation and before launching it
+				if (PrevNumLocks < ExecutionFlag)
+				{
+					checkf(PrevNumLocks != (GetPipe() == nullptr ? 1 : 0), TEXT("The task is not locked"));
+					return TrySchedule(PrevNumLocks - 1);
+				}
+
+				checkf(PrevNumLocks != ExecutionFlag, TEXT("The task is not locked"));
+				return TryComplete(PrevNumLocks - 1);
+			}
+
+			// tries to pass the task to the scheduler for eventual execution by checking if something holds it back.
+			// tasks with inline priority can be executed right here.
+			bool TrySchedule(uint32 LocalNumLocks)
+			{
 				const bool bReadyForPiping = LocalNumLocks <= 1; // the only lock that can remain is pipe's one
 				if (!bReadyForPiping)
 				{
@@ -276,28 +320,50 @@ namespace UE { namespace Tasks
 					return true;
 				}
 
-				// schedule
-				return LowLevelTasks::TryLaunch(LowLevelTask);
+				verify(LowLevelTasks::TryLaunch(LowLevelTask)); // schedule the task
+
+				return true;
+			}
+
+			// the task is already executed but it's not completed yet. This method sets completion flag if there're no pending nested tasks.
+			// The task can be deleted as the result of this call.
+			bool TryComplete(uint32 LocalNumLocks)
+			{
+				checkSlow(!IsCompleted());
+				if (LocalNumLocks == ExecutionFlag)
+				{
+					Close();
+					return true;
+				}
+
+				return false;
 			}
 
-			// execute the task body if the execution wasn't started yet
+			// prepares the task for execution and executes its body if its execution hasn't been started yet
 			bool TryExecute()
 			{
+				// can be executed only once
 				if (!bAvailableForExecution.exchange(false, std::memory_order_relaxed))
 				{
-					return false; // too late, execution already started, e.g. by task retraction
+					return false;
 				}
 
+				AddRef(); // for the reference hold by nested tasks, is released when task is closed
+
+				NumLocks.store(ExecutionFlag + 1, std::memory_order_relaxed); // +1 to hold it locked during execution, so nested tasks don't
+				// complete it before the execution finishes
+				FTaskBase* PrevTask = ExchangeCurrentTask(this);
 				TaskTrace::Started(GetTraceId());
-				{
-					TRACE_CPUPROFILER_EVENT_SCOPE(ExecuteTask);
-					StartPipeExecution();
-					TaskBody();
-					FinishPipeExecution();
-				}
+				StartPipeExecution();
+
+				TaskBody();
+
+				FinishPipeExecution();
 				TaskTrace::Finished(GetTraceId());
+				ExchangeCurrentTask(PrevTask);
+				uint32 LocalNumLocks = NumLocks.fetch_sub(1, std::memory_order_relaxed) - 1;
 
-				Close();
+				TryComplete(LocalNumLocks);
 
 				return true;
 			}
@@ -327,29 +393,29 @@ namespace UE { namespace Tasks
 				return true;
 			}
 
-			CORE_API bool PushIntoPipe();
-			CORE_API void StartPipeExecution();
-			CORE_API void FinishPipeExecution();
+			// is called when the task has no pending prerequisites. Returns false is the pipe is blocked (another task from the same pipe is being
+			// executed)
+			bool PushIntoPipe();
 
-			void Close()
-			{
-				Subsequents.Close(
-					[this](FTaskBase* Subsequent)
-					{
-						Subsequent->TryUnlock();
-					}
-				);
-			}
+			void StartPipeExecution();
+			void FinishPipeExecution();
+
+			// closes task by unlocking its subsequents and flagging it as completed
+			void Close();
 
 		private:
 			LowLevelTasks::FTask LowLevelTask;
 			TUniqueFunction<void()> TaskBody;
 			std::atomic<bool> bAvailableForExecution{ true };
 
-			// the number of times that the task should be unlocked before it can be scheduled
-			// initial count is 1 for launching the task (it can't be scheduled before it's launched) and 1 for a potential blocked pipe
+			// the number of times that the task should be unlocked before it can be scheduled or completed
+			// initial count is 1 for launching the task (it can't be scheduled before it's launched) and 1 for a potential blocked pipe. once NumLocks
+			// reaches 0 the task is scheduled for execution.
+			// NumLocks's the most significant bit (see `ExecutionFlag`) is set on task execution start, and indicates that now NumLocks is about 
+			// how many times the task must be unlocked to be completed
 			static const uint32 NumInitialLocks = 1 + 1;
 			std::atomic<uint32> NumLocks{ NumInitialLocks };
+
 			TClosableMpscQueue<FTaskBase*> Subsequents;
 
 			FPipe* Pipe{ nullptr };
@@ -383,33 +449,28 @@ namespace UE { namespace Tasks
 
 		// Extends FTaskBase by supporting execution result.
 		template<typename ResultType>
-		class TTaskWithResultBase : public FTaskBase
+		class TTaskWithResult : public FTaskBase
 		{
-			UE_NONCOPYABLE(TTaskWithResultBase);
-
-			using FSelf = TTaskWithResultBase<ResultType>;
+			UE_NONCOPYABLE(TTaskWithResult);
 
-			// Can be used only as a base class
-		protected:
-			TTaskWithResultBase() = default;
+		public:
+			TTaskWithResult() = default;
 
-			~TTaskWithResultBase()
+			virtual ~TTaskWithResult() override
 			{
 				checkf(IsCompleted(), TEXT("Every task instance must be completed before it's destroyed"));
 				DestructItem(ResultStorage.GetTypedPtr());
 			}
 
-		public:
-			template<typename TaskBodyType, typename DeleterType>
-			void Init(const TCHAR* DebugName, TaskBodyType&& TaskBody, DeleterType&& Deleter, LowLevelTasks::ETaskPriority Priority)
+			template<typename TaskBodyType>
+			void Init(const TCHAR* DebugName, TaskBodyType&& TaskBody, LowLevelTasks::ETaskPriority Priority)
 			{
 				FTaskBase::Init(
-					DebugName, 
+					DebugName,
 					[this, TaskBody = Forward<TaskBodyType>(TaskBody)]() mutable
-					{ 
-						new(&ResultStorage) ResultType{ Invoke(TaskBody) }; 
+					{
+						new(&ResultStorage) ResultType{ Invoke(TaskBody) };
 					},
-					Forward<DeleterType>(Deleter),
 					Priority
 				);
 			}
@@ -430,75 +491,15 @@ namespace UE { namespace Tasks
 		};
 
 		template<>
-		class TTaskWithResultBase<void> : public FTaskBase
+		class TTaskWithResult<void> : public FTaskBase
 		{
-		protected:
-			TTaskWithResultBase() = default; // instances can be created only on the heap by a ref-counting handler
-
 		public:
+			TTaskWithResult() = default; // instances can be created only on the heap by a ref-counting handler
+
 			void GetResult()
 			{
 				checkf(IsCompleted(), TEXT("The task must be completed to obtain its result"));
 			}
 		};
-
-		// intrusive atomic reference-counting base class.
-		// doesn't require derived classes to be polymorphic
-		template<typename ReferenceType>
-		class TRefCountedBase
-		{
-		public:
-			explicit TRefCountedBase(uint32 InitRefCount = 0)
-				: RefCount{ InitRefCount }
-			{
-				static_assert(std::is_final<ReferenceType>::value, "ReferenceType should be `final` to guarantee no slicing during destruction");
-			}
-
-			void AddRef()
-			{
-				RefCount.fetch_add(1, std::memory_order_relaxed);
-			}
-
-			void Release()
-			{
-				uint32 LocalRefCount = RefCount.fetch_sub(1, std::memory_order_release) - 1;
-				if (LocalRefCount == 0)
-				{
-					std::atomic_thread_fence(std::memory_order_acquire);
-					delete static_cast<ReferenceType*>(this);
-				}
-			}
-
-		private:
-			std::atomic<uint32> RefCount;
-		};
-
-		// A final task class that must be used only with `TRefCountPtr`
-		// @see UE::Tasks::Task<T>
-		template<typename ResultType>
-		class TTaskWithResult final : public TTaskWithResultBase<ResultType>, public TRefCountedBase<TTaskWithResult<ResultType>>
-		{
-			UE_NONCOPYABLE(TTaskWithResult);
-
-			using FSelf = TTaskWithResult<ResultType>;
-
-		public:
-			TTaskWithResult()
-				: TRefCountedBase<FSelf>(2) // one for the initial reference (we don't increment it when passing to `TRefCountPtr`) and 
-				// one for scheduler's reference which will be released on task completion
-			{
-			}
-
-			template<typename TaskBodyType>
-			void Init(const TCHAR* DebugName, TaskBodyType&& TaskBody, LowLevelTasks::ETaskPriority Priority)
-			{
-				TTaskWithResultBase<ResultType>::Init(
-					DebugName,
-					Forward<TaskBodyType>(TaskBody),
-					LowLevelTasks::TDeleter<TRefCountedBase<FSelf>, &TRefCountedBase<FSelf>::Release>{ this },
-					Priority
-				);
-			}
-		};
 	}
 }}
\ No newline at end of file
