Allowing configuration of Task priorities from configs.

diff --git a/Engine/Source/Runtime/Core/Private/Async/TaskGraph.cpp b/Engine/Source/Runtime/Core/Private/Async/TaskGraph.cpp
--- a/Engine/Source/Runtime/Core/Private/Async/TaskGraph.cpp
+++ b/Engine/Source/Runtime/Core/Private/Async/TaskGraph.cpp
@@ -338,99 +338,87 @@ FORCEINLINE void TestRandomizedThreads()
 
 #endif
 
-static FString ThreadPriorityToName(ENamedThreads::Type Priority)
+static TCHAR CharFromTaskPriority(ENamedThreads::Type InPriority)
 {
-	if (Priority == ENamedThreads::NormalThreadPriority)
+	if (InPriority == ENamedThreads::HighThreadPriority)
 	{
-		return FString(TEXT("Normal"));
+		return TEXT('h');
 	}
-	if (Priority == ENamedThreads::HighThreadPriority)
-	{
-		return FString(TEXT("High"));
-	}
-	if (Priority == ENamedThreads::BackgroundThreadPriority)
+	return TEXT('n');
+}
+
+static ENamedThreads::Type TaskPriorityFromChar(TCHAR InChar)
+{
+	if (InChar == TEXT('h'))
 	{
-		return FString(TEXT("Background"));
+		return ENamedThreads::HighTaskPriority;
 	}
-	return FString(TEXT("??Unknown??"));
+	return ENamedThreads::NormalTaskPriority;
 }
 
-static FString TaskPriorityToName(ENamedThreads::Type Priority)
+static TCHAR CharFromThreadPriority(ENamedThreads::Type InPriority)
 {
-	if (Priority == ENamedThreads::NormalTaskPriority)
+	if (InPriority == ENamedThreads::NormalThreadPriority)
 	{
-		return FString(TEXT("Normal"));
+		return TEXT('n');
 	}
-	if (Priority == ENamedThreads::HighTaskPriority)
+	if (InPriority == ENamedThreads::HighThreadPriority)
 	{
-		return FString(TEXT("High"));
+		return TEXT('h');
 	}
-	return FString(TEXT("??Unknown??"));
+	return TEXT('b');
 }
 
-void FAutoConsoleTaskPriority::CommandExecute(const TArray<FString>& Args)
+static ENamedThreads::Type ThreadPriorityFromChar(TCHAR InChar)
 {
-	if (Args.Num() > 0)
+	if (InChar == TEXT('n'))
 	{
-		if (Args[0].Compare(ThreadPriorityToName(ENamedThreads::NormalThreadPriority), ESearchCase::IgnoreCase) == 0)
-		{
-			ThreadPriority = ENamedThreads::NormalThreadPriority;
-		}
-		else if (Args[0].Compare(ThreadPriorityToName(ENamedThreads::HighThreadPriority), ESearchCase::IgnoreCase) == 0)
-		{
-			ThreadPriority = ENamedThreads::HighThreadPriority;
-		}
-		else if (Args[0].Compare(ThreadPriorityToName(ENamedThreads::BackgroundThreadPriority), ESearchCase::IgnoreCase) == 0)
-		{
-			ThreadPriority = ENamedThreads::BackgroundThreadPriority;
-		}
-		else
-		{
-			UE_LOG(LogConsoleResponse, Display, TEXT("Could not parse thread priority %s"), *Args[0]);
-		}
+		return ENamedThreads::NormalThreadPriority;
 	}
-	if (Args.Num() > 1)
+	if (InChar == TEXT('h'))
 	{
-		if (Args[1].Compare(TaskPriorityToName(ENamedThreads::NormalTaskPriority), ESearchCase::IgnoreCase) == 0)
-		{
-			TaskPriority = ENamedThreads::NormalTaskPriority;
-		}
-		else if (Args[1].Compare(TaskPriorityToName(ENamedThreads::HighTaskPriority), ESearchCase::IgnoreCase) == 0)
-		{
-			TaskPriority = ENamedThreads::HighTaskPriority;
-		}
-		else
-		{
-			UE_LOG(LogConsoleResponse, Display, TEXT("Could not parse task priority %s"), *Args[1]);
-		}
+		return ENamedThreads::HighThreadPriority;
 	}
-	if (Args.Num() > 2)
+	return ENamedThreads::BackgroundThreadPriority;
+}
+
+FString FAutoConsoleTaskPriority::CreateFullHelpText(const TCHAR* Name, const TCHAR* OriginalHelp)
+{
+	return FString::Printf(
+		TEXT("%s\n")
+		TEXT("Arguments are three characters: [ThreadPriority][TaskPriority][TaskPriorityIfForcedToNormalThreadPriority] ")
+		TEXT("where ThreadPriority is 'h' or 'n' or 'b' (high/normal/background) and TaskPriority is 'h' or 'n' (high/normal). ")
+		TEXT("Example: %s bnh")
+		, OriginalHelp, Name);
+}
+
+FString FAutoConsoleTaskPriority::ConfigStringFromPriorities(ENamedThreads::Type InThreadPriority, ENamedThreads::Type InTaskPriority, ENamedThreads::Type InTaskPriorityBackup)
+{
+	const TCHAR OutName[4] = {
+		CharFromThreadPriority(InThreadPriority),
+		CharFromTaskPriority(InTaskPriority),
+		CharFromTaskPriority(InTaskPriorityBackup),
+		TEXT('\0')
+	};
+	return FString(OutName);
+}
+
+void FAutoConsoleTaskPriority::OnSettingChanged(IConsoleVariable* InVariable)
+{
+	if (RawSetting.Len() > 0)
 	{
-		if (Args[2].Compare(TaskPriorityToName(ENamedThreads::NormalTaskPriority), ESearchCase::IgnoreCase) == 0)
-		{
-			TaskPriorityIfForcedToNormalThreadPriority = ENamedThreads::NormalTaskPriority;
-		}
-		else if (Args[2].Compare(TaskPriorityToName(ENamedThreads::HighTaskPriority), ESearchCase::IgnoreCase) == 0)
-		{
-			TaskPriorityIfForcedToNormalThreadPriority = ENamedThreads::HighTaskPriority;
-		}
-		else
-		{
-			UE_LOG(LogConsoleResponse, Display, TEXT("Could not parse task priority %s"), *Args[2]);
-		}
+		ThreadPriority = ThreadPriorityFromChar(RawSetting[0]);
 	}
-	if (ThreadPriority == ENamedThreads::NormalThreadPriority)
+	if (RawSetting.Len() > 1)
 	{
-		UE_LOG(LogConsoleResponse, Display, TEXT("%s - thread priority:%s   task priority:%s"), *CommandName, *ThreadPriorityToName(ThreadPriority), *TaskPriorityToName(TaskPriority));
+		TaskPriority = TaskPriorityFromChar(RawSetting[1]);
 	}
-	else
+	if (RawSetting.Len() > 2)
 	{
-		UE_LOG(LogConsoleResponse, Display, TEXT("%s - thread priority:%s   task priority:%s  %s (when forced to normal)"), *CommandName, *ThreadPriorityToName(ThreadPriority), *TaskPriorityToName(TaskPriority), *TaskPriorityToName(this->TaskPriorityIfForcedToNormalThreadPriority));
+		TaskPriorityIfForcedToNormalThreadPriority = TaskPriorityFromChar(RawSetting[2]);
 	}
 }
 
-
-
 /** 
  *	FTaskThreadBase
  *	Base class for a thread that executes tasks
diff --git a/Engine/Source/Runtime/Core/Public/Async/TaskGraphInterfaces.h b/Engine/Source/Runtime/Core/Public/Async/TaskGraphInterfaces.h
--- a/Engine/Source/Runtime/Core/Public/Async/TaskGraphInterfaces.h
+++ b/Engine/Source/Runtime/Core/Public/Async/TaskGraphInterfaces.h
@@ -240,16 +240,22 @@ enum class UE_DEPRECATED(4.26, "No longer supported") EPowerSavingEligibility :
 
 class CORE_API FAutoConsoleTaskPriority
 {
-	FAutoConsoleCommand Command;
-	FString CommandName;
+	FString RawSetting;
+	FString FullHelpText;
+	FAutoConsoleVariableRef Variable;
 	ENamedThreads::Type ThreadPriority;
 	ENamedThreads::Type TaskPriority;
 	ENamedThreads::Type TaskPriorityIfForcedToNormalThreadPriority;
-	void CommandExecute(const TArray<FString>& Args);
+
+	static FString CreateFullHelpText(const TCHAR* Name, const TCHAR* OriginalHelp);
+	static FString ConfigStringFromPriorities(ENamedThreads::Type InThreadPriority, ENamedThreads::Type InTaskPriority, ENamedThreads::Type InTaskPriorityBackup);
+	void OnSettingChanged(IConsoleVariable* Variable);
+
 public:
 	FAutoConsoleTaskPriority(const TCHAR* Name, const TCHAR* Help, ENamedThreads::Type DefaultThreadPriority, ENamedThreads::Type DefaultTaskPriority, ENamedThreads::Type DefaultTaskPriorityIfForcedToNormalThreadPriority = ENamedThreads::UnusedAnchor)
-		: Command(Name, Help, FConsoleCommandWithArgsDelegate::CreateRaw(this, &FAutoConsoleTaskPriority::CommandExecute))
-		, CommandName(Name)
+		: RawSetting(ConfigStringFromPriorities(DefaultThreadPriority, DefaultTaskPriority, DefaultTaskPriorityIfForcedToNormalThreadPriority))
+		, FullHelpText(CreateFullHelpText(Name, Help))
+		, Variable(Name, RawSetting, *FullHelpText, FConsoleVariableDelegate::CreateRaw(this, &FAutoConsoleTaskPriority::OnSettingChanged), ECVF_Default)
 		, ThreadPriority(DefaultThreadPriority)
 		, TaskPriority(DefaultTaskPriority)
 		, TaskPriorityIfForcedToNormalThreadPriority(DefaultTaskPriorityIfForcedToNormalThreadPriority)
