FNames can now be constructed with UTF8CHAR strings.
UTF-8 support added to FString::AppendChars() and made able to correctly handle variably-encoded source strings (both wider and narrower than TCHAR).
FUTF8ToTCHAR now accepts either UTF8CHAR* or ANSICHAR* for better compatibility.
FUtf8StringView now binds to UTF8CHAR* or any UTF8CHAR range, and can be constructed via a _U8SV string literal suffix.
CompactBinary now supports UTF-8.
Miscellaneous other UTF8CHAR fixes.

diff --git a/Engine/Source/Developer/DerivedDataCache/Private/DerivedDataBuildOutput.cpp b/Engine/Source/Developer/DerivedDataCache/Private/DerivedDataBuildOutput.cpp
--- a/Engine/Source/Developer/DerivedDataCache/Private/DerivedDataBuildOutput.cpp
+++ b/Engine/Source/Developer/DerivedDataCache/Private/DerivedDataBuildOutput.cpp
@@ -28,9 +28,9 @@ static FStringView LexToString(EBuildDiagnosticLevel Level)
 	}
 }
 
-static void LexFromString(EBuildDiagnosticLevel& OutLevel, FAnsiStringView String)
+static void LexFromString(EBuildDiagnosticLevel& OutLevel, FUtf8StringView String)
 {
-	if (String.Equals("Warning"_ASV, ESearchCase::CaseSensitive))
+	if (String.Equals("Warning"_U8SV, ESearchCase::CaseSensitive))
 	{
 		OutLevel = EBuildDiagnosticLevel::Warning;
 	}
@@ -186,8 +186,8 @@ void FBuildOutputInternal::IterateDiagnostics(TFunctionRef<void (const FBuildDia
 	{
 		EBuildDiagnosticLevel Level;
 		LexFromString(Level, Diagnostic["Level"_ASV].AsString());
-		const FAnsiStringView Category = Diagnostic["Category"_ASV].AsString();
-		const FAnsiStringView Message = Diagnostic["Message"_ASV].AsString();
+		const FUtf8StringView Category = Diagnostic["Category"_ASV].AsString();
+		const FUtf8StringView Message = Diagnostic["Message"_ASV].AsString();
 		Visitor(FBuildDiagnostic{FUTF8ToTCHAR(Category), FUTF8ToTCHAR(Message), Level});
 	}
 }
diff --git a/Engine/Source/Developer/TraceServices/Private/AnalysisCache.cpp b/Engine/Source/Developer/TraceServices/Private/AnalysisCache.cpp
--- a/Engine/Source/Developer/TraceServices/Private/AnalysisCache.cpp
+++ b/Engine/Source/Developer/TraceServices/Private/AnalysisCache.cpp
@@ -192,7 +192,7 @@ bool FAnalysisCache::FFileContents::Load()
 	for (FCbFieldView IndexEntry : IndexArray)
 	{
 		FCbObjectView IndexEntryObj = IndexEntry.AsObjectView();
-		FAnsiStringView NameView = IndexEntryObj["N"_ASV].AsString();
+		FUtf8StringView NameView = IndexEntryObj["N"_ASV].AsString();
 		const uint32 Id = IndexEntryObj["I"_ASV].AsUInt32();
 		const uint32 Flags = IndexEntryObj["F"_ASV].AsUInt32();
 		FMemoryView UserData = IndexEntryObj["UD"_ASV].AsBinaryView();
diff --git a/Engine/Source/Runtime/Core/Private/Containers/String.cpp b/Engine/Source/Runtime/Core/Private/Containers/String.cpp
--- a/Engine/Source/Runtime/Core/Private/Containers/String.cpp
+++ b/Engine/Source/Runtime/Core/Private/Containers/String.cpp
@@ -131,18 +131,43 @@ void AppendCharacters(TArray<TCHAR>& Out, const CharType* Str, int32 Count)
 
 	checkSlow(Str);
 
-	const int32 OldNum = Out.Num();
+	int32 OldEnd = Out.Num();
+	OldEnd -= OldEnd ? 1 : 0;
 
-	// Reserve enough space - including an extra gap for a null terminator if we don't already have a string allocated
-	Out.AddUninitialized(Count + (OldNum ? 0 : 1));
+	// Try to reserve enough space by guessing that the new length will be the same as the input length.
+	// Include an extra gap for a null terminator if we don't already have a string allocated
+	Out.AddUninitialized(Count + (OldEnd ? 0 : 1));
 
-	TCHAR* Dest = Out.GetData() + OldNum - (OldNum ? 1 : 0);
+	TCHAR* Dest = Out.GetData() + OldEnd;
 
-	// Copy characters to end of string, overwriting null terminator if we already have one
-	FPlatformString::Convert(Dest, Count, Str, Count);
+	// Try copying characters to end of string, overwriting null terminator if we already have one
+	TCHAR* NewEnd = FPlatformString::Convert(Dest, Count, Str, Count);
+	if (!NewEnd)
+	{
+		// If that failed, it will have meant that conversion likely contained multi-code unit characters
+		// and so the buffer wasn't long enough, so calculate it properly.
+		int32 Length = FPlatformString::ConvertedLength<TCHAR>(Str, Count);
+
+		// Add the extra bytes that we need
+		Out.AddUninitialized(Length - Count);
+
+		// Restablish destination pointer in case a realloc happened
+		Dest = Out.GetData() + OldEnd;
+
+		bool bSuccess = !!FPlatformString::Convert(Dest, Length, Str, Count);
+		checkSlow(bSuccess);
+	}
+	else
+	{
+		int32 NewEndIndex = (int32)(NewEnd - Dest);
+		if (NewEndIndex < Count)
+		{
+			Out.SetNumUninitialized(OldEnd + NewEndIndex + 1, /*bAllowShrinking=*/false);
+		}
+	}
 
 	// (Re-)establish the null terminator
-	Dest[Count] = '\0';
+	*NewEnd = '\0';
 }
 
 void FString::AppendChars(const ANSICHAR* Str, int32 Count)
@@ -163,6 +188,12 @@ void FString::AppendChars(const UCS2CHAR* Str, int32 Count)
 	AppendCharacters(Data, Str, Count);
 }
 
+void FString::AppendChars(const UTF8CHAR* Str, int32 Count)
+{
+	CheckInvariants();
+	AppendCharacters(Data, Str, Count);
+}
+
 void FString::TrimToNullTerminator()
 {
 	if( Data.Num() )
diff --git a/Engine/Source/Runtime/Core/Private/Containers/StringView.cpp b/Engine/Source/Runtime/Core/Private/Containers/StringView.cpp
--- a/Engine/Source/Runtime/Core/Private/Containers/StringView.cpp
+++ b/Engine/Source/Runtime/Core/Private/Containers/StringView.cpp
@@ -95,3 +95,4 @@ TStringView<CharType> TStringView<CharType>::TrimEnd() const
 
 template class TStringView<ANSICHAR>;
 template class TStringView<WIDECHAR>;
+template class TStringView<UTF8CHAR>;
diff --git a/Engine/Source/Runtime/Core/Private/GenericPlatform/GenericPlatformStricmp.cpp b/Engine/Source/Runtime/Core/Private/GenericPlatform/GenericPlatformStricmp.cpp
--- a/Engine/Source/Runtime/Core/Private/GenericPlatform/GenericPlatformStricmp.cpp
+++ b/Engine/Source/Runtime/Core/Private/GenericPlatform/GenericPlatformStricmp.cpp
@@ -105,6 +105,7 @@ int32 FGenericPlatformStricmp::Stricmp(const UTF16CHAR* Str1, const ANSICHAR* St
 int32 FGenericPlatformStricmp::Stricmp(const UTF32CHAR* Str1, const ANSICHAR* Str2) { return StricmpImpl(Str1, Str2); }
 int32 FGenericPlatformStricmp::Strnicmp(const ANSICHAR* Str1, const ANSICHAR* Str2, SIZE_T Count) { return StrnicmpImpl(Str1, Str2, Count); }
 int32 FGenericPlatformStricmp::Strnicmp(const WIDECHAR* Str1, const WIDECHAR* Str2, SIZE_T Count) { return StrnicmpImpl(Str1, Str2, Count); }
+int32 FGenericPlatformStricmp::Strnicmp(const UTF8CHAR* Str1, const UTF8CHAR* Str2, SIZE_T Count) { return StrnicmpImpl(Str1, Str2, Count); }
 int32 FGenericPlatformStricmp::Strnicmp(const ANSICHAR* Str1, const WIDECHAR* Str2, SIZE_T Count) { return StrnicmpImpl(Str1, Str2, Count); }
 int32 FGenericPlatformStricmp::Strnicmp(const WIDECHAR* Str1, const ANSICHAR* Str2, SIZE_T Count) { return StrnicmpImpl(Str1, Str2, Count); }
 
diff --git a/Engine/Source/Runtime/Core/Private/IO/PackageStore.cpp b/Engine/Source/Runtime/Core/Private/IO/PackageStore.cpp
--- a/Engine/Source/Runtime/Core/Private/IO/PackageStore.cpp
+++ b/Engine/Source/Runtime/Core/Private/IO/PackageStore.cpp
@@ -116,8 +116,7 @@ FPackageStoreEntryResource FPackageStoreEntryResource::FromCbObject(const FCbObj
 		for (FCbField& ArrayField : Obj["shadermaphashes"].AsArray())
 		{
 			FSHAHash& ShaderMapHash = Entry.ShaderMapHashes.AddDefaulted_GetRef();
-			FAnsiStringView AnsiStringView = ArrayField.AsString();
-			ShaderMapHash.FromString(FStringView(ANSI_TO_TCHAR(AnsiStringView.GetData()), AnsiStringView.Len()));
+			ShaderMapHash.FromString(FUTF8ToTCHAR(ArrayField.AsString()));
 		}
 	}
 
diff --git a/Engine/Source/Runtime/Core/Private/Serialization/CompactBinary.cpp b/Engine/Source/Runtime/Core/Private/Serialization/CompactBinary.cpp
--- a/Engine/Source/Runtime/Core/Private/Serialization/CompactBinary.cpp
+++ b/Engine/Source/Runtime/Core/Private/Serialization/CompactBinary.cpp
@@ -153,11 +153,11 @@ FMemoryView FCbFieldView::AsBinaryView(const FMemoryView Default)
 	}
 }
 
-FAnsiStringView FCbFieldView::AsString(const FAnsiStringView Default)
+FUtf8StringView FCbFieldView::AsString(const FUtf8StringView Default)
 {
 	if (FCbFieldType::IsString(TypeWithFlags))
 	{
-		const ANSICHAR* const PayloadChars = static_cast<const ANSICHAR*>(Payload);
+		const UTF8CHAR* const PayloadChars = static_cast<const UTF8CHAR*>(Payload);
 		uint32 ValueSizeByteCount;
 		const uint64 ValueSize = ReadVarUInt(PayloadChars, ValueSizeByteCount);
 
@@ -168,7 +168,7 @@ FAnsiStringView FCbFieldView::AsString(const FAnsiStringView Default)
 		}
 
 		Error = ECbFieldError::None;
-		return FAnsiStringView(PayloadChars + ValueSizeByteCount, int32(ValueSize));
+		return FUtf8StringView(PayloadChars + ValueSizeByteCount, int32(ValueSize));
 	}
 	else
 	{
@@ -456,9 +456,9 @@ FCbCustomByName FCbFieldView::AsCustomByName(FCbCustomByName Default)
 		PayloadBytes += TypeNameLenByteCount;
 
 		FCbCustomByName Value;
-		Value.Name = FAnsiStringView(
-			reinterpret_cast<const ANSICHAR*>(PayloadBytes),
-			static_cast<FAnsiStringView::SizeType>(TypeNameLen));
+		Value.Name = FUtf8StringView(
+			reinterpret_cast<const UTF8CHAR*>(PayloadBytes),
+			static_cast<FUtf8StringView::SizeType>(TypeNameLen));
 		Value.Data = MakeMemoryView(PayloadBytes + TypeNameLen, PayloadSize - TypeNameLen - TypeNameLenByteCount);
 		Error = ECbFieldError::None;
 		return Value;
@@ -484,7 +484,7 @@ FMemoryView FCbFieldView::AsCustom(uint64 Id, FMemoryView Default)
 	}
 }
 
-FMemoryView FCbFieldView::AsCustom(FAnsiStringView Name, FMemoryView Default)
+FMemoryView FCbFieldView::AsCustom(FUtf8StringView Name, FMemoryView Default)
 {
 	const FCbCustomByName Custom = AsCustomByName(FCbCustomByName{Name, Default});
 	if (Custom.Name.Equals(Name, ESearchCase::CaseSensitive))
@@ -621,7 +621,7 @@ FMemoryView FCbFieldView::GetViewNoType() const
 	return MakeMemoryView(static_cast<const uint8*>(Payload) - NameSize, NameSize + PayloadSize);
 }
 
-FCbFieldView FCbFieldView::operator[](FAnsiStringView Name) const
+FCbFieldView FCbFieldView::operator[](FUtf8StringView Name) const
 {
 	switch (GetType())
 	{
@@ -721,7 +721,7 @@ FCbObjectView::FCbObjectView()
 {
 }
 
-FCbFieldView FCbObjectView::FindView(const FAnsiStringView Name) const
+FCbFieldView FCbObjectView::FindView(const FUtf8StringView Name) const
 {
 	for (const FCbFieldView& Field : *this)
 	{
@@ -733,7 +733,7 @@ FCbFieldView FCbObjectView::FindView(const FAnsiStringView Name) const
 	return FCbFieldView();
 }
 
-FCbFieldView FCbObjectView::FindViewIgnoreCase(const FAnsiStringView Name) const
+FCbFieldView FCbObjectView::FindViewIgnoreCase(const FUtf8StringView Name) const
 {
 	for (const FCbFieldView& Field : *this)
 	{
diff --git a/Engine/Source/Runtime/Core/Private/Serialization/CompactBinaryValidation.cpp b/Engine/Source/Runtime/Core/Private/Serialization/CompactBinaryValidation.cpp
--- a/Engine/Source/Runtime/Core/Private/Serialization/CompactBinaryValidation.cpp
+++ b/Engine/Source/Runtime/Core/Private/Serialization/CompactBinaryValidation.cpp
@@ -158,10 +158,10 @@ static FMemoryView ValidateCbDynamicPayload(FMemoryView& View, ECbValidateMode M
  *
  * Modifies the view to start at the end of the string, and adds error flags if applicable.
  */
-static FAnsiStringView ValidateCbString(FMemoryView& View, ECbValidateMode Mode, ECbValidateError& Error)
+static FUtf8StringView ValidateCbString(FMemoryView& View, ECbValidateMode Mode, ECbValidateError& Error)
 {
 	const FMemoryView Payload = ValidateCbDynamicPayload(View, Mode, Error);
-	const FAnsiStringView Value(static_cast<const ANSICHAR*>(Payload.GetData()), static_cast<int32>(Payload.GetSize()));
+	const FUtf8StringView Value(static_cast<const UTF8CHAR*>(Payload.GetData()), static_cast<int32>(Payload.GetSize()));
 	return Value;
 }
 
@@ -219,7 +219,7 @@ static void ValidateCbObject(FMemoryView& View, ECbValidateMode Mode, ECbValidat
 
 	if (Size > 0)
 	{
-		TArray<FAnsiStringView, TInlineAllocator<16>> Names;
+		TArray<FUtf8StringView, TInlineAllocator<16>> Names;
 
 		const bool bUniformObject = FCbFieldType::GetType(ObjectType) == ECbFieldType::UniformObject;
 		const ECbFieldType ExternalType = bUniformObject ? ValidateCbFieldType(ObjectView, Mode, Error) : ECbFieldType::HasFieldType;
@@ -245,8 +245,8 @@ static void ValidateCbObject(FMemoryView& View, ECbValidateMode Mode, ECbValidat
 
 		if (EnumHasAnyFlags(Mode, ECbValidateMode::Names) && Names.Num() > 1)
 		{
-			Algo::Sort(Names, [](FAnsiStringView L, FAnsiStringView R) { return L.Compare(R) < 0; });
-			for (const FAnsiStringView* NamesIt = Names.GetData(), *NamesEnd = NamesIt + Names.Num() - 1; NamesIt != NamesEnd; ++NamesIt)
+			Algo::Sort(Names, [](FUtf8StringView L, FUtf8StringView R) { return L.Compare(R) < 0; });
+			for (const FUtf8StringView* NamesIt = Names.GetData(), *NamesEnd = NamesIt + Names.Num() - 1; NamesIt != NamesEnd; ++NamesIt)
 			{
 				if (NamesIt[0].Equals(NamesIt[1]))
 				{
@@ -296,7 +296,7 @@ static FCbFieldView ValidateCbField(FMemoryView& View, ECbValidateMode Mode, ECb
 {
 	const FMemoryView FieldView = View;
 	const ECbFieldType Type = ValidateCbFieldType(View, Mode, Error, ExternalType);
-	const FAnsiStringView Name = FCbFieldType::HasFieldName(Type) ? ValidateCbString(View, Mode, Error) : FAnsiStringView();
+	const FUtf8StringView Name = FCbFieldType::HasFieldName(Type) ? ValidateCbString(View, Mode, Error) : FUtf8StringView();
 
 	if (EnumHasAnyFlags(Error, ECbValidateError::OutOfBounds | ECbValidateError::InvalidType))
 	{
@@ -370,7 +370,7 @@ static FCbFieldView ValidateCbField(FMemoryView& View, ECbValidateMode Mode, ECb
 	case ECbFieldType::CustomByName:
 	{
 		FMemoryView Value = ValidateCbDynamicPayload(View, Mode, Error);
-		const FAnsiStringView TypeName = ValidateCbString(Value, Mode, Error);
+		const FUtf8StringView TypeName = ValidateCbString(Value, Mode, Error);
 		if (TypeName.IsEmpty() && !EnumHasAnyFlags(Error, ECbValidateError::OutOfBounds))
 		{
 			AddError(Error, ECbValidateError::InvalidType);
diff --git a/Engine/Source/Runtime/Core/Private/Serialization/CompactBinaryWriter.cpp b/Engine/Source/Runtime/Core/Private/Serialization/CompactBinaryWriter.cpp
--- a/Engine/Source/Runtime/Core/Private/Serialization/CompactBinaryWriter.cpp
+++ b/Engine/Source/Runtime/Core/Private/Serialization/CompactBinaryWriter.cpp
@@ -164,7 +164,7 @@ void FCbWriter::EndField(ECbFieldType Type)
 	Data[State.Offset] = uint8(Type);
 }
 
-FCbWriter& FCbWriter::SetName(const FAnsiStringView Name)
+FCbWriter& FCbWriter::SetName(const FUtf8StringView Name)
 {
 	FState& State = States.Last();
 	checkf((State.Flags & EStateFlags::Array) != EStateFlags::Array,
@@ -185,7 +185,7 @@ FCbWriter& FCbWriter::SetName(const FAnsiStringView Name)
 	return *this;
 }
 
-void FCbWriter::SetNameOrAddString(const FAnsiStringView NameOrValue)
+void FCbWriter::SetNameOrAddString(const FUtf8StringView NameOrValue)
 {
 	// A name is only written if it would begin a new field inside of an object.
 	if ((States.Last().Flags & (EStateFlags::Name | EStateFlags::Field | EStateFlags::Object)) == EStateFlags::Object)
@@ -198,7 +198,7 @@ void FCbWriter::SetNameOrAddString(const FAnsiStringView NameOrValue)
 	}
 }
 
-FAnsiStringView FCbWriter::GetActiveName() const
+FUtf8StringView FCbWriter::GetActiveName() const
 {
 	const FState& State = States.Last();
 	if ((State.Flags & EStateFlags::Name) == EStateFlags::Name)
@@ -206,11 +206,11 @@ FAnsiStringView FCbWriter::GetActiveName() const
 		const uint8* const EncodedName = Data.GetData() + State.Offset + sizeof(ECbFieldType);
 		uint32 NameLenByteCount;
 		const uint64 NameLen = ReadVarUInt(EncodedName, NameLenByteCount);
-		using NameLenSizeType = FAnsiStringView::SizeType;
+		using NameLenSizeType = FUtf8StringView::SizeType;
 		const NameLenSizeType ClampedNameLen = NameLenSizeType(FMath::Clamp<uint64>(NameLen, 0, ~NameLenSizeType(0)));
-		return FAnsiStringView(reinterpret_cast<const ANSICHAR*>(EncodedName + NameLenByteCount), ClampedNameLen);
+		return FUtf8StringView(reinterpret_cast<const UTF8CHAR*>(EncodedName + NameLenByteCount), ClampedNameLen);
 	}
-	return FAnsiStringView();
+	return FUtf8StringView();
 }
 
 void FCbWriter::MakeFieldsUniform(const int64 FieldBeginOffset, const int64 FieldEndOffset)
@@ -380,7 +380,7 @@ void FCbWriter::AddBinary(const FCompositeBuffer& Buffer)
 	AddBinary(Buffer.Flatten());
 }
 
-void FCbWriter::AddString(const FAnsiStringView Value)
+void FCbWriter::AddString(const FUtf8StringView Value)
 {
 	BeginField();
 	const uint64 Size = uint64(Value.Len());
@@ -391,7 +391,7 @@ void FCbWriter::AddString(const FAnsiStringView Value)
 	StringData += SizeByteCount;
 	if (Size > 0)
 	{
-		FMemory::Memcpy(StringData, Value.GetData(), Value.Len() * sizeof(ANSICHAR));
+		FMemory::Memcpy(StringData, Value.GetData(), Value.Len() * sizeof(UTF8CHAR));
 	}
 	EndField(ECbFieldType::String);
 }
@@ -537,7 +537,7 @@ void FCbWriter::AddDateTimeTicks(const int64 Ticks)
 	EndField(ECbFieldType::DateTime);
 }
 
-void FCbWriter::AddDateTime(FAnsiStringView Name, FDateTime Value)
+void FCbWriter::AddDateTime(FUtf8StringView Name, FDateTime Value)
 {
 	SetName(Name);
 	AddDateTime(Value);
@@ -556,7 +556,7 @@ void FCbWriter::AddTimeSpanTicks(const int64 Ticks)
 	EndField(ECbFieldType::TimeSpan);
 }
 
-void FCbWriter::AddTimeSpan(FAnsiStringView Name, FTimespan Value)
+void FCbWriter::AddTimeSpan(FUtf8StringView Name, FTimespan Value)
 {
 	SetName(Name);
 	AddTimeSpan(Value);
@@ -588,7 +588,7 @@ void FCbWriter::AddCustom(const uint64 TypeId, const FMemoryView Value)
 	EndField(ECbFieldType::CustomById);
 }
 
-void FCbWriter::AddCustom(const FAnsiStringView TypeName, const FMemoryView Value)
+void FCbWriter::AddCustom(const FUtf8StringView TypeName, const FMemoryView Value)
 {
 	checkf(!TypeName.IsEmpty(), TEXT("Field '%.*hs' requires a non-empty type name for its custom type."),
 		GetActiveName().Len(), GetActiveName().GetData());
diff --git a/Engine/Source/Runtime/Core/Private/Tests/Serialization/CompactBinaryTest.cpp b/Engine/Source/Runtime/Core/Private/Tests/Serialization/CompactBinaryTest.cpp
--- a/Engine/Source/Runtime/Core/Private/Tests/Serialization/CompactBinaryTest.cpp
+++ b/Engine/Source/Runtime/Core/Private/Tests/Serialization/CompactBinaryTest.cpp
@@ -53,7 +53,7 @@ UE_CBFIELD_TYPE_ACCESSOR_NO_DEFAULT(UniformObject, IsObject, AsObjectView, FCbOb
 UE_CBFIELD_TYPE_ACCESSOR_NO_DEFAULT(Array, IsArray, AsArrayView, FCbArrayView);
 UE_CBFIELD_TYPE_ACCESSOR_NO_DEFAULT(UniformArray, IsArray, AsArrayView, FCbArrayView);
 UE_CBFIELD_TYPE_ACCESSOR(Binary, IsBinary, AsBinaryView, FMemoryView);
-UE_CBFIELD_TYPE_ACCESSOR(String, IsString, AsString, FAnsiStringView);
+UE_CBFIELD_TYPE_ACCESSOR(String, IsString, AsString, FUtf8StringView);
 UE_CBFIELD_TYPE_ACCESSOR(IntegerPositive, IsInteger, AsUInt64, uint64);
 UE_CBFIELD_TYPE_ACCESSOR(IntegerNegative, IsInteger, AsInt64, int64);
 UE_CBFIELD_TYPE_ACCESSOR(Float32, IsFloat, AsFloat, float);
@@ -208,7 +208,7 @@ bool FCbFieldNoneTest::RunTest(const FString& Parameters)
 		constexpr const ANSICHAR NoneBytes[] = { ANSICHAR(FieldType), 4, 'N', 'a', 'm', 'e' };
 		FCbFieldView NoneField(NoneBytes);
 		TestEqual(TEXT("FCbFieldView(None|Type|Name)::GetSize()"), NoneField.GetSize(), uint64(sizeof(NoneBytes)));
-		TestEqual(TEXT("FCbFieldView(None|Type|Name)::GetName()"), NoneField.GetName(), "Name"_ASV);
+		TestEqual(TEXT("FCbFieldView(None|Type|Name)::GetName()"), NoneField.GetName(), "Name"_U8SV);
 		TestTrue(TEXT("FCbFieldView(None|Type|Name)::HasName()"), NoneField.HasName());
 		TestFalse(TEXT("!FCbFieldView(None|Type|Name)::HasValue()"), NoneField.HasValue());
 		TestEqual(TEXT("FCbFieldView(None|Type|Name)::GetHash()"), NoneField.GetHash(), FBlake3::HashBuffer(MakeMemoryView(NoneBytes)));
@@ -240,7 +240,7 @@ bool FCbFieldNoneTest::RunTest(const FString& Parameters)
 		constexpr const ANSICHAR NoneBytes[] = { ANSICHAR(FieldType), 4, 'N', 'a', 'm', 'e' };
 		FCbFieldView NoneField(NoneBytes + 1, FieldType);
 		TestEqual(TEXT("FCbFieldView(None|Name)::GetSize()"), NoneField.GetSize(), uint64(sizeof(NoneBytes)));
-		TestEqual(TEXT("FCbFieldView(None|Name)::GetName()"), NoneField.GetName(), "Name"_ASV);
+		TestEqual(TEXT("FCbFieldView(None|Name)::GetName()"), NoneField.GetName(), "Name"_U8SV);
 		TestTrue(TEXT("FCbFieldView(None|Name)::HasName()"), NoneField.HasName());
 		TestFalse(TEXT("!FCbFieldView(None|Name)::HasValue()"), NoneField.HasValue());
 		TestEqual(TEXT("FCbFieldView(None|Name)::GetHash()"), NoneField.GetHash(), FBlake3::HashBuffer(MakeMemoryView(NoneBytes)));
@@ -258,7 +258,7 @@ bool FCbFieldNoneTest::RunTest(const FString& Parameters)
 		constexpr const uint8 NoneBytes[] = { uint8(FieldType), 0 };
 		FCbFieldView NoneField(NoneBytes + 1, FieldType);
 		TestEqual(TEXT("FCbFieldView(None|EmptyName)::GetSize()"), NoneField.GetSize(), uint64(sizeof(NoneBytes)));
-		TestEqual(TEXT("FCbFieldView(None|EmptyName)::GetName()"), NoneField.GetName(), ""_ASV);
+		TestEqual(TEXT("FCbFieldView(None|EmptyName)::GetName()"), NoneField.GetName(), ""_U8SV);
 		TestTrue(TEXT("FCbFieldView(None|EmptyName)::HasName()"), NoneField.HasName());
 		TestFalse(TEXT("!FCbFieldView(None|EmptyName)::HasValue()"), NoneField.HasValue());
 		TestEqual(TEXT("FCbFieldView(None|EmptyName)::GetHash()"), NoneField.GetHash(), FBlake3::HashBuffer(MakeMemoryView(NoneBytes)));
@@ -341,9 +341,9 @@ bool FCbFieldObjectTest::RunTest(const FString& Parameters)
 		TestIntObject(Object, 0, 1);
 
 		// Find fields that do not exist.
-		TestFalse(TEXT("FCbObjectView()::Find(Missing)"), Object.FindView("Field"_ASV).HasValue());
-		TestFalse(TEXT("FCbObjectView()::FindViewIgnoreCase(Missing)"), Object.FindViewIgnoreCase("Field"_ASV).HasValue());
-		TestFalse(TEXT("FCbObjectView()::operator[](Missing)"), Object["Field"_ASV].HasValue());
+		TestFalse(TEXT("FCbObjectView()::Find(Missing)"), Object.FindView("Field"_U8SV).HasValue());
+		TestFalse(TEXT("FCbObjectView()::FindViewIgnoreCase(Missing)"), Object.FindViewIgnoreCase("Field"_U8SV).HasValue());
+		TestFalse(TEXT("FCbObjectView()::operator[](Missing)"), Object["Field"_U8SV].HasValue());
 
 		// Advance an iterator past the last field.
 		FCbFieldViewIterator It = Object.CreateViewIterator();
@@ -445,7 +445,7 @@ bool FCbFieldObjectTest::RunTest(const FString& Parameters)
 		for (FCbFieldIterator It = ObjectClone.CreateIterator(); It; ++It)
 		{
 			FCbField Field = *It;
-			TestEqual(TEXT("FCbObject::CreateIterator().GetName()"), Field.GetName(), "F"_ASV);
+			TestEqual(TEXT("FCbObject::CreateIterator().GetName()"), Field.GetName(), "F"_U8SV);
 			TestEqual(TEXT("FCbObject::CreateIterator().AsInt32()"), Field.AsInt32(), 8);
 			TestTrue(TEXT("FCbObject::CreateIterator().IsOwned()"), Field.IsOwned());
 		}
@@ -748,20 +748,20 @@ bool FCbFieldStringTest::RunTest(const FString& Parameters)
 	// Test FCbFieldView(String, Value)
 	{
 		const uint8 Payload[] = { 3, 'A', 'B', 'C' }; // Size: 3, Data: ABC
-		TestField<ECbFieldType::String>(TEXT("String, Value"), Payload, FAnsiStringView(reinterpret_cast<const ANSICHAR*>(Payload) + 1, 3));
+		TestField<ECbFieldType::String>(TEXT("String, Value"), Payload, FUtf8StringView(reinterpret_cast<const UTF8CHAR*>(Payload) + 1, 3));
 	}
 
 	// Test FCbFieldView(String, OutOfRangeSize)
 	{
 		uint8 Payload[9];
 		WriteVarUInt(uint64(1) << 31, Payload);
-		TestFieldError<ECbFieldType::String>(TEXT("String, OutOfRangeSize"), Payload, ECbFieldError::RangeError, "ABC"_ASV);
+		TestFieldError<ECbFieldType::String>(TEXT("String, OutOfRangeSize"), Payload, ECbFieldError::RangeError, "ABC"_U8SV);
 	}
 
 	// Test FCbFieldView(None) as String
 	{
 		FCbFieldView Field;
-		TestFieldError<ECbFieldType::String>(TEXT("String, None"), Field, ECbFieldError::TypeError, "ABC"_ASV);
+		TestFieldError<ECbFieldType::String>(TEXT("String, None"), Field, ECbFieldError::TypeError, "ABC"_U8SV);
 	}
 
 	return true;
@@ -1261,8 +1261,8 @@ bool FCbFieldCustomByNameTest::RunTest(const FString& Parameters)
 {
 	struct FCustomByNameAccessor
 	{
-		explicit FCustomByNameAccessor(FAnsiStringView Name)
-			: AsType([Name = FString(Name)](FCbFieldView& Field, FMemoryView Default) { return Field.AsCustom(TCHAR_TO_ANSI(*Name), Default); })
+		explicit FCustomByNameAccessor(FUtf8StringView Name)
+			: AsType([Name = FString(Name)](FCbFieldView& Field, FMemoryView Default) { return Field.AsCustom(FTCHARToUTF8(Name), Default); })
 		{
 		}
 
@@ -1273,26 +1273,26 @@ bool FCbFieldCustomByNameTest::RunTest(const FString& Parameters)
 	// Test FCbFieldView(CustomByName, ABC, Empty)
 	{
 		const uint8 Payload[] = {4, 3, 'A', 'B', 'C'};
-		TestField<ECbFieldType::CustomByName>(TEXT("CustomByName, MinId, Empty"), Payload, FCbCustomByName{"ABC"});
-		TestField<ECbFieldType::CustomByName>(TEXT("CustomByName, MinId, Empty, View"), Payload, FMemoryView(), MakeMemoryView<uint8>({1, 2, 3}), ECbFieldError::None, FCustomByNameAccessor("ABC"));
-		TestFieldError<ECbFieldType::CustomByName>(TEXT("CustomByName, MinId, Empty, InvalidCase"), Payload, ECbFieldError::RangeError, MakeMemoryView<uint8>({1, 2, 3}), FCustomByNameAccessor("abc"));
+		TestField<ECbFieldType::CustomByName>(TEXT("CustomByName, MinId, Empty"), Payload, FCbCustomByName{"ABC"_U8SV});
+		TestField<ECbFieldType::CustomByName>(TEXT("CustomByName, MinId, Empty, View"), Payload, FMemoryView(), MakeMemoryView<uint8>({1, 2, 3}), ECbFieldError::None, FCustomByNameAccessor("ABC"_U8SV));
+		TestFieldError<ECbFieldType::CustomByName>(TEXT("CustomByName, MinId, Empty, InvalidCase"), Payload, ECbFieldError::RangeError, MakeMemoryView<uint8>({1, 2, 3}), FCustomByNameAccessor("abc"_U8SV));
 	}
 
 	// Test FCbFieldView(CustomByName, ABC, Value)
 	{
 		const uint8 Payload[] = {8, 3, 'A', 'B', 'C', 1, 2, 3, 4};
-		TestFieldNoClone<ECbFieldType::CustomByName>(TEXT("CustomByName, MinId, Value"), Payload, FCbCustomByName{"ABC", MakeMemoryView(Payload).Right(4)});
-		TestFieldNoClone<ECbFieldType::CustomByName>(TEXT("CustomByName, MinId, Value, View"), Payload, MakeMemoryView(Payload).Right(4), FMemoryView(), ECbFieldError::None, FCustomByNameAccessor("ABC"));
-		TestFieldError<ECbFieldType::CustomByName>(TEXT("CustomByName, MinId, Value, InvalidCase"), Payload, ECbFieldError::RangeError, MakeMemoryView<uint8>({1, 2, 3}), FCustomByNameAccessor("abc"));
+		TestFieldNoClone<ECbFieldType::CustomByName>(TEXT("CustomByName, MinId, Value"), Payload, FCbCustomByName{"ABC"_U8SV, MakeMemoryView(Payload).Right(4)});
+		TestFieldNoClone<ECbFieldType::CustomByName>(TEXT("CustomByName, MinId, Value, View"), Payload, MakeMemoryView(Payload).Right(4), FMemoryView(), ECbFieldError::None, FCustomByNameAccessor("ABC"_U8SV));
+		TestFieldError<ECbFieldType::CustomByName>(TEXT("CustomByName, MinId, Value, InvalidCase"), Payload, ECbFieldError::RangeError, MakeMemoryView<uint8>({1, 2, 3}), FCustomByNameAccessor("abc"_U8SV));
 	}
 
 	// Test FCbFieldView(None) as CustomByName
 	{
 		FCbFieldView DefaultField;
-		TestFieldError<ECbFieldType::CustomByName>(TEXT("CustomByName, None"), DefaultField, ECbFieldError::TypeError, FCbCustomByName{"ABC", MakeMemoryView<uint8>({1, 2, 3})});
-		TestFieldError<ECbFieldType::CustomByName>(TEXT("CustomByName, None, View"), DefaultField, ECbFieldError::TypeError, MakeMemoryView<uint8>({1, 2, 3}), FCustomByNameAccessor("ABC"));
+		TestFieldError<ECbFieldType::CustomByName>(TEXT("CustomByName, None"), DefaultField, ECbFieldError::TypeError, FCbCustomByName{"ABC"_U8SV, MakeMemoryView<uint8>({1, 2, 3})});
+		TestFieldError<ECbFieldType::CustomByName>(TEXT("CustomByName, None, View"), DefaultField, ECbFieldError::TypeError, MakeMemoryView<uint8>({1, 2, 3}), FCustomByNameAccessor("ABC"_U8SV));
 		const uint8 DefaultValue[] = {1, 2, 3};
-		TestEqual(TEXT("FCbFieldView()::AsCustom(Name)"), DefaultField.AsCustom("ABC", MakeMemoryView(DefaultValue)), MakeMemoryView(DefaultValue));
+		TestEqual(TEXT("FCbFieldView()::AsCustom(Name)"), DefaultField.AsCustom("ABC"_U8SV, MakeMemoryView(DefaultValue)), MakeMemoryView(DefaultValue));
 	}
 
 	return true;
@@ -1826,22 +1826,22 @@ bool FCbFieldParseTest::RunTest(const FString& Parameters)
 		for (FCbFieldViewIterator It = Object.CreateViewIterator(); It;)
 		{
 			const FCbFieldViewIterator Last = It;
-			if (It.GetName().Equals("A"_ASV))
+			if (It.GetName().Equals("A"_U8SV))
 			{
 				A = It.AsUInt32();
 				++It;
 			}
-			if (It.GetName().Equals("B"_ASV))
+			if (It.GetName().Equals("B"_U8SV))
 			{
 				B = It.AsUInt32();
 				++It;
 			}
-			if (It.GetName().Equals("C"_ASV))
+			if (It.GetName().Equals("C"_U8SV))
 			{
 				C = It.AsUInt32();
 				++It;
 			}
-			if (It.GetName().Equals("D"_ASV))
+			if (It.GetName().Equals("D"_U8SV))
 			{
 				D = It.AsUInt32();
 				++It;
diff --git a/Engine/Source/Runtime/Core/Private/Tests/Serialization/CompactBinaryWriterTest.cpp b/Engine/Source/Runtime/Core/Private/Tests/Serialization/CompactBinaryWriterTest.cpp
--- a/Engine/Source/Runtime/Core/Private/Tests/Serialization/CompactBinaryWriterTest.cpp
+++ b/Engine/Source/Runtime/Core/Private/Tests/Serialization/CompactBinaryWriterTest.cpp
@@ -191,7 +191,7 @@ bool FCbWriterNullTest::RunTest(const FString& Parameters)
 		FCbField Field = Writer.Save();
 		if (TestEqual(TEXT("FCbWriter(Null, Name) Validate"), ValidateCompactBinary(Field.GetOuterBuffer(), ECbValidateMode::All), ECbValidateError::None))
 		{
-			TestEqual(TEXT("FCbWriter(Null, Name).GetName()"), Field.GetName(), "Null"_ASV);
+			TestEqual(TEXT("FCbWriter(Null, Name).GetName()"), Field.GetName(), "Null"_U8SV);
 			TestTrue(TEXT("FCbWriter(Null, Name).HasName()"), Field.HasName());
 			TestTrue(TEXT("FCbWriter(Null, Name).IsNull()"), Field.IsNull());
 		}
@@ -283,7 +283,7 @@ bool FCbWriterBinaryTest::RunTest(const FString& Parameters)
 		FCbField Field = Writer.Save();
 		if (TestEqual(TEXT("FCbWriter(Binary, Array) Validate"), ValidateCompactBinary(Field.GetOuterBuffer(), ECbValidateMode::All), ECbValidateError::None))
 		{
-			TestEqual(TEXT("FCbWriter(Binary, Array).GetName()"), Field.GetName(), "Binary"_ASV);
+			TestEqual(TEXT("FCbWriter(Binary, Array).GetName()"), Field.GetName(), "Binary"_U8SV);
 			TestTrue(TEXT("FCbWriter(Binary, Array).HasName()"), Field.HasName());
 			TestTrue(TEXT("FCbWriter(Binary, Array).IsBinary()"), Field.IsBinary());
 			TestTrue(TEXT("FCbWriter(Binary, Array).AsBinaryView()"), Field.AsBinaryView().EqualBytes(MakeMemoryView(BinaryValue)));
@@ -327,10 +327,10 @@ bool FCbWriterStringTest::RunTest(const FString& Parameters)
 		{
 			for (FCbFieldView Field : Fields)
 			{
-				TestEqual(TEXT("FCbWriter(String, Basic).GetName()"), Field.GetName(), "String"_ASV);
+				TestEqual(TEXT("FCbWriter(String, Basic).GetName()"), Field.GetName(), "String"_U8SV);
 				TestTrue(TEXT("FCbWriter(String, Basic).HasName()"), Field.HasName());
 				TestTrue(TEXT("FCbWriter(String, Basic).IsString()"), Field.IsString());
-				TestEqual(TEXT("FCbWriter(String, Basic).AsString()"), Field.AsString(), "Value"_ASV);
+				TestEqual(TEXT("FCbWriter(String, Basic).AsString()"), Field.AsString(), "Value"_U8SV);
 			}
 		}
 	}
@@ -339,7 +339,7 @@ bool FCbWriterStringTest::RunTest(const FString& Parameters)
 	{
 		Writer.Reset();
 		constexpr int DotCount = 256;
-		TAnsiStringBuilder<DotCount + 1> Dots;
+		TUtf8StringBuilder<DotCount + 1> Dots;
 		for (int Index = 0; Index < DotCount; ++Index)
 		{
 			Dots << '.';
@@ -351,7 +351,7 @@ bool FCbWriterStringTest::RunTest(const FString& Parameters)
 		{
 			for (FCbFieldView Field : Fields)
 			{
-				TestEqual(TEXT("FCbWriter(String, Long).AsString()"), Field.AsString(), FAnsiStringView(Dots));
+				TestEqual(TEXT("FCbWriter(String, Long).AsString()"), Field.AsString(), Dots.ToView());
 			}
 		}
 	}
@@ -367,7 +367,7 @@ bool FCbWriterStringTest::RunTest(const FString& Parameters)
 		{
 			for (FCbFieldView Field : Fields)
 			{
-				TestEqual(TEXT("FCbWriter(String, Unicode).AsString()"), Field.AsString(), "\xf0\x9f\x98\x80"_ASV);
+				TestEqual(TEXT("FCbWriter(String, Unicode).AsString()"), Field.AsString(), "\xf0\x9f\x98\x80"_U8SV);
 			}
 		}
 	}
@@ -807,13 +807,13 @@ bool FCbWriterCustomByNameTest::RunTest(const FString& Parameters)
 
 	struct FCustomValue
 	{
-		FAnsiStringView Type;
+		FUtf8StringView Type;
 		TArray<uint8, TInlineAllocator<16>> Bytes;
 	};
 	const FCustomValue Values[] =
 	{
-		{ "Type1", {1, 2, 3} },
-		{ "Type2", {4, 5, 6} },
+		{ "Type1"_U8SV, {1, 2, 3} },
+		{ "Type2"_U8SV, {4, 5, 6} },
 	};
 
 	for (const FCustomValue& Value : Values)
@@ -879,6 +879,7 @@ bool FCbWriterComplexTest::RunTest(const FString& Parameters)
 			Writer.AddString("WideString"_ASV, FString::ChrN(256, TEXT('.')));
 			Writer.AddString("EmptyAnsiString"_ASV, FAnsiStringView());
 			Writer.AddString("EmptyWideString"_ASV, FWideStringView());
+			Writer.AddString("EmptyUtf8String"_U8SV, FUtf8StringView());
 			Writer.AddString("AnsiStringLiteral", "AnsiValue");
 			Writer.AddString("WideStringLiteral", TEXT("AnsiValue"));
 		}
diff --git a/Engine/Source/Runtime/Core/Private/UObject/UnrealNames.cpp b/Engine/Source/Runtime/Core/Private/UObject/UnrealNames.cpp
--- a/Engine/Source/Runtime/Core/Private/UObject/UnrealNames.cpp
+++ b/Engine/Source/Runtime/Core/Private/UObject/UnrealNames.cpp
@@ -1837,6 +1837,15 @@ struct FNameAnsiStringView
 	int32 Len;
 };
 
+struct FNameUtf8StringViewWithWidth
+{
+	using CharType = UTF8CHAR;
+
+	const UTF8CHAR* Str;
+	int32 Len;
+	bool bIsWide;
+};
+
 struct FWideStringViewWithWidth
 {
 	using CharType = WIDECHAR;
@@ -1856,25 +1865,27 @@ static FNameAnsiStringView MakeUnconvertedView(const ANSICHAR* Str)
 	return { Str, Str ? FCStringAnsi::Strlen(Str) : 0 };
 }
 
-static bool IsWide(const WIDECHAR* Str, const int32 Len)
+template <typename CharType>
+static bool IsWide(const CharType* Str, const int32 Len)
 {
 	uint32 UserCharBits = 0;
 	for (int32 I = 0; I < Len; ++I)
 	{
-		UserCharBits |= TChar<WIDECHAR>::ToUnsigned(Str[I]);
+		UserCharBits |= TChar<CharType>::ToUnsigned(Str[I]);
 	}
 	return UserCharBits & 0xffffff80u;
 }
 
-static int32 GetLengthAndWidth(const WIDECHAR* Str, bool& bOutIsWide)
+template <typename CharType>
+static int32 GetLengthAndWidth(const CharType* Str, bool& bOutIsWide)
 {
 	uint32 UserCharBits = 0;
-	const WIDECHAR* It = Str;
+	const CharType* It = Str;
 	if (Str)
 	{
 		while (*It)
 		{
-			UserCharBits |= TChar<WIDECHAR>::ToUnsigned(*It);
+			UserCharBits |= TChar<CharType>::ToUnsigned(*It);
 			++It;
 		}
 	}
@@ -1884,6 +1895,19 @@ static int32 GetLengthAndWidth(const WIDECHAR* Str, bool& bOutIsWide)
 	return UE_PTRDIFF_TO_INT32(It - Str);
 }
 
+static FNameUtf8StringViewWithWidth MakeUnconvertedView(const UTF8CHAR* Str, int32 Len)
+{
+	return { Str, Len, IsWide(Str, Len) };
+}
+
+static FNameUtf8StringViewWithWidth MakeUnconvertedView(const UTF8CHAR* Str)
+{
+	FNameUtf8StringViewWithWidth View;
+	View.Str = Str;
+	View.Len = GetLengthAndWidth(Str, View.bIsWide);
+	return View;
+}
+
 static FWideStringViewWithWidth MakeUnconvertedView(const WIDECHAR* Str, int32 Len)
 {
 	return { Str, Len, IsWide(Str, Len) };
@@ -1968,6 +1992,26 @@ struct FNameHelper
 		return Make(FNameStringView(View.Str, View.Len), FindType, InternalNumber);
 	}
 
+	static FName MakeWithNumber(FNameUtf8StringViewWithWidth View, EFindName FindType, int32 InternalNumber)
+	{
+		// Ignore the supplied number if the name string is empty
+		// to keep the semantics of the old FName implementation
+		if (View.Len == 0)
+		{
+			return FName();
+		}
+
+		if (!View.bIsWide)
+		{
+			return Make(FNameStringView(reinterpret_cast<const ANSICHAR*>(View.Str), View.Len), FindType, InternalNumber);
+		}
+		else
+		{
+			TStringConversion<FUTF8ToTCHAR_Convert, NAME_SIZE> WideName(View.Str, View.Len);
+			return Make(FNameStringView(WideName.Get(), WideName.Length()), FindType, InternalNumber);
+		}
+	}
+
 	static FName MakeWithNumber(const FWideStringViewWithWidth View, EFindName FindType, int32 InternalNumber)
 	{
 		// Ignore the supplied number if the name string is empty
@@ -2108,6 +2152,10 @@ FName::FName(const ANSICHAR* Name, EFindName FindType)
 	: FName(FNameHelper::MakeDetectNumber(MakeUnconvertedView(Name), FindType))
 {}
 
+FName::FName(const UTF8CHAR* Name, EFindName FindType)
+	: FName(FNameHelper::MakeDetectNumber(MakeUnconvertedView(Name), FindType))
+{}
+
 FName::FName(int32 Len, const WIDECHAR* Name, EFindName FindType)
 	: FName(FNameHelper::MakeDetectNumber(MakeUnconvertedView(Name, Len), FindType))
 {}
@@ -2116,6 +2164,10 @@ FName::FName(int32 Len, const ANSICHAR* Name, EFindName FindType)
 	: FName(FNameHelper::MakeDetectNumber(MakeUnconvertedView(Name, Len), FindType))
 {}
 
+FName::FName(int32 Len, const UTF8CHAR* Name, EFindName FindType)
+	: FName(FNameHelper::MakeDetectNumber(MakeUnconvertedView(Name, Len), FindType))
+{}
+
 FName::FName(const WIDECHAR* Name, int32 InNumber, EFindName FindType)
 	: FName(FNameHelper::MakeWithNumber(MakeUnconvertedView(Name), FindType, InNumber))
 {}
@@ -2124,6 +2176,10 @@ FName::FName(const ANSICHAR* Name, int32 InNumber, EFindName FindType)
 	: FName(FNameHelper::MakeWithNumber(MakeUnconvertedView(Name), FindType, InNumber))
 {}
 
+FName::FName(const UTF8CHAR* Name, int32 InNumber, EFindName FindType)
+	: FName(FNameHelper::MakeWithNumber(MakeUnconvertedView(Name), FindType, InNumber))
+{}
+
 FName::FName(int32 Len, const WIDECHAR* Name, int32 InNumber, EFindName FindType)
 	: FName(InNumber != NAME_NO_NUMBER_INTERNAL ? FNameHelper::MakeWithNumber(MakeUnconvertedView(Name, Len), FindType, InNumber)
 												: FNameHelper::MakeDetectNumber(MakeUnconvertedView(Name, Len), FindType))
@@ -2134,6 +2190,11 @@ FName::FName(int32 Len, const ANSICHAR* Name, int32 InNumber, EFindName FindType
 												: FNameHelper::MakeDetectNumber(MakeUnconvertedView(Name, Len), FindType))
 {}
 
+FName::FName(int32 Len, const UTF8CHAR* Name, int32 InNumber, EFindName FindType)
+	: FName(InNumber != NAME_NO_NUMBER_INTERNAL ? FNameHelper::MakeWithNumber(MakeUnconvertedView(Name, Len), FindType, InNumber)
+												: FNameHelper::MakeDetectNumber(MakeUnconvertedView(Name, Len), FindType))
+{}
+
 FName::FName(const TCHAR* Name, int32 InNumber, EFindName FindType, bool bSplitName)
 	: FName(InNumber == NAME_NO_NUMBER_INTERNAL && bSplitName 
 			? FNameHelper::MakeDetectNumber(MakeUnconvertedView(Name), FindType)
diff --git a/Engine/Source/Runtime/Core/Public/Containers/StringConv.h b/Engine/Source/Runtime/Core/Public/Containers/StringConv.h
--- a/Engine/Source/Runtime/Core/Public/Containers/StringConv.h
+++ b/Engine/Source/Runtime/Core/Public/Containers/StringConv.h
@@ -163,6 +163,15 @@ namespace StringConv
 		StrBuffer.SetNum(NewStrLen + 1, /*bAllowShrinking*/false);
 #endif	// PLATFORM_TCHAR_IS_4_BYTES
 	}
+
+	struct FUnused;
+
+	// Helper functions for getting the conversion type from a converter
+	template <typename ConverterType>
+	typename ConverterType::LegacyFromType GetLegacyFromType(typename ConverterType::LegacyFromType*);
+
+	template <typename ConverterType>
+	FUnused GetLegacyFromType(...);
 }
 
 namespace UE4StringConv_Private
@@ -325,11 +334,9 @@ public:
 
 class FUTF8ToTCHAR_Convert
 {
-	// Leave this until we convert all the call sites and can change FromType to this.
-	using IntendedFromType = UTF8CHAR;
-
 public:
-	typedef ANSICHAR FromType;
+	typedef ANSICHAR LegacyFromType;
+	typedef UTF8CHAR FromType;
 	typedef TCHAR    ToType;
 
 	/**
@@ -346,7 +353,11 @@ public:
 	>
 	static FORCEINLINE void Convert(ToType* Dest, const int32 DestLen, const SrcBufferType* Source, const int32 SourceLen)
 	{
-		FPlatformString::Convert(Dest, DestLen, (const IntendedFromType*)Source, SourceLen);
+		FPlatformString::Convert(Dest, DestLen, (const FromType*)Source, SourceLen);
+	}
+	static FORCEINLINE void Convert(ToType* Dest, const int32 DestLen, const LegacyFromType* Source, const int32 SourceLen)
+	{
+		FPlatformString::Convert(Dest, DestLen, (const FromType*)Source, SourceLen);
 	}
 
 	/**
@@ -362,7 +373,11 @@ public:
 	>
 	static int32 ConvertedLength(const SrcBufferType* Source, const int32 SourceLen)
 	{
-		return FPlatformString::ConvertedLength<ToType>((const IntendedFromType*)Source, SourceLen);
+		return FPlatformString::ConvertedLength<ToType>((const FromType*)Source, SourceLen);
+	}
+	static int32 ConvertedLength(const LegacyFromType* Source, const int32 SourceLen)
+	{
+		return FPlatformString::ConvertedLength<ToType>((const FromType*)Source, SourceLen);
 	}
 };
 
@@ -639,6 +654,8 @@ class TStringConversion : private Converter, private TInlineAllocator<DefaultCon
 	typedef typename Converter::FromType FromType;
 	typedef typename Converter::ToType   ToType;
 
+	using LegacyFromType = decltype(StringConv::GetLegacyFromType<Converter>(nullptr));
+
 	/**
 	 * Converts the data by using the Convert() method on the base class
 	 */
@@ -671,6 +688,10 @@ public:
 			StringLength = 0;
 		}
 	}
+	explicit TStringConversion(const LegacyFromType* Source)
+		: TStringConversion((const FromType*)Source)
+	{
+	}
 
 	template <
 		typename SrcBufferType,
@@ -696,6 +717,10 @@ public:
 			StringLength = 0;
 		}
 	}
+	TStringConversion(const LegacyFromType* Source, int32 SourceLen)
+		: TStringConversion((const FromType*)Source, SourceLen)
+	{
+	}
 
 	/**
 	 * Construct from a compatible character range such as TStringView or TStringBuilder.
diff --git a/Engine/Source/Runtime/Core/Public/Containers/StringView.h b/Engine/Source/Runtime/Core/Public/Containers/StringView.h
--- a/Engine/Source/Runtime/Core/Public/Containers/StringView.h
+++ b/Engine/Source/Runtime/Core/Public/Containers/StringView.h
@@ -37,7 +37,8 @@ namespace StringViewPrivate
 	{
 		static constexpr bool Value = TOr<
 			TIsConvertibleFromTo<T, FAnsiStringView>,
-			TIsConvertibleFromTo<T, FWideStringView>
+			TIsConvertibleFromTo<T, FWideStringView>,
+			TIsConvertibleFromTo<T, FUtf8StringView>
 			>::Value;
 	};
 
@@ -48,7 +49,14 @@ namespace StringViewPrivate
 		using Type =
 			typename TChooseClass<TIsConvertibleFromTo<T, FAnsiStringView>::Value, FAnsiStringView,
 			typename TChooseClass<TIsConvertibleFromTo<T, FWideStringView>::Value, FWideStringView,
-			NotCompatible>::Result>::Result;
+			typename TChooseClass<TIsConvertibleFromTo<T, FUtf8StringView>::Value, FUtf8StringView,
+			NotCompatible>::Result>::Result>::Result;
+	};
+
+	template <typename CharRangeType, typename ElementType>
+	struct TIsCompatibleRangeType
+	{
+		static constexpr bool Value = FPlatformString::IsCharEncodingCompatibleWith<typename TRemoveCV<typename TRemovePointer<decltype(StringViewPrivate::WrapGetData(DeclVal<CharRangeType&>()))>::Type>::Type, ElementType>();
 	};
 }
 
@@ -110,7 +118,7 @@ private:
 	using TIsCharRange = TAnd<
 		TIsContiguousContainer<CharRangeType>,
 		TNot<TIsArray<typename TRemoveReference<CharRangeType>::Type>>,
-		TIsSame<ElementType, typename TRemoveCV<typename TRemovePointer<decltype(StringViewPrivate::WrapGetData(DeclVal<CharRangeType&>()))>::Type>::Type>>;
+		StringViewPrivate::TIsCompatibleRangeType<CharRangeType, ElementType>>;
 
 public:
 	/** Construct an empty view. */
@@ -121,6 +129,15 @@ public:
 	 *
 	 * The caller is responsible for ensuring that the provided character range remains valid for the lifetime of the view.
 	 */
+	template <typename OtherCharType,
+		typename TEnableIf<
+			FPlatformString::IsCharEncodingCompatibleWith<OtherCharType, CharType>()
+		>::Type* = nullptr>
+	constexpr inline TStringView(const OtherCharType* InData)
+		: DataPtr((const CharType*)InData)
+		, Size(InData ? TCString<CharType>::Strlen((const CharType*)InData) : 0)
+	{
+	}
 	constexpr inline TStringView(const CharType* InData)
 		: DataPtr(InData)
 		, Size(InData ? TCString<CharType>::Strlen(InData) : 0)
@@ -132,6 +149,15 @@ public:
 	 *
 	 * The caller is responsible for ensuring that the provided character range remains valid for the lifetime of the view.
 	 */
+	template <typename OtherCharType,
+		typename TEnableIf<
+			FPlatformString::IsCharEncodingCompatibleWith<OtherCharType, CharType>()
+		>::Type* = nullptr>
+	constexpr inline TStringView(const OtherCharType* InData, SizeType InSize)
+		: DataPtr((const CharType*)InData)
+		, Size(InSize)
+	{
+	}
 	constexpr inline TStringView(const CharType* InData, SizeType InSize)
 		: DataPtr(InData)
 		, Size(InSize)
@@ -149,11 +175,24 @@ public:
 			TIsCharRange<CharRangeType>
 		>::Value>::Type* = nullptr>
 	constexpr inline TStringView(CharRangeType&& InRange)
-		: DataPtr(StringViewPrivate::WrapGetData(InRange))
+		: DataPtr((const CharType*)StringViewPrivate::WrapGetData(InRange))
 		, Size(static_cast<SizeType>(GetNum(InRange)))
 	{
 	}
 
+	/**
+	 * Construct a view from a compatible view.
+	 */
+	template <typename OtherCharType,
+		typename TEnableIf<
+			!TIsSame<CharType, OtherCharType>::Value &&
+			FPlatformString::IsCharEncodingCompatibleWith<OtherCharType, CharType>()>::Type* = nullptr>
+	constexpr inline TStringView(TStringView<OtherCharType> InView)
+		: DataPtr(reinterpret_cast<const CharType*>(InView.GetData()))
+		, Size(InView.Len())
+	{
+	}
+
 	/** Access the character at the given index in the view. */
 	inline const CharType& operator[](SizeType Index) const;
 
@@ -338,9 +377,27 @@ constexpr inline auto GetNum(TStringView<CharType> String)
 
 //////////////////////////////////////////////////////////////////////////
 
-constexpr inline FStringView operator "" _SV(const TCHAR* String, size_t Size) { return FStringView(String, Size); }
-constexpr inline FAnsiStringView operator "" _ASV(const ANSICHAR* String, size_t Size) { return FAnsiStringView(String, Size); }
-constexpr inline FWideStringView operator "" _WSV(const WIDECHAR* String, size_t Size) { return FWideStringView(String, Size); }
+constexpr inline FStringView operator "" _SV(const TCHAR* String, size_t Size)
+{
+	return FStringView(String, Size);
+}
+
+constexpr inline FAnsiStringView operator "" _ASV(const ANSICHAR* String, size_t Size)
+{
+	return FAnsiStringView(String, Size);
+}
+
+constexpr inline FWideStringView operator "" _WSV(const WIDECHAR* String, size_t Size)
+{
+	return FWideStringView(String, Size);
+}
+
+/*constexpr*/ inline FUtf8StringView operator "" _U8SV(const ANSICHAR* String, size_t Size)
+{
+	// Would like this operator to be constexpr, but cannot be until after this operator can take a UTF8CHAR*
+	// rather than an ANSICHAR*, which won't be until we have C++20 char8_t string literals.
+	return FUtf8StringView(reinterpret_cast<const UTF8CHAR*>(String), Size);
+}
 
 //////////////////////////////////////////////////////////////////////////
 
diff --git a/Engine/Source/Runtime/Core/Public/Containers/UnrealString.h b/Engine/Source/Runtime/Core/Public/Containers/UnrealString.h
--- a/Engine/Source/Runtime/Core/Public/Containers/UnrealString.h
+++ b/Engine/Source/Runtime/Core/Public/Containers/UnrealString.h
@@ -38,6 +38,13 @@ int32        GetNum(const FString& String);
 /**
  * A dynamically sizeable string.
  * @see https://docs.unrealengine.com/latest/INT/Programming/UnrealArchitecture/StringHandling/FString/
+ *
+ * When dealing with UTF-8 literals, the following advice is recommended:
+ *
+ * - Do not use the u8"..." prefix (gives the wrong array type until C++20).
+ * - Use UTF8TEXT("...") for array literals (type is const UTF8CHAR[n]).
+ * - Use "..."_U8SV for string view literals (type is FUtf8StringView).
+ * - Use \uxxxx or \Uxxxxxxxx escape sequences rather than \x to specify Unicode code points.
  */
 class CORE_API FString
 {
@@ -495,6 +502,7 @@ public:
 	void AppendChars(const ANSICHAR* Str, int32 Count);
 	void AppendChars(const WIDECHAR* Str, int32 Count);
 	void AppendChars(const UCS2CHAR* Str, int32 Count);
+	void AppendChars(const UTF8CHAR* Str, int32 Count);
 
 	/** Append a string and return a reference to this */
 	template<class CharType>
diff --git a/Engine/Source/Runtime/Core/Public/GenericPlatform/GenericPlatformStricmp.h b/Engine/Source/Runtime/Core/Public/GenericPlatform/GenericPlatformStricmp.h
--- a/Engine/Source/Runtime/Core/Public/GenericPlatform/GenericPlatformStricmp.h
+++ b/Engine/Source/Runtime/Core/Public/GenericPlatform/GenericPlatformStricmp.h
@@ -31,6 +31,7 @@ struct FGenericPlatformStricmp
 
 	CORE_API static int32 Strnicmp(const ANSICHAR* String1, const ANSICHAR* String2, SIZE_T Count);
 	CORE_API static int32 Strnicmp(const WIDECHAR* String1, const WIDECHAR* String2, SIZE_T Count);
+	CORE_API static int32 Strnicmp(const UTF8CHAR* String1, const UTF8CHAR* String2, SIZE_T Count);
 	CORE_API static int32 Strnicmp(const ANSICHAR* String1, const WIDECHAR* String2, SIZE_T Count);
 	CORE_API static int32 Strnicmp(const WIDECHAR* String1, const ANSICHAR* String2, SIZE_T Count);
 };
diff --git a/Engine/Source/Runtime/Core/Public/Misc/Char.h b/Engine/Source/Runtime/Core/Public/Misc/Char.h
--- a/Engine/Source/Runtime/Core/Public/Misc/Char.h
+++ b/Engine/Source/Runtime/Core/Public/Misc/Char.h
@@ -170,3 +170,5 @@ template <> inline bool TChar<ANSICHAR>::IsDigit(ANSICHAR Char) { return ::isdig
 template <> inline bool TChar<ANSICHAR>::IsHexDigit(ANSICHAR Char) { return ::isxdigit((unsigned char)Char) != 0; }
 template <> inline bool TChar<ANSICHAR>::IsWhitespace(ANSICHAR Char) { return ::isspace((unsigned char)Char) != 0; }
 template <> inline bool TChar<ANSICHAR>::IsControl(ANSICHAR Char) { return ::iscntrl((unsigned char)Char) != 0; }
+
+template <> inline bool TChar<UTF8CHAR>::IsWhitespace(UTF8CHAR Char) { return ::isspace((unsigned char)Char) != 0; }
diff --git a/Engine/Source/Runtime/Core/Public/Misc/StringBuilder.h b/Engine/Source/Runtime/Core/Public/Misc/StringBuilder.h
--- a/Engine/Source/Runtime/Core/Public/Misc/StringBuilder.h
+++ b/Engine/Source/Runtime/Core/Public/Misc/StringBuilder.h
@@ -12,6 +12,7 @@
 #include "Templates/IsValidVariadicFunctionArg.h"
 #include "Templates/UnrealTemplate.h"
 #include "Traits/IsContiguousContainer.h"
+#include <type_traits>
 
 /**
  * String Builder
@@ -433,6 +434,7 @@ private:
 inline FAnsiStringBuilderBase&		operator<<(FAnsiStringBuilderBase& Builder, ANSICHAR Char)							{ return Builder.Append(Char); }
 inline FWideStringBuilderBase&		operator<<(FWideStringBuilderBase& Builder, ANSICHAR Char)							{ return Builder.Append(Char); }
 inline FWideStringBuilderBase&		operator<<(FWideStringBuilderBase& Builder, WIDECHAR Char)							{ return Builder.Append(Char); }
+inline FUtf8StringBuilderBase&		operator<<(FUtf8StringBuilderBase& Builder, ANSICHAR Char)							{ return Builder.Append(UTF8CHAR(Char)); }
 inline FUtf8StringBuilderBase&		operator<<(FUtf8StringBuilderBase& Builder, UTF8CHAR Char)							{ return Builder.Append(Char); }
 
 template <typename T>
@@ -442,9 +444,17 @@ inline auto operator<<(FAnsiStringBuilderBase& Builder, T&& Str) -> decltype(Bui
 }
 
 template <typename T>
-inline auto operator<<(FWideStringBuilderBase& Builder, T&& Str) -> decltype(Builder.AppendAnsi(ImplicitConv<FAnsiStringView>(Forward<T>(Str))))
+inline auto operator<<(FWideStringBuilderBase& Builder, T&& Str) -> decltype(Builder.AppendUtf8(ImplicitConv<FUtf8StringView>(Forward<T>(Str))))
 {
-	return Builder.AppendAnsi(ImplicitConv<FAnsiStringView>(Forward<T>(Str)));
+	// Anything convertible to an FAnsiStringView is also convertible to a FUtf8StringView, but FAnsiStringView is more efficient to convert
+	if constexpr (std::is_convertible_v<T, FAnsiStringView>)
+	{
+		return Builder.AppendAnsi(ImplicitConv<FAnsiStringView>(Forward<T>(Str)));
+	}
+	else
+	{
+		return Builder.AppendUtf8(ImplicitConv<FUtf8StringView>(Forward<T>(Str)));
+	}
 }
 
 template <typename T>
@@ -456,13 +466,20 @@ inline auto operator<<(FWideStringBuilderBase& Builder, T&& Str) -> decltype(Bui
 template <typename T>
 inline auto operator<<(FUtf8StringBuilderBase& Builder, T&& Str) -> decltype(Builder.Append(ImplicitConv<FUtf8StringView>(Forward<T>(Str))))
 {
-	return Builder.AppendUtf8(ImplicitConv<FUtf8StringView>(Forward<T>(Str)));
+	// Anything convertible to an FAnsiStringView is also convertible to a FUtf8StringView, but FAnsiStringView is more efficient to convert
+	if constexpr (std::is_convertible_v<T, FAnsiStringView>)
+	{
+		return Builder.AppendAnsi(ImplicitConv<FAnsiStringView>(Forward<T>(Str)));
+	}
+	else
+	{
+		return Builder.AppendUtf8(ImplicitConv<FUtf8StringView>(Forward<T>(Str)));
+	}
 }
 
 // Prefer using << instead of += as operator+= is only intended for mechanical FString -> FStringView replacement.
 inline FStringBuilderBase&			operator+=(FStringBuilderBase& Builder, ANSICHAR Char)								{ return Builder.Append(Char); }
 inline FStringBuilderBase&			operator+=(FStringBuilderBase& Builder, WIDECHAR Char)								{ return Builder.Append(Char); }
-inline FStringBuilderBase&			operator+=(FStringBuilderBase& Builder, FAnsiStringView Str)						{ return Builder.AppendAnsi(Str); }
 inline FStringBuilderBase&			operator+=(FStringBuilderBase& Builder, FWideStringView Str)						{ return Builder.Append(Str); }
 inline FStringBuilderBase&			operator+=(FStringBuilderBase& Builder, FUtf8StringView Str)						{ return Builder.AppendUtf8(Str); }
 
diff --git a/Engine/Source/Runtime/Core/Public/Serialization/CompactBinary.h b/Engine/Source/Runtime/Core/Public/Serialization/CompactBinary.h
--- a/Engine/Source/Runtime/Core/Public/Serialization/CompactBinary.h
+++ b/Engine/Source/Runtime/Core/Public/Serialization/CompactBinary.h
@@ -541,7 +541,7 @@ struct FCbCustomById
 struct FCbCustomByName
 {
 	/** An identifier for the sub-type of the field. Lifetime is tied to the field that the name is associated with. */
-	FAnsiStringView Name;
+	FUtf8StringView Name;
 	/** A view of the value. Lifetime is tied to the field that the value is associated with. */
 	FMemoryView Data;
 };
@@ -583,9 +583,9 @@ public:
 	CORE_API explicit FCbFieldView(const void* Data, ECbFieldType Type = ECbFieldType::HasFieldType);
 
 	/** Returns the name of the field if it has a name, otherwise an empty view. */
-	constexpr inline FAnsiStringView GetName() const
+	constexpr inline FUtf8StringView GetName() const
 	{
-		return FAnsiStringView(static_cast<const ANSICHAR*>(Payload) - NameLen, NameLen);
+		return FUtf8StringView(static_cast<const UTF8CHAR*>(Payload) - NameLen, NameLen);
 	}
 
 	/** Access the field as an object. Defaults to an empty object on error. */
@@ -597,8 +597,8 @@ public:
 	/** Access the field as binary. Returns the provided default on error. */
 	CORE_API FMemoryView AsBinaryView(FMemoryView Default = FMemoryView());
 
-	/** Access the field as a UTF-8 string. Returns the provided default on error. */
-	CORE_API FAnsiStringView AsString(FAnsiStringView Default = FAnsiStringView());
+	/** Access the field as a string. Returns the provided default on error. */
+	CORE_API FUtf8StringView AsString(FUtf8StringView Default = FUtf8StringView());
 
 	/** Access the field as an int8. Returns the provided default on error. */
 	inline int8 AsInt8(int8 Default = 0)       { return AsInteger<int8>(Default); }
@@ -667,7 +667,7 @@ public:
 	/** Access the field as a custom sub-type with an integer identifier. Returns the provided default on error. */
 	CORE_API FMemoryView AsCustom(uint64 Id, FMemoryView Default = FMemoryView());
 	/** Access the field as a custom sub-type with a string identifier. Returns the provided default on error. */
-	CORE_API FMemoryView AsCustom(FAnsiStringView Name, FMemoryView Default = FMemoryView());
+	CORE_API FMemoryView AsCustom(FUtf8StringView Name, FMemoryView Default = FMemoryView());
 
 	/** True if the field has a name. */
 	constexpr inline bool HasName() const           { return FCbFieldType::HasFieldName(TypeWithFlags); }
@@ -763,7 +763,7 @@ public:
 	}
 
 	/** Find a field of an object by case-sensitive name comparison, otherwise a field with no value. */
-	CORE_API FCbFieldView operator[](FAnsiStringView Name) const;
+	CORE_API FCbFieldView operator[](FUtf8StringView Name) const;
 
 	/** Create an iterator for the fields of an array or object, otherwise an empty iterator. */
 	CORE_API FCbFieldViewIterator CreateViewIterator() const;
@@ -1005,13 +1005,13 @@ public:
 	 * @param Name The name of the field.
 	 * @return The matching field if found, otherwise a field with no value.
 	 */
-	CORE_API FCbFieldView FindView(FAnsiStringView Name) const;
+	CORE_API FCbFieldView FindView(FUtf8StringView Name) const;
 
 	/** Find a field by case-insensitive name comparison. */
-	CORE_API FCbFieldView FindViewIgnoreCase(FAnsiStringView Name) const;
+	CORE_API FCbFieldView FindViewIgnoreCase(FUtf8StringView Name) const;
 
 	/** Find a field by case-sensitive name comparison. */
-	inline FCbFieldView operator[](FAnsiStringView Name) const { return FindView(Name); }
+	inline FCbFieldView operator[](FUtf8StringView Name) const { return FindView(Name); }
 
 	/** Access the object as an object field. */
 	inline FCbFieldView AsFieldView() const { return static_cast<const FCbFieldView&>(*this); }
@@ -1146,7 +1146,7 @@ public:
 	inline FSharedBuffer GetOuterBuffer() && { return MoveTemp(Buffer); }
 
 	/** Find a field of an object by case-sensitive name comparison, otherwise a field with no value. */
-	inline FCbField operator[](FAnsiStringView Name) const;
+	inline FCbField operator[](FUtf8StringView Name) const;
 
 	/** Create an iterator for the fields of an array or object, otherwise an empty iterator. */
 	inline FCbFieldIterator CreateIterator() const;
@@ -1236,7 +1236,7 @@ public:
 };
 
 template <typename ViewType>
-inline FCbField TCbBuffer<ViewType>::operator[](FAnsiStringView Name) const
+inline FCbField TCbBuffer<ViewType>::operator[](FUtf8StringView Name) const
 {
 	if (FCbFieldView Field = ViewType::operator[](Name))
 	{
@@ -1357,7 +1357,7 @@ public:
 	using TCbBuffer::TCbBuffer;
 
 	/** Find a field by case-sensitive name comparison. */
-	inline FCbField Find(FAnsiStringView Name) const
+	inline FCbField Find(FUtf8StringView Name) const
 	{
 		if (::FCbFieldView Field = FindView(Name))
 		{
@@ -1367,7 +1367,7 @@ public:
 	}
 
 	/** Find a field by case-insensitive name comparison. */
-	inline FCbField FindIgnoreCase(FAnsiStringView Name) const
+	inline FCbField FindIgnoreCase(FUtf8StringView Name) const
 	{
 		if (::FCbFieldView Field = FindViewIgnoreCase(Name))
 		{
@@ -1377,7 +1377,7 @@ public:
 	}
 
 	/** Find a field by case-sensitive name comparison. */
-	inline FCbField operator[](FAnsiStringView Name) const { return Find(Name); }
+	inline FCbField operator[](FUtf8StringView Name) const { return Find(Name); }
 
 	/** Access the object as an object field. */
 	inline FCbField AsField() const & { return FCbField(FCbObjectView::AsFieldView(), *this); }
diff --git a/Engine/Source/Runtime/Core/Public/Serialization/CompactBinaryWriter.h b/Engine/Source/Runtime/Core/Public/Serialization/CompactBinaryWriter.h
--- a/Engine/Source/Runtime/Core/Public/Serialization/CompactBinaryWriter.h
+++ b/Engine/Source/Runtime/Core/Public/Serialization/CompactBinaryWriter.h
@@ -107,132 +107,132 @@ public:
 	 * It is not valid to call this function when writing a field inside an array.
 	 * Names must be valid UTF-8 and must be unique within an object.
 	 */
-	CORE_API FCbWriter& SetName(FAnsiStringView Name);
+	CORE_API FCbWriter& SetName(FUtf8StringView Name);
 
 	/** Copy the value (not the name) of an existing field. */
-	inline void AddField(FAnsiStringView Name, const FCbFieldView& Value) { SetName(Name); AddField(Value); }
+	inline void AddField(FUtf8StringView Name, const FCbFieldView& Value) { SetName(Name); AddField(Value); }
 	CORE_API void AddField(const FCbFieldView& Value);
 	/** Copy the value (not the name) of an existing field. Holds a reference if owned. */
-	inline void AddField(FAnsiStringView Name, const FCbField& Value) { SetName(Name); AddField(Value); }
+	inline void AddField(FUtf8StringView Name, const FCbField& Value) { SetName(Name); AddField(Value); }
 	CORE_API void AddField(const FCbField& Value);
 
 	/** Begin a new object. Must have a matching call to EndObject. */
-	inline void BeginObject(FAnsiStringView Name) { SetName(Name); BeginObject(); }
+	inline void BeginObject(FUtf8StringView Name) { SetName(Name); BeginObject(); }
 	CORE_API void BeginObject();
 	/** End an object after its fields have been written. */
 	CORE_API void EndObject();
 
 	/** Copy the value (not the name) of an existing object. */
-	inline void AddObject(FAnsiStringView Name, const FCbObjectView& Value) { SetName(Name); AddObject(Value); }
+	inline void AddObject(FUtf8StringView Name, const FCbObjectView& Value) { SetName(Name); AddObject(Value); }
 	CORE_API void AddObject(const FCbObjectView& Value);
 	/** Copy the value (not the name) of an existing object. Holds a reference if owned. */
-	inline void AddObject(FAnsiStringView Name, const FCbObject& Value) { SetName(Name); AddObject(Value); }
+	inline void AddObject(FUtf8StringView Name, const FCbObject& Value) { SetName(Name); AddObject(Value); }
 	CORE_API void AddObject(const FCbObject& Value);
 
 	/** Begin a new array. Must have a matching call to EndArray. */
-	inline void BeginArray(FAnsiStringView Name) { SetName(Name); BeginArray(); }
+	inline void BeginArray(FUtf8StringView Name) { SetName(Name); BeginArray(); }
 	CORE_API void BeginArray();
 	/** End an array after its fields have been written. */
 	CORE_API void EndArray();
 
 	/** Copy the value (not the name) of an existing array. */
-	inline void AddArray(FAnsiStringView Name, const FCbArrayView& Value) { SetName(Name); AddArray(Value); }
+	inline void AddArray(FUtf8StringView Name, const FCbArrayView& Value) { SetName(Name); AddArray(Value); }
 	CORE_API void AddArray(const FCbArrayView& Value);
 	/** Copy the value (not the name) of an existing array. Holds a reference if owned. */
-	inline void AddArray(FAnsiStringView Name, const FCbArray& Value) { SetName(Name); AddArray(Value); }
+	inline void AddArray(FUtf8StringView Name, const FCbArray& Value) { SetName(Name); AddArray(Value); }
 	CORE_API void AddArray(const FCbArray& Value);
 
 	/** Write a null field. */
-	inline void AddNull(FAnsiStringView Name) { SetName(Name); AddNull(); }
+	inline void AddNull(FUtf8StringView Name) { SetName(Name); AddNull(); }
 	CORE_API void AddNull();
 
 	/** Write a binary field by copying Size bytes from Value. */
-	inline void AddBinary(FAnsiStringView Name, const void* Value, uint64 Size) { SetName(Name); AddBinary(Value, Size); }
+	inline void AddBinary(FUtf8StringView Name, const void* Value, uint64 Size) { SetName(Name); AddBinary(Value, Size); }
 	CORE_API void AddBinary(const void* Value, uint64 Size);
 	/** Write a binary field by copying the view. */
-	inline void AddBinary(FAnsiStringView Name, FMemoryView Value) { SetName(Name); AddBinary(Value); }
+	inline void AddBinary(FUtf8StringView Name, FMemoryView Value) { SetName(Name); AddBinary(Value); }
 	inline void AddBinary(FMemoryView Value) { AddBinary(Value.GetData(), Value.GetSize()); }
 	/** Write a binary field by copying the buffer. Holds a reference if owned. */
-	inline void AddBinary(FAnsiStringView Name, const FSharedBuffer& Value) { SetName(Name); AddBinary(Value); }
+	inline void AddBinary(FUtf8StringView Name, const FSharedBuffer& Value) { SetName(Name); AddBinary(Value); }
 	CORE_API void AddBinary(const FSharedBuffer& Value);
-	inline void AddBinary(FAnsiStringView Name, const FCompositeBuffer& Value) { SetName(Name); AddBinary(Value); }
+	inline void AddBinary(FUtf8StringView Name, const FCompositeBuffer& Value) { SetName(Name); AddBinary(Value); }
 	CORE_API void AddBinary(const FCompositeBuffer& Value);
 
 	/** Write a string field by copying the UTF-8 value. */
-	inline void AddString(FAnsiStringView Name, FAnsiStringView Value) { SetName(Name); AddString(Value); }
-	CORE_API void AddString(FAnsiStringView Value);
+	inline void AddString(FUtf8StringView Name, FUtf8StringView Value) { SetName(Name); AddString(Value); }
+	CORE_API void AddString(FUtf8StringView Value);
 	/** Write a string field by converting the UTF-16 value to UTF-8. */
-	inline void AddString(FAnsiStringView Name, FWideStringView Value) { SetName(Name); AddString(Value); }
+	inline void AddString(FUtf8StringView Name, FWideStringView Value) { SetName(Name); AddString(Value); }
 	CORE_API void AddString(FWideStringView Value);
 
 	/** Write an integer field. */
-	inline void AddInteger(FAnsiStringView Name, int32 Value) { SetName(Name); AddInteger(Value); }
+	inline void AddInteger(FUtf8StringView Name, int32 Value) { SetName(Name); AddInteger(Value); }
 	CORE_API void AddInteger(int32 Value);
 	/** Write an integer field. */
-	inline void AddInteger(FAnsiStringView Name, int64 Value) { SetName(Name); AddInteger(Value); }
+	inline void AddInteger(FUtf8StringView Name, int64 Value) { SetName(Name); AddInteger(Value); }
 	CORE_API void AddInteger(int64 Value);
 	/** Write an integer field. */
-	inline void AddInteger(FAnsiStringView Name, uint32 Value) { SetName(Name); AddInteger(Value); }
+	inline void AddInteger(FUtf8StringView Name, uint32 Value) { SetName(Name); AddInteger(Value); }
 	CORE_API void AddInteger(uint32 Value);
 	/** Write an integer field. */
-	inline void AddInteger(FAnsiStringView Name, uint64 Value) { SetName(Name); AddInteger(Value); }
+	inline void AddInteger(FUtf8StringView Name, uint64 Value) { SetName(Name); AddInteger(Value); }
 	CORE_API void AddInteger(uint64 Value);
 
 	/** Write a float field from a 32-bit float value. */
-	inline void AddFloat(FAnsiStringView Name, float Value) { SetName(Name); AddFloat(Value); }
+	inline void AddFloat(FUtf8StringView Name, float Value) { SetName(Name); AddFloat(Value); }
 	CORE_API void AddFloat(float Value);
 	/** Write a float field from a 64-bit float value. */
-	inline void AddFloat(FAnsiStringView Name, double Value) { SetName(Name); AddFloat(Value); }
+	inline void AddFloat(FUtf8StringView Name, double Value) { SetName(Name); AddFloat(Value); }
 	CORE_API void AddFloat(double Value);
 
 	/** Write a bool field. */
-	inline void AddBool(FAnsiStringView Name, bool bValue) { SetName(Name); AddBool(bValue); }
+	inline void AddBool(FUtf8StringView Name, bool bValue) { SetName(Name); AddBool(bValue); }
 	CORE_API void AddBool(bool bValue);
 
 	/** Write a field referencing an object attachment by its hash. */
-	inline void AddObjectAttachment(FAnsiStringView Name, const FIoHash& Value) { SetName(Name); AddObjectAttachment(Value); }
+	inline void AddObjectAttachment(FUtf8StringView Name, const FIoHash& Value) { SetName(Name); AddObjectAttachment(Value); }
 	CORE_API void AddObjectAttachment(const FIoHash& Value);
 	/** Write a field referencing a binary attachment by its hash. */
-	inline void AddBinaryAttachment(FAnsiStringView Name, const FIoHash& Value) { SetName(Name); AddBinaryAttachment(Value); }
+	inline void AddBinaryAttachment(FUtf8StringView Name, const FIoHash& Value) { SetName(Name); AddBinaryAttachment(Value); }
 	CORE_API void AddBinaryAttachment(const FIoHash& Value);
 	/** Write a field referencing the attachment by its hash. */
-	inline void AddAttachment(FAnsiStringView Name, const FCbAttachment& Attachment) { SetName(Name); AddAttachment(Attachment); }
+	inline void AddAttachment(FUtf8StringView Name, const FCbAttachment& Attachment) { SetName(Name); AddAttachment(Attachment); }
 	CORE_API void AddAttachment(const FCbAttachment& Attachment);
 
 	/** Write a hash field. */
-	inline void AddHash(FAnsiStringView Name, const FIoHash& Value) { SetName(Name); AddHash(Value); }
+	inline void AddHash(FUtf8StringView Name, const FIoHash& Value) { SetName(Name); AddHash(Value); }
 	CORE_API void AddHash(const FIoHash& Value);
 	/** Write a UUID field. */
-	inline void AddUuid(FAnsiStringView Name, const FGuid& Value) { SetName(Name); AddUuid(Value); }
+	inline void AddUuid(FUtf8StringView Name, const FGuid& Value) { SetName(Name); AddUuid(Value); }
 	CORE_API void AddUuid(const FGuid& Value);
 
 	/** Write a date/time field with the specified count of 100ns ticks since the epoch. */
-	inline void AddDateTimeTicks(FAnsiStringView Name, int64 Ticks) { SetName(Name); AddDateTimeTicks(Ticks); }
+	inline void AddDateTimeTicks(FUtf8StringView Name, int64 Ticks) { SetName(Name); AddDateTimeTicks(Ticks); }
 	CORE_API void AddDateTimeTicks(int64 Ticks);
 
 	/** Write a date/time field. */
-	CORE_API void AddDateTime(FAnsiStringView Name, FDateTime Value);
+	CORE_API void AddDateTime(FUtf8StringView Name, FDateTime Value);
 	CORE_API void AddDateTime(FDateTime Value);
 
 	/** Write a time span field with the specified count of 100ns ticks. */
-	inline void AddTimeSpanTicks(FAnsiStringView Name, int64 Ticks) { SetName(Name); AddTimeSpanTicks(Ticks); }
+	inline void AddTimeSpanTicks(FUtf8StringView Name, int64 Ticks) { SetName(Name); AddTimeSpanTicks(Ticks); }
 	CORE_API void AddTimeSpanTicks(int64 Ticks);
 
 	/** Write a time span field. */
-	CORE_API void AddTimeSpan(FAnsiStringView Name, FTimespan Value);
+	CORE_API void AddTimeSpan(FUtf8StringView Name, FTimespan Value);
 	CORE_API void AddTimeSpan(FTimespan Value);
 
 	/** Write an ObjectId field. */
-	inline void AddObjectId(FAnsiStringView Name, const FCbObjectId& Value) { SetName(Name); AddObjectId(Value); }
+	inline void AddObjectId(FUtf8StringView Name, const FCbObjectId& Value) { SetName(Name); AddObjectId(Value); }
 	CORE_API void AddObjectId(const FCbObjectId& Value);
 
 	/** Write a custom field with an integer sub-type identifier. */
-	inline void AddCustom(FAnsiStringView FieldName, uint64 TypeId, FMemoryView Value) { SetName(FieldName); AddCustom(TypeId, Value); }
+	inline void AddCustom(FUtf8StringView FieldName, uint64 TypeId, FMemoryView Value) { SetName(FieldName); AddCustom(TypeId, Value); }
 	CORE_API void AddCustom(uint64 TypeId, FMemoryView Value);
 
 	/** Write a custom field with a string sub-type identifier. */
-	inline void AddCustom(FAnsiStringView FieldName, FAnsiStringView TypeName, FMemoryView Value) { SetName(FieldName); AddCustom(TypeName, Value); }
-	CORE_API void AddCustom(FAnsiStringView TypeName, FMemoryView Value);
+	inline void AddCustom(FUtf8StringView FieldName, FUtf8StringView TypeName, FMemoryView Value) { SetName(FieldName); AddCustom(TypeName, Value); }
+	CORE_API void AddCustom(FUtf8StringView TypeName, FMemoryView Value);
 
 	/** Private flags that are public to work with ENUM_CLASS_FLAGS. */
 	enum class EStateFlags : uint8;
@@ -242,7 +242,7 @@ protected:
 	CORE_API explicit FCbWriter(int64 InitialSize);
 
 private:
-	friend FCbWriter& operator<<(FCbWriter& Writer, FAnsiStringView NameOrValue);
+	friend FCbWriter& operator<<(FCbWriter& Writer, FUtf8StringView NameOrValue);
 
 	/** Begin writing a field. May be called twice for named fields. */
 	void BeginField();
@@ -251,10 +251,10 @@ private:
 	void EndField(ECbFieldType Type);
 
 	/** Set the field name if valid in this state, otherwise write add a string field. */
-	CORE_API void SetNameOrAddString(FAnsiStringView NameOrValue);
+	CORE_API void SetNameOrAddString(FUtf8StringView NameOrValue);
 
 	/** Returns a view of the name of the active field, if any, otherwise the empty view. */
-	FAnsiStringView GetActiveName() const;
+	FUtf8StringView GetActiveName() const;
 
 	/** Remove field types after the first to make the sequence uniform. */
 	void MakeFieldsUniform(int64 FieldBeginOffset, int64 FieldEndOffset);
@@ -311,7 +311,7 @@ private:
 ///////////////////////////////////////////////////////////////////////////////////////////////////////////////////////
 
 /** Write the field name if valid in this state, otherwise write the string value. */
-inline FCbWriter& operator<<(FCbWriter& Writer, FAnsiStringView NameOrValue)
+inline FCbWriter& operator<<(FCbWriter& Writer, FUtf8StringView NameOrValue)
 {
 	Writer.SetNameOrAddString(NameOrValue);
 	return Writer;
@@ -323,6 +323,12 @@ inline FCbWriter& operator<<(FCbWriter& Writer, const ANSICHAR* NameOrValue)
 	return Writer << FAnsiStringView(NameOrValue);
 }
 
+/** Write the field name if valid in this state, otherwise write the string value. */
+inline FCbWriter& operator<<(FCbWriter& Writer, const UTF8CHAR* NameOrValue)
+{
+	return Writer << FUtf8StringView(NameOrValue);
+}
+
 inline FCbWriter& operator<<(FCbWriter& Writer, const FCbFieldView& Value)
 {
 	Writer.AddField(Value);
diff --git a/Engine/Source/Runtime/Core/Public/UObject/NameTypes.h b/Engine/Source/Runtime/Core/Public/UObject/NameTypes.h
--- a/Engine/Source/Runtime/Core/Public/UObject/NameTypes.h
+++ b/Engine/Source/Runtime/Core/Public/UObject/NameTypes.h
@@ -780,10 +780,12 @@ public:
 	 */
 	FName(const WIDECHAR* Name, EFindName FindType=FNAME_Add);
 	FName(const ANSICHAR* Name, EFindName FindType=FNAME_Add);
+	FName(const UTF8CHAR* Name, EFindName FindType=FNAME_Add);
 
 	/** Create FName from non-null string with known length  */
 	FName(int32 Len, const WIDECHAR* Name, EFindName FindType=FNAME_Add);
 	FName(int32 Len, const ANSICHAR* Name, EFindName FindType=FNAME_Add);
+	FName(int32 Len, const UTF8CHAR* Name, EFindName FindType=FNAME_Add);
 
 	template <typename CharRangeType,
 		typename CharType = typename TRemoveCV<typename TRemovePointer<decltype(GetData(DeclVal<CharRangeType>()))>::Type>::Type,
@@ -806,8 +808,10 @@ public:
 	 */
 	FName(const WIDECHAR* Name, int32 InNumber, EFindName FindType = FNAME_Add);
 	FName(const ANSICHAR* Name, int32 InNumber, EFindName FindType = FNAME_Add);
+	FName(const UTF8CHAR* Name, int32 InNumber, EFindName FindType = FNAME_Add);
 	FName(int32 Len, const WIDECHAR* Name, int32 Number, EFindName FindType = FNAME_Add);
 	FName(int32 Len, const ANSICHAR* Name, int32 InNumber, EFindName FindType = FNAME_Add);
+	FName(int32 Len, const UTF8CHAR* Name, int32 InNumber, EFindName FindType = FNAME_Add);
 
 	template <typename CharRangeType,
 		typename CharType = typename TRemoveCV<typename TRemovePointer<decltype(GetData(DeclVal<CharRangeType>()))>::Type>::Type,
