Archive support CompressionFormat other than NAME_Zlib
v2 header tag to indicate format is serialized in data
writes new data, continues to load old data

diff --git a/Engine/Source/Runtime/Core/Private/Compression/CompressionUtil.cpp b/Engine/Source/Runtime/Core/Private/Compression/CompressionUtil.cpp
new file mode 100644
--- /dev/null
+++ b/Engine/Source/Runtime/Core/Private/Compression/CompressionUtil.cpp
@@ -0,0 +1,92 @@
+#include "Compression/CompressionUtil.h"
+#include "Serialization/Archive.h"
+#include "UObject/NameTypes.h"
+#include "Logging/LogMacros.h"
+
+DECLARE_LOG_CATEGORY_EXTERN(LogCompression, Log, All);
+
+namespace FCompressionUtil
+{
+
+
+	// Serialize common compressor names in a byte
+	//	to avoid going through String serialization and String-to-Name
+	void CORE_API SerializeCompressorName(FArchive & Archive,FName & Compressor)
+	{
+		if ( Archive.IsLoading() )
+		{
+			uint8 CompressorNum;
+			Archive << CompressorNum;
+
+			switch(CompressorNum)
+			{
+			case 0:
+			{
+				// can't rely on archive serializing FName, so use String
+				FString LoadedString;
+				Archive << LoadedString;
+				Compressor = FName(*LoadedString);
+				break;
+			}
+			case 1:
+				Compressor = NAME_None;
+				break;
+			case 2:
+				Compressor = NAME_Oodle;
+				break;
+			case 3:
+				Compressor = NAME_Zlib;
+				break;
+			case 4:
+				Compressor = NAME_Gzip;
+				break;
+			case 5:
+				Compressor = NAME_LZ4;
+				break;
+			default:		
+				UE_LOG(LogCompression, Error, TEXT("SerializeCompressorName Unknown index:%d"),CompressorNum );
+				
+				Compressor = NAME_None;
+				break;
+			}
+		}
+		else
+		{
+			uint8 CompressorNum = 0;
+
+			if ( Compressor == NAME_None )
+			{
+				CompressorNum = 1;
+			}
+			else if ( Compressor == NAME_Oodle )
+			{
+				CompressorNum = 2;
+			}
+			else if ( Compressor == NAME_Zlib )
+			{
+				CompressorNum = 3;
+			}
+			else if ( Compressor == NAME_Gzip )
+			{
+				CompressorNum = 4;
+			}
+			else if ( Compressor == NAME_LZ4 )
+			{
+				CompressorNum = 5;
+			}
+			// else CompressorNum = 0 is right
+
+			
+			Archive << CompressorNum;
+
+			if ( CompressorNum == 0 )
+			{
+				// can't rely on archive serializing FName, so use String
+				FString SavedString(Compressor.ToString());
+				Archive << SavedString;
+			}			
+		}
+	}
+
+
+};
diff --git a/Engine/Source/Runtime/Core/Private/Serialization/Archive.cpp b/Engine/Source/Runtime/Core/Private/Serialization/Archive.cpp
--- a/Engine/Source/Runtime/Core/Private/Serialization/Archive.cpp
+++ b/Engine/Source/Runtime/Core/Private/Serialization/Archive.cpp
@@ -23,6 +23,8 @@
 #include "Serialization/CompressedChunkInfo.h"
 #include "Serialization/ArchiveSerializedPropertyChain.h"
 #include "UObject/NameTypes.h"
+#include "Compression/CompressionUtil.h"
+#include "UObject/UnrealNames.h"
 
 PRAGMA_DISABLE_UNSAFE_TYPECAST_WARNINGS
 
@@ -664,7 +666,7 @@ public:
 		, CompressedSize(0)
 		, UncompressedSize(0)
 		, BitWindow(DEFAULT_ZLIB_BIT_WINDOW)
-		, CompressionFormat(NAME_Zlib)
+		, CompressionFormat(NAME_None)
 		, Flags(COMPRESS_NoFlags)
 	{
 	}
@@ -673,13 +675,6 @@ public:
 	 */
 	void DoWork()
 	{
-		// upgrade old flag method
-		if ((Flags & COMPRESS_DeprecatedFormatFlagsMask) != 0)
-		{
-			UE_LOG(LogSerialization, Warning, TEXT("Old style compression flags are being used with FAsyncCompressionChunk, please update any code using this!"));
-			CompressionFormat = FCompression::GetCompressionFormatFromDeprecatedFlags(Flags);
-		}
-
 		// Compress from memory to memory.
 		verify( FCompression::CompressMemory(CompressionFormat, CompressedBuffer, CompressedSize, UncompressedBuffer, UncompressedSize, Flags, BitWindow) );
 	}
@@ -691,55 +686,132 @@ public:
 };
 #endif		// WITH_MULTI_THREADED_COMPRESSION
 
-void FArchive::SerializeCompressed(void* V, int64 Length, FName CompressionFormat, ECompressionFlags Flags, bool bTreatBufferAsFileReader)
+
+
+void FArchive::SerializeCompressed(void* V, int64 Length, FName CompressionFormatCannotChange, ECompressionFlags Flags, bool bTreatBufferAsFileReader)
+{
+	// with this legacy/deprecated API you can NOT change the CompressionFormat and still load old files
+	//  CompressionFormat must match exactly the format that was written to the file
+
+	SerializeCompressedNew(V,Length,CompressionFormatCannotChange,CompressionFormatCannotChange,Flags,bTreatBufferAsFileReader);
+}
+
+void FArchive::SerializeCompressedNew(void* V, int64 Length, FName CompressionFormatToEncode, FName CompressionFormatToDecodeOldV1Files,  ECompressionFlags Flags, bool bTreatBufferAsFileReader)
 {
 	TRACE_CPUPROFILER_EVENT_SCOPE(FArchive::SerializeCompressed);
+
+	// CompressionFormatToEncode can be changed freely without breaking loading of old files
+	// CompressionFormatToDecodeOldV1Files must match what was used to encode old files, cannot change
+
 	if( IsLoading() )
 	{
-		if (CompressionFormat == NAME_Zlib && FPlatformProperties::GetZlibReplacementFormat() != nullptr)
-		{
-			// use this platform's replacement format in case it's not zlib
-			CompressionFormat = FPlatformProperties::GetZlibReplacementFormat();
-		}
-
 		// Serialize package file tag used to determine endianess.
 		FCompressedChunkInfo PackageFileTag;
 		PackageFileTag.CompressedSize	= 0;
 		PackageFileTag.UncompressedSize	= 0;
 		*this << PackageFileTag;
-		bool bWasByteSwapped = PackageFileTag.CompressedSize != PACKAGE_FILE_TAG;
 
-		// Read in base summary.
-		FCompressedChunkInfo Summary;
-		*this << Summary;
+		// v1 header did not store CompressionFormatToDecode
+		//	assume it was CompressionFormatToDecodeOldV1Files (usually Zlib)
+		FName CompressionFormatToDecode = CompressionFormatToDecodeOldV1Files;
+
+		bool bHeaderWasValid=false;
+		bool bWasByteSwapped=false;
+		bool bReadCompressionFormat=false;
+		
+		// FPackageFileSummary has int32 Tag == PACKAGE_FILE_TAG
+		// this header does not otherwise match FPackageFileSummary in any way
 
-		bool bHeaderWasValid = true;
+		// low 32 bits of ARCHIVE_V2_HEADER_TAG are == PACKAGE_FILE_TAG
+		#define ARCHIVE_V2_HEADER_TAG	(PACKAGE_FILE_TAG | ((uint64)0x22222222<<32) )
 
-		if (bWasByteSwapped)
+		if ( PackageFileTag.CompressedSize == PACKAGE_FILE_TAG )
+		{
+			// v1 header, not swapped
+			bHeaderWasValid = true;
+		}
+		else if ( PackageFileTag.CompressedSize == PACKAGE_FILE_TAG_SWAPPED ||
+			PackageFileTag.CompressedSize == BYTESWAP_ORDER64((uint64)PACKAGE_FILE_TAG) )
+		{
+			// v1 header, swapped
+			bHeaderWasValid = true;
+			bWasByteSwapped = true;
+		}
+		else if ( PackageFileTag.CompressedSize == ARCHIVE_V2_HEADER_TAG ||
+			PackageFileTag.CompressedSize == BYTESWAP_ORDER64((uint64)ARCHIVE_V2_HEADER_TAG) )
+		{
+			// v2 header
+			bHeaderWasValid = true;
+			bWasByteSwapped = ( PackageFileTag.CompressedSize != ARCHIVE_V2_HEADER_TAG );
+			bReadCompressionFormat = true;
+
+			// read CompressionFormatToDecode
+			FCompressionUtil::SerializeCompressorName(*this,CompressionFormatToDecode);
+		}
+		else
 		{
-			bHeaderWasValid = PackageFileTag.CompressedSize == PACKAGE_FILE_TAG_SWAPPED;
-			if (bHeaderWasValid)
+			UE_LOG(LogSerialization, Log, TEXT("ArchiveName: %s"), *GetArchiveName());
+			UE_LOG(LogSerialization, Log, TEXT("Archive UE4 Version: %d"), UEVer());
+			UE_LOG(LogSerialization, Log, TEXT("Archive Licensee Version: %d"), LicenseeUEVer());
+			UE_LOG(LogSerialization, Log, TEXT("Position: %lld"), Tell());
+			UE_LOG(LogSerialization, Log, TEXT("Read Size: %lld"), Length);
+			UE_LOG(LogSerialization, Fatal, TEXT("BulkData compressed header read error. This package may be corrupt!"));
+		}
+		
+		if ( ! bReadCompressionFormat )
+		{
+			// upgrade old flag method
+			if ((Flags & COMPRESS_DeprecatedFormatFlagsMask) != 0)
 			{
-				Summary.CompressedSize = BYTESWAP_ORDER64(Summary.CompressedSize);
-				Summary.UncompressedSize = BYTESWAP_ORDER64(Summary.UncompressedSize);
-				PackageFileTag.UncompressedSize = BYTESWAP_ORDER64(PackageFileTag.UncompressedSize);
+				UE_LOG(LogSerialization, Warning, TEXT("Old style compression flags are being used with FAsyncCompressionChunk, please update any code using this!"));
+				CompressionFormatToDecode = FCompression::GetCompressionFormatFromDeprecatedFlags(Flags);
+			}
+
+			if (CompressionFormatToDecode == NAME_Zlib && FPlatformProperties::GetZlibReplacementFormat() != nullptr)
+			{
+				// use this platform's replacement format in case it's not zlib
+				CompressionFormatToDecode = FPlatformProperties::GetZlibReplacementFormat();
 			}
 		}
 		else
-		{
-			bHeaderWasValid = PackageFileTag.CompressedSize == PACKAGE_FILE_TAG; //-V547
+		{		
+			// shouldn't need to do this step for v2 headers; zlib should have already been changed by the encoder
+			//	need to verify that's working right for xb1
+			if (CompressionFormatToDecode == NAME_Zlib && FPlatformProperties::GetZlibReplacementFormat() != nullptr)
+			{
+				const char * ZlibReplacement = FPlatformProperties::GetZlibReplacementFormat();
+				
+				// go ahead and do it but warn :
+				CompressionFormatToDecode = FName(ZlibReplacement);
+
+				UE_LOG(LogSerialization, Warning, TEXT("Archive v2 header with ZLib not ZlibReplacement: %s"), ZlibReplacement);
+			}
 		}
 
-		if (!bHeaderWasValid)
+		// CompressionFormatToDecode came from disk, need to validate it :
+		if ( ! FCompression::IsFormatValid(CompressionFormatToDecode) )
 		{
 			UE_LOG(LogSerialization, Log, TEXT("ArchiveName: %s"), *GetArchiveName());
 			UE_LOG(LogSerialization, Log, TEXT("Archive UE4 Version: %d"), UEVer());
 			UE_LOG(LogSerialization, Log, TEXT("Archive Licensee Version: %d"), LicenseeUEVer());
 			UE_LOG(LogSerialization, Log, TEXT("Position: %lld"), Tell());
 			UE_LOG(LogSerialization, Log, TEXT("Read Size: %lld"), Length);
+			UE_LOG(LogSerialization, Log, TEXT("CompressionFormatToDecode not found : %s"), *CompressionFormatToDecode.ToString());
 			UE_LOG(LogSerialization, Fatal, TEXT("BulkData compressed header read error. This package may be corrupt!"));
 		}
 
+		// Read in base summary, contains total sizes :
+		FCompressedChunkInfo Summary;
+		*this << Summary;
+
+		if (bWasByteSwapped)
+		{
+			Summary.CompressedSize = BYTESWAP_ORDER64(Summary.CompressedSize);
+			Summary.UncompressedSize = BYTESWAP_ORDER64(Summary.UncompressedSize);
+			PackageFileTag.UncompressedSize = BYTESWAP_ORDER64(PackageFileTag.UncompressedSize);
+		}
+
+
 		// Handle change in compression chunk size in backward compatible way.
 		int64 LoadingCompressionChunkSize = PackageFileTag.UncompressedSize;
 		if (LoadingCompressionChunkSize == PACKAGE_FILE_TAG)
@@ -747,12 +819,19 @@ void FArchive::SerializeCompressed(void* V, int64 Length, FName CompressionForma
 			LoadingCompressionChunkSize = LOADING_COMPRESSION_CHUNK_SIZE;
 		}
 
+		check( LoadingCompressionChunkSize > 0 );
+
+		// check Summary.UncompressedSize vs [V,Length] passed in
+		UE_CLOG( Summary.UncompressedSize != Length, LogSerialization, Fatal, TEXT(" Archive SerializedCompressed UncompressedSize (%lld) != Length (%lld)"), (int64)Summary.UncompressedSize, (int64) Length );
+
 		// Figure out how many chunks there are going to be based on uncompressed size and compression chunk size.
 		int64	TotalChunkCount	= (Summary.UncompressedSize + LoadingCompressionChunkSize - 1) / LoadingCompressionChunkSize;
 		
 		// Allocate compression chunk infos and serialize them, keeping track of max size of compression chunks used.
-		FCompressedChunkInfo*	CompressionChunks	= new FCompressedChunkInfo[TotalChunkCount];
-		int64						MaxCompressedSize	= 0;
+		FCompressedChunkInfo* CompressionChunks = new FCompressedChunkInfo[TotalChunkCount];
+		int64 MaxCompressedSize	= 0;
+		int64 TotalChunkCompressedSize = 0;
+		int64 TotalChunkUncompressedSize = 0;
 		for( int32 ChunkIndex=0; ChunkIndex<TotalChunkCount; ChunkIndex++ )
 		{
 			*this << CompressionChunks[ChunkIndex];
@@ -762,9 +841,17 @@ void FArchive::SerializeCompressed(void* V, int64 Length, FName CompressionForma
 				CompressionChunks[ChunkIndex].UncompressedSize	= BYTESWAP_ORDER64( CompressionChunks[ChunkIndex].UncompressedSize );
 			}
 			MaxCompressedSize = FMath::Max( CompressionChunks[ChunkIndex].CompressedSize, MaxCompressedSize );
+
+			TotalChunkCompressedSize += CompressionChunks[ChunkIndex].CompressedSize;
+			TotalChunkUncompressedSize += CompressionChunks[ChunkIndex].UncompressedSize;
 		}
+		
+		// verify the CompressionChunks[] sizes we read add up to the total we read
+		UE_CLOG( TotalChunkCompressedSize != Summary.CompressedSize, LogSerialization, Fatal, TEXT(" Archive SerializedCompressed TotalChunkCompressedSize (%lld) != Summary.CompressedSize (%lld)"), (int64)TotalChunkCompressedSize, (int64) Summary.CompressedSize );
+		UE_CLOG( TotalChunkUncompressedSize != Summary.UncompressedSize, LogSerialization, Fatal, TEXT(" Archive SerializedCompressed TotalChunkUncompressedSize (%lld) != Summary.UnompressedSize (%lld)"), (int64)TotalChunkUncompressedSize, (int64) Summary.UncompressedSize );
 
 		// Set up destination pointer and allocate memory for compressed chunk[s] (one at a time).
+		check( ! bTreatBufferAsFileReader );
 		uint8*	Dest				= (uint8*) V;
 		void*	CompressedBuffer	= FMemory::Malloc( MaxCompressedSize );
 
@@ -774,9 +861,25 @@ void FArchive::SerializeCompressed(void* V, int64 Length, FName CompressionForma
 			const FCompressedChunkInfo& Chunk = CompressionChunks[ChunkIndex];
 			// Read compressed data.
 			Serialize( CompressedBuffer, Chunk.CompressedSize );
+
+			// check Serialize error before trying to decode
+			if ( IsError() )
+			{
+				UE_LOG(LogSerialization,Error, TEXT("Failed to serialize compress chunk in %s, Chunk.CompressedSize=%d"), *GetArchiveName(), (int)Chunk.CompressedSize);
+				break;
+			}
+
 			// Decompress into dest pointer directly.
-			bool bUncompressMemorySucceeded = FCompression::UncompressMemory( CompressionFormat, Dest, Chunk.UncompressedSize, CompressedBuffer, Chunk.CompressedSize, COMPRESS_NoFlags);
-			verifyf(bUncompressMemorySucceeded, TEXT("Failed to uncompress data in %s. Check log for details."), *GetArchiveName());			// And advance it by read amount.
+			bool bUncompressMemorySucceeded = FCompression::UncompressMemory( CompressionFormatToDecode, Dest, Chunk.UncompressedSize, CompressedBuffer, Chunk.CompressedSize, COMPRESS_NoFlags);
+
+			if ( ! bUncompressMemorySucceeded )
+			{
+				UE_LOG(LogSerialization,Error, TEXT("Failed to uncompress data in %s, CompressionFormatToDecode=%s"), *GetArchiveName(), *CompressionFormatToDecode.ToString());
+				SetError();
+				break;
+			}
+
+			// And advance it by read amount.
 			Dest += Chunk.UncompressedSize;
 		}
 
@@ -789,21 +892,36 @@ void FArchive::SerializeCompressed(void* V, int64 Length, FName CompressionForma
 		SCOPE_SECONDS_COUNTER(GArchiveSerializedCompressedSavingTime);
 		check( Length > 0 );
 
+		// upgrade old flag method
+		if ((Flags & COMPRESS_DeprecatedFormatFlagsMask) != 0)
+		{
+			check( CompressionFormatToEncode == NAME_Zlib );
+			UE_LOG(LogSerialization, Warning, TEXT("Old style compression flags are being used with FAsyncCompressionChunk, please update any code using this!"));
+			CompressionFormatToEncode = FCompression::GetCompressionFormatFromDeprecatedFlags(Flags);
+		}
+
 		// if there's a cooking target, and it wants to replace Zlib compression with another format, use it. When loading, 
 		// the platform will replace Zlib with that format above
-		if (CompressionFormat == NAME_Zlib && CookingTargetPlatform != nullptr)
+		if (CompressionFormatToEncode == NAME_Zlib && CookingTargetPlatform != nullptr)
 		{
 			// use the replacement format
-			CompressionFormat = CookingTargetPlatform->GetZlibReplacementFormat();
+			CompressionFormatToEncode = CookingTargetPlatform->GetZlibReplacementFormat();
+
+			// with v2 headers, the modified CompressionFormatToEncode will be written in the archive
 		}
 
-		// Serialize package file tag used to determine endianess in LoadCompressedData.
+		// Serialize package file tag used to determine endianness in LoadCompressedData.
 		FCompressedChunkInfo PackageFileTag;
-		PackageFileTag.CompressedSize	= PACKAGE_FILE_TAG;
+		//PackageFileTag.CompressedSize	= PACKAGE_FILE_TAG;
+		PackageFileTag.CompressedSize	= ARCHIVE_V2_HEADER_TAG;
 		PackageFileTag.UncompressedSize	= GSavingCompressionChunkSize;
 		*this << PackageFileTag;
 
+		// v2 header writes compressor used :
+		FCompressionUtil::SerializeCompressorName(*this,CompressionFormatToEncode);
+
 		// Figure out how many chunks there are going to be based on uncompressed size and compression chunk size.
+		//  +1 for Summary chunk
 		int64	TotalChunkCount	= (Length + GSavingCompressionChunkSize - 1) / GSavingCompressionChunkSize + 1;
 		
 		// Keep track of current position so we can later seek back and overwrite stub compression chunk infos.
@@ -895,8 +1013,8 @@ void FArchive::SerializeCompressed(void* V, int64 Length, FName CompressionForma
 				if( FreeIndex != INDEX_NONE )
 				{
 					FAsyncCompressionChunk& NewChunk = AsyncChunks[FreeIndex].GetTask();
-					// 2 times the uncompressed size should be more than enough; the compressed data shouldn't be that much larger
-					NewChunk.CompressedSize	= 2 * GSavingCompressionChunkSize;
+
+					NewChunk.CompressedSize	= FCompression::CompressMemoryBound(CompressionFormatToEncode, GSavingCompressionChunkSize);
 					// Allocate compressed buffer placeholder on first use.
 					if( NewChunk.CompressedBuffer == NULL )
 					{
@@ -917,9 +1035,9 @@ void FArchive::SerializeCompressed(void* V, int64 Length, FName CompressionForma
 						}
 						((FArchive*)V)->Serialize(NewChunk.UncompressedBuffer, NewChunk.UncompressedSize);
 					}
-					// Advance src pointer by amount to be compressed.
 					else
 					{
+						// Advance src pointer by amount to be compressed.
 						NewChunk.UncompressedBuffer = SrcBuffer;
 						SrcBuffer += NewChunk.UncompressedSize;
 					}
@@ -928,17 +1046,18 @@ void FArchive::SerializeCompressed(void* V, int64 Length, FName CompressionForma
 					BytesRemainingToKickOff -= NewChunk.UncompressedSize;
 					AsyncChunkIndex[FreeIndex] = CurrentChunkIndex++;
 					NewChunk.Flags = Flags;
-					NewChunk.CompressionFormat = CompressionFormat;
+					NewChunk.CompressionFormat = CompressionFormatToEncode;
 					NumChunksLeftToKickOff--;
 
 					AsyncChunks[FreeIndex].StartBackgroundTask();
 				}
-				// No chunks were available to use, complete some
 				else
 				{
+					// No chunks were available to use, complete some
 					bNeedToWaitForAsyncTask = true;
 				}
 			}
+
 			// Wait for the oldest task to finish instead of spinning
 			if (NumChunksLeftToKickOff == 0)
 			{
@@ -1026,8 +1145,8 @@ void FArchive::SerializeCompressed(void* V, int64 Length, FName CompressionForma
 		int64		BytesRemaining			= Length;
 		// Start at index 1 as first chunk info is summary.
 		int32		CurrentChunkIndex		= 1;
-		// 2 times the uncompressed size should be more than enough; the compressed data shouldn't be that much larger
-		int64		CompressedBufferSize	= 2 * GSavingCompressionChunkSize;
+
+		int64		CompressedBufferSize	= FCompression::CompressMemoryBound(CompressionFormatToEncode, GSavingCompressionChunkSize);
 		void*	CompressedBuffer		= FMemory::Malloc( CompressedBufferSize );
 
 		while( BytesRemaining > 0 )
@@ -1044,7 +1163,7 @@ void FArchive::SerializeCompressed(void* V, int64 Length, FName CompressionForma
 			check(CompressedSize < INT_MAX);
 			int32 CompressedSizeInt = (int32)CompressedSize;
 						
-			verify( FCompression::CompressMemory( CompressionFormat, CompressedBuffer, CompressedSizeInt, Src, BytesToCompress, Flags ) );
+			verify( FCompression::CompressMemory( CompressionFormatToEncode, CompressedBuffer, CompressedSizeInt, Src, BytesToCompress, Flags ) );
 			CompressedSize = CompressedSizeInt;
 			// move to next chunk if not reading from file
 			if (!bTreatBufferAsFileReader)
diff --git a/Engine/Source/Runtime/Core/Public/Compression/CompressionUtil.h b/Engine/Source/Runtime/Core/Public/Compression/CompressionUtil.h
new file mode 100644
--- /dev/null
+++ b/Engine/Source/Runtime/Core/Public/Compression/CompressionUtil.h
@@ -0,0 +1,19 @@
+// Copyright Epic Games, Inc. All Rights Reserved.
+
+#pragma once
+
+#include "CoreTypes.h"
+
+class FArchive;
+class FCustomVersionContainer;
+class FLinker;
+class FName;
+
+namespace FCompressionUtil
+{
+
+
+	void CORE_API SerializeCompressorName(FArchive & Archive,FName & Compressor);
+
+
+};
\ No newline at end of file
diff --git a/Engine/Source/Runtime/Core/Public/Serialization/Archive.h b/Engine/Source/Runtime/Core/Public/Serialization/Archive.h
--- a/Engine/Source/Runtime/Core/Public/Serialization/Archive.h
+++ b/Engine/Source/Runtime/Core/Public/Serialization/Archive.h
@@ -1688,18 +1688,39 @@ public:
 	}
 
 	/**
-	 * Serializes and compresses/ uncompresses data. This is a shared helper function for compression
-	 * support. The data is saved in a way compatible with FIOSystem::LoadCompressedData.
+	 * Serializes and compresses/ uncompresses data. This is a shared helper function for compression support. 
+	 
+	 // @@ old comment ? remove me ?
+	 The data is saved in a way compatible with FIOSystem::LoadCompressedData.
+
+	 *
+	 * prefer SerializeCompressedNew instead.
+	 *
 	 *
 	 * @param	V		Data pointer to serialize data from/ to
 	 * @param	Length	Length of source data if we're saving, unused otherwise
+	 * @param   CompressionFormatCannotChange  Compression Format to use for encoding & decoding - cannot be changed without breaking file compatibility
 	 * @param	Flags	Flags to control what method to use for [de]compression and optionally control memory vs speed when compressing
 	 * @param	bTreatBufferAsFileReader true if V is actually an FArchive, which is used when saving to read data - helps to avoid single huge allocations of source data
-	 * @param	bUsePlatformBitWindow use a platform specific bitwindow setting
 	 */
-	void SerializeCompressed(void* V, int64 Length, FName CompressionFormat, ECompressionFlags Flags=COMPRESS_NoFlags, bool bTreatBufferAsFileReader=false);
-
-
+	void SerializeCompressed(void* V, int64 Length, FName CompressionFormatCannotChange,
+		ECompressionFlags Flags=COMPRESS_NoFlags, bool bTreatBufferAsFileReader=false);
+	
+	/**
+	 * Serializes and compresses/ uncompresses data. This is a shared helper function for compression support. 
+	 *
+	 * call SerializeCompressedNew instead of SerializeCompressed
+	 *
+	 * @param	V		Data pointer to serialize data from/ to
+	 * @param	Length	Length of source data if we're saving, unused otherwise
+	 * @param   CompressionFormatToEncode  Compression Format to use for encoding, can be changed freely without breaking compatibility
+	 * @param   CompressionFormatToDecodeOldV1Files  Compression Format to decode old data with that didn't write compressor in header, cannot change, usually NAME_Zlib
+	 * @param	Flags	Flags to control what method to use for [de]compression and optionally control memory vs speed when compressing
+	 * @param	bTreatBufferAsFileReader true if V is actually an FArchive, which is used when saving to read data - helps to avoid single huge allocations of source data
+	 */
+	void SerializeCompressedNew(void* V, int64 Length, FName CompressionFormatToEncode, FName CompressionFormatToDecodeOldV1Files,
+		ECompressionFlags Flags=COMPRESS_NoFlags, bool bTreatBufferAsFileReader=false);
+		
 	using FArchiveState::IsByteSwapping;
 
 	/** Used to do byte swapping on small items. This does not happen usually, so we don't want it inline. */
