Added FWindowsPlatformProcess::CreateElevatedProcess

diff --git a/Engine/Source/Runtime/Core/Private/Windows/WindowsPlatformProcess.cpp b/Engine/Source/Runtime/Core/Private/Windows/WindowsPlatformProcess.cpp
--- a/Engine/Source/Runtime/Core/Private/Windows/WindowsPlatformProcess.cpp
+++ b/Engine/Source/Runtime/Core/Private/Windows/WindowsPlatformProcess.cpp
@@ -924,6 +924,25 @@ bool FWindowsPlatformProcess::ExecElevatedProcess(const TCHAR* URL, const TCHAR*
 	return bSuccess;
 }
 
+FProcHandle FWindowsPlatformProcess::CreateElevatedProcess(const TCHAR* URL, const TCHAR* Params)
+{
+	SHELLEXECUTEINFO ShellExecuteInfo;
+	ZeroMemory(&ShellExecuteInfo, sizeof(ShellExecuteInfo));
+	ShellExecuteInfo.cbSize = sizeof(ShellExecuteInfo);
+	ShellExecuteInfo.fMask = SEE_MASK_UNICODE | SEE_MASK_NOCLOSEPROCESS;
+	ShellExecuteInfo.lpFile = URL;
+	ShellExecuteInfo.lpVerb = TEXT("runas");
+	ShellExecuteInfo.nShow = SW_SHOW;
+	ShellExecuteInfo.lpParameters = Params;
+
+	if (ShellExecuteEx(&ShellExecuteInfo))
+	{
+		return FProcHandle(ShellExecuteInfo.hProcess);
+	}
+
+	return FProcHandle{};
+}
+
 const TCHAR* FWindowsPlatformProcess::BaseDir()
 {
 	static TCHAR Result[512]=TEXT("");
diff --git a/Engine/Source/Runtime/Core/Public/Windows/WindowsPlatformProcess.h b/Engine/Source/Runtime/Core/Public/Windows/WindowsPlatformProcess.h
--- a/Engine/Source/Runtime/Core/Public/Windows/WindowsPlatformProcess.h
+++ b/Engine/Source/Runtime/Core/Public/Windows/WindowsPlatformProcess.h
@@ -169,6 +169,7 @@ public:
 	static FString GetApplicationName( uint32 ProcessId );	
 	static bool ExecProcess(const TCHAR* URL, const TCHAR* Params, int32* OutReturnCode, FString* OutStdOut, FString* OutStdErr, const TCHAR* OptionalWorkingDirectory = NULL, bool bShouldEndWithParentProcess  = false);
 	static bool ExecElevatedProcess(const TCHAR* URL, const TCHAR* Params, int32* OutReturnCode);
+	static FProcHandle CreateElevatedProcess(const TCHAR* URL, const TCHAR* Params);
 	static bool LaunchFileInDefaultExternalApplication( const TCHAR* FileName, const TCHAR* Parms = NULL, ELaunchVerb::Type Verb = ELaunchVerb::Open, bool bPromptToOpenOnFailure = true );
 	static void ExploreFolder( const TCHAR* FilePath );
 	static bool ResolveNetworkPath( FString InUNCPath, FString& OutPath ); 
