Added copyright boilerplate to fix CIS

diff --git a/Engine/Source/Runtime/Core/Private/Compression/CompressionUtil.cpp b/Engine/Source/Runtime/Core/Private/Compression/CompressionUtil.cpp
--- a/Engine/Source/Runtime/Core/Private/Compression/CompressionUtil.cpp
+++ b/Engine/Source/Runtime/Core/Private/Compression/CompressionUtil.cpp
@@ -1,3 +1,5 @@
+// Copyright Epic Games, Inc. All Rights Reserved.
+
 #include "Compression/CompressionUtil.h"
 #include "Serialization/Archive.h"
 #include "UObject/NameTypes.h"
