Fix for dereferencing null pointer static analysis warning (in a branch not currently called).

diff --git a/Engine/Source/Runtime/Core/Private/Containers/String.cpp b/Engine/Source/Runtime/Core/Private/Containers/String.cpp
--- a/Engine/Source/Runtime/Core/Private/Containers/String.cpp
+++ b/Engine/Source/Runtime/Core/Private/Containers/String.cpp
@@ -154,8 +154,8 @@ void AppendCharacters(TArray<TCHAR>& Out, const CharType* Str, int32 Count)
 		// Restablish destination pointer in case a realloc happened
 		Dest = Out.GetData() + OldEnd;
 
-		bool bSuccess = !!FPlatformString::Convert(Dest, Length, Str, Count);
-		checkSlow(bSuccess);
+		NewEnd = FPlatformString::Convert(Dest, Length, Str, Count);
+		checkSlow(NewEnd);
 	}
 	else
 	{
