fix Unity stress with LogCompression

diff --git a/Engine/Source/Runtime/Core/Private/Compression/CompressionUtil.cpp b/Engine/Source/Runtime/Core/Private/Compression/CompressionUtil.cpp
--- a/Engine/Source/Runtime/Core/Private/Compression/CompressionUtil.cpp
+++ b/Engine/Source/Runtime/Core/Private/Compression/CompressionUtil.cpp
@@ -4,8 +4,7 @@
 #include "Serialization/Archive.h"
 #include "UObject/NameTypes.h"
 #include "Logging/LogMacros.h"
-
-DECLARE_LOG_CATEGORY_EXTERN(LogCompression, Log, All);
+#include "CoreGlobals.h"
 
 namespace FCompressionUtil
 {
@@ -46,7 +45,7 @@ namespace FCompressionUtil
 				Compressor = NAME_LZ4;
 				break;
 			default:		
-				UE_LOG(LogCompression, Error, TEXT("SerializeCompressorName Unknown index:%d"),CompressorNum );
+				UE_LOG(LogSerialization, Error, TEXT("SerializeCompressorName Unknown index:%d"),CompressorNum );
 				
 				Compressor = NAME_None;
 				break;
