optionally use JPEG for thumbnails instead of PNG (disabled)

diff --git a/Engine/Source/Runtime/Core/Private/Misc/ObjectThumbnail.cpp b/Engine/Source/Runtime/Core/Private/Misc/ObjectThumbnail.cpp
--- a/Engine/Source/Runtime/Core/Private/Misc/ObjectThumbnail.cpp
+++ b/Engine/Source/Runtime/Core/Private/Misc/ObjectThumbnail.cpp
@@ -3,8 +3,13 @@
 #include "Misc/ObjectThumbnail.h"
 #include "Serialization/StructuredArchive.h"
 
+#define USE_JPEG_FOR_THUMBNAILS 0
+
 /** Static: Thumbnail compressor */
-FThumbnailCompressionInterface* FObjectThumbnail::ThumbnailCompressor = NULL;
+// these are installed from UnrealEngine.cpp startup
+// because we are in Core Module we can't access ImageWrapper ourselves
+FThumbnailCompressionInterface* FObjectThumbnail::PNGThumbnailCompressor = nullptr;
+FThumbnailCompressionInterface* FObjectThumbnail::JPEGThumbnailCompressor = nullptr;
 
 
 FObjectThumbnail::FObjectThumbnail()
@@ -12,6 +17,7 @@ FObjectThumbnail::FObjectThumbnail()
 	  ImageHeight( 0 ),
 	  bIsDirty( false ),
 	  bLoadedFromDisk(false),
+	  bIsJPEG(false),
 	  bCreatedAfterCustomThumbForSharedTypesEnabled(false)
 { }
 
@@ -35,10 +41,6 @@ void FObjectThumbnail::Serialize( FArchive& Ar )
 
 void FObjectThumbnail::Serialize(FStructuredArchive::FSlot Slot)
 {
-	FStructuredArchive::FRecord Record = Slot.EnterRecord();
-	Record << SA_VALUE(TEXT("ImageWidth"), ImageWidth);
-	Record << SA_VALUE(TEXT("ImageHeight"), ImageHeight);
-
 	//if the image thinks it's empty, ensure there is no memory waste
 	if ((ImageWidth == 0) || (ImageHeight == 0))
 	{
@@ -52,6 +54,27 @@ void FObjectThumbnail::Serialize(FStructuredArchive::FSlot Slot)
 		CompressImageData();
 	}
 
+	FStructuredArchive::FRecord Record = Slot.EnterRecord();
+	Record << SA_VALUE(TEXT("ImageWidth"), ImageWidth);
+	
+	// serialize bIsJPEG via negative ImageHeight
+	if (Slot.GetArchiveState().IsLoading())
+	{
+		Record << SA_VALUE(TEXT("ImageHeight"), ImageHeight);
+		bIsJPEG = false;
+		if ( ImageHeight < 0 )
+		{
+			bIsJPEG = true;
+			ImageHeight = -ImageHeight;
+		}
+	}
+	else
+	{
+		int32 ImageHeightToSerialize = bIsJPEG ? -ImageHeight : ImageHeight;
+		
+		Record << SA_VALUE(TEXT("ImageHeight"), ImageHeightToSerialize);
+	}
+		
 	Record << SA_VALUE(TEXT("CompressedImageData"), CompressedImageData);
 
 	if (Slot.GetUnderlyingArchive().IsCountingMemory())
@@ -73,9 +96,23 @@ void FObjectThumbnail::Serialize(FStructuredArchive::FSlot Slot)
 void FObjectThumbnail::CompressImageData()
 {
 	CompressedImageData.Reset();
-	if( ThumbnailCompressor != NULL && ImageData.Num() > 0 && ImageWidth > 0 && ImageHeight > 0 )
+
+	if ( ImageData.Num() > 0 && ImageWidth > 0 && ImageHeight > 0 )
 	{
-		ThumbnailCompressor->CompressImage( ImageData, ImageWidth, ImageHeight, CompressedImageData );
+		#if USE_JPEG_FOR_THUMBNAILS
+		// prefer JPEG except on images that are trivially tiny or 1d
+		if( JPEGThumbnailCompressor != NULL && ImageWidth >= 8 && ImageHeight >= 8 )
+		{
+			JPEGThumbnailCompressor->CompressImage( ImageData, ImageWidth, ImageHeight, CompressedImageData );
+			bIsJPEG = true;
+		}
+		#endif
+
+		if( PNGThumbnailCompressor != NULL && CompressedImageData.Num() == 0 )
+		{
+			PNGThumbnailCompressor->CompressImage( ImageData, ImageWidth, ImageHeight, CompressedImageData );
+			bIsJPEG = false;
+		}
 	}
 }
 
@@ -83,9 +120,16 @@ void FObjectThumbnail::CompressImageData()
 void FObjectThumbnail::DecompressImageData()
 {
 	ImageData.Reset();
-	if( ThumbnailCompressor != NULL && CompressedImageData.Num() > 0 && ImageWidth > 0 && ImageHeight > 0 )
+	if( CompressedImageData.Num() > 0 && ImageWidth > 0 && ImageHeight > 0 )
 	{
-		ThumbnailCompressor->DecompressImage( CompressedImageData, ImageWidth, ImageHeight, ImageData );
+		if ( bIsJPEG && JPEGThumbnailCompressor != nullptr )
+		{
+			JPEGThumbnailCompressor->DecompressImage( CompressedImageData, ImageWidth, ImageHeight, ImageData );
+		}
+		else if ( !bIsJPEG && PNGThumbnailCompressor != nullptr )
+		{
+			PNGThumbnailCompressor->DecompressImage( CompressedImageData, ImageWidth, ImageHeight, ImageData );
+		}
 	}
 }
 
diff --git a/Engine/Source/Runtime/Core/Public/Misc/ObjectThumbnail.h b/Engine/Source/Runtime/Core/Public/Misc/ObjectThumbnail.h
--- a/Engine/Source/Runtime/Core/Public/Misc/ObjectThumbnail.h
+++ b/Engine/Source/Runtime/Core/Public/Misc/ObjectThumbnail.h
@@ -52,15 +52,18 @@ public:
 	 *
 	 * @param	InThumbnailCompressor	A class derived from FThumbnailCompressionInterface.
 	 */
-	static void SetThumbnailCompressor( FThumbnailCompressionInterface* InThumbnailCompressor )
+	static void SetThumbnailCompressors( FThumbnailCompressionInterface* InPNGThumbnailCompressor, 
+		FThumbnailCompressionInterface* InJPEGThumbnailCompressor)
 	{
-		ThumbnailCompressor = InThumbnailCompressor;
+		PNGThumbnailCompressor = InPNGThumbnailCompressor;
+		JPEGThumbnailCompressor = InJPEGThumbnailCompressor;
 	}
 
 private:
 
 	/** Static: Thumbnail compressor. */
-	static FThumbnailCompressionInterface* ThumbnailCompressor;
+	static FThumbnailCompressionInterface* PNGThumbnailCompressor;
+	static FThumbnailCompressionInterface* JPEGThumbnailCompressor;
 
 public:
 
@@ -228,6 +231,9 @@ private:
 	/** Whether the thumbnail has a backup on disk*/
 	bool bLoadedFromDisk;
 
+	/** Whether compressed data is JPEG (else PNG) */
+	bool bIsJPEG;
+
 	/** Whether this was saved AFTER custom-thumbnails for shared thumbnail asset types was supported */
 	bool bCreatedAfterCustomThumbForSharedTypesEnabled;
 };
diff --git a/Engine/Source/Runtime/Engine/Private/UnrealEngine.cpp b/Engine/Source/Runtime/Engine/Private/UnrealEngine.cpp
--- a/Engine/Source/Runtime/Engine/Private/UnrealEngine.cpp
+++ b/Engine/Source/Runtime/Engine/Private/UnrealEngine.cpp
@@ -1099,6 +1099,76 @@ public:
 };
 
 
+class FJPEGThumbnailCompressor
+	: public FThumbnailCompressionInterface
+{
+
+public:
+
+	/**
+	* Compresses an image
+	*
+	* @param	InUncompressedData	The uncompressed image data
+	* @param	InWidth				Width of the image
+	* @param	InHeight			Height of the image
+	* @param	OutCompressedData	[Out] Compressed image data
+	*
+	* @return	true if the image was compressed successfully, otherwise false if an error occurred
+	*/
+	virtual bool CompressImage( const TArray< uint8 >& InUncompressedData, const int32 InWidth, const int32 InHeight, TArray< uint8 >& OutCompressedData )
+	{
+		bool bSucceeded = false;
+		OutCompressedData.Reset();
+		if( InUncompressedData.Num() > 0 )
+		{
+			IImageWrapperModule& ImageWrapperModule = FModuleManager::LoadModuleChecked<IImageWrapperModule>( FName("ImageWrapper") );
+			TSharedPtr<IImageWrapper> ImageWrapper = ImageWrapperModule.CreateImageWrapper( EImageFormat::JPEG );
+			if ( ImageWrapper.IsValid() && ImageWrapper->SetRaw( &InUncompressedData[ 0 ], InUncompressedData.Num(), InWidth, InHeight, ERGBFormat::RGBA, 8 ) )
+			{
+				int ThumbnailJPEGQuality = 65;
+				OutCompressedData = ImageWrapper->GetCompressed(ThumbnailJPEGQuality);
+				bSucceeded = true;
+			}
+		}
+
+		return bSucceeded;
+	}
+
+
+	/**
+	* Decompresses an image
+	*
+	* @param	InCompressedData	The compressed image data
+	* @param	InWidth				Width of the image
+	* @param	InHeight			Height of the image
+	* @param	OutUncompressedData	[Out] Uncompressed image data
+	*
+	* @return	true if the image was decompressed successfully, otherwise false if an error occurred
+	*/
+	virtual bool DecompressImage( const TArray< uint8 >& InCompressedData, const int32 InWidth, const int32 InHeight, TArray< uint8 >& OutUncompressedData )
+	{
+		bool bSucceeded = false;
+		OutUncompressedData.Reset();
+		if( InCompressedData.Num() > 0 )
+		{
+			IImageWrapperModule& ImageWrapperModule = FModuleManager::LoadModuleChecked<IImageWrapperModule>( FName("ImageWrapper") );
+			TSharedPtr<IImageWrapper> ImageWrapper = ImageWrapperModule.CreateImageWrapper( EImageFormat::JPEG );
+			if ( ImageWrapper.IsValid() && ImageWrapper->SetCompressed( &InCompressedData[ 0 ], InCompressedData.Num() ) )
+			{
+				check( ImageWrapper->GetWidth() == InWidth );
+				check( ImageWrapper->GetHeight() == InHeight );
+				if ( ImageWrapper->GetRaw( ERGBFormat::RGBA, 8, OutUncompressedData) )
+				{
+					bSucceeded = true;
+				}
+			}
+		}
+
+		return bSucceeded;
+	}
+};
+
+
 /**
 * Helper class inhibiting screen saver by e.g. moving the mouse by 0 pixels every 50 seconds.
 */
@@ -1622,7 +1692,7 @@ void UEngine::Init(IEngineLoop* InEngineLoop)
 	}
 
 	// Assign thumbnail compressor/decompressor
-	FObjectThumbnail::SetThumbnailCompressor( new FPNGThumbnailCompressor() );
+	FObjectThumbnail::SetThumbnailCompressors( new FPNGThumbnailCompressor() , new FJPEGThumbnailCompressor() );
 
 	//UEngine::StaticClass()->GetDefaultObject(true);
 	LoadObject<UClass>(UEngine::StaticClass()->GetOuter(), *UEngine::StaticClass()->GetName(), NULL, LOAD_Quiet|LOAD_NoWarn, NULL );
