Fix bad stricmp that led to mimalloc being activated by default
Make mimalloc really activated by default now that we know it doesn't cause major issues

#rnx

diff --git a/Engine/Source/Runtime/Core/Private/Unix/UnixPlatformMemory.cpp b/Engine/Source/Runtime/Core/Private/Unix/UnixPlatformMemory.cpp
--- a/Engine/Source/Runtime/Core/Private/Unix/UnixPlatformMemory.cpp
+++ b/Engine/Source/Runtime/Core/Private/Unix/UnixPlatformMemory.cpp
@@ -147,6 +147,13 @@ class FMalloc* FUnixPlatformMemory::BaseAllocator()
 	}
 	else
 	{
+		// Mimalloc is now the default allocator for editor and programs because it has shown
+		// both great performance and as much as half the memory usage of TBB after
+		// heavy editor workloads. See CL 15887498 description for benchmarks.
+#if (WITH_EDITORONLY_DATA || IS_PROGRAM) && PLATFORM_SUPPORTS_MIMALLOC && MIMALLOC_ALLOCATOR_ALLOWED
+		AllocatorToUse = EMemoryAllocatorToUse::Mimalloc;
+#endif
+
 		// Allow overriding on the command line.
 		// We get here before main due to global ctors, so need to do some hackery to get command line args
 		if (FILE* CmdLineFile = fopen("/proc/self/cmdline", "r"))
@@ -176,7 +183,7 @@ class FMalloc* FUnixPlatformMemory::BaseAllocator()
 				}
 
 #if PLATFORM_SUPPORTS_MIMALLOC && MIMALLOC_ALLOCATOR_ALLOWED
-				if (FCStringAnsi::Stricmp(Arg, "-mimalloc"))
+				if (FCStringAnsi::Stricmp(Arg, "-mimalloc") == 0)
 				{
 					AllocatorToUse = EMemoryAllocatorToUse::Mimalloc;
 					break;
