Fix RHI thread priority to avoid starvation/preemption under high load

#rnx

diff --git a/Engine/Source/Runtime/Core/Public/GenericPlatform/GenericPlatformAffinity.h b/Engine/Source/Runtime/Core/Public/GenericPlatform/GenericPlatformAffinity.h
--- a/Engine/Source/Runtime/Core/Public/GenericPlatform/GenericPlatformAffinity.h
+++ b/Engine/Source/Runtime/Core/Public/GenericPlatform/GenericPlatformAffinity.h
@@ -137,7 +137,7 @@ public:
 	
 	static EThreadPriority GetRHIThreadPriority()
 	{
-		return TPri_SlightlyBelowNormal;
+		return TPri_Normal;
 	}
 
 	static EThreadCreateFlags GetRHIThreadFlags()
