Fix game-thread starvation under high load by better distributing the work among threads

#rnx

diff --git a/Engine/Source/Runtime/Core/Private/Serialization/MemoryImage.cpp b/Engine/Source/Runtime/Core/Private/Serialization/MemoryImage.cpp
--- a/Engine/Source/Runtime/Core/Private/Serialization/MemoryImage.cpp
+++ b/Engine/Source/Runtime/Core/Private/Serialization/MemoryImage.cpp
@@ -1808,7 +1808,7 @@ void FMemoryImage::Flatten(FMemoryImageResult& OutResult, bool bMergeDuplicateSe
 	{
 		// Find unique sections
 		// precompute hashes in parallel
-		ParallelFor(Sections.Num(), [this](int32 Index) { Sections[Index]->ComputeHash(); });
+		ParallelFor(Sections.Num(), [this](int32 Index) { Sections[Index]->ComputeHash(); }, EParallelForFlags::Unbalanced);
 
 		TMap<FSHAHash, int32> HashToSectionIndex;
 		for (int32 SectionIndex = 0; SectionIndex < Sections.Num(); ++SectionIndex)
