use JPEG for thumbnails instead of PNG

diff --git a/Engine/Source/Runtime/Core/Private/Misc/ObjectThumbnail.cpp b/Engine/Source/Runtime/Core/Private/Misc/ObjectThumbnail.cpp
--- a/Engine/Source/Runtime/Core/Private/Misc/ObjectThumbnail.cpp
+++ b/Engine/Source/Runtime/Core/Private/Misc/ObjectThumbnail.cpp
@@ -3,7 +3,7 @@
 #include "Misc/ObjectThumbnail.h"
 #include "Serialization/StructuredArchive.h"
 
-#define USE_JPEG_FOR_THUMBNAILS 0
+#define USE_JPEG_FOR_THUMBNAILS 1
 
 /** Static: Thumbnail compressor */
 // these are installed from UnrealEngine.cpp startup
