#ue5 - LWC: fix issue with double version of VectorSet_W0() caused by apparent  MSVC compiler optimization bug. It doesn't handle the casting well so changing to slightly lower throughput _mm_move_sd function. This usually manifested as bad matrix values when trying to remove scale with RemoveScaling() or creating matrices with ToMatrixNoScale().

diff --git a/Engine/Source/Runtime/Core/Public/Math/UnrealMathSSE.h b/Engine/Source/Runtime/Core/Public/Math/UnrealMathSSE.h
--- a/Engine/Source/Runtime/Core/Public/Math/UnrealMathSSE.h
+++ b/Engine/Source/Runtime/Core/Public/Math/UnrealMathSSE.h
@@ -2039,7 +2039,7 @@ FORCEINLINE VectorRegister4Double VectorSet_W0(const VectorRegister4Double& Vec)
 	VectorRegister4Double Result;
 #if !UE_PLATFORM_MATH_USE_AVX
 	Result.XY = Vec.XY;
-	Result.ZW = _mm_castsi128_pd(_mm_move_epi64(_mm_castpd_si128(Vec.ZW)));
+	Result.ZW = _mm_move_sd(_mm_setzero_pd(), Vec.ZW);
 #else
 	Result = _mm256_blend_pd(Vec, VectorZeroDouble(), 0b1000);
 #endif
