allow SerializeCompressedNew to read less than requested
fix FArchiveLoadCompressedProxy reading past end on last partial chunk

diff --git a/Engine/Source/Runtime/Core/Private/Serialization/Archive.cpp b/Engine/Source/Runtime/Core/Private/Serialization/Archive.cpp
--- a/Engine/Source/Runtime/Core/Private/Serialization/Archive.cpp
+++ b/Engine/Source/Runtime/Core/Private/Serialization/Archive.cpp
@@ -693,10 +693,10 @@ void FArchive::SerializeCompressed(void* V, int64 Length, FName CompressionForma
 	// with this legacy/deprecated API you can NOT change the CompressionFormat and still load old files
 	//  CompressionFormat must match exactly the format that was written to the file
 
-	SerializeCompressedNew(V,Length,CompressionFormatCannotChange,CompressionFormatCannotChange,Flags,bTreatBufferAsFileReader);
+	SerializeCompressedNew(V,Length,CompressionFormatCannotChange,CompressionFormatCannotChange,Flags,bTreatBufferAsFileReader,nullptr);
 }
 
-void FArchive::SerializeCompressedNew(void* V, int64 Length, FName CompressionFormatToEncode, FName CompressionFormatToDecodeOldV1Files,  ECompressionFlags Flags, bool bTreatBufferAsFileReader)
+void FArchive::SerializeCompressedNew(void* V, int64 Length, FName CompressionFormatToEncode, FName CompressionFormatToDecodeOldV1Files,  ECompressionFlags Flags, bool bTreatBufferAsFileReader, int64 * OutPartialReadLength)
 {
 	TRACE_CPUPROFILER_EVENT_SCOPE(FArchive::SerializeCompressed);
 
@@ -822,7 +822,19 @@ void FArchive::SerializeCompressedNew(void* V, int64 Length, FName CompressionFo
 		check( LoadingCompressionChunkSize > 0 );
 
 		// check Summary.UncompressedSize vs [V,Length] passed in
-		UE_CLOG( Summary.UncompressedSize != Length, LogSerialization, Fatal, TEXT(" Archive SerializedCompressed UncompressedSize (%lld) != Length (%lld)"), (int64)Summary.UncompressedSize, (int64) Length );
+		if ( OutPartialReadLength == nullptr )
+		{
+			check( Summary.UncompressedSize == Length );
+			UE_CLOG( Summary.UncompressedSize != Length, LogSerialization, Fatal, TEXT(" Archive SerializedCompressed UncompressedSize (%lld) != Length (%lld)"), (int64)Summary.UncompressedSize, (int64) Length );
+		}
+		else
+		{
+			// UncompressedSize smaller than length is okay
+			check( Summary.UncompressedSize <= Length );
+			check( Summary.UncompressedSize >= 0 );
+			UE_CLOG( Summary.UncompressedSize > Length, LogSerialization, Fatal, TEXT(" Archive SerializedCompressed UncompressedSize (%lld) > Length (%lld)"), (int64)Summary.UncompressedSize, (int64) Length );
+			*OutPartialReadLength = Summary.UncompressedSize;
+		}
 
 		// Figure out how many chunks there are going to be based on uncompressed size and compression chunk size.
 		int64	TotalChunkCount	= (Summary.UncompressedSize + LoadingCompressionChunkSize - 1) / LoadingCompressionChunkSize;
diff --git a/Engine/Source/Runtime/Core/Private/Serialization/ArchiveLoadCompressedProxy.cpp b/Engine/Source/Runtime/Core/Private/Serialization/ArchiveLoadCompressedProxy.cpp
--- a/Engine/Source/Runtime/Core/Private/Serialization/ArchiveLoadCompressedProxy.cpp
+++ b/Engine/Source/Runtime/Core/Private/Serialization/ArchiveLoadCompressedProxy.cpp
@@ -43,12 +43,25 @@ FArchiveLoadCompressedProxy::~FArchiveLoadCompressedProxy()
  */
 void FArchiveLoadCompressedProxy::DecompressMoreData()
 {
+	if ( IsError() )
+	{
+		// don't try to decode more if we hit error
+		TmpData = TmpDataStart;
+		TmpDataEnd = TmpDataStart;
+		return;
+	}
+
 	// This will call Serialize so we need to indicate that we want to serialize from array.
 	bShouldSerializeFromArray = true;
-	SerializeCompressed( TmpDataStart, LOADING_COMPRESSION_CHUNK_SIZE /** it's ignored, but that's how much we serialize */, CompressionFormat, CompressionFlags);
+	int64 DecompressedLength = 0;
+	SerializeCompressedNew( TmpDataStart, LOADING_COMPRESSION_CHUNK_SIZE, CompressionFormat,CompressionFormat, CompressionFlags, false, &DecompressedLength);
+	// last chunk will be partial :
+	//	all chunks before last should have size == LOADING_COMPRESSION_CHUNK_SIZE
+	check( DecompressedLength <= LOADING_COMPRESSION_CHUNK_SIZE );
 	bShouldSerializeFromArray = false;
 	// Buffer is filled again, reset.
 	TmpData = TmpDataStart;
+	TmpDataEnd = TmpDataStart + DecompressedLength;
 }
 
 /**
@@ -61,10 +74,10 @@ void FArchiveLoadCompressedProxy::DecompressMoreData()
 void FArchiveLoadCompressedProxy::Serialize( void* InData, int64 Count )
 {
 	uint8* DstData = (uint8*) InData;
-	// If counter > 1 it means we're calling recursively and therefore need to write to compressed data.
+
 	if( bShouldSerializeFromArray )
 	{
-		// Add space in array and copy data there.
+		// SerializedCompressed reads the compressed data from here
 		check(CurrentIndex+Count<=CompressedData.Num());
 		FMemory::Memcpy( DstData, &CompressedData[CurrentIndex], Count );
 		CurrentIndex += Count;
@@ -94,6 +107,14 @@ void FArchiveLoadCompressedProxy::Serialize( void* InData, int64 Count )
 			{
 				// Decompress more data. This will call Serialize again so we need to handle recursion.
 				DecompressMoreData();
+
+				if ( TmpDataEnd == TmpData )
+				{
+					// wanted more but couldn't get any
+					// avoid infinite loop
+					SetError();
+					return;
+				}
 			}
 		}
 	}
diff --git a/Engine/Source/Runtime/Core/Public/Serialization/Archive.h b/Engine/Source/Runtime/Core/Public/Serialization/Archive.h
--- a/Engine/Source/Runtime/Core/Public/Serialization/Archive.h
+++ b/Engine/Source/Runtime/Core/Public/Serialization/Archive.h
@@ -1717,9 +1717,11 @@ public:
 	 * @param   CompressionFormatToDecodeOldV1Files  Compression Format to decode old data with that didn't write compressor in header, cannot change, usually NAME_Zlib
 	 * @param	Flags	Flags to control what method to use for [de]compression and optionally control memory vs speed when compressing
 	 * @param	bTreatBufferAsFileReader true if V is actually an FArchive, which is used when saving to read data - helps to avoid single huge allocations of source data
+	 * @param	OutPartialReadLength if not null, partial reads are allowed and the size is filled here
 	 */
 	void SerializeCompressedNew(void* V, int64 Length, FName CompressionFormatToEncode, FName CompressionFormatToDecodeOldV1Files,
-		ECompressionFlags Flags=COMPRESS_NoFlags, bool bTreatBufferAsFileReader=false);
+		ECompressionFlags Flags=COMPRESS_NoFlags, bool bTreatBufferAsFileReader=false,
+		int64 * OutPartialReadLength=nullptr);
 		
 	using FArchiveState::IsByteSwapping;
 
diff --git a/Engine/Source/Runtime/RenderCore/Private/ShaderPipelineCache.cpp b/Engine/Source/Runtime/RenderCore/Private/ShaderPipelineCache.cpp
--- a/Engine/Source/Runtime/RenderCore/Private/ShaderPipelineCache.cpp
+++ b/Engine/Source/Runtime/RenderCore/Private/ShaderPipelineCache.cpp
@@ -10,8 +10,6 @@
 #include "HAL/FileManager.h"
 #include "Misc/Paths.h"
 #include "Serialization/MemoryWriter.h"
-#include "Serialization/ArchiveSaveCompressedProxy.h"
-#include "Serialization/ArchiveLoadCompressedProxy.h"
 #include "Serialization/CustomVersion.h"
 #include "Serialization/MemoryReader.h"
 #include "Shader.h"
