Implemented Taskaffinity to launch a job on a specific worker. This uses a seperate queue for every worker to work. We had to introduce an extra sleeping state to ensure that the worker will try to check the queues again without necessarily adding itself again to the sleepstack when drowsing.
Also fixing some issues with the throttle stealing that could lead to drowsing workers not properly waking up in certain situations.

diff --git a/Engine/Source/Runtime/Core/Private/Async/Fundamental/ReserveScheduler.cpp b/Engine/Source/Runtime/Core/Private/Async/Fundamental/ReserveScheduler.cpp
--- a/Engine/Source/Runtime/Core/Private/Async/Fundamental/ReserveScheduler.cpp
+++ b/Engine/Source/Runtime/Core/Private/Async/Fundamental/ReserveScheduler.cpp
@@ -74,7 +74,7 @@ void FReserveScheduler::StartWorkers(FScheduler& MainScheduler, uint32 NumWorker
 		UE::Trace::ThreadGroupBegin(TEXT("Reserve Workers"));
 		for (uint32 WorkerId = 0; WorkerId < NumWorkers; ++WorkerId)
 		{
-			WorkerLocalQueues.Emplace(MainScheduler.GetQueueRegistry(), ELocalQueueType::EBusyWait);
+			WorkerLocalQueues.Emplace(MainScheduler.GetQueueRegistry(), ELocalQueueType::EBusyWait, nullptr);
 			WorkerThreads.Add(CreateWorker(&WorkerLocalQueues.Last(), WorkerPriority, bIsForkable));
 		}
 		UE::Trace::ThreadGroupEnd();
diff --git a/Engine/Source/Runtime/Core/Private/Async/Fundamental/Scheduler.cpp b/Engine/Source/Runtime/Core/Private/Async/Fundamental/Scheduler.cpp
--- a/Engine/Source/Runtime/Core/Private/Async/Fundamental/Scheduler.cpp
+++ b/Engine/Source/Runtime/Core/Private/Async/Fundamental/Scheduler.cpp
@@ -40,7 +40,7 @@ namespace LowLevelTasks
 		}
 	}
 
-	TUniquePtr<FThread> FScheduler::CreateWorker(FSchedulerTls::FLocalQueueType* ExternalWorkerLocalQueue, EThreadPriority Priority, bool bPermitBackgroundWork, bool bIsForkable)
+	TUniquePtr<FThread> FScheduler::CreateWorker(FSleepEvent* ExternalWorkerEvent, FSchedulerTls::FLocalQueueType* ExternalWorkerLocalQueue, EThreadPriority Priority, bool bPermitBackgroundWork, bool bIsForkable)
 	{
 		uint32 WorkerId = NextWorkerId++;
 		const uint32 WaitTimes[8] = { 719, 991, 1361, 1237, 1597, 953, 587, 1439 };
@@ -76,10 +76,9 @@ namespace LowLevelTasks
 		return MakeUnique<FThread>
 		(
 			bPermitBackgroundWork ? *FString::Printf(TEXT("Background Worker #%d"), WorkerId) : *FString::Printf(TEXT("Foreground Worker #%d"), WorkerId),
-			[this, ExternalWorkerLocalQueue, WaitTime, bPermitBackgroundWork]
+			[this, ExternalWorkerEvent, ExternalWorkerLocalQueue, WaitTime, bPermitBackgroundWork]
 			{ 
-				FSleepEvent Event;
-				WorkerMain(&Event, ExternalWorkerLocalQueue, WaitTime, bPermitBackgroundWork);
+				WorkerMain(ExternalWorkerEvent, ExternalWorkerLocalQueue, WaitTime, bPermitBackgroundWork);
 			}, 0, Priority, FThreadAffinity{ ThreadAffinityMask & ProcessorGroups.ThreadAffinities[CpuGroup], CpuGroup }, bIsForkable
 		);
 	}
@@ -100,22 +99,26 @@ namespace LowLevelTasks
 			FScopeLock Lock(&WorkerThreadsCS);
 			check(!WorkerThreads.Num());
 			check(!WorkerLocalQueues.Num());
+			check(!WorkerEvents.Num());		
 			check(NextWorkerId == 0);
 
 			WorkerThreads.Reserve(NumForegroundWorkers + NumBackgroundWorkers);
 			WorkerLocalQueues.Reserve(NumForegroundWorkers + NumBackgroundWorkers);
+			WorkerEvents.Reserve(NumForegroundWorkers + NumBackgroundWorkers);
 			UE::Trace::ThreadGroupBegin(TEXT("Foreground Workers"));
 			for (uint32 WorkerId = 0; WorkerId < NumForegroundWorkers; ++WorkerId)
 			{
-				WorkerLocalQueues.Emplace(QueueRegistry, ELocalQueueType::EForeground);
-				WorkerThreads.Add(CreateWorker(&WorkerLocalQueues.Last(), WorkerPriority, false, bIsForkable));
+				WorkerEvents.Emplace();
+				WorkerLocalQueues.Emplace(QueueRegistry, ELocalQueueType::EForeground, &WorkerEvents.Last());
+				WorkerThreads.Add(CreateWorker(&WorkerEvents.Last(), &WorkerLocalQueues.Last(), WorkerPriority, false, bIsForkable));
 			}
 			UE::Trace::ThreadGroupEnd();
 			UE::Trace::ThreadGroupBegin(TEXT("Background Workers"));
 			for (uint32 WorkerId = 0; WorkerId < NumBackgroundWorkers; ++WorkerId)
 			{
-				WorkerLocalQueues.Emplace(QueueRegistry, ELocalQueueType::EBackground);
-				WorkerThreads.Add(CreateWorker(&WorkerLocalQueues.Last(), BackgroundPriority, true, bIsForkable));
+				WorkerEvents.Emplace();
+				WorkerLocalQueues.Emplace(QueueRegistry, ELocalQueueType::EBackground, &WorkerEvents.Last());
+				WorkerThreads.Add(CreateWorker(&WorkerEvents.Last(), &WorkerLocalQueues.Last(), BackgroundPriority, true, bIsForkable));
 			}
 			UE::Trace::ThreadGroupEnd();
 		}
@@ -137,6 +140,7 @@ namespace LowLevelTasks
 			NextWorkerId = 0;
 			WorkerThreads.Reset();
 			WorkerLocalQueues.Reset();
+			WorkerEvents.Reset();
 			for (FTask* Task = QueueRegistry.Dequeue(); Task != nullptr; Task = QueueRegistry.Dequeue())
 			{
 				Task->ExecuteTask();
@@ -199,6 +203,15 @@ namespace LowLevelTasks
 		return BusyWaitingDepth != 0;
 	}
 
+	uint32 FSchedulerTls::GetAffinityIndex()
+	{
+		if (LocalQueue)
+		{
+			return LocalQueue->GetAffinityIndex();
+		}
+		return ~0;
+	}
+
 	void FTask::PropagateUserData()
 	{
 		const FTask* ActiveTask = FSchedulerTls::GetActiveTask();
@@ -212,12 +225,12 @@ namespace LowLevelTasks
 		}
 	}
 
-	template<FTask* (FSchedulerTls::FLocalQueueType::*DequeueFunction)(bool), bool bIsBusyWaiting>
-	bool FScheduler::TryExecuteTaskFrom(FSchedulerTls::FLocalQueueType* Queue, FSchedulerTls::FQueueRegistry::FOutOfWork& OutOfWork, bool bPermitBackgroundWork)
+	template<FTask* (FSchedulerTls::FLocalQueueType::*DequeueFunction)(bool, bool), bool bIsBusyWaiting>
+	bool FScheduler::TryExecuteTaskFrom(FSchedulerTls::FLocalQueueType* Queue, FSchedulerTls::FQueueRegistry::FOutOfWork& OutOfWork, bool bPermitBackgroundWork, bool bDisableThrottleStealing)
 	{
 		for(int i = 0; i < 2; i++) // one retry if we pick up a task that cannot used during busy waiting.
 		{
-			FTask* Task = (Queue->*DequeueFunction)(bPermitBackgroundWork);
+			FTask* Task = (Queue->*DequeueFunction)(bPermitBackgroundWork, bDisableThrottleStealing);
 			if (Task)
 			{	
 				if (bIsBusyWaiting && !Task->AllowBusyWaiting())
@@ -228,7 +241,10 @@ namespace LowLevelTasks
 				if (OutOfWork.Stop())
 				{
 					//if we are ramping up work again we start waking up workers that might not have been woken while a lot of work might have been queued
-					WakeUpWorker(bPermitBackgroundWork);
+					if (!WakeUpWorker(bPermitBackgroundWork) && !FSchedulerTls::IsBackgroundWorker())
+					{
+						WakeUpWorker(!bPermitBackgroundWork);
+					}
 				}
 				FTask* OldTask = FSchedulerTls::ActiveTask;
 				FSchedulerTls::ActiveTask = Task;
@@ -244,7 +260,7 @@ namespace LowLevelTasks
 		return false;
 	}
 
-	void FScheduler::WorkerMain(FSleepEvent* WorkerEvent, FSchedulerTls::FLocalQueueType* ExternalWorkerLocalQueue, uint32 WaitCycles, bool bPermitBackgroundWork)
+	void FScheduler::WorkerMain(FSleepEvent* ExternalWorkerEvent, FSchedulerTls::FLocalQueueType* ExternalWorkerLocalQueue, uint32 WaitCycles, bool bPermitBackgroundWork)
 	{
 		FTaskTagScope WorkerScope(ETaskTag::EWorkerThread);
 		FSchedulerTls::ActiveScheduler = this;
@@ -252,6 +268,12 @@ namespace LowLevelTasks
 		FMemory::SetupTLSCachesOnCurrentThread();
 		FSchedulerTls::WorkerType = bPermitBackgroundWork ? FSchedulerTls::EWorkerType::Background : FSchedulerTls::EWorkerType::Foreground;
 
+		FSleepEvent* WorkerEvent = ExternalWorkerEvent;
+		if (!ExternalWorkerEvent)
+		{
+			WorkerEvent = new FSleepEvent();
+		}
+
 		checkSlow(FSchedulerTls::LocalQueue == nullptr);
 		if(ExternalWorkerLocalQueue)
 		{
@@ -259,26 +281,34 @@ namespace LowLevelTasks
 		}
 		else
 		{
-			FSchedulerTls::LocalQueue = FSchedulerTls::FLocalQueueType::AllocateLocalQueue(QueueRegistry, bPermitBackgroundWork ? ELocalQueueType::EBackground : ELocalQueueType::EForeground);
+			FSchedulerTls::LocalQueue = FSchedulerTls::FLocalQueueType::AllocateLocalQueue(QueueRegistry, bPermitBackgroundWork ? ELocalQueueType::EBackground : ELocalQueueType::EForeground, WorkerEvent);
 		}
+
 		FSchedulerTls::FLocalQueueType* WorkerLocalQueue = FSchedulerTls::LocalQueue;
 
-		bool Drowsing = false;
+		bool bDrowsing = false;
 		uint32 WaitCount = 0;
 		FSchedulerTls::FQueueRegistry::FOutOfWork OutOfWork = QueueRegistry.GetOutOfWorkScope(bPermitBackgroundWork ? ELocalQueueType::EBackground : ELocalQueueType::EForeground);
 		while (true)
 		{
-			while(TryExecuteTaskFrom<&FSchedulerTls::FLocalQueueType::DequeueLocal, false>(WorkerLocalQueue, OutOfWork, bPermitBackgroundWork)
-			   || TryExecuteTaskFrom<&FSchedulerTls::FLocalQueueType::DequeueGlobal, false>(WorkerLocalQueue, OutOfWork, bPermitBackgroundWork))
+			while(TryExecuteTaskFrom<&FSchedulerTls::FLocalQueueType::DequeueLocal,  false>(WorkerLocalQueue, OutOfWork, bPermitBackgroundWork, bDrowsing)
+			   || TryExecuteTaskFrom<&FSchedulerTls::FLocalQueueType::DequeueGlobal, false>(WorkerLocalQueue, OutOfWork, bPermitBackgroundWork, bDrowsing))
 			{		
-				Drowsing = false;
+				bDrowsing = false;
+				WaitCount = 0;
+			}
+
+			while(TryExecuteTaskFrom<&FSchedulerTls::FLocalQueueType::DequeueLocal, false>(WorkerLocalQueue, OutOfWork, bPermitBackgroundWork, bDrowsing)
+			   || TryExecuteTaskFrom<&FSchedulerTls::FLocalQueueType::DequeueSteal, false>(WorkerLocalQueue, OutOfWork, bPermitBackgroundWork, bDrowsing))
+			{
+				bDrowsing = false;
 				WaitCount = 0;
 			}
 
-			while(TryExecuteTaskFrom<&FSchedulerTls::FLocalQueueType::DequeueLocal, false>(WorkerLocalQueue, OutOfWork, bPermitBackgroundWork)
-			   || TryExecuteTaskFrom<&FSchedulerTls::FLocalQueueType::DequeueSteal, false>(WorkerLocalQueue, OutOfWork, bPermitBackgroundWork))
+			while(TryExecuteTaskFrom<&FSchedulerTls::FLocalQueueType::DequeueLocal, false>(WorkerLocalQueue, OutOfWork, bPermitBackgroundWork, bDrowsing)
+				|| TryExecuteTaskFrom<&FSchedulerTls::FLocalQueueType::DequeueAffinity, false>(WorkerLocalQueue, OutOfWork, bPermitBackgroundWork, true))
 			{
-				Drowsing = false;
+				bDrowsing = false;
 				WaitCount = 0;
 			}
 
@@ -295,7 +325,7 @@ namespace LowLevelTasks
 				continue;
 			}
 
-			TrySleeping(WorkerEvent, OutOfWork, Drowsing, bPermitBackgroundWork);
+			TrySleeping(WorkerEvent, OutOfWork.Stop(), bDrowsing, bPermitBackgroundWork);
 		}
 
 		while (WakeUpWorker(bPermitBackgroundWork)) {}
@@ -303,6 +333,11 @@ namespace LowLevelTasks
 		FSchedulerTls::FLocalQueueType::DeleteLocalQueue(WorkerLocalQueue, bPermitBackgroundWork ? ELocalQueueType::EBackground : ELocalQueueType::EForeground, ExternalWorkerLocalQueue != nullptr);
 		FSchedulerTls::LocalQueue = nullptr;
 
+		if (!ExternalWorkerEvent)
+		{
+			delete WorkerEvent;
+		}
+
 		FSchedulerTls::ActiveScheduler = nullptr;
 		FSchedulerTls::WorkerType = FSchedulerTls::EWorkerType::None;
 		FMemory::ClearAndDisableTLSCachesOnCurrentThread();
@@ -327,8 +362,18 @@ namespace LowLevelTasks
 		FSchedulerTls::FQueueRegistry::FOutOfWork OutOfWork = QueueRegistry.GetOutOfWorkScope(bPermitBackgroundWork ? ELocalQueueType::EBackground : ELocalQueueType::EForeground);
 		while (true)
 		{
-			while(TryExecuteTaskFrom<&FSchedulerTls::FLocalQueueType::DequeueLocal, true>(WorkerLocalQueue, OutOfWork, bPermitBackgroundWork)
-			   || TryExecuteTaskFrom<&FSchedulerTls::FLocalQueueType::DequeueGlobal, true>(WorkerLocalQueue, OutOfWork, bPermitBackgroundWork))
+			while(TryExecuteTaskFrom<&FSchedulerTls::FLocalQueueType::DequeueLocal, true>(WorkerLocalQueue, OutOfWork, bPermitBackgroundWork, true)
+			   || TryExecuteTaskFrom<&FSchedulerTls::FLocalQueueType::DequeueGlobal, true>(WorkerLocalQueue, OutOfWork, bPermitBackgroundWork, true))
+			{
+				if (Conditional())
+				{
+					return;
+				}
+				WaitCount = 0;
+			}
+
+			while(TryExecuteTaskFrom<&FSchedulerTls::FLocalQueueType::DequeueLocal, true>(WorkerLocalQueue, OutOfWork, bPermitBackgroundWork, true)
+			   || TryExecuteTaskFrom<&FSchedulerTls::FLocalQueueType::DequeueSteal, true>(WorkerLocalQueue, OutOfWork, bPermitBackgroundWork, true))
 			{
 				if (Conditional())
 				{
@@ -337,8 +382,8 @@ namespace LowLevelTasks
 				WaitCount = 0;
 			}
 
-			while(TryExecuteTaskFrom<&FSchedulerTls::FLocalQueueType::DequeueLocal, true>(WorkerLocalQueue, OutOfWork, bPermitBackgroundWork)
-			   || TryExecuteTaskFrom<&FSchedulerTls::FLocalQueueType::DequeueSteal, true>(WorkerLocalQueue, OutOfWork, bPermitBackgroundWork))
+			while(TryExecuteTaskFrom<&FSchedulerTls::FLocalQueueType::DequeueLocal, true>(WorkerLocalQueue, OutOfWork, bPermitBackgroundWork, true)
+				|| TryExecuteTaskFrom<&FSchedulerTls::FLocalQueueType::DequeueAffinity, true>(WorkerLocalQueue, OutOfWork, bPermitBackgroundWork, true))
 			{
 				if (Conditional())
 				{
diff --git a/Engine/Source/Runtime/Core/Private/Tests/Async/TaskGraphTest.cpp b/Engine/Source/Runtime/Core/Private/Tests/Async/TaskGraphTest.cpp
--- a/Engine/Source/Runtime/Core/Private/Tests/Async/TaskGraphTest.cpp
+++ b/Engine/Source/Runtime/Core/Private/Tests/Async/TaskGraphTest.cpp
@@ -714,6 +714,8 @@ namespace OldTaskGraphTests
 	}
 }
 
+extern uint32 GNumForegroundWorkers;
+
 namespace TaskGraphTests
 {
 	IMPLEMENT_SIMPLE_AUTOMATION_TEST(FTaskGraphGraphEventTest, "System.Core.Async.TaskGraph.GraphEventTest", EAutomationTestFlags::ApplicationContextMask | EAutomationTestFlags::EngineFilter);
@@ -1418,6 +1420,64 @@ PRAGMA_ENABLE_DEPRECATION_WARNINGS
 
 		return true;
 	}
+
+	IMPLEMENT_SIMPLE_AUTOMATION_TEST(FTaskGraphBlockingWorkersTest, "System.Core.TaskGraph.BlockingWorkers", EAutomationTestFlags::ApplicationContextMask | EAutomationTestFlags::EngineFilter);
+
+	template<uint32 Num, bool bBackgroundWorkersToo>
+	void BlockingWorkersByFEvent();
+
+	bool FTaskGraphBlockingWorkersTest::RunTest(const FString& Parameters)
+	{
+		FPlatformProcess::Sleep(1.0f);
+		UE_BENCHMARK(10, BlockingWorkersByFEvent<100, false>);
+		UE_BENCHMARK(10, BlockingWorkersByFEvent<100, true>);
+		return true;
+	}
+
+	template<uint32 Num, bool bBackgroundWorkersToo>
+	void BlockingWorkersByFEvent()
+	{
+		for (int N = 0; N != Num; ++N)
+		{
+			int32 NumWorkers = bBackgroundWorkersToo ? LowLevelTasks::FScheduler::Get().GetNumWorkers() : GNumForegroundWorkers;
+
+			TArray<LowLevelTasks::FTask> WorkerBlockers; // tasks that block worker threads
+			WorkerBlockers.AddDefaulted(NumWorkers);
+
+			FEventRef BlockersBlocker{ EEventMode::ManualReset }; // blocks `WorkerBlockers`, manual reset because multiple threads are waiting for it
+
+			std::atomic<int32> NumWorkersNotBlocked{ NumWorkers };
+			FEventRef AllWorkersBlockedSignal;
+
+			for (int i = 0; i != NumWorkers; ++i)
+			{
+				WorkerBlockers[i].Init(TEXT("BlockingWorkers"), LowLevelTasks::ETaskPriority::High,
+					[i, &BlockersBlocker, &AllWorkersBlockedSignal, &NumWorkersNotBlocked]
+					{
+						verify(LowLevelTasks::FSchedulerTls::GetAffinityIndex() == i);
+						if (--NumWorkersNotBlocked == 0)
+						{
+							AllWorkersBlockedSignal->Trigger();
+						}
+						else
+						{
+							BlockersBlocker->Wait();
+						}
+					});
+
+				verify(LowLevelTasks::TryLaunchAffinity(WorkerBlockers[i], i));
+			}
+
+			verify(AllWorkersBlockedSignal->Wait(FTimespan::FromSeconds(3)));
+
+			BlockersBlocker->Trigger();
+			for (int i = 0; i != NumWorkers; ++i)
+			{
+				while(!WorkerBlockers[i].IsCompleted())
+				{}
+			}
+		}
+	}
 }
 
 #endif //WITH_DEV_AUTOMATION_TESTS
diff --git a/Engine/Source/Runtime/Core/Public/Async/Fundamental/LocalQueue.h b/Engine/Source/Runtime/Core/Public/Async/Fundamental/LocalQueue.h
--- a/Engine/Source/Runtime/Core/Public/Async/Fundamental/LocalQueue.h
+++ b/Engine/Source/Runtime/Core/Public/Async/Fundamental/LocalQueue.h
@@ -245,21 +245,22 @@ public:
 		friend class TLocalQueueRegistry;
 
 	public:
-		TLocalQueue(TLocalQueueRegistry& InRegistry, ELocalQueueType QueueType) : Registry(&InRegistry)
+		TLocalQueue(TLocalQueueRegistry& InRegistry, ELocalQueueType QueueType, FSleepEvent* InSleepEvent) : Registry(&InRegistry), SleepEvent(InSleepEvent)
 		{
 			checkSlow(Registry);
 			StealHazard = FStealHazard(Registry->QueueCollection, Registry->HazardsCollection);
-			Registry->AddLocalQueue(StealHazard, this, QueueType);
+			AffinityIndex = Registry->AddLocalQueue(StealHazard, this, QueueType);
 			for (int32 PriorityIndex = 0; PriorityIndex < int32(ETaskPriority::Count); PriorityIndex++)
 			{
 				DequeueHazards[PriorityIndex] = Registry->OverflowQueues[PriorityIndex].getHeadHazard();
 			}
+			AffinityHazard = AffinityQueue.getHeadHazard();
 		}
 
-		static TLocalQueue* AllocateLocalQueue(TLocalQueueRegistry& InRegistry, ELocalQueueType QueueType)
+		static TLocalQueue* AllocateLocalQueue(TLocalQueueRegistry& InRegistry, ELocalQueueType QueueType, FSleepEvent* InSleepEvent = nullptr)
 		{
 			void* Memory = FMemory::Malloc(sizeof(TLocalQueue), 128u);
-			return new (Memory) TLocalQueue(InRegistry, QueueType);
+			return new (Memory) TLocalQueue(InRegistry, QueueType, InSleepEvent);
 		}
 
 		//delete a queue
@@ -301,7 +302,7 @@ public:
 			return Registry->LessThanHalfWorkersLookingForWork(bBackgroundTask);
 		}
 
-		inline FTask* DequeueLocal(bool GetBackGroundTasks)
+		inline FTask* DequeueLocal(bool GetBackGroundTasks, bool bDisableThrottleStealing)
 		{
 			int32 MaxPriority = GetBackGroundTasks ? int32(ETaskPriority::Count) : int32(ETaskPriority::ForegroundCount);
 			for (int32 PriorityIndex = 0; PriorityIndex < MaxPriority; PriorityIndex++)
@@ -315,9 +316,9 @@ public:
 			return nullptr;
 		}
 
-		inline FTask* DequeueGlobal(bool GetBackGroundTasks)
+		inline FTask* DequeueGlobal(bool GetBackGroundTasks, bool bDisableThrottleStealing)
 		{
-			if ((Registry->NumActiveWorkers[GetBackGroundTasks].load(std::memory_order_relaxed) >= (2 * Registry->NumWorkersLookingForWork[GetBackGroundTasks].load(std::memory_order_relaxed) - 1)) || ((Random.GetUnsignedInt() % 4) == 0))
+			if (bDisableThrottleStealing || (Registry->NumActiveWorkers[GetBackGroundTasks].load(std::memory_order_relaxed) >= (2 * Registry->NumWorkersLookingForWork[GetBackGroundTasks].load(std::memory_order_relaxed) - 1)) || ((Random.GetUnsignedInt() % 4) == 0))
 			{
 				int32 MaxPriority = GetBackGroundTasks ? int32(ETaskPriority::Count) : int32(ETaskPriority::ForegroundCount);
 				for (int32 PriorityIndex = 0; PriorityIndex < MaxPriority; PriorityIndex++)
@@ -332,9 +333,9 @@ public:
 			return nullptr;
 		}
 
-		inline FTask* DequeueSteal(bool GetBackGroundTasks)
+		inline FTask* DequeueSteal(bool GetBackGroundTasks, bool bDisableThrottleStealing)
 		{
-			if ((Registry->NumActiveWorkers[GetBackGroundTasks].load(std::memory_order_relaxed) >= (2 * Registry->NumWorkersLookingForWork[GetBackGroundTasks].load(std::memory_order_relaxed) - 1)) || ((Random.GetUnsignedInt() % 4) == 0))
+			if (bDisableThrottleStealing || (Registry->NumActiveWorkers[GetBackGroundTasks].load(std::memory_order_relaxed) >= (2 * Registry->NumWorkersLookingForWork[GetBackGroundTasks].load(std::memory_order_relaxed) - 1)) || ((Random.GetUnsignedInt() % 4) == 0))
 			{
 				if (CachedRandomIndex == InvalidIndex)
 				{
@@ -350,15 +351,46 @@ public:
 			return nullptr;
 		}
 
+		inline FTask* DequeueAffinity(bool GetBackGroundTasks, bool bDisableThrottleStealing)
+		{
+			return AffinityQueue.dequeue(AffinityHazard);
+		}
+
+		uint32 GetAffinityIndex() const
+		{
+			return AffinityIndex;
+		}
+
+	private:
+		inline void EnqueueAffinity(FTask* Item)
+		{
+			check(SleepEvent != nullptr);
+			AffinityQueue.enqueue(Item);
+			
+			ESleepState SleepState = ESleepState::Sleeping;
+			ESleepState DrowsingState = ESleepState::Drowsing;
+			if (SleepEvent->SleepState.compare_exchange_strong(DrowsingState, ESleepState::Affinity, std::memory_order_acquire))
+			{
+			}
+			else if (SleepEvent->SleepState.compare_exchange_strong(SleepState, ESleepState::Drowsing, std::memory_order_acquire))
+			{
+				SleepEvent->SleepEvent->Trigger();
+			}
+		}
+
 	private:
 		static constexpr uint32	InvalidIndex = ~0u;
 		FLocalQueueType			LocalQueues[uint32(ETaskPriority::Count)];
 		DequeueHazard			DequeueHazards[uint32(ETaskPriority::Count)];
+		FOverflowQueueType		AffinityQueue;
+		DequeueHazard			AffinityHazard;
 		FStealHazard			StealHazard;
 		TLocalQueueRegistry*	Registry;
+		FSleepEvent*			SleepEvent;
 		FRandomStream			Random;
 		uint32					CachedRandomIndex = InvalidIndex;
 		uint32					CachedPriorityIndex = 0;
+		uint32					AffinityIndex = ~0;
 	};
 
 	TLocalQueueRegistry()
@@ -368,7 +400,7 @@ public:
 
 private:
 	// add a queue to the Registry
-	void AddLocalQueue(FStealHazard& Hazard, TLocalQueue* QueueToAdd, ELocalQueueType QueueType)
+	uint32 AddLocalQueue(FStealHazard& Hazard, TLocalQueue* QueueToAdd, ELocalQueueType QueueType)
 	{
 		if (QueueType != ELocalQueueType::EBusyWait)
 		{
@@ -387,7 +419,7 @@ private:
 			}
 			HazardsCollection.Delete(Previous);
 			Hazard.Retire();
-			return;
+			return Previous->LocalQueues.Num();
 		}
 	}
 
@@ -464,6 +496,19 @@ public:
 		return LessThanHalfWorkersLookingForWork(bBackgroundTask);
 	}
 
+	inline bool EnqueueAffinity(FTask* Item, uint32 AffinityIndex)
+	{
+		FStealHazard Hazard(QueueCollection, HazardsCollection);
+		FLocalQueueCollection* Queues = Hazard.Get();
+
+		if (AffinityIndex < uint32(Queues->LocalQueues.Num()) && Queues->LocalQueues[AffinityIndex]->SleepEvent)
+		{
+			Queues->LocalQueues[AffinityIndex]->EnqueueAffinity(Item);
+			return true;
+		}
+		return false;
+	}
+
 	// grab an Item directy from the Global OverflowQueue
 	FTask* Dequeue()
 	{
@@ -486,8 +531,14 @@ public:
 private:
 	inline bool LessThanHalfWorkersLookingForWork(bool bBackgroundTask) const
 	{
-		return (NumWorkersLookingForWork[bBackgroundTask].load(std::memory_order_acquire) * 2 < NumActiveWorkers[bBackgroundTask].load(std::memory_order_acquire)) 
-			&& (bBackgroundTask || (NumWorkersLookingForWork[true].load(std::memory_order_acquire) * 2 < NumActiveWorkers[true].load(std::memory_order_acquire)));
+		if (bBackgroundTask)
+		{
+			return NumWorkersLookingForWork[true].load(std::memory_order_acquire) * 2 <= NumActiveWorkers[true].load(std::memory_order_acquire);
+		}
+		else
+		{
+			return (NumWorkersLookingForWork[false].load(std::memory_order_acquire) + NumWorkersLookingForWork[true].load(std::memory_order_acquire)) * 2 <= (NumActiveWorkers[false].load(std::memory_order_acquire) + NumActiveWorkers[true].load(std::memory_order_acquire));
+		}
 	}
 
 	FOverflowQueueType	  OverflowQueues[uint32(ETaskPriority::Count)];
diff --git a/Engine/Source/Runtime/Core/Public/Async/Fundamental/Scheduler.h b/Engine/Source/Runtime/Core/Public/Async/Fundamental/Scheduler.h
--- a/Engine/Source/Runtime/Core/Public/Async/Fundamental/Scheduler.h
+++ b/Engine/Source/Runtime/Core/Public/Async/Fundamental/Scheduler.h
@@ -116,6 +116,9 @@ namespace LowLevelTasks
 		// returns true if the current thread execution is in the context of busy-waiting
 		CORE_API static bool IsBusyWaiting();
 
+		// returns the AffinityIndex of the thread LocalQueue
+		CORE_API static uint32 GetAffinityIndex();
+
 	protected:
 		inline static bool PermitBackgroundWork()
 		{
@@ -159,6 +162,9 @@ namespace LowLevelTasks
 		//try to launch the task, the return value will specify if the task was in the ready state and has been launced
 		inline bool TryLaunch(FTask& Task, EQueuePreference QueuePreference = EQueuePreference::DefaultPreference, bool bWakeUpWorker = true);	
 
+		//try to launch the task on a specific worker ID, the return value will specify if the task was in the ready state and has been launched
+		inline bool TryLaunchAffinity(FTask& Task, uint32 AffinityIndex);	
+
 		//tries to do some work until the Task is completed
 		template<typename TaskType>
 		inline void BusyWait(const TaskType& Task, bool ForceAllowBackgroundWork = false);
@@ -178,37 +184,20 @@ namespace LowLevelTasks
 		//get the Queue registry to register additonal WorkerQueues for this Scheduler
 		inline FSchedulerTls::FQueueRegistry& GetQueueRegistry();
 
-	private: //Private Interface of the Scheduler	
-		enum class ESleepState
-		{
-			Running,
-			Drowsing,
-			Sleeping,
-		};
-
-		//the struct is naturally 8 bytes aligned, the extra alignment just 
-		//re-enforces this assumption and will error if it changes in the future
-		struct alignas(8) FSleepEvent
-		{
-			FEventRef SleepEvent;
-			std::atomic<ESleepState> SleepState { ESleepState::Running };
-			FSleepEvent* Next = nullptr;
-		};
-
 	public:
 		FScheduler() = default;
 		~FScheduler();
 
 	private: 
-		TUniquePtr<FThread> CreateWorker(FSchedulerTls::FLocalQueueType* ExternalWorkerLocalQueue = nullptr, EThreadPriority Priority = EThreadPriority::TPri_Normal, bool bPermitBackgroundWork = false, bool bIsForkable = false);
+		TUniquePtr<FThread> CreateWorker(FSleepEvent* ExternalWorkerEvent = nullptr, FSchedulerTls::FLocalQueueType* ExternalWorkerLocalQueue = nullptr, EThreadPriority Priority = EThreadPriority::TPri_Normal, bool bPermitBackgroundWork = false, bool bIsForkable = false);
 		void WorkerMain(struct FSleepEvent* WorkerEvent, FSchedulerTls::FLocalQueueType* ExternalWorkerLocalQueue, uint32 WaitCycles, bool bPermitBackgroundWork);
 		CORE_API void LaunchInternal(FTask& Task, EQueuePreference QueuePreference, bool bWakeUpWorker);
 		CORE_API void BusyWaitInternal(const FConditional& Conditional, bool ForceAllowBackgroundWork);
-		FORCENOINLINE void TrySleeping(FSleepEvent* WorkerEvent, FSchedulerTls::FQueueRegistry::FOutOfWork& OutOfWork, bool& Drowsing, bool bBackgroundWorker);
+		FORCENOINLINE void TrySleeping(FSleepEvent* WorkerEvent, bool bStopOutOfWorkScope, bool& Drowsing, bool bBackgroundWorker);
 		inline bool WakeUpWorker(bool bBackgroundWorker);
 
-		template<FTask* (FSchedulerTls::FLocalQueueType::*DequeueFunction)(bool), bool bIsBusyWaiting>
-		bool TryExecuteTaskFrom(FSchedulerTls::FLocalQueueType* Queue, FSchedulerTls::FQueueRegistry::FOutOfWork& OutOfWork, bool bPermitBackgroundWork);
+		template<FTask* (FSchedulerTls::FLocalQueueType::*DequeueFunction)(bool, bool), bool bIsBusyWaiting>
+		bool TryExecuteTaskFrom(FSchedulerTls::FLocalQueueType* Queue, FSchedulerTls::FQueueRegistry::FOutOfWork& OutOfWork, bool bPermitBackgroundWork, bool bDisableThrottleStealing);
 
 	private:
 		TEventStack<FSleepEvent> 				SleepEventStack[2];
@@ -216,6 +205,7 @@ namespace LowLevelTasks
 		FCriticalSection 						WorkerThreadsCS;
 		TArray<TUniquePtr<FThread>>				WorkerThreads;
 		TArray<FSchedulerTls::FLocalQueueType>	WorkerLocalQueues;
+		TArray<FSleepEvent>						WorkerEvents;
 		std::atomic_uint						ActiveWorkers { 0 };
 		std::atomic_uint						NextWorkerId { 0 };
 	};
@@ -225,6 +215,11 @@ namespace LowLevelTasks
 		return FScheduler::Get().TryLaunch(Task, QueuePreference, bWakeUpWorker);
 	}
 
+	FORCEINLINE_DEBUGGABLE bool TryLaunchAffinity(FTask& Task, uint32 AffinityIndex)
+	{
+		return FScheduler::Get().TryLaunchAffinity(Task, AffinityIndex);
+	}
+
 	FORCEINLINE_DEBUGGABLE void BusyWaitForTask(const FTask& Task, bool ForceAllowBackgroundWork = false)
 	{
 		FScheduler::Get().BusyWait(Task, ForceAllowBackgroundWork);
@@ -249,12 +244,21 @@ namespace LowLevelTasks
 	{
 		if(Task.TryPrepareLaunch())
 		{
-			FScheduler::Get().LaunchInternal(Task, QueuePreference, bWakeUpWorker);
+			LaunchInternal(Task, QueuePreference, bWakeUpWorker);
 			return true;
 		}
 		return false;
 	}
 
+	inline bool FScheduler::TryLaunchAffinity(FTask& Task, uint32 AffinityIndex)
+	{
+		if(Task.TryPrepareLaunch())
+		{
+			return QueueRegistry.EnqueueAffinity(&Task, AffinityIndex);
+		}
+		return false;
+	}
+
 	inline uint32 FScheduler::GetNumWorkers() const
 	{
 		return ActiveWorkers.load(std::memory_order_relaxed);
@@ -306,25 +310,30 @@ namespace LowLevelTasks
 		}
 	}
 
-	inline void FScheduler::TrySleeping(FSleepEvent* WorkerEvent, FSchedulerTls::FQueueRegistry::FOutOfWork& OutOfWork, bool& Drowsing, bool bBackgroundWorker)
+	inline void FScheduler::TrySleeping(FSleepEvent* WorkerEvent, bool bStopOutOfWorkScope, bool& Drowsing, bool bBackgroundWorker)
 	{
 		ESleepState DrowsingState1 = ESleepState::Drowsing;
 		ESleepState DrowsingState2 = ESleepState::Drowsing;
 		ESleepState RunningState  = ESleepState::Running;
-		if(!Drowsing && WorkerEvent->SleepState.compare_exchange_strong(DrowsingState1, ESleepState::Drowsing, std::memory_order_relaxed)) //continue drowsing
+		ESleepState AffinityState  = ESleepState::Affinity;
+
+		if(!Drowsing && WorkerEvent->SleepState.compare_exchange_strong(DrowsingState1, ESleepState::Drowsing, std::memory_order_release)) //continue drowsing
 		{
-			verifySlow(OutOfWork.Stop());
+			verifySlow(bStopOutOfWorkScope);
 			Drowsing = true; // Alternative State one: ((Running -> Drowsing) -> Drowsing)
 		}
-		else if(WorkerEvent->SleepState.compare_exchange_strong(DrowsingState2, ESleepState::Sleeping, std::memory_order_relaxed))
+		else if(WorkerEvent->SleepState.compare_exchange_strong(DrowsingState2, ESleepState::Sleeping, std::memory_order_release))
 		{
-			verifySlow(!OutOfWork.Stop());
+			verifySlow(!bStopOutOfWorkScope);
 			Drowsing = false;
 			WorkerEvent->SleepEvent->Wait(); // State two: ((Running -> Drowsing) -> Sleeping)
 		}
-		else if(WorkerEvent->SleepState.compare_exchange_strong(RunningState, ESleepState::Drowsing, std::memory_order_relaxed))
+		else if(WorkerEvent->SleepState.compare_exchange_strong(AffinityState, ESleepState::Drowsing, std::memory_order_release)) //continue drowsing
+		{
+			Drowsing = true; // Alternative State one: ((Running -> Drowsing) -> Drowsing)
+		}
+		else if(WorkerEvent->SleepState.compare_exchange_strong(RunningState, ESleepState::Drowsing, std::memory_order_release))
 		{
-			OutOfWork.Stop();
 			Drowsing = true;
 			SleepEventStack[bBackgroundWorker].Push(WorkerEvent); // State one: (Running -> Drowsing)
 		}
@@ -336,17 +345,15 @@ namespace LowLevelTasks
 
 	inline bool FScheduler::WakeUpWorker(bool bBackgroundWorker)
 	{
-		FSleepEvent* WorkerEvent = SleepEventStack[bBackgroundWorker].Pop();
-		if (WorkerEvent)
+		while (FSleepEvent* WorkerEvent = SleepEventStack[bBackgroundWorker].Pop())
 		{
-			ESleepState SleepState = WorkerEvent->SleepState.exchange(ESleepState::Running, std::memory_order_relaxed);
+			ESleepState SleepState = WorkerEvent->SleepState.exchange(ESleepState::Running, std::memory_order_acquire);
 			if (SleepState == ESleepState::Sleeping)
 			{
 				WorkerEvent->SleepEvent->Trigger(); 
 				return true; // Solving State two: (((Running -> Drowsing) -> Sleeping) -> Running)
 			}
-			checkf(SleepState == ESleepState::Drowsing, TEXT("Worker was not drowsing: %d"), WorkerEvent->SleepState.load(std::memory_order_relaxed));
-			return true; // Solving State one: (Running -> Drowsing) -> Running  OR ((Running -> Drowsing) -> Drowsing) -> Running
+			continue; // Solving State one: (Running -> Drowsing) -> Running  OR ((Running -> Drowsing) -> Drowsing) -> Running
 		}
 		return false;
 	}
diff --git a/Engine/Source/Runtime/Core/Public/Async/Fundamental/Task.h b/Engine/Source/Runtime/Core/Public/Async/Fundamental/Task.h
--- a/Engine/Source/Runtime/Core/Public/Async/Fundamental/Task.h
+++ b/Engine/Source/Runtime/Core/Public/Async/Fundamental/Task.h
@@ -127,7 +127,7 @@ namespace LowLevelTasks
 	private:
 		using FTaskDelegate = TTaskDelegate<PLATFORM_CACHE_LINE_SIZE - sizeof(FPackedData) - sizeof(void*), bool>;
 		FTaskDelegate Runnable;
-		mutable void* UserData;
+		mutable void* UserData = nullptr;
 		std::atomic<FPackedData> PackedData { FPackedData() };
 
 	private:
@@ -178,28 +178,16 @@ namespace LowLevelTasks
 		inline bool TryExecute();
 
 		template<typename TRunnable, typename TContinuation>
-		inline void Init(const TCHAR* InDebugName, ETaskPriority InPriority, TRunnable&& InRunnable, TContinuation&& InContinuation, bool bAllowBusyWaiting);
+		inline void Init(const TCHAR* InDebugName, ETaskPriority InPriority, TRunnable&& InRunnable, TContinuation&& InContinuation, bool bAllowBusyWaiting = true);
 
 		template<typename TRunnable>
-		inline void Init(const TCHAR* InDebugName, ETaskPriority InPriority, TRunnable&& InRunnable, bool bAllowBusyWaiting);
+		inline void Init(const TCHAR* InDebugName, ETaskPriority InPriority, TRunnable&& InRunnable, bool bAllowBusyWaiting = true);
 
 		template<typename TRunnable, typename TContinuation>
-		inline void Init(const TCHAR* InDebugName, TRunnable&& InRunnable, TContinuation&& InContinuation, bool bAllowBusyWaiting);
+		inline void Init(const TCHAR* InDebugName, TRunnable&& InRunnable, TContinuation&& InContinuation, bool bAllowBusyWaiting = true);
 
 		template<typename TRunnable>
-		inline void Init(const TCHAR* InDebugName, TRunnable&& InRunnable, bool bAllowBusyWaiting);
-
-		template<typename TRunnable, typename TContinuation>
-		inline void Init(const TCHAR* InDebugName, ETaskPriority InPriority, TRunnable&& InRunnable, TContinuation&& InContinuation);
-
-		template<typename TRunnable>
-		inline void Init(const TCHAR* InDebugName, ETaskPriority InPriority, TRunnable&& InRunnable);
-
-		template<typename TRunnable, typename TContinuation>
-		inline void Init(const TCHAR* InDebugName, TRunnable&& InRunnable, TContinuation&& InContinuation);
-
-		template<typename TRunnable>
-		inline void Init(const TCHAR* InDebugName, TRunnable&& InRunnable);
+		inline void Init(const TCHAR* InDebugName, TRunnable&& InRunnable, bool bAllowBusyWaiting = true);
 
 		inline const TCHAR* GetDebugName() const;
 		inline ETaskPriority GetPriority() const;
@@ -291,30 +279,6 @@ namespace LowLevelTasks
 		Init(InDebugName, ETaskPriority::Default, Forward<TRunnable>(InRunnable), bAllowBusyWaiting);
 	}
 
-	template<typename TRunnable, typename TContinuation>
-	inline void FTask::Init(const TCHAR* InDebugName, ETaskPriority InPriority, TRunnable&& InRunnable, TContinuation&& InContinuation)
-	{
-		Init(InDebugName, InPriority, Forward<TRunnable>(InRunnable), Forward<TContinuation>(InContinuation), true);
-	}
-
-	template<typename TRunnable>
-	inline void FTask::Init(const TCHAR* InDebugName, ETaskPriority InPriority, TRunnable&& InRunnable)
-	{
-		Init(InDebugName, InPriority, Forward<TRunnable>(InRunnable), true);
-	}
-
-	template<typename TRunnable, typename TContinuation>
-	inline void FTask::Init(const TCHAR* InDebugName, TRunnable&& InRunnable, TContinuation&& InContinuation)
-	{
-		Init(InDebugName, ETaskPriority::Default, Forward<TRunnable>(InRunnable), Forward<TContinuation>(InContinuation), true);
-	}
-
-	template<typename TRunnable>
-	inline void FTask::Init(const TCHAR* InDebugName, TRunnable&& InRunnable)
-	{
-		Init(InDebugName, ETaskPriority::Default, Forward<TRunnable>(InRunnable), true);
-	}
-
 	inline FTask::~FTask()
 	{
 		checkf(IsCompleted(), TEXT("State: %d"), PackedData.load(std::memory_order_relaxed).GetState());
@@ -388,5 +352,22 @@ namespace LowLevelTasks
 	inline bool FTask::AllowBusyWaiting() const
 	{
 		return PackedData.load(std::memory_order_relaxed).AllowBusyWaiting();
-	}	
+	}
+
+	enum class ESleepState
+	{
+		Affinity,
+		Running,
+		Drowsing,
+		Sleeping,
+	};
+
+	//the struct is naturally 64 bytes aligned, the extra alignment just 
+	//re-enforces this assumption and will error if it changes in the future
+	struct alignas(64) FSleepEvent
+	{
+		FEventRef SleepEvent;
+		std::atomic<ESleepState> SleepState { ESleepState::Running };
+		FSleepEvent* Next = nullptr;
+	};
 }
\ No newline at end of file
