missing include

diff --git a/Engine/Source/Runtime/Core/Public/Async/Fundamental/Task.h b/Engine/Source/Runtime/Core/Public/Async/Fundamental/Task.h
--- a/Engine/Source/Runtime/Core/Public/Async/Fundamental/Task.h
+++ b/Engine/Source/Runtime/Core/Public/Async/Fundamental/Task.h
@@ -3,6 +3,7 @@
 #pragma once
 #include "Logging/LogMacros.h"
 #include "TaskDelegate.h"
+#include "HAL/Event.h"
 #include "CoreTypes.h"
 #include <atomic>
 
