fix macos alignment restrictions

diff --git a/Engine/Source/Runtime/Core/Private/Async/Fundamental/Scheduler.cpp b/Engine/Source/Runtime/Core/Private/Async/Fundamental/Scheduler.cpp
--- a/Engine/Source/Runtime/Core/Private/Async/Fundamental/Scheduler.cpp
+++ b/Engine/Source/Runtime/Core/Private/Async/Fundamental/Scheduler.cpp
@@ -260,7 +260,7 @@ namespace LowLevelTasks
 		return false;
 	}
 
-	void FScheduler::WorkerMain(FSleepEvent* ExternalWorkerEvent, FSchedulerTls::FLocalQueueType* ExternalWorkerLocalQueue, uint32 WaitCycles, bool bPermitBackgroundWork)
+	void FScheduler::WorkerMain(FSleepEvent* WorkerEvent, FSchedulerTls::FLocalQueueType* ExternalWorkerLocalQueue, uint32 WaitCycles, bool bPermitBackgroundWork)
 	{
 		FTaskTagScope WorkerScope(ETaskTag::EWorkerThread);
 		FSchedulerTls::ActiveScheduler = this;
@@ -268,10 +268,10 @@ namespace LowLevelTasks
 		FMemory::SetupTLSCachesOnCurrentThread();
 		FSchedulerTls::WorkerType = bPermitBackgroundWork ? FSchedulerTls::EWorkerType::Background : FSchedulerTls::EWorkerType::Foreground;
 
-		FSleepEvent* WorkerEvent = ExternalWorkerEvent;
-		if (!ExternalWorkerEvent)
+		FSleepEvent LocalWorkerEvent;
+		if (!WorkerEvent)
 		{
-			WorkerEvent = new FSleepEvent();
+			WorkerEvent = &LocalWorkerEvent;
 		}
 
 		checkSlow(FSchedulerTls::LocalQueue == nullptr);
@@ -333,11 +333,6 @@ namespace LowLevelTasks
 		FSchedulerTls::FLocalQueueType::DeleteLocalQueue(WorkerLocalQueue, bPermitBackgroundWork ? ELocalQueueType::EBackground : ELocalQueueType::EForeground, ExternalWorkerLocalQueue != nullptr);
 		FSchedulerTls::LocalQueue = nullptr;
 
-		if (!ExternalWorkerEvent)
-		{
-			delete WorkerEvent;
-		}
-
 		FSchedulerTls::ActiveScheduler = nullptr;
 		FSchedulerTls::WorkerType = FSchedulerTls::EWorkerType::None;
 		FMemory::ClearAndDisableTLSCachesOnCurrentThread();
