Convert Stats and Csv tracing to use array/string type trace event fields instead of attachments

#rnx

diff --git a/Engine/Source/Developer/TraceServices/Private/Analyzers/CsvProfilerTraceAnalysis.cpp b/Engine/Source/Developer/TraceServices/Private/Analyzers/CsvProfilerTraceAnalysis.cpp
--- a/Engine/Source/Developer/TraceServices/Private/Analyzers/CsvProfilerTraceAnalysis.cpp
+++ b/Engine/Source/Developer/TraceServices/Private/Analyzers/CsvProfilerTraceAnalysis.cpp
@@ -130,9 +130,17 @@ bool FCsvProfilerAnalyzer::OnEvent(uint16 RouteId, EStyle Style, const FOnEventC
 	}
 	case RouteId_Metadata:
 	{
-		const TCHAR* Key = reinterpret_cast<const TCHAR*>(EventData.GetAttachment());
-		const TCHAR* Value = reinterpret_cast<const TCHAR*>(EventData.GetAttachment() + EventData.GetValue<uint16>("ValueOffset"));
-		CsvProfilerProvider.SetMetadata(Session.StoreString(Key), Session.StoreString(Value));
+		FString Key, Value;
+		if (EventData.GetString("Key", Key))
+		{
+			EventData.GetString("Value", Value);
+		}
+		else
+		{
+			Key = reinterpret_cast<const TCHAR*>(EventData.GetAttachment());
+			Value = reinterpret_cast<const TCHAR*>(EventData.GetAttachment() + EventData.GetValue<uint16>("ValueOffset"));
+		}
+		CsvProfilerProvider.SetMetadata(Session.StoreString(*Key), Session.StoreString(*Value));
 		break;
 	}
 	case RouteId_BeginCapture:
@@ -141,8 +149,9 @@ bool FCsvProfilerAnalyzer::OnEvent(uint16 RouteId, EStyle Style, const FOnEventC
 		RHIThreadId = FTraceAnalyzerUtils::GetThreadIdField(Context, "RHIThreadId");
 		uint32 CaptureStartFrame = GetFrameNumberForTimestamp(TraceFrameType_Game, Context.EventTime.AsSeconds(EventData.GetValue<uint64>("Cycle")));
 		bEnableCounts = EventData.GetValue<bool>("EnableCounts");
-		const TCHAR* Filename = Session.StoreString(reinterpret_cast<const TCHAR*>(EventData.GetAttachment()));
-		CsvProfilerProvider.StartCapture(Filename, CaptureStartFrame);
+		FString FileName = FTraceAnalyzerUtils::LegacyAttachmentString<TCHAR>("FileName", Context);
+		const TCHAR* StoredFileName = Session.StoreString(*FileName);
+		CsvProfilerProvider.StartCapture(StoredFileName, CaptureStartFrame);
 		break;
 	}
 	case RouteId_EndCapture:
@@ -401,7 +410,7 @@ void FCsvProfilerAnalyzer::HandleEventEvent(const FOnEventContext& Context)
 	FThreadState& ThreadState = GetThreadState(ThreadId);
 	uint64 Cycle = Context.EventData.GetValue<uint64>("Cycle");
 	uint32 FrameNumber = GetFrameNumberForTimestamp(ThreadState.FrameType, Context.EventTime.AsSeconds(Cycle));
-	FString EventText = reinterpret_cast<const TCHAR*>(Context.EventData.GetAttachment());
+	FString EventText = FTraceAnalyzerUtils::LegacyAttachmentString<TCHAR>("Text", Context);
 	int32 CategoryIndex = Context.EventData.GetValue<int32>("CategoryIndex");
 	if (CategoryIndex > 0)
 	{
diff --git a/Engine/Source/Developer/TraceServices/Private/Analyzers/StatsTraceAnalysis.cpp b/Engine/Source/Developer/TraceServices/Private/Analyzers/StatsTraceAnalysis.cpp
--- a/Engine/Source/Developer/TraceServices/Private/Analyzers/StatsTraceAnalysis.cpp
+++ b/Engine/Source/Developer/TraceServices/Private/Analyzers/StatsTraceAnalysis.cpp
@@ -67,8 +67,9 @@ bool FStatsAnalyzer::OnEvent(uint16 RouteId, EStyle Style, const FOnEventContext
 	{
 		uint32 ThreadId = FTraceAnalyzerUtils::GetThreadIdField(Context);
 		TSharedRef<FThreadState> ThreadState = GetThreadState(ThreadId);
-		uint64 BufferSize = EventData.GetAttachmentSize();
-		const uint8* BufferPtr = EventData.GetAttachment();
+		TArrayView<const uint8> DataView = FTraceAnalyzerUtils::LegacyAttachmentArray("Data", Context);
+		uint64 BufferSize = DataView.Num();
+		const uint8* BufferPtr = DataView.GetData();
 		const uint8* BufferEnd = BufferPtr + BufferSize;
 		while (BufferPtr < BufferEnd)
 		{
diff --git a/Engine/Source/Runtime/Core/Private/ProfilingDebugging/CsvProfilerTrace.cpp b/Engine/Source/Runtime/Core/Private/ProfilingDebugging/CsvProfilerTrace.cpp
--- a/Engine/Source/Runtime/Core/Private/ProfilingDebugging/CsvProfilerTrace.cpp
+++ b/Engine/Source/Runtime/Core/Private/ProfilingDebugging/CsvProfilerTrace.cpp
@@ -62,6 +62,7 @@ UE_TRACE_EVENT_END()
 UE_TRACE_EVENT_BEGIN(CsvProfiler, Event)
 	UE_TRACE_EVENT_FIELD(uint64, Cycle)
 	UE_TRACE_EVENT_FIELD(int32, CategoryIndex)
+	UE_TRACE_EVENT_FIELD(UE::Trace::WideString, Text)
 UE_TRACE_EVENT_END()
 
 UE_TRACE_EVENT_BEGIN(CsvProfiler, BeginCapture)
@@ -69,6 +70,7 @@ UE_TRACE_EVENT_BEGIN(CsvProfiler, BeginCapture)
 	UE_TRACE_EVENT_FIELD(uint32, RenderThreadId)
 	UE_TRACE_EVENT_FIELD(uint32, RHIThreadId)
 	UE_TRACE_EVENT_FIELD(bool, EnableCounts)
+	UE_TRACE_EVENT_FIELD(UE::Trace::WideString, FileName)
 UE_TRACE_EVENT_END()
 
 UE_TRACE_EVENT_BEGIN(CsvProfiler, EndCapture)
@@ -76,7 +78,8 @@ UE_TRACE_EVENT_BEGIN(CsvProfiler, EndCapture)
 UE_TRACE_EVENT_END()
 
 UE_TRACE_EVENT_BEGIN(CsvProfiler, Metadata)
-	UE_TRACE_EVENT_FIELD(uint16, ValueOffset)
+	UE_TRACE_EVENT_FIELD(UE::Trace::WideString, Key)
+	UE_TRACE_EVENT_FIELD(UE::Trace::WideString, Value)
 UE_TRACE_EVENT_END()
 
 struct FCsvProfilerTraceInternal
@@ -227,22 +230,20 @@ void FCsvProfilerTrace::OutputCustomStat(const FName& StatName, int32 CategoryIn
 void FCsvProfilerTrace::OutputBeginCapture(const TCHAR* Filename, uint32 RenderThreadId, uint32 RHIThreadId, const char* DefaultWaitStatName, bool bEnableCounts)
 {
 	OutputInlineStat(DefaultWaitStatName, CSV_CATEGORY_INDEX(Exclusive));
-	uint16 NameSize = (uint16)((FCString::Strlen(Filename) + 1) * sizeof(TCHAR));
-	UE_TRACE_LOG(CsvProfiler, BeginCapture, CountersChannel, NameSize)
+	UE_TRACE_LOG(CsvProfiler, BeginCapture, CountersChannel)
 		<< BeginCapture.Cycle(FPlatformTime::Cycles64())
 		<< BeginCapture.RenderThreadId(RenderThreadId)
 		<< BeginCapture.RHIThreadId(RHIThreadId)
 		<< BeginCapture.EnableCounts(bEnableCounts)
-		<< BeginCapture.Attachment(Filename, NameSize);
+		<< BeginCapture.FileName(Filename);
 }
 
 void FCsvProfilerTrace::OutputEvent(const TCHAR* Text, int32 CategoryIndex, uint64 Cycles)
 {
-	uint16 TextSize = (uint16)((FCString::Strlen(Text) + 1) * sizeof(TCHAR));
-	UE_TRACE_LOG(CsvProfiler, Event, CountersChannel, TextSize)
+	UE_TRACE_LOG(CsvProfiler, Event, CountersChannel)
 		<< Event.Cycle(Cycles)
 		<< Event.CategoryIndex(CategoryIndex)
-		<< Event.Attachment(Text, TextSize);
+		<< Event.Text(Text);
 }
 
 void FCsvProfilerTrace::OutputEndCapture()
@@ -253,16 +254,9 @@ void FCsvProfilerTrace::OutputEndCapture()
 
 void FCsvProfilerTrace::OutputMetadata(const TCHAR* Key, const TCHAR* Value)
 {
-	uint16 KeySize = (uint16)((FCString::Strlen(Key) + 1) * sizeof(TCHAR));
-	uint16 ValueSize = (uint16)((FCString::Strlen(Value) + 1) * sizeof(TCHAR));
-	auto Attachment = [Key, KeySize, Value, ValueSize](uint8* Out)
-	{
-		memcpy(Out, Key, KeySize);
-		memcpy(Out + KeySize, Value, ValueSize);
-	};
-	UE_TRACE_LOG(CsvProfiler, Metadata, CountersChannel, KeySize + ValueSize)
-		<< Metadata.ValueOffset(KeySize)
-		<< Metadata.Attachment(Attachment);
+	UE_TRACE_LOG(CsvProfiler, Metadata, CountersChannel)
+		<< Metadata.Key(Key)
+		<< Metadata.Value(Value);
 }
 
 #endif
diff --git a/Engine/Source/Runtime/Core/Private/Stats/StatsTrace.cpp b/Engine/Source/Runtime/Core/Private/Stats/StatsTrace.cpp
--- a/Engine/Source/Runtime/Core/Private/Stats/StatsTrace.cpp
+++ b/Engine/Source/Runtime/Core/Private/Stats/StatsTrace.cpp
@@ -22,6 +22,7 @@ UE_TRACE_EVENT_BEGIN(Stats, Spec, NoSync|Important)
 UE_TRACE_EVENT_END()
 
 UE_TRACE_EVENT_BEGIN(Stats, EventBatch)
+	UE_TRACE_EVENT_FIELD(uint8[], Data)
 UE_TRACE_EVENT_END()
 
 struct FStatsTraceInternal
@@ -73,8 +74,8 @@ FStatsTraceInternal::FThreadState* FStatsTraceInternal::InitThreadState()
 
 void FStatsTraceInternal::FlushThreadBuffer(FThreadState* ThreadState)
 {
-	UE_TRACE_LOG(Stats, EventBatch, StatsChannel, ThreadState->BufferSize)
-		<< EventBatch.Attachment(ThreadState->Buffer, ThreadState->BufferSize);
+	UE_TRACE_LOG(Stats, EventBatch, StatsChannel)
+		<< EventBatch.Data(ThreadState->Buffer, ThreadState->BufferSize);
 	ThreadState->BufferSize = 0;
 }
 
