Compile error fix if CPUPROFILERTRACE_ENABLED was off by STATS is on.

#rnx

diff --git a/Engine/Source/Runtime/Core/Public/Stats/Stats2.h b/Engine/Source/Runtime/Core/Public/Stats/Stats2.h
--- a/Engine/Source/Runtime/Core/Public/Stats/Stats2.h
+++ b/Engine/Source/Runtime/Core/Public/Stats/Stats2.h
@@ -2324,7 +2324,7 @@ namespace Stats
 {
 	FORCEINLINE bool IsThreadCollectingData()
 	{
-#if STATS
+#if STATS && CPUPROFILERTRACE_ENABLED
 		return UE_TRACE_CHANNELEXPR_IS_ENABLED(CpuChannel) || FThreadStats::IsCollectingData();
 #else
 		return GCycleStatsShouldEmitNamedEvents > 0;
