Implements -cvarsini to be able to load a local ConsoleVariables.ini on a cooked build

diff --git a/Engine/Source/Runtime/Core/Private/Misc/ConfigCacheIni.cpp b/Engine/Source/Runtime/Core/Private/Misc/ConfigCacheIni.cpp
--- a/Engine/Source/Runtime/Core/Private/Misc/ConfigCacheIni.cpp
+++ b/Engine/Source/Runtime/Core/Private/Misc/ConfigCacheIni.cpp
@@ -19,6 +19,7 @@
 #include "Misc/ConfigManifest.h"
 #include "Misc/DataDrivenPlatformInfoRegistry.h"
 #include "Misc/StringBuilder.h"
+#include "Misc/Paths.h"
 #include "Serialization/MemoryReader.h"
 #include "Serialization/MemoryWriter.h"
 #include "Serialization/LargeMemoryReader.h"
@@ -4542,6 +4543,20 @@ void FConfigCacheIni::LoadConsoleVariablesFromINI()
 	ApplyCVarSettingsFromIni(TEXT("Startup"), *ConsoleVariablesPath, ECVF_SetByConsoleVariablesIni, true);
 #endif // !DISABLE_CHEAT_CVARS
 
+#if !DISABLE_CHEAT_CVARS && !UE_BUILD_SHIPPING
+	{
+		FString OverrideConsoleVariablesPath;
+		FParse::Value(FCommandLine::Get(), TEXT("-cvarsini="), OverrideConsoleVariablesPath);
+
+		if (!OverrideConsoleVariablesPath.IsEmpty())
+		{
+			ensureMsgf(FPaths::FileExists(OverrideConsoleVariablesPath), TEXT("-cvarsini's file %s doesn't exist"), *OverrideConsoleVariablesPath);
+			ApplyCVarSettingsFromIni(TEXT("Startup"), *OverrideConsoleVariablesPath, ECVF_SetByConsoleVariablesIni, true);
+		}
+
+	}
+#endif
+
 	// We also apply from Engine.ini [ConsoleVariables] section
 	ApplyCVarSettingsFromIni(TEXT("ConsoleVariables"), *GEngineIni, ECVF_SetBySystemSettingsIni);
 
