FTSTicker: cosmetic

diff --git a/Engine/Source/Runtime/Core/Public/Containers/Ticker.h b/Engine/Source/Runtime/Core/Public/Containers/Ticker.h
--- a/Engine/Source/Runtime/Core/Public/Containers/Ticker.h
+++ b/Engine/Source/Runtime/Core/Public/Containers/Ticker.h
@@ -13,7 +13,7 @@
 
 /**
  * Ticker delegates return true to automatically reschedule at the same delay or false for one-shot.
- * You will not get more than one fire per "frame", which is just a FTicker::Tick call.
+ * You will not get more than one fire per "frame", which is just a FTSTicker::Tick call.
  * DeltaTime argument is the time since the last game frame, *not* since the last tick the
  * delegate received.
  */
@@ -131,7 +131,7 @@ public:
 	 * Constructor
 	 *
 	 * @param InDelay Delay until next fire; 0 means "next frame"
-	 * @param Ticker the ticker to register with. Defaults to FTicker::GetCoreTicker().
+	 * @param Ticker the ticker to register with. Defaults to FTSTicker::GetCoreTicker().
 	*/
 	FTSTickerObjectBase(float InDelay = 0.0f, FTSTicker& Ticker = FTSTicker::GetCoreTicker());
 
