Fixed a bug in double version of NEON's VectorAnyGreaterThan

diff --git a/Engine/Source/Runtime/Core/Private/Tests/Math/UnrealMathTest.cpp b/Engine/Source/Runtime/Core/Private/Tests/Math/UnrealMathTest.cpp
--- a/Engine/Source/Runtime/Core/Private/Tests/Math/UnrealMathTest.cpp
+++ b/Engine/Source/Runtime/Core/Private/Tests/Math/UnrealMathTest.cpp
@@ -20,7 +20,13 @@
 #include <limits>
 #include <cmath>
 
-#if WITH_DEV_AUTOMATION_TESTS
+#if !defined (_MSC_VER)
+#include <arm_neon.h>
+#endif
+
+namespace
+{
+//#if WITH_DEV_AUTOMATION_TESTS
 
 DEFINE_LOG_CATEGORY_STATIC(LogUnrealMathTest, Log, All);
 
@@ -1628,6 +1634,11 @@ bool RunDoubleVectorTest()
 	bool bIsVAGT_TRUE = VectorAnyGreaterThan(V0, V1) != 0;
 	LogTest<double>(TEXT("VectorAnyGreaterThan-true"), bIsVAGT_TRUE);
 
+	V0 = MakeVectorRegister(0.0, -0.0, 1.0, 0.8);
+	V1 = MakeVectorRegister(0.0, 0.0, 0.0, 1.0);
+	bIsVAGT_TRUE = VectorAnyGreaterThan(V0, V1) != 0;
+	LogTest<double>(TEXT("VectorAnyGreaterThan-true"), bIsVAGT_TRUE);
+
 	V0 = MakeVectorRegister(1.0, 3.0, 2.0, 1.0);
 	V1 = MakeVectorRegister(2.0, 4.0, 6.0, 8.0);
 	bool bIsVAGT_FALSE = VectorAnyGreaterThan(V0, V1) == 0;
@@ -1843,6 +1854,20 @@ bool RunDoubleVectorTest()
 	return GPassing;
 }
 
+#if !defined (_MSC_VER)
+static float InvSqrt2(float F)
+{
+	float32x2_t	val = vdup_n_f32(F);
+
+	float32x2_t inv_sqrt = vrsqrte_f32(val);
+
+	inv_sqrt = vmul_f32(vrsqrts_f32(vmul_f32(inv_sqrt, inv_sqrt), val), inv_sqrt);
+	inv_sqrt = vmul_f32(vrsqrts_f32(vmul_f32(inv_sqrt, inv_sqrt), val), inv_sqrt);
+
+	return vget_lane_f32(inv_sqrt, 0);
+}
+#endif
+
 /**
  * Run a suite of vector operations to validate vector intrinsics are working on the platform
  */
@@ -2133,6 +2158,11 @@ bool FVectorRegisterAbstractionTest::RunTest(const FString& Parameters)
 	bool bIsVAGT_TRUE = VectorAnyGreaterThan( V0, V1 ) != 0;
 	LogTest<float>( TEXT("VectorAnyGreaterThan-true"), bIsVAGT_TRUE );
 
+	V0 = MakeVectorRegister(0.0f, -0.0f, 1.0f, 0.8f);
+	V1 = MakeVectorRegister(0.0f, 0.0f, 0.0f, 1.0f);
+	bIsVAGT_TRUE = VectorAnyGreaterThan(V0, V1) != 0;
+	LogTest<float>(TEXT("VectorAnyGreaterThan-true"), bIsVAGT_TRUE);
+
 	V0 = MakeVectorRegister( 1.0f, 3.0f, 2.0f, 1.0f );
 	V1 = MakeVectorRegister( 2.0f, 4.0f, 6.0f, 8.0f );
 	bool bIsVAGT_FALSE = VectorAnyGreaterThan( V0, V1 ) == 0;
@@ -2601,6 +2631,24 @@ bool FVectorRegisterAbstractionTest::RunTest(const FString& Parameters)
 	// Exp, Log tests
 	TestVectorExpLogFunctions<float, VectorRegister4Float>();
 
+#if !defined (_MSC_VER)
+	float Value = 0.1;
+	float Step = 0.1;
+	float LargestDiff = 0.0f;
+	for (int i = 0; i < 10000; ++i)
+	{
+		const float A = FMath::InvSqrt(Value);
+		const float B = InvSqrt2(Value);
+		const float Diff = A - B;
+		if (abs(Diff) > LargestDiff)
+		{
+			LargestDiff = Diff;
+		}
+
+		Value += Step;
+	}
+#endif
+
 	// Quat<->Rotator conversions and equality
 	{
 		// Identity conversion
@@ -3497,4 +3545,5 @@ bool FBitCastTest::RunTest(const FString& Parameters)
 	return true;
 }
 
-#endif //WITH_DEV_AUTOMATION_TESTS
+//#endif //WITH_DEV_AUTOMATION_TESTS
+}
\ No newline at end of file
diff --git a/Engine/Source/Runtime/Core/Public/Math/UnrealMathNeon.h b/Engine/Source/Runtime/Core/Public/Math/UnrealMathNeon.h
--- a/Engine/Source/Runtime/Core/Public/Math/UnrealMathNeon.h
+++ b/Engine/Source/Runtime/Core/Public/Math/UnrealMathNeon.h
@@ -2198,7 +2198,8 @@ FORCEINLINE int32 VectorAnyGreaterThan(VectorRegister4Double Vec1, VectorRegiste
 {
 	uint16x8_t u16x8_1 = (uint16x8_t)vcgtq_f64(Vec1.XY, Vec2.XY);
 	uint16x8_t u16x8_2 = (uint16x8_t)vcgtq_f64(Vec1.ZW, Vec2.ZW);
-	uint8x8_t u8x8 = (uint8x8_t)vget_low_u16(vuzpq_u16(u16x8_1, u16x8_2).val[0]);
+	uint16x8x2_t tmp = vuzpq_u16(u16x8_1, u16x8_2);
+	uint8x8_t u8x8 = (uint8x8_t)vget_low_u16(vuzpq_u16(tmp.val[0], tmp.val[0]).val[0]);
 	u8x8 = vuzp_u8(u8x8, u8x8).val[0];
 	uint32_t buf[2];
 	vst1_u8((uint8_t*)buf, u8x8);
