Removed development debug code from UnrealMathTest

diff --git a/Engine/Source/Runtime/Core/Private/Tests/Math/UnrealMathTest.cpp b/Engine/Source/Runtime/Core/Private/Tests/Math/UnrealMathTest.cpp
--- a/Engine/Source/Runtime/Core/Private/Tests/Math/UnrealMathTest.cpp
+++ b/Engine/Source/Runtime/Core/Private/Tests/Math/UnrealMathTest.cpp
@@ -20,13 +20,7 @@
 #include <limits>
 #include <cmath>
 
-#if !defined (_MSC_VER)
-#include <arm_neon.h>
-#endif
-
-namespace
-{
-//#if WITH_DEV_AUTOMATION_TESTS
+#if WITH_DEV_AUTOMATION_TESTS
 
 DEFINE_LOG_CATEGORY_STATIC(LogUnrealMathTest, Log, All);
 
@@ -1854,20 +1848,6 @@ bool RunDoubleVectorTest()
 	return GPassing;
 }
 
-#if !defined (_MSC_VER)
-static float InvSqrt2(float F)
-{
-	float32x2_t	val = vdup_n_f32(F);
-
-	float32x2_t inv_sqrt = vrsqrte_f32(val);
-
-	inv_sqrt = vmul_f32(vrsqrts_f32(vmul_f32(inv_sqrt, inv_sqrt), val), inv_sqrt);
-	inv_sqrt = vmul_f32(vrsqrts_f32(vmul_f32(inv_sqrt, inv_sqrt), val), inv_sqrt);
-
-	return vget_lane_f32(inv_sqrt, 0);
-}
-#endif
-
 /**
  * Run a suite of vector operations to validate vector intrinsics are working on the platform
  */
@@ -2631,24 +2611,6 @@ bool FVectorRegisterAbstractionTest::RunTest(const FString& Parameters)
 	// Exp, Log tests
 	TestVectorExpLogFunctions<float, VectorRegister4Float>();
 
-#if !defined (_MSC_VER)
-	float Value = 0.1;
-	float Step = 0.1;
-	float LargestDiff = 0.0f;
-	for (int i = 0; i < 10000; ++i)
-	{
-		const float A = FMath::InvSqrt(Value);
-		const float B = InvSqrt2(Value);
-		const float Diff = A - B;
-		if (abs(Diff) > LargestDiff)
-		{
-			LargestDiff = Diff;
-		}
-
-		Value += Step;
-	}
-#endif
-
 	// Quat<->Rotator conversions and equality
 	{
 		// Identity conversion
@@ -3545,5 +3507,4 @@ bool FBitCastTest::RunTest(const FString& Parameters)
 	return true;
 }
 
-//#endif //WITH_DEV_AUTOMATION_TESTS
-}
\ No newline at end of file
+#endif //WITH_DEV_AUTOMATION_TESTS
\ No newline at end of file
