Quotes added around strings from failed comparisons during automation tests to aid readability.

diff --git a/Engine/Source/Runtime/Core/Private/Misc/AutomationTest.cpp b/Engine/Source/Runtime/Core/Private/Misc/AutomationTest.cpp
--- a/Engine/Source/Runtime/Core/Private/Misc/AutomationTest.cpp
+++ b/Engine/Source/Runtime/Core/Private/Misc/AutomationTest.cpp
@@ -1427,7 +1427,7 @@ bool FAutomationTestBase::TestEqual(const TCHAR* What, const TCHAR* Actual, cons
 {
 	if (FCString::Strcmp(Actual, Expected) != 0)
 	{
-		AddError(FString::Printf(TEXT("Expected '%s' to be %s, but it was %s."), What, Expected, Actual), 1);
+		AddError(FString::Printf(TEXT("Expected '%s' to be \"%s\", but it was \"%s\"."), What, Expected, Actual), 1);
 		return false;
 	}
 	return true;
@@ -1437,7 +1437,7 @@ bool FAutomationTestBase::TestEqualInsensitive(const TCHAR* What, const TCHAR* A
 {
 	if (FCString::Stricmp(Actual, Expected) != 0)
 	{
-		AddError(FString::Printf(TEXT("Expected '%s' to be %s, but it was %s."), What, Expected, Actual), 1);
+		AddError(FString::Printf(TEXT("Expected '%s' to be \"%s\", but it was \"%s\"."), What, Expected, Actual), 1);
 		return false;
 	}
 	return true;
