It is possible that the trace store server binary legitimately doesn't exist. While it is still useful to inform the user that the store couldn't be started, it has be downgraded to a normal log message instead of a warning.

#rnx

diff --git a/Engine/Source/Runtime/Core/Private/ProfilingDebugging/TraceAuxiliary.cpp b/Engine/Source/Runtime/Core/Private/ProfilingDebugging/TraceAuxiliary.cpp
--- a/Engine/Source/Runtime/Core/Private/ProfilingDebugging/TraceAuxiliary.cpp
+++ b/Engine/Source/Runtime/Core/Private/ProfilingDebugging/TraceAuxiliary.cpp
@@ -518,7 +518,7 @@ static void LaunchUnrealTraceInternal()
 
 	if (!bOk)
 	{
-		UE_LOG(LogCore, Warning, TEXT("UnrealTraceServer: Unable to launch the trace store from '%s' (%08x)"), *BinPath, GetLastError());
+		UE_LOG(LogCore, Display, TEXT("UnrealTraceServer: Unable to launch the trace store from '%s' (%08x)"), *BinPath, GetLastError());
 		return;
 	}
 
