NonUnity build fix - Added UnrealTemplate.h include to StaticArray.h

diff --git a/Engine/Source/Runtime/Core/Public/Containers/StaticArray.h b/Engine/Source/Runtime/Core/Public/Containers/StaticArray.h
--- a/Engine/Source/Runtime/Core/Public/Containers/StaticArray.h
+++ b/Engine/Source/Runtime/Core/Public/Containers/StaticArray.h
@@ -4,6 +4,7 @@
 
 #include "CoreTypes.h"
 #include "Misc/AssertionMacros.h"
+#include "Templates/UnrealTemplate.h"
 #include "Templates/UnrealTypeTraits.h"
 #include "Delegates/IntegerSequence.h"
 
