Don't warn when compiling with VS2022 toolchain

#rn none
#rnx

diff --git a/Engine/Source/Runtime/Core/Public/Windows/WindowsPlatformCompilerSetup.h b/Engine/Source/Runtime/Core/Public/Windows/WindowsPlatformCompilerSetup.h
--- a/Engine/Source/Runtime/Core/Public/Windows/WindowsPlatformCompilerSetup.h
+++ b/Engine/Source/Runtime/Core/Public/Windows/WindowsPlatformCompilerSetup.h
@@ -9,8 +9,8 @@
 //
 // Future-proofing the min version check so we keep bumping it whenever we upgrade.
 //
-#if defined(_MSC_VER) && _MSC_VER > 1929 
-	#pragma message("Detected compiler newer than Visual Studio 2019, please update min version checking in WindowsPlatformCompilerSetup.h")
+#if defined(_MSC_VER) && _MSC_VER > 1939 
+	#pragma message("Detected compiler newer than Visual Studio 2022, please update min version checking in WindowsPlatformCompilerSetup.h")
 #endif
 
 //
