Fix for FStringView::Mid() which differed from FString::Mid().
Tests for FStringView slicing functions to ensure they match FString's.
New slicing functions for TArrayView.

diff --git a/Engine/Source/Runtime/Core/Private/Tests/Containers/StringViewTest.cpp b/Engine/Source/Runtime/Core/Private/Tests/Containers/StringViewTest.cpp
--- a/Engine/Source/Runtime/Core/Private/Tests/Containers/StringViewTest.cpp
+++ b/Engine/Source/Runtime/Core/Private/Tests/Containers/StringViewTest.cpp
@@ -854,6 +854,109 @@ bool FStringViewTestFindLastChar::RunTest(const FString& Parameters)
 	return true; 
 }
 
+void TestSlicing(FAutomationTestBase& Test, const FString& Str)
+{
+	const FStringView View = Str;
+	const int32       Len  = Str.Len();
+
+	// Left tests
+	{
+		// Try all lengths of the string, including +/- 5 either side
+		for (int32 Index = -5; Index != Len + 5; ++Index)
+		{
+			FString     Substring = Str .Left(Index);
+			FStringView Subview   = View.Left(Index);
+
+			Test.TestEqual(FString::Printf(TEXT("FStringView(\"%s\")::Left(%d)"), *Str, Index), FString(Subview), Substring);
+		}
+	}
+
+	// LeftChop tests
+	{
+		// Try all lengths of the string, including +/- 5 either side
+		for (int32 Index = -5; Index != Len + 5; ++Index)
+		{
+			FString     Substring = Str .LeftChop(Index);
+			FStringView Subview   = View.LeftChop(Index);
+
+			Test.TestEqual(FString::Printf(TEXT("FStringView(\"%s\")::LeftChop(%d)"), *Str, Index), FString(Subview), Substring);
+		}
+	}
+
+	// Right tests
+	{
+		// Try all lengths of the string, including +/- 5 either side
+		for (int32 Index = -5; Index != Len + 5; ++Index)
+		{
+			FString     Substring = Str .Right(Index);
+			FStringView Subview   = View.Right(Index);
+
+			Test.TestEqual(FString::Printf(TEXT("FStringView(\"%s\")::Right(%d)"), *Str, Index), FString(Subview), Substring);
+		}
+	}
+
+	// RightChop tests
+	{
+		// Try all lengths of the string, including +/- 5 either side
+		for (int32 Index = -5; Index != Len + 5; ++Index)
+		{
+			FString     Substring = Str .RightChop(Index);
+			FStringView Subview   = View.RightChop(Index);
+
+			Test.TestEqual(FString::Printf(TEXT("FStringView(\"%s\").RightChop(%d)"), *Str, Index), FString(Subview), Substring);
+		}
+	}
+
+	// Mid tests
+	{
+		// Try all lengths of the string, including +/- 5 either side
+		for (int32 Index = -5; Index != Len + 5; ++Index)
+		{
+			for (int32 Count = -5; Count != Len + 5; ++Count)
+			{
+				FString     Substring = Str .Mid(Index, Count);
+				FStringView Subview   = View.Mid(Index, Count);
+
+				Test.TestEqual(FString::Printf(TEXT("FStringView(\"%s\")::Mid(%d, %d)"), *Str, Index, Count), FString(Subview), Substring);
+			}
+		}
+
+		// Test near limits of int32
+		for (int32 IndexOffset = 0; IndexOffset != Len + 5; ++IndexOffset)
+		{
+			for (int32 CountOffset = 0; CountOffset != Len + 5; ++CountOffset)
+			{
+				int32 Index = MIN_int32 + IndexOffset;
+				int32 Count = MAX_int32 - CountOffset;
+
+				FString     Substring = Str .Mid(Index, Count);
+				FStringView Subview   = View.Mid(Index, Count);
+
+				Test.TestEqual(FString::Printf(TEXT("FStringView(\"%s\")::Mid(%d, %d)"), *Str, Index, Count), FString(Subview), Substring);
+			}
+		}
+	}
+}
+
+IMPLEMENT_SIMPLE_AUTOMATION_TEST(FStringViewTestSlice, TEST_NAME_ROOT ".Slice", TestFlags)
+bool FStringViewTestSlice::RunTest(const FString& Parameters)
+{
+	// We assume that FString has already passed its tests, and we just want views to be consistent with it
+
+	// Test an aribtrary string
+	TestSlicing(*this, TEXT("Test string"));
+
+	// Test an empty string
+	TestSlicing(*this, FString());
+
+	// Test an null-terminator-only empty string
+	FString TerminatorOnly;
+	TerminatorOnly.GetCharArray().Add('\0');
+	TestSlicing(*this, TerminatorOnly);
+
+	return true; 
+}
+
 #undef TEST_NAME_ROOT
 
 #endif //WITH_DEV_AUTOMATION_TESTS
diff --git a/Engine/Source/Runtime/Core/Public/Containers/ArrayView.h b/Engine/Source/Runtime/Core/Public/Containers/ArrayView.h
--- a/Engine/Source/Runtime/Core/Public/Containers/ArrayView.h
+++ b/Engine/Source/Runtime/Core/Public/Containers/ArrayView.h
@@ -11,6 +11,7 @@
 #include "Templates/UnrealTypeTraits.h"
 #include "Traits/ElementType.h"
 #include "Containers/Array.h"
+#include "Math/UnrealMathUtility.h"
 
 namespace ArrayViewPrivate
 {
@@ -327,17 +328,97 @@ public:
 
 	/**
 	 * Returns a sliced view
+	   The is similar to Mid(), but with a narrow contract, i.e. slicing outside of the range of the view is illegal.
 	 *
 	 * @param Index starting index of the new view
 	 * @param InNum number of elements in the new view
 	 * @returns Sliced view
+	 *
+	 * @see Mid
 	 */
-	FORCEINLINE TArrayView Slice(SizeType Index, SizeType InNum) const
+	[[nodiscard]] FORCEINLINE TArrayView Slice(SizeType Index, SizeType InNum) const
 	{
 		SliceRangeCheck(Index, InNum);
 		return TArrayView(DataPtr + Index, InNum);
 	}
 
+	/** Returns the left-most part of the view by taking the given number of elements from the left. */
+	[[nodiscard]] inline TArrayView Left(SizeType Count) const
+	{
+		return TArrayView(DataPtr, FMath::Clamp(Count, 0, ArrayNum));
+	}
+
+	/** Returns the left-most part of the view by chopping the given number of elements from the right. */
+	[[nodiscard]] inline TArrayView LeftChop(SizeType Count) const
+	{
+		return TArrayView(DataPtr, FMath::Clamp(ArrayNum - Count, 0, ArrayNum));
+	}
+
+	/** Returns the right-most part of the view by taking the given number of elements from the right. */
+	[[nodiscard]] inline TArrayView Right(SizeType Count) const
+	{
+		const SizeType OutLen = FMath::Clamp(Count, 0, ArrayNum);
+		return TArrayView(DataPtr + ArrayNum - OutLen, OutLen);
+	}
+
+	/** Returns the right-most part of the view by chopping the given number of elements from the left. */
+	[[nodiscard]] inline TArrayView RightChop(SizeType Count) const
+	{
+		const SizeType OutLen = FMath::Clamp(ArrayNum - Count, 0, ArrayNum);
+		return TArrayView(DataPtr + ArrayNum - OutLen, OutLen);
+	}
+
+	/** Returns the middle part of the view by taking up to the given number of elements from the given position. */
+	[[nodiscard]] inline TArrayView Mid(SizeType Index, SizeType Count = TNumericLimits<SizeType>::Max()) const
+	{
+		const ElementType* CurrentStart  = GetData();
+		const SizeType     CurrentLength = Num();
+
+		// Clamp minimum index at the start of the range, adjusting the length down if necessary
+		const SizeType NegativeIndexOffset = (Index < 0) ? Index : 0;
+		Count += NegativeIndexOffset;
+		Index -= NegativeIndexOffset;
+
+		// Clamp maximum index at the end of the range
+		Index = (Index > CurrentLength) ? CurrentLength : Index;
+
+		// Clamp count between 0 and the distance to the end of the range
+		Count = FMath::Clamp(Count, 0, (CurrentLength - Index));
+
+		TArrayView Result = TArrayView(CurrentStart + Index, Count);
+		return Result;
+	}
+
+	/** Modifies the view to be the given number of elements from the left. */
+	inline void LeftInline(SizeType CharCount)
+	{
+		*this = Left(CharCount);
+	}
+
+	/** Modifies the view by chopping the given number of elements from the right. */
+	inline void LeftChopInline(SizeType CharCount)
+	{
+		*this = LeftChop(CharCount);
+	}
+
+	/** Modifies the view to be the given number of elements from the right. */
+	inline void RightInline(SizeType CharCount)
+	{
+		*this = Right(CharCount);
+	}
+
+	/** Modifies the view by chopping the given number of elements from the left. */
+	inline void RightChopInline(SizeType CharCount)
+	{
+		*this = RightChop(CharCount);
+	}
+
+	/** Modifies the view to be the middle part by taking up to the given number of elements from the given position. */
+	inline void MidInline(SizeType Position, SizeType CharCount = TNumericLimits<SizeType>::Max())
+	{
+		*this = Mid(Position, CharCount);
+	}
+
 	/**
 	 * Finds element within the array.
 	 *
diff --git a/Engine/Source/Runtime/Core/Public/Containers/StringView.inl b/Engine/Source/Runtime/Core/Public/Containers/StringView.inl
--- a/Engine/Source/Runtime/Core/Public/Containers/StringView.inl
+++ b/Engine/Source/Runtime/Core/Public/Containers/StringView.inl
@@ -1,6 +1,7 @@
 // Copyright Epic Games, Inc. All Rights Reserved.
 
 #include "HAL/UnrealMemory.h"
+#include "Math/UnrealMathUtility.h"
 #include "Traits/IntType.h"
 
 template <typename CharType>
@@ -47,10 +48,22 @@ inline TStringView<CharType> TStringView<CharType>::RightChop(SizeType CharCount
 template <typename CharType>
 inline TStringView<CharType> TStringView<CharType>::Mid(SizeType Position, SizeType CharCount) const
 {
-	using USizeType = TUnsignedIntType_T<sizeof(SizeType)>;
-	Position = FMath::Clamp<USizeType>(Position, 0, Size);
-	CharCount = FMath::Clamp<USizeType>(CharCount, 0, Size - Position);
-	return ViewType(DataPtr + Position, CharCount);
+	const CharType* CurrentStart  = GetData();
+	const SizeType  CurrentLength = Len();
+
+	// Clamp minimum position at the start of the string, adjusting the length down if necessary
+	const SizeType NegativePositionOffset = (Position < 0) ? Position : 0;
+	CharCount += NegativePositionOffset;
+	Position  -= NegativePositionOffset;
+
+	// Clamp maximum position at the end of the string
+	Position = (Position > CurrentLength) ? CurrentLength : Position;
+
+	// Clamp count between 0 and the distance to the end of the string
+	CharCount = FMath::Clamp(CharCount, 0, (CurrentLength - Position));
+
+	ViewType Result = ViewType(CurrentStart + Position, CharCount);
+	return Result;
 }
 
 template <typename CharType>
