Corrected a spelling mistake.

diff --git a/Engine/Source/Runtime/Core/Private/ProfilingDebugging/TraceAuxiliary.cpp b/Engine/Source/Runtime/Core/Private/ProfilingDebugging/TraceAuxiliary.cpp
--- a/Engine/Source/Runtime/Core/Private/ProfilingDebugging/TraceAuxiliary.cpp
+++ b/Engine/Source/Runtime/Core/Private/ProfilingDebugging/TraceAuxiliary.cpp
@@ -494,13 +494,13 @@ static FAutoConsoleCommand TraceAuxiliaryResumeCmd(
 
 #if WITH_UNREAL_TRACE_LAUNCH
 ////////////////////////////////////////////////////////////////////////////////
-static std::atomic<int32> GUnrealTraceLanuched; // = 0;
+static std::atomic<int32> GUnrealTraceLaunched; // = 0;
 
 ////////////////////////////////////////////////////////////////////////////////
 #if PLATFORM_WINDOWS
 static void LaunchUnrealTraceInternal()
 {
-	if (GUnrealTraceLanuched.load(std::memory_order_relaxed))
+	if (GUnrealTraceLaunched.load(std::memory_order_relaxed))
 	{
 		UE_LOG(LogCore, Log, TEXT("UnrealTraceServer: Trace store already started"));
 		return;
@@ -537,7 +537,7 @@ static void LaunchUnrealTraceInternal()
 		else
 		{
 			UE_LOG(LogCore, Log, TEXT("UnrealTraceServer: Trace store launch successful"));
-			GUnrealTraceLanuched.fetch_add(1, std::memory_order_relaxed);
+			GUnrealTraceLaunched.fetch_add(1, std::memory_order_relaxed);
 		}
 	}
 
