the fix proposed by [REDACTED_PERSON] to match reserve worker's priority to the priority of the blocked worker it's replacing

diff --git a/Engine/Source/Runtime/Core/Private/Async/Fundamental/ReserveScheduler.cpp b/Engine/Source/Runtime/Core/Private/Async/Fundamental/ReserveScheduler.cpp
--- a/Engine/Source/Runtime/Core/Private/Async/Fundamental/ReserveScheduler.cpp
+++ b/Engine/Source/Runtime/Core/Private/Async/Fundamental/ReserveScheduler.cpp
@@ -48,7 +48,8 @@ bool FReserveScheduler::DoReserveWorkUntil(FConditional&& Condition)
 	if (FYieldedWork* WorkerEvent = EventStack.Pop())
 	{
 		WorkerEvent->CompletedDelegate = MoveTemp(Condition);
-		WorkerEvent->bPermitBackgroundWork = FSchedulerTls::PermitBackgroundWork();
+		// become a background worker if the reserve worker is replacing a blocked background worker
+		WorkerEvent->bPermitBackgroundWork = FSchedulerTls::IsBackgroundWorker();
 		WorkerEvent->SleepEvent->Trigger();
 		return true;
 	}
