CompressedBuffer: Removed calls to Ar.TotalSize() to fix loading from a text-format archive

#rnx

diff --git a/Engine/Source/Runtime/Core/Private/Compression/CompressedBuffer.cpp b/Engine/Source/Runtime/Core/Private/Compression/CompressedBuffer.cpp
--- a/Engine/Source/Runtime/Core/Private/Compression/CompressedBuffer.cpp
+++ b/Engine/Source/Runtime/Core/Private/Compression/CompressedBuffer.cpp
@@ -749,26 +749,24 @@ FCompressedBuffer FCompressedBuffer::FromCompressed(FArchive& Ar)
 	using namespace UE::CompressedBuffer;
 	check(Ar.IsLoading());
 
-	if (sizeof(FHeader) > Ar.TotalSize() - Ar.Tell())
-	{
-		return FCompressedBuffer();
-	}
-
 	FHeader Header;
 	Ar.Serialize(&Header, sizeof(FHeader));
 	Header.ByteSwap();
 
-	if (Header.TotalCompressedSize + Ar.Tell() > Ar.TotalSize() + sizeof(FHeader))
+	FCompressedBuffer Local;
+	if (Header.Magic == Header.ExpectedMagic)
 	{
-		return FCompressedBuffer();
+		FUniqueBuffer MutableBuffer = FUniqueBuffer::Alloc(Header.TotalCompressedSize);
+		Header.ByteSwap();
+		const FMutableMemoryView MutableView = MutableBuffer.GetView().CopyFrom(MakeMemoryView(&Header, &Header + 1));
+		Ar.Serialize(MutableView.GetData(), static_cast<int64>(MutableView.GetSize()));
+		Local.CompressedData = UE::CompressedBuffer::ValidBufferOrEmpty(MutableBuffer.MoveToShared());
 	}
-
-	FUniqueBuffer MutableBuffer = FUniqueBuffer::Alloc(Header.TotalCompressedSize);
-	Header.ByteSwap();
-	const FMutableMemoryView MutableView = MutableBuffer.GetView().CopyFrom(MakeMemoryView(&Header, &Header + 1));
-	Ar.Serialize(MutableView.GetData(), static_cast<int64>(MutableView.GetSize()));
-
-	return FromCompressed(MutableBuffer.MoveToShared());
+	if (Local.IsNull())
+	{
+		Ar.SetError();
+	}
+	return Local;
 }
 
 uint64 FCompressedBuffer::GetRawSize() const
