Add setting of affinity and priority for the task graph when starting workers

diff --git a/Engine/Source/Runtime/Core/Private/Async/Fundamental/ReserveScheduler.cpp b/Engine/Source/Runtime/Core/Private/Async/Fundamental/ReserveScheduler.cpp
--- a/Engine/Source/Runtime/Core/Private/Async/Fundamental/ReserveScheduler.cpp
+++ b/Engine/Source/Runtime/Core/Private/Async/Fundamental/ReserveScheduler.cpp
@@ -8,7 +8,7 @@ namespace LowLevelTasks
 
 FReserveScheduler FReserveScheduler::Singleton;
 
-TUniquePtr<FThread> FReserveScheduler::CreateWorker(FSchedulerTls::FLocalQueueType* WorkerLocalQueue, EThreadPriority Priority, bool bIsForkable)
+TUniquePtr<FThread> FReserveScheduler::CreateWorker(bool bIsForkable, FSchedulerTls::FLocalQueueType* WorkerLocalQueue, EThreadPriority Priority)
 {
 	uint32 WorkerId = NextWorkerId++;
 	return MakeUnique<FThread>
@@ -56,7 +56,7 @@ bool FReserveScheduler::DoReserveWorkUntil(FConditional&& Condition)
 	return false;
 }
 
-void FReserveScheduler::StartWorkers(FScheduler& MainScheduler, uint32 NumWorkers, EThreadPriority WorkerPriority, bool bIsForkable)
+void FReserveScheduler::StartWorkers(FScheduler& MainScheduler, uint32 NumWorkers, bool bIsForkable, EThreadPriority WorkerPriority)
 {
 	if (NumWorkers == 0)
 	{
@@ -76,7 +76,7 @@ void FReserveScheduler::StartWorkers(FScheduler& MainScheduler, uint32 NumWorker
 		for (uint32 WorkerId = 0; WorkerId < NumWorkers; ++WorkerId)
 		{
 			WorkerLocalQueues.Emplace(MainScheduler.GetQueueRegistry(), ELocalQueueType::EBusyWait, nullptr);
-			WorkerThreads.Add(CreateWorker(&WorkerLocalQueues.Last(), WorkerPriority, bIsForkable));
+			WorkerThreads.Add(CreateWorker(bIsForkable, &WorkerLocalQueues.Last(), WorkerPriority));
 		}
 		UE::Trace::ThreadGroupEnd();
 	}
diff --git a/Engine/Source/Runtime/Core/Private/Async/Fundamental/Scheduler.cpp b/Engine/Source/Runtime/Core/Private/Async/Fundamental/Scheduler.cpp
--- a/Engine/Source/Runtime/Core/Private/Async/Fundamental/Scheduler.cpp
+++ b/Engine/Source/Runtime/Core/Private/Async/Fundamental/Scheduler.cpp
@@ -40,7 +40,7 @@ namespace LowLevelTasks
 		}
 	}
 
-	TUniquePtr<FThread> FScheduler::CreateWorker(FSleepEvent* ExternalWorkerEvent, FSchedulerTls::FLocalQueueType* ExternalWorkerLocalQueue, EThreadPriority Priority, bool bPermitBackgroundWork, bool bIsForkable)
+	TUniquePtr<FThread> FScheduler::CreateWorker(bool bPermitBackgroundWork, bool bIsForkable, FSleepEvent* ExternalWorkerEvent, FSchedulerTls::FLocalQueueType* ExternalWorkerLocalQueue, EThreadPriority Priority, uint64 InAffinity)
 	{
 		uint32 WorkerId = NextWorkerId++;
 		const uint32 WaitTimes[8] = { 719, 991, 1361, 1237, 1597, 953, 587, 1439 };
@@ -50,6 +50,11 @@ namespace LowLevelTasks
 		{
 			ThreadAffinityMask = FPlatformAffinity::GetTaskGraphBackgroundTaskMask();
 		}
+		if (InAffinity)
+		{
+			// we can override the affinity!
+			ThreadAffinityMask = InAffinity;
+		}
 
 		const FProcessorGroupDesc& ProcessorGroups = FPlatformMisc::GetProcessorGroupDesc();
 		int32 CpuGroupCount = ProcessorGroups.NumProcessorGroups;
@@ -83,7 +88,7 @@ namespace LowLevelTasks
 		);
 	}
 
-	void FScheduler::StartWorkers(uint32 NumForegroundWorkers, uint32 NumBackgroundWorkers, EThreadPriority WorkerPriority,  EThreadPriority BackgroundPriority, bool bIsForkable)
+	void FScheduler::StartWorkers(uint32 NumForegroundWorkers, uint32 NumBackgroundWorkers, bool bIsForkable, EThreadPriority InWorkerPriority,  EThreadPriority InBackgroundPriority, uint64 InWorkerAffinity, uint64 InBackgroundAffinity)
 	{
 		if (NumForegroundWorkers == 0 && NumBackgroundWorkers == 0)
 		{
@@ -91,6 +96,18 @@ namespace LowLevelTasks
 			NumBackgroundWorkers = FMath::Max<int32>(1, FPlatformMisc::NumberOfWorkerThreadsToSpawn() - NumForegroundWorkers);
 		}
 
+		WorkerPriority = InWorkerPriority;
+		BackgroundPriority = InBackgroundPriority;
+
+		if (InWorkerAffinity)
+		{
+			WorkerAffinity = InWorkerAffinity;
+		}
+		if (InBackgroundAffinity)
+		{
+			BackgroundAffinity = InBackgroundAffinity;
+		}
+
 		const bool bSupportsMultithreading = FPlatformProcess::SupportsMultithreading() || FForkProcessHelper::IsForkedMultithreadInstance();
 
 		uint32 OldActiveWorkers = ActiveWorkers.load(std::memory_order_relaxed);
@@ -110,7 +127,7 @@ namespace LowLevelTasks
 			{
 				WorkerEvents.Emplace();
 				WorkerLocalQueues.Emplace(QueueRegistry, ELocalQueueType::EForeground, &WorkerEvents.Last());
-				WorkerThreads.Add(CreateWorker(&WorkerEvents.Last(), &WorkerLocalQueues.Last(), WorkerPriority, false, bIsForkable));
+				WorkerThreads.Add(CreateWorker(false, bIsForkable, &WorkerEvents.Last(), &WorkerLocalQueues.Last(), WorkerPriority, WorkerAffinity));
 			}
 			UE::Trace::ThreadGroupEnd();
 			UE::Trace::ThreadGroupBegin(TEXT("Background Workers"));
@@ -118,13 +135,13 @@ namespace LowLevelTasks
 			{
 				WorkerEvents.Emplace();
 				WorkerLocalQueues.Emplace(QueueRegistry, ELocalQueueType::EBackground, &WorkerEvents.Last());
-				WorkerThreads.Add(CreateWorker(&WorkerEvents.Last(), &WorkerLocalQueues.Last(), BackgroundPriority, true, bIsForkable));
+				WorkerThreads.Add(CreateWorker(true, bIsForkable, &WorkerEvents.Last(), &WorkerLocalQueues.Last(), BackgroundPriority, BackgroundAffinity));
 			}
 			UE::Trace::ThreadGroupEnd();
 		}
 	}
 
-	void FScheduler::StopWorkers()
+	void FScheduler::StopWorkers(bool DrainGlobalQueue)
 	{
 		uint32 OldActiveWorkers = ActiveWorkers.load(std::memory_order_relaxed);
 		if(OldActiveWorkers != 0 && ActiveWorkers.compare_exchange_strong(OldActiveWorkers, 0, std::memory_order_relaxed))
@@ -141,16 +158,28 @@ namespace LowLevelTasks
 			WorkerThreads.Reset();
 			WorkerLocalQueues.Reset();
 			WorkerEvents.Reset();
-			for (FTask* Task = QueueRegistry.Dequeue(); Task != nullptr; Task = QueueRegistry.Dequeue())
+
+			if (DrainGlobalQueue)
 			{
-				Task->ExecuteTask();
+				for (FTask* Task = QueueRegistry.Dequeue(); Task != nullptr; Task = QueueRegistry.Dequeue())
+				{
+					Task->ExecuteTask();
+				}
 			}
 		}
 	}
+	void FScheduler::RestartWorkers(uint32 NumForegroundWorkers, uint32 NumBackgroundWorkers, bool bIsForkable, EThreadPriority InWorkerPriority, EThreadPriority InBackgroundPriority, uint64 InWorkerAffinity, uint64 InBackgroundAffinity)
+	{
+		FScopeLock Lock(&WorkerThreadsCS);
+		TemporaryShutdown.store(true, std::memory_order_release);
+		StopWorkers(false);
+		StartWorkers(NumForegroundWorkers, NumBackgroundWorkers, bIsForkable, InWorkerPriority, InBackgroundPriority, InWorkerAffinity, InBackgroundAffinity);
+		TemporaryShutdown.store(false, std::memory_order_release);
+	}
 
 	void FScheduler::LaunchInternal(FTask& Task, EQueuePreference QueuePreference, bool bWakeUpWorker)
 	{
-		if (ActiveWorkers.load(std::memory_order_relaxed))
+		if (ActiveWorkers.load(std::memory_order_relaxed) || TemporaryShutdown.load(std::memory_order_acquire))
 		{			
 			const bool bIsBackgroundTask = Task.IsBackgroundTask();
 			const bool bIsBackgroundWorker = FSchedulerTls::IsBackgroundWorker();
diff --git a/Engine/Source/Runtime/Core/Private/Async/TaskGraph.cpp b/Engine/Source/Runtime/Core/Private/Async/TaskGraph.cpp
--- a/Engine/Source/Runtime/Core/Private/Async/TaskGraph.cpp
+++ b/Engine/Source/Runtime/Core/Private/Async/TaskGraph.cpp
@@ -101,7 +101,13 @@ static FAutoConsoleVariableRef CVarForkedProcessMaxWorkerThreads(
 );
 
 CORE_API int32 GUseNewTaskBackend = 1;
-uint32 GNumForegroundWorkers = 2;
+int32 GNumForegroundWorkers = 2;
+static FAutoConsoleVariableRef CVarNumForegroundWorkers(
+	TEXT("TaskGraph.NumForegroundWorkers"),
+	GNumForegroundWorkers,
+	TEXT("Configures the number of foreground worker threads. Requires the scheduler to be restarted to have an affect")
+);
+
 bool GDisableReserveWorkers = 0;
 
 #if CREATE_HIPRI_TASK_THREADS || CREATE_BACKGROUND_TASK_THREADS
@@ -1867,10 +1873,10 @@ public:
 			int32 NumBackgroundWorkers = FMath::Max(1, NumWorkerThreads - FMath::Min<int>(GNumForegroundWorkers, NumWorkerThreads));
 			int32 NumForegroundWorkers =  FMath::Max(1, NumWorkerThreads - NumBackgroundWorkers);
 
-			LowLevelTasks::FScheduler::Get().StartWorkers(NumForegroundWorkers, NumBackgroundWorkers, EThreadPriority::TPri_Normal, EThreadPriority::TPri_BelowNormal, FForkProcessHelper::IsForkedMultithreadInstance());
+			LowLevelTasks::FScheduler::Get().StartWorkers(NumForegroundWorkers, NumBackgroundWorkers, FForkProcessHelper::IsForkedMultithreadInstance(), EThreadPriority::TPri_Normal, EThreadPriority::TPri_BelowNormal);
 			if (!GDisableReserveWorkers)
 			{
-				LowLevelTasks::FReserveScheduler::Get().StartWorkers(LowLevelTasks::FScheduler::Get(), NumForegroundWorkers + NumBackgroundWorkers, EThreadPriority::TPri_BelowNormal, FForkProcessHelper::IsForkedMultithreadInstance());
+				LowLevelTasks::FReserveScheduler::Get().StartWorkers(LowLevelTasks::FScheduler::Get(), NumForegroundWorkers + NumBackgroundWorkers, FForkProcessHelper::IsForkedMultithreadInstance(), EThreadPriority::TPri_BelowNormal);
 			}
 
 			NumNamedThreads = ENamedThreads::ActualRenderingThread + 1;
@@ -1950,12 +1956,12 @@ public:
 			int32 NumWorkers =  FMath::Max(1, NumWorkerThreads - NumBackgroundWorkers);
 
 			LowLevelTasks::FScheduler::Get().StopWorkers();
-			LowLevelTasks::FScheduler::Get().StartWorkers(NumWorkers, NumBackgroundWorkers, Pri, EThreadPriority::TPri_BelowNormal, FForkProcessHelper::IsForkedMultithreadInstance());
+			LowLevelTasks::FScheduler::Get().StartWorkers(NumWorkers, NumBackgroundWorkers, FForkProcessHelper::IsForkedMultithreadInstance(), Pri, EThreadPriority::TPri_BelowNormal);
 
 			if (!GDisableReserveWorkers)
 			{
 				LowLevelTasks::FReserveScheduler::Get().StopWorkers();
-				LowLevelTasks::FReserveScheduler::Get().StartWorkers(LowLevelTasks::FScheduler::Get(), NumWorkers + NumBackgroundWorkers, EThreadPriority::TPri_BelowNormal, FForkProcessHelper::IsForkedMultithreadInstance());
+				LowLevelTasks::FReserveScheduler::Get().StartWorkers(LowLevelTasks::FScheduler::Get(), NumWorkers + NumBackgroundWorkers, FForkProcessHelper::IsForkedMultithreadInstance(), EThreadPriority::TPri_BelowNormal);
 			}
 		}
 	}
diff --git a/Engine/Source/Runtime/Core/Private/Tests/Async/TaskGraphTest.cpp b/Engine/Source/Runtime/Core/Private/Tests/Async/TaskGraphTest.cpp
--- a/Engine/Source/Runtime/Core/Private/Tests/Async/TaskGraphTest.cpp
+++ b/Engine/Source/Runtime/Core/Private/Tests/Async/TaskGraphTest.cpp
@@ -715,7 +715,7 @@ namespace OldTaskGraphTests
 	}
 }
 
-extern uint32 GNumForegroundWorkers;
+extern int32 GNumForegroundWorkers;
 
 namespace TaskGraphTests
 {
diff --git a/Engine/Source/Runtime/Core/Public/Async/Fundamental/ReserveScheduler.h b/Engine/Source/Runtime/Core/Public/Async/Fundamental/ReserveScheduler.h
--- a/Engine/Source/Runtime/Core/Public/Async/Fundamental/ReserveScheduler.h
+++ b/Engine/Source/Runtime/Core/Public/Async/Fundamental/ReserveScheduler.h
@@ -30,14 +30,14 @@ namespace LowLevelTasks
 		FORCEINLINE_DEBUGGABLE static FReserveScheduler& Get();
 
 		//start number of reserve workers where 0 is the system default
-		CORE_API void StartWorkers(FScheduler& MainScheduler, uint32 ReserveWorkers = 0, EThreadPriority WorkerPriority = EThreadPriority::TPri_Normal, bool bIsForkable = false);
+		CORE_API void StartWorkers(FScheduler& MainScheduler, uint32 ReserveWorkers = 0, bool bIsForkable = false, EThreadPriority WorkerPriority = EThreadPriority::TPri_Normal);
 		CORE_API void StopWorkers();
 
 		//tries to yield this thread using the YieldEvent and do busywork on a reserve worker.
 		CORE_API bool DoReserveWorkUntil(FConditional&& Condition);
 
 	private: 
-		TUniquePtr<FThread> CreateWorker(FSchedulerTls::FLocalQueueType* WorkerLocalQueue, EThreadPriority Priority = EThreadPriority::TPri_Normal, bool bIsForkable = false);
+		TUniquePtr<FThread> CreateWorker(bool bIsForkable = false, FSchedulerTls::FLocalQueueType* WorkerLocalQueue = nullptr, EThreadPriority Priority = EThreadPriority::TPri_Normal);
 
 	private:
 		TEventStack<FYieldedWork> 				EventStack;
diff --git a/Engine/Source/Runtime/Core/Public/Async/Fundamental/Scheduler.h b/Engine/Source/Runtime/Core/Public/Async/Fundamental/Scheduler.h
--- a/Engine/Source/Runtime/Core/Public/Async/Fundamental/Scheduler.h
+++ b/Engine/Source/Runtime/Core/Public/Async/Fundamental/Scheduler.h
@@ -156,10 +156,11 @@ namespace LowLevelTasks
 		FORCEINLINE_DEBUGGABLE static FScheduler& Get();
 
 		//start number of workers where 0 is the system default
-		CORE_API void StartWorkers(uint32 NumForegroundWorkers = 0, uint32 NumBackgroundWorkers = 0, EThreadPriority WorkerPriority = EThreadPriority::TPri_Normal,  EThreadPriority BackgroundPriority = EThreadPriority::TPri_BelowNormal, bool bIsForkable = false);
-		CORE_API void StopWorkers();
+		CORE_API void StartWorkers(uint32 NumForegroundWorkers = 0, uint32 NumBackgroundWorkers = 0, bool bIsForkable = false, EThreadPriority InWorkerPriority = EThreadPriority::TPri_Normal, EThreadPriority InBackgroundPriority = EThreadPriority::TPri_BelowNormal, uint64 InWorkerAffinity = 0, uint64 InBackgroundAffinity = 0);
+		CORE_API void StopWorkers(bool DrainGlobalQueue = true);
+		CORE_API void RestartWorkers(uint32 NumForegroundWorkers = 0, uint32 NumBackgroundWorkers = 0, bool bIsForkable = false, EThreadPriority WorkerPriority = EThreadPriority::TPri_Normal, EThreadPriority BackgroundPriority = EThreadPriority::TPri_BelowNormal, uint64 InWorkerAffinity = 0, uint64 InBackgroundAffinity = 0);
 
-		//try to launch the task, the return value will specify if the task was in the ready state and has been launced
+		//try to launch the task, the return value will specify if the task was in the ready state and has been launhced
 		inline bool TryLaunch(FTask& Task, EQueuePreference QueuePreference = EQueuePreference::DefaultPreference, bool bWakeUpWorker = true);	
 
 		//try to launch the task on a specific worker ID, the return value will specify if the task was in the ready state and has been launched
@@ -189,7 +190,7 @@ namespace LowLevelTasks
 		~FScheduler();
 
 	private: 
-		TUniquePtr<FThread> CreateWorker(FSleepEvent* ExternalWorkerEvent = nullptr, FSchedulerTls::FLocalQueueType* ExternalWorkerLocalQueue = nullptr, EThreadPriority Priority = EThreadPriority::TPri_Normal, bool bPermitBackgroundWork = false, bool bIsForkable = false);
+		TUniquePtr<FThread> CreateWorker(bool bPermitBackgroundWork = false, bool bIsForkable = false, FSleepEvent* ExternalWorkerEvent = nullptr, FSchedulerTls::FLocalQueueType* ExternalWorkerLocalQueue = nullptr, EThreadPriority Priority = EThreadPriority::TPri_Normal, uint64 InAffinity = 0);
 		void WorkerMain(struct FSleepEvent* WorkerEvent, FSchedulerTls::FLocalQueueType* ExternalWorkerLocalQueue, uint32 WaitCycles, bool bPermitBackgroundWork);
 		CORE_API void LaunchInternal(FTask& Task, EQueuePreference QueuePreference, bool bWakeUpWorker);
 		CORE_API void BusyWaitInternal(const FConditional& Conditional, bool ForceAllowBackgroundWork);
@@ -208,6 +209,11 @@ namespace LowLevelTasks
 		TArray<FSleepEvent>						WorkerEvents;
 		std::atomic_uint						ActiveWorkers { 0 };
 		std::atomic_uint						NextWorkerId { 0 };
+		uint64									WorkerAffinity = 0;
+		uint64									BackgroundAffinity = 0;
+		EThreadPriority							WorkerPriority = EThreadPriority::TPri_Normal;
+		EThreadPriority							BackgroundPriority = EThreadPriority::TPri_BelowNormal;
+		std::atomic_bool						TemporaryShutdown{ false };
 	};
 
 	FORCEINLINE_DEBUGGABLE bool TryLaunch(FTask& Task, EQueuePreference QueuePreference = EQueuePreference::DefaultPreference, bool bWakeUpWorker = true)
