OodleDataCompression fix 64 bit size cast to int

diff --git a/Engine/Source/Runtime/Core/Private/Compression/OodleDataCompression.cpp b/Engine/Source/Runtime/Core/Private/Compression/OodleDataCompression.cpp
--- a/Engine/Source/Runtime/Core/Private/Compression/OodleDataCompression.cpp
+++ b/Engine/Source/Runtime/Core/Private/Compression/OodleDataCompression.cpp
@@ -128,7 +128,7 @@ struct OodleDataCompressionDecoders
 
 					OodleDecoders[i].OodleDecoderMutex.Unlock();
 
-					return (int) Result;
+					return (int64) Result;
 				}
 
 				OodleDecoders[i].OodleDecoderMutex.Unlock();
