Added float/double variants for all FMatrix specializations

diff --git a/Engine/Source/Runtime/Core/Private/Math/UnrealMath.cpp b/Engine/Source/Runtime/Core/Private/Math/UnrealMath.cpp
--- a/Engine/Source/Runtime/Core/Private/Math/UnrealMath.cpp
+++ b/Engine/Source/Runtime/Core/Private/Math/UnrealMath.cpp
@@ -657,163 +657,6 @@ float FQuat::GetTwistAngle(const FVector& TwistAxis) const
 	return FMath::UnwindRadians(2.0f * (float)FMath::Atan2(XYZ, W));
 }
 
-FMatrix FRotationAboutPointMatrix::Make(const FQuat& Rot, const FVector& Origin)
-{
-	return FRotationAboutPointMatrix(Rot.Rotator(), Origin);
-}
-
-FMatrix FRotationMatrix::Make(FQuat const& Rot)
-{
-	return FQuatRotationTranslationMatrix(Rot, FVector::ZeroVector);
-}
-
-FMatrix FRotationMatrix::MakeFromX(FVector const& XAxis)
-{
-	FVector const NewX = XAxis.GetSafeNormal();
-
-	// try to use up if possible
-	FVector const UpVector = ( FMath::Abs(NewX.Z) < (1.f - KINDA_SMALL_NUMBER) ) ? FVector(0,0,1.f) : FVector(1.f,0,0);
-
-	const FVector NewY = (UpVector ^ NewX).GetSafeNormal();
-	const FVector NewZ = NewX ^ NewY;
-
-	return FMatrix(NewX, NewY, NewZ, FVector::ZeroVector);
-}
-
-FMatrix FRotationMatrix::MakeFromY(FVector const& YAxis)
-{
-	FVector const NewY = YAxis.GetSafeNormal();
-
-	// try to use up if possible
-	FVector const UpVector = ( FMath::Abs(NewY.Z) < (1.f - KINDA_SMALL_NUMBER) ) ? FVector(0,0,1.f) : FVector(1.f,0,0);
-
-	const FVector NewZ = (UpVector ^ NewY).GetSafeNormal();
-	const FVector NewX = NewY ^ NewZ;
-
-	return FMatrix(NewX, NewY, NewZ, FVector::ZeroVector);
-}
-
-FMatrix FRotationMatrix::MakeFromZ(FVector const& ZAxis)
-{
-	FVector const NewZ = ZAxis.GetSafeNormal();
-
-	// try to use up if possible
-	FVector const UpVector = ( FMath::Abs(NewZ.Z) < (1.f - KINDA_SMALL_NUMBER) ) ? FVector(0,0,1.f) : FVector(1.f,0,0);
-
-	const FVector NewX = (UpVector ^ NewZ).GetSafeNormal();
-	const FVector NewY = NewZ ^ NewX;
-
-	return FMatrix(NewX, NewY, NewZ, FVector::ZeroVector);
-}
-
-FMatrix FRotationMatrix::MakeFromXY(FVector const& XAxis, FVector const& YAxis)
-{
-	FVector NewX = XAxis.GetSafeNormal();
-	FVector Norm = YAxis.GetSafeNormal();
-
-	// if they're almost same, we need to find arbitrary vector
-	if ( FMath::IsNearlyEqual(FMath::Abs(NewX | Norm), FVector::FReal(1.f)) )
-	{
-		// make sure we don't ever pick the same as NewX
-		Norm = ( FMath::Abs(NewX.Z) < (1.f - KINDA_SMALL_NUMBER) ) ? FVector(0,0,1.f) : FVector(1.f,0,0);
-	}
-
-	const FVector NewZ = (NewX ^ Norm).GetSafeNormal();
-	const FVector NewY = NewZ ^ NewX;
-
-	return FMatrix(NewX, NewY, NewZ, FVector::ZeroVector);
-}
-
-FMatrix FRotationMatrix::MakeFromXZ(FVector const& XAxis, FVector const& ZAxis)
-{
-	FVector const NewX = XAxis.GetSafeNormal();
-	FVector Norm = ZAxis.GetSafeNormal();
-
-	// if they're almost same, we need to find arbitrary vector
-	if ( FMath::IsNearlyEqual(FMath::Abs(NewX | Norm), FVector::FReal(1.f)) )
-	{
-		// make sure we don't ever pick the same as NewX
-		Norm = ( FMath::Abs(NewX.Z) < (1.f - KINDA_SMALL_NUMBER) ) ? FVector(0,0,1.f) : FVector(1.f,0,0);
-	}
-
-	const FVector NewY = (Norm ^ NewX).GetSafeNormal();
-	const FVector NewZ = NewX ^ NewY;
-
-	return FMatrix(NewX, NewY, NewZ, FVector::ZeroVector);
-}
-
-FMatrix FRotationMatrix::MakeFromYX(FVector const& YAxis, FVector const& XAxis)
-{
-	FVector const NewY = YAxis.GetSafeNormal();
-	FVector Norm = XAxis.GetSafeNormal();
-
-	// if they're almost same, we need to find arbitrary vector
-	if ( FMath::IsNearlyEqual(FMath::Abs(NewY | Norm), FVector::FReal(1.f)) )
-	{
-		// make sure we don't ever pick the same as NewX
-		Norm = ( FMath::Abs(NewY.Z) < (1.f - KINDA_SMALL_NUMBER) ) ? FVector(0,0,1.f) : FVector(1.f,0,0);
-	}
-	
-	const FVector NewZ = (Norm ^ NewY).GetSafeNormal();
-	const FVector NewX = NewY ^ NewZ;
-
-	return FMatrix(NewX, NewY, NewZ, FVector::ZeroVector);
-}
-
-FMatrix FRotationMatrix::MakeFromYZ(FVector const& YAxis, FVector const& ZAxis)
-{
-	FVector const NewY = YAxis.GetSafeNormal();
-	FVector Norm = ZAxis.GetSafeNormal();
-
-	// if they're almost same, we need to find arbitrary vector
-	if ( FMath::IsNearlyEqual(FMath::Abs(NewY | Norm), FVector::FReal(1.f)) )
-	{
-		// make sure we don't ever pick the same as NewX
-		Norm = ( FMath::Abs(NewY.Z) < (1.f - KINDA_SMALL_NUMBER) ) ? FVector(0,0,1.f) : FVector(1.f,0,0);
-	}
-
-	const FVector NewX = (NewY ^ Norm).GetSafeNormal();
-	const FVector NewZ = NewX ^ NewY;
-
-	return FMatrix(NewX, NewY, NewZ, FVector::ZeroVector);
-}
-
-FMatrix FRotationMatrix::MakeFromZX(FVector const& ZAxis, FVector const& XAxis)
-{
-	FVector const NewZ = ZAxis.GetSafeNormal();
-	FVector Norm = XAxis.GetSafeNormal();
-
-	// if they're almost same, we need to find arbitrary vector
-	if ( FMath::IsNearlyEqual(FMath::Abs(NewZ | Norm), FVector::FReal(1.f)) )
-	{
-		// make sure we don't ever pick the same as NewX
-		Norm = ( FMath::Abs(NewZ.Z) < (1.f - KINDA_SMALL_NUMBER) ) ? FVector(0,0,1.f) : FVector(1.f,0,0);
-	}
-
-	const FVector NewY = (NewZ ^ Norm).GetSafeNormal();
-	const FVector NewX = NewY ^ NewZ;
-
-	return FMatrix(NewX, NewY, NewZ, FVector::ZeroVector);
-}
-
-FMatrix FRotationMatrix::MakeFromZY(FVector const& ZAxis, FVector const& YAxis)
-{
-	FVector const NewZ = ZAxis.GetSafeNormal();
-	FVector Norm = YAxis.GetSafeNormal();
-
-	// if they're almost same, we need to find arbitrary vector
-	if ( FMath::IsNearlyEqual(FMath::Abs(NewZ | Norm), FVector::FReal(1.f)) )
-	{
-		// make sure we don't ever pick the same as NewX
-		Norm = ( FMath::Abs(NewZ.Z) < (1.f - KINDA_SMALL_NUMBER) ) ? FVector(0,0,1.f) : FVector(1.f,0,0);
-	}
-
-	const FVector NewX = (Norm ^ NewZ).GetSafeNormal();
-	const FVector NewY = NewZ ^ NewX;
-
-	return FMatrix(NewX, NewY, NewZ, FVector::ZeroVector);
-}
-
 FVector FQuat::Euler() const
 {
 	return Rotator().Euler();
diff --git a/Engine/Source/Runtime/Core/Public/Math/ClipProjectionMatrix.h b/Engine/Source/Runtime/Core/Public/Math/ClipProjectionMatrix.h
--- a/Engine/Source/Runtime/Core/Public/Math/ClipProjectionMatrix.h
+++ b/Engine/Source/Runtime/Core/Public/Math/ClipProjectionMatrix.h
@@ -6,37 +6,43 @@
 #include "Math/Plane.h"
 #include "Math/Matrix.h"
 
+namespace UE {
+namespace Math {
+
 /**
  * Realigns the near plane for an existing projection matrix 
  * with an arbitrary clip plane
  * from: http://sourceforge.net/mailarchive/message.php?msg_id=000901c26324%242181ea90%24a1e93942%40firefly
  * Updated for the fact that our FPlane uses Ax+By+Cz=D.
  */
-class FClipProjectionMatrix : public FMatrix
+template<typename T>
+struct TClipProjectionMatrix : public TMatrix<T>
 {
 public:
+	using TMatrix<T>::M;
 	/**
 	 * Constructor
 	 *
 	 * @param	SrcProjMat - source projection matrix to premultiply with the clip matrix
 	 * @param	Plane - clipping plane used to build the clip matrix (assumed to be in camera space)
 	 */
-	FClipProjectionMatrix( const FMatrix& SrcProjMat, const FPlane& Plane );
+	TClipProjectionMatrix( const TMatrix<T>& SrcProjMat, const TPlane<T>& Plane );
 
 private:
 	/** return sign of a number */
-	FORCEINLINE FMatrix::FReal sgn(FMatrix::FReal a );
+	FORCEINLINE T sgn(T a );
 };
 
 
-FORCEINLINE FClipProjectionMatrix::FClipProjectionMatrix( const FMatrix& SrcProjMat, const FPlane& Plane ) :
-FMatrix(SrcProjMat)
+template<typename T>
+FORCEINLINE TClipProjectionMatrix<T>::TClipProjectionMatrix( const TMatrix<T>& SrcProjMat, const TPlane<T>& Plane ) :
+TMatrix<T>(SrcProjMat)
 {
 	// Calculate the clip-space corner point opposite the clipping plane
 	// as (sgn(clipPlane.x), sgn(clipPlane.y), 1, 1) and
 	// transform it into camera space by multiplying it
 	// by the inverse of the projection matrix
-	FPlane CornerPlane( 
+	TPlane<T> CornerPlane( 
 		sgn(Plane.X) / SrcProjMat.M[0][0],
 		sgn(Plane.Y) / SrcProjMat.M[1][1],
 		1.0f,
@@ -44,7 +50,7 @@ FMatrix(SrcProjMat)
 	);
 
 	// Calculate the scaled plane vector
-	FPlane ProjPlane( Plane * (1.0f / (Plane | CornerPlane)) );
+	TPlane<T> ProjPlane( Plane * (1.0f / (Plane | CornerPlane)) );
 
 	// use the projected space clip plane in z column 
 	// Note: (account for our negated W coefficient)
@@ -54,10 +60,15 @@ FMatrix(SrcProjMat)
 	M[3][2] = -ProjPlane.W;
 }
 
-
-FORCEINLINE FMatrix::FReal FClipProjectionMatrix::sgn(FMatrix::FReal a )
+template<typename T>
+FORCEINLINE T TClipProjectionMatrix<T>::sgn(T a)
 {
 	if (a > 0.0f) return (1.0f);
 	if (a < 0.0f) return (-1.0f);
 	return (0.0f);
 }
+
+} // namespace Math
+} // namespace UE
+
+DECLARE_LWC_TYPE(ClipProjectionMatrix, 44);
\ No newline at end of file
diff --git a/Engine/Source/Runtime/Core/Public/Math/InverseRotationMatrix.h b/Engine/Source/Runtime/Core/Public/Math/InverseRotationMatrix.h
--- a/Engine/Source/Runtime/Core/Public/Math/InverseRotationMatrix.h
+++ b/Engine/Source/Runtime/Core/Public/Math/InverseRotationMatrix.h
@@ -7,9 +7,15 @@
 #include "Math/Plane.h"
 #include "Math/Matrix.h"
 
+// LWC_TODO: FRotator -> TRotator<T>
+
+namespace UE {
+namespace Math {
+ 	
 /** Inverse Rotation matrix */
-class FInverseRotationMatrix
-	: public FMatrix
+template<typename T>
+struct TInverseRotationMatrix
+	: public TMatrix<T>
 {
 public:
 	/**
@@ -17,26 +23,31 @@ public:
 	 *
 	 * @param Rot rotation
 	 */
-	FInverseRotationMatrix(const FRotator& Rot);
+	TInverseRotationMatrix(const FRotator& Rot);
 };
 
-
-FORCEINLINE FInverseRotationMatrix::FInverseRotationMatrix(const FRotator& Rot)
-	: FMatrix(
-		FMatrix( // Yaw
-		FPlane(+FMath::Cos(Rot.Yaw * PI / 180.f), -FMath::Sin(Rot.Yaw * PI / 180.f), 0.0f, 0.0f),
-		FPlane(+FMath::Sin(Rot.Yaw * PI / 180.f), +FMath::Cos(Rot.Yaw * PI / 180.f), 0.0f, 0.0f),
-		FPlane(0.0f, 0.0f, 1.0f, 0.0f),
-		FPlane(0.0f, 0.0f, 0.0f, 1.0f)) *
-		FMatrix( // Pitch
-		FPlane(+FMath::Cos(Rot.Pitch * PI / 180.f), 0.0f, -FMath::Sin(Rot.Pitch * PI / 180.f), 0.0f),
-		FPlane(0.0f, 1.0f, 0.0f, 0.0f),
-		FPlane(+FMath::Sin(Rot.Pitch * PI / 180.f), 0.0f, +FMath::Cos(Rot.Pitch * PI / 180.f), 0.0f),
-		FPlane(0.0f, 0.0f, 0.0f, 1.0f)) *
-		FMatrix( // Roll
-		FPlane(1.0f, 0.0f, 0.0f, 0.0f),
-		FPlane(0.0f, +FMath::Cos(Rot.Roll * PI / 180.f), +FMath::Sin(Rot.Roll * PI / 180.f), 0.0f),
-		FPlane(0.0f, -FMath::Sin(Rot.Roll * PI / 180.f), +FMath::Cos(Rot.Roll * PI / 180.f), 0.0f),
-		FPlane(0.0f, 0.0f, 0.0f, 1.0f))
+template<typename T>
+FORCEINLINE TInverseRotationMatrix<T>::TInverseRotationMatrix(const FRotator& Rot)
+	: TMatrix<T>(
+		TMatrix<T>( // Yaw
+			TPlane<T>(+FMath::Cos(Rot.Yaw * PI / 180.f), -FMath::Sin(Rot.Yaw * PI / 180.f), 0.0f, 0.0f),
+			TPlane<T>(+FMath::Sin(Rot.Yaw * PI / 180.f), +FMath::Cos(Rot.Yaw * PI / 180.f), 0.0f, 0.0f),
+			TPlane<T>(0.0f, 0.0f, 1.0f, 0.0f),
+			TPlane<T>(0.0f, 0.0f, 0.0f, 1.0f)) *
+		TMatrix<T>( // Pitch
+			TPlane<T>(+FMath::Cos(Rot.Pitch * PI / 180.f), 0.0f, -FMath::Sin(Rot.Pitch * PI / 180.f), 0.0f),
+			TPlane<T>(0.0f, 1.0f, 0.0f, 0.0f),
+			TPlane<T>(+FMath::Sin(Rot.Pitch * PI / 180.f), 0.0f, +FMath::Cos(Rot.Pitch * PI / 180.f), 0.0f),
+			TPlane<T>(0.0f, 0.0f, 0.0f, 1.0f)) *
+		TMatrix<T>( // Roll
+			TPlane<T>(1.0f, 0.0f, 0.0f, 0.0f),
+			TPlane<T>(0.0f, +FMath::Cos(Rot.Roll * PI / 180.f), +FMath::Sin(Rot.Roll * PI / 180.f), 0.0f),
+			TPlane<T>(0.0f, -FMath::Sin(Rot.Roll * PI / 180.f), +FMath::Cos(Rot.Roll * PI / 180.f), 0.0f),
+			TPlane<T>(0.0f, 0.0f, 0.0f, 1.0f))
 	)
 { }
+
+} // namespace Math
+} // namespace UE
+
+DECLARE_LWC_TYPE(InverseRotationMatrix, 44);
\ No newline at end of file
diff --git a/Engine/Source/Runtime/Core/Public/Math/Matrix.h b/Engine/Source/Runtime/Core/Public/Math/Matrix.h
--- a/Engine/Source/Runtime/Core/Public/Math/Matrix.h
+++ b/Engine/Source/Runtime/Core/Public/Math/Matrix.h
@@ -11,6 +11,7 @@
 #include "Math/Plane.h"
 #include "Math/Rotator.h"
 #include "Math/Axis.h"
+#include <type_traits>
 
 #ifdef _MSC_VER
 #pragma warning (push)
@@ -33,6 +34,8 @@ namespace Math
 template<typename T>
 struct TMatrix
 {
+	static_assert(std::is_floating_point_v<T>, "T must be floating point");
+
 public:
 	using FReal = T;
 
diff --git a/Engine/Source/Runtime/Core/Public/Math/MirrorMatrix.h b/Engine/Source/Runtime/Core/Public/Math/MirrorMatrix.h
--- a/Engine/Source/Runtime/Core/Public/Math/MirrorMatrix.h
+++ b/Engine/Source/Runtime/Core/Public/Math/MirrorMatrix.h
@@ -6,11 +6,15 @@
 #include "Math/Plane.h"
 #include "Math/Matrix.h"
 
-/**
- * Mirrors a point about an abitrary plane 
+namespace UE {
+namespace Math {
+
+ /**
+ * Mirrors a point about an arbitrary plane 
  */
-class FMirrorMatrix
-	: public FMatrix
+template<typename T>
+struct TMirrorMatrix
+	: public TMatrix<T>
 {
 public:
 
@@ -19,16 +23,21 @@ public:
 	 * 
 	 * @param Plane source plane for mirroring (assumed normalized)
 	 */
-	FMirrorMatrix( const FPlane& Plane );
+	TMirrorMatrix( const TPlane<T>& Plane );
 };
 
-
-FORCEINLINE FMirrorMatrix::FMirrorMatrix( const FPlane& Plane ) :
+template<typename T>
+FORCEINLINE TMirrorMatrix<T>::TMirrorMatrix( const TPlane<T>& Plane ) :
 FMatrix(
-	FPlane( -2.f*Plane.X*Plane.X + 1.f,	-2.f*Plane.Y*Plane.X,		-2.f*Plane.Z*Plane.X,		0.f ),
-	FPlane( -2.f*Plane.X*Plane.Y,			-2.f*Plane.Y*Plane.Y + 1.f,	-2.f*Plane.Z*Plane.Y,		0.f ),
-	FPlane( -2.f*Plane.X*Plane.Z,			-2.f*Plane.Y*Plane.Z,		-2.f*Plane.Z*Plane.Z + 1.f,	0.f ),
-	FPlane(  2.f*Plane.X*Plane.W,			 2.f*Plane.Y*Plane.W,		 2.f*Plane.Z*Plane.W,		1.f ) )
+	TPlane<T>( -2.f*Plane.X*Plane.X + 1.f,		-2.f*Plane.Y*Plane.X,		-2.f*Plane.Z*Plane.X,		0.f ),
+	TPlane<T>( -2.f*Plane.X*Plane.Y,			-2.f*Plane.Y*Plane.Y + 1.f,	-2.f*Plane.Z*Plane.Y,		0.f ),
+	TPlane<T>( -2.f*Plane.X*Plane.Z,			-2.f*Plane.Y*Plane.Z,		-2.f*Plane.Z*Plane.Z + 1.f,	0.f ),
+	TPlane<T>(  2.f*Plane.X*Plane.W,			 2.f*Plane.Y*Plane.W,		 2.f*Plane.Z*Plane.W,		1.f ) )
 {
 	//check( FMath::Abs(1.f - Plane.SizeSquared()) < KINDA_SMALL_NUMBER && TEXT("not normalized"));
 }
+
+} // namespace Math
+} // namespace UE
+
+DECLARE_LWC_TYPE(MirrorMatrix, 44);
\ No newline at end of file
diff --git a/Engine/Source/Runtime/Core/Public/Math/OrthoMatrix.h b/Engine/Source/Runtime/Core/Public/Math/OrthoMatrix.h
--- a/Engine/Source/Runtime/Core/Public/Math/OrthoMatrix.h
+++ b/Engine/Source/Runtime/Core/Public/Math/OrthoMatrix.h
@@ -6,8 +6,13 @@
 #include "Math/Plane.h"
 #include "Math/Matrix.h"
 
-class FOrthoMatrix
-	: public FMatrix
+namespace UE {
+namespace Math {
+
+
+template<typename T>
+struct TOrthoMatrix
+	: public TMatrix<T>
 {
 public:
 
@@ -19,41 +24,51 @@ public:
 	 * @param ZScale scale in the Z axis
 	 * @param ZOffset offset in the Z axis
 	 */
-	FOrthoMatrix(float Width,float Height,float ZScale,float ZOffset);
+	TOrthoMatrix(T Width,T Height,T ZScale,T ZOffset);
 };
 
 
-class FReversedZOrthoMatrix : public FMatrix
+template<typename T>
+struct TReversedZOrthoMatrix : public TMatrix<T>
 {
 public:
-	FReversedZOrthoMatrix(float Width,float Height,float ZScale,float ZOffset);
-	FReversedZOrthoMatrix(float Left, float Right, float Bottom, float Top, float ZScale, float ZOffset);
+	TReversedZOrthoMatrix(T Width,T Height,T ZScale,T ZOffset);
+	TReversedZOrthoMatrix(T Left, T Right, T Bottom, T Top, T ZScale, T ZOffset);
 };
 
-FORCEINLINE FOrthoMatrix::FOrthoMatrix(float Width,float Height,float ZScale,float ZOffset)
-	: FMatrix(
-		FPlane((Width != 0.0f) ? (1.0f / Width) : 1.0f, 0.0f, 0.0f, 0.0f),
-		FPlane(0.0f, (Height != 0.0f) ? (1.0f / Height) : 1.f, 0.0f, 0.0f),
-		FPlane(0.0f, 0.0f, ZScale, 0.0f),
-		FPlane(0.0f, 0.0f, ZOffset * ZScale, 1.0f)
+template<typename T>
+FORCEINLINE TOrthoMatrix<T>::TOrthoMatrix(T Width, T Height, T ZScale, T ZOffset)
+	: TMatrix<T>(
+		TPlane<T>((Width != 0.0f) ? (1.0f / Width) : 1.0f, 0.0f, 0.0f, 0.0f),
+		TPlane<T>(0.0f, (Height != 0.0f) ? (1.0f / Height) : 1.f, 0.0f, 0.0f),
+		TPlane<T>(0.0f, 0.0f, ZScale, 0.0f),
+		TPlane<T>(0.0f, 0.0f, ZOffset * ZScale, 1.0f)
 	)
 { }
 
 
-FORCEINLINE FReversedZOrthoMatrix::FReversedZOrthoMatrix(float Width,float Height,float ZScale,float ZOffset)
-	: FMatrix(
-		FPlane((Width != 0.0f) ? (1.0f / Width) : 1.0f, 0.0f, 0.0f, 0.0f),
-		FPlane(0.0f, (Height != 0.0f) ? (1.0f / Height) : 1.f, 0.0f, 0.0f),
-		FPlane(0.0f, 0.0f, -ZScale, 0.0f),
-		FPlane(0.0f, 0.0f, 1.0f - ZOffset * ZScale, 1.0f)
+template<typename T>
+FORCEINLINE TReversedZOrthoMatrix<T>::TReversedZOrthoMatrix(T Width, T Height, T ZScale, T ZOffset)
+	: TMatrix<T>(
+		TPlane<T>((Width != 0.0f) ? (1.0f / Width) : 1.0f, 0.0f, 0.0f, 0.0f),
+		TPlane<T>(0.0f, (Height != 0.0f) ? (1.0f / Height) : 1.f, 0.0f, 0.0f),
+		TPlane<T>(0.0f, 0.0f, -ZScale, 0.0f),
+		TPlane<T>(0.0f, 0.0f, 1.0f - ZOffset * ZScale, 1.0f)
 	)
 { }
 
-FORCEINLINE FReversedZOrthoMatrix::FReversedZOrthoMatrix(float Left, float Right, float Bottom, float Top, float ZScale, float ZOffset)
-	: FMatrix(
-		FPlane(1.0f / (Right - Left), 0.0f, 0.0f, 0.0f),
-		FPlane(0.0f, 1.0f / (Top - Bottom), 0.0f, 0.0f),
-		FPlane(0.0f, 0.0f, -ZScale, 0.0f),
-		FPlane((Left + Right) / (Left - Right), (Top + Bottom) / (Bottom - Top), 1.0f - ZOffset * ZScale, 1.0f)
+template<typename T>
+FORCEINLINE TReversedZOrthoMatrix<T>::TReversedZOrthoMatrix(T Left, T Right, T Bottom, T Top, T ZScale, T ZOffset)
+	: TMatrix<T>(
+		TPlane<T>(1.0f / (Right - Left), 0.0f, 0.0f, 0.0f),
+		TPlane<T>(0.0f, 1.0f / (Top - Bottom), 0.0f, 0.0f),
+		TPlane<T>(0.0f, 0.0f, -ZScale, 0.0f),
+		TPlane<T>((Left + Right) / (Left - Right), (Top + Bottom) / (Bottom - Top), 1.0f - ZOffset * ZScale, 1.0f)
 	)
-{ }
\ No newline at end of file
+{ }
+
+ } // namespace Math
+ } // namespace UE
+
+ DECLARE_LWC_TYPE(OrthoMatrix, 44);
+ DECLARE_LWC_TYPE(ReversedZOrthoMatrix, 44);
\ No newline at end of file
diff --git a/Engine/Source/Runtime/Core/Public/Math/PerspectiveMatrix.h b/Engine/Source/Runtime/Core/Public/Math/PerspectiveMatrix.h
--- a/Engine/Source/Runtime/Core/Public/Math/PerspectiveMatrix.h
+++ b/Engine/Source/Runtime/Core/Public/Math/PerspectiveMatrix.h
@@ -7,8 +7,12 @@
 #include "Math/Plane.h"
 #include "Math/Matrix.h"
 
-class FPerspectiveMatrix
-	: public FMatrix
+namespace UE {
+namespace Math {
+
+template<typename T>
+struct TPerspectiveMatrix
+	: public TMatrix<T>
 {
 public:
 
@@ -25,7 +29,7 @@ public:
 	 * @param MinZ distance to the near Z plane
 	 * @param MaxZ distance to the far Z plane
 	 */
-	FPerspectiveMatrix(float HalfFOVX, float HalfFOVY, float MultFOVX, float MultFOVY, float MinZ, float MaxZ);
+	TPerspectiveMatrix(T HalfFOVX, T HalfFOVY, T MultFOVX, T MultFOVY, T MinZ, T MaxZ);
 
 	/**
 	 * Constructor
@@ -37,7 +41,7 @@ public:
 	 * @param MaxZ distance to the far Z plane
 	 * @note that the FOV you pass in is actually half the FOV, unlike most perspective matrix functions (D3DXMatrixPerspectiveFovLH).
 	 */
-	FPerspectiveMatrix(float HalfFOV, float Width, float Height, float MinZ, float MaxZ);
+	TPerspectiveMatrix(T HalfFOV, T Width, T Height, T MinZ, T MaxZ);
 
 	/**
 	 * Constructor
@@ -48,16 +52,17 @@ public:
 	 * @param MinZ distance to the near Z plane
 	 * @note that the FOV you pass in is actually half the FOV, unlike most perspective matrix functions (D3DXMatrixPerspectiveFovLH).
 	 */
-	FPerspectiveMatrix(float HalfFOV, float Width, float Height, float MinZ);
+	TPerspectiveMatrix(T HalfFOV, T Width, T Height, T MinZ);
 };
 
 
-class FReversedZPerspectiveMatrix : public FMatrix
+template<typename T>
+struct TReversedZPerspectiveMatrix : public TMatrix<T>
 {
 public:
-	FReversedZPerspectiveMatrix(float HalfFOVX, float HalfFOVY, float MultFOVX, float MultFOVY, float MinZ, float MaxZ);
-	FReversedZPerspectiveMatrix(float HalfFOV, float Width, float Height, float MinZ, float MaxZ);
-	FReversedZPerspectiveMatrix(float HalfFOV, float Width, float Height, float MinZ);
+	TReversedZPerspectiveMatrix(T HalfFOVX, T HalfFOVY, T MultFOVX, T MultFOVY, T MinZ, T MaxZ);
+	TReversedZPerspectiveMatrix(T HalfFOV, T Width, T Height, T MinZ, T MaxZ);
+	TReversedZPerspectiveMatrix(T HalfFOV, T Width, T Height, T MinZ);
 };
 
 
@@ -68,66 +73,78 @@ public:
 #endif
 
 
-FORCEINLINE FPerspectiveMatrix::FPerspectiveMatrix(float HalfFOVX, float HalfFOVY, float MultFOVX, float MultFOVY, float MinZ, float MaxZ)
-	: FMatrix(
-		FPlane(MultFOVX / FMath::Tan(HalfFOVX),	0.0f,								0.0f,																	0.0f),
-		FPlane(0.0f,							MultFOVY / FMath::Tan(HalfFOVY),	0.0f,																	0.0f),
-		FPlane(0.0f,							0.0f,								((MinZ == MaxZ) ? (1.0f - Z_PRECISION) : MaxZ / (MaxZ - MinZ)),			1.0f),
-		FPlane(0.0f,							0.0f,								-MinZ * ((MinZ == MaxZ) ? (1.0f - Z_PRECISION) : MaxZ / (MaxZ - MinZ)),	0.0f)
+template<typename T>
+FORCEINLINE TPerspectiveMatrix<T>::TPerspectiveMatrix(T HalfFOVX, T HalfFOVY, T MultFOVX, T MultFOVY, T MinZ, T MaxZ)
+	: TMatrix<T>(
+		TPlane<T>(MultFOVX / FMath::Tan(HalfFOVX),	0.0f,								0.0f,																	0.0f),
+		TPlane<T>(0.0f,								MultFOVY / FMath::Tan(HalfFOVY),	0.0f,																	0.0f),
+		TPlane<T>(0.0f,								0.0f,								((MinZ == MaxZ) ? (1.0f - Z_PRECISION) : MaxZ / (MaxZ - MinZ)),			1.0f),
+		TPlane<T>(0.0f,								0.0f,								-MinZ * ((MinZ == MaxZ) ? (1.0f - Z_PRECISION) : MaxZ / (MaxZ - MinZ)),	0.0f)
 	)
 { }
 
 
-FORCEINLINE FPerspectiveMatrix::FPerspectiveMatrix(float HalfFOV, float Width, float Height, float MinZ, float MaxZ)
-	: FMatrix(
-		FPlane(1.0f / FMath::Tan(HalfFOV),	0.0f,									0.0f,							0.0f),
-		FPlane(0.0f,						Width / FMath::Tan(HalfFOV) / Height,	0.0f,							0.0f),
-		FPlane(0.0f,						0.0f,									((MinZ == MaxZ) ? (1.0f - Z_PRECISION) : MaxZ / (MaxZ - MinZ)),			1.0f),
-		FPlane(0.0f,						0.0f,									-MinZ * ((MinZ == MaxZ) ? (1.0f - Z_PRECISION) : MaxZ / (MaxZ - MinZ)),	0.0f)
+template<typename T>
+FORCEINLINE TPerspectiveMatrix<T>::TPerspectiveMatrix(T HalfFOV, T Width, T Height, T MinZ, T MaxZ)
+	: TMatrix<T>(
+		TPlane<T>(1.0f / FMath::Tan(HalfFOV),	0.0f,									0.0f,							0.0f),
+		TPlane<T>(0.0f,							Width / FMath::Tan(HalfFOV) / Height,	0.0f,							0.0f),
+		TPlane<T>(0.0f,							0.0f,									((MinZ == MaxZ) ? (1.0f - Z_PRECISION) : MaxZ / (MaxZ - MinZ)),			1.0f),
+		TPlane<T>(0.0f,							0.0f,									-MinZ * ((MinZ == MaxZ) ? (1.0f - Z_PRECISION) : MaxZ / (MaxZ - MinZ)),	0.0f)
 	)
 { }
 
 
-FORCEINLINE FPerspectiveMatrix::FPerspectiveMatrix(float HalfFOV, float Width, float Height, float MinZ)
-	: FMatrix(
-		FPlane(1.0f / FMath::Tan(HalfFOV),	0.0f,									0.0f,							0.0f),
-		FPlane(0.0f,						Width / FMath::Tan(HalfFOV) / Height,	0.0f,							0.0f),
-		FPlane(0.0f,						0.0f,									(1.0f - Z_PRECISION),			1.0f),
-		FPlane(0.0f,						0.0f,									-MinZ * (1.0f - Z_PRECISION),	0.0f)
+template<typename T>
+FORCEINLINE TPerspectiveMatrix<T>::TPerspectiveMatrix(T HalfFOV, T Width, T Height, T MinZ)
+	: TMatrix<T>(
+		TPlane<T>(1.0f / FMath::Tan(HalfFOV),	0.0f,									0.0f,							0.0f),
+		TPlane<T>(0.0f,							Width / FMath::Tan(HalfFOV) / Height,	0.0f,							0.0f),
+		TPlane<T>(0.0f,							0.0f,									(1.0f - Z_PRECISION),			1.0f),
+		TPlane<T>(0.0f,							0.0f,									-MinZ * (1.0f - Z_PRECISION),	0.0f)
 	)
 { }
 
 
-FORCEINLINE FReversedZPerspectiveMatrix::FReversedZPerspectiveMatrix(float HalfFOVX, float HalfFOVY, float MultFOVX, float MultFOVY, float MinZ, float MaxZ)
-	: FMatrix(
-		FPlane(MultFOVX / FMath::Tan(HalfFOVX),	0.0f,								0.0f,													0.0f),
-		FPlane(0.0f,							MultFOVY / FMath::Tan(HalfFOVY),	0.0f,													0.0f),
-		FPlane(0.0f,							0.0f,								((MinZ == MaxZ) ? 0.0f : MinZ / (MinZ - MaxZ)),			1.0f),
-		FPlane(0.0f,							0.0f,								((MinZ == MaxZ) ? MinZ : -MaxZ * MinZ / (MinZ - MaxZ)),	0.0f)
+template<typename T>
+FORCEINLINE TReversedZPerspectiveMatrix<T>::TReversedZPerspectiveMatrix(T HalfFOVX, T HalfFOVY, T MultFOVX, T MultFOVY, T MinZ, T MaxZ)
+	: TMatrix<T>(
+		TPlane<T>(MultFOVX / FMath::Tan(HalfFOVX),	0.0f,								0.0f,													0.0f),
+		TPlane<T>(0.0f,								MultFOVY / FMath::Tan(HalfFOVY),	0.0f,													0.0f),
+		TPlane<T>(0.0f,								0.0f,								((MinZ == MaxZ) ? 0.0f : MinZ / (MinZ - MaxZ)),			1.0f),
+		TPlane<T>(0.0f,								0.0f,								((MinZ == MaxZ) ? MinZ : -MaxZ * MinZ / (MinZ - MaxZ)),	0.0f)
 	)
 { }
 
 
-FORCEINLINE FReversedZPerspectiveMatrix::FReversedZPerspectiveMatrix(float HalfFOV, float Width, float Height, float MinZ, float MaxZ)
-	: FMatrix(
-		FPlane(1.0f / FMath::Tan(HalfFOV),	0.0f,									0.0f,													0.0f),
-		FPlane(0.0f,						Width / FMath::Tan(HalfFOV) / Height,	0.0f,													0.0f),
-		FPlane(0.0f,						0.0f,									((MinZ == MaxZ) ? 0.0f : MinZ / (MinZ - MaxZ)),			1.0f),
-		FPlane(0.0f,						0.0f,									((MinZ == MaxZ) ? MinZ : -MaxZ * MinZ / (MinZ - MaxZ)),	0.0f)
+template<typename T>
+FORCEINLINE TReversedZPerspectiveMatrix<T>::TReversedZPerspectiveMatrix(T HalfFOV, T Width, T Height, T MinZ, T MaxZ)
+	: TMatrix<T>(
+		TPlane<T>(1.0f / FMath::Tan(HalfFOV),	0.0f,									0.0f,													0.0f),
+		TPlane<T>(0.0f,							Width / FMath::Tan(HalfFOV) / Height,	0.0f,													0.0f),
+		TPlane<T>(0.0f,							0.0f,									((MinZ == MaxZ) ? 0.0f : MinZ / (MinZ - MaxZ)),			1.0f),
+		TPlane<T>(0.0f,							0.0f,									((MinZ == MaxZ) ? MinZ : -MaxZ * MinZ / (MinZ - MaxZ)),	0.0f)
 	)
 { }
 
 
-FORCEINLINE FReversedZPerspectiveMatrix::FReversedZPerspectiveMatrix(float HalfFOV, float Width, float Height, float MinZ)
-	: FMatrix(
-		FPlane(1.0f / FMath::Tan(HalfFOV),	0.0f,									0.0f, 0.0f),
-		FPlane(0.0f,						Width / FMath::Tan(HalfFOV) / Height,	0.0f, 0.0f),
-		FPlane(0.0f,						0.0f,									0.0f, 1.0f),
-		FPlane(0.0f,						0.0f,									MinZ, 0.0f)
+template<typename T>
+FORCEINLINE TReversedZPerspectiveMatrix<T>::TReversedZPerspectiveMatrix(T HalfFOV, T Width, T Height, T MinZ)
+	: TMatrix<T>(
+		TPlane<T>(1.0f / FMath::Tan(HalfFOV),	0.0f,									0.0f, 0.0f),
+		TPlane<T>(0.0f,							Width / FMath::Tan(HalfFOV) / Height,	0.0f, 0.0f),
+		TPlane<T>(0.0f,							0.0f,									0.0f, 1.0f),
+		TPlane<T>(0.0f,							0.0f,									MinZ, 0.0f)
 	)
 { }
 
-
 #ifdef _MSC_VER
 #pragma warning (pop)
 #endif
+
+
+} // namespace Math
+} // namespace UE
+
+ DECLARE_LWC_TYPE(PerspectiveMatrix, 44);
+ DECLARE_LWC_TYPE(ReversedZPerspectiveMatrix, 44);
\ No newline at end of file
diff --git a/Engine/Source/Runtime/Core/Public/Math/QuatRotationTranslationMatrix.h b/Engine/Source/Runtime/Core/Public/Math/QuatRotationTranslationMatrix.h
--- a/Engine/Source/Runtime/Core/Public/Math/QuatRotationTranslationMatrix.h
+++ b/Engine/Source/Runtime/Core/Public/Math/QuatRotationTranslationMatrix.h
@@ -8,63 +8,79 @@
 #include "Math/Matrix.h"
 #include "Math/Quat.h"
 
+namespace UE {
+namespace Math {
+
+// LWC_TODO: FQuat -> TQuat<T>
+	
 /** Rotation and translation matrix using quaternion rotation */
-class FQuatRotationTranslationMatrix
-	: public FMatrix
+template <typename T>
+struct TQuatRotationTranslationMatrix
+	: public TMatrix<T>
 {
 public:
+	using TMatrix<T>::M;
 
 	/** Constructor
 	*
 	* @param Q rotation
 	* @param Origin translation to apply
 	*/
-	FQuatRotationTranslationMatrix(const FQuat& Q, const FVector& Origin);
+	TQuatRotationTranslationMatrix(const FQuat& Q, const TVector<T>& Origin);
 
 	/** Matrix factory. Return an FMatrix so we don't have type conversion issues in expressions. */
-	static FMatrix Make(const FQuat& Q, const FVector& Origin)
+	static TMatrix<T> Make(const FQuat& Q, const TVector<T>& Origin)
 	{
-		return FQuatRotationTranslationMatrix(Q, Origin);
+		return TQuatRotationTranslationMatrix<T>(Q, Origin);
 	}
 };
 
 
 /** Rotation matrix using quaternion rotation */
-class FQuatRotationMatrix
-	: public FQuatRotationTranslationMatrix
+template<typename T>
+struct TQuatRotationMatrix
+	: public TQuatRotationTranslationMatrix<T>
 {
 public:
+	using TQuatRotationTranslationMatrix<T>::M;
 
 	/** Constructor
 	*
 	* @param Q rotation
 	*/
-	FQuatRotationMatrix(const FQuat& Q)
-		: FQuatRotationTranslationMatrix(Q, FVector::ZeroVector)
+	TQuatRotationMatrix(const FQuat& Q)
+		: TQuatRotationTranslationMatrix<T>(Q, TVector<T>::ZeroVector)
 	{
 	}
 
 	/** Matrix factory. Return an FMatrix so we don't have type conversion issues in expressions. */
-	static FMatrix Make(const FQuat& Q)
+	static TMatrix<T> Make(const FQuat& Q)
 	{
-		return FQuatRotationMatrix(Q);
+		return TQuatRotationMatrix<T>(Q);
 	}
 };
 
 
-FORCEINLINE FQuatRotationTranslationMatrix::FQuatRotationTranslationMatrix(const FQuat& Q, const FVector& Origin)
+template<typename T>
+FORCEINLINE TQuatRotationTranslationMatrix<T>::TQuatRotationTranslationMatrix(const FQuat& Q, const TVector<T>& Origin)
 {
 #if !(UE_BUILD_SHIPPING || UE_BUILD_TEST) && WITH_EDITORONLY_DATA
 	// Make sure Quaternion is normalized
 	check( Q.IsNormalized() );
 #endif
-	const float x2 = Q.X + Q.X;  const float y2 = Q.Y + Q.Y;  const float z2 = Q.Z + Q.Z;
-	const float xx = Q.X * x2;   const float xy = Q.X * y2;   const float xz = Q.X * z2;
-	const float yy = Q.Y * y2;   const float yz = Q.Y * z2;   const float zz = Q.Z * z2;
-	const float wx = Q.W * x2;   const float wy = Q.W * y2;   const float wz = Q.W * z2;
+	const T x2 = Q.X + Q.X;  const T y2 = Q.Y + Q.Y;  const T z2 = Q.Z + Q.Z;
+	const T xx = Q.X * x2;   const T xy = Q.X * y2;   const T xz = Q.X * z2;
+	const T yy = Q.Y * y2;   const T yz = Q.Y * z2;   const T zz = Q.Z * z2;
+	const T wx = Q.W * x2;   const T wy = Q.W * y2;   const T wz = Q.W * z2;
 
 	M[0][0] = 1.0f - (yy + zz);	M[1][0] = xy - wz;				M[2][0] = xz + wy;			M[3][0] = Origin.X;
 	M[0][1] = xy + wz;			M[1][1] = 1.0f - (xx + zz);		M[2][1] = yz - wx;			M[3][1] = Origin.Y;
 	M[0][2] = xz - wy;			M[1][2] = yz + wx;				M[2][2] = 1.0f - (xx + yy);	M[3][2] = Origin.Z;
 	M[0][3] = 0.0f;				M[1][3] = 0.0f;					M[2][3] = 0.0f;				M[3][3] = 1.0f;
 }
+
+} // namespace Math
+} // namespace UE
+
+DECLARE_LWC_TYPE(QuatRotationTranslationMatrix, 44);
+DECLARE_LWC_TYPE(QuatRotationMatrix, 44);
\ No newline at end of file
diff --git a/Engine/Source/Runtime/Core/Public/Math/RotationAboutPointMatrix.h b/Engine/Source/Runtime/Core/Public/Math/RotationAboutPointMatrix.h
--- a/Engine/Source/Runtime/Core/Public/Math/RotationAboutPointMatrix.h
+++ b/Engine/Source/Runtime/Core/Public/Math/RotationAboutPointMatrix.h
@@ -5,13 +5,21 @@
 #include "CoreTypes.h"
 #include "Math/Vector.h"
 #include "Math/Matrix.h"
+#include "Math/Quat.h"
 #include "Math/RotationTranslationMatrix.h"
 
+// LWC_TODO: FRotator -> TRotator<T>
+
+namespace UE {
+namespace Math {
+
 /** Rotates about an Origin point. */
-class FRotationAboutPointMatrix
-	: public FRotationTranslationMatrix
+template<typename T>	
+struct TRotationAboutPointMatrix
+	: public TRotationTranslationMatrix<T>
 {
 public:
+	using TRotationTranslationMatrix<T>::M;
 
 	/**
 	 * Constructor.
@@ -19,29 +27,37 @@ public:
 	 * @param Rot rotation
 	 * @param Origin about which to rotate.
 	 */
-	FRotationAboutPointMatrix(const FRotator& Rot, const FVector& Origin);
+	TRotationAboutPointMatrix(const FRotator& Rot, const TVector<T>& Origin);
 
-	/** Matrix factory. Return an FMatrix so we don't have type conversion issues in expressions. */
-	static FMatrix Make(const FRotator& Rot, const FVector& Origin)
+	/** Matrix factory. Return an TMatrix<T> so we don't have type conversion issues in expressions. */
+	static TMatrix<T> Make(const FRotator& Rot, const TVector<T>& Origin)
 	{
-		return FRotationAboutPointMatrix(Rot, Origin);
+		return TRotationAboutPointMatrix(Rot, Origin);
 	}
 
-	/** Matrix factory. Return an FMatrix so we don't have type conversion issues in expressions. */
-	static CORE_API FMatrix Make(const FQuat& Rot, const FVector& Origin);
+	/** Matrix factory. Return an TMatrix<T> so we don't have type conversion issues in expressions. */
+	static TMatrix<T> Make(const FQuat& Rot, const TVector<T>& Origin)
+	{
+		return TRotationAboutPointMatrix(Rot.Rotator(), Origin);
+	}
 };
 
-
-FORCEINLINE FRotationAboutPointMatrix::FRotationAboutPointMatrix(const FRotator& Rot, const FVector& Origin)
-	: FRotationTranslationMatrix(Rot, Origin)
+template<typename T>
+FORCEINLINE TRotationAboutPointMatrix<T>::TRotationAboutPointMatrix(const FRotator& Rot, const TVector<T>& Origin)
+	: TRotationTranslationMatrix<T>(Rot, Origin)
 {
 	// FRotationTranslationMatrix generates R * T.
 	// We need -T * R * T, so prepend that translation:
-	FVector XAxis(M[0][0], M[1][0], M[2][0]);
-	FVector YAxis(M[0][1], M[1][1], M[2][1]);
-	FVector ZAxis(M[0][2], M[1][2], M[2][2]);
+	TVector<T> XAxis(M[0][0], M[1][0], M[2][0]);
+	TVector<T> YAxis(M[0][1], M[1][1], M[2][1]);
+	TVector<T> ZAxis(M[0][2], M[1][2], M[2][2]);
 
 	M[3][0]	-= XAxis | Origin;
 	M[3][1]	-= YAxis | Origin;
 	M[3][2]	-= ZAxis | Origin;
 }
+
+} // namespace Math
+} // namespace UE
+
+DECLARE_LWC_TYPE(RotationAboutPointMatrix, 44);
\ No newline at end of file
diff --git a/Engine/Source/Runtime/Core/Public/Math/RotationMatrix.h b/Engine/Source/Runtime/Core/Public/Math/RotationMatrix.h
--- a/Engine/Source/Runtime/Core/Public/Math/RotationMatrix.h
+++ b/Engine/Source/Runtime/Core/Public/Math/RotationMatrix.h
@@ -6,56 +6,228 @@
 #include "Math/Vector.h"
 #include "Math/Rotator.h"
 #include "Math/Matrix.h"
-#include "Math/RotationTranslationMatrix.h"
+#include "Math/QuatRotationTranslationMatrix.h"
+
+namespace UE {
+namespace Math {
 
 /** Rotation matrix no translation */
-class FRotationMatrix
-	: public FRotationTranslationMatrix
+template<typename T>
+struct TRotationMatrix
+	: public TRotationTranslationMatrix<T>
 {
 public:
+	using TRotationTranslationMatrix<T>::M;
 
 	/**
 	 * Constructor.
 	 *
 	 * @param Rot rotation
 	 */
-	FRotationMatrix(const FRotator& Rot)
-		: FRotationTranslationMatrix(Rot, FVector::ZeroVector)
+	TRotationMatrix(const FRotator& Rot)
+		: TRotationTranslationMatrix<T>(Rot, TVector<T>::ZeroVector)
 	{ }
 
-	/** Matrix factory. Return an FMatrix so we don't have type conversion issues in expressions. */
-	static FMatrix Make(FRotator const& Rot)
+	/** Matrix factory. Return an TMatrix<T> so we don't have type conversion issues in expressions. */
+	static TMatrix<T> Make(FRotator const& Rot)
 	{
-		return FRotationMatrix(Rot);
+		return TRotationMatrix(Rot);
 	}
 
-	/** Matrix factory. Return an FMatrix so we don't have type conversion issues in expressions. */
-	static CORE_API FMatrix Make(FQuat const& Rot);
+	/** Matrix factory. Return an TMatrix<T> so we don't have type conversion issues in expressions. */
+	static TMatrix<T> Make(FQuat const& Rot);
 
 	/** Builds a rotation matrix given only a XAxis. Y and Z are unspecified but will be orthonormal. XAxis need not be normalized. */
-	static CORE_API FMatrix MakeFromX(FVector const& XAxis);
+	static TMatrix<T> MakeFromX(TVector<T> const& XAxis);
 
 	/** Builds a rotation matrix given only a YAxis. X and Z are unspecified but will be orthonormal. YAxis need not be normalized. */
-	static CORE_API FMatrix MakeFromY(FVector const& YAxis);
+	static TMatrix<T> MakeFromY(TVector<T> const& YAxis);
 
 	/** Builds a rotation matrix given only a ZAxis. X and Y are unspecified but will be orthonormal. ZAxis need not be normalized. */
-	static CORE_API FMatrix MakeFromZ(FVector const& ZAxis);
+	static TMatrix<T> MakeFromZ(TVector<T> const& ZAxis);
 
 	/** Builds a matrix with given X and Y axes. X will remain fixed, Y may be changed minimally to enforce orthogonality. Z will be computed. Inputs need not be normalized. */
-	static CORE_API FMatrix MakeFromXY(FVector const& XAxis, FVector const& YAxis);
+	static TMatrix<T> MakeFromXY(TVector<T> const& XAxis, TVector<T> const& YAxis);
 
 	/** Builds a matrix with given X and Z axes. X will remain fixed, Z may be changed minimally to enforce orthogonality. Y will be computed. Inputs need not be normalized. */
-	static CORE_API FMatrix MakeFromXZ(FVector const& XAxis, FVector const& ZAxis);
+	static TMatrix<T> MakeFromXZ(TVector<T> const& XAxis, TVector<T> const& ZAxis);
 
 	/** Builds a matrix with given Y and X axes. Y will remain fixed, X may be changed minimally to enforce orthogonality. Z will be computed. Inputs need not be normalized. */
-	static CORE_API FMatrix MakeFromYX(FVector const& YAxis, FVector const& XAxis);
+	static TMatrix<T> MakeFromYX(TVector<T> const& YAxis, TVector<T> const& XAxis);
 
 	/** Builds a matrix with given Y and Z axes. Y will remain fixed, Z may be changed minimally to enforce orthogonality. X will be computed. Inputs need not be normalized. */
-	static CORE_API FMatrix MakeFromYZ(FVector const& YAxis, FVector const& ZAxis);
+	static TMatrix<T> MakeFromYZ(TVector<T> const& YAxis, TVector<T> const& ZAxis);
 
 	/** Builds a matrix with given Z and X axes. Z will remain fixed, X may be changed minimally to enforce orthogonality. Y will be computed. Inputs need not be normalized. */
-	static CORE_API FMatrix MakeFromZX(FVector const& ZAxis, FVector const& XAxis);
+	static TMatrix<T> MakeFromZX(TVector<T> const& ZAxis, TVector<T> const& XAxis);
 
 	/** Builds a matrix with given Z and Y axes. Z will remain fixed, Y may be changed minimally to enforce orthogonality. X will be computed. Inputs need not be normalized. */
-	static CORE_API FMatrix MakeFromZY(FVector const& ZAxis, FVector const& YAxis);
+	static TMatrix<T> MakeFromZY(TVector<T> const& ZAxis, TVector<T> const& YAxis);
 };
+
+template<typename T>
+TMatrix<T> TRotationMatrix<T>::Make(FQuat const& Rot)	// LWC_TODO: FQuat -> TQuat<T>
+{
+	return TQuatRotationTranslationMatrix<T>(Rot, TVector<T>::ZeroVector);
+}
+
+template<typename T>
+TMatrix<T> TRotationMatrix<T>::MakeFromX(TVector<T> const& XAxis)
+{
+	TVector<T> const NewX = XAxis.GetSafeNormal();
+
+	// try to use up if possible
+	TVector<T> const UpVector = (FMath::Abs(NewX.Z) < (1.f - KINDA_SMALL_NUMBER)) ? TVector<T>(0, 0, 1.f) : TVector<T>(1.f, 0, 0);
+
+	const TVector<T> NewY = (UpVector ^ NewX).GetSafeNormal();
+	const TVector<T> NewZ = NewX ^ NewY;
+
+	return TMatrix<T>(NewX, NewY, NewZ, TVector<T>::ZeroVector);
+}
+
+template<typename T>
+TMatrix<T> TRotationMatrix<T>::MakeFromY(TVector<T> const& YAxis)
+{
+	TVector<T> const NewY = YAxis.GetSafeNormal();
+
+	// try to use up if possible
+	TVector<T> const UpVector = (FMath::Abs(NewY.Z) < (1.f - KINDA_SMALL_NUMBER)) ? TVector<T>(0, 0, 1.f) : TVector<T>(1.f, 0, 0);
+
+	const TVector<T> NewZ = (UpVector ^ NewY).GetSafeNormal();
+	const TVector<T> NewX = NewY ^ NewZ;
+
+	return TMatrix<T>(NewX, NewY, NewZ, TVector<T>::ZeroVector);
+}
+
+template<typename T>
+TMatrix<T> TRotationMatrix<T>::MakeFromZ(TVector<T> const& ZAxis)
+{
+	TVector<T> const NewZ = ZAxis.GetSafeNormal();
+
+	// try to use up if possible
+	TVector<T> const UpVector = (FMath::Abs(NewZ.Z) < (1.f - KINDA_SMALL_NUMBER)) ? TVector<T>(0, 0, 1.f) : TVector<T>(1.f, 0, 0);
+
+	const TVector<T> NewX = (UpVector ^ NewZ).GetSafeNormal();
+	const TVector<T> NewY = NewZ ^ NewX;
+
+	return TMatrix<T>(NewX, NewY, NewZ, TVector<T>::ZeroVector);
+}
+
+template<typename T>
+TMatrix<T> TRotationMatrix<T>::MakeFromXY(TVector<T> const& XAxis, TVector<T> const& YAxis)
+{
+	TVector<T> NewX = XAxis.GetSafeNormal();
+	TVector<T> Norm = YAxis.GetSafeNormal();
+
+	// if they're almost same, we need to find arbitrary vector
+	if (FMath::IsNearlyEqual(FMath::Abs(NewX | Norm), T(1.f)))
+	{
+		// make sure we don't ever pick the same as NewX
+		Norm = (FMath::Abs(NewX.Z) < (1.f - KINDA_SMALL_NUMBER)) ? TVector<T>(0, 0, 1.f) : TVector<T>(1.f, 0, 0);
+	}
+
+	const TVector<T> NewZ = (NewX ^ Norm).GetSafeNormal();
+	const TVector<T> NewY = NewZ ^ NewX;
+
+	return TMatrix<T>(NewX, NewY, NewZ, TVector<T>::ZeroVector);
+}
+
+template<typename T>
+TMatrix<T> TRotationMatrix<T>::MakeFromXZ(TVector<T> const& XAxis, TVector<T> const& ZAxis)
+{
+	TVector<T> const NewX = XAxis.GetSafeNormal();
+	TVector<T> Norm = ZAxis.GetSafeNormal();
+
+	// if they're almost same, we need to find arbitrary vector
+	if (FMath::IsNearlyEqual(FMath::Abs(NewX | Norm), T(1.f)))
+	{
+		// make sure we don't ever pick the same as NewX
+		Norm = (FMath::Abs(NewX.Z) < (1.f - KINDA_SMALL_NUMBER)) ? TVector<T>(0, 0, 1.f) : TVector<T>(1.f, 0, 0);
+	}
+
+	const TVector<T> NewY = (Norm ^ NewX).GetSafeNormal();
+	const TVector<T> NewZ = NewX ^ NewY;
+
+	return TMatrix<T>(NewX, NewY, NewZ, TVector<T>::ZeroVector);
+}
+
+template<typename T>
+TMatrix<T> TRotationMatrix<T>::MakeFromYX(TVector<T> const& YAxis, TVector<T> const& XAxis)
+{
+	TVector<T> const NewY = YAxis.GetSafeNormal();
+	TVector<T> Norm = XAxis.GetSafeNormal();
+
+	// if they're almost same, we need to find arbitrary vector
+	if (FMath::IsNearlyEqual(FMath::Abs(NewY | Norm), T(1.f)))
+	{
+		// make sure we don't ever pick the same as NewX
+		Norm = (FMath::Abs(NewY.Z) < (1.f - KINDA_SMALL_NUMBER)) ? TVector<T>(0, 0, 1.f) : TVector<T>(1.f, 0, 0);
+	}
+
+	const TVector<T> NewZ = (Norm ^ NewY).GetSafeNormal();
+	const TVector<T> NewX = NewY ^ NewZ;
+
+	return TMatrix<T>(NewX, NewY, NewZ, TVector<T>::ZeroVector);
+}
+
+template<typename T>
+TMatrix<T> TRotationMatrix<T>::MakeFromYZ(TVector<T> const& YAxis, TVector<T> const& ZAxis)
+{
+	TVector<T> const NewY = YAxis.GetSafeNormal();
+	TVector<T> Norm = ZAxis.GetSafeNormal();
+
+	// if they're almost same, we need to find arbitrary vector
+	if (FMath::IsNearlyEqual(FMath::Abs(NewY | Norm), T(1.f)))
+	{
+		// make sure we don't ever pick the same as NewX
+		Norm = (FMath::Abs(NewY.Z) < (1.f - KINDA_SMALL_NUMBER)) ? TVector<T>(0, 0, 1.f) : TVector<T>(1.f, 0, 0);
+	}
+
+	const TVector<T> NewX = (NewY ^ Norm).GetSafeNormal();
+	const TVector<T> NewZ = NewX ^ NewY;
+
+	return TMatrix<T>(NewX, NewY, NewZ, TVector<T>::ZeroVector);
+}
+
+template<typename T>
+TMatrix<T> TRotationMatrix<T>::MakeFromZX(TVector<T> const& ZAxis, TVector<T> const& XAxis)
+{
+	TVector<T> const NewZ = ZAxis.GetSafeNormal();
+	TVector<T> Norm = XAxis.GetSafeNormal();
+
+	// if they're almost same, we need to find arbitrary vector
+	if (FMath::IsNearlyEqual(FMath::Abs(NewZ | Norm), T(1.f)))
+	{
+		// make sure we don't ever pick the same as NewX
+		Norm = (FMath::Abs(NewZ.Z) < (1.f - KINDA_SMALL_NUMBER)) ? TVector<T>(0, 0, 1.f) : TVector<T>(1.f, 0, 0);
+	}
+
+	const TVector<T> NewY = (NewZ ^ Norm).GetSafeNormal();
+	const TVector<T> NewX = NewY ^ NewZ;
+
+	return TMatrix<T>(NewX, NewY, NewZ, TVector<T>::ZeroVector);
+}
+
+template<typename T>
+TMatrix<T> TRotationMatrix<T>::MakeFromZY(TVector<T> const& ZAxis, TVector<T> const& YAxis)
+{
+	TVector<T> const NewZ = ZAxis.GetSafeNormal();
+	TVector<T> Norm = YAxis.GetSafeNormal();
+
+	// if they're almost same, we need to find arbitrary vector
+	if (FMath::IsNearlyEqual(FMath::Abs(NewZ | Norm), T(1.f)))
+	{
+		// make sure we don't ever pick the same as NewX
+		Norm = (FMath::Abs(NewZ.Z) < (1.f - KINDA_SMALL_NUMBER)) ? TVector<T>(0, 0, 1.f) : TVector<T>(1.f, 0, 0);
+	}
+
+	const TVector<T> NewX = (Norm ^ NewZ).GetSafeNormal();
+	const TVector<T> NewY = NewZ ^ NewX;
+
+	return TMatrix<T>(NewX, NewY, NewZ, TVector<T>::ZeroVector);
+}
+
+} // namespace Math
+} // namespace UE
+
+ DECLARE_LWC_TYPE(RotationMatrix, 44);
\ No newline at end of file
diff --git a/Engine/Source/Runtime/Core/Public/Math/RotationTranslationMatrix.h b/Engine/Source/Runtime/Core/Public/Math/RotationTranslationMatrix.h
--- a/Engine/Source/Runtime/Core/Public/Math/RotationTranslationMatrix.h
+++ b/Engine/Source/Runtime/Core/Public/Math/RotationTranslationMatrix.h
@@ -7,11 +7,18 @@
 #include "Math/VectorRegister.h"
 #include "Math/Matrix.h"
 
+// LWC_TODO: FRotator -> TRotator<T>
+
+namespace UE {
+namespace Math {
+
 /** Combined rotation and translation matrix */
-class FRotationTranslationMatrix
-	: public FMatrix
+template<typename T>
+struct TRotationTranslationMatrix
+	: public TMatrix<T>
 {
 public:
+	using TMatrix<T>::M;
 
 	/**
 	 * Constructor.
@@ -19,40 +26,41 @@ public:
 	 * @param Rot rotation
 	 * @param Origin translation to apply
 	 */
-	FRotationTranslationMatrix(const FRotator& Rot, const FVector& Origin);
+	TRotationTranslationMatrix(const FRotator& Rot, const TVector<T>& Origin);
 
-	/** Matrix factory. Return an FMatrix so we don't have type conversion issues in expressions. */
-	static FMatrix Make(const FRotator& Rot, const FVector& Origin)
+	/** Matrix factory. Return an TMatrix<T> so we don't have type conversion issues in expressions. */
+	static TMatrix<T> Make(const FRotator& Rot, const TVector<T>& Origin)
 	{
-		return FRotationTranslationMatrix(Rot, Origin);
+		return TRotationTranslationMatrix(Rot, Origin);
 	}
 };
 
 
-FORCEINLINE FRotationTranslationMatrix::FRotationTranslationMatrix(const FRotator& Rot, const FVector& Origin)
+template<typename T>
+FORCEINLINE TRotationTranslationMatrix<T>::TRotationTranslationMatrix(const FRotator& Rot, const TVector<T>& Origin)
 {
 #if PLATFORM_ENABLE_VECTORINTRINSICS
 
-	const VectorRegister Angles = MakeVectorRegister(Rot.Pitch, Rot.Yaw, Rot.Roll, 0.0f);
-	const VectorRegister HalfAngles = VectorMultiply(Angles, GlobalVectorConstants::DEG_TO_RAD);
+	const TVectorRegisterType<T> Angles = MakeVectorRegister(Rot.Pitch, Rot.Yaw, Rot.Roll, 0.0f);
+	const TVectorRegisterType<T> HalfAngles = VectorMultiply(Angles, GlobalVectorConstants::DEG_TO_RAD);
 
-	union { VectorRegister v; FMatrix::FReal f[4]; } SinAngles, CosAngles;
+	union { TVectorRegisterType<T> v; T f[4]; } SinAngles, CosAngles;
 	VectorSinCos(&SinAngles.v, &CosAngles.v, &HalfAngles);
 
-	const FMatrix::FReal SP	= SinAngles.f[0];
-	const FMatrix::FReal SY	= SinAngles.f[1];
-	const FMatrix::FReal SR	= SinAngles.f[2];
-	const FMatrix::FReal CP	= CosAngles.f[0];
-	const FMatrix::FReal CY	= CosAngles.f[1];
-	const FMatrix::FReal CR	= CosAngles.f[2];
+	const T SP	= SinAngles.f[0];
+	const T SY	= SinAngles.f[1];
+	const T SR	= SinAngles.f[2];
+	const T CP	= CosAngles.f[0];
+	const T CY	= CosAngles.f[1];
+	const T CR	= CosAngles.f[2];
 
 #else
-	
-	FMatrix::FReal SP, SY, SR;
-	FMatrix::FReal CP, CY, CR;
-	FMath::SinCos(&SP, &CP, (FMatrix::FReal) FMath::DegreesToRadians(Rot.Pitch));
-	FMath::SinCos(&SY, &CY, (FMatrix::FReal) FMath::DegreesToRadians(Rot.Yaw));
-	FMath::SinCos(&SR, &CR, (FMatrix::FReal) FMath::DegreesToRadians(Rot.Roll));
+
+	T SP, SY, SR;
+	T CP, CY, CR;
+	FMath::SinCos(&SP, &CP, (T)FMath::DegreesToRadians(Rot.Pitch));
+	FMath::SinCos(&SY, &CY, (T)FMath::DegreesToRadians(Rot.Yaw));
+	FMath::SinCos(&SR, &CR, (T)FMath::DegreesToRadians(Rot.Roll));
 
 #endif // PLATFORM_ENABLE_VECTORINTRINSICS
 
@@ -76,3 +84,8 @@ FORCEINLINE FRotationTranslationMatrix::FRotationTranslationMatrix(const FRotato
 	M[3][2]	= Origin.Z;
 	M[3][3]	= 1.f;
 }
+
+} // namespace Math
+} // namespace UE
+
+DECLARE_LWC_TYPE(RotationTranslationMatrix, 44);
\ No newline at end of file
diff --git a/Engine/Source/Runtime/Core/Public/Math/ScaleMatrix.h b/Engine/Source/Runtime/Core/Public/Math/ScaleMatrix.h
--- a/Engine/Source/Runtime/Core/Public/Math/ScaleMatrix.h
+++ b/Engine/Source/Runtime/Core/Public/Math/ScaleMatrix.h
@@ -6,34 +6,38 @@
 #include "Math/Plane.h"
 #include "Math/Matrix.h"
 
+namespace UE {
+namespace Math {
+
 /**
  * Scale matrix.
  */
-class FScaleMatrix
-	: public FMatrix
+template<typename T>
+struct TScaleMatrix
+	: public TMatrix<T>
 {
 public:
 
 	/**
 	 * @param Scale uniform scale to apply to matrix.
 	 */
-	FScaleMatrix( float Scale );
+	TScaleMatrix( T Scale );
 
 	/**
 	 * @param Scale Non-uniform scale to apply to matrix.
 	 */
-	FScaleMatrix( const FVector& Scale );
+	TScaleMatrix( const TVector<T>& Scale );
 
-	/** Matrix factory. Return an FMatrix so we don't have type conversion issues in expressions. */
-	static FMatrix Make(float Scale)
+	/** Matrix factory. Return an TMatrix<T> so we don't have type conversion issues in expressions. */
+	static TMatrix<T> Make(T Scale)
 	{
-		return FScaleMatrix(Scale);
+		return TScaleMatrix<T>(Scale);
 	}
 
-	/** Matrix factory. Return an FMatrix so we don't have type conversion issues in expressions. */
-	static FMatrix Make(const FVector& Scale)
+	/** Matrix factory. Return an TMatrix<T> so we don't have type conversion issues in expressions. */
+	static TMatrix<T> Make(const TVector<T>& Scale)
 	{
-		return FScaleMatrix(Scale);
+		return TScaleMatrix<T>(Scale);
 	}
 };
 
@@ -41,21 +45,28 @@ public:
 /* FScaleMatrix inline functions
  *****************************************************************************/
 
-FORCEINLINE FScaleMatrix::FScaleMatrix( float Scale )
-	: FMatrix(
-		FPlane(Scale,	0.0f,	0.0f,	0.0f),
-		FPlane(0.0f,	Scale,	0.0f,	0.0f),
-		FPlane(0.0f,	0.0f,	Scale,	0.0f),
-		FPlane(0.0f,	0.0f,	0.0f,	1.0f)
+template<typename T>
+FORCEINLINE TScaleMatrix<T>::TScaleMatrix( T Scale )
+	: TMatrix<T>(
+		TPlane<T>(Scale,	0.0f,	0.0f,	0.0f),
+		TPlane<T>(0.0f,		Scale,	0.0f,	0.0f),
+		TPlane<T>(0.0f,		0.0f,	Scale,	0.0f),
+		TPlane<T>(0.0f,		0.0f,	0.0f,	1.0f)
 	)
 { }
 
 
-FORCEINLINE FScaleMatrix::FScaleMatrix( const FVector& Scale )
-	: FMatrix(
-		FPlane(Scale.X,	0.0f,		0.0f,		0.0f),
-		FPlane(0.0f,	Scale.Y,	0.0f,		0.0f),
-		FPlane(0.0f,	0.0f,		Scale.Z,	0.0f),
-		FPlane(0.0f,	0.0f,		0.0f,		1.0f)
+template<typename T>
+FORCEINLINE TScaleMatrix<T>::TScaleMatrix( const TVector<T>& Scale )
+	: TMatrix<T>(
+		TPlane<T>(Scale.X,	0.0f,		0.0f,		0.0f),
+		TPlane<T>(0.0f,		Scale.Y,	0.0f,		0.0f),
+		TPlane<T>(0.0f,		0.0f,		Scale.Z,	0.0f),
+		TPlane<T>(0.0f,		0.0f,		0.0f,		1.0f)
 	)
 { }
+	
+} // namespace Math
+} // namespace UE
+
+DECLARE_LWC_TYPE(ScaleMatrix, 44);
\ No newline at end of file
diff --git a/Engine/Source/Runtime/Core/Public/Math/ScaleRotationTranslationMatrix.h b/Engine/Source/Runtime/Core/Public/Math/ScaleRotationTranslationMatrix.h
--- a/Engine/Source/Runtime/Core/Public/Math/ScaleRotationTranslationMatrix.h
+++ b/Engine/Source/Runtime/Core/Public/Math/ScaleRotationTranslationMatrix.h
@@ -6,12 +6,19 @@
 #include "Math/UnrealMathUtility.h"
 #include "Math/Matrix.h"
 
+// LWC_TODO: FRotator -> TRotator<T>
+
+namespace UE {
+namespace Math {
+		
 /** Combined Scale rotation and translation matrix */
-class FScaleRotationTranslationMatrix
-	: public FMatrix
+template<typename T>
+struct TScaleRotationTranslationMatrix
+	: public TMatrix<T>
 {
 public:
-
+	using TMatrix<T>::M;
+	
 	/**
 	 * Constructor.
 	 *
@@ -19,12 +26,13 @@ public:
 	 * @param Rot rotation
 	 * @param Origin translation to apply
 	 */
-	FScaleRotationTranslationMatrix(const FVector& Scale, const FRotator& Rot, const FVector& Origin);
+	TScaleRotationTranslationMatrix(const TVector<T>& Scale, const FRotator& Rot, const TVector<T>& Origin);
 };
 
 namespace
 {
-	void GetSinCos(float& S, float& C, float Degrees)
+	template<typename T>
+	void GetSinCos(T& S, T& C, T Degrees)
 	{
 		if (Degrees == 0.f)
 		{
@@ -53,13 +61,14 @@ namespace
 	}
 }
 
-FORCEINLINE FScaleRotationTranslationMatrix::FScaleRotationTranslationMatrix(const FVector& Scale, const FRotator& Rot, const FVector& Origin)
+template<typename T>
+FORCEINLINE TScaleRotationTranslationMatrix<T>::TScaleRotationTranslationMatrix(const TVector<T>& Scale, const FRotator& Rot, const TVector<T>& Origin)
 {
-	float SP, SY, SR;
-	float CP, CY, CR;
-	GetSinCos(SP, CP, Rot.Pitch);
-	GetSinCos(SY, CY, Rot.Yaw);
-	GetSinCos(SR, CR, Rot.Roll);
+	T SP, SY, SR;
+	T CP, CY, CR;
+	GetSinCos(SP, CP, (T)Rot.Pitch);
+	GetSinCos(SY, CY, (T)Rot.Yaw);
+	GetSinCos(SR, CR, (T)Rot.Roll);
 
 	M[0][0]	= (CP * CY) * Scale.X;
 	M[0][1]	= (CP * SY) * Scale.X;
@@ -81,3 +90,8 @@ FORCEINLINE FScaleRotationTranslationMatrix::FScaleRotationTranslationMatrix(con
 	M[3][2]	= Origin.Z;
 	M[3][3]	= 1.f;
 }
+	
+} // namespace Math
+} // namespace UE
+
+DECLARE_LWC_TYPE(ScaleRotationTranslationMatrix, 44);
\ No newline at end of file
diff --git a/Engine/Source/Runtime/Core/Public/Math/TranslationMatrix.h b/Engine/Source/Runtime/Core/Public/Math/TranslationMatrix.h
--- a/Engine/Source/Runtime/Core/Public/Math/TranslationMatrix.h
+++ b/Engine/Source/Runtime/Core/Public/Math/TranslationMatrix.h
@@ -7,27 +7,36 @@
 #include "Math/Plane.h"
 #include "Math/Matrix.h"
 
-class FTranslationMatrix
-	: public FMatrix
+namespace UE {
+namespace Math {
+
+template<typename T>
+struct TTranslationMatrix
+	: public TMatrix<T>
 {
 public:
 
 	/** Constructor translation matrix based on given vector */
-	FTranslationMatrix(const FVector& Delta);
+	TTranslationMatrix(const TVector<T>& Delta);
 
 	/** Matrix factory. Return an FMatrix so we don't have type conversion issues in expressions. */
-	static FMatrix Make(FVector const& Delta)
+	static TMatrix<T> Make(TVector<T> const& Delta)
 	{
-		return FTranslationMatrix(Delta);
+		return TTranslationMatrix<T>(Delta);
 	}
 };
 
-
-FORCEINLINE FTranslationMatrix::FTranslationMatrix(const FVector& Delta)
-	: FMatrix(
-		FPlane(1.0f,	0.0f,	0.0f,	0.0f),
-		FPlane(0.0f,	1.0f,	0.0f,	0.0f),
-		FPlane(0.0f,	0.0f,	1.0f,	0.0f),
-		FPlane(Delta.X,	Delta.Y,Delta.Z,1.0f)
+template<typename T>
+FORCEINLINE TTranslationMatrix<T>::TTranslationMatrix(const TVector<T>& Delta)
+	: TMatrix<T>(
+		TPlane<T>(1.0f,		0.0f,	0.0f,	0.0f),
+		TPlane<T>(0.0f,		1.0f,	0.0f,	0.0f),
+		TPlane<T>(0.0f,		0.0f,	1.0f,	0.0f),
+		TPlane<T>(Delta.X,	Delta.Y,Delta.Z,1.0f)
 	)
 { }
+	
+} // namespace Math
+} // namespace UE
+
+DECLARE_LWC_TYPE(TranslationMatrix, 44);
\ No newline at end of file
diff --git a/Engine/Source/Runtime/Core/Public/Math/Vector.h b/Engine/Source/Runtime/Core/Public/Math/Vector.h
--- a/Engine/Source/Runtime/Core/Public/Math/Vector.h
+++ b/Engine/Source/Runtime/Core/Public/Math/Vector.h
@@ -43,6 +43,8 @@ template<typename T> struct TPlane;
 template<typename T>
 struct TVector
 {
+	static_assert(std::is_floating_point_v<T>, "T must be floating point");
+
 public:
 	using FReal = T;
 
diff --git a/Engine/Source/Runtime/Core/Public/Math/VectorRegister.h b/Engine/Source/Runtime/Core/Public/Math/VectorRegister.h
--- a/Engine/Source/Runtime/Core/Public/Math/VectorRegister.h
+++ b/Engine/Source/Runtime/Core/Public/Math/VectorRegister.h
@@ -23,6 +23,11 @@
 
 #define SIMD_ALIGNMENT (alignof(VectorRegister))
 
+
+// Alias which helps resolve a template type (float/double) to the appropriate VectorRegister type (VectorRegister4Float/VectorRegister4Double).
+template<typename T>
+using TVectorRegisterType = std::conditional_t<std::is_same_v<T, float>, VectorRegister4Float, std::conditional_t<std::is_same_v<T, double>, VectorRegister4Double, void> >;
+
 // 'Cross-platform' vector intrinsics (built on the platform-specific ones defined above)
 #include "Math/UnrealMathVectorCommon.h"
 
