NonUnity build fix for RotationMatrix

diff --git a/Engine/Source/Runtime/Core/Public/Math/RotationMatrix.h b/Engine/Source/Runtime/Core/Public/Math/RotationMatrix.h
--- a/Engine/Source/Runtime/Core/Public/Math/RotationMatrix.h
+++ b/Engine/Source/Runtime/Core/Public/Math/RotationMatrix.h
@@ -6,6 +6,7 @@
 #include "Math/Vector.h"
 #include "Math/Rotator.h"
 #include "Math/Matrix.h"
+#include "Math/RotationTranslationMatrix.h"
 #include "Math/QuatRotationTranslationMatrix.h"
 
 namespace UE {
