fix read out of bounds in Trace startup
caused crash with -stompmalloc

diff --git a/Engine/Source/Runtime/Core/Private/ProfilingDebugging/PlatformFileTrace.cpp b/Engine/Source/Runtime/Core/Private/ProfilingDebugging/PlatformFileTrace.cpp
--- a/Engine/Source/Runtime/Core/Private/ProfilingDebugging/PlatformFileTrace.cpp
+++ b/Engine/Source/Runtime/Core/Private/ProfilingDebugging/PlatformFileTrace.cpp
@@ -87,10 +87,10 @@ namespace
 
 void FPlatformFileTrace::BeginOpen(const TCHAR* Path)
 {
-	uint32 PathSize = (FCString::Strlen(Path) + 1) * sizeof(TCHAR);
+	int32 PathLength = FCString::Strlen(Path);
 	UE_TRACE_LOG(PlatformFile, BeginOpen, FileChannel)
 		<< BeginOpen.Cycle(FPlatformTime::Cycles64())
-		<< BeginOpen.Path(Path, PathSize);
+		<< BeginOpen.Path(Path, PathLength);
 }
 
 void FPlatformFileTrace::EndOpen(uint64 FileHandle)
