Logging: Fixed a race condition causing log lines to merge together when writing to stdout

diff --git a/Engine/Source/Runtime/Core/Public/HAL/FeedbackContextAnsi.h b/Engine/Source/Runtime/Core/Public/HAL/FeedbackContextAnsi.h
--- a/Engine/Source/Runtime/Core/Public/HAL/FeedbackContextAnsi.h
+++ b/Engine/Source/Runtime/Core/Public/HAL/FeedbackContextAnsi.h
@@ -83,8 +83,8 @@ public:
 					AddWarning(Format);
 				}
 			}
+			Format.AppendChar(TEXT('\n'));
 			LocalPrint(*Format);
-			LocalPrint( TEXT("\n") );
 		}
 		else if (Verbosity == ELogVerbosity::SetColor)
 		{
