TIsCharType specialization added for UTF32CHAR.
New PLATFORM_UCS2CHAR_IS_UTF16CHAR macro which we can use to start deprecating that relationship.

diff --git a/Engine/Source/Runtime/Core/Public/HAL/Platform.h b/Engine/Source/Runtime/Core/Public/HAL/Platform.h
--- a/Engine/Source/Runtime/Core/Public/HAL/Platform.h
+++ b/Engine/Source/Runtime/Core/Public/HAL/Platform.h
@@ -263,6 +263,10 @@
 #ifndef PLATFORM_TCHAR_IS_CHAR16
 	#define PLATFORM_TCHAR_IS_CHAR16			0
 #endif
+#ifndef PLATFORM_UCS2CHAR_IS_UTF16CHAR
+	// Currently true, but we don't want it to be true
+	#define PLATFORM_UCS2CHAR_IS_UTF16CHAR		1
+#endif
 #ifndef PLATFORM_HAS_BSD_TIME
 	#define PLATFORM_HAS_BSD_TIME				1
 #endif
@@ -1039,9 +1043,16 @@ namespace TypeTests
 
 	static_assert((!TAreTypesEqual<ANSICHAR, WIDECHAR>::Value),  "ANSICHAR and WIDECHAR should be different types.");
 	static_assert((!TAreTypesEqual<ANSICHAR, UTF8CHAR>::Value),  "ANSICHAR and UTF8CHAR should be different types.");
-//	static_assert((!TAreTypesEqual<UCS2CHAR, UTF16CHAR>::Value), "UCS2CHAR and UTF16CHAR should be different types."); // should be true - not currently
-	static_assert((!TAreTypesEqual<ANSICHAR, UCS2CHAR>::Value),  "ANSICHAR and CHAR16 should be different types.");
-	static_assert((!TAreTypesEqual<WIDECHAR, UCS2CHAR>::Value),  "WIDECHAR and CHAR16 should be different types.");
+#if !PLATFORM_UCS2CHAR_IS_UTF16CHAR
+	// We want these types to be different, because we want to be able to determine whether an encoding
+	// is fixed-width (UCS2CHAR) or variable-width (UTF16CHAR) at compile time.
+	static_assert((!TAreTypesEqual<UCS2CHAR, UTF16CHAR>::Value), "UCS2CHAR and UTF16CHAR should be different types.");
+#else
+	// We don't want these types to be equal, but while this macro exists, they ought to be equal
+	static_assert(TAreTypesEqual<UCS2CHAR, UTF16CHAR>::Value, "UCS2CHAR and UTF16CHAR are expected to be the same type.");
+#endif
+	static_assert((!TAreTypesEqual<ANSICHAR, UCS2CHAR>::Value),  "ANSICHAR and UCS2CHAR should be different types.");
+	static_assert((!TAreTypesEqual<WIDECHAR, UCS2CHAR>::Value),  "WIDECHAR and UCS2CHAR should be different types.");
 	static_assert((TAreTypesEqual<TCHAR, ANSICHAR>::Value == true || TAreTypesEqual<TCHAR, WIDECHAR>::Value == true), "TCHAR should either be ANSICHAR or WIDECHAR.");
 
 	static_assert(sizeof(uint8) == 1, "uint8 type size test failed.");
diff --git a/Engine/Source/Runtime/Core/Public/Templates/UnrealTypeTraits.h b/Engine/Source/Runtime/Core/Public/Templates/UnrealTypeTraits.h
--- a/Engine/Source/Runtime/Core/Public/Templates/UnrealTypeTraits.h
+++ b/Engine/Source/Runtime/Core/Public/Templates/UnrealTypeTraits.h
@@ -97,11 +97,15 @@ struct TNthTypeFromParameterPack<0, T, OtherTypes...>
 /**
  * TIsCharType
  */
-template<typename T> struct TIsCharType           { enum { Value = false }; };
-template<>           struct TIsCharType<ANSICHAR> { enum { Value = true  }; };
-template<>           struct TIsCharType<UCS2CHAR> { enum { Value = true  }; };
-template<>           struct TIsCharType<WIDECHAR> { enum { Value = true  }; };
-template<>           struct TIsCharType<UTF8CHAR> { enum { Value = true  }; };
+template<typename T> struct TIsCharType            { enum { Value = false }; };
+template<>           struct TIsCharType<ANSICHAR>  { enum { Value = true  }; };
+template<>           struct TIsCharType<UCS2CHAR>  { enum { Value = true  }; };
+#if !PLATFORM_UCS2CHAR_IS_UTF16CHAR
+template<>           struct TIsCharType<UTF16CHAR> { enum { Value = true  }; };
+#endif
+template<>           struct TIsCharType<WIDECHAR>  { enum { Value = true  }; };
+template<>           struct TIsCharType<UTF8CHAR>  { enum { Value = true  }; };
+template<>           struct TIsCharType<UTF32CHAR> { enum { Value = true  }; };
 
 /**
  * TFormatSpecifier, only applies to numeric types
