IsFixedWidthEncoding() constexpr function to replace the TIsFixedWidthEncoding trait (now deprecated).

diff --git a/Engine/Source/Runtime/Core/Public/GenericPlatform/GenericPlatformString.h b/Engine/Source/Runtime/Core/Public/GenericPlatform/GenericPlatformString.h
--- a/Engine/Source/Runtime/Core/Public/GenericPlatform/GenericPlatformString.h
+++ b/Engine/Source/Runtime/Core/Public/GenericPlatform/GenericPlatformString.h
@@ -39,6 +39,24 @@ namespace UE::Core::Private
  */
 struct FGenericPlatformString : public FGenericPlatformStricmp
 {
+	/**
+	 * Tests whether an encoding has fixed-width characters
+	 */
+	template <typename Encoding>
+	static constexpr bool IsFixedWidthEncoding()
+	{
+		static_assert(TIsCharType<Encoding>::Value, "Encoding is not a char type");
+
+		return
+			std::is_same_v<Encoding, ANSICHAR> ||
+			std::is_same_v<Encoding, UCS2CHAR> || // this may not be true when PLATFORM_UCS2CHAR_IS_UTF16CHAR == 1, but this is the legacy behavior
+			std::is_same_v<Encoding, WIDECHAR> || // the UCS2CHAR comment also applies to WIDECHAR
+#if PLATFORM_TCHAR_IS_CHAR16
+			std::is_same_v<Encoding, wchar_t>  || // the UCS2CHAR comment also applies to wchar_t
+#endif
+			std::is_same_v<Encoding, UTF32CHAR>;
+	}
+
 	/**
 	 * Tests whether a particular character is a valid member of its encoding.
 	 *
@@ -101,17 +119,10 @@ struct FGenericPlatformString : public FGenericPlatformStricmp
 		enum { Value = false };
 	};
 
-	template <bool Dummy> struct TIsFixedWidthEncoding_Helper<Dummy, ANSICHAR>  { enum { Value = true }; };
-	template <bool Dummy> struct TIsFixedWidthEncoding_Helper<Dummy, WIDECHAR>  { enum { Value = true }; };
-	template <bool Dummy> struct TIsFixedWidthEncoding_Helper<Dummy, UCS2CHAR>  { enum { Value = true }; };
-	template <bool Dummy> struct TIsFixedWidthEncoding_Helper<Dummy, UTF32CHAR> { enum { Value = true }; };
-#if PLATFORM_TCHAR_IS_CHAR16
-	template <bool Dummy> struct TIsFixedWidthEncoding_Helper<Dummy, wchar_t>   { enum { Value = true }; };
-#endif
-
 	template <typename T>
-	struct TIsFixedWidthEncoding : TIsFixedWidthEncoding_Helper<false, T>
+	struct UE_DEPRECATED(5.0, "TIsFixedWidthEncoding is deprecated, use IsFixedWidthEncoding<T>() instead.") TIsFixedWidthEncoding
 	{
+		enum { Value = IsFixedWidthEncoding<T>() };
 	};
 
 
@@ -126,7 +137,7 @@ struct FGenericPlatformString : public FGenericPlatformStricmp
 	template <typename EncodingA, typename EncodingB>
 	struct UE_DEPRECATED(5.0, "TAreEncodingsCompatible is deprecated, use IsCharEncodingCompatibleWith<SrcEncoding, DestEncoding>() instead.") TAreEncodingsCompatible
 	{
-		enum { Value = TIsFixedWidthEncoding<EncodingA>::Value && TIsFixedWidthEncoding<EncodingB>::Value && sizeof(EncodingA) == sizeof(EncodingB) };
+		enum { Value = IsFixedWidthEncoding<EncodingA>() && IsFixedWidthEncoding<EncodingB>() && sizeof(EncodingA) == sizeof(EncodingB) };
 	};
 
 	/**
@@ -198,7 +209,7 @@ struct FGenericPlatformString : public FGenericPlatformStricmp
 
 			return (DestEncoding*)Memcpy(Dest, Src, SrcSize * sizeof(SourceEncoding)) + SrcSize;
 		}
-		else if constexpr (TIsFixedWidthEncoding<SourceEncoding>::Value && TIsFixedWidthEncoding<DestEncoding>::Value)
+		else if constexpr (IsFixedWidthEncoding<SourceEncoding>() && IsFixedWidthEncoding<DestEncoding>())
 		{
 			const int32 Size = DestSize <= SrcSize ? DestSize : SrcSize;
 			bool bInvalidChars = false;
@@ -242,7 +253,7 @@ struct FGenericPlatformString : public FGenericPlatformStricmp
 	template <typename DestEncoding, typename SourceEncoding>
 	static int32 ConvertedLength(const SourceEncoding* Src, int32 SrcSize)
 	{
-		if constexpr (std::is_same_v<SourceEncoding, DestEncoding> || (TIsFixedWidthEncoding<SourceEncoding>::Value && TIsFixedWidthEncoding<DestEncoding>::Value))
+		if constexpr (std::is_same_v<SourceEncoding, DestEncoding> || (IsFixedWidthEncoding<SourceEncoding>() && IsFixedWidthEncoding<DestEncoding>()))
 		{
 			return SrcSize;
 		}
