Temporary compile fix.

diff --git a/Engine/Source/Runtime/Core/Public/GenericPlatform/GenericPlatformString.h b/Engine/Source/Runtime/Core/Public/GenericPlatform/GenericPlatformString.h
--- a/Engine/Source/Runtime/Core/Public/GenericPlatform/GenericPlatformString.h
+++ b/Engine/Source/Runtime/Core/Public/GenericPlatform/GenericPlatformString.h
@@ -45,7 +45,7 @@ struct FGenericPlatformString : public FGenericPlatformStricmp
 	template <typename Encoding>
 	static constexpr bool IsFixedWidthEncoding()
 	{
-		static_assert(TIsCharType<Encoding>::Value, "Encoding is not a char type");
+//		static_assert(TIsCharType<Encoding>::Value, "Encoding is not a char type");
 
 		return
 			std::is_same_v<Encoding, ANSICHAR> ||
