Added UE_  prefix to MTAccesDetector macros.

#rnx

diff --git a/Engine/Source/Runtime/Core/Private/Misc/MTAccessDetector.cpp b/Engine/Source/Runtime/Core/Private/Misc/MTAccessDetector.cpp
--- a/Engine/Source/Runtime/Core/Private/Misc/MTAccessDetector.cpp
+++ b/Engine/Source/Runtime/Core/Private/Misc/MTAccessDetector.cpp
@@ -17,7 +17,7 @@ bool FRWAccessDetector_ConcurentReadTest::RunTest(const FString& Parameters)
 {
 	bool bReading1 = false;
 	bool bReading2 = false;
-	MT_DECLARE_RW_ACCESS_DETECTOR(MTAccessDetector);
+	UE_MT_DECLARE_RW_ACCESS_DETECTOR(MTAccessDetector);
 
 	bool Success = false;
 	FThread AccessThread = FThread(TEXT("AccessThread"), [&]()
@@ -42,7 +42,7 @@ bool FRWAccessDetector_ConcurentWriteTest::RunTest(const FString& Parameters)
 {
 	bool bWriting1 = false;
 	bool bWriting2 = false;
-	MT_DECLARE_RW_ACCESS_DETECTOR(MTAccessDetector);
+	UE_MT_DECLARE_RW_ACCESS_DETECTOR(MTAccessDetector);
 
 	bool Success = false;
 	FThread AccessThread = FThread(TEXT("AccessThread"), [&]()
@@ -67,7 +67,7 @@ bool FRWAccessDetector_ConcurentReadWriteTest::RunTest(const FString& Parameters
 {
 	bool bReading = false;
 	bool bWriting = false;
-	MT_DECLARE_RW_ACCESS_DETECTOR(MTAccessDetector);
+	UE_MT_DECLARE_RW_ACCESS_DETECTOR(MTAccessDetector);
 
 	bool Success = false;
 	FThread AccessThread = FThread(TEXT("AccessThread"), [&]()
@@ -90,7 +90,7 @@ bool FRWAccessDetector_ConcurentReadWriteTest::RunTest(const FString& Parameters
 IMPLEMENT_SIMPLE_AUTOMATION_TEST(FRWAccessDetector_WriteReentrance, "System.Core.Misc.MTAccessDetector.WriteReentrance", EAutomationTestFlags::ApplicationContextMask | EAutomationTestFlags::EngineFilter)
 bool FRWAccessDetector_WriteReentrance::RunTest(const FString& Parameters)
 {
-	MT_DECLARE_RW_ACCESS_DETECTOR(MTAccessDetector);
+	UE_MT_DECLARE_RW_ACCESS_DETECTOR(MTAccessDetector);
 
 	bool Success = false;
 	Success |= TestTrue(TEXT("Aquiring 1st Write Access"), MTAccessDetector.AcquireWriteAccess());
@@ -104,7 +104,7 @@ bool FRWAccessDetector_WriteReentrance::RunTest(const FString& Parameters)
 IMPLEMENT_SIMPLE_AUTOMATION_TEST(FRWAccessDetector_ReadReentrance, "System.Core.Misc.MTAccessDetector.ReadReentrance", EAutomationTestFlags::ApplicationContextMask | EAutomationTestFlags::EngineFilter)
 bool FRWAccessDetector_ReadReentrance::RunTest(const FString& Parameters)
 {
-	MT_DECLARE_RW_ACCESS_DETECTOR(MTAccessDetector);
+	UE_MT_DECLARE_RW_ACCESS_DETECTOR(MTAccessDetector);
 
 	bool Success = false;
 	Success |= TestTrue(TEXT("Aquiring Write Access"), MTAccessDetector.AcquireWriteAccess());
@@ -123,7 +123,7 @@ bool FRWAccessDetector_RecursiveConcurentReadTest::RunTest(const FString& Parame
 {
 	bool bReading1 = false;
 	bool bReading2 = false;
-	MT_DECLARE_RW_RECURSIVE_ACCESS_DETECTOR(MTAccessDetector);
+	UE_MT_DECLARE_RW_RECURSIVE_ACCESS_DETECTOR(MTAccessDetector);
 
 	bool Success = false;
 	FThread AccessThread = FThread(TEXT("AccessThread"), [&]()
@@ -148,7 +148,7 @@ bool FRWAccessDetector_RecursiveConcurentWriteTest::RunTest(const FString& Param
 {
 	bool bWriting1 = false;
 	bool bWriting2 = false;
-	MT_DECLARE_RW_RECURSIVE_ACCESS_DETECTOR(MTAccessDetector);
+	UE_MT_DECLARE_RW_RECURSIVE_ACCESS_DETECTOR(MTAccessDetector);
 
 	bool Success = false;
 	FThread AccessThread = FThread(TEXT("AccessThread"), [&]()
@@ -173,7 +173,7 @@ bool FRWAccessDetector_RecursiveConcurentReadWriteTest::RunTest(const FString& P
 {
 	bool bReading = false;
 	bool bWriting = false;
-	MT_DECLARE_RW_RECURSIVE_ACCESS_DETECTOR(MTAccessDetector);
+	UE_MT_DECLARE_RW_RECURSIVE_ACCESS_DETECTOR(MTAccessDetector);
 
 	bool Success = false;
 	FThread AccessThread = FThread(TEXT("AccessThread"), [&]()
@@ -196,7 +196,7 @@ bool FRWAccessDetector_RecursiveConcurentReadWriteTest::RunTest(const FString& P
 IMPLEMENT_SIMPLE_AUTOMATION_TEST(FRWAccessDetector_RecursiveWriteReentrance, "System.Core.Misc.MTAccessDetector.RecursiveWriteReentrance", EAutomationTestFlags::ApplicationContextMask | EAutomationTestFlags::EngineFilter)
 bool FRWAccessDetector_RecursiveWriteReentrance::RunTest(const FString& Parameters)
 {
-	MT_DECLARE_RW_RECURSIVE_ACCESS_DETECTOR(MTAccessDetector);
+	UE_MT_DECLARE_RW_RECURSIVE_ACCESS_DETECTOR(MTAccessDetector);
 
 	bool Success = false;
 	Success |= TestTrue(TEXT("Aquiring 1st Write Access"), MTAccessDetector.AcquireWriteAccess());
@@ -210,7 +210,7 @@ bool FRWAccessDetector_RecursiveWriteReentrance::RunTest(const FString& Paramete
 IMPLEMENT_SIMPLE_AUTOMATION_TEST(FRWAccessDetector_RecursiveReadReentrance, "System.Core.Misc.MTAccessDetector.RecursiveReadReentrance", EAutomationTestFlags::ApplicationContextMask | EAutomationTestFlags::EngineFilter)
 bool FRWAccessDetector_RecursiveReadReentrance::RunTest(const FString& Parameters)
 {
-	MT_DECLARE_RW_RECURSIVE_ACCESS_DETECTOR(MTAccessDetector);
+	UE_MT_DECLARE_RW_RECURSIVE_ACCESS_DETECTOR(MTAccessDetector);
 
 	bool Success = false;
 	Success |= TestTrue(TEXT("Aquiring Write Access"), MTAccessDetector.AcquireWriteAccess());
@@ -229,7 +229,7 @@ bool FRWAccessDetector_FullyRecursiveConcurentReadTest::RunTest(const FString& P
 {
 	bool bReading1 = false;
 	bool bReading2 = false;
-	MT_DECLARE_RW_FULLY_RECURSIVE_ACCESS_DETECTOR(MTAccessDetector);
+	UE_MT_DECLARE_RW_FULLY_RECURSIVE_ACCESS_DETECTOR(MTAccessDetector);
 
 	bool Success = false;
 	FThread AccessThread = FThread(TEXT("AccessThread"), [&]()
@@ -254,7 +254,7 @@ bool FRWAccessDetector_FullyRecursiveConcurentWriteTest::RunTest(const FString&
 {
 	bool bWriting1 = false;
 	bool bWriting2 = false;
-	MT_DECLARE_RW_FULLY_RECURSIVE_ACCESS_DETECTOR(MTAccessDetector);
+	UE_MT_DECLARE_RW_FULLY_RECURSIVE_ACCESS_DETECTOR(MTAccessDetector);
 
 	bool Success = false;
 	FThread AccessThread = FThread(TEXT("AccessThread"), [&]()
@@ -279,7 +279,7 @@ bool FRWAccessDetector_FullyRecursiveConcurentReadWriteTest::RunTest(const FStri
 {
 	bool bReading = false;
 	bool bWriting = false;
-	MT_DECLARE_RW_FULLY_RECURSIVE_ACCESS_DETECTOR(MTAccessDetector);
+	UE_MT_DECLARE_RW_FULLY_RECURSIVE_ACCESS_DETECTOR(MTAccessDetector);
 
 	bool Success = false;
 	FThread AccessThread = FThread(TEXT("AccessThread"), [&]()
@@ -302,7 +302,7 @@ bool FRWAccessDetector_FullyRecursiveConcurentReadWriteTest::RunTest(const FStri
 IMPLEMENT_SIMPLE_AUTOMATION_TEST(FRWAccessDetector_FullyRecursiveWriteReentrance, "System.Core.Misc.MTAccessDetector.FullyRecursiveWriteReentrance", EAutomationTestFlags::ApplicationContextMask | EAutomationTestFlags::EngineFilter)
 bool FRWAccessDetector_FullyRecursiveWriteReentrance::RunTest(const FString& Parameters)
 {
-	MT_DECLARE_RW_FULLY_RECURSIVE_ACCESS_DETECTOR(MTAccessDetector);
+	UE_MT_DECLARE_RW_FULLY_RECURSIVE_ACCESS_DETECTOR(MTAccessDetector);
 
 	bool Success = false;
 	Success |= TestTrue(TEXT("Aquiring 1st Write Access"), MTAccessDetector.AcquireWriteAccess());
@@ -316,7 +316,7 @@ bool FRWAccessDetector_FullyRecursiveWriteReentrance::RunTest(const FString& Par
 IMPLEMENT_SIMPLE_AUTOMATION_TEST(FRWAccessDetector_FullyRecursiveReadReentrance, "System.Core.Misc.MTAccessDetector.FullyRecursiveReadReentrance", EAutomationTestFlags::ApplicationContextMask | EAutomationTestFlags::EngineFilter)
 bool FRWAccessDetector_FullyRecursiveReadReentrance::RunTest(const FString& Parameters)
 {
-	MT_DECLARE_RW_FULLY_RECURSIVE_ACCESS_DETECTOR(MTAccessDetector);
+	UE_MT_DECLARE_RW_FULLY_RECURSIVE_ACCESS_DETECTOR(MTAccessDetector);
 
 	bool Success = false;
 	Success |= TestTrue(TEXT("Aquiring Write Access"), MTAccessDetector.AcquireWriteAccess());
diff --git a/Engine/Source/Runtime/Core/Public/Misc/MTAccessDetector.h b/Engine/Source/Runtime/Core/Public/Misc/MTAccessDetector.h
--- a/Engine/Source/Runtime/Core/Public/Misc/MTAccessDetector.h
+++ b/Engine/Source/Runtime/Core/Public/Misc/MTAccessDetector.h
@@ -234,30 +234,30 @@ FORCEINLINE TScopedWriterDetector<RWAccessDetector> MakeScopedWriterAccessDetect
 	return TScopedWriterDetector<RWAccessDetector>(InAccessDetector);
 }
 
-#define MT_DECLARE_RW_ACCESS_DETECTOR(AccessDetector) FRWAccessDetector AccessDetector;
-#define MT_DECLARE_RW_RECURSIVE_ACCESS_DETECTOR(AccessDetector) FRWRecursiveAccessDetector AccessDetector;
-#define MT_DECLARE_RW_FULLY_RECURSIVE_ACCESS_DETECTOR(AccessDetector) FRWFullyRecursiveAccessDetector AccessDetector;
+#define UE_MT_DECLARE_RW_ACCESS_DETECTOR(AccessDetector) FRWAccessDetector AccessDetector;
+#define UE_MT_DECLARE_RW_RECURSIVE_ACCESS_DETECTOR(AccessDetector) FRWRecursiveAccessDetector AccessDetector;
+#define UE_MT_DECLARE_RW_FULLY_RECURSIVE_ACCESS_DETECTOR(AccessDetector) FRWFullyRecursiveAccessDetector AccessDetector;
 
-#define MT_SCOPED_READ_ACCESS(AccessDetector) const FBaseScopedAccessDetector& PREPROCESSOR_JOIN(ScopedMTAccessDetector_,__LINE__) = MakeScopedReaderAccessDetector(AccessDetector);
-#define MT_SCOPED_WRITE_ACCESS(AccessDetector) const FBaseScopedAccessDetector& PREPROCESSOR_JOIN(ScopedMTAccessDetector_,__LINE__) = MakeScopedWriterAccessDetector(AccessDetector);
+#define UE_MT_SCOPED_READ_ACCESS(AccessDetector) const FBaseScopedAccessDetector& PREPROCESSOR_JOIN(ScopedMTAccessDetector_,__LINE__) = MakeScopedReaderAccessDetector(AccessDetector);
+#define UE_MT_SCOPED_WRITE_ACCESS(AccessDetector) const FBaseScopedAccessDetector& PREPROCESSOR_JOIN(ScopedMTAccessDetector_,__LINE__) = MakeScopedWriterAccessDetector(AccessDetector);
 
-#define MT_ACQUIRE_READ_ACCESS(AccessDetector) (AccessDetector).AcquireReadAccess();
-#define MT_RELEASE_READ_ACCESS(AccessDetector) (AccessDetector).ReleaseReadAccess();
-#define MT_ACQUIRE_WRITE_ACCESS(AccessDetector) (AccessDetector).AcquireWriteAccess();
-#define MT_RELEASE_WRITE_ACCESS(AccessDetector) (AccessDetector).ReleaseWriteAccess();
+#define UE_MT_ACQUIRE_READ_ACCESS(AccessDetector) (AccessDetector).AcquireReadAccess();
+#define UE_MT_RELEASE_READ_ACCESS(AccessDetector) (AccessDetector).ReleaseReadAccess();
+#define UE_MT_ACQUIRE_WRITE_ACCESS(AccessDetector) (AccessDetector).AcquireWriteAccess();
+#define UE_MT_RELEASE_WRITE_ACCESS(AccessDetector) (AccessDetector).ReleaseWriteAccess();
 
 #else // ENABLE_MT_DETECTOR
 
-#define MT_DECLARE_RW_ACCESS_DETECTOR(AccessDetector)
-#define MT_DECLARE_RW_RECURSIVE_ACCESS_DETECTOR(AccessDetector)
-#define MT_DECLARE_RW_FULLY_RECURSIVE_ACCESS_DETECTOR(AccessDetector)
+#define UE_MT_DECLARE_RW_ACCESS_DETECTOR(AccessDetector)
+#define UE_MT_DECLARE_RW_RECURSIVE_ACCESS_DETECTOR(AccessDetector)
+#define UE_MT_DECLARE_RW_FULLY_RECURSIVE_ACCESS_DETECTOR(AccessDetector)
 
-#define MT_SCOPED_READ_ACCESS(AccessDetector) 
-#define MT_SCOPED_WRITE_ACCESS(AccessDetector)
+#define UE_MT_SCOPED_READ_ACCESS(AccessDetector) 
+#define UE_MT_SCOPED_WRITE_ACCESS(AccessDetector)
 
-#define MT_ACQUIRE_READ_ACCESS(AccessDetector)
-#define MT_RELEASE_READ_ACCESS(AccessDetector)
-#define MT_ACQUIRE_WRITE_ACCESS(AccessDetector)
-#define MT_RELEASE_WRITE_ACCESS(AccessDetector)
+#define UE_MT_ACQUIRE_READ_ACCESS(AccessDetector)
+#define UE_MT_RELEASE_READ_ACCESS(AccessDetector)
+#define UE_MT_ACQUIRE_WRITE_ACCESS(AccessDetector)
+#define UE_MT_RELEASE_WRITE_ACCESS(AccessDetector)
 
 #endif // ENABLE_MT_DETECTOR
