Add thread priority defaults for Task and TaskBP threads

diff --git a/Engine/Source/Runtime/Core/Private/Async/TaskGraph.cpp b/Engine/Source/Runtime/Core/Private/Async/TaskGraph.cpp
--- a/Engine/Source/Runtime/Core/Private/Async/TaskGraph.cpp
+++ b/Engine/Source/Runtime/Core/Private/Async/TaskGraph.cpp
@@ -1873,10 +1873,10 @@ public:
 			int32 NumBackgroundWorkers = FMath::Max(1, NumWorkerThreads - FMath::Min<int>(GNumForegroundWorkers, NumWorkerThreads));
 			int32 NumForegroundWorkers =  FMath::Max(1, NumWorkerThreads - NumBackgroundWorkers);
 
-			LowLevelTasks::FScheduler::Get().StartWorkers(NumForegroundWorkers, NumBackgroundWorkers, FForkProcessHelper::IsForkedMultithreadInstance(), EThreadPriority::TPri_Normal, EThreadPriority::TPri_BelowNormal);
+			LowLevelTasks::FScheduler::Get().StartWorkers(NumForegroundWorkers, NumBackgroundWorkers, FForkProcessHelper::IsForkedMultithreadInstance(), FPlatformAffinity::GetTaskThreadPriority(), FPlatformAffinity::GetTaskBPThreadPriority());
 			if (!GDisableReserveWorkers)
 			{
-				LowLevelTasks::FReserveScheduler::Get().StartWorkers(LowLevelTasks::FScheduler::Get(), NumForegroundWorkers + NumBackgroundWorkers, FForkProcessHelper::IsForkedMultithreadInstance(), EThreadPriority::TPri_BelowNormal);
+				LowLevelTasks::FReserveScheduler::Get().StartWorkers(LowLevelTasks::FScheduler::Get(), NumForegroundWorkers + NumBackgroundWorkers, FForkProcessHelper::IsForkedMultithreadInstance(), FPlatformAffinity::GetTaskBPThreadPriority());
 			}
 
 			NumNamedThreads = ENamedThreads::ActualRenderingThread + 1;
@@ -1956,12 +1956,12 @@ public:
 			int32 NumWorkers =  FMath::Max(1, NumWorkerThreads - NumBackgroundWorkers);
 
 			LowLevelTasks::FScheduler::Get().StopWorkers();
-			LowLevelTasks::FScheduler::Get().StartWorkers(NumWorkers, NumBackgroundWorkers, FForkProcessHelper::IsForkedMultithreadInstance(), Pri, EThreadPriority::TPri_BelowNormal);
+			LowLevelTasks::FScheduler::Get().StartWorkers(NumWorkers, NumBackgroundWorkers, FForkProcessHelper::IsForkedMultithreadInstance(), Pri, FPlatformAffinity::GetTaskBPThreadPriority());
 
 			if (!GDisableReserveWorkers)
 			{
 				LowLevelTasks::FReserveScheduler::Get().StopWorkers();
-				LowLevelTasks::FReserveScheduler::Get().StartWorkers(LowLevelTasks::FScheduler::Get(), NumWorkers + NumBackgroundWorkers, FForkProcessHelper::IsForkedMultithreadInstance(), EThreadPriority::TPri_BelowNormal);
+				LowLevelTasks::FReserveScheduler::Get().StartWorkers(LowLevelTasks::FScheduler::Get(), NumWorkers + NumBackgroundWorkers, FForkProcessHelper::IsForkedMultithreadInstance(), FPlatformAffinity::GetTaskBPThreadPriority());
 			}
 		}
 	}
diff --git a/Engine/Source/Runtime/Core/Public/GenericPlatform/GenericPlatformAffinity.h b/Engine/Source/Runtime/Core/Public/GenericPlatform/GenericPlatformAffinity.h
--- a/Engine/Source/Runtime/Core/Public/GenericPlatform/GenericPlatformAffinity.h
+++ b/Engine/Source/Runtime/Core/Public/GenericPlatform/GenericPlatformAffinity.h
@@ -149,6 +149,16 @@ public:
 	{
 		return EThreadCreateFlags::None;
 	}
+
+	static EThreadPriority GetTaskThreadPriority()
+	{
+		return TPri_SlightlyBelowNormal;
+	}
+
+	static EThreadPriority GetTaskBPThreadPriority()
+	{
+		return TPri_BelowNormal;
+	}
 };
 
 
diff --git a/Engine/Source/Runtime/Engine/Private/UnrealEngine.cpp b/Engine/Source/Runtime/Engine/Private/UnrealEngine.cpp
--- a/Engine/Source/Runtime/Engine/Private/UnrealEngine.cpp
+++ b/Engine/Source/Runtime/Engine/Private/UnrealEngine.cpp
@@ -16233,8 +16233,8 @@ void GetDefaultThreadConfigs()
 	RenderThreadConfig.Affinity = FPlatformAffinity::GetRenderingThreadMask();
 	RHIThreadConfig.Affinity = FPlatformAffinity::GetRHIThreadMask();
 
-	TaskThreadConfig.Priority = EThreadPriority::TPri_Normal;
-	TaskBPThreadConfig.Priority = EThreadPriority::TPri_BelowNormal;
+	TaskThreadConfig.Priority = FPlatformAffinity::GetTaskThreadPriority();
+	TaskBPThreadConfig.Priority = FPlatformAffinity::GetTaskBPThreadPriority();
 	GameThreadConfig.Priority = FPlatformAffinity::GetGameThreadPriority();
 	RenderThreadConfig.Priority = FPlatformAffinity::GetRenderingThreadPriority();
 	RHIThreadConfig.Priority = FPlatformAffinity::GetRHIThreadPriority();
