Disable thread check in FTicker::AddTicker on Android platform

diff --git a/Engine/Source/Runtime/Core/Private/Containers/Ticker.cpp b/Engine/Source/Runtime/Core/Private/Containers/Ticker.cpp
--- a/Engine/Source/Runtime/Core/Private/Containers/Ticker.cpp
+++ b/Engine/Source/Runtime/Core/Private/Containers/Ticker.cpp
@@ -163,8 +163,10 @@ FTicker::~FTicker()
 
 FDelegateHandle FTicker::AddTicker(const FTickerDelegate& InDelegate, float InDelay)
 {
+#if !PLATFORM_ANDROID	//temp work around to prevent it crashing on Android, since it was "always" working before
 	// not currently threadsafe
 	check(IsInGameThread());
+#endif
 
 	// We can add elements safely even during tick.
 	Elements.Emplace(CurrentTime + InDelay, InDelay, InDelegate);
