Replaced FTicker with a thread safe alternative FTSTicker for Crash Report

diff --git a/Engine/Source/Runtime/Core/Private/Containers/Ticker.cpp b/Engine/Source/Runtime/Core/Private/Containers/Ticker.cpp
--- a/Engine/Source/Runtime/Core/Private/Containers/Ticker.cpp
+++ b/Engine/Source/Runtime/Core/Private/Containers/Ticker.cpp
@@ -163,11 +163,6 @@ FTicker::~FTicker()
 
 FDelegateHandle FTicker::AddTicker(const FTickerDelegate& InDelegate, float InDelay)
 {
-#if !PLATFORM_ANDROID	//temp work around to prevent it crashing on Android, since it was "always" working before
-	// not currently threadsafe
-	check(IsInGameThread());
-#endif
-
 	// We can add elements safely even during tick.
 	Elements.Emplace(CurrentTime + InDelay, InDelay, InDelegate);
 	// @todo this needs a unique handle for each add call to allow you to register the same delegate twice with a different delay safely.
diff --git a/Engine/Source/Runtime/CrashReportCore/Private/CrashReportCoreUnattended.cpp b/Engine/Source/Runtime/CrashReportCore/Private/CrashReportCoreUnattended.cpp
--- a/Engine/Source/Runtime/CrashReportCore/Private/CrashReportCoreUnattended.cpp
+++ b/Engine/Source/Runtime/CrashReportCore/Private/CrashReportCoreUnattended.cpp
@@ -73,5 +73,5 @@ bool FCrashReportCoreUnattended::Tick(float UnusedDeltaTime)
 
 void FCrashReportCoreUnattended::StartTicker()
 {
-	FTicker::GetCoreTicker().AddTicker(FTickerDelegate::CreateRaw(this, &FCrashReportCoreUnattended::Tick), 1.f);
+	FTSTicker::GetCoreTicker().AddTicker(FTickerDelegate::CreateRaw(this, &FCrashReportCoreUnattended::Tick), 1.f);
 }
