TIsCharType moved to its own Traits header so that it can be included separately from others.

diff --git a/Engine/Source/Runtime/Core/Public/Containers/UnrealString.h b/Engine/Source/Runtime/Core/Public/Containers/UnrealString.h
--- a/Engine/Source/Runtime/Core/Public/Containers/UnrealString.h
+++ b/Engine/Source/Runtime/Core/Public/Containers/UnrealString.h
@@ -25,6 +25,7 @@
 #include "Templates/IsArrayOrRefOfType.h"
 #include "Templates/TypeHash.h"
 #include "Templates/IsFloatingPoint.h"
+#include "Traits/IsCharType.h"
 
 struct FStringFormatArg;
 template<typename InKeyType,typename InValueType,typename SetAllocator ,typename KeyFuncs > class TMap;
diff --git a/Engine/Source/Runtime/Core/Public/GenericPlatform/GenericPlatformString.h b/Engine/Source/Runtime/Core/Public/GenericPlatform/GenericPlatformString.h
--- a/Engine/Source/Runtime/Core/Public/GenericPlatform/GenericPlatformString.h
+++ b/Engine/Source/Runtime/Core/Public/GenericPlatform/GenericPlatformString.h
@@ -4,6 +4,7 @@
 
 #include "CoreTypes.h"
 #include "Templates/EnableIf.h"
+#include "Traits/IsCharType.h"
 #include "GenericPlatform/GenericPlatformStricmp.h"
 #include <type_traits>
 
@@ -45,7 +46,7 @@ struct FGenericPlatformString : public FGenericPlatformStricmp
 	template <typename Encoding>
 	static constexpr bool IsFixedWidthEncoding()
 	{
-//		static_assert(TIsCharType<Encoding>::Value, "Encoding is not a char type");
+		static_assert(TIsCharType<Encoding>::Value, "Encoding is not a char type");
 
 		return
 			std::is_same_v<Encoding, ANSICHAR> ||
diff --git a/Engine/Source/Runtime/Core/Public/ProfilingDebugging/FormatArgsTrace.h b/Engine/Source/Runtime/Core/Public/ProfilingDebugging/FormatArgsTrace.h
--- a/Engine/Source/Runtime/Core/Public/ProfilingDebugging/FormatArgsTrace.h
+++ b/Engine/Source/Runtime/Core/Public/ProfilingDebugging/FormatArgsTrace.h
@@ -7,6 +7,7 @@
 #include "Templates/IsFloatingPoint.h"
 #include "Templates/UnrealTypeTraits.h"
 #include "Templates/UnrealTemplate.h"
+#include "Traits/IsCharType.h"
 
 struct FFormatArgsTrace
 {
diff --git a/Engine/Source/Runtime/Core/Public/Serialization/MemoryImage.h b/Engine/Source/Runtime/Core/Public/Serialization/MemoryImage.h
--- a/Engine/Source/Runtime/Core/Public/Serialization/MemoryImage.h
+++ b/Engine/Source/Runtime/Core/Public/Serialization/MemoryImage.h
@@ -12,6 +12,7 @@
 #include "Serialization/MemoryLayout.h"
 #include "Serialization/MemoryImageWriter.h"
 #include "Templates/RefCounting.h"
+#include "Traits/IsCharType.h"
 
 #if defined(WITH_RTTI) || defined(_CPPRTTI) || defined(__GXX_RTTI) || WITH_EDITOR
 #include <typeinfo>
diff --git a/Engine/Source/Runtime/Core/Public/Templates/UnrealTypeTraits.h b/Engine/Source/Runtime/Core/Public/Templates/UnrealTypeTraits.h
--- a/Engine/Source/Runtime/Core/Public/Templates/UnrealTypeTraits.h
+++ b/Engine/Source/Runtime/Core/Public/Templates/UnrealTypeTraits.h
@@ -94,19 +94,6 @@ struct TNthTypeFromParameterPack<0, T, OtherTypes...>
 	using Type = T;
 };
 
-/**
- * TIsCharType
- */
-template<typename T> struct TIsCharType            { enum { Value = false }; };
-template<>           struct TIsCharType<ANSICHAR>  { enum { Value = true  }; };
-template<>           struct TIsCharType<UCS2CHAR>  { enum { Value = true  }; };
-#if !PLATFORM_UCS2CHAR_IS_UTF16CHAR
-template<>           struct TIsCharType<UTF16CHAR> { enum { Value = true  }; };
-#endif
-template<>           struct TIsCharType<WIDECHAR>  { enum { Value = true  }; };
-template<>           struct TIsCharType<UTF8CHAR>  { enum { Value = true  }; };
-template<>           struct TIsCharType<UTF32CHAR> { enum { Value = true  }; };
-
 /**
  * TFormatSpecifier, only applies to numeric types
  */
diff --git a/Engine/Source/Runtime/Core/Public/Traits/IsCharType.h b/Engine/Source/Runtime/Core/Public/Traits/IsCharType.h
new file mode 100644
--- /dev/null
+++ b/Engine/Source/Runtime/Core/Public/Traits/IsCharType.h
@@ -0,0 +1,18 @@
+// Copyright Epic Games, Inc. All Rights Reserved.
+
+#pragma once
+
+#include "CoreTypes.h"
+
+/**
+ * Type trait which tests if a type is a char type.
+ */
+template<typename T> struct TIsCharType            { enum { Value = false }; };
+template<>           struct TIsCharType<ANSICHAR>  { enum { Value = true  }; };
+template<>           struct TIsCharType<UCS2CHAR>  { enum { Value = true  }; };
+#if !PLATFORM_UCS2CHAR_IS_UTF16CHAR
+template<>           struct TIsCharType<UTF16CHAR> { enum { Value = true  }; };
+#endif
+template<>           struct TIsCharType<WIDECHAR>  { enum { Value = true  }; };
+template<>           struct TIsCharType<UTF8CHAR>  { enum { Value = true  }; };
+template<>           struct TIsCharType<UTF32CHAR> { enum { Value = true  }; };
