Enabled memory tracking with callstacks for PS4/5

Implementation details:
- PS4 & PS5 uses dwarf symbol format and uses Syms symbol resolver.
- To resolve symbols the path to folder where .self file is built currently must be specified in UE_INSIGHTS_SYMBOL_PATH env variable for Insights.
- Multiple paths can be separated by ; in this variable.
- Build for PS5 does not seem to have PLATFORM_PS5 define, I used defined(__PS5__).
- PS5 runtime collects and sends callstacks, but Syms resolver does not support dwarf v5 format yet, which is used on PS5 toolchain.

diff --git a/Engine/Source/Developer/TraceServices/Private/Analyzers/ModuleAnalysis.cpp b/Engine/Source/Developer/TraceServices/Private/Analyzers/ModuleAnalysis.cpp
--- a/Engine/Source/Developer/TraceServices/Private/Analyzers/ModuleAnalysis.cpp
+++ b/Engine/Source/Developer/TraceServices/Private/Analyzers/ModuleAnalysis.cpp
@@ -70,7 +70,7 @@ bool FModuleAnalyzer::OnEvent(uint16 RouteId, EStyle Style, const FOnEventContex
 				FStringView ModuleName;
 				if (Context.EventData.GetString("Name", ModuleName))
 				{
-					const uint64 Base = GetRealBaseAddress(Context.EventData.GetValue<uint32>("Base"));
+					const uint64 Base = GetBaseAddress(Context.EventData);
 					const uint32 Size = Context.EventData.GetValue<uint32>("Size");
 					const TArrayReader<uint8>& ImageId = Context.EventData.GetArray<uint8>("ImageId");
 
@@ -82,7 +82,7 @@ bool FModuleAnalyzer::OnEvent(uint16 RouteId, EStyle Style, const FOnEventContex
 		case RouteId_ModuleUnload:
 			if (Provider != nullptr)
 			{
-				const uint32 Base = Context.EventData.GetValue<uint32>("Base");
+				const uint64 Base = GetBaseAddress(Context.EventData);
 				Provider->OnModuleUnload(Base);
 			}
 			break;
@@ -91,10 +91,19 @@ bool FModuleAnalyzer::OnEvent(uint16 RouteId, EStyle Style, const FOnEventContex
 }
 
 ////////////////////////////////////////////////////////////////////////////////
-uint64 FModuleAnalyzer::GetRealBaseAddress(uint32 EventBase) const
+uint64	FModuleAnalyzer::GetBaseAddress(const FEventData& EventData) const
 {
-	// Unshift base address according to capture alignment (e.g. Windows 4k aligned)
-	return uint64(EventBase) << ModuleBaseShift;
+	if (ModuleBaseShift == 0)
+	{
+		// With 0 base shift assume newer trace event is used with 64-bit base address
+		return EventData.GetValue<uint64>("Base");
+	}
+	else
+	{
+		// Unshift base address according to capture alignment (e.g. Windows 4k aligned)
+		uint64 Base = uint64(EventData.GetValue<uint32>("Base"));
+		return Base << ModuleBaseShift;
+	}
 }
 
 ////////////////////////////////////////////////////////////////////////////////
diff --git a/Engine/Source/Developer/TraceServices/Private/Analyzers/ModuleAnalysis.h b/Engine/Source/Developer/TraceServices/Private/Analyzers/ModuleAnalysis.h
--- a/Engine/Source/Developer/TraceServices/Private/Analyzers/ModuleAnalysis.h
+++ b/Engine/Source/Developer/TraceServices/Private/Analyzers/ModuleAnalysis.h
@@ -19,7 +19,7 @@ public:
 	virtual bool 	OnEvent(uint16 RouteId, EStyle Style, const FOnEventContext& Context) override;
 
 private:
-	uint64			GetRealBaseAddress(uint32 EventBase) const;
+	uint64			GetBaseAddress(const FEventData& EventData) const;
 
 	IAnalysisSession& Session;
 	IModuleAnalysisProvider* Provider;
diff --git a/Engine/Source/Developer/TraceServices/Private/Model/ModuleProvider.cpp b/Engine/Source/Developer/TraceServices/Private/Model/ModuleProvider.cpp
--- a/Engine/Source/Developer/TraceServices/Private/Model/ModuleProvider.cpp
+++ b/Engine/Source/Developer/TraceServices/Private/Model/ModuleProvider.cpp
@@ -246,20 +246,18 @@ void TModuleProvider<SymbolResolverType>::LoadSymbolsFromCache(IAnalysisCache& C
 IModuleAnalysisProvider* CreateModuleProvider(IAnalysisSession& InSession, const FAnsiStringView& InSymbolFormat)
 {
 	IModuleAnalysisProvider* Provider(nullptr);
-#if PLATFORM_WINDOWS
-#if USE_SYMSLIB
-	if (!Provider && InSymbolFormat.Equals("pdb"))
+#if PLATFORM_WINDOWS && USE_SYMSLIB
+	if (!Provider && (InSymbolFormat.Equals("pdb") || InSymbolFormat.Equals("dwarf")))
 	{
 		Provider = new TModuleProvider<FSymslibResolver>(InSession);
 	}
-#endif // USE_SYMSLIB
-#if USE_DBGHELP
+#endif // PLATFORM_WINDOWS && USE_SYMSLIB
+#if PLATFORM_WINDOWS && USE_DBGHELP
 	if (!Provider && InSymbolFormat.Equals("pdb"))
 	{
 		Provider = new TModuleProvider<FDbgHelpResolver>(InSession);
 	}
-#endif // USE_DBGHELP
-#endif //PLATFORM_WINDOWS
+#endif // PLATFORM_WINDOWS && USE_DBGHELP
 	return Provider;
 }
 
diff --git a/Engine/Source/Developer/TraceServices/Private/Model/SymslibResolver.cpp b/Engine/Source/Developer/TraceServices/Private/Model/SymslibResolver.cpp
--- a/Engine/Source/Developer/TraceServices/Private/Model/SymslibResolver.cpp
+++ b/Engine/Source/Developer/TraceServices/Private/Model/SymslibResolver.cpp
@@ -299,6 +299,36 @@ void FSymslibResolver::LoadModuleTracked(FModuleEntry* Module)
 	}
 }
 
+FString FSymslibResolver::FindSymbolFile(const FString& Path)
+{
+	IPlatformFile* PlatformFile = &FPlatformFileManager::Get().GetPlatformFile();
+
+	if (PlatformFile->FileExists(*Path))
+	{
+		return Path;
+	}
+
+	FString SymbolPath = FPlatformMisc::GetEnvironmentVariable(L"UE_INSIGHTS_SYMBOL_PATH");
+
+	FString SymbolPathPart;
+	while (SymbolPath.Split(TEXT(";"), &SymbolPathPart, &SymbolPath))
+	{
+		FString Result = FPaths::Combine(SymbolPathPart, Path);
+		if (PlatformFile->FileExists(*Result))
+		{
+			return Result;
+		}
+	}
+
+	FString Result = FPaths::Combine(SymbolPath, Path);
+	if (PlatformFile->FileExists(*Result))
+	{
+		return Result;
+	}
+
+	return FString();
+}
+
 FSymslibResolver::EModuleStatus FSymslibResolver::LoadModule(FModuleEntry* Module)
 {
 	IPlatformFile* PlatformFile = &FPlatformFileManager::Get().GetPlatformFile();
@@ -313,17 +343,18 @@ FSymslibResolver::EModuleStatus FSymslibResolver::LoadModule(FModuleEntry* Modul
 		return EModuleStatus::Failed;
 	}
 
-	if (!PlatformFile->FileExists(Module->Path))
+	FString SymbolFilePath = FindSymbolFile(Module->Path);
+	if (SymbolFilePath.IsEmpty())
 	{
 		UE_LOG(LogSymslib, Warning, TEXT("File '%s' does not exist"), Module->Path);
 		return EModuleStatus::Failed;
 	}
 
 	// Map the image file to memory
-	IMappedFileHandle* ImageFileHandle = PlatformFile->OpenMapped(Module->Path);
+	IMappedFileHandle* ImageFileHandle = PlatformFile->OpenMapped(*SymbolFilePath);
 	if (ImageFileHandle == nullptr)
 	{
-		UE_LOG(LogSymslib, Warning, TEXT("Failed to open '%s'"), Module->Path);
+		UE_LOG(LogSymslib, Warning, TEXT("Failed to open '%s'"), *SymbolFilePath);
 		return EModuleStatus::Failed;
 	}
 	IMappedFileRegion* ImageFileRegion = ImageFileHandle->MapRegion(0, ImageFileHandle->GetFileSize());
@@ -337,7 +368,7 @@ FSymslibResolver::EModuleStatus FSymslibResolver::LoadModule(FModuleEntry* Modul
 
 	if (SYMS_RESULT_FAIL(Result))
 	{
-		UE_LOG(LogSymslib, Warning, TEXT("Failed to load '%s'"), Module->Path);
+		UE_LOG(LogSymslib, Warning, TEXT("Failed to load '%s'"), *SymbolFilePath);
 		return EModuleStatus::Failed;
 	}
 
@@ -386,7 +417,7 @@ FSymslibResolver::EModuleStatus FSymslibResolver::LoadModule(FModuleEntry* Modul
 
 	if (SYMS_RESULT_FAIL(Result))
 	{
-		UE_LOG(LogSymslib, Display, TEXT("No debug information for '%s'"), Module->Path);
+		UE_LOG(LogSymslib, Display, TEXT("No debug information for '%s'"), *SymbolFilePath);
 		return EModuleStatus::Failed;
 	}
 
diff --git a/Engine/Source/Developer/TraceServices/Private/Model/SymslibResolver.h b/Engine/Source/Developer/TraceServices/Private/Model/SymslibResolver.h
--- a/Engine/Source/Developer/TraceServices/Private/Model/SymslibResolver.h
+++ b/Engine/Source/Developer/TraceServices/Private/Model/SymslibResolver.h
@@ -79,6 +79,7 @@ private:
 	FModuleEntry* GetModuleForAddress(uint64 Address);
 	static void UpdateResolvedSymbol(FResolvedSymbol* Symbol, ESymbolQueryResult Result, const TCHAR* Module, const TCHAR* Name, const TCHAR* File, uint16 Line);
 	void LoadModuleTracked(FModuleEntry* Module);
+	FString FindSymbolFile(const FString& Path);
 	EModuleStatus LoadModule(FModuleEntry* Module);
 	void ResolveSymbolTracked(uint64 Address, FResolvedSymbol* Target, class FSymbolStringAllocator& StringAllocator);
 	bool ResolveSymbol(uint64 Address, FResolvedSymbol* Target, class FSymbolStringAllocator& StringAllocator);
diff --git a/Engine/Source/Runtime/Core/Private/ProfilingDebugging/CallstackTrace.cpp b/Engine/Source/Runtime/Core/Private/ProfilingDebugging/CallstackTrace.cpp
new file mode 100644
--- /dev/null
+++ b/Engine/Source/Runtime/Core/Private/ProfilingDebugging/CallstackTrace.cpp
@@ -0,0 +1,7 @@
+// Copyright Epic Games, Inc. All Rights Reserved.
+
+#include "ProfilingDebugging/CallstackTrace.h"
+
+UE_TRACE_CHANNEL_DEFINE(CallstackChannel)
+
+UE_TRACE_EVENT_DEFINE(Memory, CallstackSpec)
diff --git a/Engine/Source/Runtime/Core/Private/ProfilingDebugging/CallstackTrace.h b/Engine/Source/Runtime/Core/Private/ProfilingDebugging/CallstackTrace.h
new file mode 100644
--- /dev/null
+++ b/Engine/Source/Runtime/Core/Private/ProfilingDebugging/CallstackTrace.h
@@ -0,0 +1,54 @@
+// Copyright Epic Games, Inc. All Rights Reserved.
+#pragma once
+
+#include "CoreTypes.h"
+#include "Experimental/Containers/SherwoodHashTable.h"
+#include "HAL/CriticalSection.h"
+#include "Trace/Trace.h"
+#include "Trace/Trace.inl"
+#include "Misc/ScopeLock.h"
+
+////////////////////////////////////////////////////////////////////////////////
+UE_TRACE_CHANNEL_EXTERN(CallstackChannel)
+
+UE_TRACE_EVENT_BEGIN_EXTERN(Memory, CallstackSpec, NoSync)
+	UE_TRACE_EVENT_FIELD(uint64, Id)
+	UE_TRACE_EVENT_FIELD(uint64[], Frames)
+UE_TRACE_EVENT_END()
+
+////////////////////////////////////////////////////////////////////////////////
+namespace {
+
+	class FCallstackTracer
+	{
+	public:
+		struct FBacktraceEntry
+		{
+			enum {MaxStackDepth = 256};
+			uint64	Id = 0;
+			uint32	FrameCount = 0;
+			uint64	Frames[MaxStackDepth] = { 0 };
+		};
+
+		FCallstackTracer()
+		{
+			KnownSet.Reserve(1024 * 1024 * 2);
+		}
+		
+		void AddCallstack(const FBacktraceEntry& Entry)
+		{
+			FScopeLock _(&ProducerCs);
+			bool bAlreadyAdded = false;
+			KnownSet.Add(Entry.Id, &bAlreadyAdded);
+			if (!bAlreadyAdded)
+			{
+				UE_TRACE_LOG(Memory, CallstackSpec, CallstackChannel)
+					<< CallstackSpec.Id(Entry.Id)
+					<< CallstackSpec.Frames(Entry.Frames, Entry.FrameCount);
+			}
+		}
+	private:	
+		Experimental::TSherwoodSet<uint64> 	KnownSet;
+		FCriticalSection					ProducerCs;
+	};
+}
diff --git a/Engine/Source/Runtime/Core/Private/ProfilingDebugging/MemoryAllocationTrace.cpp b/Engine/Source/Runtime/Core/Private/ProfilingDebugging/MemoryAllocationTrace.cpp
new file mode 100644
--- /dev/null
+++ b/Engine/Source/Runtime/Core/Private/ProfilingDebugging/MemoryAllocationTrace.cpp
@@ -0,0 +1,167 @@
+// Copyright Epic Games, Inc. All Rights Reserved.
+
+#include "ProfilingDebugging/MemoryAllocationTrace.h"
+
+#if UE_MEMORY_TRACE_ENABLED
+
+////////////////////////////////////////////////////////////////////////////////
+namespace UE {
+namespace Trace {
+
+TRACELOG_API void Update();
+
+} // namespace Trace
+} // namespace UE
+
+////////////////////////////////////////////////////////////////////////////////
+UE_TRACE_CHANNEL(MemSummaryChannel)
+UE_TRACE_CHANNEL_DEFINE(MemAllocChannel)
+
+UE_TRACE_EVENT_BEGIN(Memory, Init, NoSync|Important)
+	UE_TRACE_EVENT_FIELD(uint32, MarkerPeriod)
+	UE_TRACE_EVENT_FIELD(uint8, MinAlignment)
+	UE_TRACE_EVENT_FIELD(uint8, SizeShift)
+	UE_TRACE_EVENT_FIELD(uint8, Mode)
+UE_TRACE_EVENT_END()
+
+UE_TRACE_EVENT_BEGIN(Memory, Marker)
+	UE_TRACE_EVENT_FIELD(uint64, Cycle)
+UE_TRACE_EVENT_END()
+
+UE_TRACE_EVENT_BEGIN(Memory, CoreAdd)
+	UE_TRACE_EVENT_FIELD(uint64, Owner)
+	UE_TRACE_EVENT_FIELD(void*, Base)
+	UE_TRACE_EVENT_FIELD(uint32, Size)
+UE_TRACE_EVENT_END()
+
+UE_TRACE_EVENT_BEGIN(Memory, CoreRemove)
+	UE_TRACE_EVENT_FIELD(uint64, Owner)
+	UE_TRACE_EVENT_FIELD(void*, Base)
+	UE_TRACE_EVENT_FIELD(uint32, Size)
+UE_TRACE_EVENT_END()
+
+UE_TRACE_EVENT_BEGIN(Memory, Alloc)
+	UE_TRACE_EVENT_FIELD(uint64, Owner)
+	UE_TRACE_EVENT_FIELD(void*, Address)
+	UE_TRACE_EVENT_FIELD(uint32, Size)
+	UE_TRACE_EVENT_FIELD(uint8, Alignment_SizeLower)
+UE_TRACE_EVENT_END()
+
+UE_TRACE_EVENT_BEGIN(Memory, Free)
+	UE_TRACE_EVENT_FIELD(void*, Address)
+UE_TRACE_EVENT_END()
+
+UE_TRACE_EVENT_BEGIN(Memory, ReallocAlloc)
+	UE_TRACE_EVENT_FIELD(uint64, Owner)
+	UE_TRACE_EVENT_FIELD(void*, Address)
+	UE_TRACE_EVENT_FIELD(uint32, Size)
+	UE_TRACE_EVENT_FIELD(uint8, Alignment_SizeLower)
+UE_TRACE_EVENT_END()
+
+UE_TRACE_EVENT_BEGIN(Memory, ReallocFree)
+	UE_TRACE_EVENT_FIELD(void*, Address)
+UE_TRACE_EVENT_END()
+
+
+////////////////////////////////////////////////////////////////////////////////
+void FAllocationTrace::Initialize()
+{
+	UE_TRACE_LOG(Memory, Init, MemAllocChannel)
+		<< Init.MarkerPeriod(MarkerSamplePeriod + 1)
+		<< Init.MinAlignment(uint8(MIN_ALIGNMENT))
+		<< Init.SizeShift(uint8(SizeShift));
+
+	static_assert((1 << SizeShift) - 1 <= MIN_ALIGNMENT, "Not enough bits to pack size fields");
+}
+
+////////////////////////////////////////////////////////////////////////////////
+void FAllocationTrace::EnableTracePump()
+{
+	bPumpTrace = true;
+}
+
+////////////////////////////////////////////////////////////////////////////////
+void FAllocationTrace::Update()
+{
+	uint32 TheCount = MarkerCounter.fetch_add(1, std::memory_order_relaxed);
+	if ((TheCount & MarkerSamplePeriod) == 0)
+	{
+		UE_TRACE_LOG(Memory, Marker, MemAllocChannel)
+			<< Marker.Cycle(FPlatformTime::Cycles64());
+	}
+
+	if (bPumpTrace)
+	{
+		UE::Trace::Update();
+	}
+}
+
+////////////////////////////////////////////////////////////////////////////////
+void FAllocationTrace::CoreAdd(void* Base, size_t Size, void* Owner)
+{
+	UE_TRACE_LOG(Memory, CoreAdd, MemAllocChannel)
+		<< CoreAdd.Owner(uint64(Owner))
+		<< CoreAdd.Base(Base)
+		<< CoreAdd.Size(uint32(Size >> SizeShift));
+
+	Update();
+}
+
+////////////////////////////////////////////////////////////////////////////////
+void FAllocationTrace::CoreRemove(void* Base, size_t Size, void* Owner)
+{
+	UE_TRACE_LOG(Memory, CoreRemove, MemAllocChannel)
+		<< CoreRemove.Owner(uint64(Owner))
+		<< CoreRemove.Base(Base)
+		<< CoreRemove.Size(uint32(Size >> SizeShift));
+
+	Update();
+}
+
+////////////////////////////////////////////////////////////////////////////////
+void FAllocationTrace::Alloc(void* Address, size_t Size, uint32 Alignment, void* Owner)
+{
+	uint32 Alignment_SizeLower = Alignment | (Size & ((1 << SizeShift) - 1));
+
+	UE_TRACE_LOG(Memory, Alloc, MemAllocChannel)
+		<< Alloc.Owner(uint64(Owner))
+		<< Alloc.Address(Address)
+		<< Alloc.Size(uint32(Size >> SizeShift))
+		<< Alloc.Alignment_SizeLower(uint8(Alignment_SizeLower));
+
+	Update();
+}
+
+////////////////////////////////////////////////////////////////////////////////
+void FAllocationTrace::Free(void* Address)
+{
+	UE_TRACE_LOG(Memory, Free, MemAllocChannel)
+		<< Free.Address(Address);
+
+	Update();
+}
+
+////////////////////////////////////////////////////////////////////////////////
+void FAllocationTrace::ReallocAlloc(void* Address, size_t Size, uint32 Alignment, void* Owner)
+{
+	uint32 Alignment_SizeLower = Alignment | (Size & ((1 << SizeShift) - 1));
+
+	UE_TRACE_LOG(Memory, ReallocAlloc, MemAllocChannel)
+		<< ReallocAlloc.Owner(uint64(Owner))
+		<< ReallocAlloc.Address(Address)
+		<< ReallocAlloc.Size(uint32(Size >> SizeShift))
+		<< ReallocAlloc.Alignment_SizeLower(uint8(Alignment_SizeLower));
+
+	Update();
+}
+
+////////////////////////////////////////////////////////////////////////////////
+void FAllocationTrace::ReallocFree(void* Address)
+{
+	UE_TRACE_LOG(Memory, ReallocFree, MemAllocChannel)
+		<< ReallocFree.Address(Address);
+
+	Update();
+}
+
+#endif // UE_MEMORY_TRACE_ENABLED
diff --git a/Engine/Source/Runtime/Core/Private/ProfilingDebugging/MemoryAllocationTrace.h b/Engine/Source/Runtime/Core/Private/ProfilingDebugging/MemoryAllocationTrace.h
new file mode 100644
--- /dev/null
+++ b/Engine/Source/Runtime/Core/Private/ProfilingDebugging/MemoryAllocationTrace.h
@@ -0,0 +1,36 @@
+// Copyright Epic Games, Inc. All Rights Reserved.
+#pragma once
+
+#include "ProfilingDebugging/MemoryTrace.h"
+
+#if UE_MEMORY_TRACE_ENABLED
+
+#include "CoreTypes.h"
+#include "HAL/Platform.h"
+#include "Trace/Trace.inl"
+
+#include <atomic>
+
+////////////////////////////////////////////////////////////////////////////////
+class FAllocationTrace
+{
+public:
+	void	Initialize();
+	void	EnableTracePump();
+	void	CoreAdd(void* Base, size_t Size, void* Owner);
+	void	CoreRemove(void* Base, size_t Size, void* Owner);
+	void	Alloc(void* Address, size_t Size, uint32 Alignment, void* Owner);
+	void	Free(void* Address);
+	void	ReallocAlloc(void* Address, size_t Size, uint32 Alignment, void* Owner);
+	void	ReallocFree(void* Address);
+
+private:
+	void				Update();
+	uint64				BaseCycle;
+	std::atomic<uint32>	MarkerCounter;
+	bool				bPumpTrace = false;
+	static const uint32 MarkerSamplePeriod	= (4 << 10) - 1;
+	static const uint32 SizeShift = 3;
+};
+
+#endif // UE_MEMORY_TRACE_ENABLED
diff --git a/Engine/Source/Runtime/Core/Private/ProfilingDebugging/MemoryTrace.h b/Engine/Source/Runtime/Core/Private/ProfilingDebugging/MemoryTrace.h
--- a/Engine/Source/Runtime/Core/Private/ProfilingDebugging/MemoryTrace.h
+++ b/Engine/Source/Runtime/Core/Private/ProfilingDebugging/MemoryTrace.h
@@ -4,10 +4,13 @@
 #include "Trace/Trace.h"
 
 #if !defined(UE_MEMORY_TRACE_ENABLED) && UE_TRACE_ENABLED
-#	if PLATFORM_WINDOWS
+#	if PLATFORM_WINDOWS || PLATFORM_PS4 || defined(__PS5__)
 #		if !PLATFORM_USES_FIXED_GMalloc_CLASS && PLATFORM_64BITS
 #			define UE_MEMORY_TRACE_ENABLED !UE_BUILD_SHIPPING
 #		endif
+#		if PLATFORM_PS4 || defined(__PS5__)
+#			define UE_MEMORY_TRACE_LATE_INIT UE_MEMORY_TRACE_ENABLED
+#		endif
 #	endif
 #endif
 
@@ -15,12 +18,23 @@
 #	define UE_MEMORY_TRACE_ENABLED 0
 #endif
 
+#if !defined(UE_MEMORY_TRACE_LATE_INIT)
+#	define UE_MEMORY_TRACE_LATE_INIT 0
+#endif
+
 #if UE_MEMORY_TRACE_ENABLED
 
 UE_TRACE_CHANNEL_EXTERN(MemTrackChannel);
 
 class FMalloc*	MemoryTrace_Create(class FMalloc* InMalloc);
 void			MemoryTrace_Initialize();
+
+#if UE_MEMORY_TRACE_LATE_INIT
+void			MemoryTrace_InitializeLate();
+#else
+static void		MemoryTrace_InitializeLate() {}
+#endif // UE_MEMORY_TRACE_LATE_INIT
+
 #if 0
 CORE_API void	MemoryTrace_CoreAdd(void* Base, SIZE_T Size);
 CORE_API void	MemoryTrace_CoreRemove(void* Base, SIZE_T Size);
diff --git a/Engine/Source/Runtime/Core/Private/ProfilingDebugging/ModuleDiagnostics.cpp b/Engine/Source/Runtime/Core/Private/ProfilingDebugging/ModuleDiagnostics.cpp
new file mode 100644
--- /dev/null
+++ b/Engine/Source/Runtime/Core/Private/ProfilingDebugging/ModuleDiagnostics.cpp
@@ -0,0 +1,10 @@
+// Copyright Epic Games, Inc. All Rights Reserved.
+
+#include "ProfilingDebugging/ModuleDiagnostics.h"
+
+////////////////////////////////////////////////////////////////////////////////
+UE_TRACE_CHANNEL_DEFINE(ModuleChannel, "Module information needed for symbols resolution", true)
+ 
+UE_TRACE_EVENT_DEFINE(Diagnostics, ModuleInit)
+UE_TRACE_EVENT_DEFINE(Diagnostics, ModuleLoad)
+UE_TRACE_EVENT_DEFINE(Diagnostics, ModuleUnload)
diff --git a/Engine/Source/Runtime/Core/Private/ProfilingDebugging/ModuleDiagnostics.h b/Engine/Source/Runtime/Core/Private/ProfilingDebugging/ModuleDiagnostics.h
new file mode 100644
--- /dev/null
+++ b/Engine/Source/Runtime/Core/Private/ProfilingDebugging/ModuleDiagnostics.h
@@ -0,0 +1,25 @@
+// Copyright Epic Games, Inc. All Rights Reserved.
+#pragma once
+
+#include "CoreMinimal.h"
+#include "Trace/Trace.h"
+#include "Trace/Trace.inl"
+
+////////////////////////////////////////////////////////////////////////////////
+UE_TRACE_CHANNEL_EXTERN(ModuleChannel)
+ 
+UE_TRACE_EVENT_BEGIN_EXTERN(Diagnostics, ModuleInit, NoSync|Important)
+	UE_TRACE_EVENT_FIELD(UE::Trace::AnsiString, SymbolFormat)
+	UE_TRACE_EVENT_FIELD(uint8, ModuleBaseShift)
+UE_TRACE_EVENT_END()
+
+UE_TRACE_EVENT_BEGIN_EXTERN(Diagnostics, ModuleLoad, NoSync|Important)
+	UE_TRACE_EVENT_FIELD(UE::Trace::WideString, Name)
+	UE_TRACE_EVENT_FIELD(uint64, Base)
+	UE_TRACE_EVENT_FIELD(uint32, Size)
+	UE_TRACE_EVENT_FIELD(uint8[], ImageId) // Platform specific id for this image, used to match debug files were available
+UE_TRACE_EVENT_END()
+
+UE_TRACE_EVENT_BEGIN_EXTERN(Diagnostics, ModuleUnload, NoSync|Important)
+	UE_TRACE_EVENT_FIELD(uint64, Base)
+UE_TRACE_EVENT_END()
diff --git a/Engine/Source/Runtime/Core/Private/ProfilingDebugging/TraceAuxiliary.cpp b/Engine/Source/Runtime/Core/Private/ProfilingDebugging/TraceAuxiliary.cpp
--- a/Engine/Source/Runtime/Core/Private/ProfilingDebugging/TraceAuxiliary.cpp
+++ b/Engine/Source/Runtime/Core/Private/ProfilingDebugging/TraceAuxiliary.cpp
@@ -43,6 +43,7 @@
 #include "ProfilingDebugging/MiscTrace.h"
 #include "ProfilingDebugging/PlatformFileTrace.h"
 #include "ProfilingDebugging/PlatformEvents.h"
+#include "ProfilingDebugging/MemoryTrace.h"
 #include "String/ParseTokens.h"
 #include "Templates/UnrealTemplate.h"
 #include "Trace/Trace.inl"
@@ -592,6 +593,10 @@ void FTraceAuxiliary::Initialize(const TCHAR* CommandLine)
 	}
 #endif
 
+#if UE_MEMORY_TRACE_ENABLED
+	MemoryTrace_InitializeLate();
+#endif
+
 	// Trace out information about this session. This is done before initialisation
 	// so that it is always sent (all channels are enabled prior to initialisation)
 	const TCHAR* BranchName = BuildSettings::GetBranchName();
diff --git a/Engine/Source/Runtime/Core/Private/ProfilingDebugging/Windows/CallstackTrace.cpp b/Engine/Source/Runtime/Core/Private/ProfilingDebugging/Windows/CallstackTrace.cpp
deleted file mode 100644
diff --git a/Engine/Source/Runtime/Core/Private/ProfilingDebugging/Windows/MemoryTrace.cpp b/Engine/Source/Runtime/Core/Private/ProfilingDebugging/Windows/MemoryTrace.cpp
--- a/Engine/Source/Runtime/Core/Private/ProfilingDebugging/Windows/MemoryTrace.cpp
+++ b/Engine/Source/Runtime/Core/Private/ProfilingDebugging/Windows/MemoryTrace.cpp
@@ -4,6 +4,7 @@
 
 #if UE_MEMORY_TRACE_ENABLED
 
+#include "ProfilingDebugging/MemoryAllocationTrace.h"
 #include "Containers/StringView.h"
 #include "CoreTypes.h"
 #include "HAL/MemoryBase.h"
@@ -16,17 +17,8 @@
 #include "Windows/AllowWindowsPlatformTypes.h"
 #include <winternl.h>
 
-#include <atomic>
 #include <intrin.h>
 
-////////////////////////////////////////////////////////////////////////////////
-namespace UE {
-namespace Trace {
-
-TRACELOG_API void Update();
-
-} // namespace Trace
-} // namespace UE
 
 ////////////////////////////////////////////////////////////////////////////////
 void	Backtracer_Create(FMalloc*);
@@ -75,185 +67,9 @@ struct FAddrPack
 static_assert(sizeof(FAddrPack) == sizeof(uint64), "");
 
 
-
-////////////////////////////////////////////////////////////////////////////////
-UE_TRACE_CHANNEL(MemSummaryChannel)
-UE_TRACE_CHANNEL_DEFINE(MemAllocChannel)
-
-UE_TRACE_EVENT_BEGIN(Memory, Init, NoSync|Important)
-	UE_TRACE_EVENT_FIELD(uint32, MarkerPeriod)
-	UE_TRACE_EVENT_FIELD(uint8, MinAlignment)
-	UE_TRACE_EVENT_FIELD(uint8, SizeShift)
-	UE_TRACE_EVENT_FIELD(uint8, Mode)
-UE_TRACE_EVENT_END()
-
-UE_TRACE_EVENT_BEGIN(Memory, Marker)
-	UE_TRACE_EVENT_FIELD(uint64, Cycle)
-UE_TRACE_EVENT_END()
-
-UE_TRACE_EVENT_BEGIN(Memory, CoreAdd)
-	UE_TRACE_EVENT_FIELD(uint64, Owner)
-	UE_TRACE_EVENT_FIELD(void*, Base)
-	UE_TRACE_EVENT_FIELD(uint32, Size)
-UE_TRACE_EVENT_END()
-
-UE_TRACE_EVENT_BEGIN(Memory, CoreRemove)
-	UE_TRACE_EVENT_FIELD(uint64, Owner)
-	UE_TRACE_EVENT_FIELD(void*, Base)
-	UE_TRACE_EVENT_FIELD(uint32, Size)
-UE_TRACE_EVENT_END()
-
-UE_TRACE_EVENT_BEGIN(Memory, Alloc)
-	UE_TRACE_EVENT_FIELD(uint64, Owner)
-	UE_TRACE_EVENT_FIELD(void*, Address)
-	UE_TRACE_EVENT_FIELD(uint32, Size)
-	UE_TRACE_EVENT_FIELD(uint8, Alignment_SizeLower)
-UE_TRACE_EVENT_END()
-
-UE_TRACE_EVENT_BEGIN(Memory, Free)
-	UE_TRACE_EVENT_FIELD(void*, Address)
-UE_TRACE_EVENT_END()
-
-UE_TRACE_EVENT_BEGIN(Memory, ReallocAlloc)
-	UE_TRACE_EVENT_FIELD(uint64, Owner)
-	UE_TRACE_EVENT_FIELD(void*, Address)
-	UE_TRACE_EVENT_FIELD(uint32, Size)
-	UE_TRACE_EVENT_FIELD(uint8, Alignment_SizeLower)
-UE_TRACE_EVENT_END()
-
-UE_TRACE_EVENT_BEGIN(Memory, ReallocFree)
-	UE_TRACE_EVENT_FIELD(void*, Address)
-UE_TRACE_EVENT_END()
-
-
-
-////////////////////////////////////////////////////////////////////////////////
-class FAllocationTrace
-{
-public:
-	void	Initialize();
-	void	EnableTracePump();
-	void	CoreAdd(void* Base, size_t Size, void* Owner);
-	void	CoreRemove(void* Base, size_t Size, void* Owner);
-	void	Alloc(void* Address, size_t Size, uint32 Alignment, void* Owner);
-	void	Free(void* Address);
-	void	ReallocAlloc(void* Address, size_t Size, uint32 Alignment, void* Owner);
-	void	ReallocFree(void* Address);
-
-private:
-	void				Update();
-	uint64				BaseCycle;
-	std::atomic<uint32>	MarkerCounter;
-	bool				bPumpTrace = false;
-	static const uint32 MarkerSamplePeriod	= (4 << 10) - 1;
-	static const uint32 SizeShift = 3;
-};
-
 ////////////////////////////////////////////////////////////////////////////////
 static FUndestructed<FAllocationTrace> GAllocationTrace;
 
-////////////////////////////////////////////////////////////////////////////////
-void FAllocationTrace::Initialize()
-{
-	UE_TRACE_LOG(Memory, Init, MemAllocChannel)
-		<< Init.MarkerPeriod(MarkerSamplePeriod + 1)
-		<< Init.MinAlignment(uint8(MIN_ALIGNMENT))
-		<< Init.SizeShift(uint8(SizeShift));
-
-	static_assert((1 << SizeShift) - 1 <= MIN_ALIGNMENT, "Not enough bits to pack size fields");
-}
-
-////////////////////////////////////////////////////////////////////////////////
-void FAllocationTrace::EnableTracePump()
-{
-	bPumpTrace = true;
-}
-
-////////////////////////////////////////////////////////////////////////////////
-void FAllocationTrace::Update()
-{
-	uint32 TheCount = MarkerCounter.fetch_add(1, std::memory_order_relaxed);
-	if ((TheCount & MarkerSamplePeriod) == 0)
-	{
-		UE_TRACE_LOG(Memory, Marker, MemAllocChannel)
-			<< Marker.Cycle(FPlatformTime::Cycles64());
-	}
-
-	if (bPumpTrace)
-	{
-		UE::Trace::Update();
-	}
-}
-
-////////////////////////////////////////////////////////////////////////////////
-void FAllocationTrace::CoreAdd(void* Base, size_t Size, void* Owner)
-{
-	UE_TRACE_LOG(Memory, CoreAdd, MemAllocChannel)
-		<< CoreAdd.Owner(uint64(Owner))
-		<< CoreAdd.Base(Base)
-		<< CoreAdd.Size(uint32(Size >> SizeShift));
-
-	Update();
-}
-
-////////////////////////////////////////////////////////////////////////////////
-void FAllocationTrace::CoreRemove(void* Base, size_t Size, void* Owner)
-{
-	UE_TRACE_LOG(Memory, CoreRemove, MemAllocChannel)
-		<< CoreRemove.Owner(uint64(Owner))
-		<< CoreRemove.Base(Base)
-		<< CoreRemove.Size(uint32(Size >> SizeShift));
-
-	Update();
-}
-
-////////////////////////////////////////////////////////////////////////////////
-void FAllocationTrace::Alloc(void* Address, size_t Size, uint32 Alignment, void* Owner)
-{
-	uint32 Alignment_SizeLower = Alignment | (Size & ((1 << SizeShift) - 1));
-
-	UE_TRACE_LOG(Memory, Alloc, MemAllocChannel)
-		<< Alloc.Owner(uint64(Owner))
-		<< Alloc.Address(Address)
-		<< Alloc.Size(uint32(Size >> SizeShift))
-		<< Alloc.Alignment_SizeLower(uint8(Alignment_SizeLower));
-
-	Update();
-}
-
-////////////////////////////////////////////////////////////////////////////////
-void FAllocationTrace::Free(void* Address)
-{
-	UE_TRACE_LOG(Memory, Free, MemAllocChannel)
-		<< Free.Address(Address);
-
-	Update();
-}
-
-////////////////////////////////////////////////////////////////////////////////
-void FAllocationTrace::ReallocAlloc(void* Address, size_t Size, uint32 Alignment, void* Owner)
-{
-	uint32 Alignment_SizeLower = Alignment | (Size & ((1 << SizeShift) - 1));
-
-	UE_TRACE_LOG(Memory, ReallocAlloc, MemAllocChannel)
-		<< ReallocAlloc.Owner(uint64(Owner))
-		<< ReallocAlloc.Address(Address)
-		<< ReallocAlloc.Size(uint32(Size >> SizeShift))
-		<< ReallocAlloc.Alignment_SizeLower(uint8(Alignment_SizeLower));
-
-	Update();
-}
-
-////////////////////////////////////////////////////////////////////////////////
-void FAllocationTrace::ReallocFree(void* Address)
-{
-	UE_TRACE_LOG(Memory, ReallocFree, MemAllocChannel)
-		<< ReallocFree.Address(Address);
-
-	Update();
-}
-
-
 
 ////////////////////////////////////////////////////////////////////////////////
 class FMallocWrapper
diff --git a/Engine/Source/Runtime/Core/Private/ProfilingDebugging/Windows/ModuleDiagnostics.cpp b/Engine/Source/Runtime/Core/Private/ProfilingDebugging/Windows/ModuleDiagnostics.cpp
deleted file mode 100644
diff --git a/Engine/Source/Runtime/Core/Private/ProfilingDebugging/Windows/WindowsCallstackTrace.cpp b/Engine/Source/Runtime/Core/Private/ProfilingDebugging/Windows/WindowsCallstackTrace.cpp
new file mode 100644
--- /dev/null
+++ b/Engine/Source/Runtime/Core/Private/ProfilingDebugging/Windows/WindowsCallstackTrace.cpp
@@ -0,0 +1,666 @@
+// Copyright Epic Games, Inc. All Rights Reserved.
+
+#include "Containers/Array.h"
+#include "Containers/ArrayView.h"
+#include "CoreTypes.h"
+#include "HAL/CriticalSection.h"
+#include "HAL/MemoryBase.h"
+#include "HAL/RunnableThread.h"
+#include "HAL/LowLevelMemTracker.h"
+#include "ProfilingDebugging/MemoryTrace.h"
+
+#include "Windows/AllowWindowsPlatformTypes.h"
+#	include <winnt.h>
+#	include <winternl.h>
+#include "Windows/HideWindowsPlatformTypes.h"
+#include "Misc/CoreDelegates.h"
+#include "Trace/Trace.inl"
+#include "HAL/Runnable.h"
+#include "HAL/PlatformProcess.h"
+#include "Misc/ScopeLock.h"
+#include "ProfilingDebugging/CallstackTrace.h"
+
+#include <atomic>
+
+// 0=off, 1=stats, 2=validation, 3=truth_compare
+#define BACKTRACE_DBGLVL 0
+
+/*
+ * Windows' x64 binaries contain a ".pdata" section that describes the location
+ * and size of its functions and details on how to unwind them. The unwind
+ * information includes descriptions about a function's stack frame size and
+ * the non-volatile registers it pushes onto the stack. From this we can
+ * calculate where a call instruction wrote its return address. This is enough
+ * to walk the callstack and by caching this information it can be done
+ * efficiently.
+ *
+ * Some functions need a variable amount of stack (such as those that use
+ * alloc() for example) will use a frame pointer. Frame pointers involve saving
+ * and restoring the stack pointer in the function's prologue/epilogue. This
+ * frees the function up to modify the stack pointer arbitrarily. This
+ * significantly complicates establishing where a return address is, so this
+ * pdata scheme of walking the stack just doesn't support functions like this.
+ * Walking stops if it encounters such a function. Fortunately there are
+ * usually very few such functions, saving us from having to read and track
+ * non-volatile registers which adds a significant amount of work.
+ *
+ * A further optimisation is to to assume we are only interested methods that
+ * are part of engine or game code. As such we only build lookup tables for
+ * such modules and never accept OS or third party modules. Backtracing stops
+ * if an address is encountered which doesn't map to a known module.
+ */
+
+////////////////////////////////////////////////////////////////////////////////
+static uint32 AddressToId(UPTRINT Address)
+{
+	return uint32(Address >> 16);
+}
+
+static UPTRINT IdToAddress(uint32 Id)
+{
+	return static_cast<uint32>(UPTRINT(Id) << 16);
+}
+
+struct FIdPredicate
+{
+	template <class T> bool operator () (uint32 Id, const T& Item) const { return Id < Item.Id; }
+	template <class T> bool operator () (const T& Item, uint32 Id) const { return Item.Id < Id; }
+};
+
+////////////////////////////////////////////////////////////////////////////////
+struct FUnwindInfo
+{
+	uint8	Version : 3;
+	uint8	Flags : 5;
+	uint8	PrologBytes;
+	uint8	NumUnwindCodes;
+	uint8	FrameReg	: 4;
+	uint8	FrameRspBias : 4;
+};
+
+struct FUnwindCode
+{
+	uint8	PrologOffset;
+	uint8	OpCode : 4;
+	uint8	OpInfo : 4;
+	uint16	Params[];
+};
+
+enum
+{
+	UWOP_PUSH_NONVOL		= 0,	// 1 node
+	UWOP_ALLOC_LARGE		= 1,	// 2 or 3 nodes
+	UWOP_ALLOC_SMALL		= 2,	// 1 node
+	UWOP_SET_FPREG			= 3,	// 1 node
+	UWOP_SAVE_NONVOL		= 4,	// 2 nodes
+	UWOP_SAVE_NONVOL_FAR	= 5,	// 3 nodes
+	UWOP_SAVE_XMM128		= 8,	// 2 nodes
+	UWOP_SAVE_XMM128_FAR	= 9,	// 3 nodes
+	UWOP_PUSH_MACHFRAME		= 10,	// 1 node
+};
+
+////////////////////////////////////////////////////////////////////////////////
+class FBacktracer
+{
+public:
+							FBacktracer(FMalloc* InMalloc);
+							~FBacktracer();
+	static FBacktracer*		Get();
+	void					AddModule(UPTRINT Base, const TCHAR* Name);
+	void					RemoveModule(UPTRINT Base);
+	void*					GetBacktraceId(void* AddressOfReturnAddress) const;
+	int32					GetFrameSize(void* FunctionAddress) const;
+
+private:
+	struct FFunction
+	{
+		uint32				Id;
+		int32				RspBias;
+#if BACKTRACE_DBGLVL >= 2
+		uint32				Size;
+		const FUnwindInfo*	UnwindInfo;
+#endif
+	};
+
+	struct FModule
+	{
+		uint32				Id;
+		uint32				IdSize = 0;
+		uint32				NumFunctions;
+#if BACKTRACE_DBGLVL >= 1
+		uint16				NumFpTypes;
+		//uint16			*padding*
+#else
+		//uint32			*padding*
+#endif
+		FFunction*			Functions;
+	};
+
+	struct FLookupState
+	{
+		FModule				Module;
+	};
+
+	const FFunction*		LookupFunction(UPTRINT Address, FLookupState& State) const;
+	static FBacktracer*		Instance;
+	mutable FCriticalSection Lock;
+	FModule*				Modules;
+	int32					ModulesNum;
+	int32					ModulesCapacity;
+	FMalloc*				Malloc;
+	FCallstackTracer*		CallstackTracer;
+#if BACKTRACE_DBGLVL >= 1
+	mutable uint32			NumFpTruncations = 0;
+	mutable uint32			TotalFunctions = 0;
+#endif
+};
+
+////////////////////////////////////////////////////////////////////////////////
+FBacktracer* FBacktracer::Instance = nullptr;
+
+////////////////////////////////////////////////////////////////////////////////
+FBacktracer::FBacktracer(FMalloc* InMalloc)
+	: Malloc(InMalloc)
+{
+	ModulesCapacity = 8;
+	ModulesNum = 0;
+	Modules = (FModule*)Malloc->Malloc(sizeof(FModule) * ModulesCapacity);
+
+	Instance = this;
+	CallstackTracer = (FCallstackTracer*) Malloc->Malloc(sizeof(FCallstackTracer), alignof(FCallstackTracer));
+	CallstackTracer = new(CallstackTracer) FCallstackTracer();
+}
+
+////////////////////////////////////////////////////////////////////////////////
+FBacktracer::~FBacktracer()
+{
+	TArrayView<FModule> ModulesView(Modules, ModulesNum);
+	for (FModule& Module : ModulesView)
+	{
+		Malloc->Free(Module.Functions);
+	}
+
+}
+
+////////////////////////////////////////////////////////////////////////////////
+FBacktracer* FBacktracer::Get()
+{
+	return Instance;
+}
+
+////////////////////////////////////////////////////////////////////////////////
+void FBacktracer::AddModule(UPTRINT ModuleBase, const TCHAR* Name)
+{
+	const auto* DosHeader = (IMAGE_DOS_HEADER*)ModuleBase;
+	const auto* NtHeader = (IMAGE_NT_HEADERS*)(ModuleBase + DosHeader->e_lfanew);
+	const IMAGE_FILE_HEADER* FileHeader = &(NtHeader->FileHeader);
+
+	if (FCString::Strfind(Name, TEXT("Binaries")) == nullptr || FCString::Strfind(Name, TEXT("ThirdParty")) != nullptr)
+	{
+		return;
+	}
+
+	uint32 NumSections = FileHeader->NumberOfSections;
+	const auto* Sections = (IMAGE_SECTION_HEADER*)(UPTRINT(&(NtHeader->OptionalHeader)) + FileHeader->SizeOfOptionalHeader);
+
+	// Find ".pdata" section
+	UPTRINT PdataBase = 0;
+	UPTRINT PdataEnd = 0;
+	for (uint32 i = 0; i < NumSections; ++i)
+	{
+		const IMAGE_SECTION_HEADER* Section = Sections + i;
+		if (*(uint64*)(Section->Name) == 0x61'74'61'64'70'2eull) // Sections names are eight bytes and zero padded. This constant is '.pdata'
+		{
+			PdataBase = ModuleBase + Section->VirtualAddress;
+			PdataEnd = PdataBase + Section->SizeOfRawData;
+			break;
+		}
+	}
+
+	if (PdataBase == 0)
+	{
+		return;
+	}
+
+	// Count the number of functions. The assumption here is that if we have got this far then there is at least one function
+	uint32 NumFunctions = uint32(PdataEnd - PdataBase) / sizeof(RUNTIME_FUNCTION);
+	const auto* FunctionTables = (RUNTIME_FUNCTION*)PdataBase;
+	do
+	{
+		const RUNTIME_FUNCTION* Function = FunctionTables + NumFunctions - 1;
+		if (uint32(Function->BeginAddress) < uint32(Function->EndAddress))
+		{
+			break;
+		}
+
+		--NumFunctions;
+	}
+	while (NumFunctions != 0);
+		
+	// Allocate some space for the module's function-to-frame-size table
+	auto* OutTable = (FFunction*)Malloc->Malloc(sizeof(FFunction) * NumFunctions);
+	FFunction* OutTableCursor = OutTable;
+
+	// Extract frame size for each function from pdata's unwind codes.
+	uint32 NumFpFuncs = 0;
+	for (uint32 i = 0; i < NumFunctions; ++i)
+	{
+		const RUNTIME_FUNCTION* FunctionTable = FunctionTables + i;
+
+		UPTRINT UnwindInfoAddr = ModuleBase + FunctionTable->UnwindInfoAddress;
+		const auto* UnwindInfo = (FUnwindInfo*)UnwindInfoAddr;
+
+		if (UnwindInfo->Version != 1)
+		{
+			/* some v2s have been seen in msvc. Always seem to be assembly
+			 * routines (memset, memcpy, etc) */
+			continue;
+		}
+
+		int32 FpInfo = 0;
+		int32 RspBias = 0;
+
+#if BACKTRACE_DBGLVL >= 2
+		uint32 PrologVerify = UnwindInfo->PrologBytes;
+#endif
+
+		const auto* Code = (FUnwindCode*)(UnwindInfo + 1);
+		const auto* EndCode = Code + UnwindInfo->NumUnwindCodes;
+		while (Code < EndCode)
+		{
+#if BACKTRACE_DBGLVL >= 2
+			if (Code->PrologOffset > PrologVerify)
+			{
+				PLATFORM_BREAK();
+			}
+			PrologVerify = Code->PrologOffset;
+#endif
+
+			switch (Code->OpCode)
+			{
+			case UWOP_PUSH_NONVOL:
+				RspBias += 8;
+				Code += 1;
+				break;
+
+			case UWOP_ALLOC_LARGE:
+				if (Code->OpInfo)
+				{
+					RspBias += *(uint32*)(Code->Params);
+					Code += 3;
+				}
+				else
+				{
+					RspBias += Code->Params[0] * 8;
+					Code += 2;
+				}
+				break;
+
+			case UWOP_ALLOC_SMALL:
+				RspBias += (Code->OpInfo * 8) + 8;
+				Code += 1;
+				break;
+
+			case UWOP_SET_FPREG:
+				// Function will adjust RSP (e.g. through use of alloca()) so it
+				// uses a frame pointer register. There's instructions like;
+				//
+				//   push FRAME_REG
+				//   lea FRAME_REG, [rsp + (FRAME_RSP_BIAS * 16)]
+				//   ...
+				//   add rsp, rax
+				//   ...
+				//   sub rsp, FRAME_RSP_BIAS * 16
+				//   pop FRAME_REG
+				//   ret
+				//
+				// To recover the stack frame we would need to track non-volatile
+				// registers which adds a lot of overhead for a small subset of
+				// functions. Instead we'll end backtraces at these functions.
+
+
+				// MSB is set to detect variable sized frames that we can't proceed
+				// past when back-tracing.
+				NumFpFuncs++;
+				FpInfo |= 0x80000000 | (uint32(UnwindInfo->FrameReg) << 27) | (uint32(UnwindInfo->FrameRspBias) << 23);
+				Code += 1;
+				break;
+
+			case UWOP_PUSH_MACHFRAME:
+				RspBias = Code->OpInfo ? 48 : 40;
+				Code += 1;
+				break;
+
+			case UWOP_SAVE_NONVOL:		Code += 2; break; /* saves are movs instead of pushes */
+			case UWOP_SAVE_NONVOL_FAR:	Code += 3; break;
+			case UWOP_SAVE_XMM128:		Code += 2; break;
+			case UWOP_SAVE_XMM128_FAR:	Code += 3; break;
+
+			default:
+#if BACKTRACE_DBGLVL >= 2
+				PLATFORM_BREAK();
+#endif
+				break;
+			}
+		}
+
+		// "Chained" simply means that multiple RUNTIME_FUNCTIONs pertains to a
+		// single actual function in the .text segment.
+		bool bIsChained = (UnwindInfo->Flags & UNW_FLAG_CHAININFO);
+
+		RspBias /= sizeof(void*);	// stack push/popds in units of one machine word
+		RspBias += !bIsChained;		// and one extra push for the ret address
+		RspBias |= FpInfo;			// pack in details about possible frame pointer
+
+		if (bIsChained)
+		{
+			OutTableCursor[-1].RspBias += RspBias;
+#if BACKTRACE_DBGLVL >= 2
+			OutTableCursor[-1].Size += (FunctionTable->EndAddress - FunctionTable->BeginAddress);
+#endif
+		}
+		else
+		{
+			*OutTableCursor = {
+				FunctionTable->BeginAddress,
+				RspBias,
+#if BACKTRACE_DBGLVL >= 2
+				FunctionTable->EndAddress - FunctionTable->BeginAddress,
+				UnwindInfo,
+#endif
+			};
+
+			++OutTableCursor;
+		}
+	}
+
+	UPTRINT ModuleSize = NtHeader->OptionalHeader.SizeOfImage;
+	ModuleSize += 0xffff; // to align up to next 64K page. it'll get shifted by AddressToId()
+
+	FModule Module = {
+		AddressToId(ModuleBase),
+		AddressToId(ModuleSize),
+		uint32(UPTRINT(OutTableCursor - OutTable)),
+#if BACKTRACE_DBGLVL >= 1
+		uint16(NumFpFuncs),
+#endif
+		OutTable,
+	};
+
+	{
+		FScopeLock _(&Lock);
+
+		TArrayView<FModule> ModulesView(Modules, ModulesNum);
+		int32 Index = Algo::UpperBound(ModulesView, Module.Id, FIdPredicate());
+
+		if (ModulesNum + 1 > ModulesCapacity)
+		{
+			ModulesCapacity += 8;
+			Modules = (FModule*) Malloc->Realloc(Modules, sizeof(FModule)* ModulesCapacity);
+		}
+		Modules[ModulesNum++] = Module;
+		Algo::Sort(TArrayView<FModule>(Modules, ModulesNum), [](const FModule& A, const FModule& B) { return A.Id < B.Id; });
+	}
+
+#if BACKTRACE_DBGLVL >= 1
+	NumFpTruncations += NumFpFuncs;
+	TotalFunctions += NumFunctions;
+#endif
+}
+
+////////////////////////////////////////////////////////////////////////////////
+void FBacktracer::RemoveModule(UPTRINT ModuleBase)
+{
+	// When Windows' RequestExit() is called it hard-terminates all threads except
+	// the main thread and then proceeds to unload the process' DLLs. This hard 
+	// thread termination can result is dangling locked locks. Not an issue as
+	// the rule is "do not do anything multithreaded in DLL load/unload". And here
+	// we are, taking write locks during DLL unload which is, quite unsurprisingly,
+	// deadlocking. In reality tracking Windows' DLL unloads doesn't tell us
+	// anything due to how DLLs and processes' address spaces work. So we will...
+#if defined PLATFORM_WINDOWS
+	return;
+#endif
+
+	FScopeLock _(&Lock);
+
+	uint32 ModuleId = AddressToId(ModuleBase);
+	TArrayView<FModule> ModulesView(Modules, ModulesNum);
+	int32 Index = Algo::LowerBound(ModulesView, ModuleId, FIdPredicate());
+	if (Index >= ModulesNum)
+	{
+		return;
+	}
+
+	const FModule& Module = Modules[Index];
+	if (Module.Id != ModuleId)
+	{
+		return;
+	}
+
+#if BACKTRACE_DBGLVL >= 1
+	NumFpTruncations -= Module.NumFpTypes;
+	TotalFunctions -= Module.NumFunctions;
+#endif
+
+	// no code should be executing at this point so we can safely free the
+	// table knowing know one is looking at it.
+	Malloc->Free(Module.Functions);
+
+	for (SIZE_T i = Index; i < ModulesNum; i++)
+	{
+		Modules[i] = Modules[i + 1];
+	}
+
+	--ModulesNum;
+}
+
+////////////////////////////////////////////////////////////////////////////////
+const FBacktracer::FFunction* FBacktracer::LookupFunction(UPTRINT Address, FLookupState& State) const
+{
+	// This function caches the previous module look up. The theory here is that
+	// a series of return address in a backtrace often cluster around one module
+
+	FIdPredicate IdPredicate;
+
+	// Look up the module that Address belongs to.
+	uint32 AddressId = AddressToId(Address);
+	if ((AddressId - State.Module.Id) >= State.Module.IdSize)
+	{
+		TArrayView<FModule> ModulesView(Modules, ModulesNum);
+		uint32 Index = Algo::UpperBound(ModulesView, AddressId, IdPredicate);
+		if (Index == 0)
+		{
+			return nullptr;
+		}
+
+		State.Module = Modules[Index - 1];
+	}
+
+	// Check that the address is within the address space of the best-found module
+	const FModule* Module = &(State.Module);
+	if ((AddressId - Module->Id) >= Module->IdSize)
+	{
+		return nullptr;
+	}
+
+	// Now we've a module we have a table of functions and their stack sizes so
+	// we can get the frame size for Address
+	uint32 FuncId = uint32(Address - IdToAddress(Module->Id)); 
+	TArrayView<FFunction> FuncsView(Module->Functions, Module->NumFunctions);
+	uint32 Index = Algo::UpperBound(FuncsView, FuncId, IdPredicate);
+
+	const FFunction* Function = Module->Functions + (Index - 1);
+#if BACKTRACE_DBGLVL >= 2
+	if ((FuncId - Function->Id) >= Function->Size)
+	{
+		PLATFORM_BREAK();
+		return nullptr;
+	}
+#endif
+	return Function;
+}
+
+////////////////////////////////////////////////////////////////////////////////
+void* FBacktracer::GetBacktraceId(void* AddressOfReturnAddress) const
+{
+	FLookupState LookupState;
+	FCallstackTracer::FBacktraceEntry BacktraceEntry;
+
+	UPTRINT* StackPointer = (UPTRINT*)AddressOfReturnAddress;
+
+#if BACKTRACE_DBGLVL >= 3
+	UPTRINT TruthBacktrace[1024];
+	uint32 NumTruth = RtlCaptureStackBackTrace(0, 1024, (void**)TruthBacktrace, nullptr);
+	UPTRINT* TruthCursor = TruthBacktrace;
+	for (; *TruthCursor != *StackPointer; ++TruthCursor);
+#endif
+
+#if BACKTRACE_DBGLVL >= 2
+	struct { void* Sp; void* Ip; const FFunction* Function; } Backtrace[1024] = {};
+	uint32 NumBacktrace = 0;
+#endif
+
+	uint64 BacktraceId = 0;
+	uint32 FrameIdx = 0;
+
+	FScopeLock _(&Lock);
+	do
+	{
+		UPTRINT RetAddr = *StackPointer;
+		
+		BacktraceEntry.Frames[FrameIdx++] = RetAddr;
+		BacktraceEntry.FrameCount = FrameIdx;
+
+		// This is a simple order-dependent LCG. Should be sufficient enough
+		BacktraceId += RetAddr;
+		BacktraceId *= 0x30be8efa499c249dull;
+
+		const FFunction* Function = LookupFunction(RetAddr, LookupState);
+		if (Function == nullptr)
+		{
+			break;
+		}
+
+#if BACKTRACE_DBGLVL >= 2
+		if (NumBacktrace < 1024)
+		{
+			Backtrace[NumBacktrace++] = {
+				StackPointer,
+				(void*)RetAddr,
+				Function,
+			};
+		}
+#endif
+
+		if (Function->RspBias < 0)
+		{
+			// This is a frame with a variable-sized stack pointer. We don't
+			// track enough information to proceed.
+#if BACKTRACE_DBGLVL >= 1
+			NumFpTruncations++;
+#endif
+			break;
+		}
+
+		StackPointer += Function->RspBias;
+	}
+	// Trunkate callstacks longer than MaxStackDepth
+	while (*StackPointer && FrameIdx < FCallstackTracer::FBacktraceEntry::MaxStackDepth);
+
+	// Save the collected id
+	BacktraceEntry.Id = BacktraceId;
+
+#if BACKTRACE_DBGLVL >= 3
+	for (uint32 i = 0; i < NumBacktrace; ++i)
+	{
+		if ((void*)TruthCursor[i] != Backtrace[i].Ip)
+		{
+			PLATFORM_BREAK();
+			break;
+		}
+	}
+#endif
+
+	// Add to queue to be processed. This might block until there is room in the
+	// queue (i.e. the processing thread has caught up processing).
+	CallstackTracer->AddCallstack(BacktraceEntry);
+	return (void*)BacktraceId;
+}
+
+////////////////////////////////////////////////////////////////////////////////
+int32 FBacktracer::GetFrameSize(void* FunctionAddress) const
+{
+	FScopeLock _(&Lock);
+
+	FLookupState LookupState;
+	const FFunction* Function = LookupFunction(UPTRINT(FunctionAddress), LookupState);
+	if (FunctionAddress == nullptr)
+	{
+		return -1;
+	}
+
+	if (Function->RspBias < 0)
+	{
+		return -1;
+	}
+
+	return Function->RspBias;
+}
+
+////////////////////////////////////////////////////////////////////////////////
+void Modules_Create(FMalloc*);
+void Modules_Subscribe(void (*)(bool, void*, const TCHAR*));
+void Modules_Initialize();
+
+////////////////////////////////////////////////////////////////////////////////
+void Backtracer_Create(FMalloc* Malloc)
+{
+	if (FBacktracer::Get() != nullptr)
+	{
+		return;
+	}
+
+	// Allocate, construct and intentionally leak backtracer
+	void* Alloc = Malloc->Malloc(sizeof(FBacktracer), alignof(FBacktracer));
+	new (Alloc) FBacktracer(Malloc);
+
+	Modules_Create(Malloc);
+	Modules_Subscribe(
+		[] (bool bLoad, void* Module, const TCHAR* Name)
+		{
+			bLoad
+				? FBacktracer::Get()->AddModule(UPTRINT(Module), Name)
+				: FBacktracer::Get()->RemoveModule(UPTRINT(Module));
+		}
+	);
+}
+
+////////////////////////////////////////////////////////////////////////////////
+void Backtracer_Initialize()
+{
+	Modules_Initialize();
+}
+
+////////////////////////////////////////////////////////////////////////////////
+int32 Backtracer_GetFrameSize(void* FunctionAddress)
+{
+	if (FBacktracer* Instance = FBacktracer::Get())
+	{
+		return Instance->GetFrameSize(FunctionAddress);
+	}
+
+	return -1;
+}
+
+////////////////////////////////////////////////////////////////////////////////
+void* Backtracer_GetBacktraceId(void* AddressOfReturnAddress)
+{
+	if (FBacktracer* Instance = FBacktracer::Get())
+	{
+		return Instance->GetBacktraceId(AddressOfReturnAddress);
+	}
+
+	return nullptr;
+}
+
diff --git a/Engine/Source/Runtime/Core/Private/ProfilingDebugging/Windows/WindowsModuleDiagnostics.cpp b/Engine/Source/Runtime/Core/Private/ProfilingDebugging/Windows/WindowsModuleDiagnostics.cpp
new file mode 100644
--- /dev/null
+++ b/Engine/Source/Runtime/Core/Private/ProfilingDebugging/Windows/WindowsModuleDiagnostics.cpp
@@ -0,0 +1,242 @@
+// Copyright Epic Games, Inc. All Rights Reserved.
+
+#include "CoreMinimal.h"
+#include "Trace/Trace.h"
+#include "Trace/Trace.inl"
+#include "ProfilingDebugging/ModuleDiagnostics.h"
+
+#include "Windows/AllowWindowsPlatformTypes.h"
+#	include <winternl.h>
+#include "Windows/HideWindowsPlatformTypes.h"
+
+////////////////////////////////////////////////////////////////////////////////
+struct FNtDllFunction
+{
+	FARPROC Addr;
+
+	FNtDllFunction(const char* Name)
+	{
+		HMODULE NtDll = LoadLibraryW(L"ntdll.dll");
+		check(NtDll);
+		Addr = GetProcAddress(NtDll, Name);
+	}
+
+	template <typename... ArgTypes>
+	unsigned int operator () (ArgTypes... Args)
+	{
+		typedef unsigned int (NTAPI *Prototype)(ArgTypes...);
+		return (Prototype((void*)Addr))(Args...);
+	}
+};
+
+////////////////////////////////////////////////////////////////////////////////
+class FModuleTrace
+{
+public:
+	typedef void				(*SubscribeFunc)(bool, void*, const TCHAR*);
+	typedef TArray<SubscribeFunc, TFixedAllocator<64>> SubscriberArray;
+
+								FModuleTrace(FMalloc* InMalloc);
+								~FModuleTrace();
+	static FModuleTrace*		Get();
+	void						Initialize();
+	void						Subscribe(SubscribeFunc Function);
+
+private:
+	void						OnDllLoaded(const UNICODE_STRING& Name, UPTRINT Base);
+	void						OnDllUnloaded(UPTRINT Base);
+	void						OnDllNotification(unsigned int Reason, const void* DataPtr);
+	static FModuleTrace*		Instance;
+	SubscriberArray				Subscribers;
+	void*						CallbackCookie;
+};
+
+////////////////////////////////////////////////////////////////////////////////
+FModuleTrace* FModuleTrace::Instance = nullptr;
+
+////////////////////////////////////////////////////////////////////////////////
+FModuleTrace::FModuleTrace(FMalloc* InMalloc)
+{
+	Instance = this;
+}
+
+////////////////////////////////////////////////////////////////////////////////
+FModuleTrace::~FModuleTrace()
+{
+	FNtDllFunction UnregisterFunc("LdrUnregisterDllNotification");
+	UnregisterFunc(CallbackCookie);
+}
+
+////////////////////////////////////////////////////////////////////////////////
+FModuleTrace* FModuleTrace::Get()
+{
+	return Instance;
+}
+
+////////////////////////////////////////////////////////////////////////////////
+void FModuleTrace::Initialize()
+{
+	using namespace UE::Trace;
+
+	UE_TRACE_LOG(Diagnostics, ModuleInit, ModuleChannel, sizeof(char) * 3)
+		<< ModuleInit.SymbolFormat("pdb", 3)
+		<< ModuleInit.ModuleBaseShift(uint8(0));
+
+	// Register for DLL load/unload notifications.
+	auto Thunk = [] (ULONG Reason, const void* Data, void* Context)
+	{
+		auto* Self = (FModuleTrace*)Context;
+		Self->OnDllNotification(Reason, Data);
+	};
+
+	typedef void (CALLBACK *ThunkType)(ULONG, const void*, void*);
+	auto ThunkImpl = ThunkType(Thunk);
+
+	FNtDllFunction RegisterFunc("LdrRegisterDllNotification");
+	RegisterFunc(0, ThunkImpl, this, &CallbackCookie);
+
+	// Enumerate already loaded modules.
+	const TEB* ThreadEnvBlock = NtCurrentTeb();
+	const PEB* ProcessEnvBlock = ThreadEnvBlock->ProcessEnvironmentBlock;
+	const LIST_ENTRY* ModuleIter = ProcessEnvBlock->Ldr->InMemoryOrderModuleList.Flink;
+	const LIST_ENTRY* ModuleIterEnd = ModuleIter->Blink;
+	do
+	{
+		const auto& ModuleData = *(LDR_DATA_TABLE_ENTRY*)(ModuleIter - 1);
+		if (ModuleData.DllBase == 0)
+		{
+			break;
+		}
+
+		OnDllLoaded(ModuleData.FullDllName, UPTRINT(ModuleData.DllBase));
+		ModuleIter = ModuleIter->Flink;
+	}
+	while (ModuleIter != ModuleIterEnd);
+}
+
+////////////////////////////////////////////////////////////////////////////////
+void FModuleTrace::Subscribe(SubscribeFunc Function)
+{
+	Subscribers.Add(Function);
+}
+
+////////////////////////////////////////////////////////////////////////////////
+void FModuleTrace::OnDllNotification(unsigned int Reason, const void* DataPtr)
+{
+	enum
+	{
+		LDR_DLL_NOTIFICATION_REASON_LOADED = 1,
+		LDR_DLL_NOTIFICATION_REASON_UNLOADED = 2,
+	};
+
+	struct FNotificationData
+	{
+		uint32					Flags;
+		const UNICODE_STRING&	FullPath;
+		const UNICODE_STRING&	BaseName;
+		UPTRINT					Base;
+	};
+	const auto& Data = *(FNotificationData*)DataPtr;
+
+	switch (Reason)
+	{
+	case LDR_DLL_NOTIFICATION_REASON_LOADED:	OnDllLoaded(Data.FullPath, Data.Base);	break;
+	case LDR_DLL_NOTIFICATION_REASON_UNLOADED:	OnDllUnloaded(Data.Base);				break;
+	}
+}
+
+////////////////////////////////////////////////////////////////////////////////
+void FModuleTrace::OnDllLoaded(const UNICODE_STRING& Name, UPTRINT Base)
+{
+	const auto* DosHeader = (IMAGE_DOS_HEADER*)Base;
+	const auto* NtHeaders = (IMAGE_NT_HEADERS*)(Base + DosHeader->e_lfanew);
+	const IMAGE_OPTIONAL_HEADER& OptionalHeader = NtHeaders->OptionalHeader;
+	TArray<uint8, TFixedAllocator<20>> ImageId;
+
+	// Find the guid and age of the binary, used to match debug files
+	const IMAGE_DATA_DIRECTORY& DebugInfoEntry = OptionalHeader.DataDirectory[IMAGE_DIRECTORY_ENTRY_DEBUG];
+	const auto* DebugEntries = (IMAGE_DEBUG_DIRECTORY*)(Base + DebugInfoEntry.VirtualAddress);
+	for (uint32 i = 0, n = DebugInfoEntry.Size / sizeof(DebugEntries[0]); i < n; ++i)
+	{
+		const IMAGE_DEBUG_DIRECTORY& Entry = DebugEntries[i];
+		if (Entry.Type == IMAGE_DEBUG_TYPE_CODEVIEW)
+		{
+			struct FCodeView7
+			{
+				uint32  Signature;
+				uint32  Guid[4];
+				uint32  Age;
+			};
+
+			if (Entry.SizeOfData < sizeof(FCodeView7))
+			{
+				continue;
+			}
+
+			const auto* CodeView7 = (FCodeView7*)(Base + Entry.AddressOfRawData);
+			if (CodeView7->Signature != 'SDSR')
+			{
+				continue;
+			}
+
+			ImageId.Append((uint8*)&CodeView7->Guid, sizeof(uint32) * 4);
+			ImageId.Append((uint8*)&CodeView7->Age, sizeof(uint32));
+			break;
+		}
+	}
+
+	// Note: UNICODE_STRING.Length is the size in bytes of the string buffer.
+	UE_TRACE_LOG(Diagnostics, ModuleLoad, ModuleChannel, Name.Length + ImageId.Num())
+		<< ModuleLoad.Name(Name.Buffer, Name.Length / 2)
+		<< ModuleLoad.Base(uint64(Base))
+		<< ModuleLoad.Size(OptionalHeader.SizeOfImage)
+		<< ModuleLoad.ImageId(ImageId.GetData(), ImageId.Num());
+
+	for (SubscribeFunc Subscriber : Subscribers)
+	{
+		Subscriber(true, (void*)Base, Name.Buffer);
+	}
+}
+
+////////////////////////////////////////////////////////////////////////////////
+void FModuleTrace::OnDllUnloaded(UPTRINT Base)
+{
+	UE_TRACE_LOG(Diagnostics, ModuleUnload, ModuleChannel)
+		<< ModuleUnload.Base(uint64(Base));
+
+	for (SubscribeFunc Subscriber : Subscribers)
+	{
+		Subscriber(false, (void*)Base, nullptr);
+	}
+}
+
+
+
+////////////////////////////////////////////////////////////////////////////////
+void Modules_Create(FMalloc* Malloc)
+{
+	if (FModuleTrace::Get() != nullptr)
+	{
+		return;
+	}
+
+	static FModuleTrace Instance(Malloc);
+}
+
+////////////////////////////////////////////////////////////////////////////////
+void Modules_Initialize()
+{
+	if (FModuleTrace* Instance = FModuleTrace::Get())
+	{
+		Instance->Initialize();
+	}
+}
+
+////////////////////////////////////////////////////////////////////////////////
+void Modules_Subscribe(void (*Function)(bool, void*, const TCHAR*))
+{
+	if (FModuleTrace* Instance = FModuleTrace::Get())
+	{
+		Instance->Subscribe(Function);
+	}
+}
diff --git a/Engine/Source/Runtime/Core/Public/Experimental/Containers/SherwoodHashTable.h b/Engine/Source/Runtime/Core/Public/Experimental/Containers/SherwoodHashTable.h
--- a/Engine/Source/Runtime/Core/Public/Experimental/Containers/SherwoodHashTable.h
+++ b/Engine/Source/Runtime/Core/Public/Experimental/Containers/SherwoodHashTable.h
@@ -226,7 +226,7 @@ struct TSherwoodHashTable
 
 	FORCENOINLINE ValueType* Add(int8 Distance, FData& Data, uint32 Cursor, KeyType Key, HashType Hash, ValueType Value)
 	{
-		if (Distance == MaxLookups || NumElements + 1 > (NumSlotsMinusOne + 1) * MaxLoadFactor)
+		if (Distance == MaxLookups || NumElements + 1 > int32(float(NumSlotsMinusOne + 1) * MaxLoadFactor))
 		{
 			Grow();
 			return FindOrAddByHash(Key, Hash, Value);
@@ -286,7 +286,7 @@ struct TSherwoodHashTable
 
 	void Rehash(uint32 DesiredNumSlots)
 	{
-		DesiredNumSlots = FMath::Max(DesiredNumSlots, uint32(FMath::CeilToInt(NumElements / MaxLoadFactor)));
+		DesiredNumSlots = FMath::Max(DesiredNumSlots, uint32(FMath::CeilToInt(float(NumElements) / MaxLoadFactor)));
 		if (DesiredNumSlots == 0)
 		{
 			Empty();
@@ -338,7 +338,7 @@ struct TSherwoodHashTable
 
 	void Reserve(uint32 DesiredNumElements)
 	{
-		uint32 DesiredNumSlots = FMath::Max<uint32>(DesiredNumElements, uint32(FMath::CeilToInt(DesiredNumElements / MaxLoadFactor)));
+		uint32 DesiredNumSlots = FMath::Max<uint32>(DesiredNumElements, uint32(FMath::CeilToInt(float(DesiredNumElements) / MaxLoadFactor)));
 		if (DesiredNumSlots > NumSlots())
 		{
 			Rehash(DesiredNumSlots);
diff --git a/Engine/Source/Runtime/Core/Public/ProfilingDebugging/TagTrace.h b/Engine/Source/Runtime/Core/Public/ProfilingDebugging/TagTrace.h
--- a/Engine/Source/Runtime/Core/Public/ProfilingDebugging/TagTrace.h
+++ b/Engine/Source/Runtime/Core/Public/ProfilingDebugging/TagTrace.h
@@ -22,7 +22,7 @@ CORE_API int32	MemoryTrace_GetActiveTag();
 ////////////////////////////////////////////////////////////////////////////////
 
 #if !defined(USE_MEMORY_TRACE_TAGS) && UE_TRACE_ENABLED && UE_BUILD_DEVELOPMENT
-	#if PLATFORM_WINDOWS
+	#if PLATFORM_WINDOWS || PLATFORM_PS4 || defined(__PS5__)
 		#define USE_MEMORY_TRACE_TAGS 1
 	#endif
 #endif
