Ununsed and inconsistently-applied BogusChar option removed from string conversions.

diff --git a/Engine/Source/Runtime/Core/Private/GenericPlatform/GenericPlatformString.cpp b/Engine/Source/Runtime/Core/Private/GenericPlatform/GenericPlatformString.cpp
--- a/Engine/Source/Runtime/Core/Private/GenericPlatform/GenericPlatformString.cpp
+++ b/Engine/Source/Runtime/Core/Private/GenericPlatform/GenericPlatformString.cpp
@@ -185,8 +185,8 @@ namespace UE::Core::Private
 		return true;
 	}
 
-	template <typename DestBufferType, typename FromType, typename DestType>
-	static int32 ConvertToUTF8(DestBufferType& Dest, int32 DestLen, const FromType* Source, const int32 SourceLen, DestType BogusChar)
+	template <typename DestBufferType, typename FromType>
+	static int32 ConvertToUTF8(DestBufferType& Dest, int32 DestLen, const FromType* Source, const int32 SourceLen)
 	{
 		DestBufferType DestStartingPosition = Dest;
 		if constexpr (sizeof(FromType) == 4)
@@ -266,7 +266,7 @@ namespace UE::Core::Private
 	}
 
 	template <typename FromType>
-	static uint32 CodepointFromUtf8(const FromType*& SourceString, const uint32 SourceLengthRemaining, uint32 BogusChar)
+	static uint32 CodepointFromUtf8(const FromType*& SourceString, const uint32 SourceLengthRemaining)
 	{
 		checkSlow(SourceLengthRemaining > 0);
 
@@ -286,7 +286,7 @@ namespace UE::Core::Private
 			// Apparently each of these is supposed to be flagged as a bogus
 			//  char, instead of just resyncing to the next valid codepoint.
 			++SourceString;  // skip to next possible start of codepoint.
-			return BogusChar;
+			return UNICODE_BOGUS_CHAR_CODEPOINT;
 		}
 		else if (Octet < 224)  // two octets
 		{
@@ -295,7 +295,7 @@ namespace UE::Core::Private
 			{
 				// Skip to end and write out a single char (we always have room for at least 1 char)
 				SourceString += SourceLengthRemaining;
-				return BogusChar;
+				return UNICODE_BOGUS_CHAR_CODEPOINT;
 			}
 
 			Octet -= (128+64);
@@ -303,7 +303,7 @@ namespace UE::Core::Private
 			if ((Octet2 & (128 + 64)) != 128)  // Format isn't 10xxxxxx?
 			{
 				++SourceString;  // Sequence was not valid UTF-8. Skip the first byte and continue.
-				return BogusChar;
+				return UNICODE_BOGUS_CHAR_CODEPOINT;
 			}
 
 			Codepoint = ((Octet << 6) | (Octet2 - 128));
@@ -320,7 +320,7 @@ namespace UE::Core::Private
 			{
 				// Skip to end and write out a single char (we always have room for at least 1 char)
 				SourceString += SourceLengthRemaining;
-				return BogusChar;
+				return UNICODE_BOGUS_CHAR_CODEPOINT;
 			}
 
 			Octet -= (128+64+32);
@@ -328,14 +328,14 @@ namespace UE::Core::Private
 			if ((Octet2 & (128+64)) != 128)  // Format isn't 10xxxxxx?
 			{
 				++SourceString;  // Sequence was not valid UTF-8. Skip the first byte and continue.
-				return BogusChar;
+				return UNICODE_BOGUS_CHAR_CODEPOINT;
 			}
 
 			Octet3 = (uint32) ((uint8) *(++OctetPtr));
 			if ((Octet3 & (128+64)) != 128)  // Format isn't 10xxxxxx?
 			{
 				++SourceString;  // Sequence was not valid UTF-8. Skip the first byte and continue.
-				return BogusChar;
+				return UNICODE_BOGUS_CHAR_CODEPOINT;
 			}
 
 			Codepoint = ( ((Octet << 12)) | ((Octet2-128) << 6) | ((Octet3-128)) );
@@ -344,7 +344,7 @@ namespace UE::Core::Private
 			if (IsHighSurrogate(Codepoint) || IsLowSurrogate(Codepoint))
 			{
 				++SourceString;  // Sequence was not valid UTF-8. Skip the first byte and continue.
-				return BogusChar;
+				return UNICODE_BOGUS_CHAR_CODEPOINT;
 			}
 
 			// 0xFFFE and 0xFFFF are illegal, too, so we check them at the edge.
@@ -361,7 +361,7 @@ namespace UE::Core::Private
 			{
 				// Skip to end and write out a single char (we always have room for at least 1 char)
 				SourceString += SourceLengthRemaining;
-				return BogusChar;
+				return UNICODE_BOGUS_CHAR_CODEPOINT;
 			}
 
 			Octet -= (128+64+32+16);
@@ -369,21 +369,21 @@ namespace UE::Core::Private
 			if ((Octet2 & (128+64)) != 128)  // Format isn't 10xxxxxx?
 			{
 				++SourceString;  // Sequence was not valid UTF-8. Skip the first byte and continue.
-				return BogusChar;
+				return UNICODE_BOGUS_CHAR_CODEPOINT;
 			}
 
 			Octet3 = (uint32) ((uint8) *(++OctetPtr));
 			if ((Octet3 & (128+64)) != 128)  // Format isn't 10xxxxxx?
 			{
 				++SourceString;  // Sequence was not valid UTF-8. Skip the first byte and continue.
-				return BogusChar;
+				return UNICODE_BOGUS_CHAR_CODEPOINT;
 			}
 
 			Octet4 = (uint32) ((uint8) *(++OctetPtr));
 			if ((Octet4 & (128+64)) != 128)  // Format isn't 10xxxxxx?
 			{
 				++SourceString;  // Sequence was not valid UTF-8. Skip the first byte and continue.
-				return BogusChar;
+				return UNICODE_BOGUS_CHAR_CODEPOINT;
 			}
 
 			Codepoint = ( ((Octet << 18)) | ((Octet2 - 128) << 12) |
@@ -404,39 +404,39 @@ namespace UE::Core::Private
 			{
 				// Skip to end and write out a single char (we always have room for at least 1 char)
 				SourceString += SourceLengthRemaining;
-				return BogusChar;
+				return UNICODE_BOGUS_CHAR_CODEPOINT;
 			}
 
 			Octet = (uint32) ((uint8) *(++OctetPtr));
 			if ((Octet & (128+64)) != 128)  // Format isn't 10xxxxxx?
 			{
 				++SourceString;  // Sequence was not valid UTF-8. Skip the first byte and continue.
-				return BogusChar;
+				return UNICODE_BOGUS_CHAR_CODEPOINT;
 			}
 
 			Octet = (uint32) ((uint8) *(++OctetPtr));
 			if ((Octet & (128+64)) != 128)  // Format isn't 10xxxxxx?
 			{
 				++SourceString;  // Sequence was not valid UTF-8. Skip the first byte and continue.
-				return BogusChar;
+				return UNICODE_BOGUS_CHAR_CODEPOINT;
 			}
 
 			Octet = (uint32) ((uint8) *(++OctetPtr));
 			if ((Octet & (128+64)) != 128)  // Format isn't 10xxxxxx?
 			{
 				++SourceString;  // Sequence was not valid UTF-8. Skip the first byte and continue.
-				return BogusChar;
+				return UNICODE_BOGUS_CHAR_CODEPOINT;
 			}
 
 			Octet = (uint32) ((uint8) *(++OctetPtr));
 			if ((Octet & (128+64)) != 128)  // Format isn't 10xxxxxx?
 			{
 				++SourceString;  // Sequence was not valid UTF-8. Skip the first byte and continue.
-				return BogusChar;
+				return UNICODE_BOGUS_CHAR_CODEPOINT;
 			}
 
 			SourceString += 5;  // skip to next possible start of codepoint.
-			return BogusChar;
+			return UNICODE_BOGUS_CHAR_CODEPOINT;
 		}
 
 		else  // six octets
@@ -446,57 +446,57 @@ namespace UE::Core::Private
 			{
 				// Skip to end and write out a single char (we always have room for at least 1 char)
 				SourceString += SourceLengthRemaining;
-				return BogusChar;
+				return UNICODE_BOGUS_CHAR_CODEPOINT;
 			}
 
 			Octet = (uint32) ((uint8) *(++OctetPtr));
 			if ((Octet & (128+64)) != 128)  // Format isn't 10xxxxxx?
 			{
 				++SourceString;  // Sequence was not valid UTF-8. Skip the first byte and continue.
-				return BogusChar;
+				return UNICODE_BOGUS_CHAR_CODEPOINT;
 			}
 
 			Octet = (uint32) ((uint8) *(++OctetPtr));
 			if ((Octet & (128+64)) != 128)  // Format isn't 10xxxxxx?
 			{
 				++SourceString;  // Sequence was not valid UTF-8. Skip the first byte and continue.
-				return BogusChar;
+				return UNICODE_BOGUS_CHAR_CODEPOINT;
 			}
 
 			Octet = (uint32) ((uint8) *(++OctetPtr));
 			if ((Octet & (128+64)) != 128)  // Format isn't 10xxxxxx?
 			{
 				++SourceString;  // Sequence was not valid UTF-8. Skip the first byte and continue.
-				return BogusChar;
+				return UNICODE_BOGUS_CHAR_CODEPOINT;
 			}
 
 			Octet = (uint32) ((uint8) *(++OctetPtr));
 			if ((Octet & (128+64)) != 128)  // Format isn't 10xxxxxx?
 			{
 				++SourceString;  // Sequence was not valid UTF-8. Skip the first byte and continue.
-				return BogusChar;
+				return UNICODE_BOGUS_CHAR_CODEPOINT;
 			}
 
 			Octet = (uint32) ((uint8) *(++OctetPtr));
 			if ((Octet & (128+64)) != 128)  // Format isn't 10xxxxxx?
 			{
 				++SourceString;  // Sequence was not valid UTF-8. Skip the first byte and continue.
-				return BogusChar;
+				return UNICODE_BOGUS_CHAR_CODEPOINT;
 			}
 
 			SourceString += 6;  // skip to next possible start of codepoint.
-			return BogusChar;
+			return UNICODE_BOGUS_CHAR_CODEPOINT;
 		}
 
 		++SourceString;  // Sequence was not valid UTF-8. Skip the first byte and continue.
-		return BogusChar;  // catch everything else.
+		return UNICODE_BOGUS_CHAR_CODEPOINT;  // catch everything else.
 	}
 
 	/**
 	 * Read Source string, converting the data from UTF-8 into UTF-16, and placing these in the Destination
 	 */
-	template <typename DestBufferType, typename FromType, typename DestType>
-	static int32 ConvertFromUTF8(DestBufferType& ConvertedBuffer, int32 DestLen, const FromType* Source, const int32 SourceLen, DestType BogusChar)
+	template <typename DestType, typename DestBufferType, typename FromType>
+	static int32 ConvertFromUTF8(DestBufferType& ConvertedBuffer, int32 DestLen, const FromType* Source, const int32 SourceLen)
 	{
 		DestBufferType DestStartingPosition = ConvertedBuffer;
 
@@ -540,7 +540,7 @@ namespace UE::Core::Private
 			while (Source < SourceEnd && DestLen > 0)
 			{
 				// Read our codepoint, advancing the source pointer
-				uint32 Codepoint = CodepointFromUtf8(Source, UE_PTRDIFF_TO_UINT32(SourceEnd - Source), BogusChar);
+				uint32 Codepoint = CodepointFromUtf8(Source, UE_PTRDIFF_TO_UINT32(SourceEnd - Source));
 
 				if constexpr (sizeof(DestType) != 4)
 				{
@@ -561,12 +561,12 @@ namespace UE::Core::Private
 						}
 
 						// If we don't have space, write a bogus character instead (we should have space for it)
-						Codepoint = (uint32)BogusChar;
+						Codepoint = (uint32)UNICODE_BOGUS_CHAR_CODEPOINT;
 					}
 					else if (Codepoint > ENCODED_SURROGATE_END_CODEPOINT)
 					{
 						// Ignore values higher than the supplementary plane range
-						Codepoint = (uint32)BogusChar;
+						Codepoint = (uint32)UNICODE_BOGUS_CHAR_CODEPOINT;
 					}
 				}
 
@@ -592,83 +592,83 @@ namespace UE::Core::Private
 	int32 GetConvertedLength(const UTF8CHAR*, const ANSICHAR* Source, int32 SourceLen)
 	{
 		TCountingOutputIterator<UTF8CHAR> Dest;
-		int32 Result = ConvertToUTF8(Dest, INT32_MAX, Source, SourceLen, (UTF8CHAR)'?');
+		int32 Result = ConvertToUTF8(Dest, INT32_MAX, Source, SourceLen);
 		return Result;
 	}
 	int32 GetConvertedLength(const UTF8CHAR*, const WIDECHAR* Source, int32 SourceLen)
 	{
 		TCountingOutputIterator<UTF8CHAR> Dest;
-		int32 Result = ConvertToUTF8(Dest, INT32_MAX, Source, SourceLen, (UTF8CHAR)'?');
+		int32 Result = ConvertToUTF8(Dest, INT32_MAX, Source, SourceLen);
 		return Result;
 	}
 	int32 GetConvertedLength(const UTF8CHAR*, const UCS2CHAR* Source, int32 SourceLen)
 	{
 		TCountingOutputIterator<UTF8CHAR> Dest;
-		int32 Result = ConvertToUTF8(Dest, INT32_MAX, Source, SourceLen, (UTF8CHAR)'?');
+		int32 Result = ConvertToUTF8(Dest, INT32_MAX, Source, SourceLen);
 		return Result;
 	}
 	int32 GetConvertedLength(const ANSICHAR*, const UTF8CHAR* Source, int32 SourceLen)
 	{
 		TCountingOutputIterator<ANSICHAR> Dest;
-		int32 Result = ConvertFromUTF8(Dest, INT32_MAX, Source, SourceLen, (ANSICHAR)'?');
+		int32 Result = ConvertFromUTF8<ANSICHAR>(Dest, INT32_MAX, Source, SourceLen);
 		return Result;
 	}
 	int32 GetConvertedLength(const WIDECHAR*, const UTF8CHAR* Source, int32 SourceLen)
 	{
 		TCountingOutputIterator<WIDECHAR> Dest;
-		int32 Result = ConvertFromUTF8(Dest, INT32_MAX, Source, SourceLen, (WIDECHAR)'?');
+		int32 Result = ConvertFromUTF8<WIDECHAR>(Dest, INT32_MAX, Source, SourceLen);
 		return Result;
 	}
 	int32 GetConvertedLength(const UCS2CHAR*, const UTF8CHAR* Source, int32 SourceLen)
 	{
 		TCountingOutputIterator<UCS2CHAR> Dest;
-		int32 Result = ConvertFromUTF8(Dest, INT32_MAX, Source, SourceLen, (UCS2CHAR)'?');
+		int32 Result = ConvertFromUTF8<UCS2CHAR>(Dest, INT32_MAX, Source, SourceLen);
 		return Result;
 	}
 
-	UTF8CHAR* Convert(UTF8CHAR* Dest, int32 DestLen, const ANSICHAR* Src, int32 SrcLen, UTF8CHAR BogusChar)
+	UTF8CHAR* Convert(UTF8CHAR* Dest, int32 DestLen, const ANSICHAR* Src, int32 SrcLen)
 	{
-		if (ConvertToUTF8(Dest, DestLen, Src, SrcLen, BogusChar) == -1)
+		if (ConvertToUTF8(Dest, DestLen, Src, SrcLen) == -1)
 		{
 			return nullptr;
 		}
 		return Dest;
 	}
-	UTF8CHAR* Convert(UTF8CHAR* Dest, int32 DestLen, const WIDECHAR* Src, int32 SrcLen, UTF8CHAR BogusChar)
+	UTF8CHAR* Convert(UTF8CHAR* Dest, int32 DestLen, const WIDECHAR* Src, int32 SrcLen)
 	{
-		if (ConvertToUTF8(Dest, DestLen, Src, SrcLen, BogusChar) == -1)
+		if (ConvertToUTF8(Dest, DestLen, Src, SrcLen) == -1)
 		{
 			return nullptr;
 		}
 		return Dest;
 	}
-	UTF8CHAR* Convert(UTF8CHAR* Dest, int32 DestLen, const UCS2CHAR* Src, int32 SrcLen, UTF8CHAR BogusChar)
+	UTF8CHAR* Convert(UTF8CHAR* Dest, int32 DestLen, const UCS2CHAR* Src, int32 SrcLen)
 	{
-		if (ConvertToUTF8(Dest, DestLen, Src, SrcLen, BogusChar) == -1)
+		if (ConvertToUTF8(Dest, DestLen, Src, SrcLen) == -1)
 		{
 			return nullptr;
 		}
 		return Dest;
 	}
-	ANSICHAR* Convert(ANSICHAR* Dest, int32 DestLen, const UTF8CHAR* Src, int32 SrcLen, ANSICHAR BogusChar)
+	ANSICHAR* Convert(ANSICHAR* Dest, int32 DestLen, const UTF8CHAR* Src, int32 SrcLen)
 	{
-		if (ConvertFromUTF8(Dest, DestLen, Src, SrcLen, BogusChar) == -1)
+		if (ConvertFromUTF8<ANSICHAR>(Dest, DestLen, Src, SrcLen) == -1)
 		{
 			return nullptr;
 		}
 		return Dest;
 	}
-	WIDECHAR* Convert(WIDECHAR* Dest, int32 DestLen, const UTF8CHAR* Src, int32 SrcLen, WIDECHAR BogusChar)
+	WIDECHAR* Convert(WIDECHAR* Dest, int32 DestLen, const UTF8CHAR* Src, int32 SrcLen)
 	{
-		if (ConvertFromUTF8(Dest, DestLen, Src, SrcLen, BogusChar) == -1)
+		if (ConvertFromUTF8<WIDECHAR>(Dest, DestLen, Src, SrcLen) == -1)
 		{
 			return nullptr;
 		}
 		return Dest;
 	}
-	UCS2CHAR* Convert(UCS2CHAR* Dest, int32 DestLen, const UTF8CHAR* Src, int32 SrcLen, UCS2CHAR BogusChar)
+	UCS2CHAR* Convert(UCS2CHAR* Dest, int32 DestLen, const UTF8CHAR* Src, int32 SrcLen)
 	{
-		if (ConvertFromUTF8(Dest, DestLen, Src, SrcLen, BogusChar) == -1)
+		if (ConvertFromUTF8<UCS2CHAR>(Dest, DestLen, Src, SrcLen) == -1)
 		{
 			return nullptr;
 		}
diff --git a/Engine/Source/Runtime/Core/Public/Containers/StringConv.h b/Engine/Source/Runtime/Core/Public/Containers/StringConv.h
--- a/Engine/Source/Runtime/Core/Public/Containers/StringConv.h
+++ b/Engine/Source/Runtime/Core/Public/Containers/StringConv.h
@@ -22,8 +22,6 @@
 #include <type_traits>
 
 #define DEFAULT_STRING_CONVERSION_SIZE 128u
-#define UNICODE_BOGUS_CHAR_CODEPOINT '?'
-static_assert(sizeof(UNICODE_BOGUS_CHAR_CODEPOINT) <= sizeof(ANSICHAR) && (UNICODE_BOGUS_CHAR_CODEPOINT) >= 32 && (UNICODE_BOGUS_CHAR_CODEPOINT) <= 127, "The Unicode Bogus character point is expected to fit in a single ANSICHAR here");
 
 template <typename From, typename To>
 class TStringConvert
@@ -38,7 +36,7 @@ public:
 	>
 	FORCEINLINE static void Convert(To* Dest, int32 DestLen, const CharType* Source, int32 SourceLen)
 	{
-		To* Result = FPlatformString::Convert(Dest, DestLen, (const FromType*)Source, SourceLen, (To)UNICODE_BOGUS_CHAR_CODEPOINT);
+		To* Result = FPlatformString::Convert(Dest, DestLen, (const FromType*)Source, SourceLen);
 		check(Result);
 	}
 
@@ -1074,7 +1072,7 @@ template <typename To, typename From>
 FORCEINLINE To CharCast(From Ch)
 {
 	To Result;
-	FPlatformString::Convert(&Result, 1, &Ch, 1, (To)UNICODE_BOGUS_CHAR_CODEPOINT);
+	FPlatformString::Convert(&Result, 1, &Ch, 1);
 	return Result;
 }
 
diff --git a/Engine/Source/Runtime/Core/Public/GenericPlatform/GenericPlatformString.h b/Engine/Source/Runtime/Core/Public/GenericPlatform/GenericPlatformString.h
--- a/Engine/Source/Runtime/Core/Public/GenericPlatform/GenericPlatformString.h
+++ b/Engine/Source/Runtime/Core/Public/GenericPlatform/GenericPlatformString.h
@@ -18,12 +18,12 @@ namespace UE::Core::Private
 	CORE_API int32 GetConvertedLength(const WIDECHAR* Dest, const UTF8CHAR* Src, int32 SrcLen);
 	CORE_API int32 GetConvertedLength(const UCS2CHAR* Dest, const UTF8CHAR* Src, int32 SrcLen);
 
-	CORE_API UTF8CHAR* Convert(UTF8CHAR* Dest, int32 DestLen, const ANSICHAR* Src, int32 SrcLen, UTF8CHAR BogusChar);
-	CORE_API UTF8CHAR* Convert(UTF8CHAR* Dest, int32 DestLen, const WIDECHAR* Src, int32 SrcLen, UTF8CHAR BogusChar);
-	CORE_API UTF8CHAR* Convert(UTF8CHAR* Dest, int32 DestLen, const UCS2CHAR* Src, int32 SrcLen, UTF8CHAR BogusChar);
-	CORE_API ANSICHAR* Convert(ANSICHAR* Dest, int32 DestLen, const UTF8CHAR* Src, int32 SrcLen, ANSICHAR BogusChar);
-	CORE_API WIDECHAR* Convert(WIDECHAR* Dest, int32 DestLen, const UTF8CHAR* Src, int32 SrcLen, WIDECHAR BogusChar);
-	CORE_API UCS2CHAR* Convert(UCS2CHAR* Dest, int32 DestLen, const UTF8CHAR* Src, int32 SrcLen, UCS2CHAR BogusChar);
+	CORE_API UTF8CHAR* Convert(UTF8CHAR* Dest, int32 DestLen, const ANSICHAR* Src, int32 SrcLen);
+	CORE_API UTF8CHAR* Convert(UTF8CHAR* Dest, int32 DestLen, const WIDECHAR* Src, int32 SrcLen);
+	CORE_API UTF8CHAR* Convert(UTF8CHAR* Dest, int32 DestLen, const UCS2CHAR* Src, int32 SrcLen);
+	CORE_API ANSICHAR* Convert(ANSICHAR* Dest, int32 DestLen, const UTF8CHAR* Src, int32 SrcLen);
+	CORE_API WIDECHAR* Convert(WIDECHAR* Dest, int32 DestLen, const UTF8CHAR* Src, int32 SrcLen);
+	CORE_API UCS2CHAR* Convert(UCS2CHAR* Dest, int32 DestLen, const UTF8CHAR* Src, int32 SrcLen);
 }
 
 // These will be moved inside GenericPlatformString.cpp when the platform layer handles UTF-16
@@ -35,6 +35,9 @@ namespace UE::Core::Private
 #define ENCODED_SURROGATE_START_CODEPOINT ((uint32)0x10000)
 #define ENCODED_SURROGATE_END_CODEPOINT   ((uint32)0x10FFFF)
 
+#define UNICODE_BOGUS_CHAR_CODEPOINT '?'
+static_assert(sizeof(UNICODE_BOGUS_CHAR_CODEPOINT) <= sizeof(ANSICHAR) && (UNICODE_BOGUS_CHAR_CODEPOINT) >= 32 && (UNICODE_BOGUS_CHAR_CODEPOINT) <= 127, "The Unicode Bogus character point is expected to fit in a single ANSICHAR here");
+
 /**
  * Generic string implementation for most platforms
  */
@@ -195,11 +198,10 @@ struct FGenericPlatformString : public FGenericPlatformStricmp
 	 * @param DestSize  The size of the destination buffer.
 	 * @param Src       The start of the string to convert.
 	 * @param SrcSize   The number of Src elements to convert.
-	 * @param BogusChar The char to use when the conversion process encounters a character it cannot convert.
 	 * @return          A pointer to one past the last-written element.
 	 */
 	template <typename SourceEncoding, typename DestEncoding>
-	static FORCEINLINE DestEncoding* Convert(DestEncoding* Dest, int32 DestSize, const SourceEncoding* Src, int32 SrcSize, DestEncoding BogusChar = (DestEncoding)'?')
+	static FORCEINLINE DestEncoding* Convert(DestEncoding* Dest, int32 DestSize, const SourceEncoding* Src, int32 SrcSize)
 	{
 		if constexpr (IsCharEncodingCompatibleWith<SourceEncoding, DestEncoding>())
 		{
@@ -227,7 +229,7 @@ struct FGenericPlatformString : public FGenericPlatformStricmp
 				{
 					if (!CanConvertChar<DestEncoding>(Src[I]))
 					{
-						Dest[I] = BogusChar;
+						Dest[I] = UNICODE_BOGUS_CHAR_CODEPOINT;
 					}
 				}
 
@@ -238,7 +240,7 @@ struct FGenericPlatformString : public FGenericPlatformStricmp
 		}
 		else
 		{
-			return UE::Core::Private::Convert(Dest, DestSize, Src, SrcSize, BogusChar);
+			return UE::Core::Private::Convert(Dest, DestSize, Src, SrcSize);
 		}
 	}
 
