WriteCodepointToBuffer and WriteCodepointToUTF8 merged and simplified.

diff --git a/Engine/Source/Runtime/Core/Private/GenericPlatform/GenericPlatformString.cpp b/Engine/Source/Runtime/Core/Private/GenericPlatform/GenericPlatformString.cpp
--- a/Engine/Source/Runtime/Core/Private/GenericPlatform/GenericPlatformString.cpp
+++ b/Engine/Source/Runtime/Core/Private/GenericPlatform/GenericPlatformString.cpp
@@ -113,75 +113,73 @@ namespace UE::Core::Private
 	 * @param Codepoint Codepoint to expand into UTF-8 bytes
 	 * @param OutputIterator Output iterator to write UTF-8 bytes into
 	 * @param OutputIteratorByteSizeRemaining Maximum number of ANSI characters that can be written to OutputIterator
-	 * @return Number of characters written for Codepoint
+	 * @return true if some data was written, false otherwise.
 	 */
 	template <typename BufferType>
-	static int32 WriteCodepointToUTF8(uint32 Codepoint, BufferType OutputIterator, uint32 OutputIteratorByteSizeRemaining)
+	static bool WriteCodepointToBuffer(uint32 Codepoint, BufferType& OutputIterator, int32& OutputIteratorByteSizeRemaining)
 	{
 		// Ensure we have at least one character in size to write
-		if (OutputIteratorByteSizeRemaining < sizeof(UTF8CHAR))
+		if (OutputIteratorByteSizeRemaining == 0)
 		{
-			return 0;
+			return false;
 		}
 
-		const BufferType OutputIteratorStartPosition = OutputIterator;
-
 		if (!IsValidCodepoint(Codepoint))
 		{
-			Codepoint = UNICODE_BOGUS_CHAR_CODEPOINT;
+			Codepoint = (uint32)UNICODE_BOGUS_CHAR_CODEPOINT;
 		}
 		else if (IsHighSurrogate(Codepoint) || IsLowSurrogate(Codepoint)) // UTF-8 Characters are not allowed to encode codepoints in the surrogate pair range
 		{
-			Codepoint = UNICODE_BOGUS_CHAR_CODEPOINT;
+			Codepoint = (uint32)UNICODE_BOGUS_CHAR_CODEPOINT;
 		}
 
 		// Do the encoding...
 		if (Codepoint < 0x80)
 		{
-			*(OutputIterator++) = (UTF8CHAR)Codepoint;
+			*OutputIterator++ = (UTF8CHAR)Codepoint;
+
+			OutputIteratorByteSizeRemaining -= 1;
 		}
 		else if (Codepoint < 0x800)
 		{
-			if (OutputIteratorByteSizeRemaining >= 2)
+			if (OutputIteratorByteSizeRemaining < 2)
 			{
-				*(OutputIterator++) = (UTF8CHAR)((Codepoint >> 6)         | 128 | 64);
-				*(OutputIterator++) = (UTF8CHAR)((Codepoint       & 0x3F) | 128);
+				return false;
 			}
+
+			*OutputIterator++ = (UTF8CHAR)((Codepoint >> 6)         | 128 | 64);
+			*OutputIterator++ = (UTF8CHAR)((Codepoint       & 0x3F) | 128);
+
+			OutputIteratorByteSizeRemaining -= 2;
 		}
 		else if (Codepoint < 0x10000)
 		{
-			if (OutputIteratorByteSizeRemaining >= 3)
+			if (OutputIteratorByteSizeRemaining < 3)
 			{
-				*(OutputIterator++) = (UTF8CHAR)( (Codepoint >> 12)        | 128 | 64 | 32);
-				*(OutputIterator++) = (UTF8CHAR)(((Codepoint >> 6) & 0x3F) | 128);
-				*(OutputIterator++) = (UTF8CHAR)( (Codepoint       & 0x3F) | 128);
+				return false;
 			}
+
+			*OutputIterator++ = (UTF8CHAR)( (Codepoint >> 12)        | 128 | 64 | 32);
+			*OutputIterator++ = (UTF8CHAR)(((Codepoint >> 6) & 0x3F) | 128);
+			*OutputIterator++ = (UTF8CHAR)( (Codepoint       & 0x3F) | 128);
+
+			OutputIteratorByteSizeRemaining -= 3;
 		}
 		else
 		{
-			if (OutputIteratorByteSizeRemaining >= 4)
+			if (OutputIteratorByteSizeRemaining < 4)
 			{
-				*(OutputIterator++) = (UTF8CHAR)( (Codepoint >> 18)         | 128 | 64 | 32 | 16);
-				*(OutputIterator++) = (UTF8CHAR)(((Codepoint >> 12) & 0x3F) | 128);
-				*(OutputIterator++) = (UTF8CHAR)(((Codepoint >> 6 ) & 0x3F) | 128);
-				*(OutputIterator++) = (UTF8CHAR)( (Codepoint        & 0x3F) | 128);
+				return false;
 			}
-		}
 
-		return UE_PTRDIFF_TO_INT32(OutputIterator - OutputIteratorStartPosition);
-	}
+			*OutputIterator++ = (UTF8CHAR)( (Codepoint >> 18)         | 128 | 64 | 32 | 16);
+			*OutputIterator++ = (UTF8CHAR)(((Codepoint >> 12) & 0x3F) | 128);
+			*OutputIterator++ = (UTF8CHAR)(((Codepoint >> 6 ) & 0x3F) | 128);
+			*OutputIterator++ = (UTF8CHAR)( (Codepoint        & 0x3F) | 128);
 
-	template <typename DestBufferType>
-	static bool WriteCodepointToBuffer(const uint32 Codepoint, DestBufferType& Dest, int32& DestLen)
-	{
-		int32 WrittenChars = WriteCodepointToUTF8(Codepoint, Dest, DestLen);
-		if (WrittenChars < 1)
-		{
-			return false;
+			OutputIteratorByteSizeRemaining -= 4;
 		}
 
-		Dest += WrittenChars;
-		DestLen -= WrittenChars;
 		return true;
 	}
 
