Redundant template parameter removed.

diff --git a/Engine/Source/Runtime/Core/Private/GenericPlatform/GenericPlatformString.cpp b/Engine/Source/Runtime/Core/Private/GenericPlatform/GenericPlatformString.cpp
--- a/Engine/Source/Runtime/Core/Private/GenericPlatform/GenericPlatformString.cpp
+++ b/Engine/Source/Runtime/Core/Private/GenericPlatform/GenericPlatformString.cpp
@@ -493,12 +493,12 @@ namespace UE::Core::Private
 	/**
 	 * Read Source string, converting the data from UTF-8 into UTF-16, and placing these in the Destination
 	 */
-	template <typename DestType, typename DestBufferType, typename FromType>
-	static int32 ConvertFromUTF8(DestBufferType& ConvertedBuffer, int32 DestLen, const FromType* Source, const int32 SourceLen)
+	template <typename DestType, typename DestBufferType>
+	static int32 ConvertFromUTF8(DestBufferType& ConvertedBuffer, int32 DestLen, const UTF8CHAR* Source, const int32 SourceLen)
 	{
 		DestBufferType DestStartingPosition = ConvertedBuffer;
 
-		const FromType* SourceEnd = Source + SourceLen;
+		const UTF8CHAR* SourceEnd = Source + SourceLen;
 
 		const uint64 ExtendedCharMask = 0x8080808080808080;
 		while (Source < SourceEnd)
