LWC: Supress invalid SHBasisFunction possible loss of data warning.

diff --git a/Engine/Source/Runtime/Core/Public/Math/SHMath.h b/Engine/Source/Runtime/Core/Public/Math/SHMath.h
--- a/Engine/Source/Runtime/Core/Public/Math/SHMath.h
+++ b/Engine/Source/Runtime/Core/Public/Math/SHMath.h
@@ -369,7 +369,7 @@ public:
 		// Multiply the result by the theta-dependent part of the SH bases.
 		for (int32 BasisIndex = 1; BasisIndex < TSHVector::MaxSHBasis; BasisIndex++)
 		{
-			Result.V[BasisIndex] *= LegendrePolynomial(BasisL[BasisIndex], FMath::Abs(BasisM[BasisIndex]), Vector.Z);
+			Result.V[BasisIndex] *= LegendrePolynomial(BasisL[BasisIndex], FMath::Abs(BasisM[BasisIndex]), (float)Vector.Z);
 		}
 
 		return Result;
