Unnecessary template removed.

diff --git a/Engine/Source/Runtime/Core/Private/GenericPlatform/GenericPlatformString.cpp b/Engine/Source/Runtime/Core/Private/GenericPlatform/GenericPlatformString.cpp
--- a/Engine/Source/Runtime/Core/Private/GenericPlatform/GenericPlatformString.cpp
+++ b/Engine/Source/Runtime/Core/Private/GenericPlatform/GenericPlatformString.cpp
@@ -263,12 +263,11 @@ namespace UE::Core::Private
 		return UE_PTRDIFF_TO_INT32(Dest - DestStartingPosition);
 	}
 
-	template <typename FromType>
-	static uint32 CodepointFromUtf8(const FromType*& SourceString, const uint32 SourceLengthRemaining)
+	static uint32 CodepointFromUtf8(const UTF8CHAR*& SourceString, const uint32 SourceLengthRemaining)
 	{
 		checkSlow(SourceLengthRemaining > 0);
 
-		const FromType* OctetPtr = SourceString;
+		const UTF8CHAR* OctetPtr = SourceString;
 
 		uint32 Codepoint = 0;
 		uint32 Octet = (uint32) ((uint8) *SourceString);
