add -doublefreefinder command line option on windows

diff --git a/Engine/Source/Runtime/Core/Private/Windows/WindowsPlatformMemory.cpp b/Engine/Source/Runtime/Core/Private/Windows/WindowsPlatformMemory.cpp
--- a/Engine/Source/Runtime/Core/Private/Windows/WindowsPlatformMemory.cpp
+++ b/Engine/Source/Runtime/Core/Private/Windows/WindowsPlatformMemory.cpp
@@ -21,6 +21,7 @@
 #include "HAL/MallocBinned2.h"
 #include "HAL/MallocBinned3.h"
 #include "HAL/MallocStomp2.h"
+#include "HAL/MallocDoubleFreeFinder.h"
 #include "Windows/WindowsHWrapper.h"
 
 #pragma warning(disable:6250)
@@ -173,6 +174,12 @@ FMalloc* FWindowsPlatformMemory::BaseAllocator()
 		GMallocStomp2Enabled = true;
 	}
 #endif // WITH_MALLOC_STOMP2
+
+	if (FCString::Stristr(CommandLine, TEXT("-doublefreefinder")))
+	{
+		GMallocDoubleFreeFinderEnabled = true;
+	}
+
 #endif // !UE_BUILD_SHIPPING
 
 	switch (AllocatorToUse)
