FPlatformString::CanConvertChar() changed to FPlatformString::CanConvertCodepoint(), with improved checks and documented assumptions.

diff --git a/Engine/Source/Runtime/Core/Private/GenericPlatform/GenericPlatformString.cpp b/Engine/Source/Runtime/Core/Private/GenericPlatform/GenericPlatformString.cpp
--- a/Engine/Source/Runtime/Core/Private/GenericPlatform/GenericPlatformString.cpp
+++ b/Engine/Source/Runtime/Core/Private/GenericPlatform/GenericPlatformString.cpp
@@ -676,17 +676,19 @@ namespace UE::Core::Private
 template <typename DestEncoding, typename SourceEncoding>
 void FGenericPlatformString::LogBogusChars(const SourceEncoding* Src, int32 SrcSize)
 {
+	static_assert(IsFixedWidthEncoding<SourceEncoding>(), "Currently unimplemented for non-fixed-width source conversions");
+
 	FString SrcStr;
 	bool    bFoundBogusChars = false;
 	for (; SrcSize; --SrcSize)
 	{
 		SourceEncoding SrcCh = *Src++;
-		if (!CanConvertChar<DestEncoding>(SrcCh))
+		if (!CanConvertCodepoint<DestEncoding>(SrcCh))
 		{
 			SrcStr += FString::Printf(TEXT("[0x%X]"), (int32)SrcCh);
 			bFoundBogusChars = true;
 		}
-		else if (CanConvertChar<TCHAR>(SrcCh))
+		else if (CanConvertCodepoint<TCHAR>(SrcCh))
 		{
 			if (TChar<SourceEncoding>::IsLinebreak(SrcCh))
 			{
@@ -753,13 +755,10 @@ int32 FGenericPlatformString::Strncmp(const UTF8CHAR* Str1, const UTF8CHAR* Str2
 #if !UE_BUILD_DOCS
 template CORE_API void FGenericPlatformString::LogBogusChars<ANSICHAR, WIDECHAR>(const WIDECHAR* Src, int32 SrcSize);
 template CORE_API void FGenericPlatformString::LogBogusChars<ANSICHAR, UCS2CHAR>(const UCS2CHAR* Src, int32 SrcSize);
-template CORE_API void FGenericPlatformString::LogBogusChars<ANSICHAR, UTF8CHAR>(const UTF8CHAR* Src, int32 SrcSize);
 template CORE_API void FGenericPlatformString::LogBogusChars<WIDECHAR, ANSICHAR>(const ANSICHAR* Src, int32 SrcSize);
 template CORE_API void FGenericPlatformString::LogBogusChars<WIDECHAR, UCS2CHAR>(const UCS2CHAR* Src, int32 SrcSize);
-template CORE_API void FGenericPlatformString::LogBogusChars<WIDECHAR, UTF8CHAR>(const UTF8CHAR* Src, int32 SrcSize);
 template CORE_API void FGenericPlatformString::LogBogusChars<UCS2CHAR, ANSICHAR>(const ANSICHAR* Src, int32 SrcSize);
 template CORE_API void FGenericPlatformString::LogBogusChars<UCS2CHAR, WIDECHAR>(const WIDECHAR* Src, int32 SrcSize);
-template CORE_API void FGenericPlatformString::LogBogusChars<UCS2CHAR, UTF8CHAR>(const UTF8CHAR* Src, int32 SrcSize);
 template CORE_API void FGenericPlatformString::LogBogusChars<UTF8CHAR, ANSICHAR>(const ANSICHAR* Src, int32 SrcSize);
 template CORE_API void FGenericPlatformString::LogBogusChars<UTF8CHAR, WIDECHAR>(const WIDECHAR* Src, int32 SrcSize);
 template CORE_API void FGenericPlatformString::LogBogusChars<UTF8CHAR, UCS2CHAR>(const UCS2CHAR* Src, int32 SrcSize);
diff --git a/Engine/Source/Runtime/Core/Public/GenericPlatform/GenericPlatformString.h b/Engine/Source/Runtime/Core/Public/GenericPlatform/GenericPlatformString.h
--- a/Engine/Source/Runtime/Core/Public/GenericPlatform/GenericPlatformString.h
+++ b/Engine/Source/Runtime/Core/Public/GenericPlatform/GenericPlatformString.h
@@ -62,28 +62,45 @@ struct FGenericPlatformString : public FGenericPlatformStricmp
 	}
 
 	/**
-	 * Tests whether a particular character is a valid member of its encoding.
-	 *
-	 * @param Ch The character to test.
-	 * @return True if the character is a valid member of Encoding.
-	 */
-	template <typename Encoding>
-	static bool IsValidChar(Encoding Ch)
-	{
-		return true;
-	}
-
-
-	/**
-	 * Tests whether a particular character can be converted to the destination encoding.
+	 * Tests whether a particular codepoint can be converted to the destination encoding.
 	 *
 	 * @param Ch The character to test.
 	 * @return True if Ch can be encoded as a DestEncoding.
 	 */
 	template <typename DestEncoding, typename SourceEncoding>
-	static bool CanConvertChar(SourceEncoding Ch)
+	static constexpr bool CanConvertCodepoint(SourceEncoding Codepoint)
 	{
-		return IsValidChar(Ch) && (SourceEncoding)(DestEncoding)Ch == Ch && IsValidChar((DestEncoding)Ch);
+		// It is assumed that the incoming codepoint is already valid and we're only testing if it can be converted to DestEncoding.
+
+		static_assert(TIsCharType<SourceEncoding>::Value, "Source encoding is not a char type");
+		static_assert(TIsCharType<DestEncoding  >::Value, "Destination encoding is not a char type");
+
+		// This is only defined for fixed-width encodings, because codepoints cannot be represented in a single variable-width code unit.
+		static_assert(IsFixedWidthEncoding<SourceEncoding>(), "Source encoding is not fixed-width");
+
+		if constexpr (sizeof(DestEncoding) >= sizeof(SourceEncoding))
+		{
+			// Converting to a wider encoding, or same encoding, or ANSICHAR->UTF8CHAR or UCS2CHAR->UTF16CHAR,
+			// should always be possible.
+			return true;
+		}
+		else if constexpr (!IsFixedWidthEncoding<DestEncoding>())
+		{
+			// Converting all codepoints to a variable-width encoding should always be possible
+			return true;
+		}
+		else if constexpr (std::is_same_v<DestEncoding, ANSICHAR>)
+		{
+			return (uint32)Codepoint <= 0x7F;
+		}
+		else
+		{
+			// The logic above should hopefully mean this branch is only taken for UTF32CHAR->UCS2CHAR
+			static_assert(std::is_same_v<SourceEncoding, UTF32CHAR> && std::is_same_v<DestEncoding, UCS2CHAR>, "Unimplemented conversion");
+
+			// Can't encode more than 16-bit in UCS-2
+			return (uint32)Codepoint <= 0xFFFF;
+		}
 	}
 
 
@@ -220,14 +237,14 @@ struct FGenericPlatformString : public FGenericPlatformStricmp
 			{
 				SourceEncoding SrcCh = Src[I];
 				Dest[I] = (DestEncoding)SrcCh;
-				bInvalidChars |= !CanConvertChar<DestEncoding>(SrcCh);
+				bInvalidChars |= !CanConvertCodepoint<DestEncoding>(SrcCh);
 			}
 
 			if (bInvalidChars)
 			{
 				for (int I = 0; I < Size; ++I)
 				{
-					if (!CanConvertChar<DestEncoding>(Src[I]))
+					if (!CanConvertCodepoint<DestEncoding>(Src[I]))
 					{
 						Dest[I] = UNICODE_BOGUS_CHAR_CODEPOINT;
 					}
@@ -297,13 +314,3 @@ private:
 	template <typename DestEncoding, typename SourceEncoding>
 	static CORE_API void LogBogusChars(const SourceEncoding* Src, int32 SrcSize);
 };
-
-
-/**
- * Specialization of IsValidChar for ANSICHARs.
- */
-template <>
-inline bool FGenericPlatformString::IsValidChar<ANSICHAR>(ANSICHAR Ch)
-{
-	return Ch >= 0x00 && Ch <= 0x7F;
-}
