MemPro integration for a newer platorm

diff --git a/Engine/Source/Runtime/Core/Private/MemPro/MemPro.cpp b/Engine/Source/Runtime/Core/Private/MemPro/MemPro.cpp
--- a/Engine/Source/Runtime/Core/Private/MemPro/MemPro.cpp
+++ b/Engine/Source/Runtime/Core/Private/MemPro/MemPro.cpp
@@ -4917,6 +4917,11 @@ namespace MemPro
 
 	// implemented in MemProPS4.cpp - contact [REDACTED_EMAIL] for this platform
 
+//------------------------------------------------------------------------
+#elif defined(MEMPRO_PLATFORM_PS5)
+
+// implemented in MemProPS5.cpp - contact [REDACTED_EMAIL] for this platform
+
 //------------------------------------------------------------------------
 #elif defined(MEMPRO_PLATFORM_UNIX)
 
diff --git a/Engine/Source/Runtime/Core/Public/MemPro/MemPro.h b/Engine/Source/Runtime/Core/Public/MemPro/MemPro.h
--- a/Engine/Source/Runtime/Core/Public/MemPro/MemPro.h
+++ b/Engine/Source/Runtime/Core/Public/MemPro/MemPro.h
@@ -323,6 +323,9 @@ namespace
 
 	#include "MemPro/MemProSwitch.h"
 
+#elif defined(MEMPRO_PLATFORM_PS5)
+
+	#include "MemPro/MemProPS5.h"
 //@EPIC END
 
 #elif defined(MEMPRO_PLATFORM_UNIX)
