Add UE_CALL_ONCE

diff --git a/Engine/Source/Runtime/Core/Public/Misc/CoreMiscDefines.h b/Engine/Source/Runtime/Core/Public/Misc/CoreMiscDefines.h
--- a/Engine/Source/Runtime/Core/Public/Misc/CoreMiscDefines.h
+++ b/Engine/Source/Runtime/Core/Public/Misc/CoreMiscDefines.h
@@ -139,6 +139,9 @@ const FPlatformUserId PLATFORMUSERID_NONE = INDEX_NONE;
 	#define ANONYMOUS_VARIABLE( Name ) PREPROCESSOR_JOIN(Name, __LINE__)
 #endif
 
+/** Thread-safe call once helper, similar to std::call_once without the std::once_flag */
+#define UE_CALL_ONCE(Func) static int32 ANONYMOUS_VARIABLE(ThreadSafeOnce) = ((Func)(), 1)
+
 /**
  * Macro for marking up deprecated code, functions and types.
  *
