Tasks: deep retraction - a retraction that follows task dependencies (prerequisites and nested tasks). When a task is waited for and it's blocked by dependencies, the dependencies will be tried to be retracted first.

diff --git a/Engine/Source/Runtime/Core/Private/Tasks/Pipe.cpp b/Engine/Source/Runtime/Core/Private/Tasks/Pipe.cpp
--- a/Engine/Source/Runtime/Core/Private/Tasks/Pipe.cpp
+++ b/Engine/Source/Runtime/Core/Private/Tasks/Pipe.cpp
@@ -6,7 +6,7 @@
 
 namespace UE { namespace Tasks
 {
-	bool FPipe::PushIntoPipe(Private::FTaskBase& Task)
+	Private::FTaskBase* FPipe::PushIntoPipe(Private::FTaskBase& Task)
 	{
 		// critical section: between exchanging tasks and adding the previous LastTask to prerequisites.
 		// `LastTask` can be in the process of destruction. Before this it will try to clear itself from the pipe
@@ -17,7 +17,13 @@ namespace UE { namespace Tasks
 		Private::FTaskBase* LastTask_Local = LastTask.exchange(&Task, std::memory_order_release);
 		checkf(LastTask_Local != &Task, TEXT("Dependency cycle: adding itself as a prerequisite (or use after destruction)"));
 
-		return LastTask_Local == nullptr || !LastTask_Local->AddSubsequent(Task);
+		if (LastTask_Local == nullptr || !LastTask_Local->AddSubsequent(Task))
+		{
+			return nullptr;
+		}
+
+		LastTask_Local->AddRef(); // keep it alive, it's the caller's duty to release it
+		return LastTask_Local;
 	}
 
 	void FPipe::ClearTask(Private::FTaskBase& Task)
diff --git a/Engine/Source/Runtime/Core/Private/Tasks/TaskPrivate.cpp b/Engine/Source/Runtime/Core/Private/Tasks/TaskPrivate.cpp
--- a/Engine/Source/Runtime/Core/Private/Tasks/TaskPrivate.cpp
+++ b/Engine/Source/Runtime/Core/Private/Tasks/TaskPrivate.cpp
@@ -5,7 +5,7 @@
 
 namespace UE { namespace Tasks { namespace Private
 {
-	bool FTaskBase::PushIntoPipe()
+	FTaskBase* FTaskBase::PushIntoPipe()
 	{
 		return GetPipe()->PushIntoPipe(*this);
 	}
@@ -39,7 +39,12 @@ namespace UE { namespace Tasks { namespace Private
 				Subsequent->TryUnlock();
 			}
 		);
-	
+
+		while (TOptional<FTaskBase*> Prerequisite = Prerequisites.Dequeue())
+		{
+			Prerequisite.GetValue()->Release();
+		}
+
 		Release(); // release the reference accounted for nested tasks. the task can be destroyed as the result
 	}
 
diff --git a/Engine/Source/Runtime/Core/Private/Tests/Tasks/TasksTest.cpp b/Engine/Source/Runtime/Core/Private/Tests/Tasks/TasksTest.cpp
--- a/Engine/Source/Runtime/Core/Private/Tests/Tasks/TasksTest.cpp
+++ b/Engine/Source/Runtime/Core/Private/Tests/Tasks/TasksTest.cpp
@@ -809,6 +809,156 @@ namespace UE { namespace TasksTests
 		return true;
 	}
 
+	// blocks all workers (except reserve workers) until given event is triggered. Returns blocking tasks.
+	TArray<FTask> BlockWorkers(FTaskEvent& ResumeEvent)
+	{
+		FPlatformProcess::Sleep(0.1f); // give workers time to fall asleep, to avoid any reserve worker messing around
+
+		uint32 NumWorkers = LowLevelTasks::FScheduler::Get().GetNumWorkers();
+
+		TArray<FTask> WorkerBlockers; // tasks that block worker threads
+		WorkerBlockers.Reserve(NumWorkers);
+
+		std::atomic<uint32> NumWorkersNotBlocked{ NumWorkers };
+
+		for (int i = 0; i != NumWorkers; ++i)
+		{
+			checkf(LowLevelTasks::FScheduler::Get().IsWorkerThread(), TEXT("No reserve workers are expected to get blocked"));
+			WorkerBlockers.Add(Launch(UE_SOURCE_LOCATION,
+				[&NumWorkersNotBlocked, &ResumeEvent]
+				{
+					--NumWorkersNotBlocked;
+					ResumeEvent.Wait();
+				}
+			));
+		}
+
+		UE::FTimeout Timeout{ FTimespan::FromSeconds(1) };
+		while (NumWorkersNotBlocked != 0)
+		{
+			check(!Timeout);
+			FPlatformProcess::Sleep(0.001f);
+		}
+
+		return WorkerBlockers;
+	}
+
+	// two levels of prerequisites and two levels of nested tasks
+	void TwoLevelsDeepRetractionTest()
+	{
+		FTask P11{ Launch(TEXT("P11"), [] {}) };
+		FTask P12{ Launch(TEXT("P12"), [] {}) };
+		FTask P21{ Launch(TEXT("P21"), [] {}, Prerequisites(P11, P12)) };
+		FTask P22{ Launch(TEXT("P22"), [] {}) };
+		FTask N11, N12, N21, N22;
+		Launch(UE_SOURCE_LOCATION,
+			[&N11, &N12, &N21, &N22]
+			{
+				AddNested(N11 = Launch(TEXT("N11"),
+					[&N21, &N22]
+					{
+						AddNested(N21 = Launch(TEXT("N21"), [] {}));
+						AddNested(N22 = Launch(TEXT("N22"), [] {}));
+					}
+				));
+				AddNested(N12 = Launch(TEXT("N12"), [] {}));
+			},
+			Prerequisites(P21, P22)
+		).Wait();
+		check(P11.IsCompleted() && P12.IsCompleted() && P21.IsCompleted() && P22.IsCompleted() &&
+			N11.IsCompleted() && N12.IsCompleted() && N21.IsCompleted() && N22.IsCompleted());
+	}
+
+	IMPLEMENT_SIMPLE_AUTOMATION_TEST(FTasksDeepRetractionTest, "System.Core.Tasks.DeepRetraction", EAutomationTestFlags::ApplicationContextMask | EAutomationTestFlags::EngineFilter);
+
+	bool FTasksDeepRetractionTest::RunTest(const FString& Parameters)
+	{
+		FTaskEvent ResumeEvent{ UE_SOURCE_LOCATION };
+		TArray<FTask> WorkerBlockers = BlockWorkers(ResumeEvent);
+
+		{	// basic retraction, no dependencies
+			bool bDone = false;
+			Launch(UE_SOURCE_LOCATION, [&bDone] { bDone = true; }).Wait();
+			check(bDone);
+		}
+
+		{	// basic deep retraction: single prerequisite
+			bool bPrerequisiteDone = false;
+			bool bTaskDone = false;
+			FTask Prerequisite{ Launch(UE_SOURCE_LOCATION, 
+				[&bPrerequisiteDone, &bTaskDone] 
+				{
+					check(!bPrerequisiteDone);
+					check(!bTaskDone);
+					bPrerequisiteDone = true;
+				}
+			)};
+			Launch(UE_SOURCE_LOCATION, 
+				[&bPrerequisiteDone, &bTaskDone] 
+				{ 
+					check(bPrerequisiteDone);
+					check(!bTaskDone);
+					bTaskDone = true; 
+				}
+			).Wait();
+			check(bPrerequisiteDone);
+			check(bTaskDone);
+		}
+
+		{	// basic deep retraction: single nested task
+			bool bParentDone = false;
+			bool bNestedDone = false;
+			FTask NestedTask;
+			Launch(UE_SOURCE_LOCATION,
+				[&bParentDone, &bNestedDone, &NestedTask]
+				{
+					check(!bParentDone);
+					check(!bNestedDone);
+					NestedTask = Launch(UE_SOURCE_LOCATION,
+						[&bParentDone, &bNestedDone]
+						{
+							check(bParentDone); // due to single-threaded execution (cos all workers are blocked), nested task can't be executed
+							// until the parent is
+							check(!bNestedDone);
+							bNestedDone = true;
+						}
+					);
+					AddNested(NestedTask);
+					check(!bNestedDone);
+					bParentDone = true;
+				}
+			).Wait();
+			check(bParentDone);
+			check(bNestedDone);
+			check(NestedTask.IsCompleted());
+		}
+
+		TwoLevelsDeepRetractionTest();
+
+		ResumeEvent.Trigger();
+		verify(Wait(WorkerBlockers, FTimespan::FromSeconds(1)));
+
+		return true;
+	}
+
+	template<uint32 Num>
+	void DeepRetractionStressTest()
+	{
+		for (uint32 i = 0; i != Num; ++i)
+		{
+			TwoLevelsDeepRetractionTest();
+		}
+	}
+
+	IMPLEMENT_SIMPLE_AUTOMATION_TEST(FTasksDeepRetractionStressTest, "System.Core.Tasks.DeepRetraction.Stress", EAutomationTestFlags::ApplicationContextMask | EAutomationTestFlags::EngineFilter);
+
+	bool FTasksDeepRetractionStressTest::RunTest(const FString& Parameters)
+	{
+		UE_BENCHMARK(5, DeepRetractionStressTest<1000>);
+
+		return true;
+	}
+
 	template<int NumTasks>
 	void TestPerfBasic()
 	{
@@ -991,7 +1141,7 @@ namespace UE { namespace TasksTests
 		}
 	}
 
-	IMPLEMENT_SIMPLE_AUTOMATION_TEST(FTasksPerfTest, "System.Core.Async.TaskGraph.PerfTest", EAutomationTestFlags::ApplicationContextMask | EAutomationTestFlags::EngineFilter);
+	IMPLEMENT_SIMPLE_AUTOMATION_TEST(FTasksPerfTest, "System.Core.Tasks.PerfTest", EAutomationTestFlags::ApplicationContextMask | EAutomationTestFlags::EngineFilter);
 
 	bool FTasksPerfTest::RunTest(const FString& Parameters)
 	{
diff --git a/Engine/Source/Runtime/Core/Public/Tasks/Pipe.h b/Engine/Source/Runtime/Core/Public/Tasks/Pipe.h
--- a/Engine/Source/Runtime/Core/Public/Tasks/Pipe.h
+++ b/Engine/Source/Runtime/Core/Public/Tasks/Pipe.h
@@ -88,8 +88,8 @@ namespace UE { namespace Tasks
 		friend class Private::FTaskBase;
 
 		// pushes given task into the pipe: adds the task as a subsequent to the last task if any and sets it as the new last task
-		// returns true if there wasn't any task in the pipe and the given task is free to be executed
-		bool PushIntoPipe(Private::FTaskBase& Task);
+		// returns the previous piped task if we managed to register the given task as its subsequent, otherwise nullptr
+		Private::FTaskBase* PushIntoPipe(Private::FTaskBase& Task);
 		// pipe holds a "weak" reference to a task. the task must be cleared from the pipe when its execution finished before its completion, 
 		// otherwise the next piped task can try to add itself as a subsequent to already destroyed task
 		void ClearTask(Private::FTaskBase& Task);
diff --git a/Engine/Source/Runtime/Core/Public/Tasks/TaskPrivate.h b/Engine/Source/Runtime/Core/Public/Tasks/TaskPrivate.h
--- a/Engine/Source/Runtime/Core/Public/Tasks/TaskPrivate.h
+++ b/Engine/Source/Runtime/Core/Public/Tasks/TaskPrivate.h
@@ -9,6 +9,7 @@
 #include "Templates/TypeCompatibleBytes.h"
 #include "Misc/Timeout.h"
 #include "Containers/ClosableMpscQueue.h"
+#include "Containers/SpscQueue.h"
 #include "Templates/UnrealTemplate.h"
 #include "CoreTypes.h"
 #include "Math/NumericLimits.h"
@@ -100,8 +101,9 @@ namespace UE { namespace Tasks
 				);
 			}
 
-			// the task will be executed only when all prerequisites are completed. The task type must be a task handle that holds a pointer to
-			//  FTaskBase as its `Pimpl` member (see Tasks::TTaskBase)
+			// The task will be executed only when all prerequisites are completed. The task type must be a task handle that holds a pointer to
+			// FTaskBase as its `Pimpl` member (see Tasks::TTaskBase).
+			// Must not be called concurrently
 			template<typename HigherLevelTaskType, decltype(std::declval<HigherLevelTaskType>().Pimpl)* = nullptr>
 			void AddPrerequisites(const HigherLevelTaskType& Prerequisite)
 			{
@@ -113,14 +115,21 @@ namespace UE { namespace Tasks
 				uint32 PrevNumLocks = NumLocks.fetch_add(1, std::memory_order_acquire); // `acquire` to make it happen before the task is registered as a subsequent
 				checkf(PrevNumLocks + 1 < ExecutionFlag, TEXT("Max number of nested tasks reached: %d"), ExecutionFlag);
 
-				if (!Prerequisite.Pimpl->AddSubsequent(*this))
+				if (Prerequisite.Pimpl->AddSubsequent(*this))
+				{
+					Prerequisites.Enqueue(Prerequisite.Pimpl.GetReference());
+					// keep it alive until this task's destruction
+					Prerequisite.Pimpl->AddRef();
+				}
+				else
 				{
 					// failed to add the prerequisite (too late), correct the number
 					NumLocks.fetch_sub(1, std::memory_order_release); // `release` to make it happen after the task is registered as a subsequent
 				}
 			}
 
-			// the task will be executed only when all prerequisites are completed
+			// The task will be executed only when all prerequisites are completed.
+			// Must not be called concurrently.
 			// @param Prerequisites - an iterable collection of tasks
 			template<typename PrerequisiteCollectionType, decltype(std::declval<PrerequisiteCollectionType>().begin())* = nullptr>
 			void AddPrerequisites(const PrerequisiteCollectionType& InPrerequisites)
@@ -148,7 +157,13 @@ namespace UE { namespace Tasks
 						Prerequisite = Prereq.Pimpl;
 					}
 
-					if (!Prerequisite->AddSubsequent(*this))
+					if (Prerequisite->AddSubsequent(*this))
+					{
+						Prerequisites.Enqueue(Prerequisite);
+						// keep it alive until this task's destruction
+						Prerequisite->AddRef();
+					}
+					else
 					{
 						++NumCompletedPrerequisites;
 					}
@@ -192,22 +207,91 @@ namespace UE { namespace Tasks
 			}
 
 			// Tries to pull out the task from the scheduler and execute it. Returns false if task execution is already started.
-			bool TryRetractAndExecute()
+			bool TryRetractAndExecute(uint32 RecursionDepth = 0)
 			{
 				if (IsCompleted())
 				{
 					return true;
 				}
 
-				uint32 LocalNumLocks = NumLocks.load(std::memory_order_acquire);
-				bool bExecutionUnlocked = LocalNumLocks < ExecutionFlag && LocalNumLocks == (GetPipe() == nullptr ? 1 : 0);
-				if (bExecutionUnlocked && TryExecute())
+				// avoid stack overflow. is not expected in a real-life cases but happens in stress tests
+				if (RecursionDepth == 200)
+				{
+					return false;
+				}
+				++RecursionDepth;
+
+				// prevent concurrent retraction from multiple threads. while it's desirable to allow this (so different threads retract and
+				// execute prerequisites in parallel), multiple threads waiting for the same task with multiple not completed prerequisites is expected
+				// to happen very rarely, and the implementation would be much more complicated potentially leading to a slower common-case
+				// version
+				if (!TryGetExecutionPermission())
+				{
+					return false;
+				}
+
+				auto IsLocked = [this]
+				{
+					return NumLocks.load(std::memory_order_acquire) != (Pipe == nullptr ? 1 : 0);
+				};
+
+				auto TryRetractPrerequisites = [this, RecursionDepth]
+				{
+					// keep trying retracting all prerequisites even if some of them fail, so the current worker can contribute instead of being blocked
+					bool bSucceeded = true;
+					// prerequisites are "consumed" here even if their retraction fails. retraction can fail only if task execution has 
+					// already started, and it can't become "retractable" again, so no need to keep them
+					while (TOptional<FTaskBase*> Prerequisite = Prerequisites.Dequeue())
+					{
+						if (!Prerequisite.GetValue()->TryRetractAndExecute(RecursionDepth))
+						{
+							bSucceeded = false;
+						}
+						Prerequisite.GetValue()->Release();
+					}
+
+					return bSucceeded;
+				};
+
+				if (IsLocked())
+				{
+					// try to unlock the task. even if prerequisites retraction fails we still need to proceed to try to execute the task here in case
+					// prerequisites were completed in parallel
+					TryRetractPrerequisites();
+				}
+
+				// the task can be still locked if either prerequisite retraction (above) failed, or the task was piped by another thread, or the task 
+				// wasn't launched yet (Tasks::FTaskEvent can be used as a prerequisites before it's triggered)
+				if (IsLocked())
+				{
+					RevokeExecutionPermission();
+					// prerequisites could be completed in parallel after `IsLocked()` and before we revoked execution permission, so the worker who
+					// unlocked the task won't be able to execute it. double check if the task is still locked and if not - try to execute it
+					if (IsLocked() || !TryGetExecutionPermission()) // if it's the case, try to get back execution permission
+					{
+						return false;
+					}
+				}
+
+				// the task is unlocked and we have execution permission
+				DoExecute();
+				// no need to cancel task execution by the scheduler, when the scheduler will execute its runnable, it will fail to get execution
+				// permissions and will do nothing
+				// LowLevelTask.TryCancel();
+
+				if (IsCompleted()) // still can be hold back by nested tasks
 				{
-					LowLevelTask.TryCancel();
 					return IsCompleted(); // still can be hold back by nested tasks
 				}
 
-				return false;
+				// try to retract nested tasks
+				if (!TryRetractPrerequisites())
+				{
+					return false;
+				}
+
+				checkSlow(IsCompleted());
+				return true;
 			}
 
 			// adds a nested task that must be completed before the parent (this) is completed
@@ -217,7 +301,14 @@ namespace UE { namespace Tasks
 				checkf(PrevNumLocks + 1 < TNumericLimits<uint32>::Max(), TEXT("Max number of nested tasks reached: %d"), TNumericLimits<uint32>::Max() - ExecutionFlag);
 				checkf(PrevNumLocks >= ExecutionFlag, TEXT("Internal error: nested tasks can be added only during parent's execution (%u)"), PrevNumLocks);
 
-				if (!Nested.AddSubsequent(*this))
+				if (Nested.AddSubsequent(*this))
+				{
+					// there's no race with, as `Prerequisites` can't be accessed by another thread during task execution (execution "locks" the task)
+					Prerequisites.Enqueue(&Nested);
+					// keep it alive until this task's destruction
+					Nested.AddRef();
+				}
+				else
 				{
 					NumLocks.fetch_sub(1, std::memory_order_relaxed);
 				}
@@ -342,13 +433,34 @@ namespace UE { namespace Tasks
 			// prepares the task for execution and executes its body if its execution hasn't been started yet
 			bool TryExecute()
 			{
-				// can be executed only once
-				if (!bAvailableForExecution.exchange(false, std::memory_order_relaxed))
+				if (TryGetExecutionPermission())
 				{
-					return false;
+					DoExecute();
+					return true;
 				}
 
-				AddRef(); // for the reference hold by nested tasks, is released when task is closed
+				return false; // too late, execution already started
+			}
+
+			bool TryGetExecutionPermission()
+			{
+				return bAvailableForExecution.exchange(false, std::memory_order_acq_rel);
+			}
+
+			void RevokeExecutionPermission()
+			{
+				bAvailableForExecution.store(true, std::memory_order_release);
+			}
+
+			void DoExecute()
+			{
+				AddRef(); // for the reference hold by nested tasks, is released when the task is closed
+
+				// release prerequisites refs as they are not needed anymore
+				while (TOptional<FTaskBase*> Prerequisite = Prerequisites.Dequeue())
+				{
+					Prerequisite.GetValue()->Release();
+				}
 
 				NumLocks.store(ExecutionFlag + 1, std::memory_order_relaxed); // +1 to hold it locked during execution, so nested tasks don't
 				// complete it before the execution finishes
@@ -364,8 +476,6 @@ namespace UE { namespace Tasks
 				uint32 LocalNumLocks = NumLocks.fetch_sub(1, std::memory_order_relaxed) - 1;
 
 				TryComplete(LocalNumLocks);
-
-				return true;
 			}
 
 			// checks if the task is ready to be launched by trying to push it into the pipe.
@@ -373,7 +483,7 @@ namespace UE { namespace Tasks
 			// `LocalNumLocks` is the value that was used to make a decision to push into the pipe, no need to read `NumLocks` again
 			bool TryPushIntoPipe(uint32 LocalNumLocks)
 			{
-				if (GetPipe() == nullptr)
+				if (Pipe == nullptr)
 				{
 					// the task is locked for a pipe initially even if eventually there's no pipe
 					check(LocalNumLocks == 1);
@@ -387,15 +497,23 @@ namespace UE { namespace Tasks
 				bool bFirstAttempt = LocalNumLocks == 1;
 				if (bFirstAttempt)
 				{
-					return PushIntoPipe();
+					FTaskBase* PrevPipedTask = PushIntoPipe();
+					if (PrevPipedTask == nullptr) // we are free to go
+					{
+						NumLocks.store(0, std::memory_order_relaxed);
+						return true;
+					}
+
+					Prerequisites.Enqueue(PrevPipedTask);
+					// no need to AddRef as it's already sorted in `FPipe::PushIntoPipe`
+					return false;
 				}
 
 				return true;
 			}
 
-			// is called when the task has no pending prerequisites. Returns false is the pipe is blocked (another task from the same pipe is being
-			// executed)
-			bool PushIntoPipe();
+			// is called when the task has no pending prerequisites. Returns the previous piped task if any
+			FTaskBase* PushIntoPipe();
 
 			void StartPipeExecution();
 			void FinishPipeExecution();
@@ -408,6 +526,7 @@ namespace UE { namespace Tasks
 			TUniqueFunction<void()> TaskBody;
 			std::atomic<bool> bAvailableForExecution{ true };
 
+
 			// the number of times that the task should be unlocked before it can be scheduled or completed
 			// initial count is 1 for launching the task (it can't be scheduled before it's launched) and 1 for a potential blocked pipe. once NumLocks
 			// reaches 0 the task is scheduled for execution.
@@ -416,6 +535,17 @@ namespace UE { namespace Tasks
 			static const uint32 NumInitialLocks = 1 + 1;
 			std::atomic<uint32> NumLocks{ NumInitialLocks };
 
+			// A single-producer/single-consumer container to store back links to "prerequsites" (either execution prerequisites or nested tasks 
+			// that are completion prerequisites).
+			// It's populated in three stages:
+			// 1) by adding execution prerequisites, before the task is launched, single thread, when nobody else accesses it. doesn't need
+			// synchronisation
+			// 2) by piping, when the previous piped task is added as a prerequisite. after adding prerequisites. this can happen in parallel with
+			// retraction that consumes prerequisites. multiple threads can try to retract the task but only one will be allowed to proceed, 
+			// so single producer/single consumer
+			// 3) by adding nested tasks. after piping. during task execution, thus retraction is not possible and won't touch it. single-threaded
+			TSpscQueue<FTaskBase*> Prerequisites;
+
 			TClosableMpscQueue<FTaskBase*> Subsequents;
 
 			FPipe* Pipe{ nullptr };
