Small change to STACK_OFFSET in order to have log events from tests that show up in index.json report the correct filename and line number.

diff --git a/Engine/Source/Runtime/Core/Private/Misc/AutomationTest.cpp b/Engine/Source/Runtime/Core/Private/Misc/AutomationTest.cpp
--- a/Engine/Source/Runtime/Core/Private/Misc/AutomationTest.cpp
+++ b/Engine/Source/Runtime/Core/Private/Misc/AutomationTest.cpp
@@ -79,7 +79,7 @@ CORE_API ELogVerbosity::Type GetAutomationLogLevel(ELogVerbosity::Type LogVerbos
 
 void FAutomationTestFramework::FAutomationTestOutputDevice::Serialize( const TCHAR* V, ELogVerbosity::Type Verbosity, const class FName& Category )
 {
-	const int32 STACK_OFFSET = 5;//FMsg::Logf_InternalImpl
+	const int32 STACK_OFFSET = 8;//FMsg::Logf_InternalImpl
 	// TODO would be nice to search for the first stack frame that isn't in outputdevice or other logging files, would be more robust.
 
 	if (!IsRunningCommandlet() && (Verbosity == ELogVerbosity::SetColor))
