Unused helper removed.

diff --git a/Engine/Source/Runtime/Core/Public/GenericPlatform/GenericPlatformString.h b/Engine/Source/Runtime/Core/Public/GenericPlatform/GenericPlatformString.h
--- a/Engine/Source/Runtime/Core/Public/GenericPlatform/GenericPlatformString.h
+++ b/Engine/Source/Runtime/Core/Public/GenericPlatform/GenericPlatformString.h
@@ -130,17 +130,7 @@ struct FGenericPlatformString : public FGenericPlatformStricmp
 
 	/**
 	 * Metafunction which tests whether a given character type represents a fixed-width encoding.
-	 *
-	 * We need to 'forward' the metafunction to a helper because of the C++ requirement that
-	 * nested structs cannot be fully specialized.  They can be partially specialized however, hence the
-	 * helper function.
 	 */
-	template <bool Dummy, typename T>
-	struct TIsFixedWidthEncoding_Helper
-	{
-		enum { Value = false };
-	};
-
 	template <typename T>
 	struct UE_DEPRECATED(5.0, "TIsFixedWidthEncoding is deprecated, use IsFixedWidthEncoding<T>() instead.") TIsFixedWidthEncoding
 	{
