New FPlatformString::IsCharEncodingTriviallyConvertibleTo() function which tests if one encoding can be per-character converted to another encoding.

diff --git a/Engine/Source/Runtime/Core/Public/GenericPlatform/GenericPlatformString.h b/Engine/Source/Runtime/Core/Public/GenericPlatform/GenericPlatformString.h
--- a/Engine/Source/Runtime/Core/Public/GenericPlatform/GenericPlatformString.h
+++ b/Engine/Source/Runtime/Core/Public/GenericPlatform/GenericPlatformString.h
@@ -62,5 +62,73 @@
 	}
 
 	/**
+	 * Function which returns whether one encoding type is binary compatible with another.
+	 *
+	 * Unlike TAreEncodingsCompatible, this is not commutative.  For example, ANSI is compatible with
+	 * UTF-8, but UTF-8 is not compatible with ANSI.
+	 */
+	template <typename SrcEncoding, typename DestEncoding>
+	static constexpr bool IsCharEncodingCompatibleWith()
+	{
+		if constexpr (std::is_same_v<SrcEncoding, DestEncoding>)
+		{
+			return true;
+		}
+		else if constexpr (std::is_same_v<SrcEncoding, ANSICHAR> && std::is_same_v<DestEncoding, UTF8CHAR>)
+		{
+			return true;
+		}
+		else if constexpr (std::is_same_v<SrcEncoding, UCS2CHAR> && std::is_same_v<DestEncoding, UTF16CHAR>)
+		{
+			return true;
+		}
+		else if constexpr (std::is_same_v<SrcEncoding, WIDECHAR> && std::is_same_v<DestEncoding, UCS2CHAR>)
+		{
+			return true;
+		}
+		else if constexpr (std::is_same_v<SrcEncoding, UCS2CHAR> && std::is_same_v<DestEncoding, WIDECHAR>)
+		{
+			return true;
+		}
+#if PLATFORM_TCHAR_IS_CHAR16
+		else if constexpr (std::is_same_v<SrcEncoding, WIDECHAR> && std::is_same_v<DestEncoding, wchar_t>)
+		{
+			return true;
+		}
+		else if constexpr (std::is_same_v<SrcEncoding, wchar_t> && std::is_same_v<DestEncoding, WIDECHAR>)
+		{
+			return true;
+		}
+#endif
+		else
+		{
+			return false;
+		}
+	};
+
+	/**
+	 * Tests whether you can simply (i.e. by assignment) encode code units from the source encoding as the destination encoding.
+	 */
+	template <typename SourceEncoding, typename DestEncoding>
+	static constexpr bool IsCharEncodingSimplyConvertibleTo()
+	{
+		if constexpr (IsCharEncodingCompatibleWith<SourceEncoding, DestEncoding>())
+		{
+			// Binary-compatible conversions are always simple
+			return true;
+		}
+		else if constexpr (IsFixedWidthEncoding<SourceEncoding>() && sizeof(DestEncoding) >= sizeof(SourceEncoding))
+		{
+			// Converting from a fixed-width encoding to a wider or same encoding should always be possible,
+			// as should ANSICHAR->UTF8CHAR and UCS2CHAR->UTF16CHAR,
+			return true;
+		}
+		else
+		{
+			return false;
+		}
+	}
+
+	/**
 	 * Tests whether a particular codepoint can be converted to the destination encoding.
 	 *
@@ -78,10 +146,9 @@
 		// This is only defined for fixed-width encodings, because codepoints cannot be represented in a single variable-width code unit.
 		static_assert(IsFixedWidthEncoding<SourceEncoding>(), "Source encoding is not fixed-width");
 
-		if constexpr (sizeof(DestEncoding) >= sizeof(SourceEncoding))
+		if constexpr (IsCharEncodingSimplyConvertibleTo<SourceEncoding, DestEncoding>())
 		{
-			// Converting to a wider encoding, or same encoding, or ANSICHAR->UTF8CHAR or UCS2CHAR->UTF16CHAR,
-			// should always be possible.
+			// Simple conversions mean conversion is always possible
 			return true;
 		}
 		else if constexpr (!IsFixedWidthEncoding<DestEncoding>())
@@ -153,50 +220,5 @@
 	};
 
 	/**
-	 * Function which returns whether one encoding type is binary compatible with another.
-	 *
-	 * Unlike TAreEncodingsCompatible, this is not commutative.  For example, ANSI is compatible with
-	 * UTF-8, but UTF-8 is not compatible with ANSI.
-	 */
-	template <typename SrcEncoding, typename DestEncoding>
-	static constexpr bool IsCharEncodingCompatibleWith()
-	{
-		if constexpr (std::is_same_v<SrcEncoding, DestEncoding>)
-		{
-			return true;
-		}
-		else if constexpr (std::is_same_v<SrcEncoding, ANSICHAR> && std::is_same_v<DestEncoding, UTF8CHAR>)
-		{
-			return true;
-		}
-		else if constexpr (std::is_same_v<SrcEncoding, UCS2CHAR> && std::is_same_v<DestEncoding, UTF16CHAR>)
-		{
-			return true;
-		}
-		else if constexpr (std::is_same_v<SrcEncoding, WIDECHAR> && std::is_same_v<DestEncoding, UCS2CHAR>)
-		{
-			return true;
-		}
-		else if constexpr (std::is_same_v<SrcEncoding, UCS2CHAR> && std::is_same_v<DestEncoding, WIDECHAR>)
-		{
-			return true;
-		}
-#if PLATFORM_TCHAR_IS_CHAR16
-		else if constexpr (std::is_same_v<SrcEncoding, WIDECHAR> && std::is_same_v<DestEncoding, wchar_t>)
-		{
-			return true;
-		}
-		else if constexpr (std::is_same_v<SrcEncoding, wchar_t> && std::is_same_v<DestEncoding, WIDECHAR>)
-		{
-			return true;
-		}
-#endif
-		else
-		{
-			return false;
-		}
-	};
-
-	/**
 	 * Converts the [Src, Src+SrcSize) string range from SourceEncoding to DestEncoding and writes it to the [Dest, Dest+DestSize) range.
 	 * The Src range should contain a null terminator if a null terminator is required in the output.
@@ -220,6 +242,17 @@
 
 			return (DestEncoding*)Memcpy(Dest, Src, SrcSize * sizeof(SourceEncoding)) + SrcSize;
 		}
+		else if constexpr (IsCharEncodingSimplyConvertibleTo<SourceEncoding, DestEncoding>())
+		{
+			const int32 Size = DestSize <= SrcSize ? DestSize : SrcSize;
+			for (int I = 0; I < Size; ++I)
+			{
+				SourceEncoding SrcCh = Src[I];
+				Dest[I] = (DestEncoding)SrcCh;
+			}
+
+			return DestSize < SrcSize ? nullptr : Dest + Size;
+		}
 		else if constexpr (IsFixedWidthEncoding<SourceEncoding>() && IsFixedWidthEncoding<DestEncoding>())
 		{
 			const int32 Size = DestSize <= SrcSize ? DestSize : SrcSize;
@@ -264,7 +297,7 @@
 	template <typename DestEncoding, typename SourceEncoding>
 	static int32 ConvertedLength(const SourceEncoding* Src, int32 SrcSize)
 	{
-		if constexpr (std::is_same_v<SourceEncoding, DestEncoding> || (IsFixedWidthEncoding<SourceEncoding>() && IsFixedWidthEncoding<DestEncoding>()))
+		if constexpr (IsCharEncodingSimplyConvertibleTo<SourceEncoding, DestEncoding>() || (IsFixedWidthEncoding<SourceEncoding>() && IsFixedWidthEncoding<DestEncoding>()))
 		{
 			return SrcSize;
 		}
