fixed ticker delegate being removed from inside delegate execution

diff --git a/Engine/Source/Runtime/Core/Private/Containers/Ticker.cpp b/Engine/Source/Runtime/Core/Private/Containers/Ticker.cpp
--- a/Engine/Source/Runtime/Core/Private/Containers/Ticker.cpp
+++ b/Engine/Source/Runtime/Core/Private/Containers/Ticker.cpp
@@ -24,18 +24,26 @@ FTSTicker::FDelegateHandle FTSTicker::AddTicker(const TCHAR* InName, float InDel
 	return NewElement;
 }
 
+uint32 GetThreadId(uint64 State)
+{
+	return (uint32)(State >> 32);
+}
+
 void FTSTicker::RemoveTicker(FDelegateHandle Handle)
 {
 	if (FElementPtr Element = Handle.Pin())
 	{
 		// mark the element as removed and if it's being ticked atm, spin-wait until its execution is finished
-		uint8 PrevState = Element->State.fetch_add(FElement::Removed, std::memory_order_acquire); // "acquire" to prevent potential 
+		uint64 PrevState = Element->State.fetch_add(FElement::RemovedState, std::memory_order_acquire); // "acquire" to prevent potential 
 		// resource release after RemoveTicker() to be reordered before it
-		checkf(PrevState < FElement::Removed, TEXT("The delegate is already removed (%u)"), PrevState);
-		while (PrevState != FElement::Idle && PrevState != FElement::Removed) // is not being executed
+		checkf((PrevState & 0x1) == FElement::DefaultState, TEXT("The delegate is already removed (%u)"), PrevState);
+		uint32 ExecutingThreadId = GetThreadId(PrevState);
+		
+		while (ExecutingThreadId != 0 && // is being executed right now
+			FPlatformTLS::GetCurrentThreadId() != ExecutingThreadId) // and is not removed from inside its execution
 		{
 			FPlatformProcess::Yield();
-			PrevState = Element->State.load(std::memory_order_relaxed);
+			ExecutingThreadId = GetThreadId(Element->State.load(std::memory_order_relaxed));
 		}
 	}
 }
@@ -46,12 +54,14 @@ void FTSTicker::Tick(float DeltaTime)
 
 	QUICK_SCOPE_CYCLE_COUNTER(STAT_FTicker_Tick);
 
+	const uint64 CurrentThreadId = ((uint64)FPlatformTLS::GetCurrentThreadId()) << 32; // already prepared to be used in `FElement::State`
+
 	// take in all new elements
 	while (TOptional<FElementPtr> AddedElement = AddedElements.Dequeue())
 	{
-		uint8 State = AddedElement.GetValue()->State.load(std::memory_order_relaxed);
-		checkf(State == FElement::Idle || State == FElement::Removed, TEXT("Invalid state %u"), State);
-		if (State == FElement::Idle)
+		uint64 State = AddedElement.GetValue()->State.load(std::memory_order_relaxed);
+		checkf(GetThreadId(State) == 0, TEXT("Invalid state %u"), State);
+		if (State == FElement::DefaultState)
 		{
 			Elements.Add(MoveTemp(AddedElement.GetValue()));
 		}
@@ -67,19 +77,20 @@ void FTSTicker::Tick(float DeltaTime)
 	TArray<FElementPtr> TickedElements;
 	for (FElementPtr& Element : Elements)
 	{
-		// set the execution flag
-		uint8 PrevState = Element->State.fetch_add(FElement::Executing, std::memory_order_acquire); // "acquire" to prevent anything between this and `ClearExecutionFlag()` to be reordered outside of this scope, is coupled with "release" in `ClearExecutionFlag`
-		checkf(PrevState == FElement::Idle || PrevState == FElement::Removed, TEXT("Invalid state %u"), PrevState);
-		auto ClearExecutionFlag = [](FElementPtr& Element) { return (uint8)(Element->State.fetch_sub(FElement::Executing, std::memory_order_release) - FElement::Executing); };
+		// set the execution state
+		uint64 PrevState = Element->State.fetch_add(CurrentThreadId, std::memory_order_acquire); // "acquire" to prevent anything between this and `ClearExecutionFlag()` to be reordered outside of this scope, is coupled with "release" in `ClearExecutionFlag`
+		checkf(GetThreadId(PrevState) == 0, TEXT("Invalid state %u"), PrevState);
+		auto ClearExecutionFlag = [CurrentThreadId](FElementPtr& Element) 
+		{
+			return (uint8)(Element->State.fetch_sub(CurrentThreadId, std::memory_order_release) - CurrentThreadId); 
+		};
 
-		if (PrevState == FElement::Removed)
+		if (PrevState == FElement::RemovedState)
 		{
-			PrevState = ClearExecutionFlag(Element);
+			ClearExecutionFlag(Element);
 			continue;
 		}
 
-		checkf(PrevState == FElement::Idle, TEXT("Invalid state %u"), PrevState);
-
 		if (Element->FireTime > CurrentTime)
 		{
 			ClearExecutionFlag(Element);
@@ -90,9 +101,9 @@ void FTSTicker::Tick(float DeltaTime)
 		if (Element->Fire(DeltaTime))
 		{
 			PrevState = ClearExecutionFlag(Element); // it can be removed during execution
-			if (PrevState != FElement::Removed)
+			if (PrevState != FElement::RemovedState)
 			{
-				checkf(PrevState == FElement::Idle, TEXT("Invalid state %u"), PrevState);
+				checkf(PrevState == FElement::DefaultState, TEXT("Invalid state %u"), PrevState);
 				Element->FireTime = CurrentTime + Element->DelayTime;
 				TickedElements.Add(MoveTemp(Element));
 			}
@@ -100,7 +111,7 @@ void FTSTicker::Tick(float DeltaTime)
 		else
 		{
 			PrevState = ClearExecutionFlag(Element);
-			checkf(PrevState == FElement::Idle || PrevState == FElement::Removed, TEXT("Invalid state %u"), PrevState);
+			checkf(PrevState == FElement::DefaultState || PrevState == FElement::RemovedState, TEXT("Invalid state %u"), PrevState);
 		}
 	}
 
diff --git a/Engine/Source/Runtime/Core/Private/Tests/Containers/TickerTests.cpp b/Engine/Source/Runtime/Core/Private/Tests/Containers/TickerTests.cpp
--- a/Engine/Source/Runtime/Core/Private/Tests/Containers/TickerTests.cpp
+++ b/Engine/Source/Runtime/Core/Private/Tests/Containers/TickerTests.cpp
@@ -136,6 +136,13 @@ bool FTSTickerTest::RunTest(const FString& Parameters)
 		verify(RemoveTickerTask.Wait(FTimespan::FromSeconds(0.1)));
 	}
 
+	{	// a delegate removal from inside the delegate execution (used to be a deadlock)
+		FTSTicker Ticker;
+		FTSTicker::FDelegateHandle DelegateHandle;
+		DelegateHandle = Ticker.AddTicker(nullptr, 0.0f, [&DelegateHandle](float) { FTSTicker::RemoveTicker(DelegateHandle); return true; });
+		Ticker.Tick(0.0f);
+	}
+
 	// multithreaded stress test
 	{
 		using namespace UE::Tasks;
diff --git a/Engine/Source/Runtime/Core/Public/Containers/Ticker.h b/Engine/Source/Runtime/Core/Public/Containers/Ticker.h
--- a/Engine/Source/Runtime/Core/Public/Containers/Ticker.h
+++ b/Engine/Source/Runtime/Core/Public/Containers/Ticker.h
@@ -95,13 +95,16 @@ private:
 		/** Delegate to fire **/
 		FTickerDelegate Delegate;
 
+		static const uint64 DefaultState = 0;
+		static const uint64 RemovedState = 1;
 		// The element can be in 4 states: "idle", "executing", "idle and removed" and "executing and removed"
 		// Often right after `RemoveTicket()` call resources that are used by the delegate are destroyed. We must ensure that 
 		// these resources are not accessed by the delegate after `RemoveTicket()` returns.
 		// `State` is used to remove the element safely, if removal happens in the middle of the delegate execution, removal will not return
 		// until the execution is finished.
-		enum EState : uint8 { Idle, Executing, Removed };
-		std::atomic<uint8> State{ Idle };
+		// `State` constists of two 32 bits parts: lower part is 1 if the delegate is removed, otherwise it's 0. the higher part contains thread id if the
+		// delegate being executed, otherwise it's 0
+		std::atomic<uint64> State{ DefaultState };
 
 		/** Default ctor is only required to implement CurrentElement handling without making it a pointer. */
 		CORE_API FElement();
