IoDispatcher:
Fixed condition for determining if a block should be cached or not
Only check the caching conditon when storing blocks, not when reading them back
#rnx

diff --git a/Engine/Source/Runtime/Core/Private/IO/IoDispatcherFileBackend.cpp b/Engine/Source/Runtime/Core/Private/IO/IoDispatcherFileBackend.cpp
--- a/Engine/Source/Runtime/Core/Private/IO/IoDispatcherFileBackend.cpp
+++ b/Engine/Source/Runtime/Core/Private/IO/IoDispatcherFileBackend.cpp
@@ -204,31 +204,28 @@ void FFileIoStoreBlockCache::Initialize(uint64 InCacheMemorySize, uint64 InReadB
 
 bool FFileIoStoreBlockCache::Read(FFileIoStoreReadRequest* Block)
 {
-	bool bIsCacheableBlock = CacheMemory != nullptr && Block->bIsCacheable;
-	if (!bIsCacheableBlock)
+	if (!CacheMemory)
 	{
 		return false;
 	}
 	check(Block->Buffer);
 	FCachedBlock* CachedBlock = CachedBlocks.FindRef(Block->Key.Hash);
-	if (CachedBlock)
-	{
-		CachedBlock->LruPrev->LruNext = CachedBlock->LruNext;
-		CachedBlock->LruNext->LruPrev = CachedBlock->LruPrev;
-
-		CachedBlock->LruPrev = &CacheLruHead;
-		CachedBlock->LruNext = CacheLruHead.LruNext;
-
-		CachedBlock->LruPrev->LruNext = CachedBlock;
-		CachedBlock->LruNext->LruPrev = CachedBlock;
-	}
-
 	if (!CachedBlock)
 	{
 		FFileIoStats::OnBlockCacheMiss(ReadBufferSize);
 		TRACE_COUNTER_INCREMENT(IoDispatcherFileBackendCacheMisses);
 		return false;
 	}
+	
+	CachedBlock->LruPrev->LruNext = CachedBlock->LruNext;
+	CachedBlock->LruNext->LruPrev = CachedBlock->LruPrev;
+
+	CachedBlock->LruPrev = &CacheLruHead;
+	CachedBlock->LruNext = CacheLruHead.LruNext;
+
+	CachedBlock->LruPrev->LruNext = CachedBlock;
+	CachedBlock->LruNext->LruPrev = CachedBlock;
+
 	check(CachedBlock->Buffer);
 	FFileIoStats::OnBlockCacheHit(ReadBufferSize);
 	FMemory::Memcpy(Block->Buffer->Memory, CachedBlock->Buffer, ReadBufferSize);
@@ -239,7 +236,7 @@ bool FFileIoStoreBlockCache::Read(FFileIoStoreReadRequest* Block)
 
 void FFileIoStoreBlockCache::Store(const FFileIoStoreReadRequest* Block)
 {
-	bool bIsCacheableBlock = CacheMemory != nullptr && Block->bIsCacheable;
+	bool bIsCacheableBlock = CacheMemory != nullptr && Block->BytesUsed < Block->Size;
 	if (!bIsCacheableBlock)
 	{
 		return;
@@ -1670,8 +1667,6 @@ void FFileIoStore::ReadBlocks(FFileIoStoreResolvedRequest& ResolvedRequest)
 		if (bCompressedBlockWasAdded)
 		{
 			CompressedBlock->EncryptionKey = ContainerFile.EncryptionKey;
-			bool bCacheable = OffsetInRequest > 0 || RequestRemainingBytes < CompressionBlockSize;
-
 			const FIoStoreTocCompressedBlockEntry& CompressionBlockEntry = ContainerFile.CompressionBlocks[CompressedBlockIndex];
 			CompressedBlock->UncompressedSize = CompressionBlockEntry.GetUncompressedSize();
 			CompressedBlock->CompressedSize = CompressionBlockEntry.GetCompressedSize();
@@ -1701,12 +1696,14 @@ void FFileIoStore::ReadBlocks(FFileIoStoreResolvedRequest& ResolvedRequest)
 				{
 					RawBlock->Priority = ResolvedRequest.GetPriority();
 					RawBlock->FileHandle = Partition.FileHandle;
-					RawBlock->bIsCacheable = bCacheable;
 					RawBlock->Offset = RawBlockIndex * ReadBufferSize;
 					uint64 ReadSize = FMath::Min(Partition.FileSize, RawBlock->Offset + ReadBufferSize) - RawBlock->Offset;
 					RawBlock->Size = ReadSize;
 					NewBlocks.Add(RawBlock);
 				}
+				RawBlock->BytesUsed += 
+					uint32(FMath::Min(CompressedBlock->RawOffset + CompressedBlock->RawSize, RawBlock->Offset + RawBlock->Size) -
+						   FMath::Max(CompressedBlock->RawOffset, RawBlock->Offset));
 				CompressedBlock->RawBlocks.Add(RawBlock);
 				++CompressedBlock->UnfinishedRawBlocksCount;
 				++CompressedBlock->RefCount;
diff --git a/Engine/Source/Runtime/Core/Private/IO/IoDispatcherFileBackendTypes.h b/Engine/Source/Runtime/Core/Private/IO/IoDispatcherFileBackendTypes.h
--- a/Engine/Source/Runtime/Core/Private/IO/IoDispatcherFileBackendTypes.h
+++ b/Engine/Source/Runtime/Core/Private/IO/IoDispatcherFileBackendTypes.h
@@ -131,7 +131,7 @@ struct FFileIoStoreReadRequest
 	int32 Priority = 0;
 	uint64 CreationTime;	// Potentially used to circuit break request ordering optimizations when outstanding requests have been delayed too long
 	FFileIoStoreBlockScatter ImmediateScatter;
-	bool bIsCacheable = false;
+	uint32 BytesUsed = 0;
 	bool bFailed = false;
 	bool bCancelled = false;
 	EQueueStatus QueueStatus = QueueStatus_NotInQueue;
