platform extensions can handle ensure failures via their crash handlers, if supported.

#rnx

diff --git a/Engine/Source/Runtime/Core/Private/HAL/ExceptionHandling.cpp b/Engine/Source/Runtime/Core/Private/HAL/ExceptionHandling.cpp
--- a/Engine/Source/Runtime/Core/Private/HAL/ExceptionHandling.cpp
+++ b/Engine/Source/Runtime/Core/Private/HAL/ExceptionHandling.cpp
@@ -34,7 +34,7 @@ void ReportInteractiveEnsure(const TCHAR* InMessage)
 {
 	GEnsureShowsCRC = true;
 
-#if PLATFORM_DESKTOP
+#if PLATFORM_USE_REPORT_ENSURE
 	GLog->PanicFlushThreadedLogs();
 	ReportEnsure(InMessage, nullptr);
 #endif
diff --git a/Engine/Source/Runtime/Core/Private/HAL/ThreadHeartBeat.cpp b/Engine/Source/Runtime/Core/Private/HAL/ThreadHeartBeat.cpp
--- a/Engine/Source/Runtime/Core/Private/HAL/ThreadHeartBeat.cpp
+++ b/Engine/Source/Runtime/Core/Private/HAL/ThreadHeartBeat.cpp
@@ -273,7 +273,7 @@ void FORCENOINLINE FThreadHeartBeat::OnHang(double HangDuration, uint32 ThreadTh
 
 		const FString ErrorMessage = FString::Printf(TEXT("Hang detected on %s:%s%s%sCheck log for full callstack."), *ThreadName, LINE_TERMINATOR, *StackTrimmed, LINE_TERMINATOR);
 
-#if PLATFORM_DESKTOP
+#if PLATFORM_USE_REPORT_ENSURE
 		UE_LOG(LogCore, Error, TEXT("%s"), *ErrorMessage);
 		GLog->PanicFlushThreadedLogs();
 
diff --git a/Engine/Source/Runtime/Core/Private/Misc/AssertionMacros.cpp b/Engine/Source/Runtime/Core/Private/Misc/AssertionMacros.cpp
--- a/Engine/Source/Runtime/Core/Private/Misc/AssertionMacros.cpp
+++ b/Engine/Source/Runtime/Core/Private/Misc/AssertionMacros.cpp
@@ -491,7 +491,7 @@ FORCENOINLINE void FDebug::EnsureFailed(const FFailureInfo& Info, const TCHAR* M
 			SCOPE_LOG_TIME_IN_SECONDS(*SendNewReportMessage, nullptr)
 #endif
 
-#if PLATFORM_DESKTOP
+#if PLATFORM_USE_REPORT_ENSURE
 			FScopeLock Lock(&GetFailDebugCriticalSection());
 
 			ReportEnsure(ErrorMsg, Info.ProgramCounter);
diff --git a/Engine/Source/Runtime/Core/Public/HAL/Platform.h b/Engine/Source/Runtime/Core/Public/HAL/Platform.h
--- a/Engine/Source/Runtime/Core/Public/HAL/Platform.h
+++ b/Engine/Source/Runtime/Core/Public/HAL/Platform.h
@@ -627,6 +627,10 @@
 	#define PLATFORM_NEEDS_GPU_UAV_RESOURCE_INIT_WORKAROUND 0
 #endif
 
+#ifndef PLATFORM_USE_REPORT_ENSURE
+	#define PLATFORM_USE_REPORT_ENSURE PLATFORM_DESKTOP
+#endif
+
 // deprecated, do not use
 #define PLATFORM_HAS_THREADSAFE_RHIGetRenderQueryResult	#
 #define PLATFORM_SUPPORTS_RHI_THREAD #
