Returning of the specialization for Vector CubicInterp, since it's faster than generic version

diff --git a/Engine/Source/Runtime/Core/Public/Math/Vector.h b/Engine/Source/Runtime/Core/Public/Math/Vector.h
--- a/Engine/Source/Runtime/Core/Public/Math/Vector.h
+++ b/Engine/Source/Runtime/Core/Public/Math/Vector.h
@@ -22,6 +22,10 @@
 #include "UObject/ObjectVersion.h"
 #include <type_traits>
 
+#if PLATFORM_ENABLE_VECTORINTRINSICS
+#include "Math/VectorRegister.h"
+#endif
+
 #ifdef _MSC_VER
 #pragma warning (push)
 // Ensure template functions don't generate shadowing warnings against global variables at the point of instantiation.
@@ -2586,6 +2590,31 @@ inline FVector FVector2D::SphericalToUnitCartesian() const
     return FVector(FMath::Cos(Y) * SinTheta, FMath::Sin(Y) * SinTheta, FMath::Cos(X));
 }
 
+#if PLATFORM_ENABLE_VECTORINTRINSICS
+template<>
+FORCEINLINE_DEBUGGABLE FVector FMath::CubicInterp(const FVector& P0, const FVector& T0, const FVector& P1, const FVector& T1, const float& A)
+{
+	static_assert(PLATFORM_ENABLE_VECTORINTRINSICS == 1, "Requires vector intrinsics.");
+	FVector res;
+
+	const float A2 = A * A;
+	const float A3 = A2 * A;
+
+	const float s0 = (2 * A3) - (3 * A2) + 1;
+	const float s1 = A3 - (2 * A2) + A;
+	const float s2 = (A3 - A2);
+	const float s3 = (-2 * A3) + (3 * A2);
+
+	VectorRegister v0 = VectorMultiply(VectorLoadFloat1(&s0), VectorLoadFloat3(&P0));
+	v0 = VectorMultiplyAdd(VectorLoadFloat1(&s1), VectorLoadFloat3(&T0), v0);
+	VectorRegister v1 = VectorMultiply(VectorLoadFloat1(&s2), VectorLoadFloat3(&T1));
+	v1 = VectorMultiplyAdd(VectorLoadFloat1(&s3), VectorLoadFloat3(&P1), v1);
+
+	VectorStoreFloat3(VectorAdd(v0, v1), &res);
+
+	return res;
+}
+#endif
 
 #ifdef _MSC_VER
 #pragma warning (pop)
