Added UTF-8 support to a few functions

#rnx

diff --git a/Engine/Source/Runtime/Core/Private/String/BytesToHex.cpp b/Engine/Source/Runtime/Core/Private/String/BytesToHex.cpp
--- a/Engine/Source/Runtime/Core/Private/String/BytesToHex.cpp
+++ b/Engine/Source/Runtime/Core/Private/String/BytesToHex.cpp
@@ -9,10 +9,10 @@ namespace UE
 namespace String
 {
 
-template <typename CharType, CharType LetterA>
+template <typename CharType, ANSICHAR LetterA>
 static inline void BytesToHexImpl(TConstArrayView<uint8> Bytes, CharType* OutHex)
 {
-	const auto NibbleToHex = [](uint8 Value) -> CharType { return Value + (Value > 9 ? LetterA - 10 : '0'); };
+	const auto NibbleToHex = [](uint8 Value) -> CharType { return CharType(Value + (Value > 9 ? LetterA - 10 : '0')); };
 	const uint8* Data = Bytes.GetData();
 	
 	for (const uint8* DataEnd = Data + Bytes.Num(); Data != DataEnd; ++Data)
@@ -32,6 +32,11 @@ void BytesToHex(TConstArrayView<uint8> Bytes, WIDECHAR* OutHex)
 	BytesToHexImpl<WIDECHAR, 'A'>(Bytes, OutHex);
 }
 
+void BytesToHex(TConstArrayView<uint8> Bytes, UTF8CHAR* OutHex)
+{
+	BytesToHexImpl<UTF8CHAR, 'A'>(Bytes, OutHex);
+}
+
 void BytesToHexLower(TConstArrayView<uint8> Bytes, ANSICHAR* OutHex)
 {
 	BytesToHexImpl<ANSICHAR, 'a'>(Bytes, OutHex);
@@ -42,6 +47,11 @@ void BytesToHexLower(TConstArrayView<uint8> Bytes, WIDECHAR* OutHex)
 	BytesToHexImpl<WIDECHAR, 'a'>(Bytes, OutHex);
 }
 
+void BytesToHexLower(TConstArrayView<uint8> Bytes, UTF8CHAR* OutHex)
+{
+	BytesToHexImpl<UTF8CHAR, 'a'>(Bytes, OutHex);
+}
+
 void BytesToHex(TConstArrayView<uint8> Bytes, FAnsiStringBuilderBase& Builder)
 {
 	const int32 Offset = Builder.AddUninitialized(Bytes.Num() * 2);
@@ -54,6 +64,12 @@ void BytesToHex(TConstArrayView<uint8> Bytes, FWideStringBuilderBase& Builder)
 	BytesToHexImpl<WIDECHAR, 'A'>(Bytes, GetData(Builder) + Offset);
 }
 
+void BytesToHex(TConstArrayView<uint8> Bytes, FUtf8StringBuilderBase& Builder)
+{
+	const int32 Offset = Builder.AddUninitialized(Bytes.Num() * 2);
+	BytesToHexImpl<UTF8CHAR, 'A'>(Bytes, GetData(Builder) + Offset);
+}
+
 void BytesToHexLower(TConstArrayView<uint8> Bytes, FAnsiStringBuilderBase& Builder)
 {
 	const int32 Offset = Builder.AddUninitialized(Bytes.Num() * 2);
@@ -66,5 +82,11 @@ void BytesToHexLower(TConstArrayView<uint8> Bytes, FWideStringBuilderBase& Build
 	BytesToHexImpl<WIDECHAR, 'a'>(Bytes, GetData(Builder) + Offset);
 }
 
+void BytesToHexLower(TConstArrayView<uint8> Bytes, FUtf8StringBuilderBase& Builder)
+{
+	const int32 Offset = Builder.AddUninitialized(Bytes.Num() * 2);
+	BytesToHexImpl<UTF8CHAR, 'a'>(Bytes, GetData(Builder) + Offset);
+}
+
 }
 }
diff --git a/Engine/Source/Runtime/Core/Private/String/HexToBytes.cpp b/Engine/Source/Runtime/Core/Private/String/HexToBytes.cpp
--- a/Engine/Source/Runtime/Core/Private/String/HexToBytes.cpp
+++ b/Engine/Source/Runtime/Core/Private/String/HexToBytes.cpp
@@ -39,5 +39,10 @@ int32 HexToBytes(FWideStringView Hex, uint8* OutBytes)
 	return HexToBytesImpl<WIDECHAR>(Hex, OutBytes);
 }
 
+int32 HexToBytes(FUtf8StringView Hex, uint8* OutBytes)
+{
+	return HexToBytesImpl<UTF8CHAR>(Hex, OutBytes);
+}
+
 }
 }
diff --git a/Engine/Source/Runtime/Core/Public/Hash/Blake3.h b/Engine/Source/Runtime/Core/Public/Hash/Blake3.h
--- a/Engine/Source/Runtime/Core/Public/Hash/Blake3.h
+++ b/Engine/Source/Runtime/Core/Public/Hash/Blake3.h
@@ -154,13 +154,8 @@ inline uint32 GetTypeHash(const FBlake3Hash& Hash)
 	return *reinterpret_cast<const uint32*>(Hash.GetBytes());
 }
 
-inline FAnsiStringBuilderBase& operator<<(FAnsiStringBuilderBase& Builder, const FBlake3Hash& Hash)
-{
-	UE::String::BytesToHexLower(Hash.GetBytes(), Builder);
-	return Builder;
-}
-
-inline FWideStringBuilderBase& operator<<(FWideStringBuilderBase& Builder, const FBlake3Hash& Hash)
+template <typename CharType>
+inline TStringBuilderBase<CharType>& operator<<(TStringBuilderBase<CharType>& Builder, const FBlake3Hash& Hash)
 {
 	UE::String::BytesToHexLower(Hash.GetBytes(), Builder);
 	return Builder;
diff --git a/Engine/Source/Runtime/Core/Public/IO/IoHash.h b/Engine/Source/Runtime/Core/Public/IO/IoHash.h
--- a/Engine/Source/Runtime/Core/Public/IO/IoHash.h
+++ b/Engine/Source/Runtime/Core/Public/IO/IoHash.h
@@ -152,13 +152,8 @@ inline uint32 GetTypeHash(const FIoHash& Hash)
 	return *reinterpret_cast<const uint32*>(Hash.GetBytes());
 }
 
-inline FAnsiStringBuilderBase& operator<<(FAnsiStringBuilderBase& Builder, const FIoHash& Hash)
-{
-	UE::String::BytesToHexLower(Hash.GetBytes(), Builder);
-	return Builder;
-}
-
-inline FWideStringBuilderBase& operator<<(FWideStringBuilderBase& Builder, const FIoHash& Hash)
+template <typename CharType>
+inline TStringBuilderBase<CharType>& operator<<(TStringBuilderBase<CharType>& Builder, const FIoHash& Hash)
 {
 	UE::String::BytesToHexLower(Hash.GetBytes(), Builder);
 	return Builder;
diff --git a/Engine/Source/Runtime/Core/Public/String/BytesToHex.h b/Engine/Source/Runtime/Core/Public/String/BytesToHex.h
--- a/Engine/Source/Runtime/Core/Public/String/BytesToHex.h
+++ b/Engine/Source/Runtime/Core/Public/String/BytesToHex.h
@@ -19,8 +19,10 @@ namespace String
  */
 CORE_API void BytesToHex(TConstArrayView<uint8> Bytes, ANSICHAR* OutHex);
 CORE_API void BytesToHex(TConstArrayView<uint8> Bytes, WIDECHAR* OutHex);
+CORE_API void BytesToHex(TConstArrayView<uint8> Bytes, UTF8CHAR* OutHex);
 CORE_API void BytesToHexLower(TConstArrayView<uint8> Bytes, ANSICHAR* OutHex);
 CORE_API void BytesToHexLower(TConstArrayView<uint8> Bytes, WIDECHAR* OutHex);
+CORE_API void BytesToHexLower(TConstArrayView<uint8> Bytes, UTF8CHAR* OutHex);
 
 /**
  * Append an array of bytes to a string builder as hex digits.
@@ -30,8 +32,10 @@ CORE_API void BytesToHexLower(TConstArrayView<uint8> Bytes, WIDECHAR* OutHex);
  */
 CORE_API void BytesToHex(TConstArrayView<uint8> Bytes, FAnsiStringBuilderBase& Builder);
 CORE_API void BytesToHex(TConstArrayView<uint8> Bytes, FWideStringBuilderBase& Builder);
+CORE_API void BytesToHex(TConstArrayView<uint8> Bytes, FUtf8StringBuilderBase& Builder);
 CORE_API void BytesToHexLower(TConstArrayView<uint8> Bytes, FAnsiStringBuilderBase& Builder);
 CORE_API void BytesToHexLower(TConstArrayView<uint8> Bytes, FWideStringBuilderBase& Builder);
+CORE_API void BytesToHexLower(TConstArrayView<uint8> Bytes, FUtf8StringBuilderBase& Builder);
 
 }
 }
diff --git a/Engine/Source/Runtime/Core/Public/String/HexToBytes.h b/Engine/Source/Runtime/Core/Public/String/HexToBytes.h
--- a/Engine/Source/Runtime/Core/Public/String/HexToBytes.h
+++ b/Engine/Source/Runtime/Core/Public/String/HexToBytes.h
@@ -20,6 +20,7 @@ namespace String
  */
 CORE_API int32 HexToBytes(FAnsiStringView Hex, uint8* OutBytes);
 CORE_API int32 HexToBytes(FWideStringView Hex, uint8* OutBytes);
+CORE_API int32 HexToBytes(FUtf8StringView Hex, uint8* OutBytes);
 
 }
 }
