Mark all FPlatformProcess::SleepInfinite implementations as [[noreturn]]

diff --git a/Engine/Source/Runtime/Core/Private/Android/AndroidSignals.h b/Engine/Source/Runtime/Core/Private/Android/AndroidSignals.h
--- a/Engine/Source/Runtime/Core/Private/Android/AndroidSignals.h
+++ b/Engine/Source/Runtime/Core/Private/Android/AndroidSignals.h
@@ -114,10 +114,7 @@ private:
 	{
 		ForwardingThreadID = FPlatformTLS::GetCurrentThreadId();
 		FPlatformAtomics::AtomicStore(&SignalThreadStatus, (int32)ESignalThreadStatus::Ready);
-		while (true)
-		{
-			FPlatformProcess::SleepInfinite();
-		}
+		FPlatformProcess::SleepInfinite();
 		return nullptr;
 	}
 
diff --git a/Engine/Source/Runtime/Core/Private/GenericPlatform/GenericPlatformProcess.cpp b/Engine/Source/Runtime/Core/Private/GenericPlatform/GenericPlatformProcess.cpp
--- a/Engine/Source/Runtime/Core/Private/GenericPlatform/GenericPlatformProcess.cpp
+++ b/Engine/Source/Runtime/Core/Private/GenericPlatform/GenericPlatformProcess.cpp
@@ -357,7 +357,10 @@ void FGenericPlatformProcess::SleepNoStats( float Seconds )
 void FGenericPlatformProcess::SleepInfinite()
 {
 	// stop this thread forever
-	pause();
+	while (true)
+	{
+		pause();
+	}
 }
 
 void FGenericPlatformProcess::YieldThread()
diff --git a/Engine/Source/Runtime/Core/Private/HoloLens/HoloLensProcess.cpp b/Engine/Source/Runtime/Core/Private/HoloLens/HoloLensProcess.cpp
--- a/Engine/Source/Runtime/Core/Private/HoloLens/HoloLensProcess.cpp
+++ b/Engine/Source/Runtime/Core/Private/HoloLens/HoloLensProcess.cpp
@@ -86,7 +86,10 @@ void FHoloLensProcess::SleepNoStats(float Seconds)
 
 void FHoloLensProcess::SleepInfinite()
 {
-	::Sleep(INFINITE);
+	while (true)
+	{
+		::Sleep(INFINITE);
+	}
 }
 
 void FHoloLensProcess::YieldThread()
diff --git a/Engine/Source/Runtime/Core/Private/IOS/IOSMallocZone.cpp b/Engine/Source/Runtime/Core/Private/IOS/IOSMallocZone.cpp
--- a/Engine/Source/Runtime/Core/Private/IOS/IOSMallocZone.cpp
+++ b/Engine/Source/Runtime/Core/Private/IOS/IOSMallocZone.cpp
@@ -190,9 +190,7 @@ bool FIOSMallocCrashHandler::IsOnCrashedThread( void )
 	}
 	else
 	{
-		do {
-			FPlatformProcess::SleepInfinite();
-		} while (true);
+		FPlatformProcess::SleepInfinite();
 		return false;
 	}
 }
diff --git a/Engine/Source/Runtime/Core/Private/Mac/MacMallocZone.cpp b/Engine/Source/Runtime/Core/Private/Mac/MacMallocZone.cpp
--- a/Engine/Source/Runtime/Core/Private/Mac/MacMallocZone.cpp
+++ b/Engine/Source/Runtime/Core/Private/Mac/MacMallocZone.cpp
@@ -184,9 +184,7 @@ bool FMacMallocCrashHandler::IsOnCrashedThread( void )
 	}
 	else
 	{
-		do {
-			FPlatformProcess::SleepInfinite();
-		} while (true);
+		FPlatformProcess::SleepInfinite();
 		return false;
 	}
 }
diff --git a/Engine/Source/Runtime/Core/Private/Windows/WindowsPlatformProcess.cpp b/Engine/Source/Runtime/Core/Private/Windows/WindowsPlatformProcess.cpp
--- a/Engine/Source/Runtime/Core/Private/Windows/WindowsPlatformProcess.cpp
+++ b/Engine/Source/Runtime/Core/Private/Windows/WindowsPlatformProcess.cpp
@@ -1369,8 +1369,10 @@ void FWindowsPlatformProcess::SleepNoStats(float Seconds)
 
 void FWindowsPlatformProcess::SleepInfinite()
 {
-	check(FPlatformProcess::SupportsMultithreading());
-	::Sleep(INFINITE);
+	while (true)
+	{
+		::Sleep(INFINITE);
+	}
 }
 
 void FWindowsPlatformProcess::YieldThread()
diff --git a/Engine/Source/Runtime/Core/Public/GenericPlatform/GenericPlatformProcess.h b/Engine/Source/Runtime/Core/Public/GenericPlatform/GenericPlatformProcess.h
--- a/Engine/Source/Runtime/Core/Public/GenericPlatform/GenericPlatformProcess.h
+++ b/Engine/Source/Runtime/Core/Public/GenericPlatform/GenericPlatformProcess.h
@@ -516,7 +516,7 @@ struct CORE_API FGenericPlatformProcess
 	/** Sleep this thread for Seconds.  0.0 means release the current time slice to let other threads get some attention. */
 	static void SleepNoStats( float Seconds );
 	/** Sleep this thread infinitely. */
-	static void SleepInfinite();
+	[[noreturn]] static void SleepInfinite();
 	/** Yield this thread so another may run for a while. */
 	static void YieldThread();
 
diff --git a/Engine/Source/Runtime/Core/Public/HoloLens/HoloLensPlatformProcess.h b/Engine/Source/Runtime/Core/Public/HoloLens/HoloLensPlatformProcess.h
--- a/Engine/Source/Runtime/Core/Public/HoloLens/HoloLensPlatformProcess.h
+++ b/Engine/Source/Runtime/Core/Public/HoloLens/HoloLensPlatformProcess.h
@@ -30,7 +30,7 @@ struct CORE_API FHoloLensProcess : public FGenericPlatformProcess
 	static const TCHAR* BaseDir();
 	static void Sleep( float Seconds );
 	static void SleepNoStats(float Seconds); 
-	static void SleepInfinite();
+	[[noreturn]] static void SleepInfinite();
 	static void YieldThread();
 	UE_DEPRECATED(5.0, "Please use GetSynchEventFromPool to create a new event, and ReturnSynchEventToPool to release the event.")
 	static class FEvent* CreateSynchEvent(bool bIsManualReset = false);
diff --git a/Engine/Source/Runtime/Core/Public/Windows/WindowsPlatformProcess.h b/Engine/Source/Runtime/Core/Public/Windows/WindowsPlatformProcess.h
--- a/Engine/Source/Runtime/Core/Public/Windows/WindowsPlatformProcess.h
+++ b/Engine/Source/Runtime/Core/Public/Windows/WindowsPlatformProcess.h
@@ -175,7 +175,7 @@ public:
 	static bool ResolveNetworkPath( FString InUNCPath, FString& OutPath ); 
 	static void Sleep(float Seconds);
 	static void SleepNoStats(float Seconds);
-	static void SleepInfinite();
+	[[noreturn]] static void SleepInfinite();
 	static void YieldThread();
 	UE_DEPRECATED(5.0, "Please use GetSynchEventFromPool to create a new event, and ReturnSynchEventToPool to release the event.")
 	static class FEvent* CreateSynchEvent(bool bIsManualReset = false);
