Fix OnOutOfMemory thread-safety issue when concurrent calls return from the function

Change all FPlatformMemory::OnOutOfMemory implementations to [[noreturn]]

diff --git a/Engine/Source/Runtime/Core/Private/GenericPlatform/GenericPlatformMemory.cpp b/Engine/Source/Runtime/Core/Private/GenericPlatform/GenericPlatformMemory.cpp
--- a/Engine/Source/Runtime/Core/Private/GenericPlatform/GenericPlatformMemory.cpp
+++ b/Engine/Source/Runtime/Core/Private/GenericPlatform/GenericPlatformMemory.cpp
@@ -153,53 +153,54 @@ void FGenericPlatformMemory::Init()
 
 void FGenericPlatformMemory::OnOutOfMemory(uint64 Size, uint32 Alignment)
 {
-	// Update memory stats before we enter the crash handler.
-	OOMAllocationSize = Size;
-	OOMAllocationAlignment = Alignment;
-
-	// only call this code one time - if already OOM, abort
-	if (bIsOOM)
+	auto HandleOOM = [&]()
 	{
-		return;
-	}
-	bIsOOM = true;
+		// Update memory stats before we enter the crash handler.
+		OOMAllocationSize = Size;
+		OOMAllocationAlignment = Alignment;
 
-	const int ErrorMsgSize = 256;
-	TCHAR ErrorMsg[ErrorMsgSize];
-	FPlatformMisc::GetSystemErrorMessage(ErrorMsg, ErrorMsgSize, 0);
+		bIsOOM = true;
 
-	FPlatformMemoryStats PlatformMemoryStats = FPlatformMemory::GetStats();
-	if (BackupOOMMemoryPool)
-	{
-		FPlatformMemory::BinnedFreeToOS(BackupOOMMemoryPool, FPlatformMemory::GetBackMemoryPoolSize());
-		UE_LOG(LogMemory, Warning, TEXT("Freeing %d bytes from backup pool to handle out of memory."), FPlatformMemory::GetBackMemoryPoolSize());
+		const int ErrorMsgSize = 256;
+		TCHAR ErrorMsg[ErrorMsgSize];
+		FPlatformMisc::GetSystemErrorMessage(ErrorMsg, ErrorMsgSize, 0);
+
+		FPlatformMemoryStats PlatformMemoryStats = FPlatformMemory::GetStats();
+		if (BackupOOMMemoryPool)
+		{
+			FPlatformMemory::BinnedFreeToOS(BackupOOMMemoryPool, FPlatformMemory::GetBackMemoryPoolSize());
+			UE_LOG(LogMemory, Warning, TEXT("Freeing %d bytes from backup pool to handle out of memory."), FPlatformMemory::GetBackMemoryPoolSize());
         
-		LLM(FLowLevelMemTracker::Get().OnLowLevelFree(ELLMTracker::Default, BackupOOMMemoryPool));
-	}
+			LLM(FLowLevelMemTracker::Get().OnLowLevelFree(ELLMTracker::Default, BackupOOMMemoryPool));
+		}
 
-	UE_LOG(LogMemory, Warning, TEXT("MemoryStats:")\
-		TEXT("\n\tAvailablePhysical %llu")\
-		TEXT("\n\t AvailableVirtual %llu")\
-		TEXT("\n\t     UsedPhysical %llu")\
-		TEXT("\n\t PeakUsedPhysical %llu")\
-		TEXT("\n\t      UsedVirtual %llu")\
-		TEXT("\n\t  PeakUsedVirtual %llu"),
-		(uint64)PlatformMemoryStats.AvailablePhysical,
-		(uint64)PlatformMemoryStats.AvailableVirtual,
-		(uint64)PlatformMemoryStats.UsedPhysical,
-		(uint64)PlatformMemoryStats.PeakUsedPhysical,
-		(uint64)PlatformMemoryStats.UsedVirtual,
-		(uint64)PlatformMemoryStats.PeakUsedVirtual);
-	if (GWarn)
-	{
-		GMalloc->DumpAllocatorStats(*GWarn);
-	}
+		UE_LOG(LogMemory, Warning, TEXT("MemoryStats:")\
+			TEXT("\n\tAvailablePhysical %llu")\
+			TEXT("\n\t AvailableVirtual %llu")\
+			TEXT("\n\t     UsedPhysical %llu")\
+			TEXT("\n\t PeakUsedPhysical %llu")\
+			TEXT("\n\t      UsedVirtual %llu")\
+			TEXT("\n\t  PeakUsedVirtual %llu"),
+			(uint64)PlatformMemoryStats.AvailablePhysical,
+			(uint64)PlatformMemoryStats.AvailableVirtual,
+			(uint64)PlatformMemoryStats.UsedPhysical,
+			(uint64)PlatformMemoryStats.PeakUsedPhysical,
+			(uint64)PlatformMemoryStats.UsedVirtual,
+			(uint64)PlatformMemoryStats.PeakUsedVirtual);
+		if (GWarn)
+		{
+			GMalloc->DumpAllocatorStats(*GWarn);
+		}
 
-	// let any registered handlers go
-	FCoreDelegates::GetOutOfMemoryDelegate().Broadcast();
+		// let any registered handlers go
+		FCoreDelegates::GetOutOfMemoryDelegate().Broadcast();
 
-	// ErrorMsg might be unrelated to OoM error in some cases as the code that calls OnOutOfMemory could have called other system functions that modified errno
-	UE_LOG(LogMemory, Fatal, TEXT("Ran out of memory allocating %llu bytes with alignment %u. Last error msg: %s."), Size, Alignment, ErrorMsg);
+		// ErrorMsg might be unrelated to OoM error in some cases as the code that calls OnOutOfMemory could have called other system functions that modified errno
+		UE_LOG(LogMemory, Fatal, TEXT("Ran out of memory allocating %llu bytes with alignment %u. Last error msg: %s."), Size, Alignment, ErrorMsg);
+	};
+	
+	UE_CALL_ONCE(HandleOOM);
+	FPlatformProcess::SleepInfinite(); // Unreachable
 }
 
 FMalloc* FGenericPlatformMemory::BaseAllocator()
diff --git a/Engine/Source/Runtime/Core/Private/IOS/IOSPlatformMemory.cpp b/Engine/Source/Runtime/Core/Private/IOS/IOSPlatformMemory.cpp
--- a/Engine/Source/Runtime/Core/Private/IOS/IOSPlatformMemory.cpp
+++ b/Engine/Source/Runtime/Core/Private/IOS/IOSPlatformMemory.cpp
@@ -13,51 +13,52 @@
 
 void FIOSPlatformMemory::OnOutOfMemory(uint64 Size, uint32 Alignment)
 {
-    // Update memory stats before we enter the crash handler.
-    OOMAllocationSize = Size;
-    OOMAllocationAlignment = Alignment;
+	auto HandleOOM = [&]()
+	{
+		// Update memory stats before we enter the crash handler.
+		OOMAllocationSize = Size;
+		OOMAllocationAlignment = Alignment;
     
-    // only call this code one time - if already OOM, abort
-    if (bIsOOM)
-    {
-        return;
-    }
-    bIsOOM = true;
+		bIsOOM = true;
 
-	const int ErrorMsgSize = 256;
-	TCHAR ErrorMsg[ErrorMsgSize];
-	FPlatformMisc::GetSystemErrorMessage(ErrorMsg, ErrorMsgSize, 0);
+		const int ErrorMsgSize = 256;
+		TCHAR ErrorMsg[ErrorMsgSize];
+		FPlatformMisc::GetSystemErrorMessage(ErrorMsg, ErrorMsgSize, 0);
     
-    FPlatformMemoryStats PlatformMemoryStats = FPlatformMemory::GetStats();
-    if (BackupOOMMemoryPool)
-    {
-        FPlatformMemory::BinnedFreeToOS(BackupOOMMemoryPool, FPlatformMemory::GetBackMemoryPoolSize());
-        UE_LOG(LogMemory, Warning, TEXT("Freeing %d bytes from backup pool to handle out of memory."), FPlatformMemory::GetBackMemoryPoolSize());
+		FPlatformMemoryStats PlatformMemoryStats = FPlatformMemory::GetStats();
+		if (BackupOOMMemoryPool)
+		{
+			FPlatformMemory::BinnedFreeToOS(BackupOOMMemoryPool, FPlatformMemory::GetBackMemoryPoolSize());
+			UE_LOG(LogMemory, Warning, TEXT("Freeing %d bytes from backup pool to handle out of memory."), FPlatformMemory::GetBackMemoryPoolSize());
         
-        LLM(FLowLevelMemTracker::Get().OnLowLevelFree(ELLMTracker::Default, BackupOOMMemoryPool));
-    }
+			LLM(FLowLevelMemTracker::Get().OnLowLevelFree(ELLMTracker::Default, BackupOOMMemoryPool));
+		}
     
-    UE_LOG(LogMemory, Warning, TEXT("MemoryStats:")\
-           TEXT("\n\tAvailablePhysical %llu")\
-           TEXT("\n\t AvailableVirtual %llu")\
-           TEXT("\n\t     UsedPhysical %llu")\
-           TEXT("\n\t PeakUsedPhysical %llu")\
-           TEXT("\n\t      UsedVirtual %llu")\
-           TEXT("\n\t  PeakUsedVirtual %llu"),
-           (uint64)PlatformMemoryStats.AvailablePhysical,
-           (uint64)PlatformMemoryStats.AvailableVirtual,
-           (uint64)PlatformMemoryStats.UsedPhysical,
-           (uint64)PlatformMemoryStats.PeakUsedPhysical,
-           (uint64)PlatformMemoryStats.UsedVirtual,
-           (uint64)PlatformMemoryStats.PeakUsedVirtual);
+		UE_LOG(LogMemory, Warning, TEXT("MemoryStats:")\
+			   TEXT("\n\tAvailablePhysical %llu")\
+			   TEXT("\n\t AvailableVirtual %llu")\
+			   TEXT("\n\t     UsedPhysical %llu")\
+			   TEXT("\n\t PeakUsedPhysical %llu")\
+			   TEXT("\n\t      UsedVirtual %llu")\
+			   TEXT("\n\t  PeakUsedVirtual %llu"),
+			   (uint64)PlatformMemoryStats.AvailablePhysical,
+			   (uint64)PlatformMemoryStats.AvailableVirtual,
+			   (uint64)PlatformMemoryStats.UsedPhysical,
+			   (uint64)PlatformMemoryStats.PeakUsedPhysical,
+			   (uint64)PlatformMemoryStats.UsedVirtual,
+			   (uint64)PlatformMemoryStats.PeakUsedVirtual);
     
-    // let any registered handlers go
-    FCoreDelegates::GetMemoryTrimDelegate().Broadcast();
+		// let any registered handlers go
+		FCoreDelegates::GetMemoryTrimDelegate().Broadcast();
     
-	// ErrorMsg might be unrelated to OoM error in some cases as the code that calls OnOutOfMemory could have called other system functions that modified errno
-    UE_LOG(LogMemory, Warning, TEXT("Ran out of memory allocating %llu bytes with alignment %u. Last error msg: %s."), Size, Alignment, ErrorMsg);
+		// ErrorMsg might be unrelated to OoM error in some cases as the code that calls OnOutOfMemory could have called other system functions that modified errno
+		UE_LOG(LogMemory, Warning, TEXT("Ran out of memory allocating %llu bytes with alignment %u. Last error msg: %s."), Size, Alignment, ErrorMsg);
     
-    // make this a fatal error that ends here not in the log
-    // changed to 3 from NULL because clang noticed writing to NULL and warned about it
-    *(int32 *)3 = 123;
+		// make this a fatal error that ends here not in the log
+		// changed to 3 from NULL because clang noticed writing to NULL and warned about it
+		*(int32 *)3 = 123;
+	};
+	
+	UE_CALL_ONCE(HandleOOM);
+	FPlatformProcess::SleepInfinite(); // Unreachable
 }
diff --git a/Engine/Source/Runtime/Core/Private/Unix/UnixPlatformMemory.cpp b/Engine/Source/Runtime/Core/Private/Unix/UnixPlatformMemory.cpp
--- a/Engine/Source/Runtime/Core/Private/Unix/UnixPlatformMemory.cpp
+++ b/Engine/Source/Runtime/Core/Private/Unix/UnixPlatformMemory.cpp
@@ -1052,50 +1052,50 @@ bool FUnixPlatformMemory::UnmapNamedSharedMemoryRegion(FSharedMemoryRegion * Mem
 
 void FUnixPlatformMemory::OnOutOfMemory(uint64 Size, uint32 Alignment)
 {
-	// Update memory stats before we enter the crash handler.
-	OOMAllocationSize = Size;
-	OOMAllocationAlignment = Alignment;
-
-	// only call this code one time - if already OOM, abort
-	if (bIsOOM)
-	{
-		return;
-	}
-	bIsOOM = true;
-
-	const int ErrorMsgSize = 256;
-	TCHAR ErrorMsg[ErrorMsgSize];
-	FPlatformMisc::GetSystemErrorMessage(ErrorMsg, ErrorMsgSize, 0);
-
-	FMalloc* Prev = GMalloc;
-	FPlatformMallocCrash::Get().SetAsGMalloc();
-
-	FPlatformMemoryStats PlatformMemoryStats = FPlatformMemory::GetStats();
-
-	UE_LOG(LogMemory, Warning, TEXT("MemoryStats:")\
-		TEXT("\n\tAvailablePhysical %llu")\
-		TEXT("\n\t AvailableVirtual %llu")\
-		TEXT("\n\t     UsedPhysical %llu")\
-		TEXT("\n\t PeakUsedPhysical %llu")\
-		TEXT("\n\t      UsedVirtual %llu")\
-		TEXT("\n\t  PeakUsedVirtual %llu"),
-		(uint64)PlatformMemoryStats.AvailablePhysical,
-		(uint64)PlatformMemoryStats.AvailableVirtual,
-		(uint64)PlatformMemoryStats.UsedPhysical,
-		(uint64)PlatformMemoryStats.PeakUsedPhysical,
-		(uint64)PlatformMemoryStats.UsedVirtual,
-		(uint64)PlatformMemoryStats.PeakUsedVirtual);
-	if (GWarn)
+	auto HandleOOM = [&]()
 	{
-		Prev->DumpAllocatorStats(*GWarn);
-	}
+		// Update memory stats before we enter the crash handler.
+		OOMAllocationSize = Size;
+		OOMAllocationAlignment = Alignment;
+
+		bIsOOM = true;
+
+		const int ErrorMsgSize = 256;
+		TCHAR ErrorMsg[ErrorMsgSize];
+		FPlatformMisc::GetSystemErrorMessage(ErrorMsg, ErrorMsgSize, 0);
+
+		FMalloc* Prev = GMalloc;
+		FPlatformMallocCrash::Get().SetAsGMalloc();
+
+		FPlatformMemoryStats PlatformMemoryStats = FPlatformMemory::GetStats();
+
+		UE_LOG(LogMemory, Warning, TEXT("MemoryStats:")\
+			TEXT("\n\tAvailablePhysical %llu")\
+			TEXT("\n\t AvailableVirtual %llu")\
+			TEXT("\n\t     UsedPhysical %llu")\
+			TEXT("\n\t PeakUsedPhysical %llu")\
+			TEXT("\n\t      UsedVirtual %llu")\
+			TEXT("\n\t  PeakUsedVirtual %llu"),
+			(uint64)PlatformMemoryStats.AvailablePhysical,
+			(uint64)PlatformMemoryStats.AvailableVirtual,
+			(uint64)PlatformMemoryStats.UsedPhysical,
+			(uint64)PlatformMemoryStats.PeakUsedPhysical,
+			(uint64)PlatformMemoryStats.UsedVirtual,
+			(uint64)PlatformMemoryStats.PeakUsedVirtual);
+		if (GWarn)
+		{
+			Prev->DumpAllocatorStats(*GWarn);
+		}
 
-	// let any registered handlers go
-	FCoreDelegates::GetOutOfMemoryDelegate().Broadcast();
+		// let any registered handlers go
+		FCoreDelegates::GetOutOfMemoryDelegate().Broadcast();
 
-	// ErrorMsg might be unrelated to OoM error in some cases as the code that calls OnOutOfMemory could have called other system functions that modified errno
-	UE_LOG(LogMemory, Fatal, TEXT("Ran out of memory allocating %llu bytes with alignment %u. Last error msg: %s."), Size, Alignment, ErrorMsg);
-	// unreachable
+		// ErrorMsg might be unrelated to OoM error in some cases as the code that calls OnOutOfMemory could have called other system functions that modified errno
+		UE_LOG(LogMemory, Fatal, TEXT("Ran out of memory allocating %llu bytes with alignment %u. Last error msg: %s."), Size, Alignment, ErrorMsg);
+	};
+	
+	UE_CALL_ONCE(HandleOOM);
+	FPlatformProcess::SleepInfinite(); // Unreachable
 }
 
 /**
diff --git a/Engine/Source/Runtime/Core/Public/GenericPlatform/GenericPlatformMemory.h b/Engine/Source/Runtime/Core/Public/GenericPlatform/GenericPlatformMemory.h
--- a/Engine/Source/Runtime/Core/Public/GenericPlatform/GenericPlatformMemory.h
+++ b/Engine/Source/Runtime/Core/Public/GenericPlatform/GenericPlatformMemory.h
@@ -293,7 +293,7 @@ struct CORE_API FGenericPlatformMemory
 	/** Initializes platform memory specific constants. */
 	static void Init();
 	
-	static CA_NO_RETURN void OnOutOfMemory(uint64 Size, uint32 Alignment);
+	[[noreturn]] static void OnOutOfMemory(uint64 Size, uint32 Alignment);
 
 	/** Initializes the memory pools, should be called by the init function. */
 	static void SetupMemoryPools();
diff --git a/Engine/Source/Runtime/Core/Public/IOS/IOSPlatformMemory.h b/Engine/Source/Runtime/Core/Public/IOS/IOSPlatformMemory.h
--- a/Engine/Source/Runtime/Core/Public/IOS/IOSPlatformMemory.h
+++ b/Engine/Source/Runtime/Core/Public/IOS/IOSPlatformMemory.h
@@ -15,7 +15,7 @@
 struct CORE_API FIOSPlatformMemory : public FApplePlatformMemory
 {
     // added this for now because Crashlytics doesn't properly break up different callstacks all ending in UE_LOG(LogXXX, Fatal, ...)
-    static CA_NO_RETURN void OnOutOfMemory(uint64 Size, uint32 Alignment);
+    [[noreturn]] static void OnOutOfMemory(uint64 Size, uint32 Alignment);
 };
 
 typedef FIOSPlatformMemory FPlatformMemory;
diff --git a/Engine/Source/Runtime/Core/Public/Unix/UnixPlatformMemory.h b/Engine/Source/Runtime/Core/Public/Unix/UnixPlatformMemory.h
--- a/Engine/Source/Runtime/Core/Public/Unix/UnixPlatformMemory.h
+++ b/Engine/Source/Runtime/Core/Public/Unix/UnixPlatformMemory.h
@@ -132,7 +132,7 @@ struct CORE_API FUnixPlatformMemory : public FGenericPlatformMemory
 	static FSharedMemoryRegion * MapNamedSharedMemoryRegion(const FString& InName, bool bCreate, uint32 AccessMode, SIZE_T Size);
 	static bool UnmapNamedSharedMemoryRegion(FSharedMemoryRegion * MemoryRegion);
 	static bool GetLLMAllocFunctions(void*(*&OutAllocFunction)(size_t), void(*&OutFreeFunction)(void*, size_t), int32& OutAlignment);
-	static CA_NO_RETURN void OnOutOfMemory(uint64 Size, uint32 Alignment);
+	[[noreturn]] static void OnOutOfMemory(uint64 Size, uint32 Alignment);
 	//~ End FGenericPlatformMemory Interface
 
 	static bool HasForkPageProtectorEnabled();
