LowLevelMemoryTracker: Add FTagData::IsReporting to hide the CustomName tag from Insights and other reporting methods.
#rnx

diff --git a/Engine/Source/Runtime/Core/Private/HAL/LowLevelMemTracker.cpp b/Engine/Source/Runtime/Core/Private/HAL/LowLevelMemTracker.cpp
--- a/Engine/Source/Runtime/Core/Private/HAL/LowLevelMemTracker.cpp
+++ b/Engine/Source/Runtime/Core/Private/HAL/LowLevelMemTracker.cpp
@@ -1286,6 +1286,11 @@ void FLowLevelMemTracker::BootstrapTagDatas()
 		}
 		LLM_ENUM_GENERIC_TAGS(REGISTER_ELLMTAG);
 #undef REGISTER_ELLMTAG
+
+		// The CustomName tag is an adapter for connecting LLM_SCOPE_BYNAME tags with platforms that used the ELLMTag-based reporting.
+		// We want to hide this procedural tags in systems that use FName-based reporting; if it is displayed it confusingly just displays 
+		// a sum of every LLM_SCOPE_BYNAME tag.
+		TagDataEnumMap[static_cast<int32>(ELLMTag::CustomName)]->SetIsReportable(false);
 	}
 }
 
@@ -2621,7 +2626,9 @@ namespace LLMPrivate
 	}
 
 	FTagData::FTagData(FName InName, FName InDisplayName, FName InParentName, FName InStatName, FName InSummaryStatName, bool bInHasEnumTag, ELLMTag InEnumTag, ETagReferenceSource InReferenceSource)
-		: Name(InName), DisplayName(InDisplayName), ParentName(InParentName), StatName(InStatName), SummaryStatName(InSummaryStatName), EnumTag(InEnumTag), ReferenceSource(InReferenceSource), bIsFinishConstructed(false), bParentIsName(true), bHasEnumTag(bInHasEnumTag)
+		: Name(InName), DisplayName(InDisplayName), ParentName(InParentName), StatName(InStatName)
+		, SummaryStatName(InSummaryStatName), EnumTag(InEnumTag), ReferenceSource(InReferenceSource)
+		, bIsFinishConstructed(false), bParentIsName(true), bHasEnumTag(bInHasEnumTag), bIsReportable(true)
 	{
 	}
 
@@ -2745,6 +2752,11 @@ namespace LLMPrivate
 		return Index;
 	}
 
+	bool FTagData::IsReportable() const
+	{
+		return bIsReportable;
+	}
+
 	void FTagData::SetParent(const FTagData* InParent)
 	{
 		if (bParentIsName)
@@ -2791,6 +2803,11 @@ namespace LLMPrivate
 		Index = InIndex;
 	}
 
+	void FTagData::SetIsReportable(bool bInIsReportable)
+	{
+		bIsReportable = bInIsReportable;
+	}
+
 	bool FTagData::IsUsedAsDisplayParent() const
 	{
 		// All Tags but one are UsedAsDisplayParent - their name is prepended during GetDisplayPath
@@ -3228,6 +3245,10 @@ namespace LLMPrivate
 		for (const TPair<const FTagData*, FTrackerTagSizeData>& It : TagSizes)
 		{
 			const FTagData* TagData = It.Key;
+			if (!TagData->IsReportable())
+			{
+				continue;
+			}
 			if (OverrideUntaggedTagData && TagData->GetName() == TagName_Untagged)
 			{
 				// Handled separately by OverrideUntaggedTagData
@@ -3718,6 +3739,10 @@ namespace LLMPrivate
 		for (const TPair<const FTagData*, FTrackerTagSizeData>& It : TagSizes)
 		{
 			const FTagData* TagData = It.Key;
+			if (!TagData->IsReportable())
+			{
+				continue;
+			}
 			if (TagData->GetName() == TagName_Untagged)
 			{
 				continue; // Handled by OverrideUntaggedName
@@ -3897,6 +3922,10 @@ namespace LLMPrivate
 		for (const TPair<const FTagData*, FTrackerTagSizeData>& It : TagSizes)
 		{
 			const FTagData* TagData = It.Key;
+			if (!TagData->IsReportable())
+			{
+				continue;
+			}
 			if (OverrideUntaggedTagData != nullptr && TagData->GetName() == TagName_Untagged)
 			{
 				continue; // Handled by OverrideUntaggedTagData
@@ -3936,6 +3965,10 @@ namespace LLMPrivate
 		for (const TPair<const FTagData*, FTrackerTagSizeData>& It : TagSizes)
 		{
 			const FTagData* TagData = It.Key;
+			if (!TagData->IsReportable())
+			{
+				continue;
+			}
 			if (OverrideUntaggedTagData && TagData->GetName() == TagName_Untagged)
 			{
 				continue; // Handled by OverrideUntaggedTagData
@@ -4008,6 +4041,10 @@ namespace LLMPrivate
 		for (const TPair<const FTagData*, FTrackerTagSizeData>& It : TagSizes)
 		{
 			const FTagData* TagData = It.Key;
+			if (!TagData->IsReportable())
+			{
+				continue;
+			}
 			if (OverrideUntaggedTagData && TagData->GetName() == TagName_Untagged)
 			{
 				// Handled separately by OverrideUntaggedTagData
diff --git a/Engine/Source/Runtime/Core/Private/HAL/LowLevelMemTrackerPrivate.h b/Engine/Source/Runtime/Core/Private/HAL/LowLevelMemTrackerPrivate.h
--- a/Engine/Source/Runtime/Core/Private/HAL/LowLevelMemTrackerPrivate.h
+++ b/Engine/Source/Runtime/Core/Private/HAL/LowLevelMemTrackerPrivate.h
@@ -77,9 +77,11 @@ namespace LLMPrivate
 		ELLMTag GetContainingEnum() const;
 		ETagReferenceSource GetReferenceSource() const;
 		int32 GetIndex() const;
+		bool IsReportable() const;
 
 		void SetParent(const FTagData* InParent);
 		void SetIndex(int32 InIndex);
+		void SetIsReportable(bool bInReportable);
 		void SetFinishConstructed();
 
 		// These functions are normally invalid - these properties should be immutable - but are called for EnumTags during bootstrapping
@@ -108,6 +110,7 @@ namespace LLMPrivate
 		bool bIsFinishConstructed;
 		bool bParentIsName;
 		bool bHasEnumTag;
+		bool bIsReportable;
 	};
 
 	// TagData container types that use FLLMAllocator
