Make file cache block size configurable with (read only) CVar fc.BlockSize

diff --git a/Engine/Source/Runtime/Core/Private/FileCache/FileCache.cpp b/Engine/Source/Runtime/Core/Private/FileCache/FileCache.cpp
--- a/Engine/Source/Runtime/Core/Private/FileCache/FileCache.cpp
+++ b/Engine/Source/Runtime/Core/Private/FileCache/FileCache.cpp
@@ -24,26 +24,31 @@ DECLARE_CYCLE_STAT(TEXT("Find Eviction Candidate"), STAT_SFC_FindEvictionCandida
 
 DEFINE_LOG_CATEGORY_STATIC(LogStreamingFileCache, Log, All);
 
-static const int CacheLineSize = 64 * 1024;
+static int32 GFileCacheBlockSizeKB = 64;
+static FAutoConsoleVariableRef CVarFileCacheBlockSize(
+	TEXT("fc.BlockSize"),
+	GFileCacheBlockSizeKB,
+	TEXT("Size of each block in KB in the global file cache object\n")
+	TEXT("Should match packaging compression block size for optimal reading from packege"),
+	ECVF_ReadOnly
+);
 
 static int32 GNumFileCacheBlocks = 256;
 static FAutoConsoleVariableRef CVarNumFileCacheBlocks(
-	TEXT("fc.NumFileCacheBlocks"),
+	TEXT("fc.NumBlocks"),
 	GNumFileCacheBlocks,
-	TEXT("Number of blocks in the global file cache object\n"),
-	ECVF_RenderThreadSafe
+	TEXT("Number of blocks in the global file cache object"),
+	ECVF_ReadOnly
 );
 
 // 
 // Strongly typed ids to avoid confusion in the code
 // 
-template <int SetBlockSize, typename Parameter> class StrongBlockIdentifier
+template <typename Parameter> class StrongBlockIdentifier
 {
 	static const int InvalidHandle = 0xFFFFFFFF;
 
 public:
-	static const int32 BlockSize = SetBlockSize;
-
 	StrongBlockIdentifier() : Id(InvalidHandle) {}
 	explicit StrongBlockIdentifier(int32 SetId) : Id(SetId) {}
 
@@ -56,24 +61,24 @@ public:
 	inline StrongBlockIdentifier operator--(int) { StrongBlockIdentifier Temp(*this); operator--(); return Temp; }
 
 	// Get the offset in the file to read this block
-	inline int64 GetOffset() const { checkSlow(IsValid()); return (int64)Id * (int64)BlockSize; }
-	inline int64 GetSize() const { checkSlow(IsValid()); return BlockSize; }
+	inline static int64 GetSize() { return (int64)(GFileCacheBlockSizeKB * 1024); }
+	inline int64 GetOffset() const { checkSlow(IsValid()); return (int64)Id * GetSize(); }
 
 	// Get the number of bytes that need to be read for this block
 	// takes into account incomplete blocks at the end of the file
-	inline int64 GetSize(int64 FileSize) const { checkSlow(IsValid()); return FMath::Min((int64)BlockSize, FileSize - GetOffset()); }
+	inline int64 GetSize(int64 FileSize) const { checkSlow(IsValid()); return FMath::Min(GetSize(), FileSize - GetOffset()); }
 
-	friend inline uint32 GetTypeHash(const StrongBlockIdentifier<SetBlockSize, Parameter>& Info) { return GetTypeHash(Info.Id); }
+	friend inline uint32 GetTypeHash(const StrongBlockIdentifier<Parameter>& Info) { return GetTypeHash(Info.Id); }
 
-	inline bool operator==(const StrongBlockIdentifier<SetBlockSize, Parameter>&Other) const { return Id == Other.Id; }
-	inline bool operator!=(const StrongBlockIdentifier<SetBlockSize, Parameter>&Other) const { return Id != Other.Id; }
+	inline bool operator==(const StrongBlockIdentifier<Parameter>&Other) const { return Id == Other.Id; }
+	inline bool operator!=(const StrongBlockIdentifier<Parameter>&Other) const { return Id != Other.Id; }
 
 private:
 	int32 Id;
 };
 
-using CacheLineID = StrongBlockIdentifier<CacheLineSize, struct CacheLineStrongType>; // Unique per file handle
-using CacheSlotID = StrongBlockIdentifier<CacheLineSize, struct CacheSlotStrongType>; // Unique per cache
+using CacheLineID = StrongBlockIdentifier<struct CacheLineStrongType>; // Unique per file handle
+using CacheSlotID = StrongBlockIdentifier<struct CacheSlotStrongType>; // Unique per cache
 
 class FFileCacheHandle;
 
@@ -97,7 +102,7 @@ public:
 	{
 		check(SlotID.Get() < SlotInfo.Num() - 1);
 		check(IsSlotLocked(SlotID)); // slot must be locked in order to access memory
-		return Memory + SlotID.Get() * CacheSlotID::BlockSize;
+		return Memory + SlotID.Get() * CacheSlotID::GetSize();
 	}
 
 	CacheSlotID AcquireAndLockSlot(FFileCacheHandle* InHandle, CacheLineID InLineID);
@@ -190,7 +195,7 @@ public:
 	// allocated with an extra dummy entry at index0 for linked list head
 	TArray<FSlotInfo> SlotInfo;
 	uint8* Memory;
-	int32 SizeInBytes;
+	int64 SizeInBytes;
 	int32 NumFreeSlots;
 };
 
@@ -213,18 +218,18 @@ public:
 	// Block helper functions. These are just convenience around basic math.
 	// 
 
- // templated uses of this may end up converting int64 to int32, but it's up to the user of the template to know
+	// templated uses of this may end up converting int64 to int32, but it's up to the user of the template to know
 	PRAGMA_DISABLE_UNSAFE_TYPECAST_WARNINGS
-		/*
-		 * Get the block id that contains the specified offset
-		 */
-		template<typename BlockIDType> inline BlockIDType GetBlock(int64 Offset)
+	/*
+		* Get the block id that contains the specified offset
+		*/
+	template<typename BlockIDType> inline BlockIDType GetBlock(int64 Offset)
 	{
-		return BlockIDType(FMath::DivideAndRoundDown(Offset, (int64)BlockIDType::BlockSize));
+		return BlockIDType(FMath::DivideAndRoundDown(Offset, BlockIDType::GetSize()));
 	}
 	PRAGMA_ENABLE_UNSAFE_TYPECAST_WARNINGS
 
-		template<typename BlockIDType> inline int32 GetNumBlocks(int64 Offset, int64 Size)
+	template<typename BlockIDType> inline int32 GetNumBlocks(int64 Offset, int64 Size)
 	{
 		BlockIDType FirstBlock = GetBlock<BlockIDType>(Offset);
 		BlockIDType LastBlock = GetBlock<BlockIDType>(Offset + Size - 1);// Block containing the last byte
@@ -234,7 +239,7 @@ public:
 	// Returns the offset within the first block covering the byte range to read from
 	template<typename BlockIDType> inline int64 GetBlockOffset(int64 Offset)
 	{
-		return Offset - FMath::DivideAndRoundDown(Offset, (int64)BlockIDType::BlockSize) *  BlockIDType::BlockSize;
+		return Offset - FMath::DivideAndRoundDown(Offset, BlockIDType::GetSize()) *  BlockIDType::GetSize();
 	}
 
 	// Returns the size within the first cache line covering the byte range to read
@@ -281,7 +286,7 @@ private:
 FFileCache::FFileCache(int32 NumSlots)
 	: EvictFileCacheCommand(TEXT("r.VT.EvictFileCache"), TEXT("Evict all the file caches in the VT system."),
 		FConsoleCommandDelegate::CreateRaw(this, &FFileCache::EvictFileCacheFromConsole))
-	, SizeInBytes(NumSlots * CacheSlotID::BlockSize)
+	, SizeInBytes(NumSlots * CacheSlotID::GetSize())
 	, NumFreeSlots(NumSlots)
 {
 	LLM_SCOPE(ELLMTag::FileSystem);
@@ -530,12 +535,13 @@ public:
 		FFileCache& Cache = GetCache();
 
 		const int64 Offset = InitialSlotOffset + InOffset;
-		const int32 SlotIndex = (int32)FMath::DivideAndRoundDown(Offset, (int64)CacheSlotID::BlockSize);
-		const int32 OffsetInSlot = (int32)(Offset - SlotIndex * CacheSlotID::BlockSize);
+		const int64 BlockSize = CacheSlotID::GetSize();
+		const int32 SlotIndex = (int32)FMath::DivideAndRoundDown(Offset, BlockSize);
+		const int32 OffsetInSlot = (int32)(Offset - SlotIndex * BlockSize);
 		checkSlow(SlotIndex >= 0 && SlotIndex < NumCacheSlots);
 		const void* SlotMemory = Cache.GetSlotMemory(CacheSlots[SlotIndex]);
 
-		OutSize = FMath::Min(InSize, (int64)CacheSlotID::BlockSize - OffsetInSlot);
+		OutSize = FMath::Min(InSize, BlockSize - OffsetInSlot);
 		return (uint8*)SlotMemory + OffsetInSlot;
 	}
 
@@ -576,7 +582,7 @@ void FFileCacheHandle::CheckForSizeRequestComplete()
 		check(FileSize > 0);
 
 		// LineInfos key is int32
-		const int64 TotalNumSlots = FMath::DivideAndRoundUp(FileSize, (int64)CacheLineSize);
+		const int64 TotalNumSlots = FMath::DivideAndRoundUp(FileSize, (int64)(GFileCacheBlockSizeKB * 1024));
 		check(TotalNumSlots < MAX_int32);
 	}
 }
@@ -859,7 +865,7 @@ IFileCacheHandle* IFileCacheHandle::CreateFileCacheHandle(IAsyncReadFileHandle*
 	}
 }
 
-uint32 IFileCacheHandle::GetFileCacheSize()
+int64 IFileCacheHandle::GetFileCacheSize()
 {
 	return GetCache().SizeInBytes;
 }
diff --git a/Engine/Source/Runtime/Core/Public/FileCache/FileCache.h b/Engine/Source/Runtime/Core/Public/FileCache/FileCache.h
--- a/Engine/Source/Runtime/Core/Public/FileCache/FileCache.h
+++ b/Engine/Source/Runtime/Core/Public/FileCache/FileCache.h
@@ -46,7 +46,7 @@ public:
 	virtual ~IFileCacheHandle() {};
 
 	/** Return size of underlying file cache in bytes. */
-	CORE_API static uint32 GetFileCacheSize();
+	CORE_API static int64 GetFileCacheSize();
 
 	/**
 	 * Read a byte range form the file. This can be a high-throughput operation and done lots of times for small reads.
