Use mach_absolute_time() for Apple platform cycle counting.

- Fixes UnrealPak

#rnx

diff --git a/Engine/Source/Runtime/Core/Public/GenericPlatform/GenericPlatformProcess.h b/Engine/Source/Runtime/Core/Public/GenericPlatform/GenericPlatformProcess.h
--- a/Engine/Source/Runtime/Core/Public/GenericPlatform/GenericPlatformProcess.h
+++ b/Engine/Source/Runtime/Core/Public/GenericPlatform/GenericPlatformProcess.h
@@ -14,6 +14,11 @@
 #include <emmintrin.h>
 #endif
 
+////////////////////////////////////////////////////////////////////////////////
+#if PLATFORM_APPLE
+#include <mach/mach_time.h>
+#endif
+
 ////////////////////////////////////////////////////////////////////////////////
 #if !defined(__clang__)
 #	include <intrin.h>
@@ -728,7 +733,9 @@ struct CORE_API FGenericPlatformProcess
 	*/
 	static FORCEINLINE uint64 ReadCycleCounter()
 	{
-#if defined(_MSC_VER) && PLATFORM_CPU_X86_FAMILY
+#if PLATFORM_APPLE
+		return mach_absolute_time();
+#elif defined(_MSC_VER) && PLATFORM_CPU_X86_FAMILY
 		return __rdtsc();
 #elif defined(_MSC_VER) && defined(_M_ARM)
 		return __rdpmccntr64();
@@ -759,7 +766,7 @@ struct CORE_API FGenericPlatformProcess
 	*/
 	static FORCEINLINE void YieldCycles(uint64 cycles)
 	{
-#if PLATFORM_IOS || PLATFORM_ANDROID
+#if PLATFORM_ANDROID
 		for (uint64 i = 0; i < cycles; i++)
 		{
 			Yield();
