FTicker ignores unbound delegates instead of asserting

diff --git a/Engine/Source/Runtime/Core/Private/Containers/Ticker.cpp b/Engine/Source/Runtime/Core/Private/Containers/Ticker.cpp
--- a/Engine/Source/Runtime/Core/Private/Containers/Ticker.cpp
+++ b/Engine/Source/Runtime/Core/Private/Containers/Ticker.cpp
@@ -149,8 +149,7 @@ FTSTicker::FElement::FElement(double InFireTime, float InDelayTime, const FTicke
 
 bool FTSTicker::FElement::Fire(float DeltaTime)
 {
-	check(Delegate.IsBound());
-	return Delegate.Execute(DeltaTime);
+	return Delegate.IsBound() && Delegate.Execute(DeltaTime);
 }
 
 FTSTickerObjectBase::FTSTickerObjectBase(float InDelay, FTSTicker& InTicker)
