CompactBinary: Rename payload to value to prepare for a public value type

#rnx

diff --git a/Engine/Source/Runtime/Core/Private/Serialization/CompactBinary.cpp b/Engine/Source/Runtime/Core/Private/Serialization/CompactBinary.cpp
--- a/Engine/Source/Runtime/Core/Private/Serialization/CompactBinary.cpp
+++ b/Engine/Source/Runtime/Core/Private/Serialization/CompactBinary.cpp
@@ -18,8 +18,8 @@
 namespace CompactBinaryPrivate
 {
 
-static constexpr const uint8 GEmptyObjectPayload[] = { uint8(ECbFieldType::Object), 0x00 };
-static constexpr const uint8 GEmptyArrayPayload[] = { uint8(ECbFieldType::Array), 0x01, 0x00 };
+static constexpr const uint8 GEmptyObjectValue[] = { uint8(ECbFieldType::Object), 0x00 };
+static constexpr const uint8 GEmptyArrayValue[] = { uint8(ECbFieldType::Array), 0x01, 0x00 };
 
 template <typename T>
 static constexpr FORCEINLINE T ReadUnaligned(const void* const Memory)
@@ -104,7 +104,7 @@ FCbFieldView::FCbFieldView(const void* const Data, ECbFieldType Type)
 
 	TypeWithFlags = Type;
 	NameLen = uint32(FMath::Clamp<uint64>(NameLen64, 0, ~uint32(0)));
-	Payload = Bytes;
+	Value = Bytes;
 }
 
 FCbObjectView FCbFieldView::AsObjectView()
@@ -139,12 +139,12 @@ FMemoryView FCbFieldView::AsBinaryView(const FMemoryView Default)
 {
 	if (FCbFieldType::IsBinary(TypeWithFlags))
 	{
-		const uint8* const PayloadBytes = static_cast<const uint8*>(Payload);
+		const uint8* const ValueBytes = static_cast<const uint8*>(Value);
 		uint32 ValueSizeByteCount;
-		const uint64 ValueSize = ReadVarUInt(PayloadBytes, ValueSizeByteCount);
+		const uint64 ValueSize = ReadVarUInt(ValueBytes, ValueSizeByteCount);
 
 		Error = ECbFieldError::None;
-		return MakeMemoryView(PayloadBytes + ValueSizeByteCount, ValueSize);
+		return MakeMemoryView(ValueBytes + ValueSizeByteCount, ValueSize);
 	}
 	else
 	{
@@ -157,9 +157,9 @@ FUtf8StringView FCbFieldView::AsString(const FUtf8StringView Default)
 {
 	if (FCbFieldType::IsString(TypeWithFlags))
 	{
-		const UTF8CHAR* const PayloadChars = static_cast<const UTF8CHAR*>(Payload);
+		const UTF8CHAR* const ValueChars = static_cast<const UTF8CHAR*>(Value);
 		uint32 ValueSizeByteCount;
-		const uint64 ValueSize = ReadVarUInt(PayloadChars, ValueSizeByteCount);
+		const uint64 ValueSize = ReadVarUInt(ValueChars, ValueSizeByteCount);
 
 		if (ValueSize >= (uint64(1) << 31))
 		{
@@ -168,7 +168,7 @@ FUtf8StringView FCbFieldView::AsString(const FUtf8StringView Default)
 		}
 
 		Error = ECbFieldError::None;
-		return FUtf8StringView(PayloadChars + ValueSizeByteCount, int32(ValueSize));
+		return FUtf8StringView(ValueChars + ValueSizeByteCount, int32(ValueSize));
 	}
 	else
 	{
@@ -186,14 +186,14 @@ uint64 FCbFieldView::AsInteger(const uint64 Default, const FIntegerParams Params
 		const uint64 IsNegative = uint8(TypeWithFlags) & 1;
 
 		uint32 MagnitudeByteCount;
-		const uint64 Magnitude = ReadVarUInt(Payload, MagnitudeByteCount);
-		const uint64 Value = Magnitude ^ -int64(IsNegative);
+		const uint64 Magnitude = ReadVarUInt(Value, MagnitudeByteCount);
+		const uint64 IntValue = Magnitude ^ -int64(IsNegative);
 
 		const uint64 IsInRange = (!(Magnitude & OutOfRangeMask)) & ((!IsNegative) | Params.IsSigned);
 		Error = IsInRange ? ECbFieldError::None : ECbFieldError::RangeError;
 
 		const uint64 UseValueMask = -int64(IsInRange);
-		return (Value & UseValueMask) | (Default & ~UseValueMask);
+		return (IntValue & UseValueMask) | (Default & ~UseValueMask);
 	}
 	else
 	{
@@ -213,7 +213,7 @@ float FCbFieldView::AsFloat(const float Default)
 		constexpr uint64 OutOfRangeMask = ~((uint64(1) << /*FLT_MANT_DIG*/ 24) - 1);
 
 		uint32 MagnitudeByteCount;
-		const int64 Magnitude = ReadVarUInt(Payload, MagnitudeByteCount) + IsNegative;
+		const int64 Magnitude = ReadVarUInt(Value, MagnitudeByteCount) + IsNegative;
 		const uint64 IsInRange = !(Magnitude & OutOfRangeMask);
 		Error = IsInRange ? ECbFieldError::None : ECbFieldError::RangeError;
 		return IsInRange ? float(IsNegative ? -Magnitude : Magnitude) : Default;
@@ -221,8 +221,8 @@ float FCbFieldView::AsFloat(const float Default)
 	case ECbFieldType::Float32:
 	{
 		Error = ECbFieldError::None;
-		const uint32 Value = NETWORK_ORDER32(CompactBinaryPrivate::ReadUnaligned<uint32>(Payload));
-		return reinterpret_cast<const float&>(Value);
+		const uint32 FloatValue = NETWORK_ORDER32(CompactBinaryPrivate::ReadUnaligned<uint32>(Value));
+		return reinterpret_cast<const float&>(FloatValue);
 	}
 	case ECbFieldType::Float64:
 		Error = ECbFieldError::RangeError;
@@ -244,7 +244,7 @@ double FCbFieldView::AsDouble(const double Default)
 		constexpr uint64 OutOfRangeMask = ~((uint64(1) << /*DBL_MANT_DIG*/ 53) - 1);
 
 		uint32 MagnitudeByteCount;
-		const int64 Magnitude = ReadVarUInt(Payload, MagnitudeByteCount) + IsNegative;
+		const int64 Magnitude = ReadVarUInt(Value, MagnitudeByteCount) + IsNegative;
 		const uint64 IsInRange = !(Magnitude & OutOfRangeMask);
 		Error = IsInRange ? ECbFieldError::None : ECbFieldError::RangeError;
 		return IsInRange ? double(IsNegative ? -Magnitude : Magnitude) : Default;
@@ -252,14 +252,14 @@ double FCbFieldView::AsDouble(const double Default)
 	case ECbFieldType::Float32:
 	{
 		Error = ECbFieldError::None;
-		const uint32 Value = NETWORK_ORDER32(CompactBinaryPrivate::ReadUnaligned<uint32>(Payload));
-		return reinterpret_cast<const float&>(Value);
+		const uint32 FloatValue = NETWORK_ORDER32(CompactBinaryPrivate::ReadUnaligned<uint32>(Value));
+		return reinterpret_cast<const float&>(FloatValue);
 	}
 	case ECbFieldType::Float64:
 	{
 		Error = ECbFieldError::None;
-		const uint64 Value = NETWORK_ORDER64(CompactBinaryPrivate::ReadUnaligned<uint64>(Payload));
-		return reinterpret_cast<const double&>(Value);
+		const uint64 FloatValue = NETWORK_ORDER64(CompactBinaryPrivate::ReadUnaligned<uint64>(Value));
+		return reinterpret_cast<const double&>(FloatValue);
 	}
 	default:
 		Error = ECbFieldError::TypeError;
@@ -280,7 +280,7 @@ FIoHash FCbFieldView::AsObjectAttachment(const FIoHash& Default)
 	if (FCbFieldType::IsObjectAttachment(TypeWithFlags))
 	{
 		Error = ECbFieldError::None;
-		return FIoHash(*static_cast<const FIoHash::ByteArray*>(Payload));
+		return FIoHash(*static_cast<const FIoHash::ByteArray*>(Value));
 	}
 	else
 	{
@@ -294,7 +294,7 @@ FIoHash FCbFieldView::AsBinaryAttachment(const FIoHash& Default)
 	if (FCbFieldType::IsBinaryAttachment(TypeWithFlags))
 	{
 		Error = ECbFieldError::None;
-		return FIoHash(*static_cast<const FIoHash::ByteArray*>(Payload));
+		return FIoHash(*static_cast<const FIoHash::ByteArray*>(Value));
 	}
 	else
 	{
@@ -308,7 +308,7 @@ FIoHash FCbFieldView::AsAttachment(const FIoHash& Default)
 	if (FCbFieldType::IsAttachment(TypeWithFlags))
 	{
 		Error = ECbFieldError::None;
-		return FIoHash(*static_cast<const FIoHash::ByteArray*>(Payload));
+		return FIoHash(*static_cast<const FIoHash::ByteArray*>(Value));
 	}
 	else
 	{
@@ -322,7 +322,7 @@ FIoHash FCbFieldView::AsHash(const FIoHash& Default)
 	if (FCbFieldType::IsHash(TypeWithFlags))
 	{
 		Error = ECbFieldError::None;
-		return FIoHash(*static_cast<const FIoHash::ByteArray*>(Payload));
+		return FIoHash(*static_cast<const FIoHash::ByteArray*>(Value));
 	}
 	else
 	{
@@ -341,13 +341,13 @@ FGuid FCbFieldView::AsUuid(const FGuid& Default)
 	if (FCbFieldType::IsUuid(TypeWithFlags))
 	{
 		Error = ECbFieldError::None;
-		FGuid Value;
-		FMemory::Memcpy(&Value, Payload, sizeof(FGuid));
-		Value.A = NETWORK_ORDER32(Value.A);
-		Value.B = NETWORK_ORDER32(Value.B);
-		Value.C = NETWORK_ORDER32(Value.C);
-		Value.D = NETWORK_ORDER32(Value.D);
-		return Value;
+		FGuid UuidValue;
+		FMemory::Memcpy(&UuidValue, Value, sizeof(FGuid));
+		UuidValue.A = NETWORK_ORDER32(UuidValue.A);
+		UuidValue.B = NETWORK_ORDER32(UuidValue.B);
+		UuidValue.C = NETWORK_ORDER32(UuidValue.C);
+		UuidValue.D = NETWORK_ORDER32(UuidValue.D);
+		return UuidValue;
 	}
 	else
 	{
@@ -361,7 +361,7 @@ int64 FCbFieldView::AsDateTimeTicks(const int64 Default)
 	if (FCbFieldType::IsDateTime(TypeWithFlags))
 	{
 		Error = ECbFieldError::None;
-		return NETWORK_ORDER64(CompactBinaryPrivate::ReadUnaligned<int64>(Payload));
+		return NETWORK_ORDER64(CompactBinaryPrivate::ReadUnaligned<int64>(Value));
 	}
 	else
 	{
@@ -385,7 +385,7 @@ int64 FCbFieldView::AsTimeSpanTicks(const int64 Default)
 	if (FCbFieldType::IsTimeSpan(TypeWithFlags))
 	{
 		Error = ECbFieldError::None;
-		return NETWORK_ORDER64(CompactBinaryPrivate::ReadUnaligned<int64>(Payload));
+		return NETWORK_ORDER64(CompactBinaryPrivate::ReadUnaligned<int64>(Value));
 	}
 	else
 	{
@@ -410,7 +410,7 @@ FCbObjectId FCbFieldView::AsObjectId(const FCbObjectId& Default)
 	if (FCbFieldType::IsObjectId(TypeWithFlags))
 	{
 		Error = ECbFieldError::None;
-		return FCbObjectId(MakeMemoryView(Payload, sizeof(FCbObjectId)));
+		return FCbObjectId(MakeMemoryView(Value, sizeof(FCbObjectId)));
 	}
 	else
 	{
@@ -423,17 +423,17 @@ FCbCustomById FCbFieldView::AsCustomById(FCbCustomById Default)
 {
 	if (FCbFieldType::IsCustomById(TypeWithFlags))
 	{
-		const uint8* PayloadBytes = static_cast<const uint8*>(Payload);
-		uint32 PayloadSizeByteCount;
-		const uint64 PayloadSize = ReadVarUInt(PayloadBytes, PayloadSizeByteCount);
-		PayloadBytes += PayloadSizeByteCount;
+		const uint8* ValueBytes = static_cast<const uint8*>(Value);
+		uint32 ValueSizeByteCount;
+		const uint64 ValueSize = ReadVarUInt(ValueBytes, ValueSizeByteCount);
+		ValueBytes += ValueSizeByteCount;
 
-		FCbCustomById Value;
+		FCbCustomById CustomValue;
 		uint32 TypeIdByteCount;
-		Value.Id = ReadVarUInt(PayloadBytes, TypeIdByteCount);
-		Value.Data = MakeMemoryView(PayloadBytes + TypeIdByteCount, PayloadSize - TypeIdByteCount);
+		CustomValue.Id = ReadVarUInt(ValueBytes, TypeIdByteCount);
+		CustomValue.Data = MakeMemoryView(ValueBytes + TypeIdByteCount, ValueSize - TypeIdByteCount);
 		Error = ECbFieldError::None;
-		return Value;
+		return CustomValue;
 	}
 	else
 	{
@@ -446,22 +446,22 @@ FCbCustomByName FCbFieldView::AsCustomByName(FCbCustomByName Default)
 {
 	if (FCbFieldType::IsCustomByName(TypeWithFlags))
 	{
-		const uint8* PayloadBytes = static_cast<const uint8*>(Payload);
-		uint32 PayloadSizeByteCount;
-		const uint64 PayloadSize = ReadVarUInt(PayloadBytes, PayloadSizeByteCount);
-		PayloadBytes += PayloadSizeByteCount;
+		const uint8* ValueBytes = static_cast<const uint8*>(Value);
+		uint32 ValueSizeByteCount;
+		const uint64 ValueSize = ReadVarUInt(ValueBytes, ValueSizeByteCount);
+		ValueBytes += ValueSizeByteCount;
 
 		uint32 TypeNameLenByteCount;
-		const uint64 TypeNameLen = ReadVarUInt(PayloadBytes, TypeNameLenByteCount);
-		PayloadBytes += TypeNameLenByteCount;
+		const uint64 TypeNameLen = ReadVarUInt(ValueBytes, TypeNameLenByteCount);
+		ValueBytes += TypeNameLenByteCount;
 
-		FCbCustomByName Value;
-		Value.Name = FUtf8StringView(
-			reinterpret_cast<const UTF8CHAR*>(PayloadBytes),
+		FCbCustomByName CustomValue;
+		CustomValue.Name = FUtf8StringView(
+			reinterpret_cast<const UTF8CHAR*>(ValueBytes),
 			static_cast<FUtf8StringView::SizeType>(TypeNameLen));
-		Value.Data = MakeMemoryView(PayloadBytes + TypeNameLen, PayloadSize - TypeNameLen - TypeNameLenByteCount);
+		CustomValue.Data = MakeMemoryView(ValueBytes + TypeNameLen, ValueSize - TypeNameLen - TypeNameLenByteCount);
 		Error = ECbFieldError::None;
-		return Value;
+		return CustomValue;
 	}
 	else
 	{
@@ -503,7 +503,7 @@ uint64 FCbFieldView::GetSize() const
 	return sizeof(ECbFieldType) + GetViewNoType().GetSize();
 }
 
-uint64 FCbFieldView::GetPayloadSize() const
+uint64 FCbFieldView::GetValueSize() const
 {
 	switch (GetType())
 	{
@@ -519,13 +519,13 @@ uint64 FCbFieldView::GetPayloadSize() const
 	case ECbFieldType::CustomById:
 	case ECbFieldType::CustomByName:
 	{
-		uint32 PayloadSizeByteCount;
-		const uint64 PayloadSize = ReadVarUInt(Payload, PayloadSizeByteCount);
-		return PayloadSize + PayloadSizeByteCount;
+		uint32 ValueSizeByteCount;
+		const uint64 ValueSize = ReadVarUInt(Value, ValueSizeByteCount);
+		return ValueSize + ValueSizeByteCount;
 	}
 	case ECbFieldType::IntegerPositive:
 	case ECbFieldType::IntegerNegative:
-		return MeasureVarUInt(Payload);
+		return MeasureVarUInt(Value);
 	case ECbFieldType::Float32:
 		return 4;
 	case ECbFieldType::Float64:
@@ -610,15 +610,15 @@ FMemoryView FCbFieldView::GetView() const
 {
 	const uint32 TypeSize = FCbFieldType::HasFieldType(TypeWithFlags) ? sizeof(ECbFieldType) : 0;
 	const uint32 NameSize = FCbFieldType::HasFieldName(TypeWithFlags) ? NameLen + MeasureVarUInt(NameLen) : 0;
-	const uint64 PayloadSize = GetPayloadSize();
-	return MakeMemoryView(static_cast<const uint8*>(Payload) - TypeSize - NameSize, TypeSize + NameSize + PayloadSize);
+	const uint64 ValueSize = GetValueSize();
+	return MakeMemoryView(static_cast<const uint8*>(Value) - TypeSize - NameSize, TypeSize + NameSize + ValueSize);
 }
 
 FMemoryView FCbFieldView::GetViewNoType() const
 {
 	const uint32 NameSize = FCbFieldType::HasFieldName(TypeWithFlags) ? NameLen + MeasureVarUInt(NameLen) : 0;
-	const uint64 PayloadSize = GetPayloadSize();
-	return MakeMemoryView(static_cast<const uint8*>(Payload) - NameSize, NameSize + PayloadSize);
+	const uint64 ValueSize = GetValueSize();
+	return MakeMemoryView(static_cast<const uint8*>(Value) - NameSize, NameSize + ValueSize);
 }
 
 FCbFieldView FCbFieldView::operator[](FUtf8StringView Name) const
@@ -638,18 +638,18 @@ FCbFieldViewIterator FCbFieldView::CreateViewIterator() const
 	const ECbFieldType LocalTypeWithFlags = TypeWithFlags;
 	if (FCbFieldType::HasFields(LocalTypeWithFlags))
 	{
-		const uint8* PayloadBytes = static_cast<const uint8*>(GetPayload());
-		uint32 PayloadSizeByteCount;
-		const uint64 PayloadSize = ReadVarUInt(PayloadBytes, PayloadSizeByteCount);
-		PayloadBytes += PayloadSizeByteCount;
-		const uint64 NumByteCount = FCbFieldType::IsArray(LocalTypeWithFlags) ? MeasureVarUInt(PayloadBytes) : 0;
-		if (PayloadSize > NumByteCount)
+		const uint8* ValueBytes = static_cast<const uint8*>(Value);
+		uint32 ValueSizeByteCount;
+		const uint64 ValueSize = ReadVarUInt(ValueBytes, ValueSizeByteCount);
+		ValueBytes += ValueSizeByteCount;
+		const uint64 NumByteCount = FCbFieldType::IsArray(LocalTypeWithFlags) ? MeasureVarUInt(ValueBytes) : 0;
+		if (ValueSize > NumByteCount)
 		{
-			const void* const PayloadEnd = PayloadBytes + PayloadSize;
-			PayloadBytes += NumByteCount;
+			const void* const ValueEnd = ValueBytes + ValueSize;
+			ValueBytes += NumByteCount;
 			const ECbFieldType UniformType = FCbFieldType::HasUniformFields(LocalTypeWithFlags) ?
-				ECbFieldType(*PayloadBytes++) : ECbFieldType::HasFieldType;
-			return FCbFieldViewIterator::MakeRange(MakeMemoryView(PayloadBytes, PayloadEnd), UniformType);
+				ECbFieldType(*ValueBytes++) : ECbFieldType::HasFieldType;
+			return FCbFieldViewIterator::MakeRange(MakeMemoryView(ValueBytes, ValueEnd), UniformType);
 		}
 	}
 	return FCbFieldViewIterator();
@@ -658,21 +658,21 @@ FCbFieldViewIterator FCbFieldView::CreateViewIterator() const
 ///////////////////////////////////////////////////////////////////////////////////////////////////////////////////////
 
 FCbArrayView::FCbArrayView()
-	: FCbFieldView(CompactBinaryPrivate::GEmptyArrayPayload)
+	: FCbFieldView(CompactBinaryPrivate::GEmptyArrayValue)
 {
 }
 
 uint64 FCbArrayView::Num() const
 {
-	const uint8* PayloadBytes = static_cast<const uint8*>(GetPayload());
-	PayloadBytes += MeasureVarUInt(PayloadBytes);
+	const uint8* ValueBytes = static_cast<const uint8*>(GetValueData());
+	ValueBytes += MeasureVarUInt(ValueBytes);
 	uint32 NumByteCount;
-	return ReadVarUInt(PayloadBytes, NumByteCount);
+	return ReadVarUInt(ValueBytes, NumByteCount);
 }
 
 uint64 FCbArrayView::GetSize() const
 {
-	return sizeof(ECbFieldType) + GetPayloadSize();
+	return sizeof(ECbFieldType) + GetValueSize();
 }
 
 FBlake3Hash FCbArrayView::GetHash() const
@@ -686,18 +686,18 @@ void FCbArrayView::AppendHash(FBlake3& Builder) const
 {
 	const ECbFieldType SerializedType = GetType();
 	Builder.Update(&SerializedType, sizeof(SerializedType));
-	Builder.Update(GetPayloadView());
+	Builder.Update(GetValueView());
 }
 
 bool FCbArrayView::Equals(const FCbArrayView& Other) const
 {
 	return GetType() == Other.GetType() &&
-		GetPayloadView().EqualBytes(Other.GetPayloadView());
+		GetValueView().EqualBytes(Other.GetValueView());
 }
 
 void FCbArrayView::CopyTo(const FMutableMemoryView Buffer) const
 {
-	const FMemoryView Source = GetPayloadView();
+	const FMemoryView Source = GetValueView();
 	checkf(Buffer.GetSize() == sizeof(ECbFieldType) + Source.GetSize(),
 		TEXT("Buffer is %" UINT64_FMT " bytes but %" UINT64_FMT " is required."),
 		Buffer.GetSize(), sizeof(ECbFieldType) + Source.GetSize());
@@ -708,7 +708,7 @@ void FCbArrayView::CopyTo(const FMutableMemoryView Buffer) const
 void FCbArrayView::CopyTo(FArchive& Ar) const
 {
 	check(Ar.IsSaving());
-	const FMemoryView Source = GetPayloadView();
+	const FMemoryView Source = GetValueView();
 	ECbFieldType SerializedType = GetType();
 	Ar.Serialize(&SerializedType, sizeof(SerializedType));
 	Ar.Serialize(const_cast<void*>(Source.GetData()), static_cast<int64>(Source.GetSize()));
@@ -717,7 +717,7 @@ void FCbArrayView::CopyTo(FArchive& Ar) const
 ///////////////////////////////////////////////////////////////////////////////////////////////////////////////////////
 
 FCbObjectView::FCbObjectView()
-	: FCbFieldView(CompactBinaryPrivate::GEmptyObjectPayload)
+	: FCbFieldView(CompactBinaryPrivate::GEmptyObjectValue)
 {
 }
 
@@ -747,12 +747,12 @@ FCbFieldView FCbObjectView::FindViewIgnoreCase(const FUtf8StringView Name) const
 
 FCbObjectView::operator bool() const
 {
-	return GetSize() > sizeof(CompactBinaryPrivate::GEmptyObjectPayload);
+	return GetSize() > sizeof(CompactBinaryPrivate::GEmptyObjectValue);
 }
 
 uint64 FCbObjectView::GetSize() const
 {
-	return sizeof(ECbFieldType) + GetPayloadSize();
+	return sizeof(ECbFieldType) + GetValueSize();
 }
 
 FBlake3Hash FCbObjectView::GetHash() const
@@ -766,18 +766,18 @@ void FCbObjectView::AppendHash(FBlake3& Builder) const
 {
 	const ECbFieldType SerializedType = GetType();
 	Builder.Update(&SerializedType, sizeof(SerializedType));
-	Builder.Update(GetPayloadView());
+	Builder.Update(GetValueView());
 }
 
 bool FCbObjectView::Equals(const FCbObjectView& Other) const
 {
 	return GetType() == Other.GetType() &&
-		GetPayloadView().EqualBytes(Other.GetPayloadView());
+		GetValueView().EqualBytes(Other.GetValueView());
 }
 
 void FCbObjectView::CopyTo(const FMutableMemoryView Buffer) const
 {
-	const FMemoryView Source = GetPayloadView();
+	const FMemoryView Source = GetValueView();
 	checkf(Buffer.GetSize() == sizeof(ECbFieldType) + Source.GetSize(),
 		TEXT("Buffer is %" UINT64_FMT " bytes but %" UINT64_FMT " is required."),
 		Buffer.GetSize(), sizeof(ECbFieldType) + Source.GetSize());
@@ -788,7 +788,7 @@ void FCbObjectView::CopyTo(const FMutableMemoryView Buffer) const
 void FCbObjectView::CopyTo(FArchive& Ar) const
 {
 	check(Ar.IsSaving());
-	const FMemoryView Source = GetPayloadView();
+	const FMemoryView Source = GetValueView();
 	ECbFieldType SerializedType = GetType();
 	Ar.Serialize(&SerializedType, sizeof(SerializedType));
 	Ar.Serialize(const_cast<void*>(Source.GetData()), static_cast<int64>(Source.GetSize()));
@@ -799,10 +799,10 @@ void FCbObjectView::CopyTo(FArchive& Ar) const
 template <typename FieldType>
 TCbFieldIterator<FieldType>& TCbFieldIterator<FieldType>::operator++()
 {
-	const void* const PayloadEnd = FieldType::GetPayloadEnd();
-	const int64 AtEndMask = int64(PayloadEnd == FieldsEnd) - 1;
+	const void* const ValueEnd = FieldType::GetValueEnd();
+	const int64 AtEndMask = int64(ValueEnd == FieldsEnd) - 1;
 	const ECbFieldType NextType = ECbFieldType(int64(FieldType::GetTypeWithFlags()) & AtEndMask);
-	const void* const NextField = reinterpret_cast<const void*>(int64(PayloadEnd) & AtEndMask);
+	const void* const NextField = reinterpret_cast<const void*>(int64(ValueEnd) & AtEndMask);
 	const void* const NextFieldsEnd = reinterpret_cast<const void*>(int64(FieldsEnd) & AtEndMask);
 	FieldType::Assign(NextField, NextType);
 	FieldsEnd = NextFieldsEnd;
@@ -967,7 +967,7 @@ FCompositeBuffer FCbArray::GetBuffer() const
 	else
 	{
 		return FCompositeBuffer(MakeSharedBufferFromArray(TArray<ECbFieldType, TInlineAllocator<1>>{GetType()}),
-			FSharedBuffer::MakeView(GetPayloadView(), GetOuterBuffer()));
+			FSharedBuffer::MakeView(GetValueView(), GetOuterBuffer()));
 	}
 }
 
@@ -981,7 +981,7 @@ FCompositeBuffer FCbObject::GetBuffer() const
 	else
 	{
 		return FCompositeBuffer(MakeSharedBufferFromArray(TArray<ECbFieldType, TInlineAllocator<1>>{GetType()}),
-			FSharedBuffer::MakeView(GetPayloadView(), GetOuterBuffer()));
+			FSharedBuffer::MakeView(GetValueView(), GetOuterBuffer()));
 	}
 }
 
diff --git a/Engine/Source/Runtime/Core/Private/Serialization/CompactBinarySerialization.cpp b/Engine/Source/Runtime/Core/Private/Serialization/CompactBinarySerialization.cpp
--- a/Engine/Source/Runtime/Core/Private/Serialization/CompactBinarySerialization.cpp
+++ b/Engine/Source/Runtime/Core/Private/Serialization/CompactBinarySerialization.cpp
@@ -131,14 +131,14 @@ bool TryMeasureCompactBinary(FMemoryView View, ECbFieldType& OutType, uint64& Ou
 		}
 		else
 		{
-			uint32 PayloadSizeByteCount = MeasureVarUInt(View.GetData());
-			if (View.GetSize() < PayloadSizeByteCount)
+			uint32 ValueSizeByteCount = MeasureVarUInt(View.GetData());
+			if (View.GetSize() < ValueSizeByteCount)
 			{
-				OutSize = Size + PayloadSizeByteCount;
+				OutSize = Size + ValueSizeByteCount;
 				return false;
 			}
-			const uint64 PayloadSize = ReadVarUInt(View.GetData(), PayloadSizeByteCount);
-			OutSize = Size + PayloadSize + PayloadSizeByteCount;
+			const uint64 ValueSize = ReadVarUInt(View.GetData(), ValueSizeByteCount);
+			OutSize = Size + ValueSize + ValueSizeByteCount;
 			return true;
 		}
 	case ECbFieldType::IntegerPositive:
diff --git a/Engine/Source/Runtime/Core/Private/Serialization/CompactBinaryValidation.cpp b/Engine/Source/Runtime/Core/Private/Serialization/CompactBinaryValidation.cpp
--- a/Engine/Source/Runtime/Core/Private/Serialization/CompactBinaryValidation.cpp
+++ b/Engine/Source/Runtime/Core/Private/Serialization/CompactBinaryValidation.cpp
@@ -127,30 +127,30 @@ static void ValidateCbFloat64(FMemoryView& View, ECbValidateMode Mode, ECbValida
 }
 
 /**
- * Validate and read a fixed-size payload from the view.
+ * Validate and read a fixed-size value from the view.
  *
- * Modifies the view to start at the end of the payload, and adds error flags if applicable.
+ * Modifies the view to start at the end of the value, and adds error flags if applicable.
  */
-static FMemoryView ValidateCbFixedPayload(FMemoryView& View, ECbValidateMode Mode, ECbValidateError& Error, uint64 Size)
+static FMemoryView ValidateCbFixedValue(FMemoryView& View, ECbValidateMode Mode, ECbValidateError& Error, uint64 Size)
 {
-	const FMemoryView Payload = View.Left(Size);
+	const FMemoryView Value = View.Left(Size);
 	View += Size;
-	if (Payload.GetSize() < Size)
+	if (Value.GetSize() < Size)
 	{
 		AddError(Error, ECbValidateError::OutOfBounds);
 	}
-	return Payload;
+	return Value;
 };
 
 /**
- * Validate and read a payload from the view where the view begins with the payload size.
+ * Validate and read a value from the view where the view begins with the value size.
  *
- * Modifies the view to start at the end of the payload, and adds error flags if applicable.
+ * Modifies the view to start at the end of the value, and adds error flags if applicable.
  */
-static FMemoryView ValidateCbDynamicPayload(FMemoryView& View, ECbValidateMode Mode, ECbValidateError& Error)
+static FMemoryView ValidateCbDynamicValue(FMemoryView& View, ECbValidateMode Mode, ECbValidateError& Error)
 {
 	const uint64 ValueSize = ValidateCbUInt(View, Mode, Error);
-	return ValidateCbFixedPayload(View, Mode, Error, ValueSize);
+	return ValidateCbFixedValue(View, Mode, Error, ValueSize);
 }
 
 /**
@@ -160,9 +160,8 @@ static FMemoryView ValidateCbDynamicPayload(FMemoryView& View, ECbValidateMode M
  */
 static FUtf8StringView ValidateCbString(FMemoryView& View, ECbValidateMode Mode, ECbValidateError& Error)
 {
-	const FMemoryView Payload = ValidateCbDynamicPayload(View, Mode, Error);
-	const FUtf8StringView Value(static_cast<const UTF8CHAR*>(Payload.GetData()), static_cast<int32>(Payload.GetSize()));
-	return Value;
+	const FMemoryView Value = ValidateCbDynamicValue(View, Mode, Error);
+	return FUtf8StringView(static_cast<const UTF8CHAR*>(Value.GetData()), static_cast<int32>(Value.GetSize()));
 }
 
 static FCbFieldView ValidateCbField(FMemoryView& View, ECbValidateMode Mode, ECbValidateError& Error, ECbFieldType ExternalType);
@@ -329,7 +328,7 @@ static FCbFieldView ValidateCbField(FMemoryView& View, ECbValidateMode Mode, ECb
 		ValidateCbArray(View, Mode, Error, FCbFieldType::GetType(Type));
 		break;
 	case ECbFieldType::Binary:
-		ValidateCbDynamicPayload(View, Mode, Error);
+		ValidateCbDynamicValue(View, Mode, Error);
 		break;
 	case ECbFieldType::String:
 		ValidateCbString(View, Mode, Error);
@@ -341,7 +340,7 @@ static FCbFieldView ValidateCbField(FMemoryView& View, ECbValidateMode Mode, ECb
 		ValidateCbUInt(View, Mode, Error);
 		break;
 	case ECbFieldType::Float32:
-		ValidateCbFixedPayload(View, Mode, Error, 4);
+		ValidateCbFixedValue(View, Mode, Error, 4);
 		break;
 	case ECbFieldType::Float64:
 		ValidateCbFloat64(View, Mode, Error);
@@ -349,27 +348,27 @@ static FCbFieldView ValidateCbField(FMemoryView& View, ECbValidateMode Mode, ECb
 	case ECbFieldType::ObjectAttachment:
 	case ECbFieldType::BinaryAttachment:
 	case ECbFieldType::Hash:
-		ValidateCbFixedPayload(View, Mode, Error, 20);
+		ValidateCbFixedValue(View, Mode, Error, 20);
 		break;
 	case ECbFieldType::Uuid:
-		ValidateCbFixedPayload(View, Mode, Error, 16);
+		ValidateCbFixedValue(View, Mode, Error, 16);
 		break;
 	case ECbFieldType::DateTime:
 	case ECbFieldType::TimeSpan:
-		ValidateCbFixedPayload(View, Mode, Error, 8);
+		ValidateCbFixedValue(View, Mode, Error, 8);
 		break;
 	case ECbFieldType::ObjectId:
-		ValidateCbFixedPayload(View, Mode, Error, 12);
+		ValidateCbFixedValue(View, Mode, Error, 12);
 		break;
 	case ECbFieldType::CustomById:
 	{
-		FMemoryView Value = ValidateCbDynamicPayload(View, Mode, Error);
+		FMemoryView Value = ValidateCbDynamicValue(View, Mode, Error);
 		ValidateCbUInt(Value, Mode, Error);
 		break;
 	}
 	case ECbFieldType::CustomByName:
 	{
-		FMemoryView Value = ValidateCbDynamicPayload(View, Mode, Error);
+		FMemoryView Value = ValidateCbDynamicValue(View, Mode, Error);
 		const FUtf8StringView TypeName = ValidateCbString(Value, Mode, Error);
 		if (TypeName.IsEmpty() && !EnumHasAnyFlags(Error, ECbValidateError::OutOfBounds))
 		{
diff --git a/Engine/Source/Runtime/Core/Private/Serialization/CompactBinaryWriter.cpp b/Engine/Source/Runtime/Core/Private/Serialization/CompactBinaryWriter.cpp
--- a/Engine/Source/Runtime/Core/Private/Serialization/CompactBinaryWriter.cpp
+++ b/Engine/Source/Runtime/Core/Private/Serialization/CompactBinaryWriter.cpp
@@ -50,16 +50,16 @@ static constexpr bool IsUniformType(const ECbFieldType Type)
 	}
 }
 
-/** Append the payload from the compact binary value to the array and return its type. */
+/** Append the value from the compact binary value to the array and return its type. */
 static inline ECbFieldType AppendCompactBinary(const FCbFieldView& Value, TArray64<uint8>& OutData)
 {
 	struct FCopy : public FCbFieldView
 	{
 		using FCbFieldView::GetType;
-		using FCbFieldView::GetPayloadView;
+		using FCbFieldView::GetValueView;
 	};
 	const FCopy& ValueCopy = static_cast<const FCopy&>(Value);
-	const FMemoryView SourceView = ValueCopy.GetPayloadView();
+	const FMemoryView SourceView = ValueCopy.GetValueView();
 	const int64 TargetOffset = OutData.AddUninitialized(SourceView.GetSize());
 	MakeMemoryView(OutData).RightChop(TargetOffset).CopyFrom(SourceView);
 	return ValueCopy.GetType();
@@ -263,27 +263,27 @@ void FCbWriter::EndObject()
 	const uint64 Count = States.Last().Count;
 	States.Pop();
 
-	// Calculate the offset of the payload.
+	// Calculate the offset of the value.
 	const FState& State = States.Last();
-	int64 PayloadOffset = State.Offset + 1;
+	int64 ValueOffset = State.Offset + 1;
 	if ((State.Flags & EStateFlags::Name) == EStateFlags::Name)
 	{
 		uint32 NameLenByteCount;
-		const uint64 NameLen = ReadVarUInt(Data.GetData() + PayloadOffset, NameLenByteCount);
-		PayloadOffset += NameLen + NameLenByteCount;
+		const uint64 NameLen = ReadVarUInt(Data.GetData() + ValueOffset, NameLenByteCount);
+		ValueOffset += NameLen + NameLenByteCount;
 	}
 
 	// Remove redundant field types for uniform objects.
 	if (bUniform && Count > 1)
 	{
-		MakeFieldsUniform(PayloadOffset, Data.Num());
+		MakeFieldsUniform(ValueOffset, Data.Num());
 	}
 
 	// Insert the object size.
-	const uint64 Size = uint64(Data.Num() - PayloadOffset);
+	const uint64 Size = uint64(Data.Num() - ValueOffset);
 	const uint32 SizeByteCount = MeasureVarUInt(Size);
-	Data.InsertUninitialized(PayloadOffset, SizeByteCount);
-	WriteVarUInt(Size, Data.GetData() + PayloadOffset);
+	Data.InsertUninitialized(ValueOffset, SizeByteCount);
+	WriteVarUInt(Size, Data.GetData() + ValueOffset);
 
 	EndField(bUniform ? ECbFieldType::UniformObject : ECbFieldType::Object);
 }
@@ -316,29 +316,29 @@ void FCbWriter::EndArray()
 	const uint64 Count = States.Last().Count;
 	States.Pop();
 
-	// Calculate the offset of the payload.
+	// Calculate the offset of the value.
 	const FState& State = States.Last();
-	int64 PayloadOffset = State.Offset + 1;
+	int64 ValueOffset = State.Offset + 1;
 	if ((State.Flags & EStateFlags::Name) == EStateFlags::Name)
 	{
 		uint32 NameLenByteCount;
-		const uint64 NameLen = ReadVarUInt(Data.GetData() + PayloadOffset, NameLenByteCount);
-		PayloadOffset += NameLen + NameLenByteCount;
+		const uint64 NameLen = ReadVarUInt(Data.GetData() + ValueOffset, NameLenByteCount);
+		ValueOffset += NameLen + NameLenByteCount;
 	}
 
 	// Remove redundant field types for uniform arrays.
 	if (bUniform && Count > 1)
 	{
-		MakeFieldsUniform(PayloadOffset, Data.Num());
+		MakeFieldsUniform(ValueOffset, Data.Num());
 	}
 
 	// Insert the array size and field count.
 	const uint32 CountByteCount = MeasureVarUInt(Count);
-	const uint64 Size = uint64(Data.Num() - PayloadOffset) + CountByteCount;
+	const uint64 Size = uint64(Data.Num() - ValueOffset) + CountByteCount;
 	const uint32 SizeByteCount = MeasureVarUInt(Size);
-	Data.InsertUninitialized(PayloadOffset, SizeByteCount + CountByteCount);
-	WriteVarUInt(Size, Data.GetData() + PayloadOffset);
-	WriteVarUInt(Count, Data.GetData() + PayloadOffset + SizeByteCount);
+	Data.InsertUninitialized(ValueOffset, SizeByteCount + CountByteCount);
+	WriteVarUInt(Size, Data.GetData() + ValueOffset);
+	WriteVarUInt(Count, Data.GetData() + ValueOffset + SizeByteCount);
 
 	EndField(bUniform ? ECbFieldType::UniformArray : ECbFieldType::Array);
 }
@@ -581,9 +581,9 @@ void FCbWriter::AddCustom(const uint64 TypeId, const FMemoryView Value)
 	const uint32 TypeByteCount = MeasureVarUInt(TypeId);
 	const uint64 Size = TypeByteCount + Value.GetSize();
 	const uint32 SizeByteCount = MeasureVarUInt(Size);
-	const int64 PayloadOffset = Data.AddUninitialized(SizeByteCount + TypeByteCount);
-	WriteVarUInt(Size, Data.GetData() + PayloadOffset);
-	WriteVarUInt(TypeId, Data.GetData() + PayloadOffset + SizeByteCount);
+	const int64 ValueOffset = Data.AddUninitialized(SizeByteCount + TypeByteCount);
+	WriteVarUInt(Size, Data.GetData() + ValueOffset);
+	WriteVarUInt(TypeId, Data.GetData() + ValueOffset + SizeByteCount);
 	Data.Append(static_cast<const uint8*>(Value.GetData()), Value.GetSize());
 	EndField(ECbFieldType::CustomById);
 }
@@ -601,9 +601,9 @@ void FCbWriter::AddCustom(const FUtf8StringView TypeName, const FMemoryView Valu
 	const uint64 Size = TypeNameLenByteCount + TypeNameLen + Value.GetSize();
 	const uint32 SizeByteCount = MeasureVarUInt(Size);
 
-	const int64 PayloadOffset = Data.AddUninitialized(SizeByteCount + TypeNameLenByteCount);
-	WriteVarUInt(Size, Data.GetData() + PayloadOffset);
-	WriteVarUInt(TypeNameLen, Data.GetData() + PayloadOffset + SizeByteCount);
+	const int64 ValueOffset = Data.AddUninitialized(SizeByteCount + TypeNameLenByteCount);
+	WriteVarUInt(Size, Data.GetData() + ValueOffset);
+	WriteVarUInt(TypeNameLen, Data.GetData() + ValueOffset + SizeByteCount);
 	Data.Append(reinterpret_cast<const uint8*>(TypeName.GetData()), TypeNameLen);
 	Data.Append(static_cast<const uint8*>(Value.GetData()), Value.GetSize());
 
diff --git a/Engine/Source/Runtime/Core/Private/Tests/Serialization/CompactBinarySerializationTest.cpp b/Engine/Source/Runtime/Core/Private/Tests/Serialization/CompactBinarySerializationTest.cpp
--- a/Engine/Source/Runtime/Core/Private/Tests/Serialization/CompactBinarySerializationTest.cpp
+++ b/Engine/Source/Runtime/Core/Private/Tests/Serialization/CompactBinarySerializationTest.cpp
@@ -174,73 +174,73 @@ bool FCbSaveTest::RunTest(const FString& Parameters)
 
 	// Field
 	{
-		const uint8 Payload[] = { uint8(ECbFieldType::IntegerPositive), 42 };
-		TestSave(TEXT("Field"), FCbField(FCbFieldView(Payload), FSharedBuffer()), MakeMemoryView(Payload));
+		const uint8 Value[] = { uint8(ECbFieldType::IntegerPositive), 42 };
+		TestSave(TEXT("Field"), FCbField(FCbFieldView(Value), FSharedBuffer()), MakeMemoryView(Value));
 	}
 	{
-		const uint8 Payload[] = { uint8(ECbFieldType::IntegerPositive | ECbFieldType::HasFieldName), 1, 'A', 42 };
-		TestSave(TEXT("Field, Name"), FCbField(FCbFieldView(Payload), FSharedBuffer()), MakeMemoryView(Payload));
+		const uint8 Value[] = { uint8(ECbFieldType::IntegerPositive | ECbFieldType::HasFieldName), 1, 'A', 42 };
+		TestSave(TEXT("Field, Name"), FCbField(FCbFieldView(Value), FSharedBuffer()), MakeMemoryView(Value));
 	}
 	{
-		const uint8 Payload[] = { 42 };
-		TestSave(TEXT("Field, NoType"), FCbField(FCbFieldView(Payload, ECbFieldType::IntegerPositive), FSharedBuffer()),
+		const uint8 Value[] = { 42 };
+		TestSave(TEXT("Field, NoType"), FCbField(FCbFieldView(Value, ECbFieldType::IntegerPositive), FSharedBuffer()),
 			MakeMemoryView<uint8>({ uint8(ECbFieldType::IntegerPositive), 42 }));
 	}
 	{
-		const uint8 Payload[] = { 1, 'I', 42 };
-		TestSave(TEXT("Field, NoType, Name"), FCbField(FCbFieldView(Payload, ECbFieldType::IntegerPositive | ECbFieldType::HasFieldName), FSharedBuffer()),
+		const uint8 Value[] = { 1, 'I', 42 };
+		TestSave(TEXT("Field, NoType, Name"), FCbField(FCbFieldView(Value, ECbFieldType::IntegerPositive | ECbFieldType::HasFieldName), FSharedBuffer()),
 			MakeMemoryView<uint8>({ uint8(ECbFieldType::IntegerPositive | ECbFieldType::HasFieldName), 1, 'I', 42 }));
 	}
 
 	// Array
 	{
 		constexpr uint8 IntType = uint8(ECbFieldType::IntegerPositive);
-		const uint8 Payload[] = { uint8(ECbFieldType::UniformArray), 5, 3, IntType, 1, 2, 3 };
-		TestSave(TEXT("Array"), FCbArray(FCbArrayView(Payload), FSharedBuffer()),
-			MakeMemoryView(Payload));
+		const uint8 Value[] = { uint8(ECbFieldType::UniformArray), 5, 3, IntType, 1, 2, 3 };
+		TestSave(TEXT("Array"), FCbArray(FCbArrayView(Value), FSharedBuffer()),
+			MakeMemoryView(Value));
 	}
 	{
 		constexpr uint8 IntType = uint8(ECbFieldType::IntegerPositive);
-		const uint8 Payload[] = { uint8(ECbFieldType::UniformArray | ECbFieldType::HasFieldName), 1, 'A', 5, 3, IntType, 1, 2, 3 };
-		TestSave(TEXT("Array, Name"), FCbArray(FCbArrayView(Payload), FSharedBuffer()),
+		const uint8 Value[] = { uint8(ECbFieldType::UniformArray | ECbFieldType::HasFieldName), 1, 'A', 5, 3, IntType, 1, 2, 3 };
+		TestSave(TEXT("Array, Name"), FCbArray(FCbArrayView(Value), FSharedBuffer()),
 			MakeMemoryView<uint8>({ uint8(ECbFieldType::UniformArray), 5, 3, IntType, 1, 2, 3 }));
 	}
 	{
 		constexpr uint8 IntType = uint8(ECbFieldType::IntegerPositive);
-		const uint8 Payload[] = { 5, 3, IntType, 1, 2, 3 };
-		TestSave(TEXT("Array, NoType"), FCbArray(FCbArrayView(Payload, ECbFieldType::UniformArray), FSharedBuffer()),
+		const uint8 Value[] = { 5, 3, IntType, 1, 2, 3 };
+		TestSave(TEXT("Array, NoType"), FCbArray(FCbArrayView(Value, ECbFieldType::UniformArray), FSharedBuffer()),
 			MakeMemoryView<uint8>({ uint8(ECbFieldType::UniformArray), 5, 3, IntType, 1, 2, 3 }));
 	}
 	{
 		constexpr uint8 IntType = uint8(ECbFieldType::IntegerPositive);
-		const uint8 Payload[] = { 1, 'A', 5, 3, IntType, 1, 2, 3 };
-		TestSave(TEXT("Array, NoType, Name"), FCbArray(FCbArrayView(Payload, ECbFieldType::UniformArray | ECbFieldType::HasFieldName), FSharedBuffer()),
+		const uint8 Value[] = { 1, 'A', 5, 3, IntType, 1, 2, 3 };
+		TestSave(TEXT("Array, NoType, Name"), FCbArray(FCbArrayView(Value, ECbFieldType::UniformArray | ECbFieldType::HasFieldName), FSharedBuffer()),
 			MakeMemoryView<uint8>({ uint8(ECbFieldType::UniformArray), 5, 3, IntType, 1, 2, 3 }));
 	}
 
 	// Object
 	{
 		constexpr uint8 IntType = uint8(ECbFieldType::HasFieldName | ECbFieldType::IntegerPositive);
-		const uint8 Payload[] = { uint8(ECbFieldType::UniformObject), 10, IntType, 1, 'A', 1, 1, 'B', 2, 1, 'C', 3 };
-		TestSave(TEXT("Object"), FCbObject(FCbObjectView(Payload), FSharedBuffer()),
-			MakeMemoryView(Payload));
+		const uint8 Value[] = { uint8(ECbFieldType::UniformObject), 10, IntType, 1, 'A', 1, 1, 'B', 2, 1, 'C', 3 };
+		TestSave(TEXT("Object"), FCbObject(FCbObjectView(Value), FSharedBuffer()),
+			MakeMemoryView(Value));
 	}
 	{
 		constexpr uint8 IntType = uint8(ECbFieldType::HasFieldName | ECbFieldType::IntegerPositive);
-		const uint8 Payload[] = { uint8(ECbFieldType::UniformObject | ECbFieldType::HasFieldName), 1, 'O', 10, IntType, 1, 'A', 1, 1, 'B', 2, 1, 'C', 3 };
-		TestSave(TEXT("Object, Name"), FCbObject(FCbObjectView(Payload), FSharedBuffer()),
+		const uint8 Value[] = { uint8(ECbFieldType::UniformObject | ECbFieldType::HasFieldName), 1, 'O', 10, IntType, 1, 'A', 1, 1, 'B', 2, 1, 'C', 3 };
+		TestSave(TEXT("Object, Name"), FCbObject(FCbObjectView(Value), FSharedBuffer()),
 			MakeMemoryView<uint8>({ uint8(ECbFieldType::UniformObject), 10, IntType, 1, 'A', 1, 1, 'B', 2, 1, 'C', 3 }));
 	}
 	{
 		constexpr uint8 IntType = uint8(ECbFieldType::HasFieldName | ECbFieldType::IntegerPositive);
-		const uint8 Payload[] = { 10, IntType, 1, 'A', 1, 1, 'B', 2, 1, 'C', 3 };
-		TestSave(TEXT("Object, NoType"), FCbObject(FCbObjectView(Payload, ECbFieldType::UniformObject), FSharedBuffer()),
+		const uint8 Value[] = { 10, IntType, 1, 'A', 1, 1, 'B', 2, 1, 'C', 3 };
+		TestSave(TEXT("Object, NoType"), FCbObject(FCbObjectView(Value, ECbFieldType::UniformObject), FSharedBuffer()),
 			MakeMemoryView<uint8>({ uint8(ECbFieldType::UniformObject), 10, IntType, 1, 'A', 1, 1, 'B', 2, 1, 'C', 3 }));
 	}
 	{
 		constexpr uint8 IntType = uint8(ECbFieldType::HasFieldName | ECbFieldType::IntegerPositive);
-		const uint8 Payload[] = { 1, 'O', 10, IntType, 1, 'A', 1, 1, 'B', 2, 1, 'C', 3 };
-		TestSave(TEXT("Object, NoType, Name"), FCbObject(FCbObjectView(Payload, ECbFieldType::UniformObject | ECbFieldType::HasFieldName), FSharedBuffer()),
+		const uint8 Value[] = { 1, 'O', 10, IntType, 1, 'A', 1, 1, 'B', 2, 1, 'C', 3 };
+		TestSave(TEXT("Object, NoType, Name"), FCbObject(FCbObjectView(Value, ECbFieldType::UniformObject | ECbFieldType::HasFieldName), FSharedBuffer()),
 			MakeMemoryView<uint8>({ uint8(ECbFieldType::UniformObject), 10, IntType, 1, 'A', 1, 1, 'B', 2, 1, 'C', 3 }));
 	}
 
diff --git a/Engine/Source/Runtime/Core/Private/Tests/Serialization/CompactBinaryTest.cpp b/Engine/Source/Runtime/Core/Private/Tests/Serialization/CompactBinaryTest.cpp
--- a/Engine/Source/Runtime/Core/Private/Tests/Serialization/CompactBinaryTest.cpp
+++ b/Engine/Source/Runtime/Core/Private/Tests/Serialization/CompactBinaryTest.cpp
@@ -125,9 +125,9 @@ protected:
 	}
 
 	template <ECbFieldType FieldType, typename T = TCbFieldValueType<FieldType>, typename AccessorsType = TCbFieldTypeAccessors<FieldType>>
-	void TestFieldNoClone(const TCHAR* What, TArrayView<const uint8> Payload, T ExpectedValue = T(), T DefaultValue = T(), ECbFieldError ExpectedError = ECbFieldError::None, const AccessorsType& Accessors = AccessorsType())
+	void TestFieldNoClone(const TCHAR* What, TArrayView<const uint8> Value, T ExpectedValue = T(), T DefaultValue = T(), ECbFieldError ExpectedError = ECbFieldError::None, const AccessorsType& Accessors = AccessorsType())
 	{
-		FCbFieldView Field(Payload.GetData(), FieldType);
+		FCbFieldView Field(Value.GetData(), FieldType);
 		TestFieldNoClone<FieldType>(What, Field, ExpectedValue, DefaultValue, ExpectedError, Accessors);
 	}
 
@@ -141,10 +141,10 @@ protected:
 	}
 
 	template <ECbFieldType FieldType, typename T = TCbFieldValueType<FieldType>, typename AccessorsType = TCbFieldTypeAccessors<FieldType>>
-	void TestField(const TCHAR* What, TArrayView<const uint8> Payload, T ExpectedValue = T(), T DefaultValue = T(), ECbFieldError ExpectedError = ECbFieldError::None, const AccessorsType& Accessors = AccessorsType())
+	void TestField(const TCHAR* What, TArrayView<const uint8> Value, T ExpectedValue = T(), T DefaultValue = T(), ECbFieldError ExpectedError = ECbFieldError::None, const AccessorsType& Accessors = AccessorsType())
 	{
-		FCbFieldView Field(Payload.GetData(), FieldType);
-		TestEqual(FString::Printf(TEXT("FCbFieldView::GetSize(%s)"), What), Field.GetSize(), uint64(Payload.Num() + !FCbFieldType::HasFieldType(FieldType)));
+		FCbFieldView Field(Value.GetData(), FieldType);
+		TestEqual(FString::Printf(TEXT("FCbFieldView::GetSize(%s)"), What), Field.GetSize(), uint64(Value.Num() + !FCbFieldType::HasFieldType(FieldType)));
 		TestTrue(FString::Printf(TEXT("FCbFieldView::HasValue(%s)"), What), Field.HasValue());
 		TestFalse(FString::Printf(TEXT("FCbFieldView::HasError(%s) == false"), What), Field.HasError());
 		TestEqual(FString::Printf(TEXT("FCbFieldView::GetError(%s) == None"), What), Field.GetError(), ECbFieldError::None);
@@ -158,9 +158,9 @@ protected:
 	}
 
 	template <ECbFieldType FieldType, typename T = TCbFieldValueType<FieldType>, typename AccessorsType = TCbFieldTypeAccessors<FieldType>>
-	void TestFieldError(const TCHAR* What, TArrayView<const uint8> Payload, ECbFieldError ExpectedError, T ExpectedValue = T(), const AccessorsType& Accessors = AccessorsType())
+	void TestFieldError(const TCHAR* What, TArrayView<const uint8> Value, ECbFieldError ExpectedError, T ExpectedValue = T(), const AccessorsType& Accessors = AccessorsType())
 	{
-		FCbFieldView Field(Payload.GetData(), FieldType);
+		FCbFieldView Field(Value.GetData(), FieldType);
 		TestFieldError<FieldType>(What, Field, ExpectedError, ExpectedValue, Accessors);
 	}
 };
@@ -300,9 +300,9 @@ bool FCbFieldObjectTest::RunTest(const FString& Parameters)
 	static_assert(!std::is_convertible<FCbFieldView, FCbObjectView>::value, "Invalid conversion to FCbObjectView");
 	static_assert(!std::is_assignable<FCbObjectView, const FCbFieldView&>::value, "Invalid assignment for FCbObjectView");
 
-	auto TestIntObject = [this](const FCbObjectView& Object, int32 ExpectedNum, uint64 ExpectedPayloadSize)
+	auto TestIntObject = [this](const FCbObjectView& Object, int32 ExpectedNum, uint64 ExpectedValueSize)
 	{
-		TestEqual(TEXT("FCbFieldView(Object)::AsObjectView().GetSize()"), Object.GetSize(), ExpectedPayloadSize + sizeof(ECbFieldType));
+		TestEqual(TEXT("FCbFieldView(Object)::AsObjectView().GetSize()"), Object.GetSize(), ExpectedValueSize + sizeof(ECbFieldType));
 
 		int32 ActualNum = 0;
 		for (FCbFieldViewIterator It = Object.CreateViewIterator(); It; ++It)
@@ -361,12 +361,12 @@ bool FCbFieldObjectTest::RunTest(const FString& Parameters)
 	// Test FCbFieldView(Object, NotEmpty)
 	{
 		constexpr uint8 IntType = uint8(ECbFieldType::HasFieldName | ECbFieldType::IntegerPositive);
-		const uint8 Payload[] = { 12, IntType, 1, 'A', 1, IntType, 1, 'B', 2, IntType, 1, 'C', 3 };
-		FCbFieldView Field(Payload, ECbFieldType::Object);
-		TestField<ECbFieldType::Object>(TEXT("Object, NotEmpty"), Field, FCbObjectView(Payload, ECbFieldType::Object));
+		const uint8 Value[] = { 12, IntType, 1, 'A', 1, IntType, 1, 'B', 2, IntType, 1, 'C', 3 };
+		FCbFieldView Field(Value, ECbFieldType::Object);
+		TestField<ECbFieldType::Object>(TEXT("Object, NotEmpty"), Field, FCbObjectView(Value, ECbFieldType::Object));
 		FCbObject Object = FCbObject::Clone(Field.AsObjectView());
-		TestIntObject(Object, 3, sizeof(Payload));
-		TestIntObject(Field.AsObjectView(), 3, sizeof(Payload));
+		TestIntObject(Object, 3, sizeof(Value));
+		TestIntObject(Field.AsObjectView(), 3, sizeof(Value));
 		TestTrue(TEXT("FCbObjectView::Equals()"), Object.Equals(Field.AsObjectView()));
 		TestEqual(TEXT("FCbObjectView::Find()"), Object.FindView("B"_ASV).AsInt32(), 2);
 		TestEqual(TEXT("FCbObjectView::Find()"), Object.FindView("b"_ASV).AsInt32(4), 4);
@@ -379,12 +379,12 @@ bool FCbFieldObjectTest::RunTest(const FString& Parameters)
 	// Test FCbFieldView(UniformObject, NotEmpty)
 	{
 		constexpr uint8 IntType = uint8(ECbFieldType::HasFieldName | ECbFieldType::IntegerPositive);
-		const uint8 Payload[] = { 10, IntType, 1, 'A', 1, 1, 'B', 2, 1, 'C', 3 };
-		FCbFieldView Field(Payload, ECbFieldType::UniformObject);
-		TestField<ECbFieldType::UniformObject>(TEXT("UniformObject, NotEmpty"), Field, FCbObjectView(Payload, ECbFieldType::UniformObject));
+		const uint8 Value[] = { 10, IntType, 1, 'A', 1, 1, 'B', 2, 1, 'C', 3 };
+		FCbFieldView Field(Value, ECbFieldType::UniformObject);
+		TestField<ECbFieldType::UniformObject>(TEXT("UniformObject, NotEmpty"), Field, FCbObjectView(Value, ECbFieldType::UniformObject));
 		FCbObject Object = FCbObject::Clone(Field.AsObjectView());
-		TestIntObject(Object, 3, sizeof(Payload));
-		TestIntObject(Field.AsObjectView(), 3, sizeof(Payload));
+		TestIntObject(Object, 3, sizeof(Value));
+		TestIntObject(Field.AsObjectView(), 3, sizeof(Value));
 		TestTrue(TEXT("FCbObjectView(Uniform)::Equals()"), Object.Equals(Field.AsObjectView()));
 		TestEqual(TEXT("FCbObjectView(Uniform)::Find()"), Object.FindView("B"_ASV).AsInt32(), 2);
 		TestEqual(TEXT("FCbObjectView(Uniform)::Find()"), Object.Find("B"_ASV).AsInt32(), 2);
@@ -400,18 +400,18 @@ bool FCbFieldObjectTest::RunTest(const FString& Parameters)
 		TestTrue(TEXT("FCbObject::AsField()"), Object.GetOuterBuffer() == Object.AsField().AsObject().GetOuterBuffer());
 
 		// Equals
-		const uint8 NamedPayload[] = { 1, 'O', 10, IntType, 1, 'A', 1, 1, 'B', 2, 1, 'C', 3 };
-		FCbFieldView NamedField(NamedPayload, ECbFieldType::UniformObject | ECbFieldType::HasFieldName);
+		const uint8 NamedValue[] = { 1, 'O', 10, IntType, 1, 'A', 1, 1, 'B', 2, 1, 'C', 3 };
+		FCbFieldView NamedField(NamedValue, ECbFieldType::UniformObject | ECbFieldType::HasFieldName);
 		TestTrue(TEXT("FCbObjectView::Equals()"), Field.AsObjectView().Equals(NamedField.AsObjectView()));
 		TestTrue(TEXT("FCbObjectView::AsFieldView().Equals()"), Field.Equals(Field.AsObjectView().AsFieldView()));
 		TestTrue(TEXT("FCbObjectView::AsFieldView().Equals()"), NamedField.Equals(NamedField.AsObjectView().AsFieldView()));
 
 		// CopyTo
-		uint8 CopyBytes[sizeof(Payload) + 1];
+		uint8 CopyBytes[sizeof(Value) + 1];
 		Field.AsObjectView().CopyTo(MakeMemoryView(CopyBytes));
-		TestTrue(TEXT("FCbObjectView(NoType)::CopyTo()"), MakeMemoryView(Payload).EqualBytes(MakeMemoryView(CopyBytes) + 1));
+		TestTrue(TEXT("FCbObjectView(NoType)::CopyTo()"), MakeMemoryView(Value).EqualBytes(MakeMemoryView(CopyBytes) + 1));
 		NamedField.AsObjectView().CopyTo(MakeMemoryView(CopyBytes));
-		TestTrue(TEXT("FCbObjectView(NoType, Name)::CopyTo()"), MakeMemoryView(Payload).EqualBytes(MakeMemoryView(CopyBytes) + 1));
+		TestTrue(TEXT("FCbObjectView(NoType, Name)::CopyTo()"), MakeMemoryView(Value).EqualBytes(MakeMemoryView(CopyBytes) + 1));
 
 		// TryGetView
 		FMemoryView View;
@@ -518,9 +518,9 @@ bool FCbFieldArrayTest::RunTest(const FString& Parameters)
 	static_assert(!std::is_convertible<FCbFieldView, FCbArrayView>::value, "Invalid conversion to FCbArrayView");
 	static_assert(!std::is_assignable<FCbArrayView, const FCbFieldView&>::value, "Invalid assignment for FCbArrayView");
 
-	auto TestIntArray = [this](FCbArrayView Array, int32 ExpectedNum, uint64 ExpectedPayloadSize)
+	auto TestIntArray = [this](FCbArrayView Array, int32 ExpectedNum, uint64 ExpectedValueSize)
 	{
-		TestEqual(TEXT("FCbFieldView(Array)::AsArrayView().GetSize()"), Array.GetSize(), ExpectedPayloadSize + sizeof(ECbFieldType));
+		TestEqual(TEXT("FCbFieldView(Array)::AsArrayView().GetSize()"), Array.GetSize(), ExpectedValueSize + sizeof(ECbFieldType));
 		TestEqual(TEXT("FCbFieldView(Array)::AsArrayView().Num()"), Array.Num(), uint64(ExpectedNum));
 
 		int32 ActualNum = 0;
@@ -572,41 +572,41 @@ bool FCbFieldArrayTest::RunTest(const FString& Parameters)
 	// Test FCbFieldView(Array, NotEmpty)
 	{
 		constexpr uint8 IntType = uint8(ECbFieldType::IntegerPositive);
-		const uint8 Payload[] = { 7, 3, IntType, 1, IntType, 2, IntType, 3 };
-		FCbFieldView Field(Payload, ECbFieldType::Array);
-		TestField<ECbFieldType::Array>(TEXT("Array, NotEmpty"), Field, FCbArrayView(Payload, ECbFieldType::Array));
+		const uint8 Value[] = { 7, 3, IntType, 1, IntType, 2, IntType, 3 };
+		FCbFieldView Field(Value, ECbFieldType::Array);
+		TestField<ECbFieldType::Array>(TEXT("Array, NotEmpty"), Field, FCbArrayView(Value, ECbFieldType::Array));
 		FCbArray Array = FCbArray::Clone(Field.AsArrayView());
-		TestIntArray(Array, 3, sizeof(Payload));
-		TestIntArray(Field.AsArrayView(), 3, sizeof(Payload));
+		TestIntArray(Array, 3, sizeof(Value));
+		TestIntArray(Field.AsArrayView(), 3, sizeof(Value));
 		TestTrue(TEXT("FCbArrayView::Equals()"), Array.Equals(Field.AsArrayView()));
 	}
 
 	// Test FCbFieldView(UniformArray)
 	{
 		constexpr uint8 IntType = uint8(ECbFieldType::IntegerPositive);
-		const uint8 Payload[] = { 5, 3, IntType, 1, 2, 3 };
-		FCbFieldView Field(Payload, ECbFieldType::UniformArray);
-		TestField<ECbFieldType::UniformArray>(TEXT("UniformArray"), Field, FCbArrayView(Payload, ECbFieldType::UniformArray));
+		const uint8 Value[] = { 5, 3, IntType, 1, 2, 3 };
+		FCbFieldView Field(Value, ECbFieldType::UniformArray);
+		TestField<ECbFieldType::UniformArray>(TEXT("UniformArray"), Field, FCbArrayView(Value, ECbFieldType::UniformArray));
 		FCbArray Array = FCbArray::Clone(Field.AsArrayView());
-		TestIntArray(Array, 3, sizeof(Payload));
-		TestIntArray(Field.AsArrayView(), 3, sizeof(Payload));
+		TestIntArray(Array, 3, sizeof(Value));
+		TestIntArray(Field.AsArrayView(), 3, sizeof(Value));
 		TestTrue(TEXT("FCbArrayView(Uniform)::Equals()"), Array.Equals(Field.AsArrayView()));
 
 		TestTrue(TEXT("FCbArray::AsField()"), Array.GetOuterBuffer() == Array.AsField().AsArray().GetOuterBuffer());
 
 		// Equals
-		const uint8 NamedPayload[] = { 1, 'A', 5, 3, IntType, 1, 2, 3 };
-		FCbFieldView NamedField(NamedPayload, ECbFieldType::UniformArray | ECbFieldType::HasFieldName);
+		const uint8 NamedValue[] = { 1, 'A', 5, 3, IntType, 1, 2, 3 };
+		FCbFieldView NamedField(NamedValue, ECbFieldType::UniformArray | ECbFieldType::HasFieldName);
 		TestTrue(TEXT("FCbArrayView::Equals()"), Field.AsArrayView().Equals(NamedField.AsArrayView()));
 		TestTrue(TEXT("FCbArrayView::AsFieldView().Equals()"), Field.Equals(Field.AsArrayView().AsFieldView()));
 		TestTrue(TEXT("FCbArrayView::AsFieldView().Equals()"), NamedField.Equals(NamedField.AsArrayView().AsFieldView()));
 
 		// CopyTo
-		uint8 CopyBytes[sizeof(Payload) + 1];
+		uint8 CopyBytes[sizeof(Value) + 1];
 		Field.AsArrayView().CopyTo(MakeMemoryView(CopyBytes));
-		TestTrue(TEXT("FCbArrayView(NoType)::CopyTo()"), MakeMemoryView(Payload).EqualBytes(MakeMemoryView(CopyBytes) + 1));
+		TestTrue(TEXT("FCbArrayView(NoType)::CopyTo()"), MakeMemoryView(Value).EqualBytes(MakeMemoryView(CopyBytes) + 1));
 		NamedField.AsArrayView().CopyTo(MakeMemoryView(CopyBytes));
-		TestTrue(TEXT("FCbArrayView(NoType, Name)::CopyTo()"), MakeMemoryView(Payload).EqualBytes(MakeMemoryView(CopyBytes) + 1));
+		TestTrue(TEXT("FCbArrayView(NoType, Name)::CopyTo()"), MakeMemoryView(Value).EqualBytes(MakeMemoryView(CopyBytes) + 1));
 
 		// TryGetView
 		FMemoryView View;
@@ -715,9 +715,9 @@ bool FCbFieldBinaryTest::RunTest(const FString& Parameters)
 
 	// Test FCbFieldView(Binary, Value)
 	{
-		const uint8 Payload[] = { 3, 4, 5, 6 }; // Size: 3, Data: 4/5/6
-		FCbFieldView FieldView(Payload, ECbFieldType::Binary);
-		TestFieldNoClone<ECbFieldType::Binary>(TEXT("Binary, Value, View"), FieldView, MakeMemoryView(Payload + 1, 3));
+		const uint8 Value[] = { 3, 4, 5, 6 }; // Size: 3, Data: 4/5/6
+		FCbFieldView FieldView(Value, ECbFieldType::Binary);
+		TestFieldNoClone<ECbFieldType::Binary>(TEXT("Binary, Value, View"), FieldView, MakeMemoryView(Value + 1, 3));
 
 		FCbField Field = FCbField::Clone(FieldView);
 		Field.AsBinary();
@@ -747,15 +747,15 @@ bool FCbFieldStringTest::RunTest(const FString& Parameters)
 
 	// Test FCbFieldView(String, Value)
 	{
-		const uint8 Payload[] = { 3, 'A', 'B', 'C' }; // Size: 3, Data: ABC
-		TestField<ECbFieldType::String>(TEXT("String, Value"), Payload, FUtf8StringView(reinterpret_cast<const UTF8CHAR*>(Payload) + 1, 3));
+		const uint8 Value[] = { 3, 'A', 'B', 'C' }; // Size: 3, Data: ABC
+		TestField<ECbFieldType::String>(TEXT("String, Value"), Value, FUtf8StringView(reinterpret_cast<const UTF8CHAR*>(Value) + 1, 3));
 	}
 
 	// Test FCbFieldView(String, OutOfRangeSize)
 	{
-		uint8 Payload[9];
-		WriteVarUInt(uint64(1) << 31, Payload);
-		TestFieldError<ECbFieldType::String>(TEXT("String, OutOfRangeSize"), Payload, ECbFieldError::RangeError, "ABC"_U8SV);
+		uint8 Value[9];
+		WriteVarUInt(uint64(1) << 31, Value);
+		TestFieldError<ECbFieldType::String>(TEXT("String, OutOfRangeSize"), Value, ECbFieldError::RangeError, "ABC"_U8SV);
 	}
 
 	// Test FCbFieldView(None) as String
@@ -808,12 +808,12 @@ protected:
 
 	void TestIntegerField(ECbFieldType FieldType, EIntType ExpectedMask, uint64 Magnitude)
 	{
-		uint8 Payload[9];
+		uint8 Value[9];
 		const bool Negative = bool(uint8(FieldType) & 1);
-		WriteVarUInt(Magnitude - Negative, Payload);
+		WriteVarUInt(Magnitude - Negative, Value);
 		constexpr uint64 DefaultValue = 8;
 		const uint64 ExpectedValue = Negative ? uint64(-int64(Magnitude)) : Magnitude;
-		FCbFieldView Field(Payload, FieldType);
+		FCbFieldView Field(Value, FieldType);
 		TestField<ECbFieldType::IntegerNegative>(TEXT("Int8"), Field, int8(EnumHasAnyFlags(ExpectedMask, EIntType::Int8) ? ExpectedValue : DefaultValue),
 			int8(DefaultValue), EnumHasAnyFlags(ExpectedMask, EIntType::Int8) ? ECbFieldError::None : ECbFieldError::RangeError, TCbIntegerAccessors<int8, &FCbFieldView::AsInt8>());
 		TestField<ECbFieldType::IntegerNegative>(TEXT("Int16"), Field, int16(EnumHasAnyFlags(ExpectedMask, EIntType::Int16) ? ExpectedValue : DefaultValue),
@@ -881,99 +881,99 @@ bool FCbFieldFloatTest::RunTest(const FString& Parameters)
 {
 	// Test FCbFieldView(Float, 32-bit)
 	{
-		const uint8 Payload[] = { 0xc0, 0x12, 0x34, 0x56 }; // -2.28444433f
-		TestField<ECbFieldType::Float32>(TEXT("Float32"), Payload, -2.28444433f);
+		const uint8 Value[] = { 0xc0, 0x12, 0x34, 0x56 }; // -2.28444433f
+		TestField<ECbFieldType::Float32>(TEXT("Float32"), Value, -2.28444433f);
 
-		FCbFieldView Field(Payload, ECbFieldType::Float32);
+		FCbFieldView Field(Value, ECbFieldType::Float32);
 		TestField<ECbFieldType::Float64>(TEXT("Float32, AsDouble"), Field, -2.28444433);
 	}
 
 	// Test FCbFieldView(Float, 64-bit)
 	{
-		const uint8 Payload[] = { 0xc1, 0x23, 0x45, 0x67, 0x89, 0xab, 0xcd, 0xef }; // -631475.76888888876
-		TestField<ECbFieldType::Float64>(TEXT("Float64"), Payload, -631475.76888888876);
+		const uint8 Value[] = { 0xc1, 0x23, 0x45, 0x67, 0x89, 0xab, 0xcd, 0xef }; // -631475.76888888876
+		TestField<ECbFieldType::Float64>(TEXT("Float64"), Value, -631475.76888888876);
 
-		FCbFieldView Field(Payload, ECbFieldType::Float64);
+		FCbFieldView Field(Value, ECbFieldType::Float64);
 		TestFieldError<ECbFieldType::Float32>(TEXT("Float64, AsFloat"), Field, ECbFieldError::RangeError, 8.0f);
 	}
 
 	// Test FCbFieldView(Integer+, MaxBinary32) as Float
 	{
-		uint8 Payload[9];
-		WriteVarUInt((uint64(1) << 24) - 1, Payload); // 16,777,215
-		FCbFieldView Field(Payload, ECbFieldType::IntegerPositive);
+		uint8 Value[9];
+		WriteVarUInt((uint64(1) << 24) - 1, Value); // 16,777,215
+		FCbFieldView Field(Value, ECbFieldType::IntegerPositive);
 		TestField<ECbFieldType::Float32>(TEXT("Integer+, MaxBinary32, AsFloat"), Field, 16'777'215.0f);
 		TestField<ECbFieldType::Float64>(TEXT("Integer+, MaxBinary32, AsDouble"), Field, 16'777'215.0);
 	}
 
 	// Test FCbFieldView(Integer+, MaxBinary32+1) as Float
 	{
-		uint8 Payload[9];
-		WriteVarUInt(uint64(1) << 24, Payload); // 16,777,216
-		FCbFieldView Field(Payload, ECbFieldType::IntegerPositive);
+		uint8 Value[9];
+		WriteVarUInt(uint64(1) << 24, Value); // 16,777,216
+		FCbFieldView Field(Value, ECbFieldType::IntegerPositive);
 		TestFieldError<ECbFieldType::Float32>(TEXT("Integer+, MaxBinary32+1, AsFloat"), Field, ECbFieldError::RangeError, 8.0f);
 		TestField<ECbFieldType::Float64>(TEXT("Integer+, MaxBinary32+1, AsDouble"), Field, 16'777'216.0);
 	}
 
 	// Test FCbFieldView(Integer+, MaxBinary64) as Float
 	{
-		uint8 Payload[9];
-		WriteVarUInt((uint64(1) << 53) - 1, Payload); // 9,007,199,254,740,991
-		FCbFieldView Field(Payload, ECbFieldType::IntegerPositive);
+		uint8 Value[9];
+		WriteVarUInt((uint64(1) << 53) - 1, Value); // 9,007,199,254,740,991
+		FCbFieldView Field(Value, ECbFieldType::IntegerPositive);
 		TestFieldError<ECbFieldType::Float32>(TEXT("Integer+, MaxBinary64, AsFloat"), Field, ECbFieldError::RangeError, 8.0f);
 		TestField<ECbFieldType::Float64>(TEXT("Integer+, MaxBinary64, AsDouble"), Field, 9'007'199'254'740'991.0);
 	}
 
 	// Test FCbFieldView(Integer+, MaxBinary64+1) as Float
 	{
-		uint8 Payload[9];
-		WriteVarUInt(uint64(1) << 53, Payload); // 9,007,199,254,740,992
-		FCbFieldView Field(Payload, ECbFieldType::IntegerPositive);
+		uint8 Value[9];
+		WriteVarUInt(uint64(1) << 53, Value); // 9,007,199,254,740,992
+		FCbFieldView Field(Value, ECbFieldType::IntegerPositive);
 		TestFieldError<ECbFieldType::Float32>(TEXT("Integer+, MaxBinary64+1, AsFloat"), Field, ECbFieldError::RangeError, 8.0f);
 		TestFieldError<ECbFieldType::Float64>(TEXT("Integer+, MaxBinary64+1, AsDouble"), Field, ECbFieldError::RangeError, 8.0);
 	}
 
 	// Test FCbFieldView(Integer+, MaxUInt64) as Float
 	{
-		uint8 Payload[9];
-		WriteVarUInt(uint64(-1), Payload); // Max uint64
-		FCbFieldView Field(Payload, ECbFieldType::IntegerPositive);
+		uint8 Value[9];
+		WriteVarUInt(uint64(-1), Value); // Max uint64
+		FCbFieldView Field(Value, ECbFieldType::IntegerPositive);
 		TestFieldError<ECbFieldType::Float32>(TEXT("Integer+, MaxUInt64, AsFloat"), Field, ECbFieldError::RangeError, 8.0f);
 		TestFieldError<ECbFieldType::Float64>(TEXT("Integer+, MaxUInt64, AsDouble"), Field, ECbFieldError::RangeError, 8.0);
 	}
 
 	// Test FCbFieldView(Integer-, MaxBinary32) as Float
 	{
-		uint8 Payload[9];
-		WriteVarUInt((uint64(1) << 24) - 2, Payload); // -16,777,215
-		FCbFieldView Field(Payload, ECbFieldType::IntegerNegative);
+		uint8 Value[9];
+		WriteVarUInt((uint64(1) << 24) - 2, Value); // -16,777,215
+		FCbFieldView Field(Value, ECbFieldType::IntegerNegative);
 		TestField<ECbFieldType::Float32>(TEXT("Integer-, MaxBinary32, AsFloat"), Field, -16'777'215.0f);
 		TestField<ECbFieldType::Float64>(TEXT("Integer-, MaxBinary32, AsDouble"), Field, -16'777'215.0);
 	}
 
 	// Test FCbFieldView(Integer-, MaxBinary32+1) as Float
 	{
-		uint8 Payload[9];
-		WriteVarUInt((uint64(1) << 24) - 1, Payload); // -16,777,216
-		FCbFieldView Field(Payload, ECbFieldType::IntegerNegative);
+		uint8 Value[9];
+		WriteVarUInt((uint64(1) << 24) - 1, Value); // -16,777,216
+		FCbFieldView Field(Value, ECbFieldType::IntegerNegative);
 		TestFieldError<ECbFieldType::Float32>(TEXT("Integer-, MaxBinary32+1, AsFloat"), Field, ECbFieldError::RangeError, 8.0f);
 		TestField<ECbFieldType::Float64>(TEXT("Integer-, MaxBinary32+1, AsDouble"), Field, -16'777'216.0);
 	}
 
 	// Test FCbFieldView(Integer-, MaxBinary64) as Float
 	{
-		uint8 Payload[9];
-		WriteVarUInt((uint64(1) << 53) - 2, Payload); // -9,007,199,254,740,991
-		FCbFieldView Field(Payload, ECbFieldType::IntegerNegative);
+		uint8 Value[9];
+		WriteVarUInt((uint64(1) << 53) - 2, Value); // -9,007,199,254,740,991
+		FCbFieldView Field(Value, ECbFieldType::IntegerNegative);
 		TestFieldError<ECbFieldType::Float32>(TEXT("Integer-, MaxBinary64, AsFloat"), Field, ECbFieldError::RangeError, 8.0f);
 		TestField<ECbFieldType::Float64>(TEXT("Integer-, MaxBinary64, AsDouble"), Field, -9'007'199'254'740'991.0);
 	}
 
 	// Test FCbFieldView(Integer-, MaxBinary64+1) as Float
 	{
-		uint8 Payload[9];
-		WriteVarUInt((uint64(1) << 53) - 1, Payload); // -9,007,199,254,740,992
-		FCbFieldView Field(Payload, ECbFieldType::IntegerNegative);
+		uint8 Value[9];
+		WriteVarUInt((uint64(1) << 53) - 1, Value); // -9,007,199,254,740,992
+		FCbFieldView Field(Value, ECbFieldType::IntegerNegative);
 		TestFieldError<ECbFieldType::Float32>(TEXT("Integer-, MaxBinary64+1, AsFloat"), Field, ECbFieldError::RangeError, 8.0f);
 		TestFieldError<ECbFieldType::Float64>(TEXT("Integer-, MaxBinary64+1, AsDouble"), Field, ECbFieldError::RangeError, 8.0);
 	}
@@ -1127,8 +1127,8 @@ bool FCbFieldDateTimeTest::RunTest(const FString& Parameters)
 
 	// Test FCbFieldView(DateTime, Zero) as FDateTime
 	{
-		const uint8 Payload[] = {0, 0, 0, 0, 0, 0, 0, 0};
-		FCbFieldView Field(Payload, ECbFieldType::DateTime);
+		const uint8 Value[] = {0, 0, 0, 0, 0, 0, 0, 0};
+		FCbFieldView Field(Value, ECbFieldType::DateTime);
 		TestEqual(TEXT("FCbFieldView()::AsDateTime()"), Field.AsDateTime(), FDateTime(0));
 	}
 
@@ -1154,8 +1154,8 @@ bool FCbFieldTimeSpanTest::RunTest(const FString& Parameters)
 
 	// Test FCbFieldView(TimeSpan, Zero) as FTimeSpan
 	{
-		const uint8 Payload[] = {0, 0, 0, 0, 0, 0, 0, 0};
-		FCbFieldView Field(Payload, ECbFieldType::TimeSpan);
+		const uint8 Value[] = {0, 0, 0, 0, 0, 0, 0, 0};
+		FCbFieldView Field(Value, ECbFieldType::TimeSpan);
 		TestEqual(TEXT("FCbFieldView()::AsTimeSpan()"), Field.AsTimeSpan(), FTimespan(0));
 	}
 
@@ -1182,8 +1182,8 @@ bool FCbFieldObjectIdTest::RunTest(const FString& Parameters)
 
 	// Test FCbFieldView(ObjectId, Zero) as FCbObjectId
 	{
-		const uint8 Payload[] = {0, 0, 0, 0, 0, 0, 0, 0, 0, 0, 0, 0};
-		FCbFieldView Field(Payload, ECbFieldType::ObjectId);
+		const uint8 Value[] = {0, 0, 0, 0, 0, 0, 0, 0, 0, 0, 0, 0};
+		FCbFieldView Field(Value, ECbFieldType::ObjectId);
 		TestEqual(TEXT("FCbFieldView(ObjectId, Zero)::AsObjectId()"), Field.AsObjectId(), FCbObjectId());
 	}
 
@@ -1214,34 +1214,34 @@ bool FCbFieldCustomByIdTest::RunTest(const FString& Parameters)
 
 	// Test FCbFieldView(CustomById, MinId, Empty)
 	{
-		const uint8 Payload[] = {1, 0};
-		TestField<ECbFieldType::CustomById>(TEXT("CustomById, MinId, Empty"), Payload, FCbCustomById{0});
-		TestField<ECbFieldType::CustomById>(TEXT("CustomById, MinId, Empty, View"), Payload, FMemoryView(), MakeMemoryView<uint8>({1, 2, 3}), ECbFieldError::None, FCustomByIdAccessor(0));
-		TestFieldError<ECbFieldType::CustomById>(TEXT("CustomById, MinId, Empty, InvalidId"), Payload, ECbFieldError::RangeError, MakeMemoryView<uint8>({1, 2, 3}), FCustomByIdAccessor(MAX_uint64));
+		const uint8 Value[] = {1, 0};
+		TestField<ECbFieldType::CustomById>(TEXT("CustomById, MinId, Empty"), Value, FCbCustomById{0});
+		TestField<ECbFieldType::CustomById>(TEXT("CustomById, MinId, Empty, View"), Value, FMemoryView(), MakeMemoryView<uint8>({1, 2, 3}), ECbFieldError::None, FCustomByIdAccessor(0));
+		TestFieldError<ECbFieldType::CustomById>(TEXT("CustomById, MinId, Empty, InvalidId"), Value, ECbFieldError::RangeError, MakeMemoryView<uint8>({1, 2, 3}), FCustomByIdAccessor(MAX_uint64));
 	}
 
 	// Test FCbFieldView(CustomById, MinId, Value)
 	{
-		const uint8 Payload[] = {5, 0, 1, 2, 3, 4};
-		TestFieldNoClone<ECbFieldType::CustomById>(TEXT("CustomById, MinId, Value"), Payload, FCbCustomById{0, MakeMemoryView(Payload).Right(4)});
-		TestFieldNoClone<ECbFieldType::CustomById>(TEXT("CustomById, MinId, Value, View"), Payload, MakeMemoryView(Payload).Right(4), FMemoryView(), ECbFieldError::None, FCustomByIdAccessor(0));
-		TestFieldError<ECbFieldType::CustomById>(TEXT("CustomById, MinId, Value, InvalidId"), Payload, ECbFieldError::RangeError, MakeMemoryView<uint8>({1, 2, 3}), FCustomByIdAccessor(MAX_uint64));
+		const uint8 Value[] = {5, 0, 1, 2, 3, 4};
+		TestFieldNoClone<ECbFieldType::CustomById>(TEXT("CustomById, MinId, Value"), Value, FCbCustomById{0, MakeMemoryView(Value).Right(4)});
+		TestFieldNoClone<ECbFieldType::CustomById>(TEXT("CustomById, MinId, Value, View"), Value, MakeMemoryView(Value).Right(4), FMemoryView(), ECbFieldError::None, FCustomByIdAccessor(0));
+		TestFieldError<ECbFieldType::CustomById>(TEXT("CustomById, MinId, Value, InvalidId"), Value, ECbFieldError::RangeError, MakeMemoryView<uint8>({1, 2, 3}), FCustomByIdAccessor(MAX_uint64));
 	}
 
 	// Test FCbFieldView(CustomById, MaxId, Empty)
 	{
-		const uint8 Payload[] = {9, 0xff, 0xff, 0xff, 0xff, 0xff, 0xff, 0xff, 0xff, 0xff};
-		TestField<ECbFieldType::CustomById>(TEXT("CustomById, MaxId, Empty"), Payload, FCbCustomById{MAX_uint64});
-		TestField<ECbFieldType::CustomById>(TEXT("CustomById, MaxId, Empty, View"), Payload, FMemoryView(), MakeMemoryView<uint8>({1, 2, 3}), ECbFieldError::None, FCustomByIdAccessor(MAX_uint64));
-		TestFieldError<ECbFieldType::CustomById>(TEXT("CustomById, MaxId, Empty, InvalidId"), Payload, ECbFieldError::RangeError, MakeMemoryView<uint8>({1, 2, 3}), FCustomByIdAccessor(0));
+		const uint8 Value[] = {9, 0xff, 0xff, 0xff, 0xff, 0xff, 0xff, 0xff, 0xff, 0xff};
+		TestField<ECbFieldType::CustomById>(TEXT("CustomById, MaxId, Empty"), Value, FCbCustomById{MAX_uint64});
+		TestField<ECbFieldType::CustomById>(TEXT("CustomById, MaxId, Empty, View"), Value, FMemoryView(), MakeMemoryView<uint8>({1, 2, 3}), ECbFieldError::None, FCustomByIdAccessor(MAX_uint64));
+		TestFieldError<ECbFieldType::CustomById>(TEXT("CustomById, MaxId, Empty, InvalidId"), Value, ECbFieldError::RangeError, MakeMemoryView<uint8>({1, 2, 3}), FCustomByIdAccessor(0));
 	}
 
 	// Test FCbFieldView(CustomById, MaxId, Value)
 	{
-		const uint8 Payload[] = {13, 0xff, 0xff, 0xff, 0xff, 0xff, 0xff, 0xff, 0xff, 0xff, 1, 2, 3, 4};
-		TestFieldNoClone<ECbFieldType::CustomById>(TEXT("CustomById, MaxId, Value"), Payload, FCbCustomById{MAX_uint64, MakeMemoryView(Payload).Right(4)});
-		TestFieldNoClone<ECbFieldType::CustomById>(TEXT("CustomById, MaxId, Value, View"), Payload, MakeMemoryView(Payload).Right(4), FMemoryView(), ECbFieldError::None, FCustomByIdAccessor(MAX_uint64));
-		TestFieldError<ECbFieldType::CustomById>(TEXT("CustomById, MaxId, Value, InvalidId"), Payload, ECbFieldError::RangeError, MakeMemoryView<uint8>({1, 2, 3}), FCustomByIdAccessor(0));
+		const uint8 Value[] = {13, 0xff, 0xff, 0xff, 0xff, 0xff, 0xff, 0xff, 0xff, 0xff, 1, 2, 3, 4};
+		TestFieldNoClone<ECbFieldType::CustomById>(TEXT("CustomById, MaxId, Value"), Value, FCbCustomById{MAX_uint64, MakeMemoryView(Value).Right(4)});
+		TestFieldNoClone<ECbFieldType::CustomById>(TEXT("CustomById, MaxId, Value, View"), Value, MakeMemoryView(Value).Right(4), FMemoryView(), ECbFieldError::None, FCustomByIdAccessor(MAX_uint64));
+		TestFieldError<ECbFieldType::CustomById>(TEXT("CustomById, MaxId, Value, InvalidId"), Value, ECbFieldError::RangeError, MakeMemoryView<uint8>({1, 2, 3}), FCustomByIdAccessor(0));
 	}
 
 	// Test FCbFieldView(None) as CustomById
@@ -1272,18 +1272,18 @@ bool FCbFieldCustomByNameTest::RunTest(const FString& Parameters)
 
 	// Test FCbFieldView(CustomByName, ABC, Empty)
 	{
-		const uint8 Payload[] = {4, 3, 'A', 'B', 'C'};
-		TestField<ECbFieldType::CustomByName>(TEXT("CustomByName, MinId, Empty"), Payload, FCbCustomByName{"ABC"_U8SV});
-		TestField<ECbFieldType::CustomByName>(TEXT("CustomByName, MinId, Empty, View"), Payload, FMemoryView(), MakeMemoryView<uint8>({1, 2, 3}), ECbFieldError::None, FCustomByNameAccessor("ABC"_U8SV));
-		TestFieldError<ECbFieldType::CustomByName>(TEXT("CustomByName, MinId, Empty, InvalidCase"), Payload, ECbFieldError::RangeError, MakeMemoryView<uint8>({1, 2, 3}), FCustomByNameAccessor("abc"_U8SV));
+		const uint8 Value[] = {4, 3, 'A', 'B', 'C'};
+		TestField<ECbFieldType::CustomByName>(TEXT("CustomByName, MinId, Empty"), Value, FCbCustomByName{"ABC"_U8SV});
+		TestField<ECbFieldType::CustomByName>(TEXT("CustomByName, MinId, Empty, View"), Value, FMemoryView(), MakeMemoryView<uint8>({1, 2, 3}), ECbFieldError::None, FCustomByNameAccessor("ABC"_U8SV));
+		TestFieldError<ECbFieldType::CustomByName>(TEXT("CustomByName, MinId, Empty, InvalidCase"), Value, ECbFieldError::RangeError, MakeMemoryView<uint8>({1, 2, 3}), FCustomByNameAccessor("abc"_U8SV));
 	}
 
 	// Test FCbFieldView(CustomByName, ABC, Value)
 	{
-		const uint8 Payload[] = {8, 3, 'A', 'B', 'C', 1, 2, 3, 4};
-		TestFieldNoClone<ECbFieldType::CustomByName>(TEXT("CustomByName, MinId, Value"), Payload, FCbCustomByName{"ABC"_U8SV, MakeMemoryView(Payload).Right(4)});
-		TestFieldNoClone<ECbFieldType::CustomByName>(TEXT("CustomByName, MinId, Value, View"), Payload, MakeMemoryView(Payload).Right(4), FMemoryView(), ECbFieldError::None, FCustomByNameAccessor("ABC"_U8SV));
-		TestFieldError<ECbFieldType::CustomByName>(TEXT("CustomByName, MinId, Value, InvalidCase"), Payload, ECbFieldError::RangeError, MakeMemoryView<uint8>({1, 2, 3}), FCustomByNameAccessor("abc"_U8SV));
+		const uint8 Value[] = {8, 3, 'A', 'B', 'C', 1, 2, 3, 4};
+		TestFieldNoClone<ECbFieldType::CustomByName>(TEXT("CustomByName, MinId, Value"), Value, FCbCustomByName{"ABC"_U8SV, MakeMemoryView(Value).Right(4)});
+		TestFieldNoClone<ECbFieldType::CustomByName>(TEXT("CustomByName, MinId, Value, View"), Value, MakeMemoryView(Value).Right(4), FMemoryView(), ECbFieldError::None, FCustomByNameAccessor("ABC"_U8SV));
+		TestFieldError<ECbFieldType::CustomByName>(TEXT("CustomByName, MinId, Value, InvalidCase"), Value, ECbFieldError::RangeError, MakeMemoryView<uint8>({1, 2, 3}), FCustomByNameAccessor("abc"_U8SV));
 	}
 
 	// Test FCbFieldView(None) as CustomByName
@@ -1507,8 +1507,8 @@ bool FCbFieldBufferTest::RunTest(const FString& Parameters)
 
 	// Test Field w/ Type from Shared Buffer
 	{
-		const uint8 Payload[] = { uint8(ECbFieldType::Binary), 3, 4, 5, 6 }; // Size: 3, Data: 4/5/6
-		FSharedBuffer ViewBuffer = FSharedBuffer::MakeView(MakeMemoryView(Payload));
+		const uint8 Value[] = { uint8(ECbFieldType::Binary), 3, 4, 5, 6 }; // Size: 3, Data: 4/5/6
+		FSharedBuffer ViewBuffer = FSharedBuffer::MakeView(MakeMemoryView(Value));
 		FSharedBuffer OwnedBuffer = FSharedBuffer::Clone(ViewBuffer);
 
 		FCbField View(ViewBuffer);
@@ -1552,20 +1552,20 @@ bool FCbFieldBufferTest::RunTest(const FString& Parameters)
 
 	// Test Field w/ Type
 	{
-		const uint8 Payload[] = { uint8(ECbFieldType::Binary), 3, 4, 5, 6 }; // Size: 3, Data: 4/5/6
-		const FCbFieldView Field(Payload);
+		const uint8 Value[] = { uint8(ECbFieldType::Binary), 3, 4, 5, 6 }; // Size: 3, Data: 4/5/6
+		const FCbFieldView Field(Value);
 
-		FCbField VoidView = FCbField::MakeView(Payload);
-		FCbField VoidClone = FCbField::Clone(Payload);
+		FCbField VoidView = FCbField::MakeView(Value);
+		FCbField VoidClone = FCbField::Clone(Value);
 		FCbField FieldView = FCbField::MakeView(Field);
 		FCbField FieldClone = FCbField::Clone(Field);
 		FCbField FieldViewClone = FCbField::Clone(FieldView);
 
-		TestEqual(TEXT("FCbField::MakeView(Void)"), VoidView.AsBinaryView(), MakeMemoryView(Payload).Right(3));
-		TestNotEqual(TEXT("FCbField::Clone(Void)"), VoidClone.AsBinaryView(), MakeMemoryView(Payload).Right(3));
+		TestEqual(TEXT("FCbField::MakeView(Void)"), VoidView.AsBinaryView(), MakeMemoryView(Value).Right(3));
+		TestNotEqual(TEXT("FCbField::Clone(Void)"), VoidClone.AsBinaryView(), MakeMemoryView(Value).Right(3));
 		TestTrue(TEXT("FCbField::Clone(Void)->EqualBytes"), VoidClone.AsBinaryView().EqualBytes(VoidView.AsBinaryView()));
-		TestEqual(TEXT("FCbField::MakeView(Field)"), FieldView.AsBinaryView(), MakeMemoryView(Payload).Right(3));
-		TestNotEqual(TEXT("FCbField::Clone(Field)"), FieldClone.AsBinaryView(), MakeMemoryView(Payload).Right(3));
+		TestEqual(TEXT("FCbField::MakeView(Field)"), FieldView.AsBinaryView(), MakeMemoryView(Value).Right(3));
+		TestNotEqual(TEXT("FCbField::Clone(Field)"), FieldClone.AsBinaryView(), MakeMemoryView(Value).Right(3));
 		TestTrue(TEXT("FCbField::Clone(Field)->EqualBytes"), FieldClone.AsBinaryView().EqualBytes(VoidView.AsBinaryView()));
 		TestNotEqual(TEXT("FCbField::Clone(FieldView)"), FieldViewClone.AsBinaryView(), FieldView.AsBinaryView());
 		TestTrue(TEXT("FCbField::Clone(FieldView)->EqualBytes"), FieldViewClone.AsBinaryView().EqualBytes(VoidView.AsBinaryView()));
@@ -1579,14 +1579,14 @@ bool FCbFieldBufferTest::RunTest(const FString& Parameters)
 
 	// Test Field w/o Type
 	{
-		const uint8 Payload[] = { 3, 4, 5, 6 }; // Size: 3, Data: 4/5/6
-		const FCbFieldView Field(Payload, ECbFieldType::Binary);
+		const uint8 Value[] = { 3, 4, 5, 6 }; // Size: 3, Data: 4/5/6
+		const FCbFieldView Field(Value, ECbFieldType::Binary);
 
 		FCbField FieldView = FCbField::MakeView(Field);
 		FCbField FieldClone = FCbField::Clone(Field);
 		FCbField FieldViewClone = FCbField::Clone(FieldView);
 
-		TestEqual(TEXT("FCbField::MakeView(Field, NoType)"), FieldView.AsBinaryView(), MakeMemoryView(Payload).Right(3));
+		TestEqual(TEXT("FCbField::MakeView(Field, NoType)"), FieldView.AsBinaryView(), MakeMemoryView(Value).Right(3));
 		TestTrue(TEXT("FCbField::Clone(Field, NoType)"), FieldClone.AsBinaryView().EqualBytes(FieldView.AsBinaryView()));
 		TestTrue(TEXT("FCbField::Clone(FieldView, NoType)"), FieldViewClone.AsBinaryView().EqualBytes(FieldView.AsBinaryView()));
 
@@ -1595,7 +1595,7 @@ bool FCbFieldBufferTest::RunTest(const FString& Parameters)
 		TestTrue(TEXT("FCbField::Clone(FieldView, NoType).IsOwned()"), FieldViewClone.IsOwned());
 
 		FieldView.MakeOwned();
-		TestTrue(TEXT("FCbField::MakeView(NoType).MakeOwned()"), FieldView.AsBinaryView().EqualBytes(MakeMemoryView(Payload).Right(3)));
+		TestTrue(TEXT("FCbField::MakeView(NoType).MakeOwned()"), FieldView.AsBinaryView().EqualBytes(MakeMemoryView(Value).Right(3)));
 		TestTrue(TEXT("FCbField::MakeView(NoType).MakeOwned().IsOwned()"), FieldView.IsOwned());
 	}
 
@@ -1696,9 +1696,9 @@ bool FCbFieldBufferIteratorTest::RunTest(const FString& Parameters)
 	// Test FCbField[View]Iterator(Range)
 	{
 		constexpr uint8 T = uint8(ECbFieldType::IntegerPositive);
-		const uint8 Payload[] = { T, 0, T, 1, T, 2, T, 3 };
+		const uint8 Value[] = { T, 0, T, 1, T, 2, T, 3 };
 
-		const FSharedBuffer View = FSharedBuffer::MakeView(MakeMemoryView(Payload));
+		const FSharedBuffer View = FSharedBuffer::MakeView(MakeMemoryView(Value));
 		const FSharedBuffer Clone = FSharedBuffer::Clone(View);
 
 		const FMemoryView EmptyView;
@@ -1711,8 +1711,8 @@ bool FCbFieldBufferIteratorTest::RunTest(const FString& Parameters)
 		TestEqual(TEXT("FCbFieldIterator::GetRangeHash()"), FieldIt.GetRangeHash(), FBlake3::HashBuffer(View));
 
 		FMemoryView RangeView;
-		TestTrue(TEXT("FCbFieldViewIterator::TryGetRangeView()"), FieldViewIt.TryGetRangeView(RangeView) && RangeView == MakeMemoryView(Payload));
-		TestTrue(TEXT("FCbFieldIterator::TryGetRangeView()"), FieldIt.TryGetRangeView(RangeView) && RangeView == MakeMemoryView(Payload));
+		TestTrue(TEXT("FCbFieldViewIterator::TryGetRangeView()"), FieldViewIt.TryGetRangeView(RangeView) && RangeView == MakeMemoryView(Value));
+		TestTrue(TEXT("FCbFieldIterator::TryGetRangeView()"), FieldIt.TryGetRangeView(RangeView) && RangeView == MakeMemoryView(Value));
 
 		TestEqual(TEXT("FCbFieldIterator::CloneRange(EmptyViewIt)"), GetCount(FCbFieldIterator::CloneRange(FCbFieldViewIterator())), 0);
 		TestEqual(TEXT("FCbFieldIterator::CloneRange(EmptyIt)"), GetCount(FCbFieldIterator::CloneRange(FCbFieldIterator())), 0);
@@ -1727,7 +1727,7 @@ bool FCbFieldBufferIteratorTest::RunTest(const FString& Parameters)
 		TestEqual(TEXT("FCbFieldIterator::MakeRange(BufferNullL)"), GetCount(FCbFieldIterator::MakeRange(NullBuffer)), 0);
 		TestEqual(TEXT("FCbFieldIterator::MakeRange(BufferNullR)"), GetCount(FCbFieldIterator::MakeRange(FSharedBuffer(NullBuffer))), 0);
 
-		TestEqual(TEXT("FCbFieldViewIterator::MakeRange(BufferView)"), GetCount(FCbFieldViewIterator::MakeRange(MakeMemoryView(Payload))), 4);
+		TestEqual(TEXT("FCbFieldViewIterator::MakeRange(BufferView)"), GetCount(FCbFieldViewIterator::MakeRange(MakeMemoryView(Value))), 4);
 		TestEqual(TEXT("FCbFieldIterator::MakeRange(BufferCloneL)"), GetCount(FCbFieldIterator::MakeRange(Clone)), 4);
 		TestEqual(TEXT("FCbFieldIterator::MakeRange(BufferCloneR)"), GetCount(FCbFieldIterator::MakeRange(FSharedBuffer(Clone))), 4);
 
@@ -1742,11 +1742,11 @@ bool FCbFieldBufferIteratorTest::RunTest(const FString& Parameters)
 		TestEqual(TEXT("FCbFieldViewIterator(FieldItR)"), GetCount(FCbFieldViewIterator(FCbFieldIterator(FieldIt))), 4);
 
 		// Uniform
-		const uint8 UniformPayload[] = { 0, 1, 2, 3 };
-		const FCbFieldViewIterator UniformFieldViewIt = FCbFieldViewIterator::MakeRange(MakeMemoryView(UniformPayload), ECbFieldType::IntegerPositive);
+		const uint8 UniformValue[] = { 0, 1, 2, 3 };
+		const FCbFieldViewIterator UniformFieldViewIt = FCbFieldViewIterator::MakeRange(MakeMemoryView(UniformValue), ECbFieldType::IntegerPositive);
 		TestEqual(TEXT("FCbFieldViewIterator::MakeRange(Uniform).GetRangeHash()"), UniformFieldViewIt.GetRangeHash(), FieldViewIt.GetRangeHash());
 		TestFalse(TEXT("FCbFieldViewIterator::MakeRange(Uniform).TryGetRangeView()"), UniformFieldViewIt.TryGetRangeView(RangeView));
-		const FSharedBuffer UniformView = FSharedBuffer::MakeView(MakeMemoryView(UniformPayload));
+		const FSharedBuffer UniformView = FSharedBuffer::MakeView(MakeMemoryView(UniformValue));
 		const FCbFieldIterator UniformFieldIt = FCbFieldIterator::MakeRange(UniformView, ECbFieldType::IntegerPositive);
 		TestEqual(TEXT("FCbFieldIterator::MakeRange(Uniform).GetRangeHash()"), UniformFieldIt.GetRangeHash(), FieldViewIt.GetRangeHash());
 		TestFalse(TEXT("FCbFieldIterator::MakeRange(Uniform).TryGetRangeView()"), UniformFieldIt.TryGetRangeView(RangeView));
@@ -1762,32 +1762,32 @@ bool FCbFieldBufferIteratorTest::RunTest(const FString& Parameters)
 		TestTrue(TEXT("FCbFieldViewIterator::Equals(Uniform, OtherType)"), UniformFieldViewIt.Equals(UniformFieldIt));
 		TestTrue(TEXT("FCbFieldIterator::Equals(Uniform, Self)"), UniformFieldIt.Equals(AsConst(UniformFieldIt)));
 		TestTrue(TEXT("FCbFieldIterator::Equals(Uniform, OtherType)"), UniformFieldIt.Equals(UniformFieldViewIt));
-		TestFalse(TEXT("FCbFieldViewIterator::Equals(SamePayload, DifferentEnd)"),
-			FCbFieldViewIterator::MakeRange(MakeMemoryView(UniformPayload), ECbFieldType::IntegerPositive)
-				.Equals(FCbFieldViewIterator::MakeRange(MakeMemoryView(UniformPayload).LeftChop(1), ECbFieldType::IntegerPositive)));
-		TestFalse(TEXT("FCbFieldViewIterator::Equals(DifferentPayload, SameEnd)"),
-			FCbFieldViewIterator::MakeRange(MakeMemoryView(UniformPayload), ECbFieldType::IntegerPositive)
-				.Equals(FCbFieldViewIterator::MakeRange(MakeMemoryView(UniformPayload).RightChop(1), ECbFieldType::IntegerPositive)));
+		TestFalse(TEXT("FCbFieldViewIterator::Equals(SameValue, DifferentEnd)"),
+			FCbFieldViewIterator::MakeRange(MakeMemoryView(UniformValue), ECbFieldType::IntegerPositive)
+				.Equals(FCbFieldViewIterator::MakeRange(MakeMemoryView(UniformValue).LeftChop(1), ECbFieldType::IntegerPositive)));
+		TestFalse(TEXT("FCbFieldViewIterator::Equals(DifferentValue, SameEnd)"),
+			FCbFieldViewIterator::MakeRange(MakeMemoryView(UniformValue), ECbFieldType::IntegerPositive)
+				.Equals(FCbFieldViewIterator::MakeRange(MakeMemoryView(UniformValue).RightChop(1), ECbFieldType::IntegerPositive)));
 
 		// CopyRangeTo
-		uint8 CopyBytes[sizeof(Payload)];
+		uint8 CopyBytes[sizeof(Value)];
 		FieldViewIt.CopyRangeTo(MakeMemoryView(CopyBytes));
-		TestTrue(TEXT("FCbFieldViewIterator::MakeRange().CopyRangeTo()"), MakeMemoryView(CopyBytes).EqualBytes(MakeMemoryView(Payload)));
+		TestTrue(TEXT("FCbFieldViewIterator::MakeRange().CopyRangeTo()"), MakeMemoryView(CopyBytes).EqualBytes(MakeMemoryView(Value)));
 		FieldIt.CopyRangeTo(MakeMemoryView(CopyBytes));
-		TestTrue(TEXT("FCbFieldIterator::MakeRange().CopyRangeTo()"), MakeMemoryView(CopyBytes).EqualBytes(MakeMemoryView(Payload)));
+		TestTrue(TEXT("FCbFieldIterator::MakeRange().CopyRangeTo()"), MakeMemoryView(CopyBytes).EqualBytes(MakeMemoryView(Value)));
 		UniformFieldViewIt.CopyRangeTo(MakeMemoryView(CopyBytes));
-		TestTrue(TEXT("FCbFieldViewIterator::MakeRange(Uniform).CopyRangeTo()"), MakeMemoryView(CopyBytes).EqualBytes(MakeMemoryView(Payload)));
+		TestTrue(TEXT("FCbFieldViewIterator::MakeRange(Uniform).CopyRangeTo()"), MakeMemoryView(CopyBytes).EqualBytes(MakeMemoryView(Value)));
 
 		// MakeRangeOwned
 		FCbFieldIterator OwnedFromView = UniformFieldIt;
 		OwnedFromView.MakeRangeOwned();
-		TestTrue(TEXT("FCbFieldIterator::MakeRangeOwned(View)"), OwnedFromView.TryGetRangeView(RangeView) && RangeView.EqualBytes(MakeMemoryView(Payload)));
+		TestTrue(TEXT("FCbFieldIterator::MakeRangeOwned(View)"), OwnedFromView.TryGetRangeView(RangeView) && RangeView.EqualBytes(MakeMemoryView(Value)));
 		FCbFieldIterator OwnedFromOwned = OwnedFromView;
 		OwnedFromOwned.MakeRangeOwned();
 		TestEqual(TEXT("FCbFieldIterator::MakeRangeOwned(Owned)"), OwnedFromOwned, OwnedFromView);
 
 		// These lines are expected to assert when uncommented.
-		//const FSharedBuffer ShortView = FSharedBuffer::MakeView(MakeMemoryView(Payload).LeftChop(2));
+		//const FSharedBuffer ShortView = FSharedBuffer::MakeView(MakeMemoryView(Value).LeftChop(2));
 		//TestEqual(TEXT("FCbFieldIterator::MakeRangeView(FieldIt, InvalidBufferL)"), GetCount(FCbFieldIterator::MakeRangeView(FCbFieldViewIterator::MakeRange(*View), ShortView)), 4);
 		//TestEqual(TEXT("FCbFieldIterator::MakeRangeView(FieldIt, InvalidBufferR)"), GetCount(FCbFieldIterator::MakeRangeView(FCbFieldViewIterator::MakeRange(*View), FSharedBuffer(ShortView))), 4);
 	}
@@ -1795,12 +1795,12 @@ bool FCbFieldBufferIteratorTest::RunTest(const FString& Parameters)
 	// Test FCbField[View]Iterator(Scalar)
 	{
 		constexpr uint8 T = uint8(ECbFieldType::IntegerPositive);
-		const uint8 Payload[] = { T, 0 };
+		const uint8 Value[] = { T, 0 };
 
-		const FSharedBuffer View = FSharedBuffer::MakeView(MakeMemoryView(Payload));
+		const FSharedBuffer View = FSharedBuffer::MakeView(MakeMemoryView(Value));
 		const FSharedBuffer Clone = FSharedBuffer::Clone(View);
 
-		const FCbFieldView FieldView(Payload);
+		const FCbFieldView FieldView(Value);
 		const FCbField Field(View);
 
 		TestEqual(TEXT("FCbFieldViewIterator::MakeSingle(FieldViewL)"), GetCount(FCbFieldViewIterator::MakeSingle(FieldView)), 1);
diff --git a/Engine/Source/Runtime/Core/Private/Tests/Serialization/CompactBinaryValidationTest.cpp b/Engine/Source/Runtime/Core/Private/Tests/Serialization/CompactBinaryValidationTest.cpp
--- a/Engine/Source/Runtime/Core/Private/Tests/Serialization/CompactBinaryValidationTest.cpp
+++ b/Engine/Source/Runtime/Core/Private/Tests/Serialization/CompactBinaryValidationTest.cpp
@@ -177,27 +177,27 @@ bool FCbValidateTest::RunTest(const FString& Parameters)
 
 	TestEqual(TEXT("ValidateCompactBinary(Valid, CustomById, Empty)"), Validate({uint8(ECbFieldType::CustomById), 1, 0}), ECbValidateError::None);
 	TestEqual(TEXT("ValidateCompactBinary(Valid, CustomById, Empty, NoType)"), Validate({1, 0}, ECbFieldType::CustomById), ECbValidateError::None);
-	TestEqual(TEXT("ValidateCompactBinary(Valid, CustomById, Payload)"), Validate({uint8(ECbFieldType::CustomById), 2, 0, 0}), ECbValidateError::None);
-	TestEqual(TEXT("ValidateCompactBinary(Valid, CustomById, Payload, NoType)"), Validate({2, 0, 0}, ECbFieldType::CustomById), ECbValidateError::None);
+	TestEqual(TEXT("ValidateCompactBinary(Valid, CustomById, Value)"), Validate({uint8(ECbFieldType::CustomById), 2, 0, 0}), ECbValidateError::None);
+	TestEqual(TEXT("ValidateCompactBinary(Valid, CustomById, Value, NoType)"), Validate({2, 0, 0}, ECbFieldType::CustomById), ECbValidateError::None);
 	TestEqual(TEXT("ValidateCompactBinary(Valid, CustomById, MaxTypeId)"), Validate({uint8(ECbFieldType::CustomById), 9, 0xff, 0xff, 0xff, 0xff, 0xff, 0xff, 0xff, 0xff, 0xff}), ECbValidateError::None);
 	TestEqual(TEXT("ValidateCompactBinary(Valid, CustomById, MaxTypeId, NoType)"), Validate({9, 0xff, 0xff, 0xff, 0xff, 0xff, 0xff, 0xff, 0xff, 0xff}, ECbFieldType::CustomById), ECbValidateError::None);
-	TestEqual(TEXT("ValidateCompactBinary(Valid, CustomById, MaxTypeId, Payload)"), Validate({uint8(ECbFieldType::CustomById), 10, 0xff, 0xff, 0xff, 0xff, 0xff, 0xff, 0xff, 0xff, 0xff, 0}), ECbValidateError::None);
-	TestEqual(TEXT("ValidateCompactBinary(Valid, CustomById, MaxTypeId, Payload, NoType)"), Validate({10, 0xff, 0xff, 0xff, 0xff, 0xff, 0xff, 0xff, 0xff, 0xff, 0}, ECbFieldType::CustomById), ECbValidateError::None);
+	TestEqual(TEXT("ValidateCompactBinary(Valid, CustomById, MaxTypeId, Value)"), Validate({uint8(ECbFieldType::CustomById), 10, 0xff, 0xff, 0xff, 0xff, 0xff, 0xff, 0xff, 0xff, 0xff, 0}), ECbValidateError::None);
+	TestEqual(TEXT("ValidateCompactBinary(Valid, CustomById, MaxTypeId, Value, NoType)"), Validate({10, 0xff, 0xff, 0xff, 0xff, 0xff, 0xff, 0xff, 0xff, 0xff, 0}, ECbFieldType::CustomById), ECbValidateError::None);
 	TestEqual(TEXT("ValidateCompactBinary(OutOfBounds, CustomById, NoSize)"), Validate({uint8(ECbFieldType::CustomById)}), ECbValidateError::OutOfBounds);
 	TestEqual(TEXT("ValidateCompactBinary(OutOfBounds, CustomById, NoSize, NoType)"), Validate({}, ECbFieldType::CustomById), ECbValidateError::OutOfBounds);
 	TestEqual(TEXT("ValidateCompactBinary(OutOfBounds, CustomById, NoTypeId)"), Validate({uint8(ECbFieldType::CustomById), 0}), ECbValidateError::OutOfBounds);
 	TestEqual(TEXT("ValidateCompactBinary(OutOfBounds, CustomById, NoTypeId, NoType)"), Validate({0}, ECbFieldType::CustomById), ECbValidateError::OutOfBounds);
 	TestEqual(TEXT("ValidateCompactBinary(OutOfBounds, CustomById, ShortTypeId)"), Validate({uint8(ECbFieldType::CustomById), 1, 0x80}), ECbValidateError::OutOfBounds);
 	TestEqual(TEXT("ValidateCompactBinary(OutOfBounds, CustomById, ShortTypeId, NoType)"), Validate({1, 0x80}, ECbFieldType::CustomById), ECbValidateError::OutOfBounds);
-	TestEqual(TEXT("ValidateCompactBinary(OutOfBounds, CustomById, NoPayload)"), Validate({uint8(ECbFieldType::CustomById), 2, 0}), ECbValidateError::OutOfBounds);
-	TestEqual(TEXT("ValidateCompactBinary(OutOfBounds, CustomById, NoPayload, NoType)"), Validate({2, 0}, ECbFieldType::CustomById), ECbValidateError::OutOfBounds);
-	TestEqual(TEXT("ValidateCompactBinary(OutOfBounds, CustomById, ShortPayload)"), Validate({uint8(ECbFieldType::CustomById), 3, 0, 0}), ECbValidateError::OutOfBounds);
-	TestEqual(TEXT("ValidateCompactBinary(OutOfBounds, CustomById, ShortPayload, NoType)"), Validate({3, 0, 0}, ECbFieldType::CustomById), ECbValidateError::OutOfBounds);
+	TestEqual(TEXT("ValidateCompactBinary(OutOfBounds, CustomById, NoValue)"), Validate({uint8(ECbFieldType::CustomById), 2, 0}), ECbValidateError::OutOfBounds);
+	TestEqual(TEXT("ValidateCompactBinary(OutOfBounds, CustomById, NoValue, NoType)"), Validate({2, 0}, ECbFieldType::CustomById), ECbValidateError::OutOfBounds);
+	TestEqual(TEXT("ValidateCompactBinary(OutOfBounds, CustomById, ShortValue)"), Validate({uint8(ECbFieldType::CustomById), 3, 0, 0}), ECbValidateError::OutOfBounds);
+	TestEqual(TEXT("ValidateCompactBinary(OutOfBounds, CustomById, ShortValue, NoType)"), Validate({3, 0, 0}, ECbFieldType::CustomById), ECbValidateError::OutOfBounds);
 
 	TestEqual(TEXT("ValidateCompactBinary(Valid, CustomByName, Empty)"), Validate({uint8(ECbFieldType::CustomByName), 2, 1, 'A'}), ECbValidateError::None);
 	TestEqual(TEXT("ValidateCompactBinary(Valid, CustomByName, Empty, NoType)"), Validate({2, 1, 'A'}, ECbFieldType::CustomByName), ECbValidateError::None);
-	TestEqual(TEXT("ValidateCompactBinary(Valid, CustomByName, Payload)"), Validate({uint8(ECbFieldType::CustomByName), 3, 1, 'A', 0}), ECbValidateError::None);
-	TestEqual(TEXT("ValidateCompactBinary(Valid, CustomByName, Payload, NoType)"), Validate({3, 1, 'A', 0}, ECbFieldType::CustomByName), ECbValidateError::None);
+	TestEqual(TEXT("ValidateCompactBinary(Valid, CustomByName, Value)"), Validate({uint8(ECbFieldType::CustomByName), 3, 1, 'A', 0}), ECbValidateError::None);
+	TestEqual(TEXT("ValidateCompactBinary(Valid, CustomByName, Value, NoType)"), Validate({3, 1, 'A', 0}, ECbFieldType::CustomByName), ECbValidateError::None);
 	TestEqual(TEXT("ValidateCompactBinary(OutOfBounds, CustomByName, NoSize)"), Validate({uint8(ECbFieldType::CustomByName)}), ECbValidateError::OutOfBounds);
 	TestEqual(TEXT("ValidateCompactBinary(OutOfBounds, CustomByName, NoSize, NoType)"), Validate({}, ECbFieldType::CustomByName), ECbValidateError::OutOfBounds);
 	TestEqual(TEXT("ValidateCompactBinary(OutOfBounds, CustomByName, NoTypeSize)"), Validate({uint8(ECbFieldType::CustomByName), 0}), ECbValidateError::OutOfBounds);
@@ -206,10 +206,10 @@ bool FCbValidateTest::RunTest(const FString& Parameters)
 	TestEqual(TEXT("ValidateCompactBinary(OutOfBounds, CustomByName, ShortTypeSize, NoType)"), Validate({1, 0x80}, ECbFieldType::CustomByName), ECbValidateError::OutOfBounds);
 	TestEqual(TEXT("ValidateCompactBinary(OutOfBounds, CustomByName, ShortTypeName)"), Validate({uint8(ECbFieldType::CustomByName), 1, 1}), ECbValidateError::OutOfBounds);
 	TestEqual(TEXT("ValidateCompactBinary(OutOfBounds, CustomByName, ShortTypeName, NoType)"), Validate({1, 1}, ECbFieldType::CustomByName), ECbValidateError::OutOfBounds);
-	TestEqual(TEXT("ValidateCompactBinary(OutOfBounds, CustomByName, NoPayload)"), Validate({uint8(ECbFieldType::CustomByName), 3, 1, 'A'}), ECbValidateError::OutOfBounds);
-	TestEqual(TEXT("ValidateCompactBinary(OutOfBounds, CustomByName, NoPayload, NoType)"), Validate({3, 1, 'A'}, ECbFieldType::CustomByName), ECbValidateError::OutOfBounds);
-	TestEqual(TEXT("ValidateCompactBinary(OutOfBounds, CustomByName, ShortPayload)"), Validate({uint8(ECbFieldType::CustomByName), 4, 1, 'A', 0}), ECbValidateError::OutOfBounds);
-	TestEqual(TEXT("ValidateCompactBinary(OutOfBounds, CustomByName, ShortPayload, NoType)"), Validate({4, 1, 'A', 0}, ECbFieldType::CustomByName), ECbValidateError::OutOfBounds);
+	TestEqual(TEXT("ValidateCompactBinary(OutOfBounds, CustomByName, NoValue)"), Validate({uint8(ECbFieldType::CustomByName), 3, 1, 'A'}), ECbValidateError::OutOfBounds);
+	TestEqual(TEXT("ValidateCompactBinary(OutOfBounds, CustomByName, NoValue, NoType)"), Validate({3, 1, 'A'}, ECbFieldType::CustomByName), ECbValidateError::OutOfBounds);
+	TestEqual(TEXT("ValidateCompactBinary(OutOfBounds, CustomByName, ShortValue)"), Validate({uint8(ECbFieldType::CustomByName), 4, 1, 'A', 0}), ECbValidateError::OutOfBounds);
+	TestEqual(TEXT("ValidateCompactBinary(OutOfBounds, CustomByName, ShortValue, NoType)"), Validate({4, 1, 'A', 0}, ECbFieldType::CustomByName), ECbValidateError::OutOfBounds);
 
 	// Test InvalidType
 	TestEqual(TEXT("ValidateCompactBinary(InvalidType, Unknown)"), Validate({uint8(ECbFieldType::ObjectId) + 1}), ECbValidateError::InvalidType);
diff --git a/Engine/Source/Runtime/Core/Public/Serialization/CompactBinary.h b/Engine/Source/Runtime/Core/Public/Serialization/CompactBinary.h
--- a/Engine/Source/Runtime/Core/Public/Serialization/CompactBinary.h
+++ b/Engine/Source/Runtime/Core/Public/Serialization/CompactBinary.h
@@ -16,9 +16,9 @@
  * This file declares a compact binary data format that is compatible with JSON and only slightly
  * more expressive. The format is designed to achieve fairly small encoded sizes while also being
  * efficient to read both sequentially and through random access. An atom of data in this compact
- * binary format is called a field, which can be an object, array, byte string, character string,
- * integer, floating point value, boolean, null, uuid, date/time, time span, or a hash digest for
- * referencing blobs or compact binary data that are stored externally.
+ * binary format is called a field, which can be: an object, an array, a byte string, a character
+ * string, or one of several scalar types including integer, floating point, boolean, null, hash,
+ * uuid, date/time, time span, object identifier, or attachment reference.
  *
  * An object is a collection of name-field pairs, and an array is a collection of fields. Encoded
  * objects and arrays are both written such that they may be interpreted as a field that can then
@@ -93,116 +93,116 @@ enum class ECbFieldType : uint8
 	/** A field type that does not occur in a valid object. */
 	None                             = 0x00,
 
-	/** Null. Payload is empty. */
+	/** Null. Value is empty. */
 	Null                             = 0x01,
 
 	/**
 	 * Object is an array of fields with unique non-empty names.
 	 *
-	 * Payload is a VarUInt byte count for the encoded fields followed by the fields.
+	 * Value is a VarUInt byte count for the encoded fields, followed by the fields.
 	 */
 	Object                           = 0x02,
 	/**
 	 * UniformObject is an array of fields with the same field types and unique non-empty names.
 	 *
-	 * Payload is a VarUInt byte count for the encoded fields followed by the fields.
+	 * Value is a VarUInt byte count for the encoded fields, followed by the field type, followed
+	 * by the fields encoded with no type.
 	 */
 	UniformObject                    = 0x03,
 
 	/**
 	 * Array is an array of fields with no name that may be of different types.
 	 *
-	 * Payload is a VarUInt byte count, followed by a VarUInt item count, followed by the fields.
+	 * Value is a VarUInt byte count, followed by a VarUInt field count, followed by the fields.
 	 */
 	Array                            = 0x04,
 	/**
 	 * UniformArray is an array of fields with no name and with the same field type.
 	 *
-	 * Payload is a VarUInt byte count, followed by a VarUInt item count, followed by field type,
-	 * followed by the fields without their field type.
+	 * Value is a VarUInt byte count, followed by a VarUInt field count, followed by the field type,
+	 * followed by the fields encoded with no type.
 	 */
 	UniformArray                     = 0x05,
 
-	/** Binary. Payload is a VarUInt byte count followed by the data. */
+	/** Binary. Value is a VarUInt byte count followed by the data. */
 	Binary                           = 0x06,
 
-	/** String in UTF-8. Payload is a VarUInt byte count then an unterminated UTF-8 string. */
+	/** String in UTF-8. Value is a VarUInt byte count then an unterminated UTF-8 string. */
 	String                           = 0x07,
 
 	/**
 	 * Non-negative integer with the range of a 64-bit unsigned integer.
 	 *
-	 * Payload is the value encoded as a VarUInt.
+	 * Value is the value encoded as a VarUInt.
 	 */
 	IntegerPositive                  = 0x08,
 	/**
 	 * Negative integer with the range of a 64-bit signed integer.
 	 *
-	 * Payload is the ones' complement of the value encoded as a VarUInt.
+	 * Value is the ones' complement of the value encoded as a VarUInt.
 	 */
 	IntegerNegative                  = 0x09,
 
-	/** Single precision float. Payload is one big endian IEEE 754 binary32 float. */
+	/** Single precision float. Value is one big endian IEEE 754 binary32 float. */
 	Float32                          = 0x0a,
-	/** Double precision float. Payload is one big endian IEEE 754 binary64 float. */
+	/** Double precision float. Value is one big endian IEEE 754 binary64 float. */
 	Float64                          = 0x0b,
 
-	/** Boolean false value. Payload is empty. */
+	/** Boolean false value. Value is empty. */
 	BoolFalse                        = 0x0c,
-	/** Boolean true value. Payload is empty. */
+	/** Boolean true value. Value is empty. */
 	BoolTrue                         = 0x0d,
 
 	/**
 	 * ObjectAttachment is a reference to a compact binary object attachment stored externally.
 	 *
-	 * Payload is a 160-bit hash digest of the referenced compact binary object data.
+	 * Value is a 160-bit hash digest of the referenced compact binary object data.
 	 */
-	ObjectAttachment          = 0x0e,
+	ObjectAttachment                 = 0x0e,
 	/**
 	 * BinaryAttachment is a reference to a binary attachment stored externally.
 	 *
-	 * Payload is a 160-bit hash digest of the referenced binary data.
+	 * Value is a 160-bit hash digest of the referenced binary data.
 	 */
 	BinaryAttachment                 = 0x0f,
 
-	/** Hash. Payload is a 160-bit hash digest. */
+	/** Hash. Value is a 160-bit hash digest. */
 	Hash                             = 0x10,
-	/** UUID/GUID. Payload is a 128-bit UUID as defined by RFC 4122. */
+	/** UUID/GUID. Value is a 128-bit UUID as defined by RFC 4122. */
 	Uuid                             = 0x11,
 
 	/**
 	 * Date and time between 0001-01-01 00:00:00.0000000 and 9999-12-31 23:59:59.9999999.
 	 *
-	 * Payload is a big endian int64 count of 100ns ticks since 0001-01-01 00:00:00.0000000.
+	 * Value is a big endian int64 count of 100ns ticks since 0001-01-01 00:00:00.0000000.
 	 */
 	DateTime                         = 0x12,
 	/**
 	 * Difference between two date/time values.
 	 *
-	 * Payload is a big endian int64 count of 100ns ticks in the span, and may be negative.
+	 * Value is a big endian int64 count of 100ns ticks in the span, and may be negative.
 	 */
 	TimeSpan                         = 0x13,
 
 	/**
 	 * ObjectId is an opaque object identifier. See FCbObjectId.
 	 *
-	 * Payload is a 12-byte object identifier.
+	 * Value is a 12-byte object identifier.
 	 */
 	ObjectId                         = 0x14,
 
 	/**
-	 * CustomById identifies the sub-type of its payload by an integer identifier.
+	 * CustomById identifies the sub-type of its value by an integer identifier.
 	 *
-	 * Payload is a VarUInt byte count of the sub-type identifier and the sub-type payload, followed
-	 * by a VarUInt of the sub-type identifier then the payload of the sub-type.
+	 * Value is a VarUInt byte count, followed by a VarUInt encoding of the sub-type identifier,
+	 * followed by the value of the sub-type.
 	 */
 	CustomById                       = 0x1e,
 	/**
-	 * CustomByType identifies the sub-type of its payload by a string identifier.
+	 * CustomByType identifies the sub-type of its value by a string identifier.
 	 *
-	 * Payload is a VarUInt byte count of the sub-type identifier and the sub-type payload, followed
-	 * by a VarUInt byte count of the unterminated sub-type identifier, then the sub-type identifier
-	 * without termination, then the payload of the sub-type.
+	 * Value is a VarUInt byte count, followed by a String encoding of the sub-type identifier,
+	 * followed by the value of the sub-type.
 	 */
 	CustomByName                     = 0x1f,
 
@@ -211,13 +211,13 @@ enum class ECbFieldType : uint8
 
 	/**
 	 * A transient flag which indicates that the object or array containing this field has stored
-	 * the field type before the payload and name. Non-uniform objects and fields will set this.
+	 * the field type before the value and name. Non-uniform objects and fields will set this.
 	 *
-	 * Note: Since the flag must never be serialized, this bit may be repurposed in the future.
+	 * Note: Since the flag must never be serialized, this bit may be re-purposed in the future.
 	 */
 	HasFieldType                     = 0x40,
 
-	/** A persisted flag which indicates that the field has a name stored before the payload. */
+	/** A persisted flag which indicates that the field has a name stored before the value. */
 	HasFieldName                     = 0x80,
 };
 
@@ -362,7 +362,7 @@ public:
 	template <typename OtherFieldType>
 	constexpr inline bool Equals(const TCbFieldIterator<OtherFieldType>& Other) const
 	{
-		return FieldType::GetPayload() == Other.OtherFieldType::GetPayload() && FieldsEnd == Other.FieldsEnd;
+		return FieldType::GetValueData() == Other.OtherFieldType::GetValueData() && FieldsEnd == Other.FieldsEnd;
 	}
 
 	template <typename OtherFieldType>
@@ -414,7 +414,7 @@ protected:
 	/** Construct a field range that contains exactly one field. */
 	constexpr inline explicit TCbFieldIterator(FieldType InField)
 		: FieldType(MoveTemp(InField))
-		, FieldsEnd(FieldType::GetPayloadEnd())
+		, FieldsEnd(FieldType::GetValueEnd())
 	{
 	}
 
@@ -422,7 +422,7 @@ protected:
 	 * Construct a field range from the first field and a pointer to the end of the last field.
 	 *
 	 * @param InField The first field, or the default field if there are no fields.
-	 * @param InFieldsEnd A pointer to the end of the payload of the last field, or null.
+	 * @param InFieldsEnd A pointer to the end of the value of the last field, or null.
 	 */
 	constexpr inline TCbFieldIterator(FieldType&& InField, const void* InFieldsEnd)
 		: FieldType(MoveTemp(InField))
@@ -561,9 +561,9 @@ struct FCbCustomByName
  * the first byte, and includes the HasFieldName flag for a field in an object. The field name is
  * encoded in a variable-length unsigned integer of its size in bytes, for named fields, followed
  * by that many bytes of the UTF-8 encoding of the name with no null terminator. The remainder of
- * the field is the payload and is described in the field type enum. Every field must be uniquely
+ * the field is the value which is described in the field type enum. Every field must be uniquely
  * addressable when encoded, which means a zero-byte field is not permitted, and only arises in a
- * uniform array of fields with no payload, where the answer is to encode as a non-uniform array.
+ * uniform array of fields with no value, where the answer is to encode as a non-uniform array.
  *
  * This type only provides a view into memory and does not perform any memory management itself.
  * Use FCbField to hold a reference to the underlying memory when necessary.
@@ -585,7 +585,7 @@ public:
 	/** Returns the name of the field if it has a name, otherwise an empty view. */
 	constexpr inline FUtf8StringView GetName() const
 	{
-		return FUtf8StringView(static_cast<const UTF8CHAR*>(Payload) - NameLen, NameLen);
+		return FUtf8StringView(static_cast<const UTF8CHAR*>(Value) - NameLen, NameLen);
 	}
 
 	/** Access the field as an object. Defaults to an empty object on error. */
@@ -776,11 +776,11 @@ protected:
 	/** Returns a view of the field, including the type and name when present. */
 	CORE_API FMemoryView GetView() const;
 
-	/** Returns a view of the name and value payload, which excludes the type. */
+	/** Returns a view of the name and value, which excludes the type. */
 	CORE_API FMemoryView GetViewNoType() const;
 
-	/** Returns a view of the value payload, which excludes the type and name. */
-	inline FMemoryView GetPayloadView() const { return MakeMemoryView(Payload, GetPayloadSize()); }
+	/** Returns a view of the value, which excludes the type and name. */
+	inline FMemoryView GetValueView() const { return MakeMemoryView(Value, GetValueSize()); }
 
 	/** Returns the type of the field excluding flags. */
 	constexpr inline ECbFieldType GetType() const { return FCbFieldType::GetType(TypeWithFlags); }
@@ -788,14 +788,14 @@ protected:
 	/** Returns the type of the field including flags. */
 	constexpr inline ECbFieldType GetTypeWithFlags() const { return TypeWithFlags; }
 
-	/** Returns the start of the value payload. */
-	constexpr inline const void* GetPayload() const { return Payload; }
+	/** Returns the start of the value. */
+	constexpr inline const void* GetValueData() const { return Value; }
 
-	/** Returns the end of the value payload. */
-	inline const void* GetPayloadEnd() const { return static_cast<const uint8*>(Payload) + GetPayloadSize(); }
+	/** Returns the end of the value. */
+	inline const void* GetValueEnd() const { return static_cast<const uint8*>(Value) + GetValueSize(); }
 
-	/** Returns the size of the value payload in bytes, which is the field excluding the type and name. */
-	CORE_API uint64 GetPayloadSize() const;
+	/** Returns the size of the value in bytes, which is the field excluding the type and name. */
+	CORE_API uint64 GetValueSize() const;
 
 	/** Assign a field from a pointer to its data and an optional externally-provided type. */
 	inline void Assign(const void* InData, const ECbFieldType InType)
@@ -843,10 +843,10 @@ private:
 	ECbFieldType TypeWithFlags = ECbFieldType::None;
 	/** The error (if any) that occurred on the last field access. */
 	ECbFieldError Error = ECbFieldError::None;
-	/** The number of bytes for the name stored before the payload. */
+	/** The number of bytes for the name stored before the value. */
 	uint32 NameLen = 0;
-	/** The value payload, which also points to the end of the name. */
-	const void* Payload = nullptr;
+	/** The value, which also points to the end of the name. */
+	const void* Value = nullptr;
 };
 
 /**
diff --git a/Engine/Source/Runtime/Core/Public/Serialization/CompactBinaryValidation.h b/Engine/Source/Runtime/Core/Public/Serialization/CompactBinaryValidation.h
--- a/Engine/Source/Runtime/Core/Public/Serialization/CompactBinaryValidation.h
+++ b/Engine/Source/Runtime/Core/Public/Serialization/CompactBinaryValidation.h
@@ -88,15 +88,15 @@ enum class ECbValidateError : uint32
 
 	// Mode: Format
 
-	/** A name or string payload is not valid UTF-8. */
+	/** A name or string value is not valid UTF-8. */
 	InvalidString           = 1 << 5,
-	/** A size or integer payload can be encoded in fewer bytes. */
+	/** A size or integer value can be encoded in fewer bytes. */
 	InvalidInteger          = 1 << 6,
-	/** A float64 payload can be encoded as a float32 without loss of precision. */
+	/** A float64 value can be encoded as a float32 without loss of precision. */
 	InvalidFloat            = 1 << 7,
 	/** An object has the same type for every field but is not uniform. */
 	NonUniformObject        = 1 << 8,
-	/** An array has the same type for every field and non-empty payloads but is not uniform. */
+	/** An array has the same type for every field and non-empty values but is not uniform. */
 	NonUniformArray         = 1 << 9,
 
 	// Mode: Padding
