CompactBinary: Added FCbValue for low-level access to field values

#rnx

diff --git a/Engine/Source/Runtime/Core/Private/Serialization/CompactBinary.cpp b/Engine/Source/Runtime/Core/Private/Serialization/CompactBinary.cpp
--- a/Engine/Source/Runtime/Core/Private/Serialization/CompactBinary.cpp
+++ b/Engine/Source/Runtime/Core/Private/Serialization/CompactBinary.cpp
@@ -10,30 +10,19 @@
 #include "Misc/Guid.h"
 #include "Misc/Timespan.h"
 #include "Serialization/Archive.h"
+#include "Serialization/CompactBinaryValue.h"
 #include "Serialization/VarInt.h"
 #include "String/BytesToHex.h"
 
 ///////////////////////////////////////////////////////////////////////////////////////////////////////////////////////
 
-namespace CompactBinaryPrivate
+namespace UE::CompactBinary::Private
 {
 
 static constexpr const uint8 GEmptyObjectValue[] = { uint8(ECbFieldType::Object), 0x00 };
 static constexpr const uint8 GEmptyArrayValue[] = { uint8(ECbFieldType::Array), 0x01, 0x00 };
 
-template <typename T>
-static constexpr FORCEINLINE T ReadUnaligned(const void* const Memory)
-{
-#if PLATFORM_SUPPORTS_UNALIGNED_LOADS
-	return *static_cast<const T*>(Memory);
-#else
-	T Value;
-	FMemory::Memcpy(&Value, Memory, sizeof(Value));
-	return Value;
-#endif
-}
-
-}
+} // UE::CompactBinary::Private
 
 ///////////////////////////////////////////////////////////////////////////////////////////////////////////////////////
 
@@ -82,16 +71,6 @@ FCbObjectId::FCbObjectId(const FMemoryView ObjectId)
 	FMemory::Memcpy(Bytes, ObjectId.GetData(), sizeof(Bytes));
 }
 
-void FCbObjectId::ToString(FAnsiStringBuilderBase& Builder) const
-{
-	UE::String::BytesToHexLower(Bytes, Builder);
-}
-
-void FCbObjectId::ToString(FWideStringBuilderBase& Builder) const
-{
-	UE::String::BytesToHexLower(Bytes, Builder);
-}
-
 ///////////////////////////////////////////////////////////////////////////////////////////////////////////////////////
 
 FCbFieldView::FCbFieldView(const void* const Data, ECbFieldType Type)
@@ -137,14 +116,10 @@ FCbArrayView FCbFieldView::AsArrayView()
 
 FMemoryView FCbFieldView::AsBinaryView(const FMemoryView Default)
 {
-	if (FCbFieldType::IsBinary(TypeWithFlags))
+	if (FCbValue Accessor = GetValue(); FCbFieldType::IsBinary(Accessor.GetType()))
 	{
-		const uint8* const ValueBytes = static_cast<const uint8*>(Value);
-		uint32 ValueSizeByteCount;
-		const uint64 ValueSize = ReadVarUInt(ValueBytes, ValueSizeByteCount);
-
 		Error = ECbFieldError::None;
-		return MakeMemoryView(ValueBytes + ValueSizeByteCount, ValueSize);
+		return Accessor.AsBinary();
 	}
 	else
 	{
@@ -155,20 +130,9 @@ FMemoryView FCbFieldView::AsBinaryView(const FMemoryView Default)
 
 FUtf8StringView FCbFieldView::AsString(const FUtf8StringView Default)
 {
-	if (FCbFieldType::IsString(TypeWithFlags))
+	if (FCbValue Accessor = GetValue(); FCbFieldType::IsString(Accessor.GetType()))
 	{
-		const UTF8CHAR* const ValueChars = static_cast<const UTF8CHAR*>(Value);
-		uint32 ValueSizeByteCount;
-		const uint64 ValueSize = ReadVarUInt(ValueChars, ValueSizeByteCount);
-
-		if (ValueSize >= (uint64(1) << 31))
-		{
-			Error = ECbFieldError::RangeError;
-			return Default;
-		}
-
-		Error = ECbFieldError::None;
-		return FUtf8StringView(ValueChars + ValueSizeByteCount, int32(ValueSize));
+		return Accessor.AsString(&Error, Default);
 	}
 	else
 	{
@@ -177,23 +141,11 @@ FUtf8StringView FCbFieldView::AsString(const FUtf8StringView Default)
 	}
 }
 
-uint64 FCbFieldView::AsInteger(const uint64 Default, const FIntegerParams Params)
+uint64 FCbFieldView::AsInteger(const uint64 Default, const UE::CompactBinary::Private::FIntegerParams Params)
 {
-	if (FCbFieldType::IsInteger(TypeWithFlags))
+	if (FCbValue Accessor = GetValue(); FCbFieldType::IsInteger(Accessor.GetType()))
 	{
-		// A shift of a 64-bit value by 64 is undefined so shift by one less because magnitude is never zero.
-		const uint64 OutOfRangeMask = uint64(-2) << (Params.MagnitudeBits - 1);
-		const uint64 IsNegative = uint8(TypeWithFlags) & 1;
-
-		uint32 MagnitudeByteCount;
-		const uint64 Magnitude = ReadVarUInt(Value, MagnitudeByteCount);
-		const uint64 IntValue = Magnitude ^ -int64(IsNegative);
-
-		const uint64 IsInRange = (!(Magnitude & OutOfRangeMask)) & ((!IsNegative) | Params.IsSigned);
-		Error = IsInRange ? ECbFieldError::None : ECbFieldError::RangeError;
-
-		const uint64 UseValueMask = -int64(IsInRange);
-		return (IntValue & UseValueMask) | (Default & ~UseValueMask);
+		return Accessor.AsInteger(Params, &Error, Default);
 	}
 	else
 	{
@@ -204,16 +156,16 @@ uint64 FCbFieldView::AsInteger(const uint64 Default, const FIntegerParams Params
 
 float FCbFieldView::AsFloat(const float Default)
 {
-	switch (GetType())
+	switch (FCbValue Accessor = GetValue(); Accessor.GetType())
 	{
 	case ECbFieldType::IntegerPositive:
 	case ECbFieldType::IntegerNegative:
 	{
-		const uint64 IsNegative = uint8(TypeWithFlags) & 1;
+		const uint64 IsNegative = uint8(Accessor.GetType()) & 1;
 		constexpr uint64 OutOfRangeMask = ~((uint64(1) << /*FLT_MANT_DIG*/ 24) - 1);
 
 		uint32 MagnitudeByteCount;
-		const int64 Magnitude = ReadVarUInt(Value, MagnitudeByteCount) + IsNegative;
+		const int64 Magnitude = ReadVarUInt(Accessor.GetData(), MagnitudeByteCount) + IsNegative;
 		const uint64 IsInRange = !(Magnitude & OutOfRangeMask);
 		Error = IsInRange ? ECbFieldError::None : ECbFieldError::RangeError;
 		return IsInRange ? float(IsNegative ? -Magnitude : Magnitude) : Default;
@@ -221,8 +173,7 @@ float FCbFieldView::AsFloat(const float Default)
 	case ECbFieldType::Float32:
 	{
 		Error = ECbFieldError::None;
-		const uint32 FloatValue = NETWORK_ORDER32(CompactBinaryPrivate::ReadUnaligned<uint32>(Value));
-		return reinterpret_cast<const float&>(FloatValue);
+		return Accessor.AsFloat32();
 	}
 	case ECbFieldType::Float64:
 		Error = ECbFieldError::RangeError;
@@ -235,16 +186,16 @@ float FCbFieldView::AsFloat(const float Default)
 
 double FCbFieldView::AsDouble(const double Default)
 {
-	switch (GetType())
+	switch (FCbValue Accessor = GetValue(); Accessor.GetType())
 	{
 	case ECbFieldType::IntegerPositive:
 	case ECbFieldType::IntegerNegative:
 	{
-		const uint64 IsNegative = uint8(TypeWithFlags) & 1;
+		const uint64 IsNegative = uint8(Accessor.GetType()) & 1;
 		constexpr uint64 OutOfRangeMask = ~((uint64(1) << /*DBL_MANT_DIG*/ 53) - 1);
 
 		uint32 MagnitudeByteCount;
-		const int64 Magnitude = ReadVarUInt(Value, MagnitudeByteCount) + IsNegative;
+		const int64 Magnitude = ReadVarUInt(Accessor.GetData(), MagnitudeByteCount) + IsNegative;
 		const uint64 IsInRange = !(Magnitude & OutOfRangeMask);
 		Error = IsInRange ? ECbFieldError::None : ECbFieldError::RangeError;
 		return IsInRange ? double(IsNegative ? -Magnitude : Magnitude) : Default;
@@ -252,14 +203,12 @@ double FCbFieldView::AsDouble(const double Default)
 	case ECbFieldType::Float32:
 	{
 		Error = ECbFieldError::None;
-		const uint32 FloatValue = NETWORK_ORDER32(CompactBinaryPrivate::ReadUnaligned<uint32>(Value));
-		return reinterpret_cast<const float&>(FloatValue);
+		return Accessor.AsFloat32();
 	}
 	case ECbFieldType::Float64:
 	{
 		Error = ECbFieldError::None;
-		const uint64 FloatValue = NETWORK_ORDER64(CompactBinaryPrivate::ReadUnaligned<uint64>(Value));
-		return reinterpret_cast<const double&>(FloatValue);
+		return Accessor.AsFloat64();
 	}
 	default:
 		Error = ECbFieldError::TypeError;
@@ -269,18 +218,18 @@ double FCbFieldView::AsDouble(const double Default)
 
 bool FCbFieldView::AsBool(const bool bDefault)
 {
-	const ECbFieldType LocalTypeWithFlags = TypeWithFlags;
-	const bool bIsBool = FCbFieldType::IsBool(LocalTypeWithFlags);
+	FCbValue Accessor = GetValue();
+	const bool bIsBool = FCbFieldType::IsBool(Accessor.GetType());
 	Error = bIsBool ? ECbFieldError::None : ECbFieldError::TypeError;
-	return (uint8(bIsBool) & uint8(LocalTypeWithFlags) & 1) | ((!bIsBool) & bDefault);
+	return (uint8(bIsBool) & Accessor.AsBool()) | ((!bIsBool) & bDefault);
 }
 
 FIoHash FCbFieldView::AsObjectAttachment(const FIoHash& Default)
 {
-	if (FCbFieldType::IsObjectAttachment(TypeWithFlags))
+	if (FCbValue Accessor = GetValue(); FCbFieldType::IsObjectAttachment(Accessor.GetType()))
 	{
 		Error = ECbFieldError::None;
-		return FIoHash(*static_cast<const FIoHash::ByteArray*>(Value));
+		return Accessor.AsObjectAttachment();
 	}
 	else
 	{
@@ -291,10 +240,10 @@ FIoHash FCbFieldView::AsObjectAttachment(const FIoHash& Default)
 
 FIoHash FCbFieldView::AsBinaryAttachment(const FIoHash& Default)
 {
-	if (FCbFieldType::IsBinaryAttachment(TypeWithFlags))
+	if (FCbValue Accessor = GetValue(); FCbFieldType::IsBinaryAttachment(Accessor.GetType()))
 	{
 		Error = ECbFieldError::None;
-		return FIoHash(*static_cast<const FIoHash::ByteArray*>(Value));
+		return Accessor.AsBinaryAttachment();
 	}
 	else
 	{
@@ -305,10 +254,10 @@ FIoHash FCbFieldView::AsBinaryAttachment(const FIoHash& Default)
 
 FIoHash FCbFieldView::AsAttachment(const FIoHash& Default)
 {
-	if (FCbFieldType::IsAttachment(TypeWithFlags))
+	if (FCbValue Accessor = GetValue(); FCbFieldType::IsAttachment(Accessor.GetType()))
 	{
 		Error = ECbFieldError::None;
-		return FIoHash(*static_cast<const FIoHash::ByteArray*>(Value));
+		return Accessor.AsAttachment();
 	}
 	else
 	{
@@ -319,10 +268,10 @@ FIoHash FCbFieldView::AsAttachment(const FIoHash& Default)
 
 FIoHash FCbFieldView::AsHash(const FIoHash& Default)
 {
-	if (FCbFieldType::IsHash(TypeWithFlags))
+	if (FCbValue Accessor = GetValue(); FCbFieldType::IsHash(Accessor.GetType()))
 	{
 		Error = ECbFieldError::None;
-		return FIoHash(*static_cast<const FIoHash::ByteArray*>(Value));
+		return Accessor.AsHash();
 	}
 	else
 	{
@@ -338,16 +287,10 @@ FGuid FCbFieldView::AsUuid()
 
 FGuid FCbFieldView::AsUuid(const FGuid& Default)
 {
-	if (FCbFieldType::IsUuid(TypeWithFlags))
+	if (FCbValue Accessor = GetValue(); FCbFieldType::IsUuid(Accessor.GetType()))
 	{
 		Error = ECbFieldError::None;
-		FGuid UuidValue;
-		FMemory::Memcpy(&UuidValue, Value, sizeof(FGuid));
-		UuidValue.A = NETWORK_ORDER32(UuidValue.A);
-		UuidValue.B = NETWORK_ORDER32(UuidValue.B);
-		UuidValue.C = NETWORK_ORDER32(UuidValue.C);
-		UuidValue.D = NETWORK_ORDER32(UuidValue.D);
-		return UuidValue;
+		return Accessor.AsUuid();
 	}
 	else
 	{
@@ -358,10 +301,10 @@ FGuid FCbFieldView::AsUuid(const FGuid& Default)
 
 int64 FCbFieldView::AsDateTimeTicks(const int64 Default)
 {
-	if (FCbFieldType::IsDateTime(TypeWithFlags))
+	if (FCbValue Accessor = GetValue(); FCbFieldType::IsDateTime(Accessor.GetType()))
 	{
 		Error = ECbFieldError::None;
-		return NETWORK_ORDER64(CompactBinaryPrivate::ReadUnaligned<int64>(Value));
+		return Accessor.AsDateTimeTicks();
 	}
 	else
 	{
@@ -382,10 +325,10 @@ FDateTime FCbFieldView::AsDateTime(FDateTime Default)
 
 int64 FCbFieldView::AsTimeSpanTicks(const int64 Default)
 {
-	if (FCbFieldType::IsTimeSpan(TypeWithFlags))
+	if (FCbValue Accessor = GetValue(); FCbFieldType::IsTimeSpan(Accessor.GetType()))
 	{
 		Error = ECbFieldError::None;
-		return NETWORK_ORDER64(CompactBinaryPrivate::ReadUnaligned<int64>(Value));
+		return Accessor.AsTimeSpanTicks();
 	}
 	else
 	{
@@ -407,10 +350,10 @@ FTimespan FCbFieldView::AsTimeSpan(FTimespan Default)
 FCbObjectId FCbFieldView::AsObjectId(const FCbObjectId& Default)
 {
 	static_assert(sizeof(FCbObjectId) == 12, "FCbObjectId is expected to be 12 bytes.");
-	if (FCbFieldType::IsObjectId(TypeWithFlags))
+	if (FCbValue Accessor = GetValue(); FCbFieldType::IsObjectId(Accessor.GetType()))
 	{
 		Error = ECbFieldError::None;
-		return FCbObjectId(MakeMemoryView(Value, sizeof(FCbObjectId)));
+		return Accessor.AsObjectId();
 	}
 	else
 	{
@@ -421,19 +364,10 @@ FCbObjectId FCbFieldView::AsObjectId(const FCbObjectId& Default)
 
 FCbCustomById FCbFieldView::AsCustomById(FCbCustomById Default)
 {
-	if (FCbFieldType::IsCustomById(TypeWithFlags))
+	if (FCbValue Accessor = GetValue(); FCbFieldType::IsCustomById(Accessor.GetType()))
 	{
-		const uint8* ValueBytes = static_cast<const uint8*>(Value);
-		uint32 ValueSizeByteCount;
-		const uint64 ValueSize = ReadVarUInt(ValueBytes, ValueSizeByteCount);
-		ValueBytes += ValueSizeByteCount;
-
-		FCbCustomById CustomValue;
-		uint32 TypeIdByteCount;
-		CustomValue.Id = ReadVarUInt(ValueBytes, TypeIdByteCount);
-		CustomValue.Data = MakeMemoryView(ValueBytes + TypeIdByteCount, ValueSize - TypeIdByteCount);
 		Error = ECbFieldError::None;
-		return CustomValue;
+		return Accessor.AsCustomById();
 	}
 	else
 	{
@@ -444,24 +378,10 @@ FCbCustomById FCbFieldView::AsCustomById(FCbCustomById Default)
 
 FCbCustomByName FCbFieldView::AsCustomByName(FCbCustomByName Default)
 {
-	if (FCbFieldType::IsCustomByName(TypeWithFlags))
+	if (FCbValue Accessor = GetValue(); FCbFieldType::IsCustomByName(Accessor.GetType()))
 	{
-		const uint8* ValueBytes = static_cast<const uint8*>(Value);
-		uint32 ValueSizeByteCount;
-		const uint64 ValueSize = ReadVarUInt(ValueBytes, ValueSizeByteCount);
-		ValueBytes += ValueSizeByteCount;
-
-		uint32 TypeNameLenByteCount;
-		const uint64 TypeNameLen = ReadVarUInt(ValueBytes, TypeNameLenByteCount);
-		ValueBytes += TypeNameLenByteCount;
-
-		FCbCustomByName CustomValue;
-		CustomValue.Name = FUtf8StringView(
-			reinterpret_cast<const UTF8CHAR*>(ValueBytes),
-			static_cast<FUtf8StringView::SizeType>(TypeNameLen));
-		CustomValue.Data = MakeMemoryView(ValueBytes + TypeNameLen, ValueSize - TypeNameLen - TypeNameLenByteCount);
 		Error = ECbFieldError::None;
-		return CustomValue;
+		return Accessor.AsCustomByName();
 	}
 	else
 	{
@@ -658,7 +578,7 @@ FCbFieldViewIterator FCbFieldView::CreateViewIterator() const
 ///////////////////////////////////////////////////////////////////////////////////////////////////////////////////////
 
 FCbArrayView::FCbArrayView()
-	: FCbFieldView(CompactBinaryPrivate::GEmptyArrayValue)
+	: FCbFieldView(UE::CompactBinary::Private::GEmptyArrayValue)
 {
 }
 
@@ -717,7 +637,7 @@ void FCbArrayView::CopyTo(FArchive& Ar) const
 ///////////////////////////////////////////////////////////////////////////////////////////////////////////////////////
 
 FCbObjectView::FCbObjectView()
-	: FCbFieldView(CompactBinaryPrivate::GEmptyObjectValue)
+	: FCbFieldView(UE::CompactBinary::Private::GEmptyObjectValue)
 {
 }
 
@@ -747,7 +667,7 @@ FCbFieldView FCbObjectView::FindViewIgnoreCase(const FUtf8StringView Name) const
 
 FCbObjectView::operator bool() const
 {
-	return GetSize() > sizeof(CompactBinaryPrivate::GEmptyObjectValue);
+	return GetSize() > sizeof(UE::CompactBinary::Private::GEmptyObjectValue);
 }
 
 uint64 FCbObjectView::GetSize() const
diff --git a/Engine/Source/Runtime/Core/Public/Serialization/CompactBinary.h b/Engine/Source/Runtime/Core/Public/Serialization/CompactBinary.h
--- a/Engine/Source/Runtime/Core/Public/Serialization/CompactBinary.h
+++ b/Engine/Source/Runtime/Core/Public/Serialization/CompactBinary.h
@@ -10,6 +10,7 @@
 #include "Memory/MemoryView.h"
 #include "Memory/SharedBuffer.h"
 #include "Misc/EnumClassFlags.h"
+#include "String/BytesToHex.h"
 #include "Templates/IsTriviallyDestructible.h"
 
 /**
@@ -71,6 +72,7 @@ class FCbFieldView;
 class FCbFieldViewIterator;
 class FCbObjectId;
 class FCbObjectView;
+class FCbValue;
 struct FDateTime;
 struct FGuid;
 struct FTimespan;
@@ -479,10 +481,6 @@ public:
 	/** Construct an ObjectId from a view of 12 bytes. */
 	CORE_API explicit FCbObjectId(FMemoryView ObjectId);
 
-	/** Convert the ObjectId to a 24-character hex string. */
-	CORE_API void ToString(FAnsiStringBuilderBase& Builder) const;
-	CORE_API void ToString(FWideStringBuilderBase& Builder) const;
-
 	/** Returns a reference to the raw byte array for the ObjectId. */
 	inline const ByteArray& GetBytes() const { return Bytes; }
 	inline operator const ByteArray&() const { return Bytes; }
@@ -519,10 +517,11 @@ inline uint32 GetTypeHash(const FCbObjectId& Id)
 	return *reinterpret_cast<const uint32*>(&Id);
 }
 
+/** Convert the ObjectId to a 24-character hex string. */
 template <typename CharType>
 inline TStringBuilderBase<CharType>& operator<<(TStringBuilderBase<CharType>& Builder, const FCbObjectId& Id)
 {
-	Id.ToString(Builder);
+	UE::String::BytesToHexLower(Id.GetBytes(), Builder);
 	return Builder;
 }
 
@@ -548,6 +547,32 @@ struct FCbCustomByName
 
 ///////////////////////////////////////////////////////////////////////////////////////////////////////////////////////
 
+namespace UE::CompactBinary::Private
+{
+
+/** Parameters for converting to an integer. */
+struct FIntegerParams
+{
+	/** Whether the output type has a sign bit. */
+	uint32 IsSigned : 1;
+	/** Bits of magnitude. (7 for int8) */
+	uint32 MagnitudeBits : 31;
+};
+
+/** Make integer params for the given integer type. */
+template <typename IntType>
+static constexpr inline FIntegerParams MakeIntegerParams()
+{
+	FIntegerParams Params;
+	Params.IsSigned = IntType(-1) < IntType(0);
+	Params.MagnitudeBits = 8 * sizeof(IntType) - Params.IsSigned;
+	return Params;
+}
+
+} // UE::CompactBinary::Private
+
+///////////////////////////////////////////////////////////////////////////////////////////////////////////////////////
+
 /**
  * An atom of data in the compact binary format.
  *
@@ -582,12 +607,18 @@ public:
 	 */
 	CORE_API explicit FCbFieldView(const void* Data, ECbFieldType Type = ECbFieldType::HasFieldType);
 
+	/** Construct a field from a value, without access to the name. */
+	inline explicit FCbFieldView(const FCbValue& Value);
+
 	/** Returns the name of the field if it has a name, otherwise an empty view. */
 	constexpr inline FUtf8StringView GetName() const
 	{
 		return FUtf8StringView(static_cast<const UTF8CHAR*>(Value) - NameLen, NameLen);
 	}
 
+	/** Returns the value for unchecked access. Prefer the typed accessors below. */
+	inline FCbValue GetValue() const;
+
 	/** Access the field as an object. Defaults to an empty object on error. */
 	CORE_API FCbObjectView AsObjectView();
 
@@ -806,25 +837,6 @@ protected:
 	}
 
 private:
-	/** Parameters for converting to an integer. */
-	struct FIntegerParams
-	{
-		/** Whether the output type has a sign bit. */
-		uint32 IsSigned : 1;
-		/** Bits of magnitude. (7 for int8) */
-		uint32 MagnitudeBits : 31;
-	};
-
-	/** Make integer params for the given integer type. */
-	template <typename IntType>
-	static constexpr inline FIntegerParams MakeIntegerParams()
-	{
-		FIntegerParams Params;
-		Params.IsSigned = IntType(-1) < IntType(0);
-		Params.MagnitudeBits = 8 * sizeof(IntType) - Params.IsSigned;
-		return Params;
-	}
-
 	/**
 	 * Access the field as the given integer type.
 	 *
@@ -833,10 +845,10 @@ private:
 	template <typename IntType>
 	inline IntType AsInteger(IntType Default)
 	{
-		return IntType(AsInteger(uint64(Default), MakeIntegerParams<IntType>()));
+		return IntType(AsInteger(uint64(Default), UE::CompactBinary::Private::MakeIntegerParams<IntType>()));
 	}
 
-	CORE_API uint64 AsInteger(uint64 Default, FIntegerParams Params);
+	CORE_API uint64 AsInteger(uint64 Default, UE::CompactBinary::Private::FIntegerParams Params);
 
 private:
 	/** The field type, with the transient HasFieldType flag if the field contains its type. */
diff --git a/Engine/Source/Runtime/Core/Public/Serialization/CompactBinaryValue.h b/Engine/Source/Runtime/Core/Public/Serialization/CompactBinaryValue.h
new file mode 100644
--- /dev/null
+++ b/Engine/Source/Runtime/Core/Public/Serialization/CompactBinaryValue.h
@@ -0,0 +1,256 @@
+// Copyright Epic Games, Inc. All Rights Reserved.
+
+#pragma once
+
+#include "Containers/StringView.h"
+#include "HAL/UnrealMemory.h"
+#include "IO/IoHash.h"
+#include "Memory/MemoryView.h"
+#include "Misc/ByteSwap.h"
+#include "Misc/Guid.h"
+#include "Serialization/CompactBinary.h"
+#include "Serialization/VarInt.h"
+
+namespace UE::CompactBinary::Private
+{
+
+template <typename T>
+static constexpr FORCEINLINE T ReadUnaligned(const void* const Memory)
+{
+#if PLATFORM_SUPPORTS_UNALIGNED_LOADS
+	return *static_cast<const T*>(Memory);
+#else
+	T Value;
+	FMemory::Memcpy(&Value, Memory, sizeof(Value));
+	return Value;
+#endif
+}
+
+} // UE::CompactBinary::Private
+
+/**
+ * A type that provides unchecked access to compact binary values.
+ *
+ * The main purpose of the type is to efficiently switch on field type. For every other use case,
+ * prefer to use the field, array, and object types directly. The accessors here do not check the
+ * type before reading the value, which means they can read out of bounds even on a valid compact
+ * binary value if the wrong accessor is used.
+ */
+class FCbValue
+{
+public:
+	FCbValue(ECbFieldType Type, const void* Value);
+
+	FCbObjectView AsObjectView() const;
+	FCbArrayView AsArrayView() const;
+
+	FMemoryView AsBinary() const;
+	/** Access as a string. Checks for range errors and uses the default if OutError is not null. */
+	FUtf8StringView AsString(ECbFieldError* OutError = nullptr, FUtf8StringView Default = FUtf8StringView()) const;
+
+	/**
+	 * Access as an integer, with both positive and negative values returned as unsigned.
+	 *
+	 * Checks for range errors and uses the default if OutError is not null.
+	 */
+	uint64 AsInteger(UE::CompactBinary::Private::FIntegerParams Params, ECbFieldError* OutError = nullptr, uint64 Default = 0) const;
+
+	uint64 AsIntegerPositive() const;
+	int64 AsIntegerNegative() const;
+
+	float AsFloat32() const;
+	double AsFloat64() const;
+
+	bool AsBool() const;
+
+	FORCEINLINE FIoHash AsObjectAttachment() const { return AsHash(); }
+	FORCEINLINE FIoHash AsBinaryAttachment() const { return AsHash(); }
+	FORCEINLINE FIoHash AsAttachment() const { return AsHash(); }
+
+	FIoHash AsHash() const;
+	FGuid AsUuid() const;
+
+	int64 AsDateTimeTicks() const;
+	int64 AsTimeSpanTicks() const;
+
+	FCbObjectId AsObjectId() const;
+
+	FCbCustomById AsCustomById() const;
+	FCbCustomByName AsCustomByName() const;
+
+	FORCEINLINE ECbFieldType GetType() const { return Type; }
+	FORCEINLINE const void* GetData() const { return Data; }
+
+private:
+	const void* Data;
+	ECbFieldType Type;
+};
+
+FORCEINLINE FCbFieldView::FCbFieldView(const FCbValue& InValue)
+	: TypeWithFlags(InValue.GetType())
+	, Value(InValue.GetData())
+{
+}
+
+FORCEINLINE FCbValue FCbFieldView::GetValue() const
+{
+	return FCbValue(FCbFieldType::GetType(TypeWithFlags), Value);
+}
+
+FORCEINLINE FCbValue::FCbValue(ECbFieldType InType, const void* InValue)
+	: Data(InValue)
+	, Type(InType)
+{
+}
+
+FORCEINLINE FCbObjectView FCbValue::AsObjectView() const
+{
+	return FCbObjectView(*this);
+}
+
+FORCEINLINE FCbArrayView FCbValue::AsArrayView() const
+{
+	return FCbArrayView(*this);
+}
+
+FORCEINLINE FMemoryView FCbValue::AsBinary() const
+{
+	const uint8* const Bytes = static_cast<const uint8*>(Data);
+	uint32 ValueSizeByteCount;
+	const uint64 ValueSize = ReadVarUInt(Bytes, ValueSizeByteCount);
+	return MakeMemoryView(Bytes + ValueSizeByteCount, ValueSize);
+}
+
+FORCEINLINE FUtf8StringView FCbValue::AsString(ECbFieldError* OutError, FUtf8StringView Default) const
+{
+	const UTF8CHAR* const Chars = static_cast<const UTF8CHAR*>(Data);
+	uint32 ValueSizeByteCount;
+	const uint64 ValueSize = ReadVarUInt(Chars, ValueSizeByteCount);
+
+	if (OutError)
+	{
+		if (ValueSize >= (uint64(1) << 31))
+		{
+			*OutError = ECbFieldError::RangeError;
+			return Default;
+		}
+		*OutError = ECbFieldError::None;
+	}
+
+	return FUtf8StringView(Chars + ValueSizeByteCount, int32(ValueSize));
+}
+
+FORCEINLINE uint64 FCbValue::AsInteger(UE::CompactBinary::Private::FIntegerParams Params, ECbFieldError* OutError, uint64 Default) const
+{
+	// A shift of a 64-bit value by 64 is undefined so shift by one less because magnitude is never zero.
+	const uint64 OutOfRangeMask = uint64(-2) << (Params.MagnitudeBits - 1);
+	const uint64 IsNegative = uint8(Type) & 1;
+
+	uint32 MagnitudeByteCount;
+	const uint64 Magnitude = ReadVarUInt(Data, MagnitudeByteCount);
+	const uint64 Value = Magnitude ^ -int64(IsNegative);
+
+	if (OutError)
+	{
+		const uint64 IsInRange = (!(Magnitude & OutOfRangeMask)) & ((!IsNegative) | Params.IsSigned);
+		*OutError = IsInRange ? ECbFieldError::None : ECbFieldError::RangeError;
+
+		const uint64 UseValueMask = -int64(IsInRange);
+		return (Value & UseValueMask) | (Default & ~UseValueMask);
+	}
+
+	return Value;
+}
+
+FORCEINLINE uint64 FCbValue::AsIntegerPositive() const
+{
+	uint32 MagnitudeByteCount;
+	return ReadVarUInt(Data, MagnitudeByteCount);
+}
+
+FORCEINLINE int64 FCbValue::AsIntegerNegative() const
+{
+	uint32 MagnitudeByteCount;
+	return int64(ReadVarUInt(Data, MagnitudeByteCount)) ^ -int64(1);
+}
+
+FORCEINLINE float FCbValue::AsFloat32() const
+{
+	const uint32 Value = NETWORK_ORDER32(UE::CompactBinary::Private::ReadUnaligned<uint32>(Data));
+	return reinterpret_cast<const float&>(Value);
+}
+
+FORCEINLINE double FCbValue::AsFloat64() const
+{
+	const uint64 Value = NETWORK_ORDER64(UE::CompactBinary::Private::ReadUnaligned<uint64>(Data));
+	return reinterpret_cast<const double&>(Value);
+}
+
+FORCEINLINE bool FCbValue::AsBool() const
+{
+	return uint8(Type) & 1;
+}
+
+FORCEINLINE FIoHash FCbValue::AsHash() const
+{
+	return FIoHash(*static_cast<const FIoHash::ByteArray*>(Data));
+}
+
+FORCEINLINE FGuid FCbValue::AsUuid() const
+{
+	FGuid Value;
+	FMemory::Memcpy(&Value, Data, sizeof(FGuid));
+	Value.A = NETWORK_ORDER32(Value.A);
+	Value.B = NETWORK_ORDER32(Value.B);
+	Value.C = NETWORK_ORDER32(Value.C);
+	Value.D = NETWORK_ORDER32(Value.D);
+	return Value;
+}
+
+FORCEINLINE int64 FCbValue::AsDateTimeTicks() const
+{
+	return NETWORK_ORDER64(UE::CompactBinary::Private::ReadUnaligned<int64>(Data));
+}
+
+FORCEINLINE int64 FCbValue::AsTimeSpanTicks() const
+{
+	return NETWORK_ORDER64(UE::CompactBinary::Private::ReadUnaligned<int64>(Data));
+}
+
+FORCEINLINE FCbObjectId FCbValue::AsObjectId() const
+{
+	return FCbObjectId(MakeMemoryView(Data, sizeof(FCbObjectId)));
+}
+
+FORCEINLINE FCbCustomById FCbValue::AsCustomById() const
+{
+	const uint8* Bytes = static_cast<const uint8*>(Data);
+	uint32 DataSizeByteCount;
+	const uint64 DataSize = ReadVarUInt(Bytes, DataSizeByteCount);
+	Bytes += DataSizeByteCount;
+
+	FCbCustomById Value;
+	uint32 TypeIdByteCount;
+	Value.Id = ReadVarUInt(Bytes, TypeIdByteCount);
+	Value.Data = MakeMemoryView(Bytes + TypeIdByteCount, DataSize - TypeIdByteCount);
+	return Value;
+}
+
+FORCEINLINE FCbCustomByName FCbValue::AsCustomByName() const
+{
+	const uint8* Bytes = static_cast<const uint8*>(Data);
+	uint32 DataSizeByteCount;
+	const uint64 DataSize = ReadVarUInt(Bytes, DataSizeByteCount);
+	Bytes += DataSizeByteCount;
+
+	uint32 TypeNameLenByteCount;
+	const uint64 TypeNameLen = ReadVarUInt(Bytes, TypeNameLenByteCount);
+	Bytes += TypeNameLenByteCount;
+
+	FCbCustomByName Value;
+	Value.Name = FUtf8StringView(
+		reinterpret_cast<const UTF8CHAR*>(Bytes),
+		static_cast<FUtf8StringView::SizeType>(TypeNameLen));
+	Value.Data = MakeMemoryView(Bytes + TypeNameLen, DataSize - TypeNameLen - TypeNameLenByteCount);
+	return Value;
+}
