CompactBinary: Added CompactBinaryToJson

#rnx

diff --git a/Engine/Source/Runtime/Core/Private/Serialization/CompactBinarySerialization.cpp b/Engine/Source/Runtime/Core/Private/Serialization/CompactBinarySerialization.cpp
--- a/Engine/Source/Runtime/Core/Private/Serialization/CompactBinarySerialization.cpp
+++ b/Engine/Source/Runtime/Core/Private/Serialization/CompactBinarySerialization.cpp
@@ -2,7 +2,16 @@
 
 #include "Serialization/CompactBinarySerialization.h"
 
+#include "Containers/StringConv.h"
+#include "HAL/Platform.h"
+#include "Misc/AsciiSet.h"
+#include "Misc/Base64.h"
+#include "Misc/DateTime.h"
+#include "Misc/Guid.h"
+#include "Misc/StringBuilder.h"
+#include "Misc/Timespan.h"
 #include "Serialization/CompactBinaryValidation.h"
+#include "Serialization/CompactBinaryValue.h"
 #include "Serialization/VarInt.h"
 #include "Templates/IdentityFunctor.h"
 #include "Templates/Invoke.h"
@@ -255,3 +264,221 @@ FArchive& operator<<(FArchive& Ar, FCbObject& Object)
 }
 
 ///////////////////////////////////////////////////////////////////////////////////////////////////////////////////////
+
+class FCbJsonWriter
+{
+public:
+	explicit FCbJsonWriter(FUtf8StringBuilderBase& InBuilder)
+		: Builder(InBuilder)
+	{
+		NewLineAndIndent << LINE_TERMINATOR_ANSI;
+	}
+
+	void WriteField(FCbFieldView Field)
+	{
+		WriteOptionalComma();
+		WriteOptionalNewLine();
+
+		if (FUtf8StringView Name = Field.GetName(); !Name.IsEmpty())
+		{
+			AppendQuotedString(Name);
+			Builder << ": "_ASV;
+		}
+
+		switch (FCbValue Accessor = Field.GetValue(); Accessor.GetType())
+		{
+		case ECbFieldType::Null:
+			Builder << "null"_ASV;
+			break;
+		case ECbFieldType::Object:
+		case ECbFieldType::UniformObject:
+			Builder << '{';
+			NewLineAndIndent << '\t';
+			bNeedsNewLine = true;
+			for (FCbFieldView It : Field)
+			{
+				WriteField(It);
+			}
+			NewLineAndIndent.RemoveSuffix(1);
+			if (bNeedsComma)
+			{
+				WriteOptionalNewLine();
+			}
+			Builder << '}';
+			break;
+		case ECbFieldType::Array:
+		case ECbFieldType::UniformArray:
+			Builder << '[';
+			NewLineAndIndent << '\t';
+			bNeedsNewLine = true;
+			for (FCbFieldView It : Field)
+			{
+				WriteField(It);
+			}
+			NewLineAndIndent.RemoveSuffix(1);
+			if (bNeedsComma)
+			{
+				WriteOptionalNewLine();
+			}
+			Builder << ']';
+			break;
+		case ECbFieldType::Binary:
+			AppendBase64String(Accessor.AsBinary());
+			break;
+		case ECbFieldType::String:
+			AppendQuotedString(Accessor.AsString());
+			break;
+		case ECbFieldType::IntegerPositive:
+			Builder << Accessor.AsIntegerPositive();
+			break;
+		case ECbFieldType::IntegerNegative:
+			Builder << Accessor.AsIntegerNegative();
+			break;
+		case ECbFieldType::Float32:
+			Builder.Appendf(UTF8TEXT("%.9g"), Accessor.AsFloat32());
+			break;
+		case ECbFieldType::Float64:
+			Builder.Appendf(UTF8TEXT("%.17g"), Accessor.AsFloat64());
+			break;
+		case ECbFieldType::BoolFalse:
+			Builder << "false"_ASV;
+			break;
+		case ECbFieldType::BoolTrue:
+			Builder << "true"_ASV;
+			break;
+		case ECbFieldType::ObjectAttachment:
+		case ECbFieldType::BinaryAttachment:
+			Builder << '"' << Accessor.AsAttachment() << '"';
+			break;
+		case ECbFieldType::Hash:
+			Builder << '"' << Accessor.AsHash() << '"';
+			break;
+		case ECbFieldType::Uuid:
+			Builder << '"' << Accessor.AsUuid() << '"';
+			break;
+		case ECbFieldType::DateTime:
+			Builder << '"' << FTCHARToUTF8(FDateTime(Accessor.AsDateTimeTicks()).ToIso8601()) << '"';
+			break;
+		case ECbFieldType::TimeSpan:
+		{
+			const FTimespan Span(Accessor.AsTimeSpanTicks());
+			if (Span.GetDays() == 0)
+			{
+				Builder << '"' << FTCHARToUTF8(Span.ToString(TEXT("%h:%m:%s.%n"))) << '"';
+			}
+			else
+			{
+				Builder << '"' << FTCHARToUTF8(Span.ToString(TEXT("%d.%h:%m:%s.%n"))) << '"';
+			}
+			break;
+		}
+		case ECbFieldType::ObjectId:
+			Builder << '"' << Accessor.AsUuid() << '"';
+			break;
+		case ECbFieldType::CustomById:
+		{
+			FCbCustomById Custom = Accessor.AsCustomById();
+			Builder << "{ \"Id\": ";
+			Builder << Custom.Id;
+			Builder << ", \"Data\": ";
+			AppendBase64String(Custom.Data);
+			Builder << " }";
+			break;
+		}
+		case ECbFieldType::CustomByName:
+		{
+			FCbCustomByName Custom = Accessor.AsCustomByName();
+			Builder << "{ \"Name\": ";
+			AppendQuotedString(Custom.Name);
+			Builder << ", \"Data\": ";
+			AppendBase64String(Custom.Data);
+			Builder << " }";
+			break;
+		}
+		default:
+			checkNoEntry();
+			break;
+		}
+
+		bNeedsComma = true;
+		bNeedsNewLine = true;
+	}
+
+private:
+	void WriteOptionalComma()
+	{
+		if (bNeedsComma)
+		{
+			bNeedsComma = false;
+			Builder << ',';
+		}
+	}
+
+	void WriteOptionalNewLine()
+	{
+		if (bNeedsNewLine)
+		{
+			bNeedsNewLine = false;
+			Builder << NewLineAndIndent;
+		}
+	}
+
+	void AppendQuotedString(FUtf8StringView Value)
+	{
+		const FAsciiSet EscapeSet("\\\"\b\f\n\r\t"
+			"\x00\x01\x02\x03\x04\x05\x06\x07\x08\x09\x0a\x0b\x0c\x0d\x0e\x0f"
+			"\x10\x11\x12\x13\x14\x15\x16\x17\x18\x19\x1a\x1b\x1c\x1d\x1e\x1f");
+		Builder << '\"';
+		while (!Value.IsEmpty())
+		{
+			FUtf8StringView Verbatim = FAsciiSet::FindPrefixWithout(Value, EscapeSet);
+			Builder << Verbatim;
+			Value.RightChopInline(Verbatim.Len());
+			FUtf8StringView Escape = FAsciiSet::FindPrefixWith(Value, EscapeSet);
+			for (UTF8CHAR Char : Escape)
+			{
+				switch (Char)
+				{
+				case '\\': Builder << "\\\\"_ASV; break;
+				case '\"': Builder << "\\\""_ASV; break;
+				case '\b': Builder << "\\b"_ASV; break;
+				case '\f': Builder << "\\f"_ASV; break;
+				case '\n': Builder << "\\n"_ASV; break;
+				case '\r': Builder << "\\r"_ASV; break;
+				case '\t': Builder << "\\t"_ASV; break;
+				default:
+					Builder.Appendf(UTF8TEXT("\\u%04x"), uint32(Char));
+					break;
+				}
+			}
+			Value.RightChopInline(Escape.Len());
+		}
+		Builder << '\"';
+	}
+
+	void AppendBase64String(FMemoryView Value)
+	{
+		Builder << '"';
+		checkf(Value.GetSize() <= 512 * 1024 * 1024, TEXT("Encoding 512 MiB or larger is not supported. ")
+			TEXT("Size: " UINT64_FMT), Value.GetSize());
+		const uint32 EncodedSize = FBase64::GetEncodedDataSize(uint32(Value.GetSize()));
+		const int32 EncodedIndex = Builder.AddUninitialized(int32(EncodedSize));
+		FBase64::Encode(static_cast<const uint8*>(Value.GetData()), uint32(Value.GetSize()),
+			reinterpret_cast<ANSICHAR*>(Builder.GetData() + EncodedIndex));
+		Builder << '"';
+	}
+
+private:
+	FUtf8StringBuilderBase& Builder;
+	TUtf8StringBuilder<32> NewLineAndIndent;
+	bool bNeedsComma{false};
+	bool bNeedsNewLine{false};
+};
+
+void CompactBinaryToJson(const FCbObjectView& Object, FUtf8StringBuilderBase& Builder)
+{
+	FCbJsonWriter Writer(Builder);
+	Writer.WriteField(Object.AsFieldView());
+}
+
+///////////////////////////////////////////////////////////////////////////////////////////////////////////////////////
diff --git a/Engine/Source/Runtime/Core/Public/Serialization/CompactBinarySerialization.h b/Engine/Source/Runtime/Core/Public/Serialization/CompactBinarySerialization.h
--- a/Engine/Source/Runtime/Core/Public/Serialization/CompactBinarySerialization.h
+++ b/Engine/Source/Runtime/Core/Public/Serialization/CompactBinarySerialization.h
@@ -3,6 +3,7 @@
 #pragma once
 
 #include "CoreTypes.h"
+#include "Containers/StringFwd.h"
 #include "Memory/MemoryView.h"
 #include "Memory/SharedBuffer.h"
 #include "Serialization/CompactBinary.h"
@@ -75,3 +76,7 @@ CORE_API FArchive& operator<<(FArchive& Ar, FCbArray& Array);
 CORE_API FArchive& operator<<(FArchive& Ar, FCbObject& Object);
 
 ///////////////////////////////////////////////////////////////////////////////////////////////////////////////////////
+
+CORE_API void CompactBinaryToJson(const FCbObjectView& Object, FUtf8StringBuilderBase& Builder);
+
+///////////////////////////////////////////////////////////////////////////////////////////////////////////////////////
