Avoid an MSVC warning about using reinterpret_cast between related types in the scope of TVariant which can arise when using TVariants containing types derived from TAlignedBytes which is used as the storage mechanism in TVariantStorage.  There is no enabled equivalent in clang (tested by compilation for Linux Editor).  This was surfaced when using Visual Studio [REDACTED_PHONE] toolchain.

diff --git a/Engine/Source/Runtime/Core/Public/Misc/TVariant.h b/Engine/Source/Runtime/Core/Public/Misc/TVariant.h
--- a/Engine/Source/Runtime/Core/Public/Misc/TVariant.h
+++ b/Engine/Source/Runtime/Core/Public/Misc/TVariant.h
@@ -105,7 +105,9 @@ public:
 		static_assert(Index != (SIZE_T)-1, "The TVariant is not declared to hold the type passed to Get<>");
 
 		check(Index == TypeIndex);
-		return *reinterpret_cast<U*>(&UE::Core::Private::CastToStorage(*this).Storage);
+		// The intermediate step of casting to void* is used to avoid warnings due to use of reinterpret_cast between related types if U and the storage class are related
+		// This was specifically encountered when U derives from TAlignedBytes
+		return *reinterpret_cast<U*>(reinterpret_cast<void*>(&UE::Core::Private::CastToStorage(*this).Storage));
 	}
 
 	/** Get a reference to the held value. Bad things can happen if this is called on a variant that does not hold the type asked for */
@@ -122,7 +124,9 @@ public:
 	{
 		constexpr SIZE_T Index = UE::Core::Private::TParameterPackTypeIndex<U, T, Ts...>::Value;
 		static_assert(Index != (SIZE_T)-1, "The TVariant is not declared to hold the type passed to TryGet<>");
-		return Index == TypeIndex ? reinterpret_cast<U*>(&UE::Core::Private::CastToStorage(*this).Storage) : nullptr;
+		// The intermediate step of casting to void* is used to avoid warnings due to use of reinterpret_cast between related types if U and the storage class are related
+		// This was specifically encountered when U derives from TAlignedBytes
+		return Index == TypeIndex ? reinterpret_cast<U*>(reinterpret_cast<void*>(&UE::Core::Private::CastToStorage(*this).Storage)) : nullptr;
 	}
 
 	/** Get a pointer to the held value if the held type is the same as the one specified */
