IoDispatcher: Allocate buffers for IoRequests as late as possible
#rnx

diff --git a/Engine/Source/Runtime/Core/Private/GenericPlatform/GenericPlatformIoDispatcher.cpp b/Engine/Source/Runtime/Core/Private/GenericPlatform/GenericPlatformIoDispatcher.cpp
--- a/Engine/Source/Runtime/Core/Private/GenericPlatform/GenericPlatformIoDispatcher.cpp
+++ b/Engine/Source/Runtime/Core/Private/GenericPlatform/GenericPlatformIoDispatcher.cpp
@@ -96,21 +96,15 @@ bool FGenericFileIoStoreImpl::StartRequests(FFileIoStoreRequestQueue& RequestQue
 	}
 
 	uint8* Dest;
-	if (!NextRequest->ImmediateScatter.Request)
+	check(!NextRequest->ImmediateScatter.Request);
+	NextRequest->Buffer = BufferAllocator.AllocBuffer();
+	if (!NextRequest->Buffer)
 	{
-		NextRequest->Buffer = BufferAllocator.AllocBuffer();
-		if (!NextRequest->Buffer)
-		{
-			RequestQueue.Push(*NextRequest);
-			return false;
-		}
-		Dest = NextRequest->Buffer->Memory;
-	}
-	else
-	{
-		Dest = NextRequest->ImmediateScatter.Request->GetIoBuffer().Data() + NextRequest->ImmediateScatter.DstOffset;
+		RequestQueue.Push(*NextRequest);
+		return false;
 	}
-	
+	Dest = NextRequest->Buffer->Memory;
+
 	if (!BlockCache.Read(NextRequest))
 	{
 		IFileHandle* FileHandle = reinterpret_cast<IFileHandle*>(static_cast<UPTRINT>(NextRequest->FileHandle));
diff --git a/Engine/Source/Runtime/Core/Private/IO/IoDispatcher.cpp b/Engine/Source/Runtime/Core/Private/IO/IoDispatcher.cpp
--- a/Engine/Source/Runtime/Core/Private/IO/IoDispatcher.cpp
+++ b/Engine/Source/Runtime/Core/Private/IO/IoDispatcher.cpp
@@ -449,7 +449,7 @@ private:
 				}
 				else
 				{
-					FPlatformAtomics::InterlockedAdd(&TotalLoaded, CompletedRequestsHead->IoBuffer.DataSize());
+					FPlatformAtomics::InterlockedAdd(&TotalLoaded, CompletedRequestsHead->GetBuffer().DataSize());
 					CompleteRequest(CompletedRequestsHead, EIoErrorCode::Ok);
 				}
 				CompletedRequestsHead->ReleaseRef();
@@ -494,7 +494,7 @@ private:
 			TIoStatusOr<FIoBuffer> Result;
 			if (Status == EIoErrorCode::Ok)
 			{
-				Result = Request->IoBuffer;
+				Result = Request->GetBuffer();
 			}
 			else
 			{
@@ -910,6 +910,20 @@ FIoBatch::IssueAndDispatchSubsequents(FGraphEventRef Event)
 
 //////////////////////////////////////////////////////////////////////////
 
+void FIoRequestImpl::CreateBuffer(uint64 Size)
+{
+	if (void* TargetVa = Options.GetTargetVa())
+	{
+		Buffer.Emplace(FIoBuffer::Wrap, TargetVa, Size);
+	}
+	else
+	{
+		LLM_SCOPE(InheritedLLMTag);
+		TRACE_CPUPROFILER_EVENT_SCOPE(AllocMemoryForRequest);
+		Buffer.Emplace(Size);
+	}
+}
+
 void FIoRequestImpl::FreeRequest()
 {
 	Dispatcher.FreeRequest(this);
@@ -997,7 +1011,7 @@ FIoRequest::GetResult()
 	TIoStatusOr<FIoBuffer> Result;
 	if (Status.IsOk())
 	{
-		return Impl->IoBuffer;
+		return Impl->GetBuffer();
 	}
 	else
 	{
diff --git a/Engine/Source/Runtime/Core/Private/IO/IoDispatcherFileBackend.cpp b/Engine/Source/Runtime/Core/Private/IO/IoDispatcherFileBackend.cpp
--- a/Engine/Source/Runtime/Core/Private/IO/IoDispatcherFileBackend.cpp
+++ b/Engine/Source/Runtime/Core/Private/IO/IoDispatcherFileBackend.cpp
@@ -1134,17 +1134,6 @@ bool FFileIoStore::Resolve(FIoRequestImpl* Request)
 
 			if (ResolvedSize > 0)
 			{
-				if (void* TargetVa = Request->Options.GetTargetVa())
-				{
-					Request->IoBuffer = FIoBuffer(FIoBuffer::Wrap, TargetVa, ResolvedSize);
-				}
-				else
-				{
-					LLM_SCOPE(Request->InheritedLLMTag);
-					TRACE_CPUPROFILER_EVENT_SCOPE(AllocMemoryForRequest);
-					Request->IoBuffer = FIoBuffer(ResolvedSize);
-				}
-
 				FFileIoStoreReadRequestList CustomRequests;
 				if (PlatformImpl.CreateCustomRequests(RequestAllocator, *ResolvedRequest, CustomRequests))
 				{
@@ -1320,7 +1309,9 @@ void FFileIoStore::ScatterBlock(FFileIoStoreCompressedBlock* CompressedBlock, bo
 		{
 			if (Scatter.Size)
 			{
-				FMemory::Memcpy(Scatter.Request->GetIoBuffer().Data() + Scatter.DstOffset, UncompressedBuffer + Scatter.SrcOffset, Scatter.Size);
+				check(Scatter.DstOffset + Scatter.Size <= Scatter.Request->GetBuffer().DataSize());
+				check(Scatter.SrcOffset + Scatter.Size <= CompressedBlock->UncompressedSize);
+				FMemory::Memcpy(Scatter.Request->GetBuffer().Data() + Scatter.DstOffset, UncompressedBuffer + Scatter.SrcOffset, Scatter.Size);
 			}
 		}
 	}
@@ -1550,6 +1541,20 @@ FIoRequestImpl* FFileIoStore::GetCompletedRequests()
 		{
 			break;
 		}
+
+		for (const FFileIoStoreBlockScatter& Scatter : BlockToDecompress->ScatterList)
+		{
+			if (Scatter.Size)
+			{
+				FIoRequestImpl* DispatcherRequest = Scatter.Request->DispatcherRequest;
+				check(DispatcherRequest);
+				if (!DispatcherRequest->HasBuffer())
+				{
+					DispatcherRequest->CreateBuffer(Scatter.Request->ResolvedSize);
+				}
+			}
+		}
+
 		// Scatter block asynchronous when the block is compressed, encrypted or signed
 		const bool bScatterAsync = bIsMultithreaded && (!BlockToDecompress->CompressionMethod.IsNone() || BlockToDecompress->EncryptionKey.IsValid() || BlockToDecompress->SignatureHash);
 		if (bScatterAsync)
@@ -1713,7 +1718,7 @@ void FFileIoStore::ReadBlocks(FFileIoStoreResolvedRequest& ResolvedRequest)
 		}
 		check(CompressedBlock->UncompressedSize > RequestStartOffsetInBlock);
 		uint64 RequestSizeInBlock = FMath::Min<uint64>(CompressedBlock->UncompressedSize - RequestStartOffsetInBlock, RequestRemainingBytes);
-		check(OffsetInRequest + RequestSizeInBlock <= ResolvedRequest.GetIoBuffer().DataSize());
+		check(OffsetInRequest + RequestSizeInBlock <= ResolvedRequest.ResolvedSize);
 		check(RequestStartOffsetInBlock + RequestSizeInBlock <= CompressedBlock->UncompressedSize);
 
 		FFileIoStoreBlockScatter& Scatter = CompressedBlock->ScatterList.AddDefaulted_GetRef();
diff --git a/Engine/Source/Runtime/Core/Private/IO/IoDispatcherFileBackendTypes.h b/Engine/Source/Runtime/Core/Private/IO/IoDispatcherFileBackendTypes.h
--- a/Engine/Source/Runtime/Core/Private/IO/IoDispatcherFileBackendTypes.h
+++ b/Engine/Source/Runtime/Core/Private/IO/IoDispatcherFileBackendTypes.h
@@ -818,10 +818,22 @@ public:
 		return DispatcherRequest->Priority;
 	}
 
-	FIoBuffer& GetIoBuffer()
+	bool HasBuffer() const
 	{
 		check(DispatcherRequest);
-		return DispatcherRequest->IoBuffer;
+		return DispatcherRequest->HasBuffer();
+	}
+
+	void CreateBuffer(uint64 Size)
+	{
+		check(DispatcherRequest);
+		DispatcherRequest->CreateBuffer(Size);
+	}
+
+	FIoBuffer& GetBuffer()
+	{
+		check(DispatcherRequest);
+		return DispatcherRequest->GetBuffer();
 	}
 
 	void AddReadRequestLink(FFileIoStoreReadRequestLink* ReadRequestLink);
diff --git a/Engine/Source/Runtime/Core/Public/IO/IoDispatcherBackend.h b/Engine/Source/Runtime/Core/Public/IO/IoDispatcherBackend.h
--- a/Engine/Source/Runtime/Core/Public/IO/IoDispatcherBackend.h
+++ b/Engine/Source/Runtime/Core/Public/IO/IoDispatcherBackend.h
@@ -4,6 +4,7 @@
 
 #include "IoDispatcher.h"
 #include "HAL/LowLevelMemTracker.h"
+#include "Misc/Optional.h"
 
 class FIoRequestImpl
 {
@@ -13,7 +14,6 @@ public:
 	LLM(const UE::LLMPrivate::FTagData* InheritedLLMTag);
 	FIoChunkId ChunkId;
 	FIoReadOptions Options;
-	FIoBuffer IoBuffer;
 	int32 Priority = 0;
 
 	FIoRequestImpl(FIoDispatcherImpl& InDispatcher)
@@ -32,6 +32,18 @@ public:
 		bFailed = true;
 	}
 
+	bool HasBuffer() const
+	{
+		return Buffer.IsSet();
+	}
+
+	CORE_API void CreateBuffer(uint64 Size);
+
+	FIoBuffer& GetBuffer()
+	{
+		return Buffer.GetValue();
+	}
+
 private:
 	friend class FIoDispatcherImpl;
 	friend class FIoRequest;
@@ -55,6 +67,7 @@ private:
 	FIoDispatcherImpl& Dispatcher;
 	struct IIoDispatcherBackend* Backend = nullptr;
 	FIoBatchImpl* Batch = nullptr;
+	TOptional<FIoBuffer> Buffer;
 	FIoReadCallback Callback;
 	TAtomic<uint32> RefCount{ 0 };
 	TAtomic<EIoErrorCode> ErrorCode{ EIoErrorCode::Unknown };
diff --git a/Engine/Source/Runtime/StorageServerClient/Private/StorageServerIoDispatcherBackend.cpp b/Engine/Source/Runtime/StorageServerClient/Private/StorageServerIoDispatcherBackend.cpp
--- a/Engine/Source/Runtime/StorageServerClient/Private/StorageServerIoDispatcherBackend.cpp
+++ b/Engine/Source/Runtime/StorageServerClient/Private/StorageServerIoDispatcherBackend.cpp
@@ -128,18 +128,9 @@ bool FStorageServerIoDispatcherBackend::Resolve(FIoRequestImpl* Request)
 		TRACE_CPUPROFILER_EVENT_SCOPE(StorageServerIoDispatcherReadChunk);
 		bool bSuccess = Connection.ReadChunkRequest(Request->ChunkId, Request->Options.GetOffset(), Request->Options.GetSize(), [&Request](FStorageServerResponse& Response)
 			{
-				if (void* TargetVa = Request->Options.GetTargetVa())
-				{
-					Request->IoBuffer = FIoBuffer(FIoBuffer::Wrap, TargetVa, Response.TotalSize());
-				}
-				else
-				{
-					LLM_SCOPE(ELLMTag::FileSystem);
-					TRACE_CPUPROFILER_EVENT_SCOPE(AllocMemoryForRequest);
-					Request->IoBuffer = FIoBuffer(Response.TotalSize());
-				}
+				Request->CreateBuffer(Response.TotalSize());
 				TRACE_CPUPROFILER_EVENT_SCOPE(SerializeResponse);
-				Response.Serialize(Request->IoBuffer.Data(), Response.TotalSize());
+				Response.Serialize(Request->GetBuffer().Data(), Response.TotalSize());
 			});
 		if (!bSuccess)
 		{
@@ -298,18 +289,9 @@ void FStorageServerIoDispatcherBackend::FBatch::DoThreadedWork()
 		TRACE_CPUPROFILER_EVENT_SCOPE(StorageServerIoDispatcherReadChunk);
 		bool bSuccess = Owner.Connection.ReadChunkRequest(Request->ChunkId, Request->Options.GetOffset(), Request->Options.GetSize(), [&Request](FStorageServerResponse& Response)
 			{
-				if (void* TargetVa = Request->Options.GetTargetVa())
-				{
-					Request->IoBuffer = FIoBuffer(FIoBuffer::Wrap, TargetVa, Response.TotalSize());
-				}
-				else
-				{
-					LLM_SCOPE(ELLMTag::FileSystem);
-					TRACE_CPUPROFILER_EVENT_SCOPE(AllocMemoryForRequest);
-					Request->IoBuffer = FIoBuffer(Response.TotalSize());
-				}
+				Request->CreateBuffer(Response.TotalSize());
 				TRACE_CPUPROFILER_EVENT_SCOPE(SerializeResponse);
-				Response.Serialize(Request->IoBuffer.Data(), Response.TotalSize());
+				Response.Serialize(Request->GetBuffer().Data(), Response.TotalSize());
 			});
 		if (!bSuccess)
 		{
