fixed mac editor compilation

diff --git a/Engine/Source/Runtime/Core/Public/Tasks/TaskPrivate.h b/Engine/Source/Runtime/Core/Public/Tasks/TaskPrivate.h
--- a/Engine/Source/Runtime/Core/Public/Tasks/TaskPrivate.h
+++ b/Engine/Source/Runtime/Core/Public/Tasks/TaskPrivate.h
@@ -553,6 +553,18 @@ namespace UE { namespace Tasks
 #if UE_TASK_TRACE_ENABLED
 			TaskTrace::FId TraceId = TaskTrace::GenerateTaskId();
 #endif
+
+		// some projects are still built on macOS before v.10.14 that doesn't have aligned new/delete operators
+		public:
+			inline void* operator new(size_t Size)
+			{
+				return FMemory::Malloc(Size, 64u);
+			}
+
+			inline void operator delete(void* Ptr)
+			{
+				FMemory::Free(Ptr);
+			}
 		};
 
 		template<typename TaskCollectionType>
