minor ergonomic change to ClosableMpscQueue API

diff --git a/Engine/Source/Runtime/Core/Public/Containers/ClosableMpscQueue.h b/Engine/Source/Runtime/Core/Public/Containers/ClosableMpscQueue.h
--- a/Engine/Source/Runtime/Core/Public/Containers/ClosableMpscQueue.h
+++ b/Engine/Source/Runtime/Core/Public/Containers/ClosableMpscQueue.h
@@ -68,9 +68,10 @@ public:
 	/**
 	 * Closes the queue and consumes all items.
 	 * @param Consumer: a functor with signature `AnyReturnType (T Value)` that will receive all items in FIFO order
+	 * @return false if already closed
 	 */
 	template<typename F>
-	void Close(const F& Consumer)
+	bool Close(const F& Consumer)
 	{
 		FNode* Tail = &Sentinel;
 
@@ -80,6 +81,8 @@ public:
 		// the queue is closed at this point, and the user is free to destroy it
 		// no members should be accessed
 		Close_NonMember(Head_Local, Tail, Consumer);
+
+		return Head_Local != nullptr;
 	}
 
 	bool IsClosed() const
