FSpinLock and generic TScopeLock

diff --git a/Engine/Source/Runtime/Core/Public/Misc/ScopeLock.h b/Engine/Source/Runtime/Core/Public/Misc/ScopeLock.h
--- a/Engine/Source/Runtime/Core/Public/Misc/ScopeLock.h
+++ b/Engine/Source/Runtime/Core/Public/Misc/ScopeLock.h
@@ -136,4 +136,76 @@ private:
 	FCriticalSection* SynchObject;
 };
 
+namespace UE
+{
+	// RAII-style scope locking of a synchronisation primitive
+	// `MutexType` is required to implement `Lock` and `Unlock` methods
+	// Example:
+	//	{
+	//		TScopeLock<FCriticalSection> ScopeLock(CriticalSection);
+	//		...
+	//	}
+	template<typename MutexType>
+	class TScopeLock
+	{
+	public:
+		UE_NONCOPYABLE(TScopeLock);
+
+		TScopeLock(MutexType& InMutex)
+			: Mutex(&InMutex)
+		{
+			check(Mutex);
+			Mutex->Lock();
+		}
+
+		~TScopeLock()
+		{
+			Unlock();
+		}
+
+		void Unlock()
+		{
+			if (Mutex)
+			{
+				Mutex->Unlock();
+				Mutex = nullptr;
+			}
+		}
+
+	private:
+		MutexType* Mutex;
+	};
+
+	// RAII-style scope unlocking of a synchronisation primitive
+	// `MutexType` is required to implement `Lock` and `Unlock` methods
+	// Example:
+	//	{
+	//		TScopeLock<FCriticalSection> ScopeLock(CriticalSection);
+	//		for (FElementType& Element : ThreadUnsafeContainer)
+	//		{
+	//			TScopeUnlock<FCriticalSection> ScopeUnlock(CriticalSection);
+	//			Process(Element);
+	//		}
+	//	}
+	template<typename MutexType>
+	class TScopeUnlock
+	{
+	public:
+		UE_NONCOPYABLE(TScopeUnlock);
+
+		TScopeUnlock(MutexType& InMutex)
+			: Mutex(&InMutex)
+		{
+			check(Mutex);
+			Mutex->Unlock();
+		}
+
+		~TScopeUnlock()
+		{
+			Mutex->Lock();
+		}
 
+	private:
+		MutexType* Mutex;
+	};
+}
\ No newline at end of file
diff --git a/Engine/Source/Runtime/Core/Public/Misc/SpinLock.h b/Engine/Source/Runtime/Core/Public/Misc/SpinLock.h
new file mode 100644
--- /dev/null
+++ b/Engine/Source/Runtime/Core/Public/Misc/SpinLock.h
@@ -0,0 +1,50 @@
+// Copyright Epic Games, Inc. All Rights Reserved.
+
+#pragma once
+
+#include "CoreTypes.h"
+#include "HAL/PlatformProcess.h"
+#include <atomic>
+
+namespace UE
+{
+	// A mutex that doesn't put the thread into a WAIT state but instead repeatedly tries to aquire the lock.
+	// WARNING: Should be used only for very short locks
+	// Use with `TScopeLock`
+	class FSpinLock
+	{
+	public:
+		UE_NONCOPYABLE(FSpinLock);
+
+		FSpinLock() = default;
+
+		void Lock()
+		{
+			while (true)
+			{
+				if (!bFlag.exchange(true, std::memory_order_acquire))
+				{
+					break;
+				}
+
+				while (bFlag.load(std::memory_order_relaxed))
+				{
+					FPlatformProcess::Yield();
+				}
+			}
+		}
+
+		bool TryLock()
+		{
+			return !bFlag.exchange(true, std::memory_order_acquire);
+		}
+
+		void Unlock()
+		{
+			bFlag.store(false, std::memory_order_release);
+		}
+
+	private:
+		std::atomic<bool> bFlag{ false };
+	};
+}
