Tasks: fixed 3 races in nested tasks and deep retraction

diff --git a/Engine/Source/Runtime/Core/Private/Tasks/TaskPrivate.cpp b/Engine/Source/Runtime/Core/Private/Tasks/TaskPrivate.cpp
--- a/Engine/Source/Runtime/Core/Private/Tasks/TaskPrivate.cpp
+++ b/Engine/Source/Runtime/Core/Private/Tasks/TaskPrivate.cpp
@@ -33,16 +33,22 @@ namespace UE { namespace Tasks { namespace Private
 			GetPipe()->ClearTask(*this);
 		}
 
-		Subsequents.Close(
+		verify(Subsequents.Close(
 			[this](FTaskBase* Subsequent)
 			{
 				Subsequent->TryUnlock();
 			}
-		);
+		));
 
-		while (TOptional<FTaskBase*> Prerequisite = Prerequisites.Dequeue())
+		// release nested tasks. this can happen concurrently with task retraction, that also consumes `Prerequisite`. 
+		// `Prerequisites` is a single-producer/single-consumer queue so we need to put an aditional synchronisation for dequeueing
 		{
-			Prerequisite.GetValue()->Release();
+			TScopeLock<FSpinLock> ScopeLock(PrerequisitesLock);
+			while (TOptional<FTaskBase*> Prerequisite = Prerequisites.Dequeue())
+			{
+				TScopeUnlock<FSpinLock> ScopeUnlock(PrerequisitesLock); // only `Prerequisites.Dequeue()` needs to be locked
+				Prerequisite.GetValue()->Release();
+			}
 		}
 
 		Release(); // release the reference accounted for nested tasks. the task can be destroyed as the result
diff --git a/Engine/Source/Runtime/Core/Private/Tests/Tasks/TasksTest.cpp b/Engine/Source/Runtime/Core/Private/Tests/Tasks/TasksTest.cpp
--- a/Engine/Source/Runtime/Core/Private/Tests/Tasks/TasksTest.cpp
+++ b/Engine/Source/Runtime/Core/Private/Tests/Tasks/TasksTest.cpp
@@ -823,10 +823,10 @@ namespace UE { namespace TasksTests
 
 		for (int i = 0; i != NumWorkers; ++i)
 		{
-			checkf(LowLevelTasks::FScheduler::Get().IsWorkerThread(), TEXT("No reserve workers are expected to get blocked"));
 			WorkerBlockers.Add(Launch(UE_SOURCE_LOCATION,
 				[&NumWorkersNotBlocked, &ResumeEvent]
 				{
+					checkf(LowLevelTasks::FScheduler::Get().IsWorkerThread(), TEXT("No reserve workers are expected to get blocked"));
 					--NumWorkersNotBlocked;
 					ResumeEvent.Wait();
 				}
@@ -851,7 +851,7 @@ namespace UE { namespace TasksTests
 		FTask P21{ Launch(TEXT("P21"), [] {}, Prerequisites(P11, P12)) };
 		FTask P22{ Launch(TEXT("P22"), [] {}) };
 		FTask N11, N12, N21, N22;
-		Launch(UE_SOURCE_LOCATION,
+		FTask Task = Launch(UE_SOURCE_LOCATION,
 			[&N11, &N12, &N21, &N22]
 			{
 				AddNested(N11 = Launch(TEXT("N11"),
@@ -864,7 +864,8 @@ namespace UE { namespace TasksTests
 				AddNested(N12 = Launch(TEXT("N12"), [] {}));
 			},
 			Prerequisites(P21, P22)
-		).Wait();
+		);
+		Task.Wait();
 		check(P11.IsCompleted() && P12.IsCompleted() && P21.IsCompleted() && P22.IsCompleted() &&
 			N11.IsCompleted() && N12.IsCompleted() && N21.IsCompleted() && N22.IsCompleted());
 	}
@@ -899,7 +900,8 @@ namespace UE { namespace TasksTests
 					check(bPrerequisiteDone);
 					check(!bTaskDone);
 					bTaskDone = true; 
-				}
+				},
+				Prerequisite
 			).Wait();
 			check(bPrerequisiteDone);
 			check(bTaskDone);
@@ -954,7 +956,33 @@ namespace UE { namespace TasksTests
 
 	bool FTasksDeepRetractionStressTest::RunTest(const FString& Parameters)
 	{
-		UE_BENCHMARK(5, DeepRetractionStressTest<1000>);
+		UE_BENCHMARK(5, DeepRetractionStressTest<10000>);
+
+		return true;
+	}
+
+	template<uint64 Num>
+	void NestedTasksStressTest()
+	{
+		for (uint64 i = 0; i != Num; ++i)
+		{
+			FTask Nested;
+			FTask Parent = Launch(TEXT("Parent"),
+				[&Nested]
+				{
+					AddNested(Nested = Launch(TEXT("Nested"), [] {}));
+				}
+			);
+			Parent.Wait();
+			check(Nested.IsCompleted() && Parent.IsCompleted());
+		}
+	}
+
+	IMPLEMENT_SIMPLE_AUTOMATION_TEST(FTasksNestedTasksStressTest, "System.Core.Tasks.NestedTasks.Stress", EAutomationTestFlags::ApplicationContextMask | EAutomationTestFlags::EngineFilter);
+
+	bool FTasksNestedTasksStressTest::RunTest(const FString& Parameters)
+	{
+		UE_BENCHMARK(5, NestedTasksStressTest<10000>);
 
 		return true;
 	}
diff --git a/Engine/Source/Runtime/Core/Public/Tasks/TaskPrivate.h b/Engine/Source/Runtime/Core/Public/Tasks/TaskPrivate.h
--- a/Engine/Source/Runtime/Core/Public/Tasks/TaskPrivate.h
+++ b/Engine/Source/Runtime/Core/Public/Tasks/TaskPrivate.h
@@ -13,6 +13,8 @@
 #include "Templates/UnrealTemplate.h"
 #include "CoreTypes.h"
 #include "Math/NumericLimits.h"
+#include "Misc/SpinLock.h"
+#include "Misc/ScopeLock.h"
 
 #include <atomic>
 #include <type_traits>
@@ -118,7 +120,7 @@ namespace UE { namespace Tasks
 				if (Prerequisite.Pimpl->AddSubsequent(*this))
 				{
 					Prerequisites.Enqueue(Prerequisite.Pimpl.GetReference());
-					// keep it alive until this task's destruction
+					// keep it alive until this task's execution
 					Prerequisite.Pimpl->AddRef();
 				}
 				else
@@ -130,7 +132,7 @@ namespace UE { namespace Tasks
 
 			// The task will be executed only when all prerequisites are completed.
 			// Must not be called concurrently.
-			// @param Prerequisites - an iterable collection of tasks
+			// @param InPrerequisites - an iterable collection of tasks
 			template<typename PrerequisiteCollectionType, decltype(std::declval<PrerequisiteCollectionType>().begin())* = nullptr>
 			void AddPrerequisites(const PrerequisiteCollectionType& InPrerequisites)
 			{
@@ -160,7 +162,7 @@ namespace UE { namespace Tasks
 					if (Prerequisite->AddSubsequent(*this))
 					{
 						Prerequisites.Enqueue(Prerequisite);
-						// keep it alive until this task's destruction
+						// keep it alive until this task's execution
 						Prerequisite->AddRef();
 					}
 					else
@@ -235,29 +237,18 @@ namespace UE { namespace Tasks
 					return NumLocks.load(std::memory_order_acquire) != (Pipe == nullptr ? 1 : 0);
 				};
 
-				auto TryRetractPrerequisites = [this, RecursionDepth]
+				if (IsLocked())
 				{
-					// keep trying retracting all prerequisites even if some of them fail, so the current worker can contribute instead of being blocked
-					bool bSucceeded = true;
+					// try to unlock the task. even if prerequisites retraction fails we still need to proceed to try to execute the task here in case
+					// prerequisites were completed in parallel
+
 					// prerequisites are "consumed" here even if their retraction fails. retraction can fail only if task execution has 
 					// already started, and it can't become "retractable" again, so no need to keep them
 					while (TOptional<FTaskBase*> Prerequisite = Prerequisites.Dequeue())
 					{
-						if (!Prerequisite.GetValue()->TryRetractAndExecute(RecursionDepth))
-						{
-							bSucceeded = false;
-						}
+						Prerequisite.GetValue()->TryRetractAndExecute(RecursionDepth);
 						Prerequisite.GetValue()->Release();
 					}
-
-					return bSucceeded;
-				};
-
-				if (IsLocked())
-				{
-					// try to unlock the task. even if prerequisites retraction fails we still need to proceed to try to execute the task here in case
-					// prerequisites were completed in parallel
-					TryRetractPrerequisites();
 				}
 
 				// the task can be still locked if either prerequisite retraction (above) failed, or the task was piped by another thread, or the task 
@@ -281,16 +272,42 @@ namespace UE { namespace Tasks
 
 				if (IsCompleted()) // still can be hold back by nested tasks
 				{
-					return IsCompleted(); // still can be hold back by nested tasks
+					return true;
 				}
 
-				// try to retract nested tasks
-				if (!TryRetractPrerequisites())
+				// retract nested tasks. this can happen concurrently with `Close()` called by a nested task, that also consumes `Prerequisite`. 
+				// `Prerequisites` is a single-producer/single-consumer queue so we need to put an aditional synchronisation for dequeueing
 				{
-					return false;
+					TScopeLock<FSpinLock> ScopeLock(PrerequisitesLock);
+
+					// keep trying retracting all nested tasks even if some of them fail, so the current worker can contribute instead of being blocked
+					bool bSucceeded = true;
+					// prerequisites are "consumed" here even if their retraction fails. retraction can fail only if task execution has 
+					// already started, and it can't become "retractable" again, so no need to keep them
+					while (TOptional<FTaskBase*> Prerequisite = Prerequisites.Dequeue())
+					{
+						TScopeUnlock<FSpinLock> ScopeUnlock(PrerequisitesLock); // only `Prerequisites.Dequeue()` needs to be locked
+
+						if (!Prerequisite.GetValue()->TryRetractAndExecute(RecursionDepth))
+						{
+							bSucceeded = false;
+						}
+						Prerequisite.GetValue()->Release();
+					}
+
+					if (!bSucceeded)
+					{
+						return false;
+					}
+				}
+
+				// it happens that all nested tasks are completed and are in the process of completing the parent (this task) concurrently, but the flag
+				// is not set yet. wait for it
+				while (!IsCompleted())
+				{
+					FPlatformProcess::Yield();
 				}
 
-				checkSlow(IsCompleted());
 				return true;
 			}
 
@@ -299,13 +316,12 @@ namespace UE { namespace Tasks
 			{
 				uint32 PrevNumLocks = NumLocks.fetch_add(1, std::memory_order_relaxed); // in case we'll succeed in adding subsequent
 				checkf(PrevNumLocks + 1 < TNumericLimits<uint32>::Max(), TEXT("Max number of nested tasks reached: %d"), TNumericLimits<uint32>::Max() - ExecutionFlag);
-				checkf(PrevNumLocks >= ExecutionFlag, TEXT("Internal error: nested tasks can be added only during parent's execution (%u)"), PrevNumLocks);
+				checkf(PrevNumLocks > ExecutionFlag, TEXT("Internal error: nested tasks can be added only during parent's execution (%u)"), PrevNumLocks);
 
 				if (Nested.AddSubsequent(*this))
 				{
-					// there's no race with, as `Prerequisites` can't be accessed by another thread during task execution (execution "locks" the task)
 					Prerequisites.Enqueue(&Nested);
-					// keep it alive until this task's destruction
+					// keep it alive until the task destruction
 					Nested.AddRef();
 				}
 				else
@@ -420,9 +436,9 @@ namespace UE { namespace Tasks
 			// The task can be deleted as the result of this call.
 			bool TryComplete(uint32 LocalNumLocks)
 			{
-				checkSlow(!IsCompleted());
 				if (LocalNumLocks == ExecutionFlag)
 				{
+					checkSlow(!IsCompleted());
 					Close();
 					return true;
 				}
@@ -439,7 +455,7 @@ namespace UE { namespace Tasks
 					return true;
 				}
 
-				return false; // too late, execution already started
+				return false; // that task is being retracted in parallel
 			}
 
 			bool TryGetExecutionPermission()
@@ -449,6 +465,7 @@ namespace UE { namespace Tasks
 
 			void RevokeExecutionPermission()
 			{
+				checkSlow(!bAvailableForExecution.load(std::memory_order_relaxed));
 				bAvailableForExecution.store(true, std::memory_order_release);
 			}
 
@@ -462,8 +479,10 @@ namespace UE { namespace Tasks
 					Prerequisite.GetValue()->Release();
 				}
 
+				checkSlow(Pipe == nullptr ? NumLocks.load(std::memory_order_relaxed) == 1 : NumLocks.load(std::memory_order_relaxed) == 0);
 				NumLocks.store(ExecutionFlag + 1, std::memory_order_relaxed); // +1 to hold it locked during execution, so nested tasks don't
 				// complete it before the execution finishes
+
 				FTaskBase* PrevTask = ExchangeCurrentTask(this);
 				TaskTrace::Started(GetTraceId());
 				StartPipeExecution();
@@ -473,6 +492,7 @@ namespace UE { namespace Tasks
 				FinishPipeExecution();
 				TaskTrace::Finished(GetTraceId());
 				ExchangeCurrentTask(PrevTask);
+				
 				uint32 LocalNumLocks = NumLocks.fetch_sub(1, std::memory_order_relaxed) - 1;
 
 				TryComplete(LocalNumLocks);
@@ -544,7 +564,9 @@ namespace UE { namespace Tasks
 			// retraction that consumes prerequisites. multiple threads can try to retract the task but only one will be allowed to proceed, 
 			// so single producer/single consumer
 			// 3) by adding nested tasks. after piping. during task execution, thus retraction is not possible and won't touch it. single-threaded
+			// `Prerequisites` can be consumed concurrently by retraction and nested tasks closing this task. This case is explicitly locked
 			TSpscQueue<FTaskBase*> Prerequisites;
+			FSpinLock PrerequisitesLock;
 
 			TClosableMpscQueue<FTaskBase*> Subsequents;
 
