#UE If CrashReportClient failed to start, emit the log information about it, even when running with -unattended.

diff --git a/Engine/Source/Runtime/Core/Private/Windows/WindowsPlatformCrashContext.cpp b/Engine/Source/Runtime/Core/Private/Windows/WindowsPlatformCrashContext.cpp
--- a/Engine/Source/Runtime/Core/Private/Windows/WindowsPlatformCrashContext.cpp
+++ b/Engine/Source/Runtime/Core/Private/Windows/WindowsPlatformCrashContext.cpp
@@ -1015,16 +1015,19 @@ int32 ReportCrashUsingCrashReportClient(FWindowsPlatformCrashContext& InContext,
 			}
 		}
 
-		if (!bCrashReporterRan && !bNoDialog)
+		if (!bCrashReporterRan)
 		{
 			UE_LOG(LogWindows, Log, TEXT("Could not start crash report client using %s"), CrashReporterClientPath);
 			FPlatformMemory::DumpStats(*GWarn);
-			FText MessageTitle(FText::Format(
-				NSLOCTEXT("MessageDialog", "AppHasCrashed", "The {0} {1} has crashed and will close"),
-				FText::FromString(AppName),
-				FText::FromString(FPlatformMisc::GetEngineMode())
-				));
-			FMessageDialog::Open(EAppMsgType::Ok, FText::FromString(GErrorHist), &MessageTitle);
+			if (!bNoDialog)
+			{
+				FText MessageTitle(FText::Format(
+					NSLOCTEXT("MessageDialog", "AppHasCrashed", "The {0} {1} has crashed and will close"),
+					FText::FromString(AppName),
+					FText::FromString(FPlatformMisc::GetEngineMode())
+					));
+				FMessageDialog::Open(EAppMsgType::Ok, FText::FromString(GErrorHist), &MessageTitle);
+			}
 		}
 	}
 
