Enable Windows ETW (context switches and stack sampling) and memory tracing on XSX.
Use Perf trace for context switches and stack sampling on PS4 and PS5.

diff --git a/Engine/Source/Developer/TraceServices/Private/Model/SymslibResolver.cpp b/Engine/Source/Developer/TraceServices/Private/Model/SymslibResolver.cpp
--- a/Engine/Source/Developer/TraceServices/Private/Model/SymslibResolver.cpp
+++ b/Engine/Source/Developer/TraceServices/Private/Model/SymslibResolver.cpp
@@ -304,29 +304,36 @@ FString FSymslibResolver::FindSymbolFile(const FString& Path)
 {
 	IPlatformFile* PlatformFile = &FPlatformFileManager::Get().GetPlatformFile();
 
-	if (PlatformFile->FileExists(*Path))
-	{
-		return Path;
-	}
+	// Extract only filename part in case Path is absolute path
+	FString FileName = FPaths::GetCleanFilename(Path);
 
 	FString SymbolPath = FPlatformMisc::GetEnvironmentVariable(L"UE_INSIGHTS_SYMBOL_PATH");
 
 	FString SymbolPathPart;
 	while (SymbolPath.Split(TEXT(";"), &SymbolPathPart, &SymbolPath))
 	{
-		FString Result = FPaths::Combine(SymbolPathPart, Path);
+		FString Result = FPaths::Combine(SymbolPathPart, FileName);
 		if (PlatformFile->FileExists(*Result))
 		{
+			// File found in one of symbol paths
 			return Result;
 		}
 	}
 
-	FString Result = FPaths::Combine(SymbolPath, Path);
+	FString Result = FPaths::Combine(SymbolPath, FileName);
 	if (PlatformFile->FileExists(*Result))
 	{
+		// File found in symbol path
 		return Result;
 	}
 
+	if (PlatformFile->FileExists(*Path))
+	{
+		// File found by absolute path
+		return Path;
+	}
+
+	// File not found
 	return FString();
 }
 
@@ -374,8 +381,7 @@ FSymslibResolver::EModuleStatus FSymslibResolver::LoadModule(FModuleEntry* Modul
 	}
 
 	// Iterate and map all the debug files
-	const FStringView ModuleBasePath = FPathViews::GetPath(Module->Path);
-	TStringBuilder<256> DebugFilePath;
+	const FString ModuleBasePath = FPaths::GetPath(SymbolFilePath);
 	TArray<SymsFile> Files;
 
 	SymsDebugFileIter DebugFiles;
@@ -385,11 +391,31 @@ FSymslibResolver::EModuleStatus FSymslibResolver::LoadModule(FModuleEntry* Modul
 		while (syms_debug_file_iter_next(&DebugFiles, &Path))
 		{
 			// Map the debug file
-			DebugFilePath.Reset();
-			FPathViews::Append(DebugFilePath, ModuleBasePath);
-			FPathViews::Append(DebugFilePath, ANSI_TO_TCHAR(Path.data));
 
-			if (PlatformFile->FileExists(*DebugFilePath))
+			bool bDebugFilePathFound = false;
+
+			// First try path itself if it is absolute
+			FString DebugFilePath = ANSI_TO_TCHAR(Path.data);
+			if (!FPathViews::IsRelativePath(DebugFilePath))
+			{
+				if (PlatformFile->FileExists(*DebugFilePath))
+				{
+					bDebugFilePathFound = true;
+				}
+			}
+
+			// Otherwise try relative path to symbol file path
+			if (!bDebugFilePathFound)
+			{
+				DebugFilePath = FPaths::Combine(ModuleBasePath, FPaths::GetCleanFilename(DebugFilePath));
+				if (PlatformFile->FileExists(*DebugFilePath))
+				{
+					bDebugFilePathFound = true;
+				}
+			}
+
+			// Use path if file was found
+			if (bDebugFilePathFound)
 			{
 				IMappedFileHandle* FileHandle = PlatformFile->OpenMapped(*DebugFilePath);
 				if (FileHandle != nullptr)
diff --git a/Engine/Source/Runtime/Core/Core.Build.cs b/Engine/Source/Runtime/Core/Core.Build.cs
--- a/Engine/Source/Runtime/Core/Core.Build.cs
+++ b/Engine/Source/Runtime/Core/Core.Build.cs
@@ -307,6 +307,15 @@ public class Core : ModuleRules
 			}
 		}
 
+		if (Target.Platform == UnrealTargetPlatform.Win64
+			&& Target.Configuration != UnrealTargetConfiguration.Shipping)
+		{
+			PublicDefinitions.Add("PLATFORM_SUPPORTS_TRACE_WIN32_VIRTUAL_MEMORY_HOOKS=1");
+			PublicDefinitions.Add("PLATFORM_SUPPORTS_TRACE_WIN32_MODULE_DIAGNOSTICS=1");
+			PublicDefinitions.Add("PLATFORM_SUPPORTS_TRACE_WIN32_CALLSTACK=1");
+			PublicDefinitions.Add("UE_MEMORY_TRACE_AVAILABLE=1");
+		}
+
 		// temporary thing.
 		PrivateDefinitions.Add("PLATFORM_SUPPORTS_BINARYCONFIG=" + (SupportsBinaryConfig(Target) ? "1" : "0"));
 
diff --git a/Engine/Source/Runtime/Core/Private/ProfilingDebugging/MemoryTrace.h b/Engine/Source/Runtime/Core/Private/ProfilingDebugging/MemoryTrace.h
--- a/Engine/Source/Runtime/Core/Private/ProfilingDebugging/MemoryTrace.h
+++ b/Engine/Source/Runtime/Core/Private/ProfilingDebugging/MemoryTrace.h
@@ -3,14 +3,19 @@
 
 #include "Trace/Trace.h"
 
+#if !defined(UE_MEMORY_TRACE_AVAILABLE)
+#   define UE_MEMORY_TRACE_AVAILABLE 0
+#endif
+
+#if !defined(UE_MEMORY_TRACE_LATE_INIT)
+#	define UE_MEMORY_TRACE_LATE_INIT 0
+#endif
+
 #if !defined(UE_MEMORY_TRACE_ENABLED) && UE_TRACE_ENABLED
-#	if PLATFORM_WINDOWS || PLATFORM_PS4 || defined(__PS5__)
+#	if UE_MEMORY_TRACE_AVAILABLE
 #		if !PLATFORM_USES_FIXED_GMalloc_CLASS && PLATFORM_64BITS
 #			define UE_MEMORY_TRACE_ENABLED !UE_BUILD_SHIPPING
 #		endif
-#		if PLATFORM_PS4 || defined(__PS5__)
-#			define UE_MEMORY_TRACE_LATE_INIT UE_MEMORY_TRACE_ENABLED
-#		endif
 #	endif
 #endif
 
@@ -18,10 +23,6 @@
 #	define UE_MEMORY_TRACE_ENABLED 0
 #endif
 
-#if !defined(UE_MEMORY_TRACE_LATE_INIT)
-#	define UE_MEMORY_TRACE_LATE_INIT 0
-#endif
-
 #if UE_MEMORY_TRACE_ENABLED
 
 UE_TRACE_CHANNEL_EXTERN(MemTrackChannel);
diff --git a/Engine/Source/Runtime/Core/Private/ProfilingDebugging/Microsoft/EventTracingForWindows.cpp b/Engine/Source/Runtime/Core/Private/ProfilingDebugging/Microsoft/EventTracingForWindows.cpp
new file mode 100644
--- /dev/null
+++ b/Engine/Source/Runtime/Core/Private/ProfilingDebugging/Microsoft/EventTracingForWindows.cpp
@@ -0,0 +1,531 @@
+// Copyright Epic Games, Inc. All Rights Reserved.
+
+#include "ProfilingDebugging/PlatformEvents.h"
+#include "HAL/Runnable.h"
+#include "HAL/Event.h"
+#include "HAL/RunnableThread.h"
+#include "Logging/LogMacros.h"
+#include "Trace/Trace.h"
+
+#if PLATFORM_SUPPORTS_PLATFORM_EVENTS
+
+#include COMPILED_PLATFORM_HEADER(PlatformEventTracingForWindows.h)
+
+/////////////////////////////////////////////////////////////////////
+
+DEFINE_LOG_CATEGORY_STATIC(LogPlatformEvents, Log, All);
+
+/////////////////////////////////////////////////////////////////////
+
+namespace {
+
+class FPlatformEvents : public FRunnable
+{
+public:
+	FPlatformEvents(uint32 SamplingIntervalUsec);
+	~FPlatformEvents();
+
+	void Enable(EPlatformEvent Event);
+	void Disable(EPlatformEvent Event);
+
+private:
+	virtual uint32 Run() override;
+	virtual void Stop() override;
+
+	bool StartETW();
+	void StopETW();
+
+	static void WINAPI TraceEventCallback(PEVENT_RECORD Event);
+	static uint32 CurrentProcessId;
+
+	bool bStarting = false;
+	bool bError = false;
+	FEventRef ReadyEvent;
+	FRunnableThread* Thread = nullptr;
+	uint32 SamplingIntervalUsec;
+
+	EPlatformEvent EnabledEvents = EPlatformEvent::None;
+
+	struct FTrace
+	{
+		EVENT_TRACE_PROPERTIES Properties;
+		WCHAR SessionName[1024];
+	};
+
+	TRACEHANDLE TraceHandle = INVALID_PROCESSTRACE_HANDLE;
+	FTrace Trace;
+
+	// timestamp for each core when thread start executing on it
+	uint64 ContexSwitchStart[256];
+};
+
+// https://docs.microsoft.com/en-us/windows/win32/etw/cswitch
+struct FContextSwitchEvent
+{
+	enum { Type = 36 };
+
+	uint32 NewThreadId;					// 0
+	uint32 OldThreadId;					// 4
+	uint8  NewThreadPriority;			// 8
+	uint8  OldThreadPriority;			// 9
+	uint8  PreviousCState;				// 10
+	uint8  SpareByte;					// 11
+	uint8  OldThreadWaitReason;			// 12
+	uint8  OldThreadWaitMode;			// 13
+	uint8  OldThreadState;				// 14
+	uint8  OldThreadWaitIdealProcessor;	// 15
+	uint32 NewThreadWaitTime;			// 16
+	uint32 Reserved;					// 20
+};
+
+// https://docs.microsoft.com/en-us/windows/win32/etw/stackwalk-event
+struct FStackWalkEvent
+{
+	enum { Type = 32 };
+
+	uint64 TimeStamp;
+	uint32 ProcessId;
+	uint32 ThreadId;
+
+	// this array can be max 192 elements
+	// actual count is determined by total byte size of event
+	uint64 Stack[];
+};
+
+} // namespace
+
+/////////////////////////////////////////////////////////////////////
+
+static FPlatformEvents* GPlatformEvents;
+
+/////////////////////////////////////////////////////////////////////
+
+uint32 FPlatformEvents::CurrentProcessId;
+
+void FPlatformEvents::TraceEventCallback(PEVENT_RECORD Event)
+{
+	uint32 ProcessId = Event->EventHeader.ProcessId;
+	int Opcode = Event->EventHeader.EventDescriptor.Opcode;
+	if (Opcode == FContextSwitchEvent::Type)
+	{
+		const FContextSwitchEvent* Payload = (const FContextSwitchEvent*)Event->UserData;
+		if (Payload == nullptr || Event->UserDataLength != sizeof(FContextSwitchEvent))
+		{
+			return;
+		}
+
+		uint8 CoreNumber = Event->BufferContext.ProcessorNumber;
+
+		uint64 StartTime = GPlatformEvents->ContexSwitchStart[CoreNumber];
+		uint64 EndTime = Event->EventHeader.TimeStamp.QuadPart;
+
+		if (StartTime != 0)
+		{
+			uint32 ThreadId = Payload->OldThreadId;
+
+			if (ThreadId != 0)
+			{
+				UE_TRACE_LOG(PlatformEvent, ContextSwitch, ContextSwitchChannel)
+					<< ContextSwitch.StartTime(StartTime)
+					<< ContextSwitch.EndTime(EndTime)
+					<< ContextSwitch.ThreadId(ThreadId)
+					<< ContextSwitch.CoreNumber(CoreNumber);
+			}
+		}
+
+		GPlatformEvents->ContexSwitchStart[CoreNumber] = EndTime;
+	}
+	else if (Opcode == FStackWalkEvent::Type)
+	{
+		const FStackWalkEvent* Payload = (const FStackWalkEvent*)Event->UserData;
+		if (Payload == nullptr || Payload->ProcessId != CurrentProcessId || Event->UserDataLength < sizeof(FStackWalkEvent))
+		{
+			return;
+		}
+
+		uint32 Count = (Event->UserDataLength - sizeof(FStackWalkEvent)) / sizeof(uint64);
+
+		// sometimes stack walk from ETW gives a call-stack that is fully in kernel address space
+		// so if all addresses have a high bit set, then ignore this call stack
+		bool bValid = false;
+		for (uint32 Index = 0; Index < Count; Index++)
+		{
+			uint64 Address = Payload->Stack[Index];
+			if ((Address >> 63) == 0)
+			{
+				bValid = true;
+				break;
+			}
+		}
+
+		// sometimes we get call stack with only one or two entries, which obviously is not very useful
+		if (bValid && Count > 2)
+		{
+			UE_TRACE_LOG(PlatformEvent, StackSample, StackSamplingChannel)
+				<< StackSample.Time(Payload->TimeStamp)
+				<< StackSample.ThreadId(Payload->ThreadId)
+				<< StackSample.Addresses(Payload->Stack, Count);
+		}
+	}
+}
+
+/////////////////////////////////////////////////////////////////////
+
+FPlatformEvents::FPlatformEvents(uint32 SamplingIntervalUsec)
+	: ReadyEvent(EEventMode::ManualReset)
+	, SamplingIntervalUsec(SamplingIntervalUsec)
+{
+	ZeroMemory(ContexSwitchStart, sizeof(ContexSwitchStart));
+	CurrentProcessId = ::GetCurrentProcessId();
+
+	ReadyEvent->Trigger();
+}
+
+/////////////////////////////////////////////////////////////////////
+
+FPlatformEvents::~FPlatformEvents()
+{
+	Stop();
+}
+
+/////////////////////////////////////////////////////////////////////
+
+void FPlatformEvents::Enable(EPlatformEvent Event)
+{
+	if (bStarting)
+	{
+		// in case ETW is already starting, wait for startup code to complete
+		ReadyEvent->Wait(MAX_uint32);
+	}
+
+	if (bError)
+	{
+		// in case ETW failed to start up earlier, do nothing
+		return;
+	}
+
+	if ((EnabledEvents & Event) != EPlatformEvent::None)
+	{
+		// if event is already enabled, do nothing
+		return;
+	}
+
+	if (TraceHandle == INVALID_PROCESSTRACE_HANDLE)
+	{
+		// in case ETW is not started, start it now
+		ReadyEvent->Reset();
+		bError = false;
+		bStarting = true;
+		Thread = FRunnableThread::Create(this, TEXT("PlatformEvents"), 0, TPri_Normal);
+	}
+	else
+	{
+		// in case ETW was already running, enable the necessary event
+
+		EVENT_TRACE_PROPERTIES* Properties = &Trace.Properties;
+		if (Event == EPlatformEvent::ContextSwitch)
+		{
+			Properties->EnableFlags |= EVENT_TRACE_FLAG_CSWITCH;
+		}
+		else if (Event == EPlatformEvent::StackSampling)
+		{
+			Properties->EnableFlags |= EVENT_TRACE_FLAG_PROFILE;
+		}
+		
+		ULONG Status = ::ControlTraceW(NULL, KERNEL_LOGGER_NAMEW, Properties, EVENT_TRACE_CONTROL_UPDATE);
+		if (Status != ERROR_SUCCESS)
+		{
+			UE_LOG(LogPlatformEvents, Error, TEXT("Unable to update ETW trace: 0x%08x"), Status);
+		}
+	}
+
+	EnumAddFlags(EnabledEvents, Event);
+}
+
+/////////////////////////////////////////////////////////////////////
+
+void FPlatformEvents::Disable(EPlatformEvent Event)
+{
+	if (bStarting)
+	{
+		// in case ETW is already starting, wait for startup code to complete
+		ReadyEvent->Wait(MAX_uint32);
+	}
+
+	if (TraceHandle == INVALID_PROCESSTRACE_HANDLE)
+	{
+		// in case ETW is not running, do nothing
+		return;
+	}
+
+	if ((EnabledEvents & Event) == EPlatformEvent::None)
+	{
+		// if event is already disabled, do nothing
+		return;
+	}
+	
+	EnumRemoveFlags(EnabledEvents, Event);
+
+	if (EnabledEvents == EPlatformEvent::None)
+	{
+		// in case all events are disabled, stop ETW
+		Stop();
+	}
+	else
+	{
+		// in case ETW still needs to run, disable only specific event
+
+		EVENT_TRACE_PROPERTIES* Properties = &Trace.Properties;
+		if (Event == EPlatformEvent::ContextSwitch)
+		{
+			Properties->EnableFlags &= ~EVENT_TRACE_FLAG_CSWITCH;
+		}
+		else if (Event == EPlatformEvent::StackSampling)
+		{
+			Properties->EnableFlags &= ~EVENT_TRACE_FLAG_PROFILE;
+		}
+
+		ReadyEvent->Wait(MAX_uint32);
+		ULONG Status = ::ControlTraceW(NULL, KERNEL_LOGGER_NAMEW, Properties, EVENT_TRACE_CONTROL_UPDATE);
+		if (Status != ERROR_SUCCESS)
+		{
+			UE_LOG(LogPlatformEvents, Error, TEXT("Unable to update ETW trace: 0x%08x"), Status);
+		}
+	}
+}
+
+/////////////////////////////////////////////////////////////////////
+
+uint32 FPlatformEvents::Run()
+{
+	if (StartETW())
+	{
+		ReadyEvent->Trigger();
+		bStarting = false;
+
+		// ProcessTrace is blocking function that enters processing loop
+		// and calls TraceEventCallback callback for each event
+		// it exits processing loop when trace is closed with ::CloseTrace 
+		::ProcessTrace(&TraceHandle, 1, NULL, NULL);
+	}
+	else
+	{
+		// in case of error in StartETW we want to be sure ETW is stopped
+		StopETW();
+		bError = true;
+		bStarting = false;
+		ReadyEvent->Trigger();
+	}
+	return 0;
+}
+
+/////////////////////////////////////////////////////////////////////
+
+void FPlatformEvents::Stop()
+{
+	// because ETW is started in background thread, then first make sure
+	// ETW startup code has completed (sucessfully or not) before actually
+	// stopping it
+	ReadyEvent->Wait(MAX_uint32);
+
+	if (Thread)
+	{
+		StopETW();
+		Thread->WaitForCompletion();
+	}
+}
+
+/////////////////////////////////////////////////////////////////////
+
+bool FPlatformEvents::StartETW()
+{
+	// explicitly define GUIDs to simplify dependency on special defines in windows headers
+	// https://docs.microsoft.com/en-us/windows/win32/etw/nt-kernel-logger-constants
+	static const GUID ETW_PerfInfoGuid = { 0xce1dbfb4, 0x137e, 0x4da6, { 0x87, 0xb0, 0x3f, 0x59, 0xaa, 0x10, 0x2c, 0xbc } };
+	static const GUID ETW_SystemTraceControlGuid = { 0x9e814aad, 0x3204, 0x11d2, { 0x9a, 0x82, 0x00, 0x60, 0x08, 0xa8, 0x69, 0x39 } };
+
+	TRACEHANDLE Session;
+
+	SYSTEM_INFO SysInfo;
+	GetSystemInfo(&SysInfo);
+
+	EVENT_TRACE_PROPERTIES* Properties = &Trace.Properties;
+
+	// stop existing trace in case it is already running (from other tools)
+	{
+		ZeroMemory(Properties, sizeof(*Properties));
+		Properties->Wnode.BufferSize = sizeof(Trace);
+		Properties->Wnode.Guid = ETW_SystemTraceControlGuid;
+		Properties->LogFileMode = EVENT_TRACE_REAL_TIME_MODE;
+		Properties->LoggerNameOffset = sizeof(Trace.Properties);
+		::ControlTraceW(NULL, KERNEL_LOGGER_NAMEW, Properties, EVENT_TRACE_CONTROL_STOP);
+	}
+
+	// enable profiling privilege for process
+	{
+		HANDLE Token;
+		if (::OpenProcessToken(::GetCurrentProcess(), TOKEN_ADJUST_PRIVILEGES, &Token))
+		{
+			LUID Luid;
+			if (::LookupPrivilegeValue(NULL, SE_SYSTEM_PROFILE_NAME, &Luid))
+			{
+				TOKEN_PRIVILEGES TokenPrivileges;
+				TokenPrivileges.PrivilegeCount = 1;
+				TokenPrivileges.Privileges[0].Luid = Luid;
+				TokenPrivileges.Privileges[0].Attributes = SE_PRIVILEGE_ENABLED;
+				if (!::AdjustTokenPrivileges(Token, 0, &TokenPrivileges, sizeof(TokenPrivileges), NULL, NULL))
+				{
+					uint32 Error = ::GetLastError();
+					UE_LOG(LogPlatformEvents, Warning, TEXT("Cannot enable profile privlege for process: 0x%08x"), Error);
+				}
+			}
+			::CloseHandle(Token);
+		}
+	}
+
+	// set up sampling interval, must be done before StartTrace
+	{
+		typedef struct
+		{
+			ULONG Source;
+			ULONG Interval;
+		}
+		TRACE_PROFILE_INTERVAL;
+
+		// interval is specified in 100 nanosecond units, and allowed range is limited
+		// see "-setprofint" argument in https://docs.microsoft.com/en-us/windows-hardware/test/wpt/start
+		TRACE_PROFILE_INTERVAL Interval;
+		Interval.Source = 0;
+		Interval.Interval = FMath::Clamp<ULONG>(10 * SamplingIntervalUsec, 1221, 10000000);
+
+		ULONG Status = ::TraceSetInformation(NULL, TraceSampledProfileIntervalInfo, &Interval, sizeof(Interval));
+		if (Status != ERROR_SUCCESS)
+		{
+			UE_LOG(LogPlatformEvents, Warning, TEXT("Unable to set ETW sampling interval: 0x%08x"), Status);
+		}
+	}
+
+	// setup trace properties
+	ZeroMemory(Properties, sizeof(*Properties));
+	Properties->Wnode.BufferSize = sizeof(Trace);
+	Properties->Wnode.Guid = ETW_SystemTraceControlGuid;
+
+	// NOTE: change this depending on what timestamps you want to collect
+	// 1 - QueryPerformanceCounter (slower)
+	// 3 - RDTSC (faster)
+	Properties->Wnode.ClientContext = 1;
+
+	Properties->Wnode.Flags = WNODE_FLAG_TRACED_GUID;
+	Properties->BufferSize = 256; // 256KB
+	Properties->MinimumBuffers = 4 * SysInfo.dwNumberOfProcessors;
+	Properties->MaximumBuffers = Properties->MinimumBuffers + 20;
+	Properties->LogFileMode = EVENT_TRACE_REAL_TIME_MODE | EVENT_TRACE_SYSTEM_LOGGER_MODE;
+	Properties->LoggerNameOffset = offsetof(FTrace, SessionName);
+	Properties->FlushTimer = 1;
+	Properties->EnableFlags = 0;
+
+	if ((EnabledEvents & EPlatformEvent::ContextSwitch) != EPlatformEvent::None)
+	{
+		Properties->EnableFlags |= EVENT_TRACE_FLAG_CSWITCH;
+	}
+	if ((EnabledEvents & EPlatformEvent::StackSampling) != EPlatformEvent::None)
+	{
+		Properties->EnableFlags |= EVENT_TRACE_FLAG_PROFILE;
+	}
+
+	// start the trace
+	{
+		ULONG Status = ::StartTraceW(&Session, KERNEL_LOGGER_NAMEW, Properties);
+		if (Status == ERROR_ACCESS_DENIED)
+		{
+			UE_LOG(LogPlatformEvents, Error, TEXT("Administrator rights required for ETW"));
+			return false;
+		}
+		else if (Status != ERROR_SUCCESS)
+		{
+			UE_LOG(LogPlatformEvents, Warning, TEXT("Unable to start ETW trace: 0x%08x"), Status);
+			return false;
+		}
+	}
+
+	// select for which events ETW should collect stack (StackWalkEvent), this is harmless if stack sampling is not enabled
+	{
+		// 0x2e -> PERFINFO_LOG_TYPE_SAMPLED_PROFILE
+		CLASSIC_EVENT_ID Events[] = { { ETW_PerfInfoGuid, 0x2e } };
+
+		ULONG Status = ::TraceSetInformation(Session, TraceStackTracingInfo, Events, sizeof(Events));
+		if (Status != ERROR_SUCCESS)
+		{
+			UE_LOG(LogPlatformEvents, Error, TEXT("Cannot enable stack tracing for ETW: 0x%08x"), Status);
+			return false;
+		}
+	}
+
+	// open trace for receving callbacks
+	{
+		EVENT_TRACE_LOGFILEW LogFile;
+		ZeroMemory(&LogFile, sizeof(LogFile));
+		LogFile.LoggerName = (LPWSTR)KERNEL_LOGGER_NAMEW;
+		LogFile.ProcessTraceMode = PROCESS_TRACE_MODE_EVENT_RECORD | PROCESS_TRACE_MODE_RAW_TIMESTAMP | PROCESS_TRACE_MODE_REAL_TIME;
+		LogFile.EventRecordCallback = &TraceEventCallback;
+
+		TraceHandle = ::OpenTraceW(&LogFile);
+		if (TraceHandle == INVALID_PROCESSTRACE_HANDLE)
+		{
+			UE_LOG(LogPlatformEvents, Error, TEXT("Cannot open ETW trace"));
+			return false;
+		}
+	}
+
+	return true;
+}
+
+/////////////////////////////////////////////////////////////////////
+
+void FPlatformEvents::StopETW()
+{
+	// stop the trace
+	::ControlTraceW(0, KERNEL_LOGGER_NAMEW, &Trace.Properties, EVENT_TRACE_CONTROL_STOP);
+
+	if (TraceHandle != INVALID_PROCESSTRACE_HANDLE)
+	{
+		// close processing loop, this will wait until all pending buffers are flushed
+		// and TracingCallback called on all pending events in buffers
+		::CloseTrace(TraceHandle);
+		TraceHandle = INVALID_PROCESSTRACE_HANDLE;
+	}
+}
+
+/////////////////////////////////////////////////////////////////////
+
+void PlatformEvents_Init(uint32 SamplingIntervalUsec)
+{
+	GPlatformEvents = new FPlatformEvents(SamplingIntervalUsec);
+}
+
+void PlatformEvents_Enable(EPlatformEvent Event)
+{
+	if (GPlatformEvents)
+	{
+		GPlatformEvents->Enable(Event);
+	}
+}
+
+void PlatformEvents_Disable(EPlatformEvent Event)
+{
+	if (GPlatformEvents)
+	{
+		GPlatformEvents->Disable(Event);
+	}
+}
+
+void PlatformEvents_Stop()
+{
+	if (GPlatformEvents)
+	{
+		delete GPlatformEvents;
+		GPlatformEvents = nullptr;
+	}
+}
+
+#endif // PLATFORM_SUPPORTS_PLATFORM_EVENTS
diff --git a/Engine/Source/Runtime/Core/Private/ProfilingDebugging/Microsoft/MemoryTrace.cpp b/Engine/Source/Runtime/Core/Private/ProfilingDebugging/Microsoft/MemoryTrace.cpp
new file mode 100644
--- /dev/null
+++ b/Engine/Source/Runtime/Core/Private/ProfilingDebugging/Microsoft/MemoryTrace.cpp
@@ -0,0 +1,482 @@
+// Copyright Epic Games, Inc. All Rights Reserved.
+
+#include "ProfilingDebugging/MemoryTrace.h"
+
+#if UE_MEMORY_TRACE_ENABLED
+
+#include "ProfilingDebugging/MemoryAllocationTrace.h"
+#include "Containers/StringView.h"
+#include "CoreTypes.h"
+#include "HAL/MemoryBase.h"
+#include "HAL/Platform.h"
+#include "HAL/PlatformTime.h"
+#include "Misc/CoreDelegates.h"
+#include "Misc/CString.h"
+#include "Trace/Trace.inl"
+
+#include "Windows/AllowWindowsPlatformTypes.h"
+#include <winternl.h>
+
+#include <intrin.h>
+
+
+////////////////////////////////////////////////////////////////////////////////
+void	Backtracer_Create(FMalloc*);
+void	Backtracer_Initialize();
+void*	Backtracer_GetBacktraceId(void*);
+void 	MemoryTrace_InitTags(FMalloc*);
+
+////////////////////////////////////////////////////////////////////////////////
+FORCEINLINE void* GetOwner(bool bLight)
+{
+	void* RetAddrAddr = PLATFORM_RETURN_ADDRESS_POINTER();
+	if (bLight)
+	{
+		return *(void**)RetAddrAddr;
+	}
+
+	return Backtracer_GetBacktraceId(RetAddrAddr);
+}
+
+////////////////////////////////////////////////////////////////////////////////
+template <class T>
+class alignas(alignof(T)) FUndestructed
+{
+public:
+	template <typename... ArgTypes>
+	void Construct(ArgTypes... Args)
+	{
+		::new (Buffer) T(Args...);
+	}
+
+	T* operator & ()	{ return (T*)Buffer; }
+	T* operator -> ()	{ return (T*)Buffer; }
+
+protected:
+	uint8 Buffer[sizeof(T)];
+};
+
+////////////////////////////////////////////////////////////////////////////////
+struct FAddrPack
+{
+			FAddrPack() = default;
+			FAddrPack(UPTRINT Addr, uint16 Value) { Set(Addr, Value); }
+	void	Set(UPTRINT Addr, uint16 Value) { Inner = uint64(Addr) | (uint64(Value) << 48ull); }
+	uint64	Inner;
+};
+static_assert(sizeof(FAddrPack) == sizeof(uint64), "");
+
+
+////////////////////////////////////////////////////////////////////////////////
+static FUndestructed<FAllocationTrace> GAllocationTrace;
+
+
+////////////////////////////////////////////////////////////////////////////////
+class FMallocWrapper
+	: public FMalloc
+{
+public:
+							FMallocWrapper(FMalloc* InMalloc, bool bInLight);
+
+private:
+	struct FCookie
+	{
+		uint64				Tag  : 16;
+		uint64				Bias : 8;
+		uint64				Size : 40;
+	};
+
+	static uint32			GetActualAlignment(SIZE_T Size, uint32 Alignment);
+	virtual void*			Malloc(SIZE_T Size, uint32 Alignment) override;
+	virtual void*			Realloc(void* PrevAddress, SIZE_T NewSize, uint32 Alignment) override;
+	virtual void			Free(void* Address) override;
+	virtual bool			IsInternallyThreadSafe() const override						{ return InnerMalloc->IsInternallyThreadSafe(); }
+	virtual void			UpdateStats() override										{ InnerMalloc->UpdateStats(); }
+	virtual void			GetAllocatorStats(FGenericMemoryStats& Out) override		{ InnerMalloc->GetAllocatorStats(Out); }
+	virtual void			DumpAllocatorStats(FOutputDevice& Ar) override				{ InnerMalloc->DumpAllocatorStats(Ar); }
+	virtual bool			ValidateHeap() override										{ return InnerMalloc->ValidateHeap(); }
+	virtual bool			GetAllocationSize(void* Address, SIZE_T &SizeOut) override	{ return InnerMalloc->GetAllocationSize(Address, SizeOut); }
+	FMalloc*				InnerMalloc;
+	bool					bLight;
+};
+
+////////////////////////////////////////////////////////////////////////////////
+FMallocWrapper::FMallocWrapper(FMalloc* InMalloc, bool bInLight)
+: InnerMalloc(InMalloc)
+, bLight(bInLight)
+{
+}
+
+////////////////////////////////////////////////////////////////////////////////
+uint32 FMallocWrapper::GetActualAlignment(SIZE_T Size, uint32 Alignment)
+{
+	// Defaults; if size is < 16 then alignment is 8 else 16.
+	uint32 DefaultAlignment = 8 << uint32(Size >= 16);
+	return (Alignment < DefaultAlignment) ? DefaultAlignment : Alignment;
+}
+
+////////////////////////////////////////////////////////////////////////////////
+void* FMallocWrapper::Malloc(SIZE_T Size, uint32 Alignment)
+{
+	if (Size == 0)
+	{
+		return nullptr;
+	}
+
+	uint32 ActualAlignment = GetActualAlignment(Size, Alignment);
+	void* Address = InnerMalloc->Malloc(Size, Alignment);
+
+	void* Owner = GetOwner(bLight);
+	GAllocationTrace->Alloc(Address, Size, ActualAlignment, Owner);
+
+	return Address;
+}
+
+////////////////////////////////////////////////////////////////////////////////
+void* FMallocWrapper::Realloc(void* PrevAddress, SIZE_T NewSize, uint32 Alignment)
+{
+	// This simplifies things and means reallocs trace events are true reallocs
+	if (PrevAddress == nullptr)
+	{
+		return Malloc(NewSize, Alignment);
+	}
+
+	if (NewSize == 0)
+	{
+		Free(PrevAddress);
+		return nullptr;
+	}
+
+	GAllocationTrace->ReallocFree(PrevAddress);
+
+	void* RetAddress = InnerMalloc->Realloc(PrevAddress, NewSize, Alignment);
+
+	void* Owner = GetOwner(bLight);
+	Alignment = GetActualAlignment(NewSize, Alignment);
+	GAllocationTrace->ReallocAlloc(RetAddress, NewSize, Alignment, Owner);
+
+	return RetAddress;
+}
+
+////////////////////////////////////////////////////////////////////////////////
+void FMallocWrapper::Free(void* Address)
+{
+	if (Address == nullptr)
+	{
+		return;
+	}
+
+	GAllocationTrace->Free(Address);
+
+	void* InnerAddress = Address;
+
+	return InnerMalloc->Free(InnerAddress);
+}
+
+
+
+////////////////////////////////////////////////////////////////////////////////
+#if defined(PLATFORM_SUPPORTS_TRACE_WIN32_VIRTUAL_MEMORY_HOOKS)
+class FTextSectionEditor
+{
+public:
+							FTextSectionEditor(void* InBase);
+							~FTextSectionEditor();
+	template <typename T>
+	T*						Hook(T* Target, T* HookFunction);
+
+private:
+	static void*			GetActualAddress(void* Function);
+	uint8*					AllocateTrampoline(unsigned int Size);
+	void*					HookImpl(void* Target, void* HookFunction);
+	uint8*					TrampolineTail;
+	void*					Base;
+	size_t					Size;
+	DWORD					Protection;
+};
+
+////////////////////////////////////////////////////////////////////////////////
+FTextSectionEditor::FTextSectionEditor(void* InBase)
+{
+	InBase = GetActualAddress(InBase);
+
+	MEMORY_BASIC_INFORMATION MemInfo;
+	VirtualQuery(InBase, &MemInfo, sizeof(MemInfo));
+	Base = MemInfo.BaseAddress;
+	Size = MemInfo.RegionSize;
+
+	VirtualProtect(Base, Size, PAGE_EXECUTE_READWRITE, &Protection);
+
+	TrampolineTail = (uint8*)Base + Size;
+}
+
+////////////////////////////////////////////////////////////////////////////////
+FTextSectionEditor::~FTextSectionEditor()
+{
+	VirtualProtect(Base, Size, Protection, &Protection);
+	FlushInstructionCache(GetCurrentProcess(), Base, Size);
+}
+
+////////////////////////////////////////////////////////////////////////////////
+void* FTextSectionEditor::GetActualAddress(void* Function)
+{
+	uint8* Addr = (uint8*)Function;
+	int Offset = unsigned(Addr[0] & 0xf0) == 0x40; // REX prefix
+	if (Addr[Offset + 0] == 0xff && Addr[Offset + 1] == 0x25)
+	{
+		Addr += Offset;
+		Addr = *(uint8**)(Addr + 6 + *(uint32*)(Addr + 2));
+	}
+	return Addr;
+}
+
+////////////////////////////////////////////////////////////////////////////////
+uint8* FTextSectionEditor::AllocateTrampoline(unsigned int PatchSize)
+{
+	static const int TrampolineSize = 24;
+	uint8* NextTail = TrampolineTail - TrampolineSize;
+	do 
+	{
+		check(*NextTail == 0);
+		++NextTail;
+	}
+	while (NextTail != TrampolineTail);
+
+	TrampolineTail -= TrampolineSize;
+	return TrampolineTail;
+}
+
+////////////////////////////////////////////////////////////////////////////////
+template <typename T>
+T* FTextSectionEditor::Hook(T* Target, T* HookFunction)
+{
+	return (T*)HookImpl((void*)Target, (void*)HookFunction);
+}
+
+////////////////////////////////////////////////////////////////////////////////
+void* FTextSectionEditor::HookImpl(void* Target, void* HookFunction)
+{
+	Target = GetActualAddress(Target);
+
+	const uint8* Start = (uint8*)Target;
+	const uint8* Read = Start;
+	do
+	{
+		Read += (Read[0] & 0xf0) == 0x40; // REX prefix
+		uint8 Inst = *Read++;
+		if (unsigned(Inst - 0x80) < 0x0cu)
+		{
+			uint8 ModRm = *Read++;
+			Read += ((ModRm & 0300) < 0300) & ((ModRm & 0007) == 0004); // SIB
+			switch (ModRm & 0300) // Disp[8|32]
+			{
+			case 0100: Read += 1; break;
+			case 0200: Read += 5; break;
+			}
+			Read += (Inst == 0x83);
+		}
+		else if (unsigned(Inst - 0x50) >= 0x10u)
+			check(!"Unknown instruction");
+	}
+	while (Read - Start < 6);
+
+	int PatchSize = int(Read - Start);
+	uint8* TrampolinePtr = AllocateTrampoline(PatchSize);
+
+	*(void**)TrampolinePtr = HookFunction;
+
+	uint8* PatchJmp = TrampolinePtr + sizeof(void*);
+	memcpy(PatchJmp, Start, PatchSize);
+
+	PatchJmp += PatchSize;
+	*PatchJmp = 0xe9;
+	*(int32*)(PatchJmp + 1) = int32(intptr_t(Start + PatchSize) - intptr_t(PatchJmp)) - 5;
+
+	uint16* HookJmp = (uint16*)Target;
+	HookJmp[0] = 0x25ff;
+	*(int32*)(HookJmp + 1) = int32(intptr_t(TrampolinePtr) - intptr_t(HookJmp + 3));
+
+	return PatchJmp - PatchSize;
+}
+
+
+
+////////////////////////////////////////////////////////////////////////////////
+class FVirtualWinApiHooks
+{
+public:
+	static void				Initialize(bool bInLight);
+
+private:
+							FVirtualWinApiHooks();
+	static bool				bLight;
+	static LPVOID WINAPI	VmAlloc(LPVOID Address, SIZE_T Size, DWORD Type, DWORD Protect);
+	static BOOL WINAPI		VmFree(LPVOID Address, SIZE_T Size, DWORD Type);
+	static BOOL WINAPI		VmFreeEx(HANDLE Process, LPVOID Address, SIZE_T Size, DWORD Type);
+	static LPVOID WINAPI	VmAllocEx(HANDLE Process, LPVOID Address, SIZE_T Size, DWORD Type, DWORD Protect);
+	static LPVOID			(WINAPI *VmAllocOrig)(LPVOID, SIZE_T, DWORD, DWORD);
+	static LPVOID			(WINAPI *VmAllocExOrig)(HANDLE, LPVOID, SIZE_T, DWORD, DWORD);
+	static BOOL				(WINAPI *VmFreeOrig)(LPVOID, SIZE_T, DWORD);
+	static BOOL				(WINAPI *VmFreeExOrig)(HANDLE, LPVOID, SIZE_T, DWORD);
+};
+
+////////////////////////////////////////////////////////////////////////////////
+bool	FVirtualWinApiHooks::bLight;
+LPVOID	(WINAPI *FVirtualWinApiHooks::VmAllocOrig)(LPVOID, SIZE_T, DWORD, DWORD);
+LPVOID	(WINAPI *FVirtualWinApiHooks::VmAllocExOrig)(HANDLE, LPVOID, SIZE_T, DWORD, DWORD);
+BOOL	(WINAPI *FVirtualWinApiHooks::VmFreeOrig)(LPVOID, SIZE_T, DWORD);
+BOOL	(WINAPI *FVirtualWinApiHooks::VmFreeExOrig)(HANDLE, LPVOID, SIZE_T, DWORD);
+
+////////////////////////////////////////////////////////////////////////////////
+void FVirtualWinApiHooks::Initialize(bool bInLight)
+{
+	bLight = bInLight;
+
+	{
+		FTextSectionEditor Editor((void*)VirtualAlloc);
+		VmAllocOrig = Editor.Hook(VirtualAlloc, &FVirtualWinApiHooks::VmAlloc);
+		VmFreeOrig = Editor.Hook(VirtualFree, &FVirtualWinApiHooks::VmFree); 
+		
+		{
+			FTextSectionEditor EditorEx((void*)VirtualAllocEx);
+			VmAllocExOrig = Editor.Hook(VirtualAllocEx, &FVirtualWinApiHooks::VmAllocEx);
+			VmFreeExOrig = Editor.Hook(VirtualFreeEx, &FVirtualWinApiHooks::VmFreeEx);
+		}
+	}
+
+	
+}
+
+////////////////////////////////////////////////////////////////////////////////
+LPVOID WINAPI FVirtualWinApiHooks::VmAlloc(LPVOID Address, SIZE_T Size, DWORD Type, DWORD Protect)
+{
+	LPVOID Ret = VmAllocOrig(Address, Size, Type, Protect);
+	if (Ret != nullptr && (Type & MEM_COMMIT))
+	{
+		void* Owner = GetOwner(bLight);
+		GAllocationTrace->CoreAdd(Ret, Size, Owner);
+	}
+
+	return Ret;
+}
+
+////////////////////////////////////////////////////////////////////////////////
+BOOL WINAPI FVirtualWinApiHooks::VmFree(LPVOID Address, SIZE_T Size, DWORD Type)
+{
+	void* Owner = GetOwner(bLight);
+	GAllocationTrace->CoreRemove(Address, Size, Owner);
+	return VmFreeOrig(Address, Size, Type);
+}
+
+////////////////////////////////////////////////////////////////////////////////
+LPVOID WINAPI FVirtualWinApiHooks::VmAllocEx(HANDLE Process, LPVOID Address, SIZE_T Size, DWORD Type, DWORD Protect)
+{
+	LPVOID Ret = VmAllocExOrig(Process, Address, Size, Type, Protect);
+	if (Process == GetCurrentProcess() && Ret != nullptr && (Type & MEM_COMMIT))
+	{
+		void* Owner = GetOwner(bLight);
+		GAllocationTrace->CoreAdd(Ret, Size, Owner);
+	}
+
+	return Ret;
+}
+
+////////////////////////////////////////////////////////////////////////////////
+BOOL WINAPI FVirtualWinApiHooks::VmFreeEx(HANDLE Process, LPVOID Address, SIZE_T Size, DWORD Type)
+{
+	if (Process == GetCurrentProcess())
+	{
+		void* Owner = GetOwner(bLight);
+		GAllocationTrace->CoreRemove(Address, Size, Owner);
+	}
+
+	return VmFreeExOrig(Process, Address, Size, Type);
+}
+#endif // defined(PLATFORM_SUPPORTS_TRACE_WIN32_VIRTUAL_MEMORY_HOOKS)
+
+
+////////////////////////////////////////////////////////////////////////////////
+FMalloc* MemoryTrace_Create(FMalloc* InMalloc)
+{
+	int32 Mode = 0;
+	const TCHAR* CmdLine = ::GetCommandLineW();
+	if (const TCHAR* TraceArg = FCString::Strstr(CmdLine, TEXT("-trace=")))
+	{
+		const TCHAR* Start = TraceArg + 7;
+		const TCHAR* End = Start;
+		for (; *End != ' ' && *End != '\0'; ++End);
+
+		for (const TCHAR* c = Start; c < End + 1; ++c)
+		{
+			if (*c == ' ' || *c == '\0' || *c == ',')
+			{
+				FStringView View(Start, uint32(c - Start));
+				if (View.Equals(TEXT("memalloc"), ESearchCase::IgnoreCase) || View.Equals(TEXT("memory"), ESearchCase::IgnoreCase))
+				{
+					Mode = 2;
+					break;
+				}
+
+				Start = c + 1;
+			}
+		}
+	}
+
+	if (Mode > 0)
+	{
+		// Some OSes (i.e. Windows) will terminate all threads except the main
+		// one as part of static deinit. However we may receive more memory
+		// trace events that would get lost as Trace's worker thread has been
+		// terminated. So flush the last remaining memory events trace needs
+		// to be updated which we will do that in response to to memory events.
+		// We'll use an atexit can to know when Trace is probably no longer
+		// getting ticked.
+		atexit([] () { GAllocationTrace->EnableTracePump(); });
+
+		GAllocationTrace.Construct();
+		GAllocationTrace->Initialize();
+
+		MemoryTrace_InitTags(InMalloc);
+
+		Backtracer_Create(InMalloc);
+
+#if defined(PLATFORM_SUPPORTS_TRACE_WIN32_VIRTUAL_MEMORY_HOOKS)
+		FVirtualWinApiHooks::Initialize(false);
+#endif // defined(PLATFORM_SUPPORTS_TRACE_WIN32_VIRTUAL_MEMORY_HOOKS)
+
+		static FUndestructed<FMallocWrapper> MemoryTrace;
+		MemoryTrace.Construct(InMalloc, false);
+
+		return &MemoryTrace;
+	}
+
+	return InMalloc;
+}
+
+////////////////////////////////////////////////////////////////////////////////
+void MemoryTrace_Initialize()
+{
+	// Allocators aren't completely ready in _Create() so we have an extra step
+	// where any initialisation that may allocate can go.
+	Backtracer_Initialize();
+}
+
+#if 0
+////////////////////////////////////////////////////////////////////////////////
+void MemoryTrace_CoreAdd(void* Base, SIZE_T Size)
+{
+	void* Owner = GetOwner(bLight);
+	GAllocationTrace->CoreAdd(Base, Size, Owner);
+}
+
+////////////////////////////////////////////////////////////////////////////////
+void MemoryTrace_CoreRemove(void* Base, SIZE_T Size)
+{
+	void* Owner = GetOwner(bLight);
+	GAllocationTrace->CoreRemove(Base, Size, Owner);
+}
+#endif // 0
+
+#include "Windows/HideWindowsPlatformTypes.h"
+
+#endif // UE_MEMORY_TRACE_ENABLED
diff --git a/Engine/Source/Runtime/Core/Private/ProfilingDebugging/Microsoft/WindowsCallstackTrace.cpp b/Engine/Source/Runtime/Core/Private/ProfilingDebugging/Microsoft/WindowsCallstackTrace.cpp
new file mode 100644
--- /dev/null
+++ b/Engine/Source/Runtime/Core/Private/ProfilingDebugging/Microsoft/WindowsCallstackTrace.cpp
@@ -0,0 +1,671 @@
+// Copyright Epic Games, Inc. All Rights Reserved.
+
+#if defined(PLATFORM_SUPPORTS_TRACE_WIN32_CALLSTACK)
+
+#include "Containers/Array.h"
+#include "Containers/ArrayView.h"
+#include "CoreTypes.h"
+#include "HAL/CriticalSection.h"
+#include "HAL/MemoryBase.h"
+#include "HAL/RunnableThread.h"
+#include "HAL/LowLevelMemTracker.h"
+#include "ProfilingDebugging/MemoryTrace.h"
+
+#include "Windows/AllowWindowsPlatformTypes.h"
+#	include <winnt.h>
+#	include <winternl.h>
+#include "Windows/HideWindowsPlatformTypes.h"
+#include "Misc/CoreDelegates.h"
+#include "Trace/Trace.inl"
+#include "HAL/Runnable.h"
+#include "HAL/PlatformProcess.h"
+#include "Misc/ScopeLock.h"
+#include "ProfilingDebugging/CallstackTrace.h"
+
+#include <atomic>
+
+// 0=off, 1=stats, 2=validation, 3=truth_compare
+#define BACKTRACE_DBGLVL 0
+
+/*
+ * Windows' x64 binaries contain a ".pdata" section that describes the location
+ * and size of its functions and details on how to unwind them. The unwind
+ * information includes descriptions about a function's stack frame size and
+ * the non-volatile registers it pushes onto the stack. From this we can
+ * calculate where a call instruction wrote its return address. This is enough
+ * to walk the callstack and by caching this information it can be done
+ * efficiently.
+ *
+ * Some functions need a variable amount of stack (such as those that use
+ * alloc() for example) will use a frame pointer. Frame pointers involve saving
+ * and restoring the stack pointer in the function's prologue/epilogue. This
+ * frees the function up to modify the stack pointer arbitrarily. This
+ * significantly complicates establishing where a return address is, so this
+ * pdata scheme of walking the stack just doesn't support functions like this.
+ * Walking stops if it encounters such a function. Fortunately there are
+ * usually very few such functions, saving us from having to read and track
+ * non-volatile registers which adds a significant amount of work.
+ *
+ * A further optimisation is to to assume we are only interested methods that
+ * are part of engine or game code. As such we only build lookup tables for
+ * such modules and never accept OS or third party modules. Backtracing stops
+ * if an address is encountered which doesn't map to a known module.
+ */
+
+////////////////////////////////////////////////////////////////////////////////
+static uint32 AddressToId(UPTRINT Address)
+{
+	return uint32(Address >> 16);
+}
+
+static UPTRINT IdToAddress(uint32 Id)
+{
+	return static_cast<uint32>(UPTRINT(Id) << 16);
+}
+
+struct FIdPredicate
+{
+	template <class T> bool operator () (uint32 Id, const T& Item) const { return Id < Item.Id; }
+	template <class T> bool operator () (const T& Item, uint32 Id) const { return Item.Id < Id; }
+};
+
+////////////////////////////////////////////////////////////////////////////////
+struct FUnwindInfo
+{
+	uint8	Version : 3;
+	uint8	Flags : 5;
+	uint8	PrologBytes;
+	uint8	NumUnwindCodes;
+	uint8	FrameReg	: 4;
+	uint8	FrameRspBias : 4;
+};
+
+struct FUnwindCode
+{
+	uint8	PrologOffset;
+	uint8	OpCode : 4;
+	uint8	OpInfo : 4;
+	uint16	Params[];
+};
+
+enum
+{
+	UWOP_PUSH_NONVOL		= 0,	// 1 node
+	UWOP_ALLOC_LARGE		= 1,	// 2 or 3 nodes
+	UWOP_ALLOC_SMALL		= 2,	// 1 node
+	UWOP_SET_FPREG			= 3,	// 1 node
+	UWOP_SAVE_NONVOL		= 4,	// 2 nodes
+	UWOP_SAVE_NONVOL_FAR	= 5,	// 3 nodes
+	UWOP_SAVE_XMM128		= 8,	// 2 nodes
+	UWOP_SAVE_XMM128_FAR	= 9,	// 3 nodes
+	UWOP_PUSH_MACHFRAME		= 10,	// 1 node
+};
+
+////////////////////////////////////////////////////////////////////////////////
+class FBacktracer
+{
+public:
+							FBacktracer(FMalloc* InMalloc);
+							~FBacktracer();
+	static FBacktracer*		Get();
+	void					AddModule(UPTRINT Base, const TCHAR* Name);
+	void					RemoveModule(UPTRINT Base);
+	void*					GetBacktraceId(void* AddressOfReturnAddress) const;
+	int32					GetFrameSize(void* FunctionAddress) const;
+
+private:
+	struct FFunction
+	{
+		uint32				Id;
+		int32				RspBias;
+#if BACKTRACE_DBGLVL >= 2
+		uint32				Size;
+		const FUnwindInfo*	UnwindInfo;
+#endif
+	};
+
+	struct FModule
+	{
+		uint32				Id;
+		uint32				IdSize = 0;
+		uint32				NumFunctions;
+#if BACKTRACE_DBGLVL >= 1
+		uint16				NumFpTypes;
+		//uint16			*padding*
+#else
+		//uint32			*padding*
+#endif
+		FFunction*			Functions;
+	};
+
+	struct FLookupState
+	{
+		FModule				Module;
+	};
+
+	const FFunction*		LookupFunction(UPTRINT Address, FLookupState& State) const;
+	static FBacktracer*		Instance;
+	mutable FCriticalSection Lock;
+	FModule*				Modules;
+	int32					ModulesNum;
+	int32					ModulesCapacity;
+	FMalloc*				Malloc;
+	FCallstackTracer*		CallstackTracer;
+#if BACKTRACE_DBGLVL >= 1
+	mutable uint32			NumFpTruncations = 0;
+	mutable uint32			TotalFunctions = 0;
+#endif
+};
+
+////////////////////////////////////////////////////////////////////////////////
+FBacktracer* FBacktracer::Instance = nullptr;
+
+////////////////////////////////////////////////////////////////////////////////
+FBacktracer::FBacktracer(FMalloc* InMalloc)
+	: Malloc(InMalloc)
+{
+	ModulesCapacity = 8;
+	ModulesNum = 0;
+	Modules = (FModule*)Malloc->Malloc(sizeof(FModule) * ModulesCapacity);
+
+	Instance = this;
+	CallstackTracer = (FCallstackTracer*) Malloc->Malloc(sizeof(FCallstackTracer), alignof(FCallstackTracer));
+	CallstackTracer = new(CallstackTracer) FCallstackTracer();
+}
+
+////////////////////////////////////////////////////////////////////////////////
+FBacktracer::~FBacktracer()
+{
+	TArrayView<FModule> ModulesView(Modules, ModulesNum);
+	for (FModule& Module : ModulesView)
+	{
+		Malloc->Free(Module.Functions);
+	}
+
+}
+
+////////////////////////////////////////////////////////////////////////////////
+FBacktracer* FBacktracer::Get()
+{
+	return Instance;
+}
+
+////////////////////////////////////////////////////////////////////////////////
+void FBacktracer::AddModule(UPTRINT ModuleBase, const TCHAR* Name)
+{
+	const auto* DosHeader = (IMAGE_DOS_HEADER*)ModuleBase;
+	const auto* NtHeader = (IMAGE_NT_HEADERS*)(ModuleBase + DosHeader->e_lfanew);
+	const IMAGE_FILE_HEADER* FileHeader = &(NtHeader->FileHeader);
+
+	int32 NameLen = FCString::Strlen(Name);
+	if (!(NameLen > 4 && FCString::Strcmp(Name + NameLen - 4, TEXT(".exe")) == 0) && // we want always load an .exe file
+		(FCString::Strfind(Name, TEXT("Binaries")) == nullptr || FCString::Strfind(Name, TEXT("ThirdParty")) != nullptr))
+	{
+		return;
+	}
+
+	uint32 NumSections = FileHeader->NumberOfSections;
+	const auto* Sections = (IMAGE_SECTION_HEADER*)(UPTRINT(&(NtHeader->OptionalHeader)) + FileHeader->SizeOfOptionalHeader);
+
+	// Find ".pdata" section
+	UPTRINT PdataBase = 0;
+	UPTRINT PdataEnd = 0;
+	for (uint32 i = 0; i < NumSections; ++i)
+	{
+		const IMAGE_SECTION_HEADER* Section = Sections + i;
+		if (*(uint64*)(Section->Name) == 0x61'74'61'64'70'2eull) // Sections names are eight bytes and zero padded. This constant is '.pdata'
+		{
+			PdataBase = ModuleBase + Section->VirtualAddress;
+			PdataEnd = PdataBase + Section->SizeOfRawData;
+			break;
+		}
+	}
+
+	if (PdataBase == 0)
+	{
+		return;
+	}
+
+	// Count the number of functions. The assumption here is that if we have got this far then there is at least one function
+	uint32 NumFunctions = uint32(PdataEnd - PdataBase) / sizeof(RUNTIME_FUNCTION);
+	const auto* FunctionTables = (RUNTIME_FUNCTION*)PdataBase;
+	do
+	{
+		const RUNTIME_FUNCTION* Function = FunctionTables + NumFunctions - 1;
+		if (uint32(Function->BeginAddress) < uint32(Function->EndAddress))
+		{
+			break;
+		}
+
+		--NumFunctions;
+	}
+	while (NumFunctions != 0);
+		
+	// Allocate some space for the module's function-to-frame-size table
+	auto* OutTable = (FFunction*)Malloc->Malloc(sizeof(FFunction) * NumFunctions);
+	FFunction* OutTableCursor = OutTable;
+
+	// Extract frame size for each function from pdata's unwind codes.
+	uint32 NumFpFuncs = 0;
+	for (uint32 i = 0; i < NumFunctions; ++i)
+	{
+		const RUNTIME_FUNCTION* FunctionTable = FunctionTables + i;
+
+		UPTRINT UnwindInfoAddr = ModuleBase + FunctionTable->UnwindInfoAddress;
+		const auto* UnwindInfo = (FUnwindInfo*)UnwindInfoAddr;
+
+		if (UnwindInfo->Version != 1)
+		{
+			/* some v2s have been seen in msvc. Always seem to be assembly
+			 * routines (memset, memcpy, etc) */
+			continue;
+		}
+
+		int32 FpInfo = 0;
+		int32 RspBias = 0;
+
+#if BACKTRACE_DBGLVL >= 2
+		uint32 PrologVerify = UnwindInfo->PrologBytes;
+#endif
+
+		const auto* Code = (FUnwindCode*)(UnwindInfo + 1);
+		const auto* EndCode = Code + UnwindInfo->NumUnwindCodes;
+		while (Code < EndCode)
+		{
+#if BACKTRACE_DBGLVL >= 2
+			if (Code->PrologOffset > PrologVerify)
+			{
+				PLATFORM_BREAK();
+			}
+			PrologVerify = Code->PrologOffset;
+#endif
+
+			switch (Code->OpCode)
+			{
+			case UWOP_PUSH_NONVOL:
+				RspBias += 8;
+				Code += 1;
+				break;
+
+			case UWOP_ALLOC_LARGE:
+				if (Code->OpInfo)
+				{
+					RspBias += *(uint32*)(Code->Params);
+					Code += 3;
+				}
+				else
+				{
+					RspBias += Code->Params[0] * 8;
+					Code += 2;
+				}
+				break;
+
+			case UWOP_ALLOC_SMALL:
+				RspBias += (Code->OpInfo * 8) + 8;
+				Code += 1;
+				break;
+
+			case UWOP_SET_FPREG:
+				// Function will adjust RSP (e.g. through use of alloca()) so it
+				// uses a frame pointer register. There's instructions like;
+				//
+				//   push FRAME_REG
+				//   lea FRAME_REG, [rsp + (FRAME_RSP_BIAS * 16)]
+				//   ...
+				//   add rsp, rax
+				//   ...
+				//   sub rsp, FRAME_RSP_BIAS * 16
+				//   pop FRAME_REG
+				//   ret
+				//
+				// To recover the stack frame we would need to track non-volatile
+				// registers which adds a lot of overhead for a small subset of
+				// functions. Instead we'll end backtraces at these functions.
+
+
+				// MSB is set to detect variable sized frames that we can't proceed
+				// past when back-tracing.
+				NumFpFuncs++;
+				FpInfo |= 0x80000000 | (uint32(UnwindInfo->FrameReg) << 27) | (uint32(UnwindInfo->FrameRspBias) << 23);
+				Code += 1;
+				break;
+
+			case UWOP_PUSH_MACHFRAME:
+				RspBias = Code->OpInfo ? 48 : 40;
+				Code += 1;
+				break;
+
+			case UWOP_SAVE_NONVOL:		Code += 2; break; /* saves are movs instead of pushes */
+			case UWOP_SAVE_NONVOL_FAR:	Code += 3; break;
+			case UWOP_SAVE_XMM128:		Code += 2; break;
+			case UWOP_SAVE_XMM128_FAR:	Code += 3; break;
+
+			default:
+#if BACKTRACE_DBGLVL >= 2
+				PLATFORM_BREAK();
+#endif
+				break;
+			}
+		}
+
+		// "Chained" simply means that multiple RUNTIME_FUNCTIONs pertains to a
+		// single actual function in the .text segment.
+		bool bIsChained = (UnwindInfo->Flags & UNW_FLAG_CHAININFO);
+
+		RspBias /= sizeof(void*);	// stack push/popds in units of one machine word
+		RspBias += !bIsChained;		// and one extra push for the ret address
+		RspBias |= FpInfo;			// pack in details about possible frame pointer
+
+		if (bIsChained)
+		{
+			OutTableCursor[-1].RspBias += RspBias;
+#if BACKTRACE_DBGLVL >= 2
+			OutTableCursor[-1].Size += (FunctionTable->EndAddress - FunctionTable->BeginAddress);
+#endif
+		}
+		else
+		{
+			*OutTableCursor = {
+				FunctionTable->BeginAddress,
+				RspBias,
+#if BACKTRACE_DBGLVL >= 2
+				FunctionTable->EndAddress - FunctionTable->BeginAddress,
+				UnwindInfo,
+#endif
+			};
+
+			++OutTableCursor;
+		}
+	}
+
+	UPTRINT ModuleSize = NtHeader->OptionalHeader.SizeOfImage;
+	ModuleSize += 0xffff; // to align up to next 64K page. it'll get shifted by AddressToId()
+
+	FModule Module = {
+		AddressToId(ModuleBase),
+		AddressToId(ModuleSize),
+		uint32(UPTRINT(OutTableCursor - OutTable)),
+#if BACKTRACE_DBGLVL >= 1
+		uint16(NumFpFuncs),
+#endif
+		OutTable,
+	};
+
+	{
+		FScopeLock _(&Lock);
+
+		TArrayView<FModule> ModulesView(Modules, ModulesNum);
+		int32 Index = Algo::UpperBound(ModulesView, Module.Id, FIdPredicate());
+
+		if (ModulesNum + 1 > ModulesCapacity)
+		{
+			ModulesCapacity += 8;
+			Modules = (FModule*) Malloc->Realloc(Modules, sizeof(FModule)* ModulesCapacity);
+		}
+		Modules[ModulesNum++] = Module;
+		Algo::Sort(TArrayView<FModule>(Modules, ModulesNum), [](const FModule& A, const FModule& B) { return A.Id < B.Id; });
+	}
+
+#if BACKTRACE_DBGLVL >= 1
+	NumFpTruncations += NumFpFuncs;
+	TotalFunctions += NumFunctions;
+#endif
+}
+
+////////////////////////////////////////////////////////////////////////////////
+void FBacktracer::RemoveModule(UPTRINT ModuleBase)
+{
+	// When Windows' RequestExit() is called it hard-terminates all threads except
+	// the main thread and then proceeds to unload the process' DLLs. This hard 
+	// thread termination can result is dangling locked locks. Not an issue as
+	// the rule is "do not do anything multithreaded in DLL load/unload". And here
+	// we are, taking write locks during DLL unload which is, quite unsurprisingly,
+	// deadlocking. In reality tracking Windows' DLL unloads doesn't tell us
+	// anything due to how DLLs and processes' address spaces work. So we will...
+#if defined PLATFORM_WINDOWS
+	return;
+#endif
+
+	FScopeLock _(&Lock);
+
+	uint32 ModuleId = AddressToId(ModuleBase);
+	TArrayView<FModule> ModulesView(Modules, ModulesNum);
+	int32 Index = Algo::LowerBound(ModulesView, ModuleId, FIdPredicate());
+	if (Index >= ModulesNum)
+	{
+		return;
+	}
+
+	const FModule& Module = Modules[Index];
+	if (Module.Id != ModuleId)
+	{
+		return;
+	}
+
+#if BACKTRACE_DBGLVL >= 1
+	NumFpTruncations -= Module.NumFpTypes;
+	TotalFunctions -= Module.NumFunctions;
+#endif
+
+	// no code should be executing at this point so we can safely free the
+	// table knowing know one is looking at it.
+	Malloc->Free(Module.Functions);
+
+	for (SIZE_T i = Index; i < ModulesNum; i++)
+	{
+		Modules[i] = Modules[i + 1];
+	}
+
+	--ModulesNum;
+}
+
+////////////////////////////////////////////////////////////////////////////////
+const FBacktracer::FFunction* FBacktracer::LookupFunction(UPTRINT Address, FLookupState& State) const
+{
+	// This function caches the previous module look up. The theory here is that
+	// a series of return address in a backtrace often cluster around one module
+
+	FIdPredicate IdPredicate;
+
+	// Look up the module that Address belongs to.
+	uint32 AddressId = AddressToId(Address);
+	if ((AddressId - State.Module.Id) >= State.Module.IdSize)
+	{
+		TArrayView<FModule> ModulesView(Modules, ModulesNum);
+		uint32 Index = Algo::UpperBound(ModulesView, AddressId, IdPredicate);
+		if (Index == 0)
+		{
+			return nullptr;
+		}
+
+		State.Module = Modules[Index - 1];
+	}
+
+	// Check that the address is within the address space of the best-found module
+	const FModule* Module = &(State.Module);
+	if ((AddressId - Module->Id) >= Module->IdSize)
+	{
+		return nullptr;
+	}
+
+	// Now we've a module we have a table of functions and their stack sizes so
+	// we can get the frame size for Address
+	uint32 FuncId = uint32(Address - IdToAddress(Module->Id)); 
+	TArrayView<FFunction> FuncsView(Module->Functions, Module->NumFunctions);
+	uint32 Index = Algo::UpperBound(FuncsView, FuncId, IdPredicate);
+
+	const FFunction* Function = Module->Functions + (Index - 1);
+#if BACKTRACE_DBGLVL >= 2
+	if ((FuncId - Function->Id) >= Function->Size)
+	{
+		PLATFORM_BREAK();
+		return nullptr;
+	}
+#endif
+	return Function;
+}
+
+////////////////////////////////////////////////////////////////////////////////
+void* FBacktracer::GetBacktraceId(void* AddressOfReturnAddress) const
+{
+	FLookupState LookupState;
+	FCallstackTracer::FBacktraceEntry BacktraceEntry;
+
+	UPTRINT* StackPointer = (UPTRINT*)AddressOfReturnAddress;
+
+#if BACKTRACE_DBGLVL >= 3
+	UPTRINT TruthBacktrace[1024];
+	uint32 NumTruth = RtlCaptureStackBackTrace(0, 1024, (void**)TruthBacktrace, nullptr);
+	UPTRINT* TruthCursor = TruthBacktrace;
+	for (; *TruthCursor != *StackPointer; ++TruthCursor);
+#endif
+
+#if BACKTRACE_DBGLVL >= 2
+	struct { void* Sp; void* Ip; const FFunction* Function; } Backtrace[1024] = {};
+	uint32 NumBacktrace = 0;
+#endif
+
+	uint64 BacktraceId = 0;
+	uint32 FrameIdx = 0;
+
+	FScopeLock _(&Lock);
+	do
+	{
+		UPTRINT RetAddr = *StackPointer;
+		
+		BacktraceEntry.Frames[FrameIdx++] = RetAddr;
+		BacktraceEntry.FrameCount = FrameIdx;
+
+		// This is a simple order-dependent LCG. Should be sufficient enough
+		BacktraceId += RetAddr;
+		BacktraceId *= 0x30be8efa499c249dull;
+
+		const FFunction* Function = LookupFunction(RetAddr, LookupState);
+		if (Function == nullptr)
+		{
+			break;
+		}
+
+#if BACKTRACE_DBGLVL >= 2
+		if (NumBacktrace < 1024)
+		{
+			Backtrace[NumBacktrace++] = {
+				StackPointer,
+				(void*)RetAddr,
+				Function,
+			};
+		}
+#endif
+
+		if (Function->RspBias < 0)
+		{
+			// This is a frame with a variable-sized stack pointer. We don't
+			// track enough information to proceed.
+#if BACKTRACE_DBGLVL >= 1
+			NumFpTruncations++;
+#endif
+			break;
+		}
+
+		StackPointer += Function->RspBias;
+	}
+	// Trunkate callstacks longer than MaxStackDepth
+	while (*StackPointer && FrameIdx < FCallstackTracer::FBacktraceEntry::MaxStackDepth);
+
+	// Save the collected id
+	BacktraceEntry.Id = BacktraceId;
+
+#if BACKTRACE_DBGLVL >= 3
+	for (uint32 i = 0; i < NumBacktrace; ++i)
+	{
+		if ((void*)TruthCursor[i] != Backtrace[i].Ip)
+		{
+			PLATFORM_BREAK();
+			break;
+		}
+	}
+#endif
+
+	// Add to queue to be processed. This might block until there is room in the
+	// queue (i.e. the processing thread has caught up processing).
+	CallstackTracer->AddCallstack(BacktraceEntry);
+	return (void*)BacktraceId;
+}
+
+////////////////////////////////////////////////////////////////////////////////
+int32 FBacktracer::GetFrameSize(void* FunctionAddress) const
+{
+	FScopeLock _(&Lock);
+
+	FLookupState LookupState;
+	const FFunction* Function = LookupFunction(UPTRINT(FunctionAddress), LookupState);
+	if (FunctionAddress == nullptr)
+	{
+		return -1;
+	}
+
+	if (Function->RspBias < 0)
+	{
+		return -1;
+	}
+
+	return Function->RspBias;
+}
+
+////////////////////////////////////////////////////////////////////////////////
+void Modules_Create(FMalloc*);
+void Modules_Subscribe(void (*)(bool, void*, const TCHAR*));
+void Modules_Initialize();
+
+////////////////////////////////////////////////////////////////////////////////
+void Backtracer_Create(FMalloc* Malloc)
+{
+	if (FBacktracer::Get() != nullptr)
+	{
+		return;
+	}
+
+	// Allocate, construct and intentionally leak backtracer
+	void* Alloc = Malloc->Malloc(sizeof(FBacktracer), alignof(FBacktracer));
+	new (Alloc) FBacktracer(Malloc);
+
+	Modules_Create(Malloc);
+	Modules_Subscribe(
+		[] (bool bLoad, void* Module, const TCHAR* Name)
+		{
+			bLoad
+				? FBacktracer::Get()->AddModule(UPTRINT(Module), Name)
+				: FBacktracer::Get()->RemoveModule(UPTRINT(Module));
+		}
+	);
+}
+
+////////////////////////////////////////////////////////////////////////////////
+void Backtracer_Initialize()
+{
+	Modules_Initialize();
+}
+
+////////////////////////////////////////////////////////////////////////////////
+int32 Backtracer_GetFrameSize(void* FunctionAddress)
+{
+	if (FBacktracer* Instance = FBacktracer::Get())
+	{
+		return Instance->GetFrameSize(FunctionAddress);
+	}
+
+	return -1;
+}
+
+////////////////////////////////////////////////////////////////////////////////
+void* Backtracer_GetBacktraceId(void* AddressOfReturnAddress)
+{
+	if (FBacktracer* Instance = FBacktracer::Get())
+	{
+		return Instance->GetBacktraceId(AddressOfReturnAddress);
+	}
+
+	return nullptr;
+}
+
+#endif // defined(PLATFORM_SUPPORTS_TRACE_WIN32_CALLSTACK)
diff --git a/Engine/Source/Runtime/Core/Private/ProfilingDebugging/Microsoft/WindowsModuleDiagnostics.cpp b/Engine/Source/Runtime/Core/Private/ProfilingDebugging/Microsoft/WindowsModuleDiagnostics.cpp
new file mode 100644
--- /dev/null
+++ b/Engine/Source/Runtime/Core/Private/ProfilingDebugging/Microsoft/WindowsModuleDiagnostics.cpp
@@ -0,0 +1,244 @@
+// Copyright Epic Games, Inc. All Rights Reserved.
+
+#if defined(PLATFORM_SUPPORTS_TRACE_WIN32_MODULE_DIAGNOSTICS)
+
+#include "CoreMinimal.h"
+#include "Trace/Trace.h"
+#include "Trace/Trace.inl"
+#include "ProfilingDebugging/ModuleDiagnostics.h"
+
+#include COMPILED_PLATFORM_HEADER(PlatformModuleDiagnostics.h)
+
+////////////////////////////////////////////////////////////////////////////////
+struct FNtDllFunction
+{
+	FARPROC Addr;
+
+	FNtDllFunction(const char* Name)
+	{
+		HMODULE NtDll = LoadLibraryW(L"ntdll.dll");
+		check(NtDll);
+		Addr = GetProcAddress(NtDll, Name);
+	}
+
+	template <typename... ArgTypes>
+	unsigned int operator () (ArgTypes... Args)
+	{
+		typedef unsigned int (NTAPI *Prototype)(ArgTypes...);
+		return (Prototype((void*)Addr))(Args...);
+	}
+};
+
+////////////////////////////////////////////////////////////////////////////////
+class FModuleTrace
+{
+public:
+	typedef void				(*SubscribeFunc)(bool, void*, const TCHAR*);
+	typedef TArray<SubscribeFunc, TFixedAllocator<64>> SubscriberArray;
+
+								FModuleTrace(FMalloc* InMalloc);
+								~FModuleTrace();
+	static FModuleTrace*		Get();
+	void						Initialize();
+	void						Subscribe(SubscribeFunc Function);
+
+private:
+	void						OnDllLoaded(const UNICODE_STRING& Name, UPTRINT Base);
+	void						OnDllUnloaded(UPTRINT Base);
+	void						OnDllNotification(unsigned int Reason, const void* DataPtr);
+	static FModuleTrace*		Instance;
+	SubscriberArray				Subscribers;
+	void*						CallbackCookie;
+};
+
+////////////////////////////////////////////////////////////////////////////////
+FModuleTrace* FModuleTrace::Instance = nullptr;
+
+////////////////////////////////////////////////////////////////////////////////
+FModuleTrace::FModuleTrace(FMalloc* InMalloc)
+{
+	Instance = this;
+}
+
+////////////////////////////////////////////////////////////////////////////////
+FModuleTrace::~FModuleTrace()
+{
+	FNtDllFunction UnregisterFunc("LdrUnregisterDllNotification");
+	UnregisterFunc(CallbackCookie);
+}
+
+////////////////////////////////////////////////////////////////////////////////
+FModuleTrace* FModuleTrace::Get()
+{
+	return Instance;
+}
+
+////////////////////////////////////////////////////////////////////////////////
+void FModuleTrace::Initialize()
+{
+	using namespace UE::Trace;
+
+	UE_TRACE_LOG(Diagnostics, ModuleInit, ModuleChannel, sizeof(char) * 3)
+		<< ModuleInit.SymbolFormat("pdb", 3)
+		<< ModuleInit.ModuleBaseShift(uint8(0));
+
+	// Register for DLL load/unload notifications.
+	auto Thunk = [] (ULONG Reason, const void* Data, void* Context)
+	{
+		auto* Self = (FModuleTrace*)Context;
+		Self->OnDllNotification(Reason, Data);
+	};
+
+	typedef void (CALLBACK *ThunkType)(ULONG, const void*, void*);
+	auto ThunkImpl = ThunkType(Thunk);
+
+	FNtDllFunction RegisterFunc("LdrRegisterDllNotification");
+	RegisterFunc(0, ThunkImpl, this, &CallbackCookie);
+
+	// Enumerate already loaded modules.
+	const TEB* ThreadEnvBlock = NtCurrentTeb();
+	const PEB* ProcessEnvBlock = ThreadEnvBlock->ProcessEnvironmentBlock;
+	const LIST_ENTRY* ModuleIter = ProcessEnvBlock->Ldr->InMemoryOrderModuleList.Flink;
+	const LIST_ENTRY* ModuleIterEnd = ModuleIter->Blink;
+	do
+	{
+		const auto& ModuleData = *(LDR_DATA_TABLE_ENTRY*)(ModuleIter - 1);
+		if (ModuleData.DllBase == 0)
+		{
+			break;
+		}
+
+		OnDllLoaded(ModuleData.FullDllName, UPTRINT(ModuleData.DllBase));
+		ModuleIter = ModuleIter->Flink;
+	}
+	while (ModuleIter != ModuleIterEnd);
+}
+
+////////////////////////////////////////////////////////////////////////////////
+void FModuleTrace::Subscribe(SubscribeFunc Function)
+{
+	Subscribers.Add(Function);
+}
+
+////////////////////////////////////////////////////////////////////////////////
+void FModuleTrace::OnDllNotification(unsigned int Reason, const void* DataPtr)
+{
+	enum
+	{
+		LDR_DLL_NOTIFICATION_REASON_LOADED = 1,
+		LDR_DLL_NOTIFICATION_REASON_UNLOADED = 2,
+	};
+
+	struct FNotificationData
+	{
+		uint32					Flags;
+		const UNICODE_STRING&	FullPath;
+		const UNICODE_STRING&	BaseName;
+		UPTRINT					Base;
+	};
+	const auto& Data = *(FNotificationData*)DataPtr;
+
+	switch (Reason)
+	{
+	case LDR_DLL_NOTIFICATION_REASON_LOADED:	OnDllLoaded(Data.FullPath, Data.Base);	break;
+	case LDR_DLL_NOTIFICATION_REASON_UNLOADED:	OnDllUnloaded(Data.Base);				break;
+	}
+}
+
+////////////////////////////////////////////////////////////////////////////////
+void FModuleTrace::OnDllLoaded(const UNICODE_STRING& Name, UPTRINT Base)
+{
+	const auto* DosHeader = (IMAGE_DOS_HEADER*)Base;
+	const auto* NtHeaders = (IMAGE_NT_HEADERS*)(Base + DosHeader->e_lfanew);
+	const IMAGE_OPTIONAL_HEADER& OptionalHeader = NtHeaders->OptionalHeader;
+	TArray<uint8, TFixedAllocator<20>> ImageId;
+
+	// Find the guid and age of the binary, used to match debug files
+	const IMAGE_DATA_DIRECTORY& DebugInfoEntry = OptionalHeader.DataDirectory[IMAGE_DIRECTORY_ENTRY_DEBUG];
+	const auto* DebugEntries = (IMAGE_DEBUG_DIRECTORY*)(Base + DebugInfoEntry.VirtualAddress);
+	for (uint32 i = 0, n = DebugInfoEntry.Size / sizeof(DebugEntries[0]); i < n; ++i)
+	{
+		const IMAGE_DEBUG_DIRECTORY& Entry = DebugEntries[i];
+		if (Entry.Type == IMAGE_DEBUG_TYPE_CODEVIEW)
+		{
+			struct FCodeView7
+			{
+				uint32  Signature;
+				uint32  Guid[4];
+				uint32  Age;
+			};
+
+			if (Entry.SizeOfData < sizeof(FCodeView7))
+			{
+				continue;
+			}
+
+			const auto* CodeView7 = (FCodeView7*)(Base + Entry.AddressOfRawData);
+			if (CodeView7->Signature != 'SDSR')
+			{
+				continue;
+			}
+
+			ImageId.Append((uint8*)&CodeView7->Guid, sizeof(uint32) * 4);
+			ImageId.Append((uint8*)&CodeView7->Age, sizeof(uint32));
+			break;
+		}
+	}
+
+	// Note: UNICODE_STRING.Length is the size in bytes of the string buffer.
+	UE_TRACE_LOG(Diagnostics, ModuleLoad, ModuleChannel, Name.Length + ImageId.Num())
+		<< ModuleLoad.Name(Name.Buffer, Name.Length / 2)
+		<< ModuleLoad.Base(uint64(Base))
+		<< ModuleLoad.Size(OptionalHeader.SizeOfImage)
+		<< ModuleLoad.ImageId(ImageId.GetData(), ImageId.Num());
+
+	for (SubscribeFunc Subscriber : Subscribers)
+	{
+		Subscriber(true, (void*)Base, Name.Buffer);
+	}
+}
+
+////////////////////////////////////////////////////////////////////////////////
+void FModuleTrace::OnDllUnloaded(UPTRINT Base)
+{
+	UE_TRACE_LOG(Diagnostics, ModuleUnload, ModuleChannel)
+		<< ModuleUnload.Base(uint64(Base));
+
+	for (SubscribeFunc Subscriber : Subscribers)
+	{
+		Subscriber(false, (void*)Base, nullptr);
+	}
+}
+
+
+
+////////////////////////////////////////////////////////////////////////////////
+void Modules_Create(FMalloc* Malloc)
+{
+	if (FModuleTrace::Get() != nullptr)
+	{
+		return;
+	}
+
+	static FModuleTrace Instance(Malloc);
+}
+
+////////////////////////////////////////////////////////////////////////////////
+void Modules_Initialize()
+{
+	if (FModuleTrace* Instance = FModuleTrace::Get())
+	{
+		Instance->Initialize();
+	}
+}
+
+////////////////////////////////////////////////////////////////////////////////
+void Modules_Subscribe(void (*Function)(bool, void*, const TCHAR*))
+{
+	if (FModuleTrace* Instance = FModuleTrace::Get())
+	{
+		Instance->Subscribe(Function);
+	}
+}
+
+#endif // PLATFORM_SUPPORTS_WIN32_MEMORY_TRACE
diff --git a/Engine/Source/Runtime/Core/Private/ProfilingDebugging/PlatformEvents.cpp b/Engine/Source/Runtime/Core/Private/ProfilingDebugging/PlatformEvents.cpp
--- a/Engine/Source/Runtime/Core/Private/ProfilingDebugging/PlatformEvents.cpp
+++ b/Engine/Source/Runtime/Core/Private/ProfilingDebugging/PlatformEvents.cpp
@@ -28,25 +28,3 @@ EPlatformEvent PlatformEvents_GetEvent(const FString& Name)
 		return EPlatformEvent::None;
 	}
 }
-
-/////////////////////////////////////////////////////////////////////
-
-#if !UE_PLATFORM_EVENTS_AVAILABLE
-
-void PlatformEvents_Init(uint32 SamplingIntervalUsec)
-{
-}
-
-void PlatformEvents_Enable(EPlatformEvent Event)
-{
-}
-
-void PlatformEvents_Disable(EPlatformEvent Event)
-{
-}
-
-void PlatformEvents_Stop()
-{
-}
-
-#endif
diff --git a/Engine/Source/Runtime/Core/Private/ProfilingDebugging/PlatformEvents.h b/Engine/Source/Runtime/Core/Private/ProfilingDebugging/PlatformEvents.h
--- a/Engine/Source/Runtime/Core/Private/ProfilingDebugging/PlatformEvents.h
+++ b/Engine/Source/Runtime/Core/Private/ProfilingDebugging/PlatformEvents.h
@@ -8,10 +8,8 @@
 
 class FString;
 
-#if PLATFORM_WINDOWS
-#	define UE_PLATFORM_EVENTS_AVAILABLE 1
-#else
-#	define UE_PLATFORM_EVENTS_AVAILABLE 0
+#if !defined(PLATFORM_SUPPORTS_PLATFORM_EVENTS)
+#	define PLATFORM_SUPPORTS_PLATFORM_EVENTS 0
 #endif
 
 /////////////////////////////////////////////////////////////////////
@@ -51,9 +49,20 @@ ENUM_CLASS_FLAGS(EPlatformEvent)
 
 EPlatformEvent PlatformEvents_GetEvent(const FString& Name);
 
+#if PLATFORM_SUPPORTS_PLATFORM_EVENTS
+
 void PlatformEvents_Init(uint32 SamplingIntervalUsec);
 void PlatformEvents_Enable(EPlatformEvent Event);
 void PlatformEvents_Disable(EPlatformEvent Event);
 void PlatformEvents_Stop();
 
+#else // PLATFORM_SUPPORTS_PLATFORM_EVENTS
+
+static void PlatformEvents_Init(uint32 SamplingIntervalUsec) {}
+static void PlatformEvents_Enable(EPlatformEvent Event) {}
+static void PlatformEvents_Disable(EPlatformEvent Event) {}
+static void PlatformEvents_Stop() {}
+
+#endif // PLATFORM_SUPPORTS_PLATFORM_EVENTS
+
 /////////////////////////////////////////////////////////////////////
diff --git a/Engine/Source/Runtime/Core/Private/ProfilingDebugging/Windows/EventTracingForWindows.cpp b/Engine/Source/Runtime/Core/Private/ProfilingDebugging/Windows/EventTracingForWindows.cpp
deleted file mode 100644
diff --git a/Engine/Source/Runtime/Core/Private/ProfilingDebugging/Windows/MemoryTrace.cpp b/Engine/Source/Runtime/Core/Private/ProfilingDebugging/Windows/MemoryTrace.cpp
deleted file mode 100644
diff --git a/Engine/Source/Runtime/Core/Private/ProfilingDebugging/Windows/WindowsCallstackTrace.cpp b/Engine/Source/Runtime/Core/Private/ProfilingDebugging/Windows/WindowsCallstackTrace.cpp
deleted file mode 100644
diff --git a/Engine/Source/Runtime/Core/Private/ProfilingDebugging/Windows/WindowsModuleDiagnostics.cpp b/Engine/Source/Runtime/Core/Private/ProfilingDebugging/Windows/WindowsModuleDiagnostics.cpp
deleted file mode 100644
diff --git a/Engine/Source/Runtime/Core/Private/Windows/WindowsPlatformEventTracingForWindows.h b/Engine/Source/Runtime/Core/Private/Windows/WindowsPlatformEventTracingForWindows.h
new file mode 100644
--- /dev/null
+++ b/Engine/Source/Runtime/Core/Private/Windows/WindowsPlatformEventTracingForWindows.h
@@ -0,0 +1,8 @@
+// Copyright Epic Games, Inc. All Rights Reserved.
+
+#pragma once
+
+#include "Windows/AllowWindowsPlatformTypes.h"
+#	include <evntrace.h>
+#	include <evntcons.h>
+#include "Windows/HideWindowsPlatformTypes.h"
diff --git a/Engine/Source/Runtime/Core/Private/Windows/WindowsPlatformModuleDiagnostics.h b/Engine/Source/Runtime/Core/Private/Windows/WindowsPlatformModuleDiagnostics.h
new file mode 100644
--- /dev/null
+++ b/Engine/Source/Runtime/Core/Private/Windows/WindowsPlatformModuleDiagnostics.h
@@ -0,0 +1,7 @@
+// Copyright Epic Games, Inc. All Rights Reserved.
+
+#pragma once
+
+#include "Windows/AllowWindowsPlatformTypes.h"
+#	include <winternl.h>
+#include "Windows/HideWindowsPlatformTypes.h"
diff --git a/Engine/Source/Runtime/TraceLog/TraceLog.Build.cs b/Engine/Source/Runtime/TraceLog/TraceLog.Build.cs
--- a/Engine/Source/Runtime/TraceLog/TraceLog.Build.cs
+++ b/Engine/Source/Runtime/TraceLog/TraceLog.Build.cs
@@ -10,6 +10,11 @@ public class TraceLog : ModuleRules
 	{
 		bRequiresImplementModule = false;
 		PublicIncludePathModuleNames.Add("Core");
+	
+		if (Target.Platform == UnrealTargetPlatform.Win64)
+		{
+			PublicDefinitions.Add("PLATFORM_SUPPORTS_PLATFORM_EVENTS=1");
+		}
     }
 
 	// used by platform extension derived classes. probably should become a project setting at some point!
