#UE Add ALLOW_NON_INTERNAL_UNATTENDED_CRASH_REPORTS to allow targets to allow unattended crash reports even if they are not internal builds

diff --git a/Engine/Source/Runtime/Core/Private/Windows/WindowsPlatformCrashContext.cpp b/Engine/Source/Runtime/Core/Private/Windows/WindowsPlatformCrashContext.cpp
--- a/Engine/Source/Runtime/Core/Private/Windows/WindowsPlatformCrashContext.cpp
+++ b/Engine/Source/Runtime/Core/Private/Windows/WindowsPlatformCrashContext.cpp
@@ -47,6 +47,10 @@
 #define USE_CRASH_REPORTER_MONITOR WITH_EDITOR
 #endif
 
+#ifndef ALLOW_NON_INTERNAL_UNATTENDED_CRASH_REPORTS
+#define ALLOW_NON_INTERNAL_UNATTENDED_CRASH_REPORTS 0
+#endif
+
 #ifndef NOINITCRASHREPORTER
 #define NOINITCRASHREPORTER 0
 #endif
@@ -708,7 +712,7 @@ int32 ReportCrashForMonitor(
 		}
 	}
 
-#if !UE_EDITOR
+#if !UE_EDITOR && !ALLOW_NON_INTERNAL_UNATTENDED_CRASH_REPORTS
 	if (!FEngineBuildSettings::IsInternalBuild())
 	{
 		// do not send unattended reports in licensees' builds except for the editor, where it is governed by the above setting
@@ -891,7 +895,7 @@ int32 ReportCrashUsingCrashReportClient(FWindowsPlatformCrashContext& InContext,
 		GConfig->GetBool(TEXT("/Script/UnrealEd.AnalyticsPrivacySettings"), TEXT("bSendUsageData"), bSendUsageData, GEditorSettingsIni);
 	}
 
-#if !UE_EDITOR
+#if !UE_EDITOR && !ALLOW_NON_INTERNAL_UNATTENDED_CRASH_REPORTS
 	if (!FEngineBuildSettings::IsInternalBuild())
 	{
 		// do not send unattended reports in licensees' builds except for the editor, where it is governed by the above setting
