[LLM] - misc scopes to improve LLM tracking

diff --git a/Engine/Plugins/Media/ElectraPlayer/Source/ElectraPlayerRuntime/Private/ElectraPlayerRuntimeModule.cpp b/Engine/Plugins/Media/ElectraPlayer/Source/ElectraPlayerRuntime/Private/ElectraPlayerRuntimeModule.cpp
--- a/Engine/Plugins/Media/ElectraPlayer/Source/ElectraPlayerRuntime/Private/ElectraPlayerRuntimeModule.cpp
+++ b/Engine/Plugins/Media/ElectraPlayer/Source/ElectraPlayerRuntime/Private/ElectraPlayerRuntimeModule.cpp
@@ -42,6 +42,7 @@ public:
 
 	void StartupModule() override
 	{
+		LLM_SCOPE_BYNAME(TEXT("StartupModule"));
 		check(!bInitialized);
 
 		if (!Electra::PlatformEarlyStartup())
diff --git a/Engine/Plugins/Runtime/Metasound/Source/MetasoundFrontend/Public/MetasoundArrayNodesRegistration.h b/Engine/Plugins/Runtime/Metasound/Source/MetasoundFrontend/Public/MetasoundArrayNodesRegistration.h
--- a/Engine/Plugins/Runtime/Metasound/Source/MetasoundFrontend/Public/MetasoundArrayNodesRegistration.h
+++ b/Engine/Plugins/Runtime/Metasound/Source/MetasoundFrontend/Public/MetasoundArrayNodesRegistration.h
@@ -174,6 +174,7 @@ namespace Metasound
 	template<typename ArrayType>
 	bool RegisterArrayNodes()
 	{
+		LLM_SCOPE_BYNAME(TEXT("Metasound"));
 		using namespace MetasoundArrayNodesPrivate;
 
 		bool bSuccess = RegisterArrayNumNode<ArrayType>();
diff --git a/Engine/Source/Editor/EditorStyle/Private/StarshipStyle.cpp b/Engine/Source/Editor/EditorStyle/Private/StarshipStyle.cpp
--- a/Engine/Source/Editor/EditorStyle/Private/StarshipStyle.cpp
+++ b/Engine/Source/Editor/EditorStyle/Private/StarshipStyle.cpp
@@ -24,6 +24,7 @@
 
 void FStarshipEditorStyle::Initialize()
 {
+	LLM_SCOPE_BYNAME(TEXT("FStarshipEditorStyle"));
 	Settings = NULL;
 
 	// The core style must be initialized before the editor style
diff --git a/Engine/Source/Runtime/Core/Private/Internationalization/TextCache.cpp b/Engine/Source/Runtime/Core/Private/Internationalization/TextCache.cpp
--- a/Engine/Source/Runtime/Core/Private/Internationalization/TextCache.cpp
+++ b/Engine/Source/Runtime/Core/Private/Internationalization/TextCache.cpp
@@ -3,6 +3,7 @@
 #include "Internationalization/TextCache.h"
 #include "Misc/LazySingleton.h"
 #include "Misc/ScopeLock.h"
+#include "HAL/LowLevelMemTracker.h"
 
 FTextCache& FTextCache::Get()
 {
@@ -16,6 +17,7 @@ void FTextCache::TearDown()
 
 FText FTextCache::FindOrCache(const TCHAR* InTextLiteral, const TCHAR* InNamespace, const TCHAR* InKey)
 {
+	LLM_SCOPE(ELLMTag::Localization);
 	const FTextId TextId(InNamespace, InKey);
 
 	// First try and find a cached instance
