Implement OutStdErr in FUnixPlatformProcess::ExecProcess

Add stderr parameter to CreateProc
Add exec-process test to TestPAL

PR #6919: Add a new FPlatformProcess::CreateProc() overload to allow separating stderr and stdout (Contributed by [REDACTED_PERSON])

diff --git a/Engine/Source/Programs/TestPAL/Private/Main.cpp b/Engine/Source/Programs/TestPAL/Private/Main.cpp
--- a/Engine/Source/Programs/TestPAL/Private/Main.cpp
+++ b/Engine/Source/Programs/TestPAL/Private/Main.cpp
@@ -39,6 +39,7 @@ IMPLEMENT_APPLICATION(TestPAL, "TestPAL");
 #define ARG_STRINGS_ALLOCATION_TEST			"stringsallocation"
 #define ARG_CREATEGUID_TEST					"createguid"
 #define ARG_THREADSTACK_TEST				"threadstack"
+#define ARG_EXEC_PROCESS_TEST				"exec-process"
 #define ARG_FORK_TEST						"fork"
 
 namespace TestPAL
@@ -1615,6 +1616,51 @@ int32 ForkTest(const TCHAR* CommandLine)
 	return 0;
 }
 
+/** 
+ * Launch the external tool as a process and wait for it to complete
+ *
+ * Set executable and parameters via something like:
+ *
+ *  ./Engine/Binaries/Linux/TestPAL exec-process exe=/bin/ls params=/tmp
+ *
+ */
+int32 ExecProcessTest(const TCHAR *CommandLine)
+{
+	FString Exe;
+	FString Params;
+
+	FPlatformMisc::SetCrashHandler(NULL);
+	FPlatformMisc::SetGracefulTerminationHandler();
+
+	GEngineLoop.PreInit(CommandLine);
+
+	if (!FParse::Value(CommandLine, TEXT("exe="), Exe))
+	{
+		Exe = TEXT("nonexistent-exe");
+	}
+	if (!FParse::Value(CommandLine, TEXT("params="), Params))
+	{
+		Params = TEXT("");
+	}
+
+	UE_LOG(LogTestPAL, Display, TEXT("  Exe:'%s' Parameters:'%s'"), *Exe, *Params);
+
+	int32 ReturnCode;
+	FString OutStdOut;
+	FString OutStdErr;
+	bool Ret = FPlatformProcess::ExecProcess(*Exe, *Params, &ReturnCode, &OutStdOut, &OutStdErr);
+
+	UE_LOG(LogTestPAL, Display, TEXT("ExecProcess returns %d"), Ret);
+	UE_LOG(LogTestPAL, Display, TEXT("  ReturnCode:%d"), ReturnCode);
+	UE_LOG(LogTestPAL, Display, TEXT("  StdOut:%s"), *OutStdOut);
+	UE_LOG(LogTestPAL, Display, TEXT("  StdErr:%s"), *OutStdErr);
+
+	FEngineLoop::AppPreExit();
+	FEngineLoop::AppExit();
+
+	return 0;
+}
+
 /**
  * Selects and runs one of test cases.
  *
@@ -1704,6 +1750,10 @@ int32 MultiplexedMain(int32 ArgC, char* ArgV[])
 		{
 			return ThreadTraceTest(*TestPAL::CommandLine);
 		}
+		else if (!FCStringAnsi::Strcmp(ArgV[IdxArg], ARG_EXEC_PROCESS_TEST))
+		{
+			return ExecProcessTest(*TestPAL::CommandLine);
+		}
 		else if (!FCStringAnsi::Strcmp(ArgV[IdxArg], ARG_FORK_TEST))
 		{
 			return ForkTest(*TestPAL::CommandLine);
@@ -1736,6 +1786,7 @@ int32 MultiplexedMain(int32 ArgC, char* ArgV[])
 	UE_LOG(LogTestPAL, Warning, TEXT("  %s: test string allocations."), UTF8_TO_TCHAR(ARG_STRINGS_ALLOCATION_TEST));
 	UE_LOG(LogTestPAL, Warning, TEXT("  %s: test CreateGuid."), UTF8_TO_TCHAR(ARG_CREATEGUID_TEST));
 	UE_LOG(LogTestPAL, Warning, TEXT("  %s: test ThreadWalkStackAndDump and CaptureThreadBackTrace."), UTF8_TO_TCHAR(ARG_THREADSTACK_TEST));
+	UE_LOG(LogTestPAL, Warning, TEXT("  %s: test ExecProcess. Possible options: [-exe=executable] [-params=parameters]"), UTF8_TO_TCHAR(ARG_EXEC_PROCESS_TEST));
 
 	if (PLATFORM_LINUX)
 	{
diff --git a/Engine/Source/Runtime/Core/Private/GenericPlatform/GenericPlatformProcess.cpp b/Engine/Source/Runtime/Core/Private/GenericPlatform/GenericPlatformProcess.cpp
--- a/Engine/Source/Runtime/Core/Private/GenericPlatform/GenericPlatformProcess.cpp
+++ b/Engine/Source/Runtime/Core/Private/GenericPlatform/GenericPlatformProcess.cpp
@@ -239,12 +239,18 @@ FString FGenericPlatformProcess::GetGameBundleId()
 	return TEXT("");
 }
 
-FProcHandle FGenericPlatformProcess::CreateProc( const TCHAR* URL, const TCHAR* Parms, bool bLaunchDetached, bool bLaunchHidden, bool bLaunchReallyHidden, uint32* OutProcessID, int32 PriorityModifier, const TCHAR* OptionalWorkingDirectory, void* PipeWriteChild, void * PipeReadChild)
+FProcHandle FGenericPlatformProcess::CreateProc( const TCHAR* URL, const TCHAR* Parms, bool bLaunchDetached, bool bLaunchHidden, bool bLaunchReallyHidden, uint32* OutProcessID, int32 PriorityModifier, const TCHAR* OptionalWorkingDirectory, void* PipeWriteChild, void* PipeReadChild )
 {
 	UE_LOG(LogHAL, Fatal, TEXT("FGenericPlatformProcess::CreateProc not implemented on this platform"));
 	return FProcHandle();
 }
 
+FProcHandle FGenericPlatformProcess::CreateProc( const TCHAR* URL, const TCHAR* Parms, bool bLaunchDetached, bool bLaunchHidden, bool bLaunchReallyHidden, uint32* OutProcessID, int32 PriorityModifier, const TCHAR* OptionalWorkingDirectory, void* PipeWriteChild, void* PipeReadChild, void* PipeStdErrChild )
+{
+	UE_LOG(LogHAL, Fatal, TEXT("FGenericPlatformProcess::CreateProc with std out/in/err not implemented on this platform"));
+	return FProcHandle();
+}
+
 FProcHandle FGenericPlatformProcess::OpenProcess(uint32 ProcessID)
 {
 	UE_LOG(LogHAL, Fatal, TEXT("FGenericPlatformProcess::OpenProcess not implemented on this platform"));
diff --git a/Engine/Source/Runtime/Core/Private/Mac/MacPlatformProcess.cpp b/Engine/Source/Runtime/Core/Private/Mac/MacPlatformProcess.cpp
--- a/Engine/Source/Runtime/Core/Private/Mac/MacPlatformProcess.cpp
+++ b/Engine/Source/Runtime/Core/Private/Mac/MacPlatformProcess.cpp
@@ -309,6 +309,11 @@ FProcHandle FMacPlatformProcess::CreateProc( const TCHAR* URL, const TCHAR* Parm
 	return CreateProcInternal(URL, Parms, bLaunchDetached, bLaunchHidden, bLaunchReallyHidden, OutProcessID, PriorityModifier, OptionalWorkingDirectory, PipeWriteChild, PipeWriteChild, PipeReadChild);
 }
 
+FProcHandle FMacPlatformProcess::CreateProc( const TCHAR* URL, const TCHAR* Parms, bool bLaunchDetached, bool bLaunchHidden, bool bLaunchReallyHidden, uint32* OutProcessID, int32 PriorityModifier, const TCHAR* OptionalWorkingDirectory, void* PipeWriteChild, void* PipeReadChild, void* PipeStdErrChild)
+{
+	return CreateProcInternal(URL, Parms, bLaunchDetached, bLaunchHidden, bLaunchReallyHidden, OutProcessID, PriorityModifier, OptionalWorkingDirectory, PipeWriteChild, PipeStdErrChild, PipeReadChild);
+}
+
 FProcHandle FMacPlatformProcess::CreateProcInternal(const TCHAR* URL, const TCHAR* Parms, bool bLaunchDetached, bool bLaunchHidden, bool bLaunchReallyHidden, uint32* OutProcessID, int32 PriorityModifier, const TCHAR* OptionalWorkingDirectory, void* PipeStdOutChild, void* PipeStdErrChild, void *PipeStdInChild)
 {
 	SCOPED_AUTORELEASE_POOL;
diff --git a/Engine/Source/Runtime/Core/Private/Unix/UnixPlatformProcess.cpp b/Engine/Source/Runtime/Core/Private/Unix/UnixPlatformProcess.cpp
--- a/Engine/Source/Runtime/Core/Private/Unix/UnixPlatformProcess.cpp
+++ b/Engine/Source/Runtime/Core/Private/Unix/UnixPlatformProcess.cpp
@@ -756,7 +756,14 @@ namespace UnixPlatformProcess
 	}
 }
 
-FProcHandle FUnixPlatformProcess::CreateProc(const TCHAR* URL, const TCHAR* Parms, bool bLaunchDetached, bool bLaunchHidden, bool bLaunchReallyHidden, uint32* OutProcessID, int32 PriorityModifier, const TCHAR* OptionalWorkingDirectory, void* PipeWriteChild, void * PipeReadChild)
+FProcHandle FUnixPlatformProcess::CreateProc(const TCHAR* URL, const TCHAR* Parms, bool bLaunchDetached, bool bLaunchHidden, bool bLaunchReallyHidden, uint32* OutProcessID, int32 PriorityModifier, const TCHAR* OptionalWorkingDirectory, void* PipeWriteChild, void* PipeReadChild)
+{
+	// CreateProc used to only have a single "write" pipe argument, which Windows and Mac would pipe both stdout and stderr into.
+	// On Unix though, only stdout was piped to it, and stderr wasn't available at all, so we'll preserve that behaviour in this overload for compatibility with existing code
+	return CreateProc(URL, Parms, bLaunchDetached, bLaunchHidden, bLaunchReallyHidden, OutProcessID, PriorityModifier, OptionalWorkingDirectory, PipeWriteChild, PipeReadChild, PipeWriteChild);
+}
+
+FProcHandle FUnixPlatformProcess::CreateProc(const TCHAR* URL, const TCHAR* Parms, bool bLaunchDetached, bool bLaunchHidden, bool bLaunchReallyHidden, uint32* OutProcessID, int32 PriorityModifier, const TCHAR* OptionalWorkingDirectory, void* PipeWriteChild, void* PipeReadChild, void* PipeStdErrChild)
 {
 	// @TODO bLaunchHidden bLaunchReallyHidden are not handled
 	// We need an absolute path to executable
@@ -924,7 +931,7 @@ FProcHandle FUnixPlatformProcess::CreateProc(const TCHAR* URL, const TCHAR* Parm
 	SpawnFlags |= POSIX_SPAWN_SETPGROUP;
 
 	int PosixSpawnErrNo = -1;
-	if (PipeWriteChild || PipeReadChild)
+	if (PipeWriteChild || PipeReadChild || PipeStdErrChild)
 	{
 		posix_spawn_file_actions_t FileActions;
 		posix_spawn_file_actions_init(&FileActions);
@@ -941,6 +948,12 @@ FProcHandle FUnixPlatformProcess::CreateProc(const TCHAR* URL, const TCHAR* Parm
 			posix_spawn_file_actions_adddup2(&FileActions, PipeReadHandle->GetHandle(), STDIN_FILENO);
 		}
 
+		if (PipeStdErrChild)
+		{
+			const FPipeHandle* PipeStdErrorHandle = reinterpret_cast<const FPipeHandle*>(PipeStdErrChild);
+			posix_spawn_file_actions_adddup2(&FileActions, PipeStdErrorHandle->GetHandle(), STDERR_FILENO);
+		}
+
 		posix_spawnattr_setflags(&SpawnAttr, SpawnFlags);
 		PosixSpawnErrNo = posix_spawn(&ChildPid, TCHAR_TO_UTF8(*ProcessPath), &FileActions, &SpawnAttr, Argv, environ);
 		posix_spawn_file_actions_destroy(&FileActions);
@@ -1652,55 +1665,73 @@ bool FUnixPlatformProcess::IsApplicationRunning( const TCHAR* ProcName )
 	return !system(TCHAR_TO_UTF8(*Commandline));
 }
 
+static bool ReadPipeToStr(void *PipeRead, FString *OutStr)
+{
+	if (PipeRead)
+	{
+		FString NewLine = FPlatformProcess::ReadPipe(PipeRead);
+
+		if (NewLine.Len() > 0)
+		{
+			if (OutStr != nullptr)
+			{
+				*OutStr += NewLine;
+			}
+
+			return true;
+		}
+	}
+
+	return false;
+}
+
 bool FUnixPlatformProcess::ExecProcess(const TCHAR* URL, const TCHAR* Params, int32* OutReturnCode, FString* OutStdOut, FString* OutStdErr, const TCHAR* OptionalWorkingDirectory, bool bShouldEndWithParentProcess)
 {
 	FString CmdLineParams = Params;
 	FString ExecutableFileName = URL;
 	int32 ReturnCode = -1;
-	FString DefaultError;
-	if (!OutStdErr)
+
+	void* PipeReadStdOut = nullptr;
+	void* PipeWriteStdOut = nullptr;
+	verify(FPlatformProcess::CreatePipe(PipeReadStdOut, PipeWriteStdOut));
+
+	void* PipeReadStdErr = nullptr;
+	void* PipeWriteStdErr = nullptr;
+	if (OutStdErr)
 	{
-		OutStdErr = &DefaultError;
+		verify(FPlatformProcess::CreatePipe(PipeReadStdErr, PipeWriteStdErr));
 	}
 
-	void* PipeRead = nullptr;
-	void* PipeWrite = nullptr;
-	verify(FPlatformProcess::CreatePipe(PipeRead, PipeWrite));
-
 	bool bInvoked = false;
 
 	const bool bLaunchDetached = true;
 	const bool bLaunchHidden = false;
 	const bool bLaunchReallyHidden = bLaunchHidden;
 
-	FProcHandle ProcHandle = FPlatformProcess::CreateProc(*ExecutableFileName, *CmdLineParams, bLaunchDetached, bLaunchHidden, bLaunchReallyHidden, NULL, 0, OptionalWorkingDirectory, PipeWrite);
+	FProcHandle ProcHandle = FPlatformProcess::CreateProc(*ExecutableFileName, *CmdLineParams, bLaunchDetached, bLaunchHidden, bLaunchReallyHidden, NULL, 0, OptionalWorkingDirectory, PipeWriteStdOut, nullptr, PipeWriteStdErr);
+
 	if (ProcHandle.IsValid())
 	{
 		while (FPlatformProcess::IsProcRunning(ProcHandle))
 		{
-			FString NewLine = FPlatformProcess::ReadPipe(PipeRead);
-			if (NewLine.Len() > 0)
-			{
-				if (OutStdOut != nullptr)
-				{
-					*OutStdOut += NewLine;
-				}
-			}
+			ReadPipeToStr(PipeReadStdOut, OutStdOut);
+			ReadPipeToStr(PipeReadStdErr, OutStdErr);
 			FPlatformProcess::Sleep(0.5);
 		}
 
-		// read the remainder
-		for(;;)
+		// Read the remainder
+		bool bReadingStdOut = true;
+		bool bReadingStdErr = true;
+		while (bReadingStdOut || bReadingStdErr)
 		{
-			FString NewLine = FPlatformProcess::ReadPipe(PipeRead);
-			if (NewLine.Len() <= 0)
+			if (bReadingStdOut && !ReadPipeToStr(PipeReadStdOut, OutStdOut))
 			{
-				break;
+				bReadingStdOut = false;
 			}
 
-			if (OutStdOut != nullptr)
+			if (bReadingStdErr && !ReadPipeToStr(PipeReadStdErr, OutStdErr))
 			{
-				*OutStdOut += NewLine;
+				bReadingStdErr = false;
 			}
 		}
 
@@ -1727,9 +1758,15 @@ bool FUnixPlatformProcess::ExecProcess(const TCHAR* URL, const TCHAR* Params, in
 		{
 			*OutStdOut = "";
 		}
+		if (OutStdErr != nullptr)
+		{
+			*OutStdErr = "";
+		}
 		UE_LOG(LogHAL, Warning, TEXT("Failed to launch Tool. (%s)"), *ExecutableFileName);
 	}
-	FPlatformProcess::ClosePipe(PipeRead, PipeWrite);
+
+	FPlatformProcess::ClosePipe(PipeReadStdOut, PipeWriteStdOut);
+	FPlatformProcess::ClosePipe(PipeReadStdErr, PipeWriteStdErr);
 	return bInvoked;
 }
 
diff --git a/Engine/Source/Runtime/Core/Private/Windows/WindowsPlatformProcess.cpp b/Engine/Source/Runtime/Core/Private/Windows/WindowsPlatformProcess.cpp
--- a/Engine/Source/Runtime/Core/Private/Windows/WindowsPlatformProcess.cpp
+++ b/Engine/Source/Runtime/Core/Private/Windows/WindowsPlatformProcess.cpp
@@ -286,7 +286,14 @@ void FWindowsPlatformProcess::LaunchURL( const TCHAR* URL, const TCHAR* Parms, F
 	}
 }
 
-FProcHandle FWindowsPlatformProcess::CreateProc( const TCHAR* URL, const TCHAR* Parms, bool bLaunchDetached, bool bLaunchHidden, bool bLaunchReallyHidden, uint32* OutProcessID, int32 PriorityModifier, const TCHAR* OptionalWorkingDirectory, void* PipeWriteChild, void * PipeReadChild)
+FProcHandle FWindowsPlatformProcess::CreateProc(const TCHAR* URL, const TCHAR* Parms, bool bLaunchDetached, bool bLaunchHidden, bool bLaunchReallyHidden, uint32* OutProcessID, int32 PriorityModifier, const TCHAR* OptionalWorkingDirectory, void* PipeWriteChild, void* PipeReadChild)
+{
+	// CreateProc used to only have a single "write" pipe argument that both stdout and stderr would be piped into on Windows,
+	// so for this overload we'll preserve that behaviour for compatibility with existing code
+	return CreateProc(URL, Parms, bLaunchDetached, bLaunchHidden, bLaunchReallyHidden, OutProcessID, PriorityModifier, OptionalWorkingDirectory, PipeWriteChild, PipeReadChild, PipeWriteChild);
+}
+
+FProcHandle FWindowsPlatformProcess::CreateProc(const TCHAR* URL, const TCHAR* Parms, bool bLaunchDetached, bool bLaunchHidden, bool bLaunchReallyHidden, uint32* OutProcessID, int32 PriorityModifier, const TCHAR* OptionalWorkingDirectory, void* PipeWriteChild, void* PipeReadChild, void* PipeStdErrChild)
 {
 	//UE_LOG(LogWindows, Log,  TEXT("CreateProc %s %s"), URL, Parms );
 
@@ -319,7 +326,7 @@ FProcHandle FWindowsPlatformProcess::CreateProc( const TCHAR* URL, const TCHAR*
 		ShowWindowFlags = SW_SHOWMINNOACTIVE;
 	}
 
-	if (PipeWriteChild != nullptr || PipeReadChild != nullptr)
+	if (PipeWriteChild != nullptr || PipeReadChild != nullptr || PipeStdErrChild != nullptr)
 	{
 		dwFlags |= STARTF_USESTDHANDLES;
 	}
@@ -338,7 +345,7 @@ FProcHandle FWindowsPlatformProcess::CreateProc( const TCHAR* URL, const TCHAR*
 		0, NULL,
 		HANDLE(PipeReadChild),
 		HANDLE(PipeWriteChild),
-		HANDLE(PipeWriteChild)
+		HANDLE(PipeStdErrChild)
 	};
 
 	bool bInheritHandles = (dwFlags & STARTF_USESTDHANDLES) != 0;
@@ -2103,4 +2110,4 @@ void FWindowsPlatformProcess::SetThreadName( const TCHAR* ThreadName )
 
 PRAGMA_ENABLE_UNSAFE_TYPECAST_WARNINGS
 
-#include "Windows/HideWindowsPlatformTypes.h"
\ No newline at end of file
+#include "Windows/HideWindowsPlatformTypes.h"
diff --git a/Engine/Source/Runtime/Core/Public/GenericPlatform/GenericPlatformProcess.h b/Engine/Source/Runtime/Core/Public/GenericPlatform/GenericPlatformProcess.h
--- a/Engine/Source/Runtime/Core/Public/GenericPlatform/GenericPlatformProcess.h
+++ b/Engine/Source/Runtime/Core/Public/GenericPlatform/GenericPlatformProcess.h
@@ -404,10 +404,29 @@ struct CORE_API FGenericPlatformProcess
 	 * @param OutProcessId			if non-NULL, this will be filled in with the ProcessId
 	 * @param PriorityModifier		-2 idle, -1 low, 0 normal, 1 high, 2 higher
 	 * @param OptionalWorkingDirectory		Directory to start in when running the program, or NULL to use the current working directory
-	 * @param PipeWrite				Optional HANDLE to pipe for redirecting output
+	 * @param PipeWriteChild		Optional HANDLE to pipe for redirecting output
+	 * @param PipeReadChild			Optional HANDLE to pipe for redirecting input
 	 * @return	The process handle for use in other process functions
 	 */
-	static FProcHandle CreateProc( const TCHAR* URL, const TCHAR* Parms, bool bLaunchDetached, bool bLaunchHidden, bool bLaunchReallyHidden, uint32* OutProcessID, int32 PriorityModifier, const TCHAR* OptionalWorkingDirectory, void* PipeWriteChild, void * PipeReadChild = nullptr);
+	static FProcHandle CreateProc( const TCHAR* URL, const TCHAR* Parms, bool bLaunchDetached, bool bLaunchHidden, bool bLaunchReallyHidden, uint32* OutProcessID, int32 PriorityModifier, const TCHAR* OptionalWorkingDirectory, void* PipeWriteChild, void* PipeReadChild = nullptr);
+	
+	/**
+	 * Creates a new process and its primary thread, with separate std pipes. The new process runs the
+	 * specified executable file in the security context of the calling process.
+	 * @param URL					executable name
+	 * @param Parms					command line arguments
+	 * @param bLaunchDetached		if true, the new process will have its own window
+	 * @param bLaunchHidden			if true, the new process will be minimized in the task bar
+	 * @param bLaunchReallyHidden	if true, the new process will not have a window or be in the task bar
+	 * @param OutProcessId			if non-NULL, this will be filled in with the ProcessId
+	 * @param PriorityModifier		-2 idle, -1 low, 0 normal, 1 high, 2 higher
+	 * @param OptionalWorkingDirectory		Directory to start in when running the program, or NULL to use the current working directory
+	 * @param PipeWriteChild		Optional HANDLE to pipe for redirecting stdout
+	 * @param PipeReadChild			Optional HANDLE to pipe for redirecting stdin
+	 * @param PipeStdErrChild		Optional HANDLE to pipe for redirecting stderr
+	 * @return	The process handle for use in other process functions
+	 */
+	static FProcHandle CreateProc( const TCHAR* URL, const TCHAR* Parms, bool bLaunchDetached, bool bLaunchHidden, bool bLaunchReallyHidden, uint32* OutProcessID, int32 PriorityModifier, const TCHAR* OptionalWorkingDirectory, void* PipeWriteChild, void* PipeReadChild, void* PipeStdErrChild);
 
 	/**
 	 * Opens an existing process. 
diff --git a/Engine/Source/Runtime/Core/Public/Mac/MacPlatformProcess.h b/Engine/Source/Runtime/Core/Public/Mac/MacPlatformProcess.h
--- a/Engine/Source/Runtime/Core/Public/Mac/MacPlatformProcess.h
+++ b/Engine/Source/Runtime/Core/Public/Mac/MacPlatformProcess.h
@@ -241,7 +241,8 @@ struct CORE_API FMacPlatformProcess : public FGenericPlatformProcess
 	static bool CanLaunchURL(const TCHAR* URL);
 	static void LaunchURL(const TCHAR* URL, const TCHAR* Parms, FString* Error);
 	static FString GetGameBundleId();
-	static FProcHandle CreateProc( const TCHAR* URL, const TCHAR* Parms, bool bLaunchDetached, bool bLaunchHidden, bool bLaunchReallyHidden, uint32* OutProcessID, int32 PriorityModifier, const TCHAR* OptionalWorkingDirectory, void* PipeWriteChild, void * PipeReadChild = nullptr);
+	static FProcHandle CreateProc( const TCHAR* URL, const TCHAR* Parms, bool bLaunchDetached, bool bLaunchHidden, bool bLaunchReallyHidden, uint32* OutProcessID, int32 PriorityModifier, const TCHAR* OptionalWorkingDirectory, void* PipeWriteChild, void* PipeReadChild = nullptr );
+	static FProcHandle CreateProc( const TCHAR* URL, const TCHAR* Parms, bool bLaunchDetached, bool bLaunchHidden, bool bLaunchReallyHidden, uint32* OutProcessID, int32 PriorityModifier, const TCHAR* OptionalWorkingDirectory, void* PipeWriteChild, void* PipeReadChild, void* PipeStdErrChild );
 	static FProcHandle OpenProcess(uint32 ProcessID);
 	static bool IsProcRunning( FProcHandle & ProcessHandle );
 	static void WaitForProc( FProcHandle & ProcessHandle );
diff --git a/Engine/Source/Runtime/Core/Public/Unix/UnixPlatformProcess.h b/Engine/Source/Runtime/Core/Public/Unix/UnixPlatformProcess.h
--- a/Engine/Source/Runtime/Core/Public/Unix/UnixPlatformProcess.h
+++ b/Engine/Source/Runtime/Core/Public/Unix/UnixPlatformProcess.h
@@ -280,6 +280,7 @@ struct CORE_API FUnixPlatformProcess : public FGenericPlatformProcess
 	static bool CanLaunchURL(const TCHAR* URL);
 	static void LaunchURL(const TCHAR* URL, const TCHAR* Parms, FString* Error);
 	static FProcHandle CreateProc(const TCHAR* URL, const TCHAR* Parms, bool bLaunchDetached, bool bLaunchHidden, bool bLaunchReallyHidden, uint32* OutProcessID, int32 PriorityModifier, const TCHAR* OptionalWorkingDirectory, void* PipeWriteChild, void* PipeReadChild = nullptr);
+	static FProcHandle CreateProc(const TCHAR* URL, const TCHAR* Parms, bool bLaunchDetached, bool bLaunchHidden, bool bLaunchReallyHidden, uint32* OutProcessID, int32 PriorityModifier, const TCHAR* OptionalWorkingDirectory, void* PipeWriteChild, void* PipeReadChild , void* PipeStdErrChild);
 	static FProcHandle OpenProcess(uint32 ProcessID);
 	static bool IsProcRunning( FProcHandle & ProcessHandle );
 	static void WaitForProc( FProcHandle & ProcessHandle );
diff --git a/Engine/Source/Runtime/Core/Public/Windows/WindowsPlatformProcess.h b/Engine/Source/Runtime/Core/Public/Windows/WindowsPlatformProcess.h
--- a/Engine/Source/Runtime/Core/Public/Windows/WindowsPlatformProcess.h
+++ b/Engine/Source/Runtime/Core/Public/Windows/WindowsPlatformProcess.h
@@ -154,7 +154,8 @@ public:
 	static const FString GetModulesDirectory();
 	static bool CanLaunchURL(const TCHAR* URL);
 	static void LaunchURL( const TCHAR* URL, const TCHAR* Parms, FString* Error );
-	static FProcHandle CreateProc( const TCHAR* URL, const TCHAR* Parms, bool bLaunchDetached, bool bLaunchHidden, bool bLaunchReallyHidden, uint32* OutProcessID, int32 PriorityModifier, const TCHAR* OptionalWorkingDirectory, void* PipeWriteChild, void * PipeReadChild = nullptr);
+	static FProcHandle CreateProc( const TCHAR* URL, const TCHAR* Parms, bool bLaunchDetached, bool bLaunchHidden, bool bLaunchReallyHidden, uint32* OutProcessID, int32 PriorityModifier, const TCHAR* OptionalWorkingDirectory, void* PipeWriteChild, void* PipeReadChild = nullptr);
+	static FProcHandle CreateProc( const TCHAR* URL, const TCHAR* Parms, bool bLaunchDetached, bool bLaunchHidden, bool bLaunchReallyHidden, uint32* OutProcessID, int32 PriorityModifier, const TCHAR* OptionalWorkingDirectory, void* PipeWriteChild, void* PipeReadChild, void* PipeStdErrChild);
 	static bool SetProcPriority(FProcHandle & InProcHandle, int32 PriorityModifier);
 	static FProcHandle OpenProcess(uint32 ProcessID);
 	static bool IsProcRunning( FProcHandle & ProcessHandle );
