Move trace platform event availability define from TraceLog to Core build files

diff --git a/Engine/Source/Runtime/Core/Core.Build.cs b/Engine/Source/Runtime/Core/Core.Build.cs
--- a/Engine/Source/Runtime/Core/Core.Build.cs
+++ b/Engine/Source/Runtime/Core/Core.Build.cs
@@ -309,6 +309,7 @@ public class Core : ModuleRules
 		if (Target.Platform == UnrealTargetPlatform.Win64
 			&& Target.Configuration != UnrealTargetConfiguration.Shipping)
 		{
+			PublicDefinitions.Add("PLATFORM_SUPPORTS_PLATFORM_EVENTS=1");
 			PublicDefinitions.Add("PLATFORM_SUPPORTS_TRACE_WIN32_VIRTUAL_MEMORY_HOOKS=1");
 			PublicDefinitions.Add("PLATFORM_SUPPORTS_TRACE_WIN32_MODULE_DIAGNOSTICS=1");
 			PublicDefinitions.Add("PLATFORM_SUPPORTS_TRACE_WIN32_CALLSTACK=1");
diff --git a/Engine/Source/Runtime/TraceLog/TraceLog.Build.cs b/Engine/Source/Runtime/TraceLog/TraceLog.Build.cs
--- a/Engine/Source/Runtime/TraceLog/TraceLog.Build.cs
+++ b/Engine/Source/Runtime/TraceLog/TraceLog.Build.cs
@@ -10,11 +10,6 @@ public class TraceLog : ModuleRules
 	{
 		bRequiresImplementModule = false;
 		PublicIncludePathModuleNames.Add("Core");
-	
-		if (Target.Platform == UnrealTargetPlatform.Win64)
-		{
-			PublicDefinitions.Add("PLATFORM_SUPPORTS_PLATFORM_EVENTS=1");
-		}
     }
 
 	// used by platform extension derived classes. probably should become a project setting at some point!
