add FPlatformProcess::SetupRHIThread to match SetupGameThread, SetupRenderThread and SetupAudioThread.

#rnx

diff --git a/Engine/Source/Runtime/Core/Public/GenericPlatform/GenericPlatformProcess.h b/Engine/Source/Runtime/Core/Public/GenericPlatform/GenericPlatformProcess.h
--- a/Engine/Source/Runtime/Core/Public/GenericPlatform/GenericPlatformProcess.h
+++ b/Engine/Source/Runtime/Core/Public/GenericPlatform/GenericPlatformProcess.h
@@ -255,6 +255,9 @@ struct CORE_API FGenericPlatformProcess
 	/** Allow the platform to do anything it needs for render thread */
 	static void SetupRenderThread() { }
 
+	/** Allow the platform to do anything it needs for the RHI thread */
+	static void SetupRHIThread() { }
+
 	/** Allow the platform to do anything it needs for audio thread */
 	static void SetupAudioThread() { }
 
diff --git a/Engine/Source/Runtime/RenderCore/Private/RenderingThread.cpp b/Engine/Source/Runtime/RenderCore/Private/RenderingThread.cpp
--- a/Engine/Source/Runtime/RenderCore/Private/RenderingThread.cpp
+++ b/Engine/Source/Runtime/RenderCore/Private/RenderingThread.cpp
@@ -319,6 +319,7 @@ public:
 		FTaskGraphInterface::Get().AttachToThread(ENamedThreads::RHIThread);
 		FTaskGraphInterface::Get().ProcessThreadUntilRequestReturn(ENamedThreads::RHIThread);
 		FMemory::ClearAndDisableTLSCachesOnCurrentThread();
+		FPlatformProcess::SetupRHIThread();
 		return 0;
 	}
 
