LWC support for FQuat and FTransform

diff --git a/Engine/Source/Programs/UnrealHeaderTool/Private/CodeGenerator.cpp b/Engine/Source/Programs/UnrealHeaderTool/Private/CodeGenerator.cpp
--- a/Engine/Source/Programs/UnrealHeaderTool/Private/CodeGenerator.cpp
+++ b/Engine/Source/Programs/UnrealHeaderTool/Private/CodeGenerator.cpp
@@ -1721,7 +1721,7 @@ static bool IsAlwaysAccessible(FUnrealScriptStructDefinitionInfo& ScriptDef)
 	if( ToTest == NAME_Plane || ToTest == NAME_Plane4f || ToTest == NAME_Plane4d
 		||	ToTest == NAME_Vector || ToTest == NAME_Vector3f || ToTest == NAME_Vector3d
 		||	ToTest == NAME_Vector4
-		||	ToTest == NAME_Quat
+		||	ToTest == NAME_Quat || ToTest == NAME_Quat4f || ToTest == NAME_Quat4d
 		||	ToTest == NAME_Color
 		)
 	{
diff --git a/Engine/Source/Runtime/ApplicationCore/Private/Android/AndroidCursor.cpp b/Engine/Source/Runtime/ApplicationCore/Private/Android/AndroidCursor.cpp
--- a/Engine/Source/Runtime/ApplicationCore/Private/Android/AndroidCursor.cpp
+++ b/Engine/Source/Runtime/ApplicationCore/Private/Android/AndroidCursor.cpp
@@ -69,10 +69,10 @@ void FAndroidCursor::Lock(const RECT* const Bounds)
 	}
 	else
 	{
-		CursorClipRect.Min.X = FMath::TruncToInt(Bounds->left);
-		CursorClipRect.Min.Y = FMath::TruncToInt(Bounds->top);
-		CursorClipRect.Max.X = FMath::TruncToInt(Bounds->right) - 1;
-		CursorClipRect.Max.Y = FMath::TruncToInt(Bounds->bottom) - 1;
+		CursorClipRect.Min.X = (int32)Bounds->left;
+		CursorClipRect.Min.Y = (int32)Bounds->top;
+		CursorClipRect.Max.X = (int32)Bounds->right - 1;
+		CursorClipRect.Max.Y = (int32)Bounds->bottom - 1;
 	}
 
 	FVector2D Position = GetPosition();
diff --git a/Engine/Source/Runtime/Core/Private/Math/Transform.cpp b/Engine/Source/Runtime/Core/Private/Math/Transform.cpp
--- a/Engine/Source/Runtime/Core/Private/Math/Transform.cpp
+++ b/Engine/Source/Runtime/Core/Private/Math/Transform.cpp
@@ -12,46 +12,49 @@
 
 DEFINE_LOG_CATEGORY_STATIC(LogTransform, Log, All);
 
-// FTransform identity
-// @Note: Do not reference FVector::ZeroVector or FVector::OneVector
-// because they're not initialized yet, it will come as 0 vector
-const FTransform FTransform::Identity(FQuat(0.f,0.f,0.f,1.f), FVector(0.f), FVector(1.f));
-
+namespace UE
+{
+namespace Math
+{
 
-// Replacement of Inverse of FMatrix
+// Replacement of Inverse of TMatrix<T>
 
 /**
 * Does a debugf of the contents of this BoneAtom.
 */
-void FTransform::DebugPrint() const
+template<typename T>
+void TTransform<T>::DebugPrint() const
 {
 	UE_LOG(LogTransform, Log, TEXT("%s"), *ToHumanReadableString());
 }
 
-FString FTransform::ToHumanReadableString() const
+template<typename T>
+FString TTransform<T>::ToHumanReadableString() const
 {
 	FRotator R(GetRotation());
-	FVector T(GetTranslation());
-	FVector S(GetScale3D());
+	TVector<T> TT(GetTranslation());
+	TVector<T> S(GetScale3D());
 
 	FString Output = FString::Printf(TEXT("Rotation: Pitch %f Yaw %f Roll %f\r\n"), R.Pitch, R.Yaw, R.Roll);
-	Output += FString::Printf(TEXT("Translation: %f %f %f\r\n"), T.X, T.Y, T.Z);
+	Output += FString::Printf(TEXT("Translation: %f %f %f\r\n"), TT.X, TT.Y, TT.Z);
 	Output += FString::Printf(TEXT("Scale3D: %f %f %f\r\n"), S.X, S.Y, S.Z);
 
 	return Output;
 }
 
 
-FString FTransform::ToString() const
+template<typename T>
+FString TTransform<T>::ToString() const
 {
 	const FRotator R(Rotator());
-	const FVector T(GetTranslation());
-	const FVector S(GetScale3D());
+	const TVector<T> TT(GetTranslation());
+	const TVector<T> S(GetScale3D());
 
-	return FString::Printf(TEXT("%f,%f,%f|%f,%f,%f|%f,%f,%f"), T.X, T.Y, T.Z, R.Pitch, R.Yaw, R.Roll, S.X, S.Y, S.Z);
+	return FString::Printf(TEXT("%f,%f,%f|%f,%f,%f|%f,%f,%f"), TT.X, TT.Y, TT.Z, R.Pitch, R.Yaw, R.Roll, S.X, S.Y, S.Z);
 }
 
-bool FTransform::InitFromString( const FString& Source )
+template<typename T>
+bool TTransform<T>::InitFromString( const FString& Source )
 {
 	TArray<FString> ComponentStrings;
 	Source.ParseIntoArray(ComponentStrings, TEXT("|"), true);
@@ -62,7 +65,7 @@ bool FTransform::InitFromString( const FString& Source )
 	}
 
 	// Translation
-	FVector ParsedTranslation = FVector::ZeroVector;
+	TVector<T> ParsedTranslation = TVector<T>::ZeroVector;
 	if( !FDefaultValueHelper::ParseVector(ComponentStrings[0], ParsedTranslation) )
 	{
 		return false;
@@ -76,19 +79,20 @@ bool FTransform::InitFromString( const FString& Source )
 	}
 
 	// Scale
-	FVector ParsedScale = FVector::OneVector;
+	TVector<T> ParsedScale = TVector<T>::OneVector;
 	if( !FDefaultValueHelper::ParseVector(ComponentStrings[2], ParsedScale) )
 	{
 		return false;
 	}
 
-	SetComponents(FQuat(ParsedRotation), ParsedTranslation, ParsedScale);
+	SetComponents(TQuat<T>(ParsedRotation), ParsedTranslation, ParsedScale);
 
 	return true;
 }
 
 #define DEBUG_INVERSE_TRANSFORM 0
-FTransform FTransform::GetRelativeTransformReverse(const FTransform& Other) const
+template<typename T>
+TTransform<T> TTransform<T>::GetRelativeTransformReverse(const TTransform<T>& Other) const
 {
 	// A (-1) * B = VQS(B)(VQS (A)(-1))
 	// 
@@ -96,9 +100,9 @@ FTransform FTransform::GetRelativeTransformReverse(const FTransform& Other) cons
 	// Rotation = Q(B) * Q(A)(-1)
 	// Translation = T(B)-S(B)/S(A) *[Q(B)*Q(A)(-1)*T(A)*Q(A)*Q(B)(-1)]
 	// where A = this, and B = Other
-	FTransform Result;
+	TTransform<T> Result;
 
-	FVector SafeRecipScale3D = GetSafeScaleReciprocal(Scale3D);
+	TVector<T> SafeRecipScale3D = GetSafeScaleReciprocal(Scale3D);
 	Result.Scale3D = Other.Scale3D*SafeRecipScale3D;	
 
 	Result.Rotation = Other.Rotation*Rotation.Inverse();
@@ -106,8 +110,8 @@ FTransform FTransform::GetRelativeTransformReverse(const FTransform& Other) cons
 	Result.Translation = Other.Translation - Result.Scale3D * ( Result.Rotation * Translation );
 
 #if DEBUG_INVERSE_TRANSFORM
-	FMatrix AM = ToMatrixWithScale();
-	FMatrix BM = Other.ToMatrixWithScale();
+	TMatrix<T> AM = ToMatrixWithScale();
+	TMatrix<T> BM = Other.ToMatrixWithScale();
 
 	Result.DebugEqualMatrix(AM.InverseFast() *  BM);
 #endif
@@ -117,9 +121,10 @@ FTransform FTransform::GetRelativeTransformReverse(const FTransform& Other) cons
 
 /**
  * Set current transform and the relative to ParentTransform.
- * Equates to This = This->GetRelativeTransform(Parent), but saves the intermediate FTransform storage and copy.
+ * Equates to This = This->GetRelativeTransform(Parent), but saves the intermediate TTransform<T> storage and copy.
  */
-void FTransform::SetToRelativeTransform(const FTransform& ParentTransform)
+template<typename T>
+void TTransform<T>::SetToRelativeTransform(const TTransform<T>& ParentTransform)
 {
 	// A * B(-1) = VQS(B)(-1) (VQS (A))
 	// 
@@ -128,12 +133,12 @@ void FTransform::SetToRelativeTransform(const FTransform& ParentTransform)
 	// Translation = 1/S(B) *[Q(B)(-1)*(T(A)-T(B))*Q(B)]
 	// where A = this, B = Other
 #if DEBUG_INVERSE_TRANSFORM
- 	FMatrix AM = ToMatrixWithScale();
- 	FMatrix BM = ParentTransform.ToMatrixWithScale();
+ 	TMatrix<T> AM = ToMatrixWithScale();
+ 	TMatrix<T> BM = ParentTransform.ToMatrixWithScale();
 #endif
 
-	const FVector SafeRecipScale3D = GetSafeScaleReciprocal(ParentTransform.Scale3D, SMALL_NUMBER);
-	const FQuat InverseRot = ParentTransform.Rotation.Inverse();
+	const TVector<T> SafeRecipScale3D = GetSafeScaleReciprocal(ParentTransform.Scale3D, SMALL_NUMBER);
+	const TQuat<T> InverseRot = ParentTransform.Rotation.Inverse();
 
 	Scale3D *= SafeRecipScale3D;	
 	Translation = (InverseRot * (Translation - ParentTransform.Translation)) * SafeRecipScale3D;
@@ -144,19 +149,21 @@ void FTransform::SetToRelativeTransform(const FTransform& ParentTransform)
 #endif
 }
 
-void FTransform::GetRelativeTransformUsingMatrixWithScale(FTransform* OutTransform, const FTransform* Base, const FTransform* Relative)
+template<typename T>
+void TTransform<T>::GetRelativeTransformUsingMatrixWithScale(TTransform<T>* OutTransform, const TTransform<T>* Base, const TTransform<T>* Relative)
 {
 	// the goal of using M is to get the correct orientation
 	// but for translation, we still need scale
-	FMatrix AM = Base->ToMatrixWithScale();
-	FMatrix BM = Relative->ToMatrixWithScale();
+	TMatrix<T> AM = Base->ToMatrixWithScale();
+	TMatrix<T> BM = Relative->ToMatrixWithScale();
 	// get combined scale
-	FVector SafeRecipScale3D = GetSafeScaleReciprocal(Relative->Scale3D, SMALL_NUMBER);
-	FVector DesiredScale3D = Base->Scale3D*SafeRecipScale3D;
+	TVector<T> SafeRecipScale3D = GetSafeScaleReciprocal(Relative->Scale3D, SMALL_NUMBER);
+	TVector<T> DesiredScale3D = Base->Scale3D*SafeRecipScale3D;
 	ConstructTransformFromMatrixWithDesiredScale(AM, BM.Inverse(), DesiredScale3D, *OutTransform);
 }
 
-FTransform FTransform::GetRelativeTransform(const FTransform& Other) const
+template<typename T>
+TTransform<T> TTransform<T>::GetRelativeTransform(const TTransform<T>& Other) const
 {
 	// A * B(-1) = VQS(B)(-1) (VQS (A))
 	// 
@@ -164,7 +171,7 @@ FTransform FTransform::GetRelativeTransform(const FTransform& Other) const
 	// Rotation = Q(B)(-1) * Q(A)
 	// Translation = 1/S(B) *[Q(B)(-1)*(T(A)-T(B))*Q(B)]
 	// where A = this, B = Other
-	FTransform Result;
+	TTransform<T> Result;
 
 	if (AnyHasNegativeScale(Scale3D, Other.GetScale3D()))
 	{
@@ -173,22 +180,22 @@ FTransform FTransform::GetRelativeTransform(const FTransform& Other) const
 	}
 	else
 	{
-		FVector SafeRecipScale3D = GetSafeScaleReciprocal(Other.Scale3D, SMALL_NUMBER);
+		TVector<T> SafeRecipScale3D = GetSafeScaleReciprocal(Other.Scale3D, SMALL_NUMBER);
 		Result.Scale3D = Scale3D*SafeRecipScale3D;
 
 		if (Other.Rotation.IsNormalized() == false)
 		{
-			return FTransform::Identity;
+			return TTransform<T>::Identity;
 		}
 
-		FQuat Inverse = Other.Rotation.Inverse();
+		TQuat<T> Inverse = Other.Rotation.Inverse();
 		Result.Rotation = Inverse*Rotation;
 
 		Result.Translation = (Inverse*(Translation - Other.Translation))*(SafeRecipScale3D);
 
 #if DEBUG_INVERSE_TRANSFORM
-		FMatrix AM = ToMatrixWithScale();
-		FMatrix BM = Other.ToMatrixWithScale();
+		TMatrix<T> AM = ToMatrixWithScale();
+		TMatrix<T> BM = Other.ToMatrixWithScale();
 
 		Result.DebugEqualMatrix(AM *  BM.InverseFast());
 
@@ -198,9 +205,10 @@ FTransform FTransform::GetRelativeTransform(const FTransform& Other) const
 	return Result;
 }
 
-bool FTransform::DebugEqualMatrix(const FMatrix& Matrix) const
+template<typename T>
+bool TTransform<T>::DebugEqualMatrix(const TMatrix<T>& Matrix) const
 {
-	FTransform TestResult(Matrix);
+	TTransform<T> TestResult(Matrix);
 	if (!Equals(TestResult))
 	{
 		// see now which one isn't equal
@@ -229,4 +237,12 @@ bool FTransform::DebugEqualMatrix(const FMatrix& Matrix) const
 	return true;
 }
 
+
+} // namespace UE::Math
+} // namespace UE
+
+// Instantiate for linker.
+template struct UE::Math::TTransform<float>;
+template struct UE::Math::TTransform<double>;
+
 #endif // #if !ENABLE_VECTORIZED_TRANSFORM
diff --git a/Engine/Source/Runtime/Core/Private/Math/TransformVectorized.cpp b/Engine/Source/Runtime/Core/Private/Math/TransformVectorized.cpp
--- a/Engine/Source/Runtime/Core/Private/Math/TransformVectorized.cpp
+++ b/Engine/Source/Runtime/Core/Private/Math/TransformVectorized.cpp
@@ -10,12 +10,7 @@
 #include "Containers/Array.h"
 #include "Containers/UnrealString.h"
 #include "Logging/LogMacros.h"
-#include "Math/Vector.h"
-#include "Math/VectorRegister.h"
-#include "Math/Rotator.h"
-#include "Math/Matrix.h"
-#include "Math/Quat.h"
-#include "Math/Transform.h"
+#include "Math/MathFwd.h"
 
 #if ENABLE_VECTORIZED_TRANSFORM
 
@@ -23,45 +18,49 @@
 
 DEFINE_LOG_CATEGORY_STATIC(LogTransform, Log, All);
 
-// FTransform identity
-// @Note: Do not reference FVector::ZeroVector or FVector::OneVector
-// because they're not initialized yet, it will come as 0 vector
-const FTransform FTransform::Identity(FQuat(0.f,0.f,0.f,1.f), FVector(0.f), FVector(1.f));
+namespace UE
+{
+namespace Math
+{
 
-// Replacement of Inverse of FMatrix
+// Replacement of Inverse of TMatrix<T>
 
 /**
 * Does a debugf of the contents of this BoneAtom.
 */
-void FTransform::DebugPrint() const
+template<typename T>
+void TTransform<T>::DebugPrint() const
 {
 	UE_LOG(LogTransform, Log, TEXT("%s"), *ToHumanReadableString());
 }
 
-FString FTransform::ToHumanReadableString() const
+template<typename T>
+FString TTransform<T>::ToHumanReadableString() const
 {
 	FRotator R(GetRotation());
-	FVector T(GetTranslation());
-	FVector S(GetScale3D());
+	TVector<T> Tr(GetTranslation());
+	TVector<T> S(GetScale3D());
 
 	FString Output= FString::Printf(TEXT("Rotation: Pitch %f Yaw %f Roll %f\r\n"), R.Pitch, R.Yaw, R.Roll);
-	Output += FString::Printf(TEXT("Translation: %f %f %f\r\n"), T.X, T.Y, T.Z);
+	Output += FString::Printf(TEXT("Translation: %f %f %f\r\n"), Tr.X, Tr.Y, Tr.Z);
 	Output += FString::Printf(TEXT("Scale3D: %f %f %f\r\n"), S.X, S.Y, S.Z);
 
 	return Output;
 }
 
 
-FString FTransform::ToString() const
+template<typename T>
+FString TTransform<T>::ToString() const
 {
 	const FRotator R(Rotator());
-	const FVector T(GetTranslation());
-	const FVector S(GetScale3D());
+	const TVector<T> Tr(GetTranslation());
+	const TVector<T> S(GetScale3D());
 
-	return FString::Printf(TEXT("%f,%f,%f|%f,%f,%f|%f,%f,%f"), T.X, T.Y, T.Z, R.Pitch, R.Yaw, R.Roll, S.X, S.Y, S.Z);
+	return FString::Printf(TEXT("%f,%f,%f|%f,%f,%f|%f,%f,%f"), Tr.X, Tr.Y, Tr.Z, R.Pitch, R.Yaw, R.Roll, S.X, S.Y, S.Z);
 }
 
-bool FTransform::InitFromString( const FString& Source )
+template<typename T>
+bool TTransform<T>::InitFromString( const FString& Source )
 {
 	TArray<FString> ComponentStrings;
 	Source.ParseIntoArray(ComponentStrings, TEXT("|"), true);
@@ -72,7 +71,7 @@ bool FTransform::InitFromString( const FString& Source )
 	}
 
 	// Translation
-	FVector ParsedTranslation = FVector::ZeroVector;
+	TVector<T> ParsedTranslation = TVector<T>::ZeroVector;
 	if( !FDefaultValueHelper::ParseVector(ComponentStrings[0], ParsedTranslation) )
 	{
 		return false;
@@ -86,19 +85,20 @@ bool FTransform::InitFromString( const FString& Source )
 	}
 
 	// Scale
-	FVector Scale = FVector::OneVector;
+	TVector<T> Scale = TVector<T>::OneVector;
 	if( !FDefaultValueHelper::ParseVector(ComponentStrings[2], Scale) )
 	{
 		return false;
 	}
 
-	SetComponents(FQuat(ParsedRotation), ParsedTranslation, Scale);
+	SetComponents(TQuat<T>::MakeFromRotator(ParsedRotation), ParsedTranslation, Scale);
 
 	return true;
 }
 
 #define DEBUG_INVERSE_TRANSFORM 0
-FTransform FTransform::GetRelativeTransformReverse(const FTransform& Other) const
+template<typename T>
+TTransform<T> TTransform<T>::GetRelativeTransformReverse(const TTransform<T>& Other) const
 {
 	// A (-1) * B = VQS(B)(VQS (A)(-1))
 	// 
@@ -106,21 +106,21 @@ FTransform FTransform::GetRelativeTransformReverse(const FTransform& Other) cons
 	// Rotation = Q(B) * Q(A)(-1)
 	// Translation = T(B)-S(B)/S(A) *[Q(B)*Q(A)(-1)*T(A)*Q(A)*Q(B)(-1)]
 	// where A = this, and B = Other
-	FTransform Result;
+	TTransform<T> Result;
 
 	// Scale = S(B)/S(A)	
-	VectorRegister VSafeScale3D	= VectorSet_W0(GetSafeScaleReciprocal(Scale3D));
-	VectorRegister VScale3D = VectorMultiply(Other.Scale3D, VSafeScale3D);
+	TransformVectorRegister VSafeScale3D	= VectorSet_W0(GetSafeScaleReciprocal(Scale3D));
+	TransformVectorRegister VScale3D = VectorMultiply(Other.Scale3D, VSafeScale3D);
 	
 	// Rotation = Q(B) * Q(A)(-1)	
-	VectorRegister4Float VInverseRot = VectorQuaternionInverse(Rotation);
-	VectorRegister4Float VRotation = VectorQuaternionMultiply2(Other.Rotation, VInverseRot );
+	TransformVectorRegister VInverseRot = VectorQuaternionInverse(Rotation);
+	TransformVectorRegister VRotation = VectorQuaternionMultiply2(Other.Rotation, VInverseRot );
 	
 	// RotatedTranslation
-	VectorRegister VR = VectorQuaternionRotateVector(VRotation, Translation);
+	TransformVectorRegister VR = VectorQuaternionRotateVector(VRotation, Translation);
 
 	// Translation = T(B)-S(B)/S(A) *[Q(B)*Q(A)(-1)*T(A)*Q(A)*Q(B)(-1)]	
-	VectorRegister VTranslation = VectorSet_W0(VectorNegateMultiplyAdd(VScale3D, VR, Other.Translation));
+	TransformVectorRegister VTranslation = VectorSet_W0(VectorNegateMultiplyAdd(VScale3D, VR, Other.Translation));
 
 	Result.Scale3D = VScale3D;	
 	Result.Translation = VTranslation;
@@ -129,8 +129,8 @@ FTransform FTransform::GetRelativeTransformReverse(const FTransform& Other) cons
 	Result.DiagnosticCheckNaN_All(); 
 
 #if DEBUG_INVERSE_TRANSFORM
-	FMatrix AM = ToMatrixWithScale();
-	FMatrix BM = Other.ToMatrixWithScale();
+	TMatrix<T> AM = ToMatrixWithScale();
+	TMatrix<T> BM = Other.ToMatrixWithScale();
 
 	Result.DebugEqualMatrix(AM.InverseFast() *  BM);
 #endif
@@ -140,9 +140,10 @@ FTransform FTransform::GetRelativeTransformReverse(const FTransform& Other) cons
 
 /**
  * Set current transform and the relative to ParentTransform.
- * Equates to This = This->GetRelativeTransform(Parent), but saves the intermediate FTransform storage and copy.
+ * Equates to This = This->GetRelativeTransform(Parent), but saves the intermediate TTransform<T> storage and copy.
  */
-void FTransform::SetToRelativeTransform(const FTransform& ParentTransform)
+template<typename T>
+void TTransform<T>::SetToRelativeTransform(const TTransform<T>& ParentTransform)
 {
 	// A * B(-1) = VQS(B)(-1) (VQS (A))
 	// 
@@ -151,22 +152,22 @@ void FTransform::SetToRelativeTransform(const FTransform& ParentTransform)
 	// Translation = 1/S(B) *[Q(B)(-1)*(T(A)-T(B))*Q(B)]
 	// where A = this, B = Other
 #if DEBUG_INVERSE_TRANSFORM
-	FMatrix AM = ToMatrixWithScale();
-	FMatrix BM = ParentTransform.ToMatrixWithScale();
+	TMatrix<T> AM = ToMatrixWithScale();
+	TMatrix<T> BM = ParentTransform.ToMatrixWithScale();
 #endif
 	
 	checkSlow(ParentTransform.IsRotationNormalized());
 
 	// Scale = S(A)/S(B)	
-	VectorRegister VSafeScale3D	= VectorSet_W0(GetSafeScaleReciprocal(ParentTransform.Scale3D, ScalarRegister(SMALL_NUMBER)));
+	TransformVectorRegister VSafeScale3D	= VectorSet_W0(GetSafeScaleReciprocal(ParentTransform.Scale3D, ScalarRegister(SMALL_NUMBER)));
 	Scale3D = VectorMultiply(Scale3D, VSafeScale3D);
 	
 	//VQTranslation = (  ( T(A).X - T(B).X ),  ( T(A).Y - T(B).Y ), ( T(A).Z - T(B).Z), 0.f );
-	VectorRegister VQTranslation = VectorSet_W0(VectorSubtract(Translation, ParentTransform.Translation));
+	TransformVectorRegister VQTranslation = VectorSet_W0(VectorSubtract(Translation, ParentTransform.Translation));
 
 	// Inverse RotatedTranslation
-	VectorRegister4Float VInverseParentRot = VectorQuaternionInverse(ParentTransform.Rotation);
-	VectorRegister VR = VectorQuaternionRotateVector(VInverseParentRot, VQTranslation);
+	TransformVectorRegister VInverseParentRot = VectorQuaternionInverse(ParentTransform.Rotation);
+	TransformVectorRegister VR = VectorQuaternionRotateVector(VInverseParentRot, VQTranslation);
 
 	// Translation = 1/S(B)
 	Translation = VectorMultiply(VR, VSafeScale3D);
@@ -181,21 +182,23 @@ void FTransform::SetToRelativeTransform(const FTransform& ParentTransform)
 #endif
 }
 
-void FTransform::GetRelativeTransformUsingMatrixWithScale(FTransform* OutTransform, const FTransform* Base, const FTransform* Relative)
+template<typename T>
+void TTransform<T>::GetRelativeTransformUsingMatrixWithScale(TTransform<T>* OutTransform, const TTransform<T>* Base, const TTransform<T>* Relative)
 {
 	// the goal of using M is to get the correct orientation
 	// but for translation, we still need scale
-	FMatrix AM = Base->ToMatrixWithScale();
-	FMatrix BM = Relative->ToMatrixWithScale();
+	TMatrix<T> AM = Base->ToMatrixWithScale();
+	TMatrix<T> BM = Relative->ToMatrixWithScale();
 	// get combined scale
 	// Scale = S(A)/S(B)
 	static ScalarRegister STolerance(SMALL_NUMBER);
-	VectorRegister VSafeScale3D = VectorSet_W0(GetSafeScaleReciprocal(Relative->Scale3D, STolerance));
-	VectorRegister VScale3D = VectorMultiply(Base->Scale3D, VSafeScale3D);
+	TransformVectorRegister VSafeScale3D = VectorSet_W0(GetSafeScaleReciprocal(Relative->Scale3D, STolerance));
+	TransformVectorRegister VScale3D = VectorMultiply(Base->Scale3D, VSafeScale3D);
 	ConstructTransformFromMatrixWithDesiredScale(AM, BM.Inverse(), VScale3D, *OutTransform);
 }
 
-FTransform FTransform::GetRelativeTransform(const FTransform& Other) const
+template<typename T>
+TTransform<T> TTransform<T>::GetRelativeTransform(const TTransform<T>& Other) const
 {
 	// A * B(-1) = VQS(B)(-1) (VQS (A))
 	// 
@@ -203,11 +206,11 @@ FTransform FTransform::GetRelativeTransform(const FTransform& Other) const
 	// Rotation = Q(B)(-1) * Q(A)
 	// Translation = 1/S(B) *[Q(B)(-1)*(T(A)-T(B))*Q(B)]
 	// where A = this, B = Other
-	FTransform Result;
+	TTransform<T> Result;
 		
 	if (Other.IsRotationNormalized() == false)
 	{
-		return FTransform::Identity;
+		return TTransform<T>::Identity;
 	}
 
 	if (Private_AnyHasNegativeScale(this->Scale3D, Other.Scale3D))
@@ -219,22 +222,22 @@ FTransform FTransform::GetRelativeTransform(const FTransform& Other) const
 	{
 		// Scale = S(A)/S(B)
 		static ScalarRegister STolerance(SMALL_NUMBER);
-		VectorRegister VSafeScale3D = VectorSet_W0(GetSafeScaleReciprocal(Other.Scale3D, STolerance));
+		TransformVectorRegister VSafeScale3D = VectorSet_W0(GetSafeScaleReciprocal(Other.Scale3D, STolerance));
 
-		VectorRegister VScale3D = VectorMultiply(Scale3D, VSafeScale3D);
+		TransformVectorRegister VScale3D = VectorMultiply(Scale3D, VSafeScale3D);
 
 		//VQTranslation = (  ( T(A).X - T(B).X ),  ( T(A).Y - T(B).Y ), ( T(A).Z - T(B).Z), 0.f );
-		VectorRegister VQTranslation = VectorSet_W0(VectorSubtract(Translation, Other.Translation));
+		TransformVectorRegister VQTranslation = VectorSet_W0(VectorSubtract(Translation, Other.Translation));
 
 		// Inverse RotatedTranslation
-		VectorRegister4Float VInverseRot = VectorQuaternionInverse(Other.Rotation);
-		VectorRegister VR = VectorQuaternionRotateVector(VInverseRot, VQTranslation);
+		TransformVectorRegister VInverseRot = VectorQuaternionInverse(Other.Rotation);
+		TransformVectorRegister VR = VectorQuaternionRotateVector(VInverseRot, VQTranslation);
 
 		//Translation = 1/S(B)
-		VectorRegister VTranslation = VectorMultiply(VR, VSafeScale3D);
+		TransformVectorRegister VTranslation = VectorMultiply(VR, VSafeScale3D);
 
 		// Rotation = Q(B)(-1) * Q(A)	
-		VectorRegister4Float VRotation = VectorQuaternionMultiply2(VInverseRot, Rotation);
+		TransformVectorRegister VRotation = VectorQuaternionMultiply2(VInverseRot, Rotation);
 
 		Result.Scale3D = VScale3D;
 		Result.Translation = VTranslation;
@@ -242,8 +245,8 @@ FTransform FTransform::GetRelativeTransform(const FTransform& Other) const
 
 		Result.DiagnosticCheckNaN_All();
 #if DEBUG_INVERSE_TRANSFORM
-		FMatrix AM = ToMatrixWithScale();
-		FMatrix BM = Other.ToMatrixWithScale();
+		TMatrix<T> AM = ToMatrixWithScale();
+		TMatrix<T> BM = Other.ToMatrixWithScale();
 
 		Result.DebugEqualMatrix(AM *  BM.InverseFast());
 #endif
@@ -252,9 +255,10 @@ FTransform FTransform::GetRelativeTransform(const FTransform& Other) const
 	return Result;
 }
 
-bool FTransform::DebugEqualMatrix(const FMatrix& Matrix) const
+template<typename T>
+bool TTransform<T>::DebugEqualMatrix(const TMatrix<T>& Matrix) const
 {
-	FTransform TestResult(Matrix);
+	TTransform<T> TestResult(Matrix);
 	if (!Equals(TestResult))
 	{
 		// see now which one isn't equal		
@@ -283,4 +287,12 @@ bool FTransform::DebugEqualMatrix(const FMatrix& Matrix) const
 	return true;
 }
 
+} // namespace UE::Math
+} // namespace UE
+
+
+// Instantiate for linker.
+template struct UE::Math::TTransform<float>;
+template struct UE::Math::TTransform<double>;
+
 #endif // ENABLE_VECTORIZED_TRANSFORM
diff --git a/Engine/Source/Runtime/Core/Private/Math/UnrealMath.cpp b/Engine/Source/Runtime/Core/Private/Math/UnrealMath.cpp
--- a/Engine/Source/Runtime/Core/Private/Math/UnrealMath.cpp
+++ b/Engine/Source/Runtime/Core/Private/Math/UnrealMath.cpp
@@ -19,10 +19,6 @@ CORE_API const FVector2D FVector2D::ZeroVector(0.0f, 0.0f);
 CORE_API const FVector2D FVector2D::UnitVector(1.0f, 1.0f);
 CORE_API const FVector2D FVector2D::Unit45Deg(UE_INV_SQRT_2, UE_INV_SQRT_2);
 
-CORE_API const FRotator FRotator::ZeroRotator(0.f, 0.f, 0.f);
-
-CORE_API const VectorRegister VECTOR_INV_255 = DECLARE_VECTOR_REGISTER(1.f/255.f, 1.f/255.f, 1.f/255.f, 1.f/255.f);
-
 CORE_API const uint32 FMath::BitFlag[32] =
 {
 	(1U << 0),	(1U << 1),	(1U << 2),	(1U << 3),
@@ -260,7 +256,7 @@ FRotator FVector4::Rotation() const
 }
 
 template<>
-FQuat FVector3f::ToOrientationQuat() const
+FQuat4f FVector3f::ToOrientationQuat() const
 {
 	// Essentially an optimized Vector->Rotator->Quat made possible by knowing Roll == 0, and avoiding radians->degrees->radians.
 	// This is done to avoid adding any roll (which our API states as a constraint).
@@ -274,7 +270,7 @@ FQuat FVector3f::ToOrientationQuat() const
 	FMath::SinCos(&SP, &CP, PitchRad * DIVIDE_BY_2);
 	FMath::SinCos(&SY, &CY, YawRad * DIVIDE_BY_2);
 
-	FQuat RotationQuat;
+	FQuat4f RotationQuat;
 	RotationQuat.X =  SP*SY;
 	RotationQuat.Y = -SP*CY;
 	RotationQuat.Z =  CP*SY;
@@ -283,7 +279,7 @@ FQuat FVector3f::ToOrientationQuat() const
 }
 
 template<>
-FQuat FVector3d::ToOrientationQuat() const
+FQuat4d FVector3d::ToOrientationQuat() const
 {
 	// Essentially an optimized Vector->Rotator->Quat made possible by knowing Roll == 0, and avoiding radians->degrees->radians.
 	// This is done to avoid adding any roll (which our API states as a constraint).
@@ -297,7 +293,7 @@ FQuat FVector3d::ToOrientationQuat() const
 	FMath::SinCos(&SP, &CP, PitchRad * DIVIDE_BY_2);
 	FMath::SinCos(&SY, &CY, YawRad * DIVIDE_BY_2);
 
-	FQuat RotationQuat;
+	FQuat4d RotationQuat;
 	RotationQuat.X = static_cast<float>(SP * SY);
 	RotationQuat.Y = static_cast<float>(-SP * CY);
 	RotationQuat.Z = static_cast<float>(CP * SY);
@@ -358,7 +354,13 @@ FVector FMath::ClosestPointOnInfiniteLine(const FVector& LineStart, const FVecto
 	return ClosestPoint;
 }
 
-FRotator::FRotator(const FQuat& Quat)
+FRotator::FRotator(const FQuat4f& Quat)
+{
+	*this = Quat.Rotator();
+	DiagnosticCheckNaN();
+}
+
+FRotator::FRotator(const FQuat4d& Quat)
 {
 	*this = Quat.Rotator();
 	DiagnosticCheckNaN();
@@ -437,9 +439,8 @@ FQuat FRotator::Quaternion() const
 	const VectorRegister4Float LeftTerm  = VectorBitwiseXor(SignBitsLeft , VectorMultiply(CR, VectorMultiply(SP_SP_CP_CP, SY_CY_SY_CY)));
 	const VectorRegister4Float RightTerm = VectorBitwiseXor(SignBitsRight, VectorMultiply(SR, VectorMultiply(CP_CP_SP_SP, CY_SY_CY_SY)));
 
-	FQuat RotationQuat;
 	const VectorRegister4Float Result = VectorAdd(LeftTerm, RightTerm);	
-	VectorStoreAligned(Result, &RotationQuat);
+	FQuat RotationQuat = FQuat::MakeFromVectorRegister(Result);
 #else
 	const float DEG_TO_RAD = PI/(180.f);
 	const float RADS_DIVIDED_BY_2 = DEG_TO_RAD/2.f;
@@ -557,268 +558,116 @@ FRotator FMatrix44d::Rotator() const
 }
 
 template<>
-FQuat FMatrix44f::ToQuat() const
+FQuat4f FMatrix44f::ToQuat() const
 {
-	return FQuat(*this);
+	return FQuat4f(*this);
 }
 
 template<>
-FQuat FMatrix44d::ToQuat() const
+FQuat4d FMatrix44d::ToQuat() const
 {
-	return FQuat(*this);
+	return FQuat4d(*this);
 }
 
-CORE_API const FQuat FQuat::Identity(0,0,0,1);
-
-
-//////////////////////////////////////////////////////////////////////////
-// FQuat
-
-FRotator FQuat::Rotator() const
-{
-	DiagnosticCheckNaN();
-	const float SingularityTest = Z*X-W*Y;
-	const float YawY = 2.f*(W*Z+X*Y);
-	const float YawX = (1.f-2.f*(FMath::Square(Y) + FMath::Square(Z)));
-
-	// reference 
-	// http://en.wikipedia.org/wiki/Conversion_between_quaternions_and_Euler_angles
-	// http://www.euclideanspace.com/maths/geometry/rotations/conversions/quaternionToEuler/
-
-	// this value was found from experience, the above websites recommend different values
-	// but that isn't the case for us, so I went through different testing, and finally found the case 
-	// where both of world lives happily. 
-	const float SINGULARITY_THRESHOLD = 0.4999995f;
-	const float RAD_TO_DEG = (180.f)/PI;
-	FRotator RotatorFromQuat;
-
-	if (SingularityTest < -SINGULARITY_THRESHOLD)
-	{
-		RotatorFromQuat.Pitch = -90.f;
-		RotatorFromQuat.Yaw = FMath::Atan2(YawY, YawX) * RAD_TO_DEG;
-		RotatorFromQuat.Roll = FRotator::NormalizeAxis(-RotatorFromQuat.Yaw - (2.f * FMath::Atan2(X, W) * RAD_TO_DEG));
-	}
-	else if (SingularityTest > SINGULARITY_THRESHOLD)
-	{
-		RotatorFromQuat.Pitch = 90.f;
-		RotatorFromQuat.Yaw = FMath::Atan2(YawY, YawX) * RAD_TO_DEG;
-		RotatorFromQuat.Roll = FRotator::NormalizeAxis(RotatorFromQuat.Yaw - (2.f * FMath::Atan2(X, W) * RAD_TO_DEG));
-	}
-	else
-	{
-		RotatorFromQuat.Pitch = FMath::FastAsin(2.f*(SingularityTest)) * RAD_TO_DEG;
-		RotatorFromQuat.Yaw = FMath::Atan2(YawY, YawX) * RAD_TO_DEG;
-		RotatorFromQuat.Roll = FMath::Atan2(-2.f*(W*X+Y*Z), (1.f-2.f*(FMath::Square(X) + FMath::Square(Y)))) * RAD_TO_DEG;
-	}
-
-#if ENABLE_NAN_DIAGNOSTIC
-	if (RotatorFromQuat.ContainsNaN())
-	{
-		logOrEnsureNanError(TEXT("FQuat::Rotator(): Rotator result %s contains NaN! Quat = %s, YawY = %.9f, YawX = %.9f"), *RotatorFromQuat.ToString(), *this->ToString(), YawY, YawX);
-		RotatorFromQuat = FRotator::ZeroRotator;
-	}
-#endif
-
-	return RotatorFromQuat;
-}
 
-FQuat FQuat::MakeFromEuler(const FVector& Euler)
+namespace UE
 {
-	return FRotator::MakeFromEuler(Euler).Quaternion();
-}
-
-void FQuat::ToSwingTwist(const FVector& InTwistAxis, FQuat& OutSwing, FQuat& OutTwist) const
+namespace Math
 {
-	// Vector part projected onto twist axis
-	FVector Projection = FVector::DotProduct(InTwistAxis, FVector(X, Y, Z)) * InTwistAxis;
-
-	using FQuatReal = decltype(FQuat::X);
+	//////////////////////////////////////////////////////////////////////////
+	// FQuat
 
-	// Twist quaternion
-	OutTwist = FQuat((FQuatReal)Projection.X, (FQuatReal)Projection.Y, (FQuatReal)Projection.Z, W);
-
-	// Singularity close to 180deg
-	if(OutTwist.SizeSquared() == 0.0f)
+	template<typename T>
+	FRotator TQuat<T>::Rotator() const
 	{
-		OutTwist = FQuat::Identity;
-	}
-	else
-	{
-		OutTwist.Normalize();
-	}
+		DiagnosticCheckNaN();
+		const T SingularityTest = Z * X - W * Y;
+		const T YawY = 2.f * (W * Z + X * Y);
+		const T YawX = (1.f - 2.f * (FMath::Square(Y) + FMath::Square(Z)));
 
-	// Set swing
-	OutSwing = *this * OutTwist.Inverse();
-}
+		// reference 
+		// http://en.wikipedia.org/wiki/Conversion_between_quaternions_and_Euler_angles
+		// http://www.euclideanspace.com/maths/geometry/rotations/conversions/quaternionToEuler/
 
-float FQuat::GetTwistAngle(const FVector& TwistAxis) const
-{
-	float XYZ = (float)FVector::DotProduct(TwistAxis, FVector(X, Y, Z));
-	return FMath::UnwindRadians(2.0f * (float)FMath::Atan2(XYZ, W));
-}
+		// this value was found from experience, the above websites recommend different values
+		// but that isn't the case for us, so I went through different testing, and finally found the case 
+		// where both of world lives happily. 
+		const T SINGULARITY_THRESHOLD = 0.4999995f;
+		const T RAD_TO_DEG = (T)(180.f / PI);
+		float Pitch, Yaw, Roll;
 
-FVector FQuat::Euler() const
-{
-	return Rotator().Euler();
-}
-
-bool FQuat::NetSerialize(FArchive& Ar, class UPackageMap*, bool& bOutSuccess)
-{
-	FQuat Q;
-
-	if (Ar.IsSaving())
-	{
-		Q = *this;
-
-		// Make sure we have a non null SquareSum. It shouldn't happen with a quaternion, but better be safe.
-		if(Q.SizeSquared() <= SMALL_NUMBER)
+		if (SingularityTest < -SINGULARITY_THRESHOLD)
+		{
+			Pitch = -90.f;
+			Yaw = float(FMath::Atan2(YawY, YawX) * RAD_TO_DEG);
+			Roll = FRotator::NormalizeAxis(-Yaw - float(2.f * FMath::Atan2(X, W) * RAD_TO_DEG));
+		}
+		else if (SingularityTest > SINGULARITY_THRESHOLD)
 		{
-			Q = FQuat::Identity;
+			Pitch = 90.f;
+			Yaw = float(FMath::Atan2(YawY, YawX) * RAD_TO_DEG);
+			Roll = FRotator::NormalizeAxis(Yaw - float(2.f * FMath::Atan2(X, W) * RAD_TO_DEG));
 		}
 		else
 		{
-			// All transmitted quaternions *MUST BE* unit quaternions, in which case we can deduce the value of W.
-			if (!ensure(Q.IsNormalized()))
-			{
-				Q.Normalize();
-			}
-			// force W component to be non-negative
-			if (Q.W < 0.f)
-			{
-				Q.X *= -1.f;
-				Q.Y *= -1.f;
-				Q.Z *= -1.f;
-				Q.W *= -1.f;
-			}
+			Pitch = float(FMath::FastAsin(2.f * SingularityTest) * RAD_TO_DEG);
+			Yaw = float(FMath::Atan2(YawY, YawX) * RAD_TO_DEG);
+			Roll = float(FMath::Atan2(-2.f * (W*X + Y*Z), (1.f - 2.f * (FMath::Square(X) + FMath::Square(Y)))) * RAD_TO_DEG);
 		}
-	}
 
-	Ar << Q.X << Q.Y << Q.Z;
+		FRotator RotatorFromQuat = FRotator(Pitch, Yaw, Roll);
 
-	if ( Ar.IsLoading() )
-	{
-		const float XYZMagSquared = (Q.X*Q.X + Q.Y*Q.Y + Q.Z*Q.Z);
-		const float WSquared = 1.0f - XYZMagSquared;
-		// If mag of (X,Y,Z) <= 1.0, then we calculate W to make magnitude of Q 1.0
-		if (WSquared >= 0.f)
-		{
-			Q.W = FMath::Sqrt(WSquared);
-		}
-		// If mag of (X,Y,Z) > 1.0, we set W to zero, and then renormalize 
-		else
+#if ENABLE_NAN_DIAGNOSTIC
+		if (RotatorFromQuat.ContainsNaN())
 		{
-			Q.W = 0.f;
-
-			const float XYZInvMag = FMath::InvSqrt(XYZMagSquared);
-			Q.X *= XYZInvMag;
-			Q.Y *= XYZInvMag;
-			Q.Z *= XYZInvMag;
+			logOrEnsureNanError(TEXT("TQuat<T>::Rotator(): Rotator result %s contains NaN! Quat = %s, YawY = %.9f, YawX = %.9f"), *RotatorFromQuat.ToString(), *this->ToString(), YawY, YawX);
+			RotatorFromQuat = FRotator::ZeroRotator;
 		}
+#endif
 
-		*this = Q;
+		return RotatorFromQuat;
 	}
 
-	bOutSuccess = true;
-	return true;
-}
-
-
-//
-// Based on:
-// http://lolengine.net/blog/2014/02/24/quaternion-from-two-vectors-final
-// http://www.euclideanspace.com/maths/algebra/vectors/angleBetween/index.htm
-//
-FORCEINLINE_DEBUGGABLE FQuat FindBetween_Helper(const FVector& A, const FVector& B, FVector::FReal NormAB)
-{
-	using FReal = decltype(FQuat::X);
-	FReal W = FReal(NormAB + FVector::DotProduct(A, B));	// LWC_TODO: Precision loss
-	FQuat Result;
-
-	if (W >= 1e-6f * NormAB)
+	template<typename T>
+	TQuat<T> TQuat<T>::MakeFromEuler(const TVector<T>& Euler)
 	{
-		//Axis = FVector::CrossProduct(A, B);
-		Result = FQuat(FReal(A.Y * B.Z - A.Z * B.Y),
-					   FReal(A.Z * B.X - A.X * B.Z),
-					   FReal(A.X * B.Y - A.Y * B.X),
-					   W);
+		return TQuat<T>(FRotator::MakeFromEuler(Euler));
 	}
-	else
-	{
-		// A and B point in opposite directions
-		W = 0.f;
-		Result = FMath::Abs(A.X) > FMath::Abs(A.Y)
-				? FQuat(FReal(-A.Z), 0.f, FReal(A.X), W)
-				: FQuat(0.f, FReal(-A.Z), FReal(A.Y), W);
-	}
-
-	Result.Normalize();
-	return Result;
-}
-
-FQuat FQuat::FindBetweenNormals(const FVector& A, const FVector& B)
-{
-	const FVector::FReal NormAB = 1.f;
-	return FindBetween_Helper(A, B, NormAB);
-}
 
-FQuat FQuat::FindBetweenVectors(const FVector& A, const FVector& B)
-{
-	const FVector::FReal NormAB = FMath::Sqrt(A.SizeSquared() * B.SizeSquared());
-	return FindBetween_Helper(A, B, NormAB);
-}
-
-FQuat FQuat::Log() const
-{
-	FQuat Result;
-	Result.W = 0.f;
-
-	if ( FMath::Abs(W) < 1.f )
+	template<typename T>
+	void TQuat<T>::ToSwingTwist(const TVector<T>& InTwistAxis, TQuat<T>& OutSwing, TQuat<T>& OutTwist) const
 	{
-		const float Angle = FMath::Acos(W);
-		const float SinAngle = FMath::Sin(Angle);
+		// Vector part projected onto twist axis
+		TVector<T> Projection = TVector<T>::DotProduct(InTwistAxis, TVector<T>(X, Y, Z)) * InTwistAxis;
 
-		if ( FMath::Abs(SinAngle) >= SMALL_NUMBER )
-		{
-			const float Scale = Angle/SinAngle;
-			Result.X = Scale*X;
-			Result.Y = Scale*Y;
-			Result.Z = Scale*Z;
+		// Twist quaternion
+		OutTwist = TQuat<T>((FReal)Projection.X, (FReal)Projection.Y, (FReal)Projection.Z, W);
 
-			return Result;
+		// Singularity close to 180deg
+		if (OutTwist.SizeSquared() == 0.0f)
+		{
+			OutTwist = TQuat<T>::Identity;
+		}
+		else
+		{
+			OutTwist.Normalize();
 		}
-	}
 
-	Result.X = X;
-	Result.Y = Y;
-	Result.Z = Z;
+		// Set swing
+		OutSwing = *this * OutTwist.Inverse();
+	}
 
-	return Result;
+	template<typename T>
+	T TQuat<T>::GetTwistAngle(const TVector<T>& TwistAxis) const
+	{
+		T XYZ = (T)TVector<T>::DotProduct(TwistAxis, TVector<T>(X, Y, Z));
+		return FMath::UnwindRadians((T)2.0f * FMath::Atan2(XYZ, W));
+	}
+}
 }
 
-FQuat FQuat::Exp() const
-{
-	const float Angle = FMath::Sqrt(X*X + Y*Y + Z*Z);
-	const float SinAngle = FMath::Sin(Angle);
-
-	FQuat Result;
-	Result.W = FMath::Cos(Angle);
 
-	if ( FMath::Abs(SinAngle) >= SMALL_NUMBER )
-	{
-		const float Scale = SinAngle/Angle;
-		Result.X = Scale*X;
-		Result.Y = Scale*Y;
-		Result.Z = Scale*Z;
-	}
-	else
-	{
-		Result.X = X;
-		Result.Y = Y;
-		Result.Z = Z;
-	}
 
-	return Result;
-}
 
 /*-----------------------------------------------------------------------------
 	Swept-Box vs Box test.
@@ -1018,99 +867,280 @@ float FLinearColor::EvaluateBezier(const FLinearColor* ControlPoints, int32 NumP
 	return Length;
 }
 
-
-
-FQuat FQuat::Slerp_NotNormalized(const FQuat& Quat1,const FQuat& Quat2, float Slerp)
+namespace UE
 {
-	// Get cosine of angle between quats.
-	const float RawCosom = 
-		    Quat1.X * Quat2.X +
+namespace Math
+{
+	template<typename T>
+	TQuat<T> TQuat<T>::Slerp_NotNormalized(const TQuat<T>& Quat1, const TQuat<T>& Quat2, T Slerp)
+	{
+		// Get cosine of angle between quats.
+		const T RawCosom =
+			Quat1.X * Quat2.X +
 			Quat1.Y * Quat2.Y +
 			Quat1.Z * Quat2.Z +
 			Quat1.W * Quat2.W;
-	// Unaligned quats - compensate, results in taking shorter route.
-	const float Cosom = FMath::FloatSelect( RawCosom, RawCosom, -RawCosom );
-	
-	float Scale0, Scale1;
+		// Unaligned quats - compensate, results in taking shorter route.
+		const T Cosom = FMath::FloatSelect(RawCosom, RawCosom, -RawCosom);
 
-	if( Cosom < 0.9999f )
-	{	
-		const float Omega = FMath::Acos(Cosom);
-		const float InvSin = 1.f/FMath::Sin(Omega);
-		Scale0 = FMath::Sin( (1.f - Slerp) * Omega ) * InvSin;
-		Scale1 = FMath::Sin( Slerp * Omega ) * InvSin;
+		T Scale0, Scale1;
+
+		if (Cosom < T(0.9999f))
+		{
+			const T Omega = FMath::Acos(Cosom);
+			const T InvSin = T(1.f) / FMath::Sin(Omega);
+			Scale0 = FMath::Sin((T(1.f) - Slerp) * Omega) * InvSin;
+			Scale1 = FMath::Sin(Slerp * Omega) * InvSin;
+		}
+		else
+		{
+			// Use linear interpolation.
+			Scale0 = T(1.0f) - Slerp;
+			Scale1 = Slerp;
+		}
+
+		// In keeping with our flipped Cosom:
+		Scale1 = FMath::FloatSelect(RawCosom, Scale1, -Scale1);
+
+		TQuat<T> Result;
+
+		Result.X = Scale0 * Quat1.X + Scale1 * Quat2.X;
+		Result.Y = Scale0 * Quat1.Y + Scale1 * Quat2.Y;
+		Result.Z = Scale0 * Quat1.Z + Scale1 * Quat2.Z;
+		Result.W = Scale0 * Quat1.W + Scale1 * Quat2.W;
+
+		return Result;
 	}
-	else
+
+	template<typename T>
+	TQuat<T> TQuat<T>::SlerpFullPath_NotNormalized(const TQuat<T>&quat1, const TQuat<T>&quat2, T Alpha )
 	{
-		// Use linear interpolation.
-		Scale0 = 1.0f - Slerp;
-		Scale1 = Slerp;	
+		const T CosAngle = FMath::Clamp(quat1 | quat2, T(-1.f), T(1.f));
+		const T Angle = FMath::Acos(CosAngle);
+
+		//UE_LOG(LogUnrealMath, Log,  TEXT("CosAngle: %f Angle: %f"), CosAngle, Angle );
+
+		if ( FMath::Abs(Angle) < T(KINDA_SMALL_NUMBER))
+		{
+			return quat1;
+		}
+
+		const T SinAngle = FMath::Sin(Angle);
+		const T InvSinAngle = T(1.f)/SinAngle;
+
+		const T Scale0 = FMath::Sin((1.0f-Alpha)*Angle)*InvSinAngle;
+		const T Scale1 = FMath::Sin(Alpha*Angle)*InvSinAngle;
+
+		return quat1*Scale0 + quat2*Scale1;
 	}
 
-	// In keeping with our flipped Cosom:
-	Scale1 = FMath::FloatSelect( RawCosom, Scale1, -Scale1 );
+	template<typename T>
+	TQuat<T> TQuat<T>::Squad(const TQuat<T>& quat1, const TQuat<T>& tang1, const TQuat<T>& quat2, const TQuat<T>& tang2, T Alpha)
+	{
+		// Always slerp along the short path from quat1 to quat2 to prevent axis flipping.
+		// This approach is taken by OGRE engine, amongst others.
+		const TQuat<T> Q1 = TQuat::Slerp_NotNormalized(quat1, quat2, Alpha);
+		const TQuat<T> Q2 = TQuat::SlerpFullPath_NotNormalized(tang1, tang2, Alpha);
+		const TQuat<T> Result = TQuat::SlerpFullPath(Q1, Q2, 2.f * Alpha * (1.f - Alpha));
 
-	FQuat Result;
-		
-	Result.X = Scale0 * Quat1.X + Scale1 * Quat2.X;
-	Result.Y = Scale0 * Quat1.Y + Scale1 * Quat2.Y;
-	Result.Z = Scale0 * Quat1.Z + Scale1 * Quat2.Z;
-	Result.W = Scale0 * Quat1.W + Scale1 * Quat2.W;
+		return Result;
+	}
 
-	return Result;
-}
+	template<typename T>
+	TQuat<T> TQuat<T>::SquadFullPath(const TQuat<T>& quat1, const TQuat<T>& tang1, const TQuat<T>& quat2, const TQuat<T>& tang2, T Alpha)
+	{
+		const TQuat<T> Q1 = TQuat::SlerpFullPath_NotNormalized(quat1, quat2, Alpha);
+		const TQuat<T> Q2 = TQuat::SlerpFullPath_NotNormalized(tang1, tang2, Alpha);
+		const TQuat<T> Result = TQuat::SlerpFullPath(Q1, Q2, 2.f * Alpha * (1.f - Alpha));
 
-FQuat FQuat::SlerpFullPath_NotNormalized(const FQuat &quat1, const FQuat &quat2, float Alpha )
-{
-	const float CosAngle = FMath::Clamp(quat1 | quat2, -1.f, 1.f);
-	const float Angle = FMath::Acos(CosAngle);
+		return Result;
+	}
 
-	//UE_LOG(LogUnrealMath, Log,  TEXT("CosAngle: %f Angle: %f"), CosAngle, Angle );
+	template<typename T>
+	void TQuat<T>::CalcTangents(const TQuat<T>& PrevP, const TQuat<T>& P, const TQuat<T>& NextP, T Tension, TQuat<T>& OutTan)
+	{
+		const TQuat<T> InvP = P.Inverse();
+		const TQuat<T> Part1 = (InvP * PrevP).Log();
+		const TQuat<T> Part2 = (InvP * NextP).Log();
+
+		const TQuat<T> PreExp = (Part1 + Part2) * -0.5f;
+
+		OutTan = P * PreExp.Exp();
+	}
 
-	if ( FMath::Abs(Angle) < KINDA_SMALL_NUMBER )
+	template<typename T>
+	TVector<T> TQuat<T>::Euler() const
 	{
-		return quat1;
+		return Rotator().Euler();
 	}
 
-	const float SinAngle = FMath::Sin(Angle);
-	const float InvSinAngle = 1.f/SinAngle;
+	template<typename T>
+	bool TQuat<T>::NetSerialize(FArchive& Ar, class UPackageMap*, bool& bOutSuccess)
+	{
+		TQuat<T> Q;
 
-	const float Scale0 = FMath::Sin((1.0f-Alpha)*Angle)*InvSinAngle;
-	const float Scale1 = FMath::Sin(Alpha*Angle)*InvSinAngle;
+		if (Ar.IsSaving())
+		{
+			Q = *this;
 
-	return quat1*Scale0 + quat2*Scale1;
-}
+			// Make sure we have a non null SquareSum. It shouldn't happen with a quaternion, but better be safe.
+			if (Q.SizeSquared() <= (T)SMALL_NUMBER)
+			{
+				Q = TQuat<T>::Identity;
+			}
+			else
+			{
+				// All transmitted quaternions *MUST BE* unit quaternions, in which case we can deduce the value of W.
+				if (!ensure(Q.IsNormalized()))
+				{
+					Q.Normalize();
+				}
+				// force W component to be non-negative
+				if (Q.W < 0.f)
+				{
+					Q.X *= -1.f;
+					Q.Y *= -1.f;
+					Q.Z *= -1.f;
+					Q.W *= -1.f;
+				}
+			}
+		}
 
-FQuat FQuat::Squad(const FQuat& quat1, const FQuat& tang1, const FQuat& quat2, const FQuat& tang2, float Alpha)
-{
-	// Always slerp along the short path from quat1 to quat2 to prevent axis flipping.
-	// This approach is taken by OGRE engine, amongst others.
-	const FQuat Q1 = FQuat::Slerp_NotNormalized(quat1, quat2, Alpha);
-	const FQuat Q2 = FQuat::SlerpFullPath_NotNormalized(tang1, tang2, Alpha);
-	const FQuat Result = FQuat::SlerpFullPath(Q1, Q2, 2.f * Alpha * (1.f - Alpha));
+		Ar << Q.X << Q.Y << Q.Z;
 
-	return Result;
-}
+		if (Ar.IsLoading())
+		{
+			const T XYZMagSquared = (Q.X * Q.X + Q.Y * Q.Y + Q.Z * Q.Z);
+			const T WSquared = (T)1.0f - XYZMagSquared;
+			// If mag of (X,Y,Z) <= 1.0, then we calculate W to make magnitude of Q 1.0
+			if (WSquared >= 0.f)
+			{
+				Q.W = FMath::Sqrt(WSquared);
+			}
+			// If mag of (X,Y,Z) > 1.0, we set W to zero, and then renormalize 
+			else
+			{
+				Q.W = 0.f;
 
-FQuat FQuat::SquadFullPath(const FQuat& quat1, const FQuat& tang1, const FQuat& quat2, const FQuat& tang2, float Alpha)
-{
-	const FQuat Q1 = FQuat::SlerpFullPath_NotNormalized(quat1, quat2, Alpha);
-	const FQuat Q2 = FQuat::SlerpFullPath_NotNormalized(tang1, tang2, Alpha);
-	const FQuat Result = FQuat::SlerpFullPath(Q1, Q2, 2.f * Alpha * (1.f - Alpha));
+				const T XYZInvMag = FMath::InvSqrt(XYZMagSquared);
+				Q.X *= XYZInvMag;
+				Q.Y *= XYZInvMag;
+				Q.Z *= XYZInvMag;
+			}
 
-	return Result;
-}
+			*this = Q;
+		}
 
-void FQuat::CalcTangents(const FQuat& PrevP, const FQuat& P, const FQuat& NextP, float Tension, FQuat& OutTan)
-{
-	const FQuat InvP = P.Inverse();
-	const FQuat Part1 = (InvP * PrevP).Log();
-	const FQuat Part2 = (InvP * NextP).Log();
+		bOutSuccess = true;
+		return true;
+	}
 
-	const FQuat PreExp = (Part1 + Part2) * -0.5f;
 
-	OutTan = P * PreExp.Exp();
-}
+	//
+	// Based on:
+	// http://lolengine.net/blog/2014/02/24/quaternion-from-two-vectors-final
+	// http://www.euclideanspace.com/maths/algebra/vectors/angleBetween/index.htm
+	//
+	template<typename T>
+	FORCEINLINE_DEBUGGABLE TQuat<T> FindBetween_Helper(const TVector<T>& A, const TVector<T>& B, T NormAB)
+	{
+		T W = NormAB + TVector<T>::DotProduct(A, B);
+		TQuat<T> Result;
+
+		if (W >= 1e-6f * NormAB)
+		{
+			//Axis = FVector::CrossProduct(A, B);
+			Result = TQuat<T>(
+				A.Y * B.Z - A.Z * B.Y,
+				A.Z * B.X - A.X * B.Z,
+				A.X * B.Y - A.Y * B.X,
+				W);
+		}
+		else
+		{
+			// A and B point in opposite directions
+			W = 0.f;
+			Result = FMath::Abs(A.X) > FMath::Abs(A.Y)
+				? TQuat<T>(-A.Z, 0.f, A.X, W)
+				: TQuat<T>(0.f, -A.Z, A.Y, W);
+		}
+
+		Result.Normalize();
+		return Result;
+	}
+
+	template<typename T>
+	TQuat<T> TQuat<T>::FindBetweenNormals(const TVector<T>& A, const TVector<T>& B)
+	{
+		const T NormAB = 1.f;
+		return UE::Math::FindBetween_Helper(A, B, NormAB);
+	}
+
+	template<typename T>
+	TQuat<T> TQuat<T>::FindBetweenVectors(const TVector<T>& A, const TVector<T>& B)
+	{
+		const T NormAB = FMath::Sqrt(A.SizeSquared() * B.SizeSquared());
+		return UE::Math::FindBetween_Helper(A, B, NormAB);
+	}
+
+	template<typename T>
+	TQuat<T> TQuat<T>::Log() const
+	{
+		TQuat<T> Result;
+		Result.W = 0.f;
+
+		if (FMath::Abs(W) < 1.f)
+		{
+			const T Angle = FMath::Acos(W);
+			const T SinAngle = FMath::Sin(Angle);
+
+			if (FMath::Abs(SinAngle) >= T(SMALL_NUMBER))
+			{
+				const T Scale = Angle / SinAngle;
+				Result.X = Scale * X;
+				Result.Y = Scale * Y;
+				Result.Z = Scale * Z;
+
+				return Result;
+			}
+		}
+
+		Result.X = X;
+		Result.Y = Y;
+		Result.Z = Z;
+
+		return Result;
+	}
+
+	template<typename T>
+	TQuat<T> TQuat<T>::Exp() const
+	{
+		const T Angle = FMath::Sqrt(X * X + Y * Y + Z * Z);
+		const T SinAngle = FMath::Sin(Angle);
+
+		TQuat<T>  Result;
+		Result.W = FMath::Cos(Angle);
+
+		if (FMath::Abs(SinAngle) >= T(SMALL_NUMBER))
+		{
+			const T Scale = SinAngle / Angle;
+			Result.X = Scale * X;
+			Result.Y = Scale * Y;
+			Result.Z = Scale * Z;
+		}
+		else
+		{
+			Result.X = X;
+			Result.Y = Y;
+			Result.Z = Z;
+		}
+
+		return Result;
+	}
+
+
+} // namespace UE::Math
+} // namespace UE
 
 template<typename T>
 static void FindBounds( T& OutMin, T& OutMax,  T Start, T StartLeaveTan, float StartT, T End, T EndArriveTan, float EndT, bool bCurve )
@@ -2100,7 +2130,7 @@ CORE_API FVector FMath::VInterpNormalRotationTo(const FVector& Current, const FV
 
 	// Decompose into an axis and angle for rotation
 	FVector DeltaAxis(0.f);
-	float DeltaAngle = 0.f;
+	FQuat::FReal DeltaAngle = 0.f;
 	DeltaQuat.ToAxisAndAngle(DeltaAxis, DeltaAngle);
 
 	// Find rotation step for this frame
@@ -2108,7 +2138,7 @@ CORE_API FVector FMath::VInterpNormalRotationTo(const FVector& Current, const FV
 
 	if( FMath::Abs(DeltaAngle) > RotationStepRadians )
 	{
-		DeltaAngle = FMath::Clamp(DeltaAngle, -RotationStepRadians, RotationStepRadians);
+		DeltaAngle = FMath::Clamp<FQuat::FReal>(DeltaAngle, -RotationStepRadians, RotationStepRadians);
 		DeltaQuat = FQuat(DeltaAxis, DeltaAngle);
 		return DeltaQuat.RotateVector(Current);
 	}
@@ -2336,7 +2366,8 @@ CORE_API FLinearColor FMath::CInterpTo(const FLinearColor& Current, const FLinea
 	return Current + DeltaMove;
 }
 
-CORE_API FQuat FMath::QInterpConstantTo(const FQuat& Current, const FQuat& Target, float DeltaTime, float InterpSpeed)
+template< class T >
+CORE_API UE::Math::TQuat<T> FMath::QInterpConstantTo(const UE::Math::TQuat<T>& Current, const UE::Math::TQuat<T>& Target, float DeltaTime, float InterpSpeed)
 {
 	// If no interp speed, jump to target value
 	if (InterpSpeed <= 0.f)
@@ -2350,14 +2381,20 @@ CORE_API FQuat FMath::QInterpConstantTo(const FQuat& Current, const FQuat& Targe
 		return Target;
 	}
 
-	float DeltaInterpSpeed = FMath::Clamp(DeltaTime * InterpSpeed, 0.f, 1.f);
-	float AngularDistance = FMath::Max(SMALL_NUMBER, Target.AngularDistance(Current));
-	float Alpha = FMath::Clamp(DeltaInterpSpeed / AngularDistance, 0.f, 1.f);
+	float DeltaInterpSpeed = FMath::Clamp<float>(DeltaTime * InterpSpeed, 0.f, 1.f);
+	float AngularDistance = FMath::Max<float>(SMALL_NUMBER, (float)(Target.AngularDistance(Current)));
+	float Alpha = FMath::Clamp<float>(DeltaInterpSpeed / AngularDistance, 0.f, 1.f);
 
-	return FQuat::Slerp(Current, Target, Alpha);
+	return UE::Math::TQuat<T>::Slerp(Current, Target, Alpha);
 }
 
-CORE_API FQuat FMath::QInterpTo(const FQuat& Current, const FQuat& Target, float DeltaTime, float InterpSpeed)
+// Instantiate for linker
+template CORE_API UE::Math::TQuat<float> FMath::QInterpConstantTo<float>(const UE::Math::TQuat<float>& Current, const UE::Math::TQuat<float>& Target, float DeltaTime, float InterpSpeed);
+template CORE_API UE::Math::TQuat<double> FMath::QInterpConstantTo<double>(const UE::Math::TQuat<double>& Current, const UE::Math::TQuat<double>& Target, float DeltaTime, float InterpSpeed);
+
+
+template< class T >
+CORE_API UE::Math::TQuat<T> FMath::QInterpTo(const UE::Math::TQuat<T>& Current, const UE::Math::TQuat<T>& Target, float DeltaTime, float InterpSpeed)
 {
 	// If no interp speed, jump to target value
 	if (InterpSpeed <= 0.f)
@@ -2371,9 +2408,14 @@ CORE_API FQuat FMath::QInterpTo(const FQuat& Current, const FQuat& Target, float
 		return Target;
 	}
 
-	return FQuat::Slerp(Current, Target, FMath::Clamp(InterpSpeed * DeltaTime, 0.f, 1.f));
+	return UE::Math::TQuat<T>::Slerp(Current, Target, FMath::Clamp<float>(InterpSpeed * DeltaTime, 0.f, 1.f));
 }
 
+// Instantiate for linker
+template CORE_API UE::Math::TQuat<float> FMath::QInterpTo<float>(const UE::Math::TQuat<float>& Current, const UE::Math::TQuat<float>& Target, float DeltaTime, float InterpSpeed);
+template CORE_API UE::Math::TQuat<double> FMath::QInterpTo<double>(const UE::Math::TQuat<double>& Current, const UE::Math::TQuat<double>& Target, float DeltaTime, float InterpSpeed);
+
+
 CORE_API float ClampFloatTangent( float PrevPointVal, float PrevTime, float CurPointVal, float CurTime, float NextPointVal, float NextTime )
 {
 	const float PrevToNextTimeDiff = FMath::Max< float >( KINDA_SMALL_NUMBER, NextTime - PrevTime );
@@ -3342,3 +3384,8 @@ float FMath::PerlinNoise3D(const FVector& Location)
 		),
 		-1.0f, 1.0f);
 }
+
+
+// Instantiate for linker
+template struct UE::Math::TQuat<float>;
+template struct UE::Math::TQuat<double>;
diff --git a/Engine/Source/Runtime/Core/Private/Tests/Math/UnrealMathTest.cpp b/Engine/Source/Runtime/Core/Private/Tests/Math/UnrealMathTest.cpp
--- a/Engine/Source/Runtime/Core/Private/Tests/Math/UnrealMathTest.cpp
+++ b/Engine/Source/Runtime/Core/Private/Tests/Math/UnrealMathTest.cpp
@@ -25,8 +25,8 @@
 DEFINE_LOG_CATEGORY_STATIC(LogUnrealMathTest, Log, All);
 
 // Create some temporary storage variables
-alignas(16) static float GScratch[16];
-alignas(16) static double GScratchDouble[16];
+alignas(alignof(VectorRegister4Float)) static float GScratch[16];
+alignas(alignof(VectorRegister4Double)) static double GScratchDouble[16];
 static float GSum;
 static double GSumDouble;
 static bool GPassing;
@@ -219,7 +219,7 @@ bool TestFVector3Equal( const FVector3f& Vec0, const FVector3f& Vec1, float Tole
 	return GSum <= Tolerance;
 }
 
-bool TestQuatsEqual(const FQuat& Q0, const FQuat& Q1, float Tolerance)
+bool TestQuatsEqual(const FQuat4f& Q0, const FQuat4f& Q1, float Tolerance)
 {
 	GScratch[0] = Q0.X;
 	GScratch[1] = Q0.Y;
@@ -236,6 +236,23 @@ bool TestQuatsEqual(const FQuat& Q0, const FQuat& Q1, float Tolerance)
 	return bEqual;
 }
 
+bool TestQuatsEqual(const FQuat4d& Q0, const FQuat4d& Q1, double Tolerance)
+{
+	GScratchDouble[0] = Q0.X;
+	GScratchDouble[1] = Q0.Y;
+	GScratchDouble[2] = Q0.Z;
+	GScratchDouble[3] = Q0.W;
+	GScratchDouble[4] = Q1.X;
+	GScratchDouble[5] = Q1.Y;
+	GScratchDouble[6] = Q1.Z;
+	GScratchDouble[7] = Q1.W;
+	GSumDouble = 0.f;
+
+	const bool bEqual = Q0.Equals(Q1, Tolerance);
+	CheckPassing(bEqual);
+	return bEqual;
+}
+
 /**
  * Tests if a vector (xyz) is normalized (length 1) within a tolerance
  *
@@ -269,7 +286,7 @@ bool TestFVector3Normalized(const FVector3f& Vec0, float Tolerance)
  *
  * @return true if normalized(ish)
  */
-bool TestQuatNormalized(const FQuat& Q0, float Tolerance)
+bool TestQuatNormalized(const FQuat4f& Q0, float Tolerance)
 {
 	GScratch[0] = Q0.X;
 	GScratch[1] = Q0.Y;
@@ -451,7 +468,7 @@ void TestVectorMatrixInverse(UE::Math::TMatrix<FReal>* DstMatrix, const UE::Math
  * @param MatrixM		FMatrix44f pointer to the Matrix to apply transform
  * @return VectorRegister4Float = VecP*MatrixM
  */
-VectorRegister4Float TestVectorTransformVector(const VectorRegister4Float&  VecP,  const void* MatrixM )
+VectorRegister4Float TestVectorTransformVector(const VectorRegister4Float&  VecP,  const FMatrix44f* MatrixM )
 {
 	typedef float Float4x4[4][4];
 	union U { 
@@ -474,7 +491,7 @@ VectorRegister4Float TestVectorTransformVector(const VectorRegister4Float&  VecP
 * @param Rotator FRotator 
 * @return Rotation as a quaternion.
 */
-MATHTEST_INLINE FQuat TestRotatorToQuaternion( const FRotator& Rotator)
+MATHTEST_INLINE FQuat4f TestRotatorToQuaternion( const FRotator& Rotator)
 {
 	const float Pitch = FMath::Fmod(Rotator.Pitch, 360.f);
 	const float Yaw = FMath::Fmod(Rotator.Yaw, 360.f);
@@ -487,7 +504,7 @@ MATHTEST_INLINE FQuat TestRotatorToQuaternion( const FRotator& Rotator)
 	const float SP = FMath::Sin(FMath::DegreesToRadians(Pitch * 0.5f));
 	const float SY = FMath::Sin(FMath::DegreesToRadians(Yaw   * 0.5f));
 
-	FQuat RotationQuat;
+	FQuat4f RotationQuat;
 	RotationQuat.W = CR*CP*CY + SR*SP*SY;
 	RotationQuat.X = CR*SP*SY - SR*CP*CY;
 	RotationQuat.Y = -CR*SP*CY - SR*CP*SY;
@@ -495,7 +512,7 @@ MATHTEST_INLINE FQuat TestRotatorToQuaternion( const FRotator& Rotator)
 	return RotationQuat;
 }
 
-MATHTEST_INLINE FVector3f TestQuaternionRotateVectorScalar(const FQuat& Quat, const FVector3f& Vector)
+MATHTEST_INLINE FVector3f TestQuaternionRotateVectorScalar(const FQuat4f& Quat, const FVector3f& Vector)
 {
 	// (q.W*q.W-qv.qv)v + 2(qv.v)qv + 2 q.W (qv x v)
 	const FVector3f qv(Quat.X, Quat.Y, Quat.Z);
@@ -507,18 +524,18 @@ MATHTEST_INLINE FVector3f TestQuaternionRotateVectorScalar(const FQuat& Quat, co
 }
 
 // Q * V * Q^-1
-MATHTEST_INLINE FVector3f TestQuaternionMultiplyVector(const FQuat& Quat, const FVector3f& Vector)
+MATHTEST_INLINE FVector3f TestQuaternionMultiplyVector(const FQuat4f& Quat, const FVector3f& Vector)
 {
-	FQuat VQ(Vector.X, Vector.Y, Vector.Z, 0.f);
-	FQuat VT, VR;
-	FQuat I = Quat.Inverse();
+	FQuat4f VQ(Vector.X, Vector.Y, Vector.Z, 0.f);
+	FQuat4f VT, VR;
+	FQuat4f I = Quat.Inverse();
 	VectorQuaternionMultiply(&VT, &Quat, &VQ);
 	VectorQuaternionMultiply(&VR, &VT, &I);
 
 	return FVector3f(VR.X, VR.Y, VR.Z);
 }
 
-MATHTEST_INLINE FVector3f TestQuaternionRotateVectorRegister(const FQuat& Quat, const FVector3f &V)
+MATHTEST_INLINE FVector3f TestQuaternionRotateVectorRegister(const FQuat4f& Quat, const FVector3f &V)
 {
 	const VectorRegister4Float Rotation = *((const VectorRegister4Float*)(&Quat));
 	const VectorRegister4Float InputVectorW0 = VectorLoadFloat3_W0(&V);
@@ -537,7 +554,7 @@ MATHTEST_INLINE FVector3f TestQuaternionRotateVectorRegister(const FQuat& Quat,
 * @param Quat1	Pointer to the first quaternion (must not be the destination)
 * @param Quat2	Pointer to the second quaternion (must not be the destination)
 */
-void TestVectorQuaternionMultiply( void *Result, const void* Quat1, const void* Quat2)
+void TestVectorQuaternionMultiply( FQuat4f* Result, const FQuat4f* Quat1, const FQuat4f* Quat2)
 {
 	typedef float Float4[4];
 	const Float4& A = *((const Float4*) Quat1);
@@ -560,7 +577,7 @@ void TestVectorQuaternionMultiply( void *Result, const void* Quat1, const void*
 /**
  * Converts a Quaternion to a Rotator.
  */
-FORCENOINLINE FRotator TestQuaternionToRotator(const FQuat& Quat)
+FORCENOINLINE FRotator TestQuaternionToRotator(const FQuat4f& Quat)
 {
 	const float X = Quat.X;
 	const float Y = Quat.Y;
@@ -604,7 +621,7 @@ FORCENOINLINE FRotator TestQuaternionToRotator(const FQuat& Quat)
 }
 
 
-FORCENOINLINE FQuat FindBetween_Old(const FVector3f& vec1, const FVector3f& vec2)
+FORCENOINLINE FQuat4f FindBetween_Old(const FVector3f& vec1, const FVector3f& vec2)
 {
 	const FVector3f cross = vec1 ^ vec2;
 	const float crossMag = cross.Size();
@@ -616,7 +633,7 @@ FORCENOINLINE FQuat FindBetween_Old(const FVector3f& vec1, const FVector3f& vec2
 		const float Dot = vec1 | vec2;
 		if (Dot > -KINDA_SMALL_NUMBER)
 		{
-			return FQuat::Identity; // no rotation
+			return FQuat4f::Identity; // no rotation
 		}
 		// Exactly opposite..
 		else
@@ -628,7 +645,7 @@ FORCENOINLINE FQuat FindBetween_Old(const FVector3f& vec1, const FVector3f& vec2
 			FVector3f AxisA, AxisB;
 			Vec.FindBestAxisVectors(AxisA, AxisB);
 
-			return FQuat(AxisA.X, AxisA.Y, AxisA.Z, 0.f); // (axis*sin(pi/2), cos(pi/2)) = (axis, 0)
+			return FQuat4f(AxisA.X, AxisA.Y, AxisA.Z, 0.f); // (axis*sin(pi/2), cos(pi/2)) = (axis, 0)
 		}
 	}
 
@@ -645,7 +662,7 @@ FORCENOINLINE FQuat FindBetween_Old(const FVector3f& vec1, const FVector3f& vec2
 	FMath::SinCos(&sinHalfAng, &cosHalfAng, 0.5f * angle);
 	const FVector3f axis = cross / crossMag;
 
-	return FQuat(
+	return FQuat4f(
 		sinHalfAng * axis.X,
 		sinHalfAng * axis.Y,
 		sinHalfAng * axis.Z,
@@ -708,7 +725,7 @@ void LogRotatorTest(const TCHAR* TestName, const FRotator& A, const FRotator& B,
 	}
 }
 
-void LogQuaternionTest(const TCHAR* TestName, const FQuat& A, const FQuat& B, bool bComparison)
+void LogQuaternionTest(const TCHAR* TestName, const FQuat4f& A, const FQuat4f& B, bool bComparison)
 {
 	if (bComparison == false)
 	{
@@ -2601,31 +2618,29 @@ bool FVectorRegisterAbstractionTest::RunTest(const FString& Parameters)
 		}
 	}
 
-
-
-	FQuat Q0, Q1, Q2, Q3;
-
 	// Sin, Cos, Tan tests
 	TestVectorTrigFunctions<float, VectorRegister4Float>();
 
 	// Exp, Log tests
 	TestVectorExpLogFunctions<float, VectorRegister4Float>();
 
+	FQuat4f Q0, Q1, Q2, Q3;
+
 	// Quat<->Rotator conversions and equality
 	{
 		// Identity conversion
 		{
 			const FRotator R0 = FRotator::ZeroRotator;
-			const FRotator R1 = FRotator(FQuat::Identity);
-			LogRotatorTest(true, TEXT("FRotator::ZeroRotator ~= FQuat::Identity : Rotator"), R0, R1, R0.Equals(R1, 0.f));
-			LogRotatorTest(true, TEXT("FRotator::ZeroRotator == FQuat::Identity : Rotator"), R0, R1, R0 == R1);
-			LogRotatorTest(true, TEXT("FRotator::ZeroRotator not != FQuat::Identity : Rotator"), R0, R1, !(R0 != R1));
-
-			Q0 = FQuat::Identity;
-			Q1 = FQuat(FRotator::ZeroRotator);
-			LogQuaternionTest(TEXT("FRotator::ZeroRotator ~= FQuat::Identity : Quaternion"), Q0, Q1, Q0.Equals(Q1, 0.f));
-			LogQuaternionTest(TEXT("FRotator::ZeroRotator == FQuat::Identity : Quaternion"), Q0, Q1, Q0 == Q1);
-			LogQuaternionTest(TEXT("FRotator::ZeroRotator not != FQuat::Identity : Quaternion"), Q0, Q1, !(Q0 != Q1));
+			const FRotator R1 = FRotator(FQuat4f::Identity);
+			LogRotatorTest(true, TEXT("FRotator::ZeroRotator ~= FQuat4f::Identity : Rotator"), R0, R1, R0.Equals(R1, 0.f));
+			LogRotatorTest(true, TEXT("FRotator::ZeroRotator == FQuat4f::Identity : Rotator"), R0, R1, R0 == R1);
+			LogRotatorTest(true, TEXT("FRotator::ZeroRotator not != FQuat4f::Identity : Rotator"), R0, R1, !(R0 != R1));
+
+			Q0 = FQuat4f::Identity;
+			Q1 = FQuat4f(FRotator::ZeroRotator);
+			LogQuaternionTest(TEXT("FRotator::ZeroRotator ~= FQuat4f::Identity : Quaternion"), Q0, Q1, Q0.Equals(Q1, 0.f));
+			LogQuaternionTest(TEXT("FRotator::ZeroRotator == FQuat4f::Identity : Quaternion"), Q0, Q1, Q0 == Q1);
+			LogQuaternionTest(TEXT("FRotator::ZeroRotator not != FQuat4f::Identity : Quaternion"), Q0, Q1, !(Q0 != Q1));
 		}
 
 		const float Nudge = KINDA_SMALL_NUMBER * 0.25f;
@@ -2667,8 +2682,8 @@ bool FVectorRegisterAbstractionTest::RunTest(const FString& Parameters)
 		const float QuatTolerance = 1e-6f;
 		for (auto const& A : RotArray)
 		{
-			const FQuat QA = TestRotatorToQuaternion(A);
-			const FQuat QB = A.Quaternion();
+			const FQuat4f QA = TestRotatorToQuaternion(A);
+			const FQuat4f QB = FQuat4f(A);
 			LogQuaternionTest(TEXT("TestRotatorToQuaternion"), QA, QB, TestQuatsEqual(QA, QB, QuatTolerance));
 		}
 	}
@@ -2677,6 +2692,8 @@ bool FVectorRegisterAbstractionTest::RunTest(const FString& Parameters)
 	{
 		const float Nudge = KINDA_SMALL_NUMBER * 0.25f;
 		const FRotator RotArray[] ={
+			FRotator(0.0f, 0.0f, 0.0f),
+			FRotator(30.0f, 30.0f, 30.0f),
 			FRotator(30.0f, -45.0f, 90.0f),
 			FRotator(45.0f, 60.0f, -120.0f),
 			FRotator(0.f, 90.f, 0.f),
@@ -2686,6 +2703,8 @@ bool FVectorRegisterAbstractionTest::RunTest(const FString& Parameters)
 			FRotator(90.f, 0.f, 0.f),
 			FRotator(-90.f, 0.f, 0.f),
 			FRotator(150.f, 0.f, 0.f),
+			FRotator(0.0f, 0.0f, 45.0f),
+			FRotator(0.0f, 0.0f, -45.0f),
 			FRotator(+360.0f, -720.0f, 1080.0f),
 			FRotator(+360.0f + 1.0f, -720.0f + 1.0f, 1080.0f + 1.0f),
 			FRotator(+360.0f + Nudge, -720.0f - Nudge, 1080.0f - Nudge),
@@ -2705,22 +2724,22 @@ bool FVectorRegisterAbstractionTest::RunTest(const FString& Parameters)
 
 	// Quat -> Axis and Angle
 	{
-		FVector Axis;
+		FVector3f Axis;
 		float Angle;
 
 		// Identity -> X Axis
-		Axis = FQuat::Identity.GetRotationAxis();
-		LogTest<float>(TEXT("FQuat::Identity.GetRotationAxis() == FVector3f::XAxisVector"), TestFVector3Equal(Axis, FVector3f::XAxisVector));
-
-		const FQuat QuatArray[] = {
-			FQuat(0.0f, 0.0f, 0.0f, 1.0f),
-			FQuat(1.0f, 0.0f, 0.0f, 0.0f),
-			FQuat(0.0f, 1.0f, 0.0f, 0.0f),
-			FQuat(0.0f, 0.0f, 1.0f, 0.0f),
-			FQuat(0.000046571717f, -0.000068426132f, 0.000290602446f, 0.999999881000f) // length = 0.99999992665
+		Axis = FQuat4f::Identity.GetRotationAxis();
+		LogTest<float>(TEXT("FQuat4f::Identity.GetRotationAxis() == FVector3f::XAxisVector"), TestFVector3Equal(Axis, FVector3f::XAxisVector));
+
+		const FQuat4f QuatArray[] = {
+			FQuat4f(0.0f, 0.0f, 0.0f, 1.0f),
+			FQuat4f(1.0f, 0.0f, 0.0f, 0.0f),
+			FQuat4f(0.0f, 1.0f, 0.0f, 0.0f),
+			FQuat4f(0.0f, 0.0f, 1.0f, 0.0f),
+			FQuat4f(0.000046571717f, -0.000068426132f, 0.000290602446f, 0.999999881000f) // length = 0.99999992665
 		};
 
-		for (const FQuat& Q : QuatArray)
+		for (const FQuat4f& Q : QuatArray)
 		{
 			Q.ToAxisAndAngle(Axis, Angle);
 			LogTest<float>(TEXT("Quat -> Axis and Angle: Q is Normalized"), TestQuatNormalized(Q, 1e-6f));
@@ -2728,11 +2747,64 @@ bool FVectorRegisterAbstractionTest::RunTest(const FString& Parameters)
 		}
 	}
 
+	// Quat Lerp/Slerp
+	{
+		const FRotator Rotator0 = FRotator( 300.0f, -45.0f,  0.0f);
+		const FRotator Rotator1 = FRotator(  10.0f, 270.0f, 30.0f);
+		const float FloatAlpha = 0.25f;
+		const double DoubleAlpha = 0.25;
+		
+		// Quat<float>
+		Q0 = FQuat4f(Rotator0);
+		Q1 = FQuat4f(Rotator1);
+		
+		// float alpha
+		Q2 = FMath::Lerp(Q0, Q1, FloatAlpha);
+		Q3 = FQuat4f::Slerp(Q0, Q1, FloatAlpha);
+		LogTest<float>(TEXT("TestQuatLerp"), TestQuatsEqual(Q2, Q3, 1e-6f));
+
+		Q2 = FMath::Lerp<FQuat4f>(Q0, Q1, FloatAlpha);
+		Q3 = FQuat4f::Slerp(Q0, Q1, FloatAlpha);
+		LogTest<float>(TEXT("TestQuatLerp"), TestQuatsEqual(Q2, Q3, 1e-6f));
+		
+		// double alpha
+		Q2 = FMath::Lerp(Q0, Q1, DoubleAlpha);
+		Q3 = FQuat4f::Slerp(Q0, Q1, (float)DoubleAlpha);
+		LogTest<float>(TEXT("TestQuatLerp"), TestQuatsEqual(Q2, Q3, 1e-6f));
+
+		Q2 = FMath::Lerp<FQuat4f>(Q0, Q1, DoubleAlpha);
+		Q3 = FQuat4f::Slerp(Q0, Q1, (float)DoubleAlpha);
+		LogTest<float>(TEXT("TestQuatLerp"), TestQuatsEqual(Q2, Q3, 1e-6f));
+
+		// Quat<double>
+		FQuat4d QD0, QD1, QD2, QD3;
+		QD0 = FQuat4d(Rotator0);
+		QD1 = FQuat4d(Rotator1);
+		
+		// double alpha
+		QD2 = FMath::Lerp(QD0, QD1, DoubleAlpha);
+		QD3 = FQuat4d::Slerp(QD0, QD1, DoubleAlpha);
+		LogTest<double>(TEXT("TestQuatLerp"), TestQuatsEqual(QD2, QD3, 1e-6f));
+
+		QD2 = FMath::Lerp<FQuat4d>(QD0, QD1, DoubleAlpha);
+		QD3 = FQuat4d::Slerp(QD0, QD1, DoubleAlpha);
+		LogTest<double>(TEXT("TestQuatLerp"), TestQuatsEqual(QD2, QD3, 1e-6f));
+		
+		// float alpha
+		QD2 = FMath::Lerp(QD0, QD1, FloatAlpha);
+		QD3 = FQuat4d::Slerp(QD0, QD1, FloatAlpha);
+		LogTest<double>(TEXT("TestQuatLerp"), TestQuatsEqual(QD2, QD3, 1e-6f));
+
+		QD2 = FMath::Lerp<FQuat4d>(QD0, QD1, FloatAlpha);
+		QD3 = FQuat4d::Slerp(QD0, QD1, FloatAlpha);
+		LogTest<double>(TEXT("TestQuatLerp"), TestQuatsEqual(QD2, QD3, 1e-6f));
+	}
+
 	// Quat / Rotator conversion to vectors, matrices
 	{
 		FRotator Rotator0;
 		Rotator0 = FRotator(30.0f, -45.0f, 90.0f);
-		Q0 = Rotator0.Quaternion();
+		Q0 = FQuat4f(Rotator0);
 		Q1 = TestRotatorToQuaternion(Rotator0);
 		LogTest<float>( TEXT("TestRotatorToQuaternion"), TestQuatsEqual(Q0, Q1, 1e-6f));
 
@@ -2742,11 +2814,11 @@ bool FVectorRegisterAbstractionTest::RunTest(const FString& Parameters)
 		LogTest<float>( TEXT("Test0 Rotator::Vector()"), TestFVector3Equal(FV0, FV1, 1e-6f));
 		
 		FV0 = FRotationMatrix( Rotator0 ).GetScaledAxis( EAxis::X );
-		FV1 = FQuatRotationMatrix( Q0 ).GetScaledAxis( EAxis::X );
+		FV1 = FQuatRotationMatrix( FQuat(Q0.X, Q0.Y, Q0.Z, Q0.W) ).GetScaledAxis( EAxis::X );
 		LogTest<float>( TEXT("Test0 FQuatRotationMatrix"), TestFVector3Equal(FV0, FV1, 1e-5f));
 
 		Rotator0 = FRotator(45.0f,  60.0f, 120.0f);
-		Q0 = Rotator0.Quaternion();
+		Q0 = FQuat4f(Rotator0);
 		Q1 = TestRotatorToQuaternion(Rotator0);
 		LogTest<float>( TEXT("TestRotatorToQuaternion"), TestQuatsEqual(Q0, Q1, 1e-6f));
 
@@ -2755,24 +2827,25 @@ bool FVectorRegisterAbstractionTest::RunTest(const FString& Parameters)
 		LogTest<float>( TEXT("Test1 Rotator::Vector()"), TestFVector3Equal(FV0, FV1, 1e-6f));
 
 		FV0 = FRotationMatrix( Rotator0 ).GetScaledAxis( EAxis::X );
-		FV1 = FQuatRotationMatrix( Q0 ).GetScaledAxis( EAxis::X );
+		FV1 = FQuatRotationMatrix( FQuat(Q0.X, Q0.Y, Q0.Z, Q0.W) ).GetScaledAxis( EAxis::X );
 		LogTest<float>(TEXT("Test1 FQuatRotationMatrix"), TestFVector3Equal(FV0, FV1, 1e-5f));
 
-		FV0 = FRotationMatrix(FRotator::ZeroRotator).GetScaledAxis(EAxis::X);
-		FV1 = FQuatRotationMatrix(FQuat::Identity).GetScaledAxis(EAxis::X);
+		FV0 = FRotationMatrix( FRotator::ZeroRotator ).GetScaledAxis(EAxis::X);
+		FV1 = FQuatRotationMatrix( FQuat::Identity ).GetScaledAxis(EAxis::X);
 		LogTest<float>(TEXT("Test2 FQuatRotationMatrix"), TestFVector3Equal(FV0, FV1, 1e-6f));
 	}
 
 	// Quat Rotation tests
 	{
 		// Use these Quats...
-		const FQuat TestQuats[] = {
-			FQuat(FQuat::Identity),
-			FQuat(FRotator(30.0f, -45.0f, 90.0f)),
-			FQuat(FRotator(45.0f,  60.0f, 120.0f)),
-			FQuat(FRotator(0.0f, 180.0f, 45.0f)),
-			FQuat(FRotator(-120.0f, -90.0f, 0.0f)),
-			FQuat(FRotator(-0.01f, 0.02f, -0.03f)),
+		const FQuat4f TestQuats[] = {
+			FQuat4f(FQuat4f::Identity),
+			FQuat4f(FRotator(30.0f, 30.0f, 30.0f)),
+			FQuat4f(FRotator(30.0f, -45.0f, 90.0f)),
+			FQuat4f(FRotator(45.0f,  60.0f, 120.0f)),
+			FQuat4f(FRotator(0.0f, 180.0f, 45.0f)),
+			FQuat4f(FRotator(-120.0f, -90.0f, 0.0f)),
+			FQuat4f(FRotator(-0.01f, 0.02f, -0.03f)),
 		};
 
 		// ... to rotate these Vectors...
@@ -2790,7 +2863,7 @@ bool FVectorRegisterAbstractionTest::RunTest(const FString& Parameters)
 		// ... and test within this tolerance.
 		const float Tolerance = 1e-4f;
 
-		// Test Macro. Tests FQuat::RotateVector(_Vec) against _Func(_Vec)
+		// Test Macro. Tests FQuat4f::RotateVector(_Vec) against _Func(_Vec)
 		#define TEST_QUAT_ROTATE(_QIndex, _VIndex, _Quat, _Vec, _Func, _Tolerance) \
 		{ \
 			const FString _TestName = FString::Printf(TEXT("Test Quat%d: Vec%d: %s"), _QIndex, _VIndex, TEXT(#_Func)); \
@@ -2800,7 +2873,7 @@ bool FVectorRegisterAbstractionTest::RunTest(const FString& Parameters)
 		// Test loop
 		for (int32 QIndex = 0; QIndex < UE_ARRAY_COUNT(TestQuats); ++QIndex)
 		{
-			const FQuat& Q = TestQuats[QIndex];
+			const FQuat4f& Q = TestQuats[QIndex];
 			for (int32 VIndex = 0; VIndex < UE_ARRAY_COUNT(TestVectors); ++VIndex)
 			{
 				const FVector3f& V = TestVectors[VIndex];
@@ -2820,9 +2893,9 @@ bool FVectorRegisterAbstractionTest::RunTest(const FString& Parameters)
 					const FVector3f ANorm = A.GetSafeNormal();
 					const FVector3f BNorm = B.GetSafeNormal();
 
-					const FQuat Old = FindBetween_Old(ANorm, BNorm);
-					const FQuat NewNormal = FQuat::FindBetweenNormals(ANorm, BNorm);
-					const FQuat NewVector = FQuat::FindBetweenVectors(A, B);
+					const FQuat4f Old = FindBetween_Old(ANorm, BNorm);
+					const FQuat4f NewNormal = FQuat4f::FindBetweenNormals(ANorm, BNorm);
+					const FQuat4f NewVector = FQuat4f::FindBetweenVectors(A, B);
 
 					const FVector3f RotAOld = Old.RotateVector(ANorm);
 					const FVector3f RotANewNormal = NewNormal.RotateVector(ANorm);
@@ -2850,7 +2923,7 @@ bool FVectorRegisterAbstractionTest::RunTest(const FString& Parameters)
 			{
 				const FVector3f VNormal = V.GetSafeNormal();
 
-				Q0 = FQuat::FindBetweenNormals(FVector3f::ForwardVector, VNormal);
+				Q0 = FQuat4f::FindBetweenNormals(FVector3f::ForwardVector, VNormal);
 				Q1 = V.ToOrientationQuat();
 				const FRotator R0 = V.ToOrientationRotator();
 
@@ -2866,8 +2939,8 @@ bool FVectorRegisterAbstractionTest::RunTest(const FString& Parameters)
 
 	// Quat multiplication
 	{
-		Q0 = FQuat(FRotator(30.0f, -45.0f, 90.0f));
-		Q1 = FQuat(FRotator(45.0f, 60.0f, 120.0f));
+		Q0 = FQuat4f(FRotator(30.0f, -45.0f, 90.0f));
+		Q1 = FQuat4f(FRotator(45.0f, 60.0f, 120.0f));
 		VectorQuaternionMultiply(&Q2, &Q0, &Q1);
 		TestVectorQuaternionMultiply(&Q3, &Q0, &Q1);
 		LogTest<float>(TEXT("VectorQuaternionMultiply"), TestQuatsEqual(Q2, Q3, 1e-6f));
@@ -2877,8 +2950,8 @@ bool FVectorRegisterAbstractionTest::RunTest(const FString& Parameters)
 		V3 = VectorLoadAligned(&Q3);
 		LogTest<float>(TEXT("VectorQuaternionMultiply2"), TestVectorsEqual(V2, V3, 1e-6f));
 
-		Q0 = FQuat(FRotator(0.0f, 180.0f, 45.0f));
-		Q1 = FQuat(FRotator(-120.0f, -90.0f, 0.0f));
+		Q0 = FQuat4f(FRotator(0.0f, 180.0f, 45.0f));
+		Q1 = FQuat4f(FRotator(-120.0f, -90.0f, 0.0f));
 		VectorQuaternionMultiply(&Q2, &Q0, &Q1);
 		TestVectorQuaternionMultiply(&Q3, &Q0, &Q1);
 		LogTest<float>(TEXT("VectorMatrixInverse"), TestQuatsEqual(Q2, Q3, 1e-6f));
diff --git a/Engine/Source/Runtime/Core/Public/CoreFwd.h b/Engine/Source/Runtime/Core/Public/CoreFwd.h
--- a/Engine/Source/Runtime/Core/Public/CoreFwd.h
+++ b/Engine/Source/Runtime/Core/Public/CoreFwd.h
@@ -23,9 +23,9 @@ struct FColor;
 struct FLinearColor;
 DECLARE_LWC_TYPE(Matrix, 44);
 DECLARE_LWC_TYPE(Plane, 4);
-struct FQuat;
+DECLARE_LWC_TYPE(Quat, 4);
 struct FRotator;
-struct FTransform;
+DECLARE_LWC_TYPE(Transform, 3);
 DECLARE_LWC_TYPE(Vector, 3);
 struct FVector2D;
 struct FVector4;
diff --git a/Engine/Source/Runtime/Core/Public/GenericPlatform/GenericPlatformMath.h b/Engine/Source/Runtime/Core/Public/GenericPlatform/GenericPlatformMath.h
--- a/Engine/Source/Runtime/Core/Public/GenericPlatform/GenericPlatformMath.h
+++ b/Engine/Source/Runtime/Core/Public/GenericPlatform/GenericPlatformMath.h
@@ -199,6 +199,10 @@ struct FGenericPlatformMath
 	{
 		return (int32)F;
 	}
+	static CONSTEXPR FORCEINLINE int32 TruncToInt(double F)
+	{
+		return (int32)F;
+	}
 
 	/**
 	 * Converts a float to an integer value with truncation towards zero.
@@ -231,6 +235,12 @@ struct FGenericPlatformMath
 		I -= ( (float)I > F );
 		return I;
 	}
+	static FORCEINLINE int32 FloorToInt(double F)
+	{
+		int32 I = (int32)F;
+		I -= ((float)I > F);
+		return I;
+	}
 
 	/**
 	* Converts a float to the nearest less or equal integer.
@@ -294,6 +304,13 @@ struct FGenericPlatformMath
 		return I;
 	}
 
+	static FORCEINLINE int32 CeilToInt(double F)
+	{
+		int32 I = (int32)F;
+		I += ((float)I < F);
+		return I;
+	}
+
 	/**
 	* Converts a float to the nearest greater or equal integer.
 	* @param F		Floating point value to convert
@@ -774,7 +791,7 @@ struct FGenericPlatformMath
 	 */
 	static CONSTEXPR FORCEINLINE double FloatSelect( double Comparand, double ValueGEZero, double ValueLTZero )
 	{
-		return Comparand >= 0.f ? ValueGEZero : ValueLTZero;
+		return Comparand >= 0.0 ? ValueGEZero : ValueLTZero;
 	}
 
 	/** Computes absolute value in a generic way */
@@ -948,6 +965,12 @@ struct FGenericPlatformMath
 	template<typename T> UE_DEPRECATED(5.0, "Arguments cause function resolution ambiguity.")
 	static FORCEINLINE float InvSqrtEst(T&& F) { return InvSqrt((float)F); }
 
+	template<typename T> UE_DEPRECATED(5.0, "Arguments cause function resolution ambiguity.")
+	static FORCEINLINE int32 FloorToInt(T F) { return FloorToInt((float)F); }
+	template<typename T> UE_DEPRECATED(5.0, "Arguments cause function resolution ambiguity.")
+	static FORCEINLINE int32 CeilToInt(T F) { return CeilToInt((float)F); }
+	template<typename T> UE_DEPRECATED(5.0, "Arguments cause function resolution ambiguity.")
+	static CONSTEXPR FORCEINLINE int32 TruncToInt(T F) { return TruncToInt((float)F); }
 private:
 
 	/** Error reporting for Fmod. Not inlined to avoid compilation issues and avoid all the checks and error reporting at all callsites. */
diff --git a/Engine/Source/Runtime/Core/Public/GenericPlatform/GenericPlatformMemory.h b/Engine/Source/Runtime/Core/Public/GenericPlatform/GenericPlatformMemory.h
--- a/Engine/Source/Runtime/Core/Public/GenericPlatform/GenericPlatformMemory.h
+++ b/Engine/Source/Runtime/Core/Public/GenericPlatform/GenericPlatformMemory.h
@@ -181,7 +181,10 @@ struct FPlatformMemoryStats;
 #define __FMemory_Alloca_Func alloca
 #endif
 
-#define FMemory_Alloca(Size )((Size==0) ? 0 : (void*)(((PTRINT)__FMemory_Alloca_Func(Size + 15) + 15) & ~15))
+#define FMemory_Alloca(Size) ((Size==0) ? 0 : (void*)(((PTRINT)__FMemory_Alloca_Func(Size + 15) + 15) & ~15))
+
+ // Version that supports alignment requirement. However since the old alignment was always forced to be 16, this continues to enforce a min alignment of 16 but allows a larger value.
+#define FMemory_Alloca_Aligned(Size, Alignment) ((Size==0) ? 0 : ((Alignment <= 16) ? FMemory_Alloca(Size) : (void*)(((PTRINT)__FMemory_Alloca_Func(Size + Alignment-1) + Alignment-1) & ~(Alignment-1))))
 
 /** Generic implementation for most platforms, these tend to be unused and unimplemented. */
 struct CORE_API FGenericPlatformMemory
diff --git a/Engine/Source/Runtime/Core/Public/Math/MathFwd.h b/Engine/Source/Runtime/Core/Public/Math/MathFwd.h
--- a/Engine/Source/Runtime/Core/Public/Math/MathFwd.h
+++ b/Engine/Source/Runtime/Core/Public/Math/MathFwd.h
@@ -9,8 +9,8 @@ struct FColor;
 struct FLinearColor;
 DECLARE_LWC_TYPE(Matrix, 44);
 DECLARE_LWC_TYPE(Plane, 4);
-struct FQuat;
+DECLARE_LWC_TYPE(Quat, 4);
 struct FRotator;
-struct FTransform;
+DECLARE_LWC_TYPE(Transform, 3);
 DECLARE_LWC_TYPE(Vector, 3);
 struct FVector4;
diff --git a/Engine/Source/Runtime/Core/Public/Math/Matrix.h b/Engine/Source/Runtime/Core/Public/Math/Matrix.h
--- a/Engine/Source/Runtime/Core/Public/Math/Matrix.h
+++ b/Engine/Source/Runtime/Core/Public/Math/Matrix.h
@@ -13,6 +13,8 @@
 #include "Math/Axis.h"
 #include <type_traits>
 
+DECLARE_LWC_TYPE(Quat, 4);
+
 #ifdef _MSC_VER
 #pragma warning (push)
 // Ensure template functions don't generate shadowing warnings against global variables at the point of instantiation.
@@ -325,7 +327,7 @@ public:
 	 *
 	 * @warning rotation part will need to be unit length for this to be right!
 	 */
-	CORE_API FQuat ToQuat() const;
+	CORE_API UE::Math::TQuat<T> ToQuat() const;
 
 	// Frustum plane extraction.
 	/** @param OuTPln the near plane of the Frustum of this matrix */
@@ -532,8 +534,8 @@ DECLARE_LWC_TYPE_ISPC(Matrix, 44);
 // Forward declare all explicit specializations (in UnrealMath.cpp)
 template<> CORE_API FRotator FMatrix44f::Rotator() const;
 template<> CORE_API FRotator FMatrix44d::Rotator() const;
-template<> CORE_API FQuat FMatrix44f::ToQuat() const;
-template<> CORE_API FQuat FMatrix44d::ToQuat() const;
+template<> CORE_API FQuat4f FMatrix44f::ToQuat() const;
+template<> CORE_API FQuat4d FMatrix44d::ToQuat() const;
 
 
 #if 0	// LWC_TODO: Is this in use?
