LWC support for FQuat and FTransform

diff --git a/Engine/Source/Runtime/Core/Public/Math/Quat.h b/Engine/Source/Runtime/Core/Public/Math/Quat.h
--- a/Engine/Source/Runtime/Core/Public/Math/Quat.h
+++ b/Engine/Source/Runtime/Core/Public/Math/Quat.h
@@ -14,7 +14,10 @@
 
 class Error;
 
-// LWC_TODO: Float casting in here to inhibit precision loss warnings from LWC enabled types. Revisit(float) casts / usage when LWC is enabled for this struct.
+namespace UE
+{
+namespace Math
+{
 
 /**
  * Floating point quaternion that can represent a rotation about an axis in 3-D space.
@@ -27,38 +30,46 @@ class Error;
  * Example: LocalToWorld = (LocalToWorld * DeltaRotation) will change rotation in local space by DeltaRotation.
  * Example: LocalToWorld = (DeltaRotation * LocalToWorld) will change rotation in world space by DeltaRotation.
  */
-MS_ALIGN(16) struct FQuat 
+
+template<typename T>
+struct alignas(16) TQuat
 {
+	// Can't have a TEMPLATE_REQUIRES in the declaration because of the forward declarations, so check for allowed types here.
+	static_assert(TIsFloatingPoint<T>::Value, "TQuat only supports float and double types.");
+
 public:
+	/** Type of the template param (float or double) */
+	using FReal = T;
+	using QuatVectorRegister = TVectorRegisterType<T>;
 
 	/** The quaternion's X-component. */
-	float X;
+	T X;
 
 	/** The quaternion's Y-component. */
-	float Y;
+	T Y;
 
 	/** The quaternion's Z-component. */
-	float Z;
+	T Z;
 
 	/** The quaternion's W-component. */
-	float W;
+	T W;
 
 public:
 
 	/** Identity quaternion. */
-	static CORE_API const FQuat Identity;
+	static const TQuat<T> Identity;
 
 public:
 
 	/** Default constructor (no initialization). */
-	FORCEINLINE FQuat() { }
+	FORCEINLINE TQuat();
 
 	/**
 	 * Creates and initializes a new quaternion, with the W component either 0 or 1.
 	 *
 	 * @param EForceInit Force init enum: if equal to ForceInitToZero then W is 0, otherwise W = 1 (creating an identity transform)
 	 */
-	explicit FORCEINLINE FQuat(EForceInit);
+	explicit FORCEINLINE TQuat(EForceInit);
 
 	/**
 	 * Constructor.
@@ -68,34 +79,45 @@ public:
 	 * @param InZ Z component of the quaternion
 	 * @param InW W component of the quaternion
 	 */
-	FORCEINLINE FQuat(float InX, float InY, float InZ, float InW);
+	FORCEINLINE TQuat(T InX, T InY, T InZ, T InW);
 
 	/**
 	 * Initializes all elements to V
 	 */
-	explicit FORCEINLINE FQuat(float V);
+	template<TEMPLATE_REQUIRES(std::is_arithmetic<T>::value)>
+	explicit FORCEINLINE TQuat(T V)
+	: X(V), Y(V), Z(V), W(V)
+	{
+		DiagnosticCheckNaN();
+	}
 
+protected:
 	/**
 	 * Creates and initializes a new quaternion from the XYZW values in the given VectorRegister4Float.
 	 *
 	 * @param V XYZW components of the quaternion packed into a single VectorRegister4Float
 	 */
-	explicit FORCEINLINE FQuat(const VectorRegister4Float& V);
-	explicit FORCEINLINE FQuat(const VectorRegister4Double& V);
+	explicit TQuat(const QuatVectorRegister& V);
+
+public:
+
+	FORCEINLINE static TQuat<T> MakeFromVectorRegister(const QuatVectorRegister& V) { return TQuat<T>(V); }
 
 	/**
-	 * Creates and initializes a new quaternion from the given matrix.
+	 * Creates and initializes a new quaternion from the given rotator.
 	 *
-	 * @param M The rotation matrix to initialize from.
+	 * @param R The rotator to initialize from.
 	 */
-	explicit FQuat(const FMatrix& M);	
+	explicit TQuat(const FRotator& R);
+
+	FORCEINLINE static TQuat<T> MakeFromRotator(const FRotator& R) { return TQuat<T>(R); }
 
 	/**
-	 * Creates and initializes a new quaternion from the given rotator.
+	 * Creates and initializes a new quaternion from the given matrix.
 	 *
-	 * @param R The rotator to initialize from.
+	 * @param M The rotation matrix to initialize from.
 	 */
-	explicit FQuat(const FRotator& R);
+	explicit TQuat(const TMatrix<T>& M);
 
 	/**
 	 * Creates and initializes a new quaternion from the a rotation around the given axis.
@@ -103,17 +125,17 @@ public:
 	 * @param Axis assumed to be a normalized vector
 	 * @param Angle angle to rotate above the given axis (in radians)
 	 */
-	FQuat(FVector Axis, float AngleRad);
+	TQuat(TVector<T> Axis, T AngleRad);
 
 public:
 
 #ifdef IMPLEMENT_ASSIGNMENT_OPERATOR_MANUALLY
 	/**
-	 * Copy another FQuat into this one
+	 * Copy another TQuat into this one
 	 *
-	 * @return reference to this FQuat
+	 * @return reference to this TQuat
 	 */
-	FORCEINLINE FQuat& operator=(const FQuat& Other);
+	FORCEINLINE TQuat<T>& operator=(const TQuat<T>& Other);
 #endif
 
 	/**
@@ -123,7 +145,7 @@ public:
 	 * @param Q The Quaternion to add.
 	 * @return The result of addition.
 	 */
-	FORCEINLINE FQuat operator+(const FQuat& Q) const;
+	FORCEINLINE TQuat<T> operator+(const TQuat<T>& Q) const;
 
 	/**
 	 * Adds to this quaternion.
@@ -132,7 +154,7 @@ public:
 	 * @param Other The quaternion to add to this.
 	 * @return Result after addition.
 	 */
-	FORCEINLINE FQuat operator+=(const FQuat& Q);
+	FORCEINLINE TQuat<T> operator+=(const TQuat<T>& Q);
 
 	/**
 	 * Gets the result of subtracting a Quaternion to this.
@@ -141,14 +163,14 @@ public:
 	 * @param Q The Quaternion to subtract.
 	 * @return The result of subtraction.
 	 */
-	FORCEINLINE FQuat operator-(const FQuat& Q) const;
+	FORCEINLINE TQuat<T> operator-(const TQuat<T>& Q) const;
 
 	/**
 	* Negates the quaternion. Note that this represents the same rotation.
 	*
 	* @return The result of negation.
 	*/
-	FORCEINLINE FQuat operator-() const;
+	FORCEINLINE TQuat<T> operator-() const;
 
 	/**
 	 * Checks whether another Quaternion is equal to this, within specified tolerance.
@@ -157,7 +179,7 @@ public:
 	 * @param Tolerance Error tolerance for comparison with other Quaternion.
 	 * @return true if two Quaternions are equal, within specified tolerance, otherwise false.
 	 */
-	FORCEINLINE bool Equals(const FQuat& Q, float Tolerance=KINDA_SMALL_NUMBER) const;
+	FORCEINLINE bool Equals(const TQuat<T>& Q, T Tolerance=KINDA_SMALL_NUMBER) const;
 
 	/**
 	 * Checks whether this Quaternion is an Identity Quaternion.
@@ -166,7 +188,7 @@ public:
 	 * @param Tolerance Error tolerance for comparison with Identity Quaternion.
 	 * @return true if Quaternion is a normalized Identity Quaternion.
 	 */
-	FORCEINLINE bool IsIdentity(float Tolerance=SMALL_NUMBER) const;
+	FORCEINLINE bool IsIdentity(T Tolerance=SMALL_NUMBER) const;
 
 	/**
 	 * Subtracts another quaternion from this.
@@ -175,7 +197,7 @@ public:
 	 * @param Q The other quaternion.
 	 * @return reference to this after subtraction.
 	 */
-	FORCEINLINE FQuat operator-=(const FQuat& Q);
+	FORCEINLINE TQuat<T> operator-=(const TQuat<T>& Q);
 
 	/**
 	 * Gets the result of multiplying this by another quaternion (this * Q).
@@ -186,7 +208,7 @@ public:
 	 * @param Q The Quaternion to multiply this by.
 	 * @return The result of multiplication (this * Q).
 	 */
-	FORCEINLINE FQuat operator*(const FQuat& Q) const;
+	FORCEINLINE TQuat<T> operator*(const TQuat<T>& Q) const;
 
 	/**
 	 * Multiply this by a quaternion (this = this * Q).
@@ -197,7 +219,7 @@ public:
 	 * @param Q the quaternion to multiply this by.
 	 * @return The result of multiplication (this * Q).
 	 */
-	FORCEINLINE FQuat operator*=(const FQuat& Q);
+	FORCEINLINE TQuat<T> operator*=(const TQuat<T>& Q);
 
 	/**
 	 * Rotate a vector by this quaternion.
@@ -206,7 +228,7 @@ public:
 	 * @return vector after rotation
 	 * @see RotateVector
 	 */
-	FVector operator*(const FVector& V) const;
+	TVector<T> operator*(const TVector<T>& V) const;
 
 	/** 
 	 * Multiply this by a matrix.
@@ -217,7 +239,7 @@ public:
 	 * @param M Matrix to multiply by.
 	 * @return Matrix result after multiplication.
 	 */	
-	FMatrix operator*(const FMatrix& M) const;
+	TMatrix<T> operator*(const TMatrix<T>& M) const;
 	
 	/**
 	 * Multiply this quaternion by a scaling factor.
@@ -225,7 +247,7 @@ public:
 	 * @param Scale The scaling factor.
 	 * @return a reference to this after scaling.
 	 */
-	FORCEINLINE FQuat operator*=(const float Scale);
+	FORCEINLINE TQuat<T> operator*=(const T Scale);
 
 	/**
 	 * Get the result of scaling this quaternion.
@@ -233,7 +255,7 @@ public:
 	 * @param Scale The scaling factor.
 	 * @return The result of scaling.
 	 */
-	FORCEINLINE FQuat operator*(const float Scale) const;
+	FORCEINLINE TQuat<T> operator*(const T Scale) const;
 	
 	/**
 	 * Divide this quaternion by scale.
@@ -241,7 +263,7 @@ public:
 	 * @param Scale What to divide by.
 	 * @return a reference to this after scaling.
 	 */
-	FORCEINLINE FQuat operator/=(const float Scale);
+	FORCEINLINE TQuat<T> operator/=(const T Scale);
 
 	/**
 	 * Divide this quaternion by scale.
@@ -249,13 +271,13 @@ public:
 	 * @param Scale What to divide by.
 	 * @return new Quaternion of this after division by scale.
 	 */
-	FORCEINLINE FQuat operator/(const float Scale) const;
+	FORCEINLINE TQuat<T> operator/(const T Scale) const;
 
 	/**
-	 * Identical implementation for FQuat properties. 
+	 * Identical implementation for TQuat properties. 
 	 * Avoids intrinsics to remain consistent with previous per-property comparison.
 	 */
-	bool Identical(const FQuat* Q, const uint32 PortFlags) const;
+	bool Identical(const TQuat* Q, const uint32 PortFlags) const;
 
  	/**
 	 * Checks whether two quaternions are identical.
@@ -266,7 +288,7 @@ public:
 	 * @return true if two quaternion are identical, otherwise false.
 	 * @see Equals
 	 */
-	bool operator==(const FQuat& Q) const;
+	bool operator==(const TQuat<T>& Q) const;
 
  	/**
 	 * Checks whether two quaternions are not identical.
@@ -274,7 +296,7 @@ public:
 	 * @param Q The other quaternion.
 	 * @return true if two quaternion are not identical, otherwise false.
 	 */
-	bool operator!=(const FQuat& Q) const;
+	bool operator!=(const TQuat<T>& Q) const;
 
 	/**
 	 * Calculates dot product of two quaternions.
@@ -282,7 +304,7 @@ public:
 	 * @param Q The other quaternions.
 	 * @return The dot product.
 	 */
-	float operator|(const FQuat& Q) const;
+	T operator|(const TQuat<T>& Q) const;
 
 public:
 
@@ -290,12 +312,12 @@ public:
 	 * Convert a vector of floating-point Euler angles (in degrees) into a Quaternion.
 	 * 
 	 * @param Euler the Euler angles
-	 * @return constructed FQuat
+	 * @return constructed TQuat
 	 */
-	static CORE_API FQuat MakeFromEuler(const FVector& Euler);
+	static CORE_API TQuat<T> MakeFromEuler(const TVector<T>& Euler);
 
 	/** Convert a Quaternion into floating-point Euler angles (in degrees). */
-	CORE_API FVector Euler() const;
+	CORE_API TVector<T> Euler() const;
 
 	/**
 	 * Normalize this quaternion if it is large enough.
@@ -303,7 +325,7 @@ public:
 	 *
 	 * @param Tolerance Minimum squared length of quaternion for normalization.
 	 */
-	FORCEINLINE void Normalize(float Tolerance = SMALL_NUMBER);
+	FORCEINLINE void Normalize(T Tolerance = SMALL_NUMBER);
 
 	/**
 	 * Get a normalized copy of this quaternion.
@@ -311,7 +333,7 @@ public:
 	 *
 	 * @param Tolerance Minimum squared length of quaternion for normalization.
 	 */
-	FORCEINLINE FQuat GetNormalized(float Tolerance = SMALL_NUMBER) const;
+	FORCEINLINE TQuat<T> GetNormalized(T Tolerance = SMALL_NUMBER) const;
 
 	// Return true if this quaternion is normalized
 	bool IsNormalized() const;
@@ -321,18 +343,18 @@ public:
 	 *
 	 * @return The length of this quaternion.
 	 */
-	FORCEINLINE float Size() const;
+	FORCEINLINE T Size() const;
 
 	/**
 	 * Get the length squared of this quaternion.
 	 *
 	 * @return The length of this quaternion.
 	 */
-	FORCEINLINE float SizeSquared() const;
+	FORCEINLINE T SizeSquared() const;
 
 
 	/** Get the angle of this quaternion */
-	FORCEINLINE float GetAngle() const;
+	FORCEINLINE T GetAngle() const;
 
 	/** 
 	 * get the axis and angle of rotation of this quaternion
@@ -341,7 +363,8 @@ public:
 	 * @param Angle{out] angle of the quaternion
 	 * @warning : assumes normalized quaternions.
 	 */
-	void ToAxisAndAngle(FVector& Axis, float& Angle) const;
+	void ToAxisAndAngle(TVector<T>& Axis, float& Angle) const;
+	void ToAxisAndAngle(TVector<T>& Axis, double& Angle) const;
 
 	/** 
 	 * Get the swing and twist decomposition for a specified axis
@@ -351,7 +374,7 @@ public:
 	 * @param OutTwist Twist component quaternion
 	 * @warning assumes normalized quaternion and twist axis
 	 */
-	CORE_API void ToSwingTwist(const FVector& InTwistAxis, FQuat& OutSwing, FQuat& OutTwist) const;
+	CORE_API void ToSwingTwist(const TVector<T>& InTwistAxis, TQuat<T>& OutSwing, TQuat<T>& OutTwist) const;
 
 	/**
 	 * Get the twist angle (in radians) for a specified axis
@@ -360,7 +383,7 @@ public:
 	 * @return Twist angle (in radians)
 	 * @warning assumes normalized quaternion and twist axis
 	 */
-	CORE_API float GetTwistAngle(const FVector& TwistAxis) const;
+	CORE_API T GetTwistAngle(const TVector<T>& TwistAxis) const;
 
 	/**
 	 * Rotate a vector by this quaternion.
@@ -368,7 +391,7 @@ public:
 	 * @param V the vector to be rotated
 	 * @return vector after rotation
 	 */
-	FVector RotateVector(FVector V) const;
+	TVector<T> RotateVector(TVector<T> V) const;
 	
 	/**
 	 * Rotate a vector by the inverse of this quaternion.
@@ -376,51 +399,51 @@ public:
 	 * @param V the vector to be rotated
 	 * @return vector after rotation by the inverse of this quaternion.
 	 */
-	FVector UnrotateVector(FVector V) const;
+	TVector<T> UnrotateVector(TVector<T> V) const;
 
 	/**
 	 * @return quaternion with W=0 and V=theta*v.
 	 */
-	CORE_API FQuat Log() const;
+	CORE_API TQuat<T> Log() const;
 
 	/**
 	 * @note Exp should really only be used after Log.
 	 * Assumes a quaternion with W=0 and V=theta*v (where |v| = 1).
 	 * Exp(q) = (sin(theta)*v, cos(theta))
 	 */
-	CORE_API FQuat Exp() const;
+	CORE_API TQuat<T> Exp() const;
 
 	/**
 	 * @return inverse of this quaternion
 	 */
-	FORCEINLINE FQuat Inverse() const;
+	FORCEINLINE TQuat<T> Inverse() const;
 
 	/**
 	 * Enforce that the delta between this Quaternion and another one represents
 	 * the shortest possible rotation angle
 	 */
-	void EnforceShortestArcWith(const FQuat& OtherQuat);
+	void EnforceShortestArcWith(const TQuat<T>& OtherQuat);
 	
 	/** Get the forward direction (X axis) after it has been rotated by this Quaternion. */
-	FORCEINLINE FVector GetAxisX() const;
+	FORCEINLINE TVector<T> GetAxisX() const;
 
 	/** Get the right direction (Y axis) after it has been rotated by this Quaternion. */
-	FORCEINLINE FVector GetAxisY() const;
+	FORCEINLINE TVector<T> GetAxisY() const;
 
 	/** Get the up direction (Z axis) after it has been rotated by this Quaternion. */
-	FORCEINLINE FVector GetAxisZ() const;
+	FORCEINLINE TVector<T> GetAxisZ() const;
 
 	/** Get the forward direction (X axis) after it has been rotated by this Quaternion. */
-	FORCEINLINE FVector GetForwardVector() const;
+	FORCEINLINE TVector<T> GetForwardVector() const;
 
 	/** Get the right direction (Y axis) after it has been rotated by this Quaternion. */
-	FORCEINLINE FVector GetRightVector() const;
+	FORCEINLINE TVector<T> GetRightVector() const;
 
 	/** Get the up direction (Z axis) after it has been rotated by this Quaternion. */
-	FORCEINLINE FVector GetUpVector() const;
+	FORCEINLINE TVector<T> GetUpVector() const;
 
 	/** Convert a rotation into a unit vector facing in its direction. Equivalent to GetForwardVector(). */
-	FORCEINLINE FVector Vector() const;
+	FORCEINLINE TVector<T> Vector() const;
 
 	/** Get the FRotator representation of this Quaternion. */
 	CORE_API FRotator Rotator() const;
@@ -428,12 +451,12 @@ public:
 	/**
 	 * Get the axis of rotation of the Quaternion.
 	 * This is the axis around which rotation occurs to transform the canonical coordinate system to the target orientation.
-	 * For the identity Quaternion which has no such rotation, FVector(1,0,0) is returned.
+	 * For the identity Quaternion which has no such rotation, TVector<T>(1,0,0) is returned.
 	 */
-	FORCEINLINE FVector GetRotationAxis() const;
+	FORCEINLINE TVector<T> GetRotationAxis() const;
 
 	/** Find the angular distance between two rotation quaternions (in radians) */
-	FORCEINLINE float AngularDistance(const FQuat& Q) const;
+	FORCEINLINE T AngularDistance(const TQuat<T>& Q) const;
 
 	/**
 	 * Serializes the vector compressed for e.g. network transmission.
@@ -457,12 +480,12 @@ public:
 	FString ToString() const;
 
 	/**
-	 * Initialize this FQuat from a FString. 
+	 * Initialize this TQuat from a FString. 
 	 * The string is expected to contain X=, Y=, Z=, W=, otherwise 
-	 * this FQuat will have indeterminate (invalid) values.
+	 * this TQuat will have indeterminate (invalid) values.
 	 *
 	 * @param InSourceString FString containing the quaternion values.
-	 * @return true if the FQuat was initialized; false otherwise.
+	 * @return true if the TQuat was initialized; false otherwise.
 	 */
 	bool InitFromString(const FString& InSourceString);
 
@@ -473,8 +496,8 @@ public:
 	{
 		if (ContainsNaN())
 		{
-			logOrEnsureNanError(TEXT("FQuat contains NaN: %s"), *ToString());
-			*const_cast<FQuat*>(this) = FQuat::Identity;
+			logOrEnsureNanError(TEXT("Quat contains NaN: %s"), *ToString());
+			*const_cast<TQuat*>(this) = TQuat<T>(0.f, 0.f, 0.f, 1.f);
 		}
 	}
 
@@ -482,8 +505,8 @@ public:
 	{
 		if (ContainsNaN())
 		{
-			logOrEnsureNanError(TEXT("%s: FQuat contains NaN: %s"), Message, *ToString());
-			*const_cast<FQuat*>(this) = FQuat::Identity;
+			logOrEnsureNanError(TEXT("%s: Quat contains NaN: %s"), Message, *ToString());
+			*const_cast<TQuat*>(this) = TQuat<T>(0.f, 0.f, 0.f, 1.f);
 		}
 	}
 #else
@@ -496,7 +519,7 @@ public:
 	/**
 	 * Generates the 'smallest' (geodesic) rotation between two vectors of arbitrary length.
 	 */
-	static FORCEINLINE FQuat FindBetween(const FVector& Vector1, const FVector& Vector2)
+	static FORCEINLINE TQuat<T> FindBetween(const TVector<T>& Vector1, const TVector<T>& Vector2)
 	{
 		return FindBetweenVectors(Vector1, Vector2);
 	}
@@ -504,43 +527,43 @@ public:
 	/**
 	 * Generates the 'smallest' (geodesic) rotation between two normals (assumed to be unit length).
 	 */
-	static CORE_API FQuat FindBetweenNormals(const FVector& Normal1, const FVector& Normal2);
+	static CORE_API TQuat<T> FindBetweenNormals(const TVector<T>& Normal1, const TVector<T>& Normal2);
 
 	/**
 	 * Generates the 'smallest' (geodesic) rotation between two vectors of arbitrary length.
 	 */
-	static CORE_API FQuat FindBetweenVectors(const FVector& Vector1, const FVector& Vector2);
+	static CORE_API TQuat<T> FindBetweenVectors(const TVector<T>& Vector1, const TVector<T>& Vector2);
 
 	/**
 	 * Error measure (angle) between two quaternions, ranged [0..1].
 	 * Returns the hypersphere-angle between two quaternions; alignment shouldn't matter, though 
 	 * @note normalized input is expected.
 	 */
-	static FORCEINLINE float Error(const FQuat& Q1, const FQuat& Q2);
+	static FORCEINLINE T Error(const TQuat<T>& Q1, const TQuat<T>& Q2);
 
 	/**
-	 * FQuat::Error with auto-normalization.
+	 * TQuat<T>::Error with auto-normalization.
 	 */
-	static FORCEINLINE float ErrorAutoNormalize(const FQuat& A, const FQuat& B);
+	static FORCEINLINE T ErrorAutoNormalize(const TQuat<T>& A, const TQuat<T>& B);
 
 	/** 
 	 * Fast Linear Quaternion Interpolation.
 	 * Result is NOT normalized.
 	 */
-	static FORCEINLINE FQuat FastLerp(const FQuat& A, const FQuat& B, const float Alpha);
+	static FORCEINLINE TQuat<T> FastLerp(const TQuat<T>& A, const TQuat<T>& B, const T Alpha);
 
 	/** 
 	 * Bi-Linear Quaternion interpolation.
 	 * Result is NOT normalized.
 	 */
-	static FORCEINLINE FQuat FastBilerp(const FQuat& P00, const FQuat& P10, const FQuat& P01, const FQuat& P11, float FracX, float FracY);
+	static FORCEINLINE TQuat<T> FastBilerp(const TQuat<T>& P00, const TQuat<T>& P10, const TQuat<T>& P01, const TQuat<T>& P11, T FracX, T FracY);
 
 
 	/** Spherical interpolation. Will correct alignment. Result is NOT normalized. */
-	static CORE_API FQuat Slerp_NotNormalized(const FQuat &Quat1, const FQuat &Quat2, float Slerp);
+	static CORE_API TQuat<T> Slerp_NotNormalized(const TQuat<T>& Quat1, const TQuat<T>& Quat2, T Slerp);
 
 	/** Spherical interpolation. Will correct alignment. Result is normalized. */
-	static FORCEINLINE FQuat Slerp(const FQuat &Quat1, const FQuat &Quat2, float Slerp)
+	static FORCEINLINE TQuat<T> Slerp(const TQuat<T>& Quat1, const TQuat<T>& Quat2, T Slerp)
 	{
 		return Slerp_NotNormalized(Quat1, Quat2, Slerp).GetNormalized();
 	}
@@ -550,14 +573,14 @@ public:
 	 * We need this for the cubic interpolation stuff so that the multiple Slerps dont go in different directions.
 	 * Result is NOT normalized.
 	 */
-	static CORE_API FQuat SlerpFullPath_NotNormalized(const FQuat &quat1, const FQuat &quat2, float Alpha);
+	static CORE_API TQuat<T> SlerpFullPath_NotNormalized(const TQuat<T>& quat1, const TQuat<T>& quat2, T Alpha);
 
 	/**
 	 * Simpler Slerp that doesn't do any checks for 'shortest distance' etc.
 	 * We need this for the cubic interpolation stuff so that the multiple Slerps dont go in different directions.
 	 * Result is normalized.
 	 */
-	static FORCEINLINE FQuat SlerpFullPath(const FQuat &quat1, const FQuat &quat2, float Alpha)
+	static FORCEINLINE TQuat<T> SlerpFullPath(const TQuat<T>& quat1, const TQuat<T>& quat2, T Alpha)
 	{
 		return SlerpFullPath_NotNormalized(quat1, quat2, Alpha).GetNormalized();
 	}
@@ -566,13 +589,13 @@ public:
 	 * Given start and end quaternions of quat1 and quat2, and tangents at those points tang1 and tang2, calculate the point at Alpha (between 0 and 1) between them. Result is normalized.
 	 * This will correct alignment by ensuring that the shortest path is taken.
 	 */
-	static CORE_API FQuat Squad(const FQuat& quat1, const FQuat& tang1, const FQuat& quat2, const FQuat& tang2, float Alpha);
+	static CORE_API TQuat<T> Squad(const TQuat<T>& quat1, const TQuat<T>& tang1, const TQuat<T>& quat2, const TQuat<T>& tang2, T Alpha);
 
 	/**
 	 * Simpler Squad that doesn't do any checks for 'shortest distance' etc.
 	 * Given start and end quaternions of quat1 and quat2, and tangents at those points tang1 and tang2, calculate the point at Alpha (between 0 and 1) between them. Result is normalized.
 	 */
-	static CORE_API FQuat SquadFullPath(const FQuat& quat1, const FQuat& tang1, const FQuat& quat2, const FQuat& tang2, float Alpha);
+	static CORE_API TQuat<T> SquadFullPath(const TQuat<T>& quat1, const TQuat<T>& tang1, const TQuat<T>& quat2, const TQuat<T>& tang2, T Alpha);
 
 	/**
 	 * Calculate tangents between given points
@@ -583,42 +606,66 @@ public:
 	 * @param Tension @todo document
 	 * @param OutTan Out control point
 	 */
-	static CORE_API void CalcTangents(const FQuat& PrevP, const FQuat& P, const FQuat& NextP, float Tension, FQuat& OutTan);
+	static CORE_API void CalcTangents(const TQuat<T>& PrevP, const TQuat<T>& P, const TQuat<T>& NextP, T Tension, TQuat<T>& OutTan);
 
 public:
 
-	/**
-	 * Serializes the quaternion.
-	 *
-	 * @param Ar Reference to the serialization archive.
-	 * @param F Reference to the quaternion being serialized.
-	 * @return Reference to the Archive after serialization.
-	 */
-	friend FArchive& operator<<(FArchive& Ar, FQuat& F)
-	{
-		return Ar << F.X << F.Y << F.Z << F.W;
-	}
-
 	bool Serialize(FArchive& Ar)
 	{
-		Ar << *this;
+		Ar << (TQuat<T>&)*this;
 		return true;
 	}
 
-} GCC_ALIGN(16);
+	// Conversion to other type.
+	template<typename FArg, TEMPLATE_REQUIRES(!TIsSame<T, FArg>::Value)>
+	explicit TQuat(const TQuat<FArg>& From) : TQuat<T>((T)From.X, (T)From.Y, (T)From.Z, (T)From.W) {}
+};
 
+/**
+ * Serializes the quaternion.
+ *
+ * @param Ar Reference to the serialization archive.
+ * @param F Reference to the quaternion being serialized.
+ * @return Reference to the Archive after serialization.
+ */
+inline FArchive& operator<<(FArchive& Ar, TQuat<float>& F)
+{
+	// LWC_TODO: Serializer
+	return Ar << F.X << F.Y << F.Z << F.W;
+}
 
-/* FQuat inline functions
- *****************************************************************************/
+inline FArchive& operator<<(FArchive& Ar, TQuat<double>& F)
+{
+	// LWC: Serialize as float
+	float X, Y, Z, W;
+	X = (float)F.X;
+	Y = (float)F.Y;
+	Z = (float)F.Z;
+	W = (float)F.W;
 
-inline FQuat::FQuat(const FMatrix& M)
+	// LWC_TODO: Serializer
+	Ar << X << Y << Z << W;
+
+	if (Ar.IsLoading())
+	{
+		F = TQuat<double>(X, Y, Z, W);
+	}
+	return Ar;
+}
+
+template<typename T> inline const TQuat<T> TQuat<T>::Identity = TQuat<T>(0.f, 0.f, 0.f, 1.f);
+
+/* TQuat inline functions
+ *****************************************************************************/
+template<typename T>
+inline TQuat<T>::TQuat(const UE::Math::TMatrix<T>& M)
 {
 	// If Matrix is NULL, return Identity quaternion. If any of them is 0, you won't be able to construct rotation
 	// if you have two plane at least, we can reconstruct the frame using cross product, but that's a bit expensive op to do here
 	// for now, if you convert to matrix from 0 scale and convert back, you'll lose rotation. Don't do that. 
 	if (M.GetScaledAxis(EAxis::X).IsNearlyZero() || M.GetScaledAxis(EAxis::Y).IsNearlyZero() || M.GetScaledAxis(EAxis::Z).IsNearlyZero())
 	{
-		*this = FQuat::Identity;
+		*this = TQuat<T>::Identity;
 		return;
 	}
 
@@ -627,26 +674,26 @@ inline FQuat::FQuat(const FMatrix& M)
 	// Changed to this (same as RemoveScaling) from RotDeterminant as using two different ways of checking unit length matrix caused inconsistency. 
 	if (!ensure((FMath::Abs(1.f - M.GetScaledAxis(EAxis::X).SizeSquared()) <= KINDA_SMALL_NUMBER) && (FMath::Abs(1.f - M.GetScaledAxis(EAxis::Y).SizeSquared()) <= KINDA_SMALL_NUMBER) && (FMath::Abs(1.f - M.GetScaledAxis(EAxis::Z).SizeSquared()) <= KINDA_SMALL_NUMBER)))
 	{
-		*this = FQuat::Identity;
+		*this = TQuat<T>::Identity;
 		return;
 	}
 #endif
 
 	//const MeReal *const t = (MeReal *) tm;
-	FMatrix::FReal	s;
+	T s;
 
 	// Check diagonal (trace)
-	const FMatrix::FReal tr = M.M[0][0] + M.M[1][1] + M.M[2][2];
+	const T tr = M.M[0][0] + M.M[1][1] + M.M[2][2];
 
 	if (tr > 0.0f) 
 	{
-		FMatrix::FReal InvS = FMath::InvSqrt(tr + 1.f);
-		this->W = float(0.5f * (1.f / InvS));
-		s = 0.5f * InvS;
+		T InvS = FMath::InvSqrt(tr + T(1.f));
+		this->W = T(T(0.5f) * (T(1.f) / InvS));
+		s = T(0.5f) * InvS;
 
-		this->X = float((M.M[1][2] - M.M[2][1]) * s);
-		this->Y = float((M.M[2][0] - M.M[0][2]) * s);
-		this->Z = float((M.M[0][1] - M.M[1][0]) * s);
+		this->X = ((M.M[1][2] - M.M[2][1]) * s);
+		this->Y = ((M.M[2][0] - M.M[0][2]) * s);
+		this->Z = ((M.M[0][1] - M.M[1][0]) * s);
 	} 
 	else 
 	{
@@ -663,50 +710,55 @@ inline FQuat::FQuat(const FMatrix& M)
 		const int32 j = nxt[i];
 		const int32 k = nxt[j];
  
-		s = M.M[i][i] - M.M[j][j] - M.M[k][k] + 1.0f;
+		s = M.M[i][i] - M.M[j][j] - M.M[k][k] + T(1.0f);
 
-		FMatrix::FReal InvS = FMath::InvSqrt(s);
+		T InvS = FMath::InvSqrt(s);
 
-		FMatrix::FReal qt[4];
-		qt[i] = 0.5f * (1.f / InvS);
+		T qt[4];
+		qt[i] = T(0.5f) * (T(1.f) / InvS);
 
-		s = 0.5f * InvS;
+		s = T(0.5f) * InvS;
 
 		qt[3] = (M.M[j][k] - M.M[k][j]) * s;
 		qt[j] = (M.M[i][j] + M.M[j][i]) * s;
 		qt[k] = (M.M[i][k] + M.M[k][i]) * s;
 
-		this->X = (float)qt[0];
-		this->Y = (float)qt[1];
-		this->Z = (float)qt[2];
-		this->W = (float)qt[3];
+		this->X = qt[0];
+		this->Y = qt[1];
+		this->Z = qt[2];
+		this->W = qt[3];
 
 		DiagnosticCheckNaN();
 	}
 }
 
 
-FORCEINLINE FQuat::FQuat(const FRotator& R)
+template<typename T>
+FORCEINLINE TQuat<T>::TQuat(const FRotator& R)
 {
-	*this = R.Quaternion();
+	FQuat Quat = R.Quaternion();
+	X = T(Quat.X);
+	Y = T(Quat.Y);
+	Z = T(Quat.Z);
+	W = T(Quat.W);
 	DiagnosticCheckNaN();
 }
 
-
-FORCEINLINE FVector FQuat::operator*(const FVector& V) const
+template<typename T>
+FORCEINLINE TVector<T> TQuat<T>::operator*(const TVector<T>& V) const
 {
 	return RotateVector(V);
 }
 
-
-inline FMatrix FQuat::operator*(const FMatrix& M) const
+template<typename T>
+inline UE::Math::TMatrix<T> TQuat<T>::operator*(const UE::Math::TMatrix<T>& M) const
 {
-	FMatrix Result;
-	FQuat VT, VR;
-	FQuat Inv = Inverse();
+	UE::Math::TMatrix<T> Result;
+	TQuat<T> VT, VR;
+	TQuat<T> Inv = Inverse();
 	for (int32 I=0; I<4; ++I)
 	{
-		FQuat VQ((float)M.M[I][0], (float)M.M[I][1], (float)M.M[I][2], (float)M.M[I][3]);
+		TQuat<T> VQ(M.M[I][0], M.M[I][1], M.M[I][2], M.M[I][3]);
 		VectorQuaternionMultiply(&VT, this, &VQ);
 		VectorQuaternionMultiply(&VR, &VT, &Inv);
 		Result.M[I][0] = VR.X;
@@ -719,46 +771,45 @@ inline FMatrix FQuat::operator*(const FMatrix& M) const
 }
 
 
-/* FQuat inline functions
+/* TQuat inline functions
  *****************************************************************************/
 
-FORCEINLINE FQuat::FQuat(EForceInit ZeroOrNot)
+template<typename T>
+FORCEINLINE TQuat<T>::TQuat()
+{}
+
+template<typename T>
+FORCEINLINE TQuat<T>::TQuat(EForceInit ZeroOrNot)
 	:	X(0), Y(0), Z(0), W(ZeroOrNot == ForceInitToZero ? 0.0f : 1.0f)
 { }
 
-FORCEINLINE FQuat::FQuat(float InX, float InY, float InZ, float InW)
+template<typename T>
+FORCEINLINE TQuat<T>::TQuat(T InX, T InY, T InZ, T InW)
 	: X(InX)
 	, Y(InY)
 	, Z(InZ)
 	, W(InW)
 {
+	static_assert(TIsFloatingPoint<T>::Value);
 	DiagnosticCheckNaN();
 }
 
-FORCEINLINE FQuat::FQuat(float V)
-	: X(V), Y(V), Z(V), W(V)
+template<typename T>
+FORCEINLINE TQuat<T>::TQuat(const QuatVectorRegister& V)
 {
+	VectorStoreAligned(V, this);
 	DiagnosticCheckNaN();
 }
 
-FORCEINLINE FQuat::FQuat(const VectorRegister4Float& V)
-{
-	VectorStoreAligned(V, &(this->X));
-	DiagnosticCheckNaN();
-}
 
-FORCEINLINE FQuat::FQuat(const VectorRegister4Double& V)
-{
-	VectorStoreAligned(V, &(this->X));
-	DiagnosticCheckNaN();
-}
-
-FORCEINLINE FString FQuat::ToString() const
+template<typename T>
+FORCEINLINE FString TQuat<T>::ToString() const
 {
 	return FString::Printf(TEXT("X=%.9f Y=%.9f Z=%.9f W=%.9f"), X, Y, Z, W);
 }
 
-inline bool FQuat::InitFromString(const FString& InSourceString)
+template<typename T>
+inline bool TQuat<T>::InitFromString(const FString& InSourceString)
 {
 	X = Y = Z = 0.f;
 	W = 1.f;
@@ -773,7 +824,8 @@ inline bool FQuat::InitFromString(const FString& InSourceString)
 }
 
 #ifdef IMPLEMENT_ASSIGNMENT_OPERATOR_MANUALLY
-FORCEINLINE FQuat& FQuat::operator=(const FQuat& Other)
+template<typename T>
+FORCEINLINE TQuat<T>& TQuat<T>::operator=(const TQuat<T>& Other)
 {
 	this->X = Other.X;
 	this->Y = Other.Y;
@@ -784,39 +836,41 @@ FORCEINLINE FQuat& FQuat::operator=(const FQuat& Other)
 }
 #endif
 
-
-FORCEINLINE FQuat::FQuat(FVector Axis, float AngleRad)
+template<typename T>
+FORCEINLINE TQuat<T>::TQuat(TVector<T> Axis, T AngleRad)
 {
-	const float half_a = 0.5f * AngleRad;
-	float s, c;
+	const T half_a = 0.5f * AngleRad;
+	T s, c;
 	FMath::SinCos(&s, &c, half_a);
 
-	X = s * (float)Axis.X;
-	Y = s * (float)Axis.Y;
-	Z = s * (float)Axis.Z;
+	X = s * Axis.X;
+	Y = s * Axis.Y;
+	Z = s * Axis.Z;
 	W = c;
 
 	DiagnosticCheckNaN();
 }
 
 
-FORCEINLINE FQuat FQuat::operator+(const FQuat& Q) const
+template<typename T>
+FORCEINLINE TQuat<T> TQuat<T>::operator+(const TQuat<T>& Q) const
 {
 #if PLATFORM_ENABLE_VECTORINTRINSICS
-	VectorRegister4Float A = VectorLoadAligned(this);
-	VectorRegister4Float B = VectorLoadAligned(&Q);
-	return FQuat(VectorAdd(A, B));
+	QuatVectorRegister A = VectorLoadAligned(this);
+	QuatVectorRegister B = VectorLoadAligned(&Q);
+	return TQuat(VectorAdd(A, B));
 #else
-	return FQuat(X + Q.X, Y + Q.Y, Z + Q.Z, W + Q.W);
+	return TQuat(X + Q.X, Y + Q.Y, Z + Q.Z, W + Q.W);
 #endif // PLATFORM_ENABLE_VECTORINTRINSICS
 }
 
 
-FORCEINLINE FQuat FQuat::operator+=(const FQuat& Q)
+template<typename T>
+FORCEINLINE TQuat<T> TQuat<T>::operator+=(const TQuat<T>& Q)
 {
 #if PLATFORM_ENABLE_VECTORINTRINSICS
-	VectorRegister4Float A = VectorLoadAligned(this);
-	VectorRegister4Float B = VectorLoadAligned(&Q);
+	QuatVectorRegister A = VectorLoadAligned(this);
+	QuatVectorRegister B = VectorLoadAligned(&Q);
 	VectorStoreAligned(VectorAdd(A, B), this);
 #else
 	this->X += Q.X;
@@ -831,31 +885,34 @@ FORCEINLINE FQuat FQuat::operator+=(const FQuat& Q)
 }
 
 
-FORCEINLINE FQuat FQuat::operator-(const FQuat& Q) const
+template<typename T>
+FORCEINLINE TQuat<T> TQuat<T>::operator-(const TQuat<T>& Q) const
 {
 #if PLATFORM_ENABLE_VECTORINTRINSICS
-	VectorRegister4Float A = VectorLoadAligned(this);
-	VectorRegister4Float B = VectorLoadAligned(&Q);
-	return FQuat(VectorSubtract(A, B));
+	QuatVectorRegister A = VectorLoadAligned(this);
+	QuatVectorRegister B = VectorLoadAligned(&Q);
+	return TQuat(VectorSubtract(A, B));
 #else
-	return FQuat(X - Q.X, Y - Q.Y, Z - Q.Z, W - Q.W);
+	return TQuat(X - Q.X, Y - Q.Y, Z - Q.Z, W - Q.W);
 #endif // PLATFORM_ENABLE_VECTORINTRINSICS
 }
 
-FORCEINLINE FQuat FQuat::operator-() const
+template<typename T>
+FORCEINLINE TQuat<T> TQuat<T>::operator-() const
 {
-	return FQuat(-X, -Y, -Z, -W);
+	return TQuat(-X, -Y, -Z, -W);
 }
 
-FORCEINLINE bool FQuat::Equals(const FQuat& Q, float Tolerance) const
+template<typename T>
+FORCEINLINE bool TQuat<T>::Equals(const TQuat<T>& Q, T Tolerance) const
 {
 #if PLATFORM_ENABLE_VECTORINTRINSICS
-	const VectorRegister4Float ToleranceV = VectorLoadFloat1(&Tolerance);
-	const VectorRegister4Float A = VectorLoadAligned(this);
-	const VectorRegister4Float B = VectorLoadAligned(&Q);
+	const QuatVectorRegister ToleranceV = VectorLoadFloat1(&Tolerance);
+	const QuatVectorRegister A = VectorLoadAligned(this);
+	const QuatVectorRegister B = VectorLoadAligned(&Q);
 
-	const VectorRegister4Float RotationSub = VectorAbs(VectorSubtract(A, B));
-	const VectorRegister4Float RotationAdd = VectorAbs(VectorAdd(A, B));
+	const QuatVectorRegister RotationSub = VectorAbs(VectorSubtract(A, B));
+	const QuatVectorRegister RotationAdd = VectorAbs(VectorAdd(A, B));
 	return !VectorAnyGreaterThan(RotationSub, ToleranceV) || !VectorAnyGreaterThan(RotationAdd, ToleranceV);
 #else
 	return (FMath::Abs(X - Q.X) <= Tolerance && FMath::Abs(Y - Q.Y) <= Tolerance && FMath::Abs(Z - Q.Z) <= Tolerance && FMath::Abs(W - Q.W) <= Tolerance)
@@ -863,16 +920,18 @@ FORCEINLINE bool FQuat::Equals(const FQuat& Q, float Tolerance) const
 #endif // PLATFORM_ENABLE_VECTORINTRINSICS
 }
 
-FORCEINLINE bool FQuat::IsIdentity(float Tolerance) const
+template<typename T>
+FORCEINLINE bool TQuat<T>::IsIdentity(T Tolerance) const
 {
-	return Equals(FQuat::Identity, Tolerance);
+	return Equals(TQuat<T>::Identity, Tolerance);
 }
 
-FORCEINLINE FQuat FQuat::operator-=(const FQuat& Q)
+template<typename T>
+FORCEINLINE TQuat<T> TQuat<T>::operator-=(const TQuat<T>& Q)
 {
 #if PLATFORM_ENABLE_VECTORINTRINSICS
-	VectorRegister4Float A = VectorLoadAligned(this);
-	VectorRegister4Float B = VectorLoadAligned(&Q);
+	QuatVectorRegister A = VectorLoadAligned(this);
+	QuatVectorRegister B = VectorLoadAligned(&Q);
 	VectorStoreAligned(VectorSubtract(A, B), this);
 #else
 	this->X -= Q.X;
@@ -887,9 +946,10 @@ FORCEINLINE FQuat FQuat::operator-=(const FQuat& Q)
 }
 
 
-FORCEINLINE FQuat FQuat::operator*(const FQuat& Q) const
+template<typename T>
+FORCEINLINE TQuat<T> TQuat<T>::operator*(const TQuat<T>& Q) const
 {
-	FQuat Result;
+	TQuat<T> Result;
 	VectorQuaternionMultiply(&Result, this, &Q);
 	
 	Result.DiagnosticCheckNaN();
@@ -898,11 +958,12 @@ FORCEINLINE FQuat FQuat::operator*(const FQuat& Q) const
 }
 
 
-FORCEINLINE FQuat FQuat::operator*=(const FQuat& Q)
+template<typename T>
+FORCEINLINE TQuat<T> TQuat<T>::operator*=(const TQuat<T>& Q)
 {
-	VectorRegister4Float A = VectorLoadAligned(this);
-	VectorRegister4Float B = VectorLoadAligned(&Q);
-	VectorRegister4Float Result;
+	QuatVectorRegister A = VectorLoadAligned(this);
+	QuatVectorRegister B = VectorLoadAligned(&Q);
+	QuatVectorRegister Result;
 	VectorQuaternionMultiply(&Result, &A, &B);
 	VectorStoreAligned(Result, this);
 
@@ -912,11 +973,12 @@ FORCEINLINE FQuat FQuat::operator*=(const FQuat& Q)
 }
 
 
-FORCEINLINE FQuat FQuat::operator*=(const float Scale)
+template<typename T>
+FORCEINLINE TQuat<T> TQuat<T>::operator*=(const T Scale)
 {
 #if PLATFORM_ENABLE_VECTORINTRINSICS
-	VectorRegister4Float A = VectorLoadAligned(this);
-	VectorRegister4Float B = VectorSetFloat1(Scale);
+	QuatVectorRegister A = VectorLoadAligned(this);
+	QuatVectorRegister B = VectorSetFloat1(Scale);
 	VectorStoreAligned(VectorMultiply(A, B), this);
 #else
 	X *= Scale;
@@ -931,30 +993,41 @@ FORCEINLINE FQuat FQuat::operator*=(const float Scale)
 }
 
 
-FORCEINLINE FQuat FQuat::operator*(const float Scale) const
+template<typename T>
+FORCEINLINE TQuat<T> TQuat<T>::operator*(const T Scale) const
 {
 #if PLATFORM_ENABLE_VECTORINTRINSICS
-	VectorRegister4Float A = VectorLoadAligned(this);
-	VectorRegister4Float B = VectorSetFloat1(Scale);
-	return FQuat(VectorMultiply(A, B));
+	QuatVectorRegister A = VectorLoadAligned(this);
+	QuatVectorRegister B = VectorSetFloat1(Scale);
+	return TQuat(VectorMultiply(A, B));
 #else
-	return FQuat(Scale * X, Scale * Y, Scale * Z, Scale * W);
+	return TQuat(Scale * X, Scale * Y, Scale * Z, Scale * W);
 #endif // PLATFORM_ENABLE_VECTORINTRINSICS
 }
 
-FORCEINLINE FQuat operator*(const float Scale, const FQuat& Q)
+// Global operator for (float * Quat)
+template<typename T>
+FORCEINLINE TQuat<T> operator*(const float Scale, const TQuat<T>& Q)
 {
 	return Q.operator*(Scale);
 }
 
-FORCEINLINE FQuat FQuat::operator/=(const float Scale)
+// Global operator for (double * Quat)
+template<typename T>
+FORCEINLINE TQuat<T> operator*(const double Scale, const TQuat<T>& Q)
+{
+	return Q.operator*(Scale);
+}
+
+template<typename T>
+FORCEINLINE TQuat<T> TQuat<T>::operator/=(const T Scale)
 {
 #if PLATFORM_ENABLE_VECTORINTRINSICS
-	VectorRegister4Float A = VectorLoadAligned(this);
-	VectorRegister4Float B = VectorSetFloat1(Scale);
+	QuatVectorRegister A = VectorLoadAligned(this);
+	QuatVectorRegister B = VectorSetFloat1(Scale);
 	VectorStoreAligned(VectorDivide(A, B), this);
 #else
-	const float Recip = 1.0f / Scale;
+	const T Recip = T(1.0f) / Scale;
 	X *= Recip;
 	Y *= Recip;
 	Z *= Recip;
@@ -967,28 +1040,31 @@ FORCEINLINE FQuat FQuat::operator/=(const float Scale)
 }
 
 
-FORCEINLINE FQuat FQuat::operator/(const float Scale) const
+template<typename T>
+FORCEINLINE TQuat<T> TQuat<T>::operator/(const T Scale) const
 {
 #if PLATFORM_ENABLE_VECTORINTRINSICS
-	VectorRegister4Float A = VectorLoadAligned(this);
-	VectorRegister4Float B = VectorSetFloat1(Scale);
-	return FQuat(VectorDivide(A, B));
+	QuatVectorRegister A = VectorLoadAligned(this);
+	QuatVectorRegister B = VectorSetFloat1(Scale);
+	return TQuat(VectorDivide(A, B));
 #else
-	const float Recip = 1.0f / Scale;
-	return FQuat(X * Recip, Y * Recip, Z * Recip, W * Recip);
+	const T Recip = 1.0f / Scale;
+	return TQuat(X * Recip, Y * Recip, Z * Recip, W * Recip);
 #endif // PLATFORM_ENABLE_VECTORINTRINSICS
 }
 
-FORCEINLINE bool FQuat::Identical(const FQuat* Q, const uint32 PortFlags) const
+template<typename T>
+FORCEINLINE bool TQuat<T>::Identical(const TQuat* Q, const uint32 PortFlags) const
 {
 	return X == Q->X && Y == Q->Y && Z == Q->Z && W == Q->W;
 }
 
-FORCEINLINE bool FQuat::operator==(const FQuat& Q) const
+template<typename T>
+FORCEINLINE bool TQuat<T>::operator==(const TQuat<T>& Q) const
 {
 #if PLATFORM_ENABLE_VECTORINTRINSICS
-	const VectorRegister4Float A = VectorLoadAligned(this);
-	const VectorRegister4Float B = VectorLoadAligned(&Q);
+	const QuatVectorRegister A = VectorLoadAligned(this);
+	const QuatVectorRegister B = VectorLoadAligned(&Q);
 	return VectorMaskBits(VectorCompareEQ(A, B)) == 0x0F;
 #else
 	return X == Q.X && Y == Q.Y && Z == Q.Z && W == Q.W;
@@ -996,11 +1072,12 @@ FORCEINLINE bool FQuat::operator==(const FQuat& Q) const
 }
 
 
-FORCEINLINE bool FQuat::operator!=(const FQuat& Q) const
+template<typename T>
+FORCEINLINE bool TQuat<T>::operator!=(const TQuat<T>& Q) const
 {
 #if PLATFORM_ENABLE_VECTORINTRINSICS
-	const VectorRegister4Float A = VectorLoadAligned(this);
-	const VectorRegister4Float B = VectorLoadAligned(&Q);
+	const QuatVectorRegister A = VectorLoadAligned(this);
+	const QuatVectorRegister B = VectorLoadAligned(&Q);
 	return VectorMaskBits(VectorCompareNE(A, B)) != 0x00;
 #else
 	return X != Q.X || Y != Q.Y || Z != Q.Z || W != Q.W;
@@ -1008,30 +1085,32 @@ FORCEINLINE bool FQuat::operator!=(const FQuat& Q) const
 }
 
 
-FORCEINLINE float FQuat::operator|(const FQuat& Q) const
+template<typename T>
+FORCEINLINE T TQuat<T>::operator|(const TQuat<T>& Q) const
 {
 	return X * Q.X + Y * Q.Y + Z * Q.Z + W * Q.W;
 }
 
 
-FORCEINLINE void FQuat::Normalize(float Tolerance)
+template<typename T>
+FORCEINLINE void TQuat<T>::Normalize(T Tolerance)
 {
 #if PLATFORM_ENABLE_VECTORINTRINSICS
-	const VectorRegister4Float Vector = VectorLoadAligned(this);
+	const QuatVectorRegister Vector = VectorLoadAligned(this);
 
-	const VectorRegister4Float SquareSum = VectorDot4(Vector, Vector);
-	const VectorRegister4Float NonZeroMask = VectorCompareGE(SquareSum, VectorLoadFloat1(&Tolerance));
-	const VectorRegister4Float InvLength = VectorReciprocalSqrtAccurate(SquareSum);
-	const VectorRegister4Float NormalizedVector = VectorMultiply(InvLength, Vector);
-	VectorRegister4Float Result = VectorSelect(NonZeroMask, NormalizedVector, GlobalVectorConstants::Float0001);
+	const QuatVectorRegister SquareSum = VectorDot4(Vector, Vector);
+	const QuatVectorRegister NonZeroMask = VectorCompareGE(SquareSum, VectorLoadFloat1(&Tolerance));
+	const QuatVectorRegister InvLength = VectorReciprocalSqrtAccurate(SquareSum);
+	const QuatVectorRegister NormalizedVector = VectorMultiply(InvLength, Vector);
+	QuatVectorRegister Result = VectorSelect(NonZeroMask, NormalizedVector, GlobalVectorConstants::Float0001);
 
 	VectorStoreAligned(Result, this);
 #else
-	const float SquareSum = X * X + Y * Y + Z * Z + W * W;
+	const T SquareSum = X * X + Y * Y + Z * Z + W * W;
 
 	if (SquareSum >= Tolerance)
 	{
-		const float Scale = FMath::InvSqrt(SquareSum);
+		const T Scale = FMath::InvSqrt(SquareSum);
 
 		X *= Scale; 
 		Y *= Scale; 
@@ -1040,26 +1119,27 @@ FORCEINLINE void FQuat::Normalize(float Tolerance)
 	}
 	else
 	{
-		*this = FQuat::Identity;
+		*this = TQuat<T>::Identity;
 	}
 #endif // PLATFORM_ENABLE_VECTORINTRINSICS
 }
 
 
-FORCEINLINE FQuat FQuat::GetNormalized(float Tolerance) const
+template<typename T>
+FORCEINLINE TQuat<T> TQuat<T>::GetNormalized(T Tolerance) const
 {
-	FQuat Result(*this);
+	TQuat<T> Result(*this);
 	Result.Normalize(Tolerance);
 	return Result;
 }
 
 
-
-FORCEINLINE bool FQuat::IsNormalized() const
+template<typename T>
+FORCEINLINE bool TQuat<T>::IsNormalized() const
 {
 #if PLATFORM_ENABLE_VECTORINTRINSICS
-	VectorRegister4Float A = VectorLoadAligned(this);
-	VectorRegister4Float TestValue = VectorAbs(VectorSubtract(VectorOne(), VectorDot4(A, A)));
+	QuatVectorRegister A = VectorLoadAligned(this);
+	QuatVectorRegister TestValue = VectorAbs(VectorSubtract(VectorOne(), VectorDot4(A, A)));
 	return !VectorAnyGreaterThan(TestValue, GlobalVectorConstants::ThreshQuatNormalized);
 #else
 	return (FMath::Abs(1.f - SizeSquared()) < THRESH_QUAT_NORMALIZED);
@@ -1067,56 +1147,69 @@ FORCEINLINE bool FQuat::IsNormalized() const
 }
 
 
-FORCEINLINE float FQuat::Size() const
+template<typename T>
+FORCEINLINE T TQuat<T>::Size() const
 {
 	return FMath::Sqrt(X * X + Y * Y + Z * Z + W * W);
 }
 
-
-FORCEINLINE float FQuat::SizeSquared() const
+template<typename T>
+FORCEINLINE T TQuat<T>::SizeSquared() const
 {
 	return (X * X + Y * Y + Z * Z + W * W);
 }
 
-FORCEINLINE float FQuat::GetAngle() const
+template<typename T>
+FORCEINLINE T TQuat<T>::GetAngle() const
 {
-	return 2.f * FMath::Acos(W);
+	return T(2.0) * FMath::Acos(W);
 }
 
 
-FORCEINLINE void FQuat::ToAxisAndAngle(FVector& Axis, float& Angle) const
+template<typename T>
+FORCEINLINE void TQuat<T>::ToAxisAndAngle(TVector<T>& Axis, float& Angle) const
+{
+	Angle = (float)GetAngle();
+	Axis = GetRotationAxis();
+}
+
+template<typename T>
+FORCEINLINE void TQuat<T>::ToAxisAndAngle(TVector<T>& Axis, double& Angle) const
 {
-	Angle = GetAngle();
+	Angle = (double)GetAngle();
 	Axis = GetRotationAxis();
 }
 
-FORCEINLINE FVector FQuat::GetRotationAxis() const
+template<typename T>
+FORCEINLINE TVector<T> TQuat<T>::GetRotationAxis() const
 {
 #if PLATFORM_ENABLE_VECTORINTRINSICS
-	FVector V;
-	VectorRegister4Float A = VectorLoadAligned(this);
-	VectorRegister4Float R = VectorNormalizeSafe(VectorSet_W0(A), GlobalVectorConstants::Float1000);
+	TVector<T> V;
+	QuatVectorRegister A = VectorLoadAligned(this);
+	QuatVectorRegister R = VectorNormalizeSafe(VectorSet_W0(A), GlobalVectorConstants::Float1000);
 	VectorStoreFloat3(R, &V);
 	return V;
 #else
-	const float SquareSum = X * X + Y * Y + Z * Z;
+	const T SquareSum = X * X + Y * Y + Z * Z;
 	if (SquareSum < SMALL_NUMBER)
 	{
-		return FVector::XAxisVector;
+		return TVector<T>::XAxisVector;
 	}
-	const float Scale = FMath::InvSqrt(SquareSum);
-	return FVector(X * Scale, Y * Scale, Z * Scale);
+	const T Scale = FMath::InvSqrt(SquareSum);
+	return TVector<T>(X * Scale, Y * Scale, Z * Scale);
 #endif // PLATFORM_ENABLE_VECTORINTRINSICS
 }
 
-float FQuat::AngularDistance(const FQuat& Q) const
+template<typename T>
+T TQuat<T>::AngularDistance(const TQuat<T>& Q) const
 {
-	float InnerProd = X*Q.X + Y*Q.Y + Z*Q.Z + W*Q.W;
+	T InnerProd = X*Q.X + Y*Q.Y + Z*Q.Z + W*Q.W;
 	return FMath::Acos((2 * InnerProd * InnerProd) - 1.f);
 }
 
 
-FORCEINLINE FVector FQuat::RotateVector(FVector V) const
+template<typename T>
+FORCEINLINE TVector<T> TQuat<T>::RotateVector(TVector<T> V) const
 {	
 	// http://people.csail.mit.edu/bkph/articles/Quaternions.pdf
 	// V' = V + 2w(Q x V) + (2Q x (Q x V))
@@ -1125,37 +1218,39 @@ FORCEINLINE FVector FQuat::RotateVector(FVector V) const
 	// T = 2(Q x V);
 	// V' = V + w*(T) + (Q x T)
 
-	const FVector Q(X, Y, Z);
-	const FVector T = 2.f * FVector::CrossProduct(Q, V);
-	const FVector Result = V + (W * T) + FVector::CrossProduct(Q, T);
+	const TVector<T> Q(X, Y, Z);
+	const TVector<T> TT = 2.f * TVector<T>::CrossProduct(Q, V);
+	const TVector<T> Result = V + (W * TT) + TVector<T>::CrossProduct(Q, TT);
 	return Result;
 }
 
-FORCEINLINE FVector FQuat::UnrotateVector(FVector V) const
+template<typename T>
+FORCEINLINE TVector<T> TQuat<T>::UnrotateVector(TVector<T> V) const
 {	
-	const FVector Q(-X, -Y, -Z); // Inverse
-	const FVector T = 2.f * FVector::CrossProduct(Q, V);
-	const FVector Result = V + (W * T) + FVector::CrossProduct(Q, T);
+	const TVector<T> Q(-X, -Y, -Z); // Inverse
+	const TVector<T> TT = 2.f * TVector<T>::CrossProduct(Q, V);
+	const TVector<T> Result = V + (W * TT) + TVector<T>::CrossProduct(Q, TT);
 	return Result;
 }
 
 
-FORCEINLINE FQuat FQuat::Inverse() const
+template<typename T>
+FORCEINLINE TQuat<T> TQuat<T>::Inverse() const
 {
 	checkSlow(IsNormalized());
 
 #if PLATFORM_ENABLE_VECTORINTRINSICS
-	return FQuat(VectorQuaternionInverse(VectorLoadAligned(this)));
+	return TQuat(VectorQuaternionInverse(VectorLoadAligned(this)));
 #else
-	return FQuat(-X, -Y, -Z, W);
+	return TQuat(-X, -Y, -Z, W);
 #endif // PLATFORM_ENABLE_VECTORINTRINSICS
 }
 
-
-FORCEINLINE void FQuat::EnforceShortestArcWith(const FQuat& OtherQuat)
+template<typename T>
+FORCEINLINE void TQuat<T>::EnforceShortestArcWith(const TQuat<T>& OtherQuat)
 {
-	const float DotResult = (OtherQuat | *this);
-	const float Bias = FMath::FloatSelect(DotResult, 1.0f, -1.0f);
+	const T DotResult = (OtherQuat | *this);
+	const T Bias = FMath::FloatSelect(DotResult, T(1.0f), T(-1.0f));
 
 	X *= Bias;
 	Y *= Bias;
@@ -1163,88 +1258,100 @@ FORCEINLINE void FQuat::EnforceShortestArcWith(const FQuat& OtherQuat)
 	W *= Bias;
 }
 
-	
-FORCEINLINE FVector FQuat::GetAxisX() const
+
+template<typename T>
+FORCEINLINE TVector<T> TQuat<T>::GetAxisX() const
 {
-	return RotateVector(FVector(1.f, 0.f, 0.f));
+	return RotateVector(TVector<T>(1.f, 0.f, 0.f));
 }
 
 
-FORCEINLINE FVector FQuat::GetAxisY() const
+template<typename T>
+FORCEINLINE TVector<T> TQuat<T>::GetAxisY() const
 {
-	return RotateVector(FVector(0.f, 1.f, 0.f));
+	return RotateVector(TVector<T>(0.f, 1.f, 0.f));
 }
 
 
-FORCEINLINE FVector FQuat::GetAxisZ() const
+template<typename T>
+FORCEINLINE TVector<T> TQuat<T>::GetAxisZ() const
 {
-	return RotateVector(FVector(0.f, 0.f, 1.f));
+	return RotateVector(TVector<T>(0.f, 0.f, 1.f));
 }
 
 
-FORCEINLINE FVector FQuat::GetForwardVector() const
+template<typename T>
+FORCEINLINE TVector<T> TQuat<T>::GetForwardVector() const
 {
 	return GetAxisX();
 }
 
-FORCEINLINE FVector FQuat::GetRightVector() const
+template<typename T>
+FORCEINLINE TVector<T> TQuat<T>::GetRightVector() const
 {
 	return GetAxisY();
 }
 
-FORCEINLINE FVector FQuat::GetUpVector() const
+template<typename T>
+FORCEINLINE TVector<T> TQuat<T>::GetUpVector() const
 {
 	return GetAxisZ();
 }
 
-FORCEINLINE FVector FQuat::Vector() const
+template<typename T>
+FORCEINLINE TVector<T> TQuat<T>::Vector() const
 {
 	return GetAxisX();
 }
 
 
-FORCEINLINE float FQuat::Error(const FQuat& Q1, const FQuat& Q2)
+template<typename T>
+FORCEINLINE T TQuat<T>::Error(const TQuat<T>& Q1, const TQuat<T>& Q2)
 {
-	const float cosom = FMath::Abs(Q1.X * Q2.X + Q1.Y * Q2.Y + Q1.Z * Q2.Z + Q1.W * Q2.W);
+	const T cosom = FMath::Abs(Q1.X * Q2.X + Q1.Y * Q2.Y + Q1.Z * Q2.Z + Q1.W * Q2.W);
 	return (FMath::Abs(cosom) < 0.9999999f) ? FMath::Acos(cosom)*(1.f / PI) : 0.0f;
 }
 
 
-FORCEINLINE float FQuat::ErrorAutoNormalize(const FQuat& A, const FQuat& B)
+template<typename T>
+FORCEINLINE T TQuat<T>::ErrorAutoNormalize(const TQuat<T>& A, const TQuat<T>& B)
 {
-	FQuat Q1 = A;
+	TQuat<T> Q1 = A;
 	Q1.Normalize();
 
-	FQuat Q2 = B;
+	TQuat<T> Q2 = B;
 	Q2.Normalize();
 
-	return FQuat::Error(Q1, Q2);
+	return TQuat<T>::Error(Q1, Q2);
 }
 
 /**
  * Fast Linear Quaternion Interpolation.
  * Result is NOT normalized.
  */
-FORCEINLINE FQuat FQuat::FastLerp(const FQuat& A, const FQuat& B, const float Alpha)
+template<typename T>
+FORCEINLINE TQuat<T> TQuat<T>::FastLerp(const TQuat<T>& A, const TQuat<T>& B, const T Alpha)
 {
 	// To ensure the 'shortest route', we make sure the dot product between the both rotations is positive.
-	const float DotResult = (A | B);
-	const float Bias = FMath::FloatSelect(DotResult, 1.0f, -1.0f);
+	const T DotResult = (A | B);
+	const T Bias = FMath::FloatSelect(DotResult, T(1.0f), T(-1.0f));
 	return (B * Alpha) + (A * (Bias * (1.f - Alpha)));
 }
 
 
-FORCEINLINE FQuat FQuat::FastBilerp(const FQuat& P00, const FQuat& P10, const FQuat& P01, const FQuat& P11, float FracX, float FracY)
+template<typename T>
+FORCEINLINE TQuat<T> TQuat<T>::FastBilerp(const TQuat<T>& P00, const TQuat<T>& P10, const TQuat<T>& P01, const TQuat<T>& P11, T FracX, T FracY)
 {
-	return FQuat::FastLerp(
-		FQuat::FastLerp(P00,P10,FracX),
-		FQuat::FastLerp(P01,P11,FracX),
+	return TQuat<T>::FastLerp(
+		TQuat<T>::FastLerp(P00,P10,FracX),
+		TQuat<T>::FastLerp(P01,P11,FracX),
 		FracY
 	);
 }
 
 
-FORCEINLINE bool FQuat::ContainsNaN() const
+template<typename T>
+FORCEINLINE bool TQuat<T>::ContainsNaN() const
 {
 	return (!FMath::IsFinite(X) ||
 			!FMath::IsFinite(Y) ||
@@ -1254,35 +1361,58 @@ FORCEINLINE bool FQuat::ContainsNaN() const
 }
 
 
-template<> struct TIsPODType<FQuat> { enum { Value = true }; };
-template<> struct TIsUECoreType<FQuat> { enum { Value = true }; };
+} // namespace UE::Math
+} // namespace UE
+
+DECLARE_LWC_TYPE(Quat, 4);
+
+template<> struct TIsPODType<FQuat4f> { enum { Value = true }; };
+template<> struct TIsPODType<FQuat4d> { enum { Value = true }; };
+template<> struct TIsUECoreType<FQuat4f> { enum { Value = true }; };
+template<> struct TIsUECoreType<FQuat4d> { enum { Value = true }; };
+template<> struct TCanBulkSerialize<FQuat4f> { enum { Value = true }; };
+template<> struct TCanBulkSerialize<FQuat4d> { enum { Value = false }; }; // LWC_TODO: This can be done (via versioning) once LWC is fixed to on.
+DECLARE_INTRINSIC_TYPE_LAYOUT(FQuat4f);
+DECLARE_INTRINSIC_TYPE_LAYOUT(FQuat4d);
+
+// ispc doesn't export typedefs in generated headers, so we do it here to keep our code happy.
+DECLARE_LWC_TYPE_ISPC(Quat, 4);
+
 
 /* FMath inline functions
  *****************************************************************************/
-
-template<class U>
-FORCEINLINE_DEBUGGABLE FQuat FMath::Lerp( const FQuat& A, const FQuat& B, const U& Alpha)
+ 
+ // TCustomLerp for FMath::Lerp()
+template<class T>
+struct TCustomLerp< UE::Math::TQuat<T> >
 {
-	return FQuat::Slerp(A, B, Alpha);
-}
+	enum { IsRequired = true };
+	using QuatType = UE::Math::TQuat<T>;
 
-template<class U>
-FORCEINLINE_DEBUGGABLE FQuat FMath::BiLerp(const FQuat& P00, const FQuat& P10, const FQuat& P01, const FQuat& P11, float FracX, float FracY)
-{
-	FQuat Result;
+	template<class U>
+	static FORCEINLINE_DEBUGGABLE QuatType Lerp(const QuatType& A, const QuatType& B, const U& Alpha)
+	{
+		return QuatType::Slerp(A, B, (T)Alpha);
+	}
 
-	Result = Lerp(
-		FQuat::Slerp_NotNormalized(P00,P10,FracX),
-		FQuat::Slerp_NotNormalized(P01,P11,FracX),
-		FracY
+	template<class U>
+	static FORCEINLINE_DEBUGGABLE QuatType BiLerp(const QuatType& P00, const QuatType& P10, const QuatType& P01, const QuatType& P11, const U& FracX, const U& FracY)
+	{
+		QuatType Result;
+
+		Result = Lerp(
+			QuatType::Slerp_NotNormalized(P00, P10, (T)FracX),
+			QuatType::Slerp_NotNormalized(P01, P11, (T)FracX),
+			(T)FracY
 		);
 
-	return Result;
-}
+		return Result;
+	}
 
-template<class U>
-FORCEINLINE_DEBUGGABLE FQuat FMath::CubicInterp( const FQuat& P0, const FQuat& T0, const FQuat& P1, const FQuat& T1, const U& A)
-{
-	return FQuat::Squad(P0, T0, P1, T1, A);
-}
+	template<class U>
+	static FORCEINLINE_DEBUGGABLE QuatType CubicInterp(const QuatType& P0, const QuatType& T0, const QuatType& P1, const QuatType& T1, const U& A)
+	{
+		return QuatType::Squad(P0, T0, P1, T1, (T)A);
+	}
 
+};
diff --git a/Engine/Source/Runtime/Core/Public/Math/QuatRotationTranslationMatrix.h b/Engine/Source/Runtime/Core/Public/Math/QuatRotationTranslationMatrix.h
--- a/Engine/Source/Runtime/Core/Public/Math/QuatRotationTranslationMatrix.h
+++ b/Engine/Source/Runtime/Core/Public/Math/QuatRotationTranslationMatrix.h
@@ -11,8 +11,6 @@
 namespace UE {
 namespace Math {
 
-// LWC_TODO: FQuat -> TQuat<T>
-	
 /** Rotation and translation matrix using quaternion rotation */
 template <typename T>
 struct TQuatRotationTranslationMatrix
@@ -26,10 +24,10 @@ public:
 	* @param Q rotation
 	* @param Origin translation to apply
 	*/
-	TQuatRotationTranslationMatrix(const FQuat& Q, const TVector<T>& Origin);
+	TQuatRotationTranslationMatrix(const TQuat<T>& Q, const TVector<T>& Origin);
 
 	/** Matrix factory. Return an FMatrix so we don't have type conversion issues in expressions. */
-	static TMatrix<T> Make(const FQuat& Q, const TVector<T>& Origin)
+	static TMatrix<T> Make(const TQuat<T>& Q, const TVector<T>& Origin)
 	{
 		return TQuatRotationTranslationMatrix<T>(Q, Origin);
 	}
@@ -48,13 +46,13 @@ public:
 	*
 	* @param Q rotation
 	*/
-	TQuatRotationMatrix(const FQuat& Q)
+	TQuatRotationMatrix(const TQuat<T>& Q)
 		: TQuatRotationTranslationMatrix<T>(Q, TVector<T>::ZeroVector)
 	{
 	}
 
 	/** Matrix factory. Return an FMatrix so we don't have type conversion issues in expressions. */
-	static TMatrix<T> Make(const FQuat& Q)
+	static TMatrix<T> Make(const TQuat<T>& Q)
 	{
 		return TQuatRotationMatrix<T>(Q);
 	}
@@ -62,7 +60,7 @@ public:
 
 
 template<typename T>
-FORCEINLINE TQuatRotationTranslationMatrix<T>::TQuatRotationTranslationMatrix(const FQuat& Q, const TVector<T>& Origin)
+FORCEINLINE TQuatRotationTranslationMatrix<T>::TQuatRotationTranslationMatrix(const TQuat<T>& Q, const TVector<T>& Origin)
 {
 #if !(UE_BUILD_SHIPPING || UE_BUILD_TEST) && WITH_EDITORONLY_DATA
 	// Make sure Quaternion is normalized
diff --git a/Engine/Source/Runtime/Core/Public/Math/Rotator.h b/Engine/Source/Runtime/Core/Public/Math/Rotator.h
--- a/Engine/Source/Runtime/Core/Public/Math/Rotator.h
+++ b/Engine/Source/Runtime/Core/Public/Math/Rotator.h
@@ -30,7 +30,7 @@ public:
 public:
 
 	/** A rotator of zero degrees on each axis. */
-	static CORE_API const FRotator ZeroRotator;
+	static const FRotator ZeroRotator;
 
 public:
 
@@ -90,7 +90,8 @@ public:
 	 *
 	 * @param Quat Quaternion used to specify rotation.
 	 */
-	explicit CORE_API FRotator( const FQuat& Quat );
+	explicit CORE_API FRotator( const FQuat4f& Quat );
+	explicit CORE_API FRotator( const FQuat4d& Quat);
 
 public:
 
@@ -446,6 +447,9 @@ public:
 };
 
 
+inline const FRotator FRotator::ZeroRotator(0.f, 0.f, 0.f);
+
+
 /* FRotator inline functions
  *****************************************************************************/
 
@@ -804,19 +808,32 @@ FORCEINLINE void FRotator::SetClosestToMe(FRotator& MakeClosest) const
 template<> struct TIsPODType<FRotator> { enum { Value = true }; };
 template<> struct TIsUECoreType<FRotator> { enum { Value = true }; };
 
-
 /* FMath inline functions
  *****************************************************************************/
 
-template<class U>
-FORCEINLINE_DEBUGGABLE FRotator FMath::Lerp(const FRotator& A, const FRotator& B, const U& Alpha)
+template<> struct TCustomLerp<FRotator>
 {
-	return A + (B - A).GetNormalized() * Alpha;
-}
+	// Required to make FMath::Lerp<FRotator>() call our custom Lerp() implementation below.
+	enum { IsRequired = true };
+
+	static FORCEINLINE_DEBUGGABLE FRotator Lerp(const FRotator& A, const FRotator& B, const float& Alpha)
+	{
+		return A + (B - A).GetNormalized() * Alpha;
+	}
 
-template<class U>
-FORCEINLINE_DEBUGGABLE FRotator FMath::LerpRange(const FRotator& A, const FRotator& B, const U& Alpha)
+	static FORCEINLINE_DEBUGGABLE FRotator Lerp(const FRotator& A, const FRotator& B, const double& Alpha)
+	{
+		return Lerp(A, B, (float)Alpha);
+	}
+};
+
+FORCEINLINE_DEBUGGABLE FRotator FMath::LerpRange(const FRotator& A, const FRotator& B, float Alpha)
 {
 	// Similar to Lerp, but does not take the shortest path. Allows interpolation over more than 180 degrees.
-	return (A * (1 - Alpha) + B * Alpha).GetNormalized();
+	return (A * (1.f - Alpha) + B * Alpha).GetNormalized();
+}
+
+FORCEINLINE_DEBUGGABLE FRotator FMath::LerpRange(const FRotator& A, const FRotator& B, double Alpha)
+{
+	return FMath::LerpRange(A, B, (float)Alpha);
 }
diff --git a/Engine/Source/Runtime/Core/Public/Math/ScalarRegister.h b/Engine/Source/Runtime/Core/Public/Math/ScalarRegister.h
--- a/Engine/Source/Runtime/Core/Public/Math/ScalarRegister.h
+++ b/Engine/Source/Runtime/Core/Public/Math/ScalarRegister.h
@@ -25,7 +25,9 @@
 class ScalarRegister
 {
 public:
-	VectorRegister Value;
+	using RegisterType = VectorRegister4Float;
+
+	RegisterType Value;
 
 	/** default constructor */
 	FORCEINLINE ScalarRegister();
@@ -38,9 +40,9 @@ public:
 
 	/** Constructor
 	 *
-	 * @param VectorRegister float4 vector register type
+	 * @param RegisterType float4 vector register type
 	 */
-	explicit FORCEINLINE ScalarRegister(VectorRegister VectorValue);
+	explicit FORCEINLINE ScalarRegister(RegisterType VectorValue);
 
 	/**
 	 * Gets the result of multiplying a scalar register to this.
@@ -96,14 +98,14 @@ public:
 
 	/** assignment operator
 	 *
-	 * @param RHS a VectorRegister
+	 * @param RHS a RegisterType
 	 */
-	FORCEINLINE ScalarRegister& operator=(const VectorRegister& RHS);
+	FORCEINLINE ScalarRegister& operator=(const RegisterType& RHS);
 
 	/**
-	 * ScalarRegister to VectorRegister conversion operator.
+	 * ScalarRegister to RegisterType conversion operator.
 	 */
-	FORCEINLINE operator VectorRegister() const;
+	FORCEINLINE operator RegisterType() const;
 };
 
 FORCEINLINE ScalarRegister::ScalarRegister()
@@ -120,7 +122,7 @@ FORCEINLINE ScalarRegister::ScalarRegister(const float& ScalarValue)
 	Value = VectorLoadFloat1(&ScalarValue);
 }
 
-FORCEINLINE ScalarRegister::ScalarRegister(VectorRegister VectorValue)
+FORCEINLINE ScalarRegister::ScalarRegister(RegisterType VectorValue)
 {
 	Value = VectorValue;
 }
@@ -159,13 +161,13 @@ FORCEINLINE ScalarRegister& ScalarRegister::operator=(const ScalarRegister& RHS)
 	return *this;
 }
 
-FORCEINLINE ScalarRegister& ScalarRegister::operator=(const VectorRegister& RHS)
+FORCEINLINE ScalarRegister& ScalarRegister::operator=(const RegisterType& RHS)
 {
 	Value = RHS;
 	return *this;
 }
 
-FORCEINLINE ScalarRegister::operator VectorRegister() const
+FORCEINLINE ScalarRegister::operator RegisterType() const
 {
 	return Value;
 }
@@ -189,13 +191,18 @@ FORCEINLINE ScalarRegister ScalarMax(const ScalarRegister& A, const ScalarRegist
 	return ScalarRegister(VectorMax(A.Value, B.Value));
 }
 
-// Specialization of Lerp template that works with scalar (float in vector) registers
-template<> FORCEINLINE ScalarRegister FMath::Lerp(const ScalarRegister& A, const ScalarRegister& B, const ScalarRegister& Alpha)
+// TCustomLerp for FMath::Lerp()
+template<> struct TCustomLerp<ScalarRegister>
 {
-	const VectorRegister Delta = VectorSubtract(B.Value, A.Value);
-	return ScalarRegister(VectorMultiplyAdd(Alpha.Value, Delta, A.Value));
-}
+	enum { IsRequired = true };
 
+	// Specialization of Lerp template that works with scalar (float in vector) registers
+	FORCEINLINE ScalarRegister Lerp(const ScalarRegister& A, const ScalarRegister& B, const ScalarRegister& Alpha)
+	{
+		const ScalarRegister::RegisterType Delta = VectorSubtract(B.Value, A.Value);
+		return ScalarRegister(VectorMultiplyAdd(Alpha.Value, Delta, A.Value));
+	}
+};
 
 /**
  * Computes the reciprocal of the scalar register (component-wise) and returns the result.
diff --git a/Engine/Source/Runtime/Core/Public/Math/Transform.h b/Engine/Source/Runtime/Core/Public/Math/Transform.h
--- a/Engine/Source/Runtime/Core/Public/Math/Transform.h
+++ b/Engine/Source/Runtime/Core/Public/Math/Transform.h
@@ -11,3 +11,17 @@
 #include "Math/ScalarRegister.h"
 #include "Math/TransformVectorized.h"
 #include "Math/TransformNonVectorized.h"
+
+DECLARE_LWC_TYPE(Transform, 3);
+
+template<> struct TIsPODType<FTransform3f> { enum { Value = true }; };
+template<> struct TIsPODType<FTransform3d> { enum { Value = true }; };
+template<> struct TIsUECoreType<FTransform3f> { enum { Value = true }; };
+template<> struct TIsUECoreType<FTransform3d> { enum { Value = true }; };
+template<> struct TCanBulkSerialize<FTransform3f> { enum { Value = false }; }; // LWC_TODO: this was never defined before, should it be true or false?
+template<> struct TCanBulkSerialize<FTransform3d> { enum { Value = false }; }; // LWC_TODO: This can be done (via versioning) once LWC is fixed to on.
+DECLARE_INTRINSIC_TYPE_LAYOUT(FTransform3f);
+DECLARE_INTRINSIC_TYPE_LAYOUT(FTransform3d);
+
+// ispc doesn't export typedefs in generated headers, so we do it here to keep our code happy.
+DECLARE_LWC_TYPE_ISPC(Transform, 3);
diff --git a/Engine/Source/Runtime/Core/Public/Math/TransformCalculus3D.h b/Engine/Source/Runtime/Core/Public/Math/TransformCalculus3D.h
--- a/Engine/Source/Runtime/Core/Public/Math/TransformCalculus3D.h
+++ b/Engine/Source/Runtime/Core/Public/Math/TransformCalculus3D.h
@@ -235,8 +235,6 @@ inline FVector Concatenate(const FVector& LHS, const FVector& RHS)
 	return LHS + RHS;
 }
 	
-} // namespace UE::Core 
-} // namespace UE
 
 /**
  * Specialization for concatenating two rotations.
@@ -253,6 +251,9 @@ inline FQuat Concatenate(const FQuat& LHS, const FQuat& RHS)
 	return RHS * LHS;
 }
 
+} // namespace UE::Math
+} // namespace UE
+
 /**
  * Specialization for concatenating two rotations.
  *
