LWC support for FQuat and FTransform

diff --git a/Engine/Source/Runtime/Core/Public/Math/TransformNonVectorized.h b/Engine/Source/Runtime/Core/Public/Math/TransformNonVectorized.h
--- a/Engine/Source/Runtime/Core/Public/Math/TransformNonVectorized.h
+++ b/Engine/Source/Runtime/Core/Public/Math/TransformNonVectorized.h
@@ -20,35 +20,49 @@
 * Transformation of direction vectors is applied in the order: Scale -> Rotate.
 *
 * Order matters when composing transforms: C = A * B will yield a transform C that logically
-* first applies A then B to any subsequent transformation. Note that this is the opposite order of quaternion (FQuat) multiplication.
+* first applies A then B to any subsequent transformation. Note that this is the opposite order of quaternion (TQuat<T>) multiplication.
 *
 * Example: LocalToWorld = (DeltaRotation * LocalToWorld) will change rotation in local space by DeltaRotation.
 * Example: LocalToWorld = (LocalToWorld * DeltaRotation) will change rotation in world space by DeltaRotation.
 */
-struct FTransform
+struct Z_Construct_UScriptStruct_FTransform3f_Statics;
+struct Z_Construct_UScriptStruct_FTransform3d_Statics;
+struct Z_Construct_UScriptStruct_FTransform_Statics;
+
+namespace UE
+{
+namespace Math
+{
+
+template<typename T>
+struct alignas(alignof(TQuat<T>)) TTransform
 {
-	friend struct Z_Construct_UScriptStruct_FTransform_Statics;
+	friend Z_Construct_UScriptStruct_FTransform3f_Statics;
+	friend Z_Construct_UScriptStruct_FTransform3d_Statics;
+	friend Z_Construct_UScriptStruct_FTransform_Statics;
+
+	using FReal = T;
 
 protected:
 	/** Rotation of this transformation, as a quaternion. */
-	FQuat	Rotation;
+	TQuat<T>	Rotation;
 	/** Translation of this transformation, as a vector. */
-	FVector	Translation;
+	TVector<T>	Translation;
 	/** 3D scale (always applied in local space) as a vector. */
-	FVector	Scale3D;
+	TVector<T>	Scale3D;
 public:
 	/**
-	* The identity transformation (Rotation = FQuat::Identity, Translation = FVector::ZeroVector, Scale3D = (1,1,1)).
+	* The identity transformation (Rotation = TQuat<T>::Identity, Translation = TVector<T>::ZeroVector, Scale3D = (1,1,1)).
 	*/
-	static CORE_API const FTransform Identity;
+	static const TTransform<T> Identity;
 
 #if ENABLE_NAN_DIAGNOSTIC
 	FORCEINLINE void DiagnosticCheckNaN_Scale3D() const
 	{
 		if (Scale3D.ContainsNaN())
 		{
-			logOrEnsureNanError(TEXT("FTransform Scale3D contains NaN: %s"), *Scale3D.ToString());
-			const_cast<FTransform*>(this)->Scale3D = FVector::OneVector;
+			logOrEnsureNanError(TEXT("TTransform<T> Scale3D contains NaN: %s"), *Scale3D.ToString());
+			const_cast<TTransform<T>*>(this)->Scale3D = TVector<T>::OneVector;
 		}
 	}
 
@@ -56,8 +70,8 @@ public:
 	{
 		if (Translation.ContainsNaN())
 		{
-			logOrEnsureNanError(TEXT("FTransform Translation contains NaN: %s"), *Translation.ToString());
-			const_cast<FTransform*>(this)->Translation = FVector::ZeroVector;
+			logOrEnsureNanError(TEXT("TTransform<T> Translation contains NaN: %s"), *Translation.ToString());
+			const_cast<TTransform<T>*>(this)->Translation = TVector<T>::ZeroVector;
 		}
 	}
 
@@ -65,8 +79,8 @@ public:
 	{
 		if (Rotation.ContainsNaN())
 		{
-			logOrEnsureNanError(TEXT("FTransform Rotation contains NaN: %s"), *Rotation.ToString());
-			const_cast<FTransform*>(this)->Rotation = FQuat::Identity;
+			logOrEnsureNanError(TEXT("TTransform<T> Rotation contains NaN: %s"), *Rotation.ToString());
+			const_cast<TTransform<T>*>(this)->Rotation = TQuat<T>::Identity;
 		}
 	}
 
@@ -82,7 +96,7 @@ public:
 		DiagnosticCheckNaN_All();
 		if (!IsValid())
 		{
-			logOrEnsureNanError(TEXT("FTransform transform is not valid: %s"), *ToHumanReadableString());
+			logOrEnsureNanError(TEXT("TTransform<T> transform is not valid: %s"), *ToHumanReadableString());
 		}
 
 	}
@@ -95,10 +109,10 @@ public:
 #endif
 
 	/** Default constructor. */
-	FORCEINLINE FTransform()
+	FORCEINLINE TTransform()
 		: Rotation(0.f, 0.f, 0.f, 1.f)
 		, Translation(0.f)
-		, Scale3D(FVector::OneVector)
+		, Scale3D(TVector<T>::OneVector)
 	{
 	}
 
@@ -107,10 +121,10 @@ public:
 	*
 	* @param InTranslation The value to use for the translation component
 	*/
-	FORCEINLINE explicit FTransform(const FVector& InTranslation)
-		: Rotation(FQuat::Identity),
+	FORCEINLINE explicit TTransform(const TVector<T>& InTranslation)
+		: Rotation(TQuat<T>::Identity),
 		Translation(InTranslation),
-		Scale3D(FVector::OneVector)
+		Scale3D(TVector<T>::OneVector)
 	{
 		DiagnosticCheckNaN_All();
 	}
@@ -118,13 +132,13 @@ public:
 	/**
 	* Constructor with leaving uninitialized memory
 	*/
-	FORCEINLINE explicit FTransform(ENoInit)
+	FORCEINLINE explicit TTransform(ENoInit)
 	{
 		// Note: This can be used to track down initialization issues with bone transform arrays; but it will
 		// cause issues with transient fields such as RootMotionDelta that get initialized to 0 by default
 #if ENABLE_NAN_DIAGNOSTIC
 		// It is unsafe to call normal constructors here as they do their own NaN checks, so would always hit a false positive
-		float qnan = FMath::Log2(-5.3f);
+		T qnan = FMath::Log2(-5.3f);
 		check(FMath::IsNaN(qnan));
 		Translation.X = Translation.Y = Translation.Z = qnan;
 		Rotation.X = Rotation.Y = Rotation.Z = Rotation.W = qnan;
@@ -137,10 +151,10 @@ public:
 	*
 	* @param InRotation The value to use for rotation component
 	*/
-	FORCEINLINE explicit FTransform(const FQuat& InRotation)
+	FORCEINLINE explicit TTransform(const TQuat<T>& InRotation)
 		: Rotation(InRotation),
-		Translation(FVector::ZeroVector),
-		Scale3D(FVector::OneVector)
+		Translation(TVector<T>::ZeroVector),
+		Scale3D(TVector<T>::OneVector)
 	{
 		DiagnosticCheckNaN_All();
 	}
@@ -150,10 +164,10 @@ public:
 	*
 	* @param InRotation The value to use for rotation component  (after being converted to a quaternion)
 	*/
-	FORCEINLINE explicit FTransform(const FRotator& InRotation)
+	FORCEINLINE explicit TTransform(const FRotator& InRotation)
 		: Rotation(InRotation),
-		Translation(FVector::ZeroVector),
-		Scale3D(FVector::OneVector)
+		Translation(TVector<T>::ZeroVector),
+		Scale3D(TVector<T>::OneVector)
 	{
 		DiagnosticCheckNaN_All();
 	}
@@ -165,7 +179,7 @@ public:
 	* @param InTranslation The value to use for the translation component
 	* @param InScale3D The value to use for the scale component
 	*/
-	FORCEINLINE FTransform(const FQuat& InRotation, const FVector& InTranslation, const FVector& InScale3D = FVector::OneVector)
+	FORCEINLINE TTransform(const TQuat<T>& InRotation, const TVector<T>& InTranslation, const TVector<T>& InScale3D = TVector<T>::OneVector)
 		: Rotation(InRotation),
 		Translation(InTranslation),
 		Scale3D(InScale3D)
@@ -180,7 +194,7 @@ public:
 	* @param InTranslation The value to use for the translation component
 	* @param InScale3D The value to use for the scale component
 	*/
-	FORCEINLINE FTransform(const FRotator& InRotation, const FVector& InTranslation, const FVector& InScale3D = FVector::OneVector)
+	FORCEINLINE TTransform(const FRotator& InRotation, const TVector<T>& InTranslation, const TVector<T>& InScale3D = TVector<T>::OneVector)
 		: Rotation(InRotation),
 		Translation(InTranslation),
 		Scale3D(InScale3D)
@@ -189,18 +203,18 @@ public:
 	}
 
 	/**
-	* Constructor for converting a Matrix (including scale) into a FTransform.
+	* Constructor for converting a Matrix (including scale) into a TTransform<T>.
 	*/
-	FORCEINLINE explicit FTransform(const FMatrix& InMatrix)
+	FORCEINLINE explicit TTransform(const TMatrix<T>& InMatrix)
 	{
 		SetFromMatrix(InMatrix);
 		DiagnosticCheckNaN_All();
 	}
 
 	/** Constructor that takes basis axes and translation */
-	FORCEINLINE FTransform(const FVector& InX, const FVector& InY, const FVector& InZ, const FVector& InTranslation)
+	FORCEINLINE TTransform(const TVector<T>& InX, const TVector<T>& InY, const TVector<T>& InZ, const TVector<T>& InTranslation)
 	{
-		SetFromMatrix(FMatrix(InX, InY, InZ, InTranslation));
+		SetFromMatrix(TMatrix<T>(InX, InY, InZ, InTranslation));
 		DiagnosticCheckNaN_All();
 	}
 
@@ -210,9 +224,9 @@ public:
 	CORE_API void DebugPrint() const;
 
 	/** Debug purpose only **/
-	bool DebugEqualMatrix(const FMatrix& Matrix) const;
+	bool DebugEqualMatrix(const TMatrix<T>& Matrix) const;
 
-	/** Convert FTransform contents to a string */
+	/** Convert TTransform<T> contents to a string */
 	CORE_API FString ToHumanReadableString() const;
 
 	CORE_API FString ToString() const;
@@ -224,7 +238,7 @@ public:
 	/**
 	* Copy another Transform into this one
 	*/
-	FORCEINLINE FTransform& operator=(const FTransform& Other)
+	FORCEINLINE TTransform<T>& operator=(const TTransform<T>& Other)
 	{
 		this->Rotation = Other.Rotation;
 		this->Translation = Other.Translation;
@@ -237,9 +251,9 @@ public:
 	/**
 	* Convert this Transform to a transformation matrix with scaling.
 	*/
-	FORCEINLINE FMatrix ToMatrixWithScale() const
+	FORCEINLINE TMatrix<T> ToMatrixWithScale() const
 	{
-		FMatrix OutMatrix;
+		TMatrix<T> OutMatrix;
 
 #if !(UE_BUILD_SHIPPING || UE_BUILD_TEST) && WITH_EDITORONLY_DATA
 		// Make sure Rotation is normalized when we turn it into a matrix.
@@ -249,35 +263,35 @@ public:
 		OutMatrix.M[3][1] = Translation.Y;
 		OutMatrix.M[3][2] = Translation.Z;
 
-		const float x2 = Rotation.X + Rotation.X;
-		const float y2 = Rotation.Y + Rotation.Y;
-		const float z2 = Rotation.Z + Rotation.Z;
+		const FReal x2 = Rotation.X + Rotation.X;
+		const FReal y2 = Rotation.Y + Rotation.Y;
+		const FReal z2 = Rotation.Z + Rotation.Z;
 		{
-			const float xx2 = Rotation.X * x2;
-			const float yy2 = Rotation.Y * y2;
-			const float zz2 = Rotation.Z * z2;
+			const FReal xx2 = Rotation.X * x2;
+			const FReal yy2 = Rotation.Y * y2;
+			const FReal zz2 = Rotation.Z * z2;
 
 			OutMatrix.M[0][0] = (1.0f - (yy2 + zz2)) * Scale3D.X;
 			OutMatrix.M[1][1] = (1.0f - (xx2 + zz2)) * Scale3D.Y;
 			OutMatrix.M[2][2] = (1.0f - (xx2 + yy2)) * Scale3D.Z;
 		}
 		{
-			const float yz2 = Rotation.Y * z2;
-			const float wx2 = Rotation.W * x2;
+			const FReal yz2 = Rotation.Y * z2;
+			const FReal wx2 = Rotation.W * x2;
 
 			OutMatrix.M[2][1] = (yz2 - wx2) * Scale3D.Z;
 			OutMatrix.M[1][2] = (yz2 + wx2) * Scale3D.Y;
 		}
 		{
-			const float xy2 = Rotation.X * y2;
-			const float wz2 = Rotation.W * z2;
+			const FReal xy2 = Rotation.X * y2;
+			const FReal wz2 = Rotation.W * z2;
 
 			OutMatrix.M[1][0] = (xy2 - wz2) * Scale3D.Y;
 			OutMatrix.M[0][1] = (xy2 + wz2) * Scale3D.X;
 		}
 		{
-			const float xz2 = Rotation.X * z2;
-			const float wy2 = Rotation.W * y2;
+			const FReal xz2 = Rotation.X * z2;
+			const FReal wy2 = Rotation.W * y2;
 
 			OutMatrix.M[2][0] = (xz2 + wy2) * Scale3D.Z;
 			OutMatrix.M[0][2] = (xz2 - wy2) * Scale3D.X;
@@ -294,7 +308,7 @@ public:
 	/**
 	* Convert this Transform to matrix with scaling and compute the inverse of that.
 	*/
-	FORCEINLINE FMatrix ToInverseMatrixWithScale() const
+	FORCEINLINE TMatrix<T> ToInverseMatrixWithScale() const
 	{
 		// todo: optimize
 		return ToMatrixWithScale().Inverse();
@@ -303,22 +317,22 @@ public:
 	/**
 	* Convert this Transform to inverse.
 	*/
-	FORCEINLINE FTransform Inverse() const
+	FORCEINLINE TTransform<T> Inverse() const
 	{
-		FQuat   InvRotation = Rotation.Inverse();
+		TQuat<T>   InvRotation = Rotation.Inverse();
 		// this used to cause NaN if Scale contained 0 
-		FVector InvScale3D = GetSafeScaleReciprocal(Scale3D);
-		FVector InvTranslation = InvRotation * (InvScale3D * -Translation);
+		TVector<T> InvScale3D = GetSafeScaleReciprocal(Scale3D);
+		TVector<T> InvTranslation = InvRotation * (InvScale3D * -Translation);
 
-		return FTransform(InvRotation, InvTranslation, InvScale3D);
+		return TTransform(InvRotation, InvTranslation, InvScale3D);
 	}
 
 	/**
 	* Convert this Transform to a transformation matrix, ignoring its scaling
 	*/
-	FORCEINLINE FMatrix ToMatrixNoScale() const
+	FORCEINLINE TMatrix<T> ToMatrixNoScale() const
 	{
-		FMatrix OutMatrix;
+		TMatrix<T> OutMatrix;
 
 #if !(UE_BUILD_SHIPPING || UE_BUILD_TEST) && WITH_EDITORONLY_DATA
 		// Make sure Rotation is normalized when we turn it into a matrix.
@@ -328,89 +342,35 @@ public:
 		OutMatrix.M[3][1] = Translation.Y;
 		OutMatrix.M[3][2] = Translation.Z;
 
-		const float x2 = Rotation.X + Rotation.X;
-		const float y2 = Rotation.Y + Rotation.Y;
-		const float z2 = Rotation.Z + Rotation.Z;
-		{
-			const float xx2 = Rotation.X * x2;
-			const float yy2 = Rotation.Y * y2;
-			const float zz2 = Rotation.Z * z2;
-
-			OutMatrix.M[0][0] = (1.0f - (yy2 + zz2));
-			OutMatrix.M[1][1] = (1.0f - (xx2 + zz2));
-			OutMatrix.M[2][2] = (1.0f - (xx2 + yy2));
-		}
-		{
-			const float yz2 = Rotation.Y * z2;
-			const float wx2 = Rotation.W * x2;
-
-			OutMatrix.M[2][1] = (yz2 - wx2);
-			OutMatrix.M[1][2] = (yz2 + wx2);
-		}
-		{
-			const float xy2 = Rotation.X * y2;
-			const float wz2 = Rotation.W * z2;
-
-			OutMatrix.M[1][0] = (xy2 - wz2);
-			OutMatrix.M[0][1] = (xy2 + wz2);
-		}
-		{
-			const float xz2 = Rotation.X * z2;
-			const float wy2 = Rotation.W * y2;
-
-			OutMatrix.M[2][0] = (xz2 + wy2);
-			OutMatrix.M[0][2] = (xz2 - wy2);
-		}
-
-		OutMatrix.M[0][3] = 0.0f;
-		OutMatrix.M[1][3] = 0.0f;
-		OutMatrix.M[2][3] = 0.0f;
-		OutMatrix.M[3][3] = 1.0f;
-
-		return OutMatrix;
-	}
-
-	FORCEINLINE FMatrix44f ToMatrix32NoScale() const
-	{
-		FMatrix44f OutMatrix;
-
-#if !(UE_BUILD_SHIPPING || UE_BUILD_TEST) && WITH_EDITORONLY_DATA
-		// Make sure Rotation is normalized when we turn it into a matrix.
-		check(IsRotationNormalized());
-#endif
-		OutMatrix.M[3][0] = (float)Translation.X;
-		OutMatrix.M[3][1] = (float)Translation.Y;
-		OutMatrix.M[3][2] = (float)Translation.Z;
-
-		const float x2 = Rotation.X + Rotation.X;
-		const float y2 = Rotation.Y + Rotation.Y;
-		const float z2 = Rotation.Z + Rotation.Z;
+		const FReal x2 = Rotation.X + Rotation.X;
+		const FReal y2 = Rotation.Y + Rotation.Y;
+		const FReal z2 = Rotation.Z + Rotation.Z;
 		{
-			const float xx2 = Rotation.X * x2;
-			const float yy2 = Rotation.Y * y2;
-			const float zz2 = Rotation.Z * z2;
+			const FReal xx2 = Rotation.X * x2;
+			const FReal yy2 = Rotation.Y * y2;
+			const FReal zz2 = Rotation.Z * z2;
 
 			OutMatrix.M[0][0] = (1.0f - (yy2 + zz2));
 			OutMatrix.M[1][1] = (1.0f - (xx2 + zz2));
 			OutMatrix.M[2][2] = (1.0f - (xx2 + yy2));
 		}
 		{
-			const float yz2 = Rotation.Y * z2;
-			const float wx2 = Rotation.W * x2;
+			const FReal yz2 = Rotation.Y * z2;
+			const FReal wx2 = Rotation.W * x2;
 
 			OutMatrix.M[2][1] = (yz2 - wx2);
 			OutMatrix.M[1][2] = (yz2 + wx2);
 		}
 		{
-			const float xy2 = Rotation.X * y2;
-			const float wz2 = Rotation.W * z2;
+			const FReal xy2 = Rotation.X * y2;
+			const FReal wz2 = Rotation.W * z2;
 
 			OutMatrix.M[1][0] = (xy2 - wz2);
 			OutMatrix.M[0][1] = (xy2 + wz2);
 		}
 		{
-			const float xz2 = Rotation.X * z2;
-			const float wy2 = Rotation.W * y2;
+			const FReal xz2 = Rotation.X * z2;
+			const FReal wy2 = Rotation.W * y2;
 
 			OutMatrix.M[2][0] = (xz2 + wy2);
 			OutMatrix.M[0][2] = (xz2 - wy2);
@@ -425,7 +385,7 @@ public:
 	}
 
 	/** Set this transform to the weighted blend of the supplied two transforms. */
-	FORCEINLINE void Blend(const FTransform& Atom1, const FTransform& Atom2, float Alpha)
+	FORCEINLINE void Blend(const TTransform<T>& Atom1, const TTransform<T>& Atom2, float Alpha)
 	{
 #if !(UE_BUILD_SHIPPING || UE_BUILD_TEST) && WITH_EDITORONLY_DATA
 		// Check that all bone atoms coming from animation are normalized
@@ -447,7 +407,7 @@ public:
 			// Simple linear interpolation for translation and scale.
 			Translation = FMath::Lerp(Atom1.Translation, Atom2.Translation, Alpha);
 			Scale3D = FMath::Lerp(Atom1.Scale3D, Atom2.Scale3D, Alpha);
-			Rotation = FQuat::FastLerp(Atom1.Rotation, Atom2.Rotation, Alpha);
+			Rotation = TQuat<T>::FastLerp(Atom1.Rotation, Atom2.Rotation, Alpha);
 
 			// ..and renormalize
 			Rotation.Normalize();
@@ -455,7 +415,7 @@ public:
 	}
 
 	/** Set this Transform to the weighted blend of it and the supplied Transform. */
-	FORCEINLINE void BlendWith(const FTransform& OtherAtom, float Alpha)
+	FORCEINLINE void BlendWith(const TTransform<T>& OtherAtom, float Alpha)
 	{
 #if !(UE_BUILD_SHIPPING || UE_BUILD_TEST) && WITH_EDITORONLY_DATA
 		// Check that all bone atoms coming from animation are normalized
@@ -474,7 +434,7 @@ public:
 				// Simple linear interpolation for translation and scale.
 				Translation = FMath::Lerp(Translation, OtherAtom.Translation, Alpha);
 				Scale3D = FMath::Lerp(Scale3D, OtherAtom.Scale3D, Alpha);
-				Rotation = FQuat::FastLerp(Rotation, OtherAtom.Rotation, Alpha);
+				Rotation = TQuat<T>::FastLerp(Rotation, OtherAtom.Rotation, Alpha);
 
 				// ..and renormalize
 				Rotation.Normalize();
@@ -488,12 +448,12 @@ public:
 	* Use only within blends!!
 	* Rotation part is NOT normalized!!
 	*/
-	FORCEINLINE FTransform operator+(const FTransform& Atom) const
+	FORCEINLINE TTransform<T> operator+(const TTransform<T>& Atom) const
 	{
-		return FTransform(Rotation + Atom.Rotation, Translation + Atom.Translation, Scale3D + Atom.Scale3D);
+		return TTransform(Rotation + Atom.Rotation, Translation + Atom.Translation, Scale3D + Atom.Scale3D);
 	}
 
-	FORCEINLINE FTransform& operator+=(const FTransform& Atom)
+	FORCEINLINE TTransform<T>& operator+=(const TTransform<T>& Atom)
 	{
 		Translation += Atom.Translation;
 
@@ -507,12 +467,12 @@ public:
 		return *this;
 	}
 
-	FORCEINLINE FTransform operator*(ScalarRegister Mult) const
+	FORCEINLINE TTransform<T> operator*(T Mult) const
 	{
-		return FTransform(Rotation * Mult, Translation * Mult, Scale3D * Mult);
+		return TTransform(Rotation * Mult, Translation * Mult, Scale3D * Mult);
 	}
 
-	FORCEINLINE FTransform& operator*=(ScalarRegister Mult)
+	FORCEINLINE TTransform<T>& operator*=(T Mult)
 	{
 		Translation *= Mult;
 		Rotation.X *= Mult;
@@ -531,7 +491,7 @@ public:
 	* @param  Other other transform by which to multiply.
 	* @return new transform: this * Other
 	*/
-	FORCEINLINE FTransform operator*(const FTransform& Other) const;
+	FORCEINLINE TTransform<T> operator*(const TTransform<T>& Other) const;
 
 
 	/**
@@ -540,16 +500,16 @@ public:
 	*
 	* @param  Other other transform by which to multiply.
 	*/
-	FORCEINLINE void operator*=(const FTransform& Other);
+	FORCEINLINE void operator*=(const TTransform<T>& Other);
 
 	/**
 	* Return a transform that is the result of this multiplied by another transform (made only from a rotation).
 	* Order matters when composing transforms : C = A * B will yield a transform C that logically first applies A then B to any subsequent transformation.
 	*
 	* @param  Other other quaternion rotation by which to multiply.
-	* @return new transform: this * FTransform(Other)
+	* @return new transform: this * TTransform(Other)
 	*/
-	FORCEINLINE FTransform operator*(const FQuat& Other) const;
+	FORCEINLINE TTransform<T> operator*(const TQuat<T>& Other) const;
 
 	/**
 	* Sets this transform to the result of this multiplied by another transform (made only from a rotation).
@@ -557,71 +517,71 @@ public:
 	*
 	* @param  Other other quaternion rotation by which to multiply.
 	*/
-	FORCEINLINE void operator*=(const FQuat& Other);
+	FORCEINLINE void operator*=(const TQuat<T>& Other);
 
-	FORCEINLINE static bool AnyHasNegativeScale(const FVector& InScale3D, const  FVector& InOtherScale3D);
-	FORCEINLINE void ScaleTranslation(const FVector& InScale3D);
-	FORCEINLINE void ScaleTranslation(const float& Scale);
-	FORCEINLINE void RemoveScaling(float Tolerance = SMALL_NUMBER);
-	FORCEINLINE float GetMaximumAxisScale() const;
-	FORCEINLINE float GetMinimumAxisScale() const;
+	FORCEINLINE static bool AnyHasNegativeScale(const TVector<T>& InScale3D, const TVector<T>& InOtherScale3D);
+	FORCEINLINE void ScaleTranslation(const TVector<T>& InScale3D);
+	FORCEINLINE void ScaleTranslation(const FReal& Scale);
+	FORCEINLINE void RemoveScaling(FReal Tolerance = SMALL_NUMBER);
+	FORCEINLINE T GetMaximumAxisScale() const;
+	FORCEINLINE T GetMinimumAxisScale() const;
 
 	// Inverse does not work well with VQS format(in particular non-uniform), so removing it, but made two below functions to be used instead. 
 
 	/*******************************************************************************************
-	* The below 2 functions are the ones to get delta transform and return FTransform format that can be concatenated
+	* The below 2 functions are the ones to get delta transform and return TTransform<T> format that can be concatenated
 	* Inverse itself can't concatenate with VQS format(since VQS always transform from S->Q->T, where inverse happens from T(-1)->Q(-1)->S(-1))
 	* So these 2 provides ways to fix this
 	* GetRelativeTransform returns this*Other(-1) and parameter is Other(not Other(-1))
 	* GetRelativeTransformReverse returns this(-1)*Other, and parameter is Other.
 	*******************************************************************************************/
-	CORE_API FTransform GetRelativeTransform(const FTransform& Other) const;
-	CORE_API FTransform GetRelativeTransformReverse(const FTransform& Other) const;
+	CORE_API TTransform<T> GetRelativeTransform(const TTransform<T>& Other) const;
+	CORE_API TTransform<T> GetRelativeTransformReverse(const TTransform<T>& Other) const;
 	/**
 	* Set current transform and the relative to ParentTransform.
-	* Equates to This = This->GetRelativeTransform(Parent), but saves the intermediate FTransform storage and copy.
+	* Equates to This = This->GetRelativeTransform(Parent), but saves the intermediate TTransform<T> storage and copy.
 	*/
-	CORE_API void SetToRelativeTransform(const FTransform& ParentTransform);
+	CORE_API void SetToRelativeTransform(const TTransform<T>& ParentTransform);
 
 	FORCEINLINE FVector4 TransformFVector4(const FVector4& V) const;
 	FORCEINLINE FVector4 TransformFVector4NoScale(const FVector4& V) const;
-	FORCEINLINE FVector TransformPosition(const FVector& V) const;
-	FORCEINLINE FVector TransformPositionNoScale(const FVector& V) const;
+	FORCEINLINE TVector<T> TransformPosition(const TVector<T>& V) const;
+	FORCEINLINE TVector<T> TransformPositionNoScale(const TVector<T>& V) const;
 
 	/** Inverts the transform and then transforms V - correctly handles scaling in this transform. */
-	FORCEINLINE FVector InverseTransformPosition(const FVector &V) const;
-	FORCEINLINE FVector InverseTransformPositionNoScale(const FVector &V) const;
-	FORCEINLINE FVector TransformVector(const FVector& V) const;
-	FORCEINLINE FVector TransformVectorNoScale(const FVector& V) const;
+	FORCEINLINE TVector<T> InverseTransformPosition(const TVector<T> &V) const;
+	FORCEINLINE TVector<T> InverseTransformPositionNoScale(const TVector<T> &V) const;
+	FORCEINLINE TVector<T> TransformVector(const TVector<T>& V) const;
+	FORCEINLINE TVector<T> TransformVectorNoScale(const TVector<T>& V) const;
 
 	/**
 	*	Transform a direction vector by the inverse of this transform - will not take into account translation part.
 	*	If you want to transform a surface normal (or plane) and correctly account for non-uniform scaling you should use TransformByUsingAdjointT with adjoint of matrix inverse.
 	*/
-	FORCEINLINE FVector InverseTransformVector(const FVector &V) const;
-	FORCEINLINE FVector InverseTransformVectorNoScale(const FVector &V) const;
+	FORCEINLINE TVector<T> InverseTransformVector(const TVector<T> &V) const;
+	FORCEINLINE TVector<T> InverseTransformVectorNoScale(const TVector<T> &V) const;
 
 	/**
 	 * Transform a rotation.
 	 * For example if this is a LocalToWorld transform, TransformRotation(Q) would transform Q from local to world space.
 	 */
-	FORCEINLINE FQuat TransformRotation(const FQuat& Q) const;
+	FORCEINLINE TQuat<T> TransformRotation(const TQuat<T>& Q) const;
 	
 	/**
 	* Inverse transform a rotation.
 	* For example if this is a LocalToWorld transform, InverseTransformRotation(Q) would transform Q from world to local space.
 	*/
-	FORCEINLINE FQuat InverseTransformRotation(const FQuat& Q) const;
+	FORCEINLINE TQuat<T> InverseTransformRotation(const TQuat<T>& Q) const;
 
-	FORCEINLINE FTransform GetScaled(float Scale) const;
-	FORCEINLINE FTransform GetScaled(FVector Scale) const;
-	FORCEINLINE FVector GetScaledAxis(EAxis::Type InAxis) const;
-	FORCEINLINE FVector GetUnitAxis(EAxis::Type InAxis) const;
+	FORCEINLINE TTransform<T> GetScaled(T Scale) const;
+	FORCEINLINE TTransform<T> GetScaled(TVector<T> Scale) const;
+	FORCEINLINE TVector<T> GetScaledAxis(EAxis::Type InAxis) const;
+	FORCEINLINE TVector<T> GetUnitAxis(EAxis::Type InAxis) const;
 	FORCEINLINE void Mirror(EAxis::Type MirrorAxis, EAxis::Type FlipAxis);
-	FORCEINLINE static FVector GetSafeScaleReciprocal(const FVector& InScale, float Tolerance = SMALL_NUMBER);
+	FORCEINLINE static TVector<T> GetSafeScaleReciprocal(const TVector<T>& InScale, FReal Tolerance = SMALL_NUMBER);
 
 	// temp function for easy conversion
-	FORCEINLINE FVector GetLocation() const
+	FORCEINLINE TVector<T> GetLocation() const
 	{
 		return GetTranslation();
 	}
@@ -632,13 +592,13 @@ public:
 	}
 
 	/** Calculate the determinant of this transformation */
-	FORCEINLINE float GetDeterminant() const
+	FORCEINLINE FReal GetDeterminant() const
 	{
 		return (float)(Scale3D.X * Scale3D.Y * Scale3D.Z);
 	}
 
 	/** Set the translation of this transformation */
-	FORCEINLINE void SetLocation(const FVector& Origin)
+	FORCEINLINE void SetLocation(const TVector<T>& Origin)
 	{
 		Translation = Origin;
 		DiagnosticCheckNaN_Translate();
@@ -669,7 +629,7 @@ public:
 	}
 
 	// Serializer.
-	inline friend FArchive& operator<<(FArchive& Ar, FTransform& M)
+	inline friend FArchive& operator<<(FArchive& Ar, TTransform<T>& M)
 	{
 		Ar << M.Rotation;
 		Ar << M.Translation;
@@ -679,11 +639,11 @@ public:
 
 	// Binary comparison operators.
 	/*
-	bool operator==(const FTransform& Other) const
+	bool operator==(const TTransform<T>& Other) const
 	{
 	return Rotation==Other.Rotation && Translation==Other.Translation && Scale3D==Other.Scale3D;
 	}
-	bool operator!=(const FTransform& Other) const
+	bool operator!=(const TTransform<T>& Other) const
 	{
 	return Rotation!=Other.Rotation || Translation!=Other.Translation || Scale3D!=Other.Scale3D;
 	}
@@ -691,17 +651,17 @@ public:
 
 private:
 
-	FORCEINLINE bool Private_RotationEquals(const FQuat& InRotation, const float Tolerance = KINDA_SMALL_NUMBER) const
+	FORCEINLINE bool Private_RotationEquals(const TQuat<T>& InRotation, const FReal Tolerance = KINDA_SMALL_NUMBER) const
 	{
 		return Rotation.Equals(InRotation, Tolerance);
 	}
 
-	FORCEINLINE bool Private_TranslationEquals(const FVector& InTranslation, const float Tolerance = KINDA_SMALL_NUMBER) const
+	FORCEINLINE bool Private_TranslationEquals(const TVector<T>& InTranslation, const FReal Tolerance = KINDA_SMALL_NUMBER) const
 	{
 		return Translation.Equals(InTranslation, Tolerance);
 	}
 
-	FORCEINLINE bool Private_Scale3DEquals(const FVector& InScale3D, const float Tolerance = KINDA_SMALL_NUMBER) const
+	FORCEINLINE bool Private_Scale3DEquals(const TVector<T>& InScale3D, const FReal Tolerance = KINDA_SMALL_NUMBER) const
 	{
 		return Scale3D.Equals(InScale3D, Tolerance);
 	}
@@ -709,19 +669,19 @@ private:
 public:
 
 	// Test if A's rotation equals B's rotation, within a tolerance. Preferred over "A.GetRotation().Equals(B.GetRotation())" because it is faster on some platforms.
-	FORCEINLINE static bool AreRotationsEqual(const FTransform& A, const FTransform& B, float Tolerance = KINDA_SMALL_NUMBER)
+	FORCEINLINE static bool AreRotationsEqual(const TTransform<T>& A, const TTransform<T>& B, FReal Tolerance = KINDA_SMALL_NUMBER)
 	{
 		return A.Private_RotationEquals(B.Rotation, Tolerance);
 	}
 
 	// Test if A's translation equals B's translation, within a tolerance. Preferred over "A.GetTranslation().Equals(B.GetTranslation())" because it is faster on some platforms.
-	FORCEINLINE static bool AreTranslationsEqual(const FTransform& A, const FTransform& B, float Tolerance = KINDA_SMALL_NUMBER)
+	FORCEINLINE static bool AreTranslationsEqual(const TTransform<T>& A, const TTransform<T>& B, FReal Tolerance = KINDA_SMALL_NUMBER)
 	{
 		return A.Private_TranslationEquals(B.Translation, Tolerance);
 	}
 
 	// Test if A's scale equals B's scale, within a tolerance. Preferred over "A.GetScale3D().Equals(B.GetScale3D())" because it is faster on some platforms.
-	FORCEINLINE static bool AreScale3DsEqual(const FTransform& A, const FTransform& B, float Tolerance = KINDA_SMALL_NUMBER)
+	FORCEINLINE static bool AreScale3DsEqual(const TTransform<T>& A, const TTransform<T>& B, FReal Tolerance = KINDA_SMALL_NUMBER)
 	{
 		return A.Private_Scale3DEquals(B.Scale3D, Tolerance);
 	}
@@ -729,37 +689,37 @@ public:
 
 
 	// Test if this Transform's rotation equals another's rotation, within a tolerance. Preferred over "GetRotation().Equals(Other.GetRotation())" because it is faster on some platforms.
-	FORCEINLINE bool RotationEquals(const FTransform& Other, float Tolerance = KINDA_SMALL_NUMBER) const
+	FORCEINLINE bool RotationEquals(const TTransform<T>& Other, FReal Tolerance = KINDA_SMALL_NUMBER) const
 	{
 		return AreRotationsEqual(*this, Other, Tolerance);
 	}
 
 	// Test if this Transform's translation equals another's translation, within a tolerance. Preferred over "GetTranslation().Equals(Other.GetTranslation())" because it is faster on some platforms.
-	FORCEINLINE bool TranslationEquals(const FTransform& Other, float Tolerance = KINDA_SMALL_NUMBER) const
+	FORCEINLINE bool TranslationEquals(const TTransform<T>& Other, FReal Tolerance = KINDA_SMALL_NUMBER) const
 	{
 		return AreTranslationsEqual(*this, Other, Tolerance);
 	}
 
 	// Test if this Transform's scale equals another's scale, within a tolerance. Preferred over "GetScale3D().Equals(Other.GetScale3D())" because it is faster on some platforms.
-	FORCEINLINE bool Scale3DEquals(const FTransform& Other, float Tolerance = KINDA_SMALL_NUMBER) const
+	FORCEINLINE bool Scale3DEquals(const TTransform<T>& Other, FReal Tolerance = KINDA_SMALL_NUMBER) const
 	{
 		return AreScale3DsEqual(*this, Other, Tolerance);
 	}
 
 	// Test if all components of the transforms are equal, within a tolerance.
-	FORCEINLINE bool Equals(const FTransform& Other, float Tolerance = KINDA_SMALL_NUMBER) const
+	FORCEINLINE bool Equals(const TTransform<T>& Other, FReal Tolerance = KINDA_SMALL_NUMBER) const
 	{
 		return Private_TranslationEquals(Other.Translation, Tolerance) && Private_RotationEquals(Other.Rotation, Tolerance) && Private_Scale3DEquals(Other.Scale3D, Tolerance);
 	}
 
 	// Test if all components of the transform property are equal.
-	FORCEINLINE bool Identical(const FTransform* Other, uint32 PortFlags) const
+	FORCEINLINE bool Identical(const TTransform<T>* Other, uint32 PortFlags) const
 	{
 		return Equals(*Other, 0.f);
 	}
 
 	// Test if rotation and translation components of the transforms are equal, within a tolerance.
-	FORCEINLINE bool EqualsNoScale(const FTransform& Other, float Tolerance = KINDA_SMALL_NUMBER) const
+	FORCEINLINE bool EqualsNoScale(const TTransform<T>& Other, FReal Tolerance = KINDA_SMALL_NUMBER) const
 	{
 		return Private_TranslationEquals(Other.Translation, Tolerance) && Private_RotationEquals(Other.Rotation, Tolerance);
 	}
@@ -773,7 +733,7 @@ public:
 	* @param  A Transform A.
 	* @param  B Transform B.
 	*/
-	FORCEINLINE static void Multiply(FTransform* OutTransform, const FTransform* A, const FTransform* B);
+	FORCEINLINE static void Multiply(TTransform<T>* OutTransform, const TTransform<T>* A, const TTransform<T>* B);
 
 	/**
 	* Sets the components
@@ -781,7 +741,7 @@ public:
 	* @param InTranslation The new value for the Translation component
 	* @param InScale3D The new value for the Scale3D component
 	*/
-	FORCEINLINE void SetComponents(const FQuat& InRotation, const FVector& InTranslation, const FVector& InScale3D)
+	FORCEINLINE void SetComponents(const TQuat<T>& InRotation, const TVector<T>& InTranslation, const TVector<T>& InScale3D)
 	{
 		Rotation = InRotation;
 		Translation = InTranslation;
@@ -798,16 +758,16 @@ public:
 	*/
 	FORCEINLINE void SetIdentity()
 	{
-		Rotation = FQuat::Identity;
-		Translation = FVector::ZeroVector;
-		Scale3D = FVector(1, 1, 1);
+		Rotation = TQuat<T>::Identity;
+		Translation = TVector<T>::ZeroVector;
+		Scale3D = TVector<T>(1, 1, 1);
 	}
 
 	/**
 	* Scales the Scale3D component by a new factor
 	* @param Scale3DMultiplier The value to multiply Scale3D with
 	*/
-	FORCEINLINE void MultiplyScale3D(const FVector& Scale3DMultiplier)
+	FORCEINLINE void MultiplyScale3D(const TVector<T>& Scale3DMultiplier)
 	{
 		Scale3D *= Scale3DMultiplier;
 		DiagnosticCheckNaN_Scale3D();
@@ -817,14 +777,14 @@ public:
 	* Sets the translation component
 	* @param NewTranslation The new value for the translation component
 	*/
-	FORCEINLINE void SetTranslation(const FVector& NewTranslation)
+	FORCEINLINE void SetTranslation(const TVector<T>& NewTranslation)
 	{
 		Translation = NewTranslation;
 		DiagnosticCheckNaN_Translate();
 	}
 
-	/** Copy translation from another FTransform. */
-	FORCEINLINE void CopyTranslation(const FTransform& Other)
+	/** Copy translation from another TTransform<T>. */
+	FORCEINLINE void CopyTranslation(const TTransform<T>& Other)
 	{
 		Translation = Other.Translation;
 	}
@@ -833,7 +793,7 @@ public:
 	* Concatenates another rotation to this transformation
 	* @param DeltaRotation The rotation to concatenate in the following fashion: Rotation = Rotation * DeltaRotation
 	*/
-	FORCEINLINE void ConcatenateRotation(const FQuat& DeltaRotation)
+	FORCEINLINE void ConcatenateRotation(const TQuat<T>& DeltaRotation)
 	{
 		Rotation = Rotation * DeltaRotation;
 		DiagnosticCheckNaN_Rotate();
@@ -843,26 +803,26 @@ public:
 	* Adjusts the translation component of this transformation
 	* @param DeltaTranslation The translation to add in the following fashion: Translation += DeltaTranslation
 	*/
-	FORCEINLINE void AddToTranslation(const FVector& DeltaTranslation)
+	FORCEINLINE void AddToTranslation(const TVector<T>& DeltaTranslation)
 	{
 		Translation += DeltaTranslation;
 		DiagnosticCheckNaN_Translate();
 	}
 
 	/**
-	* Add the translations from two FTransforms and return the result.
+	* Add the translations from two TTransform<T>s and return the result.
 	* @return A.Translation + B.Translation
 	*/
-	FORCEINLINE static FVector AddTranslations(const FTransform& A, const FTransform& B)
+	FORCEINLINE static TVector<T> AddTranslations(const TTransform<T>& A, const TTransform<T>& B)
 	{
 		return A.Translation + B.Translation;
 	}
 
 	/**
-	* Subtract translations from two FTransforms and return the difference.
+	* Subtract translations from two TTransform<T>s and return the difference.
 	* @return A.Translation - B.Translation.
 	*/
-	FORCEINLINE static FVector SubtractTranslations(const FTransform& A, const FTransform& B)
+	FORCEINLINE static TVector<T> SubtractTranslations(const TTransform<T>& A, const TTransform<T>& B)
 	{
 		return A.Translation - B.Translation;
 	}
@@ -871,14 +831,14 @@ public:
 	* Sets the rotation component
 	* @param NewRotation The new value for the rotation component
 	*/
-	FORCEINLINE void SetRotation(const FQuat& NewRotation)
+	FORCEINLINE void SetRotation(const TQuat<T>& NewRotation)
 	{
 		Rotation = NewRotation;
 		DiagnosticCheckNaN_Rotate();
 	}
 
-	/** Copy rotation from another FTransform. */
-	FORCEINLINE void CopyRotation(const FTransform& Other)
+	/** Copy rotation from another TTransform<T>. */
+	FORCEINLINE void CopyRotation(const TTransform<T>& Other)
 	{
 		Rotation = Other.Rotation;
 	}
@@ -887,14 +847,14 @@ public:
 	* Sets the Scale3D component
 	* @param NewScale3D The new value for the Scale3D component
 	*/
-	FORCEINLINE void SetScale3D(const FVector& NewScale3D)
+	FORCEINLINE void SetScale3D(const TVector<T>& NewScale3D)
 	{
 		Scale3D = NewScale3D;
 		DiagnosticCheckNaN_Scale3D();
 	}
 
-	/** Copy scale from another FTransform. */
-	FORCEINLINE void CopyScale3D(const FTransform& Other)
+	/** Copy scale from another TTransform<T>. */
+	FORCEINLINE void CopyScale3D(const TTransform<T>& Other)
 	{
 		Scale3D = Other.Scale3D;
 	}
@@ -904,7 +864,7 @@ public:
 	* @param NewTranslation The new value for the translation component
 	* @param NewScale3D The new value for the Scale3D component
 	*/
-	FORCEINLINE void SetTranslationAndScale3D(const FVector& NewTranslation, const FVector& NewScale3D)
+	FORCEINLINE void SetTranslationAndScale3D(const TVector<T>& NewTranslation, const TVector<T>& NewScale3D)
 	{
 		Translation = NewTranslation;
 		Scale3D = NewScale3D;
@@ -930,7 +890,7 @@ public:
 	*
 	* @param SourceAtom The other transform to accumulate into this one
 	*/
-	FORCEINLINE void Accumulate(const FTransform& SourceAtom)
+	FORCEINLINE void Accumulate(const TTransform<T>& SourceAtom)
 	{
 		// Add ref pose relative animation to base animation, only if rotation is significant.
 		if (FMath::Square(SourceAtom.Rotation.W) < 1.f - DELTA * DELTA)
@@ -958,9 +918,9 @@ public:
 	* @param Atom The other transform to accumulate into this one
 	* @param BlendWeight The weight to multiply Atom by before it is accumulated.
 	*/
-	FORCEINLINE void Accumulate(const FTransform& Atom, float BlendWeight/* default param doesn't work since vectorized version takes ref param */)
+	FORCEINLINE void Accumulate(const TTransform<T>& Atom, FReal BlendWeight/* default param doesn't work since vectorized version takes ref param */)
 	{
-		FTransform SourceAtom(Atom * BlendWeight);
+		TTransform<T> SourceAtom(Atom * BlendWeight);
 
 		// Add ref pose relative animation to base animation, only if rotation is significant.
 		if (FMath::Square(SourceAtom.Rotation.W) < 1.f - DELTA * DELTA)
@@ -984,9 +944,9 @@ public:
 	* @param DeltaAtom The other transform to accumulate into this one
 	* @param Weight The weight to multiply DeltaAtom by before it is accumulated.
 	*/
-	FORCEINLINE void AccumulateWithShortestRotation(const FTransform& DeltaAtom, float BlendWeight/* default param doesn't work since vectorized version takes ref param */)
+	FORCEINLINE void AccumulateWithShortestRotation(const TTransform<T>& DeltaAtom, FReal BlendWeight/* default param doesn't work since vectorized version takes ref param */)
 	{
-		FTransform Atom(DeltaAtom * BlendWeight);
+		TTransform<T> Atom(DeltaAtom * BlendWeight);
 
 		// To ensure the 'shortest route', we make sure the dot product between the accumulator and the incoming child atom is positive.
 		if ((Atom.Rotation | Rotation) < 0.f)
@@ -1028,11 +988,11 @@ public:
 	* @param Atom The other transform to accumulate into this one
 	* @param BlendWeight The weight to multiply Atom by before it is accumulated.
 	*/
-	FORCEINLINE void AccumulateWithAdditiveScale(const FTransform& Atom, float BlendWeight/* default param doesn't work since vectorized version takes ref param */)
+	FORCEINLINE void AccumulateWithAdditiveScale(const TTransform<T>& Atom, T BlendWeight/* default param doesn't work since vectorized version takes ref param */)
 	{
-		const FVector DefaultScale(FVector::OneVector);
+		const TVector<T> DefaultScale(TVector<T>::OneVector);
 
-		FTransform SourceAtom(Atom * BlendWeight);
+		TTransform<T> SourceAtom(Atom * BlendWeight);
 
 		// Add ref pose relative animation to base animation, only if rotation is significant.
 		if (FMath::Square(SourceAtom.Rotation.W) < 1.f - DELTA * DELTA)
@@ -1055,7 +1015,7 @@ public:
 	* @param SourceAtom2 The ending point source atom (used 100% if Alpha is 1)
 	* @param Alpha The blending weight between SourceAtom1 and SourceAtom2
 	*/
-	FORCEINLINE void LerpTranslationScale3D(const FTransform& SourceAtom1, const FTransform& SourceAtom2, ScalarRegister Alpha)
+	FORCEINLINE void LerpTranslationScale3D(const TTransform<T>& SourceAtom1, const TTransform<T>& SourceAtom2, ScalarRegister Alpha)
 	{
 		Translation = FMath::Lerp(SourceAtom1.Translation, SourceAtom2.Translation, Alpha);
 		Scale3D = FMath::Lerp(SourceAtom1.Scale3D, SourceAtom2.Scale3D, Alpha);
@@ -1095,11 +1055,11 @@ public:
 	* @param SourceAtom The target transformation (used when BlendWeight = 1); this is modified during the process
 	* @param BlendWeight The blend weight between Identity and SourceAtom
 	*/
-	FORCEINLINE static void BlendFromIdentityAndAccumulate(FTransform& FinalAtom, const FTransform& SourceAtom, float BlendWeight)
+	FORCEINLINE static void BlendFromIdentityAndAccumulate(TTransform<T>& FinalAtom, const TTransform<T>& SourceAtom, float BlendWeight)
 	{
-		const FTransform AdditiveIdentity(FQuat::Identity, FVector::ZeroVector, FVector::ZeroVector);
-		const FVector DefaultScale(FVector::OneVector);
-		FTransform DeltaAtom = SourceAtom;
+		const TTransform<T> AdditiveIdentity(TQuat<T>::Identity, TVector<T>::ZeroVector, TVector<T>::ZeroVector);
+		const TVector<T> DefaultScale(TVector<T>::OneVector);
+		TTransform<T> DeltaAtom = SourceAtom;
 
 		// Scale delta by weight
 		if (BlendWeight < (1.f - ZERO_ANIMWEIGHT_THRESH))
@@ -1124,7 +1084,7 @@ public:
 	*
 	* @return The rotation component
 	*/
-	FORCEINLINE FQuat GetRotation() const
+	FORCEINLINE TQuat<T> GetRotation() const
 	{
 		DiagnosticCheckNaN_Rotate();
 		return Rotation;
@@ -1135,7 +1095,7 @@ public:
 	*
 	* @return The translation component
 	*/
-	FORCEINLINE FVector GetTranslation() const
+	FORCEINLINE TVector<T> GetTranslation() const
 	{
 		DiagnosticCheckNaN_Translate();
 		return Translation;
@@ -1146,7 +1106,7 @@ public:
 	*
 	* @return The Scale3D component
 	*/
-	FORCEINLINE FVector GetScale3D() const
+	FORCEINLINE TVector<T> GetScale3D() const
 	{
 		DiagnosticCheckNaN_Scale3D();
 		return Scale3D;
@@ -1157,7 +1117,7 @@ public:
 	*
 	* @param SrcBA The transform to copy rotation and Scale3D from
 	*/
-	FORCEINLINE void CopyRotationPart(const FTransform& SrcBA)
+	FORCEINLINE void CopyRotationPart(const TTransform<T>& SrcBA)
 	{
 		Rotation = SrcBA.Rotation;
 		Scale3D = SrcBA.Scale3D;
@@ -1171,7 +1131,7 @@ public:
 	*
 	* @param SrcBA The transform to copy translation and Scale3D from
 	*/
-	FORCEINLINE void CopyTranslationAndScale3D(const FTransform& SrcBA)
+	FORCEINLINE void CopyTranslationAndScale3D(const TTransform<T>& SrcBA)
 	{
 		Translation = SrcBA.Translation;
 		Scale3D = SrcBA.Scale3D;
@@ -1180,9 +1140,9 @@ public:
 		DiagnosticCheckNaN_Scale3D();
 	}
 
-	void SetFromMatrix(const FMatrix& InMatrix)
+	void SetFromMatrix(const TMatrix<T>& InMatrix)
 	{
-		FMatrix M = InMatrix;
+		TMatrix<T> M = InMatrix;
 
 		// Get the 3D scale from the matrix
 		Scale3D = M.ExtractScaling();
@@ -1196,7 +1156,7 @@ public:
 			M.SetAxis(0, -M.GetScaledAxis(EAxis::X));
 		}
 
-		Rotation = FQuat(M);
+		Rotation = TQuat<T>(M);
 		Translation = InMatrix.GetOrigin();
 
 		// Normalize rotation
@@ -1214,7 +1174,7 @@ private:
 	* @param  A Transform A.
 	* @param  B Transform B.
 	*/
-	FORCEINLINE static void MultiplyUsingMatrixWithScale(FTransform* OutTransform, const FTransform* A, const FTransform* B);
+	FORCEINLINE static void MultiplyUsingMatrixWithScale(TTransform<T>* OutTransform, const TTransform<T>* A, const TTransform<T>* B);
 	/**
 	* Create a new transform from multiplications of given to matrices (AMatrix*BMatrix) using desired scale
 	* This is used by MultiplyUsingMatrixWithScale and GetRelativeTransformUsingMatrixWithScale
@@ -1225,7 +1185,7 @@ private:
 	* @param	DesiredScale - there is no check on if the magnitude is correct here. It assumes that is correct.
 	* @param	OutTransform the constructed transform
 	*/
-	FORCEINLINE static void ConstructTransformFromMatrixWithDesiredScale(const FMatrix& AMatrix, const FMatrix& BMatrix, const FVector& DesiredScale, FTransform& OutTransform);
+	FORCEINLINE static void ConstructTransformFromMatrixWithDesiredScale(const TMatrix<T>& AMatrix, const TMatrix<T>& BMatrix, const TVector<T>& DesiredScale, TTransform<T>& OutTransform);
 	/**
 	* Create a new transform: OutTransform = Base * Relative(-1) using the matrix while keeping the scale that's given by Base and Relative
 	* Please note that this operation is a lot more expensive than normal GetRelativeTrnasform
@@ -1234,20 +1194,27 @@ private:
 	* @param  BAse Transform Base.
 	* @param  Relative Transform Relative.
 	*/
-	static void GetRelativeTransformUsingMatrixWithScale(FTransform* OutTransform, const FTransform* Base, const FTransform* Relative);
-};
+	static void GetRelativeTransformUsingMatrixWithScale(TTransform<T>* OutTransform, const TTransform<T>* Base, const TTransform<T>* Relative);
 
+public:
+	// Conversion to other type.
+	template<typename FArg, TEMPLATE_REQUIRES(!TIsSame<T, FArg>::Value)>
+	explicit TTransform(const TTransform<FArg>& From) : TTransform<T>((TQuat<T>)From.GetRotation(), (TVector<T>)From.GetTranslation(), (TVector<T>)From.GetScale3D()) {}
+};
 
-template <> struct TIsPODType<FTransform> { enum { Value = true }; };
+template<> inline const TTransform<float> TTransform<float>::Identity(FQuat4f(0.f, 0.f, 0.f, 1.f), FVector3f(0.f), FVector3f(1.f));
+template<> inline const TTransform<double> TTransform<double>::Identity(FQuat4d(0.0, 0.0, 0.0, 1.0), FVector3d(0.0), FVector3d(1.0));
 
-FORCEINLINE bool FTransform::AnyHasNegativeScale(const FVector& InScale3D, const  FVector& InOtherScale3D) 
+template<typename T>
+FORCEINLINE bool TTransform<T>::AnyHasNegativeScale(const TVector<T>& InScale3D, const  TVector<T>& InOtherScale3D) 
 {
 	return  (InScale3D.X < 0.f || InScale3D.Y < 0.f || InScale3D.Z < 0.f 
 			|| InOtherScale3D.X < 0.f || InOtherScale3D.Y < 0.f || InOtherScale3D.Z < 0.f );
 }
 
 /** Scale the translation part of the Transform by the supplied vector. */
-FORCEINLINE void FTransform::ScaleTranslation(const FVector& InScale3D)
+template<typename T>
+FORCEINLINE void TTransform<T>::ScaleTranslation(const TVector<T>& InScale3D)
 {
 	Translation *= InScale3D;
 
@@ -1255,7 +1222,8 @@ FORCEINLINE void FTransform::ScaleTranslation(const FVector& InScale3D)
 }
 
 
-FORCEINLINE void FTransform::ScaleTranslation(const float& Scale)
+template<typename T>
+FORCEINLINE void TTransform<T>::ScaleTranslation(const FReal& Scale)
 {
 	Translation *= Scale;
 
@@ -1264,31 +1232,34 @@ FORCEINLINE void FTransform::ScaleTranslation(const float& Scale)
 
 
 // this function is from matrix, and all it does is to normalize rotation portion
-FORCEINLINE void FTransform::RemoveScaling(float Tolerance/*=SMALL_NUMBER*/)
+template<typename T>
+FORCEINLINE void TTransform<T>::RemoveScaling(T Tolerance/*=SMALL_NUMBER*/)
 {
-	Scale3D = FVector(1, 1, 1);
+	Scale3D = TVector<T>(1, 1, 1);
 	Rotation.Normalize();
 
 	DiagnosticCheckNaN_Rotate();
 	DiagnosticCheckNaN_Scale3D();
 }
 
-FORCEINLINE void FTransform::MultiplyUsingMatrixWithScale(FTransform* OutTransform, const FTransform* A, const FTransform* B)
+template<typename T>
+FORCEINLINE void TTransform<T>::MultiplyUsingMatrixWithScale(TTransform<T>* OutTransform, const TTransform<T>* A, const TTransform<T>* B)
 {
 	// the goal of using M is to get the correct orientation
 	// but for translation, we still need scale
 	ConstructTransformFromMatrixWithDesiredScale(A->ToMatrixWithScale(), B->ToMatrixWithScale(), A->Scale3D*B->Scale3D, *OutTransform);
 }
 
-FORCEINLINE void FTransform::ConstructTransformFromMatrixWithDesiredScale(const FMatrix& AMatrix, const FMatrix& BMatrix, const FVector& DesiredScale, FTransform& OutTransform)
+template<typename T>
+FORCEINLINE void TTransform<T>::ConstructTransformFromMatrixWithDesiredScale(const TMatrix<T>& AMatrix, const TMatrix<T>& BMatrix, const TVector<T>& DesiredScale, TTransform<T>& OutTransform)
 {
 	// the goal of using M is to get the correct orientation
 	// but for translation, we still need scale
-	FMatrix M = AMatrix * BMatrix;
+	TMatrix<T> M = AMatrix * BMatrix;
 	M.RemoveScaling();
 
 	// apply negative scale back to axes
-	FVector SignedScale = DesiredScale.GetSignVector();
+	TVector<T> SignedScale = DesiredScale.GetSignVector();
 
 	M.SetAxis(0, SignedScale.X * M.GetScaledAxis(EAxis::X));
 	M.SetAxis(1, SignedScale.Y * M.GetScaledAxis(EAxis::Y));
@@ -1296,21 +1267,22 @@ FORCEINLINE void FTransform::ConstructTransformFromMatrixWithDesiredScale(const
 
 	// @note: if you have negative with 0 scale, this will return rotation that is identity
 	// since matrix loses that axes
-	FQuat Rotation = FQuat(M);
+	TQuat<T> Rotation = TQuat<T>(M);
 	Rotation.Normalize();
 
 	// set values back to output
 	OutTransform.Scale3D = DesiredScale;
 	OutTransform.Rotation = Rotation;
 
-	// technically I could calculate this using FTransform but then it does more quat multiplication 
+	// technically I could calculate this using TTransform<T> but then it does more quat multiplication 
 	// instead of using Scale in matrix multiplication
-	// it's a question of between RemoveScaling vs using FTransform to move translation
+	// it's a question of between RemoveScaling vs using TTransform<T> to move translation
 	OutTransform.Translation = M.GetOrigin();
 }
 
-/** Returns Multiplied Transform of 2 FTransforms **/
-FORCEINLINE void FTransform::Multiply(FTransform* OutTransform, const FTransform* A, const FTransform* B)
+/** Returns Multiplied Transform of 2 TTransform<T>s **/
+template<typename T>
+FORCEINLINE void TTransform<T>::Multiply(TTransform<T>* OutTransform, const TTransform<T>* A, const TTransform<T>* B)
 {
 	A->DiagnosticCheckNaN_All();
 	B->DiagnosticCheckNaN_All();
@@ -1351,9 +1323,10 @@ FORCEINLINE void FTransform::Multiply(FTransform* OutTransform, const FTransform
 /**
 * Apply Scale to this transform
 */
-FORCEINLINE FTransform FTransform::GetScaled(float InScale) const
+template<typename T>
+FORCEINLINE TTransform<T> TTransform<T>::GetScaled(T InScale) const
 {
-	FTransform A(*this);
+	TTransform<T> A(*this);
 	A.Scale3D *= InScale;
 
 	A.DiagnosticCheckNaN_Scale3D();
@@ -1365,9 +1338,10 @@ FORCEINLINE FTransform FTransform::GetScaled(float InScale) const
 /**
 * Apply Scale to this transform
 */
-FORCEINLINE FTransform FTransform::GetScaled(FVector InScale) const
+template<typename T>
+FORCEINLINE TTransform<T> TTransform<T>::GetScaled(TVector<T> InScale) const
 {
-	FTransform A(*this);
+	TTransform<T> A(*this);
 	A.Scale3D *= InScale;
 
 	A.DiagnosticCheckNaN_Scale3D();
@@ -1377,7 +1351,8 @@ FORCEINLINE FTransform FTransform::GetScaled(FVector InScale) const
 
 
 /** Transform homogenous FVector4, ignoring the scaling part of this transform **/
-FORCEINLINE FVector4 FTransform::TransformFVector4NoScale(const FVector4& V) const
+template<typename T>
+FORCEINLINE FVector4 TTransform<T>::TransformFVector4NoScale(const FVector4& V) const
 {
 	DiagnosticCheckNaN_All();
 
@@ -1386,7 +1361,7 @@ FORCEINLINE FVector4 FTransform::TransformFVector4NoScale(const FVector4& V) con
 
 	//Transform using QST is following
 	//QST(P) = Q*S*P*-Q + T where Q = quaternion, S = scale, T = translation
-	FVector4 Transform = FVector4(Rotation.RotateVector(FVector(V)), 0.f);
+	FVector4 Transform = FVector4(Rotation.RotateVector(TVector<T>(V)), 0.f);
 	if (V.W == 1.f)
 	{
 		Transform += FVector4(Translation, 1.f);
@@ -1397,7 +1372,8 @@ FORCEINLINE FVector4 FTransform::TransformFVector4NoScale(const FVector4& V) con
 
 
 /** Transform FVector4 **/
-FORCEINLINE FVector4 FTransform::TransformFVector4(const FVector4& V) const
+template<typename T>
+FORCEINLINE FVector4 TTransform<T>::TransformFVector4(const FVector4& V) const
 {
 	DiagnosticCheckNaN_All();
 
@@ -1407,7 +1383,7 @@ FORCEINLINE FVector4 FTransform::TransformFVector4(const FVector4& V) const
 	//Transform using QST is following
 	//QST(P) = Q*S*P*-Q + T where Q = quaternion, S = scale, T = translation
 
-	FVector4 Transform = FVector4(Rotation.RotateVector(Scale3D*FVector(V)), 0.f);
+	FVector4 Transform = FVector4(Rotation.RotateVector(Scale3D*TVector<T>(V)), 0.f);
 	if (V.W == 1.f)
 	{
 		Transform += FVector4(Translation, 1.f);
@@ -1417,28 +1393,32 @@ FORCEINLINE FVector4 FTransform::TransformFVector4(const FVector4& V) const
 }
 
 
-FORCEINLINE FVector FTransform::TransformPosition(const FVector& V) const
+template<typename T>
+FORCEINLINE TVector<T> TTransform<T>::TransformPosition(const TVector<T>& V) const
 {
 	DiagnosticCheckNaN_All();
 	return Rotation.RotateVector(Scale3D*V) + Translation;
 }
 
 
-FORCEINLINE FVector FTransform::TransformPositionNoScale(const FVector& V) const
+template<typename T>
+FORCEINLINE TVector<T> TTransform<T>::TransformPositionNoScale(const TVector<T>& V) const
 {
 	DiagnosticCheckNaN_All();
 	return Rotation.RotateVector(V) + Translation;
 }
 
 
-FORCEINLINE FVector FTransform::TransformVector(const FVector& V) const
+template<typename T>
+FORCEINLINE TVector<T> TTransform<T>::TransformVector(const TVector<T>& V) const
 {
 	DiagnosticCheckNaN_All();
 	return Rotation.RotateVector(Scale3D*V);
 }
 
 
-FORCEINLINE FVector FTransform::TransformVectorNoScale(const FVector& V) const
+template<typename T>
+FORCEINLINE TVector<T> TTransform<T>::TransformVectorNoScale(const TVector<T>& V) const
 {
 	DiagnosticCheckNaN_All();
 	return Rotation.RotateVector(V);
@@ -1446,7 +1426,8 @@ FORCEINLINE FVector FTransform::TransformVectorNoScale(const FVector& V) const
 
 
 // do backward operation when inverse, translation -> rotation -> scale
-FORCEINLINE FVector FTransform::InverseTransformPosition(const FVector &V) const
+template<typename T>
+FORCEINLINE TVector<T> TTransform<T>::InverseTransformPosition(const TVector<T> &V) const
 {
 	DiagnosticCheckNaN_All();
 	return (Rotation.UnrotateVector(V - Translation)) * GetSafeScaleReciprocal(Scale3D);
@@ -1454,7 +1435,8 @@ FORCEINLINE FVector FTransform::InverseTransformPosition(const FVector &V) const
 
 
 // do backward operation when inverse, translation -> rotation
-FORCEINLINE FVector FTransform::InverseTransformPositionNoScale(const FVector &V) const
+template<typename T>
+FORCEINLINE TVector<T> TTransform<T>::InverseTransformPositionNoScale(const TVector<T> &V) const
 {
 	DiagnosticCheckNaN_All();
 	return (Rotation.UnrotateVector(V - Translation));
@@ -1462,7 +1444,8 @@ FORCEINLINE FVector FTransform::InverseTransformPositionNoScale(const FVector &V
 
 
 // do backward operation when inverse, translation -> rotation -> scale
-FORCEINLINE FVector FTransform::InverseTransformVector(const FVector &V) const
+template<typename T>
+FORCEINLINE TVector<T> TTransform<T>::InverseTransformVector(const TVector<T> &V) const
 {
 	DiagnosticCheckNaN_All();
 	return (Rotation.UnrotateVector(V)) * GetSafeScaleReciprocal(Scale3D);
@@ -1470,95 +1453,106 @@ FORCEINLINE FVector FTransform::InverseTransformVector(const FVector &V) const
 
 
 // do backward operation when inverse, translation -> rotation
-FORCEINLINE FVector FTransform::InverseTransformVectorNoScale(const FVector &V) const
+template<typename T>
+FORCEINLINE TVector<T> TTransform<T>::InverseTransformVectorNoScale(const TVector<T> &V) const
 {
 	DiagnosticCheckNaN_All();
 	return (Rotation.UnrotateVector(V));
 }
 
-FORCEINLINE FQuat FTransform::TransformRotation(const FQuat& Q) const
+template<typename T>
+FORCEINLINE TQuat<T> TTransform<T>::TransformRotation(const TQuat<T>& Q) const
 {
 	return GetRotation() * Q;
 }
 
-FORCEINLINE FQuat FTransform::InverseTransformRotation(const FQuat& Q) const
+template<typename T>
+FORCEINLINE TQuat<T> TTransform<T>::InverseTransformRotation(const TQuat<T>& Q) const
 {
 	return GetRotation().Inverse() * Q;
 }
 
-FORCEINLINE FTransform FTransform::operator*(const FTransform& Other) const
+template<typename T>
+FORCEINLINE TTransform<T> TTransform<T>::operator*(const TTransform<T>& Other) const
 {
-	FTransform Output;
+	TTransform<T> Output;
 	Multiply(&Output, this, &Other);
 	return Output;
 }
 
 
-FORCEINLINE void FTransform::operator*=(const FTransform& Other)
+template<typename T>
+FORCEINLINE void TTransform<T>::operator*=(const TTransform<T>& Other)
 {
 	Multiply(this, this, &Other);
 }
 
 
-FORCEINLINE FTransform FTransform::operator*(const FQuat& Other) const
+template<typename T>
+FORCEINLINE TTransform<T> TTransform<T>::operator*(const TQuat<T>& Other) const
 {
-	FTransform Output, OtherTransform(Other, FVector::ZeroVector, FVector::OneVector);
+	TTransform<T> Output, OtherTransform(Other, TVector<T>::ZeroVector, TVector<T>::OneVector);
 	Multiply(&Output, this, &OtherTransform);
 	return Output;
 }
 
 
-FORCEINLINE void FTransform::operator*=(const FQuat& Other)
+template<typename T>
+FORCEINLINE void TTransform<T>::operator*=(const TQuat<T>& Other)
 {
-	FTransform OtherTransform(Other, FVector::ZeroVector, FVector::OneVector);
+	TTransform<T> OtherTransform(Other, TVector<T>::ZeroVector, TVector<T>::OneVector);
 	Multiply(this, this, &OtherTransform);
 }
 
 
 // x = 0, y = 1, z = 2
-FORCEINLINE FVector FTransform::GetScaledAxis(EAxis::Type InAxis) const
+template<typename T>
+FORCEINLINE TVector<T> TTransform<T>::GetScaledAxis(EAxis::Type InAxis) const
 {
 	if (InAxis == EAxis::X)
 	{
-		return TransformVector(FVector(1.f, 0.f, 0.f));
+		return TransformVector(TVector<T>(1.f, 0.f, 0.f));
 	}
 	else if (InAxis == EAxis::Y)
 	{
-		return TransformVector(FVector(0.f, 1.f, 0.f));
+		return TransformVector(TVector<T>(0.f, 1.f, 0.f));
 	}
 
-	return TransformVector(FVector(0.f, 0.f, 1.f));
+	return TransformVector(TVector<T>(0.f, 0.f, 1.f));
 }
 
 
 // x = 0, y = 1, z = 2
-FORCEINLINE FVector FTransform::GetUnitAxis(EAxis::Type InAxis) const
+template<typename T>
+FORCEINLINE TVector<T> TTransform<T>::GetUnitAxis(EAxis::Type InAxis) const
 {
 	if (InAxis == EAxis::X)
 	{
-		return TransformVectorNoScale(FVector(1.f, 0.f, 0.f));
+		return TransformVectorNoScale(TVector<T>(1.f, 0.f, 0.f));
 	}
 	else if (InAxis == EAxis::Y)
 	{
-		return TransformVectorNoScale(FVector(0.f, 1.f, 0.f));
+		return TransformVectorNoScale(TVector<T>(0.f, 1.f, 0.f));
 	}
 
-	return TransformVectorNoScale(FVector(0.f, 0.f, 1.f));
+	return TransformVectorNoScale(TVector<T>(0.f, 0.f, 1.f));
 }
 
 
-FORCEINLINE void FTransform::Mirror(EAxis::Type MirrorAxis, EAxis::Type FlipAxis)
+template<typename T>
+FORCEINLINE void TTransform<T>::Mirror(EAxis::Type MirrorAxis, EAxis::Type FlipAxis)
 {
 	// We do convert to Matrix for mirroring. 
-	FMatrix M = ToMatrixWithScale();
+	TMatrix<T> M = ToMatrixWithScale();
 	M.Mirror(MirrorAxis, FlipAxis);
 	SetFromMatrix(M);
 }
 
 
-/** same version of FMatrix::GetMaximumAxisScale function **/
+/** same version of TMatrix<T>::GetMaximumAxisScale function **/
 /** @return the maximum magnitude of all components of the 3D scale. */
-inline float FTransform::GetMaximumAxisScale() const
+template<typename T>
+inline T TTransform<T>::GetMaximumAxisScale() const
 {
 	DiagnosticCheckNaN_Scale3D();
 	return (float)Scale3D.GetAbsMax();
@@ -1566,7 +1560,8 @@ inline float FTransform::GetMaximumAxisScale() const
 
 
 /** @return the minimum magnitude of all components of the 3D scale. */
-inline float FTransform::GetMinimumAxisScale() const
+template<typename T>
+inline T TTransform<T>::GetMinimumAxisScale() const
 {
 	DiagnosticCheckNaN_Scale3D();
 	return (float)Scale3D.GetAbsMin();
@@ -1578,9 +1573,10 @@ inline float FTransform::GetMinimumAxisScale() const
 // anymore because you should be instead of showing gigantic infinite mesh
 // also returning BIG_NUMBER causes sequential NaN issues by multiplying 
 // so we hardcode as 0
-FORCEINLINE FVector FTransform::GetSafeScaleReciprocal(const FVector& InScale, float Tolerance)
+template<typename T>
+FORCEINLINE TVector<T> TTransform<T>::GetSafeScaleReciprocal(const TVector<T>& InScale, T Tolerance)
 {
-	FVector SafeReciprocalScale;
+	TVector<T> SafeReciprocalScale;
 	if (FMath::Abs(InScale.X) <= Tolerance)
 	{
 		SafeReciprocalScale.X = 0.f;
@@ -1611,4 +1607,7 @@ FORCEINLINE FVector FTransform::GetSafeScaleReciprocal(const FVector& InScale, f
 	return SafeReciprocalScale;
 }
 
+} // namespace UE::Math
+} // namespace UE
+
 #endif
