LWC support for FQuat and FTransform

diff --git a/Engine/Source/Runtime/Core/Public/Math/TransformVectorized.h b/Engine/Source/Runtime/Core/Public/Math/TransformVectorized.h
--- a/Engine/Source/Runtime/Core/Public/Math/TransformVectorized.h
+++ b/Engine/Source/Runtime/Core/Public/Math/TransformVectorized.h
@@ -17,6 +17,23 @@
 
 #if ENABLE_VECTORIZED_TRANSFORM
 
+struct Z_Construct_UScriptStruct_FTransform3f_Statics;
+struct Z_Construct_UScriptStruct_FTransform3d_Statics;
+struct Z_Construct_UScriptStruct_FTransform_Statics;
+
+namespace UE
+{
+namespace Math
+{
+	// Selector for alignment of TTransform based on template type used.
+	// The base has no 'Value' defined, you must define a specialized type below for types allowed as TTransform's template param.
+	template<typename T>
+	struct TAlignOfTransform {  };
+
+	// Specializations for allowed TTransform::FReal types
+	template<> struct TAlignOfTransform<float>		{ enum { Value = FMath::Max<uint64>(16, alignof(TPersistentVectorRegisterType<float>)) }; };
+	template<> struct TAlignOfTransform<double>		{ enum { Value = FMath::Max<uint64>(16, alignof(TPersistentVectorRegisterType<double>)) }; };
+
 /**
  * Transform composed of Scale, Rotation (as a quaternion), and Translation.
  *
@@ -27,36 +44,44 @@
  * Transformation of direction vectors is applied in the order: Scale -> Rotate.
  *
  * Order matters when composing transforms: C = A * B will yield a transform C that logically
- * first applies A then B to any subsequent transformation. Note that this is the opposite order of quaternion (FQuat) multiplication.
+ * first applies A then B to any subsequent transformation. Note that this is the opposite order of quaternion (TQuat<T>) multiplication.
  *
  * Example: LocalToWorld = (DeltaRotation * LocalToWorld) will change rotation in local space by DeltaRotation.
  * Example: LocalToWorld = (LocalToWorld * DeltaRotation) will change rotation in world space by DeltaRotation.
  */
-
-struct alignas(alignof(VectorRegister)) FTransform
+template<typename T>
+struct alignas( TAlignOfTransform<T>::Value ) TTransform
 {
-	friend struct Z_Construct_UScriptStruct_FTransform_Statics;
+	// Can't have a TEMPLATE_REQUIRES in the declaration because of the forward declarations, so check for allowed types here.
+	static_assert(TIsFloatingPoint<T>::Value, "TTransform only supports float and double types.");
+
+	friend Z_Construct_UScriptStruct_FTransform3f_Statics;
+	friend Z_Construct_UScriptStruct_FTransform3d_Statics;
+	friend Z_Construct_UScriptStruct_FTransform_Statics;
+
+	using FReal = T;
+	using TransformVectorRegister = TVectorRegisterType<T>;
 
 protected:
 	/** Rotation of this transformation, as a quaternion */
-	VectorRegister4Float Rotation; // LWC_TODO : change back to VectorRegister when FQuat uses double 
+	TPersistentVectorRegisterType<T> Rotation;
 	/** Translation of this transformation, as a vector */
-	VectorRegister	Translation;
+	TPersistentVectorRegisterType<T> Translation;
 	/** 3D scale (always applied in local space) as a vector */
-	VectorRegister Scale3D;	
+	TPersistentVectorRegisterType<T> Scale3D;
 public:
 	/**
-	 * The identity transformation (Rotation = FQuat::Identity, Translation = FVector::ZeroVector, Scale3D = (1,1,1))
+	 * The identity transformation (Rotation = TQuat<T>::Identity, Translation = TVector<T>::ZeroVector, Scale3D = (1,1,1))
 	 */
-	static CORE_API const FTransform Identity;
+	static const TTransform<T> Identity;
 
 #if ENABLE_NAN_DIAGNOSTIC
 	FORCEINLINE void DiagnosticCheckNaN_Scale3D() const
 	{
 		if (VectorContainsNaNOrInfinite(Scale3D))
 		{
-			logOrEnsureNanError(TEXT("FTransform Vectorized Scale3D contains NaN"));
-			const_cast<FTransform*>(this)->Scale3D = VectorSet_W0( VectorOneVectorRegister() );
+			logOrEnsureNanError(TEXT("TTransform<T> Vectorized Scale3D contains NaN"));
+			const_cast<TTransform<T>*>(this)->Scale3D = GlobalVectorConstants::Float1110;
 		}
 	}
 
@@ -64,8 +89,8 @@ public:
 	{
 		if (VectorContainsNaNOrInfinite(Translation))
 		{
-			logOrEnsureNanError(TEXT("FTransform Vectorized Translation contains NaN"));
-			const_cast<FTransform*>(this)->Translation = VectorOneVectorRegister();
+			logOrEnsureNanError(TEXT("TTransform<T> Vectorized Translation contains NaN"));
+			const_cast<TTransform<T>*>(this)->Translation = VectorZero();
 		}
 	}
 
@@ -73,8 +98,8 @@ public:
 	{
 		if (VectorContainsNaNOrInfinite(Rotation))
 		{
-			logOrEnsureNanError(TEXT("FTransform Vectorized Rotation contains NaN"));
-			const_cast<FTransform*>(this)->Rotation = VectorSet_W1( VectorZeroFloat() );
+			logOrEnsureNanError(TEXT("TTransform<T> Vectorized Rotation contains NaN"));
+			const_cast<TTransform<T>*>(this)->Rotation = GlobalVectorConstants::Float0001;
 		}
 	}
 
@@ -90,7 +115,7 @@ public:
 		DiagnosticCheckNaN_All();
 		if (!IsValid())
 		{
-			logOrEnsureNanError(TEXT("FTransform Vectorized transform is not valid: %s"), *ToHumanReadableString());
+			logOrEnsureNanError(TEXT("TTransform Vectorized transform is not valid: %s"), *ToHumanReadableString());
 		}
 		
 	}
@@ -105,14 +130,14 @@ public:
 	/**
 	 * Constructor with initialization to the identity transform.
 	 */
-	FORCEINLINE FTransform()
+	FORCEINLINE TTransform()
 	{
 		// Rotation = {0,0,0,1)
-		Rotation = GlobalVectorConstants::Float0001; //VectorSet_W1( VectorZero() );
+		Rotation = GlobalVectorConstants::Float0001;
 		// Translation = {0,0,0,0)
-		Translation = VectorZeroVectorRegister();
+		Translation = GlobalVectorConstants::FloatZero;
 		// Scale3D = {1,1,1,0);
-		Scale3D = GlobalVectorConstants::Vector1110; //VectorSet_W0( VectorOne() );
+		Scale3D = GlobalVectorConstants::Float1110;
 	}
 
 	/**
@@ -120,14 +145,14 @@ public:
 	 *
 	 * @param InTranslation The value to use for the translation component
 	 */
-	FORCEINLINE explicit FTransform(const FVector& InTranslation) 
+	FORCEINLINE explicit TTransform(const TVector<T>& InTranslation) 
 	{
 		// Rotation = {0,0,0,1) quaternion identity
-		Rotation = GlobalVectorConstants::Float0001; //VectorSet_W1( VectorZero() );
+		Rotation = GlobalVectorConstants::Float0001;
 		//Translation = InTranslation;
 		Translation = MakeVectorRegister(InTranslation.X, InTranslation.Y, InTranslation.Z, 0.0f );
 		// Scale3D = {1,1,1,0);
-		Scale3D = GlobalVectorConstants::Vector1110; //VectorSet_W0( VectorOne() );
+		Scale3D = GlobalVectorConstants::Float1110;
 
 		DiagnosticCheckNaN_All();
 	}
@@ -137,14 +162,14 @@ public:
 	 *
 	 * @param InRotation The value to use for rotation component
 	 */
-	FORCEINLINE explicit FTransform(const FQuat& InRotation) 
+	FORCEINLINE explicit TTransform(const TQuat<T>& InRotation) 
 	{
 		// Rotation = InRotation
-		Rotation =  VectorLoadAligned( &InRotation.X );
+		Rotation =  VectorLoadAligned( &InRotation );
 		// Translation = {0,0,0,0)
-		Translation = VectorZeroVectorRegister();
+		Translation = GlobalVectorConstants::FloatZero;
 		// Scale3D = {1,1,1,0);
-		Scale3D = GlobalVectorConstants::Vector1110; //VectorSet_W0( VectorOne() );
+		Scale3D = GlobalVectorConstants::Float1110;
 
 		DiagnosticCheckNaN_All();
 	}
@@ -154,15 +179,15 @@ public:
 	 *
 	 * @param InRotation The value to use for rotation component  (after being converted to a quaternion)
 	 */
-	FORCEINLINE explicit FTransform(const FRotator& InRotation) 
+	FORCEINLINE explicit TTransform(const FRotator& InRotation) 
 	{
-		FQuat InQuatRotation = InRotation.Quaternion();
+		TQuat<T> InQuatRotation = TQuat<T>::MakeFromRotator(InRotation);
 		// Rotation = InRotation
-		Rotation =  VectorLoadAligned( &InQuatRotation.X );
+		Rotation =  VectorLoadAligned( &InQuatRotation );
 		// Translation = {0,0,0,0)
-		Translation = VectorZeroVectorRegister();
+		Translation = GlobalVectorConstants::FloatZero;
 		// Scale3D = {1,1,1,0);
-		Scale3D = GlobalVectorConstants::Vector1110; //VectorSet_W0( VectorOne() );
+		Scale3D = GlobalVectorConstants::Float1110;
 
 		DiagnosticCheckNaN_All();
 	}
@@ -174,10 +199,10 @@ public:
 	 * @param InTranslation The value to use for the translation component
 	 * @param InScale3D The value to use for the scale component
 	 */
-	FORCEINLINE FTransform(const FQuat& InRotation, const FVector& InTranslation, const FVector& InScale3D = FVector::OneVector)
+	FORCEINLINE TTransform(const TQuat<T>& InRotation, const TVector<T>& InTranslation, const TVector<T>& InScale3D = TVector<T>(1.f, 1.f, 1.f))
 	{
 		// Rotation = InRotation
-		Rotation =  VectorLoadAligned( &InRotation.X );
+		Rotation =  VectorLoadAligned( &InRotation );
 		// Translation = InTranslation
 		Translation = MakeVectorRegister(InTranslation.X, InTranslation.Y, InTranslation.Z, 0.0f );
 		// Scale3D = InScale3D
@@ -193,7 +218,7 @@ public:
 	 * @param InTranslation The value to use for the translation component
 	 * @param InScale3D The value to use for the scale component
 	 */
-	FORCEINLINE FTransform(const VectorRegister4Float& InRotation, const VectorRegister& InTranslation, const VectorRegister& InScale3D) 
+	FORCEINLINE TTransform(const TransformVectorRegister& InRotation, const TransformVectorRegister& InTranslation, const TransformVectorRegister& InScale3D)
 		: Rotation(InRotation),
 		Translation(InTranslation),
 		Scale3D(InScale3D)
@@ -208,11 +233,11 @@ public:
 	 * @param InTranslation The value to use for the translation component
 	 * @param InScale3D The value to use for the scale component
 	 */
-	FORCEINLINE FTransform(const FRotator& InRotation, const FVector& InTranslation, const FVector& InScale3D = FVector::OneVector)
+	FORCEINLINE TTransform(const FRotator& InRotation, const TVector<T>& InTranslation, const TVector<T>& InScale3D = TVector<T>(1.f, 1.f, 1.f))
 	{
-		FQuat InQuatRotation = InRotation.Quaternion();
+		TQuat<T> InQuatRotation = TQuat<T>::MakeFromRotator(InRotation);
 		// Rotation = InRotation
-		Rotation =  VectorLoadAligned( &InQuatRotation.X );
+		Rotation =  VectorLoadAligned( &InQuatRotation );
 		// Translation = InTranslation
 		Translation = MakeVectorRegister(InTranslation.X, InTranslation.Y, InTranslation.Z, 0.0f );
 		// Scale3D = InScale3D
@@ -224,12 +249,12 @@ public:
 	/**
 	 * Constructor with leaving uninitialized memory
 	 */
-	FORCEINLINE explicit FTransform(ENoInit) 
+	FORCEINLINE explicit TTransform(ENoInit) 
 	{
 		// Note: This can be used to track down initialization issues with bone transform arrays; but it will
 		// cause issues with transient fields such as RootMotionDelta that get initialized to 0 by default
 #if ENABLE_NAN_DIAGNOSTIC
-		float qnan = FMath::Log2(-5.3f);
+		FReal qnan = FMath::Log2(-5.3f);
 		check(FMath::IsNaN(qnan));
 		Translation = MakeVectorRegister(qnan, qnan, qnan, qnan);
 		Rotation = MakeVectorRegister(qnan, qnan, qnan, qnan);
@@ -238,18 +263,18 @@ public:
 	}
 
 	/**
-	 * Constructor for converting a Matrix (including scale) into a FTransform.
+	 * Constructor for converting a Matrix (including scale) into a TTransform.
 	 */
-	FORCEINLINE explicit FTransform(const FMatrix& InMatrix)
+	FORCEINLINE explicit TTransform(const TMatrix<T>& InMatrix)
 	{
 		SetFromMatrix(InMatrix);
 		DiagnosticCheckNaN_All();
 	}
 
 	/** Constructor that takes basis axes and translation */
-	FORCEINLINE FTransform(const FVector& InX, const FVector& InY, const FVector& InZ, const FVector& InTranslation)
+	FORCEINLINE TTransform(const TVector<T>& InX, const TVector<T>& InY, const TVector<T>& InZ, const TVector<T>& InTranslation)
 	{
-		SetFromMatrix(FMatrix(InX, InY, InZ, InTranslation));
+		SetFromMatrix(TMatrix<T>(InX, InY, InZ, InTranslation));
 		DiagnosticCheckNaN_All();
 	}
 
@@ -259,9 +284,9 @@ public:
 	CORE_API void DebugPrint() const;
 
 	/** Debug purpose only **/
-	bool DebugEqualMatrix(const FMatrix& Matrix) const;
+	bool DebugEqualMatrix(const TMatrix<T>& Matrix) const;
 
-	/** Convert FTransform contents to a string */
+	/** Convert TTransform contents to a string */
 	CORE_API FString ToHumanReadableString() const;
 
 	CORE_API FString ToString() const;
@@ -272,7 +297,7 @@ public:
 	/**
 	* Copy another Transform into this one
 	*/
-	FORCEINLINE FTransform& operator=(const FTransform& Other)
+	FORCEINLINE TTransform& operator=(const TTransform& Other)
 	{
 		this->Rotation = Other.Rotation;
 		this->Translation = Other.Translation;
@@ -282,39 +307,39 @@ public:
 	}
 
 
-	FORCEINLINE FMatrix ToMatrixWithScale() const
+	FORCEINLINE TMatrix<T> ToMatrixWithScale() const
 	{
 
-		FMatrix OutMatrix;
-		VectorRegister DiagonalsXYZ;
-		VectorRegister Adds;
-		VectorRegister Subtracts;
+		TMatrix<T> OutMatrix;
+		TransformVectorRegister DiagonalsXYZ;
+		TransformVectorRegister Adds;
+		TransformVectorRegister Subtracts;
 
 		ToMatrixInternal( DiagonalsXYZ, Adds, Subtracts );
-		const VectorRegister DiagonalsXYZ_W0 = VectorSet_W0(DiagonalsXYZ);
+		const TransformVectorRegister DiagonalsXYZ_W0 = VectorSet_W0(DiagonalsXYZ);
 
 		// OutMatrix.M[0][0] = (1.0f - (yy2 + zz2)) * Scale.X;    // Diagonal.X
 		// OutMatrix.M[0][1] = (xy2 + wz2) * Scale.X;             // Adds.X
 		// OutMatrix.M[0][2] = (xz2 - wy2) * Scale.X;             // Subtracts.Z
 		// OutMatrix.M[0][3] = 0.0f;                              // DiagonalsXYZ_W0.W
-		const VectorRegister AddX_DC_DiagX_DC = VectorShuffle(Adds, DiagonalsXYZ_W0, 0, 0, 0, 0);
-		const VectorRegister SubZ_DC_DiagW_DC = VectorShuffle(Subtracts, DiagonalsXYZ_W0, 2, 0, 3, 0);
-		const VectorRegister Row0 = VectorShuffle(AddX_DC_DiagX_DC, SubZ_DC_DiagW_DC, 2, 0, 0, 2);
+		const TransformVectorRegister AddX_DC_DiagX_DC = VectorShuffle(Adds, DiagonalsXYZ_W0, 0, 0, 0, 0);
+		const TransformVectorRegister SubZ_DC_DiagW_DC = VectorShuffle(Subtracts, DiagonalsXYZ_W0, 2, 0, 3, 0);
+		const TransformVectorRegister Row0 = VectorShuffle(AddX_DC_DiagX_DC, SubZ_DC_DiagW_DC, 2, 0, 0, 2);
 
 		// OutMatrix.M[1][0] = (xy2 - wz2) * Scale.Y;             // Subtracts.X
 		// OutMatrix.M[1][1] = (1.0f - (xx2 + zz2)) * Scale.Y;    // Diagonal.Y
 		// OutMatrix.M[1][2] = (yz2 + wx2) * Scale.Y;             // Adds.Y
 		// OutMatrix.M[1][3] = 0.0f;                            // DiagonalsXYZ_W0.W
-		const VectorRegister SubX_DC_DiagY_DC = VectorShuffle(Subtracts, DiagonalsXYZ_W0, 0, 0, 1, 0);
-		const VectorRegister AddY_DC_DiagW_DC = VectorShuffle(Adds, DiagonalsXYZ_W0, 1, 0, 3, 0);
-		const VectorRegister Row1 = VectorShuffle(SubX_DC_DiagY_DC, AddY_DC_DiagW_DC, 0, 2, 0, 2);
+		const TransformVectorRegister SubX_DC_DiagY_DC = VectorShuffle(Subtracts, DiagonalsXYZ_W0, 0, 0, 1, 0);
+		const TransformVectorRegister AddY_DC_DiagW_DC = VectorShuffle(Adds, DiagonalsXYZ_W0, 1, 0, 3, 0);
+		const TransformVectorRegister Row1 = VectorShuffle(SubX_DC_DiagY_DC, AddY_DC_DiagW_DC, 0, 2, 0, 2);
 
 		// OutMatrix.M[2][0] = (xz2 + wy2) * Scale.Z;             // Adds.Z
 		// OutMatrix.M[2][1] = (yz2 - wx2) * Scale.Z;             // Subtracts.Y
 		// OutMatrix.M[2][2] = (1.0f - (xx2 + yy2)) * Scale.Z;    // Diagonals.Z
 		// OutMatrix.M[2][3] = 0.0f;                              // DiagonalsXYZ_W0.W
-		const VectorRegister AddZ_DC_SubY_DC = VectorShuffle(Adds, Subtracts, 2, 0, 1, 0);
-		const VectorRegister Row2 = VectorShuffle(AddZ_DC_SubY_DC, DiagonalsXYZ_W0, 0, 2, 2, 3);
+		const TransformVectorRegister AddZ_DC_SubY_DC = VectorShuffle(Adds, Subtracts, 2, 0, 1, 0);
+		const TransformVectorRegister Row2 = VectorShuffle(AddZ_DC_SubY_DC, DiagonalsXYZ_W0, 0, 2, 2, 3);
 
 		VectorStoreAligned(Row0, &(OutMatrix.M[0][0]));
 		VectorStoreAligned(Row1, &(OutMatrix.M[1][0]));
@@ -324,7 +349,7 @@ public:
 		// OutMatrix.M[3][1] = Translation.Y;
 		// OutMatrix.M[3][2] = Translation.Z;
 		// OutMatrix.M[3][3] = 1.0f;
-		const VectorRegister Row3 = VectorSet_W1(Translation);
+		const TransformVectorRegister Row3 = VectorSet_W1(Translation);
 		VectorStoreAligned(Row3, &(OutMatrix.M[3][0]));
 
 		return OutMatrix;
@@ -333,7 +358,7 @@ public:
 	/**
 	* Convert this Transform to matrix with scaling and compute the inverse of that.
 	*/
-	FORCEINLINE FMatrix ToInverseMatrixWithScale() const
+	FORCEINLINE TMatrix<T> ToInverseMatrixWithScale() const
 	{
 		// todo: optimize
 		return ToMatrixWithScale().Inverse();
@@ -342,54 +367,54 @@ public:
 	/**
 	* Convert this Transform to inverse.
 	*/
-	FORCEINLINE FTransform Inverse() const
+	FORCEINLINE TTransform Inverse() const
 	{
-		// Replacement of Inverse of FMatrix
+		// Replacement of Inverse of TMatrix<T>
 		if (VectorAnyGreaterThan(VectorAbs(Scale3D), GlobalVectorConstants::SmallNumber))
 		{
 			return InverseFast();
 		}
 		else
 		{
-			return FTransform::Identity;
+			return TTransform<T>::Identity;
 		}
 	}
 
 	/**
 	* Convert this Transform to a transformation matrix, ignoring its scaling
 	*/
-	FORCEINLINE FMatrix ToMatrixNoScale() const
+	FORCEINLINE TMatrix<T> ToMatrixNoScale() const
 	{
-		FMatrix OutMatrix;
-		VectorRegister DiagonalsXYZ;
-		VectorRegister Adds;
-		VectorRegister Subtracts;
+		TMatrix<T> OutMatrix;
+		TransformVectorRegister DiagonalsXYZ;
+		TransformVectorRegister Adds;
+		TransformVectorRegister Subtracts;
 
 		ToMatrixInternalNoScale( DiagonalsXYZ, Adds, Subtracts );
-		const VectorRegister DiagonalsXYZ_W0 = VectorSet_W0(DiagonalsXYZ);
+		const TransformVectorRegister DiagonalsXYZ_W0 = VectorSet_W0(DiagonalsXYZ);
 
 		// OutMatrix.M[0][0] = (1.0f - (yy2 + zz2));			// Diagonal.X
 		// OutMatrix.M[0][1] = (xy2 + wz2);						// Adds.X
 		// OutMatrix.M[0][2] = (xz2 - wy2);						// Subtracts.Z
 		// OutMatrix.M[0][3] = 0.0f;                            // DiagonalsXYZ_W0.W
-		const VectorRegister AddX_DC_DiagX_DC = VectorShuffle(Adds, DiagonalsXYZ_W0, 0, 0, 0, 0);
-		const VectorRegister SubZ_DC_DiagW_DC = VectorShuffle(Subtracts, DiagonalsXYZ_W0, 2, 0, 3, 0);
-		const VectorRegister Row0 = VectorShuffle(AddX_DC_DiagX_DC, SubZ_DC_DiagW_DC, 2, 0, 0, 2);
+		const TransformVectorRegister AddX_DC_DiagX_DC = VectorShuffle(Adds, DiagonalsXYZ_W0, 0, 0, 0, 0);
+		const TransformVectorRegister SubZ_DC_DiagW_DC = VectorShuffle(Subtracts, DiagonalsXYZ_W0, 2, 0, 3, 0);
+		const TransformVectorRegister Row0 = VectorShuffle(AddX_DC_DiagX_DC, SubZ_DC_DiagW_DC, 2, 0, 0, 2);
 
 		// OutMatrix.M[1][0] = (xy2 - wz2);			            // Subtracts.X
 		// OutMatrix.M[1][1] = (1.0f - (xx2 + zz2));		    // Diagonal.Y
 		// OutMatrix.M[1][2] = (yz2 + wx2);						// Adds.Y
 		// OutMatrix.M[1][3] = 0.0f;                            // DiagonalsXYZ_W0.W
-		const VectorRegister SubX_DC_DiagY_DC = VectorShuffle(Subtracts, DiagonalsXYZ_W0, 0, 0, 1, 0);
-		const VectorRegister AddY_DC_DiagW_DC = VectorShuffle(Adds, DiagonalsXYZ_W0, 1, 0, 3, 0);
-		const VectorRegister Row1 = VectorShuffle(SubX_DC_DiagY_DC, AddY_DC_DiagW_DC, 0, 2, 0, 2);
+		const TransformVectorRegister SubX_DC_DiagY_DC = VectorShuffle(Subtracts, DiagonalsXYZ_W0, 0, 0, 1, 0);
+		const TransformVectorRegister AddY_DC_DiagW_DC = VectorShuffle(Adds, DiagonalsXYZ_W0, 1, 0, 3, 0);
+		const TransformVectorRegister Row1 = VectorShuffle(SubX_DC_DiagY_DC, AddY_DC_DiagW_DC, 0, 2, 0, 2);
 
 		// OutMatrix.M[2][0] = (xz2 + wy2);						// Adds.Z
 		// OutMatrix.M[2][1] = (yz2 - wx2);						// Subtracts.Y
 		// OutMatrix.M[2][2] = (1.0f - (xx2 + yy2));		    // Diagonals.Z
 		// OutMatrix.M[2][3] = 0.0f;                            // DiagonalsXYZ_W0.W
-		const VectorRegister AddZ_DC_SubY_DC = VectorShuffle(Adds, Subtracts, 2, 0, 1, 0);
-		const VectorRegister Row2 = VectorShuffle(AddZ_DC_SubY_DC, DiagonalsXYZ_W0, 0, 2, 2, 3);
+		const TransformVectorRegister AddZ_DC_SubY_DC = VectorShuffle(Adds, Subtracts, 2, 0, 1, 0);
+		const TransformVectorRegister Row2 = VectorShuffle(AddZ_DC_SubY_DC, DiagonalsXYZ_W0, 0, 2, 2, 3);
 
 		VectorStoreAligned(Row0, &(OutMatrix.M[0][0]));
 		VectorStoreAligned(Row1, &(OutMatrix.M[1][0]));
@@ -399,14 +424,14 @@ public:
 		// OutMatrix.M[3][1] = Translation.Y;
 		// OutMatrix.M[3][2] = Translation.Z;
 		// OutMatrix.M[3][3] = 1.0f;
-		const VectorRegister Row3 = VectorSet_W1(Translation);
+		const TransformVectorRegister Row3 = VectorSet_W1(Translation);
 		VectorStoreAligned(Row3, &(OutMatrix.M[3][0]));
 
 		return OutMatrix;
 	}
 
 	/** Set this transform to the weighted blend of the supplied two transforms. */
-	FORCEINLINE void Blend(const FTransform& Atom1, const FTransform& Atom2, float Alpha)
+	FORCEINLINE void Blend(const TTransform<T>& Atom1, const TTransform<T>& Atom2, float Alpha)
 	{
 #if !(UE_BUILD_SHIPPING || UE_BUILD_TEST) && WITH_EDITORONLY_DATA
 		// Check that all bone atoms coming from animation are normalized
@@ -427,13 +452,12 @@ public:
 		else
 		{
 			// Simple linear interpolation for translation and scale.			
-			ScalarRegister BlendWeight = ScalarRegister(Alpha);
+			TransformVectorRegister BlendWeight = VectorLoadFloat1(&Alpha);
 
-			Translation = FMath::Lerp(Atom1.Translation, Atom2.Translation, BlendWeight.Value);			
-			Scale3D = FMath::Lerp(Atom1.Scale3D, Atom2.Scale3D, BlendWeight.Value);			
+			Translation = FMath::Lerp<TransformVectorRegister>(Atom1.Translation, Atom2.Translation, BlendWeight);
+			Scale3D = FMath::Lerp<TransformVectorRegister>(Atom1.Scale3D, Atom2.Scale3D, BlendWeight);
 
-			VectorRegister4Float BlendWeightFloat = VectorLoadFloat1(&Alpha);
-			VectorRegister4Float VRotation = VectorLerpQuat(Atom1.Rotation, Atom2.Rotation, BlendWeightFloat);
+			TransformVectorRegister VRotation = VectorLerpQuat(Atom1.Rotation, Atom2.Rotation, BlendWeight);
 
 			// ..and renormalize
 			Rotation = VectorNormalizeQuaternion(VRotation);
@@ -443,7 +467,7 @@ public:
 	}
 
 	/** Set this Transform to the weighted blend of it and the supplied Transform. */
-	FORCEINLINE void BlendWith(const FTransform& OtherAtom, float Alpha)
+	FORCEINLINE void BlendWith(const TTransform<T>& OtherAtom, float Alpha)
 	{
 #if !(UE_BUILD_SHIPPING || UE_BUILD_TEST) && WITH_EDITORONLY_DATA
 		// Check that all bone atoms coming from animation are normalized
@@ -461,13 +485,12 @@ public:
 			else 
 			{
 				// Simple linear interpolation for translation and scale.				
-				ScalarRegister BlendWeight = ScalarRegister(Alpha);
-				Translation = FMath::Lerp(Translation, OtherAtom.Translation, BlendWeight.Value);
+				TransformVectorRegister BlendWeight = VectorLoadFloat1(&Alpha);
+				Translation = FMath::Lerp<TransformVectorRegister>(Translation, OtherAtom.Translation, BlendWeight);
 				
-				Scale3D = FMath::Lerp(Scale3D, OtherAtom.Scale3D, BlendWeight.Value);
+				Scale3D = FMath::Lerp<TransformVectorRegister>(Scale3D, OtherAtom.Scale3D, BlendWeight);
 				
-				VectorRegister4Float BlendWeightFloat = VectorLoadFloat1(&Alpha);
-				VectorRegister4Float VRotation = VectorLerpQuat(Rotation, OtherAtom.Rotation, BlendWeightFloat);
+				TransformVectorRegister VRotation = VectorLerpQuat(Rotation, OtherAtom.Rotation, BlendWeight);
 				
 				// ..and renormalize
 				Rotation = VectorNormalizeQuaternion(VRotation);
@@ -483,12 +506,12 @@ public:
 	 * Use only within blends!!
 	 * Rotation part is NOT normalized!!
 	 */
-	FORCEINLINE FTransform operator+(const FTransform& Atom) const
+	FORCEINLINE TTransform<T> operator+(const TTransform<T>& Atom) const
 	{
-		return FTransform( VectorAdd(Rotation, Atom.Rotation), VectorAdd(Translation, Atom.Translation ), VectorAdd( Scale3D, Atom.Scale3D ) );
+		return TTransform<T>( VectorAdd(Rotation, Atom.Rotation), VectorAdd(Translation, Atom.Translation ), VectorAdd( Scale3D, Atom.Scale3D ) );
 	}
 
-	FORCEINLINE FTransform& operator+=(const FTransform& Atom)
+	FORCEINLINE TTransform<T>& operator+=(const TTransform<T>& Atom)
 	{ 
 		Translation = VectorAdd(Translation, Atom.Translation);
 		Rotation = VectorAdd(Rotation, Atom.Rotation);
@@ -497,95 +520,97 @@ public:
 		return *this;
 	}
 
-	FORCEINLINE FTransform operator*(const ScalarRegister& Mult) const
+	FORCEINLINE TTransform<T> operator*(const ScalarRegister& MultScalar) const
 	{		
-		return FTransform( VectorMultiply(Rotation, MakeVectorRegisterFloatFromDouble(Mult.Value)), VectorMultiply(Translation, Mult), VectorMultiply(Scale3D, Mult) );
+		const TransformVectorRegister Mult = MultScalar.Value;
+		return TTransform<T>( VectorMultiply(Rotation, Mult), VectorMultiply(Translation, Mult), VectorMultiply(Scale3D, Mult) );
 	}
 
-	FORCEINLINE FTransform& operator*=(const ScalarRegister& Mult)
+	FORCEINLINE TTransform<T>& operator*=(const ScalarRegister& MultScalar)
 	{			
+		const TransformVectorRegister Mult = MultScalar.Value;
 		Translation= VectorMultiply(Translation, Mult);
-		Rotation = VectorMultiply(Rotation, MakeVectorRegisterFloatFromDouble(Mult.Value));
+		Rotation = VectorMultiply(Rotation, Mult);
 		Scale3D = VectorMultiply(Scale3D, Mult);
 
 		return *this;
 	}
 
-	FORCEINLINE FTransform		operator*(const FTransform& Other) const;
-	FORCEINLINE void			operator*=(const FTransform& Other);
-	FORCEINLINE FTransform		operator*(const FQuat& Other) const;
-	FORCEINLINE void			operator*=(const FQuat& Other);
+	FORCEINLINE TTransform<T>	operator*(const TTransform<T>& Other) const;
+	FORCEINLINE void			operator*=(const TTransform<T>& Other);
+	FORCEINLINE TTransform<T>	operator*(const TQuat<T>& Other) const;
+	FORCEINLINE void			operator*=(const TQuat<T>& Other);
 
-	FORCEINLINE static bool AnyHasNegativeScale(const FVector& InScale3D, const FVector& InOtherScale3D);
-	FORCEINLINE void ScaleTranslation(const FVector& InScale3D);
-	FORCEINLINE void ScaleTranslation(const float& Scale);
-	FORCEINLINE void RemoveScaling(float Tolerance=SMALL_NUMBER);
+	FORCEINLINE static bool AnyHasNegativeScale(const TVector<T>& InScale3D, const TVector<T>& InOtherScale3D);
+	FORCEINLINE void ScaleTranslation(const TVector<T>& InScale3D);
+	FORCEINLINE void ScaleTranslation(const T& Scale);
+	FORCEINLINE void RemoveScaling(T Tolerance=SMALL_NUMBER);
 	FORCEINLINE float GetMaximumAxisScale() const;
 	FORCEINLINE float GetMinimumAxisScale() const;
 	// Inverse does not work well with VQS format(in particular non-uniform), so removing it, but made two below functions to be used instead. 
 
 	/*******************************************************************************************
-	 * The below 2 functions are the ones to get delta transform and return FTransform format that can be concatenated
+	 * The below 2 functions are the ones to get delta transform and return TTransform<T> format that can be concatenated
 	 * Inverse itself can't concatenate with VQS format(since VQS always transform from S->Q->T, where inverse happens from T(-1)->Q(-1)->S(-1))
 	 * So these 2 provides ways to fix this
 	 * GetRelativeTransform returns this*Other(-1) and parameter is Other(not Other(-1))
 	 * GetRelativeTransformReverse returns this(-1)*Other, and parameter is Other. 
 	 *******************************************************************************************/
-	CORE_API FTransform GetRelativeTransform(const FTransform& Other) const;
-	CORE_API FTransform GetRelativeTransformReverse(const FTransform& Other) const;
+	CORE_API TTransform<T> GetRelativeTransform(const TTransform<T>& Other) const;
+	CORE_API TTransform<T> GetRelativeTransformReverse(const TTransform<T>& Other) const;
 	/**
 	 * Set current transform and the relative to ParentTransform.
-	 * Equates to This = This->GetRelativeTransform(Parent), but saves the intermediate FTransform storage and copy.
+	 * Equates to This = This->GetRelativeTransform(Parent), but saves the intermediate TTransform<T> storage and copy.
 	 */
-	CORE_API void		SetToRelativeTransform(const FTransform& ParentTransform);
+	CORE_API void		SetToRelativeTransform(const TTransform<T>& ParentTransform);
 
 	FORCEINLINE FVector4	TransformFVector4(const FVector4& V) const;
 	FORCEINLINE FVector4	TransformFVector4NoScale(const FVector4& V) const;
-	FORCEINLINE FVector		TransformPosition(const FVector& V) const;
-	FORCEINLINE FVector		TransformPositionNoScale(const FVector& V) const;
+	FORCEINLINE TVector<T>		TransformPosition(const TVector<T>& V) const;
+	FORCEINLINE TVector<T>		TransformPositionNoScale(const TVector<T>& V) const;
 
 
 	/** Inverts the transform and then transforms V - correctly handles scaling in this transform. */
-	FORCEINLINE FVector		InverseTransformPosition(const FVector &V) const;
-	FORCEINLINE FVector		InverseTransformPositionNoScale(const FVector &V) const;
-	FORCEINLINE FVector		TransformVector(const FVector& V) const;
-	FORCEINLINE FVector		TransformVectorNoScale(const FVector& V) const;
+	FORCEINLINE TVector<T>		InverseTransformPosition(const TVector<T> &V) const;
+	FORCEINLINE TVector<T>		InverseTransformPositionNoScale(const TVector<T> &V) const;
+	FORCEINLINE TVector<T>		TransformVector(const TVector<T>& V) const;
+	FORCEINLINE TVector<T>		TransformVectorNoScale(const TVector<T>& V) const;
 
 	/** 
 	 *	Transform a direction vector by the inverse of this matrix - will not take into account translation part.
 	 *	If you want to transform a surface normal (or plane) and correctly account for non-uniform scaling you should use TransformByUsingAdjointT with adjoint of matrix inverse.
 	 */
-	FORCEINLINE FVector InverseTransformVector(const FVector &V) const;
-	FORCEINLINE FVector InverseTransformVectorNoScale(const FVector &V) const;
+	FORCEINLINE TVector<T> InverseTransformVector(const TVector<T> &V) const;
+	FORCEINLINE TVector<T> InverseTransformVectorNoScale(const TVector<T> &V) const;
 
 	/**
 	* Transform a rotation.
 	* For example if this is a LocalToWorld transform, TransformRotation(Q) would transform Q from local to world space.
 	*/
-	FORCEINLINE FQuat TransformRotation(const FQuat& Q) const;
+	FORCEINLINE TQuat<T> TransformRotation(const TQuat<T>& Q) const;
 
 	/**
 	* Inverse transform a rotation.
 	* For example if this is a LocalToWorld transform, InverseTransformRotation(Q) would transform Q from world to local space.
 	*/
-	FORCEINLINE FQuat InverseTransformRotation(const FQuat& Q) const;
+	FORCEINLINE TQuat<T> InverseTransformRotation(const TQuat<T>& Q) const;
 
-	FORCEINLINE FTransform	GetScaled(float Scale) const;
-	FORCEINLINE FTransform	GetScaled(FVector Scale) const;
-	FORCEINLINE FVector		GetScaledAxis(EAxis::Type InAxis) const;
-	FORCEINLINE FVector		GetUnitAxis(EAxis::Type InAxis) const;
+	FORCEINLINE TTransform<T>	GetScaled(T Scale) const;
+	FORCEINLINE TTransform<T>	GetScaled(TVector<T> Scale) const;
+	FORCEINLINE TVector<T>		GetScaledAxis(EAxis::Type InAxis) const;
+	FORCEINLINE TVector<T>		GetUnitAxis(EAxis::Type InAxis) const;
 	FORCEINLINE void		Mirror(EAxis::Type MirrorAxis, EAxis::Type FlipAxis);
-	FORCEINLINE static FVector	GetSafeScaleReciprocal(const FVector& InScale, float Tolerance=SMALL_NUMBER);
+	FORCEINLINE static TVector<T>	GetSafeScaleReciprocal(const TVector<T>& InScale, T Tolerance=SMALL_NUMBER);
 
 
-	FORCEINLINE FVector GetLocation() const
+	FORCEINLINE TVector<T> GetLocation() const
 	{
 		return GetTranslation();
 	}
 
 	FORCEINLINE FRotator Rotator() const
 	{
-		FQuat OutRotation;
+		TQuat<T> OutRotation;
 		VectorStoreAligned(Rotation, &OutRotation);
 		return OutRotation.Rotator();
 	}
@@ -600,7 +625,7 @@ public:
 	}
 
 	/** Set the translation of this transformation */
-	FORCEINLINE void SetLocation(const FVector& Origin)
+	FORCEINLINE void SetLocation(const TVector<T>& Origin)
 	{		
 		Translation = VectorLoadFloat3_W0(&Origin);
 		DiagnosticCheckNaN_Translate();
@@ -644,12 +669,12 @@ public:
 	}
 
 	// Serializer.
-	inline friend FArchive& operator<<(FArchive& Ar,FTransform& M)
+	inline friend FArchive& operator<<(FArchive& Ar,TTransform<T>& M)
 	{
-		//@LWC_TODO: fix this when FQuat finally are double 
-		Ar << *reinterpret_cast<FVector4*>(&(M.Rotation));
-		Ar << *reinterpret_cast<FVector*>(&(M.Translation));
-		Ar << *reinterpret_cast<FVector*>(&(M.Scale3D));
+		// LWC_TODO: Serializer
+		Ar << *reinterpret_cast<TQuat<T>*>(&(M.Rotation));
+		Ar << *reinterpret_cast<TVector<T>*>(&(M.Translation));
+		Ar << *reinterpret_cast<TVector<T>*>(&(M.Scale3D));
 		
 		if (Ar.IsLoading())
 		{
@@ -661,11 +686,11 @@ public:
 
 	// Binary comparison operators.
 	/*
-	bool operator==( const FTransform& Other ) const
+	bool operator==( const TTransform<T>& Other ) const
 	{
 		return Rotation==Other.Rotation && Translation==Other.Translation && Scale3D==Other.Scale3D;
 	}
-	bool operator!=( const FTransform& Other ) const
+	bool operator!=( const TTransform<T>& Other ) const
 	{
 		return Rotation!=Other.Rotation || Translation!=Other.Translation || Scale3D!=Other.Scale3D;
 	}
@@ -673,102 +698,102 @@ public:
 
 private:
 
-	FORCEINLINE static bool Private_AnyHasNegativeScale(const VectorRegister& InScale3D, const  VectorRegister& InOtherScale3D)
+	FORCEINLINE static bool Private_AnyHasNegativeScale(const TransformVectorRegister& InScale3D, const TransformVectorRegister& InOtherScale3D)
 	{
-		return !!VectorAnyLesserThan(VectorMin(InScale3D, InOtherScale3D), VectorZeroVectorRegister());
+		return !!VectorAnyLesserThan(VectorMin(InScale3D, InOtherScale3D), (TransformVectorRegister)GlobalVectorConstants::FloatZero);
 	}
 
-	FORCEINLINE bool Private_RotationEquals( const VectorRegister4Float& InRotation, float Tolerance = KINDA_SMALL_NUMBER) const
+	FORCEINLINE bool Private_RotationEquals( const TransformVectorRegister& InRotation, const FReal ToleranceScalar = KINDA_SMALL_NUMBER) const
 	{			
+		const TransformVectorRegister Tolerance = VectorLoadFloat1(&ToleranceScalar);
 		// !( (FMath::Abs(X-Q.X) > Tolerance) || (FMath::Abs(Y-Q.Y) > Tolerance) || (FMath::Abs(Z-Q.Z) > Tolerance) || (FMath::Abs(W-Q.W) > Tolerance) )
-		const VectorRegister4Float RotationSub = VectorAbs(VectorSubtract(Rotation, InRotation));
+		const TransformVectorRegister RotationSub = VectorAbs(VectorSubtract(Rotation, InRotation));		
 		// !( (FMath::Abs(X+Q.X) > Tolerance) || (FMath::Abs(Y+Q.Y) > Tolerance) || (FMath::Abs(Z+Q.Z) > Tolerance) || (FMath::Abs(W+Q.W) > Tolerance) )
-		const VectorRegister4Float RotationAdd = VectorAbs(VectorAdd(Rotation, InRotation));
-		const VectorRegister4Float ToleranceFloat = VectorLoadFloat1(&Tolerance);
-		return !VectorAnyGreaterThan(RotationSub, ToleranceFloat) || !VectorAnyGreaterThan(RotationAdd, ToleranceFloat);
+		const TransformVectorRegister RotationAdd = VectorAbs(VectorAdd(Rotation, InRotation));
+		return !VectorAnyGreaterThan(RotationSub, Tolerance) || !VectorAnyGreaterThan(RotationAdd, Tolerance);
 	}
 
-	FORCEINLINE bool Private_TranslationEquals( const VectorRegister& InTranslation, const ScalarRegister& Tolerance = ScalarRegister(GlobalVectorConstants::KindaSmallNumber)) const
+	FORCEINLINE bool Private_TranslationEquals( const TransformVectorRegister& InTranslation, const FReal ToleranceScalar = KINDA_SMALL_NUMBER) const
 	{			
+		const TransformVectorRegister Tolerance = VectorLoadFloat1(&ToleranceScalar);
 		// !( (FMath::Abs(X-V.X) > Tolerance) || (FMath::Abs(Y-V.Y) > Tolerance) || (FMath::Abs(Z-V.Z) > Tolerance) )
-		const VectorRegister TranslationDiff = VectorAbs(VectorSubtract(Translation, InTranslation));		
-		return !VectorAnyGreaterThan(TranslationDiff, Tolerance.Value);
+		const TransformVectorRegister TranslationDiff = VectorAbs(VectorSubtract(Translation, InTranslation));		
+		return !VectorAnyGreaterThan(TranslationDiff, Tolerance);
 	}
 
-	FORCEINLINE bool Private_Scale3DEquals( const VectorRegister& InScale3D, const ScalarRegister& Tolerance = ScalarRegister(GlobalVectorConstants::KindaSmallNumber)) const
+	FORCEINLINE bool Private_Scale3DEquals( const TransformVectorRegister& InScale3D, const FReal ToleranceScalar = KINDA_SMALL_NUMBER) const
 	{
+		const TransformVectorRegister Tolerance = VectorLoadFloat1(&ToleranceScalar);
 		// !( (FMath::Abs(X-V.X) > Tolerance) || (FMath::Abs(Y-V.Y) > Tolerance) || (FMath::Abs(Z-V.Z) > Tolerance) )
-		const VectorRegister ScaleDiff = VectorAbs(VectorSubtract(Scale3D, InScale3D));
-		return !VectorAnyGreaterThan(ScaleDiff, Tolerance.Value);
+		const TransformVectorRegister ScaleDiff = VectorAbs(VectorSubtract(Scale3D, InScale3D));
+		return !VectorAnyGreaterThan(ScaleDiff, Tolerance);
 	}
 
 public:
 
 	// Test if A's rotation equals B's rotation, within a tolerance. Preferred over "A.GetRotation().Equals(B.GetRotation())" because it is faster on some platforms.
-	FORCEINLINE static bool AreRotationsEqual(const FTransform& A, const FTransform& B, float Tolerance=KINDA_SMALL_NUMBER)
+	FORCEINLINE static bool AreRotationsEqual(const TTransform<T>& A, const TTransform<T>& B, FReal Tolerance=KINDA_SMALL_NUMBER)
 	{
 		return A.Private_RotationEquals(B.Rotation, Tolerance);
 	}
 
-	// Test if A's translation equals B's translation, within a tolerance. Preferred over "A.GetTranslation().Equals(B.GetTranslation())" because it avoids VectorRegister->FVector conversion.
-	FORCEINLINE static bool AreTranslationsEqual(const FTransform& A, const FTransform& B, float Tolerance=KINDA_SMALL_NUMBER)
+	// Test if A's translation equals B's translation, within a tolerance. Preferred over "A.GetTranslation().Equals(B.GetTranslation())" because it avoids TransformVectorRegister->TVector<T> conversion.
+	FORCEINLINE static bool AreTranslationsEqual(const TTransform<T>& A, const TTransform<T>& B, FReal Tolerance=KINDA_SMALL_NUMBER)
 	{
-		return A.Private_TranslationEquals(B.Translation, ScalarRegister(Tolerance));
+		return A.Private_TranslationEquals(B.Translation, Tolerance);
 	}
 
-	// Test if A's scale equals B's scale, within a tolerance. Preferred over "A.GetScale3D().Equals(B.GetScale3D())" because it avoids VectorRegister->FVector conversion.
-	FORCEINLINE static bool AreScale3DsEqual(const FTransform& A, const FTransform& B, float Tolerance=KINDA_SMALL_NUMBER)
+	// Test if A's scale equals B's scale, within a tolerance. Preferred over "A.GetScale3D().Equals(B.GetScale3D())" because it avoids TransformVectorRegister->TVector<T> conversion.
+	FORCEINLINE static bool AreScale3DsEqual(const TTransform<T>& A, const TTransform<T>& B, FReal Tolerance=KINDA_SMALL_NUMBER)
 	{
-		return A.Private_Scale3DEquals(B.Scale3D, ScalarRegister(Tolerance));
+		return A.Private_Scale3DEquals(B.Scale3D, Tolerance);
 	}
 
 
 	// Test if this Transform's rotation equals another's rotation, within a tolerance. Preferred over "GetRotation().Equals(Other.GetRotation())" because it is faster on some platforms.
-	FORCEINLINE bool RotationEquals(const FTransform& Other, float Tolerance=KINDA_SMALL_NUMBER) const
+	FORCEINLINE bool RotationEquals(const TTransform<T>& Other, FReal Tolerance=KINDA_SMALL_NUMBER) const
 	{
 		return AreRotationsEqual(*this, Other, Tolerance);
 	}
 
-	// Test if this Transform's translation equals another's translation, within a tolerance. Preferred over "GetTranslation().Equals(Other.GetTranslation())" because it avoids VectorRegister->FVector conversion.
-	FORCEINLINE bool TranslationEquals(const FTransform& Other, float Tolerance=KINDA_SMALL_NUMBER) const
+	// Test if this Transform's translation equals another's translation, within a tolerance. Preferred over "GetTranslation().Equals(Other.GetTranslation())" because it avoids TransformVectorRegister->TVector<T> conversion.
+	FORCEINLINE bool TranslationEquals(const TTransform<T>& Other, FReal Tolerance=KINDA_SMALL_NUMBER) const
 	{
 		return AreTranslationsEqual(*this, Other, Tolerance);
 	}
 
-	// Test if this Transform's scale equals another's scale, within a tolerance. Preferred over "GetScale3D().Equals(Other.GetScale3D())" because it avoids VectorRegister->FVector conversion.
-	FORCEINLINE bool Scale3DEquals(const FTransform& Other, float Tolerance=KINDA_SMALL_NUMBER) const
+	// Test if this Transform's scale equals another's scale, within a tolerance. Preferred over "GetScale3D().Equals(Other.GetScale3D())" because it avoids TransformVectorRegister->TVector<T> conversion.
+	FORCEINLINE bool Scale3DEquals(const TTransform<T>& Other, FReal Tolerance=KINDA_SMALL_NUMBER) const
 	{
 		return AreScale3DsEqual(*this, Other, Tolerance);
 	}
 
 	// Test if all components of the transforms are equal, within a tolerance.
-	FORCEINLINE bool Equals(const FTransform& Other, float Tolerance=KINDA_SMALL_NUMBER) const
+	FORCEINLINE bool Equals(const TTransform<T>& Other, FReal Tolerance=KINDA_SMALL_NUMBER) const
 	{
-		const ScalarRegister ToleranceRegister(Tolerance);
-		return Private_TranslationEquals(Other.Translation, ToleranceRegister) && Private_RotationEquals(Other.Rotation, Tolerance) && Private_Scale3DEquals(Other.Scale3D, ToleranceRegister);
+		return Private_TranslationEquals(Other.Translation, Tolerance) && Private_RotationEquals(Other.Rotation, Tolerance) && Private_Scale3DEquals(Other.Scale3D, Tolerance);
 	}
 
 	// Test if all components of the transform property are equal.
-	FORCEINLINE bool Identical(const FTransform* Other, uint32 PortFlags) const
+	FORCEINLINE bool Identical(const TTransform<T>* Other, uint32 PortFlags) const
 	{
 		return Equals(*Other, 0.f);
 	}
 
 	// Test if rotation and translation components of the transforms are equal, within a tolerance.
-	FORCEINLINE bool EqualsNoScale(const FTransform& Other, float Tolerance=KINDA_SMALL_NUMBER) const
+	FORCEINLINE bool EqualsNoScale(const TTransform<T>& Other, FReal Tolerance=KINDA_SMALL_NUMBER) const
 	{
-		const ScalarRegister ToleranceRegister(Tolerance);
-		return Private_TranslationEquals(Other.Translation, ToleranceRegister) && Private_RotationEquals(Other.Rotation, Tolerance);
+		return Private_TranslationEquals(Other.Translation, Tolerance) && Private_RotationEquals(Other.Rotation, Tolerance);
 	}
 
-	FORCEINLINE static void Multiply(FTransform* OutTransform, const FTransform* A, const FTransform* B);
+	FORCEINLINE static void Multiply(TTransform<T>* OutTransform, const TTransform<T>* A, const TTransform<T>* B);
 	/**
 	 * Sets the components
 	 * @param InRotation The new value for the Rotation component
 	 * @param InTranslation The new value for the Translation component
 	 * @param InScale3D The new value for the Scale3D component
 	 */
-	FORCEINLINE void SetComponents(const FQuat& InRotation, const FVector& InTranslation, const FVector& InScale3D) 
+	FORCEINLINE void SetComponents(const TQuat<T>& InRotation, const TVector<T>& InTranslation, const TVector<T>& InScale3D) 
 	{
 		Rotation = VectorLoadAligned(&InRotation);
 		Translation = VectorLoadFloat3_W0(&InTranslation);
@@ -788,16 +813,16 @@ public:
 		// Rotation = {0,0,0,1)
 		Rotation = GlobalVectorConstants::Float0001; //VectorSet_W1( VectorZero() );
 		// Translation = {0,0,0,0)
-		Translation = VectorZeroVectorRegister();
+		Translation = GlobalVectorConstants::FloatZero;
 		// Scale3D = {1,1,1,0);
-		Scale3D = GlobalVectorConstants::Vector1110; //VectorSet_W0( VectorOne() );
+		Scale3D = GlobalVectorConstants::Float1110; //VectorSet_W0( VectorOne() );
 	}
 
 	/**
 	 * Scales the Scale3D component by a new factor
 	 * @param Scale3DMultiplier The value to multiply Scale3D with
 	 */
-	FORCEINLINE void MultiplyScale3D(const FVector& Scale3DMultiplier)
+	FORCEINLINE void MultiplyScale3D(const TVector<T>& Scale3DMultiplier)
 	{
 		Scale3D = VectorMultiply(Scale3D, VectorLoadFloat3_W0(&Scale3DMultiplier));
 		DiagnosticCheckNaN_Scale3D();
@@ -807,14 +832,14 @@ public:
 	 * Sets the translation component
 	 * @param NewTranslation The new value for the translation component
 	 */
-	FORCEINLINE void SetTranslation(const FVector& NewTranslation)
+	FORCEINLINE void SetTranslation(const TVector<T>& NewTranslation)
 	{
 		Translation = VectorLoadFloat3_W0(&NewTranslation);
 		DiagnosticCheckNaN_Translate();
 	}
 
-	/** Copy translation from another FTransform. */
-	FORCEINLINE void CopyTranslation(const FTransform& Other)
+	/** Copy translation from another TTransform<T>. */
+	FORCEINLINE void CopyTranslation(const TTransform<T>& Other)
 	{
 		Translation = Other.Translation;
 	}
@@ -823,7 +848,7 @@ public:
 	 * Concatenates another rotation to this transformation 
 	 * @param DeltaRotation The rotation to concatenate in the following fashion: Rotation = Rotation * DeltaRotation
 	 */
-	FORCEINLINE void ConcatenateRotation(const FQuat& DeltaRotation)
+	FORCEINLINE void ConcatenateRotation(const TQuat<T>& DeltaRotation)
 	{		
 		Rotation = VectorQuaternionMultiply2(Rotation, VectorLoadAligned(&DeltaRotation));		
 		DiagnosticCheckNaN_Rotate();
@@ -833,7 +858,7 @@ public:
 	 * Adjusts the translation component of this transformation 
 	 * @param DeltaTranslation The translation to add in the following fashion: Translation += DeltaTranslation
 	 */
-	FORCEINLINE void AddToTranslation(const FVector& DeltaTranslation)
+	FORCEINLINE void AddToTranslation(const TVector<T>& DeltaTranslation)
 	{		
 		Translation = VectorAdd(Translation, VectorLoadFloat3_W0(&DeltaTranslation));
 		DiagnosticCheckNaN_Translate();
@@ -843,9 +868,9 @@ public:
 	 * Add the translations from two FTransforms and return the result.
 	 * @return A.Translation + B.Translation
 	 */
-	FORCEINLINE static FVector AddTranslations(const FTransform& A, const FTransform& B)
+	FORCEINLINE static TVector<T> AddTranslations(const TTransform<T>& A, const TTransform<T>& B)
 	{
-		FVector Result;
+		TVector<T> Result;
 		VectorStoreFloat3(VectorAdd(A.Translation, B.Translation), &Result);
 		return Result;
 	}
@@ -854,9 +879,9 @@ public:
 	 * Subtract translations from two FTransforms and return the difference.
 	 * @return A.Translation - B.Translation.
 	 */
-	FORCEINLINE static FVector SubtractTranslations(const FTransform& A, const FTransform& B)
+	FORCEINLINE static TVector<T> SubtractTranslations(const TTransform<T>& A, const TTransform<T>& B)
 	{
-		FVector Result;
+		TVector<T> Result;
 		VectorStoreFloat3(VectorSubtract(A.Translation, B.Translation), &Result);
 		return Result;
 	}
@@ -865,14 +890,14 @@ public:
 	 * Sets the rotation component
 	 * @param NewRotation The new value for the rotation component
 	 */
-	FORCEINLINE void SetRotation(const FQuat& NewRotation)
+	FORCEINLINE void SetRotation(const TQuat<T>& NewRotation)
 	{
 		Rotation = VectorLoadAligned(&NewRotation);
 		DiagnosticCheckNaN_Rotate();
 	}
 
-	/** Copy rotation from another FTransform. */
-	FORCEINLINE void CopyRotation(const FTransform& Other)
+	/** Copy rotation from another TTransform<T>. */
+	FORCEINLINE void CopyRotation(const TTransform<T>& Other)
 	{
 		Rotation = Other.Rotation;
 	}
@@ -881,14 +906,14 @@ public:
 	 * Sets the Scale3D component
 	 * @param NewScale3D The new value for the Scale3D component
 	 */
-	FORCEINLINE void SetScale3D(const FVector& NewScale3D)
+	FORCEINLINE void SetScale3D(const TVector<T>& NewScale3D)
 	{
 		Scale3D = VectorLoadFloat3_W0(&NewScale3D);
 		DiagnosticCheckNaN_Scale3D();
 	}
 
-	/** Copy scale from another FTransform. */
-	FORCEINLINE void CopyScale3D(const FTransform& Other)
+	/** Copy scale from another TTransform<T>. */
+	FORCEINLINE void CopyScale3D(const TTransform<T>& Other)
 	{
 		Scale3D = Other.Scale3D;
 	}
@@ -898,7 +923,7 @@ public:
 	 * @param NewTranslation The new value for the translation component
 	 * @param NewScale3D The new value for the Scale3D component
 	 */
-	FORCEINLINE void SetTranslationAndScale3D(const FVector& NewTranslation, const FVector& NewScale3D)
+	FORCEINLINE void SetTranslationAndScale3D(const TVector<T>& NewTranslation, const TVector<T>& NewScale3D)
 	{
 		Translation = VectorLoadFloat3_W0(&NewTranslation);
 		Scale3D = VectorLoadFloat3_W0(&NewScale3D);
@@ -924,13 +949,13 @@ public:
 	*
 	* @param SourceAtom The other transform to accumulate into this one
 	*/
-	FORCEINLINE void Accumulate(const FTransform& SourceAtom)
+	FORCEINLINE void Accumulate(const TTransform<T>& SourceAtom)
 	{
-		const VectorRegister4Float BlendedRotation = SourceAtom.Rotation;
-		const VectorRegister4Float RotationW = VectorReplicate(BlendedRotation, 3);
+		const TransformVectorRegister BlendedRotation = SourceAtom.Rotation;
+		const TransformVectorRegister RotationW = VectorReplicate(BlendedRotation, 3);
 
 		// if( Square(SourceAtom.Rotation.W) < 1.f - DELTA * DELTA )
-		if (VectorAnyGreaterThan(GlobalVectorConstants::RotationSignificantThresholdFloat, VectorMultiply(RotationW, RotationW)))
+		if (VectorAnyGreaterThan(GlobalVectorConstants::RotationSignificantThreshold, VectorMultiply(RotationW, RotationW)))
 		{
 			// Rotation = SourceAtom.Rotation * Rotation;
 			Rotation = VectorQuaternionMultiply2(BlendedRotation, Rotation);
@@ -959,18 +984,18 @@ public:
 	* @param Atom The other transform to accumulate into this one
 	* @param BlendWeight The weight to multiply Atom by before it is accumulated.
 	*/
-	FORCEINLINE void Accumulate(const FTransform& Atom, const ScalarRegister& BlendWeight)
+	FORCEINLINE void Accumulate(const TTransform<T>& Atom, const ScalarRegister& BlendWeightScalar)
 	{
 		// SourceAtom = Atom * BlendWeight;
-		VectorRegister4Float BlendWeightFloat = MakeVectorRegisterFloatFromDouble(BlendWeight.Value);
-		const VectorRegister4Float BlendedRotation = VectorMultiply(Atom.Rotation, BlendWeightFloat);
-		const VectorRegister BlendedScale = VectorMultiply(Atom.Scale3D, BlendWeight.Value);
+		const TransformVectorRegister BlendWeight(BlendWeightScalar.Value);
+		const TransformVectorRegister BlendedRotation = VectorMultiply(Atom.Rotation, BlendWeight);
+		const TransformVectorRegister BlendedScale = VectorMultiply(Atom.Scale3D, BlendWeight);
 
-		const VectorRegister4Float RotationW = VectorReplicate(BlendedRotation, 3);
+		const TransformVectorRegister RotationW = VectorReplicate(BlendedRotation, 3);
 
 		// Add ref pose relative animation to base animation, only if rotation is significant.
 		// if( Square(SourceAtom.Rotation.W) < 1.f - DELTA * DELTA )
-		if (VectorAnyGreaterThan(GlobalVectorConstants::RotationSignificantThresholdFloat, VectorMultiply(RotationW, RotationW)))
+		if (VectorAnyGreaterThan(GlobalVectorConstants::RotationSignificantThreshold, VectorMultiply(RotationW, RotationW)))
 		{
 			// Rotation = SourceAtom.Rotation * Rotation;
 			Rotation = VectorQuaternionMultiply2(BlendedRotation, Rotation);
@@ -978,7 +1003,7 @@ public:
 
 		// Translation += SourceAtom.Translation;
 		// Scale *= SourceAtom.Scale;
-		Translation = VectorMultiplyAdd(Atom.Translation, BlendWeight.Value, Translation);
+		Translation = VectorMultiplyAdd(Atom.Translation, BlendWeight, Translation);
 		Scale3D = VectorMultiply(Scale3D, BlendedScale);
 
 		DiagnosticCheckNaN_All();
@@ -993,9 +1018,10 @@ public:
 	 * @param DeltaAtom The other transform to accumulate into this one
 	 * @param Weight The weight to multiply DeltaAtom by before it is accumulated.
 	 */
-	FORCEINLINE void AccumulateWithShortestRotation(const FTransform& DeltaAtom, const ScalarRegister& BlendWeight)
+	FORCEINLINE void AccumulateWithShortestRotation(const TTransform<T>& DeltaAtom, const ScalarRegister& BlendWeightScalar)
 	{
-		const VectorRegister4Float BlendedRotation = VectorMultiply(DeltaAtom.Rotation, MakeVectorRegisterFloatFromDouble(BlendWeight.Value));
+		const TransformVectorRegister BlendWeight(BlendWeightScalar.Value);
+		const TransformVectorRegister BlendedRotation = VectorMultiply(DeltaAtom.Rotation, BlendWeight);
 
 		Rotation = VectorAccumulateQuaternionShortestPath(Rotation, BlendedRotation);
 
@@ -1023,18 +1049,19 @@ public:
 	* @param Atom The other transform to accumulate into this one
 	* @param BlendWeight The weight to multiply Atom by before it is accumulated.
 	*/
-	FORCEINLINE void AccumulateWithAdditiveScale(const FTransform& Atom, const ScalarRegister& BlendWeight)
+	FORCEINLINE void AccumulateWithAdditiveScale(const TTransform<T>& Atom, const ScalarRegister& BlendWeightScalar)
 	{
-		const VectorRegister DefaultScale = GlobalVectorConstants::Vector1110;
+		const TransformVectorRegister BlendWeight(BlendWeightScalar.Value);
+		const TransformVectorRegister DefaultScale = GlobalVectorConstants::Float1110;
 
 		// SourceAtom = Atom * BlendWeight;
-		const VectorRegister4Float BlendedRotation = VectorMultiply(Atom.Rotation, MakeVectorRegisterFloatFromDouble(BlendWeight.Value));
+		const TransformVectorRegister BlendedRotation = VectorMultiply(Atom.Rotation, BlendWeight);
 
-		const VectorRegister4Float RotationW = VectorReplicate(BlendedRotation, 3);
+		const TransformVectorRegister RotationW = VectorReplicate(BlendedRotation, 3);
 
 		// Add ref pose relative animation to base animation, only if rotation is significant.
 		// if( Square(SourceAtom.Rotation.W) < 1.f - DELTA * DELTA )
-		if (VectorAnyGreaterThan(GlobalVectorConstants::RotationSignificantThresholdFloat, VectorMultiply(RotationW, RotationW)))
+		if (VectorAnyGreaterThan(GlobalVectorConstants::RotationSignificantThreshold, VectorMultiply(RotationW, RotationW)))
 		{
 			// Rotation = SourceAtom.Rotation * Rotation;
 			Rotation = VectorQuaternionMultiply2(BlendedRotation, Rotation);
@@ -1042,8 +1069,9 @@ public:
 
 		// Translation += SourceAtom.Translation;
 		// Scale *= SourceAtom.Scale;
-		Translation = VectorMultiplyAdd(Atom.Translation, BlendWeight.Value, Translation);
-		Scale3D = VectorMultiply(Scale3D, VectorMultiplyAdd(Atom.Scale3D, BlendWeight.Value, DefaultScale));
+		Translation = VectorMultiplyAdd(Atom.Translation, BlendWeight, Translation);
+		Scale3D = VectorMultiply(Scale3D, VectorMultiplyAdd(Atom.Scale3D, BlendWeight, DefaultScale));
+
 
 		DiagnosticCheckNaN_All();
 	}
@@ -1058,10 +1086,11 @@ public:
 	 * @param SourceAtom2 The ending point source atom (used 100% if Alpha is 1)
 	 * @param Alpha The blending weight between SourceAtom1 and SourceAtom2
 	 */
-	FORCEINLINE void LerpTranslationScale3D(const FTransform& SourceAtom1, const FTransform& SourceAtom2, const ScalarRegister& Alpha)
+	FORCEINLINE void LerpTranslationScale3D(const TTransform<T>& SourceAtom1, const TTransform<T>& SourceAtom2, const ScalarRegister& AlphaScalar)
 	{
-		Translation	= FMath::Lerp(SourceAtom1.Translation, SourceAtom2.Translation, Alpha.Value);
-		Scale3D = FMath::Lerp(SourceAtom1.Scale3D, SourceAtom2.Scale3D, Alpha.Value);
+		const TransformVectorRegister Alpha(AlphaScalar.Value);
+		Translation	= FMath::Lerp<TransformVectorRegister>(SourceAtom1.Translation, SourceAtom2.Translation, Alpha);
+		Scale3D = FMath::Lerp<TransformVectorRegister>(SourceAtom1.Scale3D, SourceAtom2.Scale3D, Alpha);
 
 		DiagnosticCheckNaN_Translate();
 		DiagnosticCheckNaN_Scale3D();
@@ -1083,7 +1112,7 @@ public:
 	 */
 	FORCEINLINE bool IsRotationNormalized() const
 	{		
-		const VectorRegister4Float TestValue = VectorAbs(VectorSubtract(VectorOneFloat(), VectorDot4(Rotation, Rotation)));
+		const TransformVectorRegister TestValue = VectorAbs(VectorSubtract(VectorOne(), VectorDot4(Rotation, Rotation)));
 		return !VectorAnyGreaterThan(TestValue, GlobalVectorConstants::ThreshQuatNormalized);
 	}
 
@@ -1099,13 +1128,13 @@ public:
 	 * @param SourceAtom The target transformation (used when BlendWeight = 1)
 	 * @param Alpha The blend weight between Identity and SourceAtom
 	 */
-	FORCEINLINE static void BlendFromIdentityAndAccumulate(FTransform& FinalAtom, const FTransform& SourceAtom, const ScalarRegister& BlendWeight)
+	FORCEINLINE static void BlendFromIdentityAndAccumulate(TTransform<T>& FinalAtom, const TTransform<T>& SourceAtom, const ScalarRegister& BlendWeightScalar)
 	{
-		const VectorRegister4Float BlendWeightFloat = MakeVectorRegisterFloatFromDouble(BlendWeight.Value);
-		const VectorRegister4Float Const0001 = GlobalVectorConstants::Float0001;
-		const VectorRegister4Float ConstNegative0001 = VectorSubtract(VectorZeroFloat(), Const0001);
-		const VectorRegister4Float VOneMinusAlpha = VectorSubtract(VectorOneFloat(), BlendWeightFloat);
-		const VectorRegister DefaultScale = GlobalVectorConstants::Vector1110;
+		const TransformVectorRegister BlendWeight(BlendWeightScalar.Value);
+		const TransformVectorRegister Const0001 = GlobalVectorConstants::Float0001;
+		const TransformVectorRegister ConstNegative0001 = VectorSubtract(VectorZero(), Const0001);
+		const TransformVectorRegister VOneMinusAlpha = VectorSubtract(VectorOne(), BlendWeight);
+		const TransformVectorRegister DefaultScale = GlobalVectorConstants::Float1110;
 
 		// Blend rotation
 		//     To ensure the 'shortest route', we make sure the dot product between the both rotations is positive.
@@ -1113,15 +1142,15 @@ public:
 		//     BlendedAtom.Rotation = (B * Alpha) + (A * (Bias * (1.f - Alpha)));
 		//     BlendedAtom.Rotation.QuaternionNormalize();
 		//  Note: A = (0,0,0,1), which simplifies things a lot; only care about sign of B.W now, instead of doing a dot product
-		const VectorRegister4Float RotationB = SourceAtom.Rotation;
+		const TransformVectorRegister RotationB = SourceAtom.Rotation;
 
-		const VectorRegister4Float QuatRotationDirMask = VectorCompareGE(RotationB, VectorZeroFloat());
-		const VectorRegister4Float BiasTimesA = VectorSelect(QuatRotationDirMask, Const0001, ConstNegative0001);
-		const VectorRegister4Float RotateBTimesWeight = VectorMultiply(RotationB, BlendWeightFloat);
-		const VectorRegister4Float UnnormalizedRotation = VectorMultiplyAdd(BiasTimesA, VOneMinusAlpha, RotateBTimesWeight);
+		const TransformVectorRegister QuatRotationDirMask = VectorCompareGE(RotationB, VectorZero());
+		const TransformVectorRegister BiasTimesA = VectorSelect(QuatRotationDirMask, Const0001, ConstNegative0001);
+		const TransformVectorRegister RotateBTimesWeight = VectorMultiply(RotationB, BlendWeight);
+		const TransformVectorRegister UnnormalizedRotation = VectorMultiplyAdd(BiasTimesA, VOneMinusAlpha, RotateBTimesWeight);
 
 		// Normalize blended rotation ( result = (Q.Q >= 1e-8) ? (Q / |Q|) : (0,0,0,1) )
-		const VectorRegister4Float BlendedRotation = VectorNormalizeSafe(UnnormalizedRotation, GlobalVectorConstants::Float0001);
+		const TransformVectorRegister BlendedRotation = VectorNormalizeSafe(UnnormalizedRotation, Const0001);
 
 		// FinalAtom.Rotation = BlendedAtom.Rotation * FinalAtom.Rotation;
 		FinalAtom.Rotation = VectorQuaternionMultiply2(BlendedRotation, FinalAtom.Rotation);
@@ -1129,8 +1158,9 @@ public:
 		// Blend translation and scale
 		//    BlendedAtom.Translation = Lerp(Zero, SourceAtom.Translation, Alpha);
 		//    BlendedAtom.Scale = Lerp(0, SourceAtom.Scale, Alpha);
-		const VectorRegister BlendedTranslation	= FMath::Lerp<VectorRegister>(VectorZeroVectorRegister(), SourceAtom.Translation, BlendWeight.Value);
-		const VectorRegister BlendedScale3D	= FMath::Lerp<VectorRegister>(VectorZeroVectorRegister(), SourceAtom.Scale3D, BlendWeight.Value);
+		const TransformVectorRegister Zero = VectorZero();
+		const TransformVectorRegister BlendedTranslation = FMath::Lerp<TransformVectorRegister>(Zero, SourceAtom.Translation, BlendWeight);
+		const TransformVectorRegister BlendedScale3D = FMath::Lerp<TransformVectorRegister>(Zero, SourceAtom.Scale3D, BlendWeight);
 
 		// Apply translation and scale to final atom
 		//     FinalAtom.Translation += BlendedAtom.Translation
@@ -1146,10 +1176,10 @@ public:
 	 *
 	 * @return The rotation component
 	 */
-	FORCEINLINE FQuat GetRotation() const
+	FORCEINLINE TQuat<T> GetRotation() const
 	{
 		DiagnosticCheckNaN_Rotate();
-		FQuat OutRotation;
+		TQuat<T> OutRotation;
 		VectorStoreAligned(Rotation, &OutRotation);
 		return OutRotation;
 	}
@@ -1159,10 +1189,10 @@ public:
 	 *
 	 * @return The translation component
 	 */
-	FORCEINLINE FVector GetTranslation() const
+	FORCEINLINE TVector<T> GetTranslation() const
 	{
 		DiagnosticCheckNaN_Translate();
-		FVector OutTranslation;
+		TVector<T> OutTranslation;
 		VectorStoreFloat3(Translation, &OutTranslation);
 		return OutTranslation;
 	}
@@ -1172,10 +1202,10 @@ public:
 	 *
 	 * @return The Scale3D component
 	 */
-	FORCEINLINE FVector GetScale3D() const
+	FORCEINLINE TVector<T> GetScale3D() const
 	{
 		DiagnosticCheckNaN_Scale3D();
-		FVector OutScale3D;
+		TVector<T> OutScale3D;
 		VectorStoreFloat3(Scale3D, &OutScale3D);
 		return OutScale3D;
 	}
@@ -1185,7 +1215,7 @@ public:
 	 *
 	 * @param SrcBA The transform to copy rotation and Scale3D from
 	 */
-	FORCEINLINE void CopyRotationPart(const FTransform& SrcBA)
+	FORCEINLINE void CopyRotationPart(const TTransform<T>& SrcBA)
 	{
 		Rotation = SrcBA.Rotation;
 		Scale3D = SrcBA.Scale3D;
@@ -1199,7 +1229,7 @@ public:
 	 *
 	 * @param SrcBA The transform to copy translation and Scale3D from
 	 */
-	FORCEINLINE void CopyTranslationAndScale3D(const FTransform& SrcBA)
+	FORCEINLINE void CopyTranslationAndScale3D(const TTransform<T>& SrcBA)
 	{
 		Translation = SrcBA.Translation;
 		Scale3D = SrcBA.Scale3D;
@@ -1208,12 +1238,12 @@ public:
 		DiagnosticCheckNaN_Scale3D();
 	}
 
-	void SetFromMatrix(const FMatrix& InMatrix)
+	void SetFromMatrix(const TMatrix<T>& InMatrix)
 	{
-		FMatrix M = InMatrix;
+		TMatrix<T> M = InMatrix;
 
 		// Get the 3D scale from the matrix
-		FVector InScale = M.ExtractScaling();
+		TVector<T> InScale = M.ExtractScaling();
 		Scale3D = VectorLoadFloat3_W0(&InScale);
 
 		// If there is negative scaling going on, we handle that here
@@ -1221,17 +1251,13 @@ public:
 		{
 			// Assume it is along X and modify transform accordingly. 
 			// It doesn't actually matter which axis we choose, the 'appearance' will be the same
-#if !UE_LARGE_WORLD_COORDINATES_DISABLED
-			Scale3D = VectorMultiply(Scale3D, GlobalVectorConstants::DoubleMinus1_111 );			
-#else
-			Scale3D = VectorMultiply(Scale3D, GlobalVectorConstants::FloatMinus1_111);
-#endif
+			Scale3D = VectorMultiply(Scale3D, (TransformVectorRegister)GlobalVectorConstants::FloatMinus1_111 );
 			M.SetAxis(0, -M.GetScaledAxis( EAxis::X ));
 		}
 
-		FQuat InRotation = FQuat(M);
+		TQuat<T> InRotation = TQuat<T>(M);
 		Rotation = VectorLoadAligned(&InRotation);
-		FVector InTranslation = InMatrix.GetOrigin();
+		TVector<T> InTranslation = InMatrix.GetOrigin();
 		Translation = VectorLoadFloat3_W0(&InTranslation);
 
 		// Normalize rotation
@@ -1239,24 +1265,24 @@ public:
 	}
 
 private:
-	FORCEINLINE void ToMatrixInternal( VectorRegister& OutDiagonals, VectorRegister& OutAdds, VectorRegister& OutSubtracts ) const
+	FORCEINLINE void ToMatrixInternal( TransformVectorRegister& OutDiagonals, TransformVectorRegister& OutAdds, TransformVectorRegister& OutSubtracts ) const
 	{
 #if !(UE_BUILD_SHIPPING || UE_BUILD_TEST) && WITH_EDITORONLY_DATA
 		// Make sure Rotation is normalized when we turn it into a matrix.
 		check( IsRotationNormalized() );
 #endif		
 
-		const VectorRegister4Float RotationX2Y2Z2 = VectorAdd(Rotation, Rotation);	// x2, y2, z2
-		const VectorRegister4Float RotationXX2YY2ZZ2 = VectorMultiply(RotationX2Y2Z2, Rotation);	// xx2, yy2, zz2		
+		const TransformVectorRegister RotationX2Y2Z2 = VectorAdd(Rotation, Rotation);	// x2, y2, z2
+		const TransformVectorRegister RotationXX2YY2ZZ2 = VectorMultiply(RotationX2Y2Z2, Rotation);	// xx2, yy2, zz2		
 
 		// The diagonal terms of the rotation matrix are:
 		//   (1 - (yy2 + zz2)) * scale
 		//   (1 - (xx2 + zz2)) * scale
 		//   (1 - (xx2 + yy2)) * scale
-		const VectorRegister4Float yy2_xx2_xx2 = VectorSwizzle(RotationXX2YY2ZZ2, 1, 0, 0, 0);
-		const VectorRegister4Float zz2_zz2_yy2 = VectorSwizzle(RotationXX2YY2ZZ2, 2, 2, 1, 0);
-		const VectorRegister4Float DiagonalSum = VectorAdd(yy2_xx2_xx2, zz2_zz2_yy2);
-		const VectorRegister4Float Diagonals = VectorSubtract(VectorOneFloat(), DiagonalSum);
+		const TransformVectorRegister yy2_xx2_xx2 = VectorSwizzle(RotationXX2YY2ZZ2, 1, 0, 0, 0);
+		const TransformVectorRegister zz2_zz2_yy2 = VectorSwizzle(RotationXX2YY2ZZ2, 2, 2, 1, 0);
+		const TransformVectorRegister DiagonalSum = VectorAdd(yy2_xx2_xx2, zz2_zz2_yy2);
+		const TransformVectorRegister Diagonals = VectorSubtract(VectorOne(), DiagonalSum);
 		OutDiagonals = VectorMultiply(Diagonals, Scale3D);
 
 		// Grouping the non-diagonal elements in the rotation block by operations:
@@ -1267,40 +1293,40 @@ private:
 
 		// RotBase = x*y2, y*z2, x*z2
 		// RotOffset = w*z2, w*x2, w*y2
-		const VectorRegister4Float x_y_x = VectorSwizzle(Rotation, 0, 1, 0, 0);
-		const VectorRegister4Float y2_z2_z2 = VectorSwizzle(RotationX2Y2Z2, 1, 2, 2, 0);
-		const VectorRegister4Float RotBase = VectorMultiply(x_y_x, y2_z2_z2);
+		const TransformVectorRegister x_y_x = VectorSwizzle(Rotation, 0, 1, 0, 0);
+		const TransformVectorRegister y2_z2_z2 = VectorSwizzle(RotationX2Y2Z2, 1, 2, 2, 0);
+		const TransformVectorRegister RotBase = VectorMultiply(x_y_x, y2_z2_z2);
 
-		const VectorRegister4Float w_w_w = VectorReplicate(Rotation, 3);
-		const VectorRegister4Float z2_x2_y2 = VectorSwizzle(RotationX2Y2Z2, 2, 0, 1, 0);
-		const VectorRegister4Float RotOffset = VectorMultiply(w_w_w, z2_x2_y2);
+		const TransformVectorRegister w_w_w = VectorReplicate(Rotation, 3);
+		const TransformVectorRegister z2_x2_y2 = VectorSwizzle(RotationX2Y2Z2, 2, 0, 1, 0);
+		const TransformVectorRegister RotOffset = VectorMultiply(w_w_w, z2_x2_y2);
 
 		// Adds = (RotBase + RotOffset)*Scale3D :  (x*y2 + w*z2) * Scale3D.X , (y*z2 + w*x2) * Scale3D.Y, (x*z2 + w*y2) * Scale3D.Z
 		// Subtracts = (RotBase - RotOffset)*Scale3DYZX :  (x*y2 - w*z2) * Scale3D.Y , (y*z2 - w*x2) * Scale3D.Z, (x*z2 - w*y2) * Scale3D.X
-		const VectorRegister4Float Adds = VectorAdd(RotBase, RotOffset);
+		const TransformVectorRegister Adds = VectorAdd(RotBase, RotOffset);
 		OutAdds = VectorMultiply(Adds, Scale3D);
-		const VectorRegister Scale3DYZXW = VectorSwizzle( Scale3D, 1, 2, 0, 3);
-		const VectorRegister4Float Subtracts = VectorSubtract(RotBase, RotOffset);
-		OutSubtracts = VectorMultiply(Subtracts, Scale3DYZXW);
+		const TransformVectorRegister Scale3DYZXW = VectorSwizzle( Scale3D, 1, 2, 0, 3);
+		const TransformVectorRegister Subtracts = VectorSubtract(RotBase, RotOffset);
+		OutSubtracts = VectorMultiply(Subtracts , Scale3DYZXW);
 	}
 
-	FORCEINLINE void ToMatrixInternalNoScale( VectorRegister& OutDiagonals, VectorRegister& OutAdds, VectorRegister& OutSubtracts ) const
+	FORCEINLINE void ToMatrixInternalNoScale( TransformVectorRegister& OutDiagonals, TransformVectorRegister& OutAdds, TransformVectorRegister& OutSubtracts ) const
 	{
 #if !(UE_BUILD_SHIPPING || UE_BUILD_TEST) && WITH_EDITORONLY_DATA
 		// Make sure Rotation is normalized when we turn it into a matrix.
 		ensure( IsRotationNormalized() );
 #endif		
-		const VectorRegister4Float RotationX2Y2Z2 = VectorAdd(Rotation, Rotation);	// x2, y2, z2
-		const VectorRegister4Float RotationXX2YY2ZZ2 = VectorMultiply(RotationX2Y2Z2, Rotation);	// xx2, yy2, zz2		
+		const TransformVectorRegister RotationX2Y2Z2 = VectorAdd(Rotation, Rotation);	// x2, y2, z2
+		const TransformVectorRegister RotationXX2YY2ZZ2 = VectorMultiply(RotationX2Y2Z2, Rotation);	// xx2, yy2, zz2		
 
 		// The diagonal terms of the rotation matrix are:
 		//   (1 - (yy2 + zz2))
 		//   (1 - (xx2 + zz2))
 		//   (1 - (xx2 + yy2))
-		const VectorRegister4Float yy2_xx2_xx2 = VectorSwizzle(RotationXX2YY2ZZ2, 1, 0, 0, 0);
-		const VectorRegister4Float zz2_zz2_yy2 = VectorSwizzle(RotationXX2YY2ZZ2, 2, 2, 1, 0);
-		const VectorRegister4Float DiagonalSum = VectorAdd(yy2_xx2_xx2, zz2_zz2_yy2);
-		OutDiagonals = VectorSubtract(VectorOneFloat(), DiagonalSum);
+		const TransformVectorRegister yy2_xx2_xx2 = VectorSwizzle(RotationXX2YY2ZZ2, 1, 0, 0, 0);
+		const TransformVectorRegister zz2_zz2_yy2 = VectorSwizzle(RotationXX2YY2ZZ2, 2, 2, 1, 0);
+		const TransformVectorRegister DiagonalSum = VectorAdd(yy2_xx2_xx2, zz2_zz2_yy2);
+		OutDiagonals = VectorSubtract(VectorOne(), DiagonalSum);
 
 		// Grouping the non-diagonal elements in the rotation block by operations:
 		//    ((x*y2,y*z2,x*z2) + (w*z2,w*x2,w*y2)) and
@@ -1310,13 +1336,13 @@ private:
 
 		// RotBase = x*y2, y*z2, x*z2
 		// RotOffset = w*z2, w*x2, w*y2
-		const VectorRegister4Float x_y_x = VectorSwizzle(Rotation, 0, 1, 0, 0);
-		const VectorRegister4Float y2_z2_z2 = VectorSwizzle(RotationX2Y2Z2, 1, 2, 2, 0);
-		const VectorRegister4Float RotBase = VectorMultiply(x_y_x, y2_z2_z2);
+		const TransformVectorRegister x_y_x = VectorSwizzle(Rotation, 0, 1, 0, 0);
+		const TransformVectorRegister y2_z2_z2 = VectorSwizzle(RotationX2Y2Z2, 1, 2, 2, 0);
+		const TransformVectorRegister RotBase = VectorMultiply(x_y_x, y2_z2_z2);
 
-		const VectorRegister4Float w_w_w = VectorReplicate(Rotation, 3);
-		const VectorRegister4Float z2_x2_y2 = VectorSwizzle(RotationX2Y2Z2, 2, 0, 1, 0);
-		const VectorRegister4Float RotOffset = VectorMultiply(w_w_w, z2_x2_y2);
+		const TransformVectorRegister w_w_w = VectorReplicate(Rotation, 3);
+		const TransformVectorRegister z2_x2_y2 = VectorSwizzle(RotationX2Y2Z2, 2, 0, 1, 0);
+		const TransformVectorRegister RotOffset = VectorMultiply(w_w_w, z2_x2_y2);
 
 		// Adds = (RotBase + RotOffset):  (x*y2 + w*z2) , (y*z2 + w*x2), (x*z2 + w*y2)
 		// Subtracts = (RotBase - RotOffset) :  (x*y2 - w*z2) , (y*z2 - w*x2), (x*z2 - w*y2)
@@ -1331,27 +1357,26 @@ private:
 	 * also returning BIG_NUMBER causes sequential NaN issues by multiplying 
 	 * so we hardcode as 0
 	 */
-	static FORCEINLINE VectorRegister		GetSafeScaleReciprocal(const VectorRegister& InScale, const ScalarRegister& Tolerance = ScalarRegister(GlobalVectorConstants::SmallNumber))
+	static FORCEINLINE TransformVectorRegister		GetSafeScaleReciprocal(const TransformVectorRegister& InScale, const ScalarRegister& Tolerance = ScalarRegister(GlobalVectorConstants::SmallNumber))
 	{		
 		// SafeReciprocalScale.X = (InScale.X == 0) ? 0.f : 1/InScale.X; // same for YZW
-		VectorRegister SafeReciprocalScale;
 
-		/// VectorRegister( 1.0f / InScale.x, 1.0f / InScale.y, 1.0f / InScale.z, 1.0f / InScale.w )
-		const VectorRegister ReciprocalScale = VectorReciprocalAccurate(InScale);
+		/// TransformVectorRegister( 1.0f / InScale.x, 1.0f / InScale.y, 1.0f / InScale.z, 1.0f / InScale.w )
+		const TransformVectorRegister ReciprocalScale = VectorReciprocalAccurate(InScale);
 		
-		//VectorRegister( Vec1.x == Vec2.x ? 0xFFFFFFFF : 0, same for yzw )
-		const VectorRegister ScaleZeroMask = VectorCompareGE(Tolerance.Value, VectorAbs(InScale));
+		//TransformVectorRegister( Vec1.x == Vec2.x ? 0xFFFFFFFF : 0, same for yzw )
+		const TransformVectorRegister ScaleZeroMask = VectorCompareGE(TransformVectorRegister(Tolerance.Value), VectorAbs(InScale));
 
-		//const VectorRegister ScaleZeroMask = VectorCompareEQ(InScale, VectorZero());
+		//const TransformVectorRegister ScaleZeroMask = VectorCompareEQ(InScale, VectorZero());
 
-		// VectorRegister( for each bit i: Mask[i] ? Vec1[i] : Vec2[i] )
-		SafeReciprocalScale = VectorSelect(ScaleZeroMask, VectorZeroVectorRegister(), ReciprocalScale);
+		// TransformVectorRegister( for each bit i: Mask[i] ? Vec1[i] : Vec2[i] )
+		const TransformVectorRegister SafeReciprocalScale = VectorSelect(ScaleZeroMask, VectorZero(), ReciprocalScale);
 
 		return SafeReciprocalScale;
 	}
 
-	/** Returns Inverse Transform of this FTransform **/
-	FORCEINLINE FTransform InverseFast() const
+	/** Returns Inverse Transform of this TTransform<T> **/
+	FORCEINLINE TTransform<T> InverseFast() const
 	{
 		// Inverse QST (A) = QST (~A)
 		// Since A*~A = Identity, 
@@ -1371,17 +1396,17 @@ private:
 		checkSlow(VectorAnyGreaterThan(VectorAbs(Scale3D), GlobalVectorConstants::SmallNumber));
 
 		// Invert the scale
-		const VectorRegister InvScale = VectorSet_W0(GetSafeScaleReciprocal(VectorSet_W1(Scale3D), ScalarRegister(GlobalVectorConstants::SmallNumber)));
+		const TransformVectorRegister InvScale = VectorSet_W0(GetSafeScaleReciprocal(VectorSet_W1(Scale3D), ScalarRegister(GlobalVectorConstants::SmallNumber)));
 
 		// Invert the rotation
-		const VectorRegister4Float InvRotation = VectorQuaternionInverse(Rotation);
+		const TransformVectorRegister InvRotation = VectorQuaternionInverse(Rotation);
 
 		// Invert the translation
-		const VectorRegister ScaledTranslation = VectorMultiply(InvScale, Translation);
-		const VectorRegister t2 = VectorQuaternionRotateVector(InvRotation, ScaledTranslation);
-		const VectorRegister InvTranslation = VectorSet_W0(VectorNegate(t2));
+		const TransformVectorRegister ScaledTranslation = VectorMultiply(InvScale, Translation);
+		const TransformVectorRegister t2 = VectorQuaternionRotateVector(InvRotation, ScaledTranslation);
+		const TransformVectorRegister InvTranslation = VectorSet_W0(VectorNegate(t2));
 
-		return FTransform(InvRotation, InvTranslation, InvScale);
+		return TTransform<T>(InvRotation, InvTranslation, InvScale);
 	}
 
 	/**
@@ -1394,7 +1419,7 @@ private:
 	* @param  A Transform A.
 	* @param  B Transform B.
 	*/
-	FORCEINLINE static void MultiplyUsingMatrixWithScale(FTransform* OutTransform, const FTransform* A, const FTransform* B);
+	FORCEINLINE static void MultiplyUsingMatrixWithScale(TTransform<T>* OutTransform, const TTransform<T>* A, const TTransform<T>* B);
 	/**
 	* Create a new transform from multiplications of given to matrices (AMatrix*BMatrix) using desired scale
 	* This is used by MultiplyUsingMatrixWithScale and GetRelativeTransformUsingMatrixWithScale
@@ -1405,7 +1430,7 @@ private:
 	* @param	DesiredScale - there is no check on if the magnitude is correct here. It assumes that is correct. 
 	* @param	OutTransform the constructed transform 
 	*/
-	FORCEINLINE static void ConstructTransformFromMatrixWithDesiredScale(const FMatrix& AMatrix, const FMatrix& BMatrix, const VectorRegister& DesiredScale, FTransform& OutTransform);
+	FORCEINLINE static void ConstructTransformFromMatrixWithDesiredScale(const TMatrix<T>& AMatrix, const TMatrix<T>& BMatrix, const TransformVectorRegister& DesiredScale, TTransform<T>& OutTransform);
 	/**
 	* Create a new transform: OutTransform = Base * Relative(-1) using the matrix while keeping the scale that's given by Base and Relative
 	* Please note that this operation is a lot more expensive than normal GetRelativeTrnasform
@@ -1414,57 +1439,76 @@ private:
 	* @param  BAse Transform Base.
 	* @param  Relative Transform Relative.
 	*/
-	static void GetRelativeTransformUsingMatrixWithScale(FTransform* OutTransform, const FTransform* Base, const FTransform* Relative);
+	static void GetRelativeTransformUsingMatrixWithScale(TTransform<T>* OutTransform, const TTransform<T>* Base, const TTransform<T>* Relative);
+
+public:
+
+	// Conversion to other type.
+	friend struct TTransform<double>;
+	template<typename FArg, TEMPLATE_REQUIRES(!TIsSame<T, FArg>::Value && TIsSame<T, float>::Value)>
+	explicit TTransform(const TTransform<FArg>& From) : TTransform(MakeVectorRegisterFloatFromDouble(From.Rotation), MakeVectorRegisterFloatFromDouble(From.Translation), MakeVectorRegisterFloatFromDouble(From.Scale3D)) {}
+
+	friend struct TTransform<float>;
+	template<typename FArg, TEMPLATE_REQUIRES(!TIsSame<T, FArg>::Value && TIsSame<T, double>::Value)>
+	explicit TTransform(const TTransform<FArg>& From) : TTransform(MakeVectorRegisterDouble(From.Rotation), MakeVectorRegisterDouble(From.Translation), MakeVectorRegisterDouble(From.Scale3D)) {}
 };
 
-template <> struct TIsPODType<FTransform> { enum { Value = true }; };
+template<> inline const TTransform<float> TTransform<float>::Identity(FQuat4f(0.f, 0.f, 0.f, 1.f), FVector3f(0.f), FVector3f(1.f));
+template<> inline const TTransform<double> TTransform<double>::Identity(FQuat4d(0.0, 0.0, 0.0, 1.0), FVector3d(0.0), FVector3d(1.0));
 
-FORCEINLINE bool FTransform::AnyHasNegativeScale(const FVector& InScale3D, const FVector& InOtherScale3D)
+
+template<typename T>
+FORCEINLINE bool TTransform<T>::AnyHasNegativeScale(const TVector<T>& InScale3D, const TVector<T>& InOtherScale3D)
 {
-	VectorRegister VectorInScale3D = VectorLoadFloat3_W0(&InScale3D);
-	VectorRegister VectorInOtherScale3D = VectorLoadFloat3_W0(&InOtherScale3D);
+	TransformVectorRegister VectorInScale3D = VectorLoadFloat3_W0(&InScale3D);
+	TransformVectorRegister VectorInOtherScale3D = VectorLoadFloat3_W0(&InOtherScale3D);
 
 	return Private_AnyHasNegativeScale(VectorInScale3D, VectorInOtherScale3D);
 }
 
 /** Scale the translation part of the Transform by the supplied vector. */
-FORCEINLINE void FTransform::ScaleTranslation(const FVector& InScale3D)
+template<typename T>
+FORCEINLINE void TTransform<T>::ScaleTranslation(const TVector<T>& InScale3D)
 {
-	VectorRegister VectorInScale3D = VectorLoadFloat3_W0(&InScale3D);
+	TransformVectorRegister VectorInScale3D = VectorLoadFloat3_W0(&InScale3D);
 	Translation = VectorMultiply( Translation, VectorInScale3D );
 	DiagnosticCheckNaN_Translate();
 }
 
 /** Scale the translation part of the Transform by the supplied value. */
-FORCEINLINE void FTransform::ScaleTranslation(const float& InScale)
+template<typename T>
+FORCEINLINE void TTransform<T>::ScaleTranslation(const T& InScale)
 {
-	ScaleTranslation( FVector(InScale) );
+	ScaleTranslation( TVector<T>(InScale) );
 }
 
 // this function is from matrix, and all it does is to normalize rotation portion
-FORCEINLINE void FTransform::RemoveScaling(float Tolerance/*=SMALL_NUMBER*/)
+template<typename T>
+FORCEINLINE void TTransform<T>::RemoveScaling(T Tolerance/*=SMALL_NUMBER*/)
 {
-	Scale3D = VectorSet_W0( VectorOneVectorRegister() );
+	Scale3D = GlobalVectorConstants::Float1110;
 	NormalizeRotation();	
 
 	DiagnosticCheckNaN_Rotate();
 	DiagnosticCheckNaN_Scale3D();
 }
 
-FORCEINLINE void FTransform::MultiplyUsingMatrixWithScale(FTransform* OutTransform, const FTransform* A, const FTransform* B)
+template<typename T>
+FORCEINLINE void TTransform<T>::MultiplyUsingMatrixWithScale(TTransform<T>* OutTransform, const TTransform<T>* A, const TTransform<T>* B)
 {
 	ConstructTransformFromMatrixWithDesiredScale(A->ToMatrixWithScale(), B->ToMatrixWithScale(), VectorMultiply(A->Scale3D, B->Scale3D), *OutTransform);
 }
 
-FORCEINLINE void FTransform::ConstructTransformFromMatrixWithDesiredScale(const FMatrix& AMatrix, const FMatrix& BMatrix, const VectorRegister& DesiredScale, FTransform& OutTransform)
+template<typename T>
+FORCEINLINE void TTransform<T>::ConstructTransformFromMatrixWithDesiredScale(const UE::Math::TMatrix<T>& AMatrix, const UE::Math::TMatrix<T>& BMatrix, const TransformVectorRegister& DesiredScale, TTransform<T>& OutTransform)
 {
 	// the goal of using M is to get the correct orientation
 	// but for translation, we still need scale
-	FMatrix M = AMatrix * BMatrix;
+	UE::Math::TMatrix<T> M = AMatrix * BMatrix;
 	M.RemoveScaling();
 
 	// apply negative scale back to axes
-	FVector SignedScale;
+	TVector<T> SignedScale;
 	VectorStoreFloat3(VectorSign(DesiredScale), &SignedScale);
 
 	M.SetAxis(0, SignedScale.X * M.GetScaledAxis(EAxis::X));
@@ -1473,22 +1517,23 @@ FORCEINLINE void FTransform::ConstructTransformFromMatrixWithDesiredScale(const
 
 	// @note: if you have negative with 0 scale, this will return rotation that is identity
 	// since matrix loses that axes
-	FQuat Rotation = FQuat(M);
+	TQuat<T> Rotation = TQuat<T>(M);
 	Rotation.Normalize();
 
 	// set values back to output
 	OutTransform.Scale3D = DesiredScale;
 	OutTransform.Rotation = VectorLoadAligned(&Rotation);
 
-	// technically I could calculate this using FTransform but then it does more quat multiplication 
+	// technically I could calculate this using TTransform<T> but then it does more quat multiplication 
 	// instead of using Scale in matrix multiplication
-	// it's a question of between RemoveScaling vs using FTransform to move translation
-	FVector Translation = M.GetOrigin();
+	// it's a question of between RemoveScaling vs using TTransform<T> to move translation
+	TVector<T> Translation = M.GetOrigin();
 	OutTransform.Translation = VectorLoadFloat3_W0(&Translation);
 }
 
 /** Returns Multiplied Transform of 2 FTransforms **/
-FORCEINLINE void FTransform::Multiply(FTransform* OutTransform, const FTransform* A, const FTransform* B)
+template<typename T>
+FORCEINLINE void TTransform<T>::Multiply(TTransform<T>* OutTransform, const TTransform<T>* A, const TTransform<T>* B)
 {
 	A->DiagnosticCheckNaN_All();
 	B->DiagnosticCheckNaN_All();
@@ -1520,19 +1565,19 @@ FORCEINLINE void FTransform::Multiply(FTransform* OutTransform, const FTransform
 	}
 	else
 	{
-		const VectorRegister4Float QuatA = A->Rotation;
-		const VectorRegister4Float QuatB = B->Rotation;
-		const VectorRegister TranslateA = A->Translation;
-		const VectorRegister TranslateB = B->Translation;
-		const VectorRegister ScaleA = A->Scale3D;
-		const VectorRegister ScaleB = B->Scale3D;
+		const TransformVectorRegister QuatA = A->Rotation;
+		const TransformVectorRegister QuatB = B->Rotation;
+		const TransformVectorRegister TranslateA = A->Translation;
+		const TransformVectorRegister TranslateB = B->Translation;
+		const TransformVectorRegister ScaleA = A->Scale3D;
+		const TransformVectorRegister ScaleB = B->Scale3D;
 
 		// RotationResult = B.Rotation * A.Rotation
 		OutTransform->Rotation = VectorQuaternionMultiply2(QuatB, QuatA);
 
 		// TranslateResult = B.Rotate(B.Scale * A.Translation) + B.Translate
-		const VectorRegister ScaledTransA = VectorMultiply(TranslateA, ScaleB);
-		const VectorRegister RotatedTranslate = VectorQuaternionRotateVector(QuatB, ScaledTransA);
+		const TransformVectorRegister ScaledTransA = VectorMultiply(TranslateA, ScaleB);
+		const TransformVectorRegister RotatedTranslate = VectorQuaternionRotateVector(QuatB, ScaledTransA);
 		OutTransform->Translation = VectorAdd(RotatedTranslate, TranslateB);
 
 		// ScaleResult = Scale.B * Scale.A
@@ -1542,11 +1587,12 @@ FORCEINLINE void FTransform::Multiply(FTransform* OutTransform, const FTransform
 /** 
  * Apply Scale to this transform
  */
-FORCEINLINE FTransform FTransform::GetScaled(float InScale) const
+template<typename T>
+FORCEINLINE TTransform<T> TTransform<T>::GetScaled(T InScale) const
 {
-	FTransform A(*this);
+	TTransform<T> A(*this);
 	
-	VectorRegister VScale = VectorLoadFloat1(&InScale);
+	TransformVectorRegister VScale = VectorLoadFloat1(&InScale);
 	A.Scale3D = VectorMultiply( A.Scale3D, VScale);
 
 	A.DiagnosticCheckNaN_Scale3D();
@@ -1557,11 +1603,12 @@ FORCEINLINE FTransform FTransform::GetScaled(float InScale) const
 /** 
  * Apply Scale to this transform
  */
-FORCEINLINE FTransform FTransform::GetScaled(FVector InScale) const
+template<typename T>
+FORCEINLINE TTransform<T> TTransform<T>::GetScaled(TVector<T> InScale) const
 {
-	FTransform A(*this);
+	TTransform<T> A(*this);
 
-	VectorRegister VScale = VectorLoadFloat3_W0(&InScale);
+	TransformVectorRegister VScale = VectorLoadFloat3_W0(&InScale);
 	A.Scale3D = VectorMultiply( A.Scale3D, VScale);
 
 	A.DiagnosticCheckNaN_Scale3D();
@@ -1569,53 +1616,55 @@ FORCEINLINE FTransform FTransform::GetScaled(FVector InScale) const
 	return A;
 }
 
-FORCEINLINE FVector4 FTransform::TransformFVector4NoScale(const FVector4& V) const
+template<typename T>
+FORCEINLINE FVector4 TTransform<T>::TransformFVector4NoScale(const FVector4& V) const
 {
 	DiagnosticCheckNaN_All();
 
 	// if not, this won't work
 	checkSlow (V.W == 0.f || V.W == 1.f);
 
-	const VectorRegister InputVector = VectorLoadAligned(&V);
+	const TransformVectorRegister InputVector = VectorLoadAligned(&V);
 
 	//Transform using QST is following
 	//QST(P) = Q.Rotate(S*P) + T where Q = quaternion, S = 1.0f, T = translation
 
 	//RotatedVec = Q.Rotate(V.X, V.Y, V.Z, 0.f)
-	const VectorRegister InputVectorW0 = VectorSet_W0(InputVector);	
-	const VectorRegister RotatedVec = VectorQuaternionRotateVector(Rotation, InputVectorW0);
+	const TransformVectorRegister InputVectorW0 = VectorSet_W0(InputVector);	
+	const TransformVectorRegister RotatedVec = VectorQuaternionRotateVector(Rotation, InputVectorW0);
 
 	// NewVect.XYZ += Translation * W
 	// NewVect.W += 1 * W
-	const VectorRegister WWWW = VectorReplicate(InputVector, 3);
-	const VectorRegister TranslatedVec = VectorMultiplyAdd(Translation, WWWW, RotatedVec);
+	const TransformVectorRegister WWWW = VectorReplicate(InputVector, 3);
+	const TransformVectorRegister TranslatedVec = VectorMultiplyAdd(Translation, WWWW, RotatedVec);
 
 	FVector4 NewVectOutput;
 	VectorStoreAligned(TranslatedVec, &NewVectOutput);
 	return NewVectOutput;
 }
 
-FORCEINLINE FVector4 FTransform::TransformFVector4(const FVector4& V) const
+template<typename T>
+FORCEINLINE FVector4 TTransform<T>::TransformFVector4(const FVector4& V) const
 {
 	DiagnosticCheckNaN_All();
 
 	// if not, this won't work
 	checkSlow (V.W == 0.f || V.W == 1.f);
 
-	const VectorRegister InputVector = VectorLoadAligned(&V);
+	const TransformVectorRegister InputVector = VectorLoadAligned(&V);
 
 	//Transform using QST is following
 	//QST(P) = Q.Rotate(S*P) + T where Q = quaternion, S = scale, T = translation
 
 	//RotatedVec = Q.Rotate(Scale*V.X, Scale*V.Y, Scale*V.Z, 0.f)
-	const VectorRegister InputVectorW0 = VectorSet_W0(InputVector);
-	const VectorRegister ScaledVec = VectorMultiply(Scale3D, InputVectorW0);
-	const VectorRegister RotatedVec = VectorQuaternionRotateVector(Rotation, ScaledVec);
+	const TransformVectorRegister InputVectorW0 = VectorSet_W0(InputVector);
+	const TransformVectorRegister ScaledVec = VectorMultiply(Scale3D, InputVectorW0);
+	const TransformVectorRegister RotatedVec = VectorQuaternionRotateVector(Rotation, ScaledVec);
 
 	// NewVect.XYZ += Translation * W
 	// NewVect.W += 1 * W
-	const VectorRegister WWWW = VectorReplicate(InputVector, 3);
-	const VectorRegister TranslatedVec = VectorMultiplyAdd(Translation, WWWW, RotatedVec);
+	const TransformVectorRegister WWWW = VectorReplicate(InputVector, 3);
+	const TransformVectorRegister TranslatedVec = VectorMultiplyAdd(Translation, WWWW, RotatedVec);
 
 	FVector4 NewVectOutput;
 	VectorStoreAligned(TranslatedVec, &NewVectOutput);
@@ -1623,243 +1672,261 @@ FORCEINLINE FVector4 FTransform::TransformFVector4(const FVector4& V) const
 }
 
 
-FORCEINLINE FVector FTransform::TransformPosition(const FVector& V) const
+template<typename T>
+FORCEINLINE TVector<T> TTransform<T>::TransformPosition(const TVector<T>& V) const
 {
 	DiagnosticCheckNaN_All();
 
-	const VectorRegister InputVectorW0 = VectorLoadFloat3_W0(&V);
+	const TransformVectorRegister InputVectorW0 = VectorLoadFloat3_W0(&V);
 
 	//Transform using QST is following
 	//QST(P) = Q.Rotate(S*P) + T where Q = quaternion, S = scale, T = translation
 	
 	//RotatedVec = Q.Rotate(Scale*V.X, Scale*V.Y, Scale*V.Z, 0.f)
-	const VectorRegister ScaledVec = VectorMultiply(Scale3D, InputVectorW0);
-	const VectorRegister RotatedVec = VectorQuaternionRotateVector(Rotation, ScaledVec);
+	const TransformVectorRegister ScaledVec = VectorMultiply(Scale3D, InputVectorW0);
+	const TransformVectorRegister RotatedVec = VectorQuaternionRotateVector(Rotation, ScaledVec);
 
-	const VectorRegister TranslatedVec = VectorAdd(RotatedVec, Translation);
+	const TransformVectorRegister TranslatedVec = VectorAdd(RotatedVec, Translation);
 
-	FVector Result;
+	TVector<T> Result;
 	VectorStoreFloat3(TranslatedVec, &Result);
 	return Result;
 }
 
-FORCEINLINE FVector FTransform::TransformPositionNoScale(const FVector& V) const
+template<typename T>
+FORCEINLINE TVector<T> TTransform<T>::TransformPositionNoScale(const TVector<T>& V) const
 {
 	DiagnosticCheckNaN_All();
 
-	const VectorRegister InputVectorW0 = VectorLoadFloat3_W0(&V);
+	const TransformVectorRegister InputVectorW0 = VectorLoadFloat3_W0(&V);
 
 	//Transform using QST is following
 	//QST(P) = Q.Rotate(S*P) + T where Q = quaternion, S = 1.0f, T = translation
 
 	//RotatedVec = Q.Rotate(V.X, V.Y, V.Z, 0.f)
-	const VectorRegister RotatedVec = VectorQuaternionRotateVector(Rotation, InputVectorW0);
+	const TransformVectorRegister RotatedVec = VectorQuaternionRotateVector(Rotation, InputVectorW0);
 
-	const VectorRegister TranslatedVec = VectorAdd(RotatedVec, Translation);
+	const TransformVectorRegister TranslatedVec = VectorAdd(RotatedVec, Translation);
 
-	FVector Result;
+	TVector<T> Result;
 	VectorStoreFloat3(TranslatedVec, &Result);
 	return Result;
 }
 
-FORCEINLINE FVector FTransform::TransformVector(const FVector& V) const
+template<typename T>
+FORCEINLINE TVector<T> TTransform<T>::TransformVector(const TVector<T>& V) const
 {
 	DiagnosticCheckNaN_All();
 
-	const VectorRegister InputVectorW0 = VectorLoadFloat3_W0(&V);
+	const TransformVectorRegister InputVectorW0 = VectorLoadFloat3_W0(&V);
 
 	//RotatedVec = Q.Rotate(Scale*V.X, Scale*V.Y, Scale*V.Z, 0.f)
-	const VectorRegister ScaledVec = VectorMultiply(Scale3D, InputVectorW0);
-	const VectorRegister RotatedVec = VectorQuaternionRotateVector(Rotation, ScaledVec);
+	const TransformVectorRegister ScaledVec = VectorMultiply(Scale3D, InputVectorW0);
+	const TransformVectorRegister RotatedVec = VectorQuaternionRotateVector(Rotation, ScaledVec);
 
-	FVector Result;
+	TVector<T> Result;
 	VectorStoreFloat3(RotatedVec, &Result);
 	return Result;
 }
 
-FORCEINLINE FVector FTransform::TransformVectorNoScale(const FVector& V) const
+template<typename T>
+FORCEINLINE TVector<T> TTransform<T>::TransformVectorNoScale(const TVector<T>& V) const
 {
 	DiagnosticCheckNaN_All();
 
-	const VectorRegister InputVectorW0 = VectorLoadFloat3_W0(&V);
+	const TransformVectorRegister InputVectorW0 = VectorLoadFloat3_W0(&V);
 
 	//RotatedVec = Q.Rotate(V.X, V.Y, V.Z, 0.f)
-	const VectorRegister RotatedVec = VectorQuaternionRotateVector(Rotation, InputVectorW0);
+	const TransformVectorRegister RotatedVec = VectorQuaternionRotateVector(Rotation, InputVectorW0);
 
-	FVector Result;
+	TVector<T> Result;
 	VectorStoreFloat3(RotatedVec, &Result);
 	return Result;
 }
 
 // do backward operation when inverse, translation -> rotation -> scale
-FORCEINLINE FVector FTransform::InverseTransformPosition(const FVector &V) const
+template<typename T>
+FORCEINLINE TVector<T> TTransform<T>::InverseTransformPosition(const TVector<T> &V) const
 {
 	DiagnosticCheckNaN_All();
 
-	const VectorRegister InputVector = VectorLoadFloat3_W0(&V);
+	const TransformVectorRegister InputVector = VectorLoadFloat3_W0(&V);
 
 	// (V-Translation)
-	const VectorRegister TranslatedVec = VectorSet_W0(VectorSubtract(InputVector, Translation));
+	const TransformVectorRegister TranslatedVec = VectorSet_W0(VectorSubtract(InputVector, Translation));
 
 	// ( Rotation.Inverse() * (V-Translation) )
-	const VectorRegister VR = VectorQuaternionInverseRotateVector(Rotation, TranslatedVec);
+	const TransformVectorRegister VR = VectorQuaternionInverseRotateVector(Rotation, TranslatedVec);
 
 	// GetSafeScaleReciprocal(Scale3D);
-	const VectorRegister SafeReciprocal = GetSafeScaleReciprocal(Scale3D);	
+	const TransformVectorRegister SafeReciprocal = GetSafeScaleReciprocal(Scale3D);	
 
 	// ( Rotation.Inverse() * (V-Translation) ) * GetSafeScaleReciprocal(Scale3D);
-	const VectorRegister VResult = VectorMultiply(VR, SafeReciprocal);
+	const TransformVectorRegister VResult = VectorMultiply(VR, SafeReciprocal);
 
-	FVector Result;
+	TVector<T> Result;
 	VectorStoreFloat3(VResult, &Result);
 	return Result;
 }
 
 // do backward operation when inverse, translation -> rotation
-FORCEINLINE FVector FTransform::InverseTransformPositionNoScale(const FVector &V) const
+template<typename T>
+FORCEINLINE TVector<T> TTransform<T>::InverseTransformPositionNoScale(const TVector<T> &V) const
 {
 	DiagnosticCheckNaN_All();
 
-	const VectorRegister InputVector = VectorLoadFloat3_W0(&V);
+	const TransformVectorRegister InputVector = VectorLoadFloat3_W0(&V);
 
 	// (V-Translation)
-	const VectorRegister TranslatedVec = VectorSet_W0(VectorSubtract(InputVector, Translation));
+	const TransformVectorRegister TranslatedVec = VectorSet_W0(VectorSubtract(InputVector, Translation));
 
 	// ( Rotation.Inverse() * (V-Translation) )
-	const VectorRegister VResult = VectorQuaternionInverseRotateVector(Rotation, TranslatedVec);
+	const TransformVectorRegister VResult = VectorQuaternionInverseRotateVector(Rotation, TranslatedVec);
 
-	FVector Result;
+	TVector<T> Result;
 	VectorStoreFloat3(VResult, &Result);
 	return Result;
 }
 
 
 // do backward operation when inverse, translation -> rotation -> scale
-FORCEINLINE FVector FTransform::InverseTransformVector(const FVector &V) const
+template<typename T>
+FORCEINLINE TVector<T> TTransform<T>::InverseTransformVector(const TVector<T> &V) const
 {
 	DiagnosticCheckNaN_All();
 
-	const VectorRegister InputVector = VectorLoadFloat3_W0(&V);
+	const TransformVectorRegister InputVector = VectorLoadFloat3_W0(&V);
 
-	// ( Rotation.Inverse() * V ) aka. FVector FQuat::operator*( const FVector& V ) const
-	const VectorRegister VR = VectorQuaternionInverseRotateVector(Rotation, InputVector);
+	// ( Rotation.Inverse() * V ) aka. TVector<T> TQuat<T>::operator*( const TVector<T>& V ) const
+	const TransformVectorRegister VR = VectorQuaternionInverseRotateVector(Rotation, InputVector);
 
 	// GetSafeScaleReciprocal(Scale3D);
-	const VectorRegister SafeReciprocal = GetSafeScaleReciprocal(Scale3D);
+	const TransformVectorRegister SafeReciprocal = GetSafeScaleReciprocal(Scale3D);
 
 	// ( Rotation.Inverse() * V) * GetSafeScaleReciprocal(Scale3D);
-	const VectorRegister VResult = VectorMultiply(VR, SafeReciprocal);
+	const TransformVectorRegister VResult = VectorMultiply(VR, SafeReciprocal);
 
-	FVector Result;
+	TVector<T> Result;
 	VectorStoreFloat3(VResult, &Result);
 	return Result;
 }
 
 // do backward operation when inverse, translation -> rotation
-FORCEINLINE FVector FTransform::InverseTransformVectorNoScale(const FVector &V) const
+template<typename T>
+FORCEINLINE TVector<T> TTransform<T>::InverseTransformVectorNoScale(const TVector<T> &V) const
 {
 	DiagnosticCheckNaN_All();
 
-	VectorRegister InputVector = VectorLoadFloat3_W0(&V);
+	TransformVectorRegister InputVector = VectorLoadFloat3_W0(&V);
 
 	// ( Rotation.Inverse() * V )
-	VectorRegister VResult = VectorQuaternionInverseRotateVector(Rotation, InputVector);
+	TransformVectorRegister VResult = VectorQuaternionInverseRotateVector(Rotation, InputVector);
 
-	FVector Result;
+	TVector<T> Result;
 	VectorStoreFloat3(VResult, &Result);
 	return Result;
 }
 
-FORCEINLINE FQuat FTransform::TransformRotation(const FQuat& Q) const
+template<typename T>
+FORCEINLINE TQuat<T> TTransform<T>::TransformRotation(const TQuat<T>& Q) const
 {
 	return GetRotation() * Q;
 }
 
-FORCEINLINE FQuat FTransform::InverseTransformRotation(const FQuat& Q) const
+template<typename T>
+FORCEINLINE TQuat<T> TTransform<T>::InverseTransformRotation(const TQuat<T>& Q) const
 {
 	return GetRotation().Inverse() * Q;
 }
 
-FORCEINLINE FTransform FTransform::operator*(const FTransform& Other) const
+template<typename T>
+FORCEINLINE TTransform<T> TTransform<T>::operator*(const TTransform<T>& Other) const
 {
-	FTransform Output;
+	TTransform<T> Output;
 	Multiply(&Output, this, &Other);
 	return Output;
 }
 
-FORCEINLINE void FTransform::operator*=(const FTransform& Other)
+template<typename T>
+FORCEINLINE void TTransform<T>::operator*=(const TTransform<T>& Other)
 {
 	Multiply(this, this, &Other);
 }
 
-FORCEINLINE FTransform FTransform::operator*(const FQuat& Other) const
+template<typename T>
+FORCEINLINE TTransform<T> TTransform<T>::operator*(const TQuat<T>& Other) const
 {
-	FTransform Output, OtherTransform(Other, FVector::ZeroVector, FVector::OneVector);
+	TTransform<T> Output, OtherTransform(Other, TVector<T>::ZeroVector, TVector<T>::OneVector);
 	Multiply(&Output, this, &OtherTransform);
 	return Output;
 }
 
-FORCEINLINE void FTransform::operator*=(const FQuat& Other)
+template<typename T>
+FORCEINLINE void TTransform<T>::operator*=(const TQuat<T>& Other)
 {
-	FTransform OtherTransform(Other, FVector::ZeroVector, FVector::OneVector);
+	TTransform<T> OtherTransform(Other, TVector<T>::ZeroVector, TVector<T>::OneVector);
 	Multiply(this, this, &OtherTransform);
 }
 
 // x = 0, y = 1, z = 2
-FORCEINLINE FVector FTransform::GetScaledAxis( EAxis::Type InAxis ) const
+template<typename T>
+FORCEINLINE TVector<T> TTransform<T>::GetScaledAxis( EAxis::Type InAxis ) const
 {
 	if ( InAxis == EAxis::X )
 	{
-		return TransformVector(FVector(1.f, 0.f, 0.f));
+		return TransformVector(TVector<T>(1.f, 0.f, 0.f));
 	}
 	else if ( InAxis == EAxis::Y )
 	{
-		return TransformVector(FVector(0.f, 1.f, 0.f));
+		return TransformVector(TVector<T>(0.f, 1.f, 0.f));
 	}
 
-	return TransformVector(FVector(0.f, 0.f, 1.f));
+	return TransformVector(TVector<T>(0.f, 0.f, 1.f));
 }
 
 // x = 0, y = 1, z = 2
-FORCEINLINE FVector FTransform::GetUnitAxis( EAxis::Type InAxis ) const
+template<typename T>
+FORCEINLINE TVector<T> TTransform<T>::GetUnitAxis( EAxis::Type InAxis ) const
 {
 	if ( InAxis == EAxis::X )
 	{
-		return TransformVectorNoScale(FVector(1.f, 0.f, 0.f));
+		return TransformVectorNoScale(TVector<T>(1.f, 0.f, 0.f));
 	}
 	else if ( InAxis == EAxis::Y )
 	{
-		return TransformVectorNoScale(FVector(0.f, 1.f, 0.f));
+		return TransformVectorNoScale(TVector<T>(0.f, 1.f, 0.f));
 	}
 
-	return TransformVectorNoScale(FVector(0.f, 0.f, 1.f));
+	return TransformVectorNoScale(TVector<T>(0.f, 0.f, 1.f));
 }
 
-FORCEINLINE void FTransform::Mirror(EAxis::Type MirrorAxis, EAxis::Type FlipAxis)
+template<typename T>
+FORCEINLINE void TTransform<T>::Mirror(EAxis::Type MirrorAxis, EAxis::Type FlipAxis)
 {
 	// We do convert to Matrix for mirroring. 
-	FMatrix M = ToMatrixWithScale();
+	UE::Math::TMatrix<T> M = ToMatrixWithScale();
 	M.Mirror(MirrorAxis, FlipAxis);
 	SetFromMatrix(M);
 }
 
 /** same version of FMatrix::GetMaximumAxisScale function **/
 /** @return the maximum magnitude of any row of the matrix. */
-inline float FTransform::GetMaximumAxisScale() const
+template<typename T>
+inline float TTransform<T>::GetMaximumAxisScale() const
 {
 	DiagnosticCheckNaN_Scale3D();
 
 	float Scale3DAbsMax;
 	// Scale3DAbsXYZ1 = { Abs(X), Abs(Y)), Abs(Z), 0 }
-	const VectorRegister Scale3DAbsXYZ0 =  VectorAbs(Scale3D);
+	const TransformVectorRegister Scale3DAbsXYZ0 =  VectorAbs(Scale3D);
 	// Scale3DAbsYZX1 = { Abs(Y),Abs(Z)),Abs(X), 0 }
-	const VectorRegister Scale3DAbsYZX0 = VectorSwizzle(Scale3DAbsXYZ0, 1,2,0,3);
+	const TransformVectorRegister Scale3DAbsYZX0 = VectorSwizzle(Scale3DAbsXYZ0, 1,2,0,3);
 	// Scale3DAbsZXY1 = { Abs(Z),Abs(X)),Abs(Y), 0 }
-	const VectorRegister Scale3DAbsZXY0 = VectorSwizzle(Scale3DAbsXYZ0, 2,0,1,3);
+	const TransformVectorRegister Scale3DAbsZXY0 = VectorSwizzle(Scale3DAbsXYZ0, 2,0,1,3);
 	// t0 = { Max(Abs(X), Abs(Y)),  Max(Abs(Y), Abs(Z)), Max(Abs(Z), Abs(X)), 0 }
-	const VectorRegister t0 = VectorMax(Scale3DAbsXYZ0, Scale3DAbsYZX0);
+	const TransformVectorRegister t0 = VectorMax(Scale3DAbsXYZ0, Scale3DAbsYZX0);
 	// t1 = { Max(Abs(X), Abs(Y), Abs(Z)), Max(Abs(Y), Abs(Z), Abs(X)), Max(Abs(Z), Abs(X), Abs(Y)), 0 }
-	const VectorRegister t2 = VectorMax(t0, Scale3DAbsZXY0);
+	const TransformVectorRegister t2 = VectorMax(t0, Scale3DAbsZXY0);
 	// Scale3DAbsMax = Max(Abs(X), Abs(Y), Abs(Z));
 	VectorStoreFloat1(t2, &Scale3DAbsMax);
 
@@ -1867,21 +1934,22 @@ inline float FTransform::GetMaximumAxisScale() const
 }
 
 /** @return the minimum magnitude of all components of the 3D scale. */
-inline float FTransform::GetMinimumAxisScale() const
+template<typename T>
+inline float TTransform<T>::GetMinimumAxisScale() const
 {
 	DiagnosticCheckNaN_Scale3D();
 
 	float Scale3DAbsMin;
 	// Scale3DAbsXYZ1 = { Abs(X), Abs(Y)), Abs(Z), 0 }
-	const VectorRegister Scale3DAbsXYZ0 =  VectorAbs(Scale3D);
+	const TransformVectorRegister Scale3DAbsXYZ0 =  VectorAbs(Scale3D);
 	// Scale3DAbsYZX1 = { Abs(Y),Abs(Z)),Abs(X), 0 }
-	const VectorRegister Scale3DAbsYZX0 = VectorSwizzle(Scale3DAbsXYZ0, 1,2,0,3);
+	const TransformVectorRegister Scale3DAbsYZX0 = VectorSwizzle(Scale3DAbsXYZ0, 1,2,0,3);
 	// Scale3DAbsZXY1 = { Abs(Z),Abs(X)),Abs(Y), 0 }
-	const VectorRegister Scale3DAbsZXY0 = VectorSwizzle(Scale3DAbsXYZ0, 2,0,1,3);
+	const TransformVectorRegister Scale3DAbsZXY0 = VectorSwizzle(Scale3DAbsXYZ0, 2,0,1,3);
 	// t0 = { Min(Abs(X), Abs(Y)),  Min(Abs(Y), Abs(Z)), Min(Abs(Z), Abs(X)), 0 }
-	const VectorRegister t0 = VectorMin(Scale3DAbsXYZ0, Scale3DAbsYZX0);
+	const TransformVectorRegister t0 = VectorMin(Scale3DAbsXYZ0, Scale3DAbsYZX0);
 	// t1 = { Min(Abs(X), Abs(Y), Abs(Z)), Min(Abs(Y), Abs(Z), Abs(X)), Min(Abs(Z), Abs(X), Abs(Y)), 0 }
-	const VectorRegister t2 = VectorMin(t0, Scale3DAbsZXY0);
+	const TransformVectorRegister t2 = VectorMin(t0, Scale3DAbsZXY0);
 	// Scale3DAbsMax = Min(Abs(X), Abs(Y), Abs(Z));
 	VectorStoreFloat1(t2, &Scale3DAbsMin);
 
@@ -1895,16 +1963,17 @@ inline float FTransform::GetMinimumAxisScale() const
  * also returning BIG_NUMBER causes sequential NaN issues by multiplying 
  * so we hardcode as 0
  */
-FORCEINLINE FVector FTransform::GetSafeScaleReciprocal(const FVector& InScale, float Tolerance)
+template<typename T>
+FORCEINLINE TVector<T> TTransform<T>::GetSafeScaleReciprocal(const TVector<T>& InScale, T Tolerance)
 {
-	FVector SafeReciprocalScale;
+	TVector<T> SafeReciprocalScale;
 	if (FMath::Abs(InScale.X) <= Tolerance)
 	{
 		SafeReciprocalScale.X = 0.f;
 	}
 	else
 	{
-		SafeReciprocalScale.X = 1/InScale.X;
+		SafeReciprocalScale.X = 1.f/InScale.X;
 	}
 
 	if (FMath::Abs(InScale.Y) <= Tolerance)
@@ -1913,7 +1982,7 @@ FORCEINLINE FVector FTransform::GetSafeScaleReciprocal(const FVector& InScale, f
 	}
 	else
 	{
-		SafeReciprocalScale.Y = 1/InScale.Y;
+		SafeReciprocalScale.Y = 1.f/InScale.Y;
 	}
 
 	if (FMath::Abs(InScale.Z) <= Tolerance)
@@ -1922,10 +1991,15 @@ FORCEINLINE FVector FTransform::GetSafeScaleReciprocal(const FVector& InScale, f
 	}
 	else
 	{
-		SafeReciprocalScale.Z = 1/InScale.Z;
+		SafeReciprocalScale.Z = 1.f/InScale.Z;
 	}
 
 	return SafeReciprocalScale;
 }
 
-#endif
+
+} // namespace UE::Math
+} // namespace UE
+
+
+#endif // ENABLE_VECTORIZED_TRANSFORM
