LWC support for FQuat and FTransform

diff --git a/Engine/Source/Runtime/Core/Public/Math/UnrealMathFPU.h b/Engine/Source/Runtime/Core/Public/Math/UnrealMathFPU.h
--- a/Engine/Source/Runtime/Core/Public/Math/UnrealMathFPU.h
+++ b/Engine/Source/Runtime/Core/Public/Math/UnrealMathFPU.h
@@ -1379,7 +1379,44 @@ FORCEINLINE void VectorQuaternionMultiply(VectorRegister4Float* RESTRICT Result,
 	R[3] = TW;
 #endif
 }
-// TODO: LWC: implement for doubles and remove 'void*'
+
+FORCEINLINE void VectorQuaternionMultiply(VectorRegister4Double* RESTRICT Result, const VectorRegister4Double* RESTRICT Quat1, const VectorRegister4Double* RESTRICT Quat2)
+{
+	typedef double Double4[4];
+	const Double4& A = *((const Double4*)Quat1);
+	const Double4& B = *((const Double4*)Quat2);
+	Double4& R = *((Double4*)Result);
+
+#if USE_FAST_QUAT_MUL
+	const double T0 = (A[2] - A[1]) * (B[1] - B[2]);
+	const double T1 = (A[3] + A[0]) * (B[3] + B[0]);
+	const double T2 = (A[3] - A[0]) * (B[1] + B[2]);
+	const double T3 = (A[1] + A[2]) * (B[3] - B[0]);
+	const double T4 = (A[2] - A[0]) * (B[0] - B[1]);
+	const double T5 = (A[2] + A[0]) * (B[0] + B[1]);
+	const double T6 = (A[3] + A[1]) * (B[3] - B[2]);
+	const double T7 = (A[3] - A[1]) * (B[3] + B[2]);
+	const double T8 = T5 + T6 + T7;
+	const double T9 = 0.5 * (T4 + T8);
+
+	R[0] = T1 + T9 - T8;
+	R[1] = T2 + T9 - T7;
+	R[2] = T3 + T9 - T6;
+	R[3] = T0 + T9 - T5;
+#else
+	// store intermediate results in temporaries
+	const double TX = A[3] * B[0] + A[0] * B[3] + A[1] * B[2] - A[2] * B[1];
+	const double TY = A[3] * B[1] - A[0] * B[2] + A[1] * B[3] + A[2] * B[0];
+	const double TZ = A[3] * B[2] + A[0] * B[1] - A[1] * B[0] + A[2] * B[3];
+	const double TW = A[3] * B[3] - A[0] * B[0] - A[1] * B[1] - A[2] * B[2];
+
+	// copy intermediate result to *this
+	R[0] = TX;
+	R[1] = TY;
+	R[2] = TZ;
+	R[3] = TW;
+#endif
+}
 
 /**
 * Multiplies two quaternions; the order matters.
@@ -1400,12 +1437,8 @@ FORCEINLINE VectorRegister4Float VectorQuaternionMultiply2( const VectorRegister
 
 FORCEINLINE VectorRegister4Double VectorQuaternionMultiply2(const VectorRegister4Double& Quat1, const VectorRegister4Double& Quat2)
 {
-	// TODO: LWC: implement for doubles, this is a lossy conversion for the interim.
-	VectorRegister4Float FloatResult;
-	VectorRegister4Float FloatQ1 = MakeVectorRegisterFloatFromDouble(Quat1); // Lossy conversion!
-	VectorRegister4Float FloatQ2 = MakeVectorRegisterFloatFromDouble(Quat2); // Lossy conversion!
-	VectorQuaternionMultiply(&FloatResult, &FloatQ1, &FloatQ2);
-	VectorRegister4Double Result = FloatResult; // convert back to double
+	VectorRegister4Double Result;
+	VectorQuaternionMultiply(&Result, &Quat1, &Quat2);
 	return Result;
 }
 
diff --git a/Engine/Source/Runtime/Core/Public/Math/UnrealMathNeon.h b/Engine/Source/Runtime/Core/Public/Math/UnrealMathNeon.h
--- a/Engine/Source/Runtime/Core/Public/Math/UnrealMathNeon.h
+++ b/Engine/Source/Runtime/Core/Public/Math/UnrealMathNeon.h
@@ -325,6 +325,12 @@ FORCEINLINE VectorRegister2Int64 MakeVectorRegisterInt64(int64 X, int64 Y)
 	return Tmp.V;
 }
 
+// Make double register from float register
+FORCEINLINE VectorRegister4Double MakeVectorRegisterDouble(const VectorRegister4Float& From)
+{
+	return VectorRegister4Double(From);
+}
+
 // Lossy conversion: double->float vector
 FORCEINLINE VectorRegister4Float MakeVectorRegisterFloatFromDouble(const VectorRegister4Double& Vec)
 {
@@ -2277,6 +2283,11 @@ FORCEINLINE void VectorQuaternionMultiply(VectorRegister4Float* RESTRICT Result,
 	*Result = VectorQuaternionMultiply2(*Quat1, *Quat2);
 }
 
+FORCEINLINE void VectorQuaternionMultiply(VectorRegister4Double* RESTRICT Result, const VectorRegister4Double* RESTRICT Quat1, const VectorRegister4Double* RESTRICT Quat2)
+{
+	*Result = VectorQuaternionMultiply2(*Quat1, *Quat2);
+}
+
 /**
 * Computes the sine and cosine of each component of a Vector.
 *
diff --git a/Engine/Source/Runtime/Core/Public/Math/UnrealMathSSE.h b/Engine/Source/Runtime/Core/Public/Math/UnrealMathSSE.h
--- a/Engine/Source/Runtime/Core/Public/Math/UnrealMathSSE.h
+++ b/Engine/Source/Runtime/Core/Public/Math/UnrealMathSSE.h
@@ -10,7 +10,7 @@
 #endif
 
 #ifndef UE_PLATFORM_MATH_USE_AVX
-#define UE_PLATFORM_MATH_USE_AVX			0 //PLATFORM_ALWAYS_HAS_AVX
+#define UE_PLATFORM_MATH_USE_AVX			PLATFORM_ALWAYS_HAS_AVX
 #endif
 
 #ifndef UE_PLATFORM_MATH_USE_AVX_2
@@ -502,14 +502,12 @@ FORCEINLINE VectorRegister4Double VectorLoad(const double* Ptr)
  */
 FORCEINLINE VectorRegister4Double VectorLoadFloat3(const double* Ptr)
 {
-#if !UE_PLATFORM_MATH_USE_AVX
+#if !UE_PLATFORM_MATH_USE_AVX_2
 	VectorRegister4Double Result;
 	Result.XY = _mm_loadu_pd((double*)(Ptr));
 	Result.ZW = _mm_load_sd((double*)(Ptr+2));
 	return Result;
 #else
-	// TODO: AVX: compare performance. MaskLoad is supposed to be efficient compared to MaskStore (from Intel and AMD manuals citing instruction latency), but validate this.
-	// Otherwise, use the m128d code above.
 	return _mm256_maskload_pd(Ptr, _mm256_castpd_si256(GlobalVectorConstants::DoubleXYZMask));
 #endif	
 }
@@ -522,14 +520,17 @@ FORCEINLINE VectorRegister4Double VectorLoadFloat3(const double* Ptr)
  */
 FORCEINLINE VectorRegister4Double VectorLoadFloat3_W1(const double* Ptr)
 {
-#if !UE_PLATFORM_MATH_USE_AVX
+#if !UE_PLATFORM_MATH_USE_AVX_2
 	VectorRegister4Double Result;
 	Result.XY = _mm_loadu_pd((double*)(Ptr));
 	Result.ZW = MakeVectorRegister2Double(Ptr[2], 1.0);
 	return Result;
 #else
-	// TODO: AVX: improve?
-	return MakeVectorRegisterDouble(Ptr[0], Ptr[1], Ptr[2], 1.0);
+	//return MakeVectorRegisterDouble(Ptr[0], Ptr[1], Ptr[2], 1.0);
+	VectorRegister4Double Result;
+	Result = _mm256_maskload_pd(Ptr, _mm256_castpd_si256(GlobalVectorConstants::DoubleXYZMask));
+	Result = _mm256_blend_pd(Result, VectorOneDouble(), 0b1000);
+	return Result;
 #endif	
 }
 
@@ -2121,40 +2122,50 @@ FORCEINLINE void VectorMatrixMultiply(FMatrix44f* Result, const FMatrix44f* Matr
 
 FORCEINLINE void VectorMatrixMultiply(FMatrix44d* Result, const FMatrix44d* Matrix1, const FMatrix44d* Matrix2)
 {
-	const VectorRegister4Double* A = (const VectorRegister4Double*)Matrix1;
-	const VectorRegister4Double* B = (const VectorRegister4Double*)Matrix2;
-	VectorRegister4Double* R = (VectorRegister4Double*)Result;
-	VectorRegister4Double Temp, R0, R1, R2;
+	// Warning: FMatrix44d alignment may not match VectorRegister4Double, so you can't just cast to VectorRegister4Double*.
+	typedef double Double4x4[4][4];
+	const Double4x4& AMRows = *((const Double4x4*)Matrix1);
+	const Double4x4& BMRows = *((const Double4x4*)Matrix2);
+	Double4x4& ResultDst = *((Double4x4*)Result);
+	VectorRegister4Double Temp, ARow;
+
+	VectorRegister4Double B[4];
+	B[0] = VectorLoad(BMRows[0]);
+	B[1] = VectorLoad(BMRows[1]);
+	B[2] = VectorLoad(BMRows[2]);
+	B[3] = VectorLoad(BMRows[3]);
 
 	// First row of result (Matrix1[0] * Matrix2).
-	Temp = VectorMultiply(VectorReplicate(A[0], 0), B[0]);
-	Temp = VectorMultiplyAdd(VectorReplicate(A[0], 1), B[1], Temp);
-	Temp = VectorMultiplyAdd(VectorReplicate(A[0], 2), B[2], Temp);
-	R0 = VectorMultiplyAdd(VectorReplicate(A[0], 3), B[3], Temp);
+	ARow = VectorLoad(AMRows[0]);
+	Temp = VectorMultiply(VectorReplicate(ARow, 0), B[0]);
+	Temp = VectorMultiplyAdd(VectorReplicate(ARow, 1), B[1], Temp);
+	Temp = VectorMultiplyAdd(VectorReplicate(ARow, 2), B[2], Temp);
+	Temp = VectorMultiplyAdd(VectorReplicate(ARow, 3), B[3], Temp);
+	VectorStore(Temp, ResultDst[0]);
 
 	// Second row of result (Matrix1[1] * Matrix2).
-	Temp = VectorMultiply(VectorReplicate(A[1], 0), B[0]);
-	Temp = VectorMultiplyAdd(VectorReplicate(A[1], 1), B[1], Temp);
-	Temp = VectorMultiplyAdd(VectorReplicate(A[1], 2), B[2], Temp);
-	R1 = VectorMultiplyAdd(VectorReplicate(A[1], 3), B[3], Temp);
+	ARow = VectorLoad(AMRows[1]);
+	Temp = VectorMultiply(VectorReplicate(ARow, 0), B[0]);
+	Temp = VectorMultiplyAdd(VectorReplicate(ARow, 1), B[1], Temp);
+	Temp = VectorMultiplyAdd(VectorReplicate(ARow, 2), B[2], Temp);
+	Temp = VectorMultiplyAdd(VectorReplicate(ARow, 3), B[3], Temp);
+	VectorStore(Temp, ResultDst[1]);
 
 	// Third row of result (Matrix1[2] * Matrix2).
-	Temp = VectorMultiply(VectorReplicate(A[2], 0), B[0]);
-	Temp = VectorMultiplyAdd(VectorReplicate(A[2], 1), B[1], Temp);
-	Temp = VectorMultiplyAdd(VectorReplicate(A[2], 2), B[2], Temp);
-	R2 = VectorMultiplyAdd(VectorReplicate(A[2], 3), B[3], Temp);
+	ARow = VectorLoad(AMRows[2]);
+	Temp = VectorMultiply(VectorReplicate(ARow, 0), B[0]);
+	Temp = VectorMultiplyAdd(VectorReplicate(ARow, 1), B[1], Temp);
+	Temp = VectorMultiplyAdd(VectorReplicate(ARow, 2), B[2], Temp);
+	Temp = VectorMultiplyAdd(VectorReplicate(ARow, 3), B[3], Temp);
+	VectorStore(Temp, ResultDst[2]);
 
 	// Fourth row of result (Matrix1[3] * Matrix2).
-	Temp = VectorMultiply(VectorReplicate(A[3], 0), B[0]);
-	Temp = VectorMultiplyAdd(VectorReplicate(A[3], 1), B[1], Temp);
-	Temp = VectorMultiplyAdd(VectorReplicate(A[3], 2), B[2], Temp);
-	Temp = VectorMultiplyAdd(VectorReplicate(A[3], 3), B[3], Temp);
-
-	// Store result
-	R[0] = R0;
-	R[1] = R1;
-	R[2] = R2;
-	R[3] = Temp;
+	ARow = VectorLoad(AMRows[3]);
+	Temp = VectorMultiply(VectorReplicate(ARow, 0), B[0]);
+	Temp = VectorMultiplyAdd(VectorReplicate(ARow, 1), B[1], Temp);
+	Temp = VectorMultiplyAdd(VectorReplicate(ARow, 2), B[2], Temp);
+	Temp = VectorMultiplyAdd(VectorReplicate(ARow, 3), B[3], Temp);
+	VectorStore(Temp, ResultDst[3]);
 }
 
 
@@ -2362,7 +2373,16 @@ FORCEINLINE VectorRegister4Float VectorTransformVector(const VectorRegister4Floa
 
 FORCEINLINE VectorRegister4Float VectorTransformVector(const VectorRegister4Float& VecP, const FMatrix44d* MatrixM)
 {
-	const VectorRegister4Double* M = (const VectorRegister4Double*) MatrixM;
+	// Warning: FMatrix44d alignment may not match VectorRegister4Double, so you can't just cast to VectorRegister4Double*.
+	typedef double Double4x4[4][4];
+	const Double4x4& MRows = *((const Double4x4*)MatrixM);
+
+	VectorRegister4Double M[4];
+	M[0] = VectorLoad(MRows[0]);
+	M[1] = VectorLoad(MRows[1]);
+	M[2] = VectorLoad(MRows[2]);
+	M[3] = VectorLoad(MRows[3]);
+
 	VectorRegister4Double VTempX, VTempY, VTempZ, VTempW;
 	VectorRegister4Double VecPDouble = VecP;
 
@@ -2383,7 +2403,16 @@ FORCEINLINE VectorRegister4Float VectorTransformVector(const VectorRegister4Floa
 
 FORCEINLINE VectorRegister4Double VectorTransformVector(const VectorRegister4Double& VecP, const FMatrix44d* MatrixM)
 {
-	const VectorRegister4Double* M = (const VectorRegister4Double*) MatrixM;
+	// Warning: FMatrix44d alignment may not match VectorRegister4Double, so you can't just cast to VectorRegister4Double*.
+	typedef double Double4x4[4][4];
+	const Double4x4& MRows = *((const Double4x4*)MatrixM);
+
+	VectorRegister4Double M[4];
+	M[0] = VectorLoad(MRows[0]);
+	M[1] = VectorLoad(MRows[1]);
+	M[2] = VectorLoad(MRows[2]);
+	M[3] = VectorLoad(MRows[3]);
+
 	VectorRegister4Double VTempX, VTempY, VTempZ, VTempW;
 
 	// Splat x,y,z and w
@@ -2818,7 +2847,11 @@ FORCEINLINE void VectorQuaternionMultiply(VectorRegister4Float* RESTRICT Result,
 {
 	*Result = VectorQuaternionMultiply2(*Quat1, *Quat2);
 }
-// TODO: LWC: implement for doubles / FQuat
+
+FORCEINLINE void VectorQuaternionMultiply(VectorRegister4Double* RESTRICT Result, const VectorRegister4Double* RESTRICT Quat1, const VectorRegister4Double* RESTRICT Quat2)
+{
+	*Result = VectorQuaternionMultiply2(*Quat1, *Quat2);
+}
 
 // Returns true if the vector contains a component that is either NAN or +/-infinite.
 FORCEINLINE bool VectorContainsNaNOrInfinite(const VectorRegister4Float& Vec)
diff --git a/Engine/Source/Runtime/Core/Public/Math/UnrealMathUtility.h b/Engine/Source/Runtime/Core/Public/Math/UnrealMathUtility.h
--- a/Engine/Source/Runtime/Core/Public/Math/UnrealMathUtility.h
+++ b/Engine/Source/Runtime/Core/Public/Math/UnrealMathUtility.h
@@ -31,9 +31,9 @@ DECLARE_LWC_TYPE(Plane, 4);
 struct  FBox;
 struct  FRotator;
 DECLARE_LWC_TYPE(Matrix, 44);
-struct  FQuat;
+DECLARE_LWC_TYPE(Quat, 4);
 struct  FTwoVectors;
-struct  FTransform;
+DECLARE_LWC_TYPE(Transform, 3);
 class  FSphere;
 struct FVector2D;
 struct FLinearColor;
@@ -147,6 +147,31 @@ class TRange;
 	Global functions.
 -----------------------------------------------------------------------------*/
 
+/**
+ * Template helper for FMath::Lerp<>() and related functions.
+ * By default, any type T is assumed to not need a custom Lerp implementation (IsRequired=false).
+ * However a class that requires custom functionality (eg FQuat) can specialize the template to define IsRequired=true and
+ * implement the Lerp() function and other similar functions and provide a custom implementation.
+ * Example:
+ * 
+ *	template<> struct TCustomLerp< MyClass >
+ *	{
+ *		// Required to use our custom Lerp() function below.
+ *		enum { IsRequired = true };
+ *
+ *		// Implements for float Alpha param. You could also add overrides or make it a template param.
+ *		static inline MyClass Lerp(const MyClass& A, const MyClass& B, const float& Alpha)
+ *		{
+ *			return MyClass::Lerp(A, B, Alpha); // Or do the computation here directly.
+ *		}
+ *	};
+ */
+template <typename T>
+struct TCustomLerp
+{
+	enum { IsRequired = false };
+};
+
 /**
  * Structure for all math helper functions, inherits from platform math to pick up platform-specific implementations
  * Check GenericPlatformMath.h for additional math functions
@@ -269,19 +294,18 @@ struct FMath : public FPlatformMath
 	 */
 	static FORCEINLINE bool IsNearlyEqual(float A, float B, float ErrorTolerance = SMALL_NUMBER)
 	{
-		return Abs<float>( A - B ) <= ErrorTolerance;
+		return Abs( A - B ) <= ErrorTolerance;
 	}
 
-	/**
-	 *	Checks if two floating point numbers are nearly equal.
-	 *	@param A				First number to compare
-	 *	@param B				Second number to compare
-	 *	@param ErrorTolerance	Maximum allowed difference for considering them as 'nearly equal'
-	 *	@return					true if A and B are nearly equal
-	 */
-	static FORCEINLINE bool IsNearlyEqual(double A, double B, double ErrorTolerance = SMALL_NUMBER)
+	static FORCEINLINE bool IsNearlyEqual(double A, double B, double ErrorTolerance = DOUBLE_SMALL_NUMBER)
+	{
+		return Abs(A - B) <= ErrorTolerance;
+	}
+
+	// Overload to resolve common scenario with the first value from a variable and 2 hard-coded constants.
+	static FORCEINLINE bool IsNearlyEqual(double A, float B, float ErrorTolerance = SMALL_NUMBER)
 	{
-		return Abs<double>( A - B ) <= ErrorTolerance;
+		return Abs(A - (double)B) <= (double)ErrorTolerance;
 	}
 
 	/**
@@ -632,6 +656,11 @@ public:
 	}
 #undef FASTASIN_HALF_PI
 
+	static FORCEINLINE double FastAsin(double Value)
+	{
+		// TODO: add fast approximation
+		return FMath::Asin(Value);
+	}
 
 	// Conversion Functions
 
@@ -667,10 +696,11 @@ public:
 	static CORE_API float ClampAngle(float AngleDegrees, float MinAngleDegrees, float MaxAngleDegrees);
 
 	/** Find the smallest angle between two headings (in degrees) */
-	static float FindDeltaAngleDegrees(float A1, float A2)
+	template<typename T, TEMPLATE_REQUIRES(TIsFloatingPoint<T>::Value)>
+	static T FindDeltaAngleDegrees(T A1, T A2)
 	{
 		// Find the difference
-		float Delta = A2 - A1;
+		T Delta = A2 - A1;
 
 		// If change is larger than 180
 		if (Delta > 180.0f)
@@ -690,10 +720,11 @@ public:
 	}
 
 	/** Find the smallest angle between two headings (in radians) */
-	static float FindDeltaAngleRadians(float A1, float A2)
+	template<typename T, TEMPLATE_REQUIRES(TIsFloatingPoint<T>::Value)>
+	static T FindDeltaAngleRadians(T A1, T A2)
 	{
 		// Find the difference
-		float Delta = A2 - A1;
+		T Delta = A2 - A1;
 
 		// If change is larger than PI
 		if (Delta > PI)
@@ -719,7 +750,8 @@ public:
 	}
 
 	/** Given a heading which may be outside the +/- PI range, 'unwind' it back into that range. */
-	static float UnwindRadians(float A)
+	template<typename T, TEMPLATE_REQUIRES(TIsFloatingPoint<T>::Value)>
+	static T UnwindRadians(T A)
 	{
 		while(A > PI)
 		{
@@ -735,7 +767,8 @@ public:
 	}
 
 	/** Utility to ensure angle is between +/- 180 degrees by unwinding. */
-	static float UnwindDegrees(float A)
+	template<typename T, TEMPLATE_REQUIRES(TIsFloatingPoint<T>::Value)>
+	static T UnwindDegrees(T A)
 	{
 		while(A > 180.f)
 		{
@@ -881,12 +914,19 @@ public:
 	}
 
 	/** Performs a linear interpolation between two values, Alpha ranges from 0-1 */
-	template< class T, class U > 
+	template< class T, class U, TEMPLATE_REQUIRES(!TCustomLerp<T>::IsRequired && (TIsFloatingPoint<U>::Value || TIsSame<T,U>::Value)) >
 	static FORCEINLINE_DEBUGGABLE T Lerp( const T& A, const T& B, const U& Alpha )
 	{
 		return (T)(A + Alpha * (B-A));
 	}
 
+	/** Custom lerps defined for those classes not suited to the default implemenation. */
+	template< class T, class U, TEMPLATE_REQUIRES(TCustomLerp<T>::IsRequired) >
+	static FORCEINLINE_DEBUGGABLE T Lerp(const T& A, const T& B, const U& Alpha)
+	{
+		return TCustomLerp<T>::Lerp(A, B, Alpha);
+	}
+
 	/** Performs a linear interpolation between two values, Alpha ranges from 0-1. Handles full numeric range of T */
 	template< class T > 
 	static FORCEINLINE_DEBUGGABLE T LerpStable( const T& A, const T& B, double Alpha )
@@ -902,7 +942,7 @@ public:
 	}
 
 	/** Performs a 2D linear interpolation between four values values, FracX, FracY ranges from 0-1 */
-	template< class T, class U > 
+	template< class T, class U, TEMPLATE_REQUIRES(!TCustomLerp<T>::IsRequired && (TIsFloatingPoint<U>::Value || TIsSame<T, U>::Value)) >
 	static FORCEINLINE_DEBUGGABLE T BiLerp(const T& P00,const T& P10,const T& P01,const T& P11, const U& FracX, const U& FracY)
 	{
 		return Lerp(
@@ -912,6 +952,13 @@ public:
 			);
 	}
 
+	/** Custom lerps defined for those classes not suited to the default implemenation. */
+	template< class T, class U, TEMPLATE_REQUIRES(TCustomLerp<T>::IsRequired) >
+	static FORCEINLINE_DEBUGGABLE T BiLerp(const T& P00, const T& P10, const T& P01, const T& P11, const U& FracX, const U& FracY)
+	{
+		return TCustomLerp<T>::BiLerp(P00, P10, P01, P11, FracX, FracY);
+	}
+
 	/**
 	 * Performs a cubic interpolation
 	 *
@@ -921,7 +968,7 @@ public:
 	 *
 	 * @return  Interpolated value
 	 */
-	template< class T, class U > 
+	template< class T, class U, TEMPLATE_REQUIRES(!TCustomLerp<T>::IsRequired && (TIsFloatingPoint<U>::Value || TIsSame<T, U>::Value)) >
 	static FORCEINLINE_DEBUGGABLE T CubicInterp( const T& P0, const T& T0, const T& P1, const T& T1, const U& A )
 	{
 		const U A2 = A * A;
@@ -930,6 +977,13 @@ public:
 		return T((((2*A3)-(3*A2)+1) * P0) + ((A3-(2*A2)+A) * T0) + ((A3-A2) * T1) + (((-2*A3)+(3*A2)) * P1));
 	}
 
+	/** Custom lerps defined for those classes not suited to the default implemenation. */
+	template< class T, class U, TEMPLATE_REQUIRES(TCustomLerp<T>::IsRequired) >
+	static FORCEINLINE_DEBUGGABLE T CubicInterp(const T& P0, const T& T0, const T& P1, const T& T1, const U& A)
+	{
+		return TCustomLerp<T>::CubicInterp(P0, T0, P1, T1, A);
+	}
+
 	/**
 	 * Performs a first derivative cubic interpolation
 	 *
@@ -939,7 +993,7 @@ public:
 	 *
 	 * @return  Interpolated value
 	 */
-	template< class T, class U > 
+	template< class T, class U, TEMPLATE_REQUIRES(TIsFloatingPoint<U>::Value) >
 	static FORCEINLINE_DEBUGGABLE T CubicInterpDerivative( const T& P0, const T& T0, const T& P1, const T& T1, const U& A )
 	{
 		T a = 6.f*P0 + 3.f*T0 + 3.f*T1 - 6.f*P1;
@@ -960,7 +1014,7 @@ public:
 	 *
 	 * @return  Interpolated value
 	 */
-	template< class T, class U > 
+	template< class T, class U, TEMPLATE_REQUIRES(TIsFloatingPoint<U>::Value) >
 	static FORCEINLINE_DEBUGGABLE T CubicInterpSecondDerivative( const T& P0, const T& T0, const T& P1, const T& T1, const U& A )
 	{
 		T a = 12.f*P0 + 6.f*T0 + 6.f*T1 - 12.f*P1;
@@ -1090,19 +1144,8 @@ public:
 	}
 
 	// Rotator specific interpolation
-	template< class U > static FRotator Lerp(const FRotator& A, const FRotator& B, const U& Alpha);
-	template< class U > static FRotator LerpRange(const FRotator& A, const FRotator& B, const U& Alpha);
-
-	// Quat-specific interpolation
-
-	template< class U > static FQuat Lerp(const FQuat& A, const FQuat& B, const U& Alpha);
-	template< class U > static FQuat BiLerp(const FQuat& P00, const FQuat& P10, const FQuat& P01, const FQuat& P11, float FracX, float FracY);
-
-	/**
-	 * In the case of quaternions, we use a bezier like approach.
-	 * T - Actual 'control' orientations.
-	 */
-	template< class U > static FQuat CubicInterp( const FQuat& P0, const FQuat& T0, const FQuat& P1, const FQuat& T1, const U& A);
+	static FRotator LerpRange(const FRotator& A, const FRotator& B, float Alpha);
+	static FRotator LerpRange(const FRotator& A, const FRotator& B, double Alpha);
 
 	/*
 	 *	Cubic Catmull-Rom Spline interpolation. Based on http://www.cemyuksel.com/research/catmullrom_param/catmullrom.pdf 
@@ -1208,10 +1251,12 @@ public:
 	static CORE_API FLinearColor CInterpTo(const FLinearColor& Current, const FLinearColor& Target, float DeltaTime, float InterpSpeed);
 
 	/** Interpolate quaternion from Current to Target with constant step (in radians) */
-	static CORE_API FQuat QInterpConstantTo(const FQuat& Current, const FQuat& Target, float DeltaTime, float InterpSpeed);
+	template< class T >
+	static CORE_API UE::Math::TQuat<T> QInterpConstantTo(const UE::Math::TQuat<T>& Current, const UE::Math::TQuat<T>& Target, float DeltaTime, float InterpSpeed);
 
 	/** Interpolate quaternion from Current to Target. Scaled by angle to Target, so it has a strong start speed and ease out. */
-	static CORE_API FQuat QInterpTo(const FQuat& Current, const FQuat& Target, float DeltaTime, float InterpSpeed);
+	template< class T >
+	static CORE_API UE::Math::TQuat<T> QInterpTo(const UE::Math::TQuat<T>& Current, const UE::Math::TQuat<T>& Target, float DeltaTime, float InterpSpeed);
 
 	/**
 	* Returns an approximation of Exp(-X) based on a Taylor expansion that has had the coefficients adjusted (using
diff --git a/Engine/Source/Runtime/Core/Public/Math/UnrealMathVectorCommon.h b/Engine/Source/Runtime/Core/Public/Math/UnrealMathVectorCommon.h
--- a/Engine/Source/Runtime/Core/Public/Math/UnrealMathVectorCommon.h
+++ b/Engine/Source/Runtime/Core/Public/Math/UnrealMathVectorCommon.h
@@ -25,13 +25,14 @@ FORCEINLINE VectorRegister4Float VectorOne(void)
  * Overloads for code passing &MyVector, &MyQuat, etc.
  */
 
- ///////////////////////////////////////////////////////////////////////////////////////////////////////////////////
- /// VectorLoad
- ///////////////////////////////////////////////////////////////////////////////////////////////////////////////////
+///////////////////////////////////////////////////////////////////////////////////////////////////////////////////
+/// VectorLoad
+///////////////////////////////////////////////////////////////////////////////////////////////////////////////////
 
-FORCEINLINE VectorRegister4Float VectorLoad(const FQuat* Ptr)
+template<typename T>
+FORCEINLINE TVectorRegisterType<T> VectorLoad(const UE::Math::TQuat<T>* Ptr)
 {
-	return VectorLoad((const float*)(Ptr));
+	return VectorLoadAligned((const T*)(Ptr));
 }
 
 FORCEINLINE VectorRegister4Float VectorLoad(const FVector4* Ptr)
@@ -39,14 +40,10 @@ FORCEINLINE VectorRegister4Float VectorLoad(const FVector4* Ptr)
 	return VectorLoad((const float*)(Ptr));
 }
 
-FORCEINLINE VectorRegister4Float VectorLoad(const FVector3f* Ptr)
-{
-	return VectorLoad((const float*)(Ptr));
-}
-
-FORCEINLINE VectorRegister4Double VectorLoad(const FVector3d* Ptr)
+template<typename T>
+FORCEINLINE TVectorRegisterType<T> VectorLoad(const UE::Math::TVector<T>* Ptr)
 {
-	return VectorLoad((const double*)(Ptr));
+	return VectorLoad((const T*)(Ptr));
 }
 
 FORCEINLINE VectorRegister4Float VectorLoad(const FVector2D* Ptr)
@@ -71,9 +68,10 @@ FORCEINLINE VectorRegister4Double VectorLoad(const VectorRegister4Double* Ptr)
 /// VectorLoadAligned
 ///////////////////////////////////////////////////////////////////////////////////////////////////////////////////
 
-FORCEINLINE VectorRegister4Float VectorLoadAligned(const FQuat* Ptr)
+template<typename T>
+FORCEINLINE TVectorRegisterType<T> VectorLoadAligned(const UE::Math::TQuat<T>* Ptr)
 {
-	return VectorLoadAligned((const float*)(Ptr));
+	return VectorLoadAligned((const T*)(Ptr));
 }
 
 FORCEINLINE VectorRegister4Float VectorLoadAligned(const FVector4* Ptr)
@@ -81,14 +79,10 @@ FORCEINLINE VectorRegister4Float VectorLoadAligned(const FVector4* Ptr)
 	return VectorLoadAligned((const float*)(Ptr));
 }
 
-FORCEINLINE VectorRegister4Float VectorLoadAligned(const FPlane4f* Ptr)
+template<typename T>
+FORCEINLINE TVectorRegisterType<T> VectorLoadAligned(const UE::Math::TPlane<T>* Ptr)
 {
-	return VectorLoadAligned((const float*)(Ptr));
-}
-
-FORCEINLINE VectorRegister4Double VectorLoadAligned(const FPlane4d* Ptr)
-{
-	return VectorLoadAligned((const double*)(Ptr));
+	return VectorLoadAligned((const T*)(Ptr));
 }
 
 FORCEINLINE VectorRegister4Float VectorLoadAligned(const VectorRegister4Float* Ptr)
@@ -119,14 +113,10 @@ FORCEINLINE VectorRegister4Float VectorLoadFloat3(const float* Ptr)
 	return MakeVectorRegister(Ptr[0], Ptr[1], Ptr[2], 0.0f);
 }
 
-FORCEINLINE VectorRegister4Float VectorLoadFloat3(const FVector3f* Ptr)
+template<typename T>
+FORCEINLINE TVectorRegisterType<T> VectorLoadFloat3(const UE::Math::TVector<T>* Ptr)
 {
-	return VectorLoadFloat3((const float*)(Ptr));
-}
-
-FORCEINLINE VectorRegister4Double VectorLoadFloat3(const FVector3d* Ptr)
-{
-	return VectorLoadFloat3((const double*)(Ptr));
+	return VectorLoadFloat3((const T*)(Ptr));
 }
 
 FORCEINLINE VectorRegister4Double VectorLoadDouble3(const double* Ptr)
@@ -155,14 +145,10 @@ FORCEINLINE VectorRegister4Double VectorLoadFloat3_W0(const double* Ptr)
 	return VectorLoadFloat3(Ptr);
 }
 
-FORCEINLINE VectorRegister4Float VectorLoadFloat3_W0(const FVector3f* Ptr)
+template<typename T>
+FORCEINLINE TVectorRegisterType<T> VectorLoadFloat3_W0(const UE::Math::TVector<T>* Ptr)
 {
-	return VectorLoadFloat3_W0((const float*)(Ptr));
-}
-
-FORCEINLINE VectorRegister4Double VectorLoadFloat3_W0(const FVector3d* Ptr)
-{
-	return VectorLoadFloat3_W0((const double*)(Ptr));
+	return VectorLoadFloat3_W0((const T*)(Ptr));
 }
 
 FORCEINLINE VectorRegister4Float VectorLoadFloat3_W0(const FRotator* Ptr)
@@ -191,14 +177,10 @@ FORCEINLINE VectorRegister4Float VectorLoadFloat3_W1(const float* Ptr)
 	return MakeVectorRegister(Ptr[0], Ptr[1], Ptr[2], 1.0f);
 }
 
-FORCEINLINE VectorRegister4Float VectorLoadFloat3_W1(const FVector3f* Ptr)
-{
-	return VectorLoadFloat3_W1((const float*)Ptr);
-}
-
-FORCEINLINE VectorRegister4Double VectorLoadFloat3_W1(const FVector3d* Ptr)
+template<typename T>
+FORCEINLINE TVectorRegisterType<T> VectorLoadFloat3_W1(const UE::Math::TVector<T>* Ptr)
 {
-	return VectorLoadFloat3_W1((const double*)Ptr);
+	return VectorLoadFloat3_W1((const T*)Ptr);
 }
 
 FORCEINLINE VectorRegister4Double VectorLoadDouble3_W1(const double* Ptr)
@@ -238,16 +220,16 @@ FORCEINLINE void VectorStore(const VectorRegister4Double& Vec, float* Dst)
 }
 #endif
 
-template<typename FReal>
-FORCEINLINE void VectorStore(const VectorRegister4Float& Vec, UE::Math::TVector<FReal>* Dst)
+template<typename T>
+FORCEINLINE void VectorStore(const VectorRegister4Float& Vec, UE::Math::TVector<T>* Dst)
 {
-	VectorStore(Vec, (FReal*)Dst);
+	VectorStore(Vec, (T*)Dst);
 }
 
-template<typename FReal>
-FORCEINLINE void VectorStore(const VectorRegister4Double& Vec, UE::Math::TVector<FReal>* Dst)
+template<typename T>
+FORCEINLINE void VectorStore(const VectorRegister4Double& Vec, UE::Math::TVector<T>* Dst)
 {
-	VectorStore(Vec, (FReal*)Dst);
+	VectorStore(Vec, (T*)Dst);
 }
 
 FORCEINLINE void VectorStore(const VectorRegister4Float& Vec, FVector4* Dst)
@@ -255,6 +237,19 @@ FORCEINLINE void VectorStore(const VectorRegister4Float& Vec, FVector4* Dst)
 	VectorStore(Vec, (float*)Dst); // TODO: LWC: convert once templated.
 }
 
+template<typename T>
+FORCEINLINE void VectorStore(const TVectorRegisterType<T>& Vec, struct UE::Math::TQuat<T>* Dst)
+{
+	VectorStore(Vec, (T*)Dst);
+}
+
+// Specific overload to support promoting float->double vector and storing in TQuat<double>
+FORCEINLINE void VectorStore(const VectorRegister4Float& Vec, struct UE::Math::TQuat<double>* Dst)
+{
+	VectorRegister4Double DoubleVec(Vec);
+	VectorStore(DoubleVec, Dst);
+}
+
 FORCEINLINE void VectorStore(const VectorRegister4Float& Vec, VectorRegister4Float* Dst)
 {
 	//VectorStore(Vec, (float*)Dst);
@@ -300,15 +295,17 @@ FORCEINLINE void VectorStoreAligned(const VectorRegister4Double& Vec, FVector4*
 }
 #endif
 
-// TODO: LWC: remove temp hack once FQuat is templated
-FORCEINLINE void VectorStoreAligned(const VectorRegister4Float& Vec, struct FQuat* Dst)
+template<typename T>
+FORCEINLINE void VectorStoreAligned(const TVectorRegisterType<T>& Vec, struct UE::Math::TQuat<T>* Dst)
 {
-	VectorStoreAligned(Vec, (float*)Dst);
+	VectorStoreAligned(Vec, (T*)Dst);
 }
-// TODO: LWC: remove temp hack once FQuat is templated
-FORCEINLINE void VectorStoreAligned(const VectorRegister4Double& Vec, struct FQuat* Dst)
+
+// Specific overload to support promoting float->double vector and storing in TQuat<double>
+FORCEINLINE void VectorStoreAligned(const VectorRegister4Float& Vec, struct UE::Math::TQuat<double>* Dst)
 {
-	VectorStoreAligned(Vec, (float*)Dst);
+	VectorRegister4Double DoubleVec(Vec);
+	VectorStoreAligned(DoubleVec, Dst);
 }
 
 FORCEINLINE void VectorStoreAligned(const VectorRegister4Float& Vec, VectorRegister4Float* Dst)
@@ -334,16 +331,16 @@ FORCEINLINE void VectorStoreFloat3(const VectorRegister4Double& Vec, float* Dst)
 }
 #endif
 
-template<typename FReal>
-FORCEINLINE void VectorStoreFloat3(const VectorRegister4Float& Vec, UE::Math::TVector<FReal>* Dst)
+template<typename T>
+FORCEINLINE void VectorStoreFloat3(const VectorRegister4Float& Vec, UE::Math::TVector<T>* Dst)
 {
-	VectorStoreFloat3(Vec, (FReal*)Dst);
+	VectorStoreFloat3(Vec, (T*)Dst);
 }
 
-template<typename FReal>
-FORCEINLINE void VectorStoreFloat3(const VectorRegister4Double& Vec, UE::Math::TVector<FReal>* Dst)
+template<typename T>
+FORCEINLINE void VectorStoreFloat3(const VectorRegister4Double& Vec, UE::Math::TVector<T>* Dst)
 {
-	VectorStoreFloat3(Vec, (FReal*)Dst);
+	VectorStoreFloat3(Vec, (T*)Dst);
 }
 
 FORCEINLINE void VectorStoreFloat3(const VectorRegister4Float& Vec, FRotator* Dst)
@@ -388,12 +385,12 @@ FORCEINLINE void VectorStoreFloat1(const VectorRegister4Double& Vec, int64* Dst)
 
 // Returns a normalized 4 vector = Vector / |Vector|.
 // There is no handling of zero length vectors, use VectorNormalizeSafe if this is a possible input.
-template<typename VectorRegisterType>
-FORCEINLINE VectorRegisterType VectorNormalizeAccurate( const VectorRegisterType& Vector )
+template<typename TVectorRegisterType>
+FORCEINLINE TVectorRegisterType VectorNormalizeAccurate( const TVectorRegisterType& Vector )
 {
-	const VectorRegisterType SquareSum = VectorDot4(Vector, Vector);
-	const VectorRegisterType InvLength = VectorReciprocalSqrtAccurate(SquareSum);
-	const VectorRegisterType NormalizedVector = VectorMultiply(InvLength, Vector);
+	const TVectorRegisterType SquareSum = VectorDot4(Vector, Vector);
+	const TVectorRegisterType InvLength = VectorReciprocalSqrtAccurate(SquareSum);
+	const TVectorRegisterType NormalizedVector = VectorMultiply(InvLength, Vector);
 	return NormalizedVector;
 }
 
@@ -494,18 +491,33 @@ template<> FORCEINLINE VectorRegister4Double FGenericPlatformMath::Max(const Vec
 	return VectorMax(A, B);
 }
 
-// Specialization of Lerp template that works with vector registers
-template<> FORCEINLINE VectorRegister4Float FMath::Lerp(const VectorRegister4Float& A, const VectorRegister4Float& B, const VectorRegister4Float& Alpha)
+// TCustomLerp for FMath::Lerp<VectorRegister4Float>()
+template<>
+struct TCustomLerp<VectorRegister4Float>
 {
-	const VectorRegister4Float Delta = VectorSubtract(B, A);
-	return VectorMultiplyAdd(Alpha, Delta, A);
-}
+	enum { IsRequired = true };
+	
+	// Specialization of Lerp function that works with vector registers
+	static FORCEINLINE VectorRegister4Float Lerp(const VectorRegister4Float& A, const VectorRegister4Float& B, const VectorRegister4Float& Alpha)
+	{
+		const VectorRegister4Float Delta = VectorSubtract(B, A);
+		return VectorMultiplyAdd(Alpha, Delta, A);
+	}
+};
+
+// TCustomLerp for FMath::Lerp<VectorRegister4Double>()
+template<>
+struct TCustomLerp<VectorRegister4Double>
+{
+	enum { IsRequired = true };
 
-//template<> FORCEINLINE VectorRegister4Double FMath::Lerp(const VectorRegister4Double& A, const VectorRegister4Double& B, const VectorRegister4Double& Alpha)
-//{
-//	const VectorRegister4Double Delta = VectorSubtract(B, A);
-//	return VectorMultiplyAdd(Alpha, Delta, A);
-//}
+	// Specialization of Lerp function that works with vector registers
+	static FORCEINLINE VectorRegister4Double Lerp(const VectorRegister4Double& A, const VectorRegister4Double& B, const VectorRegister4Double& Alpha)
+	{
+		const VectorRegister4Double Delta = VectorSubtract(B, A);
+		return VectorMultiplyAdd(Alpha, Delta, A);
+	}
+};
 
 // A and B are quaternions.  The result is A + (|A.B| >= 0 ? 1 : -1) * B
 FORCEINLINE VectorRegister4Float VectorAccumulateQuaternionShortestPath(const VectorRegister4Float& A, const VectorRegister4Float& B)
@@ -720,34 +732,6 @@ FORCEINLINE VectorRegister4Double VectorQuaternionInverseRotateVector(const Vect
 	return VectorQuaternionRotateVector(QInv, VectorW0);
 }
 
-
-/**
-* Rotate a vector using a pointer to a unit Quaternion.
-*
-* @param Result		Pointer to where the result should be stored
-* @param Quat		Pointer to the unit quaternion (must not be the destination)
-* @param VectorW0	Pointer to the vector (must not be the destination). W component must be zero.
-*/
-FORCEINLINE void VectorQuaternionRotateVectorPtr( void* RESTRICT Result, const void* RESTRICT Quat, const void* RESTRICT VectorW0)
-{
-	*((VectorRegister *)Result) = VectorQuaternionRotateVector(*((const VectorRegister *)Quat), *((const VectorRegister *)VectorW0));
-}
-// TODO: LWC: implement for template type FQuat
-
-/**
-* Rotate a vector using the inverse of a unit Quaternion (rotation in the opposite direction).
-*
-* @param Result		Pointer to where the result should be stored
-* @param Quat		Pointer to the unit quaternion (must not be the destination)
-* @param VectorW0	Pointer to the vector (must not be the destination). W component must be zero.
-*/
-FORCEINLINE void VectorQuaternionInverseRotateVectorPtr(void* RESTRICT Result, const void* RESTRICT Quat, const void* RESTRICT VectorW0)
-{
-	*((VectorRegister *)Result) = VectorQuaternionInverseRotateVector(*((const VectorRegister *)Quat), *((const VectorRegister *)VectorW0));
-}
-// TODO: LWC: implement for template type FQuat
-
-
 FORCEINLINE VectorRegister4Double VectorMultiply(VectorRegister4Double Vec1, VectorRegister4Float Vec2)
 {
 	return VectorMultiply(Vec1, VectorRegister4Double(Vec2));
@@ -763,21 +747,6 @@ FORCEINLINE VectorRegister4Double VectorMultiplyAdd(VectorRegister4Double Vec1,
 	return VectorMultiplyAdd(Vec1, VectorRegister4Double(Vec2), VectorRegister4Double(Acc));
 }
 
-FORCEINLINE VectorRegister4Double operator+(const VectorRegister4Double& Vec1, const VectorRegister4Double& Vec2)
-{
-	return VectorAdd(Vec1, Vec2);
-}
-
-FORCEINLINE VectorRegister4Double operator-(const VectorRegister4Double& Vec1, const VectorRegister4Double& Vec2)
-{
-	return VectorSubtract(Vec1, Vec2);
-}
-
-FORCEINLINE VectorRegister4Double operator*(const VectorRegister4Double& Vec1, const VectorRegister4Double& Vec2)
-{
-	return VectorMultiply(Vec1, Vec2);
-}
-
 FORCEINLINE VectorRegister4Double VectorSetDouble(double X, double Y, double Z, double W)
 {
 	return MakeVectorRegisterDouble(X, Y, Z, W);
@@ -874,9 +843,23 @@ FORCEINLINE VectorRegister4Double VectorFractional(const VectorRegister4Double&
 }
 
 
-FORCEINLINE void VectorQuaternionMultiply(FQuat* RESTRICT Result, const FQuat* RESTRICT Quat1, const FQuat* RESTRICT Quat2)
+FORCEINLINE void VectorQuaternionMultiply(FQuat4f* RESTRICT Result, const FQuat4f* RESTRICT Quat1, const FQuat4f* RESTRICT Quat2)
+{
+	//*((VectorRegister4Float*)Result) = VectorQuaternionMultiply2(*((const VectorRegister4Float*)Quat1), *((const VectorRegister4Float*)Quat2));
+	const VectorRegister4Float Q1 = VectorLoadAligned(Quat1);
+	const VectorRegister4Float Q2 = VectorLoadAligned(Quat2);
+	const VectorRegister4Float QResult = VectorQuaternionMultiply2(Q1, Q2);
+	VectorStoreAligned(QResult, Result);
+}
+
+FORCEINLINE void VectorQuaternionMultiply(FQuat4d* RESTRICT Result, const FQuat4d* RESTRICT Quat1, const FQuat4d* RESTRICT Quat2)
 {
-	*((VectorRegister4Float*)Result) = VectorQuaternionMultiply2(*((const VectorRegister4Float*)Quat1), *((const VectorRegister4Float*)Quat2));
+	// Warning: don't try to cast FQuat4d to VectorRegister4Double, they may not be of the same alignment.
+	// We can still use VectorLoadAligned/VectorStoreAligned, because those only enforce 16-byte alignment and fall back to unaligned otherwise.
+	const VectorRegister4Double Q1 = VectorLoadAligned(Quat1);
+	const VectorRegister4Double Q2 = VectorLoadAligned(Quat2);
+	const VectorRegister4Double QResult = VectorQuaternionMultiply2(Q1, Q2);
+	VectorStoreAligned(QResult, Result);
 }
 
 /**
diff --git a/Engine/Source/Runtime/Core/Public/Math/UnrealMathVectorConstants.h b/Engine/Source/Runtime/Core/Public/Math/UnrealMathVectorConstants.h
--- a/Engine/Source/Runtime/Core/Public/Math/UnrealMathVectorConstants.h
+++ b/Engine/Source/Runtime/Core/Public/Math/UnrealMathVectorConstants.h
@@ -5,125 +5,125 @@
 
 namespace GlobalVectorConstants
 {
-	static const VectorRegister4Float FloatOne = MakeVectorRegisterFloat(1.0f, 1.0f, 1.0f, 1.0f);
-	static const VectorRegister4Float FloatZero = MakeVectorRegisterFloat(0.0f, 0.0f, 0.0f, 0.0f);
-	static const VectorRegister4Float FloatMinusOne = MakeVectorRegisterFloat(-1.0f, -1.0f, -1.0f, -1.0f);
-	static const VectorRegister4Float Float0001 = MakeVectorRegisterFloat( 0.0f, 0.0f, 0.0f, 1.0f );
-	static const VectorRegister4Float Float1000 = MakeVectorRegisterFloat( 1.0f, 0.0f, 0.0f, 0.0f );
-	static const VectorRegister4Float Float1110 = MakeVectorRegisterFloat(1.0f, 1.0f, 1.0f, 0.0f);
-	static const VectorRegister4Float SmallLengthThreshold = MakeVectorRegisterFloat(1.e-8f, 1.e-8f, 1.e-8f, 1.e-8f);
-	static const VectorRegister4Float FloatOneHundredth = MakeVectorRegisterFloat(0.01f, 0.01f, 0.01f, 0.01f);
-	static const VectorRegister4Float Float111_Minus1 = MakeVectorRegisterFloat( 1.f, 1.f, 1.f, -1.f );
-	static const VectorRegister4Float FloatMinus1_111= MakeVectorRegisterFloat( -1.f, 1.f, 1.f, 1.f );
-	static const VectorRegister4Float FloatOneHalf = MakeVectorRegisterFloat( 0.5f, 0.5f, 0.5f, 0.5f );
-	static const VectorRegister4Float FloatMinusOneHalf = MakeVectorRegisterFloat( -0.5f, -0.5f, -0.5f, -0.5f );
-	static const VectorRegister4Float KindaSmallNumber = MakeVectorRegisterFloat( KINDA_SMALL_NUMBER, KINDA_SMALL_NUMBER, KINDA_SMALL_NUMBER, KINDA_SMALL_NUMBER );
-	static const VectorRegister4Float SmallNumber = MakeVectorRegisterFloat( SMALL_NUMBER, SMALL_NUMBER, SMALL_NUMBER, SMALL_NUMBER );
-	static const VectorRegister4Float ThreshQuatNormalized = MakeVectorRegisterFloat( THRESH_QUAT_NORMALIZED, THRESH_QUAT_NORMALIZED, THRESH_QUAT_NORMALIZED, THRESH_QUAT_NORMALIZED );
-	static const VectorRegister4Float BigNumber = MakeVectorRegisterFloat(BIG_NUMBER, BIG_NUMBER, BIG_NUMBER, BIG_NUMBER);
-
-	static const VectorRegister2Double DoubleOne2d = MakeVectorRegister2Double(1.0, 1.0);
-	static const VectorRegister4Double DoubleOne = MakeVectorRegisterDouble(1.0, 1.0, 1.0, 1.0);
-	static const VectorRegister4Double DoubleZero = MakeVectorRegisterDouble(0.0, 0.0, 0.0, 0.0);
-	static const VectorRegister4Double DoubleMinusOne = MakeVectorRegisterDouble(-1.0, -1.0, -1.0, -1.0);
-	static const VectorRegister4Double Double0001 = MakeVectorRegisterDouble(0.0f, 0.0, 0.0, 1.0);
-	static const VectorRegister4Double Double1000 = MakeVectorRegisterDouble(1.0, 0.0, 0.0, 0.0);
-	static const VectorRegister4Double Double1110 = MakeVectorRegisterDouble(1.0, 1.0, 1.0, 0.0);
-	static const VectorRegister4Double DoubleSmallLengthThreshold = MakeVectorRegisterDouble(1.e-8, 1.e-8, 1.e-8, 1.e-8);
-	static const VectorRegister4Double DoubleOneHundredth = MakeVectorRegisterDouble(0.01, 0.01, 0.01, 0.01);
-	static const VectorRegister4Double Double111_Minus1 = MakeVectorRegisterDouble(1., 1., 1., -1.);
-	static const VectorRegister4Double DoubleMinus1_111 = MakeVectorRegisterDouble(-1., 1., 1., 1.);
-	static const VectorRegister4Double DoubleOneHalf = MakeVectorRegisterDouble(0.5, 0.5, 0.5, 0.5);
-	static const VectorRegister4Double DoubleMinusOneHalf = MakeVectorRegisterDouble(-0.5, -0.5, -0.5, -0.5);
-	static const VectorRegister4Double DoubleKindaSmallNumber = MakeVectorRegisterDouble(DOUBLE_KINDA_SMALL_NUMBER, DOUBLE_KINDA_SMALL_NUMBER, DOUBLE_KINDA_SMALL_NUMBER, DOUBLE_KINDA_SMALL_NUMBER);
-	static const VectorRegister4Double DoubleSmallNumber = MakeVectorRegisterDouble(DOUBLE_SMALL_NUMBER, DOUBLE_SMALL_NUMBER, DOUBLE_SMALL_NUMBER, DOUBLE_SMALL_NUMBER);
-	static const VectorRegister4Double DoubleThreshQuatNormalized = MakeVectorRegisterDouble(DOUBLE_THRESH_QUAT_NORMALIZED, DOUBLE_THRESH_QUAT_NORMALIZED, DOUBLE_THRESH_QUAT_NORMALIZED, DOUBLE_THRESH_QUAT_NORMALIZED);
-	static const VectorRegister4Double DoubleBigNumber = MakeVectorRegisterDouble(DOUBLE_BIG_NUMBER, DOUBLE_BIG_NUMBER, DOUBLE_BIG_NUMBER, DOUBLE_BIG_NUMBER);
-
-	static const VectorRegister Vector0001 = Float0001;
-	static const VectorRegister Vector1110 = Float1110;
-
-	static const VectorRegister4Int IntOne = MakeVectorRegisterInt(1, 1, 1, 1);
-	static const VectorRegister4Int IntZero = MakeVectorRegisterInt(0, 0, 0, 0);
-	static const VectorRegister4Int IntMinusOne = MakeVectorRegisterInt(-1, -1, -1, -1);
+	inline static const VectorRegister4Float FloatOne = MakeVectorRegisterFloat(1.0f, 1.0f, 1.0f, 1.0f);
+	inline static const VectorRegister4Float FloatZero = MakeVectorRegisterFloat(0.0f, 0.0f, 0.0f, 0.0f);
+	inline static const VectorRegister4Float FloatMinusOne = MakeVectorRegisterFloat(-1.0f, -1.0f, -1.0f, -1.0f);
+	inline static const VectorRegister4Float Float0001 = MakeVectorRegisterFloat( 0.0f, 0.0f, 0.0f, 1.0f );
+	inline static const VectorRegister4Float Float1000 = MakeVectorRegisterFloat( 1.0f, 0.0f, 0.0f, 0.0f );
+	inline static const VectorRegister4Float Float1110 = MakeVectorRegisterFloat(1.0f, 1.0f, 1.0f, 0.0f);
+	inline static const VectorRegister4Float SmallLengthThreshold = MakeVectorRegisterFloat(1.e-8f, 1.e-8f, 1.e-8f, 1.e-8f);
+	inline static const VectorRegister4Float FloatOneHundredth = MakeVectorRegisterFloat(0.01f, 0.01f, 0.01f, 0.01f);
+	inline static const VectorRegister4Float Float111_Minus1 = MakeVectorRegisterFloat( 1.f, 1.f, 1.f, -1.f );
+	inline static const VectorRegister4Float FloatMinus1_111= MakeVectorRegisterFloat( -1.f, 1.f, 1.f, 1.f );
+	inline static const VectorRegister4Float FloatOneHalf = MakeVectorRegisterFloat( 0.5f, 0.5f, 0.5f, 0.5f );
+	inline static const VectorRegister4Float FloatMinusOneHalf = MakeVectorRegisterFloat( -0.5f, -0.5f, -0.5f, -0.5f );
+	inline static const VectorRegister4Float KindaSmallNumber = MakeVectorRegisterFloat( KINDA_SMALL_NUMBER, KINDA_SMALL_NUMBER, KINDA_SMALL_NUMBER, KINDA_SMALL_NUMBER );
+	inline static const VectorRegister4Float SmallNumber = MakeVectorRegisterFloat( SMALL_NUMBER, SMALL_NUMBER, SMALL_NUMBER, SMALL_NUMBER );
+	inline static const VectorRegister4Float ThreshQuatNormalized = MakeVectorRegisterFloat( THRESH_QUAT_NORMALIZED, THRESH_QUAT_NORMALIZED, THRESH_QUAT_NORMALIZED, THRESH_QUAT_NORMALIZED );
+	inline static const VectorRegister4Float BigNumber = MakeVectorRegisterFloat(BIG_NUMBER, BIG_NUMBER, BIG_NUMBER, BIG_NUMBER);
+
+	inline static const VectorRegister2Double DoubleOne2d = MakeVectorRegister2Double(1.0, 1.0);
+	inline static const VectorRegister4Double DoubleOne = MakeVectorRegisterDouble(1.0, 1.0, 1.0, 1.0);
+	inline static const VectorRegister4Double DoubleZero = MakeVectorRegisterDouble(0.0, 0.0, 0.0, 0.0);
+	inline static const VectorRegister4Double DoubleMinusOne = MakeVectorRegisterDouble(-1.0, -1.0, -1.0, -1.0);
+	inline static const VectorRegister4Double Double0001 = MakeVectorRegisterDouble(0.0f, 0.0, 0.0, 1.0);
+	inline static const VectorRegister4Double Double1000 = MakeVectorRegisterDouble(1.0, 0.0, 0.0, 0.0);
+	inline static const VectorRegister4Double Double1110 = MakeVectorRegisterDouble(1.0, 1.0, 1.0, 0.0);
+	inline static const VectorRegister4Double DoubleSmallLengthThreshold = MakeVectorRegisterDouble(1.e-8, 1.e-8, 1.e-8, 1.e-8);
+	inline static const VectorRegister4Double DoubleOneHundredth = MakeVectorRegisterDouble(0.01, 0.01, 0.01, 0.01);
+	inline static const VectorRegister4Double Double111_Minus1 = MakeVectorRegisterDouble(1., 1., 1., -1.);
+	inline static const VectorRegister4Double DoubleMinus1_111 = MakeVectorRegisterDouble(-1., 1., 1., 1.);
+	inline static const VectorRegister4Double DoubleOneHalf = MakeVectorRegisterDouble(0.5, 0.5, 0.5, 0.5);
+	inline static const VectorRegister4Double DoubleMinusOneHalf = MakeVectorRegisterDouble(-0.5, -0.5, -0.5, -0.5);
+	inline static const VectorRegister4Double DoubleKindaSmallNumber = MakeVectorRegisterDouble(DOUBLE_KINDA_SMALL_NUMBER, DOUBLE_KINDA_SMALL_NUMBER, DOUBLE_KINDA_SMALL_NUMBER, DOUBLE_KINDA_SMALL_NUMBER);
+	inline static const VectorRegister4Double DoubleSmallNumber = MakeVectorRegisterDouble(DOUBLE_SMALL_NUMBER, DOUBLE_SMALL_NUMBER, DOUBLE_SMALL_NUMBER, DOUBLE_SMALL_NUMBER);
+	inline static const VectorRegister4Double DoubleThreshQuatNormalized = MakeVectorRegisterDouble(DOUBLE_THRESH_QUAT_NORMALIZED, DOUBLE_THRESH_QUAT_NORMALIZED, DOUBLE_THRESH_QUAT_NORMALIZED, DOUBLE_THRESH_QUAT_NORMALIZED);
+	inline static const VectorRegister4Double DoubleBigNumber = MakeVectorRegisterDouble(DOUBLE_BIG_NUMBER, DOUBLE_BIG_NUMBER, DOUBLE_BIG_NUMBER, DOUBLE_BIG_NUMBER);
+
+	inline static const VectorRegister Vector0001 = Float0001;
+	inline static const VectorRegister Vector1110 = Float1110;
+
+	inline static const VectorRegister4Int IntOne = MakeVectorRegisterInt(1, 1, 1, 1);
+	inline static const VectorRegister4Int IntZero = MakeVectorRegisterInt(0, 0, 0, 0);
+	inline static const VectorRegister4Int IntMinusOne = MakeVectorRegisterInt(-1, -1, -1, -1);
 
 	/** This is to speed up Quaternion Inverse. Static variable to keep sign of inverse **/
-	static const VectorRegister4Float QINV_SIGN_MASK = MakeVectorRegisterFloat( -1.f, -1.f, -1.f, 1.f );
-	static const VectorRegister4Double DOUBLE_QINV_SIGN_MASK = MakeVectorRegisterDouble(-1., -1., -1., 1.);
-
-	static const VectorRegister4Float QMULTI_SIGN_MASK0 = MakeVectorRegisterFloat( 1.f, -1.f, 1.f, -1.f );
-	static const VectorRegister4Float QMULTI_SIGN_MASK1 = MakeVectorRegisterFloat( 1.f, 1.f, -1.f, -1.f );
-	static const VectorRegister4Float QMULTI_SIGN_MASK2 = MakeVectorRegisterFloat( -1.f, 1.f, 1.f, -1.f );
-	static const VectorRegister4Double DOUBLE_QMULTI_SIGN_MASK0 = MakeVectorRegisterDouble(1., -1., 1., -1.);
-	static const VectorRegister4Double DOUBLE_QMULTI_SIGN_MASK1 = MakeVectorRegisterDouble(1., 1., -1., -1.);
-	static const VectorRegister4Double DOUBLE_QMULTI_SIGN_MASK2 = MakeVectorRegisterDouble(-1., 1., 1., -1.);
-
-	static const VectorRegister4Float DEG_TO_RAD = MakeVectorRegister(PI/(180.f), PI/(180.f), PI/(180.f), PI/(180.f));
-	static const VectorRegister4Float DEG_TO_RAD_HALF = MakeVectorRegister((PI/180.f)*0.5f, (PI/180.f)*0.5f, (PI/180.f)*0.5f, (PI/180.f)*0.5f);
-	static const VectorRegister4Float RAD_TO_DEG = MakeVectorRegister((180.f)/PI, (180.f)/PI, (180.f)/PI, (180.f)/PI);
-	static const VectorRegister4Double DOUBLE_DEG_TO_RAD = MakeVectorRegister(DOUBLE_PI/(180.), DOUBLE_PI/(180.), DOUBLE_PI/(180.), DOUBLE_PI/(180.));
-	static const VectorRegister4Double DOUBLE_DEG_TO_RAD_HALF = MakeVectorRegister((DOUBLE_PI/180.) * 0.5, (DOUBLE_PI/180.) * 0.5, (DOUBLE_PI/180.) * 0.5, (DOUBLE_PI/180.) * 0.5);
-	static const VectorRegister4Double DOUBLE_RAD_TO_DEG = MakeVectorRegister((180.)/DOUBLE_PI, (180.)/DOUBLE_PI, (180.)/DOUBLE_PI, (180.)/DOUBLE_PI);
+	inline static const VectorRegister4Float QINV_SIGN_MASK = MakeVectorRegisterFloat( -1.f, -1.f, -1.f, 1.f );
+	inline static const VectorRegister4Double DOUBLE_QINV_SIGN_MASK = MakeVectorRegisterDouble(-1., -1., -1., 1.);
+
+	inline static const VectorRegister4Float QMULTI_SIGN_MASK0 = MakeVectorRegisterFloat( 1.f, -1.f, 1.f, -1.f );
+	inline static const VectorRegister4Float QMULTI_SIGN_MASK1 = MakeVectorRegisterFloat( 1.f, 1.f, -1.f, -1.f );
+	inline static const VectorRegister4Float QMULTI_SIGN_MASK2 = MakeVectorRegisterFloat( -1.f, 1.f, 1.f, -1.f );
+	inline static const VectorRegister4Double DOUBLE_QMULTI_SIGN_MASK0 = MakeVectorRegisterDouble(1., -1., 1., -1.);
+	inline static const VectorRegister4Double DOUBLE_QMULTI_SIGN_MASK1 = MakeVectorRegisterDouble(1., 1., -1., -1.);
+	inline static const VectorRegister4Double DOUBLE_QMULTI_SIGN_MASK2 = MakeVectorRegisterDouble(-1., 1., 1., -1.);
+
+	inline static const VectorRegister4Float DEG_TO_RAD = MakeVectorRegister(PI/(180.f), PI/(180.f), PI/(180.f), PI/(180.f));
+	inline static const VectorRegister4Float DEG_TO_RAD_HALF = MakeVectorRegister((PI/180.f)*0.5f, (PI/180.f)*0.5f, (PI/180.f)*0.5f, (PI/180.f)*0.5f);
+	inline static const VectorRegister4Float RAD_TO_DEG = MakeVectorRegister((180.f)/PI, (180.f)/PI, (180.f)/PI, (180.f)/PI);
+	inline static const VectorRegister4Double DOUBLE_DEG_TO_RAD = MakeVectorRegister(DOUBLE_PI/(180.), DOUBLE_PI/(180.), DOUBLE_PI/(180.), DOUBLE_PI/(180.));
+	inline static const VectorRegister4Double DOUBLE_DEG_TO_RAD_HALF = MakeVectorRegister((DOUBLE_PI/180.) * 0.5, (DOUBLE_PI/180.) * 0.5, (DOUBLE_PI/180.) * 0.5, (DOUBLE_PI/180.) * 0.5);
+	inline static const VectorRegister4Double DOUBLE_RAD_TO_DEG = MakeVectorRegister((180.)/DOUBLE_PI, (180.)/DOUBLE_PI, (180.)/DOUBLE_PI, (180.)/DOUBLE_PI);
 
 	/** Bitmask to AND out the XYZ components in a vector */
-	static const VectorRegister4Float XYZMask = MakeVectorRegisterFloatMask((uint32)0xffffffff, (uint32)0xffffffff, (uint32)0xffffffff, (uint32)0x00000000);
-	static const VectorRegister4Double DoubleXYZMask = MakeVectorRegisterDoubleMask((uint64)0xFFFFFFFFFFFFFFFF, (uint64)0xFFFFFFFFFFFFFFFF, (uint64)0xFFFFFFFFFFFFFFFF, (uint64)0);
+	inline static const VectorRegister4Float XYZMask = MakeVectorRegisterFloatMask((uint32)0xffffffff, (uint32)0xffffffff, (uint32)0xffffffff, (uint32)0x00000000);
+	inline static const VectorRegister4Double DoubleXYZMask = MakeVectorRegisterDoubleMask((uint64)0xFFFFFFFFFFFFFFFF, (uint64)0xFFFFFFFFFFFFFFFF, (uint64)0xFFFFFFFFFFFFFFFF, (uint64)0);
 	
 	/** Bitmask to AND out the sign bit of each components in a vector */
 #define SIGN_BIT ((1 << 31))
-	static const VectorRegister4Float SignBit = MakeVectorRegisterFloatMask((uint32)SIGN_BIT, (uint32)SIGN_BIT, (uint32)SIGN_BIT, (uint32)SIGN_BIT);
-	static const VectorRegister4Float SignMask = MakeVectorRegisterFloatMask((uint32)(~SIGN_BIT), (uint32)(~SIGN_BIT), (uint32)(~SIGN_BIT), (uint32)(~SIGN_BIT));
-	static const VectorRegister4Int IntSignBit = MakeVectorRegisterInt(SIGN_BIT, SIGN_BIT, SIGN_BIT, SIGN_BIT);
-	static const VectorRegister4Int IntSignMask = MakeVectorRegisterInt((~SIGN_BIT), (~SIGN_BIT), (~SIGN_BIT), (~SIGN_BIT));
+	inline static const VectorRegister4Float SignBit = MakeVectorRegisterFloatMask((uint32)SIGN_BIT, (uint32)SIGN_BIT, (uint32)SIGN_BIT, (uint32)SIGN_BIT);
+	inline static const VectorRegister4Float SignMask = MakeVectorRegisterFloatMask((uint32)(~SIGN_BIT), (uint32)(~SIGN_BIT), (uint32)(~SIGN_BIT), (uint32)(~SIGN_BIT));
+	inline static const VectorRegister4Int IntSignBit = MakeVectorRegisterInt(SIGN_BIT, SIGN_BIT, SIGN_BIT, SIGN_BIT);
+	inline static const VectorRegister4Int IntSignMask = MakeVectorRegisterInt((~SIGN_BIT), (~SIGN_BIT), (~SIGN_BIT), (~SIGN_BIT));
 #undef SIGN_BIT
-	static const VectorRegister4Float AllMask = MakeVectorRegisterFloatMask(0xFFFFFFFF, 0xFFFFFFFF, 0xFFFFFFFF, 0xFFFFFFFF);
-	static const VectorRegister4Double DoubleAllMask = MakeVectorRegisterDoubleMask((uint64)(0xFFFFFFFFFFFFFFFF), (uint64)(0xFFFFFFFFFFFFFFFF), (uint64)(0xFFFFFFFFFFFFFFFF), (uint64)(0xFFFFFFFFFFFFFFFF));
-	static const VectorRegister4Int IntAllMask = MakeVectorRegisterInt(0xFFFFFFFF, 0xFFFFFFFF, 0xFFFFFFFF, 0xFFFFFFFF);
+	inline static const VectorRegister4Float AllMask = MakeVectorRegisterFloatMask(0xFFFFFFFF, 0xFFFFFFFF, 0xFFFFFFFF, 0xFFFFFFFF);
+	inline static const VectorRegister4Double DoubleAllMask = MakeVectorRegisterDoubleMask((uint64)(0xFFFFFFFFFFFFFFFF), (uint64)(0xFFFFFFFFFFFFFFFF), (uint64)(0xFFFFFFFFFFFFFFFF), (uint64)(0xFFFFFFFFFFFFFFFF));
+	inline static const VectorRegister4Int IntAllMask = MakeVectorRegisterInt(0xFFFFFFFF, 0xFFFFFFFF, 0xFFFFFFFF, 0xFFFFFFFF);
 
 #define DOUBLE_SIGN_BIT (uint64(1) << uint64(63))
-	static const VectorRegister4Double DoubleSignBit = MakeVectorRegisterDoubleMask((uint64)DOUBLE_SIGN_BIT, (uint64)DOUBLE_SIGN_BIT, (uint64)DOUBLE_SIGN_BIT, (uint64)DOUBLE_SIGN_BIT);
-	static const VectorRegister4Double DoubleSignMask = MakeVectorRegisterDoubleMask((uint64)(~DOUBLE_SIGN_BIT), (uint64)(~DOUBLE_SIGN_BIT), (uint64)(~DOUBLE_SIGN_BIT), (uint64)(~DOUBLE_SIGN_BIT));
+	inline static const VectorRegister4Double DoubleSignBit = MakeVectorRegisterDoubleMask((uint64)DOUBLE_SIGN_BIT, (uint64)DOUBLE_SIGN_BIT, (uint64)DOUBLE_SIGN_BIT, (uint64)DOUBLE_SIGN_BIT);
+	inline static const VectorRegister4Double DoubleSignMask = MakeVectorRegisterDoubleMask((uint64)(~DOUBLE_SIGN_BIT), (uint64)(~DOUBLE_SIGN_BIT), (uint64)(~DOUBLE_SIGN_BIT), (uint64)(~DOUBLE_SIGN_BIT));
 #undef DOUBLE_SIGN_BIT
 
 	/** Vector full of positive infinity */
-	static const VectorRegister4Float FloatInfinity = MakeVectorRegisterFloatMask((uint32)0x7F800000, (uint32)0x7F800000, (uint32)0x7F800000, (uint32)0x7F800000);
-	static const VectorRegister4Double DoubleInfinity = MakeVectorRegisterDoubleMask((uint64)0x7FF0000000000000, (uint64)0x7FF0000000000000, (uint64)0x7FF0000000000000, (uint64)0x7FF0000000000000);
-
-	static const VectorRegister4Float Pi = MakeVectorRegister(PI, PI, PI, PI);
-	static const VectorRegister4Float TwoPi = MakeVectorRegister(2.0f*PI, 2.0f*PI, 2.0f*PI, 2.0f*PI);
-	static const VectorRegister4Float PiByTwo = MakeVectorRegister(0.5f*PI, 0.5f*PI, 0.5f*PI, 0.5f*PI);
-	static const VectorRegister4Float PiByFour = MakeVectorRegister(0.25f*PI, 0.25f*PI, 0.25f*PI, 0.25f*PI);
-	static const VectorRegister4Float OneOverPi = MakeVectorRegister(1.0f / PI, 1.0f / PI, 1.0f / PI, 1.0f / PI);
-	static const VectorRegister4Float OneOverTwoPi = MakeVectorRegister(1.0f / (2.0f*PI), 1.0f / (2.0f*PI), 1.0f / (2.0f*PI), 1.0f / (2.0f*PI));
-
-	static const VectorRegister4Double DoublePi = MakeVectorRegisterDouble(DOUBLE_PI, DOUBLE_PI, DOUBLE_PI, DOUBLE_PI);
-	static const VectorRegister4Double DoubleTwoPi = MakeVectorRegisterDouble(2.0 * DOUBLE_PI, 2.0 * DOUBLE_PI, 2.0 * DOUBLE_PI, 2.0 * DOUBLE_PI);
-	static const VectorRegister4Double DoublePiByTwo = MakeVectorRegisterDouble(0.5 * DOUBLE_PI, 0.5 * DOUBLE_PI, 0.5 * DOUBLE_PI, 0.5 * DOUBLE_PI);
-	static const VectorRegister4Double DoublePiByFour = MakeVectorRegisterDouble(0.25 * DOUBLE_PI, 0.25 * DOUBLE_PI, 0.25 * DOUBLE_PI, 0.25 * DOUBLE_PI);
-	static const VectorRegister4Double DoubleOneOverPi = MakeVectorRegisterDouble(1.0 / DOUBLE_PI, 1.0 / DOUBLE_PI, 1.0 / DOUBLE_PI, 1.0 / DOUBLE_PI);
-	static const VectorRegister4Double DoubleOneOverTwoPi = MakeVectorRegisterDouble(1.0 / (2.0 * DOUBLE_PI), 1.0 / (2.0 * DOUBLE_PI), 1.0 / (2.0 * DOUBLE_PI), 1.0 / (2.0 * DOUBLE_PI));
-
-	static const VectorRegister4Float Float255 = MakeVectorRegister(255.0f, 255.0f, 255.0f, 255.0f);
-	static const VectorRegister4Float Float127 = MakeVectorRegister(127.0f, 127.0f, 127.0f, 127.0f);
-	static const VectorRegister4Float FloatNeg127 = MakeVectorRegister(-127.0f, -127.0f, -127.0f, -127.0f);
-	static const VectorRegister4Float Float360 = MakeVectorRegister(360.f, 360.f, 360.f, 360.f);
-	static const VectorRegister4Float Float180 = MakeVectorRegister(180.f, 180.f, 180.f, 180.f);
-
-	static const VectorRegister4Double Double255 = MakeVectorRegisterDouble(255.0, 255.0, 255.0, 255.0);
-	static const VectorRegister4Double Double127 = MakeVectorRegisterDouble(127.0, 127.0, 127.0, 127.0);
-	static const VectorRegister4Double DoubleNeg127 = MakeVectorRegisterDouble(-127.0, -127.0, -127.0, -127.0);
-	static const VectorRegister4Double Double360 = MakeVectorRegisterDouble(360., 360., 360., 360.);
-	static const VectorRegister4Double Double180 = MakeVectorRegisterDouble(180., 180., 180., 180.);
+	inline static const VectorRegister4Float FloatInfinity = MakeVectorRegisterFloatMask((uint32)0x7F800000, (uint32)0x7F800000, (uint32)0x7F800000, (uint32)0x7F800000);
+	inline static const VectorRegister4Double DoubleInfinity = MakeVectorRegisterDoubleMask((uint64)0x7FF0000000000000, (uint64)0x7FF0000000000000, (uint64)0x7FF0000000000000, (uint64)0x7FF0000000000000);
+
+	inline static const VectorRegister4Float Pi = MakeVectorRegister(PI, PI, PI, PI);
+	inline static const VectorRegister4Float TwoPi = MakeVectorRegister(2.0f*PI, 2.0f*PI, 2.0f*PI, 2.0f*PI);
+	inline static const VectorRegister4Float PiByTwo = MakeVectorRegister(0.5f*PI, 0.5f*PI, 0.5f*PI, 0.5f*PI);
+	inline static const VectorRegister4Float PiByFour = MakeVectorRegister(0.25f*PI, 0.25f*PI, 0.25f*PI, 0.25f*PI);
+	inline static const VectorRegister4Float OneOverPi = MakeVectorRegister(1.0f / PI, 1.0f / PI, 1.0f / PI, 1.0f / PI);
+	inline static const VectorRegister4Float OneOverTwoPi = MakeVectorRegister(1.0f / (2.0f*PI), 1.0f / (2.0f*PI), 1.0f / (2.0f*PI), 1.0f / (2.0f*PI));
+
+	inline static const VectorRegister4Double DoublePi = MakeVectorRegisterDouble(DOUBLE_PI, DOUBLE_PI, DOUBLE_PI, DOUBLE_PI);
+	inline static const VectorRegister4Double DoubleTwoPi = MakeVectorRegisterDouble(2.0 * DOUBLE_PI, 2.0 * DOUBLE_PI, 2.0 * DOUBLE_PI, 2.0 * DOUBLE_PI);
+	inline static const VectorRegister4Double DoublePiByTwo = MakeVectorRegisterDouble(0.5 * DOUBLE_PI, 0.5 * DOUBLE_PI, 0.5 * DOUBLE_PI, 0.5 * DOUBLE_PI);
+	inline static const VectorRegister4Double DoublePiByFour = MakeVectorRegisterDouble(0.25 * DOUBLE_PI, 0.25 * DOUBLE_PI, 0.25 * DOUBLE_PI, 0.25 * DOUBLE_PI);
+	inline static const VectorRegister4Double DoubleOneOverPi = MakeVectorRegisterDouble(1.0 / DOUBLE_PI, 1.0 / DOUBLE_PI, 1.0 / DOUBLE_PI, 1.0 / DOUBLE_PI);
+	inline static const VectorRegister4Double DoubleOneOverTwoPi = MakeVectorRegisterDouble(1.0 / (2.0 * DOUBLE_PI), 1.0 / (2.0 * DOUBLE_PI), 1.0 / (2.0 * DOUBLE_PI), 1.0 / (2.0 * DOUBLE_PI));
+
+	inline static const VectorRegister4Float Float255 = MakeVectorRegister(255.0f, 255.0f, 255.0f, 255.0f);
+	inline static const VectorRegister4Float Float127 = MakeVectorRegister(127.0f, 127.0f, 127.0f, 127.0f);
+	inline static const VectorRegister4Float FloatNeg127 = MakeVectorRegister(-127.0f, -127.0f, -127.0f, -127.0f);
+	inline static const VectorRegister4Float Float360 = MakeVectorRegister(360.f, 360.f, 360.f, 360.f);
+	inline static const VectorRegister4Float Float180 = MakeVectorRegister(180.f, 180.f, 180.f, 180.f);
+
+	inline static const VectorRegister4Double Double255 = MakeVectorRegisterDouble(255.0, 255.0, 255.0, 255.0);
+	inline static const VectorRegister4Double Double127 = MakeVectorRegisterDouble(127.0, 127.0, 127.0, 127.0);
+	inline static const VectorRegister4Double DoubleNeg127 = MakeVectorRegisterDouble(-127.0, -127.0, -127.0, -127.0);
+	inline static const VectorRegister4Double Double360 = MakeVectorRegisterDouble(360., 360., 360., 360.);
+	inline static const VectorRegister4Double Double180 = MakeVectorRegisterDouble(180., 180., 180., 180.);
 
 	// All float numbers greater than or equal to this have no fractional value.
-	static const VectorRegister4Float FloatNonFractional = MakeVectorRegister(FLOAT_NON_FRACTIONAL, FLOAT_NON_FRACTIONAL, FLOAT_NON_FRACTIONAL, FLOAT_NON_FRACTIONAL);
-	static const VectorRegister4Double DoubleNonFractional = MakeVectorRegisterDouble(DOUBLE_NON_FRACTIONAL, DOUBLE_NON_FRACTIONAL, DOUBLE_NON_FRACTIONAL, DOUBLE_NON_FRACTIONAL);
+	inline static const VectorRegister4Float FloatNonFractional = MakeVectorRegister(FLOAT_NON_FRACTIONAL, FLOAT_NON_FRACTIONAL, FLOAT_NON_FRACTIONAL, FLOAT_NON_FRACTIONAL);
+	inline static const VectorRegister4Double DoubleNonFractional = MakeVectorRegisterDouble(DOUBLE_NON_FRACTIONAL, DOUBLE_NON_FRACTIONAL, DOUBLE_NON_FRACTIONAL, DOUBLE_NON_FRACTIONAL);
 
-	static const VectorRegister4Float FloatTwo = MakeVectorRegister(2.0f, 2.0f, 2.0f, 2.0f);
-	static const uint32 AlmostTwoBits = 0x3fffffff;
-	static const VectorRegister4Float FloatAlmostTwo = MakeVectorRegisterFloatMask(AlmostTwoBits, AlmostTwoBits, AlmostTwoBits, AlmostTwoBits);
+	inline static const VectorRegister4Float FloatTwo = MakeVectorRegister(2.0f, 2.0f, 2.0f, 2.0f);
+	inline static const uint32 AlmostTwoBits = 0x3fffffff;
+	inline static const VectorRegister4Float FloatAlmostTwo = MakeVectorRegisterFloatMask(AlmostTwoBits, AlmostTwoBits, AlmostTwoBits, AlmostTwoBits);
 
-	static const VectorRegister4Double DoubleTwo = MakeVectorRegisterDouble(2.0, 2.0, 2.0, 2.0);
-	static const uint64 DoubleAlmostTwoBits = 0x3FFFFFFFFFFFFFFF;
-	static const VectorRegister4Double DoubleAlmostTwo = MakeVectorRegisterDoubleMask(DoubleAlmostTwoBits, DoubleAlmostTwoBits, DoubleAlmostTwoBits, DoubleAlmostTwoBits);
+	inline static const VectorRegister4Double DoubleTwo = MakeVectorRegisterDouble(2.0, 2.0, 2.0, 2.0);
+	inline static const uint64 DoubleAlmostTwoBits = 0x3FFFFFFFFFFFFFFF;
+	inline static const VectorRegister4Double DoubleAlmostTwo = MakeVectorRegisterDoubleMask(DoubleAlmostTwoBits, DoubleAlmostTwoBits, DoubleAlmostTwoBits, DoubleAlmostTwoBits);
 }
diff --git a/Engine/Source/Runtime/Core/Public/Math/Vector.h b/Engine/Source/Runtime/Core/Public/Math/Vector.h
--- a/Engine/Source/Runtime/Core/Public/Math/Vector.h
+++ b/Engine/Source/Runtime/Core/Public/Math/Vector.h
@@ -32,10 +32,6 @@
 #pragma warning (disable : 4459)
 #endif
 
-/**
- * A vector in 3-D space composed of components (X, Y, Z) with floating point precision.
- */
-
 // Move out of global namespace to avoid collisions with Chaos::TVector within the physics code.
 namespace UE
 {
@@ -44,6 +40,9 @@ namespace Math
 
 template<typename T> struct TPlane;
 
+/**
+ * A vector in 3-D space composed of components (X, Y, Z) with floating point precision.
+ */
 template<typename T>
 struct TVector
 {
@@ -95,19 +94,19 @@ public:
 	static const TVector<T> ZAxisVector;
 
 	/** @return Zero Vector (0,0,0) */
-	static TVector<T> Zero() { return ZeroVector; }
+	static inline TVector<T> Zero() { return ZeroVector; }
 
 	/** @return One Vector (1,1,1) */
-	static TVector<T> One() { return OneVector; }
+	static inline TVector<T> One() { return OneVector; }
 
 	/** @return Unit X Vector (1,0,0)  */
-	static TVector<T> UnitX() { return XAxisVector; }
+	static inline TVector<T> UnitX() { return XAxisVector; }
 
 	/** @return Unit Y Vector (0,1,0)  */
-	static TVector<T> UnitY() { return YAxisVector; }
+	static inline TVector<T> UnitY() { return YAxisVector; }
 
 	/** @return Unit Z Vector (0,0,1)  */
-	static TVector<T> UnitZ() { return ZAxisVector; }
+	static inline TVector<T> UnitZ() { return ZAxisVector; }
 
 public:
 
@@ -798,7 +797,7 @@ public:
      * @return Quaternion from the Vector's direction, without any roll.
      * @see ToOrientationRotator(), FQuat::FindBetweenVectors()
      */
-    CORE_API FQuat ToOrientationQuat() const;
+    CORE_API TQuat<T> ToOrientationQuat() const;
 
     /**
      * Return the FRotator orientation corresponding to the direction in which the vector points.
@@ -1183,7 +1182,6 @@ inline void operator<<(FStructuredArchive::FSlot Slot, TVector<double>& V)
 	V.DiagnosticCheckNaN();
 }
 
-
 /* FVector inline functions
  *****************************************************************************/
 
@@ -2374,17 +2372,17 @@ FORCEINLINE TVector<T> TVector<T>::Max3(const TVector<T>& A, const TVector<T>& B
         ); 
 }
 
-template<typename T> const TVector<T> TVector<T>::ZeroVector(0, 0, 0);
-template<typename T> const TVector<T> TVector<T>::OneVector(1, 1, 1);
-template<typename T> const TVector<T> TVector<T>::UpVector(0, 0, 1);
-template<typename T> const TVector<T> TVector<T>::DownVector(0, 0, -1);
-template<typename T> const TVector<T> TVector<T>::ForwardVector(1, 0, 0);
-template<typename T> const TVector<T> TVector<T>::BackwardVector(-1, 0, 0);
-template<typename T> const TVector<T> TVector<T>::RightVector(0, 1, 0);
-template<typename T> const TVector<T> TVector<T>::LeftVector(0, -1, 0);
-template<typename T> const TVector<T> TVector<T>::XAxisVector(1, 0, 0);
-template<typename T> const TVector<T> TVector<T>::YAxisVector(0, 1, 0);
-template<typename T> const TVector<T> TVector<T>::ZAxisVector(0, 0, 1);
+template<typename T> inline const TVector<T> TVector<T>::ZeroVector(0, 0, 0);
+template<typename T> inline const TVector<T> TVector<T>::OneVector(1, 1, 1);
+template<typename T> inline const TVector<T> TVector<T>::UpVector(0, 0, 1);
+template<typename T> inline const TVector<T> TVector<T>::DownVector(0, 0, -1);
+template<typename T> inline const TVector<T> TVector<T>::ForwardVector(1, 0, 0);
+template<typename T> inline const TVector<T> TVector<T>::BackwardVector(-1, 0, 0);
+template<typename T> inline const TVector<T> TVector<T>::RightVector(0, 1, 0);
+template<typename T> inline const TVector<T> TVector<T>::LeftVector(0, -1, 0);
+template<typename T> inline const TVector<T> TVector<T>::XAxisVector(1, 0, 0);
+template<typename T> inline const TVector<T> TVector<T>::YAxisVector(0, 1, 0);
+template<typename T> inline const TVector<T> TVector<T>::ZAxisVector(0, 0, 1);
 
 /**
  * Multiplies a vector by a scaling factor.
@@ -2412,7 +2410,7 @@ FORCEINLINE uint32 GetTypeHash(const TVector<T>& Vector)
 	return FCrc::MemCrc_DEPRECATED(&Vector, sizeof(Vector));
 }
 
-} // namespace UE::Core
+} // namespace UE::Math
 } // namespace UE
 
 DECLARE_LWC_TYPE(Vector, 3);
@@ -2433,8 +2431,8 @@ DECLARE_LWC_TYPE_ISPC(Vector, 3);
 // Forward declare all explicit specializations (in UnrealMath.cpp)
 template<> CORE_API FRotator FVector3f::ToOrientationRotator() const;
 template<> CORE_API FRotator FVector3d::ToOrientationRotator() const;
-template<> CORE_API FQuat FVector3f::ToOrientationQuat() const;
-template<> CORE_API FQuat FVector3d::ToOrientationQuat() const;
+template<> CORE_API FQuat4f FVector3f::ToOrientationQuat() const;
+template<> CORE_API FQuat4d FVector3d::ToOrientationQuat() const;
 template<> CORE_API FRotator FVector3f::Rotation() const;
 template<> CORE_API FRotator FVector3d::Rotation() const;
 
diff --git a/Engine/Source/Runtime/Core/Public/Math/VectorRegister.h b/Engine/Source/Runtime/Core/Public/Math/VectorRegister.h
--- a/Engine/Source/Runtime/Core/Public/Math/VectorRegister.h
+++ b/Engine/Source/Runtime/Core/Public/Math/VectorRegister.h
@@ -31,8 +31,73 @@ using TVectorRegisterType = std::conditional_t<std::is_same_v<T, float>, VectorR
 // 'Cross-platform' vector intrinsics (built on the platform-specific ones defined above)
 #include "Math/UnrealMathVectorCommon.h"
 
+
+#if !defined(UE_SSE_DOUBLE_ALIGNMENT) || (UE_SSE_DOUBLE_ALIGNMENT <= 16)
+// Suitable to just use regular VectorRegister4Double as persistent variables.
+using PersistentVectorRegister4Double = VectorRegister4Double;
+#else
+/**
+ * VectorRegister type to be used as member variables of structs, when alignment should not be over 16 bytes.
+ * 32-byte alignment requires large allocations in some allocators, and this allows us to keep at a 16-byte aligned type in situations (ie AVX) where the underlying
+ * VectorRegister type might want 32-byte alignment instead. It does imply use of unaligned loads/stores (which should be fine) in the conversion process, but try to
+ * formulate code using these types to only load/store once by avoiding repeated type conversion.
+ *
+ * See TPersistentVectorRegisterType<> below, for automatic selection of the correct type for a member variable.
+ */
+struct alignas(16) PersistentVectorRegister4Double
+{
+	double XYZW[4];
+
+	FORCEINLINE PersistentVectorRegister4Double() = default;
+
+	FORCEINLINE PersistentVectorRegister4Double(const VectorRegister4Double& Register)
+	{
+		VectorStore(Register, XYZW);
+	}
+
+	FORCEINLINE PersistentVectorRegister4Double(const VectorRegister4Float& Register)
+	{
+		VectorStore(VectorRegister4Double(Register), XYZW);
+	}
+
+	FORCEINLINE PersistentVectorRegister4Double& operator=(const VectorRegister4Double& Register)
+	{
+		VectorStore(Register, XYZW);
+		return *this;
+	}
+
+	FORCEINLINE PersistentVectorRegister4Double& operator=(const VectorRegister4Float& Register)
+	{
+		VectorStore(VectorRegister4Double(Register), XYZW);
+		return *this;
+	}
+
+	FORCEINLINE operator VectorRegister4Double() const
+	{
+		return VectorLoad(XYZW);
+	}
+};
+#endif // Alignment check for VectorRegister4Double
+
+// Suitable to just use regular VectorRegister4Float as persistent variables.
+using PersistentVectorRegister4Float = VectorRegister4Float;
+
+/**
+ * Alias for VectorRegister type to be used as member variables of structs, when alignment should not be over 16 bytes.
+ * 32-byte alignment requires large allocations in some allocators, which we want to avoid.
+ * See PersistentVectorRegister4Double for more details.
+ * 
+ * Usage example:
+ * TPersistentVectorRegisterType<float> MyFloatRegister;
+ * TPersistentVectorRegisterType<double> MyDoubleRegister;
+ */
+template<typename T>
+using TPersistentVectorRegisterType = std::conditional_t<std::is_same_v<T, float>, PersistentVectorRegister4Float, std::conditional_t<std::is_same_v<T, double>, PersistentVectorRegister4Double, void> >;
+
+
+
 /** Vector that represents (1/255,1/255,1/255,1/255) */
-extern CORE_API const VectorRegister VECTOR_INV_255;
+inline static const VectorRegister VECTOR_INV_255 = DECLARE_VECTOR_REGISTER(1.f / 255.f, 1.f / 255.f, 1.f / 255.f, 1.f / 255.f);
 
 // Old names for comparison functions, kept for compatibility.
 #define VectorMask_LT( Vec1, Vec2 )			VectorCompareLT(Vec1, Vec2)
@@ -49,9 +114,7 @@ extern CORE_API const VectorRegister VECTOR_INV_255;
 
 namespace GlobalVectorConstants
 {
-	static const VectorRegister AnimWeightThreshold = MakeVectorRegister(ZERO_ANIMWEIGHT_THRESH, ZERO_ANIMWEIGHT_THRESH, ZERO_ANIMWEIGHT_THRESH, ZERO_ANIMWEIGHT_THRESH);
-	static const VectorRegister RotationSignificantThreshold = MakeVectorRegister(1.0f - DELTA*DELTA, 1.0f - DELTA*DELTA, 1.0f - DELTA*DELTA, 1.0f - DELTA*DELTA);
-	// LWC_TODO : only needed until FQuat is using doubles (see TransformVectorized)
-	static const VectorRegister4Float RotationSignificantThresholdFloat = MakeVectorRegisterFloat(1.0f - DELTA * DELTA, 1.0f - DELTA * DELTA, 1.0f - DELTA * DELTA, 1.0f - DELTA * DELTA);
+	inline static const VectorRegister AnimWeightThreshold = MakeVectorRegister(ZERO_ANIMWEIGHT_THRESH, ZERO_ANIMWEIGHT_THRESH, ZERO_ANIMWEIGHT_THRESH, ZERO_ANIMWEIGHT_THRESH);
+	inline static const VectorRegister RotationSignificantThreshold = MakeVectorRegister(1.0f - DELTA*DELTA, 1.0f - DELTA*DELTA, 1.0f - DELTA*DELTA, 1.0f - DELTA*DELTA);
 }
 
diff --git a/Engine/Source/Runtime/Core/Public/UObject/UnrealNames.inl b/Engine/Source/Runtime/Core/Public/UObject/UnrealNames.inl
--- a/Engine/Source/Runtime/Core/Public/UObject/UnrealNames.inl
+++ b/Engine/Source/Runtime/Core/Public/UObject/UnrealNames.inl
@@ -81,6 +81,10 @@ REGISTER_NAME(74,Plane4f)
 REGISTER_NAME(75,Plane4d)
 REGISTER_NAME(76,Matrix44f)
 REGISTER_NAME(77,Matrix44d)
+REGISTER_NAME(78,Quat4f)
+REGISTER_NAME(79,Quat4d)
+REGISTER_NAME(80,Transform4f)
+REGISTER_NAME(81,Transform4d)
 
 // Object class names.
 REGISTER_NAME(100,Object)
diff --git a/Engine/Source/Runtime/CoreUObject/Private/UObject/Class.cpp b/Engine/Source/Runtime/CoreUObject/Private/UObject/Class.cpp
--- a/Engine/Source/Runtime/CoreUObject/Private/UObject/Class.cpp
+++ b/Engine/Source/Runtime/CoreUObject/Private/UObject/Class.cpp
@@ -3305,7 +3305,18 @@ public:
 			else
 			{
 				// Allocate space for the struct
-				const int32 RequiredAllocSize = ScriptStruct->GetStructureSize();
+				int32 RequiredAllocSize = ScriptStruct->GetStructureSize();
+				if (StructOps)
+				{
+					const int32 CppAllocSize = Align(StructOps->GetSize(), ScriptStruct->GetMinAlignment());
+					if (!ensure(RequiredAllocSize >= CppAllocSize))
+					{
+						UE_LOG(LogClass, Warning, TEXT("Struct %s%s has Cpp alloc size = %d > ScriptStruct->GetStructureSize() = %d, this could result in allocations that are too small to fit the structure."),
+							InStruct->GetPrefixCPP(), *InStruct->GetName(), CppAllocSize, RequiredAllocSize);
+						// We'd probably crash below, so use a larger allocated size.
+						RequiredAllocSize = StructOps->GetSize();
+					}
+				}
 				TempBuffer = (uint8*)FMemory::Malloc(RequiredAllocSize, ScriptStruct->GetMinAlignment());
 
 				// The following section is a partial duplication of ScriptStruct->InitializeStruct, except we initialize with 0xFD instead of 0x00
@@ -6515,7 +6526,9 @@ bool UScriptStruct::CanSerializeAsAlias(const FPropertyTag& Tag) const
 	return (
 		(Tag.StructName == NAME_Vector && (GetFName() == NAME_Vector3f || GetFName() == NAME_Vector3d)) ||
 		(Tag.StructName == NAME_Plane && (GetFName() == NAME_Plane4f || GetFName() == NAME_Plane4d)) ||
-		(Tag.StructName == NAME_Matrix && (GetFName() == NAME_Matrix44f || GetFName() == NAME_Matrix44d))
+		(Tag.StructName == NAME_Matrix && (GetFName() == NAME_Matrix44f || GetFName() == NAME_Matrix44d)) ||
+		(Tag.StructName == NAME_Quat && (GetFName() == NAME_Quat4f || GetFName() == NAME_Quat4d)) ||
+		(Tag.StructName == NAME_Transform && (GetFName() == NAME_Transform4f || GetFName() == NAME_Transform4d))
 		);
 }
 
diff --git a/Engine/Source/Runtime/CoreUObject/Private/UObject/Property.cpp b/Engine/Source/Runtime/CoreUObject/Private/UObject/Property.cpp
--- a/Engine/Source/Runtime/CoreUObject/Private/UObject/Property.cpp
+++ b/Engine/Source/Runtime/CoreUObject/Private/UObject/Property.cpp
@@ -256,7 +256,7 @@ IMPLEMENT_STRUCT(Color);
 
 
 template<>
-struct TStructOpsTypeTraits<FQuat> : public TStructOpsTypeTraitsBase2<FQuat>
+struct TStructOpsTypeTraits<FQuat4f> : public TStructOpsTypeTraitsBase2<FQuat4f>
 {
 	enum 
 	{
@@ -265,9 +265,26 @@ struct TStructOpsTypeTraits<FQuat> : public TStructOpsTypeTraitsBase2<FQuat>
 		WithNetSerializer = true,
 		WithNetSharedSerialization = true,
 		WithIdentical = true,
+		WithSerializer = true,
+	};
+};
+IMPLEMENT_STRUCT(Quat4f);
+
+template<>
+struct TStructOpsTypeTraits<FQuat4d> : public TStructOpsTypeTraitsBase2<FQuat4d>
+{
+	enum
+	{
+		//quat is somewhat special in that it initialized w to one
+		WithNoInitConstructor = true,
+		WithNetSerializer = true,
+		WithNetSharedSerialization = true,
+		WithIdentical = true,
+		WithSerializer = true,
 	};
 };
-IMPLEMENT_STRUCT(Quat);
+IMPLEMENT_STRUCT(Quat4d);
+IMPLEMENT_STRUCT(Quat);		// Aliased to one of FQuat4f/FQuat4d
 
 template<>
 struct TStructOpsTypeTraits<FTwoVectors> : public TStructOpsTypeTraitsBase2<FTwoVectors>
@@ -358,14 +375,27 @@ struct TStructOpsTypeTraits<FGuid> : public TStructOpsTypeTraitsBase2<FGuid>
 IMPLEMENT_STRUCT(Guid);
 
 template<>
-struct TStructOpsTypeTraits<FTransform> : public TStructOpsTypeTraitsBase2<FTransform>
+struct TStructOpsTypeTraits<FTransform3f> : public TStructOpsTypeTraitsBase2<FTransform3f>
+{
+	enum
+	{
+		//WithSerializer = true,
+		WithIdentical = true,
+	};
+};
+IMPLEMENT_STRUCT(Transform3f);
+
+template<>
+struct TStructOpsTypeTraits<FTransform3d> : public TStructOpsTypeTraitsBase2<FTransform3d>
 {
 	enum
 	{
+		//WithSerializer = true,
 		WithIdentical = true,
 	};
 };
-IMPLEMENT_STRUCT(Transform);
+IMPLEMENT_STRUCT(Transform3d);
+IMPLEMENT_STRUCT(Transform); // Aliased
 
 template<>
 struct TStructOpsTypeTraits<FRandomStream> : public TStructOpsTypeTraitsBase2<FRandomStream>
diff --git a/Engine/Source/Runtime/CoreUObject/Private/UObject/ScriptCore.cpp b/Engine/Source/Runtime/CoreUObject/Private/UObject/ScriptCore.cpp
--- a/Engine/Source/Runtime/CoreUObject/Private/UObject/ScriptCore.cpp
+++ b/Engine/Source/Runtime/CoreUObject/Private/UObject/ScriptCore.cpp
@@ -759,7 +759,7 @@ COREUOBJECT_API uint8 GRegisterCast( int32 CastCode, const FNativeFuncPtr& Func
 void UObject::SkipFunction(FFrame& Stack, RESULT_DECL, UFunction* Function)
 {
 	// allocate temporary memory on the stack for evaluating parameters
-	uint8* Frame = (uint8*)FMemory_Alloca(Function->PropertiesSize);
+	uint8* Frame = (uint8*)FMemory_Alloca_Aligned(Function->PropertiesSize, Function->GetMinAlignment());
 	FMemory::Memzero(Frame, Function->PropertiesSize);
 	for (FProperty* Property = (FProperty*)(Function->ChildProperties); *Stack.Code != EX_EndFunctionParms; Property = (FProperty*)(Property->Next))
 	{
@@ -817,7 +817,7 @@ void ProcessScriptFunction(UObject* Context, UFunction* Function, FFrame& Stack,
 	bool bUsePersistentFrame = (nullptr != FrameMemory);
 	if (!bUsePersistentFrame)
 	{
-		FrameMemory = (uint8*)FMemory_Alloca(Function->PropertiesSize);
+		FrameMemory = (uint8*)FMemory_Alloca_Aligned(Function->PropertiesSize, Function->GetMinAlignment());
 		FMemory::Memzero(FrameMemory, Function->PropertiesSize);
 	}
 
@@ -977,7 +977,7 @@ void UObject::CallFunction( FFrame& Stack, RESULT_DECL, UFunction* Function )
 		if (FunctionCallspace & FunctionCallspace::Remote)
 		{
 			// Call native networkable function.
-			uint8* Buffer = (uint8*)FMemory_Alloca(Function->ParmsSize);
+			uint8* Buffer = (uint8*)FMemory_Alloca_Aligned(Function->ParmsSize, Function->GetMinAlignment());
 
 			SavedCode = Stack.Code; // Since this is native, we need to rollback the stack if we are calling both remotely and locally
 
@@ -1219,7 +1219,7 @@ bool UObject::CallFunctionByNameWithArguments(const TCHAR* Str, FOutputDevice& A
 	}
 
 	// Parse all function parameters.
-	uint8* Parms = (uint8*)FMemory_Alloca(Function->ParmsSize);
+	uint8* Parms = (uint8*)FMemory_Alloca_Aligned(Function->ParmsSize, Function->GetMinAlignment());
 	FMemory::Memzero( Parms, Function->ParmsSize );
 
 	for (TFieldIterator<FProperty> It(Function); It && It->HasAnyPropertyFlags(CPF_Parm); ++It)
@@ -1941,7 +1941,7 @@ void UObject::ProcessEvent( UFunction* Function, void* Parms )
 		const bool bUsePersistentFrame = (NULL != Frame);
 		if (!bUsePersistentFrame)
 		{
-			Frame = (uint8*)FMemory_Alloca(Function->PropertiesSize);
+			Frame = (uint8*)FMemory_Alloca_Aligned(Function->PropertiesSize, Function->GetMinAlignment());
 			// zero the local property memory
 			FMemory::Memzero(Frame + Function->ParmsSize, Function->PropertiesSize - Function->ParmsSize);
 		}
@@ -2529,7 +2529,7 @@ DEFINE_FUNCTION(UObject::execSwitchValue)
 
 	bool bProperCaseUsed = false;
 	{
-		auto LocalTempIndexMem = (uint8*)FMemory_Alloca(IndexProperty->GetSize());
+		auto LocalTempIndexMem = (uint8*)FMemory_Alloca_Aligned(IndexProperty->GetSize(), IndexProperty->GetMinAlignment());
 		IndexProperty->InitializeValue(LocalTempIndexMem);
 		for (int32 CaseIndex = 0; CaseIndex < NumCases; ++CaseIndex)
 		{
@@ -2650,7 +2650,7 @@ DEFINE_FUNCTION(UObject::execLet)
 
 		if (LocallyKnownProperty)
 		{
-			LocalTempResult = (uint8*)FMemory_Alloca(LocallyKnownProperty->GetSize());
+			LocalTempResult = (uint8*)FMemory_Alloca_Aligned(LocallyKnownProperty->GetSize(), LocallyKnownProperty->GetMinAlignment());
 			LocallyKnownProperty->InitializeValue(LocalTempResult);
 			Stack.MostRecentPropertyAddress = LocalTempResult;
 		}
@@ -3023,7 +3023,7 @@ public:
 		const FMulticastScriptDelegate* DelegateAddr = (DelegateProp ? DelegateProp->GetMulticastDelegate(Stack.MostRecentPropertyAddress) : nullptr);
 
 		//Fill parameters
-		uint8* Parameters = (uint8*)FMemory_Alloca(SignatureFunction->ParmsSize);
+		uint8* Parameters = (uint8*)FMemory_Alloca_Aligned(SignatureFunction->ParmsSize, SignatureFunction->GetMinAlignment());
 		FMemory::Memzero(Parameters, SignatureFunction->ParmsSize);
 		for (FProperty* Property = (FProperty*)SignatureFunction->ChildProperties; *Stack.Code != EX_EndFunctionParms; Property = (FProperty*)Property->Next)
 		{
diff --git a/Engine/Source/Runtime/CoreUObject/Public/UObject/NoExportTypes.h b/Engine/Source/Runtime/CoreUObject/Public/UObject/NoExportTypes.h
--- a/Engine/Source/Runtime/CoreUObject/Public/UObject/NoExportTypes.h
+++ b/Engine/Source/Runtime/CoreUObject/Public/UObject/NoExportTypes.h
@@ -588,23 +588,67 @@ struct FRotator
  * Quaternion.
  * @note The full C++ class is located here: Engine\Source\Runtime\Core\Public\Math\Quat.h
  */
-USTRUCT(immutable, noexport, BlueprintType)
-struct FQuat
+USTRUCT(immutable, noexport)
+struct FQuat4f
 {
-	UPROPERTY(EditAnywhere, BlueprintReadWrite, Category=Quat, SaveGame)
+	UPROPERTY(EditAnywhere, Category=Quat, SaveGame)
 	float X;
 
-	UPROPERTY(EditAnywhere, BlueprintReadWrite, Category=Quat, SaveGame)
+	UPROPERTY(EditAnywhere, Category=Quat, SaveGame)
 	float Y;
 
-	UPROPERTY(EditAnywhere, BlueprintReadWrite, Category=Quat, SaveGame)
+	UPROPERTY(EditAnywhere, Category=Quat, SaveGame)
 	float Z;
 
-	UPROPERTY(EditAnywhere, BlueprintReadWrite, Category=Quat, SaveGame)
+	UPROPERTY(EditAnywhere, Category=Quat, SaveGame)
 	float W;
 
 };
 
+
+/**
+ * Quaternion.
+ * @note The full C++ class is located here: Engine\Source\Runtime\Core\Public\Math\Quat.h
+ */
+USTRUCT(immutable, noexport)
+struct FQuat4d
+{
+	UPROPERTY(EditAnywhere, Category = Quat, SaveGame)
+	double X;
+
+	UPROPERTY(EditAnywhere, Category = Quat, SaveGame)
+	double Y;
+
+	UPROPERTY(EditAnywhere, Category = Quat, SaveGame)
+	double Z;
+
+	UPROPERTY(EditAnywhere, Category = Quat, SaveGame)
+	double W;
+
+};
+
+
+/**
+ * Quaternion.
+ * @note The full C++ class is located here: Engine\Source\Runtime\Core\Public\Math\Quat.h
+ */
+USTRUCT(immutable, noexport, BlueprintType)
+struct FQuat
+{
+	UPROPERTY(EditAnywhere, BlueprintReadWrite, Category = Quat, SaveGame)
+	FLargeWorldCoordinatesReal X;		// Alias for float/double depending on LWC status. Note: Will be refactored to double before UE5 ships.
+
+	UPROPERTY(EditAnywhere, BlueprintReadWrite, Category = Quat, SaveGame)
+	FLargeWorldCoordinatesReal Y;
+
+	UPROPERTY(EditAnywhere, BlueprintReadWrite, Category = Quat, SaveGame)
+	FLargeWorldCoordinatesReal Z;
+
+	UPROPERTY(EditAnywhere, BlueprintReadWrite, Category = Quat, SaveGame)
+	FLargeWorldCoordinatesReal W;
+};
+
+
 /**
  * A packed normal.
  * @note The full C++ class is located here: Engine\Source\Runtime\RenderCore\Public\PackedNormal.h
@@ -1169,6 +1213,46 @@ struct FInterpCurveLinearColor
 	float LoopKeyOffset;
 };
 
+/**
+ * Transform composed of Quat/Translation/Scale.
+ * @note This is implemented in either TransformVectorized.h or TransformNonVectorized.h depending on the platform.
+ */
+USTRUCT(immutable, noexport)
+struct FTransform3f
+{
+	/** Rotation of this transformation, as a quaternion. */
+	UPROPERTY(EditAnywhere, Category = Transform, SaveGame)
+	FQuat4f Rotation;
+
+	/** Translation of this transformation, as a vector. */
+	UPROPERTY(EditAnywhere, Category = Transform, SaveGame)
+	FVector3f Translation;
+
+	/** 3D scale (always applied in local space) as a vector. */
+	UPROPERTY(EditAnywhere, Category = Transform, SaveGame)
+	FVector3f Scale3D;
+};
+
+/**
+ * Transform composed of Quat/Translation/Scale.
+ * @note This is implemented in either TransformVectorized.h or TransformNonVectorized.h depending on the platform.
+ */
+USTRUCT(immutable, noexport)
+struct FTransform3d
+{
+	/** Rotation of this transformation, as a quaternion. */
+	UPROPERTY(EditAnywhere, Category = Transform, SaveGame)
+	FQuat4d Rotation;
+
+	/** Translation of this transformation, as a vector. */
+	UPROPERTY(EditAnywhere, Category = Transform, SaveGame)
+	FVector3d Translation;
+
+	/** 3D scale (always applied in local space) as a vector. */
+	UPROPERTY(EditAnywhere, Category = Transform, SaveGame)
+	FVector3d Scale3D;
+};
+
 /**
  * Transform composed of Quat/Translation/Scale.
  * @note This is implemented in either TransformVectorized.h or TransformNonVectorized.h depending on the platform.
diff --git a/Engine/Source/Runtime/Engine/Private/BlueprintGeneratedClass.cpp b/Engine/Source/Runtime/Engine/Private/BlueprintGeneratedClass.cpp
--- a/Engine/Source/Runtime/Engine/Private/BlueprintGeneratedClass.cpp
+++ b/Engine/Source/Runtime/Engine/Private/BlueprintGeneratedClass.cpp
@@ -1448,7 +1448,7 @@ void UBlueprintGeneratedClass::CreatePersistentUberGraphFrame(UObject* Obj, bool
 			if (bUberGraphFunctionIsReady)
 			{
 				INC_MEMORY_STAT_BY(STAT_PersistentUberGraphFrameMemory, UberGraphFunction->GetStructureSize());
-				FrameMemory = (uint8*)FMemory::Malloc(UberGraphFunction->GetStructureSize());
+				FrameMemory = (uint8*)FMemory::Malloc(UberGraphFunction->GetStructureSize(), UberGraphFunction->GetMinAlignment());
 
 				FMemory::Memzero(FrameMemory, UberGraphFunction->GetStructureSize());
 				for (FProperty* Property = UberGraphFunction->PropertyLink; Property; Property = Property->PropertyLinkNext)
diff --git a/Engine/Source/Runtime/Experimental/Chaos/Private/Chaos/MassProperties.cpp b/Engine/Source/Runtime/Experimental/Chaos/Private/Chaos/MassProperties.cpp
--- a/Engine/Source/Runtime/Experimental/Chaos/Private/Chaos/MassProperties.cpp
+++ b/Engine/Source/Runtime/Experimental/Chaos/Private/Chaos/MassProperties.cpp
@@ -79,7 +79,7 @@ namespace Chaos
 
 		// NOTE: UE Matrix are column-major, so the PMatrix constructor is not setting eigenvectors - we need to transpose it to get a UE rotation matrix.
 		FinalRotation = FRotation3(RotationMatrix.GetTransposed());
-		if (!ensure(FMath::IsNearlyEqual(FinalRotation.Size(), 1.0f, KINDA_SMALL_NUMBER)))
+		if (!ensure(FMath::IsNearlyEqual((float)FinalRotation.Size(), 1.0f, KINDA_SMALL_NUMBER)))
 		{
 			return FRotation3::FromElements(FVec3(0), 1);
 		}
diff --git a/Engine/Source/Runtime/Experimental/Chaos/Private/Chaos/Transform.cpp b/Engine/Source/Runtime/Experimental/Chaos/Private/Chaos/Transform.cpp
--- a/Engine/Source/Runtime/Experimental/Chaos/Private/Chaos/Transform.cpp
+++ b/Engine/Source/Runtime/Experimental/Chaos/Private/Chaos/Transform.cpp
@@ -5,14 +5,9 @@
 
 using namespace Chaos;
 
-PMatrix<Chaos::FReal, 4, 4> operator*(const TRigidTransform<Chaos::FReal, 3>& Transform, const PMatrix<Chaos::FReal, 4, 4>& Matrix)
+PMatrix<Chaos::FReal, 4, 4> TRigidTransform<Chaos::FReal, 3>::operator*(const PMatrix<Chaos::FReal, 4, 4>& Matrix) const
 {
 	// LWC_TODO: Perf pessimization
-	return Transform.ToMatrixNoScale() * static_cast<const FMatrix&>(Matrix);
+	return ToMatrixNoScale() * static_cast<const UE::Math::TMatrix<FReal>&>(Matrix);
 }
 
-PMatrix<Chaos::FReal, 4, 4> operator*(const PMatrix<Chaos::FReal, 4, 4>& Matrix, const TRigidTransform<Chaos::FReal, 3>& Transform)
-{
-	// LWC_TODO: Perf pessimization
-	return static_cast<const FMatrix&>(Matrix) * Transform.ToMatrixNoScale();
-}
\ No newline at end of file
diff --git a/Engine/Source/Runtime/Experimental/Chaos/Private/PhysicsProxy/SingleParticlePhysicsProxy.cpp b/Engine/Source/Runtime/Experimental/Chaos/Private/PhysicsProxy/SingleParticlePhysicsProxy.cpp
--- a/Engine/Source/Runtime/Experimental/Chaos/Private/PhysicsProxy/SingleParticlePhysicsProxy.cpp
+++ b/Engine/Source/Runtime/Experimental/Chaos/Private/PhysicsProxy/SingleParticlePhysicsProxy.cpp
@@ -283,7 +283,7 @@ bool FSingleParticlePhysicsProxy::PullFromPhysicsState(const Chaos::FDirtyRigidP
 					FQuat Target = FMath::Lerp(*Prev, NextPullData->R, *Alpha);
 					if (LeashAlpha)
 					{
-						Target = FMath::Lerp(Rigid->R(), Target, *LeashAlpha);
+						Target = FMath::Lerp<FQuat>(Rigid->R(), Target, *LeashAlpha);
 					}
 					Rigid->SetR(Target, false);
 				}
diff --git a/Engine/Source/Runtime/Experimental/Chaos/Public/Chaos/Rotation.h b/Engine/Source/Runtime/Experimental/Chaos/Public/Chaos/Rotation.h
--- a/Engine/Source/Runtime/Experimental/Chaos/Public/Chaos/Rotation.h
+++ b/Engine/Source/Runtime/Experimental/Chaos/Public/Chaos/Rotation.h
@@ -44,15 +44,16 @@ namespace Chaos
 	};
 
 	template<>
-	class TRotation<FReal, 3> : public FQuat
+	class TRotation<FReal, 3> : public UE::Math::TQuat<FReal>
 	{
+		using BaseQuat = UE::Math::TQuat<FReal>;
 	public:
 		TRotation()
-		    : FQuat() {}
-		TRotation(const FQuat& Quat)
-		    : FQuat(Quat) {}
+		    : BaseQuat() {}
+		TRotation(const BaseQuat& Quat)
+		    : BaseQuat(Quat) {}
 		TRotation(const FMatrix& Matrix)
-		    : FQuat(Matrix) {}
+		    : BaseQuat(Matrix) {}
 
 		PMatrix<FReal, 3, 3> ToMatrix() const
 		{
@@ -107,9 +108,9 @@ namespace Chaos
 		 * Extract the Swing and Twist rotations, assuming that the Twist Axis is (1,0,0).
 		 * /see ToSwingTwist
 		 */
-		void ToSwingTwistX(FQuat& OutSwing, FQuat& OutTwist) const
+		void ToSwingTwistX(BaseQuat& OutSwing, BaseQuat& OutTwist) const
 		{
-			OutTwist = (X != 0.0f)? FQuat(X, 0, 0, W).GetNormalized() : FQuat::Identity;
+			OutTwist = (X != 0.0f)? BaseQuat(X, 0, 0, W).GetNormalized() : BaseQuat::Identity;
 			OutSwing = *this * OutTwist.Inverse();
 		}
 
@@ -144,7 +145,7 @@ namespace Chaos
 		 */
 		static TRotation<FReal, 3> FromIdentity()
 		{
-			return FQuat(0, 0, 0, 1);
+			return BaseQuat(0, 0, 0, 1);
 		}
 
 		/**
@@ -152,8 +153,8 @@ namespace Chaos
 		 */
 		static TRotation<FReal, 3> FromElements(const FReal X, const FReal Y, const FReal Z, const FReal W)
 		{
-			using FQuatReal = decltype(FQuat::X);	// LWC_TODO: Remove once FQuat supports variants
-			return FQuat((FQuatReal)X, (FQuatReal)Y, (FQuatReal)Z, (FQuatReal)W);
+			using FQuatReal = BaseQuat::FReal;
+			return BaseQuat((FQuatReal)X, (FQuatReal)Y, (FQuatReal)Z, (FQuatReal)W);
 		}
 
 		/**
@@ -169,7 +170,7 @@ namespace Chaos
 		 */
 		static TRotation<FReal, 3> FromAxisAngle(const ::Chaos::TVector<FReal, 3>& Axis, const FReal AngleRad)
 		{
-			return FQuat(FVector(Axis.X, Axis.Y, Axis.Z), (decltype(FQuat::X))AngleRad);
+			return BaseQuat(UE::Math::TVector<FReal>(Axis.X, Axis.Y, Axis.Z), (BaseQuat::FReal)AngleRad);
 		}
 
 		/**
@@ -177,7 +178,7 @@ namespace Chaos
 		 */
 		static TRotation<FReal, 3> FromVector(const ::Chaos::TVector<FReal, 3>& V)
 		{
-			using FQuatReal = decltype(FQuat::X);	// LWC_TODO: Remove once FQuat supports variants
+			using FQuatReal = BaseQuat::FReal;
 			TRotation<FReal, 3> Rot;
 			FReal HalfSize = 0.5f * V.Size();
 			FReal sinc = (FMath::Abs(HalfSize) > 1e-8) ? FMath::Sin(HalfSize) / HalfSize : 1;
diff --git a/Engine/Source/Runtime/Experimental/Chaos/Public/Chaos/Transform.h b/Engine/Source/Runtime/Experimental/Chaos/Public/Chaos/Transform.h
--- a/Engine/Source/Runtime/Experimental/Chaos/Public/Chaos/Transform.h
+++ b/Engine/Source/Runtime/Experimental/Chaos/Public/Chaos/Transform.h
@@ -113,44 +113,63 @@ namespace Chaos
 	};
 
 	template<>
-	class TRigidTransform<FReal, 2> : public FTransform
+	class TRigidTransform<FReal, 2> : public UE::Math::TTransform<FReal>
 	{
+		using BaseTransform = UE::Math::TTransform<FReal>;
 	public:
 		TRigidTransform()
-			: FTransform() {}
+			: BaseTransform() {}
 		TRigidTransform(const TVector<FReal, 3>& Translation, const TRotation<FReal, 3>& Rotation)
-			: FTransform(Rotation, FVector(Translation)) {}	// LWC_TODO: Perf pessimization
+			: BaseTransform(Rotation, Translation) {}
 		TRigidTransform(const FMatrix44d& Matrix)
-			: FTransform(Matrix) {}
+			: BaseTransform(Matrix) {}
 		TRigidTransform(const FMatrix44f& Matrix)
-			: FTransform(Matrix) {}
-		TRigidTransform(const FTransform& Transform)
-			: FTransform(Transform) {}
-		PMatrix<FReal, 3, 3> Inverse() const
+			: BaseTransform(Matrix) {}
+		TRigidTransform(const BaseTransform& Transform)
+			: BaseTransform(Transform) {}
+		TRigidTransform<FReal, 2> Inverse() const
 		{
-			return (FMatrix44f)ToMatrixNoScale().Inverse();	// LWC_TODO: Perf pessimization
+			return BaseTransform::Inverse();
+		}
+
+		inline TRigidTransform<FReal, 2> operator*(const TRigidTransform<FReal, 2>& Other) const
+		{
+			return BaseTransform::operator*(Other);
 		}
 	};
 
 	template<>
-	class TRigidTransform<FReal, 3> : public FTransform
+	class TRigidTransform<FReal, 3> : public UE::Math::TTransform<FReal>
 	{
+		using BaseTransform = UE::Math::TTransform<FReal>;
 	public:
 		TRigidTransform()
-			: FTransform() {}
+			: BaseTransform() {}
 		TRigidTransform(const TVector<FReal, 3>& Translation, const TRotation<FReal, 3>& Rotation)
-			: FTransform(Rotation, FVector(Translation)) {}	// LWC_TODO: Perf pessimization
+			: BaseTransform(Rotation, Translation) {}
 		TRigidTransform(const TVector<FReal, 3>& Translation, const TRotation<FReal, 3>& Rotation, const TVector<FReal, 3>& Scale)
-			: FTransform(Rotation, FVector(Translation), FVector(Scale)) {}	// LWC_TODO: Perf pessimization
+			: BaseTransform(Rotation, Translation, Scale) {}
 		TRigidTransform(const FMatrix44d& Matrix)
-			: FTransform(Matrix) {}
+			: BaseTransform(Matrix) {}
 		TRigidTransform(const FMatrix44f& Matrix)
-			: FTransform(Matrix) {}
-		TRigidTransform(const FTransform& Transform)
-			: FTransform(Transform) {}
-		PMatrix<FReal, 4, 4> Inverse() const
+			: BaseTransform(Matrix) {}
+		TRigidTransform(const BaseTransform& Transform)
+			: BaseTransform(Transform) {}
+		TRigidTransform<FReal, 3> Inverse() const
+		{
+			return BaseTransform::Inverse();
+		}
+
+		PMatrix<FReal, 4, 4> ToMatrixNoScale() const
+		{
+			return BaseTransform::ToMatrixNoScale();
+		}
+
+		CHAOS_API Chaos::PMatrix<Chaos::FReal, 4, 4> operator*(const Chaos::PMatrix<Chaos::FReal, 4, 4>& Matrix) const;
+		
+		inline TRigidTransform<FReal, 3> operator*(const TRigidTransform<FReal, 3>& Other) const
 		{
-			return ToMatrixNoScale().Inverse();
+			return BaseTransform::operator*(Other);
 		}
 
 		// Get the transform which maps from Other to This, ignoring the scale on both.
@@ -190,5 +209,3 @@ inline uint32 GetTypeHash(const Chaos::TRigidTransform<Chaos::FReal, 3>& InTrans
 	return HashCombine(GetTypeHash(InTransform.GetTranslation()), HashCombine(GetTypeHash(InTransform.GetRotation().Euler()), GetTypeHash(InTransform.GetScale3D())));
 }
 
-CHAOS_API Chaos::PMatrix<Chaos::FReal, 4, 4> operator*(const Chaos::TRigidTransform<Chaos::FReal, 3>& Transform, const Chaos::PMatrix<Chaos::FReal, 4, 4>& Matrix);
-CHAOS_API Chaos::PMatrix<Chaos::FReal, 4, 4> operator*(const Chaos::PMatrix<Chaos::FReal, 4, 4>& Matrix, const Chaos::TRigidTransform<Chaos::FReal, 3>& Transform);
\ No newline at end of file
diff --git a/Engine/Source/Runtime/Experimental/Chaos/Public/GeometryCollection/RecordedTransformTrack.h b/Engine/Source/Runtime/Experimental/Chaos/Public/GeometryCollection/RecordedTransformTrack.h
--- a/Engine/Source/Runtime/Experimental/Chaos/Public/GeometryCollection/RecordedTransformTrack.h
+++ b/Engine/Source/Runtime/Experimental/Chaos/Public/GeometryCollection/RecordedTransformTrack.h
@@ -544,7 +544,7 @@ struct CHAOS_API FRecordedTransformTrack
 
 		FQuat Delta = Curr.GetRotation() * Prev.GetRotation().Inverse();
 		FVector Axis;
-		float Angle;
+		FVector::FReal Angle;
 		Delta.ToAxisAndAngle(Axis, Angle);
 
 		return (Axis * Angle) / SampleWidth;
diff --git a/Engine/Source/Runtime/Experimental/ChaosCore/Public/Chaos/Matrix.h b/Engine/Source/Runtime/Experimental/ChaosCore/Public/Chaos/Matrix.h
--- a/Engine/Source/Runtime/Experimental/ChaosCore/Public/Chaos/Matrix.h
+++ b/Engine/Source/Runtime/Experimental/ChaosCore/Public/Chaos/Matrix.h
@@ -148,6 +148,11 @@ namespace Chaos
 			    M[3][0] * Other[0] + M[3][1] * Other[1] + M[3][2] * Other[2] + M[3][3] * Other[3]);
 		}
 #endif
+
+		PMatrix<FReal, 4, 4> operator*(const PMatrix<FReal, 4, 4>& Other) const
+		{
+			return static_cast<const UE::Math::TMatrix<FReal>*>(this)->operator*(static_cast<const UE::Math::TMatrix<FReal>&>(Other));
+		}
 	};
 
 	// TODO(mlentine): Do not use 4x4 matrix for 3x3 implementation
diff --git a/Engine/Source/Runtime/HeadMountedDisplay/Public/HeadMountedDisplayTypes.h b/Engine/Source/Runtime/HeadMountedDisplay/Public/HeadMountedDisplayTypes.h
--- a/Engine/Source/Runtime/HeadMountedDisplay/Public/HeadMountedDisplayTypes.h
+++ b/Engine/Source/Runtime/HeadMountedDisplay/Public/HeadMountedDisplayTypes.h
@@ -348,7 +348,7 @@ struct HEADMOUNTEDDISPLAY_API FXRMotionControllerData
 	UPROPERTY(BlueprintReadOnly, Category = "XR")
 	TArray<FVector> HandKeyPositions;
 	UPROPERTY(BlueprintReadOnly, Category = "XR")
-	TArray<struct FQuat> HandKeyRotations;
+	TArray<FQuat> HandKeyRotations;
 	UPROPERTY(BlueprintReadOnly, Category = "XR")
 	TArray<float> HandKeyRadii;
 
diff --git a/Engine/Source/Runtime/HeadMountedDisplay/Public/IHandTracker.h b/Engine/Source/Runtime/HeadMountedDisplay/Public/IHandTracker.h
--- a/Engine/Source/Runtime/HeadMountedDisplay/Public/IHandTracker.h
+++ b/Engine/Source/Runtime/HeadMountedDisplay/Public/IHandTracker.h
@@ -43,7 +43,7 @@ public:
 	 */
 	virtual bool GetKeypointState(EControllerHand Hand, EHandKeypoint Keypoint, FTransform& OutTransform, float& OutRadius) const = 0;
 
-	virtual bool GetAllKeypointStates(EControllerHand Hand, TArray<FVector>& OutPositions, TArray<struct FQuat>& OutRotations, TArray<float>& OutRadii) const = 0;
+	virtual bool GetAllKeypointStates(EControllerHand Hand, TArray<FVector>& OutPositions, TArray<FQuat>& OutRotations, TArray<float>& OutRadii) const = 0;
 
 	virtual bool HasHandMeshData() const
 	{
diff --git a/Engine/Source/Runtime/PhysicsCore/Public/PhysicsInterfaceDeclaresCore.h b/Engine/Source/Runtime/PhysicsCore/Public/PhysicsInterfaceDeclaresCore.h
--- a/Engine/Source/Runtime/PhysicsCore/Public/PhysicsInterfaceDeclaresCore.h
+++ b/Engine/Source/Runtime/PhysicsCore/Public/PhysicsInterfaceDeclaresCore.h
@@ -63,7 +63,7 @@ using FPhysActorDummy = ChaosInterface::FDummyPhysActor;
 template<typename T>
 using FCallbackDummy = ChaosInterface::FDummyCallback<T>;
 
-struct FTransform;
+DECLARE_LWC_TYPE(Transform, 3);
 
 using FHitLocation = ChaosInterface::FLocationHit;
 using FHitSweep = ChaosInterface::FSweepHit;
diff --git a/Engine/Source/Runtime/RigVM/Public/RigVMCore/RigVMTraits.h b/Engine/Source/Runtime/RigVM/Public/RigVMCore/RigVMTraits.h
--- a/Engine/Source/Runtime/RigVM/Public/RigVMCore/RigVMTraits.h
+++ b/Engine/Source/Runtime/RigVM/Public/RigVMCore/RigVMTraits.h
@@ -45,9 +45,7 @@ struct TRigVMIsBaseStructure
 
 struct FRotator;
 template<> struct TRigVMIsBaseStructure<FRotator> { enum { Value = true }; };
-struct FQuat;
 template<> struct TRigVMIsBaseStructure<FQuat> { enum { Value = true }; };
-struct FTransform;
 template<> struct TRigVMIsBaseStructure<FTransform> { enum { Value = true }; };
 struct FLinearColor;
 template<> struct TRigVMIsBaseStructure<FLinearColor> { enum { Value = true }; };
diff --git a/Engine/Source/Runtime/TimeManagement/Private/GenlockedTimecodeProvider.cpp b/Engine/Source/Runtime/TimeManagement/Private/GenlockedTimecodeProvider.cpp
--- a/Engine/Source/Runtime/TimeManagement/Private/GenlockedTimecodeProvider.cpp
+++ b/Engine/Source/Runtime/TimeManagement/Private/GenlockedTimecodeProvider.cpp
@@ -32,7 +32,7 @@ namespace
 		const double FpsB = RateB.AsDecimal();
 
 		// protect against divide by zero
-		if (FMath::IsNearlyEqual(FpsB, 0))
+		if (FMath::IsNearlyEqual(FpsB, 0.0))
 		{
 			return false;
 		}
@@ -63,7 +63,7 @@ FQualifiedFrameTime UGenlockedTimecodeProvider::CorrectFromGenlock(FQualifiedFra
 	}
 	
 	// Protect against divide by zero. Also eliminates timecode rate being faster than genlock rate.
-	if (FMath::IsNearlyEqual(SyncsPerTcFrame, 0))
+	if (FMath::IsNearlyEqual(SyncsPerTcFrame, 0.0))
 	{
 		return InFrameTime;
 	}
