IoDispatcher: Implemented perfect hashing for iostore tocs
#rnx

diff --git a/Engine/Source/Runtime/Core/Private/IO/IoDispatcherFileBackend.cpp b/Engine/Source/Runtime/Core/Private/IO/IoDispatcherFileBackend.cpp
--- a/Engine/Source/Runtime/Core/Private/IO/IoDispatcherFileBackend.cpp
+++ b/Engine/Source/Runtime/Core/Private/IO/IoDispatcherFileBackend.cpp
@@ -34,6 +34,7 @@ TRACE_DECLARE_INT_COUNTER(IoDispatcherFileBackendForwardSeeks, TEXT("IoDispatche
 TRACE_DECLARE_INT_COUNTER(IoDispatcherFileBackendBackwardSeeks, TEXT("IoDispatcherFileBackend/BackwardSeeks"));
 TRACE_DECLARE_INT_COUNTER(IoDispatcherFileBackendSwitchContainerSeeks, TEXT("IoDispatcherFileBackend/SwitchContainerSeeks"));
 TRACE_DECLARE_MEMORY_COUNTER(IoDispatcherFileBackendTotalSeekDistance, TEXT("IoDispatcherFileBackend/TotalSeekDistance"));
+TRACE_DECLARE_MEMORY_COUNTER(IoStoreTocMemory, TEXT("IoDispatcher/TocMemory"));
 
 //PRAGMA_DISABLE_OPTIMIZATION
 
@@ -670,12 +671,26 @@ FIoStatus FFileIoStoreReader::Initialize(const TCHAR* InContainerPath, int32 InO
 		Partition.ContainerFileIndex = GlobalPartitionIndex++;
 	}
 
-	Toc.Reserve(TocResource.Header.TocEntryCount);
-
-	for (uint32 ChunkIndex = 0; ChunkIndex < TocResource.Header.TocEntryCount; ++ChunkIndex)
+	if (TocResource.ChunkHashSeeds.Num() == TocResource.Header.TocEntryCount)
+	{
+		PerfectHashMap.TocChunkHashSeeds = MoveTemp(TocResource.ChunkHashSeeds);
+		PerfectHashMap.TocOffsetAndLengths = MoveTemp(TocResource.ChunkOffsetLengths);
+		PerfectHashMap.TocChunkHashes.SetNumUninitialized(TocResource.Header.TocEntryCount);
+		// Store only the chunk hashes, assumes that the perfect hash function is different from the
+		// default hash function and that they won't both collide
+		for (uint32 ChunkIndex = 0; ChunkIndex < TocResource.Header.TocEntryCount; ++ChunkIndex)
+		{
+			PerfectHashMap.TocChunkHashes[ChunkIndex] = GetTypeHash(TocResource.ChunkIds[ChunkIndex]);
+		}
+		bHasPerfectHashMap = true;
+	}
+	else
 	{
-		const FIoOffsetAndLength& ChunkOffsetLength = TocResource.ChunkOffsetLengths[ChunkIndex];
-		Toc.Add(TocResource.ChunkIds[ChunkIndex], ChunkOffsetLength);
+		for (uint32 ChunkIndex = 0; ChunkIndex < TocResource.Header.TocEntryCount; ++ChunkIndex)
+		{
+			TocImperfectHashMapFallback.Add(TocResource.ChunkIds[ChunkIndex], TocResource.ChunkOffsetLengths[ChunkIndex]);
+		}
+		bHasPerfectHashMap = false;
 	}
 	
 	ContainerFile.CompressionMethods	= MoveTemp(TocResource.CompressionMethods);
@@ -686,6 +701,14 @@ FIoStatus FFileIoStoreReader::Initialize(const TCHAR* InContainerPath, int32 InO
 	ContainerFile.BlockSignatureHashes	= MoveTemp(TocResource.ChunkBlockSignatures);
 	ContainerFile.ContainerInstanceId	= ++GlobalContainerInstanceId;
 
+	TRACE_COUNTER_ADD(IoStoreTocMemory,
+		TocImperfectHashMapFallback.GetAllocatedSize() +
+		PerfectHashMap.TocOffsetAndLengths.GetAllocatedSize() +
+		PerfectHashMap.TocChunkHashes.GetAllocatedSize() +
+		PerfectHashMap.TocChunkHashSeeds.GetAllocatedSize() +
+		ContainerFile.CompressionBlocks.GetAllocatedSize() +
+		ContainerFile.BlockSignatureHashes.GetAllocatedSize());
+
 	ContainerId = TocResource.Header.ContainerId;
 	Order = InOrder;
 	return FIoStatus::Ok;
@@ -703,7 +726,10 @@ FIoStatus FFileIoStoreReader::Close()
 		PlatformImpl.CloseContainer(Partition.FileHandle);
 	}
 
-	Toc.Empty();
+	PerfectHashMap.TocChunkHashSeeds.Empty();
+	PerfectHashMap.TocChunkHashes.Empty();
+	PerfectHashMap.TocOffsetAndLengths.Empty();
+	TocImperfectHashMapFallback.Empty();
 	ContainerFile = FFileIoStoreContainerFile();
 	ContainerId = FIoContainerId();
 	Order = INDEX_NONE;
@@ -712,18 +738,54 @@ FIoStatus FFileIoStoreReader::Close()
 	return FIoStatus::Ok;
 }
 
+const FIoOffsetAndLength* FFileIoStoreReader::FindChunkInternal(const FIoChunkId& ChunkId) const
+{
+	if (bHasPerfectHashMap)
+	{
+		// See FIoStoreWriterImpl::GeneratePerfectHashes
+		const uint32 ChunkCount = PerfectHashMap.TocChunkHashes.Num();
+		if (!ChunkCount)
+		{
+			return nullptr;
+		}
+		uint32 SeedIndex = FIoStoreTocResource::HashChunkIdWithSeed(0, ChunkId) % ChunkCount;
+		const int32 Seed = PerfectHashMap.TocChunkHashSeeds[SeedIndex];
+		if (Seed == 0)
+		{
+			return nullptr;
+		}
+		uint32 Slot;
+		if (Seed < 0)
+		{
+			Slot = static_cast<uint32>(-Seed - 1);
+		}
+		else
+		{
+			Slot = FIoStoreTocResource::HashChunkIdWithSeed(static_cast<uint32>(Seed), ChunkId) % ChunkCount;
+		}
+		if (PerfectHashMap.TocChunkHashes[Slot] == GetTypeHash(ChunkId))
+		{
+			return &PerfectHashMap.TocOffsetAndLengths[Slot];
+		}
+		return nullptr;
+	}
+	else
+	{
+		return TocImperfectHashMapFallback.Find(ChunkId);
+	}
+}
+
 bool FFileIoStoreReader::DoesChunkExist(const FIoChunkId& ChunkId) const
 {
 	check(!bClosed);
-	return Toc.Find(ChunkId) != nullptr;
+	return FindChunkInternal(ChunkId) != nullptr;
 }
 
 TIoStatusOr<uint64> FFileIoStoreReader::GetSizeForChunk(const FIoChunkId& ChunkId) const
 {
 	check(!bClosed);
-	const FIoOffsetAndLength* OffsetAndLength = Toc.Find(ChunkId);
-
-	if (OffsetAndLength != nullptr)
+	const FIoOffsetAndLength* OffsetAndLength = FindChunkInternal(ChunkId);
+	if (OffsetAndLength)
 	{
 		return OffsetAndLength->GetLength();
 	}
@@ -736,7 +798,7 @@ TIoStatusOr<uint64> FFileIoStoreReader::GetSizeForChunk(const FIoChunkId& ChunkI
 const FIoOffsetAndLength* FFileIoStoreReader::Resolve(const FIoChunkId& ChunkId) const
 {
 	check(!bClosed);
-	return Toc.Find(ChunkId);
+	return FindChunkInternal(ChunkId);
 }
 
 IMappedFileHandle* FFileIoStoreReader::GetMappedContainerFileHandle(uint64 TocOffset)
diff --git a/Engine/Source/Runtime/Core/Private/IO/IoDispatcherFileBackend.h b/Engine/Source/Runtime/Core/Private/IO/IoDispatcherFileBackend.h
--- a/Engine/Source/Runtime/Core/Private/IO/IoDispatcherFileBackend.h
+++ b/Engine/Source/Runtime/Core/Private/IO/IoDispatcherFileBackend.h
@@ -60,13 +60,23 @@ public:
 	const FAES::FAESKey& GetEncryptionKey() const { return ContainerFile.EncryptionKey; }
 
 private:
-	FFileIoStoreImpl& PlatformImpl;
+	const FIoOffsetAndLength* FindChunkInternal(const FIoChunkId& ChunkId) const;
 
-	TMap<FIoChunkId, FIoOffsetAndLength> Toc;
+	FFileIoStoreImpl& PlatformImpl;
+	
+	struct FPerfectHashMap
+	{
+		TArray<int32> TocChunkHashSeeds;
+		TArray<uint32> TocChunkHashes;
+		TArray<FIoOffsetAndLength> TocOffsetAndLengths;
+	};
+	FPerfectHashMap PerfectHashMap;
+	TMap<FIoChunkId, FIoOffsetAndLength> TocImperfectHashMapFallback;
 	FFileIoStoreContainerFile ContainerFile;
 	FIoContainerId ContainerId;
 	int32 Order;
 	bool bClosed = false;
+	bool bHasPerfectHashMap = false;
 
 	static TAtomic<uint32> GlobalPartitionIndex;
 	static TAtomic<uint32> GlobalContainerInstanceId;
diff --git a/Engine/Source/Runtime/Core/Private/IO/IoStore.cpp b/Engine/Source/Runtime/Core/Private/IO/IoStore.cpp
--- a/Engine/Source/Runtime/Core/Private/IO/IoStore.cpp
+++ b/Engine/Source/Runtime/Core/Private/IO/IoStore.cpp
@@ -662,6 +662,100 @@ public:
 		}
 	}
 
+	void GeneratePerfectHashes(FIoStoreTocResource& TocResource)
+	{
+		// https://en.wikipedia.org/wiki/Perfect_hash_function
+		TRACE_CPUPROFILER_EVENT_SCOPE(TocGeneratePerfectHashes);
+		uint32 ChunkCount = TocResource.ChunkIds.Num();
+		check(TocResource.ChunkOffsetLengths.Num() == ChunkCount);
+		
+		TArray<FIoChunkId> OutTocChunkIds;
+		OutTocChunkIds.SetNum(ChunkCount);
+		TArray<FIoOffsetAndLength> OutTocOffsetAndLengths;
+		OutTocOffsetAndLengths.SetNum(ChunkCount);
+		TocResource.ChunkHashSeeds.SetNumZeroed(ChunkCount);
+		
+		TArray<TArray<int32>> Buckets;
+		Buckets.SetNum(ChunkCount);
+
+		TBitArray FreeSlots(true, ChunkCount);
+		// Put each chunk in a bucket, each bucket contains the chunk ids that have colliding hashes
+		for (uint32 ChunkIndex = 0; ChunkIndex < ChunkCount; ++ChunkIndex)
+		{
+			const FIoChunkId& ChunkId = TocResource.ChunkIds[ChunkIndex];
+			Buckets[FIoStoreTocResource::HashChunkIdWithSeed(0, ChunkId) % ChunkCount].Add(ChunkIndex);
+		}
+		// For each bucket containing more than one chunk id find a seed that makes its chunk ids
+		// hash to unused slots in the output array
+		Algo::Sort(Buckets, [](const TArray<int32>& A, const TArray<int32>& B)
+			{
+				return A.Num() > B.Num();
+			});
+		for (uint32 BucketIndex = 0; BucketIndex < ChunkCount; ++BucketIndex)
+		{
+			const TArray<int32>& Bucket = Buckets[BucketIndex];
+			if (Bucket.Num() <= 1)
+			{
+				break;
+			}
+
+			uint32 Seed = 1;
+			TBitArray BucketUsedSlots(false, ChunkCount);
+			int32 IndexInBucket = 0;
+			while (IndexInBucket < Bucket.Num())
+			{
+				const FIoChunkId& ChunkId = TocResource.ChunkIds[Bucket[IndexInBucket]];
+				uint32 Slot = FIoStoreTocResource::HashChunkIdWithSeed(Seed, ChunkId) % ChunkCount;
+				if (!FreeSlots[Slot] || BucketUsedSlots[Slot])
+				{
+					++Seed;
+					IndexInBucket = 0;
+					BucketUsedSlots.Init(false, ChunkCount);
+				}
+				else
+				{
+					BucketUsedSlots[Slot] = true;
+					++IndexInBucket;
+				}
+			}
+
+			uint64 BucketHash = FIoStoreTocResource::HashChunkIdWithSeed(0, TocResource.ChunkIds[Bucket[0]]);
+			TocResource.ChunkHashSeeds[BucketHash % ChunkCount] = Seed;
+			for (IndexInBucket = 0; IndexInBucket < Bucket.Num(); ++IndexInBucket)
+			{
+				int32 ChunkIndex = Bucket[IndexInBucket];
+				const FIoChunkId& ChunkId = TocResource.ChunkIds[ChunkIndex];
+				uint32 Slot = FIoStoreTocResource::HashChunkIdWithSeed(Seed, ChunkId) % ChunkCount;
+				check(FreeSlots[Slot]);
+				FreeSlots[Slot] = false;
+				OutTocChunkIds[Slot] = ChunkId;
+				OutTocOffsetAndLengths[Slot] = TocResource.ChunkOffsetLengths[ChunkIndex];
+			}
+		}
+
+		// For the remaining buckets with only one chunk id put that chunk id in the first empty position in
+		// the output array and store the index as a negative seed for the bucket (-1 to allow use of slot 0)
+		TConstSetBitIterator<> FreeSlotIt(FreeSlots);
+		for (uint32 BucketIndex = 0; BucketIndex < ChunkCount; ++BucketIndex)
+		{
+			const TArray<int32>& Bucket = Buckets[BucketIndex];
+			if (Bucket.Num() == 1)
+			{
+				uint32 Slot = FreeSlotIt.GetIndex();
+				++FreeSlotIt;
+				int32 ChunkIndex = Bucket[0];
+				const FIoChunkId& ChunkId = TocResource.ChunkIds[ChunkIndex];
+				uint64 BucketHash = FIoStoreTocResource::HashChunkIdWithSeed(0, ChunkId);
+				TocResource.ChunkHashSeeds[BucketHash % ChunkCount] = -static_cast<int32>(Slot) - 1;
+				OutTocChunkIds[Slot] = ChunkId;
+				OutTocOffsetAndLengths[Slot] = TocResource.ChunkOffsetLengths[ChunkIndex];
+			}
+		}
+
+		TocResource.ChunkIds = OutTocChunkIds;
+		TocResource.ChunkOffsetLengths = OutTocOffsetAndLengths;
+	}
+
 	UE_NODISCARD TIoStatusOr<FIoStoreWriterResult> Flush()
 	{
 		TRACE_CPUPROFILER_EVENT_SCOPE(FlushContainer);
@@ -736,6 +830,8 @@ public:
 
 		FIoStoreTocResource& TocResource = Toc.GetTocResource();
 
+		GeneratePerfectHashes(TocResource);
+
 		if (ContainerSettings.IsIndexed())
 		{
 			TRACE_CPUPROFILER_EVENT_SCOPE(BuildIndex);
@@ -2017,30 +2113,43 @@ FIoStatus FIoStoreTocResource::Read(const TCHAR* TocFilePath, EIoStoreTocReadOpt
 	}
 
 	// Chunk IDs
-	const FIoChunkId* ChunkIds = reinterpret_cast<const FIoChunkId*>(TocBuffer.Get());
+	const uint8* DataPtr = TocBuffer.Get();
+	const FIoChunkId* ChunkIds = reinterpret_cast<const FIoChunkId*>(DataPtr);
 	OutTocResource.ChunkIds = MakeArrayView<FIoChunkId const>(ChunkIds, Header.TocEntryCount);
+	DataPtr += Header.TocEntryCount * sizeof(FIoChunkId);
 
 	// Chunk offsets
-	const FIoOffsetAndLength* ChunkOffsetLengths = reinterpret_cast<const FIoOffsetAndLength*>(ChunkIds + Header.TocEntryCount);
+	const FIoOffsetAndLength* ChunkOffsetLengths = reinterpret_cast<const FIoOffsetAndLength*>(DataPtr);
 	OutTocResource.ChunkOffsetLengths = MakeArrayView<FIoOffsetAndLength const>(ChunkOffsetLengths, Header.TocEntryCount);
+	DataPtr += Header.TocEntryCount * sizeof(FIoOffsetAndLength);
+
+	// Chunk hash seeds
+	if (Header.Version >= static_cast<uint8>(EIoStoreTocVersion::PerfectHash))
+	{
+		const int32* ChunkHashSeeds = reinterpret_cast<const int32*>(DataPtr);
+		OutTocResource.ChunkHashSeeds = MakeArrayView<int32 const>(ChunkHashSeeds, Header.TocEntryCount);
+		DataPtr += Header.TocEntryCount * sizeof(int32);
+	}
 
 	// Compression blocks
-	const FIoStoreTocCompressedBlockEntry* CompressionBlocks = reinterpret_cast<const FIoStoreTocCompressedBlockEntry*>(ChunkOffsetLengths + Header.TocEntryCount);
+	const FIoStoreTocCompressedBlockEntry* CompressionBlocks = reinterpret_cast<const FIoStoreTocCompressedBlockEntry*>(DataPtr);
 	OutTocResource.CompressionBlocks = MakeArrayView<FIoStoreTocCompressedBlockEntry const>(CompressionBlocks, Header.TocCompressedBlockEntryCount);
+	DataPtr += Header.TocCompressedBlockEntryCount * sizeof(FIoStoreTocCompressedBlockEntry);
 
 	// Compression methods
 	OutTocResource.CompressionMethods.Reserve(Header.CompressionMethodNameCount + 1);
 	OutTocResource.CompressionMethods.Add(NAME_None);
 
-	const ANSICHAR* AnsiCompressionMethodNames = reinterpret_cast<const ANSICHAR*>(CompressionBlocks + Header.TocCompressedBlockEntryCount);
+	const ANSICHAR* AnsiCompressionMethodNames = reinterpret_cast<const ANSICHAR*>(DataPtr);
 	for (uint32 CompressonNameIndex = 0; CompressonNameIndex < Header.CompressionMethodNameCount; CompressonNameIndex++)
 	{
 		const ANSICHAR* AnsiCompressionMethodName = AnsiCompressionMethodNames + CompressonNameIndex * Header.CompressionMethodNameLength;
 		OutTocResource.CompressionMethods.Add(FName(AnsiCompressionMethodName));
 	}
+	DataPtr += Header.CompressionMethodNameCount * Header.CompressionMethodNameLength;
 
 	// Chunk block signatures
-	const uint8* SignatureBuffer = reinterpret_cast<const uint8*>(AnsiCompressionMethodNames + Header.CompressionMethodNameCount * Header.CompressionMethodNameLength);
+	const uint8* SignatureBuffer = reinterpret_cast<const uint8*>(DataPtr);
 	const uint8* DirectoryIndexBuffer = SignatureBuffer;
 
 	const bool bIsSigned = EnumHasAnyFlags(Header.ContainerFlags, EIoContainerFlags::Signed);
@@ -2175,6 +2284,12 @@ TIoStatusOr<uint64> FIoStoreTocResource::Write(
 		return FIoStatus(EIoErrorCode::WriteError, TEXT("Failed to write chunk offsets"));
 	}
 
+	// Chunk hash seeds
+	if (!WriteArray(TocFileHandle.Get(), TocResource.ChunkHashSeeds))
+	{
+		return FIoStatus(EIoErrorCode::WriteError, TEXT("Failed to write chunk hash seeds"));
+	}
+
 	// Compression blocks
 	if (!WriteArray(TocFileHandle.Get(), TocResource.CompressionBlocks))
 	{
@@ -2243,6 +2358,18 @@ TIoStatusOr<uint64> FIoStoreTocResource::Write(
 	return TocFileHandle->Tell();
 }
 
+uint64 FIoStoreTocResource::HashChunkIdWithSeed(int32 Seed, const FIoChunkId& ChunkId)
+{
+	const uint8* Data = ChunkId.GetData();
+	const uint32 DataSize = ChunkId.GetSize();
+	uint64 Hash = Seed ? static_cast<uint64>(Seed) : 0xcbf29ce484222325;
+	for (uint32 Index = 0; Index < DataSize; ++Index)
+	{
+		Hash = (Hash * 0x00000100000001B3) ^ Data[Index];
+	}
+	return Hash;
+}
+
 void FIoStoreReader::GetFilenames(TArray<FString>& OutFileList) const
 {
 	const FIoDirectoryIndexReader& DirectoryIndex = GetDirectoryIndexReader();
diff --git a/Engine/Source/Runtime/Core/Private/IO/IoStore.h b/Engine/Source/Runtime/Core/Private/IO/IoStore.h
--- a/Engine/Source/Runtime/Core/Private/IO/IoStore.h
+++ b/Engine/Source/Runtime/Core/Private/IO/IoStore.h
@@ -14,6 +14,7 @@ enum class EIoStoreTocVersion : uint8
 	Initial,
 	DirectoryIndex,
 	PartitionSize,
+	PerfectHash,
 	LatestPlusOne,
 	Latest = LatestPlusOne - 1
 };
@@ -215,6 +216,8 @@ struct FIoStoreTocResource
 
 	TArray<FIoOffsetAndLength> ChunkOffsetLengths;
 
+	TArray<int32> ChunkHashSeeds;
+
 	TArray<FIoStoreTocCompressedBlockEntry> CompressionBlocks;
 
 	TArray<FName> CompressionMethods;
@@ -228,4 +231,6 @@ struct FIoStoreTocResource
 	UE_NODISCARD static FIoStatus Read(const TCHAR* TocFilePath, EIoStoreTocReadOptions ReadOptions, FIoStoreTocResource& OutTocResource);
 
 	UE_NODISCARD static TIoStatusOr<uint64> Write(const TCHAR* TocFilePath, FIoStoreTocResource& TocResource, const FIoContainerSettings& ContainerSettings, const FIoStoreWriterSettings& WriterSettings);
+
+	static uint64 HashChunkIdWithSeed(int32 Seed, const FIoChunkId& ChunkId);
 };
