Tasks: traced waiting, plus a minor refactoring

diff --git a/Engine/Source/Developer/TraceServices/Private/Model/TasksProfiler.cpp b/Engine/Source/Developer/TraceServices/Private/Model/TasksProfiler.cpp
--- a/Engine/Source/Developer/TraceServices/Private/Model/TasksProfiler.cpp
+++ b/Engine/Source/Developer/TraceServices/Private/Model/TasksProfiler.cpp
@@ -405,7 +405,10 @@ namespace TraceServices
 
 	const FWaitingForTasks* FTasksProvider::TryGetWaiting(const TCHAR* TimerName, uint32 ThreadId, double Timestamp) const
 	{
-		if (FCString::Strcmp(TimerName, TEXT("WaitUntilTasksComplete")) != 0 && FCString::Strcmp(TimerName, TEXT("GameThreadWaitForTask")) != 0)
+		if (FCString::Strcmp(TimerName, TEXT("WaitUntilTasksComplete")) != 0 && 
+			FCString::Strcmp(TimerName, TEXT("GameThreadWaitForTask")) != 0 &&
+			FCString::Strcmp(TimerName, TEXT("Tasks::Wait")) != 0 &&
+			FCString::Strcmp(TimerName, TEXT("Tasks::BusyWait")) != 0)
 		{
 			return nullptr;
 		}
diff --git a/Engine/Source/Runtime/Core/Private/Async/TaskTrace.cpp b/Engine/Source/Runtime/Core/Private/Async/TaskTrace.cpp
--- a/Engine/Source/Runtime/Core/Private/Async/TaskTrace.cpp
+++ b/Engine/Source/Runtime/Core/Private/Async/TaskTrace.cpp
@@ -196,6 +196,15 @@ namespace TaskTrace
 			<< WaitingStarted.Tasks(Tasks.GetData(), Tasks.Num());
 	}
 
+	FWaitingScope::FWaitingScope(FId TaskId)
+	{
+		check(bGTaskTraceInitialized);
+
+		UE_TRACE_LOG(TaskTrace, WaitingStarted, TaskChannel)
+			<< WaitingStarted.Timestamp(FPlatformTime::Cycles64())
+			<< WaitingStarted.Tasks(&TaskId, 1);
+	}
+
 	FWaitingScope::~FWaitingScope()
 	{
 		check(bGTaskTraceInitialized);
diff --git a/Engine/Source/Runtime/Core/Public/Async/TaskTrace.h b/Engine/Source/Runtime/Core/Public/Async/TaskTrace.h
--- a/Engine/Source/Runtime/Core/Public/Async/TaskTrace.h
+++ b/Engine/Source/Runtime/Core/Public/Async/TaskTrace.h
@@ -42,6 +42,7 @@ namespace TaskTrace
 	struct CORE_API FWaitingScope
 	{
 		explicit FWaitingScope(const TArray<FId>& Tasks); // waiting for given tasks completion
+		explicit FWaitingScope(FId TaskId);
 		~FWaitingScope();
 	};
 
@@ -58,6 +59,7 @@ namespace TaskTrace
 	inline void Finished(FId TaskId) {}
 	inline void Completed(FId TaskId) {}
 	inline FWaitingScope::FWaitingScope(const TArray<FId>& Tasks) {}
+	inline FWaitingScope::FWaitingScope(FId TaskId) {}
 	inline FWaitingScope::~FWaitingScope() {}
 #endif // UE_TASK_TRACE_ENABLED
 }
diff --git a/Engine/Source/Runtime/Core/Public/Tasks/Task.h b/Engine/Source/Runtime/Core/Public/Tasks/Task.h
--- a/Engine/Source/Runtime/Core/Public/Tasks/Task.h
+++ b/Engine/Source/Runtime/Core/Public/Tasks/Task.h
@@ -21,10 +21,16 @@ namespace UE { namespace Tasks
 	TTask<TInvokeResult_T<TaskBodyType>> Launch(const TCHAR* DebugName, TaskBodyType&& TaskBody, PrerequisitesCollectionType&& Prerequisites, LowLevelTasks::ETaskPriority Priority = LowLevelTasks::ETaskPriority::Normal);
 
 	template<typename TaskCollectionType>
-	bool Wait(TaskCollectionType&& Tasks, FTimespan InTimeout = FTimespan::MaxValue());
+	bool Wait(const TaskCollectionType& Tasks, FTimespan InTimeout = FTimespan::MaxValue());
 
 	template<typename TaskCollectionType>
-	bool BusyWait(TaskCollectionType&& Tasks, FTimespan InTimeout = FTimespan::MaxValue());
+	bool BusyWait(const TaskCollectionType& Tasks, FTimespan InTimeout = FTimespan::MaxValue());
+
+	namespace Private
+	{
+		template<typename TaskCollectionType>
+		TArray<TaskTrace::FId> GetTraceIds(const TaskCollectionType& Tasks);
+	}
 
 	// a common part of the generic `TTask<ResultType>` and its `TTask<void>` specialisation
 	template<typename ResultType>
@@ -37,16 +43,19 @@ namespace UE { namespace Tasks
 		friend class TPrerequisites;
 
 		template<typename TaskCollectionType>
-		friend bool Private::TryRetractAndExecute(TaskCollectionType&& Tasks, FTimespan Timeout);
+		friend bool Private::TryRetractAndExecute(const TaskCollectionType& Tasks, FTimespan Timeout);
+
+		template<typename TaskCollectionType>
+		friend TArray<TaskTrace::FId> Private::GetTraceIds(const TaskCollectionType& Tasks);
 
 		template<typename TaskType>
-		friend void AddNested(TaskType&& Nested);
+		friend void AddNested(const TaskType& Nested);
 
 		template<typename TaskCollectionType>
-		friend bool Wait(TaskCollectionType&& Tasks, FTimespan InTimeout/* = FTimespan::MaxValue()*/);
+		friend bool Wait(const TaskCollectionType& Tasks, FTimespan InTimeout/* = FTimespan::MaxValue()*/);
 
 		template<typename TaskCollectionType>
-		friend bool BusyWait(TaskCollectionType&& Tasks, FTimespan InTimeout/* = FTimespan::MaxValue()*/);
+		friend bool BusyWait(const TaskCollectionType& Tasks, FTimespan InTimeout/* = FTimespan::MaxValue()*/);
 
 	protected:
 		TTaskBase() = default;
@@ -72,7 +81,15 @@ namespace UE { namespace Tasks
 		// @return true if the task is completed
 		bool Wait(FTimespan Timeout = FTimespan::MaxValue())
 		{
-			if (!IsValid() || Pimpl->TryRetractAndExecute())
+			if (!IsValid())
+			{
+				return true;
+			}
+
+			TaskTrace::FWaitingScope WaitingScope(Pimpl->GetTraceId());
+			TRACE_CPUPROFILER_EVENT_SCOPE(Tasks::Wait);
+
+			if (Pimpl->TryRetractAndExecute())
 			{
 				return true;
 			}
@@ -99,6 +116,8 @@ namespace UE { namespace Tasks
 		{
 			if (IsValid())
 			{
+				TaskTrace::FWaitingScope WaitingScope(Pimpl->GetTraceId());
+				TRACE_CPUPROFILER_EVENT_SCOPE(Tasks::BusyWait);
 				Pimpl->BusyWait();
 			}
 		}
@@ -108,7 +127,13 @@ namespace UE { namespace Tasks
 		// @return true if the task is completed
 		bool BusyWait(FTimespan Timeout)
 		{
-			return !IsValid() || Pimpl->BusyWait(Timeout);
+			if (IsValid())
+			{
+				TaskTrace::FWaitingScope WaitingScope(Pimpl->GetTraceId());
+				return Pimpl->BusyWait(Timeout);
+			}
+
+			return true;
 		}
 
 		// waits for task's completion or the given condition becomes true, while executing other tasks.
@@ -117,7 +142,14 @@ namespace UE { namespace Tasks
 		template<typename ConditionType>
 		bool BusyWait(ConditionType&& Condition)
 		{
-			return !IsValid() || Pimpl->BusyWait(Forward<ConditionType>(Condition));
+			if (IsValid())
+			{
+				TaskTrace::FWaitingScope WaitingScope(Pimpl->GetTraceId());
+				TRACE_CPUPROFILER_EVENT_SCOPE(Tasks::BusyWait);
+				return Pimpl->BusyWait(Forward<ConditionType>(Condition));
+			}
+
+			return true;
 		}
 
 	protected:
@@ -256,10 +288,38 @@ namespace UE { namespace Tasks
 		return TTask<FResult>{ Task };
 	}
 
+	namespace Private
+	{
+		template<typename TaskCollectionType>
+		TArray<TaskTrace::FId> GetTraceIds(const TaskCollectionType& Tasks)
+		{
+#if UE_TASK_TRACE_ENABLED
+			TArray<TaskTrace::FId> TasksIds;
+			TasksIds.Reserve(Tasks.Num());
+
+			for (auto& Task : Tasks)
+			{
+				if (Task.IsValid())
+				{
+					TasksIds.Add(Task.Pimpl->GetTraceId());
+				}
+			}
+
+			return TasksIds;
+#else
+			return {};
+#endif
+		}
+	}
+
 	// wait for multiple tasks, with optional timeout
+	// @param TaskCollectionType - an iterable collection of `TTask<T>`, e.g. `TArray<FTask>`
 	template<typename TaskCollectionType>
-	bool Wait(TaskCollectionType&& Tasks, FTimespan InTimeout/* = FTimespan::MaxValue()*/)
+	bool Wait(const TaskCollectionType& Tasks, FTimespan InTimeout/* = FTimespan::MaxValue()*/)
 	{
+		TaskTrace::FWaitingScope WaitingScope(Private::GetTraceIds(Tasks));
+		TRACE_CPUPROFILER_EVENT_SCOPE(Tasks::Wait);
+
 		FTimeout Timeout{ InTimeout };
 
 		if (Private::TryRetractAndExecute(Tasks, InTimeout))
@@ -286,8 +346,11 @@ namespace UE { namespace Tasks
 
 	// wait for multiple tasks while executing other tasks
 	template<typename TaskCollectionType>
-	bool BusyWait(TaskCollectionType&& Tasks, FTimespan TimeoutValue/* = FTimespan::MaxValue()*/)
+	bool BusyWait(const TaskCollectionType& Tasks, FTimespan TimeoutValue/* = FTimespan::MaxValue()*/)
 	{
+		TaskTrace::FWaitingScope WaitingScope(Private::GetTraceIds(Tasks));
+		TRACE_CPUPROFILER_EVENT_SCOPE(Tasks::BusyWait);
+
 		FTimeout Timeout{ TimeoutValue };
 
 		if (Private::TryRetractAndExecute(Tasks, TimeoutValue))
@@ -344,7 +407,7 @@ namespace UE { namespace Tasks
 	// until all nested tasks are completed. It's similar to explicitly waiting for a sub-task at the end of its parent task, except explicit waiting
 	// blocks the worker executing the parent task until the sub-task is completed. With nested tasks, the worker won't be blocked.
 	template<typename TaskType>
-	void AddNested(TaskType&& Nested)
+	void AddNested(const TaskType& Nested)
 	{
 		Private::FTaskBase* Parent = Private::FTaskBase::GetCurrentTask();
 		check(Parent != nullptr);
diff --git a/Engine/Source/Runtime/Core/Public/Tasks/TaskPrivate.h b/Engine/Source/Runtime/Core/Public/Tasks/TaskPrivate.h
--- a/Engine/Source/Runtime/Core/Public/Tasks/TaskPrivate.h
+++ b/Engine/Source/Runtime/Core/Public/Tasks/TaskPrivate.h
@@ -590,14 +590,14 @@ namespace UE { namespace Tasks
 		};
 
 		template<typename TaskCollectionType>
-		bool TryRetractAndExecute(TaskCollectionType&& Tasks, FTimespan InTimeout = FTimespan::MaxValue())
+		bool TryRetractAndExecute(const TaskCollectionType& Tasks, FTimespan InTimeout = FTimespan::MaxValue())
 		{
 			FTimeout Timeout{ InTimeout };
 			bool bResult = true;
 
 			for (auto& Task : Tasks)
 			{
-				if (!Task.Pimpl->TryRetractAndExecute())
+				if (Task.IsValid() && !Task.Pimpl->TryRetractAndExecute())
 				{
 					bResult = false;
 				}
