#ue5 - LWC: don't assume TQuat<> is 16 or 32-byte aligned, some people cast pointers to floats or doubles and then use VectorLoad (unaligned) on them.

diff --git a/Engine/Source/Runtime/Core/Public/Math/UnrealMathVectorCommon.h b/Engine/Source/Runtime/Core/Public/Math/UnrealMathVectorCommon.h
--- a/Engine/Source/Runtime/Core/Public/Math/UnrealMathVectorCommon.h
+++ b/Engine/Source/Runtime/Core/Public/Math/UnrealMathVectorCommon.h
@@ -32,7 +32,7 @@ FORCEINLINE VectorRegister4Float VectorOne(void)
 template<typename T>
 FORCEINLINE TVectorRegisterType<T> VectorLoad(const UE::Math::TQuat<T>* Ptr)
 {
-	return VectorLoadAligned((const T*)(Ptr));
+	return VectorLoad((const T*)(Ptr));
 }
 
 FORCEINLINE VectorRegister4Float VectorLoad(const FVector4* Ptr)
