Collect process name & id for threads in context switches capture for Insights

diff --git a/Engine/Source/Developer/TraceServices/Private/Analyzers/PlatformEventTraceAnalysis.cpp b/Engine/Source/Developer/TraceServices/Private/Analyzers/PlatformEventTraceAnalysis.cpp
--- a/Engine/Source/Developer/TraceServices/Private/Analyzers/PlatformEventTraceAnalysis.cpp
+++ b/Engine/Source/Developer/TraceServices/Private/Analyzers/PlatformEventTraceAnalysis.cpp
@@ -22,6 +22,7 @@ void FPlatformEventTraceAnalyzer::OnAnalysisBegin(const FOnAnalysisContext& Cont
 
 	Builder.RouteEvent(RouteId_ContextSwitch, "PlatformEvent", "ContextSwitch");
 	Builder.RouteEvent(RouteId_StackSample, "PlatformEvent", "StackSample");
+	Builder.RouteEvent(RouteId_ThreadName, "PlatformEvent", "ThreadName");
 }
 
 bool FPlatformEventTraceAnalyzer::OnEvent(uint16 RouteId, EStyle Style, const FOnEventContext& Context)
@@ -57,6 +58,18 @@ bool FPlatformEventTraceAnalyzer::OnEvent(uint16 RouteId, EStyle Style, const FO
 		break;
 	}
 
+	case RouteId_ThreadName:
+	{
+		uint32 ThreadId = EventData.GetValue<uint32>("ThreadId");
+		uint32 ProcessId = EventData.GetValue<uint32>("ProcessId");
+		FStringView Name;
+		if (EventData.GetString("Name", Name))
+		{
+			ContextSwitchesProvider.AddThreadName(ThreadId, ProcessId, Name);
+		}
+		break;
+	}
+
 	}
 
 	return true;
diff --git a/Engine/Source/Developer/TraceServices/Private/Analyzers/PlatformEventTraceAnalysis.h b/Engine/Source/Developer/TraceServices/Private/Analyzers/PlatformEventTraceAnalysis.h
--- a/Engine/Source/Developer/TraceServices/Private/Analyzers/PlatformEventTraceAnalysis.h
+++ b/Engine/Source/Developer/TraceServices/Private/Analyzers/PlatformEventTraceAnalysis.h
@@ -28,6 +28,7 @@ private:
 	{
 		RouteId_ContextSwitch,
 		RouteId_StackSample,
+		RouteId_ThreadName,
 	};
 
 	IAnalysisSession& Session;
diff --git a/Engine/Source/Developer/TraceServices/Private/Model/ContextSwitches.cpp b/Engine/Source/Developer/TraceServices/Private/Model/ContextSwitches.cpp
--- a/Engine/Source/Developer/TraceServices/Private/Model/ContextSwitches.cpp
+++ b/Engine/Source/Developer/TraceServices/Private/Model/ContextSwitches.cpp
@@ -337,6 +337,11 @@ void FContextSwitchesProvider::AddThreadInfo(uint32 ThreadId, uint32 SystemThrea
 	SystemToTraceThreadIdMap.Add(SystemThreadId, ThreadId);
 }
 
+void FContextSwitchesProvider::AddThreadName(uint32 SystemTreadId, uint32 SystemProcessId, FStringView Name)
+{
+	//
+}
+
 const IContextSwitchesProvider* ReadContextSwitchesProvider(const IAnalysisSession& Session)
 {
 	return Session.ReadProvider<IContextSwitchesProvider>(FContextSwitchesProvider::ProviderName);
diff --git a/Engine/Source/Developer/TraceServices/Private/Model/ContextSwitchesPrivate.h b/Engine/Source/Developer/TraceServices/Private/Model/ContextSwitchesPrivate.h
--- a/Engine/Source/Developer/TraceServices/Private/Model/ContextSwitchesPrivate.h
+++ b/Engine/Source/Developer/TraceServices/Private/Model/ContextSwitchesPrivate.h
@@ -35,6 +35,8 @@ public:
 	void Add(uint32 SystemThreadId, double Start, double End, uint32 CoreNumber);
 	void AddThreadInfo(uint32 ThreadId, uint32 SystemThreadId);
 
+	void AddThreadName(uint32 SystemTreadId, uint32 SystemProcessId, FStringView Name);
+
 private:
 	const TPagedArray<FContextSwitch>* GetContextSwitches(uint32 ThreadId) const;
 
diff --git a/Engine/Source/Runtime/Core/Private/ProfilingDebugging/Microsoft/EventTracingForWindows.cpp b/Engine/Source/Runtime/Core/Private/ProfilingDebugging/Microsoft/EventTracingForWindows.cpp
--- a/Engine/Source/Runtime/Core/Private/ProfilingDebugging/Microsoft/EventTracingForWindows.cpp
+++ b/Engine/Source/Runtime/Core/Private/ProfilingDebugging/Microsoft/EventTracingForWindows.cpp
@@ -5,6 +5,7 @@
 #include "HAL/Event.h"
 #include "HAL/RunnableThread.h"
 #include "Logging/LogMacros.h"
+#include "Experimental/Containers/SherwoodHashTable.h"
 #include "Trace/Trace.h"
 
 #if PLATFORM_SUPPORTS_PLATFORM_EVENTS
@@ -57,6 +58,9 @@ private:
 
 	// timestamp for each core when thread start executing on it
 	uint64 ContexSwitchStart[256];
+
+	// thread ids that have their info already sent
+	Experimental::TSherwoodSet<uint32> ThreadNameSet;
 };
 
 // https://docs.microsoft.com/en-us/windows/win32/etw/cswitch
@@ -104,7 +108,6 @@ uint32 FPlatformEvents::CurrentProcessId;
 
 void FPlatformEvents::TraceEventCallback(PEVENT_RECORD Event)
 {
-	uint32 ProcessId = Event->EventHeader.ProcessId;
 	int Opcode = Event->EventHeader.EventDescriptor.Opcode;
 	if (Opcode == FContextSwitchEvent::Type)
 	{
@@ -130,6 +133,43 @@ void FPlatformEvents::TraceEventCallback(PEVENT_RECORD Event)
 					<< ContextSwitch.EndTime(EndTime)
 					<< ContextSwitch.ThreadId(ThreadId)
 					<< ContextSwitch.CoreNumber(CoreNumber);
+
+				bool bAlreadyAdded = false;
+				GPlatformEvents->ThreadNameSet.Add(ThreadId, &bAlreadyAdded);
+				if (!bAlreadyAdded)
+				{
+					WCHAR Name[4096];
+					DWORD NameLen = 0;
+					uint32 ProcessId = 0;
+
+					HANDLE ThreadHandle = ::OpenThread(THREAD_QUERY_LIMITED_INFORMATION, false, ThreadId);
+					if (ThreadHandle)
+					{
+						ProcessId = ::GetProcessIdOfThread(ThreadHandle);
+						::CloseHandle(ThreadHandle);
+					}
+
+					// collect names only about threads not belonging to our process
+					if (ProcessId && ProcessId != CurrentProcessId)
+					{
+						HANDLE ProcessHandle = ::OpenProcess(PROCESS_QUERY_LIMITED_INFORMATION, false, ProcessId);
+						if (ProcessHandle)
+						{
+							NameLen = sizeof(Name) / sizeof(*Name);
+							if (::QueryFullProcessImageNameW(ProcessHandle, 0, Name, &NameLen) == 0)
+							{
+								// failed to get name of process executable
+								NameLen = 0;
+							}
+							::CloseHandle(ProcessHandle);
+						}
+
+						UE_TRACE_LOG(PlatformEvent, ThreadName, ContextSwitchChannel, NameLen * sizeof(TCHAR))
+							<< ThreadName.ThreadId(ThreadId)
+							<< ThreadName.ProcessId(ProcessId)
+							<< ThreadName.Name(Name, NameLen);
+					}
+				}
 			}
 		}
 
@@ -333,6 +373,8 @@ void FPlatformEvents::Stop()
 		StopETW();
 		Thread->WaitForCompletion();
 	}
+
+	ThreadNameSet.Empty();
 }
 
 /////////////////////////////////////////////////////////////////////
diff --git a/Engine/Source/Runtime/Core/Private/ProfilingDebugging/PlatformEvents.cpp b/Engine/Source/Runtime/Core/Private/ProfilingDebugging/PlatformEvents.cpp
--- a/Engine/Source/Runtime/Core/Private/ProfilingDebugging/PlatformEvents.cpp
+++ b/Engine/Source/Runtime/Core/Private/ProfilingDebugging/PlatformEvents.cpp
@@ -10,6 +10,7 @@ UE_TRACE_CHANNEL_DEFINE(StackSamplingChannel)
 
 UE_TRACE_EVENT_DEFINE(PlatformEvent, ContextSwitch)
 UE_TRACE_EVENT_DEFINE(PlatformEvent, StackSample)
+UE_TRACE_EVENT_DEFINE(PlatformEvent, ThreadName)
 
 /////////////////////////////////////////////////////////////////////
 
diff --git a/Engine/Source/Runtime/Core/Private/ProfilingDebugging/PlatformEvents.h b/Engine/Source/Runtime/Core/Private/ProfilingDebugging/PlatformEvents.h
--- a/Engine/Source/Runtime/Core/Private/ProfilingDebugging/PlatformEvents.h
+++ b/Engine/Source/Runtime/Core/Private/ProfilingDebugging/PlatformEvents.h
@@ -34,6 +34,16 @@ UE_TRACE_EVENT_BEGIN_EXTERN(PlatformEvent, StackSample, NoSync)
 	UE_TRACE_EVENT_FIELD(uint64[], Addresses)
 UE_TRACE_EVENT_END()
 
+// thread information for context switches
+// NOTE: in some cases name of process can be empty, for example when there are no
+// privileges to query it or there are no name for process thread belongs to.
+// Depending on platform actual name can be absolute path to process executable.
+UE_TRACE_EVENT_BEGIN_EXTERN(PlatformEvent, ThreadName, NoSync)
+	UE_TRACE_EVENT_FIELD(uint32, ThreadId)
+	UE_TRACE_EVENT_FIELD(uint32, ProcessId)
+	UE_TRACE_EVENT_FIELD(UE::Trace::WideString, Name)
+UE_TRACE_EVENT_END()
+
 /////////////////////////////////////////////////////////////////////
 
 enum class EPlatformEvent
