#ue5 - LWC: AVX fix for VectorMatrixMultiply().

- Added FMatrix4d validation tests to UnrealMathTest.cpp startup tests.
- Avoided MSVC compiler issue with AVX casts, and switched AVX2 swizzles to more exclusive use of permute2f128.
- Shrinked some code blocks in the VectorMatrixMultiply() function which can fix compiler bugs in some cases for register re-use (it did in CL 16042174 for example).
- Re-enabled AVX1 and AVX2 as this is confirmed to fix the regression on console.

diff --git a/Engine/Source/Runtime/Core/Private/Tests/Math/UnrealMathTest.cpp b/Engine/Source/Runtime/Core/Private/Tests/Math/UnrealMathTest.cpp
--- a/Engine/Source/Runtime/Core/Private/Tests/Math/UnrealMathTest.cpp
+++ b/Engine/Source/Runtime/Core/Private/Tests/Math/UnrealMathTest.cpp
@@ -331,6 +331,25 @@ bool TestMatricesEqual( FMatrix44f &Mat0, FMatrix44f &Mat1, float Tolerance = 0.
 	return true;
 }
 
+bool TestMatricesEqual(FMatrix44d& Mat0, FMatrix44d& Mat1, double Tolerance = 0.0f)
+{
+	for (int32 Row = 0; Row < 4; ++Row)
+	{
+		GSumDouble = 0.f;
+		for (int32 Column = 0; Column < 4; ++Column)
+		{
+			double Diff = Mat0.M[Row][Column] - Mat1.M[Row][Column];
+			GSumDouble += FMath::Abs(Diff);
+		}
+		if (GSumDouble > Tolerance)
+		{
+			CheckPassing(false);
+			return false;
+		}
+	}
+	return true;
+}
+
 
 /**
  * Multiplies two 4x4 matrices.
@@ -486,6 +505,24 @@ VectorRegister4Float TestVectorTransformVector(const VectorRegister4Float&  VecP
 	return Result.v;
 }
 
+VectorRegister4Double TestVectorTransformVector(const VectorRegister4Double& VecP, const FMatrix44d* MatrixM)
+{
+	typedef double Float4x4[4][4];
+	union U {
+		VectorRegister4Double v; double f[4];
+		FORCEINLINE U() : v() {}
+	} Tmp, Result;
+	Tmp.v = VecP;
+	const Float4x4& M = *((const Float4x4*)MatrixM);
+
+	Result.f[0] = Tmp.f[0] * M[0][0] + Tmp.f[1] * M[1][0] + Tmp.f[2] * M[2][0] + Tmp.f[3] * M[3][0];
+	Result.f[1] = Tmp.f[0] * M[0][1] + Tmp.f[1] * M[1][1] + Tmp.f[2] * M[2][1] + Tmp.f[3] * M[3][1];
+	Result.f[2] = Tmp.f[0] * M[0][2] + Tmp.f[1] * M[1][2] + Tmp.f[2] * M[2][2] + Tmp.f[3] * M[3][2];
+	Result.f[3] = Tmp.f[0] * M[0][3] + Tmp.f[1] * M[1][3] + Tmp.f[2] * M[2][3] + Tmp.f[3] * M[3][3];
+
+	return Result.v;
+}
+
 /**
 * Get Rotation as a quaternion.
 * @param Rotator FRotator 
@@ -1829,6 +1866,54 @@ bool RunDoubleVectorTest()
 	LogTest<double>(TEXT("VectorMaskBits"), MaskBits == 5);
 
 
+	FMatrix44d	M0, M1, M2, M3;
+	FVector3d Eye, LookAt, Up;
+	// Create Look at Matrix
+	Eye = FVector3d(1024.0f, -512.0f, -2048.0f);
+	LookAt = FVector3d(0.0f, 0.0f, 0.0f);
+	Up = FVector3d(0.0f, 1.0f, 0.0f);
+	M0 = FLookAtMatrix(Eye, LookAt, Up);
+
+	// Create GL ortho projection matrix
+	const double Width = 1920.0f;
+	const double Height = 1080.0f;
+	const double Left = 0.0f;
+	const double Right = Left + Width;
+	const double Top = 0.0f;
+	const double Bottom = Top + Height;
+	const double ZNear = -100.0f;
+	const double ZFar = 100.0f;
+
+	M1 = FMatrix44d(FPlane4d(2.0f / (Right - Left), 0, 0, 0),
+		FPlane4d(0, 2.0f / (Top - Bottom), 0, 0),
+		FPlane4d(0, 0, 1 / (ZNear - ZFar), 0),
+		FPlane4d((Left + Right) / (Left - Right), (Top + Bottom) / (Bottom - Top), ZNear / (ZNear - ZFar), 1));
+
+	VectorMatrixMultiply(&M2, &M0, &M1);
+	TestVectorMatrixMultiply(&M3, &M0, &M1);
+	LogTest<double>(TEXT("VectorMatrixMultiply"), TestMatricesEqual(M2, M3, 0.000001f));
+
+	VectorMatrixInverse(&M2, &M1);
+	TestVectorMatrixInverse(&M3, &M1);
+	LogTest<double>(TEXT("VectorMatrixInverse"), TestMatricesEqual(M2, M3, 0.000001f));
+	
+	// 	FTransform Transform;
+	// 	Transform.SetFromMatrix(M1);
+	// 	FTransform InvTransform = Transform.Inverse();
+	// 	FTransform InvTransform2 = FTransform(Transform.ToMatrixWithScale().Inverse());
+	// 	LogTest<double>( TEXT("FTransform Inverse"), InvTransform.Equals(InvTransform2, 1e-3f ) );
+
+	V0 = MakeVectorRegister(100.0f, -100.0f, 200.0f, 1.0f);
+	V1 = VectorTransformVector(V0, &M0);
+	V2 = TestVectorTransformVector(V0, &M0);
+	LogTest<double>(TEXT("VectorTransformVector"), TestVectorsEqual(V1, V2, 1e-8f));
+
+	V0 = MakeVectorRegister(32768.0f, 131072.0f, -8096.0f, 1.0f);
+	V1 = VectorTransformVector(V0, &M1);
+	V2 = TestVectorTransformVector(V0, &M1);
+	LogTest<double>(TEXT("VectorTransformVector"), TestVectorsEqual(V1, V2, 1e-8f));
+
+
 	// NaN / Inf tests
 	// Using a union as we need to do a bitwise cast of 0xFFFFFFFF into a float.
 	typedef union
@@ -2376,7 +2461,7 @@ bool FVectorRegisterAbstractionTest::RunTest(const FString& Parameters)
 	LogTest<float>(TEXT("VectorMaskBits"), MaskBits == 5);
 
 
-	FMatrix44f	M0, M1, M2, M3;		// LWC_TODO: Need double version tests!
+	FMatrix44f	M0, M1, M2, M3;
 	FVector3f Eye, LookAt, Up;	
 	// Create Look at Matrix
 	Eye    = FVector3f(1024.0f, -512.0f, -2048.0f);
@@ -2416,12 +2501,12 @@ bool FVectorRegisterAbstractionTest::RunTest(const FString& Parameters)
 	V0 = MakeVectorRegister( 100.0f, -100.0f, 200.0f, 1.0f );
 	V1 = VectorTransformVector(V0, &M0);
 	V2 = TestVectorTransformVector(V0, &M0);
-	LogTest<float>( TEXT("VectorTransformVector"), TestVectorsEqual( V1, V2 ) );
+	LogTest<float>( TEXT("VectorTransformVector"), TestVectorsEqual( V1, V2, 1e-8f ) );
 
 	V0 = MakeVectorRegister( 32768.0f,131072.0f, -8096.0f, 1.0f );
 	V1 = VectorTransformVector(V0, &M1);
 	V2 = TestVectorTransformVector(V0, &M1);
-	LogTest<float>( TEXT("VectorTransformVector"), TestVectorsEqual( V1, V2 ) );
+	LogTest<float>( TEXT("VectorTransformVector"), TestVectorsEqual( V1, V2, 1e-8f ) );
 
 	// NaN / Inf tests
 	// Using a union as we need to do a bitwise cast of 0xFFFFFFFF into a float.
diff --git a/Engine/Source/Runtime/Core/Public/Math/UnrealMathSSE.h b/Engine/Source/Runtime/Core/Public/Math/UnrealMathSSE.h
--- a/Engine/Source/Runtime/Core/Public/Math/UnrealMathSSE.h
+++ b/Engine/Source/Runtime/Core/Public/Math/UnrealMathSSE.h
@@ -10,7 +10,7 @@
 #endif
 
 #ifndef UE_PLATFORM_MATH_USE_AVX
-#define UE_PLATFORM_MATH_USE_AVX			0 // PLATFORM_ALWAYS_HAS_AVX
+#define UE_PLATFORM_MATH_USE_AVX			PLATFORM_ALWAYS_HAS_AVX
 #endif
 
 #ifndef UE_PLATFORM_MATH_USE_AVX_2
@@ -97,7 +97,7 @@ struct alignas(UE_SSE_DOUBLE_ALIGNMENT) VectorRegister4Double
 	};
 
 	// Use in preference when reading XY or ZW to extract values, it's better on MSVC than the generated memory reads.
-	FORCEINLINE VectorRegister2Double GetXY() const { return _mm256_castpd256_pd128(XYZW); }
+	FORCEINLINE VectorRegister2Double GetXY() const { return _mm256_extractf128_pd(XYZW, 0); } // { return _mm256_castpd256_pd128(XYZW); } // Possible MSVC compiler bug in optimized bugs when using this cast, but can be more efficient.
 	FORCEINLINE VectorRegister2Double GetZW() const { return _mm256_extractf128_pd(XYZW, 1); }
 #endif
 
@@ -838,10 +838,10 @@ namespace SSEPermuteHelpers
 
 	// Identity
 	template<> FORCEINLINE VectorRegister4Double PermuteLanes<0, 1>(const VectorRegister4Double& Vec) { return Vec; }
-	template<> FORCEINLINE VectorRegister4Double PermuteLanes<0, 0>(const VectorRegister4Double& Vec) { return _mm256_insertf128_pd(Vec, Vec.GetXY(), 1); } // copy XY to lane 1
 #if !UE_PLATFORM_MATH_USE_AVX_2
 	// On AVX1, permute2f128 can be quite slow, so look for alternatives (extract + insert). On AVX2, permute2f128 is more efficient and should equal or beat (extract + insert).
 	// Sources: https://www.agner.org/optimize/instruction_tables.pdf, https://uops.info/table.html
+	template<> FORCEINLINE VectorRegister4Double PermuteLanes<0, 0>(const VectorRegister4Double& Vec) { return _mm256_insertf128_pd(Vec, Vec.GetXY(), 1); } // copy XY to lane 1
 	template<> FORCEINLINE VectorRegister4Double PermuteLanes<1, 0>(const VectorRegister4Double& Vec) { return _mm256_insertf128_pd(_mm256_castpd128_pd256(Vec.GetZW()), Vec.GetXY(), 1); } // copy XY to lane 1
 	template<> FORCEINLINE VectorRegister4Double PermuteLanes<1, 1>(const VectorRegister4Double& Vec) { return _mm256_insertf128_pd(Vec, Vec.GetZW(), 0); } // copy ZW to lane 0
 #endif // !AVX2
@@ -986,10 +986,10 @@ namespace SSEPermuteHelpers
 
 	// Lane shuffle helper specialization
 	template<> FORCEINLINE VectorRegister4Double ShuffleLanes<0, 1>(const VectorRegister4Double& Vec1, const VectorRegister4Double& Vec2) { return _mm256_blend_pd(Vec1, Vec2, 0b1100); }
-	template<> FORCEINLINE VectorRegister4Double ShuffleLanes<0, 0>(const VectorRegister4Double& Vec1, const VectorRegister4Double& Vec2) { return _mm256_insertf128_pd(Vec1, Vec2.GetXY(), 1); } // copy XY to lane 1
 #if !UE_PLATFORM_MATH_USE_AVX_2
 	// On AVX1, permute2f128 can be quite slow, so look for alternatives (extract + insert). On AVX2, permute2f128 is more efficient and should equal or beat (extract + insert).
 	// Sources: https://www.agner.org/optimize/instruction_tables.pdf, https://uops.info/table.html
+	template<> FORCEINLINE VectorRegister4Double ShuffleLanes<0, 0>(const VectorRegister4Double& Vec1, const VectorRegister4Double& Vec2) { return _mm256_insertf128_pd(Vec1, Vec2.GetXY(), 1); } // copy XY to lane 1
 	template<> FORCEINLINE VectorRegister4Double ShuffleLanes<1, 0>(const VectorRegister4Double& Vec1, const VectorRegister4Double& Vec2) { return _mm256_insertf128_pd(_mm256_castpd128_pd256(Vec1.GetZW()), Vec2.GetXY(), 1); } // copy XY to lane 1
 	template<> FORCEINLINE VectorRegister4Double ShuffleLanes<1, 1>(const VectorRegister4Double& Vec1, const VectorRegister4Double& Vec2) { return _mm256_insertf128_pd(Vec2, Vec1.GetZW(), 0); } // copy ZW to lane 0
 #endif // !AVX2
@@ -2127,7 +2127,7 @@ FORCEINLINE void VectorMatrixMultiply(FMatrix44d* Result, const FMatrix44d* Matr
 	const Double4x4& AMRows = *((const Double4x4*)Matrix1);
 	const Double4x4& BMRows = *((const Double4x4*)Matrix2);
 	Double4x4& ResultDst = *((Double4x4*)Result);
-	VectorRegister4Double Temp, ARow;
+	VectorRegister4Double Temp;
 
 	VectorRegister4Double B[4];
 	B[0] = VectorLoad(BMRows[0]);
@@ -2136,36 +2136,44 @@ FORCEINLINE void VectorMatrixMultiply(FMatrix44d* Result, const FMatrix44d* Matr
 	B[3] = VectorLoad(BMRows[3]);
 
 	// First row of result (Matrix1[0] * Matrix2).
-	ARow = VectorLoad(AMRows[0]);
-	Temp = VectorMultiply(VectorReplicate(ARow, 0), B[0]);
-	Temp = VectorMultiplyAdd(VectorReplicate(ARow, 1), B[1], Temp);
-	Temp = VectorMultiplyAdd(VectorReplicate(ARow, 2), B[2], Temp);
-	Temp = VectorMultiplyAdd(VectorReplicate(ARow, 3), B[3], Temp);
-	VectorStore(Temp, ResultDst[0]);
+	{
+		VectorRegister4Double ARow = VectorLoad(AMRows[0]);
+		Temp = VectorMultiply(VectorReplicate(ARow, 0), B[0]);
+		Temp = VectorMultiplyAdd(VectorReplicate(ARow, 1), B[1], Temp);
+		Temp = VectorMultiplyAdd(VectorReplicate(ARow, 2), B[2], Temp);
+		Temp = VectorMultiplyAdd(VectorReplicate(ARow, 3), B[3], Temp);
+		VectorStore(Temp, ResultDst[0]);
+	}
 
 	// Second row of result (Matrix1[1] * Matrix2).
-	ARow = VectorLoad(AMRows[1]);
-	Temp = VectorMultiply(VectorReplicate(ARow, 0), B[0]);
-	Temp = VectorMultiplyAdd(VectorReplicate(ARow, 1), B[1], Temp);
-	Temp = VectorMultiplyAdd(VectorReplicate(ARow, 2), B[2], Temp);
-	Temp = VectorMultiplyAdd(VectorReplicate(ARow, 3), B[3], Temp);
-	VectorStore(Temp, ResultDst[1]);
+	{
+		VectorRegister4Double ARow = VectorLoad(AMRows[1]);
+		Temp = VectorMultiply(VectorReplicate(ARow, 0), B[0]);
+		Temp = VectorMultiplyAdd(VectorReplicate(ARow, 1), B[1], Temp);
+		Temp = VectorMultiplyAdd(VectorReplicate(ARow, 2), B[2], Temp);
+		Temp = VectorMultiplyAdd(VectorReplicate(ARow, 3), B[3], Temp);
+		VectorStore(Temp, ResultDst[1]);
+	}
 
 	// Third row of result (Matrix1[2] * Matrix2).
-	ARow = VectorLoad(AMRows[2]);
-	Temp = VectorMultiply(VectorReplicate(ARow, 0), B[0]);
-	Temp = VectorMultiplyAdd(VectorReplicate(ARow, 1), B[1], Temp);
-	Temp = VectorMultiplyAdd(VectorReplicate(ARow, 2), B[2], Temp);
-	Temp = VectorMultiplyAdd(VectorReplicate(ARow, 3), B[3], Temp);
-	VectorStore(Temp, ResultDst[2]);
+	{
+		VectorRegister4Double ARow = VectorLoad(AMRows[2]);
+		Temp = VectorMultiply(VectorReplicate(ARow, 0), B[0]);
+		Temp = VectorMultiplyAdd(VectorReplicate(ARow, 1), B[1], Temp);
+		Temp = VectorMultiplyAdd(VectorReplicate(ARow, 2), B[2], Temp);
+		Temp = VectorMultiplyAdd(VectorReplicate(ARow, 3), B[3], Temp);
+		VectorStore(Temp, ResultDst[2]);
+	}
 
 	// Fourth row of result (Matrix1[3] * Matrix2).
-	ARow = VectorLoad(AMRows[3]);
-	Temp = VectorMultiply(VectorReplicate(ARow, 0), B[0]);
-	Temp = VectorMultiplyAdd(VectorReplicate(ARow, 1), B[1], Temp);
-	Temp = VectorMultiplyAdd(VectorReplicate(ARow, 2), B[2], Temp);
-	Temp = VectorMultiplyAdd(VectorReplicate(ARow, 3), B[3], Temp);
-	VectorStore(Temp, ResultDst[3]);
+	{
+		VectorRegister4Double ARow = VectorLoad(AMRows[3]);
+		Temp = VectorMultiply(VectorReplicate(ARow, 0), B[0]);
+		Temp = VectorMultiplyAdd(VectorReplicate(ARow, 1), B[1], Temp);
+		Temp = VectorMultiplyAdd(VectorReplicate(ARow, 2), B[2], Temp);
+		Temp = VectorMultiplyAdd(VectorReplicate(ARow, 3), B[3], Temp);
+		VectorStore(Temp, ResultDst[3]);
+	}
 }
 
 
