Tasks: moved waiting logic one level down to Private::FTaskBase so when users need fine control of task lifetime, they can use Private::FTaskBase and have full functionality. Planned to be used by Oodle integration.

diff --git a/Engine/Source/Runtime/Core/Public/Tasks/Task.h b/Engine/Source/Runtime/Core/Public/Tasks/Task.h
--- a/Engine/Source/Runtime/Core/Public/Tasks/Task.h
+++ b/Engine/Source/Runtime/Core/Public/Tasks/Task.h
@@ -81,34 +81,7 @@ namespace UE { namespace Tasks
 		// @return true if the task is completed
 		bool Wait(FTimespan Timeout = FTimespan::MaxValue())
 		{
-			if (!IsValid())
-			{
-				return true;
-			}
-
-			TaskTrace::FWaitingScope WaitingScope(Pimpl->GetTraceId());
-			TRACE_CPUPROFILER_EVENT_SCOPE(Tasks::Wait);
-
-			if (Pimpl->TryRetractAndExecute())
-			{
-				return true;
-			}
-
-			// the event must be alive for the task and this function lifetime, we don't know which one will be finished first as waiting can time out
-			// before the waiting task is completed
-			FSharedEventRef CompletionEvent;
-
-			TRefCountPtr<Private::TTaskWithResult<void>> WaitingTask{ new Private::TTaskWithResult<void>, /*bAddRef=*/ false };
-			WaitingTask->Init(TEXT("Waiting Task"), [CompletionEvent] { CompletionEvent->Trigger(); }, Private::FTaskBase::InlineTaskPriority);
-			WaitingTask->AddPrerequisites(*this);
-
-			if (WaitingTask->TryLaunch())
-			{	// was executed inline
-				check(WaitingTask->IsCompleted());
-				return true;
-			}
-
-			return CompletionEvent->Wait(Timeout);
+			return !IsValid() || Pimpl->Wait(Timeout);
 		}
 
 		// waits for task's completion while executing other tasks. Shouldn't be used inside a latency-sensitive task
@@ -116,8 +89,6 @@ namespace UE { namespace Tasks
 		{
 			if (IsValid())
 			{
-				TaskTrace::FWaitingScope WaitingScope(Pimpl->GetTraceId());
-				TRACE_CPUPROFILER_EVENT_SCOPE(Tasks::BusyWait);
 				Pimpl->BusyWait();
 			}
 		}
@@ -127,13 +98,7 @@ namespace UE { namespace Tasks
 		// @return true if the task is completed
 		bool BusyWait(FTimespan Timeout)
 		{
-			if (IsValid())
-			{
-				TaskTrace::FWaitingScope WaitingScope(Pimpl->GetTraceId());
-				return Pimpl->BusyWait(Timeout);
-			}
-
-			return true;
+			return !IsValid() || Pimpl->BusyWait(Timeout);
 		}
 
 		// waits for task's completion or the given condition becomes true, while executing other tasks.
@@ -142,14 +107,7 @@ namespace UE { namespace Tasks
 		template<typename ConditionType>
 		bool BusyWait(ConditionType&& Condition)
 		{
-			if (IsValid())
-			{
-				TaskTrace::FWaitingScope WaitingScope(Pimpl->GetTraceId());
-				TRACE_CPUPROFILER_EVENT_SCOPE(Tasks::BusyWait);
-				return Pimpl->BusyWait(Forward<ConditionType>(Condition));
-			}
-
-			return true;
+			return !IsValid() || Pimpl->BusyWait(Forward<ConditionType>(Condition));
 		}
 
 	protected:
diff --git a/Engine/Source/Runtime/Core/Public/Tasks/TaskPrivate.h b/Engine/Source/Runtime/Core/Public/Tasks/TaskPrivate.h
--- a/Engine/Source/Runtime/Core/Public/Tasks/TaskPrivate.h
+++ b/Engine/Source/Runtime/Core/Public/Tasks/TaskPrivate.h
@@ -106,8 +106,7 @@ namespace UE { namespace Tasks
 			// The task will be executed only when all prerequisites are completed. The task type must be a task handle that holds a pointer to
 			// FTaskBase as its `Pimpl` member (see Tasks::TTaskBase).
 			// Must not be called concurrently
-			template<typename HigherLevelTaskType, decltype(std::declval<HigherLevelTaskType>().Pimpl)* = nullptr>
-			void AddPrerequisites(const HigherLevelTaskType& Prerequisite)
+			void AddPrerequisites(FTaskBase& Prerequisite)
 			{
 				checkf(NumLocks.load(std::memory_order_relaxed) >= NumInitialLocks && NumLocks.load(std::memory_order_relaxed) < ExecutionFlag, TEXT("Prerequisites can be added only before the task is launched"));
 
@@ -117,11 +116,11 @@ namespace UE { namespace Tasks
 				uint32 PrevNumLocks = NumLocks.fetch_add(1, std::memory_order_acquire); // `acquire` to make it happen before the task is registered as a subsequent
 				checkf(PrevNumLocks + 1 < ExecutionFlag, TEXT("Max number of nested tasks reached: %d"), ExecutionFlag);
 
-				if (Prerequisite.Pimpl->AddSubsequent(*this))
+				if (Prerequisite.AddSubsequent(*this))
 				{
-					Prerequisites.Enqueue(Prerequisite.Pimpl.GetReference());
+					Prerequisites.Enqueue(&Prerequisite);
 					// keep it alive until this task's execution
-					Prerequisite.Pimpl->AddRef();
+					Prerequisite.AddRef();
 				}
 				else
 				{
@@ -130,6 +129,15 @@ namespace UE { namespace Tasks
 				}
 			}
 
+			// The task will be executed only when all prerequisites are completed. The task type must be a task handle that holds a pointer to
+			// FTaskBase as its `Pimpl` member (see Tasks::TTaskBase).
+			// Must not be called concurrently
+			template<typename HigherLevelTaskType, decltype(std::declval<HigherLevelTaskType>().Pimpl)* = nullptr>
+			void AddPrerequisites(const HigherLevelTaskType& Prerequisite)
+			{
+				AddPrerequisites(*Prerequisite.Pimpl);
+			}
+
 			// The task will be executed only when all prerequisites are completed.
 			// Must not be called concurrently.
 			// @param InPrerequisites - an iterable collection of tasks
@@ -330,9 +338,42 @@ namespace UE { namespace Tasks
 				}
 			}
 
+			// waits for task's completion, with optional timeout. Tries to retract the task and execute it in-place, if failed - blocks until the task 
+			// is completed by another thread. If timeout is zero, tries to retract the task and returns immedially after that.
+			// @return true if the task is completed
+			bool Wait(FTimespan Timeout = FTimespan::MaxValue())
+			{
+				TaskTrace::FWaitingScope WaitingScope(GetTraceId());
+				TRACE_CPUPROFILER_EVENT_SCOPE(Tasks::Wait);
+
+				if (TryRetractAndExecute())
+				{
+					return true;
+				}
+
+				// the event must be alive for the task and this function lifetime, we don't know which one will be finished first as waiting can time out
+				// before the waiting task is completed
+				FSharedEventRef CompletionEvent;
+
+				TRefCountPtr<Private::FTaskBase> WaitingTask{ new Private::FTaskBase, /*bAddRef=*/ false };
+				WaitingTask->Init(TEXT("Waiting Task"), [CompletionEvent] { CompletionEvent->Trigger(); }, Private::FTaskBase::InlineTaskPriority);
+				WaitingTask->AddPrerequisites(*this);
+
+				if (WaitingTask->TryLaunch())
+				{	// was executed inline
+					check(WaitingTask->IsCompleted());
+					return true;
+				}
+
+				return CompletionEvent->Wait(Timeout);
+			}
+
 			// waits until the task is completed while executing other tasks
 			void BusyWait()
 			{
+				TaskTrace::FWaitingScope WaitingScope(GetTraceId());
+				TRACE_CPUPROFILER_EVENT_SCOPE(Tasks::BusyWait);
+				
 				if (!TryRetractAndExecute())
 				{
 					LowLevelTasks::BusyWaitUntil([this] { return IsCompleted(); });
@@ -342,6 +383,9 @@ namespace UE { namespace Tasks
 			// waits until the task is completed or waiting timed out, while executing other tasks
 			bool BusyWait(FTimespan InTimeout)
 			{
+				TaskTrace::FWaitingScope WaitingScope(GetTraceId());
+				TRACE_CPUPROFILER_EVENT_SCOPE(Tasks::BusyWait);
+
 				FTimeout Timeout{ InTimeout };
 				
 				if (TryRetractAndExecute())
@@ -357,6 +401,9 @@ namespace UE { namespace Tasks
 			template<typename ConditionType>
 			bool BusyWait(ConditionType&& Condition)
 			{
+				TaskTrace::FWaitingScope WaitingScope(GetTraceId());
+				TRACE_CPUPROFILER_EVENT_SCOPE(Tasks::BusyWait);
+
 				if (TryRetractAndExecute())
 				{
 					return true;
