fixed non-unity compilation

diff --git a/Engine/Source/Runtime/Core/Public/Tasks/Task.h b/Engine/Source/Runtime/Core/Public/Tasks/Task.h
--- a/Engine/Source/Runtime/Core/Public/Tasks/Task.h
+++ b/Engine/Source/Runtime/Core/Public/Tasks/Task.h
@@ -4,7 +4,6 @@
 
 #include "Tasks/TaskPrivate.h"
 #include "Async/Fundamental/Task.h"
-#include "Templates/RefCounting.h"
 #include "Containers/StaticArray.h"
 #include "HAL/Event.h"
 #include "CoreTypes.h"
diff --git a/Engine/Source/Runtime/Core/Public/Tasks/TaskPrivate.h b/Engine/Source/Runtime/Core/Public/Tasks/TaskPrivate.h
--- a/Engine/Source/Runtime/Core/Public/Tasks/TaskPrivate.h
+++ b/Engine/Source/Runtime/Core/Public/Tasks/TaskPrivate.h
@@ -11,6 +11,7 @@
 #include "Containers/ClosableMpscQueue.h"
 #include "Containers/SpscQueue.h"
 #include "Templates/UnrealTemplate.h"
+#include "Templates/RefCounting.h"
 #include "CoreTypes.h"
 #include "Math/NumericLimits.h"
 #include "Misc/SpinLock.h"
