Fixed UnrealMathNeon.h compilation for VC ARM64

diff --git a/Engine/Source/Runtime/Core/Public/Math/UnrealMathNeon.h b/Engine/Source/Runtime/Core/Public/Math/UnrealMathNeon.h
--- a/Engine/Source/Runtime/Core/Public/Math/UnrealMathNeon.h
+++ b/Engine/Source/Runtime/Core/Public/Math/UnrealMathNeon.h
@@ -30,7 +30,7 @@ PRAGMA_DISABLE_SHADOW_VARIABLE_WARNINGS
 // longer trivially constructible. The optimizer should eliminate the
 // redundant zero initialization in these cases for non-MSVC (e.g. V()
 // is called now where it wasn't before)
-template<typename T, bool is_int>
+template<typename T, typename BASE_TYPE>
 struct alignas(alignof(T)) VectorRegisterWrapper
 {
 	FORCEINLINE VectorRegisterWrapper() {}
@@ -39,14 +39,40 @@ struct alignas(alignof(T)) VectorRegisterWrapper
 	FORCEINLINE operator T&() { return m_vec; }
 	FORCEINLINE operator const T&() const { return m_vec; }
 
+	FORCEINLINE BASE_TYPE operator[](int Index) const;
+
 	T m_vec;
 };
 
+template<>
+FORCEINLINE float VectorRegisterWrapper<float32x4_t, float>::operator[](int Index) const
+{
+	return m_vec.n128_f32[Index];
+}
+
+template<>
+FORCEINLINE double VectorRegisterWrapper<float64x2_t, double>::operator[](int Index) const
+{
+	return m_vec.n128_f64[Index];
+}
+
+template<>
+FORCEINLINE int VectorRegisterWrapper<int32x4_t, int>::operator[](int Index) const
+{
+	return m_vec.n128_i32[Index];
+}
+
+template<>
+FORCEINLINE int64 VectorRegisterWrapper<int64x2_t, int64>::operator[](int Index) const
+{
+	return m_vec.n128_i64[Index];
+}
+
 /** 16-byte vector register type */
-typedef VectorRegisterWrapper<float32x4_t, false> VectorRegister4Float;
-typedef VectorRegisterWrapper<float64x2_t, false> VectorRegister2Double;
-typedef VectorRegisterWrapper<int32x4_t, true> VectorRegister4Int;
-typedef VectorRegisterWrapper<int64x2_t, true> VectorRegister2Int64;
+typedef VectorRegisterWrapper<float32x4_t, float> VectorRegister4Float;
+typedef VectorRegisterWrapper<float64x2_t, double> VectorRegister2Double;
+typedef VectorRegisterWrapper<int32x4_t, int> VectorRegister4Int;
+typedef VectorRegisterWrapper<int64x2_t, int64> VectorRegister2Int64;
 #else
 
 /** 16-byte vector register type */
@@ -60,14 +86,10 @@ typedef int64x2_t GCC_ALIGN(16) VectorRegister2Int64;
 
 struct alignas(16) VectorRegister4Double
 {
-	union
+	struct
 	{
-		float64x2x2_t Vec;
-		struct
-		{
-			VectorRegister2Double XY;
-			VectorRegister2Double ZW;
-		};
+		VectorRegister2Double XY;
+		VectorRegister2Double ZW;
 	};
 
 	FORCEINLINE VectorRegister4Double() = default;
@@ -407,8 +429,8 @@ FORCEINLINE VectorRegister4Float VectorLoad(const float* Ptr)
 
 FORCEINLINE VectorRegister4Double VectorLoad(const double* Ptr)
 {
-	VectorRegister4Double Result;
-	Result.Vec = vld1q_f64_x2(Ptr);
+	float64x2x2_t Vec = vld1q_f64_x2(Ptr);
+	VectorRegister4Double Result = *(VectorRegister4Double*)&Vec;
 	return Result;
 }
 
@@ -580,7 +602,7 @@ FORCEINLINE void VectorStoreAligned(const VectorRegister4Float& Vec, float* Ptr)
 
 FORCEINLINE void VectorStoreAligned(const VectorRegister4Double& Vec, double* Ptr)
 {
-	vst1q_f64_x2(Ptr, Vec.Vec);
+	vst1q_f64_x2(Ptr, *(float64x2x2_t*)&Vec);
 }
 
 //TODO: LWC VectorVM.cpp calls it on a line 3294, case EVectorVMOp::outputdata_half: Context.WriteExecFunction(CopyConstantToOutput<float, FFloat16, 2>); break;
@@ -614,7 +636,7 @@ FORCEINLINE void VectorStore(const VectorRegister4Float& Vec, float* Ptr)
 
 FORCEINLINE void VectorStore(const VectorRegister4Double& Vec, double* Ptr)
 {
-	vst1q_f64_x2(Ptr, Vec.Vec);
+	vst1q_f64_x2(Ptr, *(float64x2x2_t*)&Vec);
 }
 
 /**
@@ -1181,7 +1203,7 @@ FORCEINLINE VectorRegister4Float VectorSwizzle
 		0x0F0E0D0C, // XM_SWIZZLE_W
 	};
 
-	int8x8x2_t tbl;
+	uint8x8x2_t tbl;
 	tbl.val[0] = vget_low_f32(V);
 	tbl.val[1] = vget_high_f32(V);
 
@@ -1193,6 +1215,38 @@ FORCEINLINE VectorRegister4Float VectorSwizzle
 
 	return vcombine_f32(rL, rH);
 }
+
+FORCEINLINE VectorRegister4Double VectorSwizzle
+(
+	VectorRegister4Double V,
+	uint32 E0,
+	uint32 E1,
+	uint32 E2,
+	uint32 E3
+)
+{
+	check((E0 < 4) && (E1 < 4) && (E2 < 4) && (E3 < 4));
+	static const uint64_t ControlElement[4] =
+	{
+		0x0706050403020100ULL, // XM_SWIZZLE_X
+		0x0F0E0D0C0B0A0908ULL, // XM_SWIZZLE_Y
+		0x1716151413121110ULL, // XM_SWIZZLE_Z
+		0x1F1E1D1C1B1A1918ULL, // XM_SWIZZLE_W
+	};
+
+	uint8x16x2_t tbl;
+	tbl.val[0] = V.XY;
+	tbl.val[1] = V.ZW;
+
+	VectorRegister4Double Result;
+	uint32x4_t idx = vcombine_u64(vcreate_u64(ControlElement[E0]), vcreate_u64(ControlElement[E1]));
+	Result.XY = vqtbl2q_u8(tbl, idx);
+
+	idx = vcombine_u64(vcreate_u64(ControlElement[E2]), vcreate_u64(ControlElement[E3]));
+	Result.ZW = vqtbl2q_u8(tbl, idx);
+
+	return Result;
+}
 #else
 template <int X, int Y, int Z, int W>
 FORCEINLINE VectorRegister4Float VectorSwizzleImpl(VectorRegister4Float Vec)
@@ -1259,7 +1313,7 @@ FORCEINLINE VectorRegister4Float VectorShuffle
 	uint32 PermuteW
 )
 {
-	check(PermuteX <= 7 && PermuteY <= 7 && PermuteZ <= 7 && PermuteW <= 7);
+	check(PermuteX <= 3 && PermuteY <= 3 && PermuteZ <= 3 && PermuteW <= 3);
 
 	static const uint32 ControlElement[8] =
 	{
@@ -1273,7 +1327,7 @@ FORCEINLINE VectorRegister4Float VectorShuffle
 		0x1F1E1D1C, // XM_PERMUTE_1W
 	};
 
-	int8x8x4_t tbl;
+	uint8x8x4_t tbl;
 	tbl.val[0] = vget_low_f32(V1);
 	tbl.val[1] = vget_high_f32(V1);
 	tbl.val[2] = vget_low_f32(V2);
@@ -1282,11 +1336,52 @@ FORCEINLINE VectorRegister4Float VectorShuffle
 	uint32x2_t idx = vcreate_u32(static_cast<uint64>(ControlElement[PermuteX]) | (static_cast<uint64>(ControlElement[PermuteY]) << 32));
 	const uint8x8_t rL = vtbl4_u8(tbl, idx);
 
-	idx = vcreate_u32(static_cast<uint64>(ControlElement[PermuteZ]) | (static_cast<uint64>(ControlElement[PermuteW]) << 32));
+	idx = vcreate_u32(static_cast<uint64>(ControlElement[PermuteZ + 4]) | (static_cast<uint64>(ControlElement[PermuteW + 4]) << 32));
 	const uint8x8_t rH = vtbl4_u8(tbl, idx);
 
 	return vcombine_f32(rL, rH);
 }
+
+FORCEINLINE VectorRegister4Double VectorShuffle
+(
+	VectorRegister4Double V1,
+	VectorRegister4Double V2,
+	uint32 PermuteX,
+	uint32 PermuteY,
+	uint32 PermuteZ,
+	uint32 PermuteW
+)
+{
+	check(PermuteX <= 3 && PermuteY <= 3 && PermuteZ <= 3 && PermuteW <= 3);
+
+	static const uint64 ControlElement[8] =
+	{
+		0x0706050403020100ULL, // XM_PERMUTE_0X
+		0x0F0E0D0C0B0A0908ULL, // XM_PERMUTE_0Y
+		0x1716151413121110ULL, // XM_PERMUTE_0Z
+		0x1F1E1D1C1B1A1918ULL, // XM_PERMUTE_0W
+
+		0x2726252423222120ULL, // XM_PERMUTE_1X
+		0x2F2E2D2C2B2A2928ULL, // XM_PERMUTE_1Y
+		0x3736353433323130ULL, // XM_PERMUTE_1Z
+		0x3F3E3D3C3B3A3938ULL, // XM_PERMUTE_1W
+	};
+
+	uint8x16x4_t tbl;
+	tbl.val[0] = V1.XY;
+	tbl.val[1] = V1.ZW;
+	tbl.val[2] = V2.XY;
+	tbl.val[3] = V2.ZW;
+
+	VectorRegister4Double Result;
+	uint32x4_t idx = vcombine_u64(vcreate_u64(ControlElement[PermuteX]), vcreate_u64(ControlElement[PermuteY]));
+	Result.XY = vqtbl4q_u8(tbl, idx);
+
+	idx = vcombine_u64(vcreate_u64(ControlElement[PermuteZ + 4]), vcreate_u64(ControlElement[PermuteW + 4]));
+	Result.ZW = vqtbl4q_u8(tbl, idx);
+
+	return Result;
+}
 #else
 
 template <int X, int Y, int Z, int W>
@@ -1690,54 +1785,56 @@ FORCEINLINE void VectorMatrixMultiply(FMatrix44d* Result, const FMatrix44d* Matr
 	float64x2x4_t A = vld1q_f64_x4((const double*)Matrix1);
 	float64x2x4_t B1 = vld1q_f64_x4((const double*)Matrix2);
 	float64x2x4_t B2 = vld1q_f64_x4((const double*)Matrix2 + 8);
+	float64_t* V = (float64_t*)&A;
 	float64x2x4_t R;
 
 	// First row of result (Matrix1[0] * Matrix2).
-	R.val[0] = vmulq_n_f64(B1.val[0], A.val[0][0]);
-	R.val[0] = vfmaq_n_f64(R.val[0], B1.val[2], A.val[0][1]);
-	R.val[0] = vfmaq_n_f64(R.val[0], B2.val[0], A.val[1][0]);
-	R.val[0] = vfmaq_n_f64(R.val[0], B2.val[2], A.val[1][1]);
+	R.val[0] = vmulq_n_f64(B1.val[0], V[0]);
+	R.val[0] = vfmaq_n_f64(R.val[0], B1.val[2], V[1]);
+	R.val[0] = vfmaq_n_f64(R.val[0], B2.val[0], V[2]);
+	R.val[0] = vfmaq_n_f64(R.val[0], B2.val[2], V[3]);
 
-	R.val[1] = vmulq_n_f64(B1.val[1], A.val[0][0]);
-	R.val[1] = vfmaq_n_f64(R.val[1], B1.val[3], A.val[0][1]);
-	R.val[1] = vfmaq_n_f64(R.val[1], B2.val[1], A.val[1][0]);
-	R.val[1] = vfmaq_n_f64(R.val[1], B2.val[3], A.val[1][1]);
+	R.val[1] = vmulq_n_f64(B1.val[1], V[0]);
+	R.val[1] = vfmaq_n_f64(R.val[1], B1.val[3], V[1]);
+	R.val[1] = vfmaq_n_f64(R.val[1], B2.val[1], V[2]);
+	R.val[1] = vfmaq_n_f64(R.val[1], B2.val[3], V[3]);
 
 	// Second row of result (Matrix1[1] * Matrix2).
-	R.val[2] = vmulq_n_f64(B1.val[0], A.val[2][0]);
-	R.val[2] = vfmaq_n_f64(R.val[2], B1.val[2], A.val[2][1]);
-	R.val[2] = vfmaq_n_f64(R.val[2], B2.val[0], A.val[3][0]);
-	R.val[2] = vfmaq_n_f64(R.val[2], B2.val[2], A.val[3][1]);
+	R.val[2] = vmulq_n_f64(B1.val[0], V[4]);
+	R.val[2] = vfmaq_n_f64(R.val[2], B1.val[2], V[5]);
+	R.val[2] = vfmaq_n_f64(R.val[2], B2.val[0], V[6]);
+	R.val[2] = vfmaq_n_f64(R.val[2], B2.val[2], V[7]);
 
-	R.val[3] = vmulq_n_f64(B1.val[1], A.val[2][0]);
-	R.val[3] = vfmaq_n_f64(R.val[3], B1.val[3], A.val[2][1]);
-	R.val[3] = vfmaq_n_f64(R.val[3], B2.val[1], A.val[3][0]);
-	R.val[3] = vfmaq_n_f64(R.val[3], B2.val[3], A.val[3][1]);
+	R.val[3] = vmulq_n_f64(B1.val[1], V[4]);
+	R.val[3] = vfmaq_n_f64(R.val[3], B1.val[3], V[5]);
+	R.val[3] = vfmaq_n_f64(R.val[3], B2.val[1], V[6]);
+	R.val[3] = vfmaq_n_f64(R.val[3], B2.val[3], V[7]);
 
 	vst1q_f64_x4((double*)Result, R);
 	A = vld1q_f64_x4((const double*)Matrix1 + 8);
+	V = (float64_t*)&A;
 
 	// Third row of result (Matrix1[2] * Matrix2).
-	R.val[0] = vmulq_n_f64(B1.val[0], A.val[0][0]);
-	R.val[0] = vfmaq_n_f64(R.val[0], B1.val[2], A.val[0][1]);
-	R.val[0] = vfmaq_n_f64(R.val[0], B2.val[0], A.val[1][0]);
-	R.val[0] = vfmaq_n_f64(R.val[0], B2.val[2], A.val[1][1]);
+	R.val[0] = vmulq_n_f64(B1.val[0], V[0]);
+	R.val[0] = vfmaq_n_f64(R.val[0], B1.val[2], V[1]);
+	R.val[0] = vfmaq_n_f64(R.val[0], B2.val[0], V[2]);
+	R.val[0] = vfmaq_n_f64(R.val[0], B2.val[2], V[3]);
 
-	R.val[1] = vmulq_n_f64(B1.val[1], A.val[0][0]);
-	R.val[1] = vfmaq_n_f64(R.val[1], B1.val[3], A.val[0][1]);
-	R.val[1] = vfmaq_n_f64(R.val[1], B2.val[1], A.val[1][0]);
-	R.val[1] = vfmaq_n_f64(R.val[1], B2.val[3], A.val[1][1]);
+	R.val[1] = vmulq_n_f64(B1.val[1], V[0]);
+	R.val[1] = vfmaq_n_f64(R.val[1], B1.val[3], V[1]);
+	R.val[1] = vfmaq_n_f64(R.val[1], B2.val[1], V[2]);
+	R.val[1] = vfmaq_n_f64(R.val[1], B2.val[3], V[3]);
 
 	// Fourth row of result (Matrix1[3] * Matrix2).
-	R.val[2] = vmulq_n_f64(B1.val[0], A.val[2][0]);
-	R.val[2] = vfmaq_n_f64(R.val[2], B1.val[2], A.val[2][1]);
-	R.val[2] = vfmaq_n_f64(R.val[2], B2.val[0], A.val[3][0]);
-	R.val[2] = vfmaq_n_f64(R.val[2], B2.val[2], A.val[3][1]);
+	R.val[2] = vmulq_n_f64(B1.val[0], V[4]);
+	R.val[2] = vfmaq_n_f64(R.val[2], B1.val[2], V[5]);
+	R.val[2] = vfmaq_n_f64(R.val[2], B2.val[0], V[6]);
+	R.val[2] = vfmaq_n_f64(R.val[2], B2.val[2], V[7]);
 
-	R.val[3] = vmulq_n_f64(B1.val[1], A.val[2][0]);
-	R.val[3] = vfmaq_n_f64(R.val[3], B1.val[3], A.val[2][1]);
-	R.val[3] = vfmaq_n_f64(R.val[3], B2.val[1], A.val[3][0]);
-	R.val[3] = vfmaq_n_f64(R.val[3], B2.val[3], A.val[3][1]);
+	R.val[3] = vmulq_n_f64(B1.val[1], V[4]);
+	R.val[3] = vfmaq_n_f64(R.val[3], B1.val[3], V[5]);
+	R.val[3] = vfmaq_n_f64(R.val[3], B2.val[1], V[6]);
+	R.val[3] = vfmaq_n_f64(R.val[3], B2.val[3], V[7]);
 
 	vst1q_f64_x4((double*)Result + 8, R);
 }
@@ -1924,7 +2021,7 @@ FORCEINLINE void VectorMatrixInverse(FMatrix44d* DstMatrix, const FMatrix44d* Sr
  * @param MatrixM		FMatrix pointer to the Matrix to apply transform
  * @return VectorRegister4Float = VecP*MatrixM
  */
-FORCEINLINE VectorRegister4Float VectorTransformVector(const VectorRegister4Float&  VecP,  const FMatrix44f* MatrixM )
+FORCEINLINE VectorRegister4Float VectorTransformVector(const VectorRegister4Float& VecP, const FMatrix44f* MatrixM )
 {
 	float32x4x4_t M = vld1q_f32_x4((const float*)MatrixM);
 	VectorRegister4Float Result;
