Fixed compilation for HoloLens ARM64

diff --git a/Engine/Source/Runtime/Core/Public/Math/RotationTranslationMatrix.h b/Engine/Source/Runtime/Core/Public/Math/RotationTranslationMatrix.h
--- a/Engine/Source/Runtime/Core/Public/Math/RotationTranslationMatrix.h
+++ b/Engine/Source/Runtime/Core/Public/Math/RotationTranslationMatrix.h
@@ -39,7 +39,7 @@ public:
 template<typename T>
 FORCEINLINE TRotationTranslationMatrix<T>::TRotationTranslationMatrix(const FRotator& Rot, const TVector<T>& Origin)
 {
-#if PLATFORM_ENABLE_VECTORINTRINSICS
+#if PLATFORM_ENABLE_VECTORINTRINSICS && (!PLATFORM_HOLOLENS || !PLATFORM_CPU_ARM_FAMILY)
 
 	const TVectorRegisterType<T> Angles = MakeVectorRegister(Rot.Pitch, Rot.Yaw, Rot.Roll, 0.0f);
 	const TVectorRegisterType<T> HalfAngles = VectorMultiply(Angles, GlobalVectorConstants::DEG_TO_RAD);
