Trivial: edited comment to actually match the intent/error logic

diff --git a/Engine/Source/Runtime/Core/Public/Windows/WindowsPlatformCompilerSetup.h b/Engine/Source/Runtime/Core/Public/Windows/WindowsPlatformCompilerSetup.h
--- a/Engine/Source/Runtime/Core/Public/Windows/WindowsPlatformCompilerSetup.h
+++ b/Engine/Source/Runtime/Core/Public/Windows/WindowsPlatformCompilerSetup.h
@@ -246,7 +246,7 @@ static_assert(_MSC_VER >= 1924, "Visual Studio 2019 v16.4 or greater is required
 #error "Bad VCC option: C++ exception handling must be enabled" //lint !e309 suppress as lint doesn't have this defined
 #endif
 
-// Make sure characters are unsigned.
+// Make sure characters are signed
 #ifdef _CHAR_UNSIGNED
 #error "Bad VC++ option: Characters must be signed" //lint !e309 suppress as lint doesn't have this defined
 #endif
