Fixed Stats macros cause extra evaluations of their 'Value' expressions.

diff --git a/Engine/Source/Runtime/Core/Public/Stats/Stats2.h b/Engine/Source/Runtime/Core/Public/Stats/Stats2.h
--- a/Engine/Source/Runtime/Core/Public/Stats/Stats2.h
+++ b/Engine/Source/Runtime/Core/Public/Stats/Stats2.h
@@ -2025,42 +2025,37 @@ struct FStat_##StatName\
 }
 #define INC_FLOAT_STAT_BY(Stat, Amount) \
 {\
-	if (Amount != 0.0f) \
+	double AddAmount = double(Amount); \
+	if (AddAmount != 0.0) \
 	{ \
 		if (FThreadStats::IsCollectingData() || !GET_STATISEVERYFRAME(Stat)) \
 		{ \
-			FThreadStats::AddMessage(GET_STATFNAME(Stat), EStatOperation::Add, double(Amount));\
-			TRACE_STAT_ADD(GET_STATFNAME(Stat), double(Amount)); \
+			FThreadStats::AddMessage(GET_STATFNAME(Stat), EStatOperation::Add, AddAmount); \
+			TRACE_STAT_ADD(GET_STATFNAME(Stat), AddAmount); \
 		} \
 	} \
 }
 #define INC_DWORD_STAT_BY(Stat, Amount) \
 {\
-	if (Amount != 0) \
+	int64 AddAmount = int64(Amount); \
+	if (AddAmount != 0) \
 	{ \
 		if (FThreadStats::IsCollectingData() || !GET_STATISEVERYFRAME(Stat)) \
 		{ \
-			FThreadStats::AddMessage(GET_STATFNAME(Stat), EStatOperation::Add, int64(Amount));\
-			TRACE_STAT_ADD(GET_STATFNAME(Stat), int64(Amount)); \
+			FThreadStats::AddMessage(GET_STATFNAME(Stat), EStatOperation::Add, AddAmount); \
+			TRACE_STAT_ADD(GET_STATFNAME(Stat), AddAmount); \
 		} \
 	} \
 }
-#define INC_DWORD_STAT_FNAME_BY(StatFName, Amount) \
-{\
-	if (Amount != 0) \
-	{ \
-		FThreadStats::AddMessage(StatFName, EStatOperation::Add, int64(Amount));\
-		TRACE_STAT_ADD(StatFName, int64(Amount)); \
-	} \
-}
 #define INC_MEMORY_STAT_BY(Stat, Amount) \
 {\
-	if (Amount != 0) \
+	int64 AddAmount = int64(Amount); \
+	if (AddAmount != 0) \
 	{ \
 		if (FThreadStats::IsCollectingData() || !GET_STATISEVERYFRAME(Stat)) \
 		{ \
-			FThreadStats::AddMessage(GET_STATFNAME(Stat), EStatOperation::Add, int64(Amount));\
-			TRACE_STAT_ADD(GET_STATFNAME(Stat), int64(Amount)); \
+			FThreadStats::AddMessage(GET_STATFNAME(Stat), EStatOperation::Add, AddAmount); \
+			TRACE_STAT_ADD(GET_STATFNAME(Stat), AddAmount); \
 		} \
 	} \
 }
@@ -2074,42 +2069,37 @@ struct FStat_##StatName\
 }
 #define DEC_FLOAT_STAT_BY(Stat,Amount) \
 {\
-	if (Amount != 0.0f) \
+	double SubtractAmount = double(Amount); \
+	if (SubtractAmount != 0.0) \
 	{ \
 		if (FThreadStats::IsCollectingData() || !GET_STATISEVERYFRAME(Stat)) \
 		{ \
-			FThreadStats::AddMessage(GET_STATFNAME(Stat), EStatOperation::Subtract, double(Amount));\
-			TRACE_STAT_ADD(GET_STATFNAME(Stat), -double(Amount)); \
+			FThreadStats::AddMessage(GET_STATFNAME(Stat), EStatOperation::Subtract, SubtractAmount); \
+			TRACE_STAT_ADD(GET_STATFNAME(Stat), -SubtractAmount); \
 		} \
 	} \
 }
 #define DEC_DWORD_STAT_BY(Stat,Amount) \
 {\
-	if (Amount != 0) \
+	int64 SubtractAmount = int64(Amount); \
+	if (SubtractAmount != 0) \
 	{ \
 		if (FThreadStats::IsCollectingData() || !GET_STATISEVERYFRAME(Stat)) \
 		{ \
-			FThreadStats::AddMessage(GET_STATFNAME(Stat), EStatOperation::Subtract, int64(Amount));\
-			TRACE_STAT_ADD(GET_STATFNAME(Stat), -int64(Amount)); \
+			FThreadStats::AddMessage(GET_STATFNAME(Stat), EStatOperation::Subtract, SubtractAmount); \
+			TRACE_STAT_ADD(GET_STATFNAME(Stat), -SubtractAmount); \
 		} \
 	} \
 }
-#define DEC_DWORD_STAT_FNAME_BY(StatFName,Amount) \
-{\
-	if (Amount != 0) \
-	{ \
-		FThreadStats::AddMessage(StatFName, EStatOperation::Subtract, int64(Amount));\
-		TRACE_STAT_ADD(StatFName, -int64(Amount)); \
-	} \
-}
 #define DEC_MEMORY_STAT_BY(Stat,Amount) \
 {\
-	if (Amount != 0) \
+	int64 SubtractAmount = int64(Amount); \
+	if (SubtractAmount != 0) \
 	{ \
 		if (FThreadStats::IsCollectingData() || !GET_STATISEVERYFRAME(Stat)) \
 		{ \
-			FThreadStats::AddMessage(GET_STATFNAME(Stat), EStatOperation::Subtract, int64(Amount));\
-			TRACE_STAT_ADD(GET_STATFNAME(Stat), -int64(Amount)); \
+			FThreadStats::AddMessage(GET_STATFNAME(Stat), EStatOperation::Subtract, SubtractAmount); \
+			TRACE_STAT_ADD(GET_STATFNAME(Stat), -SubtractAmount); \
 		} \
 	} \
 }
@@ -2117,24 +2107,27 @@ struct FStat_##StatName\
 {\
 	if (FThreadStats::IsCollectingData() || !GET_STATISEVERYFRAME(Stat)) \
 	{ \
-		FThreadStats::AddMessage(GET_STATFNAME(Stat), EStatOperation::Set, int64(Value));\
-		TRACE_STAT_SET(GET_STATFNAME(Stat), int64(Value)); \
+		int64 SetValue = int64(Value); \
+		FThreadStats::AddMessage(GET_STATFNAME(Stat), EStatOperation::Set, SetValue); \
+		TRACE_STAT_SET(GET_STATFNAME(Stat), SetValue); \
 	} \
 }
 #define SET_DWORD_STAT(Stat,Value) \
 {\
 	if (FThreadStats::IsCollectingData() || !GET_STATISEVERYFRAME(Stat)) \
 	{ \
-		FThreadStats::AddMessage(GET_STATFNAME(Stat), EStatOperation::Set, int64(Value));\
-		TRACE_STAT_SET(GET_STATFNAME(Stat), int64(Value)); \
+		int64 SetValue = int64(Value); \
+		FThreadStats::AddMessage(GET_STATFNAME(Stat), EStatOperation::Set, SetValue); \
+		TRACE_STAT_SET(GET_STATFNAME(Stat), SetValue); \
 	} \
 }
 #define SET_FLOAT_STAT(Stat,Value) \
 {\
 	if (FThreadStats::IsCollectingData() || !GET_STATISEVERYFRAME(Stat)) \
 	{ \
-		FThreadStats::AddMessage(GET_STATFNAME(Stat), EStatOperation::Set, double(Value));\
-		TRACE_STAT_SET(GET_STATFNAME(Stat), double(Value)); \
+		double SetValue = double(Value); \
+		FThreadStats::AddMessage(GET_STATFNAME(Stat), EStatOperation::Set, SetValue); \
+		TRACE_STAT_SET(GET_STATFNAME(Stat), SetValue); \
 	} \
 }
 
@@ -2159,26 +2152,30 @@ struct FStat_##StatName\
 }
 #define INC_FLOAT_STAT_BY_FName(Stat, Amount) \
 {\
-	if (Amount != 0.0f) \
+	double AddAmount = double(Amount); \
+	if (AddAmount != 0.0) \
 	{ \
-		FThreadStats::AddMessage(Stat, EStatOperation::Add, double(Amount));\
-		TRACE_STAT_ADD(Stat, double(Amount)); \
+		FThreadStats::AddMessage(Stat, EStatOperation::Add, AddAmount); \
+		TRACE_STAT_ADD(Stat, AddAmount); \
 	} \
 }
 #define INC_DWORD_STAT_BY_FName(Stat, Amount) \
 {\
-	if (Amount != 0) \
+	int64 AddAmount = int64(Amount); \
+	if (AddAmount != 0) \
 	{ \
-		FThreadStats::AddMessage(Stat, EStatOperation::Add, int64(Amount));\
-		TRACE_STAT_ADD(Stat, int64(Amount)); \
+		FThreadStats::AddMessage(Stat, EStatOperation::Add, AddAmount); \
+		TRACE_STAT_ADD(Stat, AddAmount); \
 	} \
 }
+#define INC_DWORD_STAT_FNAME_BY(Stat, Amount) INC_DWORD_STAT_BY_FName(Stat, Amount)
 #define INC_MEMORY_STAT_BY_FName(Stat, Amount) \
 {\
-	if (Amount != 0) \
+	int64 AddAmount = int64(Amount); \
+	if (AddAmount != 0) \
 	{ \
-		FThreadStats::AddMessage(Stat, EStatOperation::Add, int64(Amount));\
-		TRACE_STAT_ADD(Stat, int64(Amount)); \
+		FThreadStats::AddMessage(Stat, EStatOperation::Add, AddAmount); \
+		TRACE_STAT_ADD(Stat, AddAmount); \
 	} \
 }
 #define DEC_DWORD_STAT_FName(Stat) \
@@ -2188,42 +2185,49 @@ struct FStat_##StatName\
 }
 #define DEC_FLOAT_STAT_BY_FName(Stat,Amount) \
 {\
-	if (Amount != 0.0f) \
+	double SubtractAmount = double(Amount); \
+	if (SubtractAmount != 0.0) \
 	{ \
-		FThreadStats::AddMessage(Stat, EStatOperation::Subtract, double(Amount));\
-		TRACE_STAT_ADD(Stat, -double(Amount)); \
+		FThreadStats::AddMessage(Stat, EStatOperation::Subtract, SubtractAmount); \
+		TRACE_STAT_ADD(Stat, -SubtractAmount); \
 	} \
 }
 #define DEC_DWORD_STAT_BY_FName(Stat,Amount) \
 {\
-	if (Amount != 0) \
+	int64 SubtractAmount = int64(Amount); \
+	if (SubtractAmount != 0) \
 	{ \
-		FThreadStats::AddMessage(Stat, EStatOperation::Subtract, int64(Amount));\
-		TRACE_STAT_ADD(Stat, -int64(Amount)); \
+		FThreadStats::AddMessage(Stat, EStatOperation::Subtract, SubtractAmount); \
+		TRACE_STAT_ADD(Stat, -SubtractAmount); \
 	} \
 }
+#define DEC_DWORD_STAT_FNAME_BY(Stat,Amount) DEC_DWORD_STAT_BY_FName(Stat,Amount)
 #define DEC_MEMORY_STAT_BY_FName(Stat,Amount) \
 {\
-	if (Amount != 0) \
+	int64 SubtractAmount = int64(Amount); \
+	if (SubtractAmount != 0) \
 	{ \
-		FThreadStats::AddMessage(Stat, EStatOperation::Subtract, int64(Amount));\
-		TRACE_STAT_ADD(Stat, -int64(Amount)); \
+		FThreadStats::AddMessage(Stat, EStatOperation::Subtract, SubtractAmount); \
+		TRACE_STAT_ADD(Stat, -SubtractAmount); \
 	} \
 }
 #define SET_MEMORY_STAT_FName(Stat,Value) \
 {\
-	FThreadStats::AddMessage(Stat, EStatOperation::Set, int64(Value));\
-	TRACE_STAT_SET(Stat, int64(Value)); \
+	int64 SetValue = int64(Value); \
+	FThreadStats::AddMessage(Stat, EStatOperation::Set, SetValue); \
+	TRACE_STAT_SET(Stat, SetValue); \
 }
 #define SET_DWORD_STAT_FName(Stat,Value) \
 {\
-	FThreadStats::AddMessage(Stat, EStatOperation::Set, int64(Value));\
-	TRACE_STAT_SET(Stat, int64(Value)); \
+	int64 SetValue = int64(Value); \
+	FThreadStats::AddMessage(Stat, EStatOperation::Set, SetValue); \
+	TRACE_STAT_SET(Stat, SetValue); \
 }
 #define SET_FLOAT_STAT_FName(Stat,Value) \
 {\
-	FThreadStats::AddMessage(Stat, EStatOperation::Set, double(Value));\
-	TRACE_STAT_SET(Stat, double(Value)); \
+	double SetValue = double(Value); \
+	FThreadStats::AddMessage(Stat, EStatOperation::Set, SetValue); \
+	TRACE_STAT_SET(Stat, SetValue); \
 }
 
 
