XSXIoDispatcher: Use direct storage hardware decompression also in the fallback path for unaligned compression block reads.
Fix broken trace counter for IoDispatcherFileBackend/QueuedFilesystemReads.
Use 64K io dispatcher buffers on XSX to match its 64K compression block size

#FROST-4056
#rnx

diff --git a/Engine/Source/Runtime/Core/Private/IO/IoDispatcherFileBackend.cpp b/Engine/Source/Runtime/Core/Private/IO/IoDispatcherFileBackend.cpp
--- a/Engine/Source/Runtime/Core/Private/IO/IoDispatcherFileBackend.cpp
+++ b/Engine/Source/Runtime/Core/Private/IO/IoDispatcherFileBackend.cpp
@@ -2077,7 +2077,7 @@ void FFileIoStats::OnFilesystemReadsQueued(uint64 BytesToRead, int32 NumReads)
 	QueuedFilesystemReads += NumReads;
 #endif
 
-	TRACE_COUNTER_INCREMENT(IoDispatcherFileBackendQueuedFilesystemReads)
+	TRACE_COUNTER_ADD(IoDispatcherFileBackendQueuedFilesystemReads, NumReads)
 	TRACE_COUNTER_ADD(IoDispatcherFileBackendQueuedFilesystemReadBytes, BytesToRead);
 #endif
 }
@@ -2145,6 +2145,7 @@ void FFileIoStats::OnReadComplete(int64 BytesRead)
 
 	TRACE_COUNTER_SUBTRACT(IoDispatcherFileBackendQueuedFilesystemReadBytes, BytesRead);
 	TRACE_COUNTER_ADD(IoDispatcherFileBackendTotalBytesRead, BytesRead);
+	TRACE_COUNTER_DECREMENT(IoDispatcherFileBackendQueuedFilesystemReads);
 #endif
 }
 
diff --git a/Engine/Source/Runtime/Core/Private/IO/IoDispatcherFileBackendTypes.h b/Engine/Source/Runtime/Core/Private/IO/IoDispatcherFileBackendTypes.h
--- a/Engine/Source/Runtime/Core/Private/IO/IoDispatcherFileBackendTypes.h
+++ b/Engine/Source/Runtime/Core/Private/IO/IoDispatcherFileBackendTypes.h
@@ -87,10 +87,10 @@ struct FFileIoStoreCompressedBlock
 	FFileIoStoreCompressedBlock* Next = nullptr;
 	FFileIoStoreBlockKey Key;
 	FName CompressionMethod;
-	uint64 RawOffset;
-	uint32 UncompressedSize;
-	uint32 CompressedSize;
-	uint32 RawSize;
+	uint64 RawOffset = uint64(-1);
+	uint32 UncompressedSize = uint32(-1);
+	uint32 CompressedSize = uint32(-1);
+	uint32 RawSize = uint32(-1);
 	uint32 RefCount = 0;
 	uint32 UnfinishedRawBlocksCount = 0;
 	TArray<struct FFileIoStoreReadRequest*, TInlineAllocator<2>> RawBlocks;
