CSA: Zero-init the FLLMScope::Tracker member variable to avoid warnings from clang's static analyzer. Zero-inited the rest of the fields and simplified the disable-check boilerplate in the constructors.

diff --git a/Engine/Source/Runtime/Core/Public/HAL/LowLevelMemTracker.h b/Engine/Source/Runtime/Core/Public/HAL/LowLevelMemTracker.h
--- a/Engine/Source/Runtime/Core/Public/HAL/LowLevelMemTracker.h
+++ b/Engine/Source/Runtime/Core/Public/HAL/LowLevelMemTracker.h
@@ -621,41 +621,34 @@ class CORE_API FLLMScope
 public:
 	FLLMScope(FName TagName, bool bIsStatTag, ELLMTagSet InTagSet, ELLMTracker InTracker)
 	{
-		if (FLowLevelMemTracker::bIsDisabled)
+		if (!FLowLevelMemTracker::bIsDisabled)
 		{
-			bEnabled = false;
-			return;
+			Init(TagName, bIsStatTag, InTagSet, InTracker);
 		}
-		Init(TagName, bIsStatTag, InTagSet, InTracker);
 	}
 	FLLMScope(ELLMTag TagEnum, bool bIsStatTag, ELLMTagSet InTagSet, ELLMTracker InTracker)
 	{
-		if (FLowLevelMemTracker::bIsDisabled)
+		if (!FLowLevelMemTracker::bIsDisabled)
 		{
-			bEnabled = false;
-			return;
+			Init(TagEnum, bIsStatTag, InTagSet, InTracker);
 		}
 
-		Init(TagEnum, bIsStatTag, InTagSet, InTracker);
 	}
 
 	FLLMScope(const UE::LLMPrivate::FTagData* TagData, bool bIsStatTag, ELLMTagSet Set, ELLMTracker Tracker)
 	{
-		if (FLowLevelMemTracker::bIsDisabled)
+		if (!FLowLevelMemTracker::bIsDisabled)
 		{
-			bEnabled = false;
-			return;
+			Init(TagData, bIsStatTag, Set, Tracker);
 		}
-		Init(TagData, bIsStatTag, Set, Tracker);
 	}
 
 	~FLLMScope()
 	{
-		if (!bEnabled)
+		if (bEnabled)
 		{
-			return;
+			Destruct();
 		}
-		Destruct();
 	}
 
 protected:
@@ -664,10 +657,10 @@ protected:
 	void Init(const UE::LLMPrivate::FTagData* TagData, bool bIsStatTag, ELLMTagSet InTagSet, ELLMTracker InTracker);
 	void Destruct();
 
-	ELLMTracker Tracker;
-	bool bEnabled;
+	ELLMTracker Tracker{};
+	bool bEnabled = false;
 #if LLM_ALLOW_ASSETS_TAGS
-	bool bIsAssetTag;
+	bool bIsAssetTag = false;
 #endif
 };
 
