Remove default brightness and reset brightness on iOS.

diff --git a/Engine/Source/Runtime/ApplicationCore/Private/IOS/IOSAppDelegate.cpp b/Engine/Source/Runtime/ApplicationCore/Private/IOS/IOSAppDelegate.cpp
--- a/Engine/Source/Runtime/ApplicationCore/Private/IOS/IOSAppDelegate.cpp
+++ b/Engine/Source/Runtime/ApplicationCore/Private/IOS/IOSAppDelegate.cpp
@@ -1277,8 +1277,6 @@ extern EDeviceScreenOrientation ConvertFromUIInterfaceOrientation(UIInterfaceOri
 FCriticalSection RenderSuspend;
 - (void)applicationWillResignActive:(UIApplication *)application
 {
-    FIOSPlatformMisc::ResetBrightness();
-    
     /*
 		Sent when the application is about to move from active to inactive
 		state. This can occur for certain types of temporary interruptions (such
diff --git a/Engine/Source/Runtime/Core/Private/IOS/IOSPlatformMisc.cpp b/Engine/Source/Runtime/Core/Private/IOS/IOSPlatformMisc.cpp
--- a/Engine/Source/Runtime/Core/Private/IOS/IOSPlatformMisc.cpp
+++ b/Engine/Source/Runtime/Core/Private/IOS/IOSPlatformMisc.cpp
@@ -64,7 +64,6 @@ void (* GMemoryWarningHandler)(const FGenericMemoryWarningContext& Context) = NU
 
 /** global for showing the splash screen */
 bool GShowSplashScreen = true;
-float GOriginalBrightness = -1.0f;
 
 static int32 GetFreeMemoryMB()
 {
@@ -75,8 +74,6 @@ static int32 GetFreeMemoryMB()
 void FIOSPlatformMisc::PlatformInit()
 {
 	FAppEntry::PlatformInit();
-    
-    GOriginalBrightness = FIOSPlatformMisc::GetBrightness();
 
 	// Increase the maximum number of simultaneously open files
 	struct rlimit Limit;
@@ -247,14 +244,6 @@ void FIOSPlatformMisc::SetBrightness(float Brightness)
 #endif // !PLATFORM_TVOS
 }
 
-void FIOSPlatformMisc::ResetBrightness()
-{
-	if (GOriginalBrightness >= 0.f)
-	{
-		SetBrightness(GOriginalBrightness);
-	}
-}
-
 bool FIOSPlatformMisc::IsRunningOnBattery()
 {
 	return [[IOSAppDelegate GetDelegate] IsRunningOnBattery];
diff --git a/Engine/Source/Runtime/Core/Public/GenericPlatform/GenericPlatformMisc.h b/Engine/Source/Runtime/Core/Public/GenericPlatform/GenericPlatformMisc.h
--- a/Engine/Source/Runtime/Core/Public/GenericPlatform/GenericPlatformMisc.h
+++ b/Engine/Source/Runtime/Core/Public/GenericPlatform/GenericPlatformMisc.h
@@ -1178,7 +1178,6 @@ public:
 
 	FORCEINLINE static void SetBrightness(float bBright) { }
 	FORCEINLINE static float GetBrightness() { return 1.0f; }
-	FORCEINLINE static void ResetBrightness() { } // resets brightness to brightness application started with
 	FORCEINLINE static bool SupportsBrightness() { return false; }
 
 	FORCEINLINE static bool IsInLowPowerMode() { return false;}
diff --git a/Engine/Source/Runtime/Core/Public/IOS/IOSPlatformMisc.h b/Engine/Source/Runtime/Core/Public/IOS/IOSPlatformMisc.h
--- a/Engine/Source/Runtime/Core/Public/IOS/IOSPlatformMisc.h
+++ b/Engine/Source/Runtime/Core/Public/IOS/IOSPlatformMisc.h
@@ -61,7 +61,6 @@ struct CORE_API FIOSPlatformMisc : public FApplePlatformMisc
 	static int32 GetDeviceVolume();
 	static void SetBrightness(float Brightness);
 	static float GetBrightness();
-	static void ResetBrightness(); //reset brightness to original value the application started with
 	static bool SupportsBrightness() { return true; }
     static bool IsInLowPowerMode();
 
