Specifying additional space for a traced event was only necessary if an event had an "attachment" (an arbitrary unstructured blob of data, "<< EventName.Attachment(...)"). Attachments were removed in 17080821. While the API wasn't updated at the time for transition purposes, it soon will be. (Note there is one slightly confusing element here which is that _important_ events do need to have a size specified due to the way they are buffered internally. The ThreadName events here however are not marked important).

diff --git a/Engine/Source/Runtime/Core/Private/ProfilingDebugging/Microsoft/EventTracingForWindows.cpp b/Engine/Source/Runtime/Core/Private/ProfilingDebugging/Microsoft/EventTracingForWindows.cpp
--- a/Engine/Source/Runtime/Core/Private/ProfilingDebugging/Microsoft/EventTracingForWindows.cpp
+++ b/Engine/Source/Runtime/Core/Private/ProfilingDebugging/Microsoft/EventTracingForWindows.cpp
@@ -164,7 +164,7 @@ void FPlatformEvents::TraceEventCallback(PEVENT_RECORD Event)
 							::CloseHandle(ProcessHandle);
 						}
 
-						UE_TRACE_LOG(PlatformEvent, ThreadName, ContextSwitchChannel, NameLen * sizeof(TCHAR))
+						UE_TRACE_LOG(PlatformEvent, ThreadName, ContextSwitchChannel)
 							<< ThreadName.ThreadId(ThreadId)
 							<< ThreadName.ProcessId(ProcessId)
 							<< ThreadName.Name(Name, NameLen);
