CSA: Disable self-deletion when running clang's static analyzer. Since it cannot reason about what happens to shared pointers/references that get passed into functions it cannot see, it assumes that it may end up in a situation where the ref count may be off and cause self-deletion in cases where the ref controller is shared among multiple instances.

diff --git a/Engine/Source/Runtime/Core/Public/Templates/SharedPointerInternals.h b/Engine/Source/Runtime/Core/Public/Templates/SharedPointerInternals.h
--- a/Engine/Source/Runtime/Core/Public/Templates/SharedPointerInternals.h
+++ b/Engine/Source/Runtime/Core/Public/Templates/SharedPointerInternals.h
@@ -251,8 +251,13 @@ namespace SharedPointerInternals
 				{
 					std::atomic_thread_fence(std::memory_order_acquire);
 
+					// Disable this if running clang's static analyzer. Passing shared pointers
+					// and references to functions it cannot reason about, produces false
+					// positives about use-after-free in the TSharedPtr/TSharedRef destructors.
+#if !defined(__clang_analyzer__)
 					// No more references to this reference count.  Destroy it!
 					delete this;
+#endif
 				}
 			}
 			else
@@ -262,7 +267,9 @@ namespace SharedPointerInternals
 				if( --WeakReferenceCount == 0 )
 				{
 					// No more references to this reference count.  Destroy it!
+#if !defined(__clang_analyzer__)
 					delete this;
+#endif
 				}
 			}
 		}
