Fixed a typo in NEON double version of VectorMaskBits

diff --git a/Engine/Source/Runtime/Core/Public/Math/UnrealMathNeon.h b/Engine/Source/Runtime/Core/Public/Math/UnrealMathNeon.h
--- a/Engine/Source/Runtime/Core/Public/Math/UnrealMathNeon.h
+++ b/Engine/Source/Runtime/Core/Public/Math/UnrealMathNeon.h
@@ -1421,7 +1421,7 @@ FORCEINLINE uint32 VectorMaskBits(VectorRegister4Double VecMask)
 	uint64x2_t mmA = vtstq_u64(vreinterpretq_u64_f64(VecMask.XY), GlobalVectorConstants::DoubleSignBit.XY); // mask with 1s every bit for vector element if it's sign is negative
 	uint64x2_t mmA1 = vtstq_u64(vreinterpretq_u64_f64(VecMask.ZW), GlobalVectorConstants::DoubleSignBit.XY);
 	uint64x2_t mmB = vandq_u64(mmA, MakeVectorRegisterInt64(0x1, 0x2)); // pick only one bit on it's corresponding position
-	uint64x2_t mmB1 = vandq_u64(mmA, MakeVectorRegisterInt64(0x4, 0x8));
+	uint64x2_t mmB1 = vandq_u64(mmA1, MakeVectorRegisterInt64(0x4, 0x8));
 	uint64x2_t mmC = vorrq_u64(mmB, mmB1);								// now combine the result
 	return (uint32)(vgetq_lane_u64(mmC, 0) | vgetq_lane_u64(mmC, 1));     // reduce the result from 2 elements to one
 }
