removed assert if ticker's delegate is removed multiple times

diff --git a/Engine/Source/Runtime/Core/Private/Containers/Ticker.cpp b/Engine/Source/Runtime/Core/Private/Containers/Ticker.cpp
--- a/Engine/Source/Runtime/Core/Private/Containers/Ticker.cpp
+++ b/Engine/Source/Runtime/Core/Private/Containers/Ticker.cpp
@@ -34,9 +34,8 @@ void FTSTicker::RemoveTicker(FDelegateHandle Handle)
 	if (FElementPtr Element = Handle.Pin())
 	{
 		// mark the element as removed and if it's being ticked atm, spin-wait until its execution is finished
-		uint64 PrevState = Element->State.fetch_add(FElement::RemovedState, std::memory_order_acquire); // "acquire" to prevent potential 
+		uint64 PrevState = Element->State.fetch_or(FElement::RemovedState, std::memory_order_acquire); // "acquire" to prevent potential 
 		// resource release after RemoveTicker() to be reordered before it
-		checkf((PrevState & 0x1) == FElement::DefaultState, TEXT("The delegate is already removed (%u)"), PrevState);
 		uint32 ExecutingThreadId = GetThreadId(PrevState);
 		
 		while (ExecutingThreadId != 0 && // is being executed right now
diff --git a/Engine/Source/Runtime/Core/Private/Tests/Containers/TickerTests.cpp b/Engine/Source/Runtime/Core/Private/Tests/Containers/TickerTests.cpp
--- a/Engine/Source/Runtime/Core/Private/Tests/Containers/TickerTests.cpp
+++ b/Engine/Source/Runtime/Core/Private/Tests/Containers/TickerTests.cpp
@@ -143,6 +143,13 @@ bool FTSTickerTest::RunTest(const FString& Parameters)
 		Ticker.Tick(0.0f);
 	}
 
+	{	// same delegate removed multiple times (used to be an assert)
+		FTSTicker Ticker;
+		FTSTicker::FDelegateHandle DelegateHandle = Ticker.AddTicker(UE_SOURCE_LOCATION, 0.0f, [](float) { return true; /* keep ticking*/ });
+		FTSTicker::RemoveTicker(DelegateHandle);
+		FTSTicker::RemoveTicker(DelegateHandle);
+	}
+
 	// multithreaded stress test
 	{
 		using namespace UE::Tasks;
