IoStore: Added fallback to normal hashing when failing to generate a perfect hash map within a reasonable amount of time

diff --git a/Engine/Source/Runtime/Core/Private/IO/IoStore.cpp b/Engine/Source/Runtime/Core/Private/IO/IoStore.cpp
--- a/Engine/Source/Runtime/Core/Private/IO/IoStore.cpp
+++ b/Engine/Source/Runtime/Core/Private/IO/IoStore.cpp
@@ -662,7 +662,7 @@ public:
 		}
 	}
 
-	void GeneratePerfectHashes(FIoStoreTocResource& TocResource)
+	bool GeneratePerfectHashes(FIoStoreTocResource& TocResource)
 	{
 		// https://en.wikipedia.org/wiki/Perfect_hash_function
 		TRACE_CPUPROFILER_EVENT_SCOPE(TocGeneratePerfectHashes);
@@ -673,7 +673,8 @@ public:
 		OutTocChunkIds.SetNum(ChunkCount);
 		TArray<FIoOffsetAndLength> OutTocOffsetAndLengths;
 		OutTocOffsetAndLengths.SetNum(ChunkCount);
-		TocResource.ChunkHashSeeds.SetNumZeroed(ChunkCount);
+		TArray<int32> OutTocChunkHashSeeds;
+		OutTocChunkHashSeeds.SetNumZeroed(ChunkCount);
 		
 		TArray<TArray<int32>> Buckets;
 		Buckets.SetNum(ChunkCount);
@@ -699,16 +700,34 @@ public:
 				break;
 			}
 
-			uint32 Seed = 1;
+			static constexpr uint32 Primes[] =
+			{
+				421,  443,  479,  491,  503,  509,  587,  617,
+				641,  691,  733,  829,  857,  881,  907,  919,
+				953,  977,  1009, 1021, 1039, 1061, 1087, 1097,
+				1117, 1151, 1171, 1193, 1217, 1231, 1259, 1283,
+				1297, 1319, 1321, 1367, 1399, 1427, 1439, 1453,
+				1459, 1483, 1499, 1553, 1571, 1579, 1607, 1663,
+				2129, 2153, 2221, 2389, 2777, 3191, 3467, 3469,
+				4079, 4783, 6011, 6983, 7589, 7741, 7877, 7919
+			};
+			static constexpr uint32 MaxIterations = UE_ARRAY_COUNT(Primes);
+
+			uint32 PrimeIndex = 0;
 			TBitArray BucketUsedSlots(false, ChunkCount);
 			int32 IndexInBucket = 0;
 			while (IndexInBucket < Bucket.Num())
 			{
 				const FIoChunkId& ChunkId = TocResource.ChunkIds[Bucket[IndexInBucket]];
+				uint32 Seed = Primes[PrimeIndex];
 				uint32 Slot = FIoStoreTocResource::HashChunkIdWithSeed(Seed, ChunkId) % ChunkCount;
 				if (!FreeSlots[Slot] || BucketUsedSlots[Slot])
 				{
-					++Seed;
+					++PrimeIndex;
+					if (PrimeIndex == MaxIterations)
+					{
+						return false;
+					}
 					IndexInBucket = 0;
 					BucketUsedSlots.Init(false, ChunkCount);
 				}
@@ -720,7 +739,8 @@ public:
 			}
 
 			uint64 BucketHash = FIoStoreTocResource::HashChunkIdWithSeed(0, TocResource.ChunkIds[Bucket[0]]);
-			TocResource.ChunkHashSeeds[BucketHash % ChunkCount] = Seed;
+			uint32 Seed = Primes[PrimeIndex];
+			OutTocChunkHashSeeds[BucketHash % ChunkCount] = Seed;
 			for (IndexInBucket = 0; IndexInBucket < Bucket.Num(); ++IndexInBucket)
 			{
 				int32 ChunkIndex = Bucket[IndexInBucket];
@@ -746,14 +766,17 @@ public:
 				int32 ChunkIndex = Bucket[0];
 				const FIoChunkId& ChunkId = TocResource.ChunkIds[ChunkIndex];
 				uint64 BucketHash = FIoStoreTocResource::HashChunkIdWithSeed(0, ChunkId);
-				TocResource.ChunkHashSeeds[BucketHash % ChunkCount] = -static_cast<int32>(Slot) - 1;
+				OutTocChunkHashSeeds[BucketHash % ChunkCount] = -static_cast<int32>(Slot) - 1;
 				OutTocChunkIds[Slot] = ChunkId;
 				OutTocOffsetAndLengths[Slot] = TocResource.ChunkOffsetLengths[ChunkIndex];
 			}
 		}
 
-		TocResource.ChunkIds = OutTocChunkIds;
-		TocResource.ChunkOffsetLengths = OutTocOffsetAndLengths;
+		TocResource.ChunkIds = MoveTemp(OutTocChunkIds);
+		TocResource.ChunkOffsetLengths = MoveTemp(OutTocOffsetAndLengths);
+		TocResource.ChunkHashSeeds = MoveTemp(OutTocChunkHashSeeds);
+
+		return true;
 	}
 
 	UE_NODISCARD TIoStatusOr<FIoStoreWriterResult> Flush()
@@ -2127,8 +2150,8 @@ FIoStatus FIoStoreTocResource::Read(const TCHAR* TocFilePath, EIoStoreTocReadOpt
 	if (Header.Version >= static_cast<uint8>(EIoStoreTocVersion::PerfectHash))
 	{
 		const int32* ChunkHashSeeds = reinterpret_cast<const int32*>(DataPtr);
-		OutTocResource.ChunkHashSeeds = MakeArrayView<int32 const>(ChunkHashSeeds, Header.TocEntryCount);
-		DataPtr += Header.TocEntryCount * sizeof(int32);
+		OutTocResource.ChunkHashSeeds = MakeArrayView<int32 const>(ChunkHashSeeds, Header.TocChunkHashSeedsCount);
+		DataPtr += Header.TocChunkHashSeedsCount * sizeof(int32);
 	}
 
 	// Compression blocks
@@ -2239,6 +2262,7 @@ TIoStatusOr<uint64> FIoStoreTocResource::Write(
 	TocHeader.Version = static_cast<uint8>(EIoStoreTocVersion::Latest);
 	TocHeader.TocHeaderSize = sizeof(TocHeader);
 	TocHeader.TocEntryCount = TocResource.ChunkIds.Num();
+	TocHeader.TocChunkHashSeedsCount = TocResource.ChunkHashSeeds.Num();
 	TocHeader.TocCompressedBlockEntryCount = TocResource.CompressionBlocks.Num();
 	TocHeader.TocCompressedBlockEntrySize = sizeof(FIoStoreTocCompressedBlockEntry);
 	TocHeader.CompressionBlockSize = uint32(WriterSettings.CompressionBlockSize);
diff --git a/Engine/Source/Runtime/Core/Private/IO/IoStore.h b/Engine/Source/Runtime/Core/Private/IO/IoStore.h
--- a/Engine/Source/Runtime/Core/Private/IO/IoStore.h
+++ b/Engine/Source/Runtime/Core/Private/IO/IoStore.h
@@ -44,7 +44,7 @@ struct FIoStoreTocHeader
 	EIoContainerFlags ContainerFlags;
 	uint8	Reserved3 = 0;
 	uint16	Reserved4 = 0;
-	uint32	Reserved5 = 0;
+	uint32	TocChunkHashSeedsCount = 0;
 	uint64	PartitionSize = 0;
 	uint64	Reserved6[6] = { 0 };
 
