int64 FMath::FloorToInt64(double)

Next step would bechanging signature of int32 FMath::FloorToInt(double) to return int64 and remove FloorToInt64 but this means hunting down conversion loss errors across projects.

diff --git a/Engine/Source/Runtime/Core/Public/GenericPlatform/GenericPlatformMath.h b/Engine/Source/Runtime/Core/Public/GenericPlatform/GenericPlatformMath.h
--- a/Engine/Source/Runtime/Core/Public/GenericPlatform/GenericPlatformMath.h
+++ b/Engine/Source/Runtime/Core/Public/GenericPlatform/GenericPlatformMath.h
@@ -241,7 +241,13 @@ struct FGenericPlatformMath
 		I -= ((float)I > F);
 		return I;
 	}
-
+	static FORCEINLINE int64 FloorToInt64(double F)
+	{
+		int64 I = (int64)F;
+		I -= ((double)I > F);
+		return I;
+	}
+	
 	/**
 	* Converts a float to the nearest less or equal integer.
 	* @param F		Floating point value to convert
diff --git a/Engine/Source/Runtime/Core/Public/Math/UnrealPlatformMathSSE.h b/Engine/Source/Runtime/Core/Public/Math/UnrealPlatformMathSSE.h
--- a/Engine/Source/Runtime/Core/Public/Math/UnrealPlatformMathSSE.h
+++ b/Engine/Source/Runtime/Core/Public/Math/UnrealPlatformMathSSE.h
@@ -243,6 +243,12 @@ struct TUnrealPlatformMathSSEBase : public Base
 		return UE4::SSE::FloorToInt(F);
 	}
 
+	template<typename T>
+	static FORCEINLINE int64 FloorToInt64(T F)
+	{
+		return UE4::SSE::FloorToInt64(F);
+	}
+
 	template<typename T>
 	static FORCEINLINE int32 CeilToInt(T F)
 	{
