XSXIoDispatcher: Use immediate direct storage hardware decompression for unaligned compression block reads (instead of the cpu decompression fallback).
Fix broken trace counter for IoDispatcherFileBackend/QueuedFilesystemReads.
Use 64K io dispatcher buffers on XSX to match its 64K compression block size.

#FROST-4056
#rnx

diff --git a/Engine/Source/Runtime/Core/Private/IO/IoDispatcherFileBackend.cpp b/Engine/Source/Runtime/Core/Private/IO/IoDispatcherFileBackend.cpp
--- a/Engine/Source/Runtime/Core/Private/IO/IoDispatcherFileBackend.cpp
+++ b/Engine/Source/Runtime/Core/Private/IO/IoDispatcherFileBackend.cpp
@@ -128,11 +128,12 @@ private:
 	IMappedFileHandle* SharedMappedFileHandle;
 };
 
-void FFileIoStoreBufferAllocator::Initialize(uint64 MemorySize, uint64 BufferSize, uint32 BufferAlignment)
+void FFileIoStoreBufferAllocator::Initialize(uint64 InMemorySize, uint64 InBufferSize, uint32 InBufferAlignment)
 {
-	uint64 BufferCount = MemorySize / BufferSize;
-	MemorySize = BufferCount * BufferSize;
-	BufferMemory = reinterpret_cast<uint8*>(FMemory::Malloc(MemorySize, BufferAlignment));
+	uint64 BufferCount = InMemorySize / InBufferSize;
+	uint64 MemorySize = BufferCount * InBufferSize;
+	BufferMemory = reinterpret_cast<uint8*>(FMemory::Malloc(MemorySize, InBufferAlignment));
+	BufferSize = InBufferSize;
 	for (uint64 BufferIndex = 0; BufferIndex < BufferCount; ++BufferIndex)
 	{
 		FFileIoStoreBuffer* Buffer = new FFileIoStoreBuffer();
@@ -2077,7 +2078,7 @@ void FFileIoStats::OnFilesystemReadsQueued(uint64 BytesToRead, int32 NumReads)
 	QueuedFilesystemReads += NumReads;
 #endif
 
-	TRACE_COUNTER_INCREMENT(IoDispatcherFileBackendQueuedFilesystemReads)
+	TRACE_COUNTER_ADD(IoDispatcherFileBackendQueuedFilesystemReads, NumReads)
 	TRACE_COUNTER_ADD(IoDispatcherFileBackendQueuedFilesystemReadBytes, BytesToRead);
 #endif
 }
@@ -2145,6 +2146,7 @@ void FFileIoStats::OnReadComplete(int64 BytesRead)
 
 	TRACE_COUNTER_SUBTRACT(IoDispatcherFileBackendQueuedFilesystemReadBytes, BytesRead);
 	TRACE_COUNTER_ADD(IoDispatcherFileBackendTotalBytesRead, BytesRead);
+	TRACE_COUNTER_DECREMENT(IoDispatcherFileBackendQueuedFilesystemReads);
 #endif
 }
 
diff --git a/Engine/Source/Runtime/Core/Private/IO/IoDispatcherFileBackendTypes.h b/Engine/Source/Runtime/Core/Private/IO/IoDispatcherFileBackendTypes.h
--- a/Engine/Source/Runtime/Core/Private/IO/IoDispatcherFileBackendTypes.h
+++ b/Engine/Source/Runtime/Core/Private/IO/IoDispatcherFileBackendTypes.h
@@ -87,10 +87,10 @@ struct FFileIoStoreCompressedBlock
 	FFileIoStoreCompressedBlock* Next = nullptr;
 	FFileIoStoreBlockKey Key;
 	FName CompressionMethod;
-	uint64 RawOffset;
-	uint32 UncompressedSize;
-	uint32 CompressedSize;
-	uint32 RawSize;
+	uint64 RawOffset = uint64(-1);
+	uint32 UncompressedSize = uint32(-1);
+	uint32 CompressedSize = uint32(-1);
+	uint32 RawSize = uint32(-1);
 	uint32 RefCount = 0;
 	uint32 UnfinishedRawBlocksCount = 0;
 	TArray<struct FFileIoStoreReadRequest*, TInlineAllocator<2>> RawBlocks;
@@ -132,6 +132,7 @@ struct FFileIoStoreReadRequest
 	uint64 CreationTime;	// Potentially used to circuit break request ordering optimizations when outstanding requests have been delayed too long
 	FFileIoStoreBlockScatter ImmediateScatter;
 	uint32 BytesUsed = 0;
+	bool bIsCustomRequest = false;
 	bool bFailed = false;
 	bool bCancelled = false;
 	EQueueStatus QueueStatus = QueueStatus_NotInQueue;
@@ -516,8 +517,10 @@ public:
 	void Initialize(uint64 MemorySize, uint64 BufferSize, uint32 BufferAlignment);
 	FFileIoStoreBuffer* AllocBuffer();
 	void FreeBuffer(FFileIoStoreBuffer* Buffer);
+	uint64 GetBufferSize() const { return BufferSize; }
 
 private:
+	uint64 BufferSize = 0;
 	uint8* BufferMemory = nullptr;
 	FCriticalSection BuffersCritical;
 	FFileIoStoreBuffer* FirstFreeBuffer = nullptr;
