[CSV] - add ExtendedUsedMB to the CSV output as a measure of the amount of extended memory currently in use.

diff --git a/Engine/Source/Runtime/Core/Private/ProfilingDebugging/CsvProfiler.cpp b/Engine/Source/Runtime/Core/Private/ProfilingDebugging/CsvProfiler.cpp
--- a/Engine/Source/Runtime/Core/Private/ProfilingDebugging/CsvProfiler.cpp
+++ b/Engine/Source/Runtime/Core/Private/ProfilingDebugging/CsvProfiler.cpp
@@ -2745,16 +2745,19 @@ void FCsvProfiler::EndFrame()
 
 		FPlatformMemoryStats MemoryStats = FPlatformMemory::GetStats();
 		float PhysicalMBFree = float(MemoryStats.AvailablePhysical) / (1024.0f * 1024.0f);
-
+		float UsedExtendedMB = 0;
 #if !UE_BUILD_SHIPPING
 		// Subtract any extra development memory from physical free. This can result in negative values in cases where we would have crashed OOM
 		PhysicalMBFree -= float(FPlatformMemory::GetExtraDevelopmentMemorySize() / 1024ull / 1024ull);
+		UsedExtendedMB = PhysicalMBFree < 0.0f ? -PhysicalMBFree : 0;
 #endif
 		float PhysicalMBUsed = float(MemoryStats.UsedPhysical) / (1024.0f * 1024.0f);
-		float VirtualMBUsed  = float(MemoryStats.UsedVirtual) / (1024.0f * 1024.0f);
+		float VirtualMBUsed = float(MemoryStats.UsedVirtual) / (1024.0f * 1024.0f);
+		
 		CSV_CUSTOM_STAT_GLOBAL(MemoryFreeMB, PhysicalMBFree, ECsvCustomStatOp::Set);
 		CSV_CUSTOM_STAT_GLOBAL(PhysicalUsedMB, PhysicalMBUsed, ECsvCustomStatOp::Set);
 		CSV_CUSTOM_STAT_GLOBAL(VirtualUsedMB, VirtualMBUsed, ECsvCustomStatOp::Set);
+		CSV_CUSTOM_STAT_GLOBAL(ExtendedUsedMB, UsedExtendedMB, ECsvCustomStatOp::Set);
 
 		// If we're single-threaded, process the stat data here
 		if (ProcessingThread == nullptr)
