CompactBinary: Added FIoHashBuilder and switched CompactBinary to use it

#rnx

diff --git a/Engine/Source/Runtime/Core/Private/Serialization/CompactBinary.cpp b/Engine/Source/Runtime/Core/Private/Serialization/CompactBinary.cpp
--- a/Engine/Source/Runtime/Core/Private/Serialization/CompactBinary.cpp
+++ b/Engine/Source/Runtime/Core/Private/Serialization/CompactBinary.cpp
@@ -469,14 +469,14 @@ uint64 FCbFieldView::GetValueSize() const
 	}
 }
 
-FBlake3Hash FCbFieldView::GetHash() const
+FIoHash FCbFieldView::GetHash() const
 {
-	FBlake3 Hash;
+	FIoHashBuilder Hash;
 	AppendHash(Hash);
 	return Hash.Finalize();
 }
 
-void FCbFieldView::AppendHash(FBlake3& Builder) const
+void FCbFieldView::AppendHash(FIoHashBuilder& Builder) const
 {
 	const ECbFieldType SerializedType = FCbFieldType::GetSerializedType(TypeWithFlags);
 	Builder.Update(&SerializedType, sizeof(SerializedType));
@@ -595,14 +595,14 @@ uint64 FCbArrayView::GetSize() const
 	return sizeof(ECbFieldType) + GetValueSize();
 }
 
-FBlake3Hash FCbArrayView::GetHash() const
+FIoHash FCbArrayView::GetHash() const
 {
-	FBlake3 Hash;
+	FIoHashBuilder Hash;
 	AppendHash(Hash);
 	return Hash.Finalize();
 }
 
-void FCbArrayView::AppendHash(FBlake3& Builder) const
+void FCbArrayView::AppendHash(FIoHashBuilder& Builder) const
 {
 	const ECbFieldType SerializedType = GetType();
 	Builder.Update(&SerializedType, sizeof(SerializedType));
@@ -675,14 +675,14 @@ uint64 FCbObjectView::GetSize() const
 	return sizeof(ECbFieldType) + GetValueSize();
 }
 
-FBlake3Hash FCbObjectView::GetHash() const
+FIoHash FCbObjectView::GetHash() const
 {
-	FBlake3 Hash;
+	FIoHashBuilder Hash;
 	AppendHash(Hash);
 	return Hash.Finalize();
 }
 
-void FCbObjectView::AppendHash(FBlake3& Builder) const
+void FCbObjectView::AppendHash(FIoHashBuilder& Builder) const
 {
 	const ECbFieldType SerializedType = GetType();
 	Builder.Update(&SerializedType, sizeof(SerializedType));
@@ -749,15 +749,15 @@ uint64 TCbFieldIterator<FieldType>::GetRangeSize() const
 }
 
 template <typename FieldType>
-FBlake3Hash TCbFieldIterator<FieldType>::GetRangeHash() const
+FIoHash TCbFieldIterator<FieldType>::GetRangeHash() const
 {
-	FBlake3 Hash;
+	FIoHashBuilder Hash;
 	AppendRangeHash(Hash);
 	return Hash.Finalize();
 }
 
 template <typename FieldType>
-void TCbFieldIterator<FieldType>::AppendRangeHash(FBlake3& Builder) const
+void TCbFieldIterator<FieldType>::AppendRangeHash(FIoHashBuilder& Builder) const
 {
 	FMemoryView View;
 	if (TryGetRangeView(View))
diff --git a/Engine/Source/Runtime/Core/Private/Tests/Serialization/CompactBinaryTest.cpp b/Engine/Source/Runtime/Core/Private/Tests/Serialization/CompactBinaryTest.cpp
--- a/Engine/Source/Runtime/Core/Private/Tests/Serialization/CompactBinaryTest.cpp
+++ b/Engine/Source/Runtime/Core/Private/Tests/Serialization/CompactBinaryTest.cpp
@@ -183,7 +183,7 @@ bool FCbFieldNoneTest::RunTest(const FString& Parameters)
 		TestFalse(TEXT("!FCbFieldView()::HasValue()"), DefaultField.HasValue());
 		TestFalse(TEXT("!FCbFieldView()::HasError()"), DefaultField.HasError());
 		TestEqual(TEXT("FCbFieldView()::GetError() == None"), DefaultField.GetError(), ECbFieldError::None);
-		TestEqual(TEXT("FCbFieldView()::GetHash()"), DefaultField.GetHash(), FBlake3::HashBuffer(MakeMemoryView<uint8>({uint8(ECbFieldType::None)})));
+		TestEqual(TEXT("FCbFieldView()::GetHash()"), DefaultField.GetHash(), FIoHash::HashBuffer(MakeMemoryView<uint8>({uint8(ECbFieldType::None)})));
 		FMemoryView View;
 		TestFalse(TEXT("FCbFieldView()::TryGetView()"), DefaultField.TryGetView(View));
 	}
@@ -211,7 +211,7 @@ bool FCbFieldNoneTest::RunTest(const FString& Parameters)
 		TestEqual(TEXT("FCbFieldView(None|Type|Name)::GetName()"), NoneField.GetName(), "Name"_U8SV);
 		TestTrue(TEXT("FCbFieldView(None|Type|Name)::HasName()"), NoneField.HasName());
 		TestFalse(TEXT("!FCbFieldView(None|Type|Name)::HasValue()"), NoneField.HasValue());
-		TestEqual(TEXT("FCbFieldView(None|Type|Name)::GetHash()"), NoneField.GetHash(), FBlake3::HashBuffer(MakeMemoryView(NoneBytes)));
+		TestEqual(TEXT("FCbFieldView(None|Type|Name)::GetHash()"), NoneField.GetHash(), FIoHash::HashBuffer(MakeMemoryView(NoneBytes)));
 		FMemoryView View;
 		TestTrue(TEXT("FCbFieldView(None|Type|Name)::TryGetView()"), NoneField.TryGetView(View) && View == MakeMemoryView(NoneBytes));
 
@@ -243,7 +243,7 @@ bool FCbFieldNoneTest::RunTest(const FString& Parameters)
 		TestEqual(TEXT("FCbFieldView(None|Name)::GetName()"), NoneField.GetName(), "Name"_U8SV);
 		TestTrue(TEXT("FCbFieldView(None|Name)::HasName()"), NoneField.HasName());
 		TestFalse(TEXT("!FCbFieldView(None|Name)::HasValue()"), NoneField.HasValue());
-		TestEqual(TEXT("FCbFieldView(None|Name)::GetHash()"), NoneField.GetHash(), FBlake3::HashBuffer(MakeMemoryView(NoneBytes)));
+		TestEqual(TEXT("FCbFieldView(None|Name)::GetHash()"), NoneField.GetHash(), FIoHash::HashBuffer(MakeMemoryView(NoneBytes)));
 		FMemoryView View;
 		TestFalse(TEXT("FCbFieldView(None|Name)::TryGetView()"), NoneField.TryGetView(View));
 
@@ -261,7 +261,7 @@ bool FCbFieldNoneTest::RunTest(const FString& Parameters)
 		TestEqual(TEXT("FCbFieldView(None|EmptyName)::GetName()"), NoneField.GetName(), ""_U8SV);
 		TestTrue(TEXT("FCbFieldView(None|EmptyName)::HasName()"), NoneField.HasName());
 		TestFalse(TEXT("!FCbFieldView(None|EmptyName)::HasValue()"), NoneField.HasValue());
-		TestEqual(TEXT("FCbFieldView(None|EmptyName)::GetHash()"), NoneField.GetHash(), FBlake3::HashBuffer(MakeMemoryView(NoneBytes)));
+		TestEqual(TEXT("FCbFieldView(None|EmptyName)::GetHash()"), NoneField.GetHash(), FIoHash::HashBuffer(MakeMemoryView(NoneBytes)));
 		FMemoryView View;
 		TestFalse(TEXT("FCbFieldView(None|EmptyName)::TryGetView()"), NoneField.TryGetView(View));
 	}
@@ -280,7 +280,7 @@ bool FCbFieldNullTest::RunTest(const FString& Parameters)
 		TestTrue(TEXT("FCbFieldView(Null)::HasValue()"), NullField.HasValue());
 		TestFalse(TEXT("!FCbFieldView(Null)::HasError()"), NullField.HasError());
 		TestEqual(TEXT("FCbFieldView(Null)::GetError() == None"), NullField.GetError(), ECbFieldError::None);
-		TestEqual(TEXT("FCbFieldView(Null)::GetHash()"), NullField.GetHash(), FBlake3::HashBuffer(MakeMemoryView<uint8>({uint8(ECbFieldType::Null)})));
+		TestEqual(TEXT("FCbFieldView(Null)::GetHash()"), NullField.GetHash(), FIoHash::HashBuffer(MakeMemoryView<uint8>({uint8(ECbFieldType::Null)})));
 	}
 
 	// Test FCbFieldView(None) as Null
@@ -1707,8 +1707,8 @@ bool FCbFieldBufferIteratorTest::RunTest(const FString& Parameters)
 		const FCbFieldViewIterator FieldViewIt = FCbFieldViewIterator::MakeRange(View);
 		const FCbFieldIterator FieldIt = FCbFieldIterator::MakeRange(View);
 
-		TestEqual(TEXT("FCbFieldViewIterator::GetRangeHash()"), FieldViewIt.GetRangeHash(), FBlake3::HashBuffer(View));
-		TestEqual(TEXT("FCbFieldIterator::GetRangeHash()"), FieldIt.GetRangeHash(), FBlake3::HashBuffer(View));
+		TestEqual(TEXT("FCbFieldViewIterator::GetRangeHash()"), FieldViewIt.GetRangeHash(), FIoHash::HashBuffer(View));
+		TestEqual(TEXT("FCbFieldIterator::GetRangeHash()"), FieldIt.GetRangeHash(), FIoHash::HashBuffer(View));
 
 		FMemoryView RangeView;
 		TestTrue(TEXT("FCbFieldViewIterator::TryGetRangeView()"), FieldViewIt.TryGetRangeView(RangeView) && RangeView == MakeMemoryView(Value));
diff --git a/Engine/Source/Runtime/Core/Public/IO/IoHash.h b/Engine/Source/Runtime/Core/Public/IO/IoHash.h
--- a/Engine/Source/Runtime/Core/Public/IO/IoHash.h
+++ b/Engine/Source/Runtime/Core/Public/IO/IoHash.h
@@ -19,8 +19,8 @@
  * Only the leading 160 bits of the 256-bit hash are used, to provide strong collision resistance
  * while minimizing the size of the hash.
  *
- * When the data to hash is not in a contiguous region of memory, FBlake3 can be used to hash the
- * data in blocks with FBlake3::Update(Block) followed by FIoHash(FBlake3::Finalize()).
+ * When the data to hash is not in a contiguous region of memory, use FIoHashBuilder to hash data
+ * in blocks with FIoHashBuilder::Update(Block) followed by FIoHashBuilder::Finalize().
  */
 struct FIoHash
 {
@@ -113,17 +113,17 @@ inline bool FIoHash::IsZero() const
 
 inline FIoHash FIoHash::HashBuffer(FMemoryView View)
 {
-	return FIoHash(FBlake3::HashBuffer(View));
+	return FBlake3::HashBuffer(View);
 }
 
 inline FIoHash FIoHash::HashBuffer(const void* Data, uint64 Size)
 {
-	return FIoHash(FBlake3::HashBuffer(Data, Size));
+	return FBlake3::HashBuffer(Data, Size);
 }
 
 inline FIoHash FIoHash::HashBuffer(const FCompositeBuffer& Buffer)
 {
-	return FIoHash(FBlake3::HashBuffer(Buffer));
+	return FBlake3::HashBuffer(Buffer);
 }
 
 inline bool operator==(const FIoHash& A, const FIoHash& B)
@@ -167,3 +167,22 @@ inline void LexFromString(FIoHash& OutHash, const TCHAR* Buffer)
 
 /** Convert a hash to a 40-character hex string. */
 [[nodiscard]] CORE_API FString LexToString(const FIoHash& Hash);
+
+///////////////////////////////////////////////////////////////////////////////////////////////////////////////////////
+
+/** Calculates a BLAKE3-160 hash. */
+class FIoHashBuilder : public FBlake3
+{
+public:
+	/**
+	 * Finalize the hash of the input data.
+	 *
+	 * May be called any number of times, and more input may be added after.
+	 */
+	[[nodiscard]] inline FIoHash Finalize() const { return FBlake3::Finalize(); }
+
+	/** Calculate the hash of the buffer. */
+	[[nodiscard]] inline static FIoHash HashBuffer(FMemoryView View) { return FBlake3::HashBuffer(View); }
+	[[nodiscard]] inline static FIoHash HashBuffer(const void* Data, uint64 Size) { return FBlake3::HashBuffer(Data, Size); }
+	[[nodiscard]] inline static FIoHash HashBuffer(const FCompositeBuffer& Buffer) { return FBlake3::HashBuffer(Buffer); }
+};
diff --git a/Engine/Source/Runtime/Core/Public/Serialization/CompactBinary.h b/Engine/Source/Runtime/Core/Public/Serialization/CompactBinary.h
--- a/Engine/Source/Runtime/Core/Public/Serialization/CompactBinary.h
+++ b/Engine/Source/Runtime/Core/Public/Serialization/CompactBinary.h
@@ -64,7 +64,7 @@
  */
 
 class FArchive;
-class FBlake3;
+class FIoHashBuilder;
 class FCbArrayView;
 class FCbField;
 class FCbFieldIterator;
@@ -355,9 +355,9 @@ public:
 	CORE_API uint64 GetRangeSize() const;
 
 	/** Calculate the hash of every field in the range. */
-	CORE_API FBlake3Hash GetRangeHash() const;
+	CORE_API FIoHash GetRangeHash() const;
 	/** Append the hash of every field in the range. */
-	CORE_API void AppendRangeHash(FBlake3& Builder) const;
+	CORE_API void AppendRangeHash(FIoHashBuilder& Builder) const;
 
 	using FieldType::Equals;
 
@@ -753,9 +753,9 @@ public:
 	CORE_API uint64 GetSize() const;
 
 	/** Calculate the hash of the field, including the type and name. */
-	CORE_API FBlake3Hash GetHash() const;
+	CORE_API FIoHash GetHash() const;
 	/** Append the hash of the field, including the type and name. */
-	CORE_API void AppendHash(FBlake3& Builder) const;
+	CORE_API void AppendHash(FIoHashBuilder& Builder) const;
 
 	/**
 	 * Whether this field is identical to the other field.
@@ -941,9 +941,9 @@ public:
 	CORE_API uint64 GetSize() const;
 
 	/** Calculate the hash of the array if serialized by itself with no name. */
-	CORE_API FBlake3Hash GetHash() const;
+	CORE_API FIoHash GetHash() const;
 	/** Append the hash of the array if serialized by itself with no name. */
-	CORE_API void AppendHash(FBlake3& Builder) const;
+	CORE_API void AppendHash(FIoHashBuilder& Builder) const;
 
 	/**
 	 * Whether this array is identical to the other array.
@@ -1038,9 +1038,9 @@ public:
 	CORE_API uint64 GetSize() const;
 
 	/** Calculate the hash of the object if serialized by itself with no name. */
-	CORE_API FBlake3Hash GetHash() const;
+	CORE_API FIoHash GetHash() const;
 	/** Append the hash of the object if serialized by itself with no name. */
-	CORE_API void AppendHash(FBlake3& Builder) const;
+	CORE_API void AppendHash(FIoHashBuilder& Builder) const;
 
 	/**
 	 * Whether this object is identical to the other object.
