CSA: Set a zero initializer for TFunctionRefBase::Callable and TDebugHelper::Ptr to avoid warnings from clang's static analyzer.

diff --git a/Engine/Source/Runtime/Core/Public/Templates/Function.h b/Engine/Source/Runtime/Core/Public/Templates/Function.h
--- a/Engine/Source/Runtime/Core/Public/Templates/Function.h
+++ b/Engine/Source/Runtime/Core/Public/Templates/Function.h
@@ -521,7 +521,7 @@ namespace UE4Function_Private
 		template <typename T>
 		struct TDebugHelper : IDebugHelper
 		{
-			T* Ptr;
+			T* Ptr = nullptr;
 		};
 	#endif
 
@@ -692,7 +692,7 @@ namespace UE4Function_Private
 
 	private:
 		// A pointer to a function which invokes the call operator on the callable object
-		Ret (*Callable)(void*, ParamTypes&...);
+		Ret (*Callable)(void*, ParamTypes&...) = nullptr;
 
 		StorageType Storage;
 
@@ -786,7 +786,7 @@ namespace UE4Function_Private
 
 	private:
 		// A pointer to the callable object
-		void* Ptr;
+		void* Ptr = nullptr;
 	};
 }
 
