Be explicit about the "fork" argument when launching the trace store server. Previously it was implicit preventing the addition of useful help text to anyone who ran the store binary explicitly. The process is also launched "detached" so that it can be in the console sub-system without popping up conhost terminals.

#rnx

diff --git a/Engine/Source/Runtime/Core/Private/ProfilingDebugging/TraceAuxiliary.cpp b/Engine/Source/Runtime/Core/Private/ProfilingDebugging/TraceAuxiliary.cpp
--- a/Engine/Source/Runtime/Core/Private/ProfilingDebugging/TraceAuxiliary.cpp
+++ b/Engine/Source/Runtime/Core/Private/ProfilingDebugging/TraceAuxiliary.cpp
@@ -507,19 +507,20 @@ static void LaunchUnrealTraceInternal()
 		return;
 	}
 
-	FString BinPath = "\"";
-	BinPath += FPaths::EngineDir();
-	BinPath += "/Binaries/Win64/UnrealTraceServer.exe\"";
+	FString CommandLine = "\"";
+	CommandLine += FPaths::EngineDir();
+	CommandLine += "/Binaries/Win64/UnrealTraceServer.exe\"";
+	CommandLine += " fork";
 
-	uint32 Flags = 0x0100'0000; // CREATE_BREAKAWAY_FROM_JOB
+	uint32 Flags = CREATE_BREAKAWAY_FROM_JOB|DETACHED_PROCESS;
 	STARTUPINFOW StartupInfo = { sizeof(STARTUPINFOW) };
 	PROCESS_INFORMATION ProcessInfo = {};
-	BOOL bOk = CreateProcessW(nullptr, LPWSTR(*BinPath), nullptr, nullptr,
+	BOOL bOk = CreateProcessW(nullptr, LPWSTR(*CommandLine), nullptr, nullptr,
 		false, Flags, nullptr, nullptr, &StartupInfo, &ProcessInfo);
 
 	if (!bOk)
 	{
-		UE_LOG(LogCore, Display, TEXT("UnrealTraceServer: Unable to launch the trace store from '%s' (%08x)"), *BinPath, GetLastError());
+		UE_LOG(LogCore, Display, TEXT("UnrealTraceServer: Unable to launch the trace store with '%s' (%08x)"), *CommandLine, GetLastError());
 		return;
 	}
 
