ElementType alias for TUniqueObj, TOptional

diff --git a/Engine/Source/Runtime/Core/Public/Misc/Optional.h b/Engine/Source/Runtime/Core/Public/Misc/Optional.h
--- a/Engine/Source/Runtime/Core/Public/Misc/Optional.h
+++ b/Engine/Source/Runtime/Core/Public/Misc/Optional.h
@@ -16,6 +16,8 @@ template<typename OptionalType>
 struct TOptional
 {
 public:
+	using ElementType = OptionalType;
+	
 	/** Construct an OptionaType with a valid value. */
 	TOptional(const OptionalType& InValue)
 	{
diff --git a/Engine/Source/Runtime/Core/Public/Templates/UniqueObj.h b/Engine/Source/Runtime/Core/Public/Templates/UniqueObj.h
--- a/Engine/Source/Runtime/Core/Public/Templates/UniqueObj.h
+++ b/Engine/Source/Runtime/Core/Public/Templates/UniqueObj.h
@@ -12,6 +12,8 @@ template <typename T>
 class TUniqueObj
 {
 public:
+	using ElementType = T;
+
 	TUniqueObj(const TUniqueObj& other)
 		: Obj(MakeUnique<T>(*other))
 	{
