CompressedBuffer: Removed the FormatName functionality and update its last remaining callers

#rnx

diff --git a/Engine/Source/Developer/DerivedDataCache/Private/DerivedDataPayloadPrivate.h b/Engine/Source/Developer/DerivedDataCache/Private/DerivedDataPayloadPrivate.h
--- a/Engine/Source/Developer/DerivedDataCache/Private/DerivedDataPayloadPrivate.h
+++ b/Engine/Source/Developer/DerivedDataCache/Private/DerivedDataPayloadPrivate.h
@@ -3,12 +3,13 @@
 #pragma once
 
 #include "CoreTypes.h"
+#include "Compression/CompressedBuffer.h"
 #include "Compression/OodleDataCompression.h"
 
 namespace UE::DerivedData
 {
 
-constexpr FOodleDataCompression::ECompressor GDefaultCompressor = FOodleDataCompression::ECompressor::Mermaid;
-constexpr FOodleDataCompression::ECompressionLevel GDefaultCompressionLevel = FOodleDataCompression::ECompressionLevel::VeryFast;
+constexpr ECompressedBufferCompressor GDefaultCompressor = ECompressedBufferCompressor::Mermaid;
+constexpr ECompressedBufferCompressionLevel GDefaultCompressionLevel = ECompressedBufferCompressionLevel::VeryFast;
 
 } // UE::DerivedData
diff --git a/Engine/Source/Developer/DesktopPlatform/Private/TargetReceiptBuildWorker.cpp b/Engine/Source/Developer/DesktopPlatform/Private/TargetReceiptBuildWorker.cpp
--- a/Engine/Source/Developer/DesktopPlatform/Private/TargetReceiptBuildWorker.cpp
+++ b/Engine/Source/Developer/DesktopPlatform/Private/TargetReceiptBuildWorker.cpp
@@ -113,8 +113,8 @@ public:
 							{
 								FileDataBuffers.Emplace(FCompressedBuffer::Compress(
 										MutableBuffer.MoveToShared(),
-										FOodleDataCompression::ECompressor::NotSet,
-										FOodleDataCompression::ECompressionLevel::None
+										ECompressedBufferCompressor::NotSet,
+										ECompressedBufferCompressionLevel::None
 									));
 							}
 						}
diff --git a/Engine/Source/Runtime/Core/Private/Compression/CompressedBuffer.cpp b/Engine/Source/Runtime/Core/Private/Compression/CompressedBuffer.cpp
--- a/Engine/Source/Runtime/Core/Private/Compression/CompressedBuffer.cpp
+++ b/Engine/Source/Runtime/Core/Private/Compression/CompressedBuffer.cpp
@@ -10,7 +10,6 @@
 #include "Misc/Crc.h"
 #include "Serialization/Archive.h"
 #include "Templates/Function.h"
-#include "UObject/NameTypes.h"
 
 THIRD_PARTY_INCLUDES_START
 #include "Compression/lz4.h"
@@ -476,7 +475,7 @@ bool FBlockDecoder::TryDecompressTo(const FHeader& Header, const FCompositeBuffe
 class FOodleEncoder final : public FBlockEncoder
 {
 public:
-	FOodleEncoder(FOodleDataCompression::ECompressor InCompressor, FOodleDataCompression::ECompressionLevel InCompressionLevel)
+	FOodleEncoder(ECompressedBufferCompressor InCompressor, ECompressedBufferCompressionLevel InCompressionLevel)
 		: Compressor(InCompressor)
 		, CompressionLevel(InCompressionLevel)
 	{
@@ -503,8 +502,8 @@ protected:
 	}
 
 private:
-	const FOodleDataCompression::ECompressor Compressor;
-	const FOodleDataCompression::ECompressionLevel CompressionLevel;
+	const ECompressedBufferCompressor Compressor;
+	const ECompressedBufferCompressionLevel CompressionLevel;
 };
 
 class FOodleDecoder final : public FBlockDecoder
@@ -588,41 +587,6 @@ static const FDecoder* GetDecoder(EMethod Method)
 	}
 }
 
-static const FEncoder* GetEncoder(FName FormatName)
-{
-	static FNoneEncoder None;
-	static FOodleEncoder Oodle(FOodleDataCompression::ECompressor::Mermaid, FOodleDataCompression::ECompressionLevel::VeryFast);
-	static FLZ4Encoder LZ4;
-	if (FormatName.IsNone())
-	{
-		return &None;
-	}
-	if (FormatName == NAME_Oodle)
-	{
-		return &Oodle;
-	}
-	if (FormatName == NAME_LZ4 || FormatName == NAME_Default)
-	{
-		return &LZ4;
-	}
-	return nullptr;
-}
-
-static FName GetMethodName(EMethod Method)
-{
-	switch (Method)
-	{
-	default:
-		return NAME_Error;
-	case EMethod::None:
-		return NAME_None;
-	case EMethod::Oodle:
-		return NAME_Oodle;
-	case EMethod::LZ4:
-		return NAME_LZ4;
-	}
-}
-
 template <typename BufferType>
 inline FCompositeBuffer ValidBufferOrEmpty(BufferType&& CompressedData)
 {
@@ -658,7 +622,7 @@ bool FHeader::IsValid(const FCompositeBuffer& CompressedData)
 
 FCompressedBuffer FCompressedBuffer::Compress(const FCompositeBuffer& RawData)
 {
-	return Compress(RawData, FOodleDataCompression::ECompressor::Mermaid, FOodleDataCompression::ECompressionLevel::VeryFast);
+	return Compress(RawData, ECompressedBufferCompressor::Mermaid, ECompressedBufferCompressionLevel::VeryFast);
 }
 
 FCompressedBuffer FCompressedBuffer::Compress(const FSharedBuffer& RawData)
@@ -668,8 +632,8 @@ FCompressedBuffer FCompressedBuffer::Compress(const FSharedBuffer& RawData)
 
 FCompressedBuffer FCompressedBuffer::Compress(
 	const FCompositeBuffer& RawData,
-	FOodleDataCompression::ECompressor Compressor,
-	FOodleDataCompression::ECompressionLevel CompressionLevel,
+	ECompressedBufferCompressor Compressor,
+	ECompressedBufferCompressionLevel CompressionLevel,
 	uint64 BlockSize)
 {
 	using namespace UE::CompressedBuffer;
@@ -680,7 +644,7 @@ FCompressedBuffer FCompressedBuffer::Compress(
 	}
 
 	FCompressedBuffer Local;
-	if (CompressionLevel == FOodleDataCompression::ECompressionLevel::None)
+	if (CompressionLevel == ECompressedBufferCompressionLevel::None)
 	{
 		Local.CompressedData = FNoneEncoder().Compress(RawData, BlockSize);
 	}
@@ -693,29 +657,13 @@ FCompressedBuffer FCompressedBuffer::Compress(
 
 FCompressedBuffer FCompressedBuffer::Compress(
 	const FSharedBuffer& RawData,
-	FOodleDataCompression::ECompressor Compressor,
-	FOodleDataCompression::ECompressionLevel CompressionLevel,
+	ECompressedBufferCompressor Compressor,
+	ECompressedBufferCompressionLevel CompressionLevel,
 	uint64 BlockSize)
 {
 	return Compress(FCompositeBuffer(RawData), Compressor, CompressionLevel, BlockSize);
 }
 
-FCompressedBuffer FCompressedBuffer::Compress(FName FormatName, const FCompositeBuffer& RawData)
-{
-	using namespace UE::CompressedBuffer;
-	FCompressedBuffer Local;
-	if (const FEncoder* const Encoder = GetEncoder(FormatName))
-	{
-		Local.CompressedData = Encoder->Compress(RawData, DefaultBlockSize);
-	}
-	return Local;
-}
-
-FCompressedBuffer FCompressedBuffer::Compress(FName FormatName, const FSharedBuffer& RawData)
-{
-	return Compress(FormatName, FCompositeBuffer(RawData));
-}
-
 FCompressedBuffer FCompressedBuffer::FromCompressed(const FCompositeBuffer& InCompressedData)
 {
 	FCompressedBuffer Local;
@@ -779,15 +727,9 @@ FBlake3Hash FCompressedBuffer::GetRawHash() const
 	return CompressedData ? UE::CompressedBuffer::FHeader::Read(CompressedData).RawHash : FBlake3Hash();
 }
 
-FName FCompressedBuffer::GetFormatName() const
-{
-	using namespace UE::CompressedBuffer;
-	return CompressedData ? GetMethodName(FHeader::Read(CompressedData).Method) : NAME_None;
-}
-
 bool FCompressedBuffer::TryGetCompressParameters(
-	FOodleDataCompression::ECompressor& OutCompressor,
-	FOodleDataCompression::ECompressionLevel& OutCompressionLevel) const
+	ECompressedBufferCompressor& OutCompressor,
+	ECompressedBufferCompressionLevel& OutCompressionLevel) const
 {
 	using namespace UE::CompressedBuffer;
 	if (CompressedData)
@@ -795,12 +737,12 @@ bool FCompressedBuffer::TryGetCompressParameters(
 		switch (const FHeader Header = FHeader::Read(CompressedData); Header.Method)
 		{
 		case EMethod::None:
-			OutCompressor = FOodleDataCompression::ECompressor::NotSet;
-			OutCompressionLevel = FOodleDataCompression::ECompressionLevel::None;
+			OutCompressor = ECompressedBufferCompressor::NotSet;
+			OutCompressionLevel = ECompressedBufferCompressionLevel::None;
 			return true;
 		case EMethod::Oodle:
-			OutCompressor = FOodleDataCompression::ECompressor(Header.Compressor);
-			OutCompressionLevel = FOodleDataCompression::ECompressionLevel(Header.CompressionLevel);
+			OutCompressor = ECompressedBufferCompressor(Header.Compressor);
+			OutCompressionLevel = ECompressedBufferCompressionLevel(Header.CompressionLevel);
 			return true;
 		default:
 			break;
diff --git a/Engine/Source/Runtime/Core/Public/Compression/CompressedBuffer.h b/Engine/Source/Runtime/Core/Public/Compression/CompressedBuffer.h
--- a/Engine/Source/Runtime/Core/Public/Compression/CompressedBuffer.h
+++ b/Engine/Source/Runtime/Core/Public/Compression/CompressedBuffer.h
@@ -6,12 +6,14 @@
 #include "Memory/MemoryFwd.h"
 
 class FArchive;
-class FName;
 struct FBlake3Hash;
 
 namespace FOodleDataCompression { enum class ECompressionLevel : int8; }
 namespace FOodleDataCompression { enum class ECompressor : uint8; }
 
+using ECompressedBufferCompressionLevel = FOodleDataCompression::ECompressionLevel;
+using ECompressedBufferCompressor = FOodleDataCompression::ECompressor;
+
 /**
  * A compressed buffer stores compressed data in a self-contained format.
  *
@@ -44,25 +46,15 @@ public:
 	 */
 	[[nodiscard]] CORE_API static FCompressedBuffer Compress(
 		const FCompositeBuffer& RawData,
-		FOodleDataCompression::ECompressor Compressor,
-		FOodleDataCompression::ECompressionLevel CompressionLevel,
+		ECompressedBufferCompressor Compressor,
+		ECompressedBufferCompressionLevel CompressionLevel,
 		uint64 BlockSize = 0);
 	[[nodiscard]] CORE_API static FCompressedBuffer Compress(
 		const FSharedBuffer& RawData,
-		FOodleDataCompression::ECompressor Compressor,
-		FOodleDataCompression::ECompressionLevel CompressionLevel,
+		ECompressedBufferCompressor Compressor,
+		ECompressedBufferCompressionLevel CompressionLevel,
 		uint64 BlockSize = 0);
 
-	/**
-	 * Compress the buffer using the requested compression format.
-	 *
-	 * @param FormatName   One of NAME_None, NAME_Default, NAME_LZ4.
-	 * @param RawData      Raw data to compress. NAME_None will reference owned raw data.
-	 * @return An owned compressed buffer, or null on error.
-	 */
-	[[nodiscard]] CORE_API static FCompressedBuffer Compress(FName FormatName, const FCompositeBuffer& RawData);
-	[[nodiscard]] CORE_API static FCompressedBuffer Compress(FName FormatName, const FSharedBuffer& RawData);
-
 	/**
 	 * Construct from a compressed buffer previously created by Compress().
 	 *
@@ -103,17 +95,6 @@ public:
 	/** Returns the hash of the raw data. Zero on error or if this is null. */
 	[[nodiscard]] CORE_API FBlake3Hash GetRawHash() const;
 
-	/**
-	 * Returns the name of the compression format used by this buffer.
-	 *
-	 * The format name may differ from the format name specified when creating the compressed buffer
-	 * because an incompressible buffer is stored with NAME_None, and a request of NAME_Default will
-	 * be stored in a specific format such as NAME_LZ4.
-	 *
-	 * @return The format name, or NAME_None if this null, or NAME_Error if the format is unknown.
-	 */
-	[[nodiscard]] CORE_API FName GetFormatName() const;
-
 	/**
 	 * Returns the compressor and compression level used by this buffer.
 	 *
@@ -124,8 +105,8 @@ public:
 	 * @return True if parameters were written, otherwise false.
 	 */
 	[[nodiscard]] CORE_API bool TryGetCompressParameters(
-		FOodleDataCompression::ECompressor& OutCompressor,
-		FOodleDataCompression::ECompressionLevel& OutCompressionLevel) const;
+		ECompressedBufferCompressor& OutCompressor,
+		ECompressedBufferCompressionLevel& OutCompressionLevel) const;
 
 	/**
 	 * Decompress into a memory view that is exactly GetRawSize() bytes.
diff --git a/Engine/Source/Runtime/CoreUObject/Private/Virtualization/VirtualizedBulkData.cpp b/Engine/Source/Runtime/CoreUObject/Private/Virtualization/VirtualizedBulkData.cpp
--- a/Engine/Source/Runtime/CoreUObject/Private/Virtualization/VirtualizedBulkData.cpp
+++ b/Engine/Source/Runtime/CoreUObject/Private/Virtualization/VirtualizedBulkData.cpp
@@ -84,8 +84,8 @@ namespace Private
 {
 
 FCompressionSettings::FCompressionSettings()
-	: Compressor(FOodleDataCompression::ECompressor::NotSet)
-	, CompressionLevel(FOodleDataCompression::ECompressionLevel::None)
+	: Compressor(ECompressedBufferCompressor::NotSet)
+	, CompressionLevel(ECompressedBufferCompressionLevel::None)
 	, bIsSet(false)
 {
 
@@ -98,8 +98,8 @@ FCompressionSettings::FCompressionSettings(const FCompressedBuffer& Buffer)
 	if (!Buffer.TryGetCompressParameters(Compressor, CompressionLevel))
 	{
 		bIsSet = true;
-		Compressor = FOodleDataCompression::ECompressor::NotSet;
-		CompressionLevel = FOodleDataCompression::ECompressionLevel::None;
+		Compressor = ECompressedBufferCompressor::NotSet;
+		CompressionLevel = ECompressedBufferCompressionLevel::None;
 	}
 }
 
@@ -117,12 +117,12 @@ bool FCompressionSettings::operator != (const FCompressionSettings& Other) const
 
 void FCompressionSettings::Reset()
 {
-	Compressor = FOodleDataCompression::ECompressor::NotSet;
-	CompressionLevel = FOodleDataCompression::ECompressionLevel::None;
+	Compressor = ECompressedBufferCompressor::NotSet;
+	CompressionLevel = ECompressedBufferCompressionLevel::None;
 	bIsSet = false;
 }
 
-void FCompressionSettings::Set(FOodleDataCompression::ECompressor InCompressor, FOodleDataCompression::ECompressionLevel InCompressionLevel)
+void FCompressionSettings::Set(ECompressedBufferCompressor InCompressor, ECompressedBufferCompressionLevel InCompressionLevel)
 {
 	Compressor = InCompressor;
 	CompressionLevel = InCompressionLevel;
@@ -131,15 +131,15 @@ void FCompressionSettings::Set(FOodleDataCompression::ECompressor InCompressor,
 
 void FCompressionSettings::SetToDefault()
 {
-	Compressor = FOodleDataCompression::ECompressor::Kraken;
-	CompressionLevel = FOodleDataCompression::ECompressionLevel::Fast;
+	Compressor = ECompressedBufferCompressor::Kraken;
+	CompressionLevel = ECompressedBufferCompressionLevel::Fast;
 	bIsSet = true;
 }
 
 void FCompressionSettings::SetToDisabled()
 {
-	Compressor = FOodleDataCompression::ECompressor::NotSet;
-	CompressionLevel = FOodleDataCompression::ECompressionLevel::None;
+	Compressor = ECompressedBufferCompressor::NotSet;
+	CompressionLevel = ECompressedBufferCompressionLevel::None;
 	bIsSet = true;
 }
 
@@ -150,15 +150,15 @@ bool FCompressionSettings::IsSet() const
 
 bool FCompressionSettings::IsCompressed() const
 {
-	return bIsSet == true && CompressionLevel != FOodleDataCompression::ECompressionLevel::None;
+	return bIsSet == true && CompressionLevel != ECompressedBufferCompressionLevel::None;
 }
 
-FOodleDataCompression::ECompressor FCompressionSettings::GetCompressor() const
+ECompressedBufferCompressor FCompressionSettings::GetCompressor() const
 {
 	return Compressor;
 }
 
-FOodleDataCompression::ECompressionLevel FCompressionSettings::GetCompressionLevel()
+ECompressedBufferCompressionLevel FCompressionSettings::GetCompressionLevel()
 {
 	return CompressionLevel;
 }
@@ -407,7 +407,7 @@ void FVirtualizedUntypedBulkData::Serialize(FArchive& Ar, UObject* Owner, bool b
 			{
 				if (bPayloadInArchive)
 				{
-					FCompressedBuffer CompressedPayload = FCompressedBuffer::Compress(NAME_None, Payload);
+					FCompressedBuffer CompressedPayload = FCompressedBuffer::Compress(Payload, ECompressedBufferCompressor::NotSet, ECompressedBufferCompressionLevel::None);
 					SerializeData(Ar, CompressedPayload, Flags);
 				}
 			}
@@ -908,7 +908,7 @@ bool FVirtualizedUntypedBulkData::SerializeData(FArchive& Ar, FCompressedBuffer&
 			Ar.Serialize(LoadPayload.GetData(), Size);
 		}
 
-		InPayload = FCompressedBuffer::Compress(NAME_None, LoadPayload.MoveToShared());
+		InPayload = FCompressedBuffer::Compress(LoadPayload.MoveToShared(), ECompressedBufferCompressor::NotSet, ECompressedBufferCompressionLevel::None);
 
 		return true;
 	}
@@ -1087,7 +1087,7 @@ FCompressedBuffer FVirtualizedUntypedBulkData::GetDataInternal() const
 	// Check if we already have the data in memory
 	if (Payload)
 	{
-		return FCompressedBuffer::Compress(NAME_None, Payload);
+		return FCompressedBuffer::Compress(Payload, ECompressedBufferCompressor::NotSet, ECompressedBufferCompressionLevel::None);
 	}
 
 	if (IsDataVirtualized())
@@ -1168,7 +1168,7 @@ void FVirtualizedUntypedBulkData::SetCompressionOptions(ECompressionOptions Opti
 		checkNoEntry();
 	}
 
-	if (CompressionSettings.GetCompressionLevel() == FOodleDataCompression::ECompressionLevel::None)
+	if (CompressionSettings.GetCompressionLevel() == ECompressedBufferCompressionLevel::None)
 	{
 		EnumAddFlags(Flags, EFlags::DisablePayloadCompression);
 	}
@@ -1178,11 +1178,11 @@ void FVirtualizedUntypedBulkData::SetCompressionOptions(ECompressionOptions Opti
 	}
 }
 
-void FVirtualizedUntypedBulkData::SetCompressionOptions(FOodleDataCompression::ECompressor Compressor, FOodleDataCompression::ECompressionLevel CompressionLevel)
+void FVirtualizedUntypedBulkData::SetCompressionOptions(ECompressedBufferCompressor Compressor, ECompressedBufferCompressionLevel CompressionLevel)
 {
 	CompressionSettings.Set(Compressor, CompressionLevel);
 
-	if (CompressionSettings.GetCompressionLevel() == FOodleDataCompression::ECompressionLevel::None)
+	if (CompressionSettings.GetCompressionLevel() == ECompressedBufferCompressionLevel::None)
 	{
 		EnumAddFlags(Flags, EFlags::DisablePayloadCompression);
 	}
diff --git a/Engine/Source/Runtime/CoreUObject/Public/Virtualization/VirtualizedBulkData.h b/Engine/Source/Runtime/CoreUObject/Public/Virtualization/VirtualizedBulkData.h
--- a/Engine/Source/Runtime/CoreUObject/Public/Virtualization/VirtualizedBulkData.h
+++ b/Engine/Source/Runtime/CoreUObject/Public/Virtualization/VirtualizedBulkData.h
@@ -33,20 +33,20 @@ struct FCompressionSettings
 
 	void Reset();
 
-	void Set(FOodleDataCompression::ECompressor InCompressor, FOodleDataCompression::ECompressionLevel InCompressionLevel);
+	void Set(ECompressedBufferCompressor InCompressor, ECompressedBufferCompressionLevel InCompressionLevel);
 	void SetToDefault();
 	void SetToDisabled();
 
 	[[nodiscard]] bool IsSet() const;
 	[[nodiscard]] bool IsCompressed() const;
 
-	[[nodiscard]] FOodleDataCompression::ECompressor GetCompressor() const;
-	[[nodiscard]] FOodleDataCompression::ECompressionLevel GetCompressionLevel();
+	[[nodiscard]] ECompressedBufferCompressor GetCompressor() const;
+	[[nodiscard]] ECompressedBufferCompressionLevel GetCompressionLevel();
 
 private:
 
-	FOodleDataCompression::ECompressor Compressor;
-	FOodleDataCompression::ECompressionLevel CompressionLevel;
+	ECompressedBufferCompressor Compressor;
+	ECompressedBufferCompressionLevel CompressionLevel;
 
 	bool bIsSet;
 };
@@ -233,10 +233,10 @@ public:
 	 * These settings will continue to be used until the bulkdata object is reset, a subsequent
 	 * call to ::SetCompressionOptions is made or the owning package is serialized to disk.
 	 * 
-	 * @param Compressor		The Oodle compressor to use. @see FOodleDataCompression::ECompressor
-	 * @param CompressionLevel	The Oodle compression level to use. @see FOodleDataCompression::ECompressionLevel
+	 * @param Compressor		The Oodle compressor to use. @see ECompressedBufferCompressor
+	 * @param CompressionLevel	The Oodle compression level to use. @see ECompressedBufferCompressionLevel
 	 */
-	void SetCompressionOptions(FOodleDataCompression::ECompressor Compressor, FOodleDataCompression::ECompressionLevel CompressionLevel);
+	void SetCompressionOptions(ECompressedBufferCompressor Compressor, ECompressedBufferCompressionLevel CompressionLevel);
 
 	/**
 	* Get the CustomVersions used in the file containing the payload. Currently this is assumed
diff --git a/Engine/Source/Runtime/Engine/Private/Texture.cpp b/Engine/Source/Runtime/Engine/Private/Texture.cpp
--- a/Engine/Source/Runtime/Engine/Private/Texture.cpp
+++ b/Engine/Source/Runtime/Engine/Private/Texture.cpp
@@ -1323,7 +1323,7 @@ void FTextureSource::Compress()
 	if ( ( CompressionFormat == TSCF_PNG && bUseOodleOnPNGz0 ) ||
 		CompressionFormat == TSCF_None )
 	{
-		BulkData.SetCompressionOptions(FOodleDataCompression::ECompressor::Kraken,FOodleDataCompression::ECompressionLevel::Fast);
+		BulkData.SetCompressionOptions(ECompressedBufferCompressor::Kraken,ECompressedBufferCompressionLevel::Fast);
 	}
 	else
 	{
