Restore Oodle allocator hooks
#rnx

diff --git a/Engine/Source/Runtime/Core/Private/Compression/OodleDataCompression.cpp b/Engine/Source/Runtime/Core/Private/Compression/OodleDataCompression.cpp
--- a/Engine/Source/Runtime/Core/Private/Compression/OodleDataCompression.cpp
+++ b/Engine/Source/Runtime/Core/Private/Compression/OodleDataCompression.cpp
@@ -314,6 +314,16 @@ bool CORE_API Decompress(
 	return true;
 }
 					
+static void* OodleAlloc(OO_SINTa Size, OO_S32 Alignment)
+{
+	return FMemory::Malloc(SIZE_T(Size), uint32(Alignment));
+}
+
+static void OodleFree(void* Ptr)
+{
+	FMemory::Free(Ptr);
+}
+
 void CORE_API StartupPreInit(void)
 {
 	// called from LaunchEngineLoop at "PreInit" time
@@ -327,6 +337,7 @@ void CORE_API StartupPreInit(void)
 	// UE5 will always read/write Oodle v9 binary data :
 	OodleConfig.m_OodleLZ_BackwardsCompatible_MajorVersion = 9;
 	Oodle_SetConfigValues(&OodleConfig);
+	OodleCore_Plugins_SetAllocators(OodleAlloc, OodleFree);
 }
 
 };
diff --git a/Engine/Source/Runtime/Core/Private/Compression/OodleDataCompressionFormat.cpp b/Engine/Source/Runtime/Core/Private/Compression/OodleDataCompressionFormat.cpp
--- a/Engine/Source/Runtime/Core/Private/Compression/OodleDataCompressionFormat.cpp
+++ b/Engine/Source/Runtime/Core/Private/Compression/OodleDataCompressionFormat.cpp
@@ -94,6 +94,10 @@ typedef OO_SINTa (OOLINK t_fp_OodleLZ_Compress)(OodleLZ_Compressor compressor,
     const OodleLZ_CompressOptions * pOptions,
     const void * dictionaryBase,
     const void * lrm);
+
+typedef OO_SINTa (OOLINK t_fp_OodleSetAllocators)(
+	t_fp_OodleCore_Plugin_MallocAligned* fp_OodleMallocAligned,
+	t_fp_OodleCore_Plugin_Free* fp_OodleFree);
 };
 
 
@@ -106,8 +110,17 @@ struct FOodleDataCompressionFormat : ICompressionFormat
 	//  OodleLZCompressFuncPtr is non-null of a DLL/so was loaded for encoding
 	t_fp_OodleLZ_Compress  * OodleLZCompressFuncPtr = nullptr;
 
+	static void* OodleAlloc(OO_SINTa Size, OO_S32 Alignment)
+	{
+		return FMemory::Malloc(SIZE_T(Size), uint32(Alignment));
+	}
+
+	static void OodleFree(void* Ptr)
+	{
+		FMemory::Free(Ptr);
+	}
 
-	FOodleDataCompressionFormat(OodleLZ_Compressor InCompressor, OodleLZ_CompressionLevel InCompressionLevel, int InSpaceSpeedTradeoffBytes, void * InOodleCompressFuncPtr)
+	FOodleDataCompressionFormat(OodleLZ_Compressor InCompressor, OodleLZ_CompressionLevel InCompressionLevel, int InSpaceSpeedTradeoffBytes, void * InOodleCompressFuncPtr, void * InOodleSetAllocatorsFuncPtr)
 	{
 		Compressor = InCompressor;
 		CompressionLevel = InCompressionLevel;
@@ -117,6 +130,10 @@ struct FOodleDataCompressionFormat : ICompressionFormat
 		CompressionOptions.makeLongRangeMatcher = false;
 		
 		OodleLZCompressFuncPtr = (t_fp_OodleLZ_Compress *) InOodleCompressFuncPtr;
+		if (InOodleSetAllocatorsFuncPtr)
+		{
+			((t_fp_OodleSetAllocators *)(InOodleSetAllocatorsFuncPtr))(OodleAlloc, OodleFree);
+		}
 	}
 
 	virtual ~FOodleDataCompressionFormat() 
@@ -326,6 +343,7 @@ ICompressionFormat * CreateOodleDataCompressionFormat()
 	OodleLZ_CompressionLevel UsedLevel = OodleLZ_CompressionLevel_VeryFast;
 	int32 SpaceSpeedTradeoff = OODLELZ_SPACESPEEDTRADEOFFBYTES_DEFAULT;
 	void * OodleCompressFuncPtr = nullptr;
+	void * OodleSetAllocatorsFuncPtr = nullptr;
 
 	#if ( (!UE_BUILD_SHIPPING) || WITH_EDITOR )
 	{
@@ -466,6 +484,7 @@ ICompressionFormat * CreateOodleDataCompressionFormat()
 				{
 					UE_LOG(OodleDataCompression, Warning, TEXT("OodleCompressDLL %s loaded but no OodleLZ_Compress !?"), *OodleDLL);
 				}
+				OodleSetAllocatorsFuncPtr = FPlatformProcess::GetDllExport( OodleDLLHandle, TEXT("OodlePlugins_SetAllocators") );
 			}
 		}
 	}
@@ -477,7 +496,7 @@ ICompressionFormat * CreateOodleDataCompressionFormat()
 	// register the compression format :
 	//  this is used by the shipping game to decode any paks compressed with Oodle :
 
-	ICompressionFormat * CompressionFormat = new FOodleDataCompressionFormat(UsedCompressor, UsedLevel, SpaceSpeedTradeoff, OodleCompressFuncPtr);
+	ICompressionFormat * CompressionFormat = new FOodleDataCompressionFormat(UsedCompressor, UsedLevel, SpaceSpeedTradeoff, OodleCompressFuncPtr, OodleSetAllocatorsFuncPtr);
 
 	IModularFeatures::Get().RegisterModularFeature(COMPRESSION_FORMAT_FEATURE_NAME, CompressionFormat);
 
