Add Insight tag over each automation tests

diff --git a/Engine/Source/Runtime/Core/Private/Misc/AutomationTest.cpp b/Engine/Source/Runtime/Core/Private/Misc/AutomationTest.cpp
--- a/Engine/Source/Runtime/Core/Private/Misc/AutomationTest.cpp
+++ b/Engine/Source/Runtime/Core/Private/Misc/AutomationTest.cpp
@@ -878,8 +878,11 @@ void FAutomationTestFramework::InternalStartTest( const FString& InTestToRun )
 
 		CurrentTest->SetTestContext(Parameters);
 
-		// Run the test!
-		bTestSuccessful = CurrentTest->RunTest(Parameters);
+		{
+			TRACE_CPUPROFILER_EVENT_SCOPE_TEXT(*FString::Printf(TEXT("AutomationTest %s"), *CurrentTest->GetBeautifiedTestName()));
+			// Run the test!
+			bTestSuccessful = CurrentTest->RunTest(Parameters);
+		}
 	}
 }
 
