StringBuilder: Fixed a static analysis warning, and fixed direct use of the base type

diff --git a/Engine/Source/Runtime/Core/Private/Misc/StringBuilder.cpp b/Engine/Source/Runtime/Core/Private/Misc/StringBuilder.cpp
--- a/Engine/Source/Runtime/Core/Private/Misc/StringBuilder.cpp
+++ b/Engine/Source/Runtime/Core/Private/Misc/StringBuilder.cpp
@@ -29,7 +29,7 @@ template <typename C>
 void TStringBuilderBase<C>::Extend(SIZE_T ExtraCapacity)
 {
 	const SIZE_T OldCapacity = End - Base;
-	const SIZE_T NewCapacity = NextPowerOfTwo(OldCapacity + ExtraCapacity);
+	const SIZE_T NewCapacity = FPlatformMath::Max<SIZE_T>(NextPowerOfTwo(OldCapacity + ExtraCapacity), 32);
 
 	C* NewBase = (C*)AllocBuffer(NewCapacity);
 
diff --git a/Engine/Source/Runtime/Core/Private/Tests/Misc/StringBuilderTest.cpp b/Engine/Source/Runtime/Core/Private/Tests/Misc/StringBuilderTest.cpp
--- a/Engine/Source/Runtime/Core/Private/Tests/Misc/StringBuilderTest.cpp
+++ b/Engine/Source/Runtime/Core/Private/Tests/Misc/StringBuilderTest.cpp
@@ -26,6 +26,22 @@ static_assert(TIsContiguousContainer<TWideStringBuilder<128>>::Value, "TWideStri
 IMPLEMENT_SIMPLE_AUTOMATION_TEST(FStringBuilderTestAppendString, "System.Core.StringBuilder.AppendString", EAutomationTestFlags::ApplicationContextMask | EAutomationTestFlags::SmokeFilter)
 bool FStringBuilderTestAppendString::RunTest(const FString& Parameters)
 {
+	// Empty Base
+	{
+		FStringBuilderBase Builder;
+		TestEqual(TEXT("Empty StringBuilderBase Len"), Builder.Len(), 0);
+		TestEqual(TEXT("Empty StringBuilderBase ToString"), Builder.ToString(), TEXT(""));
+		Builder << TEXT('A');
+		TestEqual(TEXT("Append Char to StringBuilderBase"), Builder.ToString(), TEXT("A"));
+	}
+
+	// Empty With Buffer
+	{
+		TStringBuilder<1024> Builder;
+		TestEqual(TEXT("Empty StringBuilderWithBuffer Len"), Builder.Len(), 0);
+		TestEqual(TEXT("Empty StringBuilderWithBuffer ToString"), Builder.ToString(), TEXT(""));
+	}
+
 	// Append Char
 	{
 		TStringBuilder<7> Builder;
diff --git a/Engine/Source/Runtime/Core/Public/Misc/StringBuilder.h b/Engine/Source/Runtime/Core/Public/Misc/StringBuilder.h
--- a/Engine/Source/Runtime/Core/Public/Misc/StringBuilder.h
+++ b/Engine/Source/Runtime/Core/Public/Misc/StringBuilder.h
@@ -370,10 +370,7 @@ protected:
 
 	inline void EnsureNulTerminated() const
 	{
-		if (*CurPos)
-		{
-			*CurPos = CharType(0);
-		}
+		*CurPos = CharType(0);
 	}
 
 	inline void EnsureAdditionalCapacity(int32 RequiredAdditionalCapacity)
@@ -393,9 +390,11 @@ protected:
 	CORE_API void*	AllocBuffer(SIZE_T CharCount);
 	CORE_API void	FreeBuffer(void* Buffer, SIZE_T CharCount);
 
-	CharType*	Base;
-	CharType*	CurPos;
-	CharType*	End;
+	static inline CharType EmptyBuffer[1]{};
+
+	CharType*	Base = EmptyBuffer;
+	CharType*	CurPos = Base;
+	CharType*	End = Base + 1;
 	bool		bIsDynamic = false;
 };
 
