UE5:  Optimizations when memory tracing is enabled:  5x perf increase with -llmcsv, 8x perf increase with -trace=memory.

#rnx

Over half the perf gain was from reducing critical section locks, with the rest of the gains split between an optimization to call stack traversal (caching return address lookups in a lock free map), and enabling TLS related optimizations for MallocBinned2, which were inadvertently disabled when tracing was enabled (this also saves some locks).  Multiple lock types had to be optimized before measurable perf gains were observed, due to contention on remaining locks increasing as the first few locks were removed.  A total of 6 types of locks were removed or reduced:  AnnounceFNameTag, MallocBinned2, TagDataNameMap, Enum tag data lookup, Windows callstack trace, LLMMap.

* Added hash table utility class TGrowOnlyLockFreeHash which supports lock free reads, writes using a critical section, and no deletion.  This is intended for classes that represent caches of one-time initialized debug data used during memory tracing, which never needs to be freed.
* TGrowOnlyLockFreeHash applied to 4 sets/maps used when memory tracing is enabled, including tag data, announced tag names, call stack functions, and call stack IDs.
* Lock striping added to LLMMap.  This isn't a good candidate for a lock free approach, as it has a heavy amount of modification, and lock free approaches tend to perform worse than lock striping in that case, plus it would be a massive undertaking to make this lock free.
* A future lock to be removed would be the one in "FLowLevelAllocInfo::GetTag", protecting the "FLowLevelMemTracker::TagDatas" array, saving a couple percent more perf.

diff --git a/Engine/Source/Runtime/Core/Private/HAL/LowLevelMemTracker.cpp b/Engine/Source/Runtime/Core/Private/HAL/LowLevelMemTracker.cpp
--- a/Engine/Source/Runtime/Core/Private/HAL/LowLevelMemTracker.cpp
+++ b/Engine/Source/Runtime/Core/Private/HAL/LowLevelMemTracker.cpp
@@ -442,7 +442,6 @@ namespace LLMPrivate
 
 		FCriticalSection PendingThreadStatesGuard;
 		TArray<FLLMThreadState*, FDefaultLLMAllocator> PendingThreadStates;
-		FCriticalSection AllocationMapLock;
 
 		int64 TrackedTotal GCC_ALIGN(8);
 
@@ -2993,7 +2992,6 @@ namespace LLMPrivate
 			LLMCheck(Size <= 0x0000'ffff'ffff'ffff);
 			uint32 SizeLow = uint32(Size);
 			uint16 SizeHigh = uint16(Size >> 32ull);
-			FScopeLock AllocationScopeLock(&AllocationMapLock);
 			PointerKey Key(Ptr, SizeHigh);
 			AllocationMap.Add(Key, SizeLow, AllocInfo);
 		}
@@ -3004,7 +3002,6 @@ namespace LLMPrivate
 		// look up the pointer in the tracking map
 		FLLMAllocMap::Values Values;
 		{
-			FScopeLock AllocationScopeLock(&AllocationMapLock);
 			if (!AllocationMap.Remove(PointerKey(Ptr), Values))
 			{
 				return;
@@ -3040,7 +3037,6 @@ namespace LLMPrivate
 	{
 		FLLMAllocMap::Values Values;
 		{
-			FScopeLock AllocationScopeLock(&AllocationMapLock);
 			if (!AllocationMap.Remove(PointerKey(Source), Values))
 			{
 				return;
@@ -3141,10 +3137,8 @@ namespace LLMPrivate
 			LLMRef.Allocator.Delete(ThreadState);
 		ThreadStates.Empty();
 
-		{
-			FScopeLock AllocationScopeLock(&AllocationMapLock);
-			AllocationMap.Clear();
-		}
+		AllocationMap.Clear();
+
 		CsvWriter.Clear();
 		TraceWriter.Clear();
 		CsvProfilerWriter.Clear();
@@ -3166,7 +3160,6 @@ namespace LLMPrivate
 		{
 			LastTrimTime = CurrentTime;
 			{
-				FScopeLock AllocationScopeLock(&AllocationMapLock);
 				AllocationMap.Trim();
 			}
 		}
@@ -3283,11 +3276,12 @@ namespace LLMPrivate
 		{
 			// Each allocation references the tag by its index, which we have just remapped.
 			// Remap each allocation's tag index to the new index for the tag
-			FScopeLock AllocationScopeLock(&AllocationMapLock);
+			AllocationMap.LockAll();
 			for (const FLLMAllocMap::FTuple& Tuple : AllocationMap)
 			{
 				Tuple.Value2.SetCompressedTag(OldTagDatas[Tuple.Value2.GetCompressedTag()]->GetIndex());
 			}
+			AllocationMap.UnlockAll();
 		}
 #else
 		// Values in AllocationMap are ELLMTags, and don't depend on the Index of the tagdatas
@@ -3330,7 +3324,6 @@ namespace LLMPrivate
 		FLLMThreadState* State = GetOrCreateState();
 		FLowLevelAllocInfo AllocInfo;
 		{
-			FScopeLock AllocationScopeLock(&AllocationMapLock);
 			uint32* Size;
 			FLowLevelAllocInfo* AllocInfoPtr;
 			AllocationMap.Find(PointerKey(Ptr), Size, AllocInfoPtr);
diff --git a/Engine/Source/Runtime/Core/Private/HAL/LowLevelMemoryUtils.h b/Engine/Source/Runtime/Core/Private/HAL/LowLevelMemoryUtils.h
--- a/Engine/Source/Runtime/Core/Private/HAL/LowLevelMemoryUtils.h
+++ b/Engine/Source/Runtime/Core/Private/HAL/LowLevelMemoryUtils.h
@@ -10,6 +10,7 @@
 #include "Containers/ContainerAllocationPolicies.h"
 #include "Containers/Map.h"
 #include "Containers/Set.h"
+#include "Misc/ScopeLock.h"
 #include "Templates/Tuple.h"
 
 #define LLM_PAGE_SIZE (16*1024)
@@ -201,14 +202,7 @@ public:
 	};
 
 	LLMMap()
-		: Allocator(NULL)
-		, Map(NULL)
-		, Count(0)
-		, Capacity(0)
-#ifdef PROFILE_LLMMAP
-		, IterAcc(0)
-		, IterCount(0)
-#endif
+		: bAllStripesLocked(false)
 	{
 	}
 
@@ -219,41 +213,52 @@ public:
 
 	void SetAllocator(UE::LLMPrivate::FLLMAllocator* InAllocator, SizeType InDefaultCapacity = DefaultCapacity)
 	{
-		Allocator = InAllocator;
+		for (StripeData& Stripe : MapStripes)
+		{
+			Stripe.Allocator = InAllocator;
 
-		Keys.SetAllocator(Allocator);
-		KeyHashes.SetAllocator(Allocator);
-		Values1.SetAllocator(Allocator);
-		Values2.SetAllocator(Allocator);
-		FreeKeyIndices.SetAllocator(Allocator);
+			Stripe.Keys.SetAllocator(InAllocator);
+			Stripe.KeyHashes.SetAllocator(InAllocator);
+			Stripe.Values1.SetAllocator(InAllocator);
+			Stripe.Values2.SetAllocator(InAllocator);
+			Stripe.FreeKeyIndices.SetAllocator(InAllocator);
 
-		Reserve(InDefaultCapacity);
+			Stripe.Reserve(InDefaultCapacity / StripeCount);
+		}
 	}
 
 	void Clear()
 	{
-		Keys.Clear(true);
-		KeyHashes.Clear(true);
-		Values1.Clear(true);
-		Values2.Clear(true);
-		FreeKeyIndices.Clear(true);
-
-		Allocator->Free(Map, Capacity * sizeof(SizeType));
-		Map = NULL;
-		Count = 0;
-		Capacity = 0;
+		for (StripeData& Stripe : MapStripes)
+		{
+			FScopeLock AllocationScopeLock(&Stripe.CriticalSection);
+
+			Stripe.Keys.Clear(true);
+			Stripe.KeyHashes.Clear(true);
+			Stripe.Values1.Clear(true);
+			Stripe.Values2.Clear(true);
+			Stripe.FreeKeyIndices.Clear(true);
+
+			Stripe.Allocator->Free(Stripe.Map, Stripe.Capacity * sizeof(SizeType));
+			Stripe.Map = NULL;
+			Stripe.Count = 0;
+			Stripe.Capacity = 0;
+		}
 	}
 
 	// Add a value to this set.
 	// If this set already contains the value does nothing.
 	void Add(const TKey& Key, const TValue1& Value1, const TValue2& Value2)
 	{
-		LLMCheck(Map);
-
 		SizeType KeyHash = Key.GetHashCode();
+		StripeData& Stripe = MapStripes[GetStripeIndex(KeyHash)];
 
-		SizeType MapIndex = GetMapIndex(Key, KeyHash);
-		SizeType KeyIndex = Map[MapIndex];
+		FScopeLock AllocationScopeLock(&Stripe.CriticalSection);
+
+		LLMCheck(Stripe.Map);
+
+		SizeType MapIndex = Stripe.GetMapIndex(Key, KeyHash);
+		SizeType KeyIndex = Stripe.Map[MapIndex];
 
 		if (KeyIndex != InvalidIndex)
 		{
@@ -264,53 +269,53 @@ public:
 				ShownWarning = true;
 			}
 
-			Values1[KeyIndex] = Value1;
-			Values2[KeyIndex] = Value2;
+			Stripe.Values1[KeyIndex] = Value1;
+			Stripe.Values2[KeyIndex] = Value2;
 		}
 		else
 		{
 			// Be careful with the multiplication to avoid overflow, while still using only integer arithmetic for speed
 			// Margin/256U is the actual float we want to multiply by. The exact number doesn't matter, so use one order of operations
 			// when Capacity is low, and another when it is high
-			SizeType MaxCount = (Capacity >= 256U*256U ? (Capacity/256U)*Margin : (Margin * Capacity) / 256U);
-			if (Count >= MaxCount)
+			SizeType MaxCount = (Stripe.Capacity >= 256U * 256U ? (Stripe.Capacity / 256U) * Margin : (Margin * Stripe.Capacity) / 256U);
+			if (Stripe.Count >= MaxCount)
 			{
-				if (Capacity * 2 < Capacity)
+				if (Stripe.Capacity * 2 < Stripe.Capacity)
 				{
 					// Trying to issue a check statement here will cause reentry into this function, use PLATFORM_BREAK directly instead
 					FPlatformMisc::LowLevelOutputDebugString(TEXT("LLM Error: Integer overflow in LLMap::Add, Capacity has reached its maximum size.\n"));
 					PLATFORM_BREAK();
 				}
-				if (Count > MaxCount)
+				if (Stripe.Count > MaxCount)
 				{
 					// Trying to issue a check statement here will cause reentry into this function, use PLATFORM_BREAK directly instead
 					// This shouldn't happen, because Count is only incremented here, and Capacity is only changed here, and Margin does not change, so Count should equal MaxCount before it goes over it
 					FPlatformMisc::LowLevelOutputDebugString(TEXT("LLM Assertion failure: Count > MaxCount.\n"));
 					PLATFORM_BREAK();
 				}
-				Grow();
-				MapIndex = GetMapIndex(Key, KeyHash);
+				Stripe.Grow();
+				MapIndex = Stripe.GetMapIndex(Key, KeyHash);
 			}
 
-			if (FreeKeyIndices.Num())
+			if (Stripe.FreeKeyIndices.Num())
 			{
-				SizeType FreeIndex = FreeKeyIndices.RemoveLast();
-				Map[MapIndex] = FreeIndex;
-				Keys[FreeIndex] = Key;
-				KeyHashes[FreeIndex] = KeyHash;
-				Values1[FreeIndex] = Value1;
-				Values2[FreeIndex] = Value2;
+				SizeType FreeIndex = Stripe.FreeKeyIndices.RemoveLast();
+				Stripe.Map[MapIndex] = FreeIndex;
+				Stripe.Keys[FreeIndex] = Key;
+				Stripe.KeyHashes[FreeIndex] = KeyHash;
+				Stripe.Values1[FreeIndex] = Value1;
+				Stripe.Values2[FreeIndex] = Value2;
 			}
 			else
 			{
-				Map[MapIndex] = Keys.Num();
-				Keys.Add(Key);
-				KeyHashes.Add(KeyHash);
-				Values1.Add(Value1);
-				Values2.Add(Value2);
+				Stripe.Map[MapIndex] = Stripe.Keys.Num();
+				Stripe.Keys.Add(Key);
+				Stripe.KeyHashes.Add(KeyHash);
+				Stripe.Values1.Add(Value1);
+				Stripe.Values2.Add(Value2);
 			}
 
-			++Count;
+			++Stripe.Count;
 		}
 	}
 
@@ -330,10 +335,13 @@ public:
 	TKey Find(const TKey& Key, TValue1*& OutValue1, TValue2*& OutValue2)
 	{
 		SizeType KeyHash = Key.GetHashCode();
+		StripeData& Stripe = MapStripes[GetStripeIndex(KeyHash)];
+
+		FScopeLock AllocationScopeLock(&Stripe.CriticalSection);
 
-		SizeType MapIndex = GetMapIndex(Key, KeyHash);
+		SizeType MapIndex = Stripe.GetMapIndex(Key, KeyHash);
 
-		SizeType KeyIndex = Map[MapIndex];
+		SizeType KeyIndex = Stripe.Map[MapIndex];
 		if (KeyIndex == InvalidIndex)
 		{
 			OutValue1 = nullptr;
@@ -341,40 +349,43 @@ public:
 			return TKey();
 		}
 
-		OutValue1 = &Values1[KeyIndex];
-		OutValue2 = &Values2[KeyIndex];
+		OutValue1 = &Stripe.Values1[KeyIndex];
+		OutValue2 = &Stripe.Values2[KeyIndex];
 
-		return Keys[KeyIndex];
+		return Stripe.Keys[KeyIndex];
 	}
 
 	bool Remove(const TKey& Key, Values& OutValues)
 	{
 		SizeType KeyHash = Key.GetHashCode();
+		StripeData& Stripe = MapStripes[GetStripeIndex(KeyHash)];
 
-		LLMCheck(Map);
+		FScopeLock AllocationScopeLock(&Stripe.CriticalSection);
 
-		SizeType MapIndex = GetMapIndex(Key, KeyHash);
-		if (!IsItemInUse(MapIndex))
+		LLMCheck(Stripe.Map);
+
+		SizeType MapIndex = Stripe.GetMapIndex(Key, KeyHash);
+		if (!Stripe.IsItemInUse(MapIndex))
 		{
 			return false;
 		}
 
-		SizeType KeyIndex = Map[MapIndex];
+		SizeType KeyIndex = Stripe.Map[MapIndex];
 
-		OutValues.Key = Keys[KeyIndex];
-		OutValues.Value1 = Values1[KeyIndex];
-		OutValues.Value2 = Values2[KeyIndex];
+		OutValues.Key = Stripe.Keys[KeyIndex];
+		OutValues.Value1 = Stripe.Values1[KeyIndex];
+		OutValues.Value2 = Stripe.Values2[KeyIndex];
 
-		if (KeyIndex == Keys.Num() - 1)
+		if (KeyIndex == Stripe.Keys.Num() - 1)
 		{
-			Keys.RemoveLast();
-			KeyHashes.RemoveLast();
-			Values1.RemoveLast();
-			Values2.RemoveLast();
+			Stripe.Keys.RemoveLast();
+			Stripe.KeyHashes.RemoveLast();
+			Stripe.Values1.RemoveLast();
+			Stripe.Values2.RemoveLast();
 		}
 		else
 		{
-			FreeKeyIndices.Add(KeyIndex);
+			Stripe.FreeKeyIndices.Add(KeyIndex);
 		}
 
 		// find first index in this array
@@ -382,15 +393,15 @@ public:
 		SizeType FirstIndex = MapIndex;
 		if (!IndexIter)
 		{
-			IndexIter = Capacity;
+			IndexIter = Stripe.Capacity;
 		}
 		--IndexIter;
-		while (IsItemInUse(IndexIter))
+		while (Stripe.IsItemInUse(IndexIter))
 		{
 			FirstIndex = IndexIter;
 			if (!IndexIter)
 			{
-				IndexIter = Capacity;
+				IndexIter = Stripe.Capacity;
 			}
 			--IndexIter;
 		}
@@ -399,28 +410,28 @@ public:
 		for (;;)
 		{
 			// find the last item in the array that can replace the item being removed
-			SizeType IndexIter2 = (MapIndex + 1) & (Capacity - 1);
+			SizeType IndexIter2 = (MapIndex + 1) & (Stripe.Capacity - 1);
 
 			SizeType SwapIndex = InvalidIndex;
-			while (IsItemInUse(IndexIter2))
+			while (Stripe.IsItemInUse(IndexIter2))
 			{
-				SizeType SearchKeyIndex = Map[IndexIter2];
-				const SizeType SearchHashCode = KeyHashes[SearchKeyIndex];
-				const SizeType SearchInsertIndex = SearchHashCode & (Capacity - 1);
+				SizeType SearchKeyIndex = Stripe.Map[IndexIter2];
+				const SizeType SearchHashCode = Stripe.KeyHashes[SearchKeyIndex];
+				const SizeType SearchInsertIndex = SearchHashCode & (Stripe.Capacity - 1);
 
-				if (InRange(SearchInsertIndex, FirstIndex, MapIndex))
+				if (Stripe.InRange(SearchInsertIndex, FirstIndex, MapIndex))
 				{
 					SwapIndex = IndexIter2;
 					Found = true;
 				}
 
-				IndexIter2 = (IndexIter2 + 1) & (Capacity - 1);
+				IndexIter2 = (IndexIter2 + 1) & (Stripe.Capacity - 1);
 			}
 
 			// swap the item
 			if (Found)
 			{
-				Map[MapIndex] = Map[SwapIndex];
+				Stripe.Map[MapIndex] = Stripe.Map[SwapIndex];
 				MapIndex = SwapIndex;
 				Found = false;
 			}
@@ -431,33 +442,47 @@ public:
 		}
 
 		// remove the last item
-		Map[MapIndex] = InvalidIndex;
+		Stripe.Map[MapIndex] = InvalidIndex;
 
-		--Count;
+		--Stripe.Count;
 
 		return true;
 	}
 
 	SizeType Num() const
 	{
-		return Count;
+		SizeType TotalCount = 0;
+		for (StripeData& Stripe : MapStripes)
+		{
+			FScopeLock AllocationScopeLock(&Stripe.CriticalSection);
+			TotalCount += Stripe.Count;
+		}
+		return TotalCount;
 	}
 
 	bool HasKey(const TKey& Key)
 	{
 		SizeType KeyHash = Key.GetHashCode();
+		StripeData& Stripe = MapStripes[GetStripeIndex(KeyHash)];
+
+		FScopeLock AllocationScopeLock(&Stripe.CriticalSection);
 
-		SizeType MapIndex = GetMapIndex(Key, KeyHash);
-		return IsItemInUse(MapIndex);
+		SizeType MapIndex = Stripe.GetMapIndex(Key, KeyHash);
+		return Stripe.IsItemInUse(MapIndex);
 	}
 
 	void Trim()
 	{
-		Keys.Trim();
-		KeyHashes.Trim();
-		Values1.Trim();
-		Values2.Trim();
-		FreeKeyIndices.Trim();
+		for (StripeData& Stripe : MapStripes)
+		{
+			FScopeLock AllocationScopeLock(&Stripe.CriticalSection);
+
+			Stripe.Keys.Trim();
+			Stripe.KeyHashes.Trim();
+			Stripe.Values1.Trim();
+			Stripe.Values2.Trim();
+			Stripe.FreeKeyIndices.Trim();
+		}
 	}
 
 	struct FBaseIterator
@@ -466,10 +491,22 @@ public:
 		FBaseIterator& operator++()
 		{
 			++MapIndex;
-			while (MapIndex < MapRef.Capacity && !MapRef.IsItemInUse(MapIndex))
+			while (StripeIndex < StripeCount)
 			{
-				++MapIndex;
+				StripeData& Stripe = MapRef.MapStripes[StripeIndex];
+				while (MapIndex < Stripe.Capacity)
+				{
+					if (Stripe.IsItemInUse(MapIndex))
+					{
+						return *this;
+					}
+					MapIndex++;
+				}
+
+				++StripeIndex;
+				MapIndex = 0;
 			}
+
 			return *this;
 		}
 
@@ -480,16 +517,19 @@ public:
 		{
 			if (bEnd)
 			{
-				MapIndex = MapRef.Capacity;
+				StripeIndex = StripeCount;
+				MapIndex = 0;
 			}
 			else
 			{
-				MapIndex = static_cast<SizeType>(-1);
+				StripeIndex = 0;
+				MapIndex = -1;
 				++(*this);
 			}
 		}
 
 		ThisType& MapRef;
+		int32 StripeIndex;
 		SizeType MapIndex;
 	};
 
@@ -524,8 +564,9 @@ public:
 		FTuple operator*() const
 		{
 			ThisType& LocalMap = FBaseIterator::MapRef;
-			SizeType KeyIndex = LocalMap.Map[FBaseIterator::MapIndex];
-			return FTuple(LocalMap.Keys[KeyIndex], LocalMap.Values1[KeyIndex], LocalMap.Values2[KeyIndex]);
+			StripeData& Stripe = LocalMap.MapStripes[FBaseIterator::StripeIndex];
+			SizeType KeyIndex = Stripe.Map[FBaseIterator::MapIndex];
+			return FTuple(Stripe.Keys[KeyIndex], Stripe.Values1[KeyIndex], Stripe.Values2[KeyIndex]);
 		}
 	};
 
@@ -567,6 +608,7 @@ public:
 
 	FIterator begin()
 	{
+		LLMCheck(bAllStripesLocked);
 		return FIterator(*this, false);
 	}
 
@@ -577,6 +619,7 @@ public:
 
 	FConstIterator begin() const
 	{
+		LLMCheck(bAllStripesLocked);
 		return FConstIterator(*this, false);
 	}
 
@@ -585,136 +628,181 @@ public:
 		return FConstIterator(*this, true);
 	}
 
-private:
-	void Reserve(SizeType NewCapacity)
+	void LockAll()
 	{
-		NewCapacity = GetNextPow2(NewCapacity);
-
-		// keep a copy of the old map
-		SizeType* OldMap = Map;
-		SizeType OldCapacity = Capacity;
+		// Locking all stripes is required to create an iterator
+		AllStripesLock.Lock();
+		LLMCheck(bAllStripesLocked == false);
+		bAllStripesLocked = true;
 
-		// allocate the new table
-		Capacity = NewCapacity;
-		Map = (SizeType*)Allocator->Alloc(NewCapacity * sizeof(SizeType));
-
-		for (SizeType Index = 0; Index < NewCapacity; ++Index)
-			Map[Index] = InvalidIndex;
-
-		// copy the values from the old to the new table
-		SizeType* OldItem = OldMap;
-		for (SizeType Index = 0; Index < OldCapacity; ++Index, ++OldItem)
+		for (StripeData& Stripe : MapStripes)
 		{
-			SizeType KeyIndex = *OldItem;
-
-			if (KeyIndex != InvalidIndex)
-			{
-				SizeType MapIndex = GetMapIndex(Keys[KeyIndex], KeyHashes[KeyIndex]);
-				Map[MapIndex] = KeyIndex;
-			}
+			Stripe.CriticalSection.Lock();
 		}
-
-		Allocator->Free(OldMap, OldCapacity * sizeof(SizeType));
 	}
 
-	static SizeType GetNextPow2(SizeType value)
+	void UnlockAll()
 	{
-		SizeType p = 2;
-		while (p < value)
+		for (StripeData& Stripe : MapStripes)
 		{
-			SizeType Nextp = p * 2;
-			if (Nextp < p)
-			{
-				// Trying to issue a check statement here will cause reentry into allocation, use PLATFORM_BREAK directly instead
-				FPlatformMisc::LowLevelOutputDebugString(TEXT("LLM Error: Integer overflow in LLMap::Add, GetNextPow2 called on value > 2^(NumBits-1).\n"));
-				PLATFORM_BREAK();
-				break;
-			}
-			p = Nextp;
+			Stripe.CriticalSection.Unlock();
 		}
-		return p;
+
+		LLMCheck(bAllStripesLocked == true);
+		bAllStripesLocked = false;
+		AllStripesLock.Unlock();
 	}
 
-	bool IsItemInUse(SizeType MapIndex) const
+	// data
+private:
+	static const int32 StripeCountLog2 = 4;		// 16 stripes
+	enum { StripeCount = 1 << StripeCountLog2 };
+	enum { DefaultCapacity = 1024 * 1024 };
+	enum { InvalidIndex = -1 };
+	static const SizeType Margin = (30 * 256) / 100;
+
+	static inline uint32 GetStripeIndex(uint32 KeyHash)
 	{
-		return Map[MapIndex] != InvalidIndex;
+		// FMath::Max prevents "shift too large" compile error in otherwise unreachable code when StripeCount == 1
+		return (StripeCount == 1) ? 0 : KeyHash >> (32 - FMath::Max(StripeCountLog2, 1));
 	}
 
-	SizeType GetMapIndex(const TKey& Key, SizeType Hash) const
+	static inline uint32 GetStripeIndex(uint64 KeyHash)
 	{
-		SizeType Mask = Capacity - 1;
-		SizeType MapIndex = Hash & Mask;
-		SizeType KeyIndex = Map[MapIndex];
+		// FMath::Max prevents "shift too large" compile error in otherwise unreachable code when StripeCount == 1
+		return static_cast<uint32>((StripeCount == 1) ? 0 : KeyHash >> (64 - FMath::Max(StripeCountLog2, 1)));
+	}
 
-		while (KeyIndex != InvalidIndex && !(Keys[KeyIndex] == Key))
-		{
-			MapIndex = (MapIndex + 1) & Mask;
-			KeyIndex = Map[MapIndex];
+	struct StripeData
+	{
+		StripeData()
+			: Allocator(nullptr)
+			, Map(nullptr)
+			, Count(0)
+			, Capacity(0)
 #ifdef PROFILE_LLMMAP
-			++IterAcc;
+			, IterAcc(0)
+			, IterCount(0)
 #endif
-		}
+		{}
 
-#ifdef PROFILE_LLMMAP
-		++IterCount;
-		double Average = IterAcc / (double)IterCount;
-		if (Average > 2.0)
+		void Reserve(SizeType NewCapacity)
 		{
-			static double LastWriteTime = 0.0;
-			double Now = FPlatformTime::Seconds();
-			if (Now - LastWriteTime > 5)
+			NewCapacity = FMath::Max(NewCapacity, (SizeType)1);
+			NewCapacity = static_cast<SizeType>(FPlatformMath::RoundUpToPowerOfTwo64(static_cast<uint64>(NewCapacity)));
+
+			// keep a copy of the old map
+			SizeType* OldMap = Map;
+			SizeType OldCapacity = Capacity;
+
+			LLMCheck(NewCapacity > OldCapacity);
+
+			// allocate the new table
+			Capacity = NewCapacity;
+			Map = (SizeType*)Allocator->Alloc(NewCapacity * sizeof(SizeType));
+
+			for (SizeType Index = 0; Index < NewCapacity; ++Index)
+				Map[Index] = InvalidIndex;
+
+			// copy the values from the old to the new table
+			SizeType* OldItem = OldMap;
+			for (SizeType Index = 0; Index < OldCapacity; ++Index, ++OldItem)
 			{
-				LastWriteTime = Now;
-				UE_LOG(LogStats, Log, TEXT("WARNING: LLMMap average: %f\n"), (float)Average);
+				SizeType KeyIndex = *OldItem;
+
+				if (KeyIndex != InvalidIndex)
+				{
+					SizeType MapIndex = GetMapIndex(Keys[KeyIndex], KeyHashes[KeyIndex]);
+					Map[MapIndex] = KeyIndex;
+				}
 			}
+
+			Allocator->Free(OldMap, OldCapacity * sizeof(SizeType));
+		}
+
+		bool IsItemInUse(SizeType MapIndex) const
+		{
+			return Map[MapIndex] != InvalidIndex;
 		}
+
+		SizeType GetMapIndex(const TKey& Key, SizeType Hash) const
+		{
+			SizeType Mask = Capacity - 1;
+			SizeType MapIndex = Hash & Mask;
+			SizeType KeyIndex = Map[MapIndex];
+
+			while (KeyIndex != InvalidIndex && !(Keys[KeyIndex] == Key))
+			{
+				MapIndex = (MapIndex + 1) & Mask;
+				KeyIndex = Map[MapIndex];
+#ifdef PROFILE_LLMMAP
+				++IterAcc;
 #endif
-		return MapIndex;
-	}
+			}
 
-	// Increase the capacity of the map
-	void Grow()
-	{
-		SizeType NewCapacity = Capacity ? 2 * Capacity : DefaultCapacity;
-		Reserve(NewCapacity);
-	}
+#ifdef PROFILE_LLMMAP
+			++IterCount;
+			double Average = IterAcc / (double)IterCount;
+			if (Average > 2.0)
+			{
+				static double LastWriteTime = 0.0;
+				double Now = FPlatformTime::Seconds();
+				if (Now - LastWriteTime > 5)
+				{
+					LastWriteTime = Now;
+					UE_LOG(LogStats, Log, TEXT("WARNING: LLMMap average: %f\n"), (float)Average);
+				}
+			}
+#endif
+			return MapIndex;
+		}
 
-	static bool InRange(
-		const SizeType Index,
-		const SizeType StartIndex,
-		const SizeType EndIndex)
-	{
-		return (StartIndex <= EndIndex) ?
-			Index >= StartIndex && Index <= EndIndex :
-			Index >= StartIndex || Index <= EndIndex;
-	}
+		// Increase the capacity of the map
+		void Grow()
+		{
+			SizeType NewCapacity = Capacity ? 2 * Capacity : DefaultCapacity / StripeCount;
+			Reserve(NewCapacity);
+		}
 
-	// data
-private:
-	enum { DefaultCapacity = 1024 * 1024 };
-	enum { InvalidIndex = -1 };
-	static const SizeType Margin = (30 * 256) / 100;
+		static bool InRange(
+			const SizeType Index,
+			const SizeType StartIndex,
+			const SizeType EndIndex)
+		{
+			return (StartIndex <= EndIndex) ?
+				Index >= StartIndex && Index <= EndIndex :
+				Index >= StartIndex || Index <= EndIndex;
+		}
 
-	mutable FCriticalSection CriticalSection;
+		mutable FCriticalSection CriticalSection;
 
-	UE::LLMPrivate::FLLMAllocator* Allocator;
+		UE::LLMPrivate::FLLMAllocator* Allocator;
 
-	SizeType* Map;
-	SizeType Count;
-	SizeType Capacity;
+		SizeType* Map;
+		SizeType Count;
+		SizeType Capacity;
 
-	// all these arrays must be kept in sync and are accessed by MapIndex
-	FLLMArray<TKey, SizeType> Keys;
-	FLLMArray<SizeType, SizeType> KeyHashes;
-	FLLMArray<TValue1, SizeType> Values1;
-	FLLMArray<TValue2, SizeType> Values2;
+		// all these arrays must be kept in sync and are accessed by MapIndex
+		FLLMArray<TKey, SizeType> Keys;
+		FLLMArray<SizeType, SizeType> KeyHashes;
+		FLLMArray<TValue1, SizeType> Values1;
+		FLLMArray<TValue2, SizeType> Values2;
 
-	FLLMArray<SizeType, SizeType> FreeKeyIndices;
+		FLLMArray<SizeType, SizeType> FreeKeyIndices;
 
 #ifdef PROFILE_LLMMAP
-	mutable int64 IterAcc;
-	mutable int64 IterCount;
+		mutable int64 IterAcc;
+		mutable int64 IterCount;
 #endif
+	};
+
+	// We divide map items into multiple stripes by the high bits of the hash key, to reduce lock contention.
+	// Each stripe has its own critical section lock, and given even distribution of hash codes (due to the
+	// mix function applied to the key), multiple threads will usually end up with their own lock, and not
+	// need to wait on each other.
+	StripeData MapStripes[StripeCount];
+	FCriticalSection AllStripesLock;
+	std::atomic_bool bAllStripesLocked;
 };
 
 // Pointer key for hash map
diff --git a/Engine/Source/Runtime/Core/Private/ProfilingDebugging/CallstackTrace.h b/Engine/Source/Runtime/Core/Private/ProfilingDebugging/CallstackTrace.h
--- a/Engine/Source/Runtime/Core/Private/ProfilingDebugging/CallstackTrace.h
+++ b/Engine/Source/Runtime/Core/Private/ProfilingDebugging/CallstackTrace.h
@@ -2,7 +2,7 @@
 #pragma once
 
 #include "CoreTypes.h"
-#include "Experimental/Containers/SherwoodHashTable.h"
+#include "Experimental/Containers/GrowOnlyLockFreeHash.h"
 #include "HAL/CriticalSection.h"
 #include "Trace/Trace.h"
 #include "Trace/Trace.inl"
@@ -30,16 +30,20 @@ namespace {
 			uint64	Frames[MaxStackDepth] = { 0 };
 		};
 
-		FCallstackTracer()
+		FCallstackTracer(FMalloc* InMalloc)
+			: KnownSet(InMalloc)
 		{
 			KnownSet.Reserve(1024 * 1024 * 2);
 		}
 		
 		void AddCallstack(const FBacktraceEntry& Entry)
 		{
-			FScopeLock _(&ProducerCs);
 			bool bAlreadyAdded = false;
-			KnownSet.Add(Entry.Id, &bAlreadyAdded);
+
+			// Our set implementation doesn't allow for zero entries (zero represents an empty element
+			// in the hash table), so if we get one due to really bad luck in our 64-bit Id calculation,
+			// treat it as a "1" instead, for purposes of tracking if we've seen that callstack.
+			KnownSet.FindOrAdd(FMath::Max(Entry.Id, (uint64)1), true, &bAlreadyAdded);
 			if (!bAlreadyAdded)
 			{
 				UE_TRACE_LOG(Memory, CallstackSpec, CallstackChannel)
@@ -47,8 +51,23 @@ namespace {
 					<< CallstackSpec.Frames(Entry.Frames, Entry.FrameCount);
 			}
 		}
-	private:	
-		Experimental::TSherwoodSet<uint64> 	KnownSet;
-		FCriticalSection					ProducerCs;
+	private:
+		struct FEncounteredCallstackSetEntry
+		{
+			std::atomic_uint64_t Data;
+
+			inline uint64 GetKey() const { return Data.load(std::memory_order_relaxed); }
+			inline bool GetValue() const { return true; }
+			inline bool IsEmpty() const { return Data.load(std::memory_order_relaxed) == 0; }
+			inline void SetKeyValue(uint64 Key, bool Value) { Data.store(Key, std::memory_order_relaxed); }
+			static inline uint32 KeyHash(uint64 Key) { return static_cast<uint32>(Key); }
+			static inline void ClearEntries(FEncounteredCallstackSetEntry* Entries, int32 EntryCount)
+			{
+				memset(Entries, 0, EntryCount * sizeof(FEncounteredCallstackSetEntry));
+			}
+		};
+		typedef TGrowOnlyLockFreeHash<FEncounteredCallstackSetEntry, uint64, bool> FEncounteredCallstackSet;
+
+		FEncounteredCallstackSet 	KnownSet;
 	};
 }
diff --git a/Engine/Source/Runtime/Core/Private/ProfilingDebugging/Microsoft/MemoryTrace.cpp b/Engine/Source/Runtime/Core/Private/ProfilingDebugging/Microsoft/MemoryTrace.cpp
--- a/Engine/Source/Runtime/Core/Private/ProfilingDebugging/Microsoft/MemoryTrace.cpp
+++ b/Engine/Source/Runtime/Core/Private/ProfilingDebugging/Microsoft/MemoryTrace.cpp
@@ -96,6 +96,8 @@ private:
 	virtual void			DumpAllocatorStats(FOutputDevice& Ar) override				{ InnerMalloc->DumpAllocatorStats(Ar); }
 	virtual bool			ValidateHeap() override										{ return InnerMalloc->ValidateHeap(); }
 	virtual bool			GetAllocationSize(void* Address, SIZE_T &SizeOut) override	{ return InnerMalloc->GetAllocationSize(Address, SizeOut); }
+	virtual void			SetupTLSCachesOnCurrentThread() override					{ return InnerMalloc->SetupTLSCachesOnCurrentThread(); }
+
 	FMalloc*				InnerMalloc;
 	bool					bLight;
 };
diff --git a/Engine/Source/Runtime/Core/Private/ProfilingDebugging/Microsoft/WindowsCallstackTrace.cpp b/Engine/Source/Runtime/Core/Private/ProfilingDebugging/Microsoft/WindowsCallstackTrace.cpp
--- a/Engine/Source/Runtime/Core/Private/ProfilingDebugging/Microsoft/WindowsCallstackTrace.cpp
+++ b/Engine/Source/Runtime/Core/Private/ProfilingDebugging/Microsoft/WindowsCallstackTrace.cpp
@@ -21,12 +21,15 @@
 #include "HAL/PlatformProcess.h"
 #include "Misc/ScopeLock.h"
 #include "ProfilingDebugging/CallstackTrace.h"
+#include "Experimental/Containers/GrowOnlyLockFreeHash.h"
 
 #include <atomic>
 
 // 0=off, 1=stats, 2=validation, 3=truth_compare
 #define BACKTRACE_DBGLVL 0
 
+#define BACKTRACE_LOCK_FREE (1 && (BACKTRACE_DBGLVL == 0))
+
 /*
  * Windows' x64 binaries contain a ".pdata" section that describes the location
  * and size of its functions and details on how to unwind them. The unwind
@@ -143,6 +146,37 @@ private:
 		FModule				Module;
 	};
 
+	struct FFunctionLookupSetEntry
+	{
+		// Bottom 48 bits are key (pointer), top 16 bits are data (RSP bias for function)
+		std::atomic_uint64_t Data;
+
+		inline uint64 GetKey() const { return Data.load(std::memory_order_relaxed) & 0xffffffffffffull; }
+		inline int32 GetValue() const { return static_cast<int64>(Data.load(std::memory_order_relaxed)) >> 48; }
+		inline bool IsEmpty() const { return Data.load(std::memory_order_relaxed) == 0; }
+		inline void SetKeyValue(uint64 Key, int32 Value)
+		{
+			Data.store(Key | (static_cast<int64>(Value) << 48), std::memory_order_relaxed);
+		}
+		static inline uint32 KeyHash(uint64 Key)
+		{
+			// 64 bit pointer to 32 bit hash
+			Key = (~Key) + (Key << 21);
+			Key = Key ^ (Key >> 24);
+			Key = Key * 265;
+			Key = Key ^ (Key >> 14);
+			Key = Key * 21;
+			Key = Key ^ (Key >> 28);
+			Key = Key + (Key << 31);
+			return static_cast<uint32>(Key);
+		}
+		static void ClearEntries(FFunctionLookupSetEntry* Entries, int32 EntryCount)
+		{
+			memset(Entries, 0, EntryCount * sizeof(FFunctionLookupSetEntry));
+		}
+	};
+	typedef TGrowOnlyLockFreeHash<FFunctionLookupSetEntry, uint64, int32> FFunctionLookupSet;
+
 	const FFunction*		LookupFunction(UPTRINT Address, FLookupState& State) const;
 	static FBacktracer*		Instance;
 	mutable FCriticalSection Lock;
@@ -151,6 +185,9 @@ private:
 	int32					ModulesCapacity;
 	FMalloc*				Malloc;
 	FCallstackTracer*		CallstackTracer;
+#if BACKTRACE_LOCK_FREE
+	mutable FFunctionLookupSet	FunctionLookups;
+#endif
 #if BACKTRACE_DBGLVL >= 1
 	mutable uint32			NumFpTruncations = 0;
 	mutable uint32			TotalFunctions = 0;
@@ -163,14 +200,20 @@ FBacktracer* FBacktracer::Instance = nullptr;
 ////////////////////////////////////////////////////////////////////////////////
 FBacktracer::FBacktracer(FMalloc* InMalloc)
 	: Malloc(InMalloc)
+#if BACKTRACE_LOCK_FREE
+	, FunctionLookups(InMalloc)
+#endif
 {
+#if BACKTRACE_LOCK_FREE
+	FunctionLookups.Reserve(512 * 1024);		// 4 MB
+#endif
 	ModulesCapacity = 8;
 	ModulesNum = 0;
 	Modules = (FModule*)Malloc->Malloc(sizeof(FModule) * ModulesCapacity);
 
 	Instance = this;
 	CallstackTracer = (FCallstackTracer*) Malloc->Malloc(sizeof(FCallstackTracer), alignof(FCallstackTracer));
-	CallstackTracer = new(CallstackTracer) FCallstackTracer();
+	CallstackTracer = new(CallstackTracer) FCallstackTracer(InMalloc);
 }
 
 ////////////////////////////////////////////////////////////////////////////////
@@ -528,7 +571,12 @@ void* FBacktracer::GetBacktraceId(void* AddressOfReturnAddress) const
 	uint64 BacktraceId = 0;
 	uint32 FrameIdx = 0;
 
+#if BACKTRACE_LOCK_FREE
+	// When running lock free, we defer the lock until a lock free function lookup fails
+	bool Locked = false;
+#else
 	FScopeLock _(&Lock);
+#endif
 	do
 	{
 		UPTRINT RetAddr = *StackPointer;
@@ -540,12 +588,42 @@ void* FBacktracer::GetBacktraceId(void* AddressOfReturnAddress) const
 		BacktraceId += RetAddr;
 		BacktraceId *= 0x30be8efa499c249dull;
 
+#if BACKTRACE_LOCK_FREE
+		int32 RspBias;
+		bool bIsAlreadyInTable;
+		FunctionLookups.Find(RetAddr, &RspBias, &bIsAlreadyInTable);
+		if (bIsAlreadyInTable)
+		{
+			if (RspBias < 0)
+			{
+				break;
+			}
+			else
+			{
+				StackPointer += RspBias;
+				continue;
+			}
+		}
+		if (!Locked)
+		{
+			Lock.Lock();
+			Locked = true;
+		}
+#endif  // BACKTRACE_LOCK_FREE
+
 		const FFunction* Function = LookupFunction(RetAddr, LookupState);
 		if (Function == nullptr)
 		{
+#if BACKTRACE_LOCK_FREE
+			FunctionLookups.Emplace(RetAddr, -1);
+#endif
 			break;
 		}
 
+#if BACKTRACE_LOCK_FREE
+		FunctionLookups.Emplace(RetAddr, Function->RspBias);
+#endif
+
 #if BACKTRACE_DBGLVL >= 2
 		if (NumBacktrace < 1024)
 		{
@@ -586,6 +664,12 @@ void* FBacktracer::GetBacktraceId(void* AddressOfReturnAddress) const
 	}
 #endif
 
+#if BACKTRACE_LOCK_FREE
+	if (Locked)
+	{
+		Lock.Unlock();
+	}
+#endif
 	// Add to queue to be processed. This might block until there is room in the
 	// queue (i.e. the processing thread has caught up processing).
 	CallstackTracer->AddCallstack(BacktraceEntry);
diff --git a/Engine/Source/Runtime/Core/Private/ProfilingDebugging/TagTrace.cpp b/Engine/Source/Runtime/Core/Private/ProfilingDebugging/TagTrace.cpp
--- a/Engine/Source/Runtime/Core/Private/ProfilingDebugging/TagTrace.cpp
+++ b/Engine/Source/Runtime/Core/Private/ProfilingDebugging/TagTrace.cpp
@@ -1,6 +1,7 @@
 // Copyright Epic Games, Inc. All Rights Reserved.
 #include "ProfilingDebugging/TagTrace.h"
 
+#include "Experimental/Containers/GrowOnlyLockFreeHash.h"
 #include "Containers/Set.h"
 #include "CoreTypes.h"
 #include "MemoryTrace.h"
@@ -135,53 +136,24 @@ public:
 	int32 			AnnounceFNameTag(const FName& TagName);
 
 private:
-	/**
-	 * Allocator used to allocate from the untracked FMalloc instance.
-	 */
-	class FTagAllocator : public TSizedHeapAllocator<32>
-	{
-	public:
-
-		class ForAnyElementType : public TSizedHeapAllocator<32>::ForAnyElementType
-		{
-		public:
-			ForAnyElementType()
-				: Data(nullptr)
-			{}
-
-			FORCEINLINE FScriptContainerElement* GetAllocation() const { return Data; }
-			FORCEINLINE bool HasAllocation() const { return !!Data; }
-			void ResizeAllocation(SizeType PreviousNumElements, SizeType NumElements, SIZE_T NumBytesPerElement)
-			{
-				check(FTagTrace::Malloc);
-				if (Data || NumElements)
-				{
-					Data = (FScriptContainerElement*)FTagTrace::Malloc->Realloc(Data, NumElements * NumBytesPerElement);
-				}
-			}
-
-		private:
-			FScriptContainerElement* Data;
-		};
 
-		template<typename ElementType>
-		class ForElementType : public ForAnyElementType
+	struct FTagNameSetEntry
+	{
+		std::atomic_int32_t Data;
+
+		int32 GetKey() const { return Data.load(std::memory_order_relaxed); }
+		bool GetValue() const { return true; }
+		bool IsEmpty() const { return Data.load(std::memory_order_relaxed) == 0; }			// NAME_None is treated as empty
+		void SetKeyValue(int32 Key, bool Value) { Data.store(Key, std::memory_order_relaxed); }
+		static uint32 KeyHash(int32 Key) { return static_cast<uint32>(Key); }
+		static void ClearEntries(FTagNameSetEntry* Entries, int32 EntryCount)
 		{
-		public:
-			FORCEINLINE ElementType* GetAllocation() const
-			{
-				return (ElementType*)ForAnyElementType::GetAllocation();
-			}
-		};
-
+			memset(Entries, 0, EntryCount * sizeof(FTagNameSetEntry));
+		}
 	};
-	typedef TSparseArrayAllocator<FTagAllocator> FTagSparseArrayAllocator;
-	typedef TInlineAllocator<1,FTagAllocator> FTagHashAllocator; 
-	typedef TSetAllocator<FTagSparseArrayAllocator, FTagHashAllocator> FTagSetAllocator;
-	typedef TSet<int32, DefaultKeyFuncs<int32>, FTagSetAllocator> FTagSet;
+	typedef TGrowOnlyLockFreeHash<FTagNameSetEntry, int32, bool> FTagNameSet;
 
-	FCriticalSection		Cs;
-	FTagSet 				AnnouncedNames;
+	FTagNameSet				AnnouncedNames;
 	static FMalloc* 		Malloc;
 };
 
@@ -189,7 +161,8 @@ FMalloc* FTagTrace::Malloc = nullptr;
 static FTagTrace* GTagTrace = nullptr;
 
 ////////////////////////////////////////////////////////////////////////////////
-FTagTrace::FTagTrace(FMalloc* InMalloc) 
+FTagTrace::FTagTrace(FMalloc* InMalloc)
+	: AnnouncedNames(InMalloc)
 {
 	Malloc = InMalloc;
 	AnnouncedNames.Reserve(1024);
@@ -240,14 +213,24 @@ void FTagTrace::OnAnnounceTagDeclaration(FLLMTagDeclaration& TagDeclaration)
 ////////////////////////////////////////////////////////////////////////////////
 int32 FTagTrace::AnnounceFNameTag(const FName& Name)
 {
-	FScopeLock _(&Cs);
 	const int32 NameIndex = Name.GetDisplayIndex().ToUnstableInt();
-	if (AnnouncedNames.Contains(NameIndex))
+
+	// Don't announce NAME_None, if we happen to get that passed in.  The "AnnouncedNames" container
+	// is not allowed to hold "NAME_None", as zero represents an invalid key.
+	if (!NameIndex)
+	{
+		return NameIndex;
+	}
+
+	// Find or add the item
+	bool bAlreadyInTable;
+	AnnouncedNames.FindOrAdd(NameIndex, true, &bAlreadyInTable);
+	if (bAlreadyInTable)
 	{
 		return NameIndex;
 	}
 
-	AnnouncedNames.Emplace(NameIndex);
+	// First time encountering this name, announce it
 	ANSICHAR NameString[NAME_SIZE];
 	Name.GetPlainANSIString(NameString);
 	return AnnounceCustomTag(NameIndex, -1, NameString);
diff --git a/Engine/Source/Runtime/Core/Public/Experimental/Containers/GrowOnlyLockFreeHash.h b/Engine/Source/Runtime/Core/Public/Experimental/Containers/GrowOnlyLockFreeHash.h
new file mode 100644
--- /dev/null
+++ b/Engine/Source/Runtime/Core/Public/Experimental/Containers/GrowOnlyLockFreeHash.h
@@ -0,0 +1,245 @@
+// Copyright Epic Games, Inc. All Rights Reserved.
+
+#pragma once
+#include "CoreMinimal.h"
+
+// Hash table with fast lock free reads, that only supports insertion of items, and no modification of
+// values.  KeyType must be an integer.  EntryType should be a POD with an identifiable "empty" state
+// that can't occur in the table, and include the following member functions:
+//
+//		KeyType GetKey() const;											// Get the key from EntryType
+//		ValueType GetValue() const;										// Get the value from EntryType
+//		bool IsEmpty() const;											// Query whether EntryType is empty
+//		void SetKeyValue(KeyType Key, ValueType Value);					// Write key and value into EntryType (ATOMICALLY!  See below)
+//		static uint32 KeyHash(KeyType Key);								// Convert Key to more well distributed hash
+//		static void ClearEntries(EntryType* Entries, int32 EntryCount);	// Fill an array of entries with empty values
+//
+// The function "SetKeyValue" must be multi-thread safe when writing new items!  This means writing the
+// Key last and atomically, or writing the entire EntryType in a single write (say if the key and value
+// are packed into a single integer word).  Inline is recommended, since these functions are called a
+// lot in the inner loop of the algorithm.  A simple implementation of "KeyHash" can just return the
+// Key (if it's already reasonable as a hash), or mix the bits if better distribution is required.  A
+// simple implementation of "ClearEntries" can just be a memset, if zero represents an empty entry.
+//
+// A set can be approximated by making "GetValue" a nop function, and just paying attention to the bool
+// result from FindEntry, although you do need to either reserve a certain Key as invalid, or add
+// space to store a valid flag as the Value.  This class should only be used for small value types, as
+// the values are embedded into the hash table, and not stored separately.
+//
+// Writes are implemented using a lock -- it would be possible to make writes lock free (or lock free
+// when resizing doesn't occur), but it adds complexity.  If we were to go that route, it would make
+// sense to create a fully generic lock free set, which would be much more involved to implement and
+// validate than this simple class, and might also offer somewhat worse read perf.  Lock free containers
+// that support item removal either need additional synchronization overhead on readers, so writers can
+// tell if a reader is active and spin, or need graveyard markers and a garbage collection pass called
+// periodically, which makes it no longer a simple standalone container.
+//
+// Lock free reads are accomplished by the reader atomically pulling the hash table pointer from the
+// class.  The hash table is self contained, with its size stored in the table itself, and hash tables
+// are not freed until the class's destruction.  So if the table needs to be reallocated due to a write,
+// active readers will still have valid memory.  This does mean that tables leak, but worst case, you
+// end up with half of the memory being waste.  It would be possible to garbage collect the excess
+// tables, but you'd need some kind of global synchronization to make sure no readers are active.
+//
+// Besides cleanup of wasted tables, it might be useful to provide a function to clear a table.  This
+// would involve clearing the Key for all the elements in the table (but leaving the memory allocated),
+// and can be done safely with active readers.  It's not possible to safely remove individual items due
+// to the need to potentially move other items, which would break an active reader that has already
+// searched past a moved item.  But in the case of removing all items, we don't care when a reader fails,
+// it's expected that eventually all readers will fail, regardless of where they are searching.  A clear
+// function could be useful if a lot of the data you are caching is no longer used, and you want to
+// reset the cache.
+//
+template<typename EntryType, typename KeyType, typename ValueType>
+class TGrowOnlyLockFreeHash
+{
+public:
+	TGrowOnlyLockFreeHash(FMalloc* InMalloc)
+		: Malloc(InMalloc), HashTable(nullptr)
+	{}
+
+	~TGrowOnlyLockFreeHash()
+	{
+		FHashHeader* HashTableNext;
+		for (FHashHeader* HashTableCurrent = HashTable; HashTableCurrent; HashTableCurrent = HashTableNext)
+		{
+			HashTableNext = HashTableCurrent->Next;
+
+			Malloc->Free(HashTableCurrent);
+		}
+	}
+
+	/**
+	 * Preallocate the hash table to a certain size
+	 * @param Count - Number of EntryType elements to allocate
+	 * @warning Can only be called once, and only before any items have been added!
+	 */
+	void Reserve(int32 Count)
+	{
+		FScopeLock _(&WriteCriticalSection);
+		check(HashTable.load(std::memory_order_relaxed) == nullptr);
+
+		if (Count <= 0)
+		{
+			Count = DEFAULT_INITIAL_SIZE;
+		}
+		Count = FPlatformMath::RoundUpToPowerOfTwo(Count);
+		FHashHeader* HashTableLocal = (FHashHeader*)Malloc->Malloc(sizeof(FHashHeader) + (Count - 1) * sizeof(EntryType));
+
+		HashTableLocal->Next = nullptr;
+		HashTableLocal->TableSize = Count;
+		HashTableLocal->Used = 0;
+		EntryType::ClearEntries(HashTableLocal->Elements, Count);
+
+		HashTable.store(HashTableLocal, std::memory_order_release);
+	}
+
+	/**
+	 * Find an entry in the hash table
+	 * @param Key - Key to search for
+	 * @param OutValue - Memory location to write result value to.  Left unmodified if Key isn't found.
+	 * @param bIsAlreadyInTable - Optional result for whether key was found in table.
+	 */
+	void Find(KeyType Key, ValueType *OutValue, bool* bIsAlreadyInTable = nullptr) const
+	{
+		FHashHeader* HashTableLocal = HashTable.load(std::memory_order_acquire);
+		if (HashTableLocal)
+		{
+			uint32 TableMask = HashTableLocal->TableSize - 1;
+
+			// Linear probing
+			for (uint32 TableIndex = EntryType::KeyHash(Key) & TableMask; !HashTableLocal->Elements[TableIndex].IsEmpty(); TableIndex = (TableIndex + 1) & TableMask)
+			{
+				if (HashTableLocal->Elements[TableIndex].GetKey() == Key)
+				{
+					if (OutValue)
+					{
+						*OutValue = HashTableLocal->Elements[TableIndex].GetValue();
+					}
+					if (bIsAlreadyInTable)
+					{
+						*bIsAlreadyInTable = true;
+					}
+					return;
+				}
+			}
+		}
+
+		if (bIsAlreadyInTable)
+		{
+			*bIsAlreadyInTable = false;
+		}
+	}
+
+	/**
+	 * Add an entry with the given Key to the hash table, will do nothing if the item already exists
+	 * @param Key - Key to add
+	 * @param Value - Value to add for key
+	 * @param bIsAlreadyInTable -- Optional result for whether item was already in table
+	 */
+	void Emplace(KeyType Key, ValueType Value, bool* bIsAlreadyInTable = nullptr)
+	{
+		FScopeLock _(&WriteCriticalSection);
+
+		// After locking, check if the item is already in the hash table.
+		ValueType ValueIgnore;
+		bool bFindResult;
+		Find(Key, &ValueIgnore, &bFindResult);
+		if (bFindResult == true)
+		{
+			if (bIsAlreadyInTable)
+			{
+				*bIsAlreadyInTable = true;
+			}
+			return;
+		}
+
+		// Check if there is space in the hash table for a new item.  We resize when the hash
+		// table gets half full or more.  @todo:  allow client to specify max load factor?
+		FHashHeader* HashTableLocal = HashTable;
+
+		if (!HashTableLocal || (HashTableLocal->Used >= HashTableLocal->TableSize / 2))
+		{
+			int32 GrowCount = HashTableLocal ? HashTableLocal->TableSize * 2 : DEFAULT_INITIAL_SIZE;
+			FHashHeader* HashTableGrow = (FHashHeader*)Malloc->Malloc(sizeof(FHashHeader) + (GrowCount - 1) * sizeof(EntryType));
+
+			HashTableGrow->Next = HashTableLocal;
+			HashTableGrow->TableSize = GrowCount;
+			HashTableGrow->Used = 0;
+			EntryType::ClearEntries(HashTableGrow->Elements, GrowCount);
+
+			if (HashTableLocal)
+			{
+				// Copy existing elements from the old table to the new table
+				for (int32 TableIndex = 0; TableIndex < HashTableLocal->TableSize; TableIndex++)
+				{
+					EntryType& Entry = HashTableLocal->Elements[TableIndex];
+					if (!Entry.IsEmpty())
+					{
+						HashInsertInternal(HashTableGrow, Entry.GetKey(), Entry.GetValue());
+					}
+				}
+			}
+
+			HashTableLocal = HashTableGrow;
+			HashTable.store(HashTableGrow, std::memory_order_release);
+		}
+
+		// Then add our new item
+		HashInsertInternal(HashTableLocal, Key, Value);
+
+		if (bIsAlreadyInTable)
+		{
+			*bIsAlreadyInTable = false;
+		}
+	}
+
+	void FindOrAdd(KeyType Key, ValueType Value, bool* bIsAlreadyInTable = nullptr)
+	{
+		// Attempt to find the item lock free, before calling "Emplace", which locks the container
+		bool bFindResult;
+		ValueType IgnoreResult;
+		Find(Key, &IgnoreResult, &bFindResult);
+		if (bFindResult)
+		{
+			if (bIsAlreadyInTable)
+			{
+				*bIsAlreadyInTable = true;
+			}
+			return;
+		}
+
+		Emplace(Key, Value, bIsAlreadyInTable);
+	}
+
+private:
+	struct FHashHeader
+	{
+		FHashHeader* Next;			// Old buffers are stored in a linked list for cleanup
+		int32 TableSize;
+		int32 Used;
+		EntryType Elements[1];		// Variable sized
+	};
+
+	FMalloc* Malloc;
+	std::atomic<FHashHeader*> HashTable;
+	FCriticalSection	WriteCriticalSection;
+
+	static const int32 DEFAULT_INITIAL_SIZE = 1024;
+
+	static void HashInsertInternal(FHashHeader* HashTableLocal, KeyType Key, ValueType Value)
+	{
+		int32 TableMask = HashTableLocal->TableSize - 1;
+
+		// Linear probing
+		for (int32 TableIndex = EntryType::KeyHash(Key) & TableMask;; TableIndex = (TableIndex + 1) & TableMask)
+		{
+			if (HashTableLocal->Elements[TableIndex].IsEmpty())
+			{
+				HashTableLocal->Elements[TableIndex].SetKeyValue(Key, Value);
+				HashTableLocal->Used++;
+				break;
+			}
+		}
+	}
+};
diff --git a/Engine/Source/Runtime/Core/Public/HAL/LowLevelMemTracker.h b/Engine/Source/Runtime/Core/Public/HAL/LowLevelMemTracker.h
--- a/Engine/Source/Runtime/Core/Public/HAL/LowLevelMemTracker.h
+++ b/Engine/Source/Runtime/Core/Public/HAL/LowLevelMemTracker.h
@@ -399,6 +399,15 @@ namespace LLMPrivate
 			}
 		}
 
+		LLMAllocFunction GetPlatformAlloc()
+		{
+			return PlatformAlloc;
+		}
+		LLMFreeFunction GetPlatformFree()
+		{
+			return PlatformFree;
+		}
+
 	private:
 		void* AllocPages(size_t Size);
 		void FreePages(void* Ptr, size_t Size);
