UE5_Main:  Non-unity build fix, missing include in GrowOnlyLockFreeHash.h

#rnx

diff --git a/Engine/Source/Runtime/Core/Public/Experimental/Containers/GrowOnlyLockFreeHash.h b/Engine/Source/Runtime/Core/Public/Experimental/Containers/GrowOnlyLockFreeHash.h
--- a/Engine/Source/Runtime/Core/Public/Experimental/Containers/GrowOnlyLockFreeHash.h
+++ b/Engine/Source/Runtime/Core/Public/Experimental/Containers/GrowOnlyLockFreeHash.h
@@ -2,6 +2,7 @@
 
 #pragma once
 #include "CoreMinimal.h"
+#include "Misc/ScopeLock.h"
 
 // Hash table with fast lock free reads, that only supports insertion of items, and no modification of
 // values.  KeyType must be an integer.  EntryType should be a POD with an identifiable "empty" state
