IoStore: Added support for storing compressed chunks in DDC
#rnx

diff --git a/Engine/Source/Developer/DeveloperToolSettings/Classes/Settings/ProjectPackagingSettings.h b/Engine/Source/Developer/DeveloperToolSettings/Classes/Settings/ProjectPackagingSettings.h
--- a/Engine/Source/Developer/DeveloperToolSettings/Classes/Settings/ProjectPackagingSettings.h
+++ b/Engine/Source/Developer/DeveloperToolSettings/Classes/Settings/ProjectPackagingSettings.h
@@ -358,6 +358,18 @@ public:
 	UPROPERTY(config, EditAnywhere, Category = Packaging, AdvancedDisplay, meta = (DisplayName = "Minimum percentage of a block of data which should be saved when performing compression, otherwise data remains uncompressed"))
 	int32 PackageCompressionMinPercentSaved;
 
+	/**
+	 * Specifies if DDC should be used to store and retrieve compressed data when creating IoStore containers.
+	 */
+	UPROPERTY(config, EditAnywhere, Category = Packaging, AdvancedDisplay, meta = (DisplayName = "Enable DDC for IoStore compression"))
+	bool bPackageCompressionEnableDDC;
+
+	/**
+	 * Specifies the minimum (uncompressed) size for storing a compressed IoStore chunk in DDC.
+	 */
+	UPROPERTY(config, EditAnywhere, Category = Packaging)
+	int32 PackageCompressionMinSizeToConsiderDDC;
+
 	/** 
 	 * Version name for HTTP Chunk Install Data.
 	 */
diff --git a/Engine/Source/Developer/IoStoreUtilities/Private/IoStoreUtilities.cpp b/Engine/Source/Developer/IoStoreUtilities/Private/IoStoreUtilities.cpp
--- a/Engine/Source/Developer/IoStoreUtilities/Private/IoStoreUtilities.cpp
+++ b/Engine/Source/Developer/IoStoreUtilities/Private/IoStoreUtilities.cpp
@@ -58,6 +58,7 @@
 #include "ZenStoreHttpClient.h"
 #include "IPlatformFilePak.h"
 #include "ZenStoreWriter.h"
+#include "ProfilingDebugging/CountersTrace.h"
 
 //PRAGMA_DISABLE_OPTIMIZATION
 
@@ -68,6 +69,8 @@ DEFINE_LOG_CATEGORY_STATIC(LogIoStore, Log, All);
 #define IOSTORE_CPU_SCOPE(NAME) TRACE_CPUPROFILER_EVENT_SCOPE(IoStore##NAME);
 #define IOSTORE_CPU_SCOPE_DATA(NAME, DATA) TRACE_CPUPROFILER_EVENT_SCOPE(IoStore##NAME);
 
+TRACE_DECLARE_MEMORY_COUNTER(IoStoreUsedFileBufferMemory, TEXT("IoStore/UsedFileBufferMemory"));
+
 #define OUTPUT_CHUNKID_DIRECTORY 0
 
 static const FName DefaultCompressionMethod = NAME_Zlib;
@@ -615,7 +618,7 @@ struct FContainerTargetSpec
 	FName Name;
 	FGuid EncryptionKeyGuid;
 	FString OutputPath;
-	FIoStoreWriter* IoStoreWriter;
+	TSharedPtr<IIoStoreWriter> IoStoreWriter;
 	TArray<FContainerTargetFile> TargetFiles;
 	TArray<TUniquePtr<FIoStoreReader>> PatchSourceReaders;
 	EIoContainerFlags ContainerFlags = EIoContainerFlags::None;
@@ -2097,6 +2100,7 @@ private:
 		}
 
 		UsedBufferMemory.AddExchange(SourceBufferSize);
+		TRACE_COUNTER_ADD(IoStoreUsedFileBufferMemory, SourceBufferSize);
 		QueueEntry->WriteRequest->LoadSourceBufferAsync();
 	}
 
@@ -2115,6 +2119,7 @@ private:
 	{
 		uint64 OldValue = UsedBufferMemory.SubExchange(Count);
 		check(OldValue >= Count);
+		TRACE_COUNTER_SUBTRACT(IoStoreUsedFileBufferMemory, Count);
 		MemoryAvailableEvent->Trigger();
 	}
 
@@ -2190,41 +2195,27 @@ int32 CreateTarget(const FIoStoreArguments& Arguments, const FIoStoreWriterSetti
 	}
 
 	TUniquePtr<FIoStoreWriterContext> IoStoreWriterContext(new FIoStoreWriterContext());
-	TArray<FIoStoreWriter*> IoStoreWriters;
-	FIoStoreWriter* GlobalIoStoreWriter = nullptr;
+	FIoStatus IoStatus = IoStoreWriterContext->Initialize(GeneralIoWriterSettings);
+	check(IoStatus.IsOk());
+	TArray<TSharedPtr<IIoStoreWriter>> IoStoreWriters;
+	TSharedPtr<IIoStoreWriter> GlobalIoStoreWriter;
 	{
 		IOSTORE_CPU_SCOPE(InitializeIoStoreWriters);
 		if (!Arguments.IsDLC())
 		{
-			GlobalIoStoreWriter = new FIoStoreWriter(*Arguments.GlobalContainerPath);
+			FIoContainerSettings GlobalContainerSettings;
+			if (Arguments.bSign)
+			{
+				GlobalContainerSettings.SigningKey = Arguments.KeyChain.SigningKey;
+				GlobalContainerSettings.ContainerFlags |= EIoContainerFlags::Signed;
+			}
+			GlobalIoStoreWriter = IoStoreWriterContext->CreateContainer(*Arguments.GlobalContainerPath, GlobalContainerSettings);
 			IoStoreWriters.Add(GlobalIoStoreWriter);
 		}
 		for (FContainerTargetSpec* ContainerTarget : ContainerTargets)
 		{
 			check(ContainerTarget->ContainerId.IsValid());
 			if (!ContainerTarget->OutputPath.IsEmpty())
-			{
-				ContainerTarget->IoStoreWriter = new FIoStoreWriter(*ContainerTarget->OutputPath);
-				IoStoreWriters.Add(ContainerTarget->IoStoreWriter);
-			}
-		}
-		FIoStatus IoStatus = IoStoreWriterContext->Initialize(GeneralIoWriterSettings);
-		check(IoStatus.IsOk());
-
-		FIoContainerSettings GlobalContainerSettings;
-		if (Arguments.bSign)
-		{
-			GlobalContainerSettings.SigningKey = Arguments.KeyChain.SigningKey;
-			GlobalContainerSettings.ContainerFlags |= EIoContainerFlags::Signed;
-		}
-		if (GlobalIoStoreWriter)
-		{
-			IoStatus = GlobalIoStoreWriter->Initialize(*IoStoreWriterContext, GlobalContainerSettings);
-		}
-		check(IoStatus.IsOk());
-		for (FContainerTargetSpec* ContainerTarget : ContainerTargets)
-		{
-			if (ContainerTarget->IoStoreWriter)
 			{
 				FIoContainerSettings ContainerSettings;
 				ContainerSettings.ContainerId = ContainerTarget->ContainerId;
@@ -2245,9 +2236,9 @@ int32 CreateTarget(const FIoStoreArguments& Arguments, const FIoStoreWriterSetti
 					ContainerSettings.ContainerFlags |= EIoContainerFlags::Signed;
 				}
 				ContainerSettings.bGenerateDiffPatch = ContainerTarget->bGenerateDiffPatch;
-				IoStatus = ContainerTarget->IoStoreWriter->Initialize(*IoStoreWriterContext, ContainerSettings);
-				check(IoStatus.IsOk());
+				ContainerTarget->IoStoreWriter = IoStoreWriterContext->CreateContainer(*ContainerTarget->OutputPath, ContainerSettings);
 				ContainerTarget->IoStoreWriter->EnableDiskLayoutOrdering(ContainerTarget->PatchSourceReaders);
+				IoStoreWriters.Add(ContainerTarget->IoStoreWriter);
 			}
 		}
 	}
@@ -2405,23 +2396,51 @@ int32 CreateTarget(const FIoStoreArguments& Arguments, const FIoStoreWriterSetti
 
 	UE_LOG(LogIoStore, Display, TEXT("Serializing container(s)..."));
 
-	TArray<FIoStoreWriterResult> IoStoreWriterResults;
-	IoStoreWriterResults.Reserve(IoStoreWriters.Num());
-	for (FIoStoreWriter* IoStoreWriter : IoStoreWriters)
+	TFuture<void> FlushTask = Async(EAsyncExecution::Thread, [&IoStoreWriterContext]()
 	{
-		TFuture<FIoStoreWriterResult> FlushTask = Async(EAsyncExecution::Thread, [IoStoreWriter]()
-		{
-			return IoStoreWriter->Flush().ConsumeValueOrDie();
-		});
+		IoStoreWriterContext->Flush();
+	});
 
-		while (!FlushTask.IsReady())
+	while (!FlushTask.IsReady())
+	{
+		FlushTask.WaitFor(FTimespan::FromSeconds(2.0));
+		FIoStoreWriterContext::FProgress Progress = IoStoreWriterContext->GetProgress();
+		TStringBuilder<1024> ProgressStringBuilder;
+		if (Progress.SerializedChunksCount >= Progress.TotalChunksCount)
 		{
-			FlushTask.WaitFor(FTimespan::FromSeconds(2.0));
-			FIoStoreWriterContext::FProgress Progress = IoStoreWriterContext->GetProgress();
-			UE_LOG(LogIoStore, Display, TEXT("Hashed, Compressed, Serialized: %lld, %lld, %lld / %lld"), Progress.HashedChunksCount, Progress.CompressedChunksCount, Progress.SerializedChunksCount, Progress.TotalChunksCount);
+			ProgressStringBuilder.Appendf(TEXT("Writing tocs..."));
+		}
+		else if (Progress.SerializedChunksCount)
+		{
+			ProgressStringBuilder.Appendf(TEXT("Writing chunks (%llu/%llu)..."), Progress.SerializedChunksCount, Progress.TotalChunksCount);
+			if (Progress.CompressedChunksCount)
+			{
+				ProgressStringBuilder.Appendf(TEXT(" [%llu compressed]"), Progress.CompressedChunksCount);
+			}
+			if (Progress.ScheduledCompressionTasksCount)
+			{
+				ProgressStringBuilder.Appendf(TEXT(" [%llu compression tasks scheduled]"), Progress.ScheduledCompressionTasksCount);
+			}
+			UE_LOG(LogIoStore, Display, TEXT("%s"), *ProgressStringBuilder);
+		}
+		else
+		{
+			UE_LOG(LogIoStore, Display, TEXT("Hashing chunks (%llu/%llu)..."), Progress.HashedChunksCount, Progress.TotalChunksCount);
 		}
-		IoStoreWriterResults.Emplace(FlushTask.Get());
-		delete IoStoreWriter;
+	}
+	if (GeneralIoWriterSettings.bCompressionEnableDDC)
+	{
+		FIoStoreWriterContext::FProgress Progress = IoStoreWriterContext->GetProgress();
+		uint64 TotalDDCAttempts = Progress.CompressionDDCHitCount + Progress.CompressionDDCMissCount;
+		double DDCHitRate = double(Progress.CompressionDDCHitCount) / TotalDDCAttempts * 100.0;
+		UE_LOG(LogIoStore, Display, TEXT("Compression DDC hits: %llu/%llu (%.2f%%)"), Progress.CompressionDDCHitCount, TotalDDCAttempts, DDCHitRate);
+	}
+
+	TArray<FIoStoreWriterResult> IoStoreWriterResults;
+	IoStoreWriterResults.Reserve(IoStoreWriters.Num());
+	for (TSharedPtr<IIoStoreWriter> IoStoreWriter : IoStoreWriters)
+	{
+		IoStoreWriterResults.Emplace(IoStoreWriter->GetResult().ConsumeValueOrDie());
 	}
 	IoStoreWriters.Empty();
 
@@ -2495,7 +2514,7 @@ int32 CreateContentPatch(const FIoStoreArguments& Arguments, const FIoStoreWrite
 	TUniquePtr<FIoStoreWriterContext> IoStoreWriterContext(new FIoStoreWriterContext());
 	FIoStatus IoStatus = IoStoreWriterContext->Initialize(GeneralIoWriterSettings);
 	check(IoStatus.IsOk());
-	TArray<FIoStoreWriterResult> Results;
+	TArray<TSharedPtr<IIoStoreWriter>> IoStoreWriters;
 	for (const FContainerSourceSpec& Container : Arguments.Containers)
 	{
 		TArray<TUniquePtr<FIoStoreReader>> SourceReaders = CreatePatchSourceReaders(Container.PatchSourceContainerFiles, Arguments);
@@ -2506,8 +2525,6 @@ int32 CreateContentPatch(const FIoStoreArguments& Arguments, const FIoStoreWrite
 			return -1;
 		}
 
-		FIoStoreWriter IoStoreWriter(*Container.OutputPath);
-		
 		EIoContainerFlags TargetContainerFlags = TargetReader->GetContainerFlags();
 
 		FIoContainerSettings ContainerSettings;
@@ -2536,9 +2553,8 @@ int32 CreateContentPatch(const FIoStoreArguments& Arguments, const FIoStoreWrite
 			ContainerSettings.EncryptionKey = Key->Key;
 		}
 
-		IoStatus = IoStoreWriter.Initialize(*IoStoreWriterContext, ContainerSettings);
-		check(IoStatus.IsOk());
-
+		TSharedPtr<IIoStoreWriter> IoStoreWriter = IoStoreWriterContext->CreateContainer(*Container.OutputPath, ContainerSettings);
+		IoStoreWriters.Add(IoStoreWriter);
 		TMap<FIoChunkId, FIoChunkHash> SourceHashByChunkId;
 		for (const TUniquePtr<FIoStoreReader>& SourceReader : SourceReaders)
 		{
@@ -2584,13 +2600,17 @@ int32 CreateContentPatch(const FIoStoreArguments& Arguments, const FIoStoreWrite
 				}
 				WriteOptions.bIsMemoryMapped = ChunkInfo.bIsMemoryMapped;
 				WriteOptions.bForceUncompressed = ChunkInfo.bForceUncompressed; 
-				IoStoreWriter.Append(ChunkInfo.Id, ChunkBuffer.ConsumeValueOrDie(), WriteOptions);
+				IoStoreWriter->Append(ChunkInfo.Id, ChunkBuffer.ConsumeValueOrDie(), WriteOptions);
 			}
 			return true;
 		});
+	}
 
-
-		Results.Emplace(IoStoreWriter.Flush().ConsumeValueOrDie());
+	IoStoreWriterContext->Flush();
+	TArray<FIoStoreWriterResult> Results;
+	for (TSharedPtr<IIoStoreWriter> IoStoreWriter : IoStoreWriters)
+	{
+		Results.Emplace(IoStoreWriter->GetResult().ConsumeValueOrDie());
 	}
 
 	LogWriterResults(Results);
diff --git a/Engine/Source/Runtime/Core/Private/IO/IoStore.cpp b/Engine/Source/Runtime/Core/Private/IO/IoStore.cpp
--- a/Engine/Source/Runtime/Core/Private/IO/IoStore.cpp
+++ b/Engine/Source/Runtime/Core/Private/IO/IoStore.cpp
@@ -19,6 +19,11 @@
 #include "Async/AsyncFileHandle.h"
 #include "Interfaces/ITargetPlatform.h"
 #include "Misc/ConfigCacheIni.h"
+#include "DerivedDataPluginInterface.h"
+#include "DerivedDataCacheInterface.h"
+#include "Serialization/LargeMemoryReader.h"
+#include "Serialization/MemoryWriter.h"
+#include "ProfilingDebugging/CountersTrace.h"
 
 //////////////////////////////////////////////////////////////////////////
 
@@ -26,6 +31,8 @@ constexpr char FIoStoreTocHeader::TocMagicImg[];
 
 //////////////////////////////////////////////////////////////////////////
 
+TRACE_DECLARE_INT_COUNTER(IoStoreAvailableCompressionBuffers, TEXT("IoStore/AvailableCompressionBuffers"));
+
 template<typename ArrayType>
 bool WriteArray(IFileHandle* FileHandle, const ArrayType& Array)
 {
@@ -148,15 +155,13 @@ struct FChunkBlock
 	uint64 CompressedSize = 0;
 	uint64 UncompressedSize = 0;
 	FName CompressionMethod = NAME_None;
-	int32 CompressionMinBytesSaved = 0;
-	int32 CompressionMinPercentSaved = 0;
 	FSHAHash Signature;
 };
 
 struct FIoStoreWriteQueueEntry
 {
 	FIoStoreWriteQueueEntry* Next = nullptr;
-	FIoStoreWriterImpl* Writer = nullptr;
+	class FIoStoreWriter* Writer = nullptr;
 	IIoStoreWriteRequest* Request = nullptr;
 	FIoChunkId ChunkId;
 	FIoChunkHash ChunkHash;
@@ -171,12 +176,16 @@ struct FIoStoreWriteQueueEntry
 	FGraphEventRef HashTask;
 	FGraphEventRef BeginCompressionBarrier;
 	FGraphEventRef FinishCompressionBarrier;
+	FGraphEventRef FinishEncryptionAndSigningBarrier;
 	FGraphEventRef BeginWriteBarrier;
 	FGraphEventRef WriteFinishedEvent;
+	TAtomic<int32> CompressedBlocksCount{ 0 };
 	TAtomic<int32> FinishedBlocksCount{ 0 };
 	int32 PartitionIndex = -1;
+	FString DDCCacheKey;
 	bool bAdded = false;
 	bool bModified = false;
+	bool bStoreCompressedDataInDDC = false;
 };
 
 class FIoStoreWriteQueue
@@ -260,19 +269,22 @@ private:
 
 class FIoStoreWriterContextImpl
 {
-	static constexpr uint64 DefaultMemoryLimit = 1ull << 30ull;
+	static constexpr uint64 DefaultMemoryLimit = 2ull << 30ull;
 public:
 	FIoStoreWriterContextImpl()
 	{
+		DDC = GetDerivedDataCache();
 	}
 
 	~FIoStoreWriterContextImpl()
 	{
 		BeginCompressionQueue.CompleteAdding();
-		FinishCompressionQueue.CompleteAdding();
+		BeginEncryptionAndSigningQueue.CompleteAdding();
+		FinishEncryptionAndSigningQueue.CompleteAdding();
 		WriterQueue.CompleteAdding();
 		BeginCompressionThread.Wait();
-		FinishCompressionThread.Wait();
+		BeginEncryptionAndSigningThread.Wait();
+		FinishEncryptionAndSigningThread.Wait();
 		WriterThread.Wait();
 		if (CompressionBufferAvailableEvent)
 		{
@@ -297,20 +309,26 @@ public:
 		CompressionBufferSize = FMath::Max(CompressionBufferSize, static_cast<int32>(WriterSettings.CompressionBlockSize));
 		CompressionBufferSize = Align(CompressionBufferSize, FAES::AESBlockSize);
 
-		TotalCompressionBufferCount = DefaultMemoryLimit / CompressionBufferSize;
+		TotalCompressionBufferCount = int32(DefaultMemoryLimit / CompressionBufferSize);
 		AvailableCompressionBuffers.Reserve(TotalCompressionBufferCount);
 		for (int32 BufferIndex = 0; BufferIndex < TotalCompressionBufferCount; ++BufferIndex)
 		{
 			AvailableCompressionBuffers.Add(new FIoBuffer(CompressionBufferSize));
 		}
+		TRACE_COUNTER_SET(IoStoreAvailableCompressionBuffers, AvailableCompressionBuffers.Num());
 
 		BeginCompressionThread = Async(EAsyncExecution::Thread, [this]() { BeginCompressionThreadFunc(); });
-		FinishCompressionThread = Async(EAsyncExecution::Thread, [this]() { FinishCompressionThreadFunc(); });
+		BeginEncryptionAndSigningThread = Async(EAsyncExecution::Thread, [this]() { BeginEncryptionAndSigningThreadFunc(); });
+		FinishEncryptionAndSigningThread = Async(EAsyncExecution::Thread, [this]() { FinishEncryptionAndSigningThreadFunc(); });
 		WriterThread = Async(EAsyncExecution::Thread, [this]() { WriterThreadFunc(); });
 
 		return FIoStatus::Ok;
 	}
 
+	TSharedPtr<IIoStoreWriter> CreateContainer(const TCHAR* InContainerPath, const FIoContainerSettings& InContainerSettings);
+
+	void Flush();
+
 	FIoStoreWriterContext::FProgress GetProgress() const
 	{
 		FIoStoreWriterContext::FProgress Progress;
@@ -318,6 +336,9 @@ public:
 		Progress.HashedChunksCount = HashedChunksCount.Load();
 		Progress.CompressedChunksCount = CompressedChunksCount.Load();
 		Progress.SerializedChunksCount = SerializedChunksCount.Load();
+		Progress.ScheduledCompressionTasksCount = ScheduledCompressionTasksCount.Load();
+		Progress.CompressionDDCHitCount = CompressionDDCHitCount.Load();
+		Progress.CompressionDDCMissCount = CompressionDDCMissCount.Load();
 		return Progress;
 	}
 
@@ -346,6 +367,7 @@ public:
 				if (AvailableCompressionBuffers.Num())
 				{
 					AllocatedBuffer = AvailableCompressionBuffers.Pop();
+					TRACE_COUNTER_DECREMENT(IoStoreAvailableCompressionBuffers);
 				}
 			}
 			if (!AllocatedBuffer)
@@ -369,6 +391,7 @@ public:
 			FScopeLock Lock(&AvailableCompressionBuffersCritical);
 			bTriggerEvent = AvailableCompressionBuffers.Num() == 0;
 			AvailableCompressionBuffers.Push(Buffer);
+			TRACE_COUNTER_INCREMENT(IoStoreAvailableCompressionBuffers);
 		}
 		if (bTriggerEvent)
 		{
@@ -378,27 +401,35 @@ public:
 
 private:
 	void BeginCompressionThreadFunc();
-	void FinishCompressionThreadFunc();
+	void BeginEncryptionAndSigningThreadFunc();
+	void FinishEncryptionAndSigningThreadFunc();
 	void WriterThreadFunc();
 
 	FIoStoreWriterSettings WriterSettings;
+	FDerivedDataCacheInterface* DDC = nullptr;
 	FEvent* CompressionBufferAvailableEvent = nullptr;
 	TFuture<void> BeginCompressionThread;
-	TFuture<void> FinishCompressionThread;
+	TFuture<void> BeginEncryptionAndSigningThread;
+	TFuture<void> FinishEncryptionAndSigningThread;
 	TFuture<void> WriterThread;
 	FIoStoreWriteQueue BeginCompressionQueue;
-	FIoStoreWriteQueue FinishCompressionQueue;
+	FIoStoreWriteQueue BeginEncryptionAndSigningQueue;
+	FIoStoreWriteQueue FinishEncryptionAndSigningQueue;
 	FIoStoreWriteQueue WriterQueue;
 	TAtomic<uint64> TotalChunksCount{ 0 };
 	TAtomic<uint64> HashedChunksCount{ 0 };
 	TAtomic<uint64> CompressedChunksCount{ 0 };
 	TAtomic<uint64> SerializedChunksCount{ 0 };
+	TAtomic<uint64> ScheduledCompressionTasksCount{ 0 };
+	TAtomic<uint64> CompressionDDCHitCount{ 0 };
+	TAtomic<uint64> CompressionDDCMissCount{ 0 };
 	FCriticalSection AvailableCompressionBuffersCritical;
 	TArray<FIoBuffer*> AvailableCompressionBuffers;
 	int32 CompressionBufferSize = -1;
 	int32 TotalCompressionBufferCount = -1;
+	TArray<TSharedPtr<FIoStoreWriter>> IoStoreWriters;
 
-	friend class FIoStoreWriterImpl;
+	friend class FIoStoreWriter;
 };
 
 FIoStoreWriterContext::FIoStoreWriterContext()
@@ -417,6 +448,16 @@ UE_NODISCARD FIoStatus FIoStoreWriterContext::Initialize(const FIoStoreWriterSet
 	return Impl->Initialize(InWriterSettings);
 }
 
+TSharedPtr<IIoStoreWriter> FIoStoreWriterContext::CreateContainer(const TCHAR* InContainerPath, const FIoContainerSettings& InContainerSettings)
+{
+	return Impl->CreateContainer(InContainerPath, InContainerSettings);
+}
+
+void FIoStoreWriterContext::Flush()
+{
+	Impl->Flush();
+}
+
 FIoStoreWriterContext::FProgress FIoStoreWriterContext::GetProgress() const
 {
 	return Impl->GetProgress();
@@ -534,10 +575,11 @@ private:
 	TArray<uint32> FileTocEntryIndices;
 };
 
-class FIoStoreWriterImpl
+class FIoStoreWriter
+	: public IIoStoreWriter
 {
 public:
-	FIoStoreWriterImpl(const TCHAR* InContainerPath)
+	FIoStoreWriter(const TCHAR* InContainerPath)
 		: ContainerPath(InContainerPath)
 	{
 	}
@@ -578,7 +620,7 @@ public:
 		return FIoStatus::Ok;
 	}
 
-	void EnableDiskLayoutOrdering(const TArray<TUniquePtr<FIoStoreReader>>& PatchSourceReaders)
+	virtual void EnableDiskLayoutOrdering(const TArray<TUniquePtr<FIoStoreReader>>& PatchSourceReaders) override
 	{
 		check(!LayoutEntriesHead);
 		check(!Entries.Num());
@@ -615,7 +657,7 @@ public:
 		LayoutEntriesTail->Prev = PrevEntryLink;
 	}
 
-	void Append(const FIoChunkId& ChunkId, IIoStoreWriteRequest* Request, const FIoWriteOptions& WriteOptions)
+	virtual void Append(const FIoChunkId& ChunkId, IIoStoreWriteRequest* Request, const FIoWriteOptions& WriteOptions) override
 	{
 		TRACE_CPUPROFILER_EVENT_SCOPE(AppendWriteRequest);
 		check(!bHasFlushed);
@@ -632,6 +674,7 @@ public:
 		Entry->HashBarrier = FGraphEvent::CreateGraphEvent();
 		Entry->BeginCompressionBarrier = FGraphEvent::CreateGraphEvent();
 		Entry->FinishCompressionBarrier = FGraphEvent::CreateGraphEvent();
+		Entry->FinishEncryptionAndSigningBarrier = FGraphEvent::CreateGraphEvent();
 		Entry->BeginWriteBarrier = FGraphEvent::CreateGraphEvent();
 		Entry->WriteFinishedEvent = FGraphEvent::CreateGraphEvent();
 		
@@ -643,23 +686,55 @@ public:
 			const FIoBuffer* SourceBuffer = Entry->Request->GetSourceBuffer();
 			Entry->ChunkHash = FIoChunkHash::HashBuffer(SourceBuffer->Data(), SourceBuffer->DataSize());
 			WriterContext->HashedChunksCount.IncrementExchange();
-			if (LayoutEntriesHead)
-			{
-				// Release the source data buffer if disk layout ordering is enabled, it will be reloaded
-				// later when we start compressing the chunk
-				Entry->Request->FreeSourceBuffer();
-			}
+			// Release the source data buffer, it will be reloaded later when we start compressing the chunk
+			Entry->Request->FreeSourceBuffer();
 		}, TStatId(), &HashPrereqs, ENamedThreads::AnyHiPriThreadHiPriTask);
-		// If disk layout ordering is enabled we need to hash each chunk before we can determine the
-		// layout, otherwise we can start compressing the chunk immediately  
-		if (LayoutEntriesHead)
-		{
-			Entry->Request->PrepareSourceBufferAsync(Entry->HashBarrier);
-		}
-		else
+		Entry->Request->PrepareSourceBufferAsync(Entry->HashBarrier);
+	}
+
+	virtual void Append(const FIoChunkId& ChunkId, FIoBuffer Chunk, const FIoWriteOptions& WriteOptions, uint64 OrderHint) override
+	{
+		struct FWriteRequest
+			: IIoStoreWriteRequest
 		{
-			WriterContext->ScheduleCompression(Entry);
-		}
+			FWriteRequest(FIoBuffer InSourceBuffer, uint64 InOrderHint)
+				: OrderHint(InOrderHint)
+			{
+				SourceBuffer = InSourceBuffer;
+				SourceBuffer.MakeOwned();
+			}
+
+			virtual ~FWriteRequest() = default;
+
+			void PrepareSourceBufferAsync(FGraphEventRef CompletionEvent) override
+			{
+				CompletionEvent->DispatchSubsequents();
+			}
+
+			const FIoBuffer* GetSourceBuffer() override
+			{
+				return &SourceBuffer;
+			}
+
+			void FreeSourceBuffer() override
+			{
+			}
+
+			uint64 GetOrderHint() override
+			{
+				return OrderHint;
+			}
+
+			TArrayView<const FFileRegion> GetRegions()
+			{
+				return TArrayView<const FFileRegion>();
+			}
+
+			FIoBuffer SourceBuffer;
+			uint64 OrderHint;
+		};
+
+		Append(ChunkId, new FWriteRequest(Chunk, OrderHint), WriteOptions);
 	}
 
 	bool GeneratePerfectHashes(FIoStoreTocResource& TocResource)
@@ -779,42 +854,9 @@ public:
 		return true;
 	}
 
-	UE_NODISCARD TIoStatusOr<FIoStoreWriterResult> Flush()
+	void Finalize()
 	{
-		TRACE_CPUPROFILER_EVENT_SCOPE(FlushContainer);
-
-		if (bHasFlushed)
-		{
-			return Result;
-		}
-		bHasFlushed = true;
-
-		if (LayoutEntriesHead)
-		{
-			{
-				TRACE_CPUPROFILER_EVENT_SCOPE(WaitForChunkHashes);
-				Algo::Reverse(Entries);
-				for (FIoStoreWriteQueueEntry* Entry : Entries)
-				{
-					FTaskGraphInterface::Get().WaitUntilTaskCompletes(Entry->HashTask);
-					Entry = Entry->Next;
-				}
-			}
-			FinalizeLayout(Entries);
-
-			for (FIoStoreWriteQueueEntry* Entry : Entries)
-			{
-				WriterContext->ScheduleCompression(Entry);
-			}
-		}
-
-		{
-			TRACE_CPUPROFILER_EVENT_SCOPE(WaitForWritesToComplete);
-			for (FIoStoreWriteQueueEntry* Entry : Entries)
-			{
-				Entry->WriteFinishedEvent->Wait();
-			}
-		}
+		check(bHasFlushed);
 
 		UncompressedContainerSize = TotalEntryUncompressedSize + TotalPaddingSize;
 		CompressedContainerSize = 0;
@@ -879,11 +921,8 @@ public:
 		}
 
 		TIoStatusOr<uint64> TocSize = FIoStoreTocResource::Write(*TocFilePath, TocResource, ContainerSettings, WriterContext->GetSettings());
-		if (!TocSize.IsOk())
-		{
-			return TocSize.Status();
-		}
-
+		check(TocSize.IsOk());
+		
 		Result.ContainerId = ContainerSettings.ContainerId;
 		Result.ContainerName = FPaths::GetBaseFilename(TocFilePath);
 		Result.ContainerFlags = ContainerSettings.ContainerFlags;
@@ -919,7 +958,15 @@ public:
 		}
 
 		Entries.Empty();
+		bHasResult = true;
+	}
 
+	TIoStatusOr<FIoStoreWriterResult> GetResult()
+	{
+		if (!bHasResult)
+		{
+			return FIoStatus::Invalid;
+		}
 		return Result;
 	}
 
@@ -945,11 +992,11 @@ private:
 		int32 PartitionIndex = -1;
 	};
 
-	void FinalizeLayout(TArray<FIoStoreWriteQueueEntry*>& WriteQueueEntries)
+	void FinalizeLayout()
 	{
 		TRACE_CPUPROFILER_EVENT_SCOPE(FinalizeLayout);
 		
-		Algo::Sort(WriteQueueEntries, [](const FIoStoreWriteQueueEntry* A, const FIoStoreWriteQueueEntry* B)
+		Algo::Sort(Entries, [](const FIoStoreWriteQueueEntry* A, const FIoStoreWriteQueueEntry* B)
 		{
 			uint64 AOrderHint = A->Request->GetOrderHint();
 			uint64 BOrderHint = B->Request->GetOrderHint();
@@ -963,7 +1010,7 @@ private:
 		TMap<int64, FLayoutEntry*> LayoutEntriesByOrderMap;
 		int64 IdealOrder = 0;
 		TArray<FLayoutEntry*> UnassignedEntries;
-		for (FIoStoreWriteQueueEntry* WriteQueueEntry : WriteQueueEntries)
+		for (FIoStoreWriteQueueEntry* WriteQueueEntry : Entries)
 		{
 			FLayoutEntry* FindPreviousEntry = PreviousBuildLayoutEntryByChunkId.FindRef(WriteQueueEntry->ChunkId);
 			if (FindPreviousEntry)
@@ -1050,7 +1097,7 @@ private:
 				ReserveInPartition.ReservedSpace += EntryIt->CompressedSize;
 			}
 		}
-		Swap(WriteQueueEntries, IncludedQueueEntries);
+		Swap(Entries, IncludedQueueEntries);
 
 		LayoutEntriesHead = nullptr;
 		LayoutEntriesTail = nullptr;
@@ -1090,6 +1137,55 @@ private:
 		return FIoStatus::Ok;
 	}
 
+	void CompressBlock(FChunkBlock* Block)
+	{
+		TRACE_CPUPROFILER_EVENT_SCOPE(CompressBlock);
+		check(Block->CompressionMethod != NAME_None);
+		int32 CompressedBlockSize = int32(Block->IoBuffer->DataSize());
+		bool bCompressed;
+		{
+			bCompressed = FCompression::CompressMemoryIfWorthDecompressing(
+				Block->CompressionMethod,
+				WriterContext->WriterSettings.CompressionMinBytesSaved,
+				WriterContext->WriterSettings.CompressionMinPercentSaved,
+				Block->IoBuffer->Data(),
+				CompressedBlockSize,
+				Block->UncompressedData,
+				static_cast<int32>(Block->UncompressedSize),
+				COMPRESS_ForPackaging);
+		}
+		if (!bCompressed)
+		{
+			Block->CompressionMethod = NAME_None;
+			Block->CompressedSize = Block->UncompressedSize;
+			FMemory::Memcpy(Block->IoBuffer->Data(), Block->UncompressedData, Block->UncompressedSize);
+		}
+		else
+		{
+			check(CompressedBlockSize > 0);
+			check(CompressedBlockSize < Block->UncompressedSize);
+			Block->CompressedSize = CompressedBlockSize;
+		}
+	}
+
+	void SerializeCompressedDDCData(FIoStoreWriteQueueEntry* Entry, FArchive& Ar)
+	{
+		uint32 NumBlocks = Entry->ChunkBlocks.Num();
+		Ar << NumBlocks;
+		check(NumBlocks == Entry->ChunkBlocks.Num());
+		for (FChunkBlock& Block : Entry->ChunkBlocks)
+		{
+			Ar << Block.CompressedSize;
+			check(Block.CompressedSize <= Block.UncompressedSize);
+			if (Ar.IsLoading() && Block.CompressedSize == Block.UncompressedSize)
+			{
+				Block.CompressionMethod = NAME_None;
+			}
+			check(Block.IoBuffer->DataSize() >= Block.CompressedSize);
+			Ar.Serialize(Block.IoBuffer->Data(), Block.CompressedSize);
+		}
+	}
+
 	void BeginCompress(FIoStoreWriteQueueEntry* Entry)
 	{
 		FName CompressionMethod = NAME_None;
@@ -1106,7 +1202,6 @@ private:
 		const uint64 NumChunkBlocks = Align(Entry->UncompressedSize, WriterSettings.CompressionBlockSize) / WriterSettings.CompressionBlockSize;
 		if (NumChunkBlocks == 0)
 		{
-			WriterContext->CompressedChunksCount.IncrementExchange();
 			Entry->FinishCompressionBarrier->DispatchSubsequents();
 			return;
 		}
@@ -1120,85 +1215,121 @@ private:
 			FChunkBlock& Block = Entry->ChunkBlocks[BlockIndex];
 			Block.IoBuffer = WriterContext->AllocCompressionBuffer(int32(NumChunkBlocks));
 			Block.CompressionMethod = CompressionMethod;
-			Block.CompressionMinBytesSaved = WriterSettings.CompressionMinBytesSaved;
-			Block.CompressionMinPercentSaved = WriterSettings.CompressionMinPercentSaved;
 			Block.UncompressedSize = FMath::Min(BytesToProcess, WriterSettings.CompressionBlockSize);
 			Block.UncompressedData = UncompressedData;
 			BytesToProcess -= Block.UncompressedSize;
 			UncompressedData += Block.UncompressedSize;
+		}
 
+		if (CompressionMethod == NAME_None)
+		{
+			for (FChunkBlock& Block : Entry->ChunkBlocks)
+			{
+				Block.CompressionMethod = NAME_None;
+				Block.CompressedSize = Block.UncompressedSize;
+				FMemory::Memcpy(Block.IoBuffer->Data(), Block.UncompressedData, Block.UncompressedSize);
+			}
+			Entry->FinishCompressionBarrier->DispatchSubsequents();
+			return;
+		}
+
+		bool bUseDDCCompression =
+			WriterContext->DDC &&
+			WriterSettings.bCompressionEnableDDC &&
+			Entry->UncompressedSize > WriterSettings.CompressionMinSizeToConsiderDDC;
+		if (bUseDDCCompression)
+		{
+			TStringBuilder<256> CacheKeySuffix;
+			CacheKeySuffix.Append(Entry->ChunkHash.ToString());
+			WriterSettings.CompressionMethod.AppendString(CacheKeySuffix);
+			CacheKeySuffix.Append(FCompression::GetCompressorDDCSuffix(WriterSettings.CompressionMethod));
+			CacheKeySuffix.Appendf(TEXT("%d_%d"),
+				WriterSettings.CompressionMinBytesSaved,
+				WriterSettings.CompressionMinPercentSaved);
+			Entry->DDCCacheKey = FDerivedDataCacheInterface::BuildCacheKey(
+				TEXT("IoStoreCompression"),
+				TEXT("985D3FAD-71D0-4758-A777-A910B49CC4BF"),
+				CacheKeySuffix.ToString());
+
+			TArray<uint8> DDCData;
+			TRACE_CPUPROFILER_EVENT_SCOPE(ReadFromDDC);
+			if (WriterContext->DDC->GetSynchronous(*Entry->DDCCacheKey, DDCData, Entry->Options.FileName))
+			{
+				WriterContext->CompressionDDCHitCount.IncrementExchange();
+				FLargeMemoryReader DDCDataReader(DDCData.GetData(), DDCData.Num());
+				SerializeCompressedDDCData(Entry, DDCDataReader);
+				Entry->FinishCompressionBarrier->DispatchSubsequents();
+				return;
+			}
+			else
+			{
+				WriterContext->CompressionDDCMissCount.IncrementExchange();
+				Entry->bStoreCompressedDataInDDC = true;
+			}
+		}
+		
+		for (FChunkBlock& Block : Entry->ChunkBlocks)
+		{
+			WriterContext->ScheduledCompressionTasksCount.IncrementExchange();
 			FFunctionGraphTask::CreateAndDispatchWhenReady([this, Entry, BlockPtr = &Block]()
 			{
 				CompressBlock(BlockPtr);
-				int32 FinishedBlocksCount = Entry->FinishedBlocksCount.IncrementExchange();
-				if (FinishedBlocksCount + 1 == Entry->ChunkBlocks.Num())
+				WriterContext->ScheduledCompressionTasksCount.DecrementExchange();
+				int32 CompressedBlocksCount = Entry->CompressedBlocksCount.IncrementExchange();
+				if (CompressedBlocksCount + 1 == Entry->ChunkBlocks.Num())
 				{
 					WriterContext->CompressedChunksCount.IncrementExchange();
 					Entry->FinishCompressionBarrier->DispatchSubsequents();
 				}
-			}, TStatId(), nullptr, ENamedThreads::AnyHiPriThreadHiPriTask);
+			}, TStatId(), nullptr, ENamedThreads::AnyHiPriThreadNormalTask);
 		}
 	}
 
-	void CompressBlock(FChunkBlock* Block)
+	void BeginEncryptAndSign(FIoStoreWriteQueueEntry* Entry)
 	{
-		if (Block->CompressionMethod != NAME_None)
+		if (ContainerSettings.IsEncrypted() || ContainerSettings.IsSigned())
 		{
-			int32 CompressedBlockSize = int32(Block->IoBuffer->DataSize());
-			bool bCompressed;
-
-			{
-				TRACE_CPUPROFILER_EVENT_SCOPE(CompressMemory);
-				bCompressed = FCompression::CompressMemoryIfWorthDecompressing(
-					Block->CompressionMethod,
-					Block->CompressionMinBytesSaved,
-					Block->CompressionMinPercentSaved,
-					Block->IoBuffer->Data(),
-					CompressedBlockSize,
-					Block->UncompressedData,
-					static_cast<int32>(Block->UncompressedSize),
-					COMPRESS_ForPackaging);
-			}
-			if ( !bCompressed )
+			FFunctionGraphTask::CreateAndDispatchWhenReady([this, Entry]()
 			{
-				Block->CompressionMethod = NAME_None;
-			}
-			else
-			{
-				check( CompressedBlockSize > 0 );
-				check( CompressedBlockSize < Block->UncompressedSize)
-
-				Block->CompressedSize = CompressedBlockSize;
-			}
+				EncryptAndSign(Entry);
+				Entry->FinishEncryptionAndSigningBarrier->DispatchSubsequents();
+			}, TStatId(), nullptr, ENamedThreads::AnyHiPriThreadHiPriTask);
 		}
-		if (Block->CompressionMethod == NAME_None)
+		else
 		{
-			Block->CompressedSize = Block->UncompressedSize;
-			FMemory::Memcpy(Block->IoBuffer->Data(), Block->UncompressedData, Block->UncompressedSize);
+			EncryptAndSign(Entry);
+			Entry->FinishEncryptionAndSigningBarrier->DispatchSubsequents();
 		}
+	}
 
-		// Always align each compressed block to AES block size but store the compressed block size in the TOC
-		Block->Size = Block->CompressedSize;
-		if (!IsAligned(Block->Size, FAES::AESBlockSize))
+	void EncryptAndSign(FIoStoreWriteQueueEntry* Entry)
+	{
+		TRACE_CPUPROFILER_EVENT_SCOPE(EncryptAndSignChunk);
+		for (FChunkBlock& Block : Entry->ChunkBlocks)
 		{
-			uint64 AlignedCompressedBlockSize = Align(Block->Size, FAES::AESBlockSize);
-			uint8* CompressedData = Block->IoBuffer->Data();
-			for (uint64 FillIndex = Block->Size; FillIndex < AlignedCompressedBlockSize; ++FillIndex)
+			// Always align each compressed block to AES block size but store the compressed block size in the TOC
+			Block.Size = Block.CompressedSize;
+			if (!IsAligned(Block.Size, FAES::AESBlockSize))
 			{
-				check(FillIndex < Block->IoBuffer->DataSize());
-				CompressedData[FillIndex] = CompressedData[(FillIndex - Block->Size) % Block->Size];
+				uint64 AlignedCompressedBlockSize = Align(Block.Size, FAES::AESBlockSize);
+				uint8* CompressedData = Block.IoBuffer->Data();
+				for (uint64 FillIndex = Block.Size; FillIndex < AlignedCompressedBlockSize; ++FillIndex)
+				{
+					check(FillIndex < Block.IoBuffer->DataSize());
+					CompressedData[FillIndex] = CompressedData[(FillIndex - Block.Size) % Block.Size];
+				}
+				Block.Size = AlignedCompressedBlockSize;
 			}
-			Block->Size = AlignedCompressedBlockSize;
-		}
 
-		if (ContainerSettings.IsEncrypted())
-		{
-			FAES::EncryptData(Block->IoBuffer->Data(), static_cast<uint32>(Block->Size), ContainerSettings.EncryptionKey);
-		}
+			if (ContainerSettings.IsEncrypted())
+			{
+				FAES::EncryptData(Block.IoBuffer->Data(), static_cast<uint32>(Block.Size), ContainerSettings.EncryptionKey);
+			}
 
-		if (ContainerSettings.IsSigned())
-		{
-			FSHA1::HashBuffer(Block->IoBuffer->Data(), Block->Size, Block->Signature.Hash);
+			if (ContainerSettings.IsSigned())
+			{
+				FSHA1::HashBuffer(Block.IoBuffer->Data(), Block.Size, Block.Signature.Hash);
+			}
 		}
 	}
 
@@ -1379,10 +1510,69 @@ private:
 	int32						CurrentPartitionIndex = 0;
 	bool						bHasMemoryMappedEntry = false;
 	bool						bHasFlushed = false;
+	bool						bHasResult = false;
 
 	friend class FIoStoreWriterContextImpl;
 };
 
+TSharedPtr<IIoStoreWriter> FIoStoreWriterContextImpl::CreateContainer(const TCHAR* InContainerPath, const FIoContainerSettings& InContainerSettings)
+{
+	TSharedPtr<FIoStoreWriter> IoStoreWriter = MakeShared<FIoStoreWriter>(InContainerPath);
+	FIoStatus IoStatus = IoStoreWriter->Initialize(*this, InContainerSettings);
+	check(IoStatus.IsOk());
+	IoStoreWriters.Add(IoStoreWriter);
+	return IoStoreWriter;
+}
+
+void FIoStoreWriterContextImpl::Flush()
+{
+	TArray<FIoStoreWriteQueueEntry*> AllEntries;
+	for (TSharedPtr<FIoStoreWriter> IoStoreWriter : IoStoreWriters)
+	{
+		IoStoreWriter->bHasFlushed = true;
+		AllEntries.Append(IoStoreWriter->Entries);
+	}
+	{
+		TRACE_CPUPROFILER_EVENT_SCOPE(WaitForChunkHashes);
+		for (int32 EntryIndex = AllEntries.Num() - 1; EntryIndex >= 0; --EntryIndex)
+		{
+			FTaskGraphInterface::Get().WaitUntilTaskCompletes(AllEntries[EntryIndex]->HashTask);
+		}
+	}
+	for (TSharedPtr<FIoStoreWriter> IoStoreWriter : IoStoreWriters)
+	{
+		if (IoStoreWriter->LayoutEntriesHead)
+		{
+			IoStoreWriter->FinalizeLayout();
+		}
+	}
+	// Update list of all entries after having the finilized layouts of each container
+	AllEntries.Reset();
+	for (TSharedPtr<FIoStoreWriter> IoStoreWriter : IoStoreWriters)
+	{
+		AllEntries.Append(IoStoreWriter->Entries);
+	}
+
+	for (FIoStoreWriteQueueEntry* Entry : AllEntries)
+	{
+		ScheduleCompression(Entry);
+	}
+
+	{
+		TRACE_CPUPROFILER_EVENT_SCOPE(WaitForWritesToComplete);
+		for (int32 EntryIndex = AllEntries.Num() - 1; EntryIndex >= 0; --EntryIndex)
+		{
+			AllEntries[EntryIndex]->WriteFinishedEvent->Wait();
+		}
+	}
+
+	AllEntries.Empty();
+	for (TSharedPtr<FIoStoreWriter> IoStoreWriter : IoStoreWriters)
+	{
+		IoStoreWriter->Finalize();
+	}
+}
+
 void FIoStoreWriterContextImpl::BeginCompressionThreadFunc()
 {
 	for (;;)
@@ -1396,26 +1586,18 @@ void FIoStoreWriterContextImpl::BeginCompressionThreadFunc()
 		{
 			FIoStoreWriteQueueEntry* Next = Entry->Next;
 			Entry->BeginCompressionBarrier->Wait();
-
-			FinishCompressionQueue.Enqueue(Entry);
-
-			if (!Entry->HashBarrier->IsComplete())
-			{
-				Entry->HashBarrier->DispatchSubsequents();
-			}
-			
 			Entry->Writer->BeginCompress(Entry);
-
+			BeginEncryptionAndSigningQueue.Enqueue(Entry);
 			Entry = Next;
 		}
 	}
 }
 
-void FIoStoreWriterContextImpl::FinishCompressionThreadFunc()
+void FIoStoreWriterContextImpl::BeginEncryptionAndSigningThreadFunc()
 {
 	for (;;)
 	{
-		FIoStoreWriteQueueEntry* Entry = FinishCompressionQueue.DequeueOrWait();
+		FIoStoreWriteQueueEntry* Entry = BeginEncryptionAndSigningQueue.DequeueOrWait();
 		if (!Entry)
 		{
 			return;
@@ -1424,12 +1606,36 @@ void FIoStoreWriterContextImpl::FinishCompressionThreadFunc()
 		{
 			FIoStoreWriteQueueEntry* Next = Entry->Next;
 			Entry->FinishCompressionBarrier->Wait();
-			Entry->HashTask->Wait();
-
-			WriterQueue.Enqueue(Entry);
-
 			Entry->Request->FreeSourceBuffer();
+			if (Entry->bStoreCompressedDataInDDC)
+			{
+				TRACE_CPUPROFILER_EVENT_SCOPE(StoreInDDC);
+				TArray<uint8> DDCData;
+				FMemoryWriter DDCDataWriter(DDCData, true);
+				Entry->Writer->SerializeCompressedDDCData(Entry, DDCDataWriter);
+				DDC->Put(*Entry->DDCCacheKey, DDCData, Entry->Options.FileName);
+			}
+			FinishEncryptionAndSigningQueue.Enqueue(Entry);
+			Entry->Writer->BeginEncryptAndSign(Entry);
+			Entry = Next;
+		}
+	}
+}
 
+void FIoStoreWriterContextImpl::FinishEncryptionAndSigningThreadFunc()
+{
+	for (;;)
+	{
+		FIoStoreWriteQueueEntry* Entry = FinishEncryptionAndSigningQueue.DequeueOrWait();
+		if (!Entry)
+		{
+			return;
+		}
+		while (Entry)
+		{
+			FIoStoreWriteQueueEntry* Next = Entry->Next;
+			Entry->FinishEncryptionAndSigningBarrier->Wait();
+			WriterQueue.Enqueue(Entry);
 			Entry->CompressedSize = 0;
 			for (const FChunkBlock& ChunkBlock : Entry->ChunkBlocks)
 			{
@@ -1462,83 +1668,6 @@ void FIoStoreWriterContextImpl::WriterThreadFunc()
 	}
 }
 
-
-FIoStoreWriter::FIoStoreWriter(const TCHAR* ContainerPath)
-:	Impl(new FIoStoreWriterImpl(ContainerPath))
-{
-}
-
-FIoStoreWriter::~FIoStoreWriter()
-{
-	(void)Impl->Flush();
-	delete Impl;
-}
-
-FIoStatus FIoStoreWriter::Initialize(const FIoStoreWriterContext& Context, const FIoContainerSettings& ContainerSettings)
-{
-	return Impl->Initialize(*Context.Impl, ContainerSettings);
-}
-
-void FIoStoreWriter::EnableDiskLayoutOrdering(const TArray<TUniquePtr<FIoStoreReader>>& PatchSourceReaders)
-{
-	Impl->EnableDiskLayoutOrdering(PatchSourceReaders);
-}
-
-void FIoStoreWriter::Append(const FIoChunkId& ChunkId, FIoBuffer Chunk, const FIoWriteOptions& WriteOptions, uint64 OrderHint)
-{
-	struct FWriteRequest
-		: IIoStoreWriteRequest
-	{
-		FWriteRequest(FIoBuffer InSourceBuffer, uint64 InOrderHint)
-			: OrderHint(InOrderHint)
-		{
-			SourceBuffer = InSourceBuffer;
-			SourceBuffer.MakeOwned();
-		}
-
-		virtual ~FWriteRequest() = default;
-
-		void PrepareSourceBufferAsync(FGraphEventRef CompletionEvent) override
-		{
-			CompletionEvent->DispatchSubsequents();
-		}
-
-		const FIoBuffer* GetSourceBuffer() override
-		{
-			return &SourceBuffer;
-		}
-
-		void FreeSourceBuffer() override
-		{
-		}
-
-		uint64 GetOrderHint() override
-		{
-			return OrderHint;
-		}
-
-		TArrayView<const FFileRegion> GetRegions()
-		{
-			return TArrayView<const FFileRegion>();
-		}
-
-		FIoBuffer SourceBuffer;
-		uint64 OrderHint;
-	};
-
-	Append(ChunkId, new FWriteRequest(Chunk, OrderHint), WriteOptions);
-}
-
-void FIoStoreWriter::Append(const FIoChunkId& ChunkId, IIoStoreWriteRequest* Request, const FIoWriteOptions& WriteOptions)
-{
-	Impl->Append(ChunkId, Request, WriteOptions);
-}
-
-TIoStatusOr<FIoStoreWriterResult> FIoStoreWriter::Flush()
-{
-	return Impl->Flush();
-}
-
 class FIoStoreAsyncReadRequest::FImpl
 {
 public:
@@ -2434,4 +2563,6 @@ void FIoStoreWriterSettings::InitializePlatformSpecificSettings(const ITargetPla
 	const TCHAR* IniSection = TEXT("/Script/UnrealEd.ProjectPackagingSettings");
 	PlatformConfig->GetInt(IniSection, TEXT("PackageCompressionMinBytesSaved"), CompressionMinBytesSaved, GGameIni);
 	PlatformConfig->GetInt(IniSection, TEXT("PackageCompressionMinPercentSaved"), CompressionMinPercentSaved, GGameIni);
+	PlatformConfig->GetInt(IniSection, TEXT("PackageCompressionMinSizeToConsiderDDC"), CompressionMinSizeToConsiderDDC, GGameIni);
+	PlatformConfig->GetBool(IniSection, TEXT("bPackageCompressionEnableDDC"), bCompressionEnableDDC, GGameIni);
 }
diff --git a/Engine/Source/Runtime/Core/Public/IO/IoDispatcher.h b/Engine/Source/Runtime/Core/Public/IO/IoDispatcher.h
--- a/Engine/Source/Runtime/Core/Public/IO/IoDispatcher.h
+++ b/Engine/Source/Runtime/Core/Public/IO/IoDispatcher.h
@@ -23,14 +23,12 @@
 class FIoRequest;
 class FIoDispatcher;
 class FIoStoreReader;
-class FIoStoreWriter;
 class FIoStoreEnvironment;
 
 class FIoRequestImpl;
 class FIoBatchImpl;
 class FIoDispatcherImpl;
 class FIoStoreWriterContextImpl;
-class FIoStoreWriterImpl;
 class FIoStoreReaderImpl;
 class IMappedFileHandle;
 class IMappedFileRegion;
@@ -1066,10 +1064,12 @@ struct FIoStoreWriterSettings
 	uint64 CompressionBlockAlignment = 0;
 	int32 CompressionMinBytesSaved = 0;
 	int32 CompressionMinPercentSaved = 0;
+	int32 CompressionMinSizeToConsiderDDC = 0;
 	uint64 MemoryMappingAlignment = 0;
 	uint64 MaxPartitionSize = 0;
 	bool bEnableCsvOutput = false;
 	bool bEnableFileRegions = false;
+	bool bCompressionEnableDDC = false;
 
 	CORE_API void InitializePlatformSpecificSettings(const ITargetPlatform* TargetPlatform);
 };
@@ -1149,17 +1149,20 @@ public:
 		uint64 HashedChunksCount = 0;
 		uint64 CompressedChunksCount = 0;
 		uint64 SerializedChunksCount = 0;
+		uint64 ScheduledCompressionTasksCount = 0;
+		uint64 CompressionDDCHitCount = 0;
+		uint64 CompressionDDCMissCount = 0;
 	};
 
 	CORE_API FIoStoreWriterContext();
 	CORE_API ~FIoStoreWriterContext();
 
 	UE_NODISCARD CORE_API FIoStatus Initialize(const FIoStoreWriterSettings& InWriterSettings);
+	CORE_API TSharedPtr<class IIoStoreWriter> CreateContainer(const TCHAR* InContainerPath, const FIoContainerSettings& InContainerSettings);
+	CORE_API void Flush();
 	CORE_API FProgress GetProgress() const;
 
 private:
-	friend class FIoStoreWriter;
-	
 	FIoStoreWriterContextImpl* Impl;
 };
 
@@ -1174,23 +1177,15 @@ public:
 	virtual void FreeSourceBuffer() = 0;
 };
 
-class FIoStoreWriter
+class IIoStoreWriter
 {
 public:
-	CORE_API 			FIoStoreWriter(const TCHAR* ContainerPath);
-	CORE_API virtual	~FIoStoreWriter();
-
-	FIoStoreWriter(const FIoStoreWriter&) = delete;
-	FIoStoreWriter& operator=(const FIoStoreWriter&) = delete;
+	CORE_API virtual ~IIoStoreWriter() = default;
 
-	UE_NODISCARD CORE_API FIoStatus	Initialize(const FIoStoreWriterContext& Context, const FIoContainerSettings& ContainerSettings);
-	CORE_API void EnableDiskLayoutOrdering(const TArray<TUniquePtr<FIoStoreReader>>& PatchSourceReaders = TArray<TUniquePtr<FIoStoreReader>>());
-	CORE_API void Append(const FIoChunkId& ChunkId, FIoBuffer Chunk, const FIoWriteOptions& WriteOptions, uint64 OrderHint = MAX_uint64);
-	CORE_API void Append(const FIoChunkId& ChunkId, IIoStoreWriteRequest* Request, const FIoWriteOptions& WriteOptions);
-	UE_NODISCARD CORE_API TIoStatusOr<FIoStoreWriterResult> Flush();
-
-private:
-	FIoStoreWriterImpl*		Impl;
+	CORE_API virtual void EnableDiskLayoutOrdering(const TArray<TUniquePtr<FIoStoreReader>>& PatchSourceReaders = TArray<TUniquePtr<FIoStoreReader>>()) = 0;
+	CORE_API virtual void Append(const FIoChunkId& ChunkId, FIoBuffer Chunk, const FIoWriteOptions& WriteOptions, uint64 OrderHint = MAX_uint64) = 0;
+	CORE_API virtual void Append(const FIoChunkId& ChunkId, IIoStoreWriteRequest* Request, const FIoWriteOptions& WriteOptions) = 0;
+	CORE_API virtual TIoStatusOr<FIoStoreWriterResult> GetResult() = 0;
 };
 
 struct FIoStoreTocChunkInfo
