PR #8376: Fix whitespace issue when updating config sections (Contributed by [REDACTED_PERSON])

diff --git a/Engine/Source/Runtime/Core/Private/Misc/ConfigCacheIni.cpp b/Engine/Source/Runtime/Core/Private/Misc/ConfigCacheIni.cpp
--- a/Engine/Source/Runtime/Core/Private/Misc/ConfigCacheIni.cpp
+++ b/Engine/Source/Runtime/Core/Private/Misc/ConfigCacheIni.cpp
@@ -4688,6 +4688,9 @@ void FConfigFile::UpdateSections(const TCHAR* DiskFilename, const TCHAR* IniRoot
 			}
 			else
 			{
+				// Strip any trailing whitespace to match config parsing.
+				TheLine.TrimEndInline();
+
 				// is this line a section? (must be at least [x])
 				if (TheLine.Len() > 3 && TheLine[0] == '[' && TheLine[TheLine.Len() - 1] == ']')
 				{
