Change Compact Binary Package API to represent binary attachments as compressed buffers identified by their raw hash.
Change Compact Binary Package serialization for binary attachments, compressed binary attachments and objects.

diff --git a/Engine/Source/Developer/DerivedDataCache/Internal/ZenSerialization.h b/Engine/Source/Developer/DerivedDataCache/Internal/ZenSerialization.h
new file mode 100644
--- /dev/null
+++ b/Engine/Source/Developer/DerivedDataCache/Internal/ZenSerialization.h
@@ -0,0 +1,26 @@
+// Copyright Epic Games, Inc. All Rights Reserved.
+
+#pragma once
+
+#include "ZenServerInterface.h"
+
+#if UE_WITH_ZEN
+
+#include "Memory/SharedBuffer.h"
+#include "Serialization/CompactBinary.h"
+
+class FArchive;
+class FCbAttachment;
+class FCbPackage;
+class FCbWriter;
+
+namespace UE::Zen {
+
+DERIVEDDATACACHE_API void SaveCbAttachment(const FCbAttachment& Attachment, FCbWriter& Writer);
+DERIVEDDATACACHE_API void SaveCbPackage(const FCbPackage& Package, FCbWriter& Writer);
+DERIVEDDATACACHE_API void SaveCbPackage(const FCbPackage& Package, FArchive& Ar);
+DERIVEDDATACACHE_API bool TryLoadCbPackage(FCbPackage& Package, FArchive& Ar, FCbBufferAllocator Allocator = FUniqueBuffer::Alloc);
+
+} // namespace UE::Zen
+
+#endif // UE_WITH_ZEN
\ No newline at end of file
diff --git a/Engine/Source/Developer/DerivedDataCache/Private/ZenDerivedDataBackend.cpp b/Engine/Source/Developer/DerivedDataCache/Private/ZenDerivedDataBackend.cpp
--- a/Engine/Source/Developer/DerivedDataCache/Private/ZenDerivedDataBackend.cpp
+++ b/Engine/Source/Developer/DerivedDataCache/Private/ZenDerivedDataBackend.cpp
@@ -15,6 +15,7 @@
 #include "Serialization/CompactBinaryValidation.h"
 #include "Serialization/CompactBinaryWriter.h"
 #include "ZenBackendUtils.h"
+#include "ZenSerialization.h"
 #include "ZenServerHttp.h"
 
 TRACE_DECLARE_INT_COUNTER(ZenDDC_Exist,			TEXT("ZenDDC Exist"));
@@ -722,7 +723,7 @@ bool FZenDerivedDataBackend::PutCacheRecord(const FCacheRecord& Record, FStringV
 	const FCacheKey& Key = Record.GetKey();
 	FCbPackage Package = Record.Save();
 	FCbWriter PackageWriter;
-	Package.Save(PackageWriter);
+	UE::Zen::SaveCbPackage(Package, PackageWriter);
 	FUniqueBuffer Buffer = FUniqueBuffer::Alloc(PackageWriter.GetSaveSize());
 	PackageWriter.Save(Buffer);
 	TStringBuilder<256> Uri;
diff --git a/Engine/Source/Developer/DerivedDataCache/Private/ZenSerialization.cpp b/Engine/Source/Developer/DerivedDataCache/Private/ZenSerialization.cpp
new file mode 100644
--- /dev/null
+++ b/Engine/Source/Developer/DerivedDataCache/Private/ZenSerialization.cpp
@@ -0,0 +1,131 @@
+// Copyright Epic Games, Inc. All Rights Reserved.
+
+#include "ZenSerialization.h"
+#include "Serialization/CompactBinary.h"
+#include "Serialization/CompactBinaryPackage.h"
+#include "Serialization/CompactBinarySerialization.h"
+#include "Serialization/CompactBinaryValidation.h"
+#include "Serialization/CompactBinaryWriter.h"
+
+#if UE_WITH_ZEN
+
+namespace UE::Zen {
+
+void SaveCbAttachment(const FCbAttachment& Attachment, FCbWriter& Writer)
+{
+	if (Attachment.IsObject())
+	{
+		FCbObject Object = Attachment.AsObject();
+		Writer.AddBinary(Object.GetBuffer().Flatten());
+		if (Object)
+		{
+			Writer.AddObjectAttachment(Attachment.GetHash());
+		}
+	}
+	else if (Attachment.IsBinary())
+	{
+		Writer.AddBinary(Attachment.AsBinary());
+		Writer.AddBinaryAttachment(Attachment.GetHash());
+	}
+	else if (Attachment.IsNull())
+	{
+		Writer.AddBinary(FMemoryView());
+	}
+	else
+	{
+		// NOTE: Compressed binary is not supported in this serialization format
+		checkNoEntry();
+	}
+}
+
+void SaveCbPackage(const FCbPackage& Package, FCbWriter& Writer)
+{
+	if (const FCbObject& RootObject = Package.GetObject())
+	{
+		Writer.AddObject(RootObject);
+		Writer.AddObjectAttachment(Package.GetObjectHash());
+	}
+	for (const FCbAttachment& Attachment : Package.GetAttachments())
+	{
+		SaveCbAttachment(Attachment, Writer);
+	}
+	Writer.AddNull();
+}
+
+void SaveCbPackage(const FCbPackage& Package, FArchive& Ar)
+{
+	FCbWriter Writer;
+	SaveCbPackage(Package, Writer);
+	Writer.Save(Ar);
+}
+
+bool TryLoadCbPackage(FCbPackage& Package, FArchive& Ar, FCbBufferAllocator Allocator)
+{
+	uint8 StackBuffer[64];
+	const auto StackAllocator = [&Allocator, &StackBuffer](uint64 Size) -> FUniqueBuffer
+	{
+		return Size <= sizeof(StackBuffer) ? FUniqueBuffer::MakeView(StackBuffer, Size) : Allocator(Size);
+	};
+
+	Package = FCbPackage();
+	for (;;)
+	{
+		FCbField ValueField = LoadCompactBinary(Ar, StackAllocator);
+		if (!ValueField)
+		{
+			Ar.SetError();
+			return false;
+		}
+		if (ValueField.IsNull())
+		{
+			return true;
+		}
+		else if (ValueField.IsBinary())
+		{
+			const FMemoryView View = ValueField.AsBinaryView();
+			if (View.GetSize() > 0)
+			{
+				FSharedBuffer Buffer = FSharedBuffer::MakeView(View, ValueField.GetOuterBuffer()).MakeOwned();
+				FCbField HashField = LoadCompactBinary(Ar, StackAllocator);
+				const FIoHash& Hash = HashField.AsAttachment();
+				if (HashField.HasError() || FIoHash::HashBuffer(Buffer) != Hash)
+				{
+					Ar.SetError();
+					return false;
+				}
+				if (HashField.IsObjectAttachment())
+				{
+					Package.AddAttachment(FCbAttachment(FCbObject(MoveTemp(Buffer)), Hash));
+				}
+				else
+				{
+					Package.AddAttachment(FCbAttachment(FCompositeBuffer(MoveTemp(Buffer)), Hash));
+				}
+			}
+		}
+		else
+		{
+			FCbObject Object = ValueField.AsObject();
+			if (ValueField.HasError())
+			{
+				Ar.SetError();
+				return false;
+			}
+
+			if (Object)
+			{
+				FCbField HashField = LoadCompactBinary(Ar, StackAllocator);
+				FIoHash ObjectHash = HashField.AsObjectAttachment();
+				if (HashField.HasError() || Object.GetHash() != ObjectHash)
+				{
+					Ar.SetError();
+					return false;
+				}
+				Package.SetObject(Object, ObjectHash);
+			}
+		}
+	}
+}
+
+}
+#endif // UE_WITH_ZEN
diff --git a/Engine/Source/Developer/DerivedDataCache/Private/ZenServerHttp.cpp b/Engine/Source/Developer/DerivedDataCache/Private/ZenServerHttp.cpp
--- a/Engine/Source/Developer/DerivedDataCache/Private/ZenServerHttp.cpp
+++ b/Engine/Source/Developer/DerivedDataCache/Private/ZenServerHttp.cpp
@@ -6,6 +6,7 @@
 
 #include "DerivedDataBackendInterface.h"	// For logger definition
 #include "ZenBackendUtils.h"
+#include "ZenSerialization.h"
 
 #if PLATFORM_WINDOWS || PLATFORM_HOLOLENS
 #	include "Windows/WindowsHWrapper.h"
@@ -176,7 +177,7 @@ namespace UE::Zen {
 			else
 			{
 				FLargeMemoryReader PackageLoader(PackageBytes.GetData(), PackageBytes.Num());
-				bResponseFormatValid = OutPackage.TryLoad(PackageLoader);
+				bResponseFormatValid = UE::Zen::TryLoadCbPackage(OutPackage, PackageLoader);
 			}
 		}
 		return LocalResult;
diff --git a/Engine/Source/Developer/IoStoreUtilities/Private/ZenStoreHttpClient.cpp b/Engine/Source/Developer/IoStoreUtilities/Private/ZenStoreHttpClient.cpp
--- a/Engine/Source/Developer/IoStoreUtilities/Private/ZenStoreHttpClient.cpp
+++ b/Engine/Source/Developer/IoStoreUtilities/Private/ZenStoreHttpClient.cpp
@@ -22,6 +22,7 @@
 #include "Serialization/BufferArchive.h"
 #include "Serialization/CompactBinaryWriter.h"
 #include "Serialization/LargeMemoryWriter.h"
+#include "ZenSerialization.h"
 #include "ZenServerHttp.h"
 
 DEFINE_LOG_CATEGORY_STATIC(LogZenStore, Log, All);
@@ -338,7 +339,7 @@ TFuture<TIoStatusOr<uint64>> FZenStoreHttpClient::AppendOp(FCbPackage OpEntry)
 
 					if (!IsSerialized)
 					{
-						Attachment.Save(Writer);
+						UE::Zen::SaveCbAttachment(Attachment, Writer);
 					}
 				}
 				else
@@ -516,7 +517,7 @@ FZenStoreHttpClient::EndBuildPass(FCbPackage OpEntry)
 	check(bAllowEdit);
 
 	FLargeMemoryWriter SerializedPackage;
-	OpEntry.Save(SerializedPackage);
+	UE::Zen::SaveCbPackage(OpEntry, SerializedPackage);
 
 	UE_LOG(LogZenStore, Verbose, TEXT("Package size: %lld"), SerializedPackage.TotalSize());
 
diff --git a/Engine/Source/Runtime/Core/Private/Serialization/CompactBinaryPackage.cpp b/Engine/Source/Runtime/Core/Private/Serialization/CompactBinaryPackage.cpp
--- a/Engine/Source/Runtime/Core/Private/Serialization/CompactBinaryPackage.cpp
+++ b/Engine/Source/Runtime/Core/Private/Serialization/CompactBinaryPackage.cpp
@@ -3,141 +3,201 @@
 #include "Serialization/CompactBinaryPackage.h"
 
 #include "Algo/BinarySearch.h"
+#include "Compression/OodleDataCompression.h"
 #include "Memory/CompositeBuffer.h"
 #include "Serialization/CompactBinarySerialization.h"
 #include "Serialization/CompactBinaryWriter.h"
 
 ///////////////////////////////////////////////////////////////////////////////////////////////////////////////////////
 
-FCbAttachment::FCbAttachment(FCbObject InValue, const FIoHash* const InHash)
+static bool TryLoad_ArchiveFieldIntoAttachment(FCbAttachment& TargetAttachment, FCbField&& Field, FArchive& Ar, FCbBufferAllocator Allocator)
 {
-	FMemoryView View;
-	if (!InValue.IsOwned() || !InValue.TryGetView(View))
+	if (const FCbObjectView ObjectView = Field.AsObjectView(); !Field.HasError())
+	{
+		// Is a null object or object not prefixed with a precomputed hash value
+		TargetAttachment = FCbAttachment(FCbObject(ObjectView, MoveTemp(Field)), ObjectView.GetHash());
+	}
+	else if (const FIoHash ObjectAttachmentHash = Field.AsObjectAttachment(); !Field.HasError())
+	{
+		// Is an object
+		Field = LoadCompactBinary(Ar, Allocator);
+		if (!Field.IsObject())
+		{
+			return false;
+		}
+		TargetAttachment = FCbAttachment(MoveTemp(Field).AsObject(), ObjectAttachmentHash);
+	}
+	else if (const FIoHash BinaryAttachmentHash = Field.AsBinaryAttachment(); !Field.HasError())
+	{
+		// Is an uncompressed binary blob
+		Field = LoadCompactBinary(Ar, Allocator);
+		FSharedBuffer Buffer = Field.AsBinary();
+		if (Field.HasError())
+		{
+			return false;
+		}
+		TargetAttachment = FCbAttachment(FCompositeBuffer(Buffer), BinaryAttachmentHash);
+	}
+	else if (FSharedBuffer Buffer = Field.AsBinary(); !Field.HasError())
+	{
+		if (Buffer.GetSize() > 0)
+		{
+			// Is a compressed binary blob
+			TargetAttachment = FCbAttachment(FCompressedBuffer::FromCompressed(MoveTemp(Buffer)));
+		}
+		else
+		{
+			// Is an uncompressed empty binary blob
+			TargetAttachment = FCbAttachment(FCompositeBuffer(Buffer), FIoHash::HashBuffer(nullptr, 0));
+		}
+	}
+	else
 	{
-		InValue = FCbObject::Clone(InValue);
+		return false;
 	}
 
-	Object = InValue.AsFieldView();
-	Buffer = MoveTemp(InValue).GetBuffer().Flatten();
+	return true;
+}
 
-	if (InHash)
+FCbAttachment::FCbAttachment(const FCbObject& InValue, const FIoHash* const InHash)
+{
+	FMemoryView View;
+	if (!InValue.IsOwned() || !InValue.TryGetView(View))
 	{
-		Hash = *InHash;
-		checkSlow(Hash == FIoHash::HashBuffer(Buffer));
+		FCbObject ClonedValue = FCbObject::Clone(InValue);
+
+		if (InHash)
+		{
+			checkSlow(*InHash == ClonedValue.GetHash());
+			Value.Emplace<FObjectValue>(MoveTemp(ClonedValue), *InHash);
+		}
+		else
+		{
+			FIoHash ClonedValueHash = ClonedValue.GetHash();
+			Value.Emplace<FObjectValue>(MoveTemp(ClonedValue), ClonedValueHash);
+		}
 	}
 	else
 	{
-		Hash = FIoHash::HashBuffer(Buffer);
+		if (InHash)
+		{
+			Value.Emplace<FObjectValue>(InValue, *InHash);
+			checkSlow(*InHash == InValue.GetHash());
+		}
+		else
+		{
+			Value.Emplace<FObjectValue>(InValue, InValue.GetHash());
+		}
 	}
+
 }
 
-FCbAttachment::FCbAttachment(FSharedBuffer InBuffer, const FIoHash* const InHash)
-	: Buffer(MoveTemp(InBuffer).MakeOwned())
+FIoHash FCbAttachment::GetHash() const
 {
-	if (InHash)
+	if (const FCompressedBuffer* Buffer = Value.TryGet<FCompressedBuffer>())
+	{
+		return Buffer->GetRawHash();
+	}
+	else if (const FBinaryValue* BinaryValue = Value.TryGet<FBinaryValue>())
 	{
-		Hash = *InHash;
-		checkSlow(Hash == FIoHash::HashBuffer(Buffer));
+		return BinaryValue->Hash;
+	}
+	else if (const FObjectValue* ObjectValue = Value.TryGet<FObjectValue>())
+	{
+		return ObjectValue->Hash;
 	}
 	else
 	{
-		Hash = FIoHash::HashBuffer(Buffer);
+		return FIoHash::Zero;
 	}
 }
 
 bool FCbAttachment::TryLoad(FCbFieldIterator& Fields)
 {
-	const FMemoryView View = Fields.AsBinaryView();
-	if (Fields.HasError())
+	if (const FCbObjectView ObjectView = Fields.AsObjectView(); !Fields.HasError())
 	{
-		return false;
+		// Is a null object or object not prefixed with a precomputed hash value
+		Value.Emplace<FObjectValue>(FCbObject(ObjectView, Fields.GetOuterBuffer()), ObjectView.GetHash());
+		++Fields;
 	}
-	if (View.GetSize() > 0)
+	else if (const FIoHash ObjectAttachmentHash = Fields.AsObjectAttachment(); !Fields.HasError())
 	{
-		Buffer = FSharedBuffer::MakeView(View, Fields.GetOuterBuffer()).MakeOwned();
-		Object = FCbFieldView();
+		// Is an object
 		++Fields;
-
-		Hash = Fields.AsAttachment();
+		const FCbObjectView InnerObjectView = Fields.AsObjectView();
 		if (Fields.HasError())
 		{
 			return false;
 		}
-		if (Fields.IsObjectAttachment())
-		{
-			Object = FCbFieldView(Buffer.GetData());
-		}
+		Value.Emplace<FObjectValue>(FCbObject(InnerObjectView, Fields.GetOuterBuffer()), ObjectAttachmentHash);
 		++Fields;
 	}
-	else
+	else if (const FIoHash BinaryAttachmentHash = Fields.AsBinaryAttachment(); !Fields.HasError())
 	{
+		// Is an uncompressed binary blob
+		++Fields;
+		FMemoryView BinaryView = Fields.AsBinaryView();
+		if (Fields.HasError())
+		{
+			return false;
+		}
+		Value.Emplace<FBinaryValue>(FSharedBuffer::MakeView(BinaryView, Fields.GetOuterBuffer()), BinaryAttachmentHash);
 		++Fields;
-		Buffer.Reset();
-		Object = FCbFieldView();
-		Hash.Reset();
-	}
-	return true;
-}
-
-bool FCbAttachment::TryLoad(FArchive& Ar, FCbBufferAllocator Allocator)
-{
-	FCbField BufferField = LoadCompactBinary(Ar, Allocator);
-	const FMemoryView View = BufferField.AsBinaryView();
-	if (BufferField.HasError())
-	{
-		Ar.SetError();
-		return false;
 	}
-	if (View.GetSize() > 0)
+	else if (FMemoryView BinaryView = Fields.AsBinaryView(); !Fields.HasError())
 	{
-		Buffer = FSharedBuffer::MakeView(View, BufferField.GetOuterBuffer()).MakeOwned();
-		Object = FCbFieldView();
-
-		TArray<uint8, TInlineAllocator<64>> HashBuffer;
-		FCbField HashField = LoadCompactBinary(Ar,
-			[&HashBuffer](uint64 Size) -> FUniqueBuffer
-			{
-				HashBuffer.SetNumUninitialized(int32(Size));
-				return FUniqueBuffer::MakeView(HashBuffer.GetData(), Size);
-			});
-		Hash = HashField.AsAttachment();
-		if (HashField.HasError() || FIoHash::HashBuffer(Buffer) != Hash)
+		if (BinaryView.GetSize() > 0)
 		{
-			Ar.SetError();
-			return false;
+			// Is a compressed binary blob
+			Value.Emplace<FCompressedBuffer>(FCompressedBuffer::FromCompressed(FSharedBuffer::MakeView(BinaryView, Fields.GetOuterBuffer())).MakeOwned());
+			++Fields;
 		}
-		if (HashField.IsObjectAttachment())
+		else
 		{
-			Object = FCbFieldView(Buffer.GetData());
+			// Is an uncompressed empty binary blob
+			Value.Emplace<FBinaryValue>(FSharedBuffer::MakeView(BinaryView, Fields.GetOuterBuffer()), FIoHash::HashBuffer(nullptr, 0));
+			++Fields;
 		}
 	}
 	else
 	{
-		Buffer.Reset();
-		Object = FCbFieldView();
-		Hash.Reset();
+		return false;
 	}
+
 	return true;
 }
 
+bool FCbAttachment::TryLoad(FArchive& Ar, FCbBufferAllocator Allocator)
+{
+	FCbField Field = LoadCompactBinary(Ar, Allocator);
+	return TryLoad_ArchiveFieldIntoAttachment(*this, MoveTemp(Field), Ar, Allocator);
+}
+
 void FCbAttachment::Save(FCbWriter& Writer) const
 {
-	if (Object.IsObject())
+	if (const FObjectValue* ObjectValue = Value.TryGet<FObjectValue>())
+	{
+		if (ObjectValue->Object)
+		{
+			Writer.AddObjectAttachment(ObjectValue->Hash);
+		}
+		Writer.AddObject(ObjectValue->Object);
+	}
+	else if (const FBinaryValue* BinaryValue = Value.TryGet<FBinaryValue>())
 	{
-		Writer.AddBinary(Buffer);
-		if (FCbFieldView(Object).AsObjectView())
+		if (BinaryValue->Buffer.GetSize() > 0)
 		{
-			Writer.AddObjectAttachment(Hash);
+			Writer.AddBinaryAttachment(BinaryValue->Hash);
 		}
+		Writer.AddBinary(BinaryValue->Buffer);
 	}
-	else if (Buffer.GetSize())
+	else if (const FCompressedBuffer* BufferValue = Value.TryGet<FCompressedBuffer>())
 	{
-		Writer.AddBinary(Buffer);
-		Writer.AddBinaryAttachment(Hash);
+		Writer.AddBinary(BufferValue->GetCompressed());
 	}
-	else // Null
+	else
 	{
-		Writer.AddBinary(FMemoryView());
+		checkf(false, TEXT("Null attachments cannot be serialized."))
 	}
 }
 
@@ -184,10 +244,7 @@ void FCbPackage::AddAttachment(const FCbAttachment& Attachment, FAttachmentResol
 		if (Attachments.IsValidIndex(Index) && Attachments[Index] == Attachment)
 		{
 			FCbAttachment& Existing = Attachments[Index];
-			if (Attachment.IsObject() && !Existing.IsObject())
-			{
-				Existing = FCbAttachment(FCbObject(Existing.AsBinary()), Existing.GetHash());
-			}
+			Existing = Attachment;
 		}
 		else
 		{
@@ -203,7 +260,7 @@ void FCbPackage::AddAttachment(const FCbAttachment& Attachment, FAttachmentResol
 int32 FCbPackage::RemoveAttachment(const FIoHash& Hash)
 {
 	const int32 Index = Algo::BinarySearchBy(Attachments, Hash,
-		[](const FCbAttachment& Attachment) -> const FIoHash& { return Attachment.GetHash(); });
+		[](const FCbAttachment& Attachment) -> FIoHash { return Attachment.GetHash(); });
 	if (Attachments.IsValidIndex(Index))
 	{
 		Attachments.RemoveAt(Index);
@@ -220,7 +277,7 @@ bool FCbPackage::Equals(const FCbPackage& Package) const
 const FCbAttachment* FCbPackage::FindAttachment(const FIoHash& Hash) const
 {
 	const int32 Index = Algo::BinarySearchBy(Attachments, Hash,
-		[](const FCbAttachment& Attachment) -> const FIoHash& { return Attachment.GetHash(); });
+		[](const FCbAttachment& Attachment) -> FIoHash { return Attachment.GetHash(); });
 	return Attachments.IsValidIndex(Index) ? &Attachments[Index] : nullptr;
 }
 
@@ -237,7 +294,7 @@ void FCbPackage::GatherAttachments(const FCbObject& Value, FAttachmentResolver R
 				}
 				else
 				{
-					AddAttachment(FCbAttachment(MoveTemp(Buffer), Hash));
+					AddAttachment(FCbAttachment(MoveTemp(Buffer)));
 				}
 			}
 		});
@@ -246,6 +303,7 @@ void FCbPackage::GatherAttachments(const FCbObject& Value, FAttachmentResolver R
 bool FCbPackage::TryLoad(FCbFieldIterator& Fields)
 {
 	*this = FCbPackage();
+
 	while (Fields)
 	{
 		if (Fields.IsNull())
@@ -253,34 +311,27 @@ bool FCbPackage::TryLoad(FCbFieldIterator& Fields)
 			++Fields;
 			break;
 		}
-		else if (Fields.IsBinary())
+		else if (FIoHash Hash = Fields.AsHash(); !Fields.HasError() && !Fields.IsAttachment())
 		{
-			FCbAttachment Attachment;
-			Attachment.TryLoad(Fields);
-			AddAttachment(Attachment);
-		}
-		else
-		{
-			Object = Fields.AsObject();
-			if (Fields.HasError())
+			++Fields;
+			FCbObjectView ObjectView = Fields.AsObjectView();
+			if (Fields.HasError() || Hash != ObjectView.GetHash())
 			{
 				return false;
 			}
+			Object = FCbObject(ObjectView, Fields.GetOuterBuffer());
 			Object.MakeOwned();
+			ObjectHash = Hash;
 			++Fields;
-			if (Object)
-			{
-				ObjectHash = Fields.AsObjectAttachment();
-				if (Fields.HasError())
-				{
-					return false;
-				}
-				++Fields;
-			}
-			else
+		}
+		else
+		{
+			FCbAttachment Attachment;
+			if (!Attachment.TryLoad(Fields))
 			{
-				Object.Reset();
+				return false;
 			}
+			AddAttachment(Attachment);
 		}
 	}
 	return true;
@@ -288,71 +339,39 @@ bool FCbPackage::TryLoad(FCbFieldIterator& Fields)
 
 bool FCbPackage::TryLoad(FArchive& Ar, FCbBufferAllocator Allocator)
 {
-	uint8 StackBuffer[64];
-	const auto StackAllocator = [&Allocator, &StackBuffer](uint64 Size) -> FUniqueBuffer
-	{
-		return Size <= sizeof(StackBuffer) ? FUniqueBuffer::MakeView(StackBuffer, Size) : Allocator(Size);
-	};
-
 	*this = FCbPackage();
 	for (;;)
 	{
-		FCbField ValueField = LoadCompactBinary(Ar, StackAllocator);
-		if (!ValueField)
+		FCbField Field = LoadCompactBinary(Ar, Allocator);
+		if (!Field)
 		{
 			Ar.SetError();
 			return false;
 		}
-		if (ValueField.IsNull())
+
+		if (Field.IsNull())
 		{
 			return true;
 		}
-		else if (ValueField.IsBinary())
+		else if (FIoHash Hash = Field.AsHash(); !Field.HasError() && !Field.IsAttachment())
 		{
-			const FMemoryView View = ValueField.AsBinaryView();
-			if (View.GetSize() > 0)
+			Field = LoadCompactBinary(Ar, Allocator);
+			FCbObjectView ObjectView = Field.AsObjectView();
+			if (Field.HasError() || Hash != ObjectView.GetHash())
 			{
-				FSharedBuffer Buffer = FSharedBuffer::MakeView(View, ValueField.GetOuterBuffer()).MakeOwned();
-				FCbField HashField = LoadCompactBinary(Ar, StackAllocator);
-				const FIoHash& Hash = HashField.AsAttachment();
-				if (HashField.HasError() || FIoHash::HashBuffer(Buffer) != Hash)
-				{
-					Ar.SetError();
-					return false;
-				}
-				if (HashField.IsObjectAttachment())
-				{
-					AddAttachment(FCbAttachment(FCbObject(MoveTemp(Buffer)), Hash));
-				}
-				else
-				{
-					AddAttachment(FCbAttachment(MoveTemp(Buffer), Hash));
-				}
+				return false;
 			}
+			Object = FCbObject(ObjectView, Field.GetOuterBuffer());
+			ObjectHash = Hash;
 		}
 		else
 		{
-			Object = ValueField.AsObject();
-			if (ValueField.HasError())
+			FCbAttachment Attachment;
+			if (!TryLoad_ArchiveFieldIntoAttachment(Attachment, MoveTemp(Field), Ar, Allocator))
 			{
-				Ar.SetError();
 				return false;
 			}
-			Object.MakeOwned();
-			if (Object)
-			{
-				FCbField HashField = LoadCompactBinary(Ar, StackAllocator);
-				ObjectHash = HashField.AsObjectAttachment();
-				if (HashField.HasError() || Object.GetHash() != ObjectHash)
-				{
-					Ar.SetError();
-					return false;
-				}
-			}
-			else
-			{
-				Object.Reset();
-			}
+			AddAttachment(Attachment);
 		}
 	}
 }
@@ -361,8 +380,8 @@ void FCbPackage::Save(FCbWriter& Writer) const
 {
 	if (Object)
 	{
+		Writer.AddHash(ObjectHash);
 		Writer.AddObject(Object);
-		Writer.AddObjectAttachment(ObjectHash);
 	}
 	for (const FCbAttachment& Attachment : Attachments)
 	{
diff --git a/Engine/Source/Runtime/Core/Private/Serialization/CompactBinaryValidation.cpp b/Engine/Source/Runtime/Core/Private/Serialization/CompactBinaryValidation.cpp
--- a/Engine/Source/Runtime/Core/Private/Serialization/CompactBinaryValidation.cpp
+++ b/Engine/Source/Runtime/Core/Private/Serialization/CompactBinaryValidation.cpp
@@ -2,6 +2,7 @@
 
 #include "Serialization/CompactBinaryValidation.h"
 
+#include "Compression/CompressedBuffer.h"
 #include "IO/IoHash.h"
 #include "Misc/ByteSwap.h"
 #include "Serialization/CompactBinary.h"
@@ -409,66 +410,122 @@ static FCbFieldView ValidateCbPackageField(FMemoryView& View, ECbValidateMode Mo
 
 static FIoHash ValidateCbPackageAttachment(FCbFieldView& Value, FMemoryView& View, ECbValidateMode Mode, ECbValidateError& Error)
 {
-	const FMemoryView ValueView = Value.AsBinaryView();
-	if (Value.HasError() && EnumHasAnyFlags(Mode, ECbValidateMode::Package))
+	if (const FCbObjectView ObjectView = Value.AsObjectView(); !Value.HasError())
 	{
-		if (EnumHasAnyFlags(Mode, ECbValidateMode::Package))
+		return FCbObject().GetHash();
+	}
+	else if (const FIoHash ObjectAttachmentHash = Value.AsObjectAttachment(); !Value.HasError())
+	{
+		if (FCbFieldView ObjectField = ValidateCbPackageField(View, Mode, Error))
 		{
-			AddError(Error, ECbValidateError::InvalidPackageFormat);
+			const FCbObjectView InnerObjectView = ObjectField.AsObjectView();
+			if (EnumHasAnyFlags(Mode, ECbValidateMode::Package) && ObjectField.HasError())
+			{
+				AddError(Error, ECbValidateError::InvalidPackageFormat);
+			}
+			else if (EnumHasAnyFlags(Mode, ECbValidateMode::PackageHash) && (ObjectAttachmentHash != InnerObjectView.GetHash()))
+			{
+				AddError(Error, ECbValidateError::InvalidPackageHash);
+			}
+			return ObjectAttachmentHash;
 		}
 	}
-	else if (ValueView.GetSize())
+	else if (const FIoHash BinaryAttachmentHash = Value.AsBinaryAttachment(); !Value.HasError())
 	{
-		if (FCbFieldView HashField = ValidateCbPackageField(View, Mode, Error))
+		if (FCbFieldView BinaryField = ValidateCbPackageField(View, Mode, Error))
 		{
-			const FIoHash Hash = HashField.AsAttachment();
-			if (EnumHasAnyFlags(Mode, ECbValidateMode::Package))
+			const FMemoryView BinaryView = BinaryField.AsBinaryView();
+			if (EnumHasAnyFlags(Mode, ECbValidateMode::Package) && BinaryField.HasError())
 			{
-				if (HashField.HasError())
+				AddError(Error, ECbValidateError::InvalidPackageFormat);
+			}
+			else
+			{
+				if (EnumHasAnyFlags(Mode, ECbValidateMode::Package) && BinaryView.IsEmpty())
 				{
-					AddError(Error, ECbValidateError::InvalidPackageFormat);
+					AddError(Error, ECbValidateError::NullPackageAttachment);
 				}
-				else if (Hash != FIoHash::HashBuffer(ValueView))
+				if (EnumHasAnyFlags(Mode, ECbValidateMode::PackageHash) && (BinaryAttachmentHash != FIoHash::HashBuffer(BinaryView)))
 				{
 					AddError(Error, ECbValidateError::InvalidPackageHash);
 				}
 			}
-			return Hash;
+			return BinaryAttachmentHash;
 		}
 	}
-	return FIoHash();
-}
-
-static FIoHash ValidateCbPackageObject(FCbFieldView& Value, FMemoryView& View, ECbValidateMode Mode, ECbValidateError& Error)
-{
-	FCbObjectView Object = Value.AsObjectView();
-	if (Value.HasError())
+	else if (const FMemoryView BinaryView = Value.AsBinaryView(); !Value.HasError())
+	{
+		if (BinaryView.GetSize() > 0)
+		{
+			FCompressedBuffer Buffer = FCompressedBuffer::FromCompressed(FSharedBuffer::MakeView(BinaryView));
+			if (EnumHasAnyFlags(Mode, ECbValidateMode::Package) && Buffer.IsNull())
+			{
+				AddError(Error, ECbValidateError::NullPackageAttachment);
+			}
+			if (EnumHasAnyFlags(Mode, ECbValidateMode::PackageHash) && (Buffer.GetRawHash() != FIoHash::HashBuffer(Buffer.DecompressToComposite())))
+			{
+				AddError(Error, ECbValidateError::InvalidPackageHash);
+			}
+			return Buffer.GetRawHash();
+		}
+		else
+		{
+			if (EnumHasAnyFlags(Mode, ECbValidateMode::Package))
+			{
+				AddError(Error, ECbValidateError::NullPackageAttachment);
+			}
+			return FIoHash::HashBuffer(FMemoryView());
+		}
+	}
+	else
 	{
 		if (EnumHasAnyFlags(Mode, ECbValidateMode::Package))
 		{
 			AddError(Error, ECbValidateError::InvalidPackageFormat);
 		}
 	}
-	else if (FCbFieldView HashField = ValidateCbPackageField(View, Mode, Error))
+
+	return FIoHash();
+}
+
+static FIoHash ValidateCbPackageObject(FCbFieldView& Value, FMemoryView& View, ECbValidateMode Mode, ECbValidateError& Error)
+{
+	if (FIoHash RootObjectHash = Value.AsHash(); !Value.HasError() && !Value.IsAttachment())
 	{
-		const FIoHash Hash = HashField.AsAttachment();
+		FCbFieldView RootObjectField = ValidateCbPackageField(View, Mode, Error);
+
 		if (EnumHasAnyFlags(Mode, ECbValidateMode::Package))
 		{
-			if (!Object)
-			{
-				AddError(Error, ECbValidateError::NullPackageObject);
-			}
-			if (HashField.HasError())
+			if (RootObjectField.HasError())
 			{
 				AddError(Error, ECbValidateError::InvalidPackageFormat);
 			}
-			else if (Hash != Value.GetHash())
+		}
+
+		const FCbObjectView RootObjectView = RootObjectField.AsObjectView();
+		if (EnumHasAnyFlags(Mode, ECbValidateMode::Package))
+		{
+			if (!RootObjectView)
 			{
-				AddError(Error, ECbValidateError::InvalidPackageHash);
+				AddError(Error, ECbValidateError::NullPackageObject);
 			}
 		}
-		return Hash;
+
+		if (EnumHasAnyFlags(Mode, ECbValidateMode::PackageHash) && (RootObjectHash != RootObjectView.GetHash()))
+		{
+			AddError(Error, ECbValidateError::InvalidPackageHash);
+		}
+
+		return RootObjectHash;
+	}
+	else
+	{
+		if (EnumHasAnyFlags(Mode, ECbValidateMode::Package))
+		{
+			AddError(Error, ECbValidateError::InvalidPackageFormat);
+		}
 	}
+
 	return FIoHash();
 }
 
@@ -525,24 +582,20 @@ ECbValidateError ValidateCompactBinaryPackage(FMemoryView View, ECbValidateMode
 		uint32 ObjectCount = 0;
 		while (FCbFieldView Value = ValidateCbPackageField(View, Mode, Error))
 		{
-			if (Value.IsBinary())
+			if (Value.IsHash() && !Value.IsAttachment())
 			{
-				const FIoHash Hash = ValidateCbPackageAttachment(Value, View, Mode, Error);
-				if (EnumHasAnyFlags(Mode, ECbValidateMode::Package))
+				ValidateCbPackageObject(Value, View, Mode, Error);
+				if (++ObjectCount > 1 && EnumHasAnyFlags(Mode, ECbValidateMode::Package))
 				{
-					Attachments.Add(Hash);
-					if (Value.AsBinaryView().IsEmpty())
-					{
-						AddError(Error, ECbValidateError::NullPackageAttachment);
-					}
+					AddError(Error, ECbValidateError::MultiplePackageObjects);
 				}
 			}
-			else if (Value.IsObject())
+			else if (Value.IsBinary() || Value.IsAttachment() || Value.IsObject())
 			{
-				ValidateCbPackageObject(Value, View, Mode, Error);
-				if (++ObjectCount > 1 && EnumHasAnyFlags(Mode, ECbValidateMode::Package))
+				const FIoHash Hash = ValidateCbPackageAttachment(Value, View, Mode, Error);
+				if (EnumHasAnyFlags(Mode, ECbValidateMode::Package))
 				{
-					AddError(Error, ECbValidateError::MultiplePackageObjects);
+					Attachments.Add(Hash);
 				}
 			}
 			else if (Value.IsNull())
diff --git a/Engine/Source/Runtime/Core/Private/Tests/Serialization/CompactBinaryPackageTest.cpp b/Engine/Source/Runtime/Core/Private/Tests/Serialization/CompactBinaryPackageTest.cpp
--- a/Engine/Source/Runtime/Core/Private/Tests/Serialization/CompactBinaryPackageTest.cpp
+++ b/Engine/Source/Runtime/Core/Private/Tests/Serialization/CompactBinaryPackageTest.cpp
@@ -56,9 +56,9 @@ bool FCbAttachmentTest::RunTest(const FString& Parameters)
 		TestFalse(TEXT("FCbAttachment(Null).AsBinary()"), bool(Attachment.AsBinary()));
 		TestFalse(TEXT("FCbAttachment(Null).AsObject()"), bool(Attachment.AsObject()));
 		TestFalse(TEXT("FCbAttachment(Null).IsBinary()"), Attachment.IsBinary());
+		TestFalse(TEXT("FCbAttachment(Null).IsCompressedBinary()"), Attachment.IsCompressedBinary());
 		TestFalse(TEXT("FCbAttachment(Null).IsObject()"), Attachment.IsObject());
 		TestEqual(TEXT("FCbAttachment(Null).GetHash()"), Attachment.GetHash(), FIoHash());
-		TestSaveLoadValidate(TEXT("Null"), Attachment);
 	}
 
 	// Binary Attachment
@@ -70,11 +70,27 @@ bool FCbAttachmentTest::RunTest(const FString& Parameters)
 		TestEqual(TEXT("FCbAttachment(Binary).AsBinary()"), Attachment.AsBinary(), Buffer);
 		TestFalse(TEXT("FCbAttachment(Binary).AsObject()"), bool(Attachment.AsObject()));
 		TestTrue(TEXT("FCbAttachment(Binary).IsBinary()"), Attachment.IsBinary());
+		TestFalse(TEXT("FCbAttachment(Binary).IsCompressedBinary()"), Attachment.IsCompressedBinary());
 		TestFalse(TEXT("FCbAttachment(Binary).IsObject()"), Attachment.IsObject());
 		TestEqual(TEXT("FCbAttachment(Binary).GetHash()"), Attachment.GetHash(), FIoHash::HashBuffer(Buffer));
 		TestSaveLoadValidate(TEXT("Binary"), Attachment);
 	}
 
+	// Compressed Binary Attachment
+	{
+		FCompressedBuffer Buffer = FCompressedBuffer::Compress(FSharedBuffer::Clone(MakeMemoryView<uint8>({0, 1, 2, 3})));
+		FCbAttachment Attachment(Buffer);
+		TestFalse(TEXT("FCbAttachment(CompressedBinary).IsNull()"), Attachment.IsNull());
+		TestTrue(TEXT("FCbAttachment(CompressedBinary) as bool"), bool(Attachment));
+		TestTrue(TEXT("FCbAttachment(CompressedBinary).AsCompressedBinary()"), Attachment.AsCompressedBinary().GetCompressed().Flatten().GetView().EqualBytes(Buffer.GetCompressed().Flatten().GetView()));
+		TestFalse(TEXT("FCbAttachment(CompressedBinary).AsObject()"), bool(Attachment.AsObject()));
+		TestFalse(TEXT("FCbAttachment(CompressedBinary).IsBinary()"), Attachment.IsBinary());
+		TestTrue(TEXT("FCbAttachment(CompressedBinary).IsCompressedBinary()"), Attachment.IsCompressedBinary());
+		TestFalse(TEXT("FCbAttachment(CompressedBinary).IsObject()"), Attachment.IsObject());
+		TestEqual(TEXT("FCbAttachment(CompressedBinary).GetHash()"), Attachment.GetHash(), FIoHash(Buffer.GetRawHash()));
+		TestSaveLoadValidate(TEXT("CompressedBinary"), Attachment);
+	}
+
 	// Object Attachment
 	{
 		FCbWriter Writer;
@@ -85,9 +101,10 @@ bool FCbAttachmentTest::RunTest(const FString& Parameters)
 		FCbAttachment Attachment(Object);
 		TestFalse(TEXT("FCbAttachment(Object).IsNull()"), Attachment.IsNull());
 		TestTrue(TEXT("FCbAttachment(Object) as bool"), bool(Attachment));
-		TestEqual(TEXT("FCbAttachment(Object).AsBinary()"), Attachment.AsBinary(), Object.GetOuterBuffer());
+		TestEqual(TEXT("FCbAttachment(Object).AsBinary()"), Attachment.AsBinary(), FSharedBuffer());
 		TestTrue(TEXT("FCbAttachment(Object).AsObject()"), Attachment.AsObject().Equals(Object));
-		TestTrue(TEXT("FCbAttachment(Object).IsBinary()"), Attachment.IsBinary());
+		TestFalse(TEXT("FCbAttachment(Object).IsBinary()"), Attachment.IsBinary());
+		TestFalse(TEXT("FCbAttachment(Object).IsCompressedBinary()"), Attachment.IsCompressedBinary());
 		TestTrue(TEXT("FCbAttachment(Object).IsObject()"), Attachment.IsObject());
 		TestEqual(TEXT("FCbAttachment(Object).GetHash()"), Attachment.GetHash(), FIoHash(Object.GetHash()));
 		TestSaveLoadValidate(TEXT("Object"), Attachment);
@@ -104,6 +121,7 @@ bool FCbAttachmentTest::RunTest(const FString& Parameters)
 		TestTrue(TEXT("FCbAttachment(BinaryView).AsBinary()"), Attachment.AsBinary().GetView().EqualBytes(Buffer.GetView()));
 		TestFalse(TEXT("FCbAttachment(BinaryView).AsObject()"), bool(Attachment.AsObject()));
 		TestTrue(TEXT("FCbAttachment(BinaryView).IsBinary()"), Attachment.IsBinary());
+		TestFalse(TEXT("FCbAttachment(BinaryView).IsCompressedBinary()"), Attachment.IsCompressedBinary());
 		TestFalse(TEXT("FCbAttachment(BinaryView).IsObject()"), Attachment.IsObject());
 		TestEqual(TEXT("FCbAttachment(BinaryView).GetHash()"), Attachment.GetHash(), FIoHash::HashBuffer(Buffer));
 	}
@@ -119,9 +137,9 @@ bool FCbAttachmentTest::RunTest(const FString& Parameters)
 		FCbAttachment Attachment(ObjectView);
 		TestFalse(TEXT("FCbAttachment(ObjectView).IsNull()"), Attachment.IsNull());
 		TestTrue(TEXT("FCbAttachment(ObjectView) as bool"), bool(Attachment));
-		TestNotEqual(TEXT("FCbAttachment(ObjectView).AsBinary()"), Attachment.AsBinary(), ObjectView.GetOuterBuffer());
 		TestTrue(TEXT("FCbAttachment(ObjectView).AsObject()"), Attachment.AsObject().Equals(Object));
-		TestTrue(TEXT("FCbAttachment(ObjectView).IsBinary()"), Attachment.IsBinary());
+		TestFalse(TEXT("FCbAttachment(ObjectView).IsBinary()"), Attachment.IsBinary());
+		TestFalse(TEXT("FCbAttachment(ObjectView).IsCompressedBinary()"), Attachment.IsCompressedBinary());
 		TestTrue(TEXT("FCbAttachment(ObjectView).IsObject()"), Attachment.IsObject());
 		TestEqual(TEXT("FCbAttachment(ObjectView).GetHash()"), Attachment.GetHash(), FIoHash(Object.GetHash()));
 	}
@@ -146,11 +164,38 @@ bool FCbAttachmentTest::RunTest(const FString& Parameters)
 			Attachment.AsBinary().GetView().EqualBytes(Buffer.GetView()));
 		TestFalse(TEXT("FCbAttachment(LoadBinaryView).AsObject()"), bool(Attachment.AsObject()));
 		TestTrue(TEXT("FCbAttachment(LoadBinaryView).IsBinary()"), Attachment.IsBinary());
+		TestFalse(TEXT("FCbAttachment(LoadBinaryView).IsCompressedBinary()"), Attachment.IsCompressedBinary());
 		TestFalse(TEXT("FCbAttachment(LoadBinaryView).IsObject()"), Attachment.IsObject());
 		TestEqual(TEXT("FCbAttachment(LoadBinaryView).GetHash()"),
 			Attachment.GetHash(), FIoHash::HashBuffer(MakeMemoryView(Value)));
 	}
 
+	// Compressed Binary Load from View
+	{
+		const uint8 Value[]{0, 1, 2, 3};
+		FCompressedBuffer Buffer = FCompressedBuffer::Compress(FSharedBuffer::MakeView(MakeMemoryView(Value)));
+		FCbAttachment Attachment(Buffer);
+
+		FCbWriter Writer;
+		Attachment.Save(Writer);
+		FCbFieldIterator Fields = Writer.Save();
+		FCbFieldIterator FieldsView = FCbFieldIterator::MakeRangeView(FCbFieldViewIterator(Fields));
+
+		Attachment.TryLoad(FieldsView);
+		TestFalse(TEXT("FCbAttachment(LoadCompressedBinaryView).IsNull()"), Attachment.IsNull());
+		TestTrue(TEXT("FCbAttachment(LoadCompressedBinaryView) as bool"), bool(Attachment));
+		TestFalse(TEXT("FCbAttachment(LoadCompressedBinaryView).AsBinary()->!InView"),
+			FieldsView.GetOuterBuffer().GetView().Contains(Attachment.AsCompressedBinary().GetCompressed().Flatten().GetView()));
+		TestTrue(TEXT("FCbAttachment(LoadCompressedBinaryView).AsCompressedBinary()->EqualBytes"),
+			Attachment.AsCompressedBinary().GetCompressed().Flatten().GetView().EqualBytes(Buffer.GetCompressed().Flatten().GetView()));
+		TestFalse(TEXT("FCbAttachment(LoadCompressedBinaryView).AsObject()"), bool(Attachment.AsObject()));
+		TestFalse(TEXT("FCbAttachment(LoadCompressedBinaryView).IsBinary()"), Attachment.IsBinary());
+		TestTrue(TEXT("FCbAttachment(LoadCompressedBinaryView).IsCompressedBinary()"), Attachment.IsCompressedBinary());
+		TestFalse(TEXT("FCbAttachment(LoadCompressedBinaryView).IsObject()"), Attachment.IsObject());
+		TestEqual(TEXT("FCbAttachment(LoadCompressedBinaryView).GetHash()"),
+			Attachment.GetHash(), FIoHash::HashBuffer(MakeMemoryView(Value)));
+	}
+
 	// Object Load from View
 	{
 		FCbWriter ValueWriter;
@@ -172,10 +217,11 @@ bool FCbAttachmentTest::RunTest(const FString& Parameters)
 		TestFalse(TEXT("FCbAttachment(LoadObjectView).IsNull()"), Attachment.IsNull());
 		TestTrue(TEXT("FCbAttachment(LoadObjectView) as bool"), bool(Attachment));
 		TestTrue(TEXT("FCbAttachment(LoadObjectView).AsBinary()->EqualBytes"),
-			Attachment.AsBinary().GetView().EqualBytes(Value.GetOuterBuffer().GetView()));
+			Attachment.AsBinary().GetView().EqualBytes(FMemoryView()));
 		TestFalse(TEXT("FCbAttachment(LoadObjectView).AsObject()->!InView"),
 			!Attachment.AsObject().TryGetView(View) || FieldsView.GetOuterBuffer().GetView().Contains(View));
-		TestTrue(TEXT("FCbAttachment(LoadObjectView).IsBinary()"), Attachment.IsBinary());
+		TestFalse(TEXT("FCbAttachment(LoadObjectView).IsBinary()"), Attachment.IsBinary());
+		TestFalse(TEXT("FCbAttachment(LoadObjectView).IsCompressedBinary()"), Attachment.IsCompressedBinary());
 		TestTrue(TEXT("FCbAttachment(LoadObjectView).IsObject()"), Attachment.IsObject());
 		TestEqual(TEXT("FCbAttachment(LoadObjectView).GetHash()"), Attachment.GetHash(), FIoHash(Value.GetHash()));
 	}
@@ -185,8 +231,9 @@ bool FCbAttachmentTest::RunTest(const FString& Parameters)
 		const FCbAttachment Attachment(FSharedBuffer{});
 		TestTrue(TEXT("FCbAttachment(BinaryNull).IsNull()"), Attachment.IsNull());
 		TestFalse(TEXT("FCbAttachment(BinaryNull).IsBinary()"), Attachment.IsBinary());
+		TestFalse(TEXT("FCbAttachment(BinaryNull).IsCompressedBinary()"), Attachment.IsCompressedBinary());
 		TestFalse(TEXT("FCbAttachment(BinaryNull).IsObject()"), Attachment.IsObject());
-		TestEqual(TEXT("FCbAttachment(BinaryNull).GetHash()"), Attachment.GetHash(), FIoHash::HashBuffer(FSharedBuffer{}));
+		TestEqual(TEXT("FCbAttachment(BinaryNull).GetHash()"), Attachment.GetHash(), FIoHash::Zero);
 	}
 
 	// Binary Empty
@@ -194,15 +241,27 @@ bool FCbAttachmentTest::RunTest(const FString& Parameters)
 		const FCbAttachment Attachment(FUniqueBuffer::Alloc(0).MoveToShared());
 		TestFalse(TEXT("FCbAttachment(BinaryEmpty).IsNull()"), Attachment.IsNull());
 		TestTrue(TEXT("FCbAttachment(BinaryEmpty).IsBinary()"), Attachment.IsBinary());
+		TestFalse(TEXT("FCbAttachment(BinaryEmpty).IsCompressedBinary()"), Attachment.IsCompressedBinary());
 		TestFalse(TEXT("FCbAttachment(BinaryEmpty).IsObject()"), Attachment.IsObject());
 		TestEqual(TEXT("FCbAttachment(BinaryEmpty).GetHash()"), Attachment.GetHash(), FIoHash::HashBuffer(FSharedBuffer{}));
 	}
 
+	// Compressed Binary Empty
+	{
+		const FCbAttachment Attachment(FCompressedBuffer::Compress(FUniqueBuffer::Alloc(0).MoveToShared()));
+		TestFalse(TEXT("FCbAttachment(CompressedBinaryEmpty).IsNull()"), Attachment.IsNull());
+		TestFalse(TEXT("FCbAttachment(CompressedBinaryEmpty).IsBinary()"), Attachment.IsBinary());
+		TestTrue(TEXT("FCbAttachment(CompressedBinaryEmpty).IsCompressedBinary()"), Attachment.IsCompressedBinary());
+		TestFalse(TEXT("FCbAttachment(CompressedBinaryEmpty).IsObject()"), Attachment.IsObject());
+		TestEqual(TEXT("FCbAttachment(CompressedBinaryEmpty).GetHash()"), Attachment.GetHash(), FIoHash::HashBuffer(FSharedBuffer{}));
+	}
+
 	// Object Empty
 	{
 		const FCbAttachment Attachment(FCbObject{});
 		TestFalse(TEXT("FCbAttachment(ObjectEmpty).IsNull()"), Attachment.IsNull());
-		TestTrue(TEXT("FCbAttachment(ObjectEmpty).IsBinary()"), Attachment.IsBinary());
+		TestFalse(TEXT("FCbAttachment(ObjectEmpty).IsBinary()"), Attachment.IsBinary());
+		TestFalse(TEXT("FCbAttachment(ObjectEmpty).IsCompressedBinary()"), Attachment.IsCompressedBinary());
 		TestTrue(TEXT("FCbAttachment(ObjectEmpty).IsObject()"), Attachment.IsObject());
 		TestEqual(TEXT("FCbAttachment(ObjectEmpty).GetHash()"), Attachment.GetHash(), FIoHash(FCbObject().GetHash()));
 	}
@@ -323,7 +382,8 @@ bool FCbPackageTest::RunTest(const FString& Parameters)
 		TestTrue(TEXT("FCbPackage(Attachments).FindAttachment(Object2)"),
 			Object2Attachment && Object2Attachment->AsBinary() == Object2.GetOuterBuffer());
 
-		Package.AddAttachment(FCbAttachment(FSharedBuffer::Clone(Object1.GetOuterBuffer())));
+		FSharedBuffer Object1ClonedBuffer = FSharedBuffer::Clone(Object1.GetOuterBuffer());
+		Package.AddAttachment(FCbAttachment(Object1ClonedBuffer));
 		Package.AddAttachment(FCbAttachment(FCbObject::Clone(Object2)));
 
 		TestEqual(TEXT("FCbPackage(Attachments).GetAttachments()"), Package.GetAttachments().Num(), 2);
@@ -332,14 +392,10 @@ bool FCbPackageTest::RunTest(const FString& Parameters)
 		TestEqual(TEXT("FCbPackage(Attachments).FindAttachment(Object2, Re-Add)"),
 			Package.FindAttachment(Object2.GetHash()), Object2Attachment);
 
-		TestTrue(TEXT("FCbPackage(Attachments).FindAttachment(ObjectAsObject)"),
-			Object1Attachment && Object1Attachment->AsObject().Equals(Object1));
 		TestTrue(TEXT("FCbPackage(Attachments).FindAttachment(ObjectAsBinary)"),
-			Object1Attachment && Object1Attachment->AsBinary() == Object1.GetOuterBuffer());
+			Object1Attachment&& Object1Attachment->AsBinary() == Object1ClonedBuffer);
 		TestTrue(TEXT("FCbPackage(Attachments).FindAttachment(FieldAsField)"),
-			Object2Attachment && Object2Attachment->AsObject().Equals(Object2));
-		TestTrue(TEXT("FCbPackage(Attachments).FindAttachment(FieldAsBinary)"),
-			Object2Attachment && Object2Attachment->AsBinary() == Object2.GetOuterBuffer());
+			Object2Attachment&& Object2Attachment->AsObject().Equals(Object2));
 
 		TestTrue(TEXT("FCbPackage(Attachments).GetAttachments()->Sorted"),
 			Algo::IsSorted(Package.GetAttachments()));
@@ -444,14 +500,14 @@ bool FCbPackageTest::RunTest(const FString& Parameters)
 	// Out of Order
 	{
 		TCbWriter<384> Writer;
-		Writer.AddBinary(Level2.GetOuterBuffer());
-		Writer.AddObjectAttachment(Level2Hash);
-		Writer.AddBinary(Level4);
-		Writer.AddBinaryAttachment(Level4Hash);
+		FCbAttachment Attachment2(Level2, Level2Hash);
+		Attachment2.Save(Writer);
+		FCbAttachment Attachment4(Level4);
+		Attachment4.Save(Writer);
+		Writer.AddHash(Level1Hash);
 		Writer.AddObject(Level1);
-		Writer.AddObjectAttachment(Level1Hash);
-		Writer.AddBinary(Level3.GetOuterBuffer());
-		Writer.AddObjectAttachment(Level3Hash);
+		FCbAttachment Attachment3(Level3, Level3Hash);
+		Attachment3.Save(Writer);
 		Writer.AddNull();
 
 		FCbFieldIterator Fields = Writer.Save();
@@ -469,15 +525,11 @@ bool FCbPackageTest::RunTest(const FString& Parameters)
 
 		TestTrue(TEXT("FCbPackage(OutOfOrder).TryLoad()->Level2"),
 			Level2Attachment && Level2Attachment->AsObject().Equals(Level2));
-		TestTrue(TEXT("FCbPackage(OutOfOrder).TryLoad()->Level2Buffer"),
-			Level2Attachment && Fields.GetOuterBuffer().GetView().Contains(Level2Attachment->AsBinary().GetView()));
 		TestTrue(TEXT("FCbPackage(OutOfOrder).TryLoad()->Level2Hash"),
 			Level2Attachment && Level2Attachment->GetHash() == Level2Hash);
 
 		TestTrue(TEXT("FCbPackage(OutOfOrder).TryLoad()->Level3"),
 			Level3Attachment && Level3Attachment->AsObject().Equals(Level3));
-		TestTrue(TEXT("FCbPackage(OutOfOrder).TryLoad()->Level3Buffer"),
-			Level3Attachment && Fields.GetOuterBuffer().GetView().Contains(Level3Attachment->AsBinary().GetView()));
 		TestTrue(TEXT("FCbPackage(OutOfOrder).TryLoad()->Level3Hash"),
 			Level3Attachment && Level3Attachment->GetHash() == Level3Hash);
 
@@ -498,22 +550,23 @@ bool FCbPackageTest::RunTest(const FString& Parameters)
 		FromArchive.Save(Writer);
 		FCbFieldIterator Saved = Writer.Save();
 		FMemoryView View;
-		TestTrue(TEXT("FCbPackage(OutOfOrder).Save()->Level1"), Saved.AsObject().Equals(Level1));
+		TestEqual(TEXT("FCbPackage(OutOfOrder).Save()->Level1Hash"), Saved.AsHash(), Level1Hash);
 		++Saved;
-		TestEqual(TEXT("FCbPackage(OutOfOrder).Save()->Level1Hash"), Saved.AsObjectAttachment(), Level1Hash);
-		++Saved;
-		TestTrue(TEXT("FCbPackage(OutOfOrder).Save()->Level2"), Level2.TryGetView(View) && Saved.AsBinaryView().EqualBytes(View));
+		TestTrue(TEXT("FCbPackage(OutOfOrder).Save()->Level1"), Saved.AsObject().Equals(Level1));
 		++Saved;
 		TestEqual(TEXT("FCbPackage(OutOfOrder).Save()->Level2Hash"), Saved.AsObjectAttachment(), Level2Hash);
 		++Saved;
-		TestTrue(TEXT("FCbPackage(OutOfOrder).Save()->Level3"), Level3.TryGetView(View) && Saved.AsBinaryView().EqualBytes(View));
+		TestTrue(TEXT("FCbPackage(OutOfOrder).Save()->Level2"), Saved.AsObject().Equals(Level2));
 		++Saved;
 		TestEqual(TEXT("FCbPackage(OutOfOrder).Save()->Level3Hash"), Saved.AsObjectAttachment(), Level3Hash);
 		++Saved;
-		TestTrue(TEXT("FCbPackage(OutOfOrder).Save()->Level4"), Saved.AsBinaryView().EqualBytes(Level4.GetView()));
+		TestTrue(TEXT("FCbPackage(OutOfOrder).Save()->Level3"), Saved.AsObject().Equals(Level3));
 		++Saved;
 		TestEqual(TEXT("FCbPackage(OutOfOrder).Save()->Level4Hash"), Saved.AsBinaryAttachment(), Level4Hash);
 		++Saved;
+		FSharedBuffer SavedLevel4Buffer = FSharedBuffer::MakeView(Saved.AsBinaryView());
+		TestTrue(TEXT("FCbPackage(OutOfOrder).Save()->Level4"), SavedLevel4Buffer.GetView().EqualBytes(Level4.GetView()));
+		++Saved;
 		TestTrue(TEXT("FCbPackage(OutOfOrder).Save()->Null"), Saved.IsNull());
 		++Saved;
 		TestTrue(TEXT("FCbPackage(OutOfOrder).Save()->AtEnd"), !Saved);
diff --git a/Engine/Source/Runtime/Core/Private/Tests/Serialization/CompactBinaryValidationTest.cpp b/Engine/Source/Runtime/Core/Private/Tests/Serialization/CompactBinaryValidationTest.cpp
--- a/Engine/Source/Runtime/Core/Private/Tests/Serialization/CompactBinaryValidationTest.cpp
+++ b/Engine/Source/Runtime/Core/Private/Tests/Serialization/CompactBinaryValidationTest.cpp
@@ -340,9 +340,9 @@ bool FCbValidateAttachmentTest::RunTest(const FString& Parameters)
 
 	// Test Null
 	{
-		TestEqual(TEXT("ValidateCompactBinaryAttachment(Null)"), Validate({uint8(ECbFieldType::Binary), 0}), ECbValidateError::None);
-		TestEqual(TEXT("ValidateCompactBinaryAttachment(Null, Padding)"), Validate({uint8(ECbFieldType::Binary), 0, 0}), ECbValidateError::Padding);
-		TestEqual(TEXT("ValidateCompactBinaryAttachment(Null, Padding, Mode)"), Validate({uint8(ECbFieldType::Binary), 0, 0}, ECbValidateMode::All & ~ECbValidateMode::Padding), ECbValidateError::None);
+		TestEqual(TEXT("ValidateCompactBinaryAttachment(Null)"), Validate({uint8(ECbFieldType::Binary), 0}, ECbValidateMode::All & ~ECbValidateMode::Package), ECbValidateError::None);
+		TestEqual(TEXT("ValidateCompactBinaryAttachment(Null, Padding)"), Validate({uint8(ECbFieldType::Binary), 0, 0}, ECbValidateMode::All & ~ECbValidateMode::Package), ECbValidateError::Padding);
+		TestEqual(TEXT("ValidateCompactBinaryAttachment(Null, Padding, Mode)"), Validate({uint8(ECbFieldType::Binary), 0, 0}, ECbValidateMode::All & ~(ECbValidateMode::Padding | ECbValidateMode::Package)), ECbValidateError::None);
 	}
 
 	// Test Binary
@@ -350,12 +350,10 @@ bool FCbValidateAttachmentTest::RunTest(const FString& Parameters)
 		FBufferArchive Buffer;
 		FCbAttachment(FSharedBuffer::MakeView(MakeMemoryView<uint8>({0}))).Save(Buffer);
 		TestEqual(TEXT("ValidateCompactBinaryAttachment(Binary)"), ValidateCompactBinaryAttachment(MakeMemoryView(Buffer), ECbValidateMode::All), ECbValidateError::None);
-		TestEqual(TEXT("ValidateCompactBinaryAttachment(Binary, MissingHash)"), ValidateCompactBinaryAttachment(MakeMemoryView(Buffer).Left(3), ECbValidateMode::All), ECbValidateError::InvalidPackageFormat);
-		TestEqual(TEXT("ValidateCompactBinaryAttachment(Binary, MissingHash, Mode)"), ValidateCompactBinaryAttachment(MakeMemoryView(Buffer).Left(3), ECbValidateMode::All & ~ECbValidateMode::Package), ECbValidateError::None);
-		++Buffer.Last();
+		++Buffer[2];
 		TestEqual(TEXT("ValidateCompactBinaryAttachment(Binary, InvalidHash)"), ValidateCompactBinaryAttachment(MakeMemoryView(Buffer), ECbValidateMode::All), ECbValidateError::InvalidPackageHash);
-		TestEqual(TEXT("ValidateCompactBinaryAttachment(Binary, InvalidHash, Mode)"), ValidateCompactBinaryAttachment(MakeMemoryView(Buffer), ECbValidateMode::All & ~ECbValidateMode::Package), ECbValidateError::None);
-		--Buffer.Last();
+		TestEqual(TEXT("ValidateCompactBinaryAttachment(Binary, InvalidHash, Mode)"), ValidateCompactBinaryAttachment(MakeMemoryView(Buffer), ECbValidateMode::All & ~ECbValidateMode::PackageHash), ECbValidateError::None);
+		--Buffer[2];
 		Buffer.AddZeroed(1);
 		TestEqual(TEXT("ValidateCompactBinaryAttachment(Binary, Padding)"), ValidateCompactBinaryAttachment(MakeMemoryView(Buffer), ECbValidateMode::All), ECbValidateError::Padding);
 		TestEqual(TEXT("ValidateCompactBinaryAttachment(Binary, Padding, Mode)"), ValidateCompactBinaryAttachment(MakeMemoryView(Buffer), ECbValidateMode::All & ~ECbValidateMode::Padding), ECbValidateError::None);
@@ -374,12 +372,10 @@ bool FCbValidateAttachmentTest::RunTest(const FString& Parameters)
 		FBufferArchive Buffer;
 		FCbAttachment(Writer.Save().AsObject()).Save(Buffer);
 		TestEqual(TEXT("ValidateCompactBinaryAttachment(Object)"), ValidateCompactBinaryAttachment(MakeMemoryView(Buffer), ECbValidateMode::All), ECbValidateError::None);
-		TestEqual(TEXT("ValidateCompactBinaryAttachment(Object, MissingHash)"), ValidateCompactBinaryAttachment(MakeMemoryView(Buffer).Left(7), ECbValidateMode::All), ECbValidateError::InvalidPackageFormat);
-		TestEqual(TEXT("ValidateCompactBinaryAttachment(Object, MissingHash, Mode)"), ValidateCompactBinaryAttachment(MakeMemoryView(Buffer).Left(7), ECbValidateMode::All & ~ECbValidateMode::Package), ECbValidateError::None);
-		++Buffer.Last();
+		++Buffer[2];
 		TestEqual(TEXT("ValidateCompactBinaryAttachment(Object, InvalidHash)"), ValidateCompactBinaryAttachment(MakeMemoryView(Buffer), ECbValidateMode::All), ECbValidateError::InvalidPackageHash);
-		TestEqual(TEXT("ValidateCompactBinaryAttachment(Object, InvalidHash, Mode)"), ValidateCompactBinaryAttachment(MakeMemoryView(Buffer), ECbValidateMode::All & ~ECbValidateMode::Package), ECbValidateError::None);
-		--Buffer.Last();
+		TestEqual(TEXT("ValidateCompactBinaryAttachment(Object, InvalidHash, Mode)"), ValidateCompactBinaryAttachment(MakeMemoryView(Buffer), ECbValidateMode::All & ~ECbValidateMode::PackageHash), ECbValidateError::None);
+		--Buffer[2];
 		Buffer.AddZeroed(1);
 		TestEqual(TEXT("ValidateCompactBinaryAttachment(Object, Padding)"), ValidateCompactBinaryAttachment(MakeMemoryView(Buffer), ECbValidateMode::All), ECbValidateError::Padding);
 		TestEqual(TEXT("ValidateCompactBinaryAttachment(Object, Padding, Mode)"), ValidateCompactBinaryAttachment(MakeMemoryView(Buffer), ECbValidateMode::All & ~ECbValidateMode::Padding), ECbValidateError::None);
@@ -389,31 +385,23 @@ bool FCbValidateAttachmentTest::RunTest(const FString& Parameters)
 		TestEqual(TEXT("ValidateCompactBinaryAttachment(Object, OutOfBoundsHash, Mode)"), ValidateCompactBinaryAttachment(MakeMemoryView(Buffer).Left(8), ECbValidateMode::None), ECbValidateError::None);
 	}
 
-	// Test InvalidPackageFormat
-	{
-		TestEqual(TEXT("ValidateCompactBinaryAttachment(Empty)"), Validate({}), ECbValidateError::InvalidPackageFormat);
-		TestEqual(TEXT("ValidateCompactBinaryAttachment(Empty, Mode)"), Validate({}, ECbValidateMode::All & ~ECbValidateMode::Package), ECbValidateError::None);
-	}
+	// Test Empty Object
 	{
 		FBufferArchive Buffer;
 		FCbFieldIterator::MakeSingle(FCbObject().AsField()).CopyRangeTo(Buffer);
-		TestEqual(TEXT("ValidateCompactBinaryAttachment(Object)"), ValidateCompactBinaryAttachment(MakeMemoryView(Buffer), ECbValidateMode::All), ECbValidateError::InvalidPackageFormat);
-		TestEqual(TEXT("ValidateCompactBinaryAttachment(Object, Mode)"), ValidateCompactBinaryAttachment(MakeMemoryView(Buffer), ECbValidateMode::All & ~ECbValidateMode::Package), ECbValidateError::None);
+		TestEqual(TEXT("ValidateCompactBinaryAttachment(Object)"), ValidateCompactBinaryAttachment(MakeMemoryView(Buffer), ECbValidateMode::All), ECbValidateError::None);
 	}
+
+	// Test InvalidPackageFormat
 	{
-		FBufferArchive Buffer;
-		FCbWriter Writer;
-		Writer.AddObjectAttachment(FIoHash::HashBuffer(BinaryView));
-		Writer.AddBinary(BinaryView);
-		Writer.Save(Buffer);
-		TestEqual(TEXT("ValidateCompactBinaryAttachment(HashBeforeValue)"), ValidateCompactBinaryAttachment(MakeMemoryView(Buffer), ECbValidateMode::All), ECbValidateError::InvalidPackageFormat | ECbValidateError::Padding);
-		TestEqual(TEXT("ValidateCompactBinaryAttachment(HashBeforeValue, Mode)"), ValidateCompactBinaryAttachment(MakeMemoryView(Buffer), ECbValidateMode::All & ~ECbValidateMode::Package), ECbValidateError::Padding);
+		TestEqual(TEXT("ValidateCompactBinaryAttachment(Empty)"), Validate({}), ECbValidateError::InvalidPackageFormat);
+		TestEqual(TEXT("ValidateCompactBinaryAttachment(Empty, Mode)"), Validate({}, ECbValidateMode::All & ~ECbValidateMode::Package), ECbValidateError::None);
 	}
 	{
 		FBufferArchive Buffer;
 		FCbWriter Writer;
+		Writer.AddBinaryAttachment(FIoHash::HashBuffer(BinaryView));
 		Writer.SetName("Name"_ASV).AddBinary(BinaryView);
-		Writer.AddObjectAttachment(FIoHash::HashBuffer(BinaryView));
 		Writer.Save(Buffer);
 		TestEqual(TEXT("ValidateCompactBinaryAttachment(NameOnValue)"), ValidateCompactBinaryAttachment(MakeMemoryView(Buffer), ECbValidateMode::All), ECbValidateError::InvalidPackageFormat);
 		TestEqual(TEXT("ValidateCompactBinaryAttachment(NameOnValue, Mode)"), ValidateCompactBinaryAttachment(MakeMemoryView(Buffer), ECbValidateMode::All & ~ECbValidateMode::Package), ECbValidateError::None);
@@ -421,8 +409,8 @@ bool FCbValidateAttachmentTest::RunTest(const FString& Parameters)
 	{
 		FBufferArchive Buffer;
 		FCbWriter Writer;
+		Writer.SetName("Name"_ASV).AddBinaryAttachment(FIoHash::HashBuffer(BinaryView));
 		Writer.AddBinary(BinaryView);
-		Writer.SetName("Name"_ASV).AddObjectAttachment(FIoHash::HashBuffer(BinaryView));
 		Writer.Save(Buffer);
 		TestEqual(TEXT("ValidateCompactBinaryAttachment(NameOnHash)"), ValidateCompactBinaryAttachment(MakeMemoryView(Buffer), ECbValidateMode::All), ECbValidateError::InvalidPackageFormat);
 		TestEqual(TEXT("ValidateCompactBinaryAttachment(NameOnHash, Mode)"), ValidateCompactBinaryAttachment(MakeMemoryView(Buffer), ECbValidateMode::All & ~ECbValidateMode::Package), ECbValidateError::None);
@@ -438,6 +426,16 @@ bool FCbValidatePackageTest::RunTest(const FString& Parameters)
 {
 	const uint8 BinaryValue[] = { 0, 1, 2, 3 };
 	const FMemoryView BinaryView = MakeMemoryView(BinaryValue);
+	const FSharedBuffer BinaryBuffer = FSharedBuffer::MakeView(BinaryView);
+	FIoHash BinaryBufferHash = FIoHash::HashBuffer(BinaryBuffer);
+	const FCompressedBuffer CompressedBinaryBuffer = FCompressedBuffer::Compress(FSharedBuffer::MakeView(BinaryView));
+	
+	TCbWriter<32> CommonRootObjectWriter;
+	CommonRootObjectWriter.BeginObject();
+	CommonRootObjectWriter << "Field" << 42;
+	CommonRootObjectWriter.EndObject();
+	FCbObject CommonRootObject(CommonRootObjectWriter.Save().AsObject());
+	FIoHash CommonRootObjectHash = CommonRootObject.GetHash();
 
 	// Test Null
 	{
@@ -450,10 +448,8 @@ bool FCbValidatePackageTest::RunTest(const FString& Parameters)
 	{
 		FBufferArchive Buffer;
 		TCbWriter<256> Writer;
-		Writer.BeginObject();
-		Writer << "Field" << 42;
-		Writer.EndObject();
-		Writer.AddObjectAttachment(Writer.Save().GetHash());
+		Writer.AddHash(CommonRootObjectHash);
+		Writer.AddObject(CommonRootObject);
 		Writer.AddNull();
 		Writer.Save(Buffer);
 		TestEqual(TEXT("ValidateCompactBinaryPackage(Object)"), ValidateCompactBinaryPackage(MakeMemoryView(Buffer), ECbValidateMode::All), ECbValidateError::None);
@@ -463,7 +459,7 @@ bool FCbValidatePackageTest::RunTest(const FString& Parameters)
 		TestEqual(TEXT("ValidateCompactBinaryPackage(Object, MissingNull, Mode)"), ValidateCompactBinaryPackage(MakeMemoryView(Buffer).LeftChop(1), ECbValidateMode::All & ~ECbValidateMode::Package), ECbValidateError::None);
 		++Buffer.Last(1);
 		TestEqual(TEXT("ValidateCompactBinaryPackage(Object, InvalidHash)"), ValidateCompactBinaryPackage(MakeMemoryView(Buffer), ECbValidateMode::All), ECbValidateError::InvalidPackageHash);
-		TestEqual(TEXT("ValidateCompactBinaryPackage(Object, InvalidHash, Mode)"), ValidateCompactBinaryPackage(MakeMemoryView(Buffer), ECbValidateMode::All & ~ECbValidateMode::Package), ECbValidateError::None);
+		TestEqual(TEXT("ValidateCompactBinaryPackage(Object, InvalidHash, Mode)"), ValidateCompactBinaryPackage(MakeMemoryView(Buffer), ECbValidateMode::All & ~ECbValidateMode::PackageHash), ECbValidateError::None);
 		--Buffer.Last(1);
 		Buffer.AddZeroed(1);
 		TestEqual(TEXT("ValidateCompactBinaryPackage(Object, Padding)"), ValidateCompactBinaryPackage(MakeMemoryView(Buffer), ECbValidateMode::All), ECbValidateError::Padding);
@@ -478,24 +474,20 @@ bool FCbValidatePackageTest::RunTest(const FString& Parameters)
 	{
 		FBufferArchive Buffer;
 		TCbWriter<256> Writer;
-		Writer.BeginObject();
-		Writer << "Field" << 42;
-		Writer.EndObject();
-		Writer.AddObjectAttachment(Writer.Save().GetHash());
+		Writer.AddHash(CommonRootObjectHash);
+		Writer.AddObject(CommonRootObject);
 		const uint64 AttachmentOffset = Writer.GetSaveSize();
-		Writer.AddBinary(BinaryView);
-		Writer.AddObjectAttachment(FIoHash::HashBuffer(BinaryView));
+		Writer.AddBinaryAttachment(BinaryBufferHash);
+		Writer.AddBinary(BinaryBuffer);
 		Writer.AddNull();
 		Writer.Save(Buffer);
 		TestEqual(TEXT("ValidateCompactBinaryPackage(Object+Attachment)"), ValidateCompactBinaryPackage(MakeMemoryView(Buffer), ECbValidateMode::All), ECbValidateError::None);
-		TestEqual(TEXT("ValidateCompactBinaryPackage(Object+Attachment, MissingHash)"), ValidateCompactBinaryPackage(MakeMemoryView(Buffer).LeftChop(22), ECbValidateMode::All), ECbValidateError::InvalidPackageFormat);
-		TestEqual(TEXT("ValidateCompactBinaryPackage(Object+Attachment, MissingHash, Mode)"), ValidateCompactBinaryPackage(MakeMemoryView(Buffer).LeftChop(22), ECbValidateMode::All & ~ECbValidateMode::Package), ECbValidateError::None);
 		TestEqual(TEXT("ValidateCompactBinaryPackage(Object+Attachment, MissingNull)"), ValidateCompactBinaryPackage(MakeMemoryView(Buffer).LeftChop(1), ECbValidateMode::All), ECbValidateError::InvalidPackageFormat);
 		TestEqual(TEXT("ValidateCompactBinaryPackage(Object+Attachment, MissingNull, Mode)"), ValidateCompactBinaryPackage(MakeMemoryView(Buffer).LeftChop(1), ECbValidateMode::All & ~ECbValidateMode::Package), ECbValidateError::None);
-		++Buffer.Last(1);
+		++Buffer[2];
 		TestEqual(TEXT("ValidateCompactBinaryPackage(Object+Attachment, InvalidHash)"), ValidateCompactBinaryPackage(MakeMemoryView(Buffer), ECbValidateMode::All), ECbValidateError::InvalidPackageHash);
-		TestEqual(TEXT("ValidateCompactBinaryPackage(Object+Attachment, InvalidHash, Mode)"), ValidateCompactBinaryPackage(MakeMemoryView(Buffer), ECbValidateMode::All & ~ECbValidateMode::Package), ECbValidateError::None);
-		--Buffer.Last(1);
+		TestEqual(TEXT("ValidateCompactBinaryPackage(Object+Attachment, InvalidHash, Mode)"), ValidateCompactBinaryPackage(MakeMemoryView(Buffer), ECbValidateMode::All & ~ECbValidateMode::PackageHash), ECbValidateError::None);
+		--Buffer[2];
 		Buffer.AddZeroed(1);
 		TestEqual(TEXT("ValidateCompactBinaryPackage(Object+Attachment, Padding)"), ValidateCompactBinaryPackage(MakeMemoryView(Buffer), ECbValidateMode::All), ECbValidateError::Padding);
 		TestEqual(TEXT("ValidateCompactBinaryPackage(Object+Attachment, Padding, Mode)"), ValidateCompactBinaryPackage(MakeMemoryView(Buffer), ECbValidateMode::All & ~ECbValidateMode::Padding), ECbValidateError::None);
@@ -505,6 +497,32 @@ bool FCbValidatePackageTest::RunTest(const FString& Parameters)
 		TestEqual(TEXT("ValidateCompactBinaryPackage(Object+Attachment, OutOfBoundsHash, Mode)"), ValidateCompactBinaryPackage(MakeMemoryView(Buffer).Left(AttachmentOffset + 7), ECbValidateMode::None), ECbValidateError::None);
 	}
 
+	// Test Object + Compressed Binary Attachment
+	{
+		FBufferArchive Buffer;
+		TCbWriter<256> Writer;
+		Writer.AddHash(CommonRootObjectHash);
+		Writer.AddObject(CommonRootObject);
+		const uint64 AttachmentOffset = Writer.GetSaveSize();
+		Writer.AddBinary(CompressedBinaryBuffer.GetCompressed());
+		Writer.AddNull();
+		Writer.Save(Buffer);
+		TestEqual(TEXT("ValidateCompactBinaryPackage(Object+CompressedBinaryAttachment)"), ValidateCompactBinaryPackage(MakeMemoryView(Buffer), ECbValidateMode::All), ECbValidateError::None);
+		TestEqual(TEXT("ValidateCompactBinaryPackage(Object+CompressedBinaryAttachment, MissingNull)"), ValidateCompactBinaryPackage(MakeMemoryView(Buffer).LeftChop(1), ECbValidateMode::All), ECbValidateError::InvalidPackageFormat);
+		TestEqual(TEXT("ValidateCompactBinaryPackage(Object+CompressedBinaryAttachment, MissingNull, Mode)"), ValidateCompactBinaryPackage(MakeMemoryView(Buffer).LeftChop(1), ECbValidateMode::All & ~ECbValidateMode::Package), ECbValidateError::None);
+		++Buffer[2];
+		TestEqual(TEXT("ValidateCompactBinaryPackage(Object+CompressedBinaryAttachment, InvalidHash)"), ValidateCompactBinaryPackage(MakeMemoryView(Buffer), ECbValidateMode::All), ECbValidateError::InvalidPackageHash);
+		TestEqual(TEXT("ValidateCompactBinaryPackage(Object+CompressedBinaryAttachment, InvalidHash, Mode)"), ValidateCompactBinaryPackage(MakeMemoryView(Buffer), ECbValidateMode::All & ~ECbValidateMode::PackageHash), ECbValidateError::None);
+		--Buffer[2];
+		Buffer.AddZeroed(1);
+		TestEqual(TEXT("ValidateCompactBinaryPackage(Object+CompressedBinaryAttachment, Padding)"), ValidateCompactBinaryPackage(MakeMemoryView(Buffer), ECbValidateMode::All), ECbValidateError::Padding);
+		TestEqual(TEXT("ValidateCompactBinaryPackage(Object+CompressedBinaryAttachment, Padding, Mode)"), ValidateCompactBinaryPackage(MakeMemoryView(Buffer), ECbValidateMode::All & ~ECbValidateMode::Padding), ECbValidateError::None);
+		Buffer.Pop(/*bAllowShrinking*/ false);
+		TestEqual(TEXT("ValidateCompactBinaryPackage(Object+CompressedBinaryAttachment, OutOfBoundsValue)"), ValidateCompactBinaryPackage(MakeMemoryView(Buffer).Left(AttachmentOffset + 1), ECbValidateMode::All), ECbValidateError::OutOfBounds);
+		TestEqual(TEXT("ValidateCompactBinaryPackage(Object+CompressedBinaryAttachment, OutOfBoundsHash)"), ValidateCompactBinaryPackage(MakeMemoryView(Buffer).Left(AttachmentOffset + 7), ECbValidateMode::All), ECbValidateError::OutOfBounds);
+		TestEqual(TEXT("ValidateCompactBinaryPackage(Object+CompressedBinaryAttachment, OutOfBoundsHash, Mode)"), ValidateCompactBinaryPackage(MakeMemoryView(Buffer).Left(AttachmentOffset + 7), ECbValidateMode::None), ECbValidateError::None);
+	}
+
 	// Test InvalidPackageFormat
 	{
 		TestEqual(TEXT("ValidateCompactBinaryPackage(Empty)"), ValidateCompactBinaryPackage(FMemoryView(), ECbValidateMode::All), ECbValidateError::InvalidPackageFormat);
@@ -513,11 +531,9 @@ bool FCbValidatePackageTest::RunTest(const FString& Parameters)
 	{
 		FBufferArchive Buffer;
 		TCbWriter<256> Writer;
+		Writer.AddHash(CommonRootObjectHash);
 		Writer.SetName("Name"_ASV);
-		Writer.BeginObject();
-		Writer << "Field" << 42;
-		Writer.EndObject();
-		Writer.AddObjectAttachment(Writer.Save().GetHash());
+		Writer.AddObject(CommonRootObject);
 		Writer.AddNull();
 		Writer.Save(Buffer);
 		TestEqual(TEXT("ValidateCompactBinaryPackage(NameOnObject)"), ValidateCompactBinaryPackage(MakeMemoryView(Buffer), ECbValidateMode::All), ECbValidateError::InvalidPackageFormat);
@@ -526,12 +542,9 @@ bool FCbValidatePackageTest::RunTest(const FString& Parameters)
 	{
 		FBufferArchive Buffer;
 		TCbWriter<256> Writer;
-		Writer.BeginObject();
-		Writer << "Field" << 42;
-		Writer.EndObject();
-		const FIoHash ObjectHash = Writer.Save().GetHash();
 		Writer.SetName("Name"_ASV);
-		Writer.AddObjectAttachment(ObjectHash);
+		Writer.AddHash(CommonRootObjectHash);
+		Writer.AddObject(CommonRootObject);
 		Writer.AddNull();
 		Writer.Save(Buffer);
 		TestEqual(TEXT("ValidateCompactBinaryPackage(NameOnHash)"), ValidateCompactBinaryPackage(MakeMemoryView(Buffer), ECbValidateMode::All), ECbValidateError::InvalidPackageFormat);
@@ -540,12 +553,10 @@ bool FCbValidatePackageTest::RunTest(const FString& Parameters)
 	{
 		FBufferArchive Buffer;
 		TCbWriter<256> Writer;
-		Writer.BeginObject();
-		Writer << "Field" << 42;
-		Writer.EndObject();
-		Writer.AddObjectAttachment(Writer.Save().GetHash());
-		Writer.SetName("Name"_ASV).AddBinary(BinaryView);
-		Writer.AddObjectAttachment(FIoHash::HashBuffer(BinaryView));
+		Writer.AddHash(CommonRootObjectHash);
+		Writer.AddObject(CommonRootObject);
+		Writer.AddBinaryAttachment(BinaryBufferHash);
+		Writer.SetName("Name"_ASV).AddBinary(BinaryBuffer);
 		Writer.AddNull();
 		Writer.Save(Buffer);
 		TestEqual(TEXT("ValidateCompactBinaryPackage(NameOnAttachment)"), ValidateCompactBinaryPackage(MakeMemoryView(Buffer), ECbValidateMode::All), ECbValidateError::InvalidPackageFormat);
@@ -554,15 +565,10 @@ bool FCbValidatePackageTest::RunTest(const FString& Parameters)
 	{
 		FBufferArchive Buffer;
 		TCbWriter<256> Writer;
-		Writer.BeginObject();
-		Writer << "Field" << 42;
-		Writer.EndObject();
-		const FIoHash ObjectHash = Writer.Save().GetHash();
-		Writer.AddObjectAttachment(ObjectHash);
-		Writer.BeginObject();
-		Writer << "Field" << 42;
-		Writer.EndObject();
-		Writer.AddObjectAttachment(ObjectHash);
+		Writer.AddHash(CommonRootObjectHash);
+		Writer.AddObject(CommonRootObject);
+		Writer.AddHash(CommonRootObjectHash);
+		Writer.AddObject(CommonRootObject);
 		Writer.AddNull();
 		Writer.Save(Buffer);
 		TestEqual(TEXT("ValidateCompactBinaryPackage(MultipleObjects)"), ValidateCompactBinaryPackage(MakeMemoryView(Buffer), ECbValidateMode::All), ECbValidateError::MultiplePackageObjects);
@@ -571,10 +577,10 @@ bool FCbValidatePackageTest::RunTest(const FString& Parameters)
 	{
 		FBufferArchive Buffer;
 		TCbWriter<256> Writer;
-		Writer.AddBinary(BinaryView);
-		Writer.AddObjectAttachment(FIoHash::HashBuffer(BinaryView));
-		Writer.AddBinary(BinaryView);
-		Writer.AddObjectAttachment(FIoHash::HashBuffer(BinaryView));
+		Writer.AddBinaryAttachment(BinaryBufferHash);
+		Writer.AddBinary(BinaryBuffer);
+		Writer.AddBinaryAttachment(BinaryBufferHash);
+		Writer.AddBinary(BinaryBuffer);
 		Writer.AddNull();
 		Writer.Save(Buffer);
 		TestEqual(TEXT("ValidateCompactBinaryPackage(DuplicateAttachments)"), ValidateCompactBinaryPackage(MakeMemoryView(Buffer), ECbValidateMode::All), ECbValidateError::DuplicateAttachments);
@@ -583,9 +589,35 @@ bool FCbValidatePackageTest::RunTest(const FString& Parameters)
 	{
 		FBufferArchive Buffer;
 		TCbWriter<256> Writer;
-		Writer.BeginObject();
-		Writer.EndObject();
-		Writer.AddObjectAttachment(FCbObjectView().GetHash());
+		Writer.AddBinary(CompressedBinaryBuffer.GetCompressed());
+		Writer.AddBinary(CompressedBinaryBuffer.GetCompressed());
+		Writer.AddNull();
+		Writer.Save(Buffer);
+		TestEqual(TEXT("ValidateCompactBinaryPackage(DuplicateCompressedAttachments)"), ValidateCompactBinaryPackage(MakeMemoryView(Buffer), ECbValidateMode::All), ECbValidateError::DuplicateAttachments);
+		TestEqual(TEXT("ValidateCompactBinaryPackage(DuplicateCompressedAttachments, Mode)"), ValidateCompactBinaryPackage(MakeMemoryView(Buffer), ECbValidateMode::All & ~ECbValidateMode::Package), ECbValidateError::None);
+	}
+	{
+		FBufferArchive Buffer;
+		TCbWriter<256> Writer;
+		Writer.AddBinaryAttachment(BinaryBufferHash);
+		Writer.AddBinary(BinaryBuffer);
+		Writer.AddBinary(CompressedBinaryBuffer.GetCompressed());
+		Writer.AddNull();
+		Writer.Save(Buffer);
+		TestEqual(TEXT("ValidateCompactBinaryPackage(DuplicateMixedAttachments)"), ValidateCompactBinaryPackage(MakeMemoryView(Buffer), ECbValidateMode::All), ECbValidateError::DuplicateAttachments);
+		TestEqual(TEXT("ValidateCompactBinaryPackage(DuplicateMixedAttachments, Mode)"), ValidateCompactBinaryPackage(MakeMemoryView(Buffer), ECbValidateMode::All & ~ECbValidateMode::Package), ECbValidateError::None);
+	}
+	{
+		TCbWriter<32> RootObjectWriter;
+		RootObjectWriter.BeginObject();
+		RootObjectWriter.EndObject();
+		FCbObject NullRootObject(RootObjectWriter.Save().AsObject());
+		FIoHash NullRootObjectHash = NullRootObject.GetHash();
+
+		FBufferArchive Buffer;
+		TCbWriter<256> Writer;
+		Writer.AddHash(NullRootObjectHash);
+		Writer.AddObject(NullRootObject);
 		Writer.AddNull();
 		Writer.Save(Buffer);
 		TestEqual(TEXT("ValidateCompactBinaryPackage(NullObject)"), ValidateCompactBinaryPackage(MakeMemoryView(Buffer), ECbValidateMode::All), ECbValidateError::NullPackageObject);
@@ -594,6 +626,7 @@ bool FCbValidatePackageTest::RunTest(const FString& Parameters)
 	{
 		FBufferArchive Buffer;
 		TCbWriter<256> Writer;
+		Writer.AddBinaryAttachment(FIoHash::HashBuffer(FMemoryView()));
 		Writer.AddBinary(FMemoryView());
 		Writer.AddNull();
 		Writer.Save(Buffer);
diff --git a/Engine/Source/Runtime/Core/Public/Serialization/CompactBinaryPackage.h b/Engine/Source/Runtime/Core/Public/Serialization/CompactBinaryPackage.h
--- a/Engine/Source/Runtime/Core/Public/Serialization/CompactBinaryPackage.h
+++ b/Engine/Source/Runtime/Core/Public/Serialization/CompactBinaryPackage.h
@@ -3,9 +3,11 @@
 #pragma once
 
 #include "CoreTypes.h"
+#include "Compression/CompressedBuffer.h"
 #include "Containers/Array.h"
 #include "IO/IoHash.h"
-#include "Memory/SharedBuffer.h"
+#include "Memory/MemoryFwd.h"
+#include "Misc/TVariant.h"
 #include "Serialization/CompactBinary.h"
 
 class FArchive;
@@ -15,13 +17,7 @@ template <typename FuncType> class TFunctionRef;
 ///////////////////////////////////////////////////////////////////////////////////////////////////////////////////////
 
 /**
- * An attachment is either binary or object and is identified by its hash.
- *
- * An object attachment is also a valid binary attachment and may be accessed as binary.
- *
- * Attachments are serialized as one or two compact binary fields with no name. A Binary field is
- * written first with its content. The content hash is omitted when the content size is zero, and
- * is otherwise written as a BinaryAttachment or ObjectAttachment depending on the type.
+ * An attachment is either null, raw binary, compressed binary, or an object and is identified by its uncompressed hash.
  */
 class FCbAttachment
 {
@@ -30,57 +26,110 @@ public:
 	FCbAttachment() = default;
 
 	/** Construct an object attachment. Value is cloned if not owned. */
-	inline explicit FCbAttachment(FCbObject Value)
-		: FCbAttachment(MoveTemp(Value), nullptr)
+	inline explicit FCbAttachment(const FCbObject& InValue)
+		: FCbAttachment(InValue, nullptr)
 	{
 	}
 
 	/** Construct an object attachment. Value is cloned if not owned. Hash must match Value. */
-	inline explicit FCbAttachment(FCbObject Value, const FIoHash& Hash)
-		: FCbAttachment(MoveTemp(Value), &Hash)
+	inline explicit FCbAttachment(const FCbObject& InValue, const FIoHash& Hash)
+		: FCbAttachment(InValue, &Hash)
+	{
+	}
+
+	/** Construct a raw binary attachment. Value is cloned if not owned. */
+	inline explicit FCbAttachment(const FSharedBuffer& InValue)
+		: FCbAttachment(FCompositeBuffer(InValue))
+	{
+	}
+
+	/** Construct a raw binary attachment. Value is cloned if not owned. */
+	inline explicit FCbAttachment(const FCompositeBuffer& InValue)
+		: Value(TInPlaceType<FBinaryValue>(), InValue)
+	{
+		if (Value.Get<FBinaryValue>().Buffer.IsNull())
+		{
+			Value.Emplace<TYPE_OF_NULLPTR>();
+		}
+	}
+
+	/** Construct a raw binary attachment. Value is cloned if not owned. */
+	inline explicit FCbAttachment(FCompositeBuffer&& InValue)
+		: Value(TInPlaceType<FBinaryValue>(), InValue)
 	{
+		if (Value.Get<FBinaryValue>().Buffer.IsNull())
+		{
+			Value.Emplace<TYPE_OF_NULLPTR>();
+		}
 	}
 
-	/** Construct a binary attachment. Value is cloned if not owned. */
-	inline explicit FCbAttachment(FSharedBuffer Value)
-		: FCbAttachment(MoveTemp(Value), nullptr)
+	/** Construct a raw binary attachment. Value is cloned if not owned. */
+	inline explicit FCbAttachment(FCompositeBuffer&& InValue, const FIoHash& InHash)
+		: Value(TInPlaceType<FBinaryValue>(), InValue, InHash)
 	{
+		if (Value.Get<FBinaryValue>().Buffer.IsNull())
+		{
+			Value.Emplace<TYPE_OF_NULLPTR>();
+		}
 	}
 
-	/** Construct a binary attachment. Value is cloned if not owned. Hash must match Value. */
-	inline explicit FCbAttachment(FSharedBuffer Value, const FIoHash& Hash)
-		: FCbAttachment(MoveTemp(Value), &Hash)
+	/** Construct a compressed binary attachment. Value is cloned if not owned. */
+	inline explicit FCbAttachment(const FCompressedBuffer& InValue)
+		: Value(TInPlaceType<FCompressedBuffer>(), InValue)
 	{
+		if (Value.Get<FCompressedBuffer>().IsNull())
+		{
+			Value.Emplace<TYPE_OF_NULLPTR>();
+		}
+	}
+
+	/** Construct a compressed binary attachment. Value is cloned if not owned. */
+	inline explicit FCbAttachment(FCompressedBuffer&& InValue)
+		: Value(TInPlaceType<FCompressedBuffer>(), InValue)
+	{
+		if (Value.Get<FCompressedBuffer>().IsNull())
+		{
+			Value.Emplace<TYPE_OF_NULLPTR>();
+		}
 	}
 
 	/** Reset this to a null attachment. */
 	inline void Reset() { *this = FCbAttachment(); }
 
-	/** Whether the attachment has a value. */
-	inline explicit operator bool() const { return !IsNull(); }
+	/** Access the attachment as an object. Defaults to an empty object on error. */
+	inline FCbObject AsObject() const;
 
-	/** Whether the attachment has a value. */
-	inline bool IsNull() const { return !Buffer; }
+	/** Access the attachment as raw binary in a single contiguous buffer. Defaults to a null buffer on error. */
+	inline FSharedBuffer AsBinary() const;
 
-	/** Access the attachment as binary. Defaults to a null buffer if the attachment is null. */
-	inline FSharedBuffer AsBinary() const { return Buffer; }
+	/** Access the attachment as raw binary. Defaults to a null buffer on error. */
+	inline const FCompositeBuffer& AsCompositeBinary() const;
 
-	/** Access the attachment as an object. Defaults to an empty object on error. */
-	inline FCbObject AsObject() const { return Object ? FCbObject(FCbFieldView(Object).AsObjectView(), Buffer) : FCbObject(); }
+	/** Access the attachment as compressed binary. Defaults to a null buffer on error. */
+	inline const FCompressedBuffer& AsCompressedBinary() const;
 
-	/** Returns whether the attachment is binary or an object. */
-	inline bool IsBinary() const { return !Buffer.IsNull(); }
+	/** Whether the attachment has a non-null value. */
+	inline explicit operator bool() const { return !IsNull(); }
+
+	/** Whether the attachment has a null value. */
+	inline bool IsNull() const { return Value.IsType<TYPE_OF_NULLPTR>(); }
 
 	/** Returns whether the attachment is an object. */
-	inline bool IsObject() const { return Object.IsObject(); }
+	inline bool IsObject() const { return Value.IsType<FObjectValue>(); }
+
+	/** Returns whether the attachment is raw binary. */
+	inline bool IsBinary() const { return Value.IsType<FBinaryValue>(); }
+
+	/** Returns whether the attachment is compressed binary. */
+	inline bool IsCompressedBinary() const { return Value.IsType<FCompressedBuffer>(); }
 
 	/** Returns the hash of the attachment value. */
-	inline const FIoHash& GetHash() const { return Hash; }
+	CORE_API FIoHash GetHash() const;
 
 	/** Compares attachments by their hash. Any discrepancy in type must be handled externally. */
-	inline bool operator==(const FCbAttachment& Attachment) const { return Hash == Attachment.Hash; }
-	inline bool operator!=(const FCbAttachment& Attachment) const { return Hash != Attachment.Hash; }
-	inline bool operator<(const FCbAttachment& Attachment) const { return Hash < Attachment.Hash; }
+	inline bool operator==(const FCbAttachment& Attachment) const { return GetHash() == Attachment.GetHash(); }
+	inline bool operator!=(const FCbAttachment& Attachment) const { return GetHash() != Attachment.GetHash(); }
+	inline bool operator<(const FCbAttachment& Attachment) const { return GetHash() < Attachment.GetHash(); }
 
 	/**
 	 * Load the attachment from compact binary as written by Save.
@@ -109,15 +158,63 @@ public:
 	CORE_API void Save(FArchive& Ar) const;
 
 private:
-	CORE_API FCbAttachment(FCbObject Value, const FIoHash* Hash);
-	CORE_API FCbAttachment(FSharedBuffer Value, const FIoHash* Hash);
-
-	/** An owned buffer containing the binary or object data. */
-	FSharedBuffer Buffer;
-	/** A field that is valid only for object attachments. */
-	FCbFieldView Object;
-	/** A hash of the attachment value. */
-	FIoHash Hash;
+	CORE_API FCbAttachment(const FCbObject& Value, const FIoHash* Hash);
+
+	struct FObjectValue
+	{
+		FCbObject Object;
+		FIoHash Hash;
+
+		FObjectValue(const FCbObject& InObject, const FIoHash& InHash)
+			: Object(InObject)
+			, Hash(InHash)
+		{
+		}
+
+		FObjectValue(FCbObject&& InObject, const FIoHash& InHash)
+			: Object(MoveTemp(InObject))
+			, Hash(InHash)
+		{
+		}
+	};
+
+	struct FBinaryValue
+	{
+		FCompositeBuffer Buffer;
+		FIoHash Hash;
+
+		FBinaryValue(const FCompositeBuffer& InBuffer)
+			: Buffer(InBuffer.MakeOwned())
+			, Hash(FIoHash::HashBuffer(InBuffer))
+		{
+		}
+
+		FBinaryValue(const FCompositeBuffer& InBuffer, const FIoHash& InHash)
+			: Buffer(InBuffer.MakeOwned())
+			, Hash(InHash)
+		{
+		}
+
+		FBinaryValue(FCompositeBuffer&& InBuffer)
+			: Buffer(MoveTemp(InBuffer))
+			, Hash(FIoHash::HashBuffer(Buffer))
+		{
+		}
+
+		FBinaryValue(FCompositeBuffer&& InBuffer, const FIoHash& InHash)
+			: Buffer(MoveTemp(InBuffer))
+			, Hash(InHash)
+		{
+		}
+
+		FBinaryValue(FSharedBuffer&& InBuffer, const FIoHash& InHash)
+			: Buffer(MoveTemp(InBuffer))
+			, Hash(InHash)
+		{
+		}
+	};
+
+	TVariant<TYPE_OF_NULLPTR, FObjectValue, FBinaryValue, FCompressedBuffer> Value;
 };
 
 /** Hashes attachments by their hash. Any discrepancy in type must be handled externally. */
@@ -349,3 +446,39 @@ private:
 };
 
 ///////////////////////////////////////////////////////////////////////////////////////////////////////////////////////
+
+const FCompositeBuffer& FCbAttachment::AsCompositeBinary() const
+{
+	if (const FBinaryValue* BinaryValue = Value.TryGet<FBinaryValue>())
+	{
+		return BinaryValue->Buffer;
+	}
+	return FCompositeBuffer::Null;
+}
+
+FSharedBuffer FCbAttachment::AsBinary() const
+{
+	if (const FBinaryValue* BinaryValue = Value.TryGet<FBinaryValue>())
+	{
+		return BinaryValue->Buffer.Flatten();
+	}
+	return FSharedBuffer();
+}
+
+const FCompressedBuffer& FCbAttachment::AsCompressedBinary() const
+{
+	if (const FCompressedBuffer* CompressedBuffer = Value.TryGet<FCompressedBuffer>())
+	{
+		return *CompressedBuffer;
+	}
+	return FCompressedBuffer::Null;
+}
+
+FCbObject FCbAttachment::AsObject() const
+{
+	if (const FObjectValue* ObjectValue = Value.TryGet<FObjectValue>())
+	{
+		return ObjectValue->Object;
+	}
+	return FCbObject();
+}
diff --git a/Engine/Source/Runtime/Core/Public/Serialization/CompactBinaryValidation.h b/Engine/Source/Runtime/Core/Public/Serialization/CompactBinaryValidation.h
--- a/Engine/Source/Runtime/Core/Public/Serialization/CompactBinaryValidation.h
+++ b/Engine/Source/Runtime/Core/Public/Serialization/CompactBinaryValidation.h
@@ -52,12 +52,17 @@ enum class ECbValidateMode : uint32
 	Padding                 = 1 << 3,
 
 	/**
-	 * Validate that a package or attachment has the expected fields and matches its saved hashes.
+	 * Validate that a package or attachment has the expected fields.
 	 */
 	Package                 = 1 << 4,
 
+	/**
+	 * Validate that a package or attachment matches its saved hashes.
+	 */
+	PackageHash             = 1 << 5,
+
 	/** Perform all validation described above. */
-	All                     = Default | Names | Format | Padding | Package,
+	All                     = Default | Names | Format | Padding | Package | PackageHash,
 };
 
 ENUM_CLASS_FLAGS(ECbValidateMode);
