Tasks: releasing TaskBody functor after its execution as it can contain captured resources messing with their lifetimes

diff --git a/Engine/Source/Runtime/Core/Public/Tasks/TaskPrivate.h b/Engine/Source/Runtime/Core/Public/Tasks/TaskPrivate.h
--- a/Engine/Source/Runtime/Core/Public/Tasks/TaskPrivate.h
+++ b/Engine/Source/Runtime/Core/Public/Tasks/TaskPrivate.h
@@ -536,6 +536,7 @@ namespace UE { namespace Tasks
 				StartPipeExecution();
 
 				TaskBody();
+				TaskBody.Reset();
 
 				FinishPipeExecution();
 				TaskTrace::Finished(GetTraceId());
