Change Stadia time to use Unix time again.

Functional tests discovered errors in the implementation. This way Improvements to Unix also improve Stadia without duplicated effort.

diff --git a/Engine/Source/Runtime/Core/Public/Linux/LinuxPlatformTime.h b/Engine/Source/Runtime/Core/Public/Linux/LinuxPlatformTime.h
--- a/Engine/Source/Runtime/Core/Public/Linux/LinuxPlatformTime.h
+++ b/Engine/Source/Runtime/Core/Public/Linux/LinuxPlatformTime.h
@@ -3,5 +3,3 @@
 #pragma once
 
 #include "Unix/UnixPlatformTime.h"
-
-typedef FUnixTime FPlatformTime;
\ No newline at end of file
diff --git a/Engine/Source/Runtime/Core/Public/Unix/UnixPlatformTime.h b/Engine/Source/Runtime/Core/Public/Unix/UnixPlatformTime.h
--- a/Engine/Source/Runtime/Core/Public/Unix/UnixPlatformTime.h
+++ b/Engine/Source/Runtime/Core/Public/Unix/UnixPlatformTime.h
@@ -75,3 +75,5 @@ private:
 	 */
 	static uint64 CallsPerSecondBenchmark(clockid_t BenchClockId, const char * BenchClockIdName);
 };
+
+typedef FUnixTime FPlatformTime;
\ No newline at end of file
