fixed TSan warning

diff --git a/Engine/Source/Runtime/Core/Public/Async/ParallelFor.h b/Engine/Source/Runtime/Core/Public/Async/ParallelFor.h
--- a/Engine/Source/Runtime/Core/Public/Async/ParallelFor.h
+++ b/Engine/Source/Runtime/Core/Public/Async/ParallelFor.h
@@ -17,6 +17,8 @@
 #include "Misc/App.h"
 #include "Misc/Fork.h"
 
+#include <atomic>
+
 extern CORE_API int32 GParallelForBackgroundYieldingTimeoutMs;
 
 // Flags controlling the ParallelFor's behavior.
@@ -54,13 +56,14 @@ namespace ParallelForImpl
 		FEvent* Event;
 		FThreadSafeCounter IndexToDo;
 		FThreadSafeCounter NumCompleted;
-		bool bExited;
+#if DO_CHECK
+		std::atomic<bool> bExited{ false };
+#endif
 		bool bTriggered;
 		bool bSaveLastBlockForMaster;
 		TParallelForData(int32 InTotalNum, int32 InNumThreads, bool bInSaveLastBlockForMaster, FunctionType InBody, EParallelForFlags Flags)
 			: Body(InBody)
 			, Event(FPlatformProcess::GetSynchEventFromPool(false))
-			, bExited(false)
 			, bTriggered(false)
 			, bSaveLastBlockForMaster(bInSaveLastBlockForMaster)
 		{
@@ -97,7 +100,7 @@ namespace ParallelForImpl
 		{
 			check(IndexToDo.GetValue() >= Num);
 			check(NumCompleted.GetValue() == Num);
-			check(bExited);
+			check(bExited.load(std::memory_order_relaxed));
 			FPlatformProcess::ReturnSynchEventToPool(Event);
 		}
 		bool Process(int32 TasksToSpawn, TSharedRef<TParallelForData, ESPMode::ThreadSafe>& Data, ENamedThreads::Type InDesiredThread, bool bMaster);
@@ -195,7 +198,7 @@ namespace ParallelForImpl
 				{
 					LocalBody(MyIndex * LocalBlockSize + LocalIndex);
 				}
-				checkSlow(!bExited);
+				checkSlow(!bExited.load(std::memory_order_relaxed));
 				int32 LocalNumCompleted = NumCompleted.Increment();
 				if (LocalNumCompleted == LocalNum)
 				{
@@ -306,7 +309,11 @@ namespace ParallelForImpl
 			check(!Data->bTriggered);
 		}
 		check(Data->NumCompleted.GetValue() == Data->Num);
-		Data->bExited = true;
+
+#if DO_CHECK
+		Data->bExited.store(true, std::memory_order_relaxed);
+#endif
+
 		// DoneEvent waits here if some other thread finishes the last item
 		// Data must live on until all of the tasks are cleared which might be long after this function exits
 	}
@@ -371,7 +378,11 @@ namespace ParallelForImpl
 			check(!Data->bTriggered);
 		}
 		check(Data->NumCompleted.GetValue() == Data->Num);
-		Data->bExited = true;
+
+#if DO_CHECK
+		Data->bExited.store(true, std::memory_order_relaxed);
+#endif
+
 		// Data must live on until all of the tasks are cleared which might be long after this function exits
 	}
 }
