Insights changes:
Make memory tag specs important so they don't get dropped before a trace connection is made.
Prevent Insights from crashing when opening a trace missing memory scopes or with missing scope specs.
Add a tooltip to symbol resolution noting the environment variable for symbol paths.
Allow Insights to capture full callstacks (including system modules) with the define UE_CALLSTACK_TRACE_FULL_CALLSTACKS or the command line argument -tracefullcallstacks on Windows.

diff --git a/Engine/Source/Developer/TraceInsights/Private/Insights/MemoryProfiler/Widgets/SMemAllocTableTreeView.cpp b/Engine/Source/Developer/TraceInsights/Private/Insights/MemoryProfiler/Widgets/SMemAllocTableTreeView.cpp
--- a/Engine/Source/Developer/TraceInsights/Private/Insights/MemoryProfiler/Widgets/SMemAllocTableTreeView.cpp
+++ b/Engine/Source/Developer/TraceInsights/Private/Insights/MemoryProfiler/Widgets/SMemAllocTableTreeView.cpp
@@ -826,6 +826,7 @@ TSharedPtr<SWidget> SMemAllocTableTreeView::ConstructFooter()
 		[
 			SNew(STextBlock)
 			.Text(this, &SMemAllocTableTreeView::GetSymbolResolutionStatus)
+			.ToolTipText(this, &SMemAllocTableTreeView::GetSymbolResolutionTooltip)
 		];
 }
 END_SLATE_FUNCTION_BUILD_OPTIMIZATION
@@ -859,6 +860,15 @@ FText SMemAllocTableTreeView::GetSymbolResolutionStatus() const
 
 ////////////////////////////////////////////////////////////////////////////////////////////////////
 
+FText SMemAllocTableTreeView::GetSymbolResolutionTooltip() const
+{
+	auto ModuleProvider = Session->ReadProvider<IModuleProvider>(FName("ModuleProvider"));
+
+	return FText::Format(LOCTEXT("SymbolResolutionEnvVarHelp", "Symbol paths can be provided with the environment variable {0}."), FText::FromString(TEXT("UE_INSIGHTS_SYMBOL_PATH")));
+}
+
+////////////////////////////////////////////////////////////////////////////////////////////////////
+
 FText SMemAllocTableTreeView::GetQueryInfo() const
 {
 	return QueryInfo;
diff --git a/Engine/Source/Developer/TraceInsights/Private/Insights/MemoryProfiler/Widgets/SMemAllocTableTreeView.h b/Engine/Source/Developer/TraceInsights/Private/Insights/MemoryProfiler/Widgets/SMemAllocTableTreeView.h
--- a/Engine/Source/Developer/TraceInsights/Private/Insights/MemoryProfiler/Widgets/SMemAllocTableTreeView.h
+++ b/Engine/Source/Developer/TraceInsights/Private/Insights/MemoryProfiler/Widgets/SMemAllocTableTreeView.h
@@ -112,6 +112,7 @@ private:
 	FReply OnCallstackViewClicked(bool bIsInverted);
 
 	FText GetSymbolResolutionStatus() const;
+	FText GetSymbolResolutionTooltip() const;
 	FText GetQueryInfo() const;
 	FText GetQueryInfoTooltip() const;
 
diff --git a/Engine/Source/Developer/TraceServices/Private/Model/AllocationsProvider.cpp b/Engine/Source/Developer/TraceServices/Private/Model/AllocationsProvider.cpp
--- a/Engine/Source/Developer/TraceServices/Private/Model/AllocationsProvider.cpp
+++ b/Engine/Source/Developer/TraceServices/Private/Model/AllocationsProvider.cpp
@@ -174,7 +174,7 @@ uint32 FTagTracker::GetCurrentTag(uint32 InThreadId, uint8 InTracker) const
 const TCHAR* FTagTracker::GetTagString(uint32 InTag) const
 {
 	const TagEntry* Entry = TagMap.Find(InTag);
-	return Entry ? Entry->Display : nullptr;
+	return Entry ? Entry->Display : TEXT("");
 }
 
 ////////////////////////////////////////////////////////////////////////////////////////////////////
diff --git a/Engine/Source/Runtime/Core/Private/ProfilingDebugging/Microsoft/WindowsCallstackTrace.cpp b/Engine/Source/Runtime/Core/Private/ProfilingDebugging/Microsoft/WindowsCallstackTrace.cpp
--- a/Engine/Source/Runtime/Core/Private/ProfilingDebugging/Microsoft/WindowsCallstackTrace.cpp
+++ b/Engine/Source/Runtime/Core/Private/ProfilingDebugging/Microsoft/WindowsCallstackTrace.cpp
@@ -25,6 +25,10 @@
 
 #include <atomic>
 
+#ifndef UE_CALLSTACK_TRACE_FULL_CALLSTACKS
+	#define UE_CALLSTACK_TRACE_FULL_CALLSTACKS 0
+#endif
+
 // 0=off, 1=stats, 2=validation, 3=truth_compare
 #define BACKTRACE_DBGLVL 0
 
@@ -196,6 +200,7 @@ private:
 
 ////////////////////////////////////////////////////////////////////////////////
 FBacktracer* FBacktracer::Instance = nullptr;
+static bool GFullBacktraces = UE_CALLSTACK_TRACE_FULL_CALLSTACKS;
 
 ////////////////////////////////////////////////////////////////////////////////
 FBacktracer::FBacktracer(FMalloc* InMalloc)
@@ -240,11 +245,14 @@ void FBacktracer::AddModule(UPTRINT ModuleBase, const TCHAR* Name)
 	const auto* NtHeader = (IMAGE_NT_HEADERS*)(ModuleBase + DosHeader->e_lfanew);
 	const IMAGE_FILE_HEADER* FileHeader = &(NtHeader->FileHeader);
 
-	int32 NameLen = FCString::Strlen(Name);
-	if (!(NameLen > 4 && FCString::Strcmp(Name + NameLen - 4, TEXT(".exe")) == 0) && // we want always load an .exe file
-		(FCString::Strfind(Name, TEXT("Binaries")) == nullptr || FCString::Strfind(Name, TEXT("ThirdParty")) != nullptr))
+	if(!GFullBacktraces)
 	{
-		return;
+		int32 NameLen = FCString::Strlen(Name);
+		if (!(NameLen > 4 && FCString::Strcmp(Name + NameLen - 4, TEXT(".exe")) == 0) && // we want always load an .exe file
+			(FCString::Strfind(Name, TEXT("Binaries")) == nullptr || FCString::Strfind(Name, TEXT("ThirdParty")) != nullptr))
+		{
+			return;
+		}
 	}
 
 	uint32 NumSections = FileHeader->NumberOfSections;
@@ -713,6 +721,12 @@ void Backtracer_Create(FMalloc* Malloc)
 	void* Alloc = Malloc->Malloc(sizeof(FBacktracer), alignof(FBacktracer));
 	new (Alloc) FBacktracer(Malloc);
 
+	const TCHAR* CmdLine = ::GetCommandLineW();
+	if (const TCHAR* TraceArg = FCString::Stristr(CmdLine, TEXT("-tracefullcallstacks")))
+	{
+		GFullBacktraces = true;
+	}
+
 	Modules_Create(Malloc);
 	Modules_Subscribe(
 		[] (bool bLoad, void* Module, const TCHAR* Name)
diff --git a/Engine/Source/Runtime/Core/Private/ProfilingDebugging/TagTrace.cpp b/Engine/Source/Runtime/Core/Private/ProfilingDebugging/TagTrace.cpp
--- a/Engine/Source/Runtime/Core/Private/ProfilingDebugging/TagTrace.cpp
+++ b/Engine/Source/Runtime/Core/Private/ProfilingDebugging/TagTrace.cpp
@@ -16,7 +16,7 @@
 
 UE_TRACE_CHANNEL_EXTERN(MemAllocChannel);
 
-UE_TRACE_EVENT_BEGIN(Memory, TagSpec)
+UE_TRACE_EVENT_BEGIN(Memory, TagSpec, Important|NoSync)
 	UE_TRACE_EVENT_FIELD(int32, Tag)
 	UE_TRACE_EVENT_FIELD(int32, Parent)
 	UE_TRACE_EVENT_FIELD(UE::Trace::AnsiString, Display)
@@ -177,7 +177,7 @@ void FTagTrace::AnnounceGenericTags()
 	#define TRACE_TAG_SPEC(Enum,Str,Stat,Group,ParentTag)\
 	{\
 		const uint32 DisplayLen = FCStringAnsi::Strlen(Str);\
-		UE_TRACE_LOG(Memory, TagSpec, MemAllocChannel)\
+		UE_TRACE_LOG(Memory, TagSpec, MemAllocChannel, DisplayLen * sizeof(ANSICHAR))\
 			<< TagSpec.Tag((int32) ELLMTag::Enum)\
 			<< TagSpec.Parent((int32) ParentTag)\
 			<< TagSpec.Display(Str, DisplayLen);\
@@ -240,7 +240,7 @@ int32 FTagTrace::AnnounceFNameTag(const FName& Name)
 int32 FTagTrace::AnnounceCustomTag(int32 Tag, int32 ParentTag, const ANSICHAR* Display)
 {		
 	const uint32 DisplayLen = FCStringAnsi::Strlen(Display);
-	UE_TRACE_LOG(Memory, TagSpec, MemAllocChannel)
+	UE_TRACE_LOG(Memory, TagSpec, MemAllocChannel, DisplayLen * sizeof(ANSICHAR))
 		<< TagSpec.Tag(Tag)
 		<< TagSpec.Parent(ParentTag)
 		<< TagSpec.Display(Display, DisplayLen);
