Set insights trace memory tag on io requests like we do with LLM tag.

#[robomerge][bot1] Dev-EngineMerge

diff --git a/Engine/Source/Runtime/Core/Private/IO/IoDispatcher.cpp b/Engine/Source/Runtime/Core/Private/IO/IoDispatcher.cpp
--- a/Engine/Source/Runtime/Core/Private/IO/IoDispatcher.cpp
+++ b/Engine/Source/Runtime/Core/Private/IO/IoDispatcher.cpp
@@ -919,6 +919,7 @@ void FIoRequestImpl::CreateBuffer(uint64 Size)
 	else
 	{
 		LLM_SCOPE(InheritedLLMTag);
+		UE_MEMSCOPE(InheritedTraceTag, ELLMTracker::Default);
 		TRACE_CPUPROFILER_EVENT_SCOPE(AllocMemoryForRequest);
 		Buffer.Emplace(Size);
 	}
diff --git a/Engine/Source/Runtime/Core/Public/IO/IoDispatcherBackend.h b/Engine/Source/Runtime/Core/Public/IO/IoDispatcherBackend.h
--- a/Engine/Source/Runtime/Core/Public/IO/IoDispatcherBackend.h
+++ b/Engine/Source/Runtime/Core/Public/IO/IoDispatcherBackend.h
@@ -5,6 +5,7 @@
 #include "IoDispatcher.h"
 #include "HAL/LowLevelMemTracker.h"
 #include "Misc/Optional.h"
+#include "ProfilingDebugging/TagTrace.h"
 
 class FIoRequestImpl
 {
@@ -12,6 +13,9 @@ public:
 	FIoRequestImpl* NextRequest = nullptr;
 	void* BackendData = nullptr;
 	LLM(const UE::LLMPrivate::FTagData* InheritedLLMTag);
+#if USE_MEMORY_TRACE_TAGS
+	int32 InheritedTraceTag;
+#endif
 	FIoChunkId ChunkId;
 	FIoReadOptions Options;
 	int32 Priority = 0;
@@ -20,6 +24,9 @@ public:
 		: Dispatcher(InDispatcher)
 	{
 		LLM(InheritedLLMTag = FLowLevelMemTracker::bIsDisabled ? nullptr : FLowLevelMemTracker::Get().GetActiveTagData(ELLMTracker::Default));
+#if USE_MEMORY_TRACE_TAGS
+		InheritedTraceTag = MemoryTrace_GetActiveTag();
+#endif
 	}
 
 	bool IsCancelled() const
