-traceshowstore to pop up the store server in its own console window. Useful for debugging.

#rnx

diff --git a/Engine/Source/Runtime/Core/Private/ProfilingDebugging/TraceAuxiliary.cpp b/Engine/Source/Runtime/Core/Private/ProfilingDebugging/TraceAuxiliary.cpp
--- a/Engine/Source/Runtime/Core/Private/ProfilingDebugging/TraceAuxiliary.cpp
+++ b/Engine/Source/Runtime/Core/Private/ProfilingDebugging/TraceAuxiliary.cpp
@@ -499,7 +499,7 @@ static std::atomic<int32> GUnrealTraceLaunched; // = 0;
 
 ////////////////////////////////////////////////////////////////////////////////
 #if PLATFORM_WINDOWS
-static void LaunchUnrealTraceInternal()
+static void LaunchUnrealTraceInternal(const TCHAR* CommandLine)
 {
 	if (GUnrealTraceLaunched.load(std::memory_order_relaxed))
 	{
@@ -512,11 +512,19 @@ static void LaunchUnrealTraceInternal()
 	CommandLine += "/Binaries/Win64/UnrealTraceServer.exe\"";
 	CommandLine += " fork";
 
-	uint32 Flags = CREATE_BREAKAWAY_FROM_JOB|DETACHED_PROCESS;
+	uint32 CreateProcFlags = CREATE_BREAKAWAY_FROM_JOB;
+	if (FParse::Param(CommandLine, TEXT("traceshowstore")))
+	{
+		CreateProcFlags |= DETACHED_PROCESS;
+	}
+	else
+	{
+		CreateProcFlags |= CREATE_NO_WINDOW;
+	}
 	STARTUPINFOW StartupInfo = { sizeof(STARTUPINFOW) };
 	PROCESS_INFORMATION ProcessInfo = {};
 	BOOL bOk = CreateProcessW(nullptr, LPWSTR(*CommandLine), nullptr, nullptr,
-		false, Flags, nullptr, nullptr, &StartupInfo, &ProcessInfo);
+		false, CreateProcFlags, nullptr, nullptr, &StartupInfo, &ProcessInfo);
 
 	if (!bOk)
 	{
@@ -550,7 +558,7 @@ static void LaunchUnrealTraceInternal()
 
 ////////////////////////////////////////////////////////////////////////////////
 #if PLATFORM_UNIX || PLATFORM_MAC
-static void LaunchUnrealTraceInternal()
+static void LaunchUnrealTraceInternal(const TCHAR* CommandLine)
 {
 	/* nop */
 }
@@ -580,7 +588,7 @@ void FTraceAuxiliary::Initialize(const TCHAR* CommandLine)
 	if (!FParse::Param(CommandLine, TEXT("notraceserver")))
 	{
 		TRACE_CPUPROFILER_EVENT_SCOPE(FTraceAux_LaunchUnrealTrace);
-		LaunchUnrealTraceInternal();
+		LaunchUnrealTraceInternal(CommandLine);
 	}
 #endif
 
