Swapped to a string-builder and adding missing TEXT() decoration around some literals

#rnx

diff --git a/Engine/Source/Runtime/Core/Private/ProfilingDebugging/TraceAuxiliary.cpp b/Engine/Source/Runtime/Core/Private/ProfilingDebugging/TraceAuxiliary.cpp
--- a/Engine/Source/Runtime/Core/Private/ProfilingDebugging/TraceAuxiliary.cpp
+++ b/Engine/Source/Runtime/Core/Private/ProfilingDebugging/TraceAuxiliary.cpp
@@ -507,10 +507,11 @@ static void LaunchUnrealTraceInternal(const TCHAR* CommandLine)
 		return;
 	}
 
-	FString CommandLine = "\"";
-	CommandLine += FPaths::EngineDir();
-	CommandLine += "/Binaries/Win64/UnrealTraceServer.exe\"";
-	CommandLine += " fork";
+	TWideStringBuilder<MAX_PATH + 32> CreateProcArgs;
+	CreateProcArgs << "\"";
+	CreateProcArgs << FPaths::EngineDir();
+	CreateProcArgs << TEXT("/Binaries/Win64/UnrealTraceServer.exe\"");
+	CreateProcArgs << TEXT(" fork");
 
 	uint32 CreateProcFlags = CREATE_BREAKAWAY_FROM_JOB;
 	if (FParse::Param(CommandLine, TEXT("traceshowstore")))
@@ -523,12 +524,12 @@ static void LaunchUnrealTraceInternal(const TCHAR* CommandLine)
 	}
 	STARTUPINFOW StartupInfo = { sizeof(STARTUPINFOW) };
 	PROCESS_INFORMATION ProcessInfo = {};
-	BOOL bOk = CreateProcessW(nullptr, LPWSTR(*CommandLine), nullptr, nullptr,
+	BOOL bOk = CreateProcessW(nullptr, LPWSTR(*CreateProcArgs), nullptr, nullptr,
 		false, CreateProcFlags, nullptr, nullptr, &StartupInfo, &ProcessInfo);
 
 	if (!bOk)
 	{
-		UE_LOG(LogCore, Display, TEXT("UnrealTraceServer: Unable to launch the trace store with '%s' (%08x)"), *CommandLine, GetLastError());
+		UE_LOG(LogCore, Display, TEXT("UnrealTraceServer: Unable to launch the trace store with '%s' (%08x)"), *CreateProcArgs, GetLastError());
 		return;
 	}
 
