Clarified that the shwon return value is in hexadecimal

#rnx

diff --git a/Engine/Source/Runtime/Core/Private/ProfilingDebugging/TraceAuxiliary.cpp b/Engine/Source/Runtime/Core/Private/ProfilingDebugging/TraceAuxiliary.cpp
--- a/Engine/Source/Runtime/Core/Private/ProfilingDebugging/TraceAuxiliary.cpp
+++ b/Engine/Source/Runtime/Core/Private/ProfilingDebugging/TraceAuxiliary.cpp
@@ -543,7 +543,7 @@ static void LaunchUnrealTraceInternal(const TCHAR* CommandLine)
 		GetExitCodeProcess(ProcessInfo.hProcess, &ExitCode);
 		if (ExitCode)
 		{
-			UE_LOG(LogCore, Warning, TEXT("UnrealTraceServer: Trace store returned an error (%08x)"), ExitCode);
+			UE_LOG(LogCore, Warning, TEXT("UnrealTraceServer: Trace store returned an error (0x%08x)"), ExitCode);
 		}
 		else
 		{
