An implementation of LaunchUnrealTrace() for POSIX platforms (currently disabled).

#rnx

diff --git a/Engine/Source/Runtime/Core/Private/ProfilingDebugging/TraceAuxiliary.cpp b/Engine/Source/Runtime/Core/Private/ProfilingDebugging/TraceAuxiliary.cpp
--- a/Engine/Source/Runtime/Core/Private/ProfilingDebugging/TraceAuxiliary.cpp
+++ b/Engine/Source/Runtime/Core/Private/ProfilingDebugging/TraceAuxiliary.cpp
@@ -562,6 +562,68 @@ static void LaunchUnrealTraceInternal(const TCHAR* CommandLine)
 static void LaunchUnrealTraceInternal(const TCHAR* CommandLine)
 {
 	/* nop */
+
+#if 0
+	if (GUnrealTraceLaunched.load(std::memory_order_relaxed))
+	{
+		UE_LOG(LogCore, Log, TEXT("UnrealTraceServer: Trace store already started"));
+		return;
+	}
+
+	TAnsiStringBuilder<320> BinPath;
+	BinPath << TCHAR_TO_UTF8(*FPaths::EngineDir());
+#if PLATFORM_UNIX
+	BinPath << "Binaries/Linux/UnrealTraceServer";
+#elif PLATFORM_MAC
+	BinPath << "Binaries/Mac/UnrealTraceServer";
+#endif
+
+	if (access(*BinPath, F_OK) < 0)
+	{
+		UE_LOG(LogCore, Display, TEXT("UnrealTraceServer: Binary not found (%s)"), ANSI_TO_TCHAR(*BinPath));
+		return;
+	}
+
+	TAnsiStringBuilder<64> ForkArg;
+	ForkArg << "fork";
+
+	pid_t UtsPid = vfork();
+	if (UtsPid < 0)
+	{
+		UE_LOG(LogCore, Display, TEXT("UnrealTraceServer: Unable to fork (errno: %d)"), errno);
+		return;
+	}
+	else if (UtsPid == 0)
+	{
+		char* Args[] = { BinPath.GetData(), ForkArg.GetData(), nullptr };
+		extern char** environ;
+		execve(*BinPath, Args, environ);
+		_exit(0x80 | (errno & 0x7f));
+	}
+	
+	int32 WaitStatus = 0;
+	do
+	{
+		int32 WaitRet = waitpid(UtsPid, &WaitStatus, 0);
+		if (WaitRet < 0)
+		{
+			UE_LOG(LogCore, Display, TEXT("UnrealTraceServer: waitpid() error; (errno: %d)"), errno);
+			return;
+		}
+	}
+	while (!WIFEXITED(WaitStatus));
+
+	int32 UtsRet = WEXITSTATUS(WaitStatus);
+	if (UtsRet)
+	{
+		UE_LOG(LogCore, Display, TEXT("UnrealTraceServer: Trace store returned an error (0x%08x)"), UtsRet);
+	}
+	else
+	{
+		UE_LOG(LogCore, Log, TEXT("UnrealTraceServer: Trace store launch successful"));
+		GUnrealTraceLaunched.fetch_add(1, std::memory_order_relaxed);
+	}
+#endif // 0
 }
 #endif // PLATFORM_UNIX/MAC
 #endif // WITH_UNREAL_TRACE_LAUNCH
