IoStore compile fix for some versions of clang

#rnx

diff --git a/Engine/Source/Runtime/Core/Private/IO/IoStore.cpp b/Engine/Source/Runtime/Core/Private/IO/IoStore.cpp
--- a/Engine/Source/Runtime/Core/Private/IO/IoStore.cpp
+++ b/Engine/Source/Runtime/Core/Private/IO/IoStore.cpp
@@ -754,7 +754,7 @@ public:
 		TArray<TArray<int32>> Buckets;
 		Buckets.SetNum(ChunkCount);
 
-		TBitArray FreeSlots(true, ChunkCount);
+		TBitArray<> FreeSlots(true, ChunkCount);
 		// Put each chunk in a bucket, each bucket contains the chunk ids that have colliding hashes
 		for (uint32 ChunkIndex = 0; ChunkIndex < ChunkCount; ++ChunkIndex)
 		{
@@ -789,7 +789,7 @@ public:
 			static constexpr uint32 MaxIterations = UE_ARRAY_COUNT(Primes);
 
 			uint32 PrimeIndex = 0;
-			TBitArray BucketUsedSlots(false, ChunkCount);
+			TBitArray<> BucketUsedSlots(false, ChunkCount);
 			int32 IndexInBucket = 0;
 			while (IndexInBucket < Bucket.Num())
 			{
