CompositeBuffer: Disabled the constructor for invalid parameter types

#rnx

diff --git a/Engine/Source/Runtime/Core/Private/Tests/Memory/CompositeBufferTest.cpp b/Engine/Source/Runtime/Core/Private/Tests/Memory/CompositeBufferTest.cpp
--- a/Engine/Source/Runtime/Core/Private/Tests/Memory/CompositeBufferTest.cpp
+++ b/Engine/Source/Runtime/Core/Private/Tests/Memory/CompositeBufferTest.cpp
@@ -3,9 +3,24 @@
 #include "Memory/CompositeBuffer.h"
 
 #include "Misc/AutomationTest.h"
+#include <type_traits>
 
 #if WITH_DEV_AUTOMATION_TESTS
 
+static_assert(std::is_constructible_v<FCompositeBuffer>);
+static_assert(std::is_constructible_v<FCompositeBuffer, FSharedBuffer&&>);
+static_assert(std::is_constructible_v<FCompositeBuffer, FCompositeBuffer&&>);
+static_assert(std::is_constructible_v<FCompositeBuffer, TArray<FSharedBuffer>&&>);
+static_assert(std::is_constructible_v<FCompositeBuffer, TArray<const FSharedBuffer>&&>);
+static_assert(std::is_constructible_v<FCompositeBuffer, const FSharedBuffer&>);
+static_assert(std::is_constructible_v<FCompositeBuffer, const FCompositeBuffer&>);
+static_assert(std::is_constructible_v<FCompositeBuffer, const TArray<FSharedBuffer>&>);
+static_assert(std::is_constructible_v<FCompositeBuffer, const TArray<const FSharedBuffer>&>);
+static_assert(std::is_constructible_v<FCompositeBuffer, const FSharedBuffer&, const FCompositeBuffer&, const TArray<FSharedBuffer>&, const TArray<const FSharedBuffer>&>);
+
+static_assert(!std::is_constructible_v<FCompositeBuffer, FMemoryView>);
+static_assert(!std::is_constructible_v<FCompositeBuffer, const uint8[4]>);
+
 IMPLEMENT_SIMPLE_AUTOMATION_TEST(FCompositeBufferTest, "System.Core.Memory.CompositeBuffer", EAutomationTestFlags::ApplicationContextMask | EAutomationTestFlags::SmokeFilter)
 bool FCompositeBufferTest::RunTest(const FString& Parameters)
 {
diff --git a/Engine/Source/Runtime/Core/Public/Memory/CompositeBuffer.h b/Engine/Source/Runtime/Core/Public/Memory/CompositeBuffer.h
--- a/Engine/Source/Runtime/Core/Public/Memory/CompositeBuffer.h
+++ b/Engine/Source/Runtime/Core/Public/Memory/CompositeBuffer.h
@@ -6,6 +6,7 @@
 #include "Containers/ArrayView.h"
 #include "Memory/MemoryView.h"
 #include "Memory/SharedBuffer.h"
+#include "Templates/UnrealTemplate.h"
 
 template <typename FuncType> class TFunctionRef;
 
@@ -19,19 +20,35 @@ template <typename FuncType> class TFunctionRef;
  */
 class FCompositeBuffer
 {
+	using FSegmentArray = TArray<FSharedBuffer, TInlineAllocator<1>>;
+
+	static inline int32 GetBufferCount(const FCompositeBuffer& Buffer) { return Buffer.Segments.Num(); }
+	static inline void AppendBuffers(FSegmentArray& Array, const FCompositeBuffer& Buffer) { Array.Append(Buffer.Segments); }
+	static inline void AppendBuffers(FSegmentArray& Array, FCompositeBuffer&& Buffer) { Array.Append(MoveTemp(Buffer.Segments)); }
+
+	static inline int32 GetBufferCount(const FSharedBuffer& Buffer) { return 1; }
+	static inline void AppendBuffers(FSegmentArray& Array, const FSharedBuffer& Buffer) { Array.Add(Buffer); }
+	static inline void AppendBuffers(FSegmentArray& Array, FSharedBuffer&& Buffer) { Array.Add(MoveTemp(Buffer)); }
+
+	template <typename BufferType, decltype(DeclVal<TArray<FSharedBuffer>>().Append(DeclVal<BufferType>()))* = nullptr>
+	static inline int32 GetBufferCount(BufferType&& Buffer) { return GetNum(Buffer); }
+	template <typename BufferType, decltype(DeclVal<TArray<FSharedBuffer>>().Append(DeclVal<BufferType>()))* = nullptr>
+	static inline void AppendBuffers(FSegmentArray& Array, BufferType&& Buffer) { Array.Append(Forward<BufferType>(Buffer)); }
+
 public:
 	/**
 	 * Construct a composite buffer by concatenating the buffers. Does not enforce ownership.
 	 *
 	 * Buffer parameters may be FSharedBuffer, FCompositeBuffer, or TArray<FSharedBuffer>.
 	 */
-	template <typename... BufferTypes>
+	template <typename... BufferTypes,
+		decltype((AppendBuffers(DeclVal<FSegmentArray&>(), DeclVal<BufferTypes>()), ..., int8()))* = nullptr>
 	inline explicit FCompositeBuffer(BufferTypes&&... Buffers)
 	{
 		if constexpr (sizeof...(Buffers) > 0)
 		{
 			Segments.Reserve((GetBufferCount(Forward<BufferTypes>(Buffers)) + ...));
-			(AppendBuffers(Forward<BufferTypes>(Buffers)), ...);
+			(AppendBuffers(Segments, Forward<BufferTypes>(Buffers)), ...);
 			Segments.RemoveAll(&FSharedBuffer::IsNull);
 		}
 	}
@@ -100,21 +117,7 @@ public:
 	static const FCompositeBuffer Null;
 
 private:
-	static inline int32 GetBufferCount(const FCompositeBuffer& Buffer) { return Buffer.Segments.Num(); }
-	inline void AppendBuffers(const FCompositeBuffer& Buffer) { Segments.Append(Buffer.Segments); }
-	inline void AppendBuffers(FCompositeBuffer&& Buffer) { Segments.Append(MoveTemp(Buffer.Segments)); }
-
-	static inline int32 GetBufferCount(const FSharedBuffer& Buffer) { return 1; }
-	inline void AppendBuffers(const FSharedBuffer& Buffer) { Segments.Add(Buffer); }
-	inline void AppendBuffers(FSharedBuffer&& Buffer) { Segments.Add(MoveTemp(Buffer)); }
-
-	template <typename BufferType, decltype(std::declval<TArray<FSharedBuffer>>().Append(std::declval<BufferType>()))* = nullptr>
-	static inline int32 GetBufferCount(BufferType&& Buffer) { return GetNum(Buffer); }
-	template <typename BufferType, decltype(std::declval<TArray<FSharedBuffer>>().Append(std::declval<BufferType>()))* = nullptr>
-	inline void AppendBuffers(BufferType&& Buffer) { Segments.Append(Forward<BufferType>(Buffer)); }
-
-private:
-	TArray<FSharedBuffer, TInlineAllocator<1>> Segments;
+	FSegmentArray Segments;
 };
 
 inline const FCompositeBuffer FCompositeBuffer::Null;
