CompactBinary: Fixed attachment constructors missing calls to MakeOwned, and attachment loading that was too permissive

#rnx

diff --git a/Engine/Source/Runtime/Core/Private/Serialization/CompactBinaryPackage.cpp b/Engine/Source/Runtime/Core/Private/Serialization/CompactBinaryPackage.cpp
--- a/Engine/Source/Runtime/Core/Private/Serialization/CompactBinaryPackage.cpp
+++ b/Engine/Source/Runtime/Core/Private/Serialization/CompactBinaryPackage.cpp
@@ -10,86 +10,76 @@
 
 ///////////////////////////////////////////////////////////////////////////////////////////////////////////////////////
 
-static bool TryLoad_ArchiveFieldIntoAttachment(FCbAttachment& TargetAttachment, FCbField&& Field, FArchive& Ar, FCbBufferAllocator Allocator)
+static bool TryLoadAttachmentFromArchive(FCbAttachment& OutAttachment, FCbField&& Field, FArchive& Ar, FCbBufferAllocator Allocator)
 {
-	if (const FCbObjectView ObjectView = Field.AsObjectView(); !Field.HasError())
+	if (const FCbObjectView Object = Field.AsObjectView(); !Field.HasError())
 	{
-		// Is a null object or object not prefixed with a precomputed hash value
-		TargetAttachment = FCbAttachment(FCbObject(ObjectView, MoveTemp(Field)), ObjectView.GetHash());
+		// Empty Object
+		if (!Object)
+		{
+			OutAttachment = FCbAttachment(FCbObject(), FCbObject().GetHash());
+			return true;
+		}
 	}
-	else if (const FIoHash ObjectAttachmentHash = Field.AsObjectAttachment(); !Field.HasError())
+	else if (const FIoHash ObjectHash = Field.AsObjectAttachment(); !Field.HasError())
 	{
-		// Is an object
+		// Object
 		Field = LoadCompactBinary(Ar, Allocator);
-		if (!Field.IsObject())
+		if (Field.IsObject())
 		{
-			return false;
+			OutAttachment = FCbAttachment(MoveTemp(Field).AsObject(), ObjectHash);
+			return true;
 		}
-		TargetAttachment = FCbAttachment(MoveTemp(Field).AsObject(), ObjectAttachmentHash);
 	}
-	else if (const FIoHash BinaryAttachmentHash = Field.AsBinaryAttachment(); !Field.HasError())
+	else if (const FIoHash BinaryHash = Field.AsBinaryAttachment(); !Field.HasError())
 	{
-		// Is an uncompressed binary blob
+		// Raw Binary
 		Field = LoadCompactBinary(Ar, Allocator);
-		FSharedBuffer Buffer = Field.AsBinary();
-		if (Field.HasError())
+		if (FSharedBuffer Buffer = Field.AsBinary(); !Field.HasError())
 		{
-			return false;
+			OutAttachment = FCbAttachment(FCompositeBuffer(MoveTemp(Buffer)), BinaryHash);
+			return true;
 		}
-		TargetAttachment = FCbAttachment(FCompositeBuffer(Buffer), BinaryAttachmentHash);
 	}
-	else if (FSharedBuffer Buffer = Field.AsBinary(); !Field.HasError())
+	else if (Field.IsBinary())
 	{
-		if (Buffer.GetSize() > 0)
+		FSharedBuffer Buffer = MoveTemp(Field).AsBinary();
+		if (Buffer.GetSize() == 0)
 		{
-			// Is a compressed binary blob
-			TargetAttachment = FCbAttachment(FCompressedBuffer::FromCompressed(MoveTemp(Buffer)));
+			// Empty Raw Binary
+			OutAttachment = FCbAttachment(FCompositeBuffer(MoveTemp(Buffer)));
+			return true;
 		}
-		else
+		else if (FCompressedBuffer CompressedBuffer{FCompressedBuffer::FromCompressed(MoveTemp(Buffer))})
 		{
-			// Is an uncompressed empty binary blob
-			TargetAttachment = FCbAttachment(FCompositeBuffer(Buffer), FIoHash::HashBuffer(nullptr, 0));
+			// Compressed Binary
+			OutAttachment = FCbAttachment(MoveTemp(CompressedBuffer));
+			return true;
 		}
 	}
-	else
-	{
-		return false;
-	}
-
-	return true;
+	return false;
 }
 
-FCbAttachment::FCbAttachment(const FCbObject& InValue, const FIoHash* const InHash)
+FCbAttachment::FObjectValue::FObjectValue(FCbObject&& InObject, const FIoHash* const InHash)
 {
 	FMemoryView View;
-	if (!InValue.IsOwned() || !InValue.TryGetView(View))
+	if (!InObject.IsOwned() || !InObject.TryGetView(View))
 	{
-		FCbObject ClonedValue = FCbObject::Clone(InValue);
-
-		if (InHash)
-		{
-			checkSlow(*InHash == ClonedValue.GetHash());
-			Value.Emplace<FObjectValue>(MoveTemp(ClonedValue), *InHash);
-		}
-		else
-		{
-			FIoHash ClonedValueHash = ClonedValue.GetHash();
-			Value.Emplace<FObjectValue>(MoveTemp(ClonedValue), ClonedValueHash);
-		}
+		Object = FCbObject::Clone(InObject);
 	}
 	else
 	{
-		if (InHash)
-		{
-			Value.Emplace<FObjectValue>(InValue, *InHash);
-			checkSlow(*InHash == InValue.GetHash());
-		}
-		else
-		{
-			Value.Emplace<FObjectValue>(InValue, InValue.GetHash());
-		}
+		Object = MoveTemp(InObject);
+	}
+	if (InHash)
+	{
+		checkSlow(*InHash == Object.GetHash());
+		Hash = *InHash;
+	}
+	else
+	{
+		Hash = Object.GetHash();
 	}
-
 }
 
 FIoHash FCbAttachment::GetHash() const
@@ -114,63 +104,63 @@ FIoHash FCbAttachment::GetHash() const
 
 bool FCbAttachment::TryLoad(FCbFieldIterator& Fields)
 {
-	if (const FCbObjectView ObjectView = Fields.AsObjectView(); !Fields.HasError())
+	if (const FCbObjectView Object = Fields.AsObjectView(); !Fields.HasError())
 	{
-		// Is a null object or object not prefixed with a precomputed hash value
-		Value.Emplace<FObjectValue>(FCbObject(ObjectView, Fields.GetOuterBuffer()), ObjectView.GetHash());
-		++Fields;
+		// Empty Object
+		if (!Object)
+		{
+			Value.Emplace<FObjectValue>(FCbObject(), nullptr);
+			++Fields;
+			return true;
+		}
 	}
-	else if (const FIoHash ObjectAttachmentHash = Fields.AsObjectAttachment(); !Fields.HasError())
+	else if (const FIoHash ObjectHash = Fields.AsObjectAttachment(); !Fields.HasError())
 	{
-		// Is an object
+		// Object
 		++Fields;
-		const FCbObjectView InnerObjectView = Fields.AsObjectView();
-		if (Fields.HasError())
+		if (Fields.IsObject())
 		{
-			return false;
+			Value.Emplace<FObjectValue>(Fields.AsObject(), &ObjectHash);
+			++Fields;
+			return true;
 		}
-		Value.Emplace<FObjectValue>(FCbObject(InnerObjectView, Fields.GetOuterBuffer()), ObjectAttachmentHash);
-		++Fields;
 	}
-	else if (const FIoHash BinaryAttachmentHash = Fields.AsBinaryAttachment(); !Fields.HasError())
+	else if (const FIoHash BinaryHash = Fields.AsBinaryAttachment(); !Fields.HasError())
 	{
-		// Is an uncompressed binary blob
+		// Raw Binary
 		++Fields;
-		FMemoryView BinaryView = Fields.AsBinaryView();
-		if (Fields.HasError())
+		if (FSharedBuffer Buffer = Fields.AsBinary(); !Fields.HasError())
 		{
-			return false;
+			Value.Emplace<FBinaryValue>(MoveTemp(Buffer), BinaryHash);
+			++Fields;
+			return true;
 		}
-		Value.Emplace<FBinaryValue>(FSharedBuffer::MakeView(BinaryView, Fields.GetOuterBuffer()), BinaryAttachmentHash);
-		++Fields;
 	}
-	else if (FMemoryView BinaryView = Fields.AsBinaryView(); !Fields.HasError())
+	else if (Fields.IsBinary())
 	{
-		if (BinaryView.GetSize() > 0)
+		FSharedBuffer Buffer = Fields.AsBinary();
+		if (Buffer.GetSize() == 0)
 		{
-			// Is a compressed binary blob
-			Value.Emplace<FCompressedBuffer>(FCompressedBuffer::FromCompressed(FSharedBuffer::MakeView(BinaryView, Fields.GetOuterBuffer())).MakeOwned());
+			// Empty Raw Binary
+			Value.Emplace<FBinaryValue>(MoveTemp(Buffer));
 			++Fields;
+			return true;
 		}
-		else
+		else if (FCompressedBuffer CompressedBuffer{FCompressedBuffer::FromCompressed(MoveTemp(Buffer))})
 		{
-			// Is an uncompressed empty binary blob
-			Value.Emplace<FBinaryValue>(FSharedBuffer::MakeView(BinaryView, Fields.GetOuterBuffer()), FIoHash::HashBuffer(nullptr, 0));
+			// Compressed Binary
+			Value.Emplace<FCompressedBuffer>(MoveTemp(CompressedBuffer));
 			++Fields;
+			return true;
 		}
 	}
-	else
-	{
-		return false;
-	}
-
-	return true;
+	return false;
 }
 
 bool FCbAttachment::TryLoad(FArchive& Ar, FCbBufferAllocator Allocator)
 {
 	FCbField Field = LoadCompactBinary(Ar, Allocator);
-	return TryLoad_ArchiveFieldIntoAttachment(*this, MoveTemp(Field), Ar, Allocator);
+	return TryLoadAttachmentFromArchive(*this, MoveTemp(Field), Ar, Allocator);
 }
 
 void FCbAttachment::Save(FCbWriter& Writer) const
@@ -303,25 +293,23 @@ void FCbPackage::GatherAttachments(const FCbObject& Value, FAttachmentResolver R
 bool FCbPackage::TryLoad(FCbFieldIterator& Fields)
 {
 	*this = FCbPackage();
-
 	while (Fields)
 	{
 		if (Fields.IsNull())
 		{
 			++Fields;
-			break;
+			return true;
 		}
 		else if (FIoHash Hash = Fields.AsHash(); !Fields.HasError() && !Fields.IsAttachment())
 		{
 			++Fields;
-			FCbObjectView ObjectView = Fields.AsObjectView();
-			if (Fields.HasError() || Hash != ObjectView.GetHash())
+			Object = Fields.AsObject();
+			Object.MakeOwned();
+			ObjectHash = Hash;
+			if (Fields.HasError())
 			{
 				return false;
 			}
-			Object = FCbObject(ObjectView, Fields.GetOuterBuffer());
-			Object.MakeOwned();
-			ObjectHash = Hash;
 			++Fields;
 		}
 		else
@@ -334,7 +322,7 @@ bool FCbPackage::TryLoad(FCbFieldIterator& Fields)
 			AddAttachment(Attachment);
 		}
 	}
-	return true;
+	return false;
 }
 
 bool FCbPackage::TryLoad(FArchive& Ar, FCbBufferAllocator Allocator)
@@ -356,18 +344,17 @@ bool FCbPackage::TryLoad(FArchive& Ar, FCbBufferAllocator Allocator)
 		else if (FIoHash Hash = Field.AsHash(); !Field.HasError() && !Field.IsAttachment())
 		{
 			Field = LoadCompactBinary(Ar, Allocator);
-			FCbObjectView ObjectView = Field.AsObjectView();
-			if (Field.HasError() || Hash != ObjectView.GetHash())
+			Object = MoveTemp(Field).AsObject();
+			ObjectHash = Hash;
+			if (Field.HasError())
 			{
 				return false;
 			}
-			Object = FCbObject(ObjectView, Field.GetOuterBuffer());
-			ObjectHash = Hash;
 		}
 		else
 		{
 			FCbAttachment Attachment;
-			if (!TryLoad_ArchiveFieldIntoAttachment(Attachment, MoveTemp(Field), Ar, Allocator))
+			if (!TryLoadAttachmentFromArchive(Attachment, MoveTemp(Field), Ar, Allocator))
 			{
 				return false;
 			}
diff --git a/Engine/Source/Runtime/Core/Public/Serialization/CompactBinaryPackage.h b/Engine/Source/Runtime/Core/Public/Serialization/CompactBinaryPackage.h
--- a/Engine/Source/Runtime/Core/Public/Serialization/CompactBinaryPackage.h
+++ b/Engine/Source/Runtime/Core/Public/Serialization/CompactBinaryPackage.h
@@ -6,9 +6,11 @@
 #include "Compression/CompressedBuffer.h"
 #include "Containers/Array.h"
 #include "IO/IoHash.h"
-#include "Memory/MemoryFwd.h"
+#include "Memory/CompositeBuffer.h"
+#include "Memory/SharedBuffer.h"
 #include "Misc/TVariant.h"
 #include "Serialization/CompactBinary.h"
+#include "Templates/UnrealTemplate.h"
 
 class FArchive;
 class FCbWriter;
@@ -21,51 +23,38 @@ template <typename FuncType> class TFunctionRef;
  */
 class FCbAttachment
 {
-public:
-	/** Construct a null attachment. */
-	FCbAttachment() = default;
-
-	/** Construct an object attachment. Value is cloned if not owned. */
-	inline explicit FCbAttachment(const FCbObject& InValue)
-		: FCbAttachment(InValue, nullptr)
+	struct FObjectValue
 	{
-	}
+		FCbObject Object;
+		FIoHash Hash;
 
-	/** Construct an object attachment. Value is cloned if not owned. Hash must match Value. */
-	inline explicit FCbAttachment(const FCbObject& InValue, const FIoHash& Hash)
-		: FCbAttachment(InValue, &Hash)
-	{
-	}
+		CORE_API FObjectValue(FCbObject&& Object, const FIoHash* Hash);
+	};
 
-	/** Construct a raw binary attachment. Value is cloned if not owned. */
-	inline explicit FCbAttachment(const FSharedBuffer& InValue)
-		: FCbAttachment(FCompositeBuffer(InValue))
+	struct FBinaryValue
 	{
-	}
+		FCompositeBuffer Buffer;
+		FIoHash Hash;
 
-	/** Construct a raw binary attachment. Value is cloned if not owned. */
-	inline explicit FCbAttachment(const FCompositeBuffer& InValue)
-		: Value(TInPlaceType<FBinaryValue>(), InValue)
-	{
-		if (Value.Get<FBinaryValue>().Buffer.IsNull())
+		template <typename BufferType, decltype(FCompositeBuffer(DeclVal<BufferType>().MakeOwned()))* = nullptr>
+		inline explicit FBinaryValue(BufferType&& InBuffer)
+			: Buffer(Forward<BufferType>(InBuffer).MakeOwned())
+			, Hash(FIoHash::HashBuffer(Buffer))
 		{
-			Value.Emplace<TYPE_OF_NULLPTR>();
 		}
-	}
 
-	/** Construct a raw binary attachment. Value is cloned if not owned. */
-	inline explicit FCbAttachment(FCompositeBuffer&& InValue)
-		: Value(TInPlaceType<FBinaryValue>(), InValue)
-	{
-		if (Value.Get<FBinaryValue>().Buffer.IsNull())
+		template <typename BufferType, decltype(FCompositeBuffer(DeclVal<BufferType>().MakeOwned()))* = nullptr>
+		inline explicit FBinaryValue(BufferType&& InBuffer, const FIoHash& InHash)
+			: Buffer(Forward<BufferType>(InBuffer).MakeOwned())
+			, Hash(InHash)
 		{
-			Value.Emplace<TYPE_OF_NULLPTR>();
+			checkSlow(Hash == FIoHash::HashBuffer(Buffer));
 		}
-	}
+	};
 
-	/** Construct a raw binary attachment. Value is cloned if not owned. */
-	inline explicit FCbAttachment(FCompositeBuffer&& InValue, const FIoHash& InHash)
-		: Value(TInPlaceType<FBinaryValue>(), InValue, InHash)
+	template <typename... ArgTypes, decltype(FBinaryValue(DeclVal<ArgTypes>()...))* = nullptr>
+	inline FCbAttachment(TInPlaceType<FBinaryValue>, ArgTypes&&... Args)
+		: Value(TInPlaceType<FBinaryValue>(), Forward<ArgTypes>(Args)...)
 	{
 		if (Value.Get<FBinaryValue>().Buffer.IsNull())
 		{
@@ -73,6 +62,28 @@ public:
 		}
 	}
 
+public:
+	/** Construct a null attachment. */
+	FCbAttachment() = default;
+
+	/** Construct an object attachment. Value is cloned if not owned. */
+	inline explicit FCbAttachment(const FCbObject& InValue)
+		: Value(TInPlaceType<FObjectValue>(), FCbObject(InValue), nullptr) {}
+	inline explicit FCbAttachment(FCbObject&& InValue)
+		: Value(TInPlaceType<FObjectValue>(), MoveTemp(InValue), nullptr) {}
+	inline explicit FCbAttachment(const FCbObject& InValue, const FIoHash& Hash)
+		: Value(TInPlaceType<FObjectValue>(), FCbObject(InValue), &Hash) {}
+	inline explicit FCbAttachment(FCbObject&& InValue, const FIoHash& Hash)
+		: Value(TInPlaceType<FObjectValue>(), MoveTemp(InValue), &Hash) {}
+
+	/** Construct a raw binary attachment from a shared/composite buffer. Value is cloned if not owned. */
+	template <typename BufferType, decltype(FBinaryValue(DeclVal<BufferType>().MakeOwned()))* = nullptr>
+	inline explicit FCbAttachment(BufferType&& InValue)
+		: FCbAttachment(TInPlaceType<FBinaryValue>(), Forward<BufferType>(InValue).MakeOwned()) {}
+	template <typename BufferType, decltype(FBinaryValue(DeclVal<BufferType>().MakeOwned()))* = nullptr>
+	inline explicit FCbAttachment(BufferType&& InValue, const FIoHash& InHash)
+		: FCbAttachment(TInPlaceType<FBinaryValue>(), Forward<BufferType>(InValue).MakeOwned(), InHash) {}
+
 	/** Construct a compressed binary attachment. Value is cloned if not owned. */
 	inline explicit FCbAttachment(const FCompressedBuffer& InValue)
 		: Value(TInPlaceType<FCompressedBuffer>(), InValue)
@@ -158,62 +169,6 @@ public:
 	CORE_API void Save(FArchive& Ar) const;
 
 private:
-	CORE_API FCbAttachment(const FCbObject& Value, const FIoHash* Hash);
-
-	struct FObjectValue
-	{
-		FCbObject Object;
-		FIoHash Hash;
-
-		FObjectValue(const FCbObject& InObject, const FIoHash& InHash)
-			: Object(InObject)
-			, Hash(InHash)
-		{
-		}
-
-		FObjectValue(FCbObject&& InObject, const FIoHash& InHash)
-			: Object(MoveTemp(InObject))
-			, Hash(InHash)
-		{
-		}
-	};
-
-	struct FBinaryValue
-	{
-		FCompositeBuffer Buffer;
-		FIoHash Hash;
-
-		FBinaryValue(const FCompositeBuffer& InBuffer)
-			: Buffer(InBuffer.MakeOwned())
-			, Hash(FIoHash::HashBuffer(InBuffer))
-		{
-		}
-
-		FBinaryValue(const FCompositeBuffer& InBuffer, const FIoHash& InHash)
-			: Buffer(InBuffer.MakeOwned())
-			, Hash(InHash)
-		{
-		}
-
-		FBinaryValue(FCompositeBuffer&& InBuffer)
-			: Buffer(MoveTemp(InBuffer))
-			, Hash(FIoHash::HashBuffer(Buffer))
-		{
-		}
-
-		FBinaryValue(FCompositeBuffer&& InBuffer, const FIoHash& InHash)
-			: Buffer(MoveTemp(InBuffer))
-			, Hash(InHash)
-		{
-		}
-
-		FBinaryValue(FSharedBuffer&& InBuffer, const FIoHash& InHash)
-			: Buffer(MoveTemp(InBuffer))
-			, Hash(InHash)
-		{
-		}
-	};
-
 	TVariant<TYPE_OF_NULLPTR, FObjectValue, FBinaryValue, FCompressedBuffer> Value;
 };
 
@@ -447,13 +402,13 @@ private:
 
 ///////////////////////////////////////////////////////////////////////////////////////////////////////////////////////
 
-const FCompositeBuffer& FCbAttachment::AsCompositeBinary() const
+FCbObject FCbAttachment::AsObject() const
 {
-	if (const FBinaryValue* BinaryValue = Value.TryGet<FBinaryValue>())
+	if (const FObjectValue* ObjectValue = Value.TryGet<FObjectValue>())
 	{
-		return BinaryValue->Buffer;
+		return ObjectValue->Object;
 	}
-	return FCompositeBuffer::Null;
+	return FCbObject();
 }
 
 FSharedBuffer FCbAttachment::AsBinary() const
@@ -465,20 +420,20 @@ FSharedBuffer FCbAttachment::AsBinary() const
 	return FSharedBuffer();
 }
 
-const FCompressedBuffer& FCbAttachment::AsCompressedBinary() const
+const FCompositeBuffer& FCbAttachment::AsCompositeBinary() const
 {
-	if (const FCompressedBuffer* CompressedBuffer = Value.TryGet<FCompressedBuffer>())
+	if (const FBinaryValue* BinaryValue = Value.TryGet<FBinaryValue>())
 	{
-		return *CompressedBuffer;
+		return BinaryValue->Buffer;
 	}
-	return FCompressedBuffer::Null;
+	return FCompositeBuffer::Null;
 }
 
-FCbObject FCbAttachment::AsObject() const
+const FCompressedBuffer& FCbAttachment::AsCompressedBinary() const
 {
-	if (const FObjectValue* ObjectValue = Value.TryGet<FObjectValue>())
+	if (const FCompressedBuffer* CompressedBuffer = Value.TryGet<FCompressedBuffer>())
 	{
-		return ObjectValue->Object;
+		return *CompressedBuffer;
 	}
-	return FCbObject();
+	return FCompressedBuffer::Null;
 }
