CompositeBuffer: Renamed Flatten() to ToShared()

#rnx

diff --git a/Engine/Source/Developer/DerivedDataCache/Private/DerivedDataCacheRecord.cpp b/Engine/Source/Developer/DerivedDataCache/Private/DerivedDataCacheRecord.cpp
--- a/Engine/Source/Developer/DerivedDataCache/Private/DerivedDataCacheRecord.cpp
+++ b/Engine/Source/Developer/DerivedDataCache/Private/DerivedDataCacheRecord.cpp
@@ -300,7 +300,7 @@ FCbPackage FCacheRecord::Save() const
 		Writer.AddHash("RawHash"_ASV, Payload.GetRawHash());
 		Writer.AddInteger("RawSize"_ASV, Payload.GetRawSize());
 		const FCompositeBuffer CompressedBuffer = Payload.GetData().GetCompressed();
-		FCbAttachment AttachedCompressed(CompressedBuffer.Flatten());
+		FCbAttachment AttachedCompressed(CompressedBuffer.ToShared());
 		Writer.AddAttachment("CompressedHash"_ASV, AttachedCompressed);
 		Package.AddAttachment(AttachedCompressed);
 		Writer.EndObject();
diff --git a/Engine/Source/Developer/DerivedDataCache/Private/ZenSerialization.cpp b/Engine/Source/Developer/DerivedDataCache/Private/ZenSerialization.cpp
--- a/Engine/Source/Developer/DerivedDataCache/Private/ZenSerialization.cpp
+++ b/Engine/Source/Developer/DerivedDataCache/Private/ZenSerialization.cpp
@@ -16,7 +16,7 @@ void SaveCbAttachment(const FCbAttachment& Attachment, FCbWriter& Writer)
 	if (Attachment.IsObject())
 	{
 		FCbObject Object = Attachment.AsObject();
-		Writer.AddBinary(Object.GetBuffer().Flatten());
+		Writer.AddBinary(Object.GetBuffer().ToShared());
 		if (Object)
 		{
 			Writer.AddObjectAttachment(Attachment.GetHash());
diff --git a/Engine/Source/Developer/Virtualization/Private/VirtualizationJupiterBackend.cpp b/Engine/Source/Developer/Virtualization/Private/VirtualizationJupiterBackend.cpp
--- a/Engine/Source/Developer/Virtualization/Private/VirtualizationJupiterBackend.cpp
+++ b/Engine/Source/Developer/Virtualization/Private/VirtualizationJupiterBackend.cpp
@@ -1417,7 +1417,7 @@ private:
 		// with the existing chunked code. The chunking code is most likely going to be removed before 
 		// this backend goes to production so it is not worth fixing up to work with FCompressedBuffer 
 		// properly.
-		FSharedBuffer FlattenedPayload = CompressedPayload.GetCompressed().Flatten();
+		FSharedBuffer FlattenedPayload = CompressedPayload.GetCompressed().ToShared();
 
 		const int64 NumChunks = FMath::DivideAndRoundUp(FlattenedPayload.GetSize(), ChunkSize);
 		if (NumChunks > MAX_int32)
diff --git a/Engine/Source/Runtime/Core/Private/Compression/CompressedBuffer.cpp b/Engine/Source/Runtime/Core/Private/Compression/CompressedBuffer.cpp
--- a/Engine/Source/Runtime/Core/Private/Compression/CompressedBuffer.cpp
+++ b/Engine/Source/Runtime/Core/Private/Compression/CompressedBuffer.cpp
@@ -775,7 +775,7 @@ FSharedBuffer FCompressedBuffer::Decompress() const
 		{
 			if (Header.Method == EMethod::None)
 			{
-				return Decoder->Decompress(Header, CompressedData).Flatten();
+				return Decoder->Decompress(Header, CompressedData).ToShared();
 			}
 			FUniqueBuffer RawData = FUniqueBuffer::Alloc(Header.TotalRawSize);
 			if (Decoder->TryDecompressTo(Header, CompressedData, RawData))
diff --git a/Engine/Source/Runtime/Core/Private/Memory/CompositeBuffer.cpp b/Engine/Source/Runtime/Core/Private/Memory/CompositeBuffer.cpp
--- a/Engine/Source/Runtime/Core/Private/Memory/CompositeBuffer.cpp
+++ b/Engine/Source/Runtime/Core/Private/Memory/CompositeBuffer.cpp
@@ -35,7 +35,7 @@ FCompositeBuffer FCompositeBuffer::MakeOwned() &&
 	return MoveTemp(*this);
 }
 
-FSharedBuffer FCompositeBuffer::Flatten() const &
+FSharedBuffer FCompositeBuffer::ToShared() const &
 {
 	switch (Segments.Num())
 	{
@@ -50,9 +50,9 @@ FSharedBuffer FCompositeBuffer::Flatten() const &
 	}
 }
 
-FSharedBuffer FCompositeBuffer::Flatten() &&
+FSharedBuffer FCompositeBuffer::ToShared() &&
 {
-	return Segments.Num() == 1 ? MoveTemp(Segments[0]) : AsConst(*this).Flatten();
+	return Segments.Num() == 1 ? MoveTemp(Segments[0]) : AsConst(*this).ToShared();
 }
 
 FCompositeBuffer FCompositeBuffer::Mid(uint64 Offset, uint64 Size) const
diff --git a/Engine/Source/Runtime/Core/Private/Serialization/CompactBinaryWriter.cpp b/Engine/Source/Runtime/Core/Private/Serialization/CompactBinaryWriter.cpp
--- a/Engine/Source/Runtime/Core/Private/Serialization/CompactBinaryWriter.cpp
+++ b/Engine/Source/Runtime/Core/Private/Serialization/CompactBinaryWriter.cpp
@@ -377,7 +377,7 @@ void FCbWriter::AddBinary(const FSharedBuffer& Buffer)
 
 void FCbWriter::AddBinary(const FCompositeBuffer& Buffer)
 {
-	AddBinary(Buffer.Flatten());
+	AddBinary(Buffer.ToShared());
 }
 
 void FCbWriter::AddString(const FUtf8StringView Value)
diff --git a/Engine/Source/Runtime/Core/Private/Tests/Memory/CompositeBufferTest.cpp b/Engine/Source/Runtime/Core/Private/Tests/Memory/CompositeBufferTest.cpp
--- a/Engine/Source/Runtime/Core/Private/Tests/Memory/CompositeBufferTest.cpp
+++ b/Engine/Source/Runtime/Core/Private/Tests/Memory/CompositeBufferTest.cpp
@@ -30,7 +30,7 @@ bool FCompositeBufferTest::RunTest(const FString& Parameters)
 		TestTrue(TEXT("FCompositeBuffer().IsNull()"), Buffer.IsNull());
 		TestTrue(TEXT("FCompositeBuffer().IsOwned()"), Buffer.IsOwned());
 		TestTrue(TEXT("FCompositeBuffer().MakeOwned().IsNull()"), Buffer.MakeOwned().IsNull());
-		TestTrue(TEXT("FCompositeBuffer().Flatten().IsNull()"), Buffer.Flatten().IsNull());
+		TestTrue(TEXT("FCompositeBuffer().ToShared().IsNull()"), Buffer.ToShared().IsNull());
 		TestTrue(TEXT("FCompositeBuffer().Mid(0, 0).IsNull()"), Buffer.Mid(0, 0).IsNull());
 		TestEqual(TEXT("FCompositeBuffer().GetSize()"), Buffer.GetSize(), uint64(0));
 		TestEqual(TEXT("FCompositeBuffer().GetSegments().Num()"), Buffer.GetSegments().Num(), 0);
@@ -60,8 +60,8 @@ bool FCompositeBufferTest::RunTest(const FString& Parameters)
 		TestFalse(TEXT("FCompositeBuffer(Empty).IsOwned()"), Buffer.IsOwned());
 		TestFalse(TEXT("FCompositeBuffer(Empty).MakeOwned().IsNull()"), Buffer.MakeOwned().IsNull());
 		TestTrue(TEXT("FCompositeBuffer(Empty).MakeOwned().IsOwned()"), Buffer.MakeOwned().IsOwned());
-		TestEqual(TEXT("FCompositeBuffer(Empty).Flatten()"), Buffer.Flatten(), EmptyView);
-		TestEqual(TEXT("FCompositeBuffer(Empty).Mid(0, 0)"), Buffer.Mid(0, 0).Flatten(), EmptyView);
+		TestEqual(TEXT("FCompositeBuffer(Empty).ToShared()"), Buffer.ToShared(), EmptyView);
+		TestEqual(TEXT("FCompositeBuffer(Empty).Mid(0, 0)"), Buffer.Mid(0, 0).ToShared(), EmptyView);
 		TestEqual(TEXT("FCompositeBuffer(Empty).GetSize()"), Buffer.GetSize(), uint64(0));
 		TestEqual(TEXT("FCompositeBuffer(Empty).GetSegments().Num()"), Buffer.GetSegments().Num(), 1);
 		TestEqual(TEXT("FCompositeBuffer(Empty).GetSegments()[0]"), Buffer.GetSegments()[0], EmptyView);
@@ -81,7 +81,7 @@ bool FCompositeBufferTest::RunTest(const FString& Parameters)
 		const FSharedBuffer EmptyView1 = FSharedBuffer::MakeView(EmptyArray, 0);
 		const FSharedBuffer EmptyView2 = FSharedBuffer::MakeView(EmptyArray + 1, 0);
 		FCompositeBuffer Buffer(TArray<FSharedBuffer>{EmptyView1}, FCompositeBuffer(EmptyView2));
-		TestEqual(TEXT("FCompositeBuffer(Empty, Empty).Mid(0, 0)"), Buffer.Mid(0, 0).Flatten(), EmptyView1);
+		TestEqual(TEXT("FCompositeBuffer(Empty, Empty).Mid(0, 0)"), Buffer.Mid(0, 0).ToShared(), EmptyView1);
 		TestEqual(TEXT("FCompositeBuffer(Empty, Empty).GetSize()"), Buffer.GetSize(), uint64(0));
 		TestEqual(TEXT("FCompositeBuffer(Empty, Empty).GetSegments().Num()"), Buffer.GetSegments().Num(), 2);
 		TestEqual(TEXT("FCompositeBuffer(Empty, Empty).GetSegments()[0]"), Buffer.GetSegments()[0], EmptyView1);
@@ -105,15 +105,15 @@ bool FCompositeBufferTest::RunTest(const FString& Parameters)
 		FCompositeBuffer Buffer(FlatView);
 		TestFalse(TEXT("FCompositeBuffer(Flat).IsNull()"), Buffer.IsNull());
 		TestTrue(TEXT("FCompositeBuffer(Flat).IsOwned()"), Buffer.IsOwned());
-		TestEqual(TEXT("FCompositeBuffer(Flat).Flatten()"), Buffer.Flatten(), FlatView);
-		TestEqual(TEXT("FCompositeBuffer(Flat).MakeOwned().Flatten()"), Buffer.MakeOwned().Flatten(), FlatView);
-		TestEqual(TEXT("FCompositeBuffer(Flat).Mid(0)"), Buffer.Mid(0).Flatten(), FlatView);
-		TestEqual(TEXT("FCompositeBuffer(Flat).Mid(4)"), Buffer.Mid(4).Flatten().GetView(), FlatView.GetView().Mid(4));
-		TestEqual(TEXT("FCompositeBuffer(Flat).Mid(8)"), Buffer.Mid(8).Flatten().GetView(), FlatView.GetView().Mid(8));
-		TestEqual(TEXT("FCompositeBuffer(Flat).Mid(4, 2)"), Buffer.Mid(4, 2).Flatten().GetView(), FlatView.GetView().Mid(4, 2));
-		TestEqual(TEXT("FCompositeBuffer(Flat).Mid(8).GetData()"), Buffer.Mid(8).Flatten().GetView().GetData(), FlatView.GetView().Mid(8).GetData());
-		TestEqual(TEXT("FCompositeBuffer(Flat).Mid(4, 2).GetData()"), Buffer.Mid(4, 2).Flatten().GetView().GetData(), FlatView.GetView().Mid(4, 2).GetData());
-		TestEqual(TEXT("FCompositeBuffer(Flat).Mid(8, 0).GetData()"), Buffer.Mid(8, 0).Flatten().GetView().GetData(), FlatView.GetView().Mid(8, 0).GetData());
+		TestEqual(TEXT("FCompositeBuffer(Flat).ToShared()"), Buffer.ToShared(), FlatView);
+		TestEqual(TEXT("FCompositeBuffer(Flat).MakeOwned().ToShared()"), Buffer.MakeOwned().ToShared(), FlatView);
+		TestEqual(TEXT("FCompositeBuffer(Flat).Mid(0)"), Buffer.Mid(0).ToShared(), FlatView);
+		TestEqual(TEXT("FCompositeBuffer(Flat).Mid(4)"), Buffer.Mid(4).ToShared().GetView(), FlatView.GetView().Mid(4));
+		TestEqual(TEXT("FCompositeBuffer(Flat).Mid(8)"), Buffer.Mid(8).ToShared().GetView(), FlatView.GetView().Mid(8));
+		TestEqual(TEXT("FCompositeBuffer(Flat).Mid(4, 2)"), Buffer.Mid(4, 2).ToShared().GetView(), FlatView.GetView().Mid(4, 2));
+		TestEqual(TEXT("FCompositeBuffer(Flat).Mid(8).GetData()"), Buffer.Mid(8).ToShared().GetView().GetData(), FlatView.GetView().Mid(8).GetData());
+		TestEqual(TEXT("FCompositeBuffer(Flat).Mid(4, 2).GetData()"), Buffer.Mid(4, 2).ToShared().GetView().GetData(), FlatView.GetView().Mid(4, 2).GetData());
+		TestEqual(TEXT("FCompositeBuffer(Flat).Mid(8, 0).GetData()"), Buffer.Mid(8, 0).ToShared().GetView().GetData(), FlatView.GetView().Mid(8, 0).GetData());
 		TestEqual(TEXT("FCompositeBuffer(Flat).GetSize()"), Buffer.GetSize(), uint64(sizeof(FlatArray)));
 		TestEqual(TEXT("FCompositeBuffer(Flat).GetSegments().Num()"), Buffer.GetSegments().Num(), 1);
 		TestEqual(TEXT("FCompositeBuffer(Flat).GetSegments()[0]"), Buffer.GetSegments()[0], FlatView);
@@ -138,12 +138,12 @@ bool FCompositeBufferTest::RunTest(const FString& Parameters)
 		FCompositeBuffer Buffer(FlatView1, FlatView2);
 		TestFalse(TEXT("FCompositeBuffer(Composite).IsNull()"), Buffer.IsNull());
 		TestFalse(TEXT("FCompositeBuffer(Composite).IsOwned()"), Buffer.IsOwned());
-		TestTrue(TEXT("FCompositeBuffer(Composite).Flatten()"), Buffer.Flatten().GetView().EqualBytes(MakeMemoryView(FlatArray)));
-		TestTrue(TEXT("FCompositeBuffer(Composite).Mid(2, 4)"), Buffer.Mid(2, 4).Flatten().GetView().EqualBytes(MakeMemoryView(FlatArray).Mid(2, 4)));
-		TestEqual(TEXT("FCompositeBuffer(Composite).Mid(0, 4)"), Buffer.Mid(0, 4).Flatten(), FlatView1);
-		TestEqual(TEXT("FCompositeBuffer(Composite).Mid(4, 4)"), Buffer.Mid(4, 4).Flatten(), FlatView2);
-		TestEqual(TEXT("FCompositeBuffer(Composite).Mid(4, 0).GetData()"), Buffer.Mid(4, 0).Flatten().GetView().GetData(), static_cast<const void*>(FlatArray + 4));
-		TestEqual(TEXT("FCompositeBuffer(Composite).Mid(8, 0).GetData()"), Buffer.Mid(8, 0).Flatten().GetView().GetData(), static_cast<const void*>(FlatArray + 8));
+		TestTrue(TEXT("FCompositeBuffer(Composite).ToShared()"), Buffer.ToShared().GetView().EqualBytes(MakeMemoryView(FlatArray)));
+		TestTrue(TEXT("FCompositeBuffer(Composite).Mid(2, 4)"), Buffer.Mid(2, 4).ToShared().GetView().EqualBytes(MakeMemoryView(FlatArray).Mid(2, 4)));
+		TestEqual(TEXT("FCompositeBuffer(Composite).Mid(0, 4)"), Buffer.Mid(0, 4).ToShared(), FlatView1);
+		TestEqual(TEXT("FCompositeBuffer(Composite).Mid(4, 4)"), Buffer.Mid(4, 4).ToShared(), FlatView2);
+		TestEqual(TEXT("FCompositeBuffer(Composite).Mid(4, 0).GetData()"), Buffer.Mid(4, 0).ToShared().GetView().GetData(), static_cast<const void*>(FlatArray + 4));
+		TestEqual(TEXT("FCompositeBuffer(Composite).Mid(8, 0).GetData()"), Buffer.Mid(8, 0).ToShared().GetView().GetData(), static_cast<const void*>(FlatArray + 8));
 		TestEqual(TEXT("FCompositeBuffer(Composite).GetSize()"), Buffer.GetSize(), uint64(sizeof(FlatArray)));
 		TestEqual(TEXT("FCompositeBuffer(Composite).GetSegments().Num()"), Buffer.GetSegments().Num(), 2);
 		TestEqual(TEXT("FCompositeBuffer(Composite).GetSegments()[0]"), Buffer.GetSegments()[0], FlatView1);
diff --git a/Engine/Source/Runtime/Core/Private/Tests/Serialization/CompactBinaryPackageTest.cpp b/Engine/Source/Runtime/Core/Private/Tests/Serialization/CompactBinaryPackageTest.cpp
--- a/Engine/Source/Runtime/Core/Private/Tests/Serialization/CompactBinaryPackageTest.cpp
+++ b/Engine/Source/Runtime/Core/Private/Tests/Serialization/CompactBinaryPackageTest.cpp
@@ -82,7 +82,7 @@ bool FCbAttachmentTest::RunTest(const FString& Parameters)
 		FCbAttachment Attachment(Buffer);
 		TestFalse(TEXT("FCbAttachment(CompressedBinary).IsNull()"), Attachment.IsNull());
 		TestTrue(TEXT("FCbAttachment(CompressedBinary) as bool"), bool(Attachment));
-		TestTrue(TEXT("FCbAttachment(CompressedBinary).AsCompressedBinary()"), Attachment.AsCompressedBinary().GetCompressed().Flatten().GetView().EqualBytes(Buffer.GetCompressed().Flatten().GetView()));
+		TestTrue(TEXT("FCbAttachment(CompressedBinary).AsCompressedBinary()"), Attachment.AsCompressedBinary().GetCompressed().ToShared().GetView().EqualBytes(Buffer.GetCompressed().ToShared().GetView()));
 		TestFalse(TEXT("FCbAttachment(CompressedBinary).AsObject()"), bool(Attachment.AsObject()));
 		TestFalse(TEXT("FCbAttachment(CompressedBinary).IsBinary()"), Attachment.IsBinary());
 		TestTrue(TEXT("FCbAttachment(CompressedBinary).IsCompressedBinary()"), Attachment.IsCompressedBinary());
@@ -185,9 +185,9 @@ bool FCbAttachmentTest::RunTest(const FString& Parameters)
 		TestFalse(TEXT("FCbAttachment(LoadCompressedBinaryView).IsNull()"), Attachment.IsNull());
 		TestTrue(TEXT("FCbAttachment(LoadCompressedBinaryView) as bool"), bool(Attachment));
 		TestFalse(TEXT("FCbAttachment(LoadCompressedBinaryView).AsBinary()->!InView"),
-			FieldsView.GetOuterBuffer().GetView().Contains(Attachment.AsCompressedBinary().GetCompressed().Flatten().GetView()));
+			FieldsView.GetOuterBuffer().GetView().Contains(Attachment.AsCompressedBinary().GetCompressed().ToShared().GetView()));
 		TestTrue(TEXT("FCbAttachment(LoadCompressedBinaryView).AsCompressedBinary()->EqualBytes"),
-			Attachment.AsCompressedBinary().GetCompressed().Flatten().GetView().EqualBytes(Buffer.GetCompressed().Flatten().GetView()));
+			Attachment.AsCompressedBinary().GetCompressed().ToShared().GetView().EqualBytes(Buffer.GetCompressed().ToShared().GetView()));
 		TestFalse(TEXT("FCbAttachment(LoadCompressedBinaryView).AsObject()"), bool(Attachment.AsObject()));
 		TestFalse(TEXT("FCbAttachment(LoadCompressedBinaryView).IsBinary()"), Attachment.IsBinary());
 		TestTrue(TEXT("FCbAttachment(LoadCompressedBinaryView).IsCompressedBinary()"), Attachment.IsCompressedBinary());
diff --git a/Engine/Source/Runtime/Core/Private/Tests/Serialization/CompactBinaryTest.cpp b/Engine/Source/Runtime/Core/Private/Tests/Serialization/CompactBinaryTest.cpp
--- a/Engine/Source/Runtime/Core/Private/Tests/Serialization/CompactBinaryTest.cpp
+++ b/Engine/Source/Runtime/Core/Private/Tests/Serialization/CompactBinaryTest.cpp
@@ -420,9 +420,9 @@ bool FCbFieldObjectTest::RunTest(const FString& Parameters)
 		TestFalse(TEXT("FCbObjectView(Name)::TryGetView()"), NamedField.AsObjectView().TryGetView(View));
 
 		// GetBuffer
-		TestTrue(TEXT("FCbObjectView(Clone)::GetBuffer()"), Object.GetBuffer().Flatten().GetView() == Object.GetOuterBuffer().GetView());
-		TestTrue(TEXT("FCbObjectView(NoType)::GetBuffer()"), FCbField::MakeView(Field).AsObject().GetBuffer().Flatten().GetView().EqualBytes(Object.GetOuterBuffer().GetView()));
-		TestTrue(TEXT("FCbObjectView(Name)::GetBuffer()"), FCbField::MakeView(NamedField).AsObject().GetBuffer().Flatten().GetView().EqualBytes(Object.GetOuterBuffer().GetView()));
+		TestTrue(TEXT("FCbObjectView(Clone)::GetBuffer()"), Object.GetBuffer().ToShared().GetView() == Object.GetOuterBuffer().GetView());
+		TestTrue(TEXT("FCbObjectView(NoType)::GetBuffer()"), FCbField::MakeView(Field).AsObject().GetBuffer().ToShared().GetView().EqualBytes(Object.GetOuterBuffer().GetView()));
+		TestTrue(TEXT("FCbObjectView(Name)::GetBuffer()"), FCbField::MakeView(NamedField).AsObject().GetBuffer().ToShared().GetView().EqualBytes(Object.GetOuterBuffer().GetView()));
 	}
 
 	// Test FCbFieldView(None) as Object
@@ -464,9 +464,9 @@ bool FCbFieldObjectTest::RunTest(const FString& Parameters)
 		TestTrue(TEXT("FCbObjectView()::CopyTo()"), ObjectClone.GetOuterBuffer().GetView().EqualBytes(MakeMemoryView(CopyBytes)));
 
 		// GetBuffer
-		TestTrue(TEXT("FCbField(ObjectWithName)::GetBuffer()"), FCbField(FSharedBuffer::MakeView(MakeMemoryView(Buffer))).GetBuffer().Flatten().GetView().EqualBytes(MakeMemoryView(Buffer)));
+		TestTrue(TEXT("FCbField(ObjectWithName)::GetBuffer()"), FCbField(FSharedBuffer::MakeView(MakeMemoryView(Buffer))).GetBuffer().ToShared().GetView().EqualBytes(MakeMemoryView(Buffer)));
 		TestTrue(TEXT("FCbField(ObjectWithNameNoType)::GetBuffer()"),
-			FCbField(FCbFieldView(Buffer + 1, ECbFieldType(ObjectType)), FSharedBuffer::MakeView(MakeMemoryView(Buffer))).GetBuffer().Flatten().GetView().EqualBytes(MakeMemoryView(Buffer)));
+			FCbField(FCbFieldView(Buffer + 1, ECbFieldType(ObjectType)), FSharedBuffer::MakeView(MakeMemoryView(Buffer))).GetBuffer().ToShared().GetView().EqualBytes(MakeMemoryView(Buffer)));
 
 		// Access Missing Field
 		TestFalse(TEXT("FCbObject()[Missing]"), ObjectClone["M"_ASV].HasValue());
@@ -615,9 +615,9 @@ bool FCbFieldArrayTest::RunTest(const FString& Parameters)
 		TestFalse(TEXT("FCbArrayView(Name)::TryGetView()"), NamedField.AsArrayView().TryGetView(View));
 
 		// GetBuffer
-		TestTrue(TEXT("FCbArrayView(Clone)::GetBuffer()"), Array.GetBuffer().Flatten().GetView() == Array.GetOuterBuffer().GetView());
-		TestTrue(TEXT("FCbArrayView(NoType)::GetBuffer()"), FCbField::MakeView(Field).AsArray().GetBuffer().Flatten().GetView().EqualBytes(Array.GetOuterBuffer().GetView()));
-		TestTrue(TEXT("FCbArrayView(Name)::GetBuffer()"), FCbField::MakeView(NamedField).AsArray().GetBuffer().Flatten().GetView().EqualBytes(Array.GetOuterBuffer().GetView()));
+		TestTrue(TEXT("FCbArrayView(Clone)::GetBuffer()"), Array.GetBuffer().ToShared().GetView() == Array.GetOuterBuffer().GetView());
+		TestTrue(TEXT("FCbArrayView(NoType)::GetBuffer()"), FCbField::MakeView(Field).AsArray().GetBuffer().ToShared().GetView().EqualBytes(Array.GetOuterBuffer().GetView()));
+		TestTrue(TEXT("FCbArrayView(Name)::GetBuffer()"), FCbField::MakeView(NamedField).AsArray().GetBuffer().ToShared().GetView().EqualBytes(Array.GetOuterBuffer().GetView()));
 	}
 
 	// Test FCbFieldView(None) as Array
@@ -658,9 +658,9 @@ bool FCbFieldArrayTest::RunTest(const FString& Parameters)
 		TestTrue(TEXT("FCbArrayView()::CopyTo()"), ArrayClone.GetOuterBuffer().GetView().EqualBytes(MakeMemoryView(CopyBytes)));
 
 		// GetBuffer
-		TestTrue(TEXT("FCbField(ArrayWithName)::GetBuffer()"), FCbField(FSharedBuffer::MakeView(MakeMemoryView(Buffer))).GetBuffer().Flatten().GetView().EqualBytes(MakeMemoryView(Buffer)));
+		TestTrue(TEXT("FCbField(ArrayWithName)::GetBuffer()"), FCbField(FSharedBuffer::MakeView(MakeMemoryView(Buffer))).GetBuffer().ToShared().GetView().EqualBytes(MakeMemoryView(Buffer)));
 		TestTrue(TEXT("FCbField(ArrayWithNameNoType)::GetBuffer()"),
-			FCbField(FCbFieldView(Buffer + 1, ECbFieldType(ArrayType)), FSharedBuffer::MakeView(MakeMemoryView(Buffer))).GetBuffer().Flatten().GetView().EqualBytes(MakeMemoryView(Buffer)));
+			FCbField(FCbFieldView(Buffer + 1, ECbFieldType(ArrayType)), FSharedBuffer::MakeView(MakeMemoryView(Buffer))).GetBuffer().ToShared().GetView().EqualBytes(MakeMemoryView(Buffer)));
 	}
 
 	// Test FCbArrayView as FCbFieldViewIterator
diff --git a/Engine/Source/Runtime/Core/Public/Memory/CompositeBuffer.h b/Engine/Source/Runtime/Core/Public/Memory/CompositeBuffer.h
--- a/Engine/Source/Runtime/Core/Public/Memory/CompositeBuffer.h
+++ b/Engine/Source/Runtime/Core/Public/Memory/CompositeBuffer.h
@@ -14,7 +14,7 @@ template <typename FuncType> class TFunctionRef;
  * FCompositeBuffer is a non-contiguous buffer composed of zero or more immutable shared buffers.
  *
  * A composite buffer is most efficient when its segments are consumed as they are, but it can be
- * flattened into a contiguous buffer, when necessary, by calling Flatten(). Ownership of segment
+ * flattened into a contiguous buffer when necessary, by calling ToShared(). Ownership of segment
  * buffers is not changed on construction, but if ownership of segments is required then that can
  * be guaranteed by calling MakeOwned().
  */
@@ -76,8 +76,8 @@ public:
 	[[nodiscard]] CORE_API FCompositeBuffer MakeOwned() &&;
 
 	/** Returns the concatenation of the segments into a contiguous buffer. */
-	[[nodiscard]] CORE_API FSharedBuffer Flatten() const &;
-	[[nodiscard]] CORE_API FSharedBuffer Flatten() &&;
+	[[nodiscard]] CORE_API FSharedBuffer ToShared() const &;
+	[[nodiscard]] CORE_API FSharedBuffer ToShared() &&;
 
 	/** Returns the middle part of the buffer by taking the size starting at the offset. */
 	[[nodiscard]] CORE_API FCompositeBuffer Mid(uint64 Offset, uint64 Size = MAX_uint64) const;
diff --git a/Engine/Source/Runtime/Core/Public/Serialization/CompactBinaryPackage.h b/Engine/Source/Runtime/Core/Public/Serialization/CompactBinaryPackage.h
--- a/Engine/Source/Runtime/Core/Public/Serialization/CompactBinaryPackage.h
+++ b/Engine/Source/Runtime/Core/Public/Serialization/CompactBinaryPackage.h
@@ -415,7 +415,7 @@ FSharedBuffer FCbAttachment::AsBinary() const
 {
 	if (const FBinaryValue* BinaryValue = Value.TryGet<FBinaryValue>())
 	{
-		return BinaryValue->Buffer.Flatten();
+		return BinaryValue->Buffer.ToShared();
 	}
 	return FSharedBuffer();
 }
