CompositeBuffer: Workaround for a bug in MSVC older than 14.29.30133

This is confirmed to build with MSVC 14.24.28314, 14.28.29914, 14.29.30133.

#rnx

diff --git a/Engine/Source/Runtime/Core/Public/Memory/CompositeBuffer.h b/Engine/Source/Runtime/Core/Public/Memory/CompositeBuffer.h
--- a/Engine/Source/Runtime/Core/Public/Memory/CompositeBuffer.h
+++ b/Engine/Source/Runtime/Core/Public/Memory/CompositeBuffer.h
@@ -6,7 +6,7 @@
 #include "Containers/ArrayView.h"
 #include "Memory/MemoryView.h"
 #include "Memory/SharedBuffer.h"
-#include "Templates/UnrealTemplate.h"
+#include <type_traits>
 
 template <typename FuncType> class TFunctionRef;
 
@@ -20,20 +20,11 @@ template <typename FuncType> class TFunctionRef;
  */
 class FCompositeBuffer
 {
-	using FSegmentArray = TArray<FSharedBuffer, TInlineAllocator<1>>;
-
-	static inline int32 GetBufferCount(const FCompositeBuffer& Buffer) { return Buffer.Segments.Num(); }
-	static inline void AppendBuffers(FSegmentArray& Array, const FCompositeBuffer& Buffer) { Array.Append(Buffer.Segments); }
-	static inline void AppendBuffers(FSegmentArray& Array, FCompositeBuffer&& Buffer) { Array.Append(MoveTemp(Buffer.Segments)); }
-
-	static inline int32 GetBufferCount(const FSharedBuffer& Buffer) { return 1; }
-	static inline void AppendBuffers(FSegmentArray& Array, const FSharedBuffer& Buffer) { Array.Add(Buffer); }
-	static inline void AppendBuffers(FSegmentArray& Array, FSharedBuffer&& Buffer) { Array.Add(MoveTemp(Buffer)); }
-
-	template <typename BufferType, decltype(DeclVal<TArray<FSharedBuffer>>().Append(DeclVal<BufferType>()))* = nullptr>
-	static inline int32 GetBufferCount(BufferType&& Buffer) { return GetNum(Buffer); }
-	template <typename BufferType, decltype(DeclVal<TArray<FSharedBuffer>>().Append(DeclVal<BufferType>()))* = nullptr>
-	static inline void AppendBuffers(FSegmentArray& Array, BufferType&& Buffer) { Array.Append(Forward<BufferType>(Buffer)); }
+	template <typename BufferType>
+	using CanAppendBufferType = std::bool_constant<
+		std::is_same_v<FCompositeBuffer, BufferType> ||
+		std::is_same_v<FSharedBuffer, BufferType> ||
+		TIsTArray<BufferType>::Value>;
 
 public:
 	/**
@@ -42,13 +33,13 @@ public:
 	 * Buffer parameters may be FSharedBuffer, FCompositeBuffer, or TArray<FSharedBuffer>.
 	 */
 	template <typename... BufferTypes,
-		decltype((AppendBuffers(DeclVal<FSegmentArray&>(), DeclVal<BufferTypes>()), ..., int8()))* = nullptr>
+		std::enable_if_t<std::conjunction_v<CanAppendBufferType<std::decay_t<BufferTypes>>...>>* = nullptr>
 	inline explicit FCompositeBuffer(BufferTypes&&... Buffers)
 	{
 		if constexpr (sizeof...(Buffers) > 0)
 		{
 			Segments.Reserve((GetBufferCount(Forward<BufferTypes>(Buffers)) + ...));
-			(AppendBuffers(Segments, Forward<BufferTypes>(Buffers)), ...);
+			(AppendBuffers(Forward<BufferTypes>(Buffers)), ...);
 			Segments.RemoveAll(&FSharedBuffer::IsNull);
 		}
 	}
@@ -117,7 +108,21 @@ public:
 	static const FCompositeBuffer Null;
 
 private:
-	FSegmentArray Segments;
+	static inline int32 GetBufferCount(const FCompositeBuffer& Buffer) { return Buffer.Segments.Num(); }
+	inline void AppendBuffers(const FCompositeBuffer& Buffer) { Segments.Append(Buffer.Segments); }
+	inline void AppendBuffers(FCompositeBuffer&& Buffer) { Segments.Append(MoveTemp(Buffer.Segments)); }
+
+	static inline int32 GetBufferCount(const FSharedBuffer& Buffer) { return 1; }
+	inline void AppendBuffers(const FSharedBuffer& Buffer) { Segments.Add(Buffer); }
+	inline void AppendBuffers(FSharedBuffer&& Buffer) { Segments.Add(MoveTemp(Buffer)); }
+
+	template <typename BufferType, decltype(std::declval<TArray<FSharedBuffer>>().Append(std::declval<BufferType>()))* = nullptr>
+	static inline int32 GetBufferCount(BufferType&& Buffer) { return GetNum(Buffer); }
+	template <typename BufferType, decltype(std::declval<TArray<FSharedBuffer>>().Append(std::declval<BufferType>()))* = nullptr>
+	inline void AppendBuffers(BufferType&& Buffer) { Segments.Append(Forward<BufferType>(Buffer)); }
+
+private:
+	TArray<FSharedBuffer, TInlineAllocator<1>> Segments;
 };
 
 inline const FCompositeBuffer FCompositeBuffer::Null;
