IoDispatcher: Extract a IPlatformFileIoStore interface for the file backend platform implementations.
Move private core header files to internal to allow engine platform implementations in modules outside of core.
Add -forcegenericio command line option to disable platform specific implementations and fall back to the generic platform implementation that works on all platforms.

#rnx

diff --git a/Engine/Source/Runtime/Core/Internal/IO/IoDirectoryIndex.h b/Engine/Source/Runtime/Core/Internal/IO/IoDirectoryIndex.h
new file mode 100644
--- /dev/null
+++ b/Engine/Source/Runtime/Core/Internal/IO/IoDirectoryIndex.h
@@ -0,0 +1,71 @@
+// Copyright Epic Games, Inc. All Rights Reserved.
+
+#pragma once
+
+#include "CoreMinimal.h"
+#include "Containers/StringView.h"
+#include "Misc/AES.h"
+
+class FArchive;
+
+namespace IoDirectoryIndexUtils
+{
+	FString GetCommonRootPath(const TArray<FString>& Filenames);
+}
+
+struct FIoDirectoryIndexEntry
+{
+	uint32 Name				= ~uint32(0);
+	uint32 FirstChildEntry	= ~uint32(0);
+	uint32 NextSiblingEntry	= ~uint32(0);
+	uint32 FirstFileEntry	= ~uint32(0);
+
+	friend FArchive& operator<<(FArchive& Ar, FIoDirectoryIndexEntry& Entry);
+};
+
+struct FIoFileIndexEntry
+{
+	uint32 Name				= ~uint32(0);
+	uint32 NextFileEntry	= ~uint32(0);
+	uint32 UserData			= 0;
+
+	friend FArchive& operator<<(FArchive& Ar, FIoFileIndexEntry& Entry);
+};
+
+struct FIoDirectoryIndexResource
+{
+	FString MountPoint;
+	TArray<FIoDirectoryIndexEntry> DirectoryEntries;
+	TArray<FIoFileIndexEntry> FileEntries;
+	TArray<FString> StringTable;
+
+	friend FArchive& operator<<(FArchive& Ar, FIoDirectoryIndexResource& Entry);
+};
+
+class FIoDirectoryIndexWriter
+{
+public:
+	FIoDirectoryIndexWriter();
+
+	void SetMountPoint(FString InMountPoint);
+	uint32 AddFile(const FString& InFileName);
+	void SetFileUserData(uint32 InFileEntryIndex, uint32 InUserData);
+	void Flush(TArray<uint8>& OutBuffer, FAES::FAESKey InEncryptionKey);
+
+private:
+	uint32 GetDirectory(uint32 DirectoryName, uint32 Parent);
+	uint32 CreateDirectory(const FStringView& DirectoryName, uint32 Parent);
+	uint32 GetNameIndex(const FStringView& String);
+	uint32 AddFile(const FStringView& FileName, uint32 Directory);
+
+	static bool IsValid(uint32 Index)
+	{
+		return Index != ~uint32(0);
+	}
+
+	FString MountPoint;
+	TArray<FIoDirectoryIndexEntry> DirectoryEntries;
+	TArray<FIoFileIndexEntry> FileEntries;
+	TMap<FString, uint32> StringToIndex;
+	TArray<FString> Strings;
+};
diff --git a/Engine/Source/Runtime/Core/Internal/IO/IoDispatcherFileBackend.h b/Engine/Source/Runtime/Core/Internal/IO/IoDispatcherFileBackend.h
new file mode 100644
--- /dev/null
+++ b/Engine/Source/Runtime/Core/Internal/IO/IoDispatcherFileBackend.h
@@ -0,0 +1,185 @@
+// Copyright Epic Games, Inc. All Rights Reserved.
+
+#pragma once
+
+#include "IO/IoDispatcherFileBackendTypes.h"
+#include "IO/IoDispatcher.h"
+#include "IO/IoDispatcherPrivate.h"
+#include "IO/IoStore.h"
+#include "Containers/Array.h"
+#include "Containers/Map.h"
+#include "Stats/Stats.h"
+#include "Async/TaskGraphInterfaces.h"
+#include "HAL/Runnable.h"
+#include "Misc/AES.h"
+#include "GenericPlatform/GenericPlatformFile.h"
+
+class IMappedFileHandle;
+
+struct FFileIoStoreCompressionContext
+{
+	FFileIoStoreCompressionContext* Next = nullptr;
+	uint64 UncompressedBufferSize = 0;
+	uint8* UncompressedBuffer = nullptr;
+};
+
+class FFileIoStoreReader
+{
+public:
+	FFileIoStoreReader(IPlatformFileIoStore& InPlatformImpl);
+	~FFileIoStoreReader();
+	FIoStatus Initialize(const TCHAR* ContainerPath, int32 Order);
+	uint32 GetContainerInstanceId() const
+	{
+		return ContainerFile.ContainerInstanceId;
+	}
+	FIoStatus Close();
+	bool DoesChunkExist(const FIoChunkId& ChunkId) const;
+	TIoStatusOr<uint64> GetSizeForChunk(const FIoChunkId& ChunkId) const;
+	const FIoOffsetAndLength* Resolve(const FIoChunkId& ChunkId) const;
+	const FFileIoStoreContainerFile& GetContainerFile() const { return ContainerFile; }
+	IMappedFileHandle* GetMappedContainerFileHandle(uint64 TocOffset);
+	const FIoContainerId& GetContainerId() const { return ContainerId; }
+	int32 GetOrder() const { return Order; }
+	bool IsEncrypted() const { return EnumHasAnyFlags(ContainerFile.ContainerFlags, EIoContainerFlags::Encrypted); }
+	bool IsSigned() const { return EnumHasAnyFlags(ContainerFile.ContainerFlags, EIoContainerFlags::Signed); }
+	const FGuid& GetEncryptionKeyGuid() const { return ContainerFile.EncryptionKeyGuid; }
+	void SetEncryptionKey(const FAES::FAESKey& Key) { ContainerFile.EncryptionKey = Key; }
+	const FAES::FAESKey& GetEncryptionKey() const { return ContainerFile.EncryptionKey; }
+
+private:
+	const FIoOffsetAndLength* FindChunkInternal(const FIoChunkId& ChunkId) const;
+
+	IPlatformFileIoStore& PlatformImpl;
+	
+	struct FPerfectHashMap
+	{
+		TArray<int32> TocChunkHashSeeds;
+		TArray<uint32> TocChunkHashes;
+		TArray<FIoOffsetAndLength> TocOffsetAndLengths;
+	};
+	FPerfectHashMap PerfectHashMap;
+	TMap<FIoChunkId, FIoOffsetAndLength> TocImperfectHashMapFallback;
+	FFileIoStoreContainerFile ContainerFile;
+	FIoContainerId ContainerId;
+	int32 Order;
+	bool bClosed = false;
+	bool bHasPerfectHashMap = false;
+
+	static TAtomic<uint32> GlobalPartitionIndex;
+	static TAtomic<uint32> GlobalContainerInstanceId;
+};
+
+class FFileIoStoreRequestTracker
+{
+public:
+	FFileIoStoreRequestTracker(FFileIoStoreRequestAllocator& RequestAllocator, FFileIoStoreRequestQueue& RequestQueue);
+	~FFileIoStoreRequestTracker();
+
+	FFileIoStoreCompressedBlock* FindOrAddCompressedBlock(FFileIoStoreBlockKey Key, bool& bOutWasAdded);
+	void RemoveCompressedBlock(const FFileIoStoreCompressedBlock* CompressedBlock);
+	FFileIoStoreReadRequest* FindOrAddRawBlock(FFileIoStoreBlockKey Key, bool& bOutWasAdded);
+	void RemoveRawBlock(const FFileIoStoreReadRequest* RawBlock);
+	void AddReadRequestsToResolvedRequest(FFileIoStoreCompressedBlock* CompressedBlock, FFileIoStoreResolvedRequest& ResolvedRequest);
+	void AddReadRequestsToResolvedRequest(const FFileIoStoreReadRequestList& Requests, FFileIoStoreResolvedRequest& ResolvedRequest);
+	bool CancelIoRequest(FFileIoStoreResolvedRequest& ResolvedRequest);
+	void UpdatePriorityForIoRequest(FFileIoStoreResolvedRequest& ResolvedRequest);
+	void ReleaseIoRequestReferences(FFileIoStoreResolvedRequest& ResolvedRequest);
+
+private:
+	FFileIoStoreRequestAllocator& RequestAllocator;
+	FFileIoStoreRequestQueue& RequestQueue;
+	TMap<FFileIoStoreBlockKey, FFileIoStoreCompressedBlock*> CompressedBlocksMap;
+	TMap<FFileIoStoreBlockKey, FFileIoStoreReadRequest*> RawBlocksMap;
+};
+
+class FFileIoStore final
+	: public FRunnable
+	, public IIoDispatcherFileBackend
+{
+public:
+	FFileIoStore(TUniquePtr<IPlatformFileIoStore>&& PlatformImpl);
+	~FFileIoStore();
+	void Initialize(TSharedRef<const FIoDispatcherBackendContext> Context) override;
+	TIoStatusOr<FIoContainerId> Mount(const TCHAR* ContainerPath, int32 Order, const FGuid& EncryptionKeyGuid, const FAES::FAESKey& EncryptionKey) override;
+	TIoStatusOr<FIoContainerId> Unmount(const TCHAR* ContainerPath) override;
+	bool Resolve(FIoRequestImpl* Request) override;
+	void CancelIoRequest(FIoRequestImpl* Request) override;
+	void UpdatePriorityForIoRequest(FIoRequestImpl* Request) override;
+	bool DoesChunkExist(const FIoChunkId& ChunkId) const override;
+	TIoStatusOr<uint64> GetSizeForChunk(const FIoChunkId& ChunkId) const;
+	FIoRequestImpl* GetCompletedRequests() override;
+	TIoStatusOr<FIoMappedRegion> OpenMapped(const FIoChunkId& ChunkId, const FIoReadOptions& Options) override;
+	void AppendMountedContainers(TSet<FIoContainerId>& OutContainers) override;
+
+	virtual bool Init() override;
+	virtual uint32 Run() override;
+	virtual void Stop() override;
+
+	uint32 GetThreadId() const;
+
+private:
+	class FDecompressAsyncTask
+	{
+	public:
+		FDecompressAsyncTask(FFileIoStore& InOuter, FFileIoStoreCompressedBlock* InCompressedBlock)
+			: Outer(InOuter)
+			, CompressedBlock(InCompressedBlock)
+		{
+
+		}
+
+		static FORCEINLINE TStatId GetStatId()
+		{
+			RETURN_QUICK_DECLARE_CYCLE_STAT(FIoStoreDecompressTask, STATGROUP_TaskGraphTasks);
+		}
+
+		static ENamedThreads::Type GetDesiredThread();
+
+		FORCEINLINE static ESubsequentsMode::Type GetSubsequentsMode()
+		{
+			return ESubsequentsMode::FireAndForget;
+		}
+
+		void DoTask(ENamedThreads::Type CurrentThread, const FGraphEventRef& MyCompletionGraphEvent)
+		{
+			Outer.ScatterBlock(CompressedBlock, true);
+		}
+
+	private:
+		FFileIoStore& Outer;
+		FFileIoStoreCompressedBlock* CompressedBlock;
+	};
+
+	void OnNewPendingRequestsAdded();
+	void ReadBlocks(FFileIoStoreResolvedRequest& ResolvedRequest);
+	void FreeBuffer(FFileIoStoreBuffer& Buffer);
+	FFileIoStoreCompressionContext* AllocCompressionContext();
+	void FreeCompressionContext(FFileIoStoreCompressionContext* CompressionContext);
+	void ScatterBlock(FFileIoStoreCompressedBlock* CompressedBlock, bool bIsAsync);
+	void CompleteDispatcherRequest(FFileIoStoreResolvedRequest* ResolvedRequest);
+	void FinalizeCompressedBlock(FFileIoStoreCompressedBlock* CompressedBlock);
+	void UpdateAsyncIOMinimumPriority();
+
+	uint64 ReadBufferSize = 0;
+	TSharedPtr<const FIoDispatcherBackendContext> BackendContext;
+	FFileIoStoreBlockCache BlockCache;
+	FFileIoStoreBufferAllocator BufferAllocator;
+	FFileIoStoreRequestAllocator RequestAllocator;
+	FFileIoStoreRequestQueue RequestQueue;
+	FFileIoStoreRequestTracker RequestTracker;
+	TUniquePtr<IPlatformFileIoStore> PlatformImpl;
+	FRunnableThread* Thread = nullptr;
+	bool bIsMultithreaded;
+	TAtomic<bool> bStopRequested{ false };
+	mutable FRWLock IoStoreReadersLock;
+	TArray<TUniquePtr<FFileIoStoreReader>> IoStoreReaders;
+	FFileIoStoreCompressionContext* FirstFreeCompressionContext = nullptr;
+	FFileIoStoreCompressedBlock* ReadyForDecompressionHead = nullptr;
+	FFileIoStoreCompressedBlock* ReadyForDecompressionTail = nullptr;
+	FCriticalSection DecompressedBlocksCritical;
+	FFileIoStoreCompressedBlock* FirstDecompressedBlock = nullptr;
+	FIoRequestImpl* CompletedRequestsHead = nullptr;
+	FIoRequestImpl* CompletedRequestsTail = nullptr;
+	EAsyncIOPriorityAndFlags CurrentAsyncIOMinimumPriority = AIOP_MIN;
+};
diff --git a/Engine/Source/Runtime/Core/Internal/IO/IoDispatcherFileBackendTypes.h b/Engine/Source/Runtime/Core/Internal/IO/IoDispatcherFileBackendTypes.h
new file mode 100644
--- /dev/null
+++ b/Engine/Source/Runtime/Core/Internal/IO/IoDispatcherFileBackendTypes.h
@@ -0,0 +1,975 @@
+// Copyright Epic Games, Inc. All Rights Reserved.
+
+#pragma once
+
+#include "IO/IoStore.h"
+#include "IO/IoDispatcherBackend.h"
+#include "ProfilingDebugging/CsvProfiler.h"
+#include "ProfilingDebugging/CountersTrace.h"
+#include "Containers/Ticker.h"
+#include "Modules/ModuleInterface.h"
+
+#ifndef PLATFORM_IMPLEMENTS_IO
+#define PLATFORM_IMPLEMENTS_IO 0
+#endif
+
+#ifndef PLATFORM_IODISPATCHER_MODULE
+#define PLATFORM_IODISPATCHER_MODULE PREPROCESSOR_TO_STRING(PREPROCESSOR_JOIN(PLATFORM_HEADER_NAME, PlatformIoDispatcher))
+#endif
+
+#define IO_DISPATCHER_FILE_STATS (COUNTERSTRACE_ENABLED || CSV_PROFILER)
+
+struct FFileIoStoreCompressionContext;
+
+struct FFileIoStoreContainerFilePartition
+{
+	uint64 FileHandle = 0;
+	uint64 FileSize = 0;
+	uint32 ContainerFileIndex = 0;
+	FString FilePath;
+	TUniquePtr<IMappedFileHandle> MappedFileHandle;
+};
+
+struct FFileIoStoreContainerFile
+{
+	uint64 PartitionSize = 0;
+	uint64 CompressionBlockSize = 0;
+	TArray<FName> CompressionMethods;
+	TArray<FIoStoreTocCompressedBlockEntry> CompressionBlocks;
+	FString FilePath;
+	FGuid EncryptionKeyGuid;
+	FAES::FAESKey EncryptionKey;
+	EIoContainerFlags ContainerFlags;
+	TArray<FSHAHash> BlockSignatureHashes;
+	TArray<FFileIoStoreContainerFilePartition> Partitions;
+	uint32 ContainerInstanceId = 0;
+
+	void GetPartitionFileHandleAndOffset(uint64 TocOffset, uint64& OutFileHandle, uint64& OutOffset) const
+	{
+		int32 PartitionIndex = int32(TocOffset / PartitionSize);
+		const FFileIoStoreContainerFilePartition& Partition = Partitions[PartitionIndex];
+		OutFileHandle = Partition.FileHandle;
+		OutOffset = TocOffset % PartitionSize;
+	}
+};
+
+struct FFileIoStoreBuffer
+{
+	FFileIoStoreBuffer* Next = nullptr;
+	uint8* Memory = nullptr;
+};
+
+struct FFileIoStoreBlockKey
+{
+	union
+	{
+		struct
+		{
+			uint32 FileIndex;
+			uint32 BlockIndex;
+		};
+		uint64 Hash;
+	};
+
+
+	friend bool operator==(const FFileIoStoreBlockKey& A, const FFileIoStoreBlockKey& B)
+	{
+		return A.Hash == B.Hash;
+	}
+
+	friend uint32 GetTypeHash(const FFileIoStoreBlockKey& Key)
+	{
+		return GetTypeHash(Key.Hash);
+	}
+};
+
+struct FFileIoStoreBlockScatter
+{
+	struct FFileIoStoreResolvedRequest* Request = nullptr;
+	uint64 DstOffset = 0;
+	uint64 SrcOffset = 0;
+	uint64 Size = 0;
+};
+
+struct FFileIoStoreCompressedBlock
+{
+	FFileIoStoreCompressedBlock* Next = nullptr;
+	FFileIoStoreBlockKey Key;
+	FName CompressionMethod;
+	uint64 RawOffset = uint64(-1);
+	uint32 UncompressedSize = uint32(-1);
+	uint32 CompressedSize = uint32(-1);
+	uint32 RawSize = uint32(-1);
+	uint32 RefCount = 0;
+	uint32 UnfinishedRawBlocksCount = 0;
+	TArray<struct FFileIoStoreReadRequest*, TInlineAllocator<2>> RawBlocks;
+	TArray<FFileIoStoreBlockScatter, TInlineAllocator<2>> ScatterList;
+	FFileIoStoreCompressionContext* CompressionContext = nullptr;
+	uint8* CompressedDataBuffer = nullptr;
+	FAES::FAESKey EncryptionKey;
+	const FSHAHash* SignatureHash = nullptr;
+	bool bFailed = false;
+	bool bCancelled = false;
+};
+
+struct FFileIoStoreReadRequest
+{
+	enum EQueueStatus
+	{
+		QueueStatus_NotInQueue,
+		QueueStatus_InQueue,
+		QueueStatus_Started
+	};
+
+	FFileIoStoreReadRequest()
+		: Sequence(NextSequence++)
+		, CreationTime(FPlatformTime::Cycles64())
+	{
+	}
+	FFileIoStoreReadRequest* Next = nullptr;
+	FFileIoStoreReadRequest* Previous = nullptr;
+	uint64 FileHandle = uint64(-1);
+	uint64 Offset = uint64(-1);
+	uint64 Size = uint64(-1);
+	FFileIoStoreBlockKey Key;
+	FFileIoStoreBuffer* Buffer = nullptr;
+	uint32 RefCount = 0;
+	uint32 BufferRefCount = 0;
+	TArray<FFileIoStoreCompressedBlock*, TInlineAllocator<8>> CompressedBlocks;
+	const uint32 Sequence;
+	int32 Priority = 0;
+	uint64 CreationTime;	// Potentially used to circuit break request ordering optimizations when outstanding requests have been delayed too long
+	FFileIoStoreBlockScatter ImmediateScatter;
+	uint32 BytesUsed = 0;
+	bool bIsCustomRequest = false;
+	bool bFailed = false;
+	bool bCancelled = false;
+	EQueueStatus QueueStatus = QueueStatus_NotInQueue;
+
+#if DO_CHECK
+	// For debug checks that we are in the correct owning list for our intrusive next/previous pointers
+	uint32 ListCookie = 0;
+#endif
+
+private:
+	static uint32 NextSequence;
+};
+
+#define CHECK_IO_STORE_READ_REQUEST_LIST_MEMBERSHIP (DO_CHECK && 0)
+// Iterator class for traversing and emptying a list FFileIoStoreReadRequestList at the same time
+class FFileIoStoreReadRequestListStealingIterator
+{
+public:
+	FFileIoStoreReadRequestListStealingIterator(const FFileIoStoreReadRequestListStealingIterator&) = delete;
+	FFileIoStoreReadRequestListStealingIterator& operator=(const FFileIoStoreReadRequestListStealingIterator&) = delete;
+
+	FFileIoStoreReadRequestListStealingIterator(FFileIoStoreReadRequestListStealingIterator&& Other)
+	{
+		Current = Other.Current;
+		Next = Other.Next;
+
+		Other.Current = Other.Next = nullptr;
+	}
+
+	FFileIoStoreReadRequestListStealingIterator& operator=(FFileIoStoreReadRequestListStealingIterator&& Other)
+	{
+		Current = Other.Current;
+		Next = Other.Next;
+
+		Other.Current = Other.Next = nullptr;
+
+		return *this;
+	}
+
+	void operator++()
+	{
+		AdvanceTo(Next);
+	}
+
+	FFileIoStoreReadRequest* operator*()
+	{
+		return Current;
+	}
+
+	FFileIoStoreReadRequest* operator->()
+	{
+		return Current;
+	}
+
+	explicit operator bool() const
+	{
+		return Current != nullptr;
+	}
+
+private:
+	friend class FFileIoStoreReadRequestList; // Only the list can construct us
+
+	FFileIoStoreReadRequestListStealingIterator(FFileIoStoreReadRequest* InHead)
+	{
+#if CHECK_IO_STORE_READ_REQUEST_LIST_MEMBERSHIP
+		// None of these nodes are members of the list any more
+		for (FFileIoStoreReadRequest* Cursor = InHead; Cursor; Cursor = Cursor->Next)
+		{
+			Cursor->ListCookie = 0;
+		}
+#endif
+
+		AdvanceTo(InHead);
+	}
+
+	void AdvanceTo(FFileIoStoreReadRequest* NewCurrent)
+	{
+		Current = NewCurrent;
+		if (Current)
+		{
+			// Copy off the next ptr and remove Current from its list so it can safely be added to another list
+			Next = Current->Next;
+			Current->Next = nullptr;
+		}
+		else
+		{
+			Next = nullptr;
+		}
+	}
+
+	FFileIoStoreReadRequest* Current = nullptr;
+	FFileIoStoreReadRequest* Next = nullptr;;
+};
+
+class FFileIoStoreReadRequestListIterator
+{
+public:
+	FFileIoStoreReadRequestListIterator(const FFileIoStoreReadRequestListIterator&) = default;
+	FFileIoStoreReadRequestListIterator& operator=(const FFileIoStoreReadRequestListIterator&) = default;
+
+	void operator++()
+	{
+		checkf(Current->Next == Next, TEXT("FFileIoStoreReadRequestListIterator cached Next pointer != Current->Next - was list modified during iteration?"));
+		AdvanceTo(Next);
+	}
+
+	FFileIoStoreReadRequest* operator*()
+	{
+		return Current;
+	}
+
+	FFileIoStoreReadRequest* operator->()
+	{
+		return Current;
+	}
+
+	explicit operator bool() const
+	{
+		return Current != nullptr;
+	}
+
+	bool operator!=(const FFileIoStoreReadRequestListIterator& Other) const
+	{
+		return Current != Other.Current;
+	}
+
+private:
+	friend class FFileIoStoreReadRequestList; // Only the list can construct us
+
+	FFileIoStoreReadRequestListIterator(FFileIoStoreReadRequest* InHead)
+	{
+		AdvanceTo(InHead);
+	}
+
+	void AdvanceTo(FFileIoStoreReadRequest* NewCurrent)
+	{
+		Current = NewCurrent;
+
+		if (Current)
+		{
+			// Copy off the next ptr to check for edits made to this list during iteration
+			Next = Current->Next;
+		}
+		else
+		{
+			Next = nullptr;
+		}
+	}
+
+	FFileIoStoreReadRequest* Current = nullptr;
+	FFileIoStoreReadRequest* Next = nullptr;
+};
+
+
+// Wrapper for doubly-linked intrusive list of FFileIoStoreReadRequest
+
+class FFileIoStoreReadRequestList
+{
+public:
+	FFileIoStoreReadRequestList()
+#if CHECK_IO_STORE_READ_REQUEST_LIST_MEMBERSHIP
+		: ListCookie(++NextListCookie)
+#endif
+	{
+	}
+
+	// Owns the Next/Previous pointers of the FFileIoStoreReadRequests it contains so it can't be copied
+	FFileIoStoreReadRequestList(const FFileIoStoreReadRequestList& Other) = delete;
+	FFileIoStoreReadRequestList& operator=(const FFileIoStoreReadRequestList& Other) = delete;
+
+	FFileIoStoreReadRequestList(FFileIoStoreReadRequestList&& Other)
+	{
+		Head = Other.Head;
+		Tail = Other.Tail;
+		Other.Head = Other.Tail = nullptr;
+
+#if CHECK_IO_STORE_READ_REQUEST_LIST_MEMBERSHIP
+		for (FFileIoStoreReadRequest* Cursor = Head; Cursor; Cursor = Cursor->Next)
+		{
+			Cursor->ListCookie = ListCookie;
+		}
+#endif
+	}
+
+	FFileIoStoreReadRequestList& operator=(FFileIoStoreReadRequestList&& Other)
+	{
+		check(!Head && !Tail);
+
+		Head = Other.Head;
+		Tail = Other.Tail;
+		Other.Head = Other.Tail = nullptr;
+
+#if CHECK_IO_STORE_READ_REQUEST_LIST_MEMBERSHIP
+		for (FFileIoStoreReadRequest* Cursor = Head; Cursor; Cursor = Cursor->Next)
+		{
+			Cursor->ListCookie = ListCookie;
+		}
+#endif
+
+		return *this;
+	}
+
+#if CHECK_IO_STORE_READ_REQUEST_LIST_MEMBERSHIP
+	~FFileIoStoreReadRequestList()
+	{
+		// If anything is left in this list it should still think we own it
+		for (FFileIoStoreReadRequest* Cursor = Head; Cursor; Cursor = Cursor->Next)
+		{
+			check(Cursor->ListCookie == ListCookie);
+		}
+	}
+#endif
+
+	bool IsEmpty() const
+	{
+		return Head == nullptr;
+	}
+
+	// Steal the whole list for iteration and moving the contents to other lists
+	FFileIoStoreReadRequestListStealingIterator Steal()
+	{
+#if CHECK_IO_STORE_READ_REQUEST_LIST_MEMBERSHIP
+		for (FFileIoStoreReadRequest* Cursor = Head; Cursor; Cursor = Cursor->Next)
+		{
+			check(Cursor->ListCookie == ListCookie);
+		}
+#endif
+
+		FFileIoStoreReadRequest* OldHead = Head;
+		Clear();
+		return FFileIoStoreReadRequestListStealingIterator(OldHead);
+	}
+
+	// Non-stealing iterator
+	FFileIoStoreReadRequestListIterator CreateIterator() const
+	{
+		return FFileIoStoreReadRequestListIterator(Head);
+	}
+
+	FFileIoStoreReadRequestListIterator begin() const {
+		return CreateIterator();
+	}
+	FFileIoStoreReadRequestListIterator end() const {
+		return FFileIoStoreReadRequestListIterator(nullptr);
+	}
+
+	FFileIoStoreReadRequest* PeekHead() const
+	{
+		return Head;
+	}
+
+	void Add(FFileIoStoreReadRequest* Request)
+	{
+#if CHECK_IO_STORE_READ_REQUEST_LIST_MEMBERSHIP
+		check(Request->ListCookie == 0);
+		Request->ListCookie = ListCookie;
+#endif
+
+		if (Tail)
+		{
+			Tail->Next = Request;
+			Request->Previous = Tail;
+		}
+		else
+		{
+			Head = Request;
+			Request->Previous = nullptr;
+		}
+		Tail = Request;
+		Request->Next = nullptr;
+	}
+
+
+	// Remove all FFileIoStoreReadRequests from List and add them to this list
+	void AppendSteal(FFileIoStoreReadRequestList& List)
+	{
+		if (List.Head)
+		{
+			FFileIoStoreReadRequest *ListHead = List.Head, *ListTail = List.Tail;
+			List.Clear();
+			AppendSteal(ListHead, ListTail);
+		}
+	}
+
+	void Remove(FFileIoStoreReadRequest* Request)
+	{
+#if CHECK_IO_STORE_READ_REQUEST_LIST_MEMBERSHIP
+		check(Request);
+		check(Request->ListCookie == ListCookie);
+		Request->ListCookie = 0;
+#endif
+
+		if (Head == Request && Tail == Request)
+		{
+			check(Request->Next == nullptr);
+			check(Request->Previous == nullptr);
+
+			Head = Tail = nullptr;
+		}
+		else if (Head == Request)
+		{
+			check(Request->Previous == nullptr);
+
+			Head = Request->Next;
+			Head->Previous = nullptr;
+			Request->Next = nullptr;
+		}
+		else if (Tail == Request)
+		{
+			check(Request->Next == nullptr);
+
+			Tail = Request->Previous;
+			Tail->Next = nullptr;
+			Request->Previous = nullptr;
+		}
+		else
+		{
+			check(Request->Next != nullptr && Request->Previous != nullptr); // Neither head nor tail should mean both links are live
+
+			Request->Next->Previous = Request->Previous;
+			Request->Previous->Next = Request->Next;
+
+			Request->Next = Request->Previous = nullptr;
+		}
+	}
+
+	void Clear()
+	{
+#if CHECK_IO_STORE_READ_REQUEST_LIST_MEMBERSHIP
+		for (FFileIoStoreReadRequest* Cursor = Head; Cursor; Cursor = Cursor->Next)
+		{
+			check(Cursor->ListCookie == ListCookie);
+			Cursor->ListCookie = 0;
+		}
+#endif
+		Head = Tail = nullptr;
+	}
+
+private:
+	FFileIoStoreReadRequest* Head = nullptr;
+	FFileIoStoreReadRequest* Tail = nullptr;
+
+#if CHECK_IO_STORE_READ_REQUEST_LIST_MEMBERSHIP
+	uint32 ListCookie;
+	static uint32 NextListCookie;
+#endif
+	
+	void AppendSteal(FFileIoStoreReadRequest* ListHead, FFileIoStoreReadRequest* ListTail)
+	{
+		check(ListHead);
+		check(ListTail);
+		check(!ListTail->Next);
+		check(!ListHead->Previous);
+		check(ListTail == ListHead || ListTail->Previous != nullptr);
+		check(ListTail == ListHead || ListHead->Next != nullptr);
+
+#if CHECK_IO_STORE_READ_REQUEST_LIST_MEMBERSHIP
+		for (FFileIoStoreReadRequest* Cursor = ListHead; Cursor; Cursor = Cursor->Next)
+		{
+			check(Cursor->ListCookie == 0);
+			Cursor->ListCookie = ListCookie;
+		}
+#endif
+
+		if (Tail)
+		{
+			Tail->Next = ListHead;
+			ListHead->Previous = Tail;
+		}
+		else
+		{
+			Head = ListHead;
+			ListHead->Previous = nullptr;
+		}
+		Tail = ListTail;
+	}
+};
+
+class FFileIoStoreBufferAllocator
+{
+public:
+	void Initialize(uint64 MemorySize, uint64 BufferSize, uint32 BufferAlignment);
+	FFileIoStoreBuffer* AllocBuffer();
+	void FreeBuffer(FFileIoStoreBuffer* Buffer);
+	uint64 GetBufferSize() const { return BufferSize; }
+
+private:
+	uint64 BufferSize = 0;
+	uint8* BufferMemory = nullptr;
+	FCriticalSection BuffersCritical;
+	FFileIoStoreBuffer* FirstFreeBuffer = nullptr;
+};
+
+class FFileIoStoreBlockCache
+{
+public:
+	FFileIoStoreBlockCache();
+	~FFileIoStoreBlockCache();
+
+	void Initialize(uint64 CacheMemorySize, uint64 ReadBufferSize);
+	bool Read(FFileIoStoreReadRequest* Block);
+	void Store(const FFileIoStoreReadRequest* Block);
+
+private:
+	struct FCachedBlock
+	{
+		FCachedBlock* LruPrev = nullptr;
+		FCachedBlock* LruNext = nullptr;
+		uint64 Key = 0;
+		uint8* Buffer = nullptr;
+	};
+
+	uint8* CacheMemory = nullptr;
+	TMap<uint64, FCachedBlock*> CachedBlocks;
+	FCachedBlock CacheLruHead;
+	FCachedBlock CacheLruTail;
+	uint64 ReadBufferSize = 0;
+};
+
+struct FFileIoStoreReadRequestSortKey
+{
+	uint64 Offset = 0;
+	uint64 Handle = 0;
+	int32 Priority = 0;
+
+	FFileIoStoreReadRequestSortKey() {}
+	FFileIoStoreReadRequestSortKey(FFileIoStoreReadRequest* Request)
+		: Offset(Request->Offset), Handle(Request->FileHandle), Priority(Request->Priority)
+	{
+	}
+};
+
+// Stores FFileIoStoreReadRequest sorted by file handle & offset with a parallel list sorted by insertion order 
+class FFileIoStoreOffsetSortedRequestQueue
+{
+public:
+	FFileIoStoreOffsetSortedRequestQueue(int32 InPriority);
+	FFileIoStoreOffsetSortedRequestQueue(const FFileIoStoreOffsetSortedRequestQueue&) = delete;
+	FFileIoStoreOffsetSortedRequestQueue(FFileIoStoreOffsetSortedRequestQueue&&) = default;
+	FFileIoStoreOffsetSortedRequestQueue& operator=(const FFileIoStoreOffsetSortedRequestQueue&) = delete;
+	FFileIoStoreOffsetSortedRequestQueue& operator=(FFileIoStoreOffsetSortedRequestQueue&&) = default;
+
+	int32 GetPriority() const { return Priority; }
+	bool IsEmpty() const { return Requests.Num() == 0; }
+	
+	// Remove all requests from this container and return them, for switching to a different queue scheme
+	TArray<FFileIoStoreReadRequest*> StealRequests();
+	// Remove all requests whose priority has been changed to something other than the Priority of this queue
+	TArray<FFileIoStoreReadRequest*> RemoveMisprioritizedRequests();
+
+	FFileIoStoreReadRequest* Peek(FFileIoStoreReadRequestSortKey LastSortKey);
+	FFileIoStoreReadRequest* Pop(FFileIoStoreReadRequestSortKey LastSortKey);
+	void Push(FFileIoStoreReadRequest* Request);
+
+private:
+	int32 Priority;
+	int32 PeekRequestIndex = INDEX_NONE;
+	
+	// Requests sorted by file handle & offset
+	TArray<FFileIoStoreReadRequest*> Requests;
+	
+	// Requests sorted by insertion order 
+	// We store this on the heap in case we get moved, FFileIoStoreReadRequest keeps pointers to FFileIoStoreReadRequestList for debugging
+	FFileIoStoreReadRequestList RequestsBySequence;
+
+	FFileIoStoreReadRequest* GetNextInternal(FFileIoStoreReadRequestSortKey LastSortKey, bool bPop);
+
+	static FFileIoStoreReadRequestSortKey RequestSortProjection(FFileIoStoreReadRequest* Request) { return FFileIoStoreReadRequestSortKey(Request); }
+	static bool RequestSortPredicate(const FFileIoStoreReadRequestSortKey& A, const FFileIoStoreReadRequestSortKey& B);
+};
+
+class FFileIoStoreRequestQueue
+{
+public:
+	FFileIoStoreReadRequest* Peek();
+	FFileIoStoreReadRequest* Pop();
+	void Push(FFileIoStoreReadRequest& Request);	// Takes ownership of Request and rewrites its intrustive linked list pointers
+	void Push(FFileIoStoreReadRequestList& Requests); // Consumes the request list and overwrites all intrustive linked list pointers
+	void UpdateOrder();
+	void Lock();
+	void Unlock();
+	void CancelRequestsWithFileHandle(uint64 FileHandle);
+
+private:
+	static bool QueueSortFunc(const FFileIoStoreReadRequest& A, const FFileIoStoreReadRequest& B)
+	{
+		if (A.Priority == B.Priority)
+		{
+			return A.Sequence < B.Sequence;
+		}
+		return A.Priority > B.Priority;
+	}
+	void UpdateSortRequestsByOffset(); // Check if we need to switch sorting schemes based on the CVar
+	void PushToPriorityQueues(FFileIoStoreReadRequest* Request);
+	static int32 QueuePriorityProjection(const FFileIoStoreOffsetSortedRequestQueue& A) { return A.GetPriority(); }
+
+	bool bSortRequestsByOffset = false; // Cached value of CVar controlling whether we use Heap or SortedPriorityQueues
+	
+	// Heap sorted by request order
+	TArray<FFileIoStoreReadRequest*> Heap;
+	FCriticalSection CriticalSection;
+
+	// Queues sorted by increasing priority
+	TArray<FFileIoStoreOffsetSortedRequestQueue> SortedPriorityQueues; 
+	// The last offset, file handle and priority we popped so that we can pop the closest forward read for the next IO operation
+	FFileIoStoreReadRequestSortKey LastSortKey;
+
+#if !UE_BUILD_SHIPPING
+	TMap<int32, uint32> RequestPriorityCounts;
+#endif
+};
+
+template <typename T, uint16 SlabSize = 4096>
+class TIoDispatcherSingleThreadedSlabAllocator
+{
+public:
+	TIoDispatcherSingleThreadedSlabAllocator()
+	{
+		CurrentSlab = new FSlab();
+	}
+
+	~TIoDispatcherSingleThreadedSlabAllocator()
+	{
+		check(CurrentSlab->Allocated == CurrentSlab->Freed);
+		delete CurrentSlab;
+	}
+
+	template <typename... ArgsType>
+	T* Construct(ArgsType&&... Args)
+	{
+		return new(Alloc()) T(Forward<ArgsType>(Args)...);
+	}
+
+	void Destroy(T* Ptr)
+	{
+		Ptr->~T();
+		Free(Ptr);
+	}
+
+private:
+	struct FSlab;
+
+	struct FElement
+	{
+		TTypeCompatibleBytes<T> Data;
+		FSlab* Slab = nullptr;
+	};
+
+	struct FSlab
+	{
+		uint16 Allocated = 0;
+		uint16 Freed = 0;
+		FElement Elements[SlabSize];
+	};
+
+	T* Alloc()
+	{
+		uint16 ElementIndex = CurrentSlab->Allocated++;
+		check(ElementIndex < SlabSize);
+		FElement* Element = CurrentSlab->Elements + ElementIndex;
+		Element->Slab = CurrentSlab;
+		if (CurrentSlab->Allocated == SlabSize)
+		{
+			//TRACE_CPUPROFILER_EVENT_SCOPE(AllocSlab);
+			CurrentSlab = new FSlab();
+		}
+		return Element->Data.GetTypedPtr();
+	}
+
+	void Free(T* Ptr)
+	{
+		FElement* Element = reinterpret_cast<FElement*>(Ptr);
+		FSlab* Slab = Element->Slab;
+		if (++Slab->Freed == SlabSize)
+		{
+			//TRACE_CPUPROFILER_EVENT_SCOPE(FreeSlab);
+			check(Slab->Freed == Slab->Allocated);
+			delete Slab;
+		}
+	}
+
+	FSlab* CurrentSlab = nullptr;
+};
+
+struct FFileIoStoreReadRequestLink
+{
+	FFileIoStoreReadRequestLink(FFileIoStoreReadRequest& InReadRequest)
+		: ReadRequest(InReadRequest)
+	{
+	}
+
+	FFileIoStoreReadRequestLink* Next = nullptr;
+	FFileIoStoreReadRequest& ReadRequest;
+};
+
+class FFileIoStoreRequestAllocator
+{
+public:
+	FFileIoStoreResolvedRequest* AllocResolvedRequest(
+		FIoRequestImpl& InDispatcherRequest,
+		const FFileIoStoreContainerFile& InContainerFile,
+		uint64 InResolvedOffset,
+		uint64 InResolvedSize)
+	{
+		return ResolvedRequestAllocator.Construct(
+			InDispatcherRequest,
+			InContainerFile,
+			InResolvedOffset,
+			InResolvedSize);
+	}
+
+	void Free(FFileIoStoreResolvedRequest* ResolvedRequest)
+	{
+		ResolvedRequestAllocator.Destroy(ResolvedRequest);
+	}
+
+	FFileIoStoreReadRequest* AllocReadRequest()
+	{
+		//TRACE_CPUPROFILER_EVENT_SCOPE(AllocReadRequest);
+		return ReadRequestAllocator.Construct();
+	}
+
+	void Free(FFileIoStoreReadRequest* ReadRequest)
+	{
+		//TRACE_CPUPROFILER_EVENT_SCOPE(FreeReadRequest);
+		ReadRequestAllocator.Destroy(ReadRequest);
+	}
+
+	FFileIoStoreCompressedBlock* AllocCompressedBlock()
+	{
+		//TRACE_CPUPROFILER_EVENT_SCOPE(AllocCompressedBlock);
+		return CompressedBlockAllocator.Construct();
+	}
+
+	void Free(FFileIoStoreCompressedBlock* CompressedBlock)
+	{
+		//TRACE_CPUPROFILER_EVENT_SCOPE(FreeCompressedBlock);
+		CompressedBlockAllocator.Destroy(CompressedBlock);
+	}
+
+	FFileIoStoreReadRequestLink* AllocRequestLink(FFileIoStoreReadRequest* ReadRequest)
+	{
+		check(ReadRequest);
+		return RequestLinkAllocator.Construct(*ReadRequest);
+	}
+
+	void Free(FFileIoStoreReadRequestLink* RequestLink)
+	{
+		RequestLinkAllocator.Destroy(RequestLink);
+	}
+
+private:
+	TIoDispatcherSingleThreadedSlabAllocator<FFileIoStoreResolvedRequest> ResolvedRequestAllocator;
+	TIoDispatcherSingleThreadedSlabAllocator<FFileIoStoreReadRequest> ReadRequestAllocator;
+	TIoDispatcherSingleThreadedSlabAllocator<FFileIoStoreCompressedBlock> CompressedBlockAllocator;
+	TIoDispatcherSingleThreadedSlabAllocator<FFileIoStoreReadRequestLink> RequestLinkAllocator;
+};
+
+struct FFileIoStoreResolvedRequest
+{
+public:
+	FFileIoStoreResolvedRequest(
+		FIoRequestImpl& InDispatcherRequest,
+		const FFileIoStoreContainerFile& InContainerFile,
+		uint64 InResolvedOffset,
+		uint64 InResolvedSize);
+
+	const FFileIoStoreContainerFile& GetContainerFile() const
+	{
+		return ContainerFile;
+	}
+
+	uint64 GetResolvedOffset() const
+	{
+		return ResolvedOffset;
+	}
+
+	uint64 GetResolvedSize() const
+	{
+		return ResolvedSize;
+	}
+
+	int32 GetPriority() const
+	{
+		check(DispatcherRequest);
+		return DispatcherRequest->Priority;
+	}
+
+	bool HasBuffer() const
+	{
+		check(DispatcherRequest);
+		return DispatcherRequest->HasBuffer();
+	}
+
+	void CreateBuffer(uint64 Size)
+	{
+		check(DispatcherRequest);
+		DispatcherRequest->CreateBuffer(Size);
+	}
+
+	FIoBuffer& GetBuffer()
+	{
+		check(DispatcherRequest);
+		return DispatcherRequest->GetBuffer();
+	}
+
+	void AddReadRequestLink(FFileIoStoreReadRequestLink* ReadRequestLink);
+
+private:
+	FIoRequestImpl* DispatcherRequest = nullptr;
+	const FFileIoStoreContainerFile& ContainerFile;
+	FFileIoStoreReadRequestLink* ReadRequestsHead = nullptr;
+	FFileIoStoreReadRequestLink* ReadRequestsTail = nullptr;
+	const uint64 ResolvedOffset;
+	const uint64 ResolvedSize;
+	uint32 UnfinishedReadsCount = 0;
+	bool bFailed = false;
+	bool bCancelled = false;
+
+	friend class FFileIoStore;
+	friend class FFileIoStoreRequestTracker;
+};
+
+
+// Wrapper for sending stats to both insights and csv profiling from any platform-specific dispatcher.
+struct FFileIoStats
+{
+#if CSV_PROFILER
+	static uint64 QueuedFilesystemReadBytes;
+	static uint64 QueuedFilesystemReads;
+
+	static uint64 QueuedUncompressBytesIn;
+	static uint64 QueuedUncompressBytesOut;
+	static uint64 QueuedUncompressBlocks;
+#endif
+
+#if IO_DISPATCHER_FILE_STATS
+	// Used for seek tracking 
+	static uint64 LastHandle;
+	static uint64 LastOffset;
+
+	static uint32 FileIoStoreThreadId;
+	static uint32 IoDispatcherThreadId;
+
+	static FTSTicker::FDelegateHandle TickerHandle;
+#endif
+
+	static float BytesToApproxMB(uint64 Bytes) {
+		return float(double(Bytes) / 1024.0 / 1024.0);
+	}
+	static float BytesToApproxKB(uint64 Bytes) {
+		return float(double(Bytes) / 1024.0);
+	}
+
+	static void Init();
+	static void Shutdown();
+	static bool CsvTick(float DeltaTime);
+
+	static bool IsInIoDispatcherThread();
+	static bool IsInFileIoStoreThread();
+	static void SetDispatcherThreadId(uint32 InIoDispatcherThreadId);
+	static void SetFileIoStoreThreadId(uint32 InFileIoStoreThreadId);
+
+	// Adds an underlying filesystem read to the queue to track num reads/bytes in the queue
+	// There is currently no separate stat for in-flight reads, those are still counted in the queued stat
+	static void OnFilesystemReadsQueued(const FFileIoStoreReadRequestList& Requests);
+	static void OnFilesystemReadsQueued(uint64 NumBytes, int32 NumReads);
+	// Called when the backend commits to the next read to check for estimated seek tracking
+	static void OnFilesystemReadStarted(uint64 Handle, uint64 Offset, uint64 NumBytes);
+	// Underlying filesystem read complete, BytesRead is bytes returned from filesystem, possibly already decompressed on some systems
+	static void OnFilesystemReadsComplete(const FFileIoStoreReadRequestList& CompletedRequests);
+	static void OnFilesystemReadsComplete(int64 BytesRead, int32 NumReads=1);
+	// Backend read complete, BytesRead is bytes returned from either block cache or filesystem
+	static void OnReadComplete(int64 BytesRead);
+	// NumBytes of data passed the decompression stage, even if the compression method was none
+	static void OnDecompressQueued(int64 CompressedBytes, int64 UncompressedBytes);
+	static void OnDecompressComplete(int64 CompressedBytes, int64 UncompressedBytes);
+
+	// Bytes were copied into the buffer provided to users of the io system
+	static void OnBytesScattered(int64 BytesScattered);
+	// Record stats for block cache hit/miss rate & data throughput
+	static void OnBlockCacheStore(uint64 NumBytes);
+	static void OnBlockCacheHit(uint64 NumBytes);
+	static void OnBlockCacheMiss(uint64 NumBytes);
+
+	static void OnCloseHandle(uint64 Handle);
+
+private:
+	// A read was started which the backend considers sequential
+	static void OnSequentialRead();
+	// A read was started which the backend considers a forward or reverse seek
+	static void OnSeek(uint64 LastOffset, uint64 NewOffset);
+	// A read was started from a different file handle from the last one we used
+	static void OnHandleChangeSeek();
+};
+
+struct FInitializePlatformFileIoStoreParams
+{
+	const FWakeUpIoDispatcherThreadDelegate* WakeUpDispatcherThreadDelegate = nullptr;
+	FFileIoStoreRequestAllocator* RequestAllocator = nullptr;
+	FFileIoStoreBufferAllocator* BufferAllocator = nullptr;
+	FFileIoStoreBlockCache* BlockCache = nullptr;
+};
+
+class IPlatformFileIoStore
+{
+public:
+	virtual ~IPlatformFileIoStore() = default;
+	virtual void Initialize(const FInitializePlatformFileIoStoreParams& Params) = 0;
+
+	virtual bool OpenContainer(const TCHAR* ContainerFilePath, uint64& ContainerFileHandle, uint64& ContainerFileSize) = 0;
+	virtual void CloseContainer(uint64 ContainerFileHandle) = 0;
+
+	virtual bool CreateCustomRequests(FFileIoStoreResolvedRequest& ResolvedRequest, FFileIoStoreReadRequestList& OutRequests) = 0;
+	virtual bool StartRequests(FFileIoStoreRequestQueue& RequestQueue) = 0;
+	virtual void GetCompletedRequests(FFileIoStoreReadRequestList& OutRequests) = 0;
+
+	virtual void ServiceNotify() = 0;
+	virtual void ServiceWait() = 0;
+};
+
+class IPlatformFileIoStoreModule : public IModuleInterface
+{
+public:
+	virtual TUniquePtr<IPlatformFileIoStore> CreatePlatformFileIoStore() = 0;
+};
+
+#if PLATFORM_IMPLEMENTS_IO
+CORE_API TUniquePtr<IPlatformFileIoStore> CreatePlatformFileIoStore();
+#endif
+
diff --git a/Engine/Source/Runtime/Core/Internal/IO/IoDispatcherPrivate.h b/Engine/Source/Runtime/Core/Internal/IO/IoDispatcherPrivate.h
new file mode 100644
--- /dev/null
+++ b/Engine/Source/Runtime/Core/Internal/IO/IoDispatcherPrivate.h
@@ -0,0 +1,17 @@
+// Copyright Epic Games, Inc. All Rights Reserved.
+
+#pragma once
+
+#include "IO/IoDispatcher.h"
+
+class FIoRequestImpl;
+
+class FIoBatchImpl
+{
+public:
+	TFunction<void()> Callback;
+	FEvent* Event = nullptr;
+	FGraphEventRef GraphEvent;
+	TAtomic<uint32> UnfinishedRequestsCount{ 0 };
+};
+
diff --git a/Engine/Source/Runtime/Core/Internal/IO/IoStore.h b/Engine/Source/Runtime/Core/Internal/IO/IoStore.h
new file mode 100644
--- /dev/null
+++ b/Engine/Source/Runtime/Core/Internal/IO/IoStore.h
@@ -0,0 +1,236 @@
+// Copyright Epic Games, Inc. All Rights Reserved.
+
+#pragma once
+
+#include "IO/IoDispatcher.h"
+#include "IO/IoDirectoryIndex.h"
+
+/**
+ * I/O store container format version
+ */
+enum class EIoStoreTocVersion : uint8
+{
+	Invalid = 0,
+	Initial,
+	DirectoryIndex,
+	PartitionSize,
+	PerfectHash,
+	LatestPlusOne,
+	Latest = LatestPlusOne - 1
+};
+
+/**
+ * I/O Store TOC header.
+ */
+struct FIoStoreTocHeader
+{
+	static constexpr char TocMagicImg[] = "-==--==--==--==-";
+
+	uint8	TocMagic[16];
+	uint8	Version;
+	uint8	Reserved0 = 0;
+	uint16	Reserved1 = 0;
+	uint32	TocHeaderSize;
+	uint32	TocEntryCount;
+	uint32	TocCompressedBlockEntryCount;
+	uint32	TocCompressedBlockEntrySize;	// For sanity checking
+	uint32	CompressionMethodNameCount;
+	uint32	CompressionMethodNameLength;
+	uint32	CompressionBlockSize;
+	uint32	DirectoryIndexSize;
+	uint32	PartitionCount = 0;
+	FIoContainerId ContainerId;
+	FGuid	EncryptionKeyGuid;
+	EIoContainerFlags ContainerFlags;
+	uint8	Reserved3 = 0;
+	uint16	Reserved4 = 0;
+	uint32	TocChunkHashSeedsCount = 0;
+	uint64	PartitionSize = 0;
+	uint64	Reserved6[6] = { 0 };
+
+	void MakeMagic()
+	{
+		FMemory::Memcpy(TocMagic, TocMagicImg, sizeof TocMagic);
+	}
+
+	bool CheckMagic() const
+	{
+		return FMemory::Memcmp(TocMagic, TocMagicImg, sizeof TocMagic) == 0;
+	}
+};
+
+/**
+ * Combined offset and length.
+ */
+struct FIoOffsetAndLength
+{
+public:
+	inline uint64 GetOffset() const
+	{
+		return OffsetAndLength[4]
+			| (uint64(OffsetAndLength[3]) << 8)
+			| (uint64(OffsetAndLength[2]) << 16)
+			| (uint64(OffsetAndLength[1]) << 24)
+			| (uint64(OffsetAndLength[0]) << 32)
+			;
+	}
+
+	inline uint64 GetLength() const
+	{
+		return OffsetAndLength[9]
+			| (uint64(OffsetAndLength[8]) << 8)
+			| (uint64(OffsetAndLength[7]) << 16)
+			| (uint64(OffsetAndLength[6]) << 24)
+			| (uint64(OffsetAndLength[5]) << 32)
+			;
+	}
+
+	inline void SetOffset(uint64 Offset)
+	{
+		OffsetAndLength[0] = uint8(Offset >> 32);
+		OffsetAndLength[1] = uint8(Offset >> 24);
+		OffsetAndLength[2] = uint8(Offset >> 16);
+		OffsetAndLength[3] = uint8(Offset >>  8);
+		OffsetAndLength[4] = uint8(Offset >>  0);
+	}
+
+	inline void SetLength(uint64 Length)
+	{
+		OffsetAndLength[5] = uint8(Length >> 32);
+		OffsetAndLength[6] = uint8(Length >> 24);
+		OffsetAndLength[7] = uint8(Length >> 16);
+		OffsetAndLength[8] = uint8(Length >> 8);
+		OffsetAndLength[9] = uint8(Length >> 0);
+	}
+
+private:
+	// We use 5 bytes for offset and size, this is enough to represent
+	// an offset and size of 1PB
+	uint8 OffsetAndLength[5 + 5];
+};
+
+enum class FIoStoreTocEntryMetaFlags : uint8
+{
+	None,
+	Compressed		= (1 << 0),
+	MemoryMapped	= (1 << 1)
+};
+
+ENUM_CLASS_FLAGS(FIoStoreTocEntryMetaFlags);
+
+/**
+ * TOC entry meta data
+ */
+struct FIoStoreTocEntryMeta
+{
+	FIoChunkHash ChunkHash;
+	FIoStoreTocEntryMetaFlags Flags;
+};
+
+/**
+ * Compression block entry.
+ */
+struct FIoStoreTocCompressedBlockEntry
+{
+	static constexpr uint32 OffsetBits = 40;
+	static constexpr uint64 OffsetMask = (1ull << OffsetBits) - 1ull;
+	static constexpr uint32 SizeBits = 24;
+	static constexpr uint32 SizeMask = (1 << SizeBits) - 1;
+	static constexpr uint32 SizeShift = 8;
+
+	inline uint64 GetOffset() const
+	{
+		const uint64* Offset = reinterpret_cast<const uint64*>(Data);
+		return *Offset & OffsetMask;
+	}
+
+	inline void SetOffset(uint64 InOffset)
+	{
+		uint64* Offset = reinterpret_cast<uint64*>(Data);
+		*Offset = InOffset & OffsetMask;
+	}
+
+	inline uint32 GetCompressedSize() const
+	{
+		const uint32* Size = reinterpret_cast<const uint32*>(Data) + 1;
+		return (*Size >> SizeShift) & SizeMask;
+	}
+
+	inline void SetCompressedSize(uint32 InSize)
+	{
+		uint32* Size = reinterpret_cast<uint32*>(Data) + 1;
+		*Size |= (uint32(InSize) << SizeShift);
+	}
+
+	inline uint32 GetUncompressedSize() const
+	{
+		const uint32* UncompressedSize = reinterpret_cast<const uint32*>(Data) + 2;
+		return *UncompressedSize & SizeMask;
+	}
+
+	inline void SetUncompressedSize(uint32 InSize)
+	{
+		uint32* UncompressedSize = reinterpret_cast<uint32*>(Data) + 2;
+		*UncompressedSize = InSize & SizeMask;
+	}
+
+	inline uint8 GetCompressionMethodIndex() const
+	{
+		const uint32* Index = reinterpret_cast<const uint32*>(Data) + 2;
+		return static_cast<uint8>(*Index >> SizeBits);
+	}
+
+	inline void SetCompressionMethodIndex(uint8 InIndex)
+	{
+		uint32* Index = reinterpret_cast<uint32*>(Data) + 2;
+		*Index |= uint32(InIndex) << SizeBits;
+	}
+
+private:
+	/* 5 bytes offset, 3 bytes for size / uncompressed size and 1 byte for compresseion method. */
+	uint8 Data[5 + 3 + 3 + 1];
+};
+
+/**
+ * TOC resource read options.
+ */
+enum class EIoStoreTocReadOptions
+{
+	Default,
+	ReadDirectoryIndex	= (1 << 0),
+	ReadTocMeta			= (1 << 1),
+	ReadAll				= ReadDirectoryIndex | ReadTocMeta
+};
+ENUM_CLASS_FLAGS(EIoStoreTocReadOptions);
+
+/**
+ * Container TOC data.
+ */
+struct FIoStoreTocResource
+{
+	enum { CompressionMethodNameLen = 32 };
+
+	FIoStoreTocHeader Header;
+
+	TArray<FIoChunkId> ChunkIds;
+
+	TArray<FIoOffsetAndLength> ChunkOffsetLengths;
+
+	TArray<int32> ChunkHashSeeds;
+
+	TArray<FIoStoreTocCompressedBlockEntry> CompressionBlocks;
+
+	TArray<FName> CompressionMethods;
+
+	TArray<FSHAHash> ChunkBlockSignatures;
+
+	TArray<FIoStoreTocEntryMeta> ChunkMetas;
+
+	TArray<uint8> DirectoryIndexBuffer;
+
+	UE_NODISCARD static FIoStatus Read(const TCHAR* TocFilePath, EIoStoreTocReadOptions ReadOptions, FIoStoreTocResource& OutTocResource);
+
+	UE_NODISCARD static TIoStatusOr<uint64> Write(const TCHAR* TocFilePath, FIoStoreTocResource& TocResource, const FIoContainerSettings& ContainerSettings, const FIoStoreWriterSettings& WriterSettings);
+
+	static uint64 HashChunkIdWithSeed(int32 Seed, const FIoChunkId& ChunkId);
+};
diff --git a/Engine/Source/Runtime/Core/Private/GenericPlatform/GenericPlatformIoDispatcher.cpp b/Engine/Source/Runtime/Core/Private/GenericPlatform/GenericPlatformIoDispatcher.cpp
--- a/Engine/Source/Runtime/Core/Private/GenericPlatform/GenericPlatformIoDispatcher.cpp
+++ b/Engine/Source/Runtime/Core/Private/GenericPlatform/GenericPlatformIoDispatcher.cpp
@@ -40,10 +40,7 @@ void FGenericFileIoStoreEventQueue::ServiceWait()
 	ServiceEvent->Wait();
 }
 
-FGenericFileIoStoreImpl::FGenericFileIoStoreImpl(FGenericFileIoStoreEventQueue& InEventQueue, FFileIoStoreBufferAllocator& InBufferAllocator, FFileIoStoreBlockCache& InBlockCache)
-	: EventQueue(InEventQueue)
-	, BufferAllocator(InBufferAllocator)
-	, BlockCache(InBlockCache)
+FGenericFileIoStoreImpl::FGenericFileIoStoreImpl()
 {
 }
 
@@ -97,7 +94,7 @@ bool FGenericFileIoStoreImpl::StartRequests(FFileIoStoreRequestQueue& RequestQue
 
 	uint8* Dest;
 	check(!NextRequest->ImmediateScatter.Request);
-	NextRequest->Buffer = BufferAllocator.AllocBuffer();
+	NextRequest->Buffer = BufferAllocator->AllocBuffer();
 	if (!NextRequest->Buffer)
 	{
 		RequestQueue.Push(*NextRequest);
@@ -105,7 +102,7 @@ bool FGenericFileIoStoreImpl::StartRequests(FFileIoStoreRequestQueue& RequestQue
 	}
 	Dest = NextRequest->Buffer->Memory;
 
-	if (!BlockCache.Read(NextRequest))
+	if (!BlockCache->Read(NextRequest))
 	{
 		IFileHandle* FileHandle = reinterpret_cast<IFileHandle*>(static_cast<UPTRINT>(NextRequest->FileHandle));
 		 
@@ -129,7 +126,7 @@ bool FGenericFileIoStoreImpl::StartRequests(FFileIoStoreRequestQueue& RequestQue
 				}
 				NextRequest->bFailed = false;
 				TRACE_COUNTER_ADD(IoDispatcherFileBackendFileSystemTotalBytesRead, NextRequest->Size);
-				BlockCache.Store(NextRequest);
+				BlockCache->Store(NextRequest);
 				break;
 			}
 		}
diff --git a/Engine/Source/Runtime/Core/Private/GenericPlatform/GenericPlatformIoDispatcher.h b/Engine/Source/Runtime/Core/Private/GenericPlatform/GenericPlatformIoDispatcher.h
--- a/Engine/Source/Runtime/Core/Private/GenericPlatform/GenericPlatformIoDispatcher.h
+++ b/Engine/Source/Runtime/Core/Private/GenericPlatform/GenericPlatformIoDispatcher.h
@@ -23,29 +23,39 @@ private:
 	FEvent* ServiceEvent = nullptr;
 };
 
-class FGenericFileIoStoreImpl
+class FGenericFileIoStoreImpl : public IPlatformFileIoStore
 {
 public:
-	FGenericFileIoStoreImpl(FGenericFileIoStoreEventQueue& InEventQueue, FFileIoStoreBufferAllocator& InBufferAllocator, FFileIoStoreBlockCache& InBlockCache);
+	FGenericFileIoStoreImpl();
 	~FGenericFileIoStoreImpl();
-	void Initialize(const FWakeUpIoDispatcherThreadDelegate* InWakeUpDispatcherThreadDelegate)
+	void Initialize(const FInitializePlatformFileIoStoreParams& Params) override
 	{
-		WakeUpDispatcherThreadDelegate = InWakeUpDispatcherThreadDelegate;
+		WakeUpDispatcherThreadDelegate = Params.WakeUpDispatcherThreadDelegate;
+		BufferAllocator = Params.BufferAllocator;
+		BlockCache = Params.BlockCache;
 	}
-	bool OpenContainer(const TCHAR* ContainerFilePath, uint64& ContainerFileHandle, uint64& ContainerFileSize);
-	void CloseContainer(uint64 ContainerFileHandle);
-	bool CreateCustomRequests(FFileIoStoreRequestAllocator& RequestAllocator, FFileIoStoreResolvedRequest& ResolvedRequest, FFileIoStoreReadRequestList& OutRequests)
+	bool OpenContainer(const TCHAR* ContainerFilePath, uint64& ContainerFileHandle, uint64& ContainerFileSize) override;
+	void CloseContainer(uint64 ContainerFileHandle) override;
+	bool CreateCustomRequests(FFileIoStoreResolvedRequest& ResolvedRequest, FFileIoStoreReadRequestList& OutRequests) override
 	{
 		return false;
 	}
-	bool StartRequests(FFileIoStoreRequestQueue& RequestQueue);
-	void GetCompletedRequests(FFileIoStoreReadRequestList& OutRequests);
+	bool StartRequests(FFileIoStoreRequestQueue& RequestQueue) override;
+	void GetCompletedRequests(FFileIoStoreReadRequestList& OutRequests) override;
 
+	virtual void ServiceNotify() override
+	{
+		EventQueue.ServiceNotify();
+	}
+	virtual void ServiceWait() override
+	{
+		EventQueue.ServiceWait();
+	}
 private:
 	const FWakeUpIoDispatcherThreadDelegate* WakeUpDispatcherThreadDelegate = nullptr;
-	FGenericFileIoStoreEventQueue& EventQueue;
-	FFileIoStoreBufferAllocator& BufferAllocator;
-	FFileIoStoreBlockCache& BlockCache;
+	FGenericFileIoStoreEventQueue EventQueue;
+	FFileIoStoreBufferAllocator* BufferAllocator = nullptr;
+	FFileIoStoreBlockCache* BlockCache = nullptr;
 
 	FCriticalSection CompletedRequestsCritical;
 	FFileIoStoreReadRequestList CompletedRequests;
diff --git a/Engine/Source/Runtime/Core/Private/IO/IoDirectoryIndex.h b/Engine/Source/Runtime/Core/Private/IO/IoDirectoryIndex.h
deleted file mode 100644
diff --git a/Engine/Source/Runtime/Core/Private/IO/IoDispatcherFileBackend.cpp b/Engine/Source/Runtime/Core/Private/IO/IoDispatcherFileBackend.cpp
--- a/Engine/Source/Runtime/Core/Private/IO/IoDispatcherFileBackend.cpp
+++ b/Engine/Source/Runtime/Core/Private/IO/IoDispatcherFileBackend.cpp
@@ -1,11 +1,13 @@
 // Copyright Epic Games, Inc. All Rights Reserved.
 
 #include "IO/IoDispatcherFileBackend.h"
+#include "Misc/Parse.h"
 #include "Misc/ScopeRWLock.h"
 #include "Misc/StringBuilder.h"
 #include "ProfilingDebugging/CountersTrace.h"
 #include "HAL/PlatformFileManager.h"
 #include "GenericPlatform/GenericPlatformFile.h"
+#include "GenericPlatform/GenericPlatformIoDispatcher.h"
 #include "HAL/IConsoleManager.h"
 #include "Async/AsyncWork.h"
 #include "Async/MappedFileHandle.h"
@@ -17,6 +19,7 @@
 #include "Algo/MinElement.h"
 #include "Templates/Greater.h"
 #include "Containers/Ticker.h"
+#include "Modules/ModuleManager.h"
 
 TRACE_DECLARE_INT_COUNTER(IoDispatcherLatencyCircuitBreaks, TEXT("IoDispatcher/LatencyCircuitBreaks"));
 TRACE_DECLARE_INT_COUNTER(IoDispatcherSeekDistanceCircuitBreaks, TEXT("IoDispatcher/SeekDistanceCircuitBreaks"));
@@ -623,7 +626,7 @@ void FFileIoStoreRequestQueue::CancelRequestsWithFileHandle(const uint64 FileHan
 	}
 }
 
-FFileIoStoreReader::FFileIoStoreReader(FFileIoStoreImpl& InPlatformImpl)
+FFileIoStoreReader::FFileIoStoreReader(IPlatformFileIoStore& InPlatformImpl)
 	: PlatformImpl(InPlatformImpl)
 {
 }
@@ -1055,9 +1058,9 @@ void FFileIoStoreRequestTracker::ReleaseIoRequestReferences(FFileIoStoreResolved
 	RequestAllocator.Free(&ResolvedRequest);
 }
 
-FFileIoStore::FFileIoStore()
+FFileIoStore::FFileIoStore(TUniquePtr<IPlatformFileIoStore>&& InPlatformImpl)
 	: RequestTracker(RequestAllocator, RequestQueue)
-	, PlatformImpl(EventQueue, BufferAllocator, BlockCache)
+	, PlatformImpl(MoveTemp(InPlatformImpl))
 {
 }
 
@@ -1085,7 +1088,12 @@ void FFileIoStore::Initialize(TSharedRef<const FIoDispatcherBackendContext> InCo
 	uint64 CacheMemorySize = uint64(GIoDispatcherCacheSizeMB) << 20ull;
 	BlockCache.Initialize(CacheMemorySize, BufferSize);
 
-	PlatformImpl.Initialize(&BackendContext->WakeUpDispatcherThreadDelegate);
+	PlatformImpl->Initialize({
+		&BackendContext->WakeUpDispatcherThreadDelegate,
+		&RequestAllocator,
+		&BufferAllocator,
+		&BlockCache
+	});
 
 	uint64 DecompressionContextCount = uint64(GIoDispatcherDecompressionWorkerCount > 0 ? GIoDispatcherDecompressionWorkerCount : 4);
 	for (uint64 ContextIndex = 0; ContextIndex < DecompressionContextCount; ++ContextIndex)
@@ -1101,7 +1109,7 @@ void FFileIoStore::Initialize(TSharedRef<const FIoDispatcherBackendContext> InCo
 
 TIoStatusOr<FIoContainerId> FFileIoStore::Mount(const TCHAR* ContainerPath, int32 Order, const FGuid& EncryptionKeyGuid, const FAES::FAESKey& EncryptionKey)
 {
-	TUniquePtr<FFileIoStoreReader> Reader(new FFileIoStoreReader(PlatformImpl));
+	TUniquePtr<FFileIoStoreReader> Reader(new FFileIoStoreReader(*PlatformImpl));
 	FIoStatus IoStatus = Reader->Initialize(ContainerPath, Order);
 	if (!IoStatus.IsOk())
 	{
@@ -1198,7 +1206,7 @@ bool FFileIoStore::Resolve(FIoRequestImpl* Request)
 			if (ResolvedSize > 0)
 			{
 				FFileIoStoreReadRequestList CustomRequests;
-				if (PlatformImpl.CreateCustomRequests(RequestAllocator, *ResolvedRequest, CustomRequests))
+				if (PlatformImpl->CreateCustomRequests(*ResolvedRequest, CustomRequests))
 				{
 					RequestTracker.AddReadRequestsToResolvedRequest(CustomRequests, *ResolvedRequest);
 					FFileIoStats::OnFilesystemReadsQueued(CustomRequests);
@@ -1469,11 +1477,11 @@ FIoRequestImpl* FFileIoStore::GetCompletedRequests()
 	
 	if (!bIsMultithreaded)
 	{
-		while (PlatformImpl.StartRequests(RequestQueue));
+		while (PlatformImpl->StartRequests(RequestQueue));
 	}
 
 	FFileIoStoreReadRequestList CompletedRequests;
-	PlatformImpl.GetCompletedRequests(CompletedRequests);
+	PlatformImpl->GetCompletedRequests(CompletedRequests);
 	for (auto It = CompletedRequests.Steal(); It; ++It)
 	{
 		FFileIoStoreReadRequest* CompletedRequest = *It;
@@ -1702,7 +1710,7 @@ void FFileIoStore::OnNewPendingRequestsAdded()
 {
 	if (bIsMultithreaded)
 	{
-		EventQueue.ServiceNotify();
+		PlatformImpl->ServiceNotify();
 	}
 }
 
@@ -1808,7 +1816,7 @@ void FFileIoStore::ReadBlocks(FFileIoStoreResolvedRequest& ResolvedRequest)
 void FFileIoStore::FreeBuffer(FFileIoStoreBuffer& Buffer)
 {
 	BufferAllocator.FreeBuffer(&Buffer);
-	EventQueue.ServiceNotify();
+	PlatformImpl->ServiceNotify();
 }
 
 FFileIoStoreCompressionContext* FFileIoStore::AllocCompressionContext()
@@ -1857,7 +1865,7 @@ bool FFileIoStore::Init()
 void FFileIoStore::Stop()
 {
 	bStopRequested = true;
-	EventQueue.ServiceNotify();
+	PlatformImpl->ServiceNotify();
 }
 
 uint32 FFileIoStore::Run()
@@ -1865,10 +1873,10 @@ uint32 FFileIoStore::Run()
 	while (!bStopRequested)
 	{
 		UpdateAsyncIOMinimumPriority();
-		if (!PlatformImpl.StartRequests(RequestQueue))
+		if (!PlatformImpl->StartRequests(RequestQueue))
 		{
 			UpdateAsyncIOMinimumPriority();
-			EventQueue.ServiceWait();
+			PlatformImpl->ServiceWait();
 		}
 	}
 	return 0;
@@ -1876,7 +1884,34 @@ uint32 FFileIoStore::Run()
 
 TSharedRef<IIoDispatcherFileBackend> CreateIoDispatcherFileBackend()
 {
-	return MakeShared<FFileIoStore>();
+	const bool bCheckForPlatformImplementation = UE_BUILD_SHIPPING || !FParse::Param(FCommandLine::Get(), TEXT("forcegenericio"));
+
+	if (bCheckForPlatformImplementation)
+	{
+		const TCHAR* ModuleName = ANSI_TO_TCHAR(PLATFORM_IODISPATCHER_MODULE);
+		if (FModuleManager::Get().ModuleExists(ModuleName))
+		{
+			IPlatformFileIoStoreModule* PlatformModule = FModuleManager::LoadModulePtr<IPlatformFileIoStoreModule>(ModuleName);
+			if (PlatformModule)
+			{
+				TUniquePtr<IPlatformFileIoStore> PlatformImpl = PlatformModule->CreatePlatformFileIoStore();
+				if (PlatformImpl.IsValid())
+				{
+					return MakeShared<FFileIoStore>(MoveTemp(PlatformImpl));
+				}
+			}
+		}
+#if PLATFORM_IMPLEMENTS_IO
+		{
+			TUniquePtr<IPlatformFileIoStore> PlatformImpl = CreatePlatformFileIoStore();
+			if (PlatformImpl.IsValid())
+			{
+				return MakeShared<FFileIoStore>(MoveTemp(PlatformImpl));
+			}
+		}
+#endif
+	}
+	return MakeShared<FFileIoStore>(MakeUnique<FGenericFileIoStoreImpl>());
 }
 
 uint32 FFileIoStore::GetThreadId() const
diff --git a/Engine/Source/Runtime/Core/Private/IO/IoDispatcherFileBackend.h b/Engine/Source/Runtime/Core/Private/IO/IoDispatcherFileBackend.h
deleted file mode 100644
diff --git a/Engine/Source/Runtime/Core/Private/IO/IoDispatcherFileBackendTypes.h b/Engine/Source/Runtime/Core/Private/IO/IoDispatcherFileBackendTypes.h
deleted file mode 100644
diff --git a/Engine/Source/Runtime/Core/Private/IO/IoDispatcherPrivate.h b/Engine/Source/Runtime/Core/Private/IO/IoDispatcherPrivate.h
deleted file mode 100644
diff --git a/Engine/Source/Runtime/Core/Private/IO/IoStore.h b/Engine/Source/Runtime/Core/Private/IO/IoStore.h
deleted file mode 100644
