Moved some platform defines to their corresponding Core.Build.cs

#rnx

diff --git a/Engine/Source/Runtime/Core/Public/ProfilingDebugging/TagTrace.h b/Engine/Source/Runtime/Core/Public/ProfilingDebugging/TagTrace.h
--- a/Engine/Source/Runtime/Core/Public/ProfilingDebugging/TagTrace.h
+++ b/Engine/Source/Runtime/Core/Public/ProfilingDebugging/TagTrace.h
@@ -22,7 +22,7 @@ CORE_API int32	MemoryTrace_GetActiveTag();
 ////////////////////////////////////////////////////////////////////////////////
 
 #if !defined(USE_MEMORY_TRACE_TAGS) && UE_TRACE_ENABLED && UE_BUILD_DEVELOPMENT
-	#if PLATFORM_WINDOWS || PLATFORM_PS4 || defined(__PS5__)
+	#if PLATFORM_WINDOWS
 		#define USE_MEMORY_TRACE_TAGS 1
 	#endif
 #endif
