IoDispatcher: Fix non-unity compile error caused by cl 17516499

#rnx

diff --git a/Engine/Source/Runtime/Core/Private/IO/IoDispatcherFileBackend.cpp b/Engine/Source/Runtime/Core/Private/IO/IoDispatcherFileBackend.cpp
--- a/Engine/Source/Runtime/Core/Private/IO/IoDispatcherFileBackend.cpp
+++ b/Engine/Source/Runtime/Core/Private/IO/IoDispatcherFileBackend.cpp
@@ -1,6 +1,7 @@
 // Copyright Epic Games, Inc. All Rights Reserved.
 
 #include "IO/IoDispatcherFileBackend.h"
+#include "Misc/CommandLine.h"
 #include "Misc/Parse.h"
 #include "Misc/ScopeRWLock.h"
 #include "Misc/StringBuilder.h"
