U5 - try to have workers go into drowsing early so that when wakup occurs we don't have to hard trigger sleeping workers instead.

diff --git a/Engine/Source/Runtime/Core/Private/Async/Fundamental/Scheduler.cpp b/Engine/Source/Runtime/Core/Private/Async/Fundamental/Scheduler.cpp
--- a/Engine/Source/Runtime/Core/Private/Async/Fundamental/Scheduler.cpp
+++ b/Engine/Source/Runtime/Core/Private/Async/Fundamental/Scheduler.cpp
@@ -348,6 +348,11 @@ namespace LowLevelTasks
 
 			if (WaitCount < WorkerSpinCycles)
 			{
+				if (WaitCount == 0 && !bDrowsing)
+				{
+					TrySleeping(WorkerEvent, true, bDrowsing, bPermitBackgroundWork);
+				}
+
 				OutOfWork.Start();
 				FPlatformProcess::YieldCycles(WaitCycles);
 				WaitCount++;
diff --git a/Engine/Source/Runtime/Core/Public/Async/Fundamental/Scheduler.h b/Engine/Source/Runtime/Core/Public/Async/Fundamental/Scheduler.h
--- a/Engine/Source/Runtime/Core/Public/Async/Fundamental/Scheduler.h
+++ b/Engine/Source/Runtime/Core/Public/Async/Fundamental/Scheduler.h
@@ -359,7 +359,11 @@ namespace LowLevelTasks
 				WorkerEvent->SleepEvent->Trigger(); 
 				return true; // Solving State two: (((Running -> Drowsing) -> Sleeping) -> Running)
 			}
-			continue; // Solving State one: (Running -> Drowsing) -> Running  OR ((Running -> Drowsing) -> Drowsing) -> Running
+			else if (SleepState == ESleepState::Drowsing)
+			{
+				return true; // Solving State one: (Running -> Drowsing) -> Running  OR ((Running -> Drowsing) -> Drowsing) -> Running
+			}
+			continue; 
 		}
 		return false;
 	}
