use Unordered comparison operator for AVX VectorCompareNE... this matches the normal C behavior.

#rnx

diff --git a/Engine/Source/Runtime/Core/Public/Math/UnrealMathSSE.h b/Engine/Source/Runtime/Core/Public/Math/UnrealMathSSE.h
--- a/Engine/Source/Runtime/Core/Public/Math/UnrealMathSSE.h
+++ b/Engine/Source/Runtime/Core/Public/Math/UnrealMathSSE.h
@@ -1599,7 +1599,9 @@ FORCEINLINE VectorRegister4Double VectorCompareNE(const VectorRegister4Double& V
 	Result.XY = _mm_cmpneq_pd(Vec1.XY, Vec2.XY);
 	Result.ZW = _mm_cmpneq_pd(Vec1.ZW, Vec2.ZW);
 #else
-	Result = _mm256_cmp_pd(Vec1, Vec2, _CMP_NEQ_OQ);
+	// For X != Y, if either is NaN it should return true (this matches the normal C behavior). 
+	// We use the *unordered* comparison operation that is true if either value is NaN.
+	Result = _mm256_cmp_pd(Vec1, Vec2, _CMP_NEQ_UQ);
 #endif
 	return Result;
 }
