Compile fixes for building UnrealPak in shipping config
#rnx

diff --git a/Engine/Source/Developer/IoStoreUtilities/Private/IoStoreUtilities.cpp b/Engine/Source/Developer/IoStoreUtilities/Private/IoStoreUtilities.cpp
--- a/Engine/Source/Developer/IoStoreUtilities/Private/IoStoreUtilities.cpp
+++ b/Engine/Source/Developer/IoStoreUtilities/Private/IoStoreUtilities.cpp
@@ -2,7 +2,6 @@
 
 #include "IoStoreUtilities.h"
 
-#if !UE_BUILD_SHIPPING
 #include "HAL/FileManager.h"
 #include "HAL/PlatformFileManager.h"
 #include "Hash/CityHash.h"
@@ -5083,5 +5082,3 @@ int32 CreateIoStoreContainerFiles(const TCHAR* CmdLine)
 
 	return CreateTarget(Arguments, WriterSettings);
 }
-
-#endif
diff --git a/Engine/Source/Developer/IoStoreUtilities/Public/IoStoreUtilities.h b/Engine/Source/Developer/IoStoreUtilities/Public/IoStoreUtilities.h
--- a/Engine/Source/Developer/IoStoreUtilities/Public/IoStoreUtilities.h
+++ b/Engine/Source/Developer/IoStoreUtilities/Public/IoStoreUtilities.h
@@ -4,8 +4,4 @@
 
 #include "CoreMinimal.h"
 
-#if !UE_BUILD_SHIPPING
-
 IOSTOREUTILITIES_API int32 CreateIoStoreContainerFiles(const TCHAR* CmdLine);
-
-#endif
\ No newline at end of file
diff --git a/Engine/Source/Runtime/Core/Public/UObject/NameBatchSerialization.h b/Engine/Source/Runtime/Core/Public/UObject/NameBatchSerialization.h
--- a/Engine/Source/Runtime/Core/Public/UObject/NameBatchSerialization.h
+++ b/Engine/Source/Runtime/Core/Public/UObject/NameBatchSerialization.h
@@ -6,7 +6,7 @@
 #include "UObject/NameTypes.h"
 #include "Templates/Function.h"
 
-#define ALLOW_NAME_BATCH_SAVING (UE_BUILD_DEBUG || UE_BUILD_DEVELOPMENT || WITH_EDITOR) && PLATFORM_LITTLE_ENDIAN && !PLATFORM_TCHAR_IS_4_BYTES
+#define ALLOW_NAME_BATCH_SAVING PLATFORM_LITTLE_ENDIAN && !PLATFORM_TCHAR_IS_4_BYTES
 
 #if ALLOW_NAME_BATCH_SAVING
 // Save comparison entries in given order to a name blob and a versioned hash blob.
