U5 - fix cases of early drowsing state.

diff --git a/Engine/Source/Runtime/Core/Private/Async/Fundamental/Scheduler.cpp b/Engine/Source/Runtime/Core/Private/Async/Fundamental/Scheduler.cpp
--- a/Engine/Source/Runtime/Core/Private/Async/Fundamental/Scheduler.cpp
+++ b/Engine/Source/Runtime/Core/Private/Async/Fundamental/Scheduler.cpp
@@ -346,20 +346,20 @@ namespace LowLevelTasks
 				break;
 			}
 
-			if (WaitCount < WorkerSpinCycles)
+			if (WaitCount == 0 && !bDrowsing)
+			{
+				TrySleeping(WorkerEvent, OutOfWork.Start(), bDrowsing, bPermitBackgroundWork);
+				WaitCount++;
+			}
+			else if (WaitCount < WorkerSpinCycles)
 			{
-				if (WaitCount == 0 && !bDrowsing)
-				{
-					TrySleeping(WorkerEvent, true, bDrowsing, bPermitBackgroundWork);
-				}
-
-				OutOfWork.Start();
 				FPlatformProcess::YieldCycles(WaitCycles);
 				WaitCount++;
-				continue;
 			}
-
-			TrySleeping(WorkerEvent, OutOfWork.Stop(), bDrowsing, bPermitBackgroundWork);
+			else
+			{
+				bDrowsing = TrySleeping(WorkerEvent, OutOfWork.Stop(), bDrowsing, bPermitBackgroundWork);
+			}			
 		}
 
 		while (WakeUpWorker(bPermitBackgroundWork)) {}
diff --git a/Engine/Source/Runtime/Core/Public/Async/Fundamental/Scheduler.h b/Engine/Source/Runtime/Core/Public/Async/Fundamental/Scheduler.h
--- a/Engine/Source/Runtime/Core/Public/Async/Fundamental/Scheduler.h
+++ b/Engine/Source/Runtime/Core/Public/Async/Fundamental/Scheduler.h
@@ -194,7 +194,7 @@ namespace LowLevelTasks
 		void WorkerMain(struct FSleepEvent* WorkerEvent, FSchedulerTls::FLocalQueueType* ExternalWorkerLocalQueue, uint32 WaitCycles, bool bPermitBackgroundWork);
 		CORE_API void LaunchInternal(FTask& Task, EQueuePreference QueuePreference, bool bWakeUpWorker);
 		CORE_API void BusyWaitInternal(const FConditional& Conditional, bool ForceAllowBackgroundWork);
-		FORCENOINLINE void TrySleeping(FSleepEvent* WorkerEvent, bool bStopOutOfWorkScope, bool& Drowsing, bool bBackgroundWorker);
+		FORCENOINLINE bool TrySleeping(FSleepEvent* WorkerEvent, bool bStopOutOfWorkScope, bool Drowsing, bool bBackgroundWorker);
 		inline bool WakeUpWorker(bool bBackgroundWorker);
 
 		template<FTask* (FSchedulerTls::FLocalQueueType::*DequeueFunction)(bool, bool), bool bIsBusyWaiting>
@@ -316,37 +316,39 @@ namespace LowLevelTasks
 		}
 	}
 
-	inline void FScheduler::TrySleeping(FSleepEvent* WorkerEvent, bool bStopOutOfWorkScope, bool& Drowsing, bool bBackgroundWorker)
+	inline bool FScheduler::TrySleeping(FSleepEvent* WorkerEvent, bool bStopOutOfWorkScope, bool bDrowsing, bool bBackgroundWorker)
 	{
 		ESleepState DrowsingState1 = ESleepState::Drowsing;
 		ESleepState DrowsingState2 = ESleepState::Drowsing;
 		ESleepState RunningState  = ESleepState::Running;
 		ESleepState AffinityState  = ESleepState::Affinity;
 
-		if(!Drowsing && WorkerEvent->SleepState.compare_exchange_strong(DrowsingState1, ESleepState::Drowsing, std::memory_order_release)) //continue drowsing
+		if(!bDrowsing && WorkerEvent->SleepState.compare_exchange_strong(DrowsingState1, ESleepState::Drowsing, std::memory_order_release)) //continue drowsing
 		{
 			verifySlow(bStopOutOfWorkScope);
-			Drowsing = true; // Alternative State one: ((Running -> Drowsing) -> Drowsing)
+			bDrowsing = true; // Alternative State one: ((Running -> Drowsing) -> Drowsing)
 		}
 		else if(WorkerEvent->SleepState.compare_exchange_strong(DrowsingState2, ESleepState::Sleeping, std::memory_order_release))
 		{
 			verifySlow(!bStopOutOfWorkScope);
-			Drowsing = false;
+			bDrowsing = false;
 			WorkerEvent->SleepEvent->Wait(); // State two: ((Running -> Drowsing) -> Sleeping)
 		}
 		else if(WorkerEvent->SleepState.compare_exchange_strong(AffinityState, ESleepState::Drowsing, std::memory_order_release)) //continue drowsing
 		{
-			Drowsing = true; // Alternative State one: ((Running -> Drowsing) -> Drowsing)
+			bDrowsing = true; // Alternative State one: ((Running -> Drowsing) -> Drowsing)
 		}
 		else if(WorkerEvent->SleepState.compare_exchange_strong(RunningState, ESleepState::Drowsing, std::memory_order_release))
 		{
-			Drowsing = true;
+			bDrowsing = true;
 			SleepEventStack[bBackgroundWorker].Push(WorkerEvent); // State one: (Running -> Drowsing)
 		}
 		else
 		{
 			checkf(false, TEXT("Worker was supposed to be running or drowsing: %d"), WorkerEvent->SleepState.load(std::memory_order_relaxed));
 		}
+
+		return bDrowsing;
 	}
 
 	inline bool FScheduler::WakeUpWorker(bool bBackgroundWorker)
