U5 - adding extra verfication constraints.

diff --git a/Engine/Source/Runtime/Core/Private/Async/Fundamental/Scheduler.cpp b/Engine/Source/Runtime/Core/Private/Async/Fundamental/Scheduler.cpp
--- a/Engine/Source/Runtime/Core/Private/Async/Fundamental/Scheduler.cpp
+++ b/Engine/Source/Runtime/Core/Private/Async/Fundamental/Scheduler.cpp
@@ -348,7 +348,7 @@ namespace LowLevelTasks
 
 			if (WaitCount == 0 && !bDrowsing)
 			{
-				TrySleeping(WorkerEvent, OutOfWork.Start(), bDrowsing, bPermitBackgroundWork);
+				verifySlow(TrySleeping(WorkerEvent, OutOfWork.Start(), false, bPermitBackgroundWork));
 				WaitCount++;
 			}
 			else if (WaitCount < WorkerSpinCycles)
