U5 - remove unnecessary continue

diff --git a/Engine/Source/Runtime/Core/Public/Async/Fundamental/Scheduler.h b/Engine/Source/Runtime/Core/Public/Async/Fundamental/Scheduler.h
--- a/Engine/Source/Runtime/Core/Public/Async/Fundamental/Scheduler.h
+++ b/Engine/Source/Runtime/Core/Public/Async/Fundamental/Scheduler.h
@@ -364,8 +364,7 @@ namespace LowLevelTasks
 			else if (SleepState == ESleepState::Drowsing)
 			{
 				return true; // Solving State one: (Running -> Drowsing) -> Running  OR ((Running -> Drowsing) -> Drowsing) -> Running
-			}
-			continue; 
+			} 
 		}
 		return false;
 	}
