Add an UE5 specific version EUnrealEngineObjectUE5Version to be used for global changes instead of EUnrealEngineObjectUEVersion. By splitting and storing both version numbers we allow for hypothetical future UE4 changes that will not conflict when merged to UE5.

#rnx
#tests Ran overnight preflights on several platforms, opened/cooked/staged/ran the oldest version of InfiltratorDemo that can be downloaded (4.11)

### ObjectVersion
- Add a new version enum EUnrealEngineObjectUE5Version.
-- This version number starts at 1000 which leaves more than enough for for EUnrealEngineObjectUEVersion to be expanded
- Even though very few changes (if any at all) to EUnrealEngineObjectUE4Version are expected there is a static assert to make sure that EUnrealEngineObjectUEVersion::AUTOMATIC_VERSION never overtakes EUnrealEngineObjectUE5Version::INITIAL_VERSION.
- Add a struct FPackageFileVersion that wraps around the version numbers and is used to store them instead of raw int32 values which was done before. This should make it easier to add new version numbers in the future if we desire (although this will cause problems in places that serialize the struct directly)

### FPackageFileSummary
- Adding a new entry to CurrentLegacyFileVersion at value -8 which shows the UE5 version being added. This lets us make the changes without needing to submit anything to UE4 Main.
- When loading a package that does not have a UE5 version, it will remain at 0.
- Added ::IsFileVersionTooOld and ::IsFileVersionTooNew to replace hardcoded tests in the code base for version validity. This will make it easier to make changes in the future.
- A few months ago most of the accessors of the version number were deprecated in favour of a version that did not contain the Engine number (ie UE4Ver -> UEVer in Archive) but to work with these changes the renamed methods now will return or accept the version as FPackageFileVersion rather than int32.  The old UE4 methods will remain deprecated and direct licensees to use the new methods.

### Archive
- Now stores the version as a FPackageFileVersion rather than int32

### LinkerLoad
- Reports the larger version number if we detect a higher version number than we support. Note that this could cause an issue if the UE4 version is ever raised but helps keep the code simple.

### AssetData
- Need to add a new version here to manage existing data that only has the UE4 version

### EditorDomain
- We do not need to version the format, we can just invalidate existing editor domain entries via EditorDomainVersion

### EditorServer
- When reporting that a package is too old we report the UE4 version as that is the only version that can be older than VER_UE4_OLDEST_LOADABLE_PACKAGE
- When reporting that a package is too new it can be either the UE4 or the UE5 version so we print them together "UE4Ver|UE5Ver"

### ContentCommandlets
- The min and max resave versions have been kept as a single value, you will not be able to resave against different UE4 and UE5 versions at the same time. It doesn't seem like a useful feature and would greatly increase the complexity of the code.
- We will also only report the file version as a single value.

### ManifestUObject
- This class was setting an older obsolete version on purpose to try and maintain compatibility with older clients so we need to provide a way to create an older UE4 only version that will leave the UE5 version as unset.

### NetworkPlatformFile
- I was unable to test the code path in FNetworkPlatformFile::ProcessServerCachedFilesResponse as I am unsure how to run the game in a mode that will actually use it.
- When reading an older "CookedVersion.txt" that was saved with a single version, the reads will fail and this will count as a version change in the code so that all of the existing files will be deleted. The existing code would not give the user a log message when this happens and given the very small time window where this might happen caused by this change I have opted to leave this alone and not add any additional logging.
- If we do detect a version mismatch we will still only log the version number as a single version.

### CookOnTheFlyServer
- We now add each version number to the IniVersionMap rather than merge the version and license version as a key/value pair. This allows us to a) use both the UE4 and UE5 version numbers b) we now log a warning that the version values don't match when it is changed, previously since it was a key value we would log a warning about an additional setting instead.
-- I also added "vs" to the log message when values are mismatched to make the space between the two values being printed clearer.

diff --git a/Engine/Plugins/Developer/Concert/ConcertMain/Source/Concert/Private/ConcertVersion.cpp b/Engine/Plugins/Developer/Concert/ConcertMain/Source/Concert/Private/ConcertVersion.cpp
--- a/Engine/Plugins/Developer/Concert/ConcertMain/Source/Concert/Private/ConcertVersion.cpp
+++ b/Engine/Plugins/Developer/Concert/ConcertMain/Source/Concert/Private/ConcertVersion.cpp
@@ -7,6 +7,8 @@
 
 #define LOCTEXT_NAMESPACE "ConcertVersion"
 
+
+
 namespace ConcertVersionUtil
 {
 
@@ -49,13 +51,16 @@ bool ValidateVersion(const int32 InCurrent, const int32 InOther, const FText InV
 
 void FConcertFileVersionInfo::Initialize()
 {
-	FileVersion = GPackageFileUEVersion;
+	FileVersion = GPackageFileUEVersion.FileVersionUE4;
+	FileVersionUE5 = GPackageFileUEVersion.FileVersionUE5;
+
 	FileVersionLicensee = GPackageFileLicenseeUEVersion;
 }
 
 bool FConcertFileVersionInfo::Validate(const FConcertFileVersionInfo& InOther, const EConcertVersionValidationMode InValidationMode, FText* OutFailureReason) const
 {
-	return ConcertVersionUtil::ValidateVersion(FileVersion, InOther.FileVersion, LOCTEXT("PackageVersionName", "Package Version"), InValidationMode, OutFailureReason)
+	return ConcertVersionUtil::ValidateVersion(FileVersion, InOther.FileVersion, LOCTEXT("PackageVersionName", "UE4 Package Version"), InValidationMode, OutFailureReason)
+		&& ConcertVersionUtil::ValidateVersion(FileVersionUE5, InOther.FileVersionUE5, LOCTEXT("PackageVersionName", "UE5 Package Version"), InValidationMode, OutFailureReason)
 		&& ConcertVersionUtil::ValidateVersion(FileVersionLicensee, InOther.FileVersionLicensee, LOCTEXT("LicenseePackageVersionName", "Licensee Package Version"), InValidationMode, OutFailureReason);
 }
 
diff --git a/Engine/Plugins/Developer/Concert/ConcertMain/Source/Concert/Public/ConcertVersion.h b/Engine/Plugins/Developer/Concert/ConcertMain/Source/Concert/Public/ConcertVersion.h
--- a/Engine/Plugins/Developer/Concert/ConcertMain/Source/Concert/Public/ConcertVersion.h
+++ b/Engine/Plugins/Developer/Concert/ConcertMain/Source/Concert/Public/ConcertVersion.h
@@ -29,9 +29,13 @@ struct FConcertFileVersionInfo
 	/** Validate this version info against another */
 	CONCERT_API bool Validate(const FConcertFileVersionInfo& InOther, const EConcertVersionValidationMode InValidationMode, FText* OutFailureReason = nullptr) const;
 
-	/* File version */
+	/* UE4 File version */
 	UPROPERTY()
-	int32 FileVersion = 0;
+	int32 FileVersion = 0; // Ideally we would rename this FileVersionUE4
+
+	/* UE5 File version */
+	UPROPERTY()
+	int32 FileVersionUE5 = 0;
 
 	/* Licensee file version */
 	UPROPERTY()
diff --git a/Engine/Plugins/Developer/Concert/ConcertSync/ConcertSyncCore/Source/ConcertSyncCore/Private/ConcertSyncArchives.cpp b/Engine/Plugins/Developer/Concert/ConcertSync/ConcertSyncCore/Source/ConcertSyncCore/Private/ConcertSyncArchives.cpp
--- a/Engine/Plugins/Developer/Concert/ConcertSync/ConcertSyncCore/Source/ConcertSyncCore/Private/ConcertSyncArchives.cpp
+++ b/Engine/Plugins/Developer/Concert/ConcertSync/ConcertSyncCore/Source/ConcertSyncCore/Private/ConcertSyncArchives.cpp
@@ -180,7 +180,9 @@ FConcertSyncObjectReader::FConcertSyncObjectReader(const FConcertLocalIdentifier
 
 	if (InVersionInfo)
 	{
-		SetUEVer(InVersionInfo->FileVersion.FileVersion);
+		FPackageFileVersion UEVersion(InVersionInfo->FileVersion.FileVersion, (EUnrealEngineObjectUE5Version)InVersionInfo->FileVersion.FileVersionUE5);
+
+		SetUEVer(UEVersion);
 		SetLicenseeUEVer(InVersionInfo->FileVersion.FileVersionLicensee);
 		SetEngineVer(FEngineVersionBase(InVersionInfo->EngineVersion.Major, InVersionInfo->EngineVersion.Minor, InVersionInfo->EngineVersion.Patch, InVersionInfo->EngineVersion.Changelist));
 
diff --git a/Engine/Plugins/VirtualProduction/LevelSnapshots/Source/LevelSnapshots/Private/Data/SnapshotVersion.cpp b/Engine/Plugins/VirtualProduction/LevelSnapshots/Source/LevelSnapshots/Private/Data/SnapshotVersion.cpp
--- a/Engine/Plugins/VirtualProduction/LevelSnapshots/Source/LevelSnapshots/Private/Data/SnapshotVersion.cpp
+++ b/Engine/Plugins/VirtualProduction/LevelSnapshots/Source/LevelSnapshots/Private/Data/SnapshotVersion.cpp
@@ -9,7 +9,8 @@
 
 void FSnapshotFileVersionInfo::Initialize()
 {
-	FileVersion = GPackageFileUEVersion;
+	FileVersion = GPackageFileUEVersion.FileVersionUE4;
+	FileVersionUE5 = GPackageFileUEVersion.FileVersionUE5;
 	FileVersionLicensee = GPackageFileLicenseeUEVersion;
 }
 
@@ -54,7 +55,9 @@ void FSnapshotVersionInfo::ApplyToArchive(FArchive& Archive) const
 {
 	if (ensure(Archive.IsLoading()))
 	{
-		Archive.SetUEVer(FileVersion.FileVersion);
+		FPackageFileVersion UEVersion(FileVersion.FileVersion, (EUnrealEngineObjectUE5Version)FileVersion.FileVersionUE5);
+
+		Archive.SetUEVer(UEVersion);
 		Archive.SetLicenseeUEVer(FileVersion.FileVersionLicensee);
 		Archive.SetEngineVer(FEngineVersionBase(EngineVersion.Major, EngineVersion.Minor, EngineVersion.Patch, EngineVersion.Changelist));
 
diff --git a/Engine/Plugins/VirtualProduction/LevelSnapshots/Source/LevelSnapshots/Public/Data/SnapshotVersion.h b/Engine/Plugins/VirtualProduction/LevelSnapshots/Source/LevelSnapshots/Public/Data/SnapshotVersion.h
--- a/Engine/Plugins/VirtualProduction/LevelSnapshots/Source/LevelSnapshots/Public/Data/SnapshotVersion.h
+++ b/Engine/Plugins/VirtualProduction/LevelSnapshots/Source/LevelSnapshots/Public/Data/SnapshotVersion.h
@@ -18,9 +18,13 @@ struct LEVELSNAPSHOTS_API FSnapshotFileVersionInfo
 	/** Initialize this version info from the compiled in data */
 	void Initialize();
 
-	/* UE File version */
+	/* UE4 File version */
 	UPROPERTY()
-	int32 FileVersion = 0;
+	int32 FileVersion = 0; // Ideally we would rename this FileVersionUE4
+
+	/* UE5 File version */
+	UPROPERTY()
+	int32 FileVersionUE5 = 0;
 
 	/* Licensee file version */
 	UPROPERTY()
diff --git a/Engine/Source/Editor/BlueprintGraph/Private/K2Node.cpp b/Engine/Source/Editor/BlueprintGraph/Private/K2Node.cpp
--- a/Engine/Source/Editor/BlueprintGraph/Private/K2Node.cpp
+++ b/Engine/Source/Editor/BlueprintGraph/Private/K2Node.cpp
@@ -148,7 +148,7 @@ bool UK2Node::HasNonEditorOnlyReferences() const
 
 void UK2Node::FixupPinDefaultValues()
 {
-	const int32 LinkerUEVersion = GetLinkerUEVersion();
+	const FPackageFileVersion LinkerUEVersion = GetLinkerUEVersion();
 	const int32 LinkerFrameworkVersion = GetLinkerCustomVersion(FFrameworkObjectVersion::GUID);
 	const UEdGraphSchema_K2* K2Schema = GetDefault<UEdGraphSchema_K2>();
 
diff --git a/Engine/Source/Editor/UnrealEd/Private/Commandlets/ContentCommandlets.cpp b/Engine/Source/Editor/UnrealEd/Private/Commandlets/ContentCommandlets.cpp
--- a/Engine/Source/Editor/UnrealEd/Private/Commandlets/ContentCommandlets.cpp
+++ b/Engine/Source/Editor/UnrealEd/Private/Commandlets/ContentCommandlets.cpp
@@ -430,8 +430,8 @@ int32 UResavePackagesCommandlet::InitializeResaveParameters( const TArray<FStrin
 	}
 	if ( Switches.Contains(TEXT("CHECKUE4VER")) )
 	{
-		// Limits resaving to packages with this ue4 package version or lower.
-		MaxResaveUEVersion = FMath::Max<int32>(GPackageFileUEVersion - 1, 0);
+		// Limits resaving to packages with this package version or lower.
+		MaxResaveUEVersion = FMath::Max<int32>(GPackageFileUEVersion.ToValue() - 1, 0);
 	}
 	else if ( Switches.Contains(TEXT("RESAVEDEPRECATED")) )
 	{
@@ -448,7 +448,7 @@ int32 UResavePackagesCommandlet::InitializeResaveParameters( const TArray<FStrin
 			{
 				if ( MinResaveUEVersion == CURRENT_PACKAGE_VERSION )
 				{
-					MinResaveUEVersion = GPackageFileUEVersion;
+					MinResaveUEVersion = GPackageFileUEVersion.ToValue();
 				}
 			}
 
@@ -456,7 +456,7 @@ int32 UResavePackagesCommandlet::InitializeResaveParameters( const TArray<FStrin
 			{
 				if ( MaxResaveUEVersion == CURRENT_PACKAGE_VERSION )
 				{
-					MaxResaveUEVersion = GPackageFileUEVersion;
+					MaxResaveUEVersion = GPackageFileUEVersion.ToValue();
 				}
 			}
 		}
@@ -782,7 +782,7 @@ void UResavePackagesCommandlet::LoadAndSaveOnePackage(const FString& Filename)
 				if( bIsReadOnly == true && bVerifyContent == true && bAutoCheckOut == false )
 				{
 					UE_LOG(LogContentCommandlet, Warning, TEXT("Package [%s] is read-only but needs to be resaved (UE Version: %i, Licensee Version: %i  Current UE Version: %i, Current Licensee Version: %i)"),
-						*Filename, Linker->Summary.GetFileVersionUE(), Linker->Summary.GetFileVersionLicenseeUE(), GPackageFileUEVersion, VER_LATEST_ENGINE_LICENSEEUE4 );
+						*Filename, Linker->Summary.GetFileVersionUE().ToValue(), Linker->Summary.GetFileVersionLicenseeUE(), GPackageFileUEVersion.ToValue(), VER_LATEST_ENGINE_LICENSEEUE4 );
 					if( SavePackageHelper(Package, FString(TEXT("Temp.temp"))) )
 					{
 						UE_LOG(LogContentCommandlet, Warning, TEXT("Correctly saved:  [Temp.temp].") );
@@ -848,7 +848,7 @@ void UResavePackagesCommandlet::LoadAndSaveOnePackage(const FString& Filename)
 					if (Verbosity != ONLY_ERRORS)
 					{
 						UE_LOG(LogContentCommandlet, Display, TEXT("Resaving package [%s] (UE Version: %i, Licensee Version: %i  Saved UE Version: %i, Saved Licensee Version: %i)"),
-							*Filename,Linker->Summary.GetFileVersionUE(), Linker->Summary.GetFileVersionLicenseeUE(), GPackageFileUEVersion, VER_LATEST_ENGINE_LICENSEEUE4 );
+							*Filename,Linker->Summary.GetFileVersionUE().ToValue(), Linker->Summary.GetFileVersionLicenseeUE(), GPackageFileUEVersion.ToValue(), VER_LATEST_ENGINE_LICENSEEUE4 );
 					}
 
 					const static bool bKeepPackageGUIDOnSave = FParse::Param(FCommandLine::Get(), TEXT("KeepPackageGUIDOnSave"));
@@ -1368,10 +1368,9 @@ FText UResavePackagesCommandlet::GetChangelistDescription() const
 
 void UResavePackagesCommandlet::PerformPreloadOperations( FLinkerLoad* PackageLinker, bool& bSavePackage )
 {
-	const int32 UEPackageVersion = PackageLinker->Summary.GetFileVersionUE();
+	const int32 UEPackageVersion = PackageLinker->Summary.GetFileVersionUE().ToValue();
 	const int32 LicenseeUEPackageVersion = PackageLinker->Summary.GetFileVersionLicenseeUE();
 
-
 	// validate that this package meets the minimum requirement
 	if ( MinResaveUEVersion != IGNORE_PACKAGE_VERSION && UEPackageVersion < MinResaveUEVersion )
 	{
diff --git a/Engine/Source/Editor/UnrealEd/Private/Commandlets/PackageUtilities.cpp b/Engine/Source/Editor/UnrealEd/Private/Commandlets/PackageUtilities.cpp
--- a/Engine/Source/Editor/UnrealEd/Private/Commandlets/PackageUtilities.cpp
+++ b/Engine/Source/Editor/UnrealEd/Private/Commandlets/PackageUtilities.cpp
@@ -781,7 +781,7 @@ int32 ULoadPackageCommandlet::Main( const FString& Params )
 		if (bCheckForLegacyPackages)
 		{
 			FLinkerLoad* Linker = LoadPackageLinker(nullptr, PackagePath, LOAD_NoVerify);
-			MinVersion = FMath::Min<int32>(MinVersion, Linker->Summary.GetFileVersionUE());
+			MinVersion = FMath::Min<int32>(MinVersion, Linker->Summary.GetFileVersionUE().ToValue());
 		}
 		else
 		{
@@ -945,7 +945,7 @@ void FPkgInfoReporter_Log::GeneratePackageReport( FLinkerLoad* InLinker /*=nullp
 	Out.Logf(ELogVerbosity::Display, TEXT("--------------------------------------------") );
 
 	Out.Logf(ELogVerbosity::Display, TEXT("\t         Filename: %s"), *Linker->GetPackagePath().GetLocalFullPath());
-	Out.Logf(ELogVerbosity::Display, TEXT("\t     File Version: %i"), Linker->UEVer() );
+	Out.Logf(ELogVerbosity::Display, TEXT("\t     File Version: %i"), Linker->UEVer().ToValue());
 	Out.Logf(ELogVerbosity::Display, TEXT("\t   Engine Version: %s"), *Linker->Summary.SavedByEngineVersion.ToString());
 	Out.Logf(ELogVerbosity::Display, TEXT("\t   Compat Version: %s"), *Linker->Summary.CompatibleWithEngineVersion.ToString());
 	Out.Logf(ELogVerbosity::Display, TEXT("\t     PackageFlags: %X"), Linker->Summary.GetPackageFlags() );
diff --git a/Engine/Source/Editor/UnrealEd/Private/CookOnTheFlyServer.cpp b/Engine/Source/Editor/UnrealEd/Private/CookOnTheFlyServer.cpp
--- a/Engine/Source/Editor/UnrealEd/Private/CookOnTheFlyServer.cpp
+++ b/Engine/Source/Editor/UnrealEd/Private/CookOnTheFlyServer.cpp
@@ -5074,9 +5074,9 @@ void GetAdditionalCurrentIniVersionStrings( const ITargetPlatform* TargetPlatfor
 		IniVersionMap.Add(CustomVersionString, CustomVersionValue);
 	}
 
-	FString UEVer = FString::Printf(TEXT("PackageFileVersions:%d"), GPackageFileUEVersion);
-	FString UE4Value = FString::Printf(TEXT("%d"), GPackageFileLicenseeUEVersion);
-	IniVersionMap.Add(UEVer, UE4Value);
+	IniVersionMap.Add(TEXT("PackageFileVersionUE4"), FString::Printf(TEXT("%d"), GPackageFileUEVersion.FileVersionUE4));
+	IniVersionMap.Add(TEXT("PackageFileVersionUE5"), FString::Printf(TEXT("%d"), GPackageFileUEVersion.FileVersionUE5));
+	IniVersionMap.Add(TEXT("PackageLicenseeVersion"), FString::Printf(TEXT("%d"), GPackageFileLicenseeUEVersion));
 
 	/*FString UE4EngineVersionCompatibleName = TEXT("EngineVersionCompatibleWith");
 	FString UE4EngineVersionCompatible = FEngineVersion::CompatibleWith().ToString();
@@ -5437,7 +5437,7 @@ bool UCookOnTheFlyServer::IniSettingsOutOfDate(const ITargetPlatform* TargetPlat
 
 		if ( *CurrentValue != OldIniSetting.Value )
 		{
-			UE_LOG(LogCook, Display, TEXT("Additional Setting from previous cook %s doesn't match %s %s"), *OldIniSetting.Key, **CurrentValue, *OldIniSetting.Value );
+			UE_LOG(LogCook, Display, TEXT("Additional Setting from previous cook %s doesn't match %s vs %s"), *OldIniSetting.Key, **CurrentValue, *OldIniSetting.Value );
 			return true;
 		}
 	}
diff --git a/Engine/Source/Editor/UnrealEd/Private/EditorDomain/EditorDomainUtils.cpp b/Engine/Source/Editor/UnrealEd/Private/EditorDomain/EditorDomainUtils.cpp
--- a/Engine/Source/Editor/UnrealEd/Private/EditorDomain/EditorDomainUtils.cpp
+++ b/Engine/Source/Editor/UnrealEd/Private/EditorDomain/EditorDomainUtils.cpp
@@ -45,7 +45,8 @@ bool GTargetDomainClassUseAllowList = true;
 bool GTargetDomainClassEmptyAllowList = false;
 
 // Change to a new guid when EditorDomain needs to be invalidated
-const TCHAR* EditorDomainVersion = TEXT("86E449B024FD412BB80DB1397EFAE34D");
+const TCHAR* EditorDomainVersion = TEXT("B9190641FCE0441EBC818BBE6C265735");
+											
 // Identifier of the CacheBuckets for EditorDomain tables
 const TCHAR* EditorDomainPackageBucketName = TEXT("EditorDomainPackage");
 const TCHAR* EditorDomainBulkDataListBucketName = TEXT("EditorDomainBulkDataList");
@@ -71,8 +72,8 @@ EPackageDigestResult AppendPackageDigest(FCbWriter& Writer, bool& bOutEditorDoma
 	const FAssetPackageData& PackageData, FName PackageName)
 {
 	bOutEditorDomainEnabled = true;
-
-	int32 CurrentFileVersionUE = GPackageFileUEVersion;
+	
+	FPackageFileVersion CurrentFileVersionUE = GPackageFileUEVersion;
 	int32 CurrentFileVersionLicenseeUE = GPackageFileLicenseeUEVersion;
 	Writer << EditorDomainVersion;
 	Writer << GetEditorDomainSaveUnversioned();
diff --git a/Engine/Source/Editor/UnrealEd/Private/EditorServer.cpp b/Engine/Source/Editor/UnrealEd/Private/EditorServer.cpp
--- a/Engine/Source/Editor/UnrealEd/Private/EditorServer.cpp
+++ b/Engine/Source/Editor/UnrealEd/Private/EditorServer.cpp
@@ -2393,14 +2393,12 @@ bool UEditorEngine::PackageIsAMapFile( const TCHAR* PackageFilename, FText& OutN
 			return false;
 		}
 
-		const int32 UEVersion = Summary.GetFileVersionUE();
-
 		// Validate the summary.
-		if( UEVersion < VER_UE4_OLDEST_LOADABLE_PACKAGE )
+		if (Summary.IsFileVersionTooOld())
 		{
 			FFormatNamedArguments Arguments;
 			Arguments.Add(TEXT("File"), FText::FromString( FString( PackageFilename ) ));
-			Arguments.Add(TEXT("Version"), UEVersion);
+			Arguments.Add(TEXT("Version"), Summary.GetFileVersionUE().FileVersionUE4);
 			Arguments.Add(TEXT("First"), VER_UE4_OLDEST_LOADABLE_PACKAGE);
 			OutNotMapReason = FText::Format( LOCTEXT( "UEFileIsOlder", "{File} is an UE4 map [File:v{Version}], from an engine release no longer supported [Min:v{First}]." ), 
 				Arguments);
@@ -2410,17 +2408,19 @@ bool UEditorEngine::PackageIsAMapFile( const TCHAR* PackageFilename, FText& OutN
 		const int32 UELicenseeVersion = Summary.GetFileVersionLicenseeUE();
 
 		// Don't load packages that were saved with an engine version newer than the current one.
-		if( UEVersion > GPackageFileUEVersion )
+		if (Summary.IsFileVersionTooNew())
 		{
 			FFormatNamedArguments Arguments;
 			Arguments.Add(TEXT("File"), FText::FromString( FString( PackageFilename ) ));
-			Arguments.Add(TEXT("Version"), UEVersion);
-			Arguments.Add(TEXT("Last"), GPackageFileUEVersion);
-			OutNotMapReason = FText::Format( LOCTEXT( "UEFileIsNewer", "{File} is a UE map [File:v{Version}], from an engine release newer than this [Cur:v{Last}]." ), 
+			Arguments.Add(TEXT("UE4Version"), Summary.GetFileVersionUE().FileVersionUE4);
+			Arguments.Add(TEXT("UE5Version"), Summary.GetFileVersionUE().FileVersionUE5);
+			Arguments.Add(TEXT("UE4Lastest"), GPackageFileUEVersion.FileVersionUE4);
+			Arguments.Add(TEXT("UE5Lastest"), GPackageFileUEVersion.FileVersionUE5);
+			OutNotMapReason = FText::Format( LOCTEXT( "UEFileIsNewer", "{File} is a UE map [File:v{UE4Version}|{UE5Version}], from an engine release newer than this [Cur:v{UE4Lastest}|{UE5Lastest}]." ), 
 				Arguments);
 			return false;
 		}
-		else if( UELicenseeVersion > GPackageFileLicenseeUEVersion )
+		else if (UELicenseeVersion > GPackageFileLicenseeUEVersion)
 		{
 			FFormatNamedArguments Arguments;
 			Arguments.Add(TEXT("File"), FText::FromString( FString( PackageFilename ) ));
diff --git a/Engine/Source/Runtime/AssetRegistry/Private/PackageReader.cpp b/Engine/Source/Runtime/AssetRegistry/Private/PackageReader.cpp
--- a/Engine/Source/Runtime/AssetRegistry/Private/PackageReader.cpp
+++ b/Engine/Source/Runtime/AssetRegistry/Private/PackageReader.cpp
@@ -80,7 +80,7 @@ bool FPackageReader::OpenPackageFile(EOpenPackageResult* OutErrorCode)
 	}
 
 	// Don't read packages that are too old
-	if( PackageFileSummary.GetFileVersionUE() < VER_UE4_OLDEST_LOADABLE_PACKAGE )
+	if (PackageFileSummary.IsFileVersionTooOld())
 	{
 		UE_LOG(LogAssetRegistry, Error, TEXT("Package %s is too old"), *PackageFilename);
 		SetPackageErrorCode(EOpenPackageResult::VersionTooOld);
@@ -88,7 +88,7 @@ bool FPackageReader::OpenPackageFile(EOpenPackageResult* OutErrorCode)
 	}
 
 	// Don't read packages that were saved with an package version newer than the current one.
-	if( (PackageFileSummary.GetFileVersionUE() > GPackageFileUEVersion) || (PackageFileSummary.GetFileVersionLicenseeUE() > GPackageFileLicenseeUEVersion) )
+	if (PackageFileSummary.IsFileVersionTooNew() || PackageFileSummary.GetFileVersionLicenseeUE() > GPackageFileLicenseeUEVersion)
 	{
 		UE_LOG(LogAssetRegistry, Error, TEXT("Package %s is too new"), *PackageFilename);
 		SetPackageErrorCode(EOpenPackageResult::VersionTooNew);
diff --git a/Engine/Source/Runtime/Core/Private/Serialization/Archive.cpp b/Engine/Source/Runtime/Core/Private/Serialization/Archive.cpp
--- a/Engine/Source/Runtime/Core/Private/Serialization/Archive.cpp
+++ b/Engine/Source/Runtime/Core/Private/Serialization/Archive.cpp
@@ -751,7 +751,7 @@ void FArchive::SerializeCompressedNew(void* V, int64 Length, FName CompressionFo
 		else
 		{
 			UE_LOG(LogSerialization, Log, TEXT("ArchiveName: %s"), *GetArchiveName());
-			UE_LOG(LogSerialization, Log, TEXT("Archive UE4 Version: %d"), UEVer());
+			UE_LOG(LogSerialization, Log, TEXT("Archive UE4 Version: %d"), UEVer().ToValue());
 			UE_LOG(LogSerialization, Log, TEXT("Archive Licensee Version: %d"), LicenseeUEVer());
 			UE_LOG(LogSerialization, Log, TEXT("Position: %lld"), Tell());
 			UE_LOG(LogSerialization, Log, TEXT("Read Size: %lld"), Length);
@@ -792,7 +792,7 @@ void FArchive::SerializeCompressedNew(void* V, int64 Length, FName CompressionFo
 		if ( ! FCompression::IsFormatValid(CompressionFormatToDecode) )
 		{
 			UE_LOG(LogSerialization, Log, TEXT("ArchiveName: %s"), *GetArchiveName());
-			UE_LOG(LogSerialization, Log, TEXT("Archive UE4 Version: %d"), UEVer());
+			UE_LOG(LogSerialization, Log, TEXT("Archive UE4 Version: %d"), UEVer().ToValue());
 			UE_LOG(LogSerialization, Log, TEXT("Archive Licensee Version: %d"), LicenseeUEVer());
 			UE_LOG(LogSerialization, Log, TEXT("Position: %lld"), Tell());
 			UE_LOG(LogSerialization, Log, TEXT("Read Size: %lld"), Length);
@@ -1342,7 +1342,7 @@ void FArchive::LogfImpl(const TCHAR* Fmt, ...)
 	FMemory::SystemFree( Buffer );
 }
 
-void FArchiveState::SetUEVer(int32 InVer)
+void FArchiveState::SetUEVer(FPackageFileVersion InVer)
 {
 	ArUEVer = InVer;
 }
diff --git a/Engine/Source/Runtime/Core/Private/UObject/ObjectVersion.cpp b/Engine/Source/Runtime/Core/Private/UObject/ObjectVersion.cpp
--- a/Engine/Source/Runtime/Core/Private/UObject/ObjectVersion.cpp
+++ b/Engine/Source/Runtime/Core/Private/UObject/ObjectVersion.cpp
@@ -6,10 +6,55 @@
 
 #include "UObject/ObjectVersion.h"
 
+#include "Serialization/Archive.h"
+#include "Serialization/CompactBinary.h"
+#include "Serialization/CompactBinaryWriter.h"
+
 // @see ObjectVersion.h for the list of changes/defines
-const int32 GPackageFileUEVersion = VER_LATEST_ENGINE_UE4;
+const int32 GPackageFileLicenseeUE4Version = GPackageFileLicenseeUEVersion;
+
+const FPackageFileVersion GPackageFileUEVersion(VER_LATEST_ENGINE_UE4, EUnrealEngineObjectUE5Version::AUTOMATIC_VERSION);
+const FPackageFileVersion GOldestLoadablePackageFileUEVersion = FPackageFileVersion::CreateUE4Version(VER_UE4_OLDEST_LOADABLE_PACKAGE);
+
+const int32 GPackageFileUE4Version = VER_LATEST_ENGINE_UE4;
 const int32 GPackageFileLicenseeUEVersion = VER_LATEST_ENGINE_LICENSEEUE4;
 
-const int32 GPackageFileUE4Version = GPackageFileUEVersion;
-const int32 GPackageFileLicenseeUE4Version = GPackageFileLicenseeUEVersion;
+FPackageFileVersion FPackageFileVersion::CreateUE4Version(int32 Version)
+{
+	check(Version <= EUnrealEngineObjectUE4Version::VER_UE4_AUTOMATIC_VERSION);
+	return FPackageFileVersion(Version, (EUnrealEngineObjectUE5Version)0);
+}
+
+FPackageFileVersion FPackageFileVersion::CreateUE4Version(EUnrealEngineObjectUE4Version Version)
+{
+	check(Version <= EUnrealEngineObjectUE4Version::VER_UE4_AUTOMATIC_VERSION);
+	return FPackageFileVersion(Version, (EUnrealEngineObjectUE5Version)0);
+}
+
+FArchive& operator<<(FArchive& Ar, FPackageFileVersion& Version)
+{
+	Ar << Version.FileVersionUE4;
+	Ar << Version.FileVersionUE5;
+
+	return Ar;
+}
+
+FCbWriter& operator<<(FCbWriter& Writer, const FPackageFileVersion& Version)
+{
+	Writer.BeginObject();
+	Writer << "ue4version" << Version.FileVersionUE4;
+	Writer << "ue5version" << Version.FileVersionUE5;
+	Writer.EndObject();
+
+	return Writer;
+}
+
+FPackageFileVersion FromCbObject(const FCbObject& Obj)
+{
+	FPackageFileVersion Version;
+
+	Version.FileVersionUE4 = Obj["ue4version"].AsInt32();
+	Version.FileVersionUE5 = Obj["ue5version"].AsInt32();
 
+	return Version;
+}
diff --git a/Engine/Source/Runtime/Core/Public/Serialization/Archive.h b/Engine/Source/Runtime/Core/Public/Serialization/Archive.h
--- a/Engine/Source/Runtime/Core/Public/Serialization/Archive.h
+++ b/Engine/Source/Runtime/Core/Public/Serialization/Archive.h
@@ -3,19 +3,20 @@
 #pragma once
 
 #include "CoreTypes.h"
-#include "Misc/VarArgs.h"
-#include "Misc/AssertionMacros.h"
-#include "Templates/EnableIf.h"
-#include "Templates/IsEnumClass.h"
-#include "Templates/Function.h"
 #include "HAL/PlatformProperties.h"
+#include "Internationalization/TextNamespaceFwd.h"
+#include "Misc/AssertionMacros.h"
 #include "Misc/CompressionFlags.h"
 #include "Misc/EngineVersionBase.h"
-#include "Internationalization/TextNamespaceFwd.h"
-#include "Templates/IsValidVariadicFunctionArg.h"
+#include "Misc/VarArgs.h"
 #include "Templates/AndOrNot.h"
+#include "Templates/EnableIf.h"
+#include "Templates/Function.h"
 #include "Templates/IsArrayOrRefOfType.h"
+#include "Templates/IsEnumClass.h"
 #include "Templates/IsSigned.h"
+#include "Templates/IsValidVariadicFunctionArg.h"
+#include "UObject/ObjectVersion.h"
 
 class FArchive;
 class FCustomVersionContainer;
@@ -186,7 +187,7 @@ public:
 	}
 
 	/** Returns the global engine serialization version used for this archive. */
-	FORCEINLINE int32 UEVer() const
+	FORCEINLINE FPackageFileVersion UEVer() const
 	{
 		return ArUEVer;
 	}
@@ -198,10 +199,10 @@ public:
 	}
 
 	/** Returns the global engine serialization version used for this archive. */
-	UE_DEPRECATED(5.0, "Use UEVer instead")
+	UE_DEPRECATED(5.0, "Use UEVer instead which returns the version as a FPackageFileVersion. See the @FPackageFileVersion documentation for further details")
 	FORCEINLINE int32 UE4Ver() const
 	{
-		return UEVer();
+		return ArUEVer.FileVersionUE4;
 	}
 
 	/** Returns the licensee-specific version used for this archive, will be 0 by default. */
@@ -877,12 +878,13 @@ public:
 	 *
 	 * @param UEVer	new version number
 	 */
-	virtual void SetUEVer(int32 InVer);
+	virtual void SetUEVer(FPackageFileVersion InVer);
 
-	UE_DEPRECATED(5.0, "Use SetUEVer instead")
+	UE_DEPRECATED(5.0, "Use SetUEVer instead which takes the version as a FPackageFileVersion. See the @FPackageFileVersion documentation for further details")
 	FORCEINLINE void SetUE4Ver(int32 InVer)
 	{
-		SetUEVer(InVer);
+		FPackageFileVersion PackageFileVersion = FPackageFileVersion::CreateUE4Version((EUnrealEngineObjectUE4Version)InVer);
+		SetUEVer(PackageFileVersion);
 	}
 
 	/**
@@ -920,7 +922,7 @@ public:
 // These will be private in FArchive
 protected:
 	/** Holds the archive version. */
-	int32 ArUEVer;
+	FPackageFileVersion ArUEVer;
 
 	/** Holds the archive version for licensees. */
 	int32 ArLicenseeUEVer;
@@ -1787,6 +1789,7 @@ public:
 	}
 
 	using FArchiveState::UEVer;
+	using FArchiveState::UE4Ver;
 	using FArchiveState::LicenseeUEVer;
 	using FArchiveState::EngineVer;
 	using FArchiveState::EngineNetVer;
@@ -2039,6 +2042,7 @@ public:
 	using FArchiveState::SetForceUnicode;
 	using FArchiveState::SetIsPersistent;
 	using FArchiveState::SetUEVer;
+	using FArchiveState::SetUE4Ver;
 	using FArchiveState::SetLicenseeUEVer;
 	using FArchiveState::SetEngineVer;
 	using FArchiveState::SetEngineNetVer;
diff --git a/Engine/Source/Runtime/Core/Public/UObject/ObjectVersion.h b/Engine/Source/Runtime/Core/Public/UObject/ObjectVersion.h
--- a/Engine/Source/Runtime/Core/Public/UObject/ObjectVersion.h
+++ b/Engine/Source/Runtime/Core/Public/UObject/ObjectVersion.h
@@ -4,6 +4,10 @@
 
 #include "CoreTypes.h"
 
+class FArchive;
+class FCbObject;
+class FCbWriter;
+
 // Prevents incorrect files from being loaded.
 
 #define PACKAGE_FILE_TAG			0x9E2A83C1
@@ -13,7 +17,42 @@
 // This is useful as enum vales cannot be seen by the preprocessor.
 #define PREPROCESSOR_ENUM_PROTECT(a) ((unsigned int)(a))
 
-enum EUnrealEngineObjectUEVersion
+/** 
+ * Versioning
+ * 
+ * We now have two enums for global versioning, EUnrealEngineObjectUE5Version and EUnrealEngineObjectUE4Version
+ * which can be updated independently of each other. This allows us to potentially make version changed in the UE4
+ * code base in the future and still have packages saved with those changes to be compatible with the UE5 code base.
+ * 
+ * The struct FPackageFileVersion should be used to contain the version number rather than storing the UE4/UE5 specific
+ * version numbers. It provides overloads for testing against both version enum types for convenience.
+ * 
+ * Note that when reporting versioning problems we often use FPackageFileVersion::ToValue which returns the highest 
+ * version valid in the structure. This allows us to keep reporting code simple and hide the fact that there are two
+ * version numbers from end users. However this does mean that if a version mismatch is encountered due to a future 
+ * extension of EUnrealEngineObjectUE4Version, then the error message given might look odd. Although given that we are 
+ * unlikely to be changing EUnrealEngineObjectUE4Version in the future the risk is deemed worth it.
+ */
+
+
+/** This enum is the version for UE5 and should ONLY be edited in UE5Main branch! */
+enum class EUnrealEngineObjectUE5Version : uint32
+{
+	// Note that currently the oldest loadable package version is EUnrealEngineObjectUEVersion::VER_UE4_OLDEST_LOADABLE_PACKAGE
+	// this can be enabled should we ever deprecate UE4 versions entirely
+	//OLDEST_LOADABLE_PACKAGE = ???,
+
+	// The original UE5 version, at the time this was added the UE4 version was 522, so UE5 will start from 1000 to show a clear difference
+	INITIAL_VERSION = 1000,
+
+	// -----<new versions can be added before this line>-------------------------------------------------
+	// - this needs to be the last line (see note below)
+	AUTOMATIC_VERSION_PLUS_ONE,
+	AUTOMATIC_VERSION = AUTOMATIC_VERSION_PLUS_ONE - 1
+};
+
+/** This enum is the version for UE4 and should ONLY be edited in UE4 Main branch! */
+enum EUnrealEngineObjectUE4Version
 {
 	VER_UE4_OLDEST_LOADABLE_PACKAGE = 214,
 
@@ -656,7 +695,117 @@ enum EUnrealEngineObjectLicenseeUEVersion
 	VER_LIC_AUTOMATIC_VERSION = VER_LIC_AUTOMATIC_VERSION_PLUS_ONE - 1
 };
 
+static_assert(EUnrealEngineObjectUE4Version::VER_UE4_AUTOMATIC_VERSION < (int32)EUnrealEngineObjectUE5Version::INITIAL_VERSION, "UE4 versioning must be lower than the start of UE5 versioning");
+
+/** 
+ * This object combines all of our version enums into a single easy to use structure
+ * which allows us to update older version numbers independently of the newer version numbers. 
+ */ 
+struct FPackageFileVersion
+{
+	FPackageFileVersion() = default;
+	FPackageFileVersion(int32 UE4Version, EUnrealEngineObjectUE5Version UE5Version)
+		: FileVersionUE4(UE4Version)
+		, FileVersionUE5((int32)UE5Version)
+	{
+
+	}
+
+	/** Set all versions to the default state */
+	void Reset()
+	{
+		FileVersionUE4 = 0;
+		FileVersionUE5 = 0;
+	}
+
+	/** Creates and returns a FPackageFileVersion based on a single EUnrealEngineObjectUEVersion and no other versions. */
+	CORE_API static FPackageFileVersion CreateUE4Version(int32 Version);
+	CORE_API static FPackageFileVersion CreateUE4Version(EUnrealEngineObjectUE4Version Version);
+
+	/** Returns the highest valid version number which is considered to be the 'true' version number */
+	int32 ToValue() const
+	{
+		if (FileVersionUE5 >= (int32)EUnrealEngineObjectUE5Version::INITIAL_VERSION)
+		{
+			return FileVersionUE5;
+		}
+		else
+		{
+			return FileVersionUE4;
+		}
+	}
+
+	/** UE4 version comparisons */
+	bool operator !=(EUnrealEngineObjectUE4Version Version) const
+	{
+		return FileVersionUE4 != Version;
+	}
+	
+	/** UE4 version comparisons */
+	bool operator <(EUnrealEngineObjectUE4Version Version) const
+	{
+		return FileVersionUE4 < Version;
+	}
+	
+	/** UE4 version comparisons */
+	bool operator >=(EUnrealEngineObjectUE4Version Version) const
+	{
+		return FileVersionUE4 >= Version;
+	}
+
+	/** UE5 version comparisons */
+	bool operator !=(EUnrealEngineObjectUE5Version Version) const
+	{
+		return FileVersionUE5 != (int32)Version;
+	}
+	
+	/** UE5 version comparisons  */
+	bool operator <(EUnrealEngineObjectUE5Version Version) const
+	{
+		return FileVersionUE5 < (int32)Version;
+	}
+
+	/** UE5 version comparisons */
+	bool operator >=(EUnrealEngineObjectUE5Version Version) const
+	{
+		return FileVersionUE5 >= (int32)Version;
+	}
+
+	/** 
+	 * Returns true if this object is compatible with the FPackageFileVersion passed in as the parameter. 
+	 * This means that  all version numbers for the current object are equal or greater than the 
+	 * corresponding version numbers of the other structure.
+	 */
+	bool IsCompatible(const FPackageFileVersion& Other) const
+	{
+		return FileVersionUE4 >= Other.FileVersionUE4 && FileVersionUE5 >= Other.FileVersionUE5;
+	}
+
+	/** FPackageFileVersion comparisons */
+	bool operator==(const FPackageFileVersion& Other) const
+	{
+		return FileVersionUE4 == Other.FileVersionUE4 && FileVersionUE5 == Other.FileVersionUE5;
+	}
+
+	/** FPackageFileVersion comparisons */
+	bool operator!=(const FPackageFileVersion& Other) const
+	{
+		return !(*this == Other);
+	}
+
+	/** Serialization members */
+	CORE_API friend FArchive& operator<<(FArchive& Ar, FPackageFileVersion& Version);
+	CORE_API friend FCbWriter& operator<<(FCbWriter& Writer, const FPackageFileVersion& Version);
+	CORE_API static FPackageFileVersion FromCbObject(const FCbObject& Obj);
+
+	/* UE4 file version*/
+	int32		FileVersionUE4 = 0;
+	/* UE5 file version */
+	int32		FileVersionUE5 = 0;
+};
+
 #define VER_LATEST_ENGINE_UE4           PREPROCESSOR_ENUM_PROTECT(VER_UE4_AUTOMATIC_VERSION)
+#define VER_LATEST_ENGINE_UE5           PREPROCESSOR_ENUM_PROTECT(EUnrealEngineObjectUE5Version::AUTOMATIC_VERSION)
 #define VER_UE4_DEPRECATED_PACKAGE      PREPROCESSOR_ENUM_PROTECT(VER_UE4_OLDEST_LOADABLE_PACKAGE)
 #define VER_LATEST_ENGINE_LICENSEEUE4   PREPROCESSOR_ENUM_PROTECT(VER_LIC_AUTOMATIC_VERSION)
 
@@ -666,11 +815,14 @@ enum EUnrealEngineObjectLicenseeUEVersion
 
 // Version access.
 
-UE_DEPRECATED(5.0, "Use GPackageFileUEVersion instead")
-extern const CORE_API int32			GPackageFileUE4Version;				// Unreal Version Number.
+UE_DEPRECATED(5.0, "Use GPackageFileUEVersion instead. See the @FPackageFileVersion documentation for further details")
+extern const CORE_API int32					GPackageFileUE4Version;				// Unreal4 Version Number.
 
 UE_DEPRECATED(5.0, "Use GPackageFileLicenseeUEVersion instead")
-extern const CORE_API int32			GPackageFileLicenseeUE4Version;		// Licensee Version Number.
+extern const CORE_API int32					GPackageFileLicenseeUE4Version;		// Licensee Version Number.
+
+extern const CORE_API FPackageFileVersion	GPackageFileUEVersion;				// Unreal Version Number.
+extern const CORE_API FPackageFileVersion	GOldestLoadablePackageFileUEVersion;// The oldest Unreal Version number that we can load.
+extern const CORE_API int32					GPackageFileLicenseeUEVersion;		// Licensee Version Number.
+
 
-extern const CORE_API int32			GPackageFileUEVersion;			// Unreal Version Number.
-extern const CORE_API int32			GPackageFileLicenseeUEVersion;	// Licensee Version Number.
diff --git a/Engine/Source/Runtime/CoreUObject/Private/Misc/WorldCompositionUtility.cpp b/Engine/Source/Runtime/CoreUObject/Private/Misc/WorldCompositionUtility.cpp
--- a/Engine/Source/Runtime/CoreUObject/Private/Misc/WorldCompositionUtility.cpp
+++ b/Engine/Source/Runtime/CoreUObject/Private/Misc/WorldCompositionUtility.cpp
@@ -32,7 +32,7 @@ FArchive& operator<<( FArchive& Ar, FWorldTileLayer& D )
 void operator<<(FStructuredArchive::FSlot Slot, FWorldTileLayer& D)
 {
 	FStructuredArchive::FRecord Record = Slot.EnterRecord();
-	int32 Version = Slot.GetUnderlyingArchive().UEVer();
+	const FPackageFileVersion Version = Slot.GetUnderlyingArchive().UEVer();
 
 	// Serialized with FPackageFileSummary
 	Record << SA_VALUE(TEXT("Name"), D.Name) << SA_VALUE(TEXT("Reserved0"), D.Reserved0) << SA_VALUE(TEXT("Reserved1"), D.Reserved1);
@@ -116,7 +116,7 @@ FArchive& operator<<( FArchive& Ar, FWorldTileInfo& D )
 void operator<<(FStructuredArchive::FSlot Slot, FWorldTileInfo& D)
 {
 	FStructuredArchive::FRecord Record = Slot.EnterRecord();
-	int32 ArchiveVersion = Slot.GetUnderlyingArchive().UEVer();
+	const FPackageFileVersion ArchiveVersion = Slot.GetUnderlyingArchive().UEVer();
 
 	// Serialized with FPackageFileSummary
 	Record << SA_VALUE(TEXT("Position"), D.Position) << SA_VALUE(TEXT("Bounds"), D.Bounds) << SA_VALUE(TEXT("Layer"), D.Layer);
diff --git a/Engine/Source/Runtime/CoreUObject/Private/Serialization/AsyncLoading.cpp b/Engine/Source/Runtime/CoreUObject/Private/Serialization/AsyncLoading.cpp
--- a/Engine/Source/Runtime/CoreUObject/Private/Serialization/AsyncLoading.cpp
+++ b/Engine/Source/Runtime/CoreUObject/Private/Serialization/AsyncLoading.cpp
@@ -7279,7 +7279,7 @@ void FAsyncArchive::ReadCallback(bool bWasCancelled, IAsyncReadRequest* Request)
 			FBufferReader Ar(Mem, FMath::Min<int64>(FMaxPackageSummarySize::Value, FileSize),/*bInFreeOnClose=*/ false, /*bIsPersistent=*/ true);
 			FPackageFileSummary Sum;
 			Ar << Sum;
-			if (Ar.IsError() || Sum.TotalHeaderSize > FileSize || Sum.GetFileVersionUE() < VER_UE4_OLDEST_LOADABLE_PACKAGE)
+			if (Ar.IsError() || Sum.TotalHeaderSize > FileSize || Sum.IsFileVersionTooOld())
 			{
 				SetError();
 			}
diff --git a/Engine/Source/Runtime/CoreUObject/Private/UObject/LinkerLoad.cpp b/Engine/Source/Runtime/CoreUObject/Private/UObject/LinkerLoad.cpp
--- a/Engine/Source/Runtime/CoreUObject/Private/UObject/LinkerLoad.cpp
+++ b/Engine/Source/Runtime/CoreUObject/Private/UObject/LinkerLoad.cpp
@@ -1217,10 +1217,10 @@ FLinkerLoad::ELinkerStatus FLinkerLoad::SerializePackageFileSummaryInternal()
 	}
 
 	// Validate the summary.
-	if (Summary.GetFileVersionUE() < VER_UE4_OLDEST_LOADABLE_PACKAGE)
+	if (Summary.IsFileVersionTooOld())
 	{
 		UE_LOG(LogLinker, Warning, TEXT("The file %s was saved by a previous version which is not backwards compatible with this one. Min Required Version: %i  Package Version: %i"),
-			*GetDebugName(), (int32)VER_UE4_OLDEST_LOADABLE_PACKAGE, Summary.GetFileVersionUE());
+			*GetDebugName(), (int32)VER_UE4_OLDEST_LOADABLE_PACKAGE, Summary.GetFileVersionUE().FileVersionUE4);
 		return LINKER_Failed;
 	}
 
@@ -1281,11 +1281,10 @@ FLinkerLoad::ELinkerStatus FLinkerLoad::SerializePackageFileSummaryInternal()
 	}
 
 	// Don't load packages that were saved with package version newer than the current one.
-	if (bLoaderVersionCheck && ((Summary.GetFileVersionUE() > GPackageFileUEVersion)
-								|| (Summary.GetFileVersionLicenseeUE() > GPackageFileLicenseeUEVersion)))
+	if (bLoaderVersionCheck && ((Summary.IsFileVersionTooNew()) || (Summary.GetFileVersionLicenseeUE() > GPackageFileLicenseeUEVersion)))
 	{
 		UE_LOG(LogLinker, Warning, TEXT("Unable to load package (%s) PackageVersion %i, MaxExpected %i : LicenseePackageVersion %i, MaxExpected %i."),
-			*GetDebugName(), Summary.GetFileVersionUE(), GPackageFileUEVersion, Summary.GetFileVersionLicenseeUE(), GPackageFileLicenseeUEVersion);
+			*GetDebugName(), Summary.GetFileVersionUE().ToValue(), GPackageFileUEVersion.ToValue(), Summary.GetFileVersionLicenseeUE(), GPackageFileLicenseeUEVersion);
 		return LINKER_Failed;
 	}
 
diff --git a/Engine/Source/Runtime/CoreUObject/Private/UObject/LinkerSave.cpp b/Engine/Source/Runtime/CoreUObject/Private/UObject/LinkerSave.cpp
--- a/Engine/Source/Runtime/CoreUObject/Private/UObject/LinkerSave.cpp
+++ b/Engine/Source/Runtime/CoreUObject/Private/UObject/LinkerSave.cpp
@@ -40,8 +40,8 @@ FLinkerSave::FLinkerSave(UPackage* InParent, const TCHAR* InFilename, bool bForc
 		UPackage* Package = dynamic_cast<UPackage*>(LinkerRoot);
 
 		// Set main summary info.
-		Summary.Tag           = PACKAGE_FILE_TAG;
-		Summary.SetFileVersions( GPackageFileUEVersion, GPackageFileLicenseeUEVersion, bInSaveUnversioned );
+		Summary.Tag = PACKAGE_FILE_TAG;
+		Summary.SetToLatestFileVersions(bInSaveUnversioned);
 		Summary.SavedByEngineVersion = FEngineVersion::Current();
 		Summary.CompatibleWithEngineVersion = FEngineVersion::CompatibleWith();
 		Summary.SetPackageFlags(Package ? Package->GetPackageFlags() : 0);
@@ -95,7 +95,7 @@ FLinkerSave::FLinkerSave(UPackage* InParent, FArchive *InSaver, bool bForceByteS
 
 		// Set main summary info.
 		Summary.Tag = PACKAGE_FILE_TAG;
-		Summary.SetFileVersions(GPackageFileUEVersion, GPackageFileLicenseeUEVersion, bInSaveUnversioned);
+		Summary.SetToLatestFileVersions(bInSaveUnversioned);
 		Summary.SavedByEngineVersion = FEngineVersion::Current();
 		Summary.CompatibleWithEngineVersion = FEngineVersion::CompatibleWith();
 		Summary.SetPackageFlags(Package ? Package->GetPackageFlags() : 0);
@@ -143,8 +143,8 @@ FLinkerSave::FLinkerSave(UPackage* InParent, bool bForceByteSwapping, bool bInSa
 		UPackage* Package = dynamic_cast<UPackage*>(LinkerRoot);
 
 		// Set main summary info.
-		Summary.Tag           = PACKAGE_FILE_TAG;
-		Summary.SetFileVersions( GPackageFileUEVersion, GPackageFileLicenseeUEVersion, bInSaveUnversioned );
+		Summary.Tag = PACKAGE_FILE_TAG;
+		Summary.SetToLatestFileVersions(bInSaveUnversioned);
 		Summary.SavedByEngineVersion = FEngineVersion::Current();
 		Summary.CompatibleWithEngineVersion = FEngineVersion::CompatibleWith();
 		Summary.SetPackageFlags(Package ? Package->GetPackageFlags() : 0);
diff --git a/Engine/Source/Runtime/CoreUObject/Private/UObject/ObjectBaseUtility.cpp b/Engine/Source/Runtime/CoreUObject/Private/UObject/ObjectBaseUtility.cpp
--- a/Engine/Source/Runtime/CoreUObject/Private/UObject/ObjectBaseUtility.cpp
+++ b/Engine/Source/Runtime/CoreUObject/Private/UObject/ObjectBaseUtility.cpp
@@ -16,7 +16,7 @@
  *			a) this object is a native only class, or
  *			b) this object's linker has been detached, in which case it is already fully loaded
  */
-int32 UObjectBaseUtility::GetLinkerUEVersion() const
+FPackageFileVersion UObjectBaseUtility::GetLinkerUEVersion() const
 {
 	FLinkerLoad* Loader = GetLinker();
 
diff --git a/Engine/Source/Runtime/CoreUObject/Private/UObject/PackageFileSummary.cpp b/Engine/Source/Runtime/CoreUObject/Private/UObject/PackageFileSummary.cpp
--- a/Engine/Source/Runtime/CoreUObject/Private/UObject/PackageFileSummary.cpp
+++ b/Engine/Source/Runtime/CoreUObject/Private/UObject/PackageFileSummary.cpp
@@ -30,7 +30,7 @@ static ECustomVersionSerializationFormat::Type GetCustomVersionFormatForArchive(
 	return CustomVersionFormat;
 }
 
-static void FixCorruptEngineVersion(int ObjectVersion, FEngineVersion& Version)
+static void FixCorruptEngineVersion(const FPackageFileVersion& ObjectVersion, FEngineVersion& Version)
 {
 	// The move of EpicInternal.txt in CL 12740027 broke checks for non-licensee builds in UGS. resulted in checks for Epic internal builds in UGS breaking, and assets being saved out with the licensee flag set.
 	// Detect such assets and clear the licensee bit.
@@ -98,8 +98,9 @@ void operator<<(FStructuredArchive::FSlot Slot, FPackageFileSummary& Sum)
 		*		-5 indicates the replacement of writing out the "UE3 version" so older versions of engine can gracefully fail to open newer packages
 		*		-6 indicates optimizations to how custom versions are being serialized
 		*		-7 indicates the texture allocation info has been removed from the summary
+		*		-8 indicates that the UE5 version has been added to the summary
 		*/
-		const int32 CurrentLegacyFileVersion = -7;
+		const int32 CurrentLegacyFileVersion = -8;
 		int32 LegacyFileVersion = CurrentLegacyFileVersion;
 		Record << SA_VALUE(TEXT("LegacyFileVersion"), LegacyFileVersion);
 
@@ -111,7 +112,7 @@ void operator<<(FStructuredArchive::FSlot Slot, FPackageFileSummary& Sum)
 				{
 					// we can't safely load more than this because the legacy version code differs in ways we can not predict.
 					// Make sure that the linker will fail to load with it.
-					Sum.FileVersionUE = 0;
+					Sum.FileVersionUE.Reset();
 					Sum.FileVersionLicenseeUE = 0;
 					return;
 				}
@@ -121,7 +122,14 @@ void operator<<(FStructuredArchive::FSlot Slot, FPackageFileSummary& Sum)
 					int32 LegacyUE3Version = 0;
 					Record << SA_VALUE(TEXT("LegacyUE3Version"), LegacyUE3Version);
 				}
-				Record << SA_VALUE(TEXT("FileVersionUE4"), Sum.FileVersionUE);
+
+				Record << SA_VALUE(TEXT("FileVersionUE4"), Sum.FileVersionUE.FileVersionUE4);
+				
+				if (LegacyFileVersion <= -8)
+				{
+					Record << SA_VALUE(TEXT("FileVersionUE5"), Sum.FileVersionUE.FileVersionUE5);
+				}
+			
 				Record << SA_VALUE(TEXT("FileVersionLicenseeUE4"), Sum.FileVersionLicenseeUE);
 
 				if (LegacyFileVersion <= -2)
@@ -129,7 +137,7 @@ void operator<<(FStructuredArchive::FSlot Slot, FPackageFileSummary& Sum)
 					Sum.CustomVersionContainer.Serialize(Record.EnterField(SA_FIELD_NAME(TEXT("CustomVersions"))), GetCustomVersionFormatForArchive(LegacyFileVersion));
 				}
 
-				if (!Sum.FileVersionUE && !Sum.FileVersionLicenseeUE)
+				if (!Sum.FileVersionUE.FileVersionUE4 && !Sum.FileVersionUE.FileVersionUE5 && !Sum.FileVersionLicenseeUE)
 				{
 #if WITH_EDITOR
 					if (!GAllowUnversionedContentInEditor)
@@ -149,7 +157,7 @@ void operator<<(FStructuredArchive::FSlot Slot, FPackageFileSummary& Sum)
 			else
 			{
 				// This is probably an old UE3 file, make sure that the linker will fail to load with it.
-				Sum.FileVersionUE = 0;
+				Sum.FileVersionUE.Reset();
 				Sum.FileVersionLicenseeUE = 0;
 			}
 		}
@@ -160,6 +168,7 @@ void operator<<(FStructuredArchive::FSlot Slot, FPackageFileSummary& Sum)
 				int32 Zero = 0;
 				Record << SA_VALUE(TEXT("LegacyUE3version"), Zero); // LegacyUE3version
 				Record << SA_VALUE(TEXT("FileVersionUE4"), Zero); // VersionUE4
+				Record << SA_VALUE(TEXT("FileVersionUE5"), Zero); // VersionUE5
 				Record << SA_VALUE(TEXT("FileVersionLicenseeUE4"), Zero); // VersionLicenseeUE4
 
 				FCustomVersionContainer NoCustomVersions;
@@ -170,7 +179,8 @@ void operator<<(FStructuredArchive::FSlot Slot, FPackageFileSummary& Sum)
 				// Must write out the last UE3 engine version, so that older versions identify it as new
 				int32 LegacyUE3Version = 864;
 				Record << SA_VALUE(TEXT("LegacyUE3Version"), LegacyUE3Version);
-				Record << SA_VALUE(TEXT("FileVersionUE4"), Sum.FileVersionUE);
+				Record << SA_VALUE(TEXT("FileVersionUE4"), Sum.FileVersionUE.FileVersionUE4);
+				Record << SA_VALUE(TEXT("FileVersionUE5"), Sum.FileVersionUE.FileVersionUE5);
 				Record << SA_VALUE(TEXT("FileVersionLicenseeUE4"), Sum.FileVersionLicenseeUE);
 
 				// Serialise custom version map.
@@ -207,7 +217,7 @@ void operator<<(FStructuredArchive::FSlot Slot, FPackageFileSummary& Sum)
 		Record << SA_VALUE(TEXT("ImportCount"), Sum.ImportCount) << SA_VALUE(TEXT("ImportOffset"), Sum.ImportOffset);
 		Record << SA_VALUE(TEXT("DependsOffset"), Sum.DependsOffset);
 
-		if (BaseArchive.IsLoading() && (Sum.FileVersionUE < VER_UE4_OLDEST_LOADABLE_PACKAGE || Sum.FileVersionUE > GPackageFileUEVersion))
+		if (BaseArchive.IsLoading() && (Sum.FileVersionUE < VER_UE4_OLDEST_LOADABLE_PACKAGE || Sum.IsFileVersionTooNew()))
 		{
 			return; // we can't safely load more than this because the below was different in older files.
 		}
@@ -319,7 +329,7 @@ void operator<<(FStructuredArchive::FSlot Slot, FPackageFileSummary& Sum)
 		if (!FCompression::VerifyCompressionFlagsValid(Sum.CompressionFlags))
 		{
 			UE_LOG(LogLinker, Warning, TEXT("Failed to read package file summary, the file \"%s\" has invalid compression flags (%d)."), *BaseArchive.GetArchiveName(), Sum.CompressionFlags);
-			Sum.FileVersionUE = VER_UE4_OLDEST_LOADABLE_PACKAGE - 1;
+			Sum.InvalidateFileVersion();
 			return;
 		}
 
@@ -330,7 +340,7 @@ void operator<<(FStructuredArchive::FSlot Slot, FPackageFileSummary& Sum)
 		{
 			// this file has package level compression, we won't load it.
 			UE_LOG(LogLinker, Warning, TEXT("Failed to read package file summary, the file \"%s\" is has package level compression (and is probably cooked). These old files cannot be loaded in the editor."), *BaseArchive.GetArchiveName());
-			Sum.FileVersionUE = VER_UE4_OLDEST_LOADABLE_PACKAGE - 1;
+			Sum.InvalidateFileVersion();
 			return; // we can't safely load more than this because we just changed the version to something it is not.
 		}
 
@@ -400,6 +410,20 @@ void FPackageFileSummary::SetCustomVersionContainer(const FCustomVersionContaine
 	CustomVersionContainer.SortByKey();
 }
 
+void FPackageFileSummary::SetFileVersions(const int32 EpicUE4, const int32 EpicUE5, const int32 LicenseeUE, const bool bInSaveUnversioned)
+{
+	// We could also make sure that EpicUE4 is >= VER_UE4_OLDEST_LOADABLE_PACKAGE but there might be a use case for setting
+	// an out of date version.
+	check(EpicUE4 <= EUnrealEngineObjectUE4Version::VER_UE4_AUTOMATIC_VERSION);
+	check(EpicUE5 <= (int32)EUnrealEngineObjectUE5Version::AUTOMATIC_VERSION);
+
+	FileVersionUE.FileVersionUE4 = EpicUE4;
+	FileVersionUE.FileVersionUE5 = EpicUE5;
+	FileVersionLicenseeUE = LicenseeUE;
+
+	bUnversioned = bInSaveUnversioned;
+}
+
 void FPackageFileSummary::SetPackageFlags(uint32 InPackageFlags)
 {
 	// Package summary flags should be set to false for all transient and memory-only flags
diff --git a/Engine/Source/Runtime/CoreUObject/Private/UObject/PropertyTag.cpp b/Engine/Source/Runtime/CoreUObject/Private/UObject/PropertyTag.cpp
--- a/Engine/Source/Runtime/CoreUObject/Private/UObject/PropertyTag.cpp
+++ b/Engine/Source/Runtime/CoreUObject/Private/UObject/PropertyTag.cpp
@@ -83,7 +83,7 @@ void operator<<(FStructuredArchive::FSlot Slot, FPropertyTag& Tag)
 	FArchive& UnderlyingArchive = Slot.GetUnderlyingArchive();
 	bool bIsTextFormat = UnderlyingArchive.IsTextFormat();
 
-	int32 Version = UnderlyingArchive.UEVer();
+	const FPackageFileVersion Version = UnderlyingArchive.UEVer();
 
 	check(!UnderlyingArchive.GetArchiveState().UseUnversionedPropertySerialization());
 	checkf(!UnderlyingArchive.IsSaving() || Tag.Prop, TEXT("FPropertyTag must be constructed with a valid property when used for saving data!"));
diff --git a/Engine/Source/Runtime/CoreUObject/Public/AssetRegistry/AssetData.h b/Engine/Source/Runtime/CoreUObject/Public/AssetRegistry/AssetData.h
--- a/Engine/Source/Runtime/CoreUObject/Public/AssetRegistry/AssetData.h
+++ b/Engine/Source/Runtime/CoreUObject/Public/AssetRegistry/AssetData.h
@@ -29,28 +29,29 @@ struct COREUOBJECT_API FAssetRegistryVersion
 {
 	enum Type
 	{
-		PreVersioning = 0,		// From before file versioning was implemented
-		HardSoftDependencies,	// The first version of the runtime asset registry to include file versioning.
-		AddAssetRegistryState,	// Added FAssetRegistryState and support for piecemeal serialization
-		ChangedAssetData,		// AssetData serialization format changed, versions before this are not readable
-		RemovedMD5Hash,			// Removed MD5 hash from package data
-		AddedHardManage,		// Added hard/soft manage references
-		AddedCookedMD5Hash,		// Added MD5 hash of cooked package to package data
-		AddedDependencyFlags,   // Added UE::AssetRegistry::EDependencyProperty to each dependency
-		FixedTags,				// Major tag format change that replaces USE_COMPACT_ASSET_REGISTRY:
-								// * Target tag INI settings cooked into tag data
-								// * Instead of FString values are stored directly as one of:
-								//		- Narrow / wide string
-								//		- [Numberless] FName
-								//		- [Numberless] export path
-								//		- Localized string
-								// * All value types are deduplicated
-								// * All key-value maps are cooked into a single contiguous range 
-								// * Switched from FName table to seek-free and more optimized FName batch loading
-								// * Removed global tag storage, a tag map reference-counts one store per asset registry
-								// * All configs can mix fixed and loose tag maps 
-		WorkspaceDomain,		// Added Version information to AssetPackageData
-		PackageImportedClasses, // Added ImportedClasses to AssetPackageData
+		PreVersioning = 0,					// From before file versioning was implemented
+		HardSoftDependencies,				// The first version of the runtime asset registry to include file versioning.
+		AddAssetRegistryState,				// Added FAssetRegistryState and support for piecemeal serialization
+		ChangedAssetData,					// AssetData serialization format changed, versions before this are not readable
+		RemovedMD5Hash,						// Removed MD5 hash from package data
+		AddedHardManage,					// Added hard/soft manage references
+		AddedCookedMD5Hash,					// Added MD5 hash of cooked package to package data
+		AddedDependencyFlags,				// Added UE::AssetRegistry::EDependencyProperty to each dependency
+		FixedTags,							// Major tag format change that replaces USE_COMPACT_ASSET_REGISTRY:
+											// * Target tag INI settings cooked into tag data
+											// * Instead of FString values are stored directly as one of:
+											//		- Narrow / wide string
+											//		- [Numberless] FName
+											//		- [Numberless] export path
+											//		- Localized string
+											// * All value types are deduplicated
+											// * All key-value maps are cooked into a single contiguous range 
+											// * Switched from FName table to seek-free and more optimized FName batch loading
+											// * Removed global tag storage, a tag map reference-counts one store per asset registry
+											// * All configs can mix fixed and loose tag maps 
+		WorkspaceDomain,					// Added Version information to AssetPackageData
+		PackageImportedClasses,				// Added ImportedClasses to AssetPackageData
+		PackageFileSummaryVersionChange,	// A new version number of UE5 was added to FPackageFileSummary
 
 		// -----<new versions can be added above this line>-------------------------------------------------
 		VersionPlusOne,
@@ -639,7 +640,7 @@ public:
 	int64 DiskSize;
 
 	/** UE file version that the package was saved with */
-	int32 FileVersionUE;
+	FPackageFileVersion FileVersionUE;
 
 	/** Licensee file version that the package was saved with */
 	int32 FileVersionLicenseeUE;
@@ -653,7 +654,6 @@ public:
 
 	FAssetPackageData()
 		: DiskSize(0)
-		, FileVersionUE(-1)
 		, FileVersionLicenseeUE(-1)
 		, Flags(0)
 	{
@@ -719,7 +719,18 @@ FORCEINLINE void FAssetPackageData::SerializeForCacheInternal(FArchive& Ar, FAss
 	}
 	if (Version >= FAssetRegistryVersion::WorkspaceDomain)
 	{
-		Ar << FileVersionUE;
+		if (Version >= FAssetRegistryVersion::PackageFileSummaryVersionChange)
+		{
+			Ar << FileVersionUE;
+		}
+		else
+		{
+			int32 UE4Version;
+			Ar << UE4Version;
+
+			FileVersionUE = FPackageFileVersion::CreateUE4Version(UE4Version);
+		}
+
 		Ar << FileVersionLicenseeUE;
 		Ar << Flags;
 		Ar << CustomVersions;
diff --git a/Engine/Source/Runtime/CoreUObject/Public/UObject/Package.h b/Engine/Source/Runtime/CoreUObject/Public/UObject/Package.h
--- a/Engine/Source/Runtime/CoreUObject/Public/UObject/Package.h
+++ b/Engine/Source/Runtime/CoreUObject/Public/UObject/Package.h
@@ -245,7 +245,7 @@ public:
 	class FLinkerLoad* LinkerLoad;
 
 	/** Linker package version this package has been serialized with. This is mostly used by PostLoad **/
-	int32 LinkerPackageVersion;
+	FPackageFileVersion LinkerPackageVersion;
 
 	/** Linker licensee version this package has been serialized with. This is mostly used by PostLoad **/
 	int32 LinkerLicenseeVersion;
diff --git a/Engine/Source/Runtime/CoreUObject/Public/UObject/PackageFileSummary.h b/Engine/Source/Runtime/CoreUObject/Public/UObject/PackageFileSummary.h
--- a/Engine/Source/Runtime/CoreUObject/Public/UObject/PackageFileSummary.h
+++ b/Engine/Source/Runtime/CoreUObject/Public/UObject/PackageFileSummary.h
@@ -52,7 +52,7 @@ struct FPackageFileSummary
 
 private:
 	/* UE file version */
-	int32		FileVersionUE;
+	FPackageFileVersion FileVersionUE;
 	/* Licensee file version */
 	int32		FileVersionLicenseeUE;
 	/* Custom version numbers. Keyed off a unique tag for each custom component. */
@@ -225,20 +225,21 @@ public:
 	FPackageFileSummary& operator=(FPackageFileSummary&&) = default;
 	FPackageFileSummary& operator=(const FPackageFileSummary&) = default;
 	PRAGMA_ENABLE_DEPRECATION_WARNINGS
-
-	int32 GetFileVersionUE() const
+	
+	FPackageFileVersion GetFileVersionUE() const
 	{
 		return FileVersionUE;
 	}
 
-	int32 GetFileVersionLicenseeUE() const
+	UE_DEPRECATED(5.0, "Use UEVer instead which returns the version as a FPackageFileVersion. See the @FPackageFileVersion documentation for further details")
+	int32 GetFileVersionUE4() const
 	{
-		return FileVersionLicenseeUE;
+		return FileVersionUE.FileVersionUE4;
 	}
-	UE_DEPRECATED(5.0, "Use GetFileVersion instead")
-	int32 GetFileVersionUE4() const
+
+	int32 GetFileVersionLicenseeUE() const
 	{
-		return GetFileVersionUE();
+		return FileVersionLicenseeUE;
 	}
 
 	UE_DEPRECATED(5.0, "Use GetFileVersionLicenseeUE instead")
@@ -254,13 +255,36 @@ public:
 
 	void SetCustomVersionContainer(const FCustomVersionContainer& InContainer);
 
-	void SetFileVersions(const int32 EpicUE, const int32 LicenseeUE, const bool bInSaveUnversioned = false)
+	/** Used to manually set the package file and licensee versions */
+	COREUOBJECT_API void SetFileVersions(const int32 EpicUE4, const int32 EpicUE5, const int32 LicenseeUE, const bool bInSaveUnversioned = false);
+	
+	UE_DEPRECATED(5.0, "Use the other overload of SetFileVersions that takes an UE5 version as well")
+	inline void SetFileVersions(const int32 EpicUE, const int32 LicenseeUE, const bool bInSaveUnversioned = false)
+	{
+		SetFileVersions(EpicUE, GPackageFileUEVersion.FileVersionUE5, LicenseeUE, bInSaveUnversioned);
+	}
+
+	/** Set both the package file versions and the licensee version to the most recent version numbers supported */
+	void SetToLatestFileVersions(const bool bInSaveUnversioned)
 	{
-		FileVersionUE = EpicUE;
-		FileVersionLicenseeUE = LicenseeUE;
+		FileVersionUE = GPackageFileUEVersion;
+		FileVersionLicenseeUE = GPackageFileLicenseeUEVersion;
+
 		bUnversioned = bInSaveUnversioned;
 	}
 
+	/** Returns true if any of the package file versions are older than the minimum supported versions */
+	bool IsFileVersionTooOld() const
+	{
+		return FileVersionUE < VER_UE4_OLDEST_LOADABLE_PACKAGE;
+	}
+
+	/** Returns true if any of the package file versions are newer than currently supported by the running process */
+	bool IsFileVersionTooNew() const
+	{
+		return FileVersionUE.FileVersionUE4 > GPackageFileUEVersion.FileVersionUE4 || FileVersionUE.FileVersionUE5 > GPackageFileUEVersion.FileVersionUE5;
+	}
+
 	/** Get the summary package flags. */
 	uint32 GetPackageFlags() const
 	{
@@ -273,5 +297,16 @@ public:
 	/** I/O functions */
 	friend COREUOBJECT_API FArchive& operator<<(FArchive& Ar, FPackageFileSummary& Sum);
 	friend COREUOBJECT_API void operator<<(FStructuredArchive::FSlot Slot, FPackageFileSummary& Sum);
-};
 
+private:
+
+	/** 
+	 * Set the UE4 version to below the oldest loadable version. This is used when an error is encountered during serialization 
+	 * and we need to prevent the use of the package. When calling this make sure you log a warning with the 'LogLinker' category
+	 * to inform the user of the problem.
+	 */
+	void InvalidateFileVersion()
+	{
+		FileVersionUE = FPackageFileVersion::CreateUE4Version((EUnrealEngineObjectUE4Version)(VER_UE4_OLDEST_LOADABLE_PACKAGE - 1));
+	}
+};
diff --git a/Engine/Source/Runtime/CoreUObject/Public/UObject/UObjectBaseUtility.h b/Engine/Source/Runtime/CoreUObject/Public/UObject/UObjectBaseUtility.h
--- a/Engine/Source/Runtime/CoreUObject/Public/UObject/UObjectBaseUtility.h
+++ b/Engine/Source/Runtime/CoreUObject/Public/UObject/UObjectBaseUtility.h
@@ -617,12 +617,16 @@ public:
 	 *			a) this object is a native only class, or
 	 *			b) this object's linker has been detached, in which case it is already fully loaded
 	 */
-	int32 GetLinkerUEVersion() const;
-
-
-	UE_DEPRECATED(5.0, "Use GetLinkerUEVersion instead")
-	inline int32 GetLinkerUE4Version() const { return GetLinkerUEVersion(); }
-
+	FPackageFileVersion GetLinkerUEVersion() const;
+		
+	UE_DEPRECATED(5.0, "Use GetLinkerUEVersion instead which returns the version as a FPackageFileVersion. See the @FPackageFileVersion documentation for further details")
+	inline int32 GetLinkerUE4Version() const 
+	{ 
+		// Existing code calling GetLinkerUE4Version will be testing against UE4 version numbers so 
+		// we can just return the UE4 version number.
+		// All new code that might actually need the UE5 version as well should be calling ::GetLinkerUEVersion
+		return GetLinkerUEVersion().FileVersionUE4;
+	}
 
 	/**
 	 * Returns the licensee version of the linker for this object.
diff --git a/Engine/Source/Runtime/Engine/Private/Camera/CameraComponent.cpp b/Engine/Source/Runtime/Engine/Private/Camera/CameraComponent.cpp
--- a/Engine/Source/Runtime/Engine/Private/Camera/CameraComponent.cpp
+++ b/Engine/Source/Runtime/Engine/Private/Camera/CameraComponent.cpp
@@ -149,7 +149,7 @@ void UCameraComponent::PostLoad()
 {
 	Super::PostLoad();
 
-	const int32 LinkerUEVer = GetLinkerUEVersion();
+	const FPackageFileVersion LinkerUEVer = GetLinkerUEVersion();
 
 	if (LinkerUEVer < VER_UE4_RENAME_CAMERA_COMPONENT_VIEW_ROTATION)
 	{
diff --git a/Engine/Source/Runtime/Engine/Private/Components/AudioComponent.cpp b/Engine/Source/Runtime/Engine/Private/Components/AudioComponent.cpp
--- a/Engine/Source/Runtime/Engine/Private/Components/AudioComponent.cpp
+++ b/Engine/Source/Runtime/Engine/Private/Components/AudioComponent.cpp
@@ -167,7 +167,7 @@ void UAudioComponent::Serialize(FArchive& Ar)
 void UAudioComponent::PostLoad()
 {
 #if WITH_EDITORONLY_DATA
-	const int32 LinkerUEVersion = GetLinkerUEVersion();
+	const FPackageFileVersion LinkerUEVersion = GetLinkerUEVersion();
 
 	// Translate the old HighFrequencyGainMultiplier value to the new LowPassFilterFrequency value
 	if (LinkerUEVersion < VER_UE4_USE_LOW_PASS_FILTER_FREQ)
diff --git a/Engine/Source/Runtime/Engine/Private/Components/CharacterMovementComponent.cpp b/Engine/Source/Runtime/Engine/Private/Components/CharacterMovementComponent.cpp
--- a/Engine/Source/Runtime/Engine/Private/Components/CharacterMovementComponent.cpp
+++ b/Engine/Source/Runtime/Engine/Private/Components/CharacterMovementComponent.cpp
@@ -640,7 +640,7 @@ void UCharacterMovementComponent::PostLoad()
 	Super::PostLoad();
 
 #if WITH_EDITORONLY_DATA
-	const int32 LinkerUEVer = GetLinkerUEVersion();
+	const FPackageFileVersion LinkerUEVer = GetLinkerUEVersion();
 
 	if (LinkerUEVer < VER_UE4_CHARACTER_MOVEMENT_DECELERATION)
 	{
diff --git a/Engine/Source/Runtime/Engine/Private/Components/PrimitiveComponent.cpp b/Engine/Source/Runtime/Engine/Private/Components/PrimitiveComponent.cpp
--- a/Engine/Source/Runtime/Engine/Private/Components/PrimitiveComponent.cpp
+++ b/Engine/Source/Runtime/Engine/Private/Components/PrimitiveComponent.cpp
@@ -1250,7 +1250,7 @@ void UPrimitiveComponent::PostLoad()
 {
 	Super::PostLoad();
 	
-	int32 const UEVersion = GetLinkerUEVersion();
+	FPackageFileVersion const UEVersion = GetLinkerUEVersion();
 
 	// as temporary fix for the bug TTP 299926
 	// permanent fix is coming
diff --git a/Engine/Source/Runtime/Engine/Private/Components/ProjectileMovementComponent.cpp b/Engine/Source/Runtime/Engine/Private/Components/ProjectileMovementComponent.cpp
--- a/Engine/Source/Runtime/Engine/Private/Components/ProjectileMovementComponent.cpp
+++ b/Engine/Source/Runtime/Engine/Private/Components/ProjectileMovementComponent.cpp
@@ -57,7 +57,7 @@ void UProjectileMovementComponent::PostLoad()
 {
 	Super::PostLoad();
 
-	const int32 LinkerUEVer = GetLinkerUEVersion();
+	const FPackageFileVersion LinkerUEVer = GetLinkerUEVersion();
 
 	if (LinkerUEVer < VER_UE4_REFACTOR_PROJECTILE_MOVEMENT)
 	{
diff --git a/Engine/Source/Runtime/Engine/Private/Components/SplineComponent.cpp b/Engine/Source/Runtime/Engine/Private/Components/SplineComponent.cpp
--- a/Engine/Source/Runtime/Engine/Private/Components/SplineComponent.cpp
+++ b/Engine/Source/Runtime/Engine/Private/Components/SplineComponent.cpp
@@ -124,7 +124,7 @@ void USplineComponent::Serialize(FArchive& Ar)
 	}
 
 	// Support old resources which don't have the rotation and scale splines present
-	const int32 ArchiveUEVersion = Ar.UEVer();
+	const FPackageFileVersion ArchiveUEVersion = Ar.UEVer();
 	if (ArchiveUEVersion < VER_UE4_INTERPCURVE_SUPPORTS_LOOPING)
 	{
 		int32 NumPoints = SplineCurves.Position.Points.Num();
diff --git a/Engine/Source/Runtime/Engine/Private/EngineUtils.cpp b/Engine/Source/Runtime/Engine/Private/EngineUtils.cpp
--- a/Engine/Source/Runtime/Engine/Private/EngineUtils.cpp
+++ b/Engine/Source/Runtime/Engine/Private/EngineUtils.cpp
@@ -497,12 +497,12 @@ void FConsoleOutputDevice::Serialize(const TCHAR* Text, ELogVerbosity::Type Verb
 /*-----------------------------------------------------------------------------
 	Serialized data stripping.
 -----------------------------------------------------------------------------*/
-FStripDataFlags::FStripDataFlags( class FArchive& Ar, uint8 InClassFlags /*= 0*/, int32 InVersion /*= VER_UE4_OLDEST_LOADABLE_PACKAGE */ )
+FStripDataFlags::FStripDataFlags( class FArchive& Ar, uint8 InClassFlags /*= 0*/, const FPackageFileVersion& InVersion /*= GOldestLoadablePackageFileUEVersion */ )
 	: GlobalStripFlags( 0 )
 	, ClassStripFlags( 0 )
 {
 	check(InVersion >= VER_UE4_OLDEST_LOADABLE_PACKAGE);
-	if (Ar.UEVer() >= InVersion)
+	if (Ar.UEVer().IsCompatible(InVersion))
 	{
 		if (Ar.IsCooking())
 		{
@@ -517,12 +517,12 @@ FStripDataFlags::FStripDataFlags( class FArchive& Ar, uint8 InClassFlags /*= 0*/
 	}
 }
 
-FStripDataFlags::FStripDataFlags( class FArchive& Ar, uint8 InGlobalFlags, uint8 InClassFlags, int32 InVersion /*= VER_UE4_OLDEST_LOADABLE_PACKAGE */ )
+FStripDataFlags::FStripDataFlags( class FArchive& Ar, uint8 InGlobalFlags, uint8 InClassFlags, const FPackageFileVersion& InVersion /*= GOldestLoadablePackageFileUEVersion */)
 	: GlobalStripFlags( 0 )
 	, ClassStripFlags( 0 )
 {
 	check(InVersion >= VER_UE4_OLDEST_LOADABLE_PACKAGE);
-	if (Ar.UEVer() >= InVersion)
+	if (Ar.UEVer().IsCompatible(InVersion))
 	{
 		if (Ar.IsCooking())
 		{
@@ -538,7 +538,7 @@ FStripDataFlags::FStripDataFlags( class FArchive& Ar, uint8 InGlobalFlags, uint8
 /*-----------------------------------------------------------------------------
 Serialized data stripping.
 -----------------------------------------------------------------------------*/
-FStripDataFlags::FStripDataFlags(FStructuredArchive::FSlot Slot, uint8 InClassFlags /*= 0*/, int32 InVersion /*= VER_UE4_OLDEST_LOADABLE_PACKAGE */)
+FStripDataFlags::FStripDataFlags(FStructuredArchive::FSlot Slot, uint8 InClassFlags /*= 0*/, const FPackageFileVersion& InVersion /*= GOldestLoadablePackageFileUEVersion */)
 	: GlobalStripFlags(0)
 	, ClassStripFlags(0)
 {
@@ -546,7 +546,7 @@ FStripDataFlags::FStripDataFlags(FStructuredArchive::FSlot Slot, uint8 InClassFl
 	FStructuredArchive::FRecord Record = Slot.EnterRecord();
 
 	check(InVersion >= VER_UE4_OLDEST_LOADABLE_PACKAGE);
-	if (UnderlyingArchive.UEVer() >= InVersion)
+	if (UnderlyingArchive.UEVer().IsCompatible(InVersion))
 	{
 		if (UnderlyingArchive.IsCooking())
 		{
@@ -561,7 +561,7 @@ FStripDataFlags::FStripDataFlags(FStructuredArchive::FSlot Slot, uint8 InClassFl
 	}
 }
 
-FStripDataFlags::FStripDataFlags(FStructuredArchive::FSlot Slot, uint8 InGlobalFlags, uint8 InClassFlags, int32 InVersion /*= VER_UE4_OLDEST_LOADABLE_PACKAGE */)
+FStripDataFlags::FStripDataFlags(FStructuredArchive::FSlot Slot, uint8 InGlobalFlags, uint8 InClassFlags, const FPackageFileVersion& InVersion /*= GOldestLoadablePackageFileUEVersion */)
 	: GlobalStripFlags(0)
 	, ClassStripFlags(0)
 {
@@ -569,7 +569,7 @@ FStripDataFlags::FStripDataFlags(FStructuredArchive::FSlot Slot, uint8 InGlobalF
 	FStructuredArchive::FRecord Record = Slot.EnterRecord();
 
 	check(InVersion >= VER_UE4_OLDEST_LOADABLE_PACKAGE);
-	if (UnderlyingArchive.UEVer() >= InVersion)
+	if (UnderlyingArchive.UEVer().IsCompatible(InVersion))
 	{
 		if (UnderlyingArchive.IsCooking())
 		{
diff --git a/Engine/Source/Runtime/Engine/Private/GameplayStatics.cpp b/Engine/Source/Runtime/Engine/Private/GameplayStatics.cpp
--- a/Engine/Source/Runtime/Engine/Private/GameplayStatics.cpp
+++ b/Engine/Source/Runtime/Engine/Private/GameplayStatics.cpp
@@ -64,6 +64,8 @@ struct FSaveGameFileVersion
 		InitialVersion = 1,
 		// serializing custom versions into the savegame data to handle that type of versioning
 		AddedCustomVersions = 2,
+		// added a new UE5 version number to FPackageFileSummary
+		PackageFileSummaryVersionChange = 3,
 
 		// -----<new versions can be added above this line>-------------------------------------------------
 		VersionPlusOne,
@@ -91,7 +93,7 @@ struct FSaveGameHeader
 
 	int32 FileTypeTag;
 	int32 SaveGameFileVersion;
-	int32 PackageFileUEVersion;
+	FPackageFileVersion PackageFileUEVersion;
 	FEngineVersion SavedEngineVersion;
 	int32 CustomVersionFormat;
 	FCustomVersionContainer CustomVersions;
@@ -101,7 +103,6 @@ struct FSaveGameHeader
 FSaveGameHeader::FSaveGameHeader()
 	: FileTypeTag(0)
 	, SaveGameFileVersion(0)
-	, PackageFileUEVersion(0)
 	, CustomVersionFormat(static_cast<int32>(ECustomVersionSerializationFormat::Unknown))
 {}
 
@@ -119,7 +120,7 @@ void FSaveGameHeader::Empty()
 {
 	FileTypeTag = 0;
 	SaveGameFileVersion = 0;
-	PackageFileUEVersion = 0;
+	PackageFileUEVersion.Reset();
 	SavedEngineVersion.Empty();
 	CustomVersionFormat = (int32)ECustomVersionSerializationFormat::Unknown;
 	CustomVersions.Empty();
@@ -155,7 +156,17 @@ void FSaveGameHeader::Read(FMemoryReader& MemoryReader)
 		MemoryReader << SaveGameFileVersion;
 
 		// Read engine and UE version information
-		MemoryReader << PackageFileUEVersion;
+		if (SaveGameFileVersion >= FSaveGameFileVersion::PackageFileSummaryVersionChange)
+		{
+			MemoryReader << PackageFileUEVersion;
+		}
+		else
+		{
+			int32 OldUe4Version;
+			MemoryReader << OldUe4Version;
+
+			PackageFileUEVersion = FPackageFileVersion::CreateUE4Version(OldUe4Version);
+		}
 
 		MemoryReader << SavedEngineVersion;
 
diff --git a/Engine/Source/Runtime/Engine/Private/Materials/Material.cpp b/Engine/Source/Runtime/Engine/Private/Materials/Material.cpp
--- a/Engine/Source/Runtime/Engine/Private/Materials/Material.cpp
+++ b/Engine/Source/Runtime/Engine/Private/Materials/Material.cpp
@@ -3401,7 +3401,7 @@ void UMaterial::PostLoad()
 	LoadedMaterialResources.Empty();
 
 #if WITH_EDITORONLY_DATA
-	const int32 UEVer = GetLinkerUEVersion();
+	const FPackageFileVersion UEVer = GetLinkerUEVersion();
 	const int32 RenderObjVer = GetLinkerCustomVersion(FRenderingObjectVersion::GUID);
 	const int32 UE5MainVer = GetLinkerCustomVersion(FUE5MainStreamObjectVersion::GUID);
 
diff --git a/Engine/Source/Runtime/Engine/Private/Materials/MaterialExpressions.cpp b/Engine/Source/Runtime/Engine/Private/Materials/MaterialExpressions.cpp
--- a/Engine/Source/Runtime/Engine/Private/Materials/MaterialExpressions.cpp
+++ b/Engine/Source/Runtime/Engine/Private/Materials/MaterialExpressions.cpp
@@ -829,7 +829,7 @@ void UMaterialExpression::Serialize(FStructuredArchive::FRecord Record)
 	Archive.UsingCustomVersion(FUE5MainStreamObjectVersion::GUID);
 
 #if WITH_EDITORONLY_DATA
-	const int32 UEVer = Archive.UEVer();
+	const FPackageFileVersion UEVer = Archive.UEVer();
 	const int32 RenderVer = Archive.CustomVer(FRenderingObjectVersion::GUID);
 	const int32 UE5Ver = Archive.CustomVer(FUE5MainStreamObjectVersion::GUID);
 
diff --git a/Engine/Source/Runtime/Engine/Private/Materials/MaterialShared.cpp b/Engine/Source/Runtime/Engine/Private/Materials/MaterialShared.cpp
--- a/Engine/Source/Runtime/Engine/Private/Materials/MaterialShared.cpp
+++ b/Engine/Source/Runtime/Engine/Private/Materials/MaterialShared.cpp
@@ -4199,7 +4199,7 @@ bool UMaterialInterface::IsTextureReferencedByProperty(EMaterialProperty InPrope
 //Reorder the output index for any FExpressionInput connected to a UMaterialExpressionBreakMaterialAttributes.
 //If the order of pins in the material results or the make/break attributes nodes changes 
 //then the OutputIndex stored in any FExpressionInput coming from UMaterialExpressionBreakMaterialAttributes will be wrong and needs reordering.
-void DoMaterialAttributeReorder(FExpressionInput* Input, int32 UEVer, int32 RenderObjVer, int32 UE5MainVer)
+void DoMaterialAttributeReorder(FExpressionInput* Input, const FPackageFileVersion& UEVer, int32 RenderObjVer, int32 UE5MainVer)
 {
 	if( Input && Input->Expression && Input->Expression->IsA(UMaterialExpressionBreakMaterialAttributes::StaticClass()) )
 	{
diff --git a/Engine/Source/Runtime/Engine/Private/PhysicsEngine/BodyInstance.cpp b/Engine/Source/Runtime/Engine/Private/PhysicsEngine/BodyInstance.cpp
--- a/Engine/Source/Runtime/Engine/Private/PhysicsEngine/BodyInstance.cpp
+++ b/Engine/Source/Runtime/Engine/Private/PhysicsEngine/BodyInstance.cpp
@@ -4043,7 +4043,7 @@ void FBodyInstance::FixupData(class UObject* Loader)
 {
 	check (Loader);
 
-	int32 const UEVersion = Loader->GetLinkerUEVersion();
+	FPackageFileVersion const UEVersion = Loader->GetLinkerUEVersion();
 
 #if WITH_EDITOR
 	if (UEVersion < VER_UE4_ADD_CUSTOMPROFILENAME_CHANGE)
diff --git a/Engine/Source/Runtime/Engine/Private/Rendering/ColorVertexBuffer.cpp b/Engine/Source/Runtime/Engine/Private/Rendering/ColorVertexBuffer.cpp
--- a/Engine/Source/Runtime/Engine/Private/Rendering/ColorVertexBuffer.cpp
+++ b/Engine/Source/Runtime/Engine/Private/Rendering/ColorVertexBuffer.cpp
@@ -195,7 +195,7 @@ void FColorVertexBuffer::AppendVertices( const FStaticMeshBuildVertex* Vertices,
 void FColorVertexBuffer::Serialize( FArchive& Ar, bool bNeedsCPUAccess )
 {
 	NeedsCPUAccess = bNeedsCPUAccess;
-	FStripDataFlags StripFlags(Ar, 0, VER_UE4_STATIC_SKELETAL_MESH_SERIALIZATION_FIX);
+	FStripDataFlags StripFlags(Ar, 0, FPackageFileVersion::CreateUE4Version(VER_UE4_STATIC_SKELETAL_MESH_SERIALIZATION_FIX));
 
 	if (Ar.IsSaving() && NumVertices > 0 && VertexData == NULL)
 	{
diff --git a/Engine/Source/Runtime/Engine/Private/Rendering/SkeletalMeshVertexBuffer.cpp b/Engine/Source/Runtime/Engine/Private/Rendering/SkeletalMeshVertexBuffer.cpp
--- a/Engine/Source/Runtime/Engine/Private/Rendering/SkeletalMeshVertexBuffer.cpp
+++ b/Engine/Source/Runtime/Engine/Private/Rendering/SkeletalMeshVertexBuffer.cpp
@@ -30,7 +30,7 @@ void FDummySkeletalMeshVertexBuffer::CleanUp()
 */
 FArchive& operator<<(FArchive& Ar, FDummySkeletalMeshVertexBuffer& VertexBuffer)
 {
-	FStripDataFlags StripFlags(Ar, 0, VER_UE4_STATIC_SKELETAL_MESH_SERIALIZATION_FIX);
+	FStripDataFlags StripFlags(Ar, 0, FPackageFileVersion::CreateUE4Version(VER_UE4_STATIC_SKELETAL_MESH_SERIALIZATION_FIX));
 
 	Ar << VertexBuffer.NumTexCoords;
 	Ar << VertexBuffer.bUseFullPrecisionUVs;
diff --git a/Engine/Source/Runtime/Engine/Private/Rendering/SkeletalMeshVertexClothBuffer.cpp b/Engine/Source/Runtime/Engine/Private/Rendering/SkeletalMeshVertexClothBuffer.cpp
--- a/Engine/Source/Runtime/Engine/Private/Rendering/SkeletalMeshVertexClothBuffer.cpp
+++ b/Engine/Source/Runtime/Engine/Private/Rendering/SkeletalMeshVertexClothBuffer.cpp
@@ -139,7 +139,7 @@ void FSkeletalMeshVertexClothBuffer::ReleaseRHI()
 */
 FArchive& operator<<(FArchive& Ar, FSkeletalMeshVertexClothBuffer& VertexBuffer)
 {
-	FStripDataFlags StripFlags(Ar, 0, VER_UE4_STATIC_SKELETAL_MESH_SERIALIZATION_FIX);
+	FStripDataFlags StripFlags(Ar, 0, FPackageFileVersion::CreateUE4Version(VER_UE4_STATIC_SKELETAL_MESH_SERIALIZATION_FIX));
 
 	if (Ar.IsLoading())
 	{
diff --git a/Engine/Source/Runtime/Engine/Private/Rendering/StaticMeshVertexBuffer.cpp b/Engine/Source/Runtime/Engine/Private/Rendering/StaticMeshVertexBuffer.cpp
--- a/Engine/Source/Runtime/Engine/Private/Rendering/StaticMeshVertexBuffer.cpp
+++ b/Engine/Source/Runtime/Engine/Private/Rendering/StaticMeshVertexBuffer.cpp
@@ -199,7 +199,7 @@ void FStaticMeshVertexBuffer::Serialize(FArchive& Ar, bool bNeedsCPUAccess)
 
 	DECLARE_SCOPE_CYCLE_COUNTER(TEXT("FStaticMeshVertexBuffer::Serialize"), STAT_StaticMeshVertexBuffer_Serialize, STATGROUP_LoadTime);
 
-	FStripDataFlags StripFlags(Ar, 0, VER_UE4_STATIC_SKELETAL_MESH_SERIALIZATION_FIX);
+	FStripDataFlags StripFlags(Ar, 0, FPackageFileVersion::CreateUE4Version(VER_UE4_STATIC_SKELETAL_MESH_SERIALIZATION_FIX));
 
 	SerializeMetaData(Ar);
 
diff --git a/Engine/Source/Runtime/Engine/Private/SkeletalMeshLODModel.cpp b/Engine/Source/Runtime/Engine/Private/SkeletalMeshLODModel.cpp
--- a/Engine/Source/Runtime/Engine/Private/SkeletalMeshLODModel.cpp
+++ b/Engine/Source/Runtime/Engine/Private/SkeletalMeshLODModel.cpp
@@ -848,7 +848,7 @@ void FSkeletalMeshLODModel::Serialize(FArchive& Ar, UObject* Owner, int32 Idx)
 				if (Ar.IsLoading() && Ar.CustomVer(FSkeletalMeshCustomVersion::GUID) < FSkeletalMeshCustomVersion::UseSharedColorBufferFormat)
 				{
 					TArray<FColor> OldColors;
-					FStripDataFlags LegacyColourStripFlags(Ar, 0, VER_UE4_STATIC_SKELETAL_MESH_SERIALIZATION_FIX);
+					FStripDataFlags LegacyColourStripFlags(Ar, 0, FPackageFileVersion::CreateUE4Version(VER_UE4_STATIC_SKELETAL_MESH_SERIALIZATION_FIX));
 					OldColors.BulkSerialize(Ar);
 				}
 				else
@@ -887,7 +887,7 @@ void FSkeletalMeshLODModel::Serialize(FArchive& Ar, UObject* Owner, int32 Idx)
 
 			if (Ar.UEVer() >= VER_UE4_APEX_CLOTH && HasClothData())
 			{
-				FStripDataFlags StripFlags2(Ar, 0, VER_UE4_STATIC_SKELETAL_MESH_SERIALIZATION_FIX);
+				FStripDataFlags StripFlags2(Ar, 0, FPackageFileVersion::CreateUE4Version(VER_UE4_STATIC_SKELETAL_MESH_SERIALIZATION_FIX));
 				TSkeletalMeshVertexData<FMeshToMeshVertData> DummyClothData(true);
 
 				if (!StripFlags2.IsDataStrippedForServer() || Ar.IsCountingMemory())
diff --git a/Engine/Source/Runtime/Engine/Private/SoundBase.cpp b/Engine/Source/Runtime/Engine/Private/SoundBase.cpp
--- a/Engine/Source/Runtime/Engine/Private/SoundBase.cpp
+++ b/Engine/Source/Runtime/Engine/Private/SoundBase.cpp
@@ -197,7 +197,7 @@ void USoundBase::PostLoad()
 		bOutputToBusOnly_DEPRECATED = false;
 	}
 
-	const int32 LinkerUEVersion = GetLinkerUEVersion();
+	const FPackageFileVersion LinkerUEVersion = GetLinkerUEVersion();
 
 	if (LinkerUEVersion < VER_UE4_SOUND_CONCURRENCY_PACKAGE)
 	{
diff --git a/Engine/Source/Runtime/Engine/Public/EngineUtils.h b/Engine/Source/Runtime/Engine/Public/EngineUtils.h
--- a/Engine/Source/Runtime/Engine/Public/EngineUtils.h
+++ b/Engine/Source/Runtime/Engine/Public/EngineUtils.h
@@ -842,7 +842,13 @@ public:
 	 * @param InClassFlags - User defined per class flags .
 	 * @param InVersion - Minimal strip version required to serialize strip flags
 	 */
-	FStripDataFlags( FArchive& Ar, uint8 InClassFlags = 0, int32 InVersion = VER_UE4_OLDEST_LOADABLE_PACKAGE );
+	FStripDataFlags(FArchive& Ar, uint8 InClassFlags = 0, const FPackageFileVersion& InVersion = GOldestLoadablePackageFileUEVersion);
+	UE_DEPRECATED(5.0, "Use the other overload that takes InVersion as a FPackageFileVersion. See the @FPackageFileVersion documentation for further details")
+	inline FStripDataFlags(FArchive& Ar, uint8 InClassFlags, int32 InVersion)
+		: FStripDataFlags(Ar, InClassFlags, FPackageFileVersion::CreateUE4Version(InVersion))
+	{
+
+	}
 
 	/** 
 	 * Constructor.
@@ -853,7 +859,13 @@ public:
 	 * @param InClassFlags - User defined per class flags.
 	 * @param InVersion - Minimal version required to serialize strip flags
 	 */
-	FStripDataFlags( FArchive& Ar, uint8 InGlobalFlags, uint8 InClassFlags, int32 InVersion = VER_UE4_OLDEST_LOADABLE_PACKAGE );
+	FStripDataFlags( FArchive& Ar, uint8 InGlobalFlags, uint8 InClassFlags, const FPackageFileVersion& InVersion = GOldestLoadablePackageFileUEVersion);
+	UE_DEPRECATED(5.0, "Use the other overload that takes InVersion as a FPackageFileVersion. See the @FPackageFileVersion documentation for further details")
+	inline FStripDataFlags(FArchive& Ar, uint8 InGlobalFlags, uint8 InClassFlags, int32 InVersion)
+		: FStripDataFlags(Ar, InGlobalFlags, InClassFlags, FPackageFileVersion::CreateUE4Version(InVersion))
+	{
+
+	}
 
 	/**
 	* Constructor.
@@ -864,7 +876,13 @@ public:
 	* @param InClassFlags - User defined per class flags .
 	* @param InVersion - Minimal strip version required to serialize strip flags
 	*/
-	FStripDataFlags(FStructuredArchive::FSlot Slot, uint8 InClassFlags = 0, int32 InVersion = VER_UE4_OLDEST_LOADABLE_PACKAGE);
+	FStripDataFlags(FStructuredArchive::FSlot Slot, uint8 InClassFlags = 0, const FPackageFileVersion& InVersion = GOldestLoadablePackageFileUEVersion);
+	UE_DEPRECATED(5.0, "Use the other overload that takes InVersion as a FPackageFileVersion. See the @FPackageFileVersion documentation for further details")
+	inline FStripDataFlags(FStructuredArchive::FSlot Slot, uint8 InClassFlags, int32 InVersion)
+		: FStripDataFlags(Slot, InClassFlags, FPackageFileVersion::CreateUE4Version(InVersion))
+	{
+
+	}
 
 	/**
 	* Constructor.
@@ -875,7 +893,13 @@ public:
 	* @param InClassFlags - User defined per class flags.
 	* @param InVersion - Minimal version required to serialize strip flags
 	*/
-	FStripDataFlags(FStructuredArchive::FSlot Slot, uint8 InGlobalFlags, uint8 InClassFlags, int32 InVersion = VER_UE4_OLDEST_LOADABLE_PACKAGE);
+	FStripDataFlags(FStructuredArchive::FSlot Slot, uint8 InGlobalFlags, uint8 InClassFlags, const FPackageFileVersion& InVersion = GOldestLoadablePackageFileUEVersion);
+	UE_DEPRECATED(5.0, "Use the other overload that takes InVersion as a FPackageFileVersion. See the @FPackageFileVersion documentation for further details")
+	inline FStripDataFlags(FStructuredArchive::FSlot Slot, uint8 InGlobalFlags, uint8 InClassFlags, int32 InVersion)
+		: FStripDataFlags(Slot, InGlobalFlags, InClassFlags, FPackageFileVersion::CreateUE4Version(InVersion))
+	{
+
+	}
 
 	/**
 	 * Checks if FStripDataFlags::Editor flag is set or not
diff --git a/Engine/Source/Runtime/Engine/Public/MaterialShared.h b/Engine/Source/Runtime/Engine/Public/MaterialShared.h
--- a/Engine/Source/Runtime/Engine/Public/MaterialShared.h
+++ b/Engine/Source/Runtime/Engine/Public/MaterialShared.h
@@ -2772,7 +2772,7 @@ ENGINE_API bool DoesMaterialUseTexture(const UMaterialInterface* Material,const
 
 #if WITH_EDITORONLY_DATA
 /** TODO - This can be removed whenever VER_UE4_MATERIAL_ATTRIBUTES_REORDERING is no longer relevant. */
-ENGINE_API void DoMaterialAttributeReorder(FExpressionInput* Input, int32 UEVer, int32 RenderObjVer, int32 UE5MainVer);
+ENGINE_API void DoMaterialAttributeReorder(FExpressionInput* Input, const FPackageFileVersion&, int32 RenderObjVer, int32 UE5MainVer);
 #endif // WITH_EDITORONLY_DATA
 
 /**
diff --git a/Engine/Source/Runtime/NetworkFile/Private/NetworkPlatformFile.cpp b/Engine/Source/Runtime/NetworkFile/Private/NetworkPlatformFile.cpp
--- a/Engine/Source/Runtime/NetworkFile/Private/NetworkPlatformFile.cpp
+++ b/Engine/Source/Runtime/NetworkFile/Private/NetworkPlatformFile.cpp
@@ -206,7 +206,7 @@ void FNetworkPlatformFile::InitializeAfterSetActive()
 		else
 		{
 			// receive the cooked version information
-			int32 ServerPackageVersion = 0;
+			FPackageFileVersion ServerPackageVersion;
 			int32 ServerPackageLicenseeVersion = 0;
 			ProcessServerInitialResponse(Response, ServerPackageVersion, ServerPackageLicenseeVersion);
 			ProcessServerCachedFilesResponse(Response, ServerPackageVersion, ServerPackageLicenseeVersion);
@@ -234,7 +234,7 @@ void FNetworkPlatformFile::InitializeAfterSetActive()
 
 }
 
-void FNetworkPlatformFile::ProcessServerCachedFilesResponse(FArrayReader& Response, const int32 ServerPackageVersion, const int32 ServerPackageLicenseeVersion)
+void FNetworkPlatformFile::ProcessServerCachedFilesResponse(FArrayReader& Response, const FPackageFileVersion& ServerPackageVersion, const int32 ServerPackageLicenseeVersion)
 {
 	/* The server root content directories */
 	TArray<FString> ServerRootContentDirectories;
@@ -267,9 +267,9 @@ void FNetworkPlatformFile::ProcessServerCachedFilesResponse(FArrayReader& Respon
 		TUniquePtr<IFileHandle> FileHandle(InnerPlatformFile->OpenRead(*CookedVersionFile));
 		if (FileHandle.IsValid())
 		{
-			int32 StoredPackageCookedVersion;
+			FPackageFileVersion StoredPackageCookedVersion;
 			int32 StoredPackageCookedLicenseeVersion;
-			if (FileHandle->Read((uint8*)&StoredPackageCookedVersion, sizeof(int32)) == true)
+			if (FileHandle->Read((uint8*)&StoredPackageCookedVersion, sizeof(FPackageFileVersion)) == true)
 			{
 				if (FileHandle->Read((uint8*)&StoredPackageCookedLicenseeVersion, sizeof(int32)) == true)
 				{
@@ -282,8 +282,8 @@ void FNetworkPlatformFile::ProcessServerCachedFilesResponse(FArrayReader& Respon
 					{
 						UE_LOG(LogNetworkPlatformFile, Display,
 							TEXT("Engine version mismatch: Server %d.%d, Stored %d.%d\n"),
-							ServerPackageVersion, ServerPackageLicenseeVersion,
-							StoredPackageCookedVersion, StoredPackageCookedLicenseeVersion);
+							ServerPackageVersion.ToValue(), ServerPackageLicenseeVersion,
+							StoredPackageCookedVersion.ToValue(), StoredPackageCookedLicenseeVersion);
 					}
 				}
 			}
@@ -326,7 +326,7 @@ void FNetworkPlatformFile::ProcessServerCachedFilesResponse(FArrayReader& Respon
 		TUniquePtr<IFileHandle> FileHandle(InnerPlatformFile->OpenWrite(*CookedVersionFile));
 		if (FileHandle.IsValid())
 		{
-			FileHandle->Write((const uint8*)&ServerPackageVersion, sizeof(int32));
+			FileHandle->Write((const uint8*)&ServerPackageVersion, sizeof(FPackageFileVersion));
 			FileHandle->Write((const uint8*)&ServerPackageLicenseeVersion, sizeof(int32));
 		}
 		else
@@ -822,7 +822,7 @@ void FNetworkPlatformFile::FillGetFileList(FNetworkFileArchive& Payload)
 	Payload << CustomPlatformData;
 }
 
-void FNetworkPlatformFile::ProcessServerInitialResponse(FArrayReader& InResponse, int32& OutServerPackageVersion, int32& OutServerPackageLicenseeVersion)
+void FNetworkPlatformFile::ProcessServerInitialResponse(FArrayReader& InResponse, FPackageFileVersion& OutServerPackageVersion, int32& OutServerPackageLicenseeVersion)
 {
 	// Receive the cooked version information.
 	InResponse << OutServerPackageVersion;
diff --git a/Engine/Source/Runtime/NetworkFile/Public/NetworkPlatformFile.h b/Engine/Source/Runtime/NetworkFile/Public/NetworkPlatformFile.h
--- a/Engine/Source/Runtime/NetworkFile/Public/NetworkPlatformFile.h
+++ b/Engine/Source/Runtime/NetworkFile/Public/NetworkPlatformFile.h
@@ -184,8 +184,8 @@ protected:
 
 	virtual void FillGetFileList(FNetworkFileArchive& Payload);
 
-	virtual void ProcessServerInitialResponse(FArrayReader& InResponse, int32& OutServerPackageVersion, int32& OutServerPackageLicenseeVersion);
-	virtual void ProcessServerCachedFilesResponse(FArrayReader& InReponse, const int32 ServerPackageVersion, const int32 ServerPackageLicenseeVersion );
+	virtual void ProcessServerInitialResponse(FArrayReader& InResponse, FPackageFileVersion& OutServerPackageVersion, int32& OutServerPackageLicenseeVersion);
+	virtual void ProcessServerCachedFilesResponse(FArrayReader& InReponse, const FPackageFileVersion& ServerPackageVersion, const int32 ServerPackageLicenseeVersion );
 
 private:
 
diff --git a/Engine/Source/Runtime/NetworkFileSystem/Private/NetworkFileServerConnection.cpp b/Engine/Source/Runtime/NetworkFileSystem/Private/NetworkFileServerConnection.cpp
--- a/Engine/Source/Runtime/NetworkFileSystem/Private/NetworkFileServerConnection.cpp
+++ b/Engine/Source/Runtime/NetworkFileSystem/Private/NetworkFileServerConnection.cpp
@@ -1095,7 +1095,7 @@ bool FNetworkFileServerClientConnection::ProcessGetFileList( FArchive& In, FArch
 
 	// report the package version information
 	// The downside of this is that ALL cooked data will get tossed on package version changes
-	int32 PackageFileUnrealVersion = GPackageFileUEVersion;
+	FPackageFileVersion PackageFileUnrealVersion = GPackageFileUEVersion;
 	Out << PackageFileUnrealVersion;
 	int32 PackageFileLicenseeUnrealVersion = GPackageFileLicenseeUEVersion;
 	Out << PackageFileLicenseeUnrealVersion;
diff --git a/Engine/Source/Runtime/Online/BuildPatchServices/Private/Data/ManifestUObject.cpp b/Engine/Source/Runtime/Online/BuildPatchServices/Private/Data/ManifestUObject.cpp
--- a/Engine/Source/Runtime/Online/BuildPatchServices/Private/Data/ManifestUObject.cpp
+++ b/Engine/Source/Runtime/Online/BuildPatchServices/Private/Data/ManifestUObject.cpp
@@ -519,7 +519,8 @@ bool FManifestUObject::LoadInternal(FArchive& Ar, FBuildPatchAppManifest& AppMan
 
 	// Make sure we use the correct serialization version, this is now fixed and must never use a newer version,
 	// because the property tag has changed in structure meaning older clients would not read correctly.
-	Ar.SetUEVer(VER_UE4_STRUCT_GUID_IN_PROPERTY_TAG - 1);
+	FPackageFileVersion Version = FPackageFileVersion::CreateUE4Version(VER_UE4_STRUCT_GUID_IN_PROPERTY_TAG - 1);
+	Ar.SetUEVer(Version);
 
 	if (Ar.IsLoading())
 	{
@@ -601,7 +602,8 @@ bool FManifestUObject::SaveInternal(FArchive& Ar, const FBuildPatchAppManifest&
 
 	// Make sure we use the correct serialization version, this is now fixed and must never use a newer version,
 	// because the property tag has changed in structure meaning older clients would not read correctly.
-	Ar.SetUEVer(VER_UE4_STRUCT_GUID_IN_PROPERTY_TAG - 1);
+	FPackageFileVersion Version = FPackageFileVersion::CreateUE4Version(VER_UE4_STRUCT_GUID_IN_PROPERTY_TAG - 1);
+	Ar.SetUEVer(Version);
 
 	if (Ar.IsLoading())
 	{
diff --git a/Engine/Source/Runtime/StreamingFile/Private/StreamingNetworkPlatformFile.cpp b/Engine/Source/Runtime/StreamingFile/Private/StreamingNetworkPlatformFile.cpp
--- a/Engine/Source/Runtime/StreamingFile/Private/StreamingNetworkPlatformFile.cpp
+++ b/Engine/Source/Runtime/StreamingFile/Private/StreamingNetworkPlatformFile.cpp
@@ -343,7 +343,7 @@ bool FStreamingNetworkPlatformFile::InitializeInternal(IPlatformFile* Inner, con
 	if(SendPayloadAndReceiveResponse(Payload,Response))
 	{
 		// Receive the cooked version information.
-		int32 ServerPackageVersion = 0;
+		FPackageFileVersion ServerPackageVersion;
 		int32 ServerPackageLicenseeVersion = 0;
 		ProcessServerInitialResponse(Response, ServerPackageVersion, ServerPackageLicenseeVersion);
 
diff --git a/Engine/Source/Runtime/Unix/UnixCommonStartup/Private/UnixCommonStartup.cpp b/Engine/Source/Runtime/Unix/UnixCommonStartup/Private/UnixCommonStartup.cpp
--- a/Engine/Source/Runtime/Unix/UnixCommonStartup/Private/UnixCommonStartup.cpp
+++ b/Engine/Source/Runtime/Unix/UnixCommonStartup/Private/UnixCommonStartup.cpp
@@ -198,7 +198,7 @@ int CommonUnixMain(int argc, char *argv[], int (*RealMain)(const TCHAR * Command
 	if (UE_BUILD_SHIPPING)
 	{
 		// only printed in shipping
-		printf("%s %d %d\n", StringCast<ANSICHAR>(*FEngineVersion::Current().ToString()).Get(), GPackageFileUEVersion, GPackageFileLicenseeUEVersion);
+		printf("%s %d %d\n", StringCast<ANSICHAR>(*FEngineVersion::Current().ToString()).Get(), GPackageFileUEVersion.ToValue(), GPackageFileLicenseeUEVersion);
 	}
 
 	int ErrorLevel = 0;
