Time: Clamp to int extents before casting. Otherwise, if it's out of range, the results are undefined.

diff --git a/Engine/Source/Runtime/Core/Public/Misc/FrameNumber.h b/Engine/Source/Runtime/Core/Public/Misc/FrameNumber.h
--- a/Engine/Source/Runtime/Core/Public/Misc/FrameNumber.h
+++ b/Engine/Source/Runtime/Core/Public/Misc/FrameNumber.h
@@ -62,8 +62,8 @@ struct FFrameNumber
 
 	friend FFrameNumber operator-(FFrameNumber A)                 { return FFrameNumber(-A.Value); }
 
-	friend FFrameNumber operator*(FFrameNumber A, float Scalar)   { return FFrameNumber(static_cast<int32>(FMath::FloorToDouble(double(A.Value) * Scalar))); }
-	friend FFrameNumber operator/(FFrameNumber A, float Scalar)   { return FFrameNumber(static_cast<int32>(FMath::FloorToDouble(double(A.Value) / Scalar))); }
+	friend FFrameNumber operator*(FFrameNumber A, float Scalar)   { return FFrameNumber(static_cast<int32>(FMath::Clamp(FMath::FloorToDouble(double(A.Value) * Scalar), (double)TNumericLimits<int32>::Min(), (double)TNumericLimits<int32>::Max()))); }
+	friend FFrameNumber operator/(FFrameNumber A, float Scalar)   { return FFrameNumber(static_cast<int32>(FMath::Clamp(FMath::FloorToDouble(double(A.Value) / Scalar), (double)TNumericLimits<int32>::Min(), (double)TNumericLimits<int32>::Max()))); }
 
 	/**
 	 * The value of the frame number
diff --git a/Engine/Source/Runtime/Core/Public/Misc/FrameRate.h b/Engine/Source/Runtime/Core/Public/Misc/FrameRate.h
--- a/Engine/Source/Runtime/Core/Public/Misc/FrameRate.h
+++ b/Engine/Source/Runtime/Core/Public/Misc/FrameRate.h
@@ -175,7 +175,7 @@ inline FFrameTime FFrameRate::AsFrameTime(double TimeInSeconds) const
 {
 	// @todo: sequencer-timecode: proper large number integer multiplication/division before coercion to float ?
 	const double       TimeAsFrame = (TimeInSeconds * Numerator) / Denominator;
-	const FFrameNumber FrameNumber = static_cast<int32>(FMath::FloorToDouble(TimeAsFrame));
+	const FFrameNumber FrameNumber = static_cast<int32>(FMath::Clamp(FMath::FloorToDouble(TimeAsFrame), (double)TNumericLimits<int32>::Min(), (double)TNumericLimits<int32>::Max()));
 
 	float SubFrame = static_cast<float>(TimeAsFrame - FMath::FloorToDouble(TimeAsFrame));
 	if (SubFrame > 0.f)
@@ -190,7 +190,7 @@ inline FFrameNumber FFrameRate::AsFrameNumber(double TimeInSeconds) const
 {
 	// @todo: sequencer-timecode: proper large number integer multiplication/division before coercion to float ?
 	const double       TimeAsFrame = (double(TimeInSeconds) * Numerator) / Denominator;
-	return static_cast<int32>(FMath::FloorToDouble(TimeAsFrame));
+	return static_cast<int32>(FMath::Clamp(FMath::FloorToDouble(TimeAsFrame), (double)TNumericLimits<int32>::Min(), (double)TNumericLimits<int32>::Max()));
 }
 
 inline bool operator==(const FFrameRate& A, const FFrameRate& B)
diff --git a/Engine/Source/Runtime/Core/Public/Misc/FrameTime.h b/Engine/Source/Runtime/Core/Public/Misc/FrameTime.h
--- a/Engine/Source/Runtime/Core/Public/Misc/FrameTime.h
+++ b/Engine/Source/Runtime/Core/Public/Misc/FrameTime.h
@@ -304,7 +304,7 @@ FORCEINLINE_DEBUGGABLE double FFrameTime::AsDecimal() const
 
 FORCEINLINE_DEBUGGABLE FFrameTime FFrameTime::FromDecimal(double InDecimalFrame)
 {
-	int32 NewFrame = static_cast<int32>(FMath::FloorToDouble(InDecimalFrame));
+	int32 NewFrame = static_cast<int32>(FMath::Clamp(FMath::FloorToDouble(InDecimalFrame), (double)TNumericLimits<int32>::Min(), (double)TNumericLimits<int32>::Max()));
 
 	// Ensure fractional parts above the highest sub frame float precision do not round to 0.0
 	double Fraction = InDecimalFrame - FMath::FloorToDouble(InDecimalFrame);
