Fix for TIsInvocable not compiling when using a member pointer callable and an incompatible target type.
CallableType renamed to TargetType, because it's not callable.
Deprecated macros removed.
Typo fix.

diff --git a/Engine/Source/Runtime/Core/Public/Templates/Invoke.h b/Engine/Source/Runtime/Core/Public/Templates/Invoke.h
--- a/Engine/Source/Runtime/Core/Public/Templates/Invoke.h
+++ b/Engine/Source/Runtime/Core/Public/Templates/Invoke.h
@@ -10,17 +10,20 @@
 
 namespace UE::Core::Private
 {
-	template <typename BaseType, typename CallableType>
-	FORCEINLINE decltype(auto) DereferenceIfNecessary(CallableType&& Callable)
+	template <typename OuterType, typename TargetType>
+	FORCEINLINE auto DereferenceIfNecessary(TargetType&& Target, const volatile OuterType* TargetPtr)
+		-> decltype((TargetType&&)Target)
 	{
-		if constexpr (std::is_convertible_v<std::decay_t<CallableType>*, std::decay_t<BaseType>*>)
-		{
-			return (CallableType&&)Callable;
-		}
-		else
-		{
-			return *(CallableType&&)Callable;
-		}
+		// If the target is the same as or is derived from the outer type, just return it unchanged.
+		return (TargetType&&)Target;
+	}
+
+	template <typename OuterType, typename TargetType>
+	FORCEINLINE auto DereferenceIfNecessary(TargetType&& Target, ...)
+		-> decltype(*(TargetType&&)Target)
+	{
+		// If the target is not related to the outer type, assume it's a (possibly smart) pointer and dereference it.
+		return *(TargetType&&)Target;
 	}
 }
 
@@ -44,23 +47,23 @@ FORCEINLINE auto Invoke(FuncType&& Func, ArgTypes&&... Args)
 	return Forward<FuncType>(Func)(Forward<ArgTypes>(Args)...);
 }
 
-template <typename ReturnType, typename ObjType, typename CallableType>
-FORCEINLINE auto Invoke(ReturnType ObjType::*pdm, CallableType&& Callable)
-	-> decltype(UE::Core::Private::DereferenceIfNecessary<ObjType>(Forward<CallableType>(Callable)).*pdm)
+template <typename ReturnType, typename ObjType, typename TargetType>
+FORCEINLINE auto Invoke(ReturnType ObjType::*pdm, TargetType&& Target)
+	-> decltype(UE::Core::Private::DereferenceIfNecessary<ObjType>(Forward<TargetType>(Target), &Target).*pdm)
 {
-	return UE::Core::Private::DereferenceIfNecessary<ObjType>(Forward<CallableType>(Callable)).*pdm;
+	return UE::Core::Private::DereferenceIfNecessary<ObjType>(Forward<TargetType>(Target), &Target).*pdm;
 }
 
 template <
 	typename    PtrMemFunType,
-	typename    CallableType,
+	typename    TargetType,
 	typename... ArgTypes,
 	typename    ObjType = TMemberFunctionPtrOuter_T<PtrMemFunType>
 >
-FORCEINLINE auto Invoke(PtrMemFunType PtrMemFun, CallableType&& Callable, ArgTypes&&... Args)
-	-> decltype((UE::Core::Private::DereferenceIfNecessary<ObjType>(Forward<CallableType>(Callable)).*PtrMemFun)(Forward<ArgTypes>(Args)...))
+FORCEINLINE auto Invoke(PtrMemFunType PtrMemFun, TargetType&& Target, ArgTypes&&... Args)
+	-> decltype((UE::Core::Private::DereferenceIfNecessary<ObjType>(Forward<TargetType>(Target), &Target).*PtrMemFun)(Forward<ArgTypes>(Args)...))
 {
-	return (UE::Core::Private::DereferenceIfNecessary<ObjType>(Forward<CallableType>(Callable)).*PtrMemFun)(Forward<ArgTypes>(Args)...);
+	return (UE::Core::Private::DereferenceIfNecessary<ObjType>(Forward<TargetType>(Target), &Target).*PtrMemFun)(Forward<ArgTypes>(Args)...);
 }
 
 
@@ -102,12 +105,9 @@ FORCEINLINE auto Invoke(PtrMemFunType PtrMemFun, CallableType&& Callable, ArgTyp
 #define UE_PROJECTION_MEMBER(Type, FuncName) \
 	[](auto&& Obj, auto&&... Args) -> decltype(auto) \
 	{ \
-		return UE::Core::Private::DereferenceIfNecessary<Type>(Forward<decltype(Obj)>(Obj)).FuncName(Forward<decltype(Args)>(Args)...); \
+		return UE::Core::Private::DereferenceIfNecessary<Type>(Forward<decltype(Obj)>(Obj), &Obj).FuncName(Forward<decltype(Args)>(Args)...); \
 	}
 
-#define PROJECTION(FuncName)              DEPRECATED_MACRO(4.26, "The PROJECTION macro is deprecated, please use UE_PROJECTION instead.") UE_PROJECTION(FuncName)
-#define PROJECTION_MEMBER(Type, FuncName) DEPRECATED_MACRO(4.26, "The PROJECTION_MEMBER macro is deprecated, please use UE_PROJECTION_MEMBER instead.") UE_PROJECTION_MEMBER(Type, FuncName)
-
 namespace UE::Core::Private
 {
 	template <typename, typename FuncType, typename... ArgTypes>
@@ -124,7 +124,7 @@ namespace UE::Core::Private
 
 /**
  * Trait for the type of the result when invoking a callable with the given argument types.
- * Not defined (as thus usable in SFINAE contexts) when the callable cannot be invoked with the given argument types.
+ * Not defined (and thus usable in SFINAE contexts) when the callable cannot be invoked with the given argument types.
  */
 template <typename FuncType, typename... ArgTypes>
 struct TInvokeResult : UE::Core::Private::TInvokeResult_Impl<void, FuncType, ArgTypes...>
