ensure PLATFORM_**  defines are set before checking them.

#rnx

diff --git a/Engine/Source/Runtime/Core/Public/Microsoft/AllowMicrosoftPlatformTypes.h b/Engine/Source/Runtime/Core/Public/Microsoft/AllowMicrosoftPlatformTypes.h
--- a/Engine/Source/Runtime/Core/Public/Microsoft/AllowMicrosoftPlatformTypes.h
+++ b/Engine/Source/Runtime/Core/Public/Microsoft/AllowMicrosoftPlatformTypes.h
@@ -1,5 +1,7 @@
 // Copyright Epic Games, Inc. All Rights Reserved.
 
+#include "HAL/Platform.h"
+
 #if PLATFORM_WINDOWS
 	#include "Windows/AllowWindowsPlatformTypes.h"
 #elif PLATFORM_HOLOLENS
