Memory tracing additional heaps

 * Introduce the concept of heaps to memory tracing. A heap is defined a block of memory which can be used to host allocations. Heaps can be given descriptive names and belong to logical hierarchy. At the top of the hierarchy is one of 16 possible "root heaps". Currently the only two root heaps are system memory (virtual memory) and video memory (VRAM). If no root heap is specified for an allocation, system memory is assumed.
 * Remove the CoreAdd/CoreRemove events and replace them with Alloc + MarkAllocAsHeap events.
 * Introduce a public interface for tracing memory allocations.
 * Reduce the memory footprint of allocation events by changing the "Owner" field (hash or return address) into a "CallstackId" (running unique callstack counter). Additonally introduced a AllocSystem/AllocVideo specialized event which encodes the root heap in the message type. Total size of the event went from 21 bytes to 17 bytes/18 bytes.
 * Remove realloc scopes. These scopes comprised 50% of events during memory tracing, but the information can be reconstructed during analysis in the case of reallocations. The only other use case (in Level.cpp) was using the scope to maintain tag while copying from one array to another. For this use case we add a specific scope for pushing the tag of a specific pointer.

diff --git a/Engine/Source/Runtime/Core/Private/IO/IoDispatcher.cpp b/Engine/Source/Runtime/Core/Private/IO/IoDispatcher.cpp
--- a/Engine/Source/Runtime/Core/Private/IO/IoDispatcher.cpp
+++ b/Engine/Source/Runtime/Core/Private/IO/IoDispatcher.cpp
@@ -925,7 +925,7 @@ void FIoRequestImpl::CreateBuffer(uint64 Size)
 	else
 	{
 		LLM_SCOPE(InheritedLLMTag);
-		UE_MEMSCOPE(InheritedTraceTag, ELLMTracker::Default);
+		UE_MEMSCOPE(InheritedTraceTag);
 		TRACE_CPUPROFILER_EVENT_SCOPE(AllocMemoryForRequest);
 		Buffer.Emplace(Size);
 	}
diff --git a/Engine/Source/Runtime/Core/Private/ProfilingDebugging/CallstackTrace.h b/Engine/Source/Runtime/Core/Private/ProfilingDebugging/CallstackTrace.h
--- a/Engine/Source/Runtime/Core/Private/ProfilingDebugging/CallstackTrace.h
+++ b/Engine/Source/Runtime/Core/Private/ProfilingDebugging/CallstackTrace.h
@@ -12,7 +12,7 @@
 UE_TRACE_CHANNEL_EXTERN(CallstackChannel)
 
 UE_TRACE_EVENT_BEGIN_EXTERN(Memory, CallstackSpec, NoSync)
-	UE_TRACE_EVENT_FIELD(uint64, Id)
+	UE_TRACE_EVENT_FIELD(uint32, CallstackId)
 	UE_TRACE_EVENT_FIELD(uint64[], Frames)
 UE_TRACE_EVENT_END()
 
@@ -25,49 +25,62 @@ namespace {
 		struct FBacktraceEntry
 		{
 			enum {MaxStackDepth = 256};
-			uint64	Id = 0;
+			uint64	Hash = 0;
 			uint32	FrameCount = 0;
 			uint64	Frames[MaxStackDepth] = { 0 };
 		};
 
 		FCallstackTracer(FMalloc* InMalloc)
 			: KnownSet(InMalloc)
+			, CallstackIdCounter(1) // 0 is reserved for "unknown callstack"
 		{
 			KnownSet.Reserve(1024 * 1024 * 2);
 		}
 		
-		void AddCallstack(const FBacktraceEntry& Entry)
+		uint32 AddCallstack(const FBacktraceEntry& Entry)
 		{
 			bool bAlreadyAdded = false;
 
 			// Our set implementation doesn't allow for zero entries (zero represents an empty element
 			// in the hash table), so if we get one due to really bad luck in our 64-bit Id calculation,
 			// treat it as a "1" instead, for purposes of tracking if we've seen that callstack.
-			KnownSet.FindOrAdd(FMath::Max(Entry.Id, (uint64)1), true, &bAlreadyAdded);
+			const uint64 Hash = FMath::Max(Entry.Hash, 1ull);
+			uint32 Id; 
+			KnownSet.Find(Hash, &Id, &bAlreadyAdded);
 			if (!bAlreadyAdded)
 			{
+				Id = CallstackIdCounter.fetch_add(1, std::memory_order_relaxed); 
+				KnownSet.Emplace(Hash, Id);
 				UE_TRACE_LOG(Memory, CallstackSpec, CallstackChannel)
-					<< CallstackSpec.Id(Entry.Id)
+					<< CallstackSpec.CallstackId(Id)
 					<< CallstackSpec.Frames(Entry.Frames, Entry.FrameCount);
 			}
+
+			return Id;
 		}
 	private:
 		struct FEncounteredCallstackSetEntry
 		{
-			std::atomic_uint64_t Data;
+			std::atomic_uint64_t Key;
+			std::atomic_uint32_t Value;
 
-			inline uint64 GetKey() const { return Data.load(std::memory_order_relaxed); }
-			inline bool GetValue() const { return true; }
-			inline bool IsEmpty() const { return Data.load(std::memory_order_relaxed) == 0; }
-			inline void SetKeyValue(uint64 Key, bool Value) { Data.store(Key, std::memory_order_relaxed); }
+			inline uint64 GetKey() const { return Key.load(std::memory_order_relaxed); }
+			inline uint32 GetValue() const { return Value.load(std::memory_order_relaxed); }
+			inline bool IsEmpty() const { return Key.load(std::memory_order_relaxed) == 0; }
+			inline void SetKeyValue(uint64 InKey, uint32 InValue)
+			{
+				Value.store(InValue, std::memory_order_release);
+				Key.store(InKey, std::memory_order_relaxed);
+			}
 			static inline uint32 KeyHash(uint64 Key) { return static_cast<uint32>(Key); }
 			static inline void ClearEntries(FEncounteredCallstackSetEntry* Entries, int32 EntryCount)
 			{
 				memset(Entries, 0, EntryCount * sizeof(FEncounteredCallstackSetEntry));
 			}
 		};
-		typedef TGrowOnlyLockFreeHash<FEncounteredCallstackSetEntry, uint64, bool> FEncounteredCallstackSet;
+		typedef TGrowOnlyLockFreeHash<FEncounteredCallstackSetEntry, uint64, uint32> FEncounteredCallstackSet;
 
 		FEncounteredCallstackSet 	KnownSet;
+		std::atomic_uint32_t		CallstackIdCounter;
 	};
 }
diff --git a/Engine/Source/Runtime/Core/Private/ProfilingDebugging/MemoryAllocationTrace.cpp b/Engine/Source/Runtime/Core/Private/ProfilingDebugging/MemoryAllocationTrace.cpp
--- a/Engine/Source/Runtime/Core/Private/ProfilingDebugging/MemoryAllocationTrace.cpp
+++ b/Engine/Source/Runtime/Core/Private/ProfilingDebugging/MemoryAllocationTrace.cpp
@@ -1,7 +1,6 @@
 // Copyright Epic Games, Inc. All Rights Reserved.
 
 #include "ProfilingDebugging/MemoryAllocationTrace.h"
-#include "HAL/PlatformTime.h"
 
 #if UE_MEMORY_TRACE_ENABLED
 
@@ -15,11 +14,11 @@ TRACELOG_API void Update();
 } // namespace UE
 
 ////////////////////////////////////////////////////////////////////////////////
-UE_TRACE_CHANNEL(MemSummaryChannel)
 UE_TRACE_CHANNEL_DEFINE(MemAllocChannel)
 
 UE_TRACE_EVENT_BEGIN(Memory, Init, NoSync|Important)
 	UE_TRACE_EVENT_FIELD(uint32, MarkerPeriod)
+	UE_TRACE_EVENT_FIELD(uint8, Version)
 	UE_TRACE_EVENT_FIELD(uint8, MinAlignment)
 	UE_TRACE_EVENT_FIELD(uint8, SizeShift)
 	UE_TRACE_EVENT_FIELD(uint8, Mode)
@@ -29,49 +28,95 @@ UE_TRACE_EVENT_BEGIN(Memory, Marker)
 	UE_TRACE_EVENT_FIELD(uint64, Cycle)
 UE_TRACE_EVENT_END()
 
-UE_TRACE_EVENT_BEGIN(Memory, CoreAdd)
-	UE_TRACE_EVENT_FIELD(uint64, Owner)
-	UE_TRACE_EVENT_FIELD(void*, Base)
+UE_TRACE_EVENT_BEGIN(Memory, Alloc)
+	UE_TRACE_EVENT_FIELD(uint64, Address)
+	UE_TRACE_EVENT_FIELD(uint32, CallstackId)
 	UE_TRACE_EVENT_FIELD(uint32, Size)
+	UE_TRACE_EVENT_FIELD(uint8, AlignmentPow2_SizeLower)
+	UE_TRACE_EVENT_FIELD(uint8, RootHeap)
 UE_TRACE_EVENT_END()
 
-UE_TRACE_EVENT_BEGIN(Memory, CoreRemove)
-	UE_TRACE_EVENT_FIELD(uint64, Owner)
-	UE_TRACE_EVENT_FIELD(void*, Base)
+UE_TRACE_EVENT_BEGIN(Memory, AllocSystem)
+	UE_TRACE_EVENT_FIELD(uint64, Address)
+	UE_TRACE_EVENT_FIELD(uint32, CallstackId)
 	UE_TRACE_EVENT_FIELD(uint32, Size)
+	UE_TRACE_EVENT_FIELD(uint8, AlignmentPow2_SizeLower)
 UE_TRACE_EVENT_END()
 
-UE_TRACE_EVENT_BEGIN(Memory, Alloc)
-	UE_TRACE_EVENT_FIELD(uint64, Owner)
-	UE_TRACE_EVENT_FIELD(void*, Address)
+UE_TRACE_EVENT_BEGIN(Memory, AllocVideo)
+	UE_TRACE_EVENT_FIELD(uint64, Address)
+	UE_TRACE_EVENT_FIELD(uint32, CallstackId)
 	UE_TRACE_EVENT_FIELD(uint32, Size)
-	UE_TRACE_EVENT_FIELD(uint8, Alignment_SizeLower)
+	UE_TRACE_EVENT_FIELD(uint8, AlignmentPow2_SizeLower)
 UE_TRACE_EVENT_END()
 
 UE_TRACE_EVENT_BEGIN(Memory, Free)
-	UE_TRACE_EVENT_FIELD(void*, Address)
+	UE_TRACE_EVENT_FIELD(uint64, Address)
+	UE_TRACE_EVENT_FIELD(uint8, RootHeap)
+UE_TRACE_EVENT_END()
+
+UE_TRACE_EVENT_BEGIN(Memory, FreeSystem)
+	UE_TRACE_EVENT_FIELD(uint64, Address)
+UE_TRACE_EVENT_END()
+
+UE_TRACE_EVENT_BEGIN(Memory, FreeVideo)
+	UE_TRACE_EVENT_FIELD(uint64, Address)
 UE_TRACE_EVENT_END()
 
 UE_TRACE_EVENT_BEGIN(Memory, ReallocAlloc)
-	UE_TRACE_EVENT_FIELD(uint64, Owner)
-	UE_TRACE_EVENT_FIELD(void*, Address)
+	UE_TRACE_EVENT_FIELD(uint64, Address)
+	UE_TRACE_EVENT_FIELD(uint32, CallstackId)
+	UE_TRACE_EVENT_FIELD(uint32, Size)
+	UE_TRACE_EVENT_FIELD(uint8, AlignmentPow2_SizeLower)
+	UE_TRACE_EVENT_FIELD(uint8, RootHeap)
+UE_TRACE_EVENT_END()
+
+UE_TRACE_EVENT_BEGIN(Memory, ReallocAllocSystem)
+	UE_TRACE_EVENT_FIELD(uint64, Address)
+	UE_TRACE_EVENT_FIELD(uint32, CallstackId)
 	UE_TRACE_EVENT_FIELD(uint32, Size)
-	UE_TRACE_EVENT_FIELD(uint8, Alignment_SizeLower)
+	UE_TRACE_EVENT_FIELD(uint8, AlignmentPow2_SizeLower)
 UE_TRACE_EVENT_END()
 
 UE_TRACE_EVENT_BEGIN(Memory, ReallocFree)
-	UE_TRACE_EVENT_FIELD(void*, Address)
+	UE_TRACE_EVENT_FIELD(uint64, Address_RootHeap)
 UE_TRACE_EVENT_END()
 
+UE_TRACE_EVENT_BEGIN(Memory, HeapSpec, NoSync|Important)
+	UE_TRACE_EVENT_FIELD(HeapId, Id)
+	UE_TRACE_EVENT_FIELD(HeapId, ParentId)
+	UE_TRACE_EVENT_FIELD(uint16, Flags)
+	UE_TRACE_EVENT_FIELD(UE::Trace::WideString, Name)
+UE_TRACE_EVENT_END()
+
+UE_TRACE_EVENT_BEGIN(Memory, HeapMarkAlloc)
+	UE_TRACE_EVENT_FIELD(uint64, Address)
+	UE_TRACE_EVENT_FIELD(uint16, Flags)
+	UE_TRACE_EVENT_FIELD(HeapId, Heap)
+UE_TRACE_EVENT_END()
+
+UE_TRACE_EVENT_BEGIN(Memory, HeapUnmarkAlloc)
+	UE_TRACE_EVENT_FIELD(uint64, Address)
+	UE_TRACE_EVENT_FIELD(HeapId, Heap)
+UE_TRACE_EVENT_END()
+
+// If layout of the above events are changed, bump this version number
+constexpr uint8 MemoryTraceVersion = 1;
 
 ////////////////////////////////////////////////////////////////////////////////
-void FAllocationTrace::Initialize()
+void FAllocationTrace::Initialize() const
 {
 	UE_TRACE_LOG(Memory, Init, MemAllocChannel)
 		<< Init.MarkerPeriod(MarkerSamplePeriod + 1)
+		<< Init.Version(MemoryTraceVersion)
 		<< Init.MinAlignment(uint8(MIN_ALIGNMENT))
 		<< Init.SizeShift(uint8(SizeShift));
 
+	const HeapId SystemRootHeap = RootHeapSpec(TEXT("System memory"));
+	check(SystemRootHeap == EMemoryTraceRootHeap::SystemMemory);
+	const HeapId VideoRootHeap = RootHeapSpec(TEXT("VideoMemory"));
+	check(VideoRootHeap == EMemoryTraceRootHeap::VideoMemory);
+
 	static_assert((1 << SizeShift) - 1 <= MIN_ALIGNMENT, "Not enough bits to pack size fields");
 }
 
@@ -84,7 +129,7 @@ void FAllocationTrace::EnableTracePump()
 ////////////////////////////////////////////////////////////////////////////////
 void FAllocationTrace::Update()
 {
-	uint32 TheCount = MarkerCounter.fetch_add(1, std::memory_order_relaxed);
+	const uint32 TheCount = MarkerCounter.fetch_add(1, std::memory_order_relaxed);
 	if ((TheCount & MarkerSamplePeriod) == 0)
 	{
 		UE_TRACE_LOG(Memory, Marker, MemAllocChannel)
@@ -98,71 +143,171 @@ void FAllocationTrace::Update()
 }
 
 ////////////////////////////////////////////////////////////////////////////////
-void FAllocationTrace::CoreAdd(void* Base, size_t Size, void* Owner)
+void FAllocationTrace::Alloc(void* Address, size_t Size, uint32 Alignment, uint32 Owner, HeapId RootHeap)
 {
-	UE_TRACE_LOG(Memory, CoreAdd, MemAllocChannel)
-		<< CoreAdd.Owner(uint64(Owner))
-		<< CoreAdd.Base(Base)
-		<< CoreAdd.Size(uint32(Size >> SizeShift));
+	check(RootHeap < 16);
+	const uint8 AlignmentPow2 = uint8(FPlatformMath::CountTrailingZeros(Alignment));
+	const uint32 Alignment_SizeLower = AlignmentPow2 << SizeShift | (Size & ((1 << SizeShift) - 1));
+
+	switch (RootHeap)
+	{
+		case EMemoryTraceRootHeap::SystemMemory:
+			{
+				UE_TRACE_LOG(Memory, AllocSystem, MemAllocChannel)
+					<< AllocSystem.CallstackId(Owner)
+					<< AllocSystem.Address(uint64(Address))
+					<< AllocSystem.Size(uint32(Size >> SizeShift))
+					<< AllocSystem.AlignmentPow2_SizeLower(uint8(Alignment_SizeLower));
+				break;
+			}
+		case EMemoryTraceRootHeap::VideoMemory:
+			{
+				UE_TRACE_LOG(Memory, AllocVideo, MemAllocChannel)
+					<< AllocVideo.CallstackId(Owner)
+					<< AllocVideo.Address(uint64(Address))
+					<< AllocVideo.Size(uint32(Size >> SizeShift))
+					<< AllocVideo.AlignmentPow2_SizeLower(uint8(Alignment_SizeLower));
+				break;
+			}
+		default:
+			{
+				UE_TRACE_LOG(Memory, Alloc, MemAllocChannel)
+					<< Alloc.CallstackId(Owner)
+					<< Alloc.Address(uint64(Address))
+					<< Alloc.RootHeap(uint8(RootHeap))
+					<< Alloc.Size(uint32(Size >> SizeShift))
+					<< Alloc.AlignmentPow2_SizeLower(uint8(Alignment_SizeLower));
+				break;
+			}
+	}
 
 	Update();
 }
 
 ////////////////////////////////////////////////////////////////////////////////
-void FAllocationTrace::CoreRemove(void* Base, size_t Size, void* Owner)
+void FAllocationTrace::Free(void* Address, HeapId RootHeap)
 {
-	UE_TRACE_LOG(Memory, CoreRemove, MemAllocChannel)
-		<< CoreRemove.Owner(uint64(Owner))
-		<< CoreRemove.Base(Base)
-		<< CoreRemove.Size(uint32(Size >> SizeShift));
-
+	check(RootHeap < 16);
+	switch (RootHeap)
+	{
+		case EMemoryTraceRootHeap::SystemMemory:
+			{
+				UE_TRACE_LOG(Memory, FreeSystem, MemAllocChannel)
+					<< FreeSystem.Address(uint64(Address));
+			}
+		case EMemoryTraceRootHeap::VideoMemory:
+			{
+				UE_TRACE_LOG(Memory, FreeVideo, MemAllocChannel)
+					<< FreeVideo.Address(uint64(Address));
+			}
+		default:
+			{
+				UE_TRACE_LOG(Memory, Free, MemAllocChannel)
+					<< Free.Address(uint64(Address))
+					<< Free.RootHeap(uint8(RootHeap));
+				break;
+			}
+	}
+	
 	Update();
 }
 
 ////////////////////////////////////////////////////////////////////////////////
-void FAllocationTrace::Alloc(void* Address, size_t Size, uint32 Alignment, void* Owner)
+void FAllocationTrace::ReallocAlloc(void* Address, size_t Size, uint32 Alignment, uint32 Owner, HeapId RootHeap)
 {
-	uint32 Alignment_SizeLower = Alignment | (Size & ((1 << SizeShift) - 1));
-
-	UE_TRACE_LOG(Memory, Alloc, MemAllocChannel)
-		<< Alloc.Owner(uint64(Owner))
-		<< Alloc.Address(Address)
-		<< Alloc.Size(uint32(Size >> SizeShift))
-		<< Alloc.Alignment_SizeLower(uint8(Alignment_SizeLower));
+	check(RootHeap < 16);
+	const uint8 AlignmentPow2 = uint8(FPlatformMath::CountTrailingZeros(Alignment));
+	const uint32 Alignment_SizeLower = AlignmentPow2 << SizeShift | (Size & ((1 << SizeShift) - 1));
+	
+	switch (RootHeap)
+	{
+	case EMemoryTraceRootHeap::SystemMemory:
+		{
+			UE_TRACE_LOG(Memory, ReallocAllocSystem, MemAllocChannel)
+				<< ReallocAllocSystem.CallstackId(Owner)
+				<< ReallocAllocSystem.Address(uint64(Address))
+				<< ReallocAllocSystem.Size(uint32(Size >> SizeShift))
+				<< ReallocAllocSystem.AlignmentPow2_SizeLower(uint8(Alignment_SizeLower));
+		}
+	default:
+		{
+			UE_TRACE_LOG(Memory, ReallocAlloc, MemAllocChannel)
+				<< ReallocAlloc.CallstackId(Owner)
+				<< ReallocAlloc.Address(uint64(Address))
+				<< ReallocAlloc.RootHeap(uint8(RootHeap))
+				<< ReallocAlloc.Size(uint32(Size >> SizeShift))
+				<< ReallocAlloc.AlignmentPow2_SizeLower(uint8(Alignment_SizeLower));
+			break;
+		}
+	}
 
 	Update();
 }
 
 ////////////////////////////////////////////////////////////////////////////////
-void FAllocationTrace::Free(void* Address)
+void FAllocationTrace::ReallocFree(void* Address, HeapId RootHeap)
 {
-	UE_TRACE_LOG(Memory, Free, MemAllocChannel)
-		<< Free.Address(Address);
+	check(RootHeap < 16);
+	const uint64 Address_RootHeap = uint64(Address) | (uint64(RootHeap) << HeapShift);
+	UE_TRACE_LOG(Memory, ReallocFree, MemAllocChannel)
+		<< ReallocFree.Address_RootHeap(Address_RootHeap);
 
 	Update();
 }
 
 ////////////////////////////////////////////////////////////////////////////////
-void FAllocationTrace::ReallocAlloc(void* Address, size_t Size, uint32 Alignment, void* Owner)
+HeapId FAllocationTrace::HeapSpec(HeapId ParentId, const TCHAR* Name, EMemoryTraceHeapFlags Flags) const
 {
-	uint32 Alignment_SizeLower = Alignment | (Size & ((1 << SizeShift) - 1));
-
-	UE_TRACE_LOG(Memory, ReallocAlloc, MemAllocChannel)
-		<< ReallocAlloc.Owner(uint64(Owner))
-		<< ReallocAlloc.Address(Address)
-		<< ReallocAlloc.Size(uint32(Size >> SizeShift))
-		<< ReallocAlloc.Alignment_SizeLower(uint8(Alignment_SizeLower));
-
-	Update();
+	static std::atomic<HeapId> HeapIdCount(15); //Reserve indexes for root heaps
+	const HeapId Id = HeapIdCount.fetch_add(1);
+	const uint32 NameLen = FCString::Strlen(Name);
+	const uint32 DataSize = NameLen * sizeof(TCHAR);
+	check(ParentId < Id);
+
+	UE_TRACE_LOG(Memory, HeapSpec, MemAllocChannel, DataSize)
+		<< HeapSpec.Id(Id)
+		<< HeapSpec.ParentId(ParentId)
+		<< HeapSpec.Name(Name, NameLen)
+		<< HeapSpec.Flags(uint16(Flags));
+
+	return Id;
 }
 
 ////////////////////////////////////////////////////////////////////////////////
-void FAllocationTrace::ReallocFree(void* Address)
+HeapId FAllocationTrace::RootHeapSpec(const TCHAR* Name, EMemoryTraceHeapFlags Flags) const
 {
-	UE_TRACE_LOG(Memory, ReallocFree, MemAllocChannel)
-		<< ReallocFree.Address(Address);
+	static std::atomic<HeapId> RootHeapCount(0);
+	const HeapId Id = RootHeapCount.fetch_add(1);
+	check(Id <= EMemoryTraceRootHeap::EndReserved);
+	
+	const uint32 NameLen = FCString::Strlen(Name);
+	const uint32 DataSize = NameLen * sizeof(TCHAR);
+	
+	UE_TRACE_LOG(Memory, HeapSpec, MemAllocChannel, DataSize)
+		<< HeapSpec.Id(Id)
+		<< HeapSpec.ParentId(HeapId(~0))
+		<< HeapSpec.Name(Name, NameLen)
+		<< HeapSpec.Flags(uint16(EMemoryTraceHeapFlags::Root|Flags));
+
+	return Id;
+}
 
+////////////////////////////////////////////////////////////////////////////////
+void FAllocationTrace::MarkAllocAsHeap(void* Address, HeapId Heap, EMemoryTraceHeapAllocationFlags Flags)
+{
+	UE_TRACE_LOG(Memory, HeapMarkAlloc, MemAllocChannel)
+		<< HeapMarkAlloc.Address(uint64(Address))
+		<< HeapMarkAlloc.Heap(Heap)
+		<< HeapMarkAlloc.Flags(uint16(EMemoryTraceHeapAllocationFlags::Heap | Flags));
 	Update();
 }
 
+void FAllocationTrace::UnmarkAllocAsHeap(void* Address, HeapId Heap)
+{
+	// Sets all flags to zero
+	UE_TRACE_LOG(Memory, HeapUnmarkAlloc, MemAllocChannel)
+		<< HeapUnmarkAlloc.Address(uint64(Address))
+		<< HeapUnmarkAlloc.Heap(Heap);
+	Update();
+}
 #endif // UE_MEMORY_TRACE_ENABLED
diff --git a/Engine/Source/Runtime/Core/Private/ProfilingDebugging/MemoryAllocationTrace.h b/Engine/Source/Runtime/Core/Private/ProfilingDebugging/MemoryAllocationTrace.h
--- a/Engine/Source/Runtime/Core/Private/ProfilingDebugging/MemoryAllocationTrace.h
+++ b/Engine/Source/Runtime/Core/Private/ProfilingDebugging/MemoryAllocationTrace.h
@@ -15,14 +15,16 @@
 class FAllocationTrace
 {
 public:
-	void	Initialize();
+	void	Initialize() const;
 	void	EnableTracePump();
-	void	CoreAdd(void* Base, size_t Size, void* Owner);
-	void	CoreRemove(void* Base, size_t Size, void* Owner);
-	void	Alloc(void* Address, size_t Size, uint32 Alignment, void* Owner);
-	void	Free(void* Address);
-	void	ReallocAlloc(void* Address, size_t Size, uint32 Alignment, void* Owner);
-	void	ReallocFree(void* Address);
+	HeapId	HeapSpec(HeapId ParentId, const TCHAR* Name, EMemoryTraceHeapFlags Flags = EMemoryTraceHeapFlags::None) const;
+	HeapId	RootHeapSpec(const TCHAR* Name, EMemoryTraceHeapFlags = EMemoryTraceHeapFlags::None) const;
+   	void	MarkAllocAsHeap(void* Address, HeapId Heap, EMemoryTraceHeapAllocationFlags Flags = EMemoryTraceHeapAllocationFlags::None);
+   	void	UnmarkAllocAsHeap(void* Address, HeapId Heap);
+	void	Alloc(void* Address, size_t Size, uint32 Alignment, uint32 Owner, HeapId RootHeap = EMemoryTraceRootHeap::SystemMemory);
+	void	Free(void* Address, HeapId RootHeap = EMemoryTraceRootHeap::SystemMemory);
+	void	ReallocAlloc(void* Address, size_t NewSize, uint32 Alignment, uint32 Owner, HeapId RootHeap = EMemoryTraceRootHeap::SystemMemory);
+	void	ReallocFree(void* Address, HeapId RootHeap = EMemoryTraceRootHeap::SystemMemory);
 
 private:
 	void				Update();
@@ -31,6 +33,7 @@ private:
 	bool				bPumpTrace = false;
 	static const uint32 MarkerSamplePeriod	= (4 << 10) - 1;
 	static const uint32 SizeShift = 3;
+	static const uint32 HeapShift = 60;	
 };
 
 #endif // UE_MEMORY_TRACE_ENABLED
diff --git a/Engine/Source/Runtime/Core/Private/ProfilingDebugging/MemoryTrace.h b/Engine/Source/Runtime/Core/Private/ProfilingDebugging/MemoryTrace.h
deleted file mode 100644
diff --git a/Engine/Source/Runtime/Core/Private/ProfilingDebugging/Microsoft/MemoryTrace.cpp b/Engine/Source/Runtime/Core/Private/ProfilingDebugging/Microsoft/MemoryTrace.cpp
--- a/Engine/Source/Runtime/Core/Private/ProfilingDebugging/Microsoft/MemoryTrace.cpp
+++ b/Engine/Source/Runtime/Core/Private/ProfilingDebugging/Microsoft/MemoryTrace.cpp
@@ -23,18 +23,13 @@
 ////////////////////////////////////////////////////////////////////////////////
 void	Backtracer_Create(FMalloc*);
 void	Backtracer_Initialize();
-void*	Backtracer_GetBacktraceId(void*);
+uint32	Backtracer_GetBacktraceId(void*);
 void 	MemoryTrace_InitTags(FMalloc*);
 
 ////////////////////////////////////////////////////////////////////////////////
-FORCEINLINE void* GetOwner(bool bLight)
+FORCEINLINE uint32 GetOwner()
 {
 	void* RetAddrAddr = PLATFORM_RETURN_ADDRESS_POINTER();
-	if (bLight)
-	{
-		return *(void**)RetAddrAddr;
-	}
-
 	return Backtracer_GetBacktraceId(RetAddrAddr);
 }
 
@@ -47,6 +42,12 @@ public:
 	void Construct(ArgTypes... Args)
 	{
 		::new (Buffer) T(Args...);
+		bIsConstructed = true;
+	}
+
+	bool IsConstructed() const
+	{
+		return bIsConstructed;
 	}
 
 	T* operator & ()	{ return (T*)Buffer; }
@@ -54,6 +55,7 @@ public:
 
 protected:
 	uint8 Buffer[sizeof(T)];
+	bool bIsConstructed;
 };
 
 ////////////////////////////////////////////////////////////////////////////////
@@ -76,7 +78,7 @@ class FMallocWrapper
 	: public FMalloc
 {
 public:
-							FMallocWrapper(FMalloc* InMalloc, bool bInLight);
+							FMallocWrapper(FMalloc* InMalloc);
 
 private:
 	struct FCookie
@@ -99,13 +101,11 @@ private:
 	virtual void			SetupTLSCachesOnCurrentThread() override					{ return InnerMalloc->SetupTLSCachesOnCurrentThread(); }
 
 	FMalloc*				InnerMalloc;
-	bool					bLight;
 };
 
 ////////////////////////////////////////////////////////////////////////////////
-FMallocWrapper::FMallocWrapper(FMalloc* InMalloc, bool bInLight)
+FMallocWrapper::FMallocWrapper(FMalloc* InMalloc)
 : InnerMalloc(InMalloc)
-, bLight(bInLight)
 {
 }
 
@@ -128,7 +128,7 @@ void* FMallocWrapper::Malloc(SIZE_T Size, uint32 Alignment)
 	uint32 ActualAlignment = GetActualAlignment(Size, Alignment);
 	void* Address = InnerMalloc->Malloc(Size, Alignment);
 
-	void* Owner = GetOwner(bLight);
+	const uint32 Owner = GetOwner();
 	GAllocationTrace->Alloc(Address, Size, ActualAlignment, Owner);
 
 	return Address;
@@ -153,7 +153,7 @@ void* FMallocWrapper::Realloc(void* PrevAddress, SIZE_T NewSize, uint32 Alignmen
 
 	void* RetAddress = InnerMalloc->Realloc(PrevAddress, NewSize, Alignment);
 
-	void* Owner = GetOwner(bLight);
+	const uint32 Owner = GetOwner();
 	Alignment = GetActualAlignment(NewSize, Alignment);
 	GAllocationTrace->ReallocAlloc(RetAddress, NewSize, Alignment, Owner);
 
@@ -355,8 +355,9 @@ LPVOID WINAPI FVirtualWinApiHooks::VmAlloc(LPVOID Address, SIZE_T Size, DWORD Ty
 	LPVOID Ret = VmAllocOrig(Address, Size, Type, Protect);
 	if (Ret != nullptr && (Type & MEM_COMMIT))
 	{
-		void* Owner = GetOwner(bLight);
-		GAllocationTrace->CoreAdd(Ret, Size, Owner);
+		const uint32 Owner = GetOwner();
+		GAllocationTrace->Alloc(Ret, Size, 0, Owner, EMemoryTraceRootHeap::SystemMemory);
+		GAllocationTrace->MarkAllocAsHeap(Ret, EMemoryTraceRootHeap::SystemMemory);
 	}
 
 	return Ret;
@@ -365,8 +366,7 @@ LPVOID WINAPI FVirtualWinApiHooks::VmAlloc(LPVOID Address, SIZE_T Size, DWORD Ty
 ////////////////////////////////////////////////////////////////////////////////
 BOOL WINAPI FVirtualWinApiHooks::VmFree(LPVOID Address, SIZE_T Size, DWORD Type)
 {
-	void* Owner = GetOwner(bLight);
-	GAllocationTrace->CoreRemove(Address, Size, Owner);
+	GAllocationTrace->Free(Address, EMemoryTraceRootHeap::SystemMemory);
 	return VmFreeOrig(Address, Size, Type);
 }
 
@@ -376,8 +376,9 @@ LPVOID WINAPI FVirtualWinApiHooks::VmAllocEx(HANDLE Process, LPVOID Address, SIZ
 	LPVOID Ret = VmAllocExOrig(Process, Address, Size, Type, Protect);
 	if (Process == GetCurrentProcess() && Ret != nullptr && (Type & MEM_COMMIT))
 	{
-		void* Owner = GetOwner(bLight);
-		GAllocationTrace->CoreAdd(Ret, Size, Owner);
+		const uint32 Owner = GetOwner();
+		GAllocationTrace->Alloc(Ret, Size, 0, Owner);
+		GAllocationTrace->MarkAllocAsHeap(Ret, EMemoryTraceRootHeap::SystemMemory);
 	}
 
 	return Ret;
@@ -388,8 +389,7 @@ BOOL WINAPI FVirtualWinApiHooks::VmFreeEx(HANDLE Process, LPVOID Address, SIZE_T
 {
 	if (Process == GetCurrentProcess())
 	{
-		void* Owner = GetOwner(bLight);
-		GAllocationTrace->CoreRemove(Address, Size, Owner);
+		GAllocationTrace->Free(Address, EMemoryTraceRootHeap::SystemMemory);
 	}
 
 	return VmFreeExOrig(Process, Address, Size, Type);
@@ -447,7 +447,7 @@ FMalloc* MemoryTrace_Create(FMalloc* InMalloc)
 #endif // defined(PLATFORM_SUPPORTS_TRACE_WIN32_VIRTUAL_MEMORY_HOOKS)
 
 		static FUndestructed<FMallocWrapper> MemoryTrace;
-		MemoryTrace.Construct(InMalloc, false);
+		MemoryTrace.Construct(InMalloc);
 
 		return &MemoryTrace;
 	}
@@ -463,21 +463,80 @@ void MemoryTrace_Initialize()
 	Backtracer_Initialize();
 }
 
-#if 0
+
+////////////////////////////////////////////////////////////////////////////////
+HeapId MemoryTrace_HeapSpec(HeapId ParentId, const TCHAR* Name, EMemoryTraceHeapFlags Flags)
+{
+	if (GAllocationTrace.IsConstructed())
+	{
+		return GAllocationTrace->HeapSpec(ParentId, Name, Flags);
+	}
+	return ~0;
+}
+
+////////////////////////////////////////////////////////////////////////////////
+HeapId MemoryTrace_RootHeapSpec(const TCHAR* Name)
+{
+	if (GAllocationTrace.IsConstructed())
+	{
+		return GAllocationTrace->RootHeapSpec(Name);
+	}
+	return ~0;
+}
+
 ////////////////////////////////////////////////////////////////////////////////
-void MemoryTrace_CoreAdd(void* Base, SIZE_T Size)
+void MemoryTrace_MarkAllocAsHeap(uint64 Address, HeapId Heap, EMemoryTraceHeapAllocationFlags Flags)
 {
-	void* Owner = GetOwner(bLight);
-	GAllocationTrace->CoreAdd(Base, Size, Owner);
+	if (GAllocationTrace.IsConstructed())
+	{
+		GAllocationTrace->MarkAllocAsHeap((void*)Address, Heap, Flags);
+	}
 }
 
 ////////////////////////////////////////////////////////////////////////////////
-void MemoryTrace_CoreRemove(void* Base, SIZE_T Size)
+void MemoryTrace_UnmarkAllocAsHeap(uint64 Address, HeapId Heap)
 {
-	void* Owner = GetOwner(bLight);
-	GAllocationTrace->CoreRemove(Base, Size, Owner);
+	if (GAllocationTrace.IsConstructed())
+	{
+		GAllocationTrace->UnmarkAllocAsHeap((void*)Address, Heap);
+	}
+}
+
+////////////////////////////////////////////////////////////////////////////////
+void MemoryTrace_Alloc(uint64 Address, uint64 Size, uint32 Alignment, HeapId RootHeap /*= EMemoryTraceRootHeap::SystemMemory*/)
+{
+	if (GAllocationTrace.IsConstructed())
+	{
+		GAllocationTrace->Alloc((void*)Address, Size, Alignment, GetOwner(), RootHeap);
+	}
+}
+
+////////////////////////////////////////////////////////////////////////////////
+void MemoryTrace_Free(uint64 Address, HeapId RootHeap /*= EMemoryTraceRootHeap::SystemMemory*/)
+{
+	if (GAllocationTrace.IsConstructed())
+	{
+		GAllocationTrace->Free((void*)Address, RootHeap);
+	}
+}
+
+////////////////////////////////////////////////////////////////////////////////
+void MemoryTrace_ReallocFree(uint64 Address, HeapId RootHeap /*= EMemoryTraceRootHeap::SystemMemory*/)
+{
+	if (GAllocationTrace.IsConstructed())
+	{
+		GAllocationTrace->ReallocFree((void*)Address, RootHeap);
+	}
+}
+
+////////////////////////////////////////////////////////////////////////////////
+void MemoryTrace_ReallocAlloc(uint64 Address, uint64 NewSize, uint32 Alignment, HeapId RootHeap /*= EMemoryTraceRootHeap::SystemMemory*/)
+{
+	if (GAllocationTrace.IsConstructed())
+	{
+		GAllocationTrace->ReallocAlloc((void*)Address, NewSize, Alignment, GetOwner(), RootHeap);
+	}
 }
-#endif // 0
 
 #include "Windows/HideWindowsPlatformTypes.h"
 
diff --git a/Engine/Source/Runtime/Core/Private/ProfilingDebugging/Microsoft/WindowsCallstackTrace.cpp b/Engine/Source/Runtime/Core/Private/ProfilingDebugging/Microsoft/WindowsCallstackTrace.cpp
--- a/Engine/Source/Runtime/Core/Private/ProfilingDebugging/Microsoft/WindowsCallstackTrace.cpp
+++ b/Engine/Source/Runtime/Core/Private/ProfilingDebugging/Microsoft/WindowsCallstackTrace.cpp
@@ -117,7 +117,7 @@ public:
 	static FBacktracer*		Get();
 	void					AddModule(UPTRINT Base, const TCHAR* Name);
 	void					RemoveModule(UPTRINT Base);
-	void*					GetBacktraceId(void* AddressOfReturnAddress) const;
+	uint32					GetBacktraceId(void* AddressOfReturnAddress) const;
 	int32					GetFrameSize(void* FunctionAddress) const;
 
 private:
@@ -557,7 +557,7 @@ const FBacktracer::FFunction* FBacktracer::LookupFunction(UPTRINT Address, FLook
 }
 
 ////////////////////////////////////////////////////////////////////////////////
-void* FBacktracer::GetBacktraceId(void* AddressOfReturnAddress) const
+uint32 FBacktracer::GetBacktraceId(void* AddressOfReturnAddress) const
 {
 	FLookupState LookupState;
 	FCallstackTracer::FBacktraceEntry BacktraceEntry;
@@ -576,7 +576,7 @@ void* FBacktracer::GetBacktraceId(void* AddressOfReturnAddress) const
 	uint32 NumBacktrace = 0;
 #endif
 
-	uint64 BacktraceId = 0;
+	uint64 BacktraceHash = 0;
 	uint32 FrameIdx = 0;
 
 #if BACKTRACE_LOCK_FREE
@@ -593,8 +593,8 @@ void* FBacktracer::GetBacktraceId(void* AddressOfReturnAddress) const
 		BacktraceEntry.FrameCount = FrameIdx;
 
 		// This is a simple order-dependent LCG. Should be sufficient enough
-		BacktraceId += RetAddr;
-		BacktraceId *= 0x30be8efa499c249dull;
+		BacktraceHash += RetAddr;
+		BacktraceHash *= 0x30be8efa499c249dull;
 
 #if BACKTRACE_LOCK_FREE
 		int32 RspBias;
@@ -659,7 +659,7 @@ void* FBacktracer::GetBacktraceId(void* AddressOfReturnAddress) const
 	while (*StackPointer && FrameIdx < FCallstackTracer::FBacktraceEntry::MaxStackDepth);
 
 	// Save the collected id
-	BacktraceEntry.Id = BacktraceId;
+	BacktraceEntry.Hash = BacktraceHash;
 
 #if BACKTRACE_DBGLVL >= 3
 	for (uint32 i = 0; i < NumBacktrace; ++i)
@@ -680,8 +680,7 @@ void* FBacktracer::GetBacktraceId(void* AddressOfReturnAddress) const
 #endif
 	// Add to queue to be processed. This might block until there is room in the
 	// queue (i.e. the processing thread has caught up processing).
-	CallstackTracer->AddCallstack(BacktraceEntry);
-	return (void*)BacktraceId;
+	return CallstackTracer->AddCallstack(BacktraceEntry);
 }
 
 ////////////////////////////////////////////////////////////////////////////////
@@ -756,14 +755,14 @@ int32 Backtracer_GetFrameSize(void* FunctionAddress)
 }
 
 ////////////////////////////////////////////////////////////////////////////////
-void* Backtracer_GetBacktraceId(void* AddressOfReturnAddress)
+uint32 Backtracer_GetBacktraceId(void* AddressOfReturnAddress)
 {
 	if (FBacktracer* Instance = FBacktracer::Get())
 	{
 		return Instance->GetBacktraceId(AddressOfReturnAddress);
 	}
 
-	return nullptr;
+	return 0;
 }
 
 #endif // defined(PLATFORM_SUPPORTS_TRACE_WIN32_CALLSTACK)
diff --git a/Engine/Source/Runtime/Core/Private/ProfilingDebugging/TagTrace.cpp b/Engine/Source/Runtime/Core/Private/ProfilingDebugging/TagTrace.cpp
--- a/Engine/Source/Runtime/Core/Private/ProfilingDebugging/TagTrace.cpp
+++ b/Engine/Source/Runtime/Core/Private/ProfilingDebugging/TagTrace.cpp
@@ -4,7 +4,7 @@
 #include "Experimental/Containers/GrowOnlyLockFreeHash.h"
 #include "Containers/Set.h"
 #include "CoreTypes.h"
-#include "MemoryTrace.h"
+#include "ProfilingDebugging/MemoryTrace.h"
 #include "Misc/ScopeLock.h"
 #include "HAL/LowLevelMemTracker.h"
 #include "Trace/Trace.inl"
@@ -26,7 +26,7 @@ UE_TRACE_EVENT_BEGIN(Memory, MemoryScope)
 	UE_TRACE_EVENT_FIELD(int32, Tag)
 UE_TRACE_EVENT_END()
 
-UE_TRACE_EVENT_BEGIN(Memory, MemoryScopeRealloc)
+UE_TRACE_EVENT_BEGIN(Memory, MemoryScopePtr)
 	UE_TRACE_EVENT_FIELD(uint64, Ptr)
 UE_TRACE_EVENT_END()
 
@@ -96,13 +96,13 @@ FMemScope::~FMemScope()
 }
 
 ////////////////////////////////////////////////////////////////////////////////
-FMemScopeRealloc::FMemScopeRealloc(uint64 InPtr)
+FMemScopePtr::FMemScopePtr(uint64 InPtr)
 {
 	if (InPtr != 0 && TRACE_PRIVATE_CHANNELEXPR_IS_ENABLED(MemAllocChannel))
 	{
-		if (auto LogScope = FMemoryMemoryScopeReallocFields::LogScopeType::ScopedEnter<FMemoryMemoryScopeReallocFields>())
+		if (auto LogScope = FMemoryMemoryScopePtrFields::LogScopeType::ScopedEnter<FMemoryMemoryScopePtrFields>())
 		{
-			if (const auto& __restrict MemoryScope = *(FMemoryMemoryScopeReallocFields*)(&LogScope))
+			if (const auto& __restrict MemoryScope = *(FMemoryMemoryScopePtrFields*)(&LogScope))
 			{
 				Inner.SetActive(), LogScope += LogScope << MemoryScope.Ptr(InPtr);
 			}
@@ -110,13 +110,11 @@ FMemScopeRealloc::FMemScopeRealloc(uint64 InPtr)
 	}
 }
 
-////////////////////////////////////////////////////////////////////////////////
-FMemScopeRealloc::~FMemScopeRealloc()
+/////////////////////////////////////////////////////////////////////////////////
+FMemScopePtr::~FMemScopePtr()
 {
 }
 
-
-
 /////////////////////////////////////////////////////////////////////////////////
 
 /**
diff --git a/Engine/Source/Runtime/Core/Public/Async/AsyncWork.h b/Engine/Source/Runtime/Core/Public/Async/AsyncWork.h
--- a/Engine/Source/Runtime/Core/Public/Async/AsyncWork.h
+++ b/Engine/Source/Runtime/Core/Public/Async/AsyncWork.h
@@ -103,7 +103,7 @@ class FAutoDeleteAsyncTask
 	void DoWork()
 	{
 		LLM_SCOPE(InheritedLLMTag);
-		UE_MEMSCOPE(InheritedTraceTag, ELLMTracker::Default);
+		UE_MEMSCOPE(InheritedTraceTag);
 		FScopeCycleCounter Scope(Task.GetStatId(), true);
 
 		Task.DoWork();
@@ -297,7 +297,7 @@ class FAsyncTask
 	void DoWork()
 	{	
 		LLM_SCOPE(InheritedLLMTag);
-		UE_MEMSCOPE(InheritedTraceTag, ELLMTracker::Default);
+		UE_MEMSCOPE(InheritedTraceTag);
 		FScopeCycleCounter Scope(Task.GetStatId(), true); 
 
 		Task.DoWork();		
diff --git a/Engine/Source/Runtime/Core/Public/Async/TaskGraphInterfaces.h b/Engine/Source/Runtime/Core/Public/Async/TaskGraphInterfaces.h
--- a/Engine/Source/Runtime/Core/Public/Async/TaskGraphInterfaces.h
+++ b/Engine/Source/Runtime/Core/Public/Async/TaskGraphInterfaces.h
@@ -582,7 +582,7 @@ private:
 	FORCEINLINE void Execute(TArray<FBaseGraphTask*>& NewTasks, ENamedThreads::Type CurrentThread, bool bDeleteOnCompletion)
 	{
 		LLM_SCOPE(InheritedLLMTag);
-		UE_MEMSCOPE(InheritedTraceTag, ELLMTracker::Default);
+		UE_MEMSCOPE(InheritedTraceTag);
 		checkThreadGraph(LifeStage.Increment() == int32(LS_Executing));
 		ExecuteTask(NewTasks, CurrentThread, bDeleteOnCompletion);
 	}
diff --git a/Engine/Source/Runtime/Core/Public/HAL/LowLevelMemTracker.h b/Engine/Source/Runtime/Core/Public/HAL/LowLevelMemTracker.h
--- a/Engine/Source/Runtime/Core/Public/HAL/LowLevelMemTracker.h
+++ b/Engine/Source/Runtime/Core/Public/HAL/LowLevelMemTracker.h
@@ -287,12 +287,12 @@ extern FName LLMGetTagStat(ELLMTag Tag);
  * LLM scope macros
  */
 #define LLM_SCOPE(Tag)												FLLMScope SCOPE_NAME(Tag, false /* bIsStatTag */, ELLMTagSet::None, ELLMTracker::Default);\
-																	UE_MEMSCOPE(Tag, ELLMTracker::Default /* Tracker */);
+																	UE_MEMSCOPE(Tag) 
 #define LLM_SCOPE_BYNAME(Tag) 										static FName PREPROCESSOR_JOIN(LLMScope_Name,__LINE__)(Tag);\
 																	FLLMScope SCOPE_NAME(PREPROCESSOR_JOIN(LLMScope_Name,__LINE__), false /* bIsStatTag */, ELLMTagSet::None, ELLMTracker::Default);\
-																	UE_MEMSCOPE(PREPROCESSOR_JOIN(LLMScope_Name,__LINE__), ELLMTracker::Default /* Tracker */);
+																	UE_MEMSCOPE(PREPROCESSOR_JOIN(LLMScope_Name,__LINE__));
 #define LLM_SCOPE_BYTAG(TagDeclName)								FLLMScope SCOPE_NAME(PREPROCESSOR_JOIN(LLMTagDeclaration_, TagDeclName).GetUniqueName(), false /* bIsStatTag */, ELLMTagSet::None, ELLMTracker::Default);\
-																	UE_MEMSCOPE(PREPROCESSOR_JOIN(LLMTagDeclaration_,TagDeclName).GetUniqueName(), ELLMTracker::Default /* Tracker */);
+																	UE_MEMSCOPE(PREPROCESSOR_JOIN(LLMTagDeclaration_,TagDeclName).GetUniqueName());
 #define LLM_PLATFORM_SCOPE(Tag)										FLLMScope SCOPE_NAME(Tag, false /* bIsStatTag */, ELLMTagSet::None, ELLMTracker::Platform);
 #define LLM_PLATFORM_SCOPE_BYNAME(Tag) 								static FName PREPROCESSOR_JOIN(LLMScope_Name,__LINE__)(Tag);\
 																	FLLMScope SCOPE_NAME(PREPROCESSOR_JOIN(LLMLLMScope_NameScope,__LINE__), false /* bIsStatTag */, ELLMTagSet::None, ELLMTracker::Platform);
@@ -308,8 +308,7 @@ extern FName LLMGetTagStat(ELLMTag Tag);
 /**
  * LLM realloc scope macros. Used when reallocating a pointer and you wish to retain the tagging from the source pointer
  */
-#define LLM_REALLOC_SCOPE(Ptr)										FLLMScopeFromPtr SCOPE_NAME(Ptr, ELLMTracker::Default);\
-																	UE_MEMSCOPE_REALLOC(Ptr, ELLMTracker::Default);
+#define LLM_REALLOC_SCOPE(Ptr)										FLLMScopeFromPtr SCOPE_NAME(Ptr, ELLMTracker::Default);
 #define LLM_REALLOC_PLATFORM_SCOPE(Ptr)								FLLMScopeFromPtr SCOPE_NAME(Ptr, ELLMTracker::Platform);
 
 /**
diff --git a/Engine/Source/Runtime/Core/Public/ProfilingDebugging/MemoryTrace.h b/Engine/Source/Runtime/Core/Public/ProfilingDebugging/MemoryTrace.h
new file mode 100644
--- /dev/null
+++ b/Engine/Source/Runtime/Core/Public/ProfilingDebugging/MemoryTrace.h
@@ -0,0 +1,149 @@
+// Copyright Epic Games, Inc. All Rights Reserved.
+#pragma once
+
+#include "Trace/Trace.h"
+
+#if !defined(UE_MEMORY_TRACE_AVAILABLE)
+#   define UE_MEMORY_TRACE_AVAILABLE 0
+#endif
+
+#if !defined(UE_MEMORY_TRACE_LATE_INIT)
+#	define UE_MEMORY_TRACE_LATE_INIT 0
+#endif
+
+#if !defined(UE_MEMORY_TRACE_ENABLED) && UE_TRACE_ENABLED
+#	if UE_MEMORY_TRACE_AVAILABLE
+#		if !PLATFORM_USES_FIXED_GMalloc_CLASS && PLATFORM_64BITS
+#			define UE_MEMORY_TRACE_ENABLED !UE_BUILD_SHIPPING
+#		endif
+#	endif
+#endif
+
+#if !defined(UE_MEMORY_TRACE_ENABLED)
+#	define UE_MEMORY_TRACE_ENABLED 0
+#endif
+
+////////////////////////////////////////////////////////////////////////////////
+typedef uint32 HeapId;
+
+////////////////////////////////////////////////////////////////////////////////
+enum EMemoryTraceRootHeap : uint8
+{
+	SystemMemory, // System memory
+	VideoMemory, //Vram
+	EndHardcoded = VideoMemory,
+	EndReserved = 15
+};
+
+////////////////////////////////////////////////////////////////////////////////
+// These values are traced. Do not modify existing values in order to maintain
+// compatibility.
+enum class EMemoryTraceHeapFlags : uint16
+{
+	None = 0,
+	Root = 1 << 0,
+	NeverFrees = 1 << 1, // The heap doesn't free (e.g. linear allocator)
+};
+ENUM_CLASS_FLAGS(EMemoryTraceHeapFlags);
+
+////////////////////////////////////////////////////////////////////////////////
+// These values are traced. Do not modify existing values in order to maintain
+// compatibility.
+enum class EMemoryTraceHeapAllocationFlags : uint8
+{
+	None = 0,
+	Heap = 1 << 0, // Is a heap, can be used to unmark alloc as heap.
+};
+ENUM_CLASS_FLAGS(EMemoryTraceHeapAllocationFlags);
+
+#if UE_MEMORY_TRACE_ENABLED
+
+////////////////////////////////////////////////////////////////////////////////
+CORE_API UE_TRACE_CHANNEL_EXTERN(MemAllocChannel);
+
+////////////////////////////////////////////////////////////////////////////////
+class FMalloc* MemoryTrace_Create(class FMalloc* InMalloc);
+void MemoryTrace_Initialize();
+
+#if UE_MEMORY_TRACE_LATE_INIT
+void MemoryTrace_InitializeLate();
+#else
+static void MemoryTrace_InitializeLate() {}
+#endif // UE_MEMORY_TRACE_LATE_INIT
+
+/**
+ * Register a new heap specification (name). Use the returned value when marking heaps.
+ * @param ParentId Heap id of parent heap.
+ * @param Name Descriptive name of the heap.
+ * @param Flags Properties of this heap. See \ref EMemoryTraceHeapFlags
+ * @return Heap id to use when allocating memory
+ */
+CORE_API HeapId MemoryTrace_HeapSpec(HeapId ParentId, const TCHAR* Name, EMemoryTraceHeapFlags Flags = EMemoryTraceHeapFlags::None);
+
+/**
+ * Register a new root heap specification (name). Use the returned value as parent to other heaps.
+ * @param Name Descriptive name of the root heap.
+ * @param Flags Properties of the this root heap. See \ref EMemoryTraceHeapFlags
+ * @return Heap id to use when allocating memory
+ */
+CORE_API HeapId MemoryTrace_RootHeapSpec(const TCHAR* Name, EMemoryTraceHeapFlags Flags = EMemoryTraceHeapFlags::None);
+
+/**
+ * Mark a traced allocation as being a heap.
+ * @param Address Address of the allocation
+ * @param Heap Heap id, see /ref MemoryTrace_HeapSpec
+ * @param Flags Additional properties of the heap allocation. Note that \ref EMemoryTraceHeapAllocationFlags::Heap is implicit.
+ */
+CORE_API void MemoryTrace_MarkAllocAsHeap(uint64 Address, HeapId Heap, EMemoryTraceHeapAllocationFlags Flags = EMemoryTraceHeapAllocationFlags::None);
+
+/**
+ * Unmark an allocation as a heap. When an allocation that has previously been used as a heap is reused as a regular
+ * allocation.
+ * @param Address Address of the allocation
+ * @param Heap Heap id
+ */
+CORE_API void MemoryTrace_UnmarkAllocAsHeap(uint64 Address, HeapId Heap);
+
+/**
+ * Trace an allocation event.
+ * @param Address Address of allocation
+ * @param Size Size of allocation
+ * @param Alignment Alignment of the allocation
+ * @param RootHeap Which root heap this belongs to (system memory, video memory etc)
+ */
+CORE_API void MemoryTrace_Alloc(uint64 Address, uint64 Size, uint32 Alignment, HeapId RootHeap = EMemoryTraceRootHeap::SystemMemory);
+
+/**
+ * Trace a free event.
+ * @param Address Address of the allocation being freed
+ * @param RootHeap Which root heap this belongs to (system memory, video memory etc)
+ */
+CORE_API void MemoryTrace_Free(uint64 Address, HeapId RootHeap = EMemoryTraceRootHeap::SystemMemory);
+
+/**
+ * Trace a free related to a reallocation event.
+ * @param Address Address of the allocation being freed
+ * @param RootHeap Which root heap this belongs to (system memory, video memory etc)
+ */
+CORE_API void MemoryTrace_ReallocFree(uint64 Address, HeapId RootHeap = EMemoryTraceRootHeap::SystemMemory);
+
+/** Trace an allocation related to a reallocation event.
+ * @param Address Address of allocation
+ * @param NewSize Size of allocation
+ * @param Alignment Alignment of the allocation
+ * @param RootHeap Which root heap this belongs to (system memory, video memory etc)
+ */
+CORE_API void MemoryTrace_ReallocAlloc(uint64 Address, uint64 NewSize, uint32 Alignment, HeapId RootHeap = EMemoryTraceRootHeap::SystemMemory);
+
+#else
+
+inline HeapId MemoryTrace_HeapSpec(HeapId ParentId, const TCHAR* Name) { return ~0; }
+inline void MemoryTrace_MarkAllocAsHeap(uint64 Address, HeapId Heap) {}
+inline void MemoryTrace_UnmarkAllocAsHeap(uint64 Address, HeapId Heap) {}
+inline void MemoryTrace_Alloc(uint64 Address, uint64 Size, uint32 Alignment, HeapId RootHeap = 0) {}
+inline void MemoryTrace_Free(uint64 Address, HeapId RootHeap = 0) {}
+inline void MemoryTrace_ReallocFree(uint64 Address, HeapId RootHeap = 0) {}
+inline void MemoryTrace_ReallocAlloc(uint64 Address, uint64 NewSize, uint32 Alignment, HeapId RootHeap = 0) {}
+
+
+#endif // UE_MEMORY_TRACE_ENABLED
diff --git a/Engine/Source/Runtime/Core/Public/ProfilingDebugging/TagTrace.h b/Engine/Source/Runtime/Core/Public/ProfilingDebugging/TagTrace.h
--- a/Engine/Source/Runtime/Core/Public/ProfilingDebugging/TagTrace.h
+++ b/Engine/Source/Runtime/Core/Public/ProfilingDebugging/TagTrace.h
@@ -21,8 +21,8 @@ CORE_API int32	MemoryTrace_GetActiveTag();
 
 ////////////////////////////////////////////////////////////////////////////////
 
-#if !defined(USE_MEMORY_TRACE_TAGS) && UE_TRACE_ENABLED && UE_BUILD_DEVELOPMENT
-	#if PLATFORM_WINDOWS
+#if !defined(USE_MEMORY_TRACE_TAGS) && UE_TRACE_ENABLED && !UE_BUILD_SHIPPING
+	#if PLATFORM_WINDOWS || PLATFORM_PS4 || defined(__PS5__)
 		#define USE_MEMORY_TRACE_TAGS 1
 	#endif
 #endif
@@ -62,25 +62,25 @@ private:
 /**
  * Used order to keep the tag for memory that is being reallocated.
  */
-class FMemScopeRealloc
+class FMemScopePtr
 {
 public:
-	CORE_API FMemScopeRealloc(uint64 InPtr);
-	CORE_API ~FMemScopeRealloc();
+	CORE_API FMemScopePtr(uint64 InPtr);
+	CORE_API ~FMemScopePtr();
 private:
 	UE::Trace::Private::FScopedLogScope Inner;
 };
 
 ////////////////////////////////////////////////////////////////////////////////
 
-#define UE_MEMSCOPE(InTag, InTracker)				FMemScope PREPROCESSOR_JOIN(MemScope,__LINE__)(InTag);
-#define UE_MEMSCOPE_REALLOC(InPtr, InTracker)		FMemScopeRealloc PREPROCESSOR_JOIN(MemReallocScope,__LINE__)((uint64)InPtr);
+#define UE_MEMSCOPE(InTag)				FMemScope PREPROCESSOR_JOIN(MemScope,__LINE__)(InTag);
+#define UE_MEMSCOPE_PTR(InPtr)			FMemScopePtr PREPROCESSOR_JOIN(MemPtrScope,__LINE__)((uint64)InPtr);
 
 #else // USE_MEMORY_TRACE_TAGS
 
 ////////////////////////////////////////////////////////////////////////////////
 #define UE_MEMSCOPE(...)
-#define UE_MEMSCOPE_REALLOC(...)
+#define UE_MEMSCOPE_PTR(...)
 
 #endif // USE_MEMORY_TRACE_TAGS
 
diff --git a/Engine/Source/Runtime/Engine/Private/Level.cpp b/Engine/Source/Runtime/Engine/Private/Level.cpp
--- a/Engine/Source/Runtime/Engine/Private/Level.cpp
+++ b/Engine/Source/Runtime/Engine/Private/Level.cpp
@@ -76,6 +76,7 @@ Level.cpp: Level-related functions
 #include "Algo/Copy.h"
 #include "HAL/LowLevelMemTracker.h"
 #include "ObjectTrace.h"
+#include "ProfilingDebugging/TagTrace.h"
 
 #define LOCTEXT_NAMESPACE "ULevel"
 DEFINE_LOG_CATEGORY(LogLevel);
@@ -717,6 +718,7 @@ void ULevel::SortActorList()
 		return;
 	}
 	LLM_REALLOC_SCOPE(Actors.GetData());
+	UE_MEMSCOPE_PTR(Actors.GetData());
 
 	TArray<AActor*> NewActors;
 	TArray<AActor*> NewNetActors;
