Nounity compile fix

#rnx

diff --git a/Engine/Source/Runtime/Core/Public/ProfilingDebugging/MemoryTrace.h b/Engine/Source/Runtime/Core/Public/ProfilingDebugging/MemoryTrace.h
--- a/Engine/Source/Runtime/Core/Public/ProfilingDebugging/MemoryTrace.h
+++ b/Engine/Source/Runtime/Core/Public/ProfilingDebugging/MemoryTrace.h
@@ -2,6 +2,7 @@
 #pragma once
 
 #include "Trace/Trace.h"
+#include "Misc/EnumClassFlags.h"
 
 #if !defined(UE_MEMORY_TRACE_AVAILABLE)
 #   define UE_MEMORY_TRACE_AVAILABLE 0
