SharedBuffer: Use if constexpr in TBufferOwnerPtr

#rnx

diff --git a/Engine/Source/Runtime/Core/Public/Memory/SharedBuffer.h b/Engine/Source/Runtime/Core/Public/Memory/SharedBuffer.h
--- a/Engine/Source/Runtime/Core/Public/Memory/SharedBuffer.h
+++ b/Engine/Source/Runtime/Core/Public/Memory/SharedBuffer.h
@@ -731,7 +731,7 @@ inline FBufferOwner* TBufferOwnerPtr<FOps>::CopyFrom(const TBufferOwnerPtr<FOthe
 	FBufferOwner* NewOwner = Ptr.Owner;
 	if (NewOwner)
 	{
-		if (bIsWeak || !Ptr.bIsWeak)
+		if constexpr (bIsWeak || !TBufferOwnerPtr<FOtherOps>::bIsWeak)
 		{
 			FOps::AddRef(*NewOwner);
 		}
@@ -748,13 +748,13 @@ template <typename FOtherOps>
 inline FBufferOwner* TBufferOwnerPtr<FOps>::MoveFrom(TBufferOwnerPtr<FOtherOps>&& Ptr)
 {
 	FBufferOwner* NewOwner = Ptr.Owner;
-	if (bIsWeak == Ptr.bIsWeak)
+	if constexpr (bIsWeak == TBufferOwnerPtr<FOtherOps>::bIsWeak)
 	{
 		Ptr.Owner = nullptr;
 	}
 	else if (NewOwner)
 	{
-		if (bIsWeak)
+		if constexpr (bIsWeak)
 		{
 			FOps::AddRef(*NewOwner);
 		}
