Memory trace instrumentation for generic allocators.

These allocators were missing from memory tracing. Now that there is a public interface (see changelist "Memory tracing additional heaps") we can instrument them.

diff --git a/Engine/Source/Runtime/Core/Private/GenericPlatform/GenericPlatformMallocCrash.cpp b/Engine/Source/Runtime/Core/Private/GenericPlatform/GenericPlatformMallocCrash.cpp
--- a/Engine/Source/Runtime/Core/Private/GenericPlatform/GenericPlatformMallocCrash.cpp
+++ b/Engine/Source/Runtime/Core/Private/GenericPlatform/GenericPlatformMallocCrash.cpp
@@ -5,6 +5,7 @@
 #include "HAL/PlatformTLS.h"
 #include "Templates/AlignmentTemplates.h"
 #include "HAL/LowLevelMemTracker.h"
+#include "ProfilingDebugging/MemoryTrace.h"
 
 /** Describes a pool. */
 struct FPoolDesc
@@ -216,6 +217,9 @@ FGenericPlatformMallocCrash::FGenericPlatformMallocCrash( FMalloc* MainMalloc )
 	LLM(FLowLevelMemTracker::Get().OnLowLevelAlloc(ELLMTracker::Default, LargeMemoryPool, LargeMemoryPoolSize));
 	LLM(FLowLevelMemTracker::Get().OnLowLevelAlloc(ELLMTracker::Default, SmallMemoryPool, GetSmallPoolTotalSize()));
 
+	MemoryTrace_Alloc((uint64)LargeMemoryPool, LargeMemoryPoolSize, alignof(uint8), EMemoryTraceRootHeap::SystemMemory);
+	MemoryTrace_Alloc((uint64)SmallMemoryPool, GetSmallPoolTotalSize(), alignof(uint8), EMemoryTraceRootHeap::SystemMemory);
+
 	if( !SmallMemoryPool || !LargeMemoryPool )
 	{
 		FPlatformMisc::LowLevelOutputDebugString( TEXT( "Memory pools allocations failed, exiting...\n" ) );
diff --git a/Engine/Source/Runtime/Core/Private/GenericPlatform/GenericPlatformMemory.cpp b/Engine/Source/Runtime/Core/Private/GenericPlatform/GenericPlatformMemory.cpp
--- a/Engine/Source/Runtime/Core/Private/GenericPlatform/GenericPlatformMemory.cpp
+++ b/Engine/Source/Runtime/Core/Private/GenericPlatform/GenericPlatformMemory.cpp
@@ -21,6 +21,7 @@
 #include "Misc/MessageDialog.h"
 #include "Templates/UnrealTemplate.h"
 #include "HAL/IConsoleManager.h"
+#include "ProfilingDebugging/MemoryTrace.h"
 
 DEFINE_STAT(MCR_Physical);
 DEFINE_STAT(MCR_PhysicalLLM);
@@ -134,6 +135,7 @@ void FGenericPlatformMemory::SetupMemoryPools()
 		BackupOOMMemoryPool = FPlatformMemory::BinnedAllocFromOS(FPlatformMemory::GetBackMemoryPoolSize());
 
 		LLM(FLowLevelMemTracker::Get().OnLowLevelAlloc(ELLMTracker::Default, BackupOOMMemoryPool, FPlatformMemory::GetBackMemoryPoolSize()));
+		MemoryTrace_Alloc((uint64)BackupOOMMemoryPool, FPlatformMemory::GetBackMemoryPoolSize(), alignof(void*), EMemoryTraceRootHeap::SystemMemory);
 	}
 }
 
@@ -172,6 +174,7 @@ void FGenericPlatformMemory::OnOutOfMemory(uint64 Size, uint32 Alignment)
 			UE_LOG(LogMemory, Warning, TEXT("Freeing %d bytes from backup pool to handle out of memory."), FPlatformMemory::GetBackMemoryPoolSize());
         
 			LLM(FLowLevelMemTracker::Get().OnLowLevelFree(ELLMTracker::Default, BackupOOMMemoryPool));
+			MemoryTrace_Free((uint64)BackupOOMMemoryPool, EMemoryTraceRootHeap::SystemMemory);
 		}
 
 		UE_LOG(LogMemory, Warning, TEXT("MemoryStats:")\
diff --git a/Engine/Source/Runtime/Core/Public/Misc/GenericGrowableAllocator.h b/Engine/Source/Runtime/Core/Public/Misc/GenericGrowableAllocator.h
--- a/Engine/Source/Runtime/Core/Public/Misc/GenericGrowableAllocator.h
+++ b/Engine/Source/Runtime/Core/Public/Misc/GenericGrowableAllocator.h
@@ -9,6 +9,7 @@
 #include "CoreMinimal.h"
 #include "HAL/LowLevelMemTracker.h"
 #include "Misc/ScopeLock.h"
+#include "ProfilingDebugging/MemoryTrace.h"
 #include "Stats/Stats.h"
 
 //////////////
@@ -223,6 +224,7 @@ public:
 				Allocation->OwnerType = OwnerType;
 #endif
 				LLM(FLowLevelMemTracker::Get().OnLowLevelAlloc(ELLMTracker::Default, GetAddressForTracking(Offset), Size));
+				MemoryTrace_Alloc(uint64(GetAddressForTracking(Offset)), Size, Alignment, EMemoryTraceRootHeap::SystemMemory);
 
 				// let the implementation fill in any more
 				InitializeAllocationStruct(Allocation);
@@ -247,6 +249,7 @@ public:
 		uint32 Offset = Memory->Offset;
 
 		LLM(FLowLevelMemTracker::Get().OnLowLevelFree(ELLMTracker::Default, GetAddressForTracking(Offset)));
+		MemoryTrace_Free(uint64(GetAddressForTracking(Offset)), EMemoryTraceRootHeap::SystemMemory);
 
 		// we are now done with the Allocation object
 		DestroyAllocationStruct(Memory);
