Experimental UE::FLazyEvent that allocates an event only when necessary

diff --git a/Engine/Source/Runtime/Core/Private/Async/LazyEvent.cpp b/Engine/Source/Runtime/Core/Private/Async/LazyEvent.cpp
new file mode 100644
--- /dev/null
+++ b/Engine/Source/Runtime/Core/Private/Async/LazyEvent.cpp
@@ -0,0 +1,161 @@
+// Copyright Epic Games, Inc. All Rights Reserved.
+
+#include "Experimental/Async/LazyEvent.h"
+#include "HAL/PlatformProcess.h"
+
+namespace UE::Core::Private
+{
+
+enum class ELazyEventType : UPTRINT
+{
+	NotTriggeredAutoReset = 0,
+	NotTriggeredManualReset = 1,
+	TriggeredAutoReset = 2,
+	TriggeredManualReset = 3,
+};
+
+static FEvent* GetLazyEvent(ELazyEventType EventType)
+{
+	return reinterpret_cast<FEvent*>(static_cast<UPTRINT>(EventType));
+}
+
+static ELazyEventType GetLazyEventType(FEvent* Event)
+{
+	return static_cast<ELazyEventType>(reinterpret_cast<UPTRINT>(Event));
+}
+
+} // UE::Core::Private
+
+namespace UE
+{
+
+FLazyEvent::FLazyEvent(EEventMode EventMode)
+	: AtomicEvent(Core::Private::GetLazyEvent(EventMode == EEventMode::AutoReset
+		? Core::Private::ELazyEventType::NotTriggeredAutoReset
+		: Core::Private::ELazyEventType::NotTriggeredManualReset))
+{
+}
+
+FLazyEvent::~FLazyEvent()
+{
+	using namespace Core::Private;
+	FEvent* Event = AtomicEvent.load(std::memory_order_acquire);
+	switch (const ELazyEventType EventType = GetLazyEventType(Event))
+	{
+		case ELazyEventType::NotTriggeredAutoReset:
+		case ELazyEventType::NotTriggeredManualReset:
+		case ELazyEventType::TriggeredAutoReset:
+		case ELazyEventType::TriggeredManualReset:
+			return;
+		default:
+			FPlatformProcess::ReturnSynchEventToPool(Event);
+			return;
+	}
+}
+
+void FLazyEvent::Trigger()
+{
+	using namespace Core::Private;
+	FEvent* Event = AtomicEvent.load(std::memory_order_acquire);
+	for (;;)
+	{
+		switch (const ELazyEventType EventType = GetLazyEventType(Event))
+		{
+		case ELazyEventType::NotTriggeredAutoReset:
+			if (AtomicEvent.compare_exchange_strong(Event, GetLazyEvent(ELazyEventType::TriggeredAutoReset),
+				std::memory_order_relaxed, std::memory_order_acquire))
+			{
+				return;
+			}
+			break;
+		case ELazyEventType::NotTriggeredManualReset:
+			if (AtomicEvent.compare_exchange_strong(Event, GetLazyEvent(ELazyEventType::TriggeredManualReset),
+				std::memory_order_relaxed, std::memory_order_acquire))
+			{
+				return;
+			}
+			break;
+		case ELazyEventType::TriggeredAutoReset:
+		case ELazyEventType::TriggeredManualReset:
+			return;
+		default:
+			return Event->Trigger();
+		}
+	}
+}
+
+void FLazyEvent::Reset()
+{
+	using namespace Core::Private;
+	FEvent* Event = AtomicEvent.load(std::memory_order_acquire);
+	for (;;)
+	{
+		switch (const ELazyEventType EventType = GetLazyEventType(Event))
+		{
+		case ELazyEventType::NotTriggeredAutoReset:
+		case ELazyEventType::NotTriggeredManualReset:
+			return;
+		case ELazyEventType::TriggeredAutoReset:
+			if (AtomicEvent.compare_exchange_strong(Event, GetLazyEvent(ELazyEventType::NotTriggeredAutoReset),
+				std::memory_order_relaxed, std::memory_order_acquire))
+			{
+				return;
+			}
+			break;
+		case ELazyEventType::TriggeredManualReset:
+			if (AtomicEvent.compare_exchange_strong(Event, GetLazyEvent(ELazyEventType::NotTriggeredManualReset),
+				std::memory_order_relaxed, std::memory_order_acquire))
+			{
+				return;
+			}
+			break;
+		default:
+			return Event->Reset();
+		}
+	}
+}
+
+void FLazyEvent::Wait()
+{
+	Wait(MAX_uint32);
+}
+
+bool FLazyEvent::Wait(uint32 WaitTime)
+{
+	using namespace Core::Private;
+	FEvent* Event = AtomicEvent.load(std::memory_order_acquire);
+	for (;;)
+	{
+		switch (const ELazyEventType EventType = GetLazyEventType(Event))
+		{
+		case ELazyEventType::NotTriggeredAutoReset:
+		case ELazyEventType::NotTriggeredManualReset:
+		{
+			if (WaitTime == 0)
+			{
+				return false;
+			}
+			FEvent* NewEvent = FPlatformProcess::GetSynchEventFromPool(EventType == ELazyEventType::NotTriggeredManualReset);
+			if (!AtomicEvent.compare_exchange_strong(Event, NewEvent,
+				std::memory_order_release, std::memory_order_acquire))
+			{
+				FPlatformProcess::ReturnSynchEventToPool(NewEvent);
+			}
+			break;
+		}
+		case ELazyEventType::TriggeredAutoReset:
+			if (AtomicEvent.compare_exchange_strong(Event, GetLazyEvent(ELazyEventType::NotTriggeredAutoReset),
+				std::memory_order_relaxed, std::memory_order_acquire))
+			{
+				return true;
+			}
+			break;
+		case ELazyEventType::TriggeredManualReset:
+			return true;
+		default:
+			return Event->Wait(WaitTime);
+		}
+	}
+}
+
+} // UE
diff --git a/Engine/Source/Runtime/Core/Public/Experimental/Async/LazyEvent.h b/Engine/Source/Runtime/Core/Public/Experimental/Async/LazyEvent.h
new file mode 100644
--- /dev/null
+++ b/Engine/Source/Runtime/Core/Public/Experimental/Async/LazyEvent.h
@@ -0,0 +1,30 @@
+// Copyright Epic Games, Inc. All Rights Reserved.
+
+#pragma once
+
+#include "CoreTypes.h"
+#include "HAL/Event.h"
+#include <atomic>
+
+namespace UE
+{
+
+class FLazyEvent
+{
+public:
+	FLazyEvent(const FLazyEvent&) = delete;
+	FLazyEvent& operator=(const FLazyEvent&) = delete;
+
+	CORE_API explicit FLazyEvent(EEventMode EventMode);
+	CORE_API ~FLazyEvent();
+
+	CORE_API void Trigger();
+	CORE_API void Reset();
+	CORE_API void Wait();
+	CORE_API bool Wait(uint32 WaitTime);
+
+private:
+	std::atomic<FEvent*> AtomicEvent;
+};
+
+} // UE
