Fixed static analysis warnings.

#rnx

diff --git a/Engine/Source/Runtime/Core/Private/ProfilingDebugging/MemoryAllocationTrace.cpp b/Engine/Source/Runtime/Core/Private/ProfilingDebugging/MemoryAllocationTrace.cpp
--- a/Engine/Source/Runtime/Core/Private/ProfilingDebugging/MemoryAllocationTrace.cpp
+++ b/Engine/Source/Runtime/Core/Private/ProfilingDebugging/MemoryAllocationTrace.cpp
@@ -147,8 +147,8 @@ void FAllocationTrace::Update()
 void FAllocationTrace::Alloc(void* Address, size_t Size, uint32 Alignment, uint32 Owner, HeapId RootHeap)
 {
 	check(RootHeap < 16);
-	const uint8 AlignmentPow2 = uint8(FPlatformMath::CountTrailingZeros(Alignment));
-	const uint32 Alignment_SizeLower = AlignmentPow2 << SizeShift | (Size & ((1 << SizeShift) - 1));
+	const uint32 AlignmentPow2 = uint32(FPlatformMath::CountTrailingZeros(Alignment));
+	const uint32 Alignment_SizeLower = (AlignmentPow2 << SizeShift) | uint32(Size & ((1 << SizeShift) - 1));
 
 	switch (RootHeap)
 	{
@@ -217,8 +217,8 @@ void FAllocationTrace::Free(void* Address, HeapId RootHeap)
 void FAllocationTrace::ReallocAlloc(void* Address, size_t Size, uint32 Alignment, uint32 Owner, HeapId RootHeap)
 {
 	check(RootHeap < 16);
-	const uint8 AlignmentPow2 = uint8(FPlatformMath::CountTrailingZeros(Alignment));
-	const uint32 Alignment_SizeLower = AlignmentPow2 << SizeShift | (Size & ((1 << SizeShift) - 1));
+	const uint32 AlignmentPow2 = uint32(FPlatformMath::CountTrailingZeros(Alignment));
+	const uint32 Alignment_SizeLower = (AlignmentPow2 << SizeShift) | uint32(Size & ((1 << SizeShift) - 1));
 	
 	switch (RootHeap)
 	{
