Restored static analysis warning fixes that got lost when merging
#rnx

diff --git a/Engine/Source/Runtime/Core/Private/FramePro/FramePro.cpp b/Engine/Source/Runtime/Core/Private/FramePro/FramePro.cpp
--- a/Engine/Source/Runtime/Core/Private/FramePro/FramePro.cpp
+++ b/Engine/Source/Runtime/Core/Private/FramePro/FramePro.cpp
@@ -1435,7 +1435,7 @@ namespace FramePro
 			default:
 			{
 				char temp[128];
-				sprintf_s(temp, "Error code: %ld", error_code);
+				sprintf_s(temp, "Error code: %lu", error_code);
 				error_string = temp;
 			}
 		}
@@ -1494,6 +1494,11 @@ namespace FramePro
 
 		// start the processing thread
 		HANDLE thread = CreateThread(0, 0, TracingThread_Static, this, 0, NULL);
+		if (thread == NULL)
+		{
+			error = "CreateThread returned NULL";
+			return false;
+		}
 		CloseHandle(thread);
 
 		return true;
@@ -9105,21 +9110,23 @@ namespace FramePro
 
 					int buffer_size = 1024;
 					FRAMEPRO_TCHAR* p_buffer = (FRAMEPRO_TCHAR*)HeapAlloc(GetProcessHeap(), 0, buffer_size * sizeof(FRAMEPRO_TCHAR));
-					memset(p_buffer, 0, buffer_size * sizeof(FRAMEPRO_TCHAR));
+					if (p_buffer)
+					{
+						memset(p_buffer, 0, buffer_size * sizeof(FRAMEPRO_TCHAR));
 
-					va_list args;
-					FormatMessage(
-						FORMAT_MESSAGE_FROM_SYSTEM,
-						NULL,
-						WSAGetLastError(),
-						0,
-						p_buffer,
-						buffer_size,
-						&args);
+						FormatMessage(
+							FORMAT_MESSAGE_FROM_SYSTEM,
+							NULL,
+							WSAGetLastError(),
+							0,
+							p_buffer,
+							buffer_size,
+							NULL);
 
-					FramePro::DebugWrite("FramePro Network Error: %ws\n", p_buffer);
+						FramePro::DebugWrite("FramePro Network Error: %ws\n", p_buffer);
 
-					HeapFree(GetProcessHeap(), 0, p_buffer);
+						HeapFree(GetProcessHeap(), 0, p_buffer);
+					}
 				#endif
 			}
 
@@ -9254,7 +9261,7 @@ namespace FramePro
 
 					// Resolve the server address and port
 					addrinfo* p_result_info;
-					HRESULT result = getaddrinfo(NULL, p_port, &info, &p_result_info);
+					int result = getaddrinfo(NULL, p_port, &info, &p_result_info);
 					if (result != 0)
 					{
 						HandleSocketError();
