CbObject: Use AsObjectId when serializing object ID's to JSON

#rnx

diff --git a/Engine/Source/Runtime/Core/Private/Serialization/CompactBinarySerialization.cpp b/Engine/Source/Runtime/Core/Private/Serialization/CompactBinarySerialization.cpp
--- a/Engine/Source/Runtime/Core/Private/Serialization/CompactBinarySerialization.cpp
+++ b/Engine/Source/Runtime/Core/Private/Serialization/CompactBinarySerialization.cpp
@@ -373,7 +373,7 @@ public:
 			break;
 		}
 		case ECbFieldType::ObjectId:
-			Builder << '"' << Accessor.AsUuid() << '"';
+			Builder << '"' << Accessor.AsObjectId() << '"';
 			break;
 		case ECbFieldType::CustomById:
 		{
