fixed reserve workers stalling on exit

diff --git a/Engine/Source/Runtime/Core/Private/Async/Fundamental/ReserveScheduler.cpp b/Engine/Source/Runtime/Core/Private/Async/Fundamental/ReserveScheduler.cpp
--- a/Engine/Source/Runtime/Core/Private/Async/Fundamental/ReserveScheduler.cpp
+++ b/Engine/Source/Runtime/Core/Private/Async/Fundamental/ReserveScheduler.cpp
@@ -10,30 +10,32 @@ FReserveScheduler FReserveScheduler::Singleton;
 
 TUniquePtr<FThread> FReserveScheduler::CreateWorker(bool bIsForkable, FSchedulerTls::FLocalQueueType* WorkerLocalQueue, EThreadPriority Priority)
 {
+	FYieldedWork* ReserveEvent = new FYieldedWork;
+	ReserveEvents.Emplace(ReserveEvent);
+
 	uint32 WorkerId = NextWorkerId++;
 	return MakeUnique<FThread>
 	(
 		*FString::Printf(TEXT("Reserve Worker #%d"), WorkerId),
-		[this, WorkerLocalQueue]
+		[this, WorkerLocalQueue, ReserveEvent]
 		{
 			FSchedulerTls::ActiveScheduler = this;
 			FSchedulerTls::LocalQueue = WorkerLocalQueue;
 
-			FYieldedWork ReserveEvent;
 			while (true)
 			{
-				EventStack.Push(&ReserveEvent);
-				ReserveEvent.SleepEvent->Wait();
+				EventStack.Push(ReserveEvent);
+				ReserveEvent->SleepEvent->Wait();
 				{
 					TRACE_CPUPROFILER_EVENT_SCOPE(FReserveScheduler::BusyWaitUntil);
-					FSchedulerTls::WorkerType = ReserveEvent.bPermitBackgroundWork ? FSchedulerTls::EWorkerType::Background : FSchedulerTls::EWorkerType::Foreground;
+					FSchedulerTls::WorkerType = ReserveEvent->bPermitBackgroundWork ? FSchedulerTls::EWorkerType::Background : FSchedulerTls::EWorkerType::Foreground;
 
 					if (ActiveWorkers.load(std::memory_order_relaxed) == 0)
 					{
 						break;
 					}
 
-					BusyWaitUntil(MoveTemp(ReserveEvent.CompletedDelegate), ReserveEvent.bPermitBackgroundWork);
+					BusyWaitUntil(MoveTemp(ReserveEvent->CompletedDelegate), ReserveEvent->bPermitBackgroundWork);
 				}
 			}
 			FSchedulerTls::WorkerType = EWorkerType::None;
@@ -88,11 +90,16 @@ void FReserveScheduler::StopWorkers()
 	if(OldActiveWorkers != 0 && ActiveWorkers.compare_exchange_strong(OldActiveWorkers, 0, std::memory_order_relaxed))
 	{
 		FScopeLock Lock(&WorkerThreadsCS);
-		while (FYieldedWork* Event = EventStack.Pop()) 
+
+		for (TUniquePtr<FYieldedWork>& Event : ReserveEvents)
 		{
 			Event->SleepEvent->Trigger();
 		}
 
+		while (FYieldedWork* Event = EventStack.Pop())
+		{
+		}
+
 		for (TUniquePtr<FThread>& Thread : WorkerThreads)
 		{
 			Thread->Join();
@@ -100,6 +107,7 @@ void FReserveScheduler::StopWorkers()
 		NextWorkerId = 0;
 		WorkerThreads.Reset();
 		WorkerLocalQueues.Reset();
+		ReserveEvents.Reset();
 	}
 }
 
diff --git a/Engine/Source/Runtime/Core/Public/Async/Fundamental/ReserveScheduler.h b/Engine/Source/Runtime/Core/Public/Async/Fundamental/ReserveScheduler.h
--- a/Engine/Source/Runtime/Core/Public/Async/Fundamental/ReserveScheduler.h
+++ b/Engine/Source/Runtime/Core/Public/Async/Fundamental/ReserveScheduler.h
@@ -41,6 +41,7 @@ namespace LowLevelTasks
 
 	private:
 		TEventStack<FYieldedWork> 				EventStack;
+		TArray<TUniquePtr<FYieldedWork>> ReserveEvents;
 		FCriticalSection 						WorkerThreadsCS;
 		TArray<FSchedulerTls::FLocalQueueType>	WorkerLocalQueues;
 		TArray<TUniquePtr<FThread>>				WorkerThreads;
