UE5:  Bug fix for infinite recursion when call stack tracing code is reentered due to an allocation within the code itself when memory tracking is enabled.  Added flag to block further stack tracing when this case is detected.  This wasn't observed in Development builds, because the reserve size was large enough that reallocation wasn't triggered, but in Test builds, additional functions are generated by the compiler.

#rnx

diff --git a/Engine/Source/Runtime/Core/Private/ProfilingDebugging/Microsoft/WindowsCallstackTrace.cpp b/Engine/Source/Runtime/Core/Private/ProfilingDebugging/Microsoft/WindowsCallstackTrace.cpp
--- a/Engine/Source/Runtime/Core/Private/ProfilingDebugging/Microsoft/WindowsCallstackTrace.cpp
+++ b/Engine/Source/Runtime/Core/Private/ProfilingDebugging/Microsoft/WindowsCallstackTrace.cpp
@@ -191,6 +191,7 @@ private:
 	FCallstackTracer*		CallstackTracer;
 #if BACKTRACE_LOCK_FREE
 	mutable FFunctionLookupSet	FunctionLookups;
+	mutable bool				bReentranceCheck = false;
 #endif
 #if BACKTRACE_DBGLVL >= 1
 	mutable uint32			NumFpTruncations = 0;
@@ -616,6 +617,16 @@ uint32 FBacktracer::GetBacktraceId(void* AddressOfReturnAddress) const
 		{
 			Lock.Lock();
 			Locked = true;
+
+			// If FunctionLookups.Emplace triggers a reallocation, it can cause an infinite recursion
+			// when the allocation reenters the stack trace code.  We need to break out of the recursion
+			// in that case, and let the allocation complete, with the assumption that we don't care
+			// about call stacks for internal allocations in the memory reporting system.  The "Lock()"
+			// above will only fall through with this flag set if it's a second lock in the same thread.
+			if (bReentranceCheck)
+			{
+				break;
+			}
 		}
 #endif  // BACKTRACE_LOCK_FREE
 
@@ -623,13 +634,17 @@ uint32 FBacktracer::GetBacktraceId(void* AddressOfReturnAddress) const
 		if (Function == nullptr)
 		{
 #if BACKTRACE_LOCK_FREE
+			TGuardValue<bool> ReentranceGuard(bReentranceCheck, true);
 			FunctionLookups.Emplace(RetAddr, -1);
 #endif
 			break;
 		}
 
 #if BACKTRACE_LOCK_FREE
-		FunctionLookups.Emplace(RetAddr, Function->RspBias);
+		{
+			TGuardValue<bool> ReentranceGuard(bReentranceCheck, true);
+			FunctionLookups.Emplace(RetAddr, Function->RspBias);
+		}
 #endif
 
 #if BACKTRACE_DBGLVL >= 2
