CoreGlobals.h: aggregate-like initialization constructor for FCoreTexts

In C++20, a class that defines (or deletes) any constructor may not be aggregate-initialized.
http://www.open-std.org/jtc1/sc22/wg21/docs/papers/2018/p1008r1.pdf

This change allows FCoreTextsSingleton to perform aggregate-like initialization of FCoreTexts FCoreTextsSingleton::Texts.

#trivial

diff --git a/Engine/Source/Runtime/Core/Public/CoreGlobals.h b/Engine/Source/Runtime/Core/Public/CoreGlobals.h
--- a/Engine/Source/Runtime/Core/Public/CoreGlobals.h
+++ b/Engine/Source/Runtime/Core/Public/CoreGlobals.h
@@ -106,6 +106,11 @@ struct CORE_API FCoreTexts
 	// Non-copyable
 	FCoreTexts(const FCoreTexts&) = delete;
 	FCoreTexts& operator=(const FCoreTexts&) = delete;
+
+	FCoreTexts(const FText& InTrue, const FText& InFalse, const FText& InYes, const FText& InNo, const FText& InNone)
+		: True(InTrue), False(InFalse), Yes(InYes), No(InNo), None(InNone)
+	{
+	}
 };
 
 #if !defined(DISABLE_LEGACY_CORE_TEXTS) || DISABLE_LEGACY_CORE_TEXTS == 0
