Fix clear page not clearing the currently viewed page in MessageLog

diff --git a/Engine/Source/Developer/MessageLog/Private/Model/MessageLogListingModel.cpp b/Engine/Source/Developer/MessageLog/Private/Model/MessageLogListingModel.cpp
--- a/Engine/Source/Developer/MessageLog/Private/Model/MessageLogListingModel.cpp
+++ b/Engine/Source/Developer/MessageLog/Private/Model/MessageLogListingModel.cpp
@@ -118,7 +118,7 @@ void FMessageLogListingModel::NewPage( const FText& InTitle, uint32 InMaxPages )
 	}
 }
 
-bool FMessageLogListingModel::SwitchToPage(const FText& InTitle, uint32 InMaxPages )
+bool FMessageLogListingModel::SetCurrentPage(const FText& InTitle, uint32 InMaxPages )
 {
 	if (CurrentPage().Title.CompareTo(InTitle) == 0)
 	{
@@ -145,6 +145,22 @@ bool FMessageLogListingModel::SwitchToPage(const FText& InTitle, uint32 InMaxPag
 	return true;
 }
 
+bool FMessageLogListingModel::SetCurrentPage( const uint32 InOldPageIndex  )
+{
+	if (InOldPageIndex == 0)
+	{
+		return false;
+	}
+
+	if (InOldPageIndex >= NumPages())
+	{
+		return false;
+	}
+
+	FPage* PageToSwitch = PageAtIndex(InOldPageIndex);
+	return SetCurrentPage(PageToSwitch->Title, MaxPages);
+}
+
 uint32 FMessageLogListingModel::NumPages() const
 {
 	return Pages.Num();
diff --git a/Engine/Source/Developer/MessageLog/Private/Model/MessageLogListingModel.h b/Engine/Source/Developer/MessageLog/Private/Model/MessageLogListingModel.h
--- a/Engine/Source/Developer/MessageLog/Private/Model/MessageLogListingModel.h
+++ b/Engine/Source/Developer/MessageLog/Private/Model/MessageLogListingModel.h
@@ -82,13 +82,23 @@ public:
 	void NewPage( const FText& InTitle, uint32 InMaxPages );
 
 	/**
-	* Switches the current page to the one specified by the title. If the page does not exist, creates a new page.
+	* Sets the current page to the one specified by the title. If the page does not exist, creates a new page.
 	* @param	InTitle		The title of the page to search for, or create
 	* @param	InMaxPages	The maximum number of pages we keep around. If the count is exceeded, 
 	*						we discard the oldest page.
 	* @returns true if the page has changed (either switched or created a new one).
 	*/
-	bool SwitchToPage( const FText& InTitle, uint32 InMaxPages  );
+	bool SetCurrentPage( const FText& InTitle, uint32 InMaxPages  );
+
+	/**
+	* Sets the current page to the one specified by the index. If the page does not exist, creates a new page.
+	* The page will move to index 0.
+	* @param	InTitle		The title of the page to search for, or create
+	* @param	InMaxPages	The maximum number of pages we keep around. If the count is exceeded, 
+	*						we discard the oldest page.
+	* @returns true if the page has changed (either switched or created a new one).
+	*/
+	bool SetCurrentPage( const uint32 InOldPageIndex  );
 
 	/** Get the number of pages contained in this log */
 	uint32 NumPages() const;
diff --git a/Engine/Source/Developer/MessageLog/Private/Presentation/MessageLogListingViewModel.cpp b/Engine/Source/Developer/MessageLog/Private/Presentation/MessageLogListingViewModel.cpp
--- a/Engine/Source/Developer/MessageLog/Private/Presentation/MessageLogListingViewModel.cpp
+++ b/Engine/Source/Developer/MessageLog/Private/Presentation/MessageLogListingViewModel.cpp
@@ -160,6 +160,7 @@ void FMessageLogListingViewModel::AddMessages( const TArray< TSharedRef< class F
 
 void FMessageLogListingViewModel::ClearMessages()
 {
+	SetCurrentPage(CurrentPageIndex);
 	MessageLogListingModel->ClearMessages();
 }
 
@@ -278,13 +279,25 @@ void FMessageLogListingViewModel::NewPage( const FText& Title )
 	RefreshFilteredMessages();
 }
 
-void FMessageLogListingViewModel::SwitchToPage(const FText& Title)
+void FMessageLogListingViewModel::SetCurrentPage(const FText& Title)
 {
 	// we should take this as a suggestion we want to show pages!
 	bShowPages = true;
 
 	// switch page & refresh
-	MessageLogListingModel->SwitchToPage( Title, MaxPageCount );
+	MessageLogListingModel->SetCurrentPage( Title, MaxPageCount );
+	
+	// reset so we always display the current page when we switch
+	SetCurrentPageIndex(0);
+}
+
+void FMessageLogListingViewModel::SetCurrentPage(const uint32 InOldPageIndex)
+{
+	// we should take this as a suggestion we want to show pages!
+	bShowPages = true;
+
+	// switch page & refresh
+	MessageLogListingModel->SetCurrentPage( InOldPageIndex );
 	
 	// reset so we always display the current page when we switch
 	SetCurrentPageIndex(0);
diff --git a/Engine/Source/Developer/MessageLog/Private/Presentation/MessageLogListingViewModel.h b/Engine/Source/Developer/MessageLog/Private/Presentation/MessageLogListingViewModel.h
--- a/Engine/Source/Developer/MessageLog/Private/Presentation/MessageLogListingViewModel.h
+++ b/Engine/Source/Developer/MessageLog/Private/Presentation/MessageLogListingViewModel.h
@@ -56,7 +56,8 @@ public:
 	virtual const TArray< TSharedRef< class FMessageFilter> >& GetMessageFilters() const override;
 	virtual void ExecuteToken( const TSharedRef<class IMessageToken>& Token ) const override;
 	virtual void NewPage( const FText& Title ) override;
-	virtual void SwitchToPage( const FText& Title ) override;
+	virtual void SetCurrentPage( const FText& Title ) override;
+	virtual void SetCurrentPage( const uint32 InOldPageIndex ) override;
 	virtual void NotifyIfAnyMessages( const FText& Message, EMessageSeverity::Type SeverityFilter = EMessageSeverity::Info, bool bForce = false ) override;
 	virtual void Open() override;
 	virtual int32 NumMessages( EMessageSeverity::Type SeverityFilter ) override;
diff --git a/Engine/Source/Developer/RigVMDeveloper/Private/RigVMPythonUtils.cpp b/Engine/Source/Developer/RigVMDeveloper/Private/RigVMPythonUtils.cpp
--- a/Engine/Source/Developer/RigVMDeveloper/Private/RigVMPythonUtils.cpp
+++ b/Engine/Source/Developer/RigVMDeveloper/Private/RigVMPythonUtils.cpp
@@ -78,7 +78,7 @@ void RigVMPythonUtils::Print(const FString& BlueprintTitle, const FString& InMes
 		MessageLogModule.RegisterLogListing("ControlRigPythonLog", LOCTEXT("ControlRigPythonLog", "Control Rig Python Log"), InitOptions);
 	}
 	TSharedRef<IMessageLogListing> PythonLog = MessageLogModule.GetLogListing( TEXT("ControlRigPythonLog") );
-	PythonLog->SwitchToPage(FText::FromString(BlueprintTitle));
+	PythonLog->SetCurrentPage(FText::FromString(BlueprintTitle));
 
 	TSharedRef<FTokenizedMessage> Token = FTokenizedMessage::Create(EMessageSeverity::Info, FText::FromString(InMessage));
 	PythonLog->AddMessage(Token, false);
diff --git a/Engine/Source/Runtime/Core/Private/Logging/MessageLog.cpp b/Engine/Source/Runtime/Core/Private/Logging/MessageLog.cpp
--- a/Engine/Source/Runtime/Core/Private/Logging/MessageLog.cpp
+++ b/Engine/Source/Runtime/Core/Private/Logging/MessageLog.cpp
@@ -39,11 +39,18 @@ public:
 		FMsg::Logf(__FILE__, __LINE__, LogName, ELogVerbosity::Log, TEXT("%s"), *FText::Format(LOCTEXT("BasicMessageLog_NewPage", "New Page: {PageTitle}"), Arguments).ToString());
 	}
 
-	virtual void SwitchToPage( const FText& Title) override
+	virtual void SetCurrentPage( const FText& Title) override
 	{
 		FFormatNamedArguments Arguments;
 		Arguments.Add(TEXT("PageTitle"), Title);
-		FMsg::Logf(__FILE__, __LINE__, LogName, ELogVerbosity::Log, TEXT("%s"), *FText::Format(LOCTEXT("BasicMessageLog_SwitchToPage", "Switch To Page: {PageTitle}"), Arguments).ToString());
+		FMsg::Logf(__FILE__, __LINE__, LogName, ELogVerbosity::Log, TEXT("%s"), *FText::Format(LOCTEXT("BasicMessageLog_SetCurrentPage", "Set Current Page: {PageTitle}"), Arguments).ToString());
+	}
+
+	virtual void SetCurrentPage( const uint32 PageIndex) override
+	{
+		FFormatNamedArguments Arguments;
+		Arguments.Add(TEXT("PageIndex"), PageIndex);
+		FMsg::Logf(__FILE__, __LINE__, LogName, ELogVerbosity::Log, TEXT("%s"), *FText::Format(LOCTEXT("BasicMessageLog_SetCurrentPage", "Set Current Page: {PageIndex}"), Arguments).ToString());
 	}
 
 	virtual void NotifyIfAnyMessages( const FText& Message, EMessageSeverity::Type SeverityFilter, bool bForce ) override
@@ -191,10 +198,10 @@ void FMessageLog::NewPage( const FText& InLabel )
 	MessageLog->NewPage(InLabel);
 }
 
-void FMessageLog::SwitchToPage(const FText& InLabel)
+void FMessageLog::SetCurrentPage(const FText& InLabel)
 {
 	Flush();
-	MessageLog->SwitchToPage(InLabel);
+	MessageLog->SetCurrentPage(InLabel);
 }
 
 FMessageLog& FMessageLog::SuppressLoggingToOutputLog(bool bShouldSuppress)
diff --git a/Engine/Source/Runtime/Core/Public/Logging/IMessageLog.h b/Engine/Source/Runtime/Core/Public/Logging/IMessageLog.h
--- a/Engine/Source/Runtime/Core/Public/Logging/IMessageLog.h
+++ b/Engine/Source/Runtime/Core/Public/Logging/IMessageLog.h
@@ -45,10 +45,16 @@ public:
 	virtual void NewPage( const FText& Title ) = 0;
 
 	/** 
-	* Switches the current page to the one specified by the title.
+	* Sets the current page to the one specified by the title.
 	* @param	Title	the title of the page to switch to
 	*/
-	virtual void SwitchToPage( const FText& Title ) = 0;
+	virtual void SetCurrentPage( const FText& Title ) = 0;
+
+	/** 
+	* Sets the current page to the one specified by the title.
+	* @param	InOldPageIndex	the index of the page to switch to. This page will move to index 0.
+	*/
+	virtual void SetCurrentPage( const uint32 InOldPageIndex ) = 0;
 
 	/** 
 	 * Notify the user if there are any messages on the current page for this log.
diff --git a/Engine/Source/Runtime/Core/Public/Logging/MessageLog.h b/Engine/Source/Runtime/Core/Public/Logging/MessageLog.h
--- a/Engine/Source/Runtime/Core/Public/Logging/MessageLog.h
+++ b/Engine/Source/Runtime/Core/Public/Logging/MessageLog.h
@@ -88,11 +88,11 @@ public:
 	CORE_API void NewPage( const FText& InLabel );
 
 	/**
-	* Switches the current page to the one specified by the label.
+	* Sets the current page to the one specified by the label.
 	* This call will cause a flush so that the logs state is properly reflected.
 	* @param	InLabel		The label for the page.
 	*/
-	CORE_API void SwitchToPage( const FText& InLabel );
+	CORE_API void SetCurrentPage( const FText& InLabel );
 
 	/** Should we mirror message log messages from this instance to the output log during flush? */
 	CORE_API FMessageLog& SuppressLoggingToOutputLog(bool bShouldSuppress = true);
