Core constexprification pass to reduce code bloat, improve performance and simplify searching for non-constexpr constants

diff --git a/Engine/Source/Runtime/Core/Public/Containers/Array.h b/Engine/Source/Runtime/Core/Public/Containers/Array.h
--- a/Engine/Source/Runtime/Core/Public/Containers/Array.h
+++ b/Engine/Source/Runtime/Core/Public/Containers/Array.h
@@ -2842,25 +2842,25 @@ private:
 public:
 	void WriteMemoryImage(FMemoryImageWriter& Writer) const
 	{
-		static const bool bSupportsFreezeMemoryImage = TAllocatorTraits<AllocatorType>::SupportsFreezeMemoryImage && THasTypeLayout<ElementType>::Value;
+		static constexpr bool bSupportsFreezeMemoryImage = TAllocatorTraits<AllocatorType>::SupportsFreezeMemoryImage && THasTypeLayout<ElementType>::Value;
 		TSupportsFreezeMemoryImageHelper<bSupportsFreezeMemoryImage>::WriteMemoryImage(Writer, *this);
 	}
 
 	void CopyUnfrozen(const FMemoryUnfreezeContent& Context, void* Dst) const
 	{
-		static const bool bSupportsFreezeMemoryImage = TAllocatorTraits<AllocatorType>::SupportsFreezeMemoryImage && THasTypeLayout<ElementType>::Value;
+		static constexpr bool bSupportsFreezeMemoryImage = TAllocatorTraits<AllocatorType>::SupportsFreezeMemoryImage && THasTypeLayout<ElementType>::Value;
 		TSupportsFreezeMemoryImageHelper<bSupportsFreezeMemoryImage>::CopyUnfrozen(Context, *this, Dst);
 	}
 
 	static void AppendHash(const FPlatformTypeLayoutParameters& LayoutParams, FSHA1& Hasher)
 	{
-		static const bool bSupportsFreezeMemoryImage = TAllocatorTraits<AllocatorType>::SupportsFreezeMemoryImage && THasTypeLayout<ElementType>::Value;
+		static constexpr bool bSupportsFreezeMemoryImage = TAllocatorTraits<AllocatorType>::SupportsFreezeMemoryImage && THasTypeLayout<ElementType>::Value;
 		TSupportsFreezeMemoryImageHelper<bSupportsFreezeMemoryImage>::AppendHash(LayoutParams, Hasher);
 	}
 
 	void ToString(const FPlatformTypeLayoutParameters& LayoutParams, FMemoryToStringContext& OutContext) const
 	{
-		static const bool bSupportsFreezeMemoryImage = TAllocatorTraits<AllocatorType>::SupportsFreezeMemoryImage && THasTypeLayout<ElementType>::Value;
+		static constexpr bool bSupportsFreezeMemoryImage = TAllocatorTraits<AllocatorType>::SupportsFreezeMemoryImage && THasTypeLayout<ElementType>::Value;
 		TSupportsFreezeMemoryImageHelper<bSupportsFreezeMemoryImage>::ToString(LayoutParams, OutContext, *this);
 	}
 
diff --git a/Engine/Source/Runtime/Core/Public/Containers/BasicArray.h b/Engine/Source/Runtime/Core/Public/Containers/BasicArray.h
--- a/Engine/Source/Runtime/Core/Public/Containers/BasicArray.h
+++ b/Engine/Source/Runtime/Core/Public/Containers/BasicArray.h
@@ -118,7 +118,7 @@ public:
 	}
 
 private:
-	static const int32 InitialReservationSize = 16;
+	static constexpr int32 InitialReservationSize = 16;
 
 	static FORCEINLINE int32 ApplyGrowthFactor(int32 CurrentNum)
 	{
diff --git a/Engine/Source/Runtime/Core/Public/Containers/BitArray.h b/Engine/Source/Runtime/Core/Public/Containers/BitArray.h
--- a/Engine/Source/Runtime/Core/Public/Containers/BitArray.h
+++ b/Engine/Source/Runtime/Core/Public/Containers/BitArray.h
@@ -1722,7 +1722,7 @@ private:
 public:
 	void WriteMemoryImage(FMemoryImageWriter& Writer) const
 	{
-		static const bool bSupportsFreezeMemoryImage = TAllocatorTraits<Allocator>::SupportsFreezeMemoryImage;
+		static constexpr bool bSupportsFreezeMemoryImage = TAllocatorTraits<Allocator>::SupportsFreezeMemoryImage;
 		checkf(!Writer.Is32BitTarget(), TEXT("TBitArray does not currently support freezing for 32bits"));
 		TSupportsFreezeMemoryImageHelper<bSupportsFreezeMemoryImage>::WriteMemoryImage(Writer, *this);
 	}
@@ -1942,7 +1942,7 @@ private:
 	/** Find the first bit that is set in both arrays, starting with the current bit, inclusive. */
 	void FindFirstSetBit()
 	{
-		static const uint32 EmptyArrayData = 0;
+		const uint32 EmptyArrayData = 0;
 		const uint32* ArrayDataA = IfAThenAElseB(ArrayA.GetData(),&EmptyArrayData);
 		const uint32* ArrayDataB = IfAThenAElseB(ArrayB.GetData(),&EmptyArrayData);
 
diff --git a/Engine/Source/Runtime/Core/Public/Containers/Map.h b/Engine/Source/Runtime/Core/Public/Containers/Map.h
--- a/Engine/Source/Runtime/Core/Public/Containers/Map.h
+++ b/Engine/Source/Runtime/Core/Public/Containers/Map.h
@@ -138,7 +138,7 @@ class TMapBase
 	friend struct TContainerTraits<TMapBase>;
 
 public:
-	static const bool SupportsFreezeMemoryImage = TAllocatorTraits<SetAllocator>::SupportsFreezeMemoryImage;
+	static constexpr bool SupportsFreezeMemoryImage = TAllocatorTraits<SetAllocator>::SupportsFreezeMemoryImage;
 
 	typedef typename TTypeTraits<KeyType  >::ConstPointerType KeyConstPointerType;
 	typedef typename TTypeTraits<KeyType  >::ConstInitType    KeyInitType;
diff --git a/Engine/Source/Runtime/Core/Public/Containers/Set.h b/Engine/Source/Runtime/Core/Public/Containers/Set.h
--- a/Engine/Source/Runtime/Core/Public/Containers/Set.h
+++ b/Engine/Source/Runtime/Core/Public/Containers/Set.h
@@ -286,7 +286,7 @@ template<
 class TSet
 {
 public:
-	static const bool SupportsFreezeMemoryImage = TAllocatorTraits<Allocator>::SupportsFreezeMemoryImage && THasTypeLayout<InElementType>::Value;
+	static constexpr bool SupportsFreezeMemoryImage = TAllocatorTraits<Allocator>::SupportsFreezeMemoryImage && THasTypeLayout<InElementType>::Value;
 
 	typedef InElementType ElementType;
 	typedef KeyFuncs    KeyFuncsType;
diff --git a/Engine/Source/Runtime/Core/Public/Containers/SparseArray.h b/Engine/Source/Runtime/Core/Public/Containers/SparseArray.h
--- a/Engine/Source/Runtime/Core/Public/Containers/SparseArray.h
+++ b/Engine/Source/Runtime/Core/Public/Containers/SparseArray.h
@@ -78,7 +78,7 @@ class TSparseArray
 	friend class TScriptSparseArray;
 
 public:
-	static const bool SupportsFreezeMemoryImage = TAllocatorTraits<Allocator>::SupportsFreezeMemoryImage && THasTypeLayout<InElementType>::Value;
+	static constexpr bool SupportsFreezeMemoryImage = TAllocatorTraits<Allocator>::SupportsFreezeMemoryImage && THasTypeLayout<InElementType>::Value;
 
 	/** Destructor. */
 	~TSparseArray()
diff --git a/Engine/Source/Runtime/Core/Public/Containers/StackTracker.h b/Engine/Source/Runtime/Core/Public/Containers/StackTracker.h
--- a/Engine/Source/Runtime/Core/Public/Containers/StackTracker.h
+++ b/Engine/Source/Runtime/Core/Public/Containers/StackTracker.h
@@ -16,7 +16,7 @@ struct FStackTracker
 {
 public:
 	/** Maximum number of backtrace depth. */
-	static const int32 MAX_BACKTRACE_DEPTH = 50;
+	static constexpr int32 MAX_BACKTRACE_DEPTH = 50;
 	/** Helper structure to capture callstack addresses and stack count. */
 	struct FCallStack
 	{
diff --git a/Engine/Source/Runtime/Core/Public/Containers/StaticBitArray.h b/Engine/Source/Runtime/Core/Public/Containers/StaticBitArray.h
--- a/Engine/Source/Runtime/Core/Public/Containers/StaticBitArray.h
+++ b/Engine/Source/Runtime/Core/Public/Containers/StaticBitArray.h
@@ -280,8 +280,6 @@ public:
 	 */
 	int32 FindFirstClearBit() const
 	{
-		static const int32 NumBitsPerWordLog2 = FMath::FloorLog2(NumBitsPerWord);
-
 		const int32 LocalNumBits = NumBits;
 
 		int32 WordIndex = 0;
@@ -313,8 +311,6 @@ public:
 	 */
 	int32 FindFirstSetBit() const
 	{
-		static const int32 NumBitsPerWordLog2 = FMath::FloorLog2(NumBitsPerWord);
-
 		const int32 LocalNumBits = NumBits;
 
 		int32 WordIndex = 0;
@@ -373,13 +369,13 @@ public:
 		return Str;
 	}
 
-	static const uint32 NumOfBits = NumBits;
+	static constexpr uint32 NumOfBits = NumBits;
 
 private:
-
-//	static_assert(NumBits > 0, "Must have at least 1 bit.");
-	static const uint32 NumBitsPerWord = sizeof(WordType) * 8;
-	static const uint32 NumWords = ((NumBits + NumBitsPerWord - 1) & ~(NumBitsPerWord - 1)) / NumBitsPerWord;
+	static constexpr uint32 NumBitsPerWord = sizeof(WordType) * 8;
+	static constexpr uint32 NumBitsPerWordLog2 = 6;
+	static_assert(NumBitsPerWord == (1u << NumBitsPerWordLog2), "Update NumBitsPerWordLog2 to reflect WordType");
+	static constexpr uint32 NumWords = ((NumBits + NumBitsPerWord - 1) & ~(NumBitsPerWord - 1)) / NumBitsPerWord;
 	WordType Words[NumWords];
 
 	// Helper class for bool conversion
diff --git a/Engine/Source/Runtime/Core/Public/Containers/Ticker.h b/Engine/Source/Runtime/Core/Public/Containers/Ticker.h
--- a/Engine/Source/Runtime/Core/Public/Containers/Ticker.h
+++ b/Engine/Source/Runtime/Core/Public/Containers/Ticker.h
@@ -95,8 +95,8 @@ private:
 		/** Delegate to fire **/
 		FTickerDelegate Delegate;
 
-		static const uint64 DefaultState = 0;
-		static const uint64 RemovedState = 1;
+		static constexpr uint64 DefaultState = 0;
+		static constexpr uint64 RemovedState = 1;
 		// The element can be in 4 states: "idle", "executing", "idle and removed" and "executing and removed"
 		// Often right after `RemoveTicket()` call resources that are used by the delegate are destroyed. We must ensure that 
 		// these resources are not accessed by the delegate after `RemoveTicket()` returns.
diff --git a/Engine/Source/Runtime/Core/Public/Experimental/Containers/FAAArrayQueue.h b/Engine/Source/Runtime/Core/Public/Experimental/Containers/FAAArrayQueue.h
--- a/Engine/Source/Runtime/Core/Public/Experimental/Containers/FAAArrayQueue.h
+++ b/Engine/Source/Runtime/Core/Public/Experimental/Containers/FAAArrayQueue.h
@@ -79,7 +79,7 @@
 template<typename T>
 class FAAArrayQueue 
 {
-    static const long BUFFER_SIZE = 1024;  // 1024
+    static constexpr long BUFFER_SIZE = 1024;  // 1024
 
 private:
     struct Node 
diff --git a/Engine/Source/Runtime/Core/Public/Experimental/Containers/GrowOnlyLockFreeHash.h b/Engine/Source/Runtime/Core/Public/Experimental/Containers/GrowOnlyLockFreeHash.h
--- a/Engine/Source/Runtime/Core/Public/Experimental/Containers/GrowOnlyLockFreeHash.h
+++ b/Engine/Source/Runtime/Core/Public/Experimental/Containers/GrowOnlyLockFreeHash.h
@@ -226,7 +226,7 @@ private:
 	std::atomic<FHashHeader*> HashTable;
 	FCriticalSection	WriteCriticalSection;
 
-	static const int32 DEFAULT_INITIAL_SIZE = 1024;
+	static constexpr int32 DEFAULT_INITIAL_SIZE = 1024;
 
 	static void HashInsertInternal(FHashHeader* HashTableLocal, KeyType Key, ValueType Value)
 	{
diff --git a/Engine/Source/Runtime/Core/Public/GenericPlatform/GenericPlatformCrashContext.h b/Engine/Source/Runtime/Core/Public/GenericPlatform/GenericPlatformCrashContext.h
--- a/Engine/Source/Runtime/Core/Public/GenericPlatform/GenericPlatformCrashContext.h
+++ b/Engine/Source/Runtime/Core/Public/GenericPlatform/GenericPlatformCrashContext.h
@@ -251,7 +251,7 @@ public:
 	static const TCHAR* const EnabledPluginsTag;
 	static const TCHAR* const UEMinidumpName;
 	static const TCHAR* const NewLineTag;
-	static const int32 CrashGUIDLength = 128;
+	static constexpr int32 CrashGUIDLength = 128;
 
 	static const TCHAR* const CrashTypeCrash;
 	static const TCHAR* const CrashTypeAssert;
diff --git a/Engine/Source/Runtime/Core/Public/GenericPlatform/GenericPlatformDriver.h b/Engine/Source/Runtime/Core/Public/GenericPlatform/GenericPlatformDriver.h
--- a/Engine/Source/Runtime/Core/Public/GenericPlatform/GenericPlatformDriver.h
+++ b/Engine/Source/Runtime/Core/Public/GenericPlatform/GenericPlatformDriver.h
@@ -138,7 +138,7 @@ public:
 		return (*this == rhs) || (*this < rhs);
 	}
 
-	static const uint32 Size = TSize;
+	static constexpr uint32 Size = TSize;
 
 	// [0]:left .. [Size-1]:right
 	uint32 Value[Size];
diff --git a/Engine/Source/Runtime/Core/Public/GenericPlatform/GenericPlatformMath.h b/Engine/Source/Runtime/Core/Public/GenericPlatform/GenericPlatformMath.h
--- a/Engine/Source/Runtime/Core/Public/GenericPlatform/GenericPlatformMath.h
+++ b/Engine/Source/Runtime/Core/Public/GenericPlatform/GenericPlatformMath.h
@@ -927,7 +927,7 @@ struct FGenericPlatformMath
 	// A call is considered ambiguous if it passes no float types, any long double, or multiple mismatched float/double types.
 
 	template<typename T1, typename T2>
-	struct TIsSameOrNotFP { static const bool Value = TIsSame<T1, T2>::Value || !TIsFloatingPoint<T2>::Value; };
+	struct TIsSameOrNotFP { static constexpr bool Value = TIsSame<T1, T2>::Value || !TIsFloatingPoint<T2>::Value; };
 	template<typename... Ts>
 	static constexpr bool TIsAmbiguous =	!TOr<TIsFloatingPoint<typename TDecay<Ts>::Type>...>::Value ||
 											TOr<TIsSame<long double, typename TDecay<Ts>::Type>...>::Value ||
diff --git a/Engine/Source/Runtime/Core/Public/GenericPlatform/GenericPlatformMisc.h b/Engine/Source/Runtime/Core/Public/GenericPlatform/GenericPlatformMisc.h
--- a/Engine/Source/Runtime/Core/Public/GenericPlatform/GenericPlatformMisc.h
+++ b/Engine/Source/Runtime/Core/Public/GenericPlatform/GenericPlatformMisc.h
@@ -75,22 +75,22 @@ namespace EBuildConfigurations
 	typedef EBuildConfiguration Type;
 
 	UE_DEPRECATED(4.24, "EBuildConfigurations::Unknown is deprecated. Use EBuildConfiguration::Unknown instead.")
-	static const EBuildConfiguration Unknown = EBuildConfiguration::Unknown;
+	static constexpr EBuildConfiguration Unknown = EBuildConfiguration::Unknown;
 
 	UE_DEPRECATED(4.24, "EBuildConfigurations::Debug is deprecated. Use EBuildConfiguration::Debug instead.")
-	static const EBuildConfiguration Debug = EBuildConfiguration::Debug;
+	static constexpr EBuildConfiguration Debug = EBuildConfiguration::Debug;
 
 	UE_DEPRECATED(4.24, "EBuildConfigurations::DebugGame is deprecated. Use EBuildConfiguration::DebugGame instead.")
-	static const EBuildConfiguration DebugGame = EBuildConfiguration::DebugGame;
+	static constexpr EBuildConfiguration DebugGame = EBuildConfiguration::DebugGame;
 
 	UE_DEPRECATED(4.24, "EBuildConfigurations::Development is deprecated. Use EBuildConfiguration::Development instead.")
-	static const EBuildConfiguration Development = EBuildConfiguration::Development;
+	static constexpr EBuildConfiguration Development = EBuildConfiguration::Development;
 
 	UE_DEPRECATED(4.24, "EBuildConfigurations::Test is deprecated. Use EBuildConfiguration::Test instead.")
-	static const EBuildConfiguration Test = EBuildConfiguration::Test;
+	static constexpr EBuildConfiguration Test = EBuildConfiguration::Test;
 
 	UE_DEPRECATED(4.24, "EBuildConfigurations::Shipping is deprecated. Use EBuildConfiguration::Shipping instead.")
-	static const EBuildConfiguration Shipping = EBuildConfiguration::Shipping;
+	static constexpr EBuildConfiguration Shipping = EBuildConfiguration::Shipping;
 
 	/**
 	 * Returns the string representation of the specified EBuildConfiguration value.
@@ -166,16 +166,16 @@ namespace EBuildTargets
 	typedef EBuildTargetType Type;
 
 	UE_DEPRECATED(4.24, "EBuildTargets::Unknown is deprecated. Use EBuildTargetType::Unknown instead.")
-	static const EBuildTargetType Unknown = EBuildTargetType::Unknown;
+	static constexpr EBuildTargetType Unknown = EBuildTargetType::Unknown;
 
 	UE_DEPRECATED(4.24, "EBuildTargets::Editor is deprecated. Use EBuildTargetType::Unknown instead.")
-	static const EBuildTargetType Editor = EBuildTargetType::Editor;
+	static constexpr EBuildTargetType Editor = EBuildTargetType::Editor;
 
 	UE_DEPRECATED(4.24, "EBuildTargets::Game is deprecated. Use EBuildTargetType::Unknown instead.")
-	static const EBuildTargetType Game = EBuildTargetType::Game;
+	static constexpr EBuildTargetType Game = EBuildTargetType::Game;
 
 	UE_DEPRECATED(4.24, "EBuildTargets::Server is deprecated. Use EBuildTargetType::Unknown instead.")
-	static const EBuildTargetType Server = EBuildTargetType::Server;
+	static constexpr EBuildTargetType Server = EBuildTargetType::Server;
 
 	UE_DEPRECATED(4.24, "EBuildTargets::FromString is deprecated. Use LexFromString() instead.")
 	CORE_API EBuildTargetType FromString( const FString& Target );
diff --git a/Engine/Source/Runtime/Core/Public/GenericPlatform/GenericPlatformString.h b/Engine/Source/Runtime/Core/Public/GenericPlatform/GenericPlatformString.h
--- a/Engine/Source/Runtime/Core/Public/GenericPlatform/GenericPlatformString.h
+++ b/Engine/Source/Runtime/Core/Public/GenericPlatform/GenericPlatformString.h
@@ -192,7 +192,7 @@ struct FGenericPlatformString : public FGenericPlatformStricmp
 	/**
 	 * True if the encoding type of the string is some form of unicode
 	 */
-	static const bool IsUnicodeEncoded = true;
+	static constexpr bool IsUnicodeEncoded = true;
 
 
 	/**
diff --git a/Engine/Source/Runtime/Core/Public/HAL/Allocators/CachedOSVeryLargePageAllocator.h b/Engine/Source/Runtime/Core/Public/HAL/Allocators/CachedOSVeryLargePageAllocator.h
--- a/Engine/Source/Runtime/Core/Public/HAL/Allocators/CachedOSVeryLargePageAllocator.h
+++ b/Engine/Source/Runtime/Core/Public/HAL/Allocators/CachedOSVeryLargePageAllocator.h
@@ -36,16 +36,16 @@ class FCachedOSVeryLargePageAllocator
 {
 	// we make the address space twice as big as we need and use the 1st have for small pool allocations, the 2nd half is used for other allocations that are still == SizeOfSubPage
 #if UE_VERYLARGEPAGEALLOCATOR_TAKEONALL64KBALLOCATIONS
-	static const uint64 AddressSpaceToReserve = ((1024LL * 1024LL * 1024LL) * UE_VERYLARGEPAGEALLOCATOR_RESERVED_SIZE_IN_GB * 2LL);
-	static const uint64 AddressSpaceToReserveForSmallPool = AddressSpaceToReserve/2;
+	static constexpr uint64 AddressSpaceToReserve = ((1024LL * 1024LL * 1024LL) * UE_VERYLARGEPAGEALLOCATOR_RESERVED_SIZE_IN_GB * 2LL);
+	static constexpr uint64 AddressSpaceToReserveForSmallPool = AddressSpaceToReserve/2;
 #else
-	static const uint64 AddressSpaceToReserve = ((1024 * 1024 * 1024LL) * UE_VERYLARGEPAGEALLOCATOR_RESERVED_SIZE_IN_GB);
-	static const uint64 AddressSpaceToReserveForSmallPool = AddressSpaceToReserve;
+	static constexpr uint64 AddressSpaceToReserve = ((1024 * 1024 * 1024LL) * UE_VERYLARGEPAGEALLOCATOR_RESERVED_SIZE_IN_GB);
+	static constexpr uint64 AddressSpaceToReserveForSmallPool = AddressSpaceToReserve;
 #endif
-	static const uint64 SizeOfLargePage = (UE_VERYLARGEPAGEALLOCATOR_PAGESIZE_KB * 1024);
-	static const uint64 SizeOfSubPage = (1024 * 64);
-	static const uint64 NumberOfLargePages = (AddressSpaceToReserve / SizeOfLargePage);
-	static const uint64 NumberOfSubPagesPerLargePage = (SizeOfLargePage / SizeOfSubPage);
+	static constexpr uint64 SizeOfLargePage = (UE_VERYLARGEPAGEALLOCATOR_PAGESIZE_KB * 1024);
+	static constexpr uint64 SizeOfSubPage = (1024 * 64);
+	static constexpr uint64 NumberOfLargePages = (AddressSpaceToReserve / SizeOfLargePage);
+	static constexpr uint64 NumberOfSubPagesPerLargePage = (SizeOfLargePage / SizeOfSubPage);
 public:
 
 	FCachedOSVeryLargePageAllocator()
diff --git a/Engine/Source/Runtime/Core/Public/HAL/IPlatformFileCachedWrapper.h b/Engine/Source/Runtime/Core/Public/HAL/IPlatformFileCachedWrapper.h
--- a/Engine/Source/Runtime/Core/Public/HAL/IPlatformFileCachedWrapper.h
+++ b/Engine/Source/Runtime/Core/Public/HAL/IPlatformFileCachedWrapper.h
@@ -195,9 +195,9 @@ public:
 
 private:
 
-	static const uint32 BufferCacheSize = 64 * 1024; // Seems to be the magic number for best perf
-	static const uint64 BufferSizeMask  = ~((uint64)BufferCacheSize-1);
-	static const uint32	CacheCount		= 2;
+	static constexpr uint32 BufferCacheSize = 64 * 1024; // Seems to be the magic number for best perf
+	static constexpr uint64 BufferSizeMask  = ~((uint64)BufferCacheSize-1);
+	static constexpr uint32	CacheCount		= 2;
 
 	bool InnerSeek(uint64 Pos)
 	{
diff --git a/Engine/Source/Runtime/Core/Public/HAL/LowLevelMemTracker.h b/Engine/Source/Runtime/Core/Public/HAL/LowLevelMemTracker.h
--- a/Engine/Source/Runtime/Core/Public/HAL/LowLevelMemTracker.h
+++ b/Engine/Source/Runtime/Core/Public/HAL/LowLevelMemTracker.h
@@ -247,10 +247,10 @@ enum class ELLMTag : LLM_TAG_TYPE
 };
 static_assert( ELLMTag::GenericTagCount <= ELLMTag::PlatformTagStart, "too many LLM tags defined"); 
 
-static const uint32 LLM_TAG_COUNT = 256;
-static const uint32 LLM_CUSTOM_TAG_START = (int32)ELLMTag::PlatformTagStart;
-static const uint32 LLM_CUSTOM_TAG_END = (int32)ELLMTag::ProjectTagEnd;
-static const uint32 LLM_CUSTOM_TAG_COUNT = LLM_CUSTOM_TAG_END + 1 - LLM_CUSTOM_TAG_START;
+static constexpr uint32 LLM_TAG_COUNT = 256;
+static constexpr uint32 LLM_CUSTOM_TAG_START = (int32)ELLMTag::PlatformTagStart;
+static constexpr uint32 LLM_CUSTOM_TAG_END = (int32)ELLMTag::ProjectTagEnd;
+static constexpr uint32 LLM_CUSTOM_TAG_COUNT = LLM_CUSTOM_TAG_END + 1 - LLM_CUSTOM_TAG_START;
 
 
 /**
diff --git a/Engine/Source/Runtime/Core/Public/HAL/MallocCallstackHandler.h b/Engine/Source/Runtime/Core/Public/HAL/MallocCallstackHandler.h
--- a/Engine/Source/Runtime/Core/Public/HAL/MallocCallstackHandler.h
+++ b/Engine/Source/Runtime/Core/Public/HAL/MallocCallstackHandler.h
@@ -130,8 +130,8 @@ public:
 		return UsedMalloc->GetDescriptiveName();
 	}
 
-	static const SIZE_T MaxCallStackDepth = 64;
-	static const SIZE_T CallStackEntriesToSkipCount = 2;
+	static constexpr SIZE_T MaxCallStackDepth = 64;
+	static constexpr SIZE_T CallStackEntriesToSkipCount = 2;
 
 	struct FCallStackInfo
 	{
diff --git a/Engine/Source/Runtime/Core/Public/HAL/MallocDebug.h b/Engine/Source/Runtime/Core/Public/HAL/MallocDebug.h
--- a/Engine/Source/Runtime/Core/Public/HAL/MallocDebug.h
+++ b/Engine/Source/Runtime/Core/Public/HAL/MallocDebug.h
@@ -44,7 +44,7 @@ private:
 	/** Total size of allocation overhead */
 	SIZE_T		TotalWasteSize;
 
-	static const uint32 AllocatorOverhead = sizeof(FMemDebug) + sizeof(FMemDebug*) + sizeof(int32) + ALLOCATION_ALIGNMENT + sizeof(int32);
+	static constexpr uint32 AllocatorOverhead = sizeof(FMemDebug) + sizeof(FMemDebug*) + sizeof(int32) + ALLOCATION_ALIGNMENT + sizeof(int32);
 
 public:
 	// FMalloc interface.
diff --git a/Engine/Source/Runtime/Core/Public/HAL/MallocLeakDetection.h b/Engine/Source/Runtime/Core/Public/HAL/MallocLeakDetection.h
--- a/Engine/Source/Runtime/Core/Public/HAL/MallocLeakDetection.h
+++ b/Engine/Source/Runtime/Core/Public/HAL/MallocLeakDetection.h
@@ -70,7 +70,7 @@ class CORE_API FMallocLeakDetection
 		{
 			FMemory::Memzero(this, sizeof(FCallstackTrack));
 		}
-		static const int32 Depth = 32;
+		static constexpr int32 Depth = 32;
 		uint64 CallStack[Depth];
 		uint32 FirstFrame;
 		uint32 LastFrame;
diff --git a/Engine/Source/Runtime/Core/Public/HAL/MallocStomp.h b/Engine/Source/Runtime/Core/Public/HAL/MallocStomp.h
--- a/Engine/Source/Runtime/Core/Public/HAL/MallocStomp.h
+++ b/Engine/Source/Runtime/Core/Public/HAL/MallocStomp.h
@@ -27,10 +27,10 @@ class FMallocStomp final : public FMalloc
 private:
 #if PLATFORM_64BITS
 	/** Expected value to be found in the sentinel. */
-	static const SIZE_T SentinelExpectedValue = 0xdeadbeefdeadbeef;
+	static constexpr SIZE_T SentinelExpectedValue = 0xdeadbeefdeadbeef;
 #else
 	/** Expected value to be found in the sentinel. */
-	static const SIZE_T SentinelExpectedValue = 0xdeadbeef;
+	static constexpr SIZE_T SentinelExpectedValue = 0xdeadbeef;
 #endif
 
 	const SIZE_T PageSize;
@@ -52,7 +52,7 @@ private:
 
 	UPTRINT VirtualAddressCursor = 0;
 	SIZE_T VirtualAddressMax = 0;
-	static const SIZE_T VirtualAddressBlockSize = 1 * 1024 * 1024 * 1024; // 1 GB blocks
+	static constexpr SIZE_T VirtualAddressBlockSize = 1 * 1024 * 1024 * 1024; // 1 GB blocks
 
 public:
 	// FMalloc interface.
diff --git a/Engine/Source/Runtime/Core/Public/HAL/MallocStomp2.h b/Engine/Source/Runtime/Core/Public/HAL/MallocStomp2.h
--- a/Engine/Source/Runtime/Core/Public/HAL/MallocStomp2.h
+++ b/Engine/Source/Runtime/Core/Public/HAL/MallocStomp2.h
@@ -35,10 +35,10 @@ private:
 
 #if PLATFORM_64BITS
 	/** Expected value to be found in the sentinel. */
-	static const SIZE_T SentinelExpectedValue = 0xdeadbeefdeadbeef;
+	static constexpr SIZE_T SentinelExpectedValue = 0xdeadbeefdeadbeef;
 #else
 	/** Expected value to be found in the sentinel. */
-	static const SIZE_T SentinelExpectedValue = 0xdeadbeef;
+	static constexpr SIZE_T SentinelExpectedValue = 0xdeadbeef;
 #endif
 
 	const SIZE_T PageSize;
@@ -68,7 +68,7 @@ public:
 
 	struct FAddressSpaceStompDataBlock
 	{
-		static const SIZE_T VirtualAddressSpaceToReserve = 1024LL * 1024LL * 1024LL * 64LL; // 64 GB blocks, we try to reserve that much, but will half it and retry
+		static constexpr SIZE_T VirtualAddressSpaceToReserve = 1024LL * 1024LL * 1024LL * 64LL; // 64 GB blocks, we try to reserve that much, but will half it and retry
 		std::atomic<UPTRINT>	CurrentAddress;
 		UPTRINT	StartAddress;
 		UPTRINT	EndAddress;
@@ -121,10 +121,10 @@ public:
 
 
 private:
-	static const SIZE_T TargetAddressSpaceToReserve = 1024LL * 1024LL * 1024LL * 1024LL * 32LL; // how much address space we will try and reserve, 32 TB
+	static constexpr SIZE_T TargetAddressSpaceToReserve = 1024LL * 1024LL * 1024LL * 1024LL * 32LL; // how much address space we will try and reserve, 32 TB
 
 
-	static const SIZE_T MaxAddressSpaceStompDataBlocks = 4096; // this * VirtualAddressSpaceToReserve is how much address space we will try and allocate
+	static constexpr SIZE_T MaxAddressSpaceStompDataBlocks = 4096; // this * VirtualAddressSpaceToReserve is how much address space we will try and allocate
 
 	std::atomic<int64> CurrentAddressSpaceStompDataBlockIndex;
 	FAddressSpaceStompDataBlock		AddressSpaceStompDataBlocks[MaxAddressSpaceStompDataBlocks];
diff --git a/Engine/Source/Runtime/Core/Public/HAL/Platform.h b/Engine/Source/Runtime/Core/Public/HAL/Platform.h
--- a/Engine/Source/Runtime/Core/Public/HAL/Platform.h
+++ b/Engine/Source/Runtime/Core/Public/HAL/Platform.h
@@ -1021,13 +1021,13 @@ namespace TypeTests
 	template <typename A, typename B>
 	struct TAreTypesEqual
 	{
-		static const bool Value = false;
+		static constexpr bool Value = false;
 	};
 
 	template <typename T>
 	struct TAreTypesEqual<T, T>
 	{
-		static const bool Value = true;
+		static constexpr bool Value = true;
 	};
 
 	static_assert(!PLATFORM_TCHAR_IS_4_BYTES || sizeof(TCHAR) == 4, "TCHAR size must be 4 bytes.");
diff --git a/Engine/Source/Runtime/Core/Public/HAL/ThreadHeartBeat.h b/Engine/Source/Runtime/Core/Public/HAL/ThreadHeartBeat.h
--- a/Engine/Source/Runtime/Core/Public/HAL/ThreadHeartBeat.h
+++ b/Engine/Source/Runtime/Core/Public/HAL/ThreadHeartBeat.h
@@ -320,10 +320,10 @@ class CORE_API FGameThreadHitchHeartBeatThreaded : public FRunnable
 
 #if WALK_STACK_ON_HITCH_DETECTED
 #if LOOKUP_SYMBOLS_IN_HITCH_STACK_WALK
-	static const SIZE_T StackTraceSize = 65535;
+	static constexpr SIZE_T StackTraceSize = 65535;
 	ANSICHAR StackTrace[StackTraceSize];
 #else
-	static const uint32 MaxStackDepth = 128;
+	static constexpr uint32 MaxStackDepth = 128;
 	uint64 StackTrace[MaxStackDepth]; 
 #endif // LOOKUP_SYMBOLS_IN_HITCH_STACK_WALK
 #endif // WALK_STACK_ON_HITCH_DETECTED
diff --git a/Engine/Source/Runtime/Core/Public/Internationalization/TextNamespaceUtil.h b/Engine/Source/Runtime/Core/Public/Internationalization/TextNamespaceUtil.h
--- a/Engine/Source/Runtime/Core/Public/Internationalization/TextNamespaceUtil.h
+++ b/Engine/Source/Runtime/Core/Public/Internationalization/TextNamespaceUtil.h
@@ -10,8 +10,8 @@
 namespace TextNamespaceUtil
 {
 
-static const TCHAR PackageNamespaceStartMarker = TEXT('[');
-static const TCHAR PackageNamespaceEndMarker = TEXT(']');
+static constexpr TCHAR PackageNamespaceStartMarker = TEXT('[');
+static constexpr TCHAR PackageNamespaceEndMarker = TEXT(']');
 
 /**
  * Given a text and package namespace, build the full version that should be used by the localization system.
diff --git a/Engine/Source/Runtime/Core/Public/Math/sse_mathfun.h b/Engine/Source/Runtime/Core/Public/Math/sse_mathfun.h
--- a/Engine/Source/Runtime/Core/Public/Math/sse_mathfun.h
+++ b/Engine/Source/Runtime/Core/Public/Math/sse_mathfun.h
@@ -39,21 +39,12 @@
      and we have no need for the MMX Version so all those are being replace with #if 1.
   2. It looks like some linked object files have conflicting names (log_ps, exp_ps, etc). So adding the 
      prefix SseMath_ to all the functions.
+  3. constexpr and alignas transition
 */
 #pragma once
 
 #include <xmmintrin.h>
 
-/* yes I know, the top of this file is quite ugly */
-
-#ifdef _MSC_VER /* visual c++ */
-# define ALIGN16_BEG __declspec(align(16))
-# define ALIGN16_END 
-#else /* gcc or icc */
-# define ALIGN16_BEG
-# define ALIGN16_END __attribute__((aligned(16)))
-#endif
-
 /* __m128 is ugly to write */
 typedef __m128 v4sf;  // vector of 4 float (sse1)
 
@@ -67,11 +58,11 @@ typedef __m64 v2si;   // vector of 2 int (mmx)
 
 /* declare some SSE constants -- why can't I figure a better way to do that? */
 #define _PS_CONST(Name, Val)                                            \
-  static const ALIGN16_BEG float _ps_##Name[4] ALIGN16_END = { Val, Val, Val, Val }
+  alignas(16) static constexpr float _ps_##Name[4] = { Val, Val, Val, Val }
 #define _PI32_CONST(Name, Val)                                            \
-  static const ALIGN16_BEG int _pi32_##Name[4] ALIGN16_END = { Val, Val, Val, Val }
+  alignas(16) static constexpr int _pi32_##Name[4] = { Val, Val, Val, Val }
 #define _PS_CONST_TYPE(Name, Type, Val)                                 \
-  static const ALIGN16_BEG Type _ps_##Name[4] ALIGN16_END = { Val, Val, Val, Val }
+  alignas(16) static constexpr Type _ps_##Name[4] = { Val, Val, Val, Val }
 
 _PS_CONST(1  , 1.0f);
 _PS_CONST(0p5, 0.5f);
diff --git a/Engine/Source/Runtime/Core/Public/Memory/MemoryArena.h b/Engine/Source/Runtime/Core/Public/Memory/MemoryArena.h
--- a/Engine/Source/Runtime/Core/Public/Memory/MemoryArena.h
+++ b/Engine/Source/Runtime/Core/Public/Memory/MemoryArena.h
@@ -61,9 +61,9 @@ public:
 #else
 
 private:
-	static const int	ArenaShift	= 48;
-	static const uint64 ArenaMask	= 0xffff000000000000;
-	static const uint64 PointerMask = 0x0000ffffFFFFffff;
+	static constexpr int	ArenaShift	= 48;
+	static constexpr uint64 ArenaMask	= 0xffff000000000000;
+	static constexpr uint64 PointerMask = 0x0000ffffFFFFffff;
 
 	void*	TaggedPointer = nullptr;
 
diff --git a/Engine/Source/Runtime/Core/Public/Microsoft/MinimalWindowsApi.h b/Engine/Source/Runtime/Core/Public/Microsoft/MinimalWindowsApi.h
--- a/Engine/Source/Runtime/Core/Public/Microsoft/MinimalWindowsApi.h
+++ b/Engine/Source/Runtime/Core/Public/Microsoft/MinimalWindowsApi.h
@@ -101,8 +101,8 @@ namespace Windows
 	};
 
 	// Constants
-	static const BOOL TRUE = 1;
-	static const BOOL FALSE = 0;
+	static constexpr BOOL TRUE = 1;
+	static constexpr BOOL FALSE = 0;
 
 	// Modules
 	MINIMAL_WINDOWS_API HMODULE WINAPI LoadLibraryW(LPCTSTR lpFileName);
diff --git a/Engine/Source/Runtime/Core/Public/Misc/AES.h b/Engine/Source/Runtime/Core/Public/Misc/AES.h
--- a/Engine/Source/Runtime/Core/Public/Misc/AES.h
+++ b/Engine/Source/Runtime/Core/Public/Misc/AES.h
@@ -10,14 +10,14 @@
 
 struct CORE_API FAES
 {
-	static const uint32 AESBlockSize = 16;
+	static constexpr uint32 AESBlockSize = 16;
 
 	/** 
 	 * Class representing a 256 bit AES key
 	 */
 	struct FAESKey
 	{
-		static const int32 KeySize = 32;
+		static constexpr int32 KeySize = 32;
 
 		uint8 Key[KeySize];
 
diff --git a/Engine/Source/Runtime/Core/Public/Misc/CString.h b/Engine/Source/Runtime/Core/Public/Misc/CString.h
--- a/Engine/Source/Runtime/Core/Public/Misc/CString.h
+++ b/Engine/Source/Runtime/Core/Public/Misc/CString.h
@@ -419,10 +419,10 @@ template <typename CharType = TCHAR>
 struct TCStringSpcHelper
 {
 	/** Number of characters to be stored in string. */
-	static const int32 MAX_SPACES = 255;
+	static constexpr int32 MAX_SPACES = 255;
 
 	/** Number of tabs to be stored in string. */
-	static const int32 MAX_TABS = 255;
+	static constexpr int32 MAX_TABS = 255;
 
 	static CORE_API const CharType SpcArray[MAX_SPACES + 1];
 	static CORE_API const CharType TabArray[MAX_TABS + 1];
diff --git a/Engine/Source/Runtime/Core/Public/Misc/CommandLine.h b/Engine/Source/Runtime/Core/Public/Misc/CommandLine.h
--- a/Engine/Source/Runtime/Core/Public/Misc/CommandLine.h
+++ b/Engine/Source/Runtime/Core/Public/Misc/CommandLine.h
@@ -17,7 +17,7 @@
 struct CORE_API FCommandLine
 {
 	/** maximum size of the command line */
-	static const uint32 MaxCommandLineSize = 16384;
+	static constexpr uint32 MaxCommandLineSize = 16384;
 
 	/** 
 	 * Resets FCommandLine to an uninitialised state as if Set() has never been called
diff --git a/Engine/Source/Runtime/Core/Public/Misc/EnumRange.h b/Engine/Source/Runtime/Core/Public/Misc/EnumRange.h
--- a/Engine/Source/Runtime/Core/Public/Misc/EnumRange.h
+++ b/Engine/Source/Runtime/Core/Public/Misc/EnumRange.h
@@ -48,8 +48,8 @@
 	struct NEnumRangePrivate::TEnumRangeTraits<EnumType> \
 	{ \
 		enum { RangeType = 0 }; \
-		static const __underlying_type(EnumType) Begin = (__underlying_type(EnumType))(First); \
-		static const __underlying_type(EnumType) End   = (__underlying_type(EnumType))(Last) + 1; \
+		static constexpr __underlying_type(EnumType) Begin = (__underlying_type(EnumType))(First); \
+		static constexpr __underlying_type(EnumType) End   = (__underlying_type(EnumType))(Last) + 1; \
 	};
 
 
@@ -78,7 +78,7 @@
 		template <typename Dummy> \
 		static const EnumType* GetPointer(bool bLast) \
 		{ \
-			static const EnumType Values[] = { __VA_ARGS__ }; \
+			static constexpr EnumType Values[] = { __VA_ARGS__ }; \
 			return bLast ? Values + sizeof(Values) / sizeof(EnumType) : Values; \
 		} \
 	};
diff --git a/Engine/Source/Runtime/Core/Public/Misc/ExpressionParserTypes.h b/Engine/Source/Runtime/Core/Public/Misc/ExpressionParserTypes.h
--- a/Engine/Source/Runtime/Core/Public/Misc/ExpressionParserTypes.h
+++ b/Engine/Source/Runtime/Core/Public/Misc/ExpressionParserTypes.h
@@ -209,7 +209,7 @@ public:
 private:
 
 	/** The maximum size of type we will allow allocation on the stack (for efficiency). Anything larger will be allocated on the heap. */
-	static const uint32 MaxStackAllocationSize = 64 - sizeof(FGuid);
+	static constexpr uint32 MaxStackAllocationSize = 64 - sizeof(FGuid);
 
 	/** Helper accessor to the data interface. Returns null for empty containers. */
 	Impl::IExpressionNodeStorage* GetData();
diff --git a/Engine/Source/Runtime/Core/Public/Misc/IEngineCrypto.h b/Engine/Source/Runtime/Core/Public/Misc/IEngineCrypto.h
--- a/Engine/Source/Runtime/Core/Public/Misc/IEngineCrypto.h
+++ b/Engine/Source/Runtime/Core/Public/Misc/IEngineCrypto.h
@@ -7,7 +7,7 @@
 #include "Containers/ArrayView.h"
 
 typedef void* FRSAKeyHandle;
-static const FRSAKeyHandle InvalidRSAKeyHandle = nullptr;
+static constexpr FRSAKeyHandle InvalidRSAKeyHandle = nullptr;
 
 struct IEngineCrypto : public IModularFeature
 {
diff --git a/Engine/Source/Runtime/Core/Public/Misc/TextFilterExpressionEvaluator.h b/Engine/Source/Runtime/Core/Public/Misc/TextFilterExpressionEvaluator.h
--- a/Engine/Source/Runtime/Core/Public/Misc/TextFilterExpressionEvaluator.h
+++ b/Engine/Source/Runtime/Core/Public/Misc/TextFilterExpressionEvaluator.h
@@ -163,7 +163,7 @@ namespace TextFilterExpressionParser
 	{																	\
 		struct TYPE														\
 		{																\
-			static const int32 MonikerCount = MONIKER_COUNT;			\
+			static constexpr int32 MonikerCount = MONIKER_COUNT;			\
 			static const TCHAR* Monikers[MonikerCount];					\
 		};																\
 	}																	\
diff --git a/Engine/Source/Runtime/Core/Public/Serialization/FileRegions.h b/Engine/Source/Runtime/Core/Public/Serialization/FileRegions.h
--- a/Engine/Source/Runtime/Core/Public/Serialization/FileRegions.h
+++ b/Engine/Source/Runtime/Core/Public/Serialization/FileRegions.h
@@ -25,7 +25,7 @@ enum class EFileRegionType : uint8
 // Regions are used to improve data compression and patching on some platforms.
 struct FFileRegion
 {
-	static const constexpr TCHAR* RegionsFileExtension = TEXT(".uregs");
+	static constexpr TCHAR RegionsFileExtension[] = TEXT(".uregs");
 
 	uint64 Offset;
 	uint64 Length;
diff --git a/Engine/Source/Runtime/Core/Public/Serialization/MemoryLayout.h b/Engine/Source/Runtime/Core/Public/Serialization/MemoryLayout.h
--- a/Engine/Source/Runtime/Core/Public/Serialization/MemoryLayout.h
+++ b/Engine/Source/Runtime/Core/Public/Serialization/MemoryLayout.h
@@ -180,7 +180,7 @@ const T* GetDefault();
 template<typename T>
 struct THasCustomDefaultObject
 {
-	static const bool Value = TModels<CProvidesDefaultUObject, T>::Value;
+	static constexpr bool Value = TModels<CProvidesDefaultUObject, T>::Value;
 };
 
 template<typename T>
@@ -217,19 +217,19 @@ struct TValidateInterfaceHelper;
 template<typename T> struct TValidateInterfaceHelper<T, ETypeLayoutInterface::NonVirtual>
 {
 	UE_STATIC_ONLY(TValidateInterfaceHelper);
-	static const bool Value = !TIsPolymorphic<T>::Value;
+	static constexpr bool Value = !TIsPolymorphic<T>::Value;
 };
 
 template<typename T> struct TValidateInterfaceHelper<T, ETypeLayoutInterface::Virtual>
 {
 	UE_STATIC_ONLY(TValidateInterfaceHelper);
-	static const bool Value = !TIsAbstract<T>::Value;
+	static constexpr bool Value = !TIsAbstract<T>::Value;
 };
 
 template<typename T> struct TValidateInterfaceHelper<T, ETypeLayoutInterface::Abstract>
 {
 	UE_STATIC_ONLY(TValidateInterfaceHelper);
-	static const bool Value = true;
+	static constexpr bool Value = true;
 };
 
 namespace Freeze
@@ -339,13 +339,13 @@ struct CProvidesStaticStruct
 template<typename T>
 struct TUsePropertyFreezing
 {
-	static const bool Value = (TModels<CProvidesStaticClass, T>::Value || TModels<CProvidesStaticStruct, T>::Value);
+	static constexpr bool Value = (TModels<CProvidesStaticClass, T>::Value || TModels<CProvidesStaticStruct, T>::Value);
 };
 
 template<typename T>
 struct TProvidesStaticStruct
 {
-	static const bool Value = TModels<CProvidesStaticStruct, T>::Value;
+	static constexpr bool Value = TModels<CProvidesStaticStruct, T>::Value;
 };
 
 template <typename T, bool bUsePropertyFreezing=TUsePropertyFreezing<T>::Value>
@@ -535,9 +535,9 @@ PRAGMA_ENABLE_DEPRECATION_WARNINGS
 #define INTERNAL_LAYOUT_INTERFACE_INLINE_IMPL(Type) PREPROCESSOR_JOIN(INTERNAL_LAYOUT_INTERFACE_INLINE_IMPL_, Type)
 
 #define INTERNAL_DECLARE_TYPE_LAYOUT_COMMON(T, InInterface) \
-	static const int CounterBase = __COUNTER__; \
+	static constexpr int CounterBase = __COUNTER__; \
 	public: using DerivedType = PREPROCESSOR_REMOVE_OPTIONAL_PARENS(T); \
-	static const ETypeLayoutInterface::Type InterfaceType = ETypeLayoutInterface::InInterface; \
+	static constexpr ETypeLayoutInterface::Type InterfaceType = ETypeLayoutInterface::InInterface; \
 	UE_DECLARE_INTERNAL_LINK_BASE(InternalLinkType) { UE_STATIC_ONLY(InternalLinkType); static FORCEINLINE void Initialize(FTypeLayoutDesc& TypeDesc) {} }
 
 #define INTERNAL_DECLARE_INLINE_TYPE_LAYOUT(T, InInterface) \
@@ -654,7 +654,7 @@ template<typename T>
 struct THasTypeLayout
 {
 	UE_STATIC_ONLY(THasTypeLayout);
-	static const bool Value = TModels<CTypeLayout, T>::Value;
+	static constexpr bool Value = TModels<CTypeLayout, T>::Value;
 };
 
 template<typename T>
diff --git a/Engine/Source/Runtime/Core/Public/Stats/Stats2.h b/Engine/Source/Runtime/Core/Public/Stats/Stats2.h
--- a/Engine/Source/Runtime/Core/Public/Stats/Stats2.h
+++ b/Engine/Source/Runtime/Core/Public/Stats/Stats2.h
@@ -882,7 +882,7 @@ template< typename TEnum >
 struct TStatMessage
 {
 	typedef TEnum TStructEnum;
-	static const int32 EnumCount = TEnum::Num;
+	static constexpr int32 EnumCount = TEnum::Num;
 
 	/**
 	* Generic payload
diff --git a/Engine/Source/Runtime/Core/Public/Tasks/TaskPrivate.h b/Engine/Source/Runtime/Core/Public/Tasks/TaskPrivate.h
--- a/Engine/Source/Runtime/Core/Public/Tasks/TaskPrivate.h
+++ b/Engine/Source/Runtime/Core/Public/Tasks/TaskPrivate.h
@@ -82,7 +82,7 @@ namespace UE { namespace Tasks
 
 			// a special internal task priority for "inline" task execution - a task is executed as soon as it's launched and has no 
 			// pending dependencies, "inline", w/o scheduling
-			static const ETaskPriority InlineTaskPriority{ ETaskPriority::Count };
+			static constexpr ETaskPriority InlineTaskPriority{ ETaskPriority::Count };
 
 			// initialises the task but doesn't launches it
 			template<typename TaskBodyType>
@@ -601,7 +601,7 @@ namespace UE { namespace Tasks
 			// reaches 0 the task is scheduled for execution.
 			// NumLocks's the most significant bit (see `ExecutionFlag`) is set on task execution start, and indicates that now NumLocks is about 
 			// how many times the task must be unlocked to be completed
-			static const uint32 NumInitialLocks = 1 + 1;
+			static constexpr uint32 NumInitialLocks = 1 + 1;
 			std::atomic<uint32> NumLocks{ NumInitialLocks };
 
 			// A single-producer/single-consumer container to store back links to "prerequsites" (either execution prerequisites or nested tasks 
diff --git a/Engine/Source/Runtime/Core/Public/Templates/IntegralConstant.h b/Engine/Source/Runtime/Core/Public/Templates/IntegralConstant.h
--- a/Engine/Source/Runtime/Core/Public/Templates/IntegralConstant.h
+++ b/Engine/Source/Runtime/Core/Public/Templates/IntegralConstant.h
@@ -10,5 +10,5 @@
 template <typename T, T Val>
 struct TIntegralConstant
 {
-	static const T Value = Val;
+	static constexpr T Value = Val;
 };
diff --git a/Engine/Source/Runtime/Core/Public/Templates/MaxSizeof.h b/Engine/Source/Runtime/Core/Public/Templates/MaxSizeof.h
--- a/Engine/Source/Runtime/Core/Public/Templates/MaxSizeof.h
+++ b/Engine/Source/Runtime/Core/Public/Templates/MaxSizeof.h
@@ -11,7 +11,7 @@ struct TMaxSizeof;
 template <>
 struct TMaxSizeof<>
 {
-	static const uint32 Value = 0;
+	static constexpr uint32 Value = 0;
 };
 
 template <typename T, typename... TRest>
diff --git a/Engine/Source/Runtime/Core/Public/Templates/UnrealTypeTraits.h b/Engine/Source/Runtime/Core/Public/Templates/UnrealTypeTraits.h
--- a/Engine/Source/Runtime/Core/Public/Templates/UnrealTypeTraits.h
+++ b/Engine/Source/Runtime/Core/Public/Templates/UnrealTypeTraits.h
@@ -65,9 +65,9 @@ struct TIsDerivedFrom
 	public:
 	// Test the derived type pointer. If it inherits from BaseType, the Test( BaseType* ) 
 	// will be chosen. If it does not, Test( ... ) will be chosen.
-	static const bool Value = sizeof(Test( DerivedTypePtr() )) == sizeof(Yes);
+	static constexpr bool Value = sizeof(Test( DerivedTypePtr() )) == sizeof(Yes);
 
-	static const bool IsDerived = Value;
+	static constexpr bool IsDerived = Value;
 };
 
 /**
