Validate that the item at tail of automation latent command is indeed the item that just got completed before removing it

diff --git a/Engine/Source/Runtime/Core/Private/Misc/AutomationTest.cpp b/Engine/Source/Runtime/Core/Private/Misc/AutomationTest.cpp
--- a/Engine/Source/Runtime/Core/Private/Misc/AutomationTest.cpp
+++ b/Engine/Source/Runtime/Core/Private/Misc/AutomationTest.cpp
@@ -413,8 +413,16 @@ bool FAutomationTestFramework::ExecuteLatentCommands()
 		bool bComplete = NextCommand->InternalUpdate();
 		if (bComplete)
 		{
-			//all done.  remove from the queue
-			LatentCommands.Dequeue(NextCommand);
+			TSharedPtr<IAutomationLatentCommand>* TailCommand = LatentCommands.Peek();
+			if (TailCommand != nullptr && NextCommand == *TailCommand)
+			{
+				//all done. remove the tail
+				LatentCommands.Pop();
+			}
+			else
+			{
+				UE_LOG(LogAutomationTest, Verbose, TEXT("Tail of latent command queue is not removed, because last completed automation latent command is not corresponding."));
+			}
 		}
 		else
 		{
