Added opt-in printf format string validation for logging macros (currently implemented for MSVC only)

To opt in, please define the UE_VALIDATE_FORMAT_STRINGS macro to 1

diff --git a/Engine/Source/Runtime/Core/Public/HAL/Platform.h b/Engine/Source/Runtime/Core/Public/HAL/Platform.h
--- a/Engine/Source/Runtime/Core/Public/HAL/Platform.h
+++ b/Engine/Source/Runtime/Core/Public/HAL/Platform.h
@@ -205,6 +205,9 @@
 #ifndef PLATFORM_ENABLE_VECTORINTRINSICS_NEON
 	#define PLATFORM_ENABLE_VECTORINTRINSICS_NEON	0
 #endif
+#ifndef UE_VALIDATE_FORMAT_STRINGS
+	#define UE_VALIDATE_FORMAT_STRINGS 0
+#endif
 #ifndef PLATFORM_USE_LS_SPEC_FOR_WIDECHAR
 	#define PLATFORM_USE_LS_SPEC_FOR_WIDECHAR	1
 #endif
diff --git a/Engine/Source/Runtime/Core/Public/Logging/LogMacros.h b/Engine/Source/Runtime/Core/Public/Logging/LogMacros.h
--- a/Engine/Source/Runtime/Core/Public/Logging/LogMacros.h
+++ b/Engine/Source/Runtime/Core/Public/Logging/LogMacros.h
@@ -178,6 +178,12 @@ private:
 	#define UE_SET_LOG_VERBOSITY(CategoryName, Verbosity) \
 		CategoryName.SetVerbosity(ELogVerbosity::Verbosity);
 
+	#if UE_VALIDATE_FORMAT_STRINGS && defined(_MSC_VER)
+		#define UE_VALIDATE_FORMAT_STRING(Format, ...) do { if (false) { wprintf(Format, ##__VA_ARGS__); } } while(false)
+	#else
+		#define UE_VALIDATE_FORMAT_STRING(Format, ...)
+	#endif
+
 	/** 
 	 * INTERNAL IMPLEMENTATION. DO NOT CALL DIRECTLY!
 	**/
@@ -186,6 +192,7 @@ private:
 		static_assert((ELogVerbosity::Verbosity & ELogVerbosity::VerbosityMask) < ELogVerbosity::NumVerbosity && ELogVerbosity::Verbosity > 0, "Verbosity must be constant and in range."); \
 		UE_LOG_EXPAND_IS_FATAL(Verbosity, PREPROCESSOR_NOTHING, if (!CategoryName.IsSuppressed(ELogVerbosity::Verbosity))) \
 			{ \
+				UE_VALIDATE_FORMAT_STRING(Format, ##__VA_ARGS__); \
 				DispatchCheckVerify([] (const auto& LCategoryName, const auto& LFormat, const auto&... UE_LOG_Args) FORCENOINLINE \
 				{ \
 					TRACE_LOG_MESSAGE(LCategoryName, Verbosity, LFormat, UE_LOG_Args...) \
