Compile fix for Win64 ClangGame

#rnx

diff --git a/Engine/Source/Runtime/Core/Public/Math/Matrix.h b/Engine/Source/Runtime/Core/Public/Math/Matrix.h
--- a/Engine/Source/Runtime/Core/Public/Math/Matrix.h
+++ b/Engine/Source/Runtime/Core/Public/Math/Matrix.h
@@ -446,7 +446,7 @@ private:
 	}
 };
 
-#ifndef _MSC_VER // MSVC can't forward declare explicit specializations
+#if !defined(_MSC_VER) || defined(__clang__) // MSVC can't forward declare explicit specializations
 template<> const FMatrix44f FMatrix44f::Identity;
 template<> const FMatrix44d FMatrix44d::Identity;
 #endif
diff --git a/Engine/Source/Runtime/Core/Public/Math/Quat.h b/Engine/Source/Runtime/Core/Public/Math/Quat.h
--- a/Engine/Source/Runtime/Core/Public/Math/Quat.h
+++ b/Engine/Source/Runtime/Core/Public/Math/Quat.h
@@ -653,7 +653,7 @@ inline FArchive& operator<<(FArchive& Ar, TQuat<double>& F)
 	return Ar;
 }
 
-#ifndef _MSC_VER // MSVC can't forward declare explicit specializations
+#if !defined(_MSC_VER) || defined(__clang__)  // MSVC can't forward declare explicit specializations
 template<> const FQuat4f FQuat4f::Identity;
 template<> const FQuat4d FQuat4d::Identity;
 #endif
diff --git a/Engine/Source/Runtime/Core/Public/Math/TransformNonVectorized.h b/Engine/Source/Runtime/Core/Public/Math/TransformNonVectorized.h
--- a/Engine/Source/Runtime/Core/Public/Math/TransformNonVectorized.h
+++ b/Engine/Source/Runtime/Core/Public/Math/TransformNonVectorized.h
@@ -1202,7 +1202,7 @@ public:
 	explicit TTransform(const TTransform<FArg>& From) : TTransform<T>((TQuat<T>)From.GetRotation(), (TVector<T>)From.GetTranslation(), (TVector<T>)From.GetScale3D()) {}
 };
 
-#ifndef _MSC_VER // MSVC can't forward declare explicit specializations
+#if !defined(_MSC_VER) || defined(__clang__)  // MSVC can't forward declare explicit specializations
 template<> const FTransform3f FTransform3f::Identity;
 template<> const FTransform3d FTransform3d::Identity;
 #endif
diff --git a/Engine/Source/Runtime/Core/Public/Math/TransformVectorized.h b/Engine/Source/Runtime/Core/Public/Math/TransformVectorized.h
--- a/Engine/Source/Runtime/Core/Public/Math/TransformVectorized.h
+++ b/Engine/Source/Runtime/Core/Public/Math/TransformVectorized.h
@@ -1453,7 +1453,7 @@ public:
 	explicit TTransform(const TTransform<FArg>& From) : TTransform(MakeVectorRegisterDouble(From.Rotation), MakeVectorRegisterDouble(From.Translation), MakeVectorRegisterDouble(From.Scale3D)) {}
 };
 
-#ifndef _MSC_VER // MSVC can't forward declare explicit specializations
+#if !defined(_MSC_VER) || defined(__clang__)  // MSVC can't forward declare explicit specializations
 template<> const FTransform3f FTransform3f::Identity;
 template<> const FTransform3d FTransform3d::Identity;
 #endif
diff --git a/Engine/Source/Runtime/Core/Public/Math/UnrealMathSSE.h b/Engine/Source/Runtime/Core/Public/Math/UnrealMathSSE.h
--- a/Engine/Source/Runtime/Core/Public/Math/UnrealMathSSE.h
+++ b/Engine/Source/Runtime/Core/Public/Math/UnrealMathSSE.h
@@ -386,7 +386,7 @@ FORCEINLINE constexpr VectorRegister4Int MakeVectorRegisterIntConstant(int32 X,
 {
 #if !PLATFORM_LITTLE_ENDIAN
 #error Big-endian unimplemented
-#elif defined(_MSC_VER)
+#elif defined(_MSC_VER) && !defined(__clang__)
     return {static_cast<char>(X >> 0), static_cast<char>(X >> 8), static_cast<char>(X >> 16), static_cast<char>(X >> 24),
             static_cast<char>(Y >> 0), static_cast<char>(Y >> 8), static_cast<char>(Y >> 16), static_cast<char>(Y >> 24), 
             static_cast<char>(Z >> 0), static_cast<char>(Z >> 8), static_cast<char>(Z >> 16), static_cast<char>(Z >> 24), 
diff --git a/Engine/Source/Runtime/Core/Public/Math/Vector.h b/Engine/Source/Runtime/Core/Public/Math/Vector.h
--- a/Engine/Source/Runtime/Core/Public/Math/Vector.h
+++ b/Engine/Source/Runtime/Core/Public/Math/Vector.h
@@ -2373,7 +2373,7 @@ FORCEINLINE TVector<T> TVector<T>::Max3(const TVector<T>& A, const TVector<T>& B
         ); 
 }
 
-#ifndef _MSC_VER // MSVC can't forward declare explicit specializations
+#if !defined(_MSC_VER) || defined(__clang__)  // MSVC can't forward declare explicit specializations
 template<> const FVector3f FVector3f::ZeroVector;
 template<> const FVector3f FVector3f::OneVector;
 template<> const FVector3f FVector3f::UpVector;
