THIRD_PARTY_INCLUDES_END now enables OVERLOADED_VIRTUAL_WARNINGS instead of continuing to disable them
#rnx

diff --git a/Engine/Source/Runtime/Core/Public/Android/AndroidPlatformCompilerPreSetup.h b/Engine/Source/Runtime/Core/Public/Android/AndroidPlatformCompilerPreSetup.h
--- a/Engine/Source/Runtime/Core/Public/Android/AndroidPlatformCompilerPreSetup.h
+++ b/Engine/Source/Runtime/Core/Public/Android/AndroidPlatformCompilerPreSetup.h
@@ -19,7 +19,7 @@
 #ifndef THIRD_PARTY_INCLUDES_END
 	#define THIRD_PARTY_INCLUDES_END \
 		PRAGMA_ENABLE_UNINITIALIZED_CONST_REFERENCE_WARNINGS \
-		PRAGMA_DISABLE_OVERLOADED_VIRTUAL_WARNINGS \
+		PRAGMA_ENABLE_OVERLOADED_VIRTUAL_WARNINGS \
 		PRAGMA_ENABLE_DEPRECATION_WARNINGS \
 		PRAGMA_ENABLE_MISSING_VIRTUAL_DESTRUCTOR_WARNINGS \
 		PRAGMA_ENABLE_UNDEFINED_IDENTIFIER_WARNINGS \
diff --git a/Engine/Source/Runtime/Core/Public/Apple/ApplePlatformCompilerPreSetup.h b/Engine/Source/Runtime/Core/Public/Apple/ApplePlatformCompilerPreSetup.h
--- a/Engine/Source/Runtime/Core/Public/Apple/ApplePlatformCompilerPreSetup.h
+++ b/Engine/Source/Runtime/Core/Public/Apple/ApplePlatformCompilerPreSetup.h
@@ -19,7 +19,7 @@
 #ifndef THIRD_PARTY_INCLUDES_END
 	#define THIRD_PARTY_INCLUDES_END \
 		PRAGMA_ENABLE_UNINITIALIZED_CONST_REFERENCE_WARNINGS \
-		PRAGMA_DISABLE_OVERLOADED_VIRTUAL_WARNINGS \
+		PRAGMA_ENABLE_OVERLOADED_VIRTUAL_WARNINGS \
 		PRAGMA_ENABLE_DEPRECATION_WARNINGS \
 		PRAGMA_ENABLE_MISSING_VIRTUAL_DESTRUCTOR_WARNINGS \
 		PRAGMA_ENABLE_UNDEFINED_IDENTIFIER_WARNINGS \
diff --git a/Engine/Source/Runtime/Core/Public/Unix/UnixPlatformCompilerPreSetup.h b/Engine/Source/Runtime/Core/Public/Unix/UnixPlatformCompilerPreSetup.h
--- a/Engine/Source/Runtime/Core/Public/Unix/UnixPlatformCompilerPreSetup.h
+++ b/Engine/Source/Runtime/Core/Public/Unix/UnixPlatformCompilerPreSetup.h
@@ -50,7 +50,7 @@
 	#define THIRD_PARTY_INCLUDES_END \
 		PRAGMA_ENABLE_UNINITIALIZED_CONST_REFERENCE_WARNINGS \
 		PRAGMA_ENABLE_MACRO_REDEFINED_WARNINGS \
-		PRAGMA_DISABLE_OVERLOADED_VIRTUAL_WARNINGS \
+		PRAGMA_ENABLE_OVERLOADED_VIRTUAL_WARNINGS \
 		PRAGMA_ENABLE_DEPRECATION_WARNINGS \
 		PRAGMA_ENABLE_DEPRECATED_REGISTER_WARNINGS \
 		PRAGMA_ENABLE_MISSING_VIRTUAL_DESTRUCTOR_WARNINGS \
