disabled tasks tests temporarily as they OOM on consoles, will tune them later

diff --git a/Engine/Source/Runtime/Core/Private/Tests/Async/TaskGraphTest.cpp b/Engine/Source/Runtime/Core/Private/Tests/Async/TaskGraphTest.cpp
--- a/Engine/Source/Runtime/Core/Private/Tests/Async/TaskGraphTest.cpp
+++ b/Engine/Source/Runtime/Core/Private/Tests/Async/TaskGraphTest.cpp
@@ -719,7 +719,7 @@ extern int32 GNumForegroundWorkers;
 
 namespace TaskGraphTests
 {
-	IMPLEMENT_SIMPLE_AUTOMATION_TEST(FTaskGraphGraphEventTest, "System.Core.Async.TaskGraph.GraphEventTest", EAutomationTestFlags::ApplicationContextMask | EAutomationTestFlags::EngineFilter);
+	IMPLEMENT_SIMPLE_AUTOMATION_TEST(FTaskGraphGraphEventTest, "System.Core.Async.TaskGraph.GraphEventTest", EAutomationTestFlags::ApplicationContextMask | EAutomationTestFlags::EngineFilter | EAutomationTestFlags::Disabled);
 
 	bool FTaskGraphGraphEventTest::RunTest(const FString& Parameters)
 	{
@@ -859,7 +859,7 @@ namespace TaskGraphTests
 		return true;
 	}
 
-	IMPLEMENT_SIMPLE_AUTOMATION_TEST(FTaskGraphBasicTest, "System.Core.Async.TaskGraph.BasicTest", EAutomationTestFlags::ApplicationContextMask | EAutomationTestFlags::EngineFilter);
+	IMPLEMENT_SIMPLE_AUTOMATION_TEST(FTaskGraphBasicTest, "System.Core.Async.TaskGraph.BasicTest", EAutomationTestFlags::ApplicationContextMask | EAutomationTestFlags::EngineFilter | EAutomationTestFlags::Disabled);
 
 	bool FTaskGraphBasicTest::RunTest(const FString& Parameters)
 	{
@@ -1301,7 +1301,7 @@ PRAGMA_ENABLE_DEPRECATION_WARNINGS
 		}
 	}
 
-	IMPLEMENT_SIMPLE_AUTOMATION_TEST(FTaskGraphPerfTest, "System.Core.Async.TaskGraph.PerfTest", EAutomationTestFlags::ApplicationContextMask | EAutomationTestFlags::EngineFilter);
+	IMPLEMENT_SIMPLE_AUTOMATION_TEST(FTaskGraphPerfTest, "System.Core.Async.TaskGraph.PerfTest", EAutomationTestFlags::ApplicationContextMask | EAutomationTestFlags::EngineFilter | EAutomationTestFlags::Disabled);
 
 	bool FTaskGraphPerfTest::RunTest(const FString& Parameters)
 	{
@@ -1325,7 +1325,7 @@ PRAGMA_ENABLE_DEPRECATION_WARNINGS
 		return true;
 	}
 
-	IMPLEMENT_SIMPLE_AUTOMATION_TEST(FTaskGraphStatsThreadRedirectionTest, "System.Core.Async.TaskGraph.StatsThreadRedirection", EAutomationTestFlags::ApplicationContextMask | EAutomationTestFlags::EngineFilter);
+	IMPLEMENT_SIMPLE_AUTOMATION_TEST(FTaskGraphStatsThreadRedirectionTest, "System.Core.Async.TaskGraph.StatsThreadRedirection", EAutomationTestFlags::ApplicationContextMask | EAutomationTestFlags::EngineFilter | EAutomationTestFlags::Disabled);
 
 	bool FTaskGraphStatsThreadRedirectionTest::RunTest(const FString& Parameters)
 	{
@@ -1377,7 +1377,7 @@ PRAGMA_ENABLE_DEPRECATION_WARNINGS
 		}
 	}
 
-	IMPLEMENT_SIMPLE_AUTOMATION_TEST(FTaskGraphOversubscriptionTest, "System.Core.Async.TaskGraph.Oversubscription", EAutomationTestFlags::ApplicationContextMask | EAutomationTestFlags::EngineFilter);
+	IMPLEMENT_SIMPLE_AUTOMATION_TEST(FTaskGraphOversubscriptionTest, "System.Core.Async.TaskGraph.Oversubscription", EAutomationTestFlags::ApplicationContextMask | EAutomationTestFlags::EngineFilter | EAutomationTestFlags::Disabled);
 
 	bool FTaskGraphOversubscriptionTest::RunTest(const FString& Parameters)
 	{
diff --git a/Engine/Source/Runtime/Core/Private/Tests/Tasks/TasksTest.cpp b/Engine/Source/Runtime/Core/Private/Tests/Tasks/TasksTest.cpp
--- a/Engine/Source/Runtime/Core/Private/Tests/Tasks/TasksTest.cpp
+++ b/Engine/Source/Runtime/Core/Private/Tests/Tasks/TasksTest.cpp
@@ -19,7 +19,7 @@ namespace UE { namespace TasksTests
 	{
 	}
 
-	IMPLEMENT_SIMPLE_AUTOMATION_TEST(FTasksBasicTest, "System.Core.Tasks.Basic", EAutomationTestFlags::ApplicationContextMask | EAutomationTestFlags::EngineFilter);
+	IMPLEMENT_SIMPLE_AUTOMATION_TEST(FTasksBasicTest, "System.Core.Tasks.Basic", EAutomationTestFlags::ApplicationContextMask | EAutomationTestFlags::EngineFilter | EAutomationTestFlags::Disabled);
 
 	template<uint32 SpawnerGroupsNum, uint32 SpawnersPerGroupNum>
 	void BasicStressTest()
@@ -388,7 +388,7 @@ namespace UE { namespace TasksTests
 		return true;
 	}
 
-	IMPLEMENT_SIMPLE_AUTOMATION_TEST(FTasksPipeTest, "System.Core.Tasks.Pipe", EAutomationTestFlags::ApplicationContextMask | EAutomationTestFlags::EngineFilter);
+	IMPLEMENT_SIMPLE_AUTOMATION_TEST(FTasksPipeTest, "System.Core.Tasks.Pipe", EAutomationTestFlags::ApplicationContextMask | EAutomationTestFlags::EngineFilter | EAutomationTestFlags::Disabled);
 
 	template<uint32 SpawnerGroupsNum, uint32 SpawnersPerGroupNum>
 	void PipeStressTest();
@@ -615,7 +615,7 @@ namespace UE { namespace TasksTests
 		TlsValue = Dummy;
 	}
 
-	IMPLEMENT_SIMPLE_AUTOMATION_TEST(FTlsTest, "System.Core.Tls", EAutomationTestFlags::ApplicationContextMask | EAutomationTestFlags::EngineFilter);
+	IMPLEMENT_SIMPLE_AUTOMATION_TEST(FTlsTest, "System.Core.Tls", EAutomationTestFlags::ApplicationContextMask | EAutomationTestFlags::EngineFilter | EAutomationTestFlags::Disabled);
 
 	bool FTlsTest::RunTest(const FString& Parameters)
 	{
@@ -625,7 +625,7 @@ namespace UE { namespace TasksTests
 		return true;
 	}
 
-	IMPLEMENT_SIMPLE_AUTOMATION_TEST(FTasksDependenciesTest, "System.Core.Tasks.Dependencies", EAutomationTestFlags::ApplicationContextMask | EAutomationTestFlags::EngineFilter);
+	IMPLEMENT_SIMPLE_AUTOMATION_TEST(FTasksDependenciesTest, "System.Core.Tasks.Dependencies", EAutomationTestFlags::ApplicationContextMask | EAutomationTestFlags::EngineFilter | EAutomationTestFlags::Disabled);
 
 	template<uint64 NumBranches, uint64 NumLoops, uint64 NumTasks>
 	void DependenciesPerfTest()
@@ -870,7 +870,7 @@ namespace UE { namespace TasksTests
 			N11.IsCompleted() && N12.IsCompleted() && N21.IsCompleted() && N22.IsCompleted());
 	}
 
-	IMPLEMENT_SIMPLE_AUTOMATION_TEST(FTasksDeepRetractionTest, "System.Core.Tasks.DeepRetraction", EAutomationTestFlags::ApplicationContextMask | EAutomationTestFlags::EngineFilter);
+	IMPLEMENT_SIMPLE_AUTOMATION_TEST(FTasksDeepRetractionTest, "System.Core.Tasks.DeepRetraction", EAutomationTestFlags::ApplicationContextMask | EAutomationTestFlags::EngineFilter | EAutomationTestFlags::Disabled);
 
 	bool FTasksDeepRetractionTest::RunTest(const FString& Parameters)
 	{
@@ -952,7 +952,7 @@ namespace UE { namespace TasksTests
 		}
 	}
 
-	IMPLEMENT_SIMPLE_AUTOMATION_TEST(FTasksDeepRetractionStressTest, "System.Core.Tasks.DeepRetraction.Stress", EAutomationTestFlags::ApplicationContextMask | EAutomationTestFlags::EngineFilter);
+	IMPLEMENT_SIMPLE_AUTOMATION_TEST(FTasksDeepRetractionStressTest, "System.Core.Tasks.DeepRetraction.Stress", EAutomationTestFlags::ApplicationContextMask | EAutomationTestFlags::EngineFilter | EAutomationTestFlags::Disabled);
 
 	bool FTasksDeepRetractionStressTest::RunTest(const FString& Parameters)
 	{
@@ -978,7 +978,7 @@ namespace UE { namespace TasksTests
 		}
 	}
 
-	IMPLEMENT_SIMPLE_AUTOMATION_TEST(FTasksNestedTasksStressTest, "System.Core.Tasks.NestedTasks.Stress", EAutomationTestFlags::ApplicationContextMask | EAutomationTestFlags::EngineFilter);
+	IMPLEMENT_SIMPLE_AUTOMATION_TEST(FTasksNestedTasksStressTest, "System.Core.Tasks.NestedTasks.Stress", EAutomationTestFlags::ApplicationContextMask | EAutomationTestFlags::EngineFilter | EAutomationTestFlags::Disabled);
 
 	bool FTasksNestedTasksStressTest::RunTest(const FString& Parameters)
 	{
@@ -1169,7 +1169,7 @@ namespace UE { namespace TasksTests
 		}
 	}
 
-	IMPLEMENT_SIMPLE_AUTOMATION_TEST(FTasksPerfTest, "System.Core.Tasks.PerfTest", EAutomationTestFlags::ApplicationContextMask | EAutomationTestFlags::EngineFilter);
+	IMPLEMENT_SIMPLE_AUTOMATION_TEST(FTasksPerfTest, "System.Core.Tasks.PerfTest", EAutomationTestFlags::ApplicationContextMask | EAutomationTestFlags::EngineFilter | EAutomationTestFlags::Disabled);
 
 	bool FTasksPerfTest::RunTest(const FString& Parameters)
 	{
