- Ensure we don't use aligned loads on some compilers for VectorLoad with VectorRegister4Float & VectorRegister4Double

diff --git a/Engine/Source/Runtime/Core/Public/Math/UnrealMathVectorCommon.h b/Engine/Source/Runtime/Core/Public/Math/UnrealMathVectorCommon.h
--- a/Engine/Source/Runtime/Core/Public/Math/UnrealMathVectorCommon.h
+++ b/Engine/Source/Runtime/Core/Public/Math/UnrealMathVectorCommon.h
@@ -54,12 +54,12 @@ FORCEINLINE VectorRegister4Float VectorLoad(const FVector2D* Ptr)
 
 FORCEINLINE VectorRegister4Float VectorLoad(const VectorRegister4Float* Ptr)
 {
-	return *Ptr;
+	return VectorLoad((const float*)(Ptr));
 }
 
 FORCEINLINE VectorRegister4Double VectorLoad(const VectorRegister4Double* Ptr)
 {
-	return *Ptr;
+	return VectorLoad((const double*)(Ptr));
 }
 
 
