Disabled tasks tests temporarily as they OOM on some platforms, will tune them later

- I also disabled "System.Core.Containers.TSTicker" as it was exceeding the memory of some platforms as well

diff --git a/Engine/Source/Runtime/Core/Private/Tests/Containers/TickerTests.cpp b/Engine/Source/Runtime/Core/Private/Tests/Containers/TickerTests.cpp
--- a/Engine/Source/Runtime/Core/Private/Tests/Containers/TickerTests.cpp
+++ b/Engine/Source/Runtime/Core/Private/Tests/Containers/TickerTests.cpp
@@ -9,7 +9,7 @@
 
 #if WITH_DEV_AUTOMATION_TESTS
 
-IMPLEMENT_SIMPLE_AUTOMATION_TEST(FTSTickerTest, "System.Core.Containers.TSTicker", EAutomationTestFlags::ApplicationContextMask | EAutomationTestFlags::EngineFilter);
+IMPLEMENT_SIMPLE_AUTOMATION_TEST(FTSTickerTest, "System.Core.Containers.TSTicker", EAutomationTestFlags::ApplicationContextMask | EAutomationTestFlags::EngineFilter | EAutomationTestFlags::Disabled);
 
 template<uint32 NumDelegates, uint32 NumTicks>
 void TickerPerfTest()
