Fix ASan use after scope with ANSI_TO_TCHAR.

Never store the return of ANSI_TO_TCHAR, as it goes out of scope

diff --git a/Engine/Source/Runtime/Core/Private/IO/IoDispatcherFileBackend.cpp b/Engine/Source/Runtime/Core/Private/IO/IoDispatcherFileBackend.cpp
--- a/Engine/Source/Runtime/Core/Private/IO/IoDispatcherFileBackend.cpp
+++ b/Engine/Source/Runtime/Core/Private/IO/IoDispatcherFileBackend.cpp
@@ -1965,10 +1965,9 @@ TSharedRef<IIoDispatcherFileBackend> CreateIoDispatcherFileBackend()
 
 	if (bCheckForPlatformImplementation)
 	{
-		const TCHAR* ModuleName = ANSI_TO_TCHAR(PLATFORM_IODISPATCHER_MODULE);
-		if (FModuleManager::Get().ModuleExists(ModuleName))
+		if (FModuleManager::Get().ModuleExists(ANSI_TO_TCHAR(PLATFORM_IODISPATCHER_MODULE)))
 		{
-			IPlatformFileIoStoreModule* PlatformModule = FModuleManager::LoadModulePtr<IPlatformFileIoStoreModule>(ModuleName);
+			IPlatformFileIoStoreModule* PlatformModule = FModuleManager::LoadModulePtr<IPlatformFileIoStoreModule>(ANSI_TO_TCHAR(PLATFORM_IODISPATCHER_MODULE));
 			if (PlatformModule)
 			{
 				TUniquePtr<IPlatformFileIoStore> PlatformImpl = PlatformModule->CreatePlatformFileIoStore();
