Enable NEON intrinsics only on ARM64 Android builds

diff --git a/Engine/Source/Runtime/Core/Public/Android/AndroidPlatform.h b/Engine/Source/Runtime/Core/Public/Android/AndroidPlatform.h
--- a/Engine/Source/Runtime/Core/Public/Android/AndroidPlatform.h
+++ b/Engine/Source/Runtime/Core/Public/Android/AndroidPlatform.h
@@ -71,10 +71,8 @@ typedef FAndroidTypes FPlatformTypes;
 	#define RUNNING_WITH_ASAN						0
 #endif
 
-// Conditionally set in AndroidToolChain.cs
-// always set to 1 for ARM64 builds
 #define PLATFORM_ENABLE_VECTORINTRINSICS			1
-#define PLATFORM_ENABLE_VECTORINTRINSICS_NEON		1
+#define PLATFORM_ENABLE_VECTORINTRINSICS_NEON		PLATFORM_ANDROID_ARM64
 
 #if __has_feature(cxx_decltype_auto)
 	#define PLATFORM_COMPILER_HAS_DECLTYPE_AUTO 1
