Compile fixes for ClangGame Win64 Editor and MSVC ARM64

#rnx

diff --git a/Engine/Source/Runtime/Core/Public/Math/Matrix.h b/Engine/Source/Runtime/Core/Public/Math/Matrix.h
--- a/Engine/Source/Runtime/Core/Public/Math/Matrix.h
+++ b/Engine/Source/Runtime/Core/Public/Math/Matrix.h
@@ -447,8 +447,8 @@ private:
 };
 
 #if !defined(_MSC_VER) || defined(__clang__) // MSVC can't forward declare explicit specializations
-template<> const FMatrix44f FMatrix44f::Identity;
-template<> const FMatrix44d FMatrix44d::Identity;
+template<> CORE_API const FMatrix44f FMatrix44f::Identity;
+template<> CORE_API const FMatrix44d FMatrix44d::Identity;
 #endif
 
 
diff --git a/Engine/Source/Runtime/Core/Public/Math/Quat.h b/Engine/Source/Runtime/Core/Public/Math/Quat.h
--- a/Engine/Source/Runtime/Core/Public/Math/Quat.h
+++ b/Engine/Source/Runtime/Core/Public/Math/Quat.h
@@ -654,8 +654,8 @@ inline FArchive& operator<<(FArchive& Ar, TQuat<double>& F)
 }
 
 #if !defined(_MSC_VER) || defined(__clang__)  // MSVC can't forward declare explicit specializations
-template<> const FQuat4f FQuat4f::Identity;
-template<> const FQuat4d FQuat4d::Identity;
+template<> CORE_API const FQuat4f FQuat4f::Identity;
+template<> CORE_API const FQuat4d FQuat4d::Identity;
 #endif
 /* TQuat inline functions
  *****************************************************************************/
diff --git a/Engine/Source/Runtime/Core/Public/Math/TransformNonVectorized.h b/Engine/Source/Runtime/Core/Public/Math/TransformNonVectorized.h
--- a/Engine/Source/Runtime/Core/Public/Math/TransformNonVectorized.h
+++ b/Engine/Source/Runtime/Core/Public/Math/TransformNonVectorized.h
@@ -1203,8 +1203,8 @@ public:
 };
 
 #if !defined(_MSC_VER) || defined(__clang__)  // MSVC can't forward declare explicit specializations
-template<> const FTransform3f FTransform3f::Identity;
-template<> const FTransform3d FTransform3d::Identity;
+template<> CORE_API const FTransform3f FTransform3f::Identity;
+template<> CORE_API const FTransform3d FTransform3d::Identity;
 #endif
 template<typename T>
 FORCEINLINE bool TTransform<T>::AnyHasNegativeScale(const TVector<T>& InScale3D, const  TVector<T>& InOtherScale3D) 
diff --git a/Engine/Source/Runtime/Core/Public/Math/TransformVectorized.h b/Engine/Source/Runtime/Core/Public/Math/TransformVectorized.h
--- a/Engine/Source/Runtime/Core/Public/Math/TransformVectorized.h
+++ b/Engine/Source/Runtime/Core/Public/Math/TransformVectorized.h
@@ -1454,8 +1454,8 @@ public:
 };
 
 #if !defined(_MSC_VER) || defined(__clang__)  // MSVC can't forward declare explicit specializations
-template<> const FTransform3f FTransform3f::Identity;
-template<> const FTransform3d FTransform3d::Identity;
+template<> CORE_API const FTransform3f FTransform3f::Identity;
+template<> CORE_API const FTransform3d FTransform3d::Identity;
 #endif
 
 
diff --git a/Engine/Source/Runtime/Core/Public/Math/UnrealMathDirectX.h b/Engine/Source/Runtime/Core/Public/Math/UnrealMathDirectX.h
--- a/Engine/Source/Runtime/Core/Public/Math/UnrealMathDirectX.h
+++ b/Engine/Source/Runtime/Core/Public/Math/UnrealMathDirectX.h
@@ -79,6 +79,16 @@ FORCEINLINE constexpr VectorRegister4Int MakeVectorRegisterIntConstant(int32 X,
             static_cast<char>(W >> 0), static_cast<char>(W >> 8), static_cast<char>(W >> 16), static_cast<char>(W >> 24)};
 }
 
+FORCEINLINE constexpr VectorRegister4Float MakeVectorRegisterFloatConstant(float X, float Y, float Z, float W)
+{
+	return VectorRegister4Float { X, Y, Z, W };
+}
+
+FORCEINLINE constexpr VectorRegister2Double MakeVectorRegister2DoubleConstant(double X, double Y)
+{
+	return VectorRegister2Double { X, Y };
+}
+
 /*=============================================================================
  *	Constants:
  *============================================================================*/
diff --git a/Engine/Source/Runtime/Core/Public/Math/UnrealMathFPU.h b/Engine/Source/Runtime/Core/Public/Math/UnrealMathFPU.h
--- a/Engine/Source/Runtime/Core/Public/Math/UnrealMathFPU.h
+++ b/Engine/Source/Runtime/Core/Public/Math/UnrealMathFPU.h
@@ -287,7 +287,17 @@ FORCEINLINE VectorRegister4Int MakeVectorRegisterInt(int32 X, int32 Y, int32 Z,
 
 FORCEINLINE constexpr VectorRegister4Int MakeVectorRegisterIntConstant(int32 X, int32 Y, int32 Z, int32 W)
 {
-    return VectorRegister4Int {X, Y, Z, W};
+    return VectorRegister4Int { X, Y, Z, W };
+}
+
+FORCEINLINE constexpr VectorRegister4Float MakeVectorRegisterFloatConstant(float X, float Y, float Z, float W)
+{
+	return VectorRegister4Float { X, Y, Z, W };
+}
+
+FORCEINLINE constexpr VectorRegister2Double MakeVectorRegister2DoubleConstant(double X, double Y)
+{
+	return VectorRegister2Double { X, Y };
 }
 
 /*=============================================================================
diff --git a/Engine/Source/Runtime/Core/Public/Math/UnrealMathNeon.h b/Engine/Source/Runtime/Core/Public/Math/UnrealMathNeon.h
--- a/Engine/Source/Runtime/Core/Public/Math/UnrealMathNeon.h
+++ b/Engine/Source/Runtime/Core/Public/Math/UnrealMathNeon.h
@@ -34,7 +34,7 @@ template<typename T, typename BASE_TYPE>
 struct alignas(alignof(T)) VectorRegisterWrapper
 {
 	FORCEINLINE VectorRegisterWrapper() {}
-	FORCEINLINE VectorRegisterWrapper(const T& vec) : m_vec(vec) {}
+	FORCEINLINE constexpr VectorRegisterWrapper(T vec) : m_vec(vec) {}
 
 	FORCEINLINE operator T&() { return m_vec; }
 	FORCEINLINE operator const T&() const { return m_vec; }
@@ -73,6 +73,35 @@ typedef VectorRegisterWrapper<float32x4_t, float> VectorRegister4Float;
 typedef VectorRegisterWrapper<float64x2_t, double> VectorRegister2Double;
 typedef VectorRegisterWrapper<int32x4_t, int> VectorRegister4Int;
 typedef VectorRegisterWrapper<int64x2_t, int64> VectorRegister2Int64;
+
+FORCEINLINE constexpr VectorRegister4Int MakeVectorRegisterIntConstant(int32 X, int32 Y, int32 Z, int32 W)
+{
+    int32x4_t Out = {};
+	Out.n128_i32[0] = X;
+    Out.n128_i32[1] = Y;
+    Out.n128_i32[2] = Z;
+    Out.n128_i32[3] = W;
+    return Out;
+}
+
+FORCEINLINE constexpr VectorRegister4Float MakeVectorRegisterFloatConstant(float X, float Y, float Z, float W)
+{
+    float32x4_t Out = {};
+    Out.n128_f32[0] = X;
+    Out.n128_f32[1] = Y;
+    Out.n128_f32[2] = Z;
+    Out.n128_f32[3] = W;
+	return Out;
+}
+
+FORCEINLINE constexpr VectorRegister2Double MakeVectorRegister2DoubleConstant(double X, double Y)
+{
+    float64x2_t Out = {};
+    Out.n128_f64[0] = X;
+    Out.n128_f64[1] = Y;
+	return Out;
+}
+
 #else
 
 /** 16-byte vector register type */
@@ -80,6 +109,22 @@ typedef float32x4_t GCC_ALIGN(16) VectorRegister4Float;
 typedef float64x2_t GCC_ALIGN(16) VectorRegister2Double;
 typedef int32x4_t  GCC_ALIGN(16) VectorRegister4Int;
 typedef int64x2_t GCC_ALIGN(16) VectorRegister2Int64;
+
+FORCEINLINE constexpr VectorRegister4Int MakeVectorRegisterIntConstant(int32 X, int32 Y, int32 Z, int32 W)
+{
+    return VectorRegister4Int { X, Y, Z, W };
+}
+
+FORCEINLINE constexpr VectorRegister4Float MakeVectorRegisterFloatConstant(float X, float Y, float Z, float W)
+{
+    return VectorRegister4Float { X, Y, Z, W };
+}
+
+FORCEINLINE constexpr VectorRegister2Double MakeVectorRegister2DoubleConstant(double X, double Y)
+{
+    return VectorRegister2Double { X, Y };
+}
+
 #endif
 
 #define DECLARE_VECTOR_REGISTER(X, Y, Z, W) MakeVectorRegister( X, Y, Z, W )
@@ -336,11 +381,6 @@ FORCEINLINE VectorRegister4Int MakeVectorRegisterInt(int32 X, int32 Y, int32 Z,
 	return Tmp.V;
 }
 
-FORCEINLINE constexpr VectorRegister4Int MakeVectorRegisterIntConstant(int32 X, int32 Y, int32 Z, int32 W)
-{
-    return VectorRegister4Int {X, Y, Z, W};
-}
-
 FORCEINLINE VectorRegister2Int64 MakeVectorRegisterInt64(int64 X, int64 Y)
 {
 	union U
@@ -2468,17 +2508,7 @@ inline bool VectorContainsNaNOrInfinite(const VectorRegister4Float& Vec)
 	VectorRegister4Float ExpTest = VectorBitwiseAnd(Vec, FloatInfinity);
 
 	// Compare to full exponent & combine resulting flags into lane 0
-#ifdef _MSC_VER
-	// msvc can only initialize using the first union type which is: n128_u64[2];
-#  if PLATFORM_LITTLE_ENDIAN
-	static constexpr int32x4_t Table = MakeVectorRegisterIntConstant(0x0C080400, 0, 0, 0);
-#  else
-	static constexpr int32x4_t Table = MakeVectorRegisterIntConstant(0, 0, 0, 0x0004080C);
-#  endif
-#else
-	// clang can initialize with this syntax, but not the msvc one
-	static constexpr int8x16_t Table = { 0,4,8,12, 0,0,0,0, 0,0,0,0, 0,0,0,0 };
-#endif
+	const int32x4_t Table = MakeVectorRegisterIntConstant(0x0C080400, 0, 0, 0);
 
 	uint8x16_t res = (uint8x16_t)VectorCompareEQ(ExpTest, FloatInfinity);
 	// If we have all zeros, all elements are finite
@@ -2501,18 +2531,7 @@ inline bool VectorContainsNaNOrInfinite(const VectorRegister4Double& Vec)
 	VectorRegister4Double ExpTest = VectorBitwiseAnd(Vec, DoubleInfinity);
 
 	// Compare to full exponent & combine resulting flags into lane 0
-#ifdef _MSC_VER
-	// msvc can only initialize using the first union type which is: n128_u64[2];
-	//TODO: Update MSVC Tables
-#  if PLATFORM_LITTLE_ENDIAN
-	static constexpr int32x4_t Table = MakeVectorRegisterIntConstant(0x18100800, 0, 0, 0);
-#  else
-	static constexpr int32x4_t Table = MakeVectorRegisterIntConstant(0, 0, 0, 0x00081018);
-#  endif
-#else
-	// clang can initialize with this syntax, but not the msvc one
-	static constexpr int8x16_t Table = {0,8,16,24, 0,0,0,0, 0,0,0,0, 0,0,0,0};
-#endif
+	const int32x4_t Table = MakeVectorRegisterIntConstant(0x18100800, 0, 0, 0);
 
 	VectorRegister4Double InfTestRes = VectorCompareEQ(ExpTest, DoubleInfinity);
 
diff --git a/Engine/Source/Runtime/Core/Public/Math/UnrealMathSSE.h b/Engine/Source/Runtime/Core/Public/Math/UnrealMathSSE.h
--- a/Engine/Source/Runtime/Core/Public/Math/UnrealMathSSE.h
+++ b/Engine/Source/Runtime/Core/Public/Math/UnrealMathSSE.h
@@ -398,6 +398,16 @@ FORCEINLINE constexpr VectorRegister4Int MakeVectorRegisterIntConstant(int32 X,
 #endif
 }
 
+FORCEINLINE constexpr VectorRegister4Float MakeVectorRegisterFloatConstant(float X, float Y, float Z, float W)
+{
+	return VectorRegister4Float { X, Y, Z, W };
+}
+
+FORCEINLINE constexpr VectorRegister2Double MakeVectorRegister2DoubleConstant(double X, double Y)
+{
+	return VectorRegister2Double { X, Y };
+}
+
 /*=============================================================================
  *	Constants:
  *============================================================================*/
diff --git a/Engine/Source/Runtime/Core/Public/Math/UnrealMathVectorConstants.h b/Engine/Source/Runtime/Core/Public/Math/UnrealMathVectorConstants.h
--- a/Engine/Source/Runtime/Core/Public/Math/UnrealMathVectorConstants.h
+++ b/Engine/Source/Runtime/Core/Public/Math/UnrealMathVectorConstants.h
@@ -9,24 +9,15 @@
 // Long-term, we should only have one set of constexpr make functions without the Constant suffix.
 ///////////////////////////////////////////////////////////////////////////////////////////////////////////////////
 
-FORCEINLINE constexpr VectorRegister4Float MakeVectorRegisterFloatConstant(float X, float Y, float Z, float W)
-{
-	return { X, Y, Z, W };
-}
-
-FORCEINLINE constexpr VectorRegister2Double MakeVectorRegister2DoubleConstant(double X, double Y)
+FORCEINLINE constexpr VectorRegister4Double MakeVectorRegisterDoubleConstant(const VectorRegister2Double& XY, const VectorRegister2Double& ZW)
 {
-	return { X, Y };
+	return VectorRegister4Double(XY, ZW, VectorRegisterConstInit{});
 }
 
 FORCEINLINE constexpr VectorRegister4Double MakeVectorRegisterDoubleConstant(double X, double Y, double Z, double W)
 {
-	return VectorRegister4Double(VectorRegister2Double{X, Y}, VectorRegister2Double{Z, W}, VectorRegisterConstInit{});
-}
-
-FORCEINLINE constexpr VectorRegister4Double MakeVectorRegisterDoubleConstant(const VectorRegister2Double& XY, const VectorRegister2Double& ZW)
-{
-	return VectorRegister4Double(XY, ZW, VectorRegisterConstInit{});
+	return MakeVectorRegisterDoubleConstant(MakeVectorRegister2DoubleConstant(X, Y),
+											MakeVectorRegister2DoubleConstant(Z, W));
 }
 
 FORCEINLINE constexpr VectorRegister4Float MakeVectorRegisterConstant(float X, float Y, float Z, float W)
diff --git a/Engine/Source/Runtime/Core/Public/Math/Vector.h b/Engine/Source/Runtime/Core/Public/Math/Vector.h
--- a/Engine/Source/Runtime/Core/Public/Math/Vector.h
+++ b/Engine/Source/Runtime/Core/Public/Math/Vector.h
@@ -2374,28 +2374,28 @@ FORCEINLINE TVector<T> TVector<T>::Max3(const TVector<T>& A, const TVector<T>& B
 }
 
 #if !defined(_MSC_VER) || defined(__clang__)  // MSVC can't forward declare explicit specializations
-template<> const FVector3f FVector3f::ZeroVector;
-template<> const FVector3f FVector3f::OneVector;
-template<> const FVector3f FVector3f::UpVector;
-template<> const FVector3f FVector3f::DownVector;
-template<> const FVector3f FVector3f::ForwardVector;
-template<> const FVector3f FVector3f::BackwardVector;
-template<> const FVector3f FVector3f::RightVector;
-template<> const FVector3f FVector3f::LeftVector;
-template<> const FVector3f FVector3f::XAxisVector;
-template<> const FVector3f FVector3f::YAxisVector;
-template<> const FVector3f FVector3f::ZAxisVector;
-template<> const FVector3d FVector3d::ZeroVector;
-template<> const FVector3d FVector3d::OneVector;
-template<> const FVector3d FVector3d::UpVector;
-template<> const FVector3d FVector3d::DownVector;
-template<> const FVector3d FVector3d::ForwardVector;
-template<> const FVector3d FVector3d::BackwardVector;
-template<> const FVector3d FVector3d::RightVector;
-template<> const FVector3d FVector3d::LeftVector;
-template<> const FVector3d FVector3d::XAxisVector;
-template<> const FVector3d FVector3d::YAxisVector;
-template<> const FVector3d FVector3d::ZAxisVector;
+template<> CORE_API const FVector3f FVector3f::ZeroVector;
+template<> CORE_API const FVector3f FVector3f::OneVector;
+template<> CORE_API const FVector3f FVector3f::UpVector;
+template<> CORE_API const FVector3f FVector3f::DownVector;
+template<> CORE_API const FVector3f FVector3f::ForwardVector;
+template<> CORE_API const FVector3f FVector3f::BackwardVector;
+template<> CORE_API const FVector3f FVector3f::RightVector;
+template<> CORE_API const FVector3f FVector3f::LeftVector;
+template<> CORE_API const FVector3f FVector3f::XAxisVector;
+template<> CORE_API const FVector3f FVector3f::YAxisVector;
+template<> CORE_API const FVector3f FVector3f::ZAxisVector;
+template<> CORE_API const FVector3d FVector3d::ZeroVector;
+template<> CORE_API const FVector3d FVector3d::OneVector;
+template<> CORE_API const FVector3d FVector3d::UpVector;
+template<> CORE_API const FVector3d FVector3d::DownVector;
+template<> CORE_API const FVector3d FVector3d::ForwardVector;
+template<> CORE_API const FVector3d FVector3d::BackwardVector;
+template<> CORE_API const FVector3d FVector3d::RightVector;
+template<> CORE_API const FVector3d FVector3d::LeftVector;
+template<> CORE_API const FVector3d FVector3d::XAxisVector;
+template<> CORE_API const FVector3d FVector3d::YAxisVector;
+template<> CORE_API const FVector3d FVector3d::ZAxisVector;
 #endif
 
 /**
