CompactBinary: Implemented FCbObjectId::NewObjectId and fixed JSON serialization of ObjectId fields

diff --git a/Engine/Source/Runtime/Core/Private/Serialization/CompactBinary.cpp b/Engine/Source/Runtime/Core/Private/Serialization/CompactBinary.cpp
--- a/Engine/Source/Runtime/Core/Private/Serialization/CompactBinary.cpp
+++ b/Engine/Source/Runtime/Core/Private/Serialization/CompactBinary.cpp
@@ -3,6 +3,7 @@
 #include "Serialization/CompactBinary.h"
 
 #include "Hash/Blake3.h"
+#include "Hash/xxhash.h"
 #include "Math/UnrealMathUtility.h"
 #include "Memory/CompositeBuffer.h"
 #include "Misc/ByteSwap.h"
@@ -71,6 +72,27 @@ FCbObjectId::FCbObjectId(const FMemoryView ObjectId)
 	FMemory::Memcpy(Bytes, ObjectId.GetData(), sizeof(Bytes));
 }
 
+FCbObjectId 
+FCbObjectId::NewObjectId()
+{
+	// This is pretty ad hoc and does not work like the zen-side code
+	// But since this is used very sparingly it does not matter too
+	// much right now, we can fix it up later
+
+	FCbObjectId Oid;
+
+	FGuid Guid = FGuid::NewGuid();
+
+	uint8 HashBytes[16];
+	FXxHash128 Hash = FXxHash128::HashBuffer(&Guid, sizeof Guid);
+	Hash.ToByteArray(HashBytes);
+
+	memcpy(Oid.Bytes, HashBytes, sizeof Oid);
+
+	return Oid;
+}
+
+
 ///////////////////////////////////////////////////////////////////////////////////////////////////////////////////////
 
 FCbFieldView::FCbFieldView(const void* const Data, ECbFieldType Type)
diff --git a/Engine/Source/Runtime/Core/Public/Serialization/CompactBinary.h b/Engine/Source/Runtime/Core/Public/Serialization/CompactBinary.h
--- a/Engine/Source/Runtime/Core/Public/Serialization/CompactBinary.h
+++ b/Engine/Source/Runtime/Core/Public/Serialization/CompactBinary.h
@@ -488,6 +488,8 @@ public:
 	/** Returns a view of the raw byte array for the ObjectId. */
 	constexpr inline FMemoryView GetView() const { return MakeMemoryView(Bytes); }
 
+	CORE_API static FCbObjectId NewObjectId();
+
 private:
 	alignas(uint32) ByteArray Bytes{};
 };
