CompactBinary: Updated FCbObjectId::NewObjectId() to match the structure of Oid in Zen

#rnx

diff --git a/Engine/Source/Runtime/Core/Private/Serialization/CompactBinary.cpp b/Engine/Source/Runtime/Core/Private/Serialization/CompactBinary.cpp
--- a/Engine/Source/Runtime/Core/Private/Serialization/CompactBinary.cpp
+++ b/Engine/Source/Runtime/Core/Private/Serialization/CompactBinary.cpp
@@ -14,6 +14,7 @@
 #include "Serialization/CompactBinaryValue.h"
 #include "Serialization/VarInt.h"
 #include "String/BytesToHex.h"
+#include <atomic>
 
 ///////////////////////////////////////////////////////////////////////////////////////////////////////////////////////
 
@@ -23,6 +24,9 @@ namespace UE::CompactBinary::Private
 static constexpr const uint8 GEmptyObjectValue[] = { uint8(ECbFieldType::Object), 0x00 };
 static constexpr const uint8 GEmptyArrayValue[] = { uint8(ECbFieldType::Array), 0x01, 0x00 };
 
+static std::atomic<uint32> GObjectIdRunId;
+static std::atomic<uint32> GObjectIdSerial;
+
 } // UE::CompactBinary::Private
 
 ///////////////////////////////////////////////////////////////////////////////////////////////////////////////////////
@@ -72,27 +76,34 @@ FCbObjectId::FCbObjectId(const FMemoryView ObjectId)
 	FMemory::Memcpy(Bytes, ObjectId.GetData(), sizeof(Bytes));
 }
 
-FCbObjectId 
-FCbObjectId::NewObjectId()
+FCbObjectId FCbObjectId::NewObjectId()
 {
-	// This is pretty ad hoc and does not work like the zen-side code
-	// But since this is used very sparingly it does not matter too
-	// much right now, we can fix it up later
+	using namespace UE::CompactBinary::Private;
 
-	FCbObjectId Oid;
+	FCbObjectId ObjectId;
+	uint32* Components = reinterpret_cast<uint32*>(ObjectId.Bytes);
 
-	FGuid Guid = FGuid::NewGuid();
+	const int64 Offset = 1'609'459'200; // Seconds from 1970 -> 2021
+	const int64 Time = FDateTime::Now().ToUnixTimestamp() - Offset;
 
-	uint8 HashBytes[16];
-	FXxHash128 Hash = FXxHash128::HashBuffer(&Guid, sizeof Guid);
-	Hash.ToByteArray(HashBytes);
+	uint32 RunId = GObjectIdRunId.load(std::memory_order_relaxed);
+	if (RunId == 0)
+	{
+		// Use a GUID as a stronger source of randomness than FMath::RandRange.
+		const FGuid Guid = FGuid::NewGuid();
+		const uint64 Hash = FXxHash64::HashBuffer(&Guid, sizeof(FGuid)).Hash;
+		uint32 ZeroSerial = 0;
+		GObjectIdSerial.compare_exchange_strong(ZeroSerial, uint32(Hash), std::memory_order_relaxed);
+		GObjectIdRunId.compare_exchange_strong(RunId, uint32(Hash >> 32), std::memory_order_relaxed);
+	}
 
-	memcpy(Oid.Bytes, HashBytes, sizeof Oid);
+	Components[0] = static_cast<uint32>(Time);
+	Components[1] = GObjectIdSerial.fetch_add(1, std::memory_order_relaxed);
+	Components[2] = RunId;
 
-	return Oid;
+	return ObjectId;
 }
 
-
 ///////////////////////////////////////////////////////////////////////////////////////////////////////////////////////
 
 FCbFieldView::FCbFieldView(const void* const Data, ECbFieldType Type)
