Changed YieldCycles to prevent reading hw cycle counter on ARM64 platforms that don't support user mode access to it

diff --git a/Engine/Source/Runtime/Core/Public/GenericPlatform/GenericPlatformProcess.h b/Engine/Source/Runtime/Core/Public/GenericPlatform/GenericPlatformProcess.h
--- a/Engine/Source/Runtime/Core/Public/GenericPlatform/GenericPlatformProcess.h
+++ b/Engine/Source/Runtime/Core/Public/GenericPlatform/GenericPlatformProcess.h
@@ -788,7 +788,8 @@ struct CORE_API FGenericPlatformProcess
 	*/
 	static FORCEINLINE void YieldCycles(uint64 cycles)
 	{
-#if PLATFORM_ANDROID
+#if PLATFORM_ANDROID_ARM64 || (PLATFORM_APPLE && !PLATFORM_MAC_X86) || PLATFORM_LINUXARM64
+		// We can't read cycle counter from user mode on these platform
 		for (uint64 i = 0; i < cycles; i++)
 		{
 			Yield();
