Tasks: removed invalid `verify` as it's triggered in an valid case:

Thread1 can call TaskEvent.Trigger(), which unlocks the task, and before it calls verify(TryExecute()), Thread2 starts TaskEvent retraction (because Thread2 started waiting on its subsequent task), and first it gets execution permission. then it proceeds to complete retraction as the task is unlocked but not executed yet. Now Thread1 tries to get execution permissions and fails, but can proceed as normal. its job is done

diff --git a/Engine/Source/Runtime/Core/Public/Tasks/TaskPrivate.h b/Engine/Source/Runtime/Core/Public/Tasks/TaskPrivate.h
--- a/Engine/Source/Runtime/Core/Public/Tasks/TaskPrivate.h
+++ b/Engine/Source/Runtime/Core/Public/Tasks/TaskPrivate.h
@@ -470,7 +470,7 @@ namespace UE { namespace Tasks
 				{
 					// execute before cancelling the low-level task as successful cancellation can release the last reference and destroy the task
 					// the low-level task wasn't scheduled, so successful execution and low-level task cancellation is guaranted
-					verify(TryExecute());
+					TryExecute();
 					verify(LowLevelTask.TryCancel());
 					return true;
 				}
