Fix for using alternative representation `#` modifier in FGenericWidePlatformString::GetVarArgs().

diff --git a/Engine/Source/Runtime/Core/Private/GenericPlatform/GenericWidePlatformString.cpp b/Engine/Source/Runtime/Core/Private/GenericPlatform/GenericWidePlatformString.cpp
--- a/Engine/Source/Runtime/Core/Private/GenericPlatform/GenericWidePlatformString.cpp
+++ b/Engine/Source/Runtime/Core/Private/GenericPlatform/GenericWidePlatformString.cpp
@@ -425,7 +425,7 @@ int32 FGenericWidePlatformString::GetVarArgs( WIDECHAR* Dest, SIZE_T DestSize, c
 
 		// Skip modifier flags that don't need additional processing;
 		// they still get passed to snprintf() below based on the conversion.
-		if (*Src == '+')
+		while (*Src == '+' || *Src == '#')
 		{
 			Src++;
 		}
