TaskGraph: fixed WaitUntilTasksComplete to tolerate null tasks

diff --git a/Engine/Source/Runtime/Core/Private/Async/TaskGraph.cpp b/Engine/Source/Runtime/Core/Private/Async/TaskGraph.cpp
--- a/Engine/Source/Runtime/Core/Private/Async/TaskGraph.cpp
+++ b/Engine/Source/Runtime/Core/Private/Async/TaskGraph.cpp
@@ -231,7 +231,10 @@ TArray<TaskTrace::FId> GetTraceIds(const FGraphEventArray& Tasks)
 
 	for (const FGraphEventRef& Task : Tasks)
 	{
-		TasksIds.Add(Task->GetTraceId());
+		if (Task.IsValid())
+		{
+			TasksIds.Add(Task->GetTraceId());
+		}
 	}
 
 	return TasksIds;
