Added possibility to disable checkpoints in hang detector by adding -nocheckpointhangdetector to cmd line

#rnx

diff --git a/Engine/Source/Runtime/Core/Private/HAL/ThreadHeartBeat.cpp b/Engine/Source/Runtime/Core/Private/HAL/ThreadHeartBeat.cpp
--- a/Engine/Source/Runtime/Core/Private/HAL/ThreadHeartBeat.cpp
+++ b/Engine/Source/Runtime/Core/Private/HAL/ThreadHeartBeat.cpp
@@ -685,7 +685,8 @@ uint32 FThreadHeartBeat::CheckCheckpointHeartBeat(double& OutHangDuration)
 {
 	// Editor and debug builds run too slow to measure them correctly
 #if USE_HANG_DETECTION
-	bool CheckBeats = IsEnabled();
+	static bool bCheckpointDisabled = FParse::Param(FCommandLine::Get(), TEXT("nocheckpointhangdetector"));
+	bool CheckBeats = IsEnabled() && !bCheckpointDisabled;
 
 	if (CheckBeats)
 	{
