Save ini file to temp name then move to avoid data loss

diff --git a/Engine/Source/Runtime/Core/Private/Misc/ConfigCacheIni.cpp b/Engine/Source/Runtime/Core/Private/Misc/ConfigCacheIni.cpp
--- a/Engine/Source/Runtime/Core/Private/Misc/ConfigCacheIni.cpp
+++ b/Engine/Source/Runtime/Core/Private/Misc/ConfigCacheIni.cpp
@@ -513,7 +513,17 @@ static bool SaveConfigFileWrapper(const TCHAR* IniFile, const FString& Contents)
 	FCoreDelegates::PreSaveConfigFileDelegate.Broadcast(IniFile, Contents, SavedCount);
 
 	// save it even if a delegate did as well
-	bool bLocalWriteSucceeded = FFileHelper::SaveStringToFile(Contents, IniFile, FFileHelper::EEncodingOptions::ForceUTF8WithoutBOM);
+	const FString BaseFilename = FPaths::GetBaseFilename(IniFile);
+	const FString TempFilename = FPaths::CreateTempFilename(*FPaths::ProjectSavedDir(), *BaseFilename.Left(32));
+	bool bLocalWriteSucceeded = FFileHelper::SaveStringToFile(Contents, *TempFilename, FFileHelper::EEncodingOptions::ForceUTF8WithoutBOM);
+	if (bLocalWriteSucceeded)
+	{
+		if (!IFileManager::Get().Move(IniFile, *TempFilename))
+		{
+			IFileManager::Get().Delete(*TempFilename);
+			bLocalWriteSucceeded = false;
+		}
+	}
 
 	// success is based on a delegate or file write working (or both)
 	return SavedCount > 0 || bLocalWriteSucceeded;
