Added EAsyncExecution::ThreadIfForkSafe
- Made FRunnable::GetThreadType public
- This is needed to allow Build Patch Service Installers to run on servers.

diff --git a/Engine/Source/Runtime/Core/Private/HAL/ThreadingBase.cpp b/Engine/Source/Runtime/Core/Private/HAL/ThreadingBase.cpp
--- a/Engine/Source/Runtime/Core/Private/HAL/ThreadingBase.cpp
+++ b/Engine/Source/Runtime/Core/Private/HAL/ThreadingBase.cpp
@@ -408,6 +408,11 @@ public:
 		FThreadManager::Get().RemoveThread(this);
 	}
 
+	virtual FRunnableThread::ThreadType GetThreadType() const override
+	{
+		return ThreadType::Fake;
+	}
+
 	virtual bool CreateInternal(FRunnable* InRunnable, const TCHAR* InThreadName,
 		uint32 InStackSize,
 		EThreadPriority InThreadPri, uint64 InThreadAffinityMask,
@@ -426,13 +431,6 @@ public:
 		}
 		return SingleThreadRunnable != nullptr;
 	}
-
-protected:
-
-	virtual FRunnableThread::ThreadType GetThreadType() const override
-	{
-		return ThreadType::Fake;
-	}
 };
 uint32 FFakeThread::ThreadIdCounter = 0xffff;
 
@@ -1419,6 +1417,11 @@ public:
 		Super::WaitForCompletion();
 	}
 
+	virtual FRunnableThread::ThreadType GetThreadType() const override
+	{
+		return ThreadType::Forkable;
+	}
+
 	virtual bool CreateInternal(FRunnable* InRunnable, const TCHAR* InThreadName, uint32 InStackSize, EThreadPriority InThreadPri, uint64 InThreadAffinityMask, EThreadCreateFlags InCreateFlags) override
 	{
 		checkf(FForkProcessHelper::SupportsMultithreadingPostFork(), TEXT("ForkableThreads should only be created when -PostForkThreading is enabled"));
@@ -1461,11 +1464,6 @@ protected:
 			RealThread = nullptr;
 		}
 	}
-
-	virtual FRunnableThread::ThreadType GetThreadType() const override
-	{
-		return ThreadType::Forkable;
-	}
 };
 
 //-------------------------------------------------------------------------------
diff --git a/Engine/Source/Runtime/Core/Public/Async/Async.h b/Engine/Source/Runtime/Core/Public/Async/Async.h
--- a/Engine/Source/Runtime/Core/Public/Async/Async.h
+++ b/Engine/Source/Runtime/Core/Public/Async/Async.h
@@ -14,6 +14,7 @@
 #include "Misc/CoreStats.h"
 #include "Misc/IQueuedWork.h"
 #include "Misc/QueuedThreadPool.h"
+#include "Misc/Fork.h"
 #include "Stats/Stats.h"
 #include "Templates/Function.h"
 
@@ -28,9 +29,12 @@ enum class EAsyncExecution
 	/** Execute in Task Graph on the main thread (for short running tasks). */
 	TaskGraphMainThread,
 
-	/** Execute in separate thread (for long running tasks). */
+	/** Execute in separate thread if supported (for long running tasks). */
 	Thread,
 
+	/** Execute in separate thread if supported or supported post fork (see FForkProcessHelper::CreateThreadIfForkSafe) (for long running tasks). */
+	ThreadIfForkSafe,
+
 	/** Execute in global queued thread pool. */
 	ThreadPool,
 
@@ -310,6 +314,27 @@ auto Async(EAsyncExecution Execution, CallableType&& Callable, TUniqueFunction<v
 			FRunnableThread* RunnableThread = FRunnableThread::Create(Runnable, *TAsyncThreadName);
 
 			check(RunnableThread != nullptr);
+			check(RunnableThread->GetThreadType() == FRunnableThread::ThreadType::Real);
+
+			ThreadPromise.SetValue(RunnableThread);
+		}
+		else
+		{
+			SetPromise(Promise, Function);
+		}
+		break;
+
+	case EAsyncExecution::ThreadIfForkSafe:
+		if (FPlatformProcess::SupportsMultithreading() || FForkProcessHelper::IsForkedMultithreadInstance())
+		{
+			TPromise<FRunnableThread*> ThreadPromise;
+			TAsyncRunnable<ResultType>* Runnable = new TAsyncRunnable<ResultType>(MoveTemp(Function), MoveTemp(Promise), ThreadPromise.GetFuture());
+
+			const FString TAsyncThreadName = FString::Printf(TEXT("TAsync %d"), FAsyncThreadIndex::GetNext());
+			FRunnableThread* RunnableThread = FForkProcessHelper::CreateForkableThread(Runnable, *TAsyncThreadName);
+
+			check(RunnableThread != nullptr);
+			check(RunnableThread->GetThreadType() == FRunnableThread::ThreadType::Real);
 
 			ThreadPromise.SetValue(RunnableThread);
 		}
diff --git a/Engine/Source/Runtime/Core/Public/Misc/AutomationTest.h b/Engine/Source/Runtime/Core/Public/Misc/AutomationTest.h
--- a/Engine/Source/Runtime/Core/Public/Misc/AutomationTest.h
+++ b/Engine/Source/Runtime/Core/Public/Misc/AutomationTest.h
@@ -506,19 +506,19 @@ public:
 	{
 		if (!Future.IsValid())
 		{
-			Future = Async(EAsyncExecution::Thread, Function);
+			Future = Async(EAsyncExecution::Thread, MoveTemp(Function));
 		}
 
 		return Future.IsReady();
 	}
 
-	FThreadedAutomationLatentCommand(TFunction<void()> InFunction)
-		: Function(InFunction)
+	FThreadedAutomationLatentCommand(TUniqueFunction<void()> InFunction)
+		: Function(MoveTemp(InFunction))
 	{ }
 
 protected:
 
-	const TFunction<void()> Function;
+	TUniqueFunction<void()> Function;
 
 	TFuture<void> Future;
 
