Should have been part of 16302102

diff --git a/Engine/Source/Runtime/Core/Public/HAL/RunnableThread.h b/Engine/Source/Runtime/Core/Public/HAL/RunnableThread.h
--- a/Engine/Source/Runtime/Core/Public/HAL/RunnableThread.h
+++ b/Engine/Source/Runtime/Core/Public/HAL/RunnableThread.h
@@ -88,6 +88,23 @@ public:
 	/** Halts the caller until this thread is has completed its work. */
 	virtual void WaitForCompletion() = 0;
 
+	/** List of unique thread types we can create */
+	enum class ThreadType
+	{
+		// Regular thread that executes the runnable object in it's own context
+		Real,
+		// Fake threads are created for a single threaded environment and are always executed from the main tick
+		Fake,
+		// Forkable threads will behave like fake threads for the master process, but will become real threads on forked processes
+		Forkable,
+	};
+
+	/** Returns the type of thread this is */
+	virtual FRunnableThread::ThreadType GetThreadType() const
+	{
+		return ThreadType::Real;
+	}
+
 	/**
 	 * Thread ID for this thread 
 	 *
@@ -171,19 +188,6 @@ protected:
 	/** ID set during thread creation. */
 	uint32 ThreadID;
 
-protected:
-
-	/** List of unique thread types we can create */
-	enum class ThreadType
-	{
-		// Regular thread that executes the runnable object in it's own context
-		Real,
-		// Fake threads are created for a single threaded environment and are always executed from the main tick
-		Fake,
-		// Forkable threads will behave like fake threads for the master process, but will become real threads on forked processes
-		Forkable,
-	};
-
 private:
 
 	/** Called to setup a newly created RunnableThread */
@@ -192,12 +196,6 @@ private:
 	/** Used by the thread manager to tick threads in single-threaded mode */
 	virtual void Tick() {}
 
-	/** Returns the type of thread this is */
-	virtual FRunnableThread::ThreadType GetThreadType() const
-	{
-		return ThreadType::Real;
-	}
-
 	/**
 	 * Called on the forked process when the forkable thread can create a real thread
 	 */
