Added CVar s.IoDispatcherSortRequestsByOffset to allow sorting io dispatcher request queue by file handle & offset instead of request age.
Added s.IoDispatcherMaxForwardSeekKB and s.IoDispatcherRequestLatencyCircuitBreakerMS as potential circuit breakers for the above behavior.
Added s.AsyncLoadingMaxPendingRequestsSizeMB to configure async loading read-ahead behavior.

diff --git a/Engine/Source/Runtime/Core/Private/GenericPlatform/GenericPlatformIoDispatcher.cpp b/Engine/Source/Runtime/Core/Private/GenericPlatform/GenericPlatformIoDispatcher.cpp
--- a/Engine/Source/Runtime/Core/Private/GenericPlatform/GenericPlatformIoDispatcher.cpp
+++ b/Engine/Source/Runtime/Core/Private/GenericPlatform/GenericPlatformIoDispatcher.cpp
@@ -159,6 +159,5 @@ bool FGenericFileIoStoreImpl::StartRequests(FFileIoStoreRequestQueue& RequestQue
 void FGenericFileIoStoreImpl::GetCompletedRequests(FFileIoStoreReadRequestList& OutRequests)
 {
 	FScopeLock _(&CompletedRequestsCritical);
-	OutRequests.Append(CompletedRequests);
-	CompletedRequests.Clear();
+	OutRequests.AppendSteal(CompletedRequests);
 }
diff --git a/Engine/Source/Runtime/Core/Private/IO/IoDispatcherFileBackend.cpp b/Engine/Source/Runtime/Core/Private/IO/IoDispatcherFileBackend.cpp
--- a/Engine/Source/Runtime/Core/Private/IO/IoDispatcherFileBackend.cpp
+++ b/Engine/Source/Runtime/Core/Private/IO/IoDispatcherFileBackend.cpp
@@ -12,11 +12,20 @@
 #include "HAL/RunnableThread.h"
 #include "Misc/ScopeLock.h"
 #include "Misc/Paths.h"
+#include "Algo/AllOf.h"
+#include "Algo/IsSorted.h"
+#include "Algo/MinElement.h"
+#include "Templates/Greater.h"
 
 TRACE_DECLARE_MEMORY_COUNTER(IoDispatcherTotalBytesRead, TEXT("IoDispatcher/TotalBytesRead"));
 TRACE_DECLARE_MEMORY_COUNTER(IoDispatcherTotalBytesScattered, TEXT("IoDispatcher/TotalBytesScattered"));
 TRACE_DECLARE_INT_COUNTER(IoDispatcherCacheHits, TEXT("IoDispatcher/CacheHits"));
 TRACE_DECLARE_INT_COUNTER(IoDispatcherCacheMisses, TEXT("IoDispatcher/CacheMisses"));
+TRACE_DECLARE_INT_COUNTER(IoDispatcherOutstandingReads, TEXT("IoDispatcher/OutstandingReads"));
+TRACE_DECLARE_MEMORY_COUNTER(IoDispatcherOutstandingBytesToRead, TEXT("IoDispatcher/OutstandingBytesToRead"));
+TRACE_DECLARE_INT_COUNTER(IoDispatcherLatencyCircuitBreaks, TEXT("IoDispatcher/LatencyCircuitBreaks"));
+TRACE_DECLARE_INT_COUNTER(IoDispatcherSeekDistanceCircuitBreaks, TEXT("IoDispatcher/SeekDistanceCircuitBreaks"));
+TRACE_DECLARE_INT_COUNTER(IoDispatcherNumPriorityQueues, TEXT("IoDispatcher/NumPriorityQueues"));
 
 //PRAGMA_DISABLE_OPTIMIZATION
 
@@ -55,7 +64,38 @@ static FAutoConsoleVariableRef CVar_IoDispatcherCacheSizeMB(
 	TEXT("IoDispatcher cache memory size (in megabytes).")
 );
 
+int32 GIoDispatcherSortRequestsByOffset = 1;
+static FAutoConsoleVariableRef CVar_IoDispatcherSortRequestsByOffset(
+	TEXT("s.IoDispatcherSortRequestsByOffset"),
+	GIoDispatcherSortRequestsByOffset,
+	TEXT("If > 0, io dispatcher sorts the outstanding request queue by offset rather than sequence.")
+);
+
+int32 GIoDispatcherMaintainSortingOnPriorityChange = 1;
+static FAutoConsoleVariableRef CVar_IoDispatcherMaintainSortingOnPriorityChange(
+	TEXT("s.IoDispatcherMaintainSortingOnPriorityChange"),
+	GIoDispatcherMaintainSortingOnPriorityChange,
+	TEXT("If s.IoDispatcherSortRequestsByOffset > 0 and this > 0, io dispatcher remembers the last file handle/offset read from even when switching priority levels.")
+);
+
+int32 GIoDispatcherMaxForwardSeekKB = 0;
+static FAutoConsoleVariableRef CVar_IoDispatcherMaxForwardSeekKB(
+	TEXT("s.IoDispatcherMaxForwardSeekKB"),
+	GIoDispatcherMaxForwardSeekKB,
+	TEXT("If s.IoDispatcherSortRequestsByOffset is enabled and this is > 0, if the next sequential read is further than this offset from the last one, read the oldest request instead")
+);
+
+int32 GIoDispatcherRequestLatencyCircuitBreakerMS = 0;
+static FAutoConsoleVariableRef CVar_IoDispatcherRequestLatencyCircuitBreakerMS(
+	TEXT("s.IoDispatcherRequestLatencyCircuitBreakerMS"),
+	GIoDispatcherRequestLatencyCircuitBreakerMS,
+	TEXT("If s.IoDispatcherSortRequestsByOffset is enabled and this is >0, if the oldest request has been in the queue for this long, read it instead of the most optimal read")
+);
+
 uint32 FFileIoStoreReadRequest::NextSequence = 0;
+#if CHECK_IO_STORE_READ_REQUEST_LIST_MEMBERSHIP
+uint32 FFileIoStoreReadRequestList::NextListCookie = 0;
+#endif
 TAtomic<uint32> FFileIoStoreReader::GlobalPartitionIndex{ 0 };
 TAtomic<uint32> FFileIoStoreReader::GlobalContainerInstanceId{ 0 };
 
@@ -240,52 +280,305 @@ void FFileIoStoreBlockCache::Store(const FFileIoStoreReadRequest* Block)
 	}
 }
 
+bool FFileIoStoreOffsetSortedRequestQueue::RequestSortPredicate(const FFileIoStoreReadRequestSortKey& A, const FFileIoStoreReadRequestSortKey& B)
+{
+	if (A.Handle == B.Handle)
+	{
+		return A.Offset < B.Offset;
+	}
+	return A.Handle < B.Handle;
+}
+
+FFileIoStoreOffsetSortedRequestQueue::FFileIoStoreOffsetSortedRequestQueue(int32 InPriority)
+	: Priority(InPriority)
+{
+}
+
+TArray<FFileIoStoreReadRequest*> FFileIoStoreOffsetSortedRequestQueue::StealRequests()
+{
+	RequestsBySequence.Clear();
+	PeekRequestIndex = INDEX_NONE;
+	return MoveTemp(Requests); 
+}
+
+// This could be potentially optimized if the higher level keeps track of which requests it changes the priority of, or even just the old priorty levels
+TArray<FFileIoStoreReadRequest*> FFileIoStoreOffsetSortedRequestQueue::RemoveMisprioritizedRequests()
+{
+	PeekRequestIndex = INDEX_NONE;
+	TArray<FFileIoStoreReadRequest*> RequestsToReturn;
+	for (int32 i = Requests.Num()-1; i >= 0; --i)
+	{
+		if (Requests[i]->Priority != Priority)
+		{
+			RequestsToReturn.Add(Requests[i]);
+			RequestsBySequence.Remove(Requests[i]);
+			Requests.RemoveAt(i, 1, false);
+		}
+	}
+
+	return RequestsToReturn;
+}
+
+FFileIoStoreReadRequest* FFileIoStoreOffsetSortedRequestQueue::GetNextInternal(FFileIoStoreReadRequestSortKey LastSortKey, bool bPop)
+{
+	if (Requests.Num() == 0)
+	{
+		return nullptr;
+	}
+
+	int32 RequestIndex = INDEX_NONE;
+	if (PeekRequestIndex != INDEX_NONE)
+	{
+		RequestIndex = PeekRequestIndex;
+	}
+	else 
+	{
+		bool bHeadRequestTooOld = false;
+		if (GIoDispatcherRequestLatencyCircuitBreakerMS > 0)
+		{
+			// If our oldest request has been unserviced for too long, grab that instead of the next sequential read
+			uint64 ThresholdCycles = uint64((GIoDispatcherRequestLatencyCircuitBreakerMS * 1000.0) / FPlatformTime::GetSecondsPerCycle64());
+			bHeadRequestTooOld = (FPlatformTime::Cycles64() - RequestsBySequence.PeekHead()->CreationTime) >= ThresholdCycles;
+
+			if (bPop)
+			{
+				TRACE_COUNTER_INCREMENT(IoDispatcherLatencyCircuitBreaks);
+			}
+		}
+
+		const bool bChooseByOffset = 
+				LastSortKey.Handle != 0 
+			&&	!bHeadRequestTooOld 
+			&&  (GIoDispatcherMaintainSortingOnPriorityChange || LastSortKey.Priority == Priority);
+		if (bChooseByOffset)
+		{
+			// Pick the request with the closest offset to the last thing that we read
+			RequestIndex = Algo::LowerBoundBy(Requests, LastSortKey, RequestSortProjection, RequestSortPredicate);
+			if (Requests.IsValidIndex(RequestIndex)) // If all requests are before LastOffset we get back out-of-bounds
+			{
+				if (Requests[RequestIndex]->FileHandle != LastSortKey.Handle)
+				{
+					// Changing file handle so switch back to the oldest outstanding request 
+					RequestIndex = INDEX_NONE;
+				}
+				else if (GIoDispatcherMaxForwardSeekKB > 0 && (LastSortKey.Offset - Requests[RequestIndex]->Offset) > uint64(GIoDispatcherMaxForwardSeekKB) * 1024)
+				{
+					// Large forward seek so switch back to the oldest outstanding request 
+					RequestIndex = INDEX_NONE;
+
+					if (bPop)
+					{
+						TRACE_COUNTER_INCREMENT(IoDispatcherSeekDistanceCircuitBreaks);
+					}
+				}
+			}
+		}
+
+		if (!Requests.IsValidIndex(RequestIndex))
+		{
+			RequestIndex = Requests.Find(RequestsBySequence.PeekHead());
+			check(Requests[RequestIndex] == RequestsBySequence.PeekHead());
+		}
+	}
+
+	check(Requests.IsValidIndex(RequestIndex));
+
+	FFileIoStoreReadRequest* Request = Requests[RequestIndex];
+	if (bPop)
+	{
+		Requests.RemoveAt(RequestIndex);
+		RequestsBySequence.Remove(Request);
+		PeekRequestIndex = INDEX_NONE;
+	}
+	else
+	{
+		PeekRequestIndex = RequestIndex;
+	}
+	return Request;
+}
+
+FFileIoStoreReadRequest* FFileIoStoreOffsetSortedRequestQueue::Peek(FFileIoStoreReadRequestSortKey LastSortKey)
+{
+	return GetNextInternal(LastSortKey, false);
+}
+
+FFileIoStoreReadRequest* FFileIoStoreOffsetSortedRequestQueue::Pop(FFileIoStoreReadRequestSortKey LastSortKey)
+{
+	return GetNextInternal(LastSortKey, true);
+}
+
+void FFileIoStoreOffsetSortedRequestQueue::Push(FFileIoStoreReadRequest* Request)
+{
+	// Insert sorted by file handle & offset
+	int32 InsertIndex = Algo::UpperBoundBy(Requests, RequestSortProjection(Request), RequestSortProjection, RequestSortPredicate);
+	Requests.Insert(Request, InsertIndex);
+	
+	// Insert sorted by age
+	RequestsBySequence.Add(Request);
+
+	PeekRequestIndex = INDEX_NONE;
+}
+
+void FFileIoStoreRequestQueue::UpdateSortRequestsByOffset()
+{
+	// Must hold CriticalSection here
+	if (bSortRequestsByOffset == bool(GIoDispatcherSortRequestsByOffset))
+	{
+		return;
+	}
+
+	bSortRequestsByOffset = bool(GIoDispatcherSortRequestsByOffset);
+	if (bSortRequestsByOffset)
+	{
+		// Split things into separate heaps
+		for (FFileIoStoreReadRequest* Request : Heap)
+		{
+			Push(*Request);
+		}
+		Heap.Empty();
+	}
+	else
+	{
+		// Put things back into the main heap
+		TArray< FFileIoStoreReadRequest*> AllRequests;
+		for (FFileIoStoreOffsetSortedRequestQueue& SubQueue : SortedPriorityQueues)
+		{
+			AllRequests.Append(SubQueue.StealRequests());
+		}
+		Algo::SortBy(AllRequests, [](FFileIoStoreReadRequest* Request) { return Request->Sequence; });
+		for (FFileIoStoreReadRequest* Request : AllRequests)
+		{
+			Push(*Request);
+		}
+		check(Algo::AllOf(SortedPriorityQueues, [](const FFileIoStoreOffsetSortedRequestQueue& Q) { return Q.IsEmpty(); }));
+		SortedPriorityQueues.Empty();
+	}
+}
+
 FFileIoStoreReadRequest* FFileIoStoreRequestQueue::Peek()
 {
 	TRACE_CPUPROFILER_EVENT_SCOPE(RequestQueuePeek);
 	FScopeLock _(&CriticalSection);
-	if (Heap.Num() == 0)
+	UpdateSortRequestsByOffset();
+	if (bSortRequestsByOffset)
 	{
-		return nullptr;
+		if( SortedPriorityQueues.Num() == 0)
+		{
+			return nullptr;
+		}
+
+		FFileIoStoreOffsetSortedRequestQueue& SubQueue = SortedPriorityQueues.Last();
+		check(!SubQueue.IsEmpty());
+		FFileIoStoreReadRequest* Request = SubQueue.Peek(LastSortKey);
+		check(Request);
+		// Do not update LastUsedPriority here until we actually pop
+		return Request;
+	}
+	else
+	{
+		if (Heap.Num() == 0)
+		{
+			return nullptr;
+		}
+		return Heap.HeapTop();
 	}
-	return Heap.HeapTop();
 }
 
 FFileIoStoreReadRequest* FFileIoStoreRequestQueue::Pop()
 {
-	//TRACE_CPUPROFILER_EVENT_SCOPE(RequestQueuePop);
+	TRACE_CPUPROFILER_EVENT_SCOPE(RequestQueuePop);
 	FScopeLock _(&CriticalSection);
-	if (Heap.Num() == 0)
+	UpdateSortRequestsByOffset();
+	FFileIoStoreReadRequest* Result = nullptr;
+	if (bSortRequestsByOffset)
 	{
-		return nullptr;
+		if (SortedPriorityQueues.Num() == 0)
+		{
+			return nullptr;
+		}
+
+		FFileIoStoreOffsetSortedRequestQueue& SubQueue = SortedPriorityQueues.Last();
+		check(!SubQueue.IsEmpty());
+		Result = SubQueue.Pop(LastSortKey);
+		check(Result);
+		LastSortKey = Result;
+		if (SubQueue.IsEmpty())
+		{
+			SortedPriorityQueues.Pop();
+			// SubQueue is invalid here
+			TRACE_COUNTER_DECREMENT(IoDispatcherNumPriorityQueues);
+		}
+	}
+	else
+	{
+		if (Heap.Num() == 0)
+		{
+			return nullptr;
+		}
+		Heap.HeapPop(Result, QueueSortFunc, false);
 	}
-	FFileIoStoreReadRequest* Result;
-	Heap.HeapPop(Result, QueueSortFunc, false);
+	
 	check(Result->QueueStatus == FFileIoStoreReadRequest::QueueStatus_InQueue);
 	Result->QueueStatus = FFileIoStoreReadRequest::QueueStatus_Started;
 	return Result;
 }
 
+void FFileIoStoreRequestQueue::PushToPriorityQueues(FFileIoStoreReadRequest* Request)
+{
+	int32 QueueIndex = Algo::LowerBoundBy(SortedPriorityQueues, Request->Priority, QueuePriorityProjection, TLess<int32>());
+	if (!SortedPriorityQueues.IsValidIndex(QueueIndex) || SortedPriorityQueues[QueueIndex].GetPriority() != Request->Priority)
+	{
+		SortedPriorityQueues.Insert(FFileIoStoreOffsetSortedRequestQueue(Request->Priority), QueueIndex);
+		TRACE_COUNTER_INCREMENT(IoDispatcherNumPriorityQueues);
+	}
+	check(Algo::IsSortedBy(SortedPriorityQueues, QueuePriorityProjection, TLess<int32>()));
+	FFileIoStoreOffsetSortedRequestQueue& Queue = SortedPriorityQueues[QueueIndex];
+	check(Queue.GetPriority() == Request->Priority);
+	Queue.Push(Request);
+}
+
 void FFileIoStoreRequestQueue::Push(FFileIoStoreReadRequest& Request)
 {
-	//TRACE_CPUPROFILER_EVENT_SCOPE(RequestQueuePush);
+	TRACE_CPUPROFILER_EVENT_SCOPE(RequestQueuePush);
 	FScopeLock _(&CriticalSection);
+	UpdateSortRequestsByOffset();
+	TRACE_COUNTER_INCREMENT(IoDispatcherOutstandingReads);
+	TRACE_COUNTER_ADD(IoDispatcherOutstandingBytesToRead, Request.Size);
+	
 	check(Request.QueueStatus != FFileIoStoreReadRequest::QueueStatus_InQueue);
 	Request.QueueStatus = FFileIoStoreReadRequest::QueueStatus_InQueue;
-	Heap.HeapPush(&Request, QueueSortFunc);
+	if (bSortRequestsByOffset)
+	{
+		PushToPriorityQueues(&Request);
+	}
+	else
+	{
+		Heap.HeapPush(&Request, QueueSortFunc);
+	}
 }
 
-void FFileIoStoreRequestQueue::Push(const FFileIoStoreReadRequestList& Requests)
+void FFileIoStoreRequestQueue::Push(FFileIoStoreReadRequestList& Requests)
 {
-	//TRACE_CPUPROFILER_EVENT_SCOPE(RequestQueuePush);
+	TRACE_CPUPROFILER_EVENT_SCOPE(RequestQueuePush);
 	FScopeLock _(&CriticalSection);
-	FFileIoStoreReadRequest* Request = Requests.GetHead();
-	while (Request)
+	UpdateSortRequestsByOffset();
+
+	for (auto It = Requests.Steal(); It; ++It)
 	{
-		check(Request->QueueStatus != FFileIoStoreReadRequest::QueueStatus_InQueue);
-		Request->QueueStatus = FFileIoStoreReadRequest::QueueStatus_InQueue;
-		Heap.HeapPush(Request, QueueSortFunc);
-		Request = Request->Next;
+		TRACE_COUNTER_INCREMENT(IoDispatcherOutstandingReads);
+		TRACE_COUNTER_ADD(IoDispatcherOutstandingBytesToRead, It->Size);
+
+		check(It->QueueStatus != FFileIoStoreReadRequest::QueueStatus_InQueue);
+		It->QueueStatus = FFileIoStoreReadRequest::QueueStatus_InQueue;
+
+		if (bSortRequestsByOffset)
+		{
+			PushToPriorityQueues(*It);
+		}
+		else
+		{
+			Heap.HeapPush(*It, QueueSortFunc);
+		}
 	}
 }
 
@@ -293,7 +586,26 @@ void FFileIoStoreRequestQueue::UpdateOrder()
 {
 	TRACE_CPUPROFILER_EVENT_SCOPE(RequestQueueUpdateOrder);
 	FScopeLock _(&CriticalSection);
-	Heap.Heapify(QueueSortFunc);
+	UpdateSortRequestsByOffset();
+	if (bSortRequestsByOffset)
+	{
+		TArray<FFileIoStoreReadRequest*> Requests;
+		for (FFileIoStoreOffsetSortedRequestQueue& SubQueue : SortedPriorityQueues)
+		{
+			TArray<FFileIoStoreReadRequest*> RequestsRemoved = SubQueue.RemoveMisprioritizedRequests();
+			Requests.Append(RequestsRemoved);
+		}
+
+		Algo::SortBy(Requests, [](FFileIoStoreReadRequest* Request) { return Request->Sequence;  });
+		for (FFileIoStoreReadRequest* Request : Requests)
+		{
+			PushToPriorityQueues(Request);
+		}
+	}
+	else
+	{
+		Heap.Heapify(QueueSortFunc);
+	}
 }
 
 void FFileIoStoreRequestQueue::Lock()
@@ -547,18 +859,16 @@ void FFileIoStoreRequestTracker::AddReadRequestsToResolvedRequest(FFileIoStoreCo
 	}
 }
 
-void FFileIoStoreRequestTracker::AddReadRequestsToResolvedRequest(const FFileIoStoreReadRequestList& Requests, FFileIoStoreResolvedRequest& ResolvedRequest)
+void FFileIoStoreRequestTracker::AddReadRequestsToResolvedRequest(FFileIoStoreReadRequestList& Requests, FFileIoStoreResolvedRequest& ResolvedRequest)
 {
 	//TRACE_CPUPROFILER_EVENT_SCOPE(RequestTrackerAddIoRequest);
-	FFileIoStoreReadRequest* ReadRequest = Requests.GetHead();
-	while (ReadRequest)
+	for (auto It = Requests.Steal(); It; ++It)
 	{
 		++ResolvedRequest.UnfinishedReadsCount;
-		FFileIoStoreReadRequestLink* Link = RequestAllocator.AllocRequestLink(ReadRequest);
-		++ReadRequest->RefCount;
+		FFileIoStoreReadRequestLink* Link = RequestAllocator.AllocRequestLink(*It);
+		++It->RefCount;
 		ResolvedRequest.AddReadRequestLink(Link);
-		check(ResolvedRequest.GetPriority() == ReadRequest->Priority);
-		ReadRequest = ReadRequest->Next;
+		check(ResolvedRequest.GetPriority() == It->Priority);
 	}
 }
 
@@ -1111,12 +1421,13 @@ FIoRequestImpl* FFileIoStore::GetCompletedRequests()
 
 	FFileIoStoreReadRequestList CompletedRequests;
 	PlatformImpl.GetCompletedRequests(CompletedRequests);
-	FFileIoStoreReadRequest* CompletedRequest = CompletedRequests.GetHead();
-	while (CompletedRequest)
+	for (auto It = CompletedRequests.Steal(); It; ++It)
 	{
-		FFileIoStoreReadRequest* NextRequest = CompletedRequest->Next;
+		FFileIoStoreReadRequest* CompletedRequest = *It;
 
 		TRACE_COUNTER_ADD(IoDispatcherTotalBytesRead, CompletedRequest->Size);
+		TRACE_COUNTER_DECREMENT(IoDispatcherOutstandingReads);
+		TRACE_COUNTER_SUBTRACT(IoDispatcherOutstandingBytesToRead, CompletedRequest->Size);
 
 		if (!CompletedRequest->ImmediateScatter.Request)
 		{
@@ -1206,8 +1517,6 @@ FIoRequestImpl* FFileIoStore::GetCompletedRequests()
 				RequestTracker.ReleaseIoRequestReferences(*CompletedResolvedRequest);
 			}
 		}
-		
-		CompletedRequest = NextRequest;
 	}
 	
 	FFileIoStoreCompressedBlock* BlockToReap;
diff --git a/Engine/Source/Runtime/Core/Private/IO/IoDispatcherFileBackend.h b/Engine/Source/Runtime/Core/Private/IO/IoDispatcherFileBackend.h
--- a/Engine/Source/Runtime/Core/Private/IO/IoDispatcherFileBackend.h
+++ b/Engine/Source/Runtime/Core/Private/IO/IoDispatcherFileBackend.h
@@ -83,7 +83,8 @@ public:
 	FFileIoStoreReadRequest* FindOrAddRawBlock(FFileIoStoreBlockKey Key, bool& bOutWasAdded);
 	void RemoveRawBlock(const FFileIoStoreReadRequest* RawBlock);
 	void AddReadRequestsToResolvedRequest(FFileIoStoreCompressedBlock* CompressedBlock, FFileIoStoreResolvedRequest& ResolvedRequest);
-	void AddReadRequestsToResolvedRequest(const FFileIoStoreReadRequestList& Requests, FFileIoStoreResolvedRequest& ResolvedRequest);
+	// This function takes ownership of the requests in the request list and empties it because of the intrusive linked list
+	void AddReadRequestsToResolvedRequest(FFileIoStoreReadRequestList& Requests, FFileIoStoreResolvedRequest& ResolvedRequest);
 	bool CancelIoRequest(FFileIoStoreResolvedRequest& ResolvedRequest);
 	void UpdatePriorityForIoRequest(FFileIoStoreResolvedRequest& ResolvedRequest);
 	void ReleaseIoRequestReferences(FFileIoStoreResolvedRequest& ResolvedRequest);
diff --git a/Engine/Source/Runtime/Core/Private/IO/IoDispatcherFileBackendTypes.h b/Engine/Source/Runtime/Core/Private/IO/IoDispatcherFileBackendTypes.h
--- a/Engine/Source/Runtime/Core/Private/IO/IoDispatcherFileBackendTypes.h
+++ b/Engine/Source/Runtime/Core/Private/IO/IoDispatcherFileBackendTypes.h
@@ -109,10 +109,11 @@ struct FFileIoStoreReadRequest
 
 	FFileIoStoreReadRequest()
 		: Sequence(NextSequence++)
+		, CreationTime(FPlatformTime::Cycles64())
 	{
-
 	}
 	FFileIoStoreReadRequest* Next = nullptr;
+	FFileIoStoreReadRequest* Previous = nullptr;
 	uint64 FileHandle = uint64(-1);
 	uint64 Offset = uint64(-1);
 	uint64 Size = uint64(-1);
@@ -123,80 +124,309 @@ struct FFileIoStoreReadRequest
 	TArray<FFileIoStoreCompressedBlock*, TInlineAllocator<8>> CompressedBlocks;
 	const uint32 Sequence;
 	int32 Priority = 0;
+	uint64 CreationTime;	// Potentially used to circuit break request ordering optimizations when outstanding requests have been delayed too long
 	FFileIoStoreBlockScatter ImmediateScatter;
 	bool bIsCacheable = false;
 	bool bFailed = false;
 	bool bCancelled = false;
 	EQueueStatus QueueStatus = QueueStatus_NotInQueue;
 
+#if DO_CHECK
+	// For debug checks that we are in the correct owning list for our intrusive next/previous pointers
+	uint32 ListCookie = 0;
+#endif
+
 private:
 	static uint32 NextSequence;
 };
 
+// Iterator class for traversing and emptying a list FFileIoStoreReadRequestList at the same time
+class FFileIoStoreReadRequestListStealingIterator
+{
+public:
+	FFileIoStoreReadRequestListStealingIterator(const FFileIoStoreReadRequestListStealingIterator&) = delete;
+	FFileIoStoreReadRequestListStealingIterator& operator=(const FFileIoStoreReadRequestListStealingIterator&) = delete;
+
+	FFileIoStoreReadRequestListStealingIterator(FFileIoStoreReadRequestListStealingIterator&& Other)
+	{
+		Current = Other.Current;
+		Next = Other.Next;
+
+		Other.Current = Other.Next = nullptr;
+	}
+
+	FFileIoStoreReadRequestListStealingIterator& operator=(FFileIoStoreReadRequestListStealingIterator&& Other)
+	{
+		Current = Other.Current;
+		Next = Other.Next;
+
+		Other.Current = Other.Next = nullptr;
+
+		return *this;
+	}
+
+	void operator++()
+	{
+		AdvanceTo(Next);
+	}
+
+	FFileIoStoreReadRequest* operator*()
+	{
+		return Current;
+	}
+
+	FFileIoStoreReadRequest* operator->()
+	{
+		return Current;
+	}
+
+	explicit operator bool() const
+	{
+		return Current != nullptr;
+	}
+
+private:
+	friend class FFileIoStoreReadRequestList; // Only the list can construct us
+
+	FFileIoStoreReadRequestListStealingIterator(FFileIoStoreReadRequest* InHead)
+	{
+#if DO_CHECK
+		// None of these nodes are members of the list any more
+		for (FFileIoStoreReadRequest* Cursor = InHead; Cursor; Cursor = Cursor->Next)
+		{
+			Cursor->ListCookie = 0;
+		}
+#endif
+
+		AdvanceTo(InHead);
+	}
+
+	void AdvanceTo(FFileIoStoreReadRequest* NewCurrent)
+	{
+		Current = NewCurrent;
+		if (Current)
+		{
+			// Copy off the next ptr and remove Current from its list so it can safely be added to another list
+			Next = Current->Next;
+			Current->Next = nullptr;
+		}
+		else
+		{
+			Next = nullptr;
+		}
+	}
+
+	FFileIoStoreReadRequest* Current = nullptr;
+	FFileIoStoreReadRequest* Next = nullptr;;
+};
+
+// Wrapper for doubly-linked intrusive list of FFileIoStoreReadRequest
+#define CHECK_IO_STORE_READ_REQUEST_LIST_MEMBERSHIP (DO_CHECK && 0)
 class FFileIoStoreReadRequestList
 {
 public:
+	FFileIoStoreReadRequestList()
+#if CHECK_IO_STORE_READ_REQUEST_LIST_MEMBERSHIP
+		: ListCookie(++NextListCookie)
+#endif
+	{
+	}
+
+	// Owns the Next/Previous pointers of the FFileIoStoreReadRequests it contains so it can't be copied
+	FFileIoStoreReadRequestList(const FFileIoStoreReadRequestList& Other) = delete;
+	FFileIoStoreReadRequestList& operator=(const FFileIoStoreReadRequestList& Other) = delete;
+
+	FFileIoStoreReadRequestList(FFileIoStoreReadRequestList&& Other)
+	{
+		Head = Other.Head;
+		Tail = Other.Tail;
+		Other.Head = Other.Tail = nullptr;
+
+#if CHECK_IO_STORE_READ_REQUEST_LIST_MEMBERSHIP
+		for (FFileIoStoreReadRequest* Cursor = Head; Cursor; Cursor = Cursor->Next)
+		{
+			Cursor->ListCookie = ListCookie;
+		}
+#endif
+	}
+
+	FFileIoStoreReadRequestList& operator=(FFileIoStoreReadRequestList&& Other)
+	{
+		check(!Head && !Tail);
+
+		Head = Other.Head;
+		Tail = Other.Tail;
+		Other.Head = Other.Tail = nullptr;
+
+#if CHECK_IO_STORE_READ_REQUEST_LIST_MEMBERSHIP
+		for (FFileIoStoreReadRequest* Cursor = Head; Cursor; Cursor = Cursor->Next)
+		{
+			Cursor->ListCookie = ListCookie;
+		}
+#endif
+
+		return *this;
+	}
+
+#if CHECK_IO_STORE_READ_REQUEST_LIST_MEMBERSHIP
+	~FFileIoStoreReadRequestList()
+	{
+		// If anything is left in this list it should still think we own it
+		for (FFileIoStoreReadRequest* Cursor = Head; Cursor; Cursor = Cursor->Next)
+		{
+			check(Cursor->ListCookie == ListCookie);
+		}
+	}
+#endif
+
 	bool IsEmpty() const
 	{
 		return Head == nullptr;
 	}
 
-	FFileIoStoreReadRequest* GetHead() const
+	// Steal the whole list for iteration and moving the contents to other lists
+	FFileIoStoreReadRequestListStealingIterator Steal()
 	{
-		return Head;
+#if CHECK_IO_STORE_READ_REQUEST_LIST_MEMBERSHIP
+		for (FFileIoStoreReadRequest* Cursor = Head; Cursor; Cursor = Cursor->Next)
+		{
+			check(Cursor->ListCookie == ListCookie);
+		}
+#endif
+
+		FFileIoStoreReadRequest* OldHead = Head;
+		Clear();
+		return FFileIoStoreReadRequestListStealingIterator(OldHead);
 	}
 
-	FFileIoStoreReadRequest* GetTail() const
+	FFileIoStoreReadRequest* PeekHead() const
 	{
-		return Tail;
+		return Head;
 	}
 
 	void Add(FFileIoStoreReadRequest* Request)
 	{
+#if CHECK_IO_STORE_READ_REQUEST_LIST_MEMBERSHIP
+		check(Request->ListCookie == 0);
+		Request->ListCookie = ListCookie;
+#endif
+
 		if (Tail)
 		{
 			Tail->Next = Request;
+			Request->Previous = Tail;
 		}
 		else
 		{
 			Head = Request;
+			Request->Previous = nullptr;
 		}
 		Tail = Request;
 		Request->Next = nullptr;
 	}
 
-	void Append(FFileIoStoreReadRequest* ListHead, FFileIoStoreReadRequest* ListTail)
+
+	// Remove all FFileIoStoreReadRequests from List and add them to this list
+	void AppendSteal(FFileIoStoreReadRequestList& List)
 	{
-		check(ListHead);
-		check(ListTail);
-		check(!ListTail->Next);
-		if (Tail)
-		{
-			Tail->Next = ListHead;
-		}
-		else
+		if (List.Head)
 		{
-			Head = ListHead;
+			FFileIoStoreReadRequest *ListHead = List.Head, *ListTail = List.Tail;
+			List.Clear();
+			AppendSteal(ListHead, ListTail);
 		}
-		Tail = ListTail;
 	}
 
-	void Append(FFileIoStoreReadRequestList& List)
+	void Remove(FFileIoStoreReadRequest* Request)
 	{
-		if (List.Head)
+#if CHECK_IO_STORE_READ_REQUEST_LIST_MEMBERSHIP
+		check(Request);
+		check(Request->ListCookie == ListCookie);
+		Request->ListCookie = 0;
+#endif
+
+		if (Head == Request && Tail == Request)
+		{
+			check(Request->Next == nullptr);
+			check(Request->Previous == nullptr);
+
+			Head = Tail = nullptr;
+		}
+		else if (Head == Request)
 		{
-			Append(List.Head, List.Tail);
+			check(Request->Previous == nullptr);
+
+			Head = Request->Next;
+			Head->Previous = nullptr;
+			Request->Next = nullptr;
+		}
+		else if (Tail == Request)
+		{
+			check(Request->Next == nullptr);
+
+			Tail = Request->Previous;
+			Tail->Next = nullptr;
+			Request->Previous = nullptr;
+		}
+		else
+		{
+			check(Request->Next != nullptr && Request->Previous != nullptr); // Neither head nor tail should mean both links are live
+
+			Request->Next->Previous = Request->Previous;
+			Request->Previous->Next = Request->Next;
+
+			Request->Next = Request->Previous = nullptr;
 		}
 	}
 
 	void Clear()
 	{
+#if CHECK_IO_STORE_READ_REQUEST_LIST_MEMBERSHIP
+		for (FFileIoStoreReadRequest* Cursor = Head; Cursor; Cursor = Cursor->Next)
+		{
+			check(Cursor->ListCookie == ListCookie);
+			Cursor->ListCookie = 0;
+		}
+#endif
 		Head = Tail = nullptr;
 	}
 
 private:
 	FFileIoStoreReadRequest* Head = nullptr;
 	FFileIoStoreReadRequest* Tail = nullptr;
+
+	uint32 ListCookie;
+	static uint32 NextListCookie;
+	
+	void AppendSteal(FFileIoStoreReadRequest* ListHead, FFileIoStoreReadRequest* ListTail)
+	{
+		check(ListHead);
+		check(ListTail);
+		check(!ListTail->Next);
+		check(!ListHead->Previous);
+		check(ListTail == ListHead || ListTail->Previous != nullptr);
+		check(ListTail == ListHead || ListHead->Next != nullptr);
+
+#if CHECK_IO_STORE_READ_REQUEST_LIST_MEMBERSHIP
+		for (FFileIoStoreReadRequest* Cursor = ListHead; Cursor; Cursor = Cursor->Next)
+		{
+			check(Cursor->ListCookie == 0);
+			Cursor->ListCookie = ListCookie;
+		}
+#endif
+
+		if (Tail)
+		{
+			Tail->Next = ListHead;
+			ListHead->Previous = Tail;
+		}
+		else
+		{
+			Head = ListHead;
+			ListHead->Previous = nullptr;
+		}
+		Tail = ListTail;
+	}
 };
 
 class FFileIoStoreBufferAllocator
@@ -240,13 +470,65 @@ private:
 	uint64 ReadBufferSize = 0;
 };
 
+struct FFileIoStoreReadRequestSortKey
+{
+	uint64 Offset = 0;
+	uint64 Handle = 0;
+	int32 Priority = 0;
+
+	FFileIoStoreReadRequestSortKey() {}
+	FFileIoStoreReadRequestSortKey(FFileIoStoreReadRequest* Request)
+		: Offset(Request->Offset), Handle(Request->FileHandle), Priority(Request->Priority)
+	{
+	}
+};
+
+// Stores FFileIoStoreReadRequest sorted by file handle & offset with a parallel list sorted by insertion order 
+class FFileIoStoreOffsetSortedRequestQueue
+{
+public:
+	FFileIoStoreOffsetSortedRequestQueue(int32 InPriority);
+	FFileIoStoreOffsetSortedRequestQueue(const FFileIoStoreOffsetSortedRequestQueue&) = delete;
+	FFileIoStoreOffsetSortedRequestQueue(FFileIoStoreOffsetSortedRequestQueue&&) = default;
+	FFileIoStoreOffsetSortedRequestQueue& operator=(const FFileIoStoreOffsetSortedRequestQueue&) = delete;
+	FFileIoStoreOffsetSortedRequestQueue& operator=(FFileIoStoreOffsetSortedRequestQueue&&) = default;
+
+	int32 GetPriority() const { return Priority; }
+	bool IsEmpty() const { return Requests.Num() == 0; }
+	
+	// Remove all requests from this container and return them, for switching to a different queue scheme
+	TArray<FFileIoStoreReadRequest*> StealRequests();
+	// Remove all requests whose priority has been changed to something other than the Priority of this queue
+	TArray<FFileIoStoreReadRequest*> RemoveMisprioritizedRequests();
+
+	FFileIoStoreReadRequest* Peek(FFileIoStoreReadRequestSortKey LastSortKey);
+	FFileIoStoreReadRequest* Pop(FFileIoStoreReadRequestSortKey LastSortKey);
+	void Push(FFileIoStoreReadRequest* Request);
+
+private:
+	int32 Priority;
+	int32 PeekRequestIndex = INDEX_NONE;
+	
+	// Requests sorted by file handle & offset
+	TArray<FFileIoStoreReadRequest*> Requests;
+	
+	// Requests sorted by insertion order 
+	// We store this on the heap in case we get moved, FFileIoStoreReadRequest keeps pointers to FFileIoStoreReadRequestList for debugging
+	FFileIoStoreReadRequestList RequestsBySequence;
+
+	FFileIoStoreReadRequest* GetNextInternal(FFileIoStoreReadRequestSortKey LastSortKey, bool bPop);
+
+	static FFileIoStoreReadRequestSortKey RequestSortProjection(FFileIoStoreReadRequest* Request) { return FFileIoStoreReadRequestSortKey(Request); }
+	static bool RequestSortPredicate(const FFileIoStoreReadRequestSortKey& A, const FFileIoStoreReadRequestSortKey& B);
+};
+
 class FFileIoStoreRequestQueue
 {
 public:
 	FFileIoStoreReadRequest* Peek();
 	FFileIoStoreReadRequest* Pop();
-	void Push(FFileIoStoreReadRequest& Request);
-	void Push(const FFileIoStoreReadRequestList& Requests);
+	void Push(FFileIoStoreReadRequest& Request);	// Takes ownership of Request and rewrites its intrustive linked list pointers
+	void Push(FFileIoStoreReadRequestList& Requests); // Consumes the request list and overwrites all intrustive linked list pointers
 	void UpdateOrder();
 	void Lock();
 	void Unlock();
@@ -261,9 +543,24 @@ private:
 		}
 		return A.Priority > B.Priority;
 	}
+	void UpdateSortRequestsByOffset(); // Check if we need to switch sorting schemes based on the CVar
+	void PushToPriorityQueues(FFileIoStoreReadRequest* Request);
+	static int32 QueuePriorityProjection(const FFileIoStoreOffsetSortedRequestQueue& A) { return A.GetPriority(); }
+
+	bool bSortRequestsByOffset = false; // Cached value of CVar controlling whether we use Heap or SortedPriorityQueues
 	
+	// Heap sorted by request order
 	TArray<FFileIoStoreReadRequest*> Heap;
 	FCriticalSection CriticalSection;
+
+	// Queues sorted by increasing priority
+	TArray<FFileIoStoreOffsetSortedRequestQueue> SortedPriorityQueues; 
+	// The last offset, file handle and priority we popped so that we can pop the closest forward read for the next IO operation
+	FFileIoStoreReadRequestSortKey LastSortKey;
+
+#if !UE_BUILD_SHIPPING
+	TMap<int32, uint32> RequestPriorityCounts;
+#endif
 };
 
 template <typename T, uint16 SlabSize = 4096>
diff --git a/Engine/Source/Runtime/CoreUObject/Private/Serialization/AsyncLoading2.cpp b/Engine/Source/Runtime/CoreUObject/Private/Serialization/AsyncLoading2.cpp
--- a/Engine/Source/Runtime/CoreUObject/Private/Serialization/AsyncLoading2.cpp
+++ b/Engine/Source/Runtime/CoreUObject/Private/Serialization/AsyncLoading2.cpp
@@ -66,6 +66,14 @@
 #if UE_BUILD_DEVELOPMENT || UE_BUILD_DEBUG
 PRAGMA_DISABLE_OPTIMIZATION
 #endif
+int32 GAsyncLoadingMaxPendingRequestsSizeMB = 256;
+static FAutoConsoleVariableRef CVar_AsyncLoadingMaxPendingRequestsSizeMB(
+	TEXT("s.AsyncLoadingMaxPendingRequestsSizeMB"),
+	GAsyncLoadingMaxPendingRequestsSizeMB,
+	TEXT("Max amount of data in MB to request from IO"),
+	ECVF_Default
+);
+
 
 FArchive& operator<<(FArchive& Ar, FContainerHeader& ContainerHeader)
 {
@@ -2592,7 +2600,7 @@ void FAsyncLoadingThread2::BundleIoRequestCompleted(FAsyncPackage2* Package)
 void FAsyncLoadingThread2::StartBundleIoRequests()
 {
 	TRACE_CPUPROFILER_EVENT_SCOPE(StartBundleIoRequests);
-	constexpr uint64 MaxPendingRequestsSize = 256 << 20;
+	const uint64 MaxPendingRequestsSize = GAsyncLoadingMaxPendingRequestsSizeMB << 20;
 	FIoBatch IoBatch = IoDispatcher.NewBatch();
 	while (WaitingIoRequests.Num())
 	{
