Safely save ini files only when not unattended

diff --git a/Engine/Source/Runtime/Core/Private/Misc/ConfigCacheIni.cpp b/Engine/Source/Runtime/Core/Private/Misc/ConfigCacheIni.cpp
--- a/Engine/Source/Runtime/Core/Private/Misc/ConfigCacheIni.cpp
+++ b/Engine/Source/Runtime/Core/Private/Misc/ConfigCacheIni.cpp
@@ -513,15 +513,23 @@ static bool SaveConfigFileWrapper(const TCHAR* IniFile, const FString& Contents)
 	FCoreDelegates::PreSaveConfigFileDelegate.Broadcast(IniFile, Contents, SavedCount);
 
 	// save it even if a delegate did as well
-	const FString BaseFilename = FPaths::GetBaseFilename(IniFile);
-	const FString TempFilename = FPaths::CreateTempFilename(*FPaths::ProjectSavedDir(), *BaseFilename.Left(32));
-	bool bLocalWriteSucceeded = FFileHelper::SaveStringToFile(Contents, *TempFilename, FFileHelper::EEncodingOptions::ForceUTF8WithoutBOM);
-	if (bLocalWriteSucceeded)
+	bool bLocalWriteSucceeded = false;
+	if (FApp::IsUnattended())
 	{
-		if (!IFileManager::Get().Move(IniFile, *TempFilename))
+		bLocalWriteSucceeded = FFileHelper::SaveStringToFile(Contents, IniFile, FFileHelper::EEncodingOptions::ForceUTF8WithoutBOM);
+	}
+	else
+	{
+		const FString BaseFilename = FPaths::GetBaseFilename(IniFile);
+		const FString TempFilename = FPaths::CreateTempFilename(*FPaths::ProjectSavedDir(), *BaseFilename.Left(32));
+		bLocalWriteSucceeded = FFileHelper::SaveStringToFile(Contents, *TempFilename, FFileHelper::EEncodingOptions::ForceUTF8WithoutBOM);
+		if (bLocalWriteSucceeded)
 		{
-			IFileManager::Get().Delete(*TempFilename);
-			bLocalWriteSucceeded = false;
+			if (!IFileManager::Get().Move(IniFile, *TempFilename))
+			{
+				IFileManager::Get().Delete(*TempFilename);
+				bLocalWriteSucceeded = false;
+			}
 		}
 	}
 
