Restore FortniteMainBranchObjectVersion to prevent content load failure in assets that were saved with a FortniteMainBranchObjectVersion

diff --git a/Engine/Source/Runtime/Core/Public/UObject/FortniteMainBranchObjectVersion.h b/Engine/Source/Runtime/Core/Public/UObject/FortniteMainBranchObjectVersion.h
--- a/Engine/Source/Runtime/Core/Public/UObject/FortniteMainBranchObjectVersion.h
+++ b/Engine/Source/Runtime/Core/Public/UObject/FortniteMainBranchObjectVersion.h
@@ -167,6 +167,9 @@ struct CORE_API FFortniteMainBranchObjectVersion
 		// Enable material (edge, bending, and area stiffness) weight maps
 		ChaosClothAddMaterialWeightMaps,
 
+		// Added bShowCurve for movie scene float channel serialization
+		SerializeFloatChannelShowCurve,
+
 		// -----<new versions can be added above this line>-------------------------------------------------
 		VersionPlusOne,
 		LatestVersion = VersionPlusOne - 1
