Fix issue with access to UE4Game directories

#ue4
#android

diff --git a/Engine/Source/Runtime/Core/Private/Android/AndroidPlatformFile.cpp b/Engine/Source/Runtime/Core/Private/Android/AndroidPlatformFile.cpp
--- a/Engine/Source/Runtime/Core/Private/Android/AndroidPlatformFile.cpp
+++ b/Engine/Source/Runtime/Core/Private/Android/AndroidPlatformFile.cpp
@@ -1142,7 +1142,7 @@ public:
 		mkdir(TCHAR_TO_UTF8(*(FileBaseDir + GAndroidProjectName)), 0777);
 
 		// make sure the log directory exists if override applied
-		if (GOverrideAndroidLogDir)
+		//if (GOverrideAndroidLogDir)
 		{
 			FString LogBaseDir = GExternalFilePath + FString(FILEBASE_DIRECTORY);
 			mkdir(TCHAR_TO_UTF8(*LogBaseDir), 0777);
