Restrict number of situations we write temp file then move when saving .ini files
#rnx

diff --git a/Engine/Source/Runtime/Core/Private/Misc/ConfigCacheIni.cpp b/Engine/Source/Runtime/Core/Private/Misc/ConfigCacheIni.cpp
--- a/Engine/Source/Runtime/Core/Private/Misc/ConfigCacheIni.cpp
+++ b/Engine/Source/Runtime/Core/Private/Misc/ConfigCacheIni.cpp
@@ -514,11 +514,14 @@ static bool SaveConfigFileWrapper(const TCHAR* IniFile, const FString& Contents)
 
 	// save it even if a delegate did as well
 	bool bLocalWriteSucceeded = false;
-	if (FApp::IsUnattended())
-	{
-		bLocalWriteSucceeded = FFileHelper::SaveStringToFile(Contents, IniFile, FFileHelper::EEncodingOptions::ForceUTF8WithoutBOM);
-	}
-	else
+
+#if PLATFORM_DESKTOP && WITH_EDITOR
+	bool bWriteTempFileThenMove = !FApp::IsGame() && !FApp::IsUnattended();
+#else // PLATFORM_DESKTOP
+	bool bWriteTempFileThenMove = false;
+#endif
+
+	if (bWriteTempFileThenMove)
 	{
 		const FString BaseFilename = FPaths::GetBaseFilename(IniFile);
 		const FString TempFilename = FPaths::CreateTempFilename(*FPaths::ProjectSavedDir(), *BaseFilename.Left(32));
@@ -532,6 +535,10 @@ static bool SaveConfigFileWrapper(const TCHAR* IniFile, const FString& Contents)
 			}
 		}
 	}
+	else
+	{
+		bLocalWriteSucceeded = FFileHelper::SaveStringToFile(Contents, IniFile, FFileHelper::EEncodingOptions::ForceUTF8WithoutBOM);
+	}
 
 	// success is based on a delegate or file write working (or both)
 	return SavedCount > 0 || bLocalWriteSucceeded;
