Fix Unique Page finder tool, now that we are threading on forked servers, need to not mess things up on the data structure

diff --git a/Engine/Source/Runtime/Core/Private/Unix/UnixForkPageProtector.cpp b/Engine/Source/Runtime/Core/Private/Unix/UnixForkPageProtector.cpp
--- a/Engine/Source/Runtime/Core/Private/Unix/UnixForkPageProtector.cpp
+++ b/Engine/Source/Runtime/Core/Private/Unix/UnixForkPageProtector.cpp
@@ -5,8 +5,10 @@
 #if COMPILE_FORK_PAGE_PROTECTOR
 #include "HAL/PlatformStackWalk.h"
 #include "Misc/Fork.h"
+#include "Misc/ScopeExit.h"
 #include "Misc/ScopeLock.h"
 #include "Unix/UnixPlatformFile.h"
+#include "Unix/UnixPlatformAtomics.h"
 
 #include <fcntl.h>
 #include <pthread.h>
@@ -596,6 +598,7 @@ bool FForkPageProtector::HandleNewCrashAddress(void* CrashAddress)
 		return false;
 	}
 
+	// TODO deal with threading issue
 	if (LastCrashAddress == CrashAddress)
 	{
 		return false;
@@ -662,6 +665,14 @@ bool FForkPageProtector::DumpCallstackInfoToFile()
 	FPlatformStackWalk::StackWalkAndDump(StackTrace, StackTraceSize, IgnoreCount);
 	uint64 Hash = CityHash64(StackTrace, FCStringAnsi::Strlen(StackTrace));
 
+	// we can have multiple threads coming into this signal handler, need to make sure
+	// that only 1 thread touches CallstackHashCount at a time.
+	static volatile sig_atomic_t SignalBeingHandled = 0;
+
+	while (!__sync_bool_compare_and_swap(&SignalBeingHandled, 0, 1)) { sched_yield(); }
+
+	ON_SCOPE_EXIT { FPlatformAtomics::InterlockedDecrement(&SignalBeingHandled); };
+
 	CallstackHashData* Data = CallstackHashCount.Find(Hash);
 
 	if (Data == nullptr)
