Add FConfigFile::WriteTempFileThenMove() so downstream code can know if FConfigFile will use temp files when saving.

diff --git a/Engine/Source/Runtime/Core/Private/Misc/ConfigCacheIni.cpp b/Engine/Source/Runtime/Core/Private/Misc/ConfigCacheIni.cpp
--- a/Engine/Source/Runtime/Core/Private/Misc/ConfigCacheIni.cpp
+++ b/Engine/Source/Runtime/Core/Private/Misc/ConfigCacheIni.cpp
@@ -515,13 +515,7 @@ static bool SaveConfigFileWrapper(const TCHAR* IniFile, const FString& Contents)
 	// save it even if a delegate did as well
 	bool bLocalWriteSucceeded = false;
 
-#if PLATFORM_DESKTOP && WITH_EDITOR
-	bool bWriteTempFileThenMove = !FApp::IsGame() && !FApp::IsUnattended();
-#else // PLATFORM_DESKTOP
-	bool bWriteTempFileThenMove = false;
-#endif
-
-	if (bWriteTempFileThenMove)
+	if (FConfigFile::WriteTempFileThenMove())
 	{
 		const FString BaseFilename = FPaths::GetBaseFilename(IniFile);
 		const FString TempFilename = FPaths::CreateTempFilename(*FPaths::ProjectSavedDir(), *BaseFilename.Left(32));
@@ -1534,6 +1528,17 @@ static void ClearHierarchyCache( const TCHAR* BaseIniName )
 #endif
 }
 
+bool FConfigFile::WriteTempFileThenMove()
+{
+#if PLATFORM_DESKTOP && WITH_EDITOR
+	bool bWriteTempFileThenMove = !FApp::IsGame() && !FApp::IsUnattended();
+#else // PLATFORM_DESKTOP
+	bool bWriteTempFileThenMove = false;
+#endif
+
+	return bWriteTempFileThenMove;
+}
+
 bool FConfigFile::Write(const FString& Filename, bool bDoRemoteWrite/* = true*/, const FString& PrefixText/*=FString()*/)
 {
 	TMap<FString, FString> SectionTexts;
diff --git a/Engine/Source/Runtime/Core/Public/Misc/ConfigCacheIni.h b/Engine/Source/Runtime/Core/Public/Misc/ConfigCacheIni.h
--- a/Engine/Source/Runtime/Core/Public/Misc/ConfigCacheIni.h
+++ b/Engine/Source/Runtime/Core/Public/Misc/ConfigCacheIni.h
@@ -442,6 +442,9 @@ public:
 	CORE_API void CombineFromBuffer(const FString& Buffer);
 	CORE_API void Read( const FString& Filename );
 
+	/** Whether to write a temp file then move it to it's destination when saving. */
+	CORE_API static bool WriteTempFileThenMove();
+
 	/** Write this ConfigFile to the given Filename, constructed the text from the config sections in *this, prepended by the optional PrefixText */
 	CORE_API bool Write( const FString& Filename, bool bDoRemoteWrite=true, const FString& PrefixText=FString());
 
