Refactor Android's FPlatformVirtualMemoryBlock::AllocateVirtual to just allocate virtual memory and commit it in a separate call to Commit. Also added mprotect(PROT_NONE) to Decommit as it seems to be the right way to decommit memory on Linux

diff --git a/Engine/Source/Runtime/Core/Private/Android/AndroidPlatformMemory.cpp b/Engine/Source/Runtime/Core/Private/Android/AndroidPlatformMemory.cpp
--- a/Engine/Source/Runtime/Core/Private/Android/AndroidPlatformMemory.cpp
+++ b/Engine/Source/Runtime/Core/Private/Android/AndroidPlatformMemory.cpp
@@ -438,7 +438,7 @@ FAndroidPlatformMemory::FPlatformVirtualMemoryBlock FAndroidPlatformMemory::FPla
 	size_t Alignment = FMath::Max(InAlignment, GetVirtualSizeAlignment());
 	check(Alignment <= GetVirtualSizeAlignment());
 
-	Result.Ptr = mmap(nullptr, Result.GetActualSize(), PROT_READ | PROT_WRITE, MAP_PRIVATE | MAP_ANON, -1, 0);
+	Result.Ptr = mmap(nullptr, Result.GetActualSize(), PROT_NONE, MAP_PRIVATE | MAP_ANON, -1, 0);
 	if (!LIKELY(Result.Ptr != MAP_FAILED))
 	{
 		FPlatformMemory::OnOutOfMemory(Result.GetActualSize(), InAlignment);
@@ -469,12 +469,24 @@ void FAndroidPlatformMemory::FPlatformVirtualMemoryBlock::Commit(size_t InOffset
 {
 	check(IsAligned(InOffset, GetCommitAlignment()) && IsAligned(InSize, GetCommitAlignment()));
 	check(InOffset >= 0 && InSize >= 0 && InOffset + InSize <= GetActualSize() && Ptr);
+	const int Result = mprotect(((uint8*)Ptr) + InOffset, InSize, PROT_READ | PROT_WRITE);
+	if (UNLIKELY(Result != 0))
+	{
+		FPlatformMemory::OnOutOfMemory(InSize, GetCommitAlignment());
+	}
 }
 
 void FAndroidPlatformMemory::FPlatformVirtualMemoryBlock::Decommit(size_t InOffset, size_t InSize)
 {
 	check(IsAligned(InOffset, GetCommitAlignment()) && IsAligned(InSize, GetCommitAlignment()));
 	check(InOffset >= 0 && InSize >= 0 && InOffset + InSize <= GetActualSize() && Ptr);
+	const int Result = mprotect(((uint8*)Ptr) + InOffset, InSize, PROT_NONE);
+	if (UNLIKELY(Result != 0))
+	{
+		// We are most likely out of page table entries
+		FPlatformMemory::OnOutOfMemory(InSize, GetCommitAlignment());
+	}
+
 	madvise(((uint8*)Ptr) + InOffset, InSize, MADV_DONTNEED);
 }
 
