Minimalistic MallocStomp3 that can run on Android

diff --git a/Engine/Source/Runtime/Core/Private/HAL/MallocStomp3.cpp b/Engine/Source/Runtime/Core/Private/HAL/MallocStomp3.cpp
new file mode 100644
--- /dev/null
+++ b/Engine/Source/Runtime/Core/Private/HAL/MallocStomp3.cpp
@@ -0,0 +1,98 @@
+// Copyright Epic Games, Inc. All Rights Reserved.
+
+#include "HAL/MallocStomp3.h"
+#include "Templates/AlignmentTemplates.h"
+
+namespace {
+struct AllocationInfo 
+{
+	uint32 PagesCount;
+	uint32 OriginalSize;
+};
+
+inline size_t RoundUpToPageSize(size_t Size, size_t PageSize)
+{
+	return (Size + PageSize - 1) & ~(PageSize - 1);
+}
+}
+
+FMallocStomp3::FMallocStomp3(EOptions Options)
+	: Options(Options)
+{
+}
+
+void* FMallocStomp3::Malloc(SIZE_T Count, uint32 Alignment)
+{
+	Alignment = (Alignment == DEFAULT_ALIGNMENT) ? 1 : Alignment;
+	Alignment = (Options & EOptions::EForceIgnoreAlignment) ? 1 : Alignment;
+
+	const size_t PageSize = FPlatformMemory::FPlatformVirtualMemoryBlock::GetVirtualSizeAlignment();
+	check(Alignment <= PageSize);
+
+	const size_t OriginalSize = Count;
+	Count = RoundUpToPageSize(Count + sizeof(AllocationInfo) + Alignment, PageSize);
+	const size_t AllocationSize = Count + PageSize;
+	FPlatformMemory::FPlatformVirtualMemoryBlock Block = FPlatformMemory::FPlatformVirtualMemoryBlock::AllocateVirtual(AllocationSize);
+	Block.CommitByPtr(0, Count);
+
+	uintptr_t Ptr = (uintptr_t)Block.GetVirtualPointer();
+	Ptr += Count - OriginalSize;
+	Ptr = AlignDown(Ptr, Alignment);
+	AllocationInfo* Info = (AllocationInfo*)Ptr;
+	--Info;
+	Info->PagesCount = Count / PageSize;
+	Info->OriginalSize = OriginalSize;
+	return (void*)Ptr;
+}
+
+void* FMallocStomp3::Realloc(void* Original, SIZE_T Count, uint32 Alignment)
+{
+	if (Count == 0)
+	{
+		Free(Original);
+		return nullptr;
+	}
+
+	void* Result = Malloc(Count, Alignment);
+	if (Original)
+	{
+		AllocationInfo* Info = (AllocationInfo*)Original;
+		--Info;
+		FPlatformMemory::Memcpy(Result, Original, FMath::Min(Info->OriginalSize, (uint32)Count));
+		Free(Original);
+	}
+
+	return Result;
+}
+
+void FMallocStomp3::Free(void* Original)
+{
+	if (!Original)
+	{
+		return;
+	}
+
+	const size_t PageSize = FPlatformMemory::FPlatformVirtualMemoryBlock::GetVirtualSizeAlignment();
+
+	AllocationInfo* Info = (AllocationInfo*)Original;
+	--Info;
+	const size_t Size = (Info->PagesCount + 1) * PageSize;
+	Original = AlignDown(Original, PageSize);
+
+	FPlatformMemory::FPlatformVirtualMemoryBlock Block(Original, Size / PageSize);
+	//TODO: free physical mem, but keep virtual to detect dangling pointers?
+	Block.FreeVirtual();
+}
+
+bool FMallocStomp3::GetAllocationSize(void* Original, SIZE_T& SizeOut)
+{
+	if (Original)
+	{
+		AllocationInfo* Info = (AllocationInfo*)Original;
+		--Info;
+		SizeOut = Info->OriginalSize;
+		return true;
+	}
+
+	return false;
+}
\ No newline at end of file
diff --git a/Engine/Source/Runtime/Core/Public/HAL/MallocStomp3.h b/Engine/Source/Runtime/Core/Public/HAL/MallocStomp3.h
new file mode 100644
--- /dev/null
+++ b/Engine/Source/Runtime/Core/Public/HAL/MallocStomp3.h
@@ -0,0 +1,54 @@
+// Copyright Epic Games, Inc. All Rights Reserved.
+
+#pragma once
+
+#include "HAL/MemoryBase.h"
+
+/*
+	A simple stomp allocator with support for mobile platforms
+*/
+class FMallocStomp3 final : public FMalloc
+{
+public:
+	enum EOptions
+	{
+		ENone,
+		EForceIgnoreAlignment,	//Ignore alignment requirements to place the allocation exactly on a boundary of a page
+		//TODO: Add options to protect from buffer under run and detecting dangling pointers
+	};
+
+	FMallocStomp3(EOptions Options);
+
+	virtual void* Malloc(SIZE_T Count, uint32 Alignment = DEFAULT_ALIGNMENT) override;
+
+	virtual void* Realloc(void* Original, SIZE_T Count, uint32 Alignment = DEFAULT_ALIGNMENT) override;
+
+	virtual void Free(void* Original) override;
+
+	/**
+	* If possible determine the size of the memory allocated at the given address
+	*
+	* @param Original - Pointer to memory we are checking the size of
+	* @param SizeOut - If possible, this value is set to the size of the passed in pointer
+	* @return true if succeeded
+	*/
+	virtual bool GetAllocationSize(void* Original, SIZE_T& SizeOut) override;
+
+	/**
+	 * Gets descriptive name for logging purposes.
+	 *
+	 * @return pointer to human-readable malloc name
+	 */
+	virtual const TCHAR* GetDescriptiveName() override
+	{
+		return TEXT("Stomp3");
+	}
+
+	virtual bool IsInternallyThreadSafe() const override
+	{
+		return true;
+	}
+
+private:
+	uint32 Options;
+};
\ No newline at end of file
