Reflect the platform type layout parameters in the SM DDC key.

- We expect ShaderMapId written to the derived data to be exact same as the one that we used to generate the DDC key with. It can happen (and already happened to MarcA), that material DDC keys of two close enough platforms with inessentially different layout parameters collided. This change prevents such a collision.

#rnx

diff --git a/Engine/Source/Runtime/Core/Private/Serialization/MemoryImage.cpp b/Engine/Source/Runtime/Core/Private/Serialization/MemoryImage.cpp
--- a/Engine/Source/Runtime/Core/Private/Serialization/MemoryImage.cpp
+++ b/Engine/Source/Runtime/Core/Private/Serialization/MemoryImage.cpp
@@ -108,6 +108,7 @@ FArchive& FPlatformTypeLayoutParameters::Serialize(FArchive& Ar)
 
 void FPlatformTypeLayoutParameters::AppendKeyString(FString& KeyString) const
 {
+	KeyString += FString::Printf(TEXT("FL_%08x_MFA_%08x_"), Flags, MaxFieldAlignment);
 }
 
 // evaluated during static-initialization, so logging from regular check() macros won't work correctly
diff --git a/Engine/Source/Runtime/Core/Public/Serialization/MemoryLayout.h b/Engine/Source/Runtime/Core/Public/Serialization/MemoryLayout.h
--- a/Engine/Source/Runtime/Core/Public/Serialization/MemoryLayout.h
+++ b/Engine/Source/Runtime/Core/Public/Serialization/MemoryLayout.h
@@ -883,9 +883,9 @@ struct FPlatformTypeLayoutParameters
 	CORE_API FArchive& Serialize(FArchive& Ar);
 
 	/**
-	 * Allow the layout parameters to modify the given DDC key string
-	 * Normally this should not be used, and instead hashed layouts of any dependent types should be preferred
-	 * This is mostly needed when fixing bugs related to writing memory images with certain parameter combinations
+	 * Allow the layout parameters to modify the given DDC key string.
+	 * Since layout parameters are part of e.g. material shadermap ID, they should result in two different DDC entries for two different IDs,
+	 * even if binary layouts happen to be compatible.
 	 */
 	CORE_API void AppendKeyString(FString& KeyString) const;
 };
diff --git a/Engine/Source/Runtime/Engine/Private/Materials/MaterialShader.cpp b/Engine/Source/Runtime/Engine/Private/Materials/MaterialShader.cpp
--- a/Engine/Source/Runtime/Engine/Private/Materials/MaterialShader.cpp
+++ b/Engine/Source/Runtime/Engine/Private/Materials/MaterialShader.cpp
@@ -1257,7 +1257,8 @@ void FMaterialShaderMap::LoadFromDerivedDataCache(const FMaterial* Material, con
 				InOutShaderMap->Serialize(Ar);
 				//InOutShaderMap->RegisterSerializedShaders(false);
 		
-				checkSlow(InOutShaderMap->GetShaderMapId() == ShaderMapId);
+				checkf(InDataKey == DataKey, TEXT("Data deserialized from the DDC would need a different DDC key!"));
+				checkf(InOutShaderMap->GetShaderMapId() == ShaderMapId, TEXT("Shadermap data deserialized from the DDC does not match the ID we used to build the key!"));
 
 				// Register in the global map
 				InOutShaderMap->Register(InPlatform);
