Enable the Crash Report Monitor for the client
  - Split the Editor analytics summary in and Engine summary and an Editor summary

diff --git a/Engine/Source/Programs/CrashReportClient/CrashReportClient.Build.cs b/Engine/Source/Programs/CrashReportClient/CrashReportClient.Build.cs
--- a/Engine/Source/Programs/CrashReportClient/CrashReportClient.Build.cs
+++ b/Engine/Source/Programs/CrashReportClient/CrashReportClient.Build.cs
@@ -55,5 +55,12 @@ public class CrashReportClient : ModuleRules
 		}
 
 		PrivateIncludePaths.Add("Runtime/Launch/Private");		// For LaunchEngineLoop.cpp include
+
+		PrivateDefinitions.AddRange(
+			new string[]
+			{
+				"CRASH_REPORT_WITH_MTBF=1",
+			}
+		);
 	}
 }
diff --git a/Engine/Source/Programs/CrashReportClient/Private/CrashReportAnalyticsSessionSummary.cpp b/Engine/Source/Programs/CrashReportClient/Private/CrashReportAnalyticsSessionSummary.cpp
--- a/Engine/Source/Programs/CrashReportClient/Private/CrashReportAnalyticsSessionSummary.cpp
+++ b/Engine/Source/Programs/CrashReportClient/Private/CrashReportAnalyticsSessionSummary.cpp
@@ -10,6 +10,7 @@
 #include "GenericPlatform/GenericPlatformMath.h"
 #include "GenericPlatform/GenericPlatformCrashContext.h"
 #include "HAL/PlatformTime.h"
+#include "HAL/Thread.h"
 #include "Logging/LogMacros.h"
 #include "Misc/EngineVersion.h"
 #include "Misc/Paths.h"
diff --git a/Engine/Source/Runtime/Analytics/AnalyticsET/Private/AnalyticsSessionSummaryManager.cpp b/Engine/Source/Runtime/Analytics/AnalyticsET/Private/AnalyticsSessionSummaryManager.cpp
--- a/Engine/Source/Runtime/Analytics/AnalyticsET/Private/AnalyticsSessionSummaryManager.cpp
+++ b/Engine/Source/Runtime/Analytics/AnalyticsET/Private/AnalyticsSessionSummaryManager.cpp
@@ -557,7 +557,7 @@ bool FAnalyticsSessionSummaryManager::AggregateSummaries(const FString& InProces
 		OutSummaryProperties.Emplace(AnalyticsManagerProperties::SessionId.Key, InternalSessionId->GetValue());
 	}
 
-	if (!MissingDataFromProcesses.IsEmpty())
+	if (MissingDataFromProcesses.Num())
 	{
 		// Add the subsidiary process tag name(s) for which the data file couldn't be loaded (file was locked, corrupted, ...)
 		OutSummaryProperties.Emplace(AnalyticsManagerProperties::MissingDataFrom.Key, FString::Join(MissingDataFromProcesses, TEXT(",")));
diff --git a/Engine/Source/Runtime/Core/Private/GenericPlatform/GenericPlatformCrashContext.cpp b/Engine/Source/Runtime/Core/Private/GenericPlatform/GenericPlatformCrashContext.cpp
--- a/Engine/Source/Runtime/Core/Private/GenericPlatform/GenericPlatformCrashContext.cpp
+++ b/Engine/Source/Runtime/Core/Private/GenericPlatform/GenericPlatformCrashContext.cpp
@@ -32,6 +32,10 @@
 #define NOINITCRASHREPORTER 0
 #endif
 
+#ifndef CRASH_REPORTER_WITH_ANALYTICS
+#define CRASH_REPORTER_WITH_ANALYTICS 0
+#endif
+
 DEFINE_LOG_CATEGORY_STATIC(LogCrashContext, Display, All);
 
 extern CORE_API bool GIsGPUCrashed;
@@ -416,9 +420,14 @@ void FGenericCrashContext::InitializeFromConfig()
 	// Read the initial un-localized crash context text
 	UpdateLocalizedStrings();
 
+#if WITH_EDITOR
 	// Set privacy settings -> WARNING: Ensure those setting have a default values in Engine/Config/BaseEditorSettings.ini file, otherwise, they will not be found.
 	GConfig->GetBool(TEXT("/Script/UnrealEd.CrashReportsPrivacySettings"), TEXT("bSendUnattendedBugReports"), NCached::UserSettings.bSendUnattendedBugReports, GEditorSettingsIni);
 	GConfig->GetBool(TEXT("/Script/UnrealEd.AnalyticsPrivacySettings"), TEXT("bSendUsageData"), NCached::UserSettings.bSendUsageData, GEditorSettingsIni);
+#elif CRASH_REPORTER_WITH_ANALYTICS
+	NCached::UserSettings.bSendUnattendedBugReports = true; // Give CRC permission to generate and send an 'AbnormalShutdown' report if the application died suddently, mainly to collect the logs post-mortem and count those occurrences.
+	NCached::UserSettings.bSendUsageData = true; // Give CRC permission to send its analytics and the application session summary (if one exist)
+#endif
 	
 	// Write a marker file to disk indicating the user has allowed unattended crash reports being
 	// sent. This allows us to submit reports for crashes during static initialization when user
diff --git a/Engine/Source/Runtime/Core/Private/Windows/WindowsPlatformCrashContext.cpp b/Engine/Source/Runtime/Core/Private/Windows/WindowsPlatformCrashContext.cpp
--- a/Engine/Source/Runtime/Core/Private/Windows/WindowsPlatformCrashContext.cpp
+++ b/Engine/Source/Runtime/Core/Private/Windows/WindowsPlatformCrashContext.cpp
@@ -43,10 +43,8 @@
 	#define UE_LOG_CRASH_CALLSTACK 1
 #endif
 
-#if WITH_EDITOR
-	#define USE_CRASH_REPORTER_MONITOR 1
-#else 
-	#define USE_CRASH_REPORTER_MONITOR 0
+#ifndef USE_CRASH_REPORTER_MONITOR
+#define USE_CRASH_REPORTER_MONITOR WITH_EDITOR
 #endif
 
 #ifndef NOINITCRASHREPORTER
@@ -546,7 +544,9 @@ FProcHandle LaunchCrashReportClient(void** OutWritePipe, void** OutReadPipe, uin
 			PipeChildInRead, //Pass this to allow inherit handles in child proc
 			nullptr);
 
-#if WITH_EDITOR
+		DWORD pid = GetProcessId(Handle.Get());
+		UE_LOG(LogWindows, Log, TEXT("Started CrashReportClient (pid=%d)"), pid);
+
 		// The CRC instance launched above will respanwn itself to sever the link with the Editor process group. This way, if the user kills the Editor
 		// process group in TaskManager, CRC will not die at the same moment and will be able to capture the Editor exit code and send the session summary.
 		if (Handle.IsValid())
@@ -586,7 +586,6 @@ FProcHandle LaunchCrashReportClient(void** OutWritePipe, void** OutReadPipe, uin
 				*OutCrashReportClientProcessId = Handle.IsValid() ? RepawnedCrcPid : 0;
 			}
 		}
-#endif
 	}
 
 	return Handle;
