Add ReadFlags to FFileHelper::LoadFileToString() overload that didn't have them.

diff --git a/Engine/Source/Runtime/Core/Private/Misc/FileHelper.cpp b/Engine/Source/Runtime/Core/Private/Misc/FileHelper.cpp
--- a/Engine/Source/Runtime/Core/Private/Misc/FileHelper.cpp
+++ b/Engine/Source/Runtime/Core/Private/Misc/FileHelper.cpp
@@ -198,7 +198,7 @@ bool FFileHelper::LoadFileToString( FString& Result, const TCHAR* Filename, EHas
 	return LoadFileToString(Result, *Reader.Get(), VerifyFlags);
 }
 
-bool FFileHelper::LoadFileToString(FString& Result, IPlatformFile* PlatformFile, const TCHAR* Filename, EHashOptions VerifyFlags /*= EHashOptions::None*/)
+bool FFileHelper::LoadFileToString(FString& Result, IPlatformFile* PlatformFile, const TCHAR* Filename, EHashOptions VerifyFlags /*= EHashOptions::None*/, uint32 ReadFlags /*= 0*/)
 {
 	if (!PlatformFile)
 	{
@@ -208,7 +208,16 @@ bool FFileHelper::LoadFileToString(FString& Result, IPlatformFile* PlatformFile,
 	IFileHandle* File = PlatformFile->OpenRead(Filename);
 	if (!File)
 	{
-		UE_LOG(LogStreaming, Warning, TEXT("Failed to read file '%s' error."), Filename);
+		if (ReadFlags & FILEREAD_NoFail)
+		{
+			UE_LOG(LogStreaming, Fatal, TEXT("Failed to read file: %s"), Filename);
+		}
+
+		if (!(ReadFlags & FILEREAD_Silent))
+		{
+			UE_LOG(LogStreaming, Warning, TEXT("Failed to read file '%s' error."), Filename);
+		}
+
 		return false;
 	}
 
diff --git a/Engine/Source/Runtime/Core/Public/Misc/FileHelper.h b/Engine/Source/Runtime/Core/Public/Misc/FileHelper.h
--- a/Engine/Source/Runtime/Core/Public/Misc/FileHelper.h
+++ b/Engine/Source/Runtime/Core/Public/Misc/FileHelper.h
@@ -92,7 +92,7 @@ struct CORE_API FFileHelper
 	 * @param Filename     Name of the file to load
 	 * @param VerifyFlags  Flags controlling the hash verification behavior ( see EHashOptions )
 	 */
-	static bool LoadFileToString(FString& Result, IPlatformFile* PlatformFile, const TCHAR* Filename, EHashOptions VerifyFlags = EHashOptions::None);
+	static bool LoadFileToString(FString& Result, IPlatformFile* PlatformFile, const TCHAR* Filename, EHashOptions VerifyFlags = EHashOptions::None, uint32 ReadFlags = 0);
 
 	/**
 	 * Load a text file to an array of strings. Supports all combination of ANSI/Unicode files and platforms.
