Add FConfigFile::SetBool

diff --git a/Engine/Source/Runtime/Core/Private/Misc/ConfigCacheIni.cpp b/Engine/Source/Runtime/Core/Private/Misc/ConfigCacheIni.cpp
--- a/Engine/Source/Runtime/Core/Private/Misc/ConfigCacheIni.cpp
+++ b/Engine/Source/Runtime/Core/Private/Misc/ConfigCacheIni.cpp
@@ -2026,6 +2026,11 @@ void FConfigFile::SetDouble(const TCHAR* Section, const TCHAR* Key, double Value
 	SetString(Section, Key, Text);
 }
 
+void FConfigFile::SetBool(const TCHAR* Section, const TCHAR* Key, bool Value)
+{
+	SetString(Section, Key, Value ? TEXT("True") : TEXT("False"));
+}
+
 void FConfigFile::SetInt64( const TCHAR* Section, const TCHAR* Key, int64 Value )
 {
 	TCHAR Text[MAX_SPRINTF]=TEXT("");
@@ -2666,25 +2671,12 @@ void FConfigCacheIni::SetString( const TCHAR* Section, const TCHAR* Key, const T
 {
 	FConfigFile* File = Find(Filename);
 
-	if ( !File )
+	if (!File)
 	{
 		return;
 	}
 
-	FConfigSection* Sec = File->FindOrAddSection( Section );
-
-	FConfigValue* ConfigValue = Sec->Find( Key );
-	if( !ConfigValue )
-	{
-		Sec->Add( Key, Value );
-		File->Dirty = true;
-	}
-	// Use GetSavedValueForWriting rather than GetSavedValue to avoid having the is-it-dirty query mark the values as having been accessed for dependency tracking
-	else if( FCString::Strcmp(*UE::ConfigCacheIni::Private::FAccessor::GetSavedValueForWriting(*ConfigValue),Value)!=0 )
-	{
-		File->Dirty = true;
-		*ConfigValue = FConfigValue(Value);
-	}
+	File->SetString(Section, Key, Value);
 }
 
 void FConfigCacheIni::SetText( const TCHAR* Section, const TCHAR* Key, const FText& Value, const FString& Filename )
@@ -3165,7 +3157,13 @@ void FConfigCacheIni::SetBool
 	const FString&	Filename
 )
 {
-	SetString( Section, Key, Value ? TEXT("True") : TEXT("False"), Filename );
+	FConfigFile* File = Find(Filename);
+	if (!File)
+	{
+		return;
+	}
+
+	File->SetBool(Section, Key, Value);
 }
 
 void FConfigCacheIni::SetArray
diff --git a/Engine/Source/Runtime/Core/Public/Misc/ConfigCacheIni.h b/Engine/Source/Runtime/Core/Public/Misc/ConfigCacheIni.h
--- a/Engine/Source/Runtime/Core/Public/Misc/ConfigCacheIni.h
+++ b/Engine/Source/Runtime/Core/Public/Misc/ConfigCacheIni.h
@@ -506,11 +506,12 @@ public:
 		return GetArray(Section, Key, Value);
 	}
 
-	CORE_API void SetString( const TCHAR* Section, const TCHAR* Key, const TCHAR* Value );
-	CORE_API void SetText( const TCHAR* Section, const TCHAR* Key, const FText& Value );
-	CORE_API void SetFloat( const TCHAR* Section, const TCHAR* Key, float Value );
-	CORE_API void SetDouble( const TCHAR* Section, const TCHAR* Key, double Value );
-	CORE_API void SetInt64( const TCHAR* Section, const TCHAR* Key, const int64 Value );
+	CORE_API void SetString(const TCHAR* Section, const TCHAR* Key, const TCHAR* Value);
+	CORE_API void SetText(const TCHAR* Section, const TCHAR* Key, const FText& Value);
+	CORE_API void SetFloat(const TCHAR* Section, const TCHAR* Key, float Value);
+	CORE_API void SetDouble(const TCHAR* Section, const TCHAR* Key, double Value);
+	CORE_API void SetBool(const TCHAR* Section, const TCHAR* Key, bool Value);
+	CORE_API void SetInt64(const TCHAR* Section, const TCHAR* Key, const int64 Value);
 	CORE_API void SetArray(const TCHAR* Section, const TCHAR* Key, const TArray<FString>& Value);
 	
 	/**
