Allow refering to CSV begin/end events by FName. This allows FName to be the stable index for these calls and allows tools outside of CsvProfiler to also use these calls directly.

diff --git a/Engine/Source/Runtime/Core/Private/ProfilingDebugging/CsvProfiler.cpp b/Engine/Source/Runtime/Core/Private/ProfilingDebugging/CsvProfiler.cpp
--- a/Engine/Source/Runtime/Core/Private/ProfilingDebugging/CsvProfiler.cpp
+++ b/Engine/Source/Runtime/Core/Private/ProfilingDebugging/CsvProfiler.cpp
@@ -3004,6 +3004,17 @@ void FCsvProfiler::BeginStat(const char * StatName, uint32 CategoryIndex)
 #endif
 }
 
+void FCsvProfiler::BeginStat(const FName& StatName, uint32 CategoryIndex)
+{
+#if RECORD_TIMESTAMPS
+	if (GCsvProfilerIsCapturing && GCsvCategoriesEnabled[CategoryIndex])
+	{
+		LLM_SCOPE(ELLMTag::CsvProfiler);
+		FCsvProfilerThreadData::Get().AddTimestampBegin(StatName, CategoryIndex);
+	}
+#endif
+}
+
 void FCsvProfiler::EndStat(const char * StatName, uint32 CategoryIndex)
 {
 #if RECORD_TIMESTAMPS
@@ -3015,6 +3026,17 @@ void FCsvProfiler::EndStat(const char * StatName, uint32 CategoryIndex)
 #endif
 }
 
+void FCsvProfiler::EndStat(const FName& StatName, uint32 CategoryIndex)
+{
+#if RECORD_TIMESTAMPS
+	if (GCsvProfilerIsCapturing && GCsvCategoriesEnabled[CategoryIndex])
+	{
+		LLM_SCOPE(ELLMTag::CsvProfiler);
+		FCsvProfilerThreadData::Get().AddTimestampEnd(StatName, CategoryIndex);
+	}
+#endif
+}
+
 void FCsvProfiler::BeginExclusiveStat(const char * StatName)
 {
 #if RECORD_TIMESTAMPS
diff --git a/Engine/Source/Runtime/Core/Public/ProfilingDebugging/CsvProfiler.h b/Engine/Source/Runtime/Core/Public/ProfilingDebugging/CsvProfiler.h
--- a/Engine/Source/Runtime/Core/Public/ProfilingDebugging/CsvProfiler.h
+++ b/Engine/Source/Runtime/Core/Public/ProfilingDebugging/CsvProfiler.h
@@ -232,8 +232,10 @@ public:
 
 	/** Begin static interface (used by macros)*/
 	/** Push/pop events */
-	CORE_API static void BeginStat(const char * StatName, uint32 CategoryIndex);
-	CORE_API static void EndStat(const char * StatName, uint32 CategoryIndex);
+	CORE_API static void BeginStat(const char* StatName, uint32 CategoryIndex);
+	CORE_API static void BeginStat(const FName& StatName, uint32 CategoryIndex);
+	CORE_API static void EndStat(const char* StatName, uint32 CategoryIndex);
+	CORE_API static void EndStat(const FName& StatName, uint32 CategoryIndex);
 
 	CORE_API static void BeginExclusiveStat(const char * StatName);
 	CORE_API static void EndExclusiveStat(const char * StatName);
