Implement PersistentStorageManager on Android when PLATFORM_USE_PLATFORM_FILE_MANAGED_STORAGE_WRAPPER is set.
#tests none

diff --git a/Engine/Source/Runtime/Core/Private/Android/AndroidPlatformFile.cpp b/Engine/Source/Runtime/Core/Private/Android/AndroidPlatformFile.cpp
--- a/Engine/Source/Runtime/Core/Private/Android/AndroidPlatformFile.cpp
+++ b/Engine/Source/Runtime/Core/Private/Android/AndroidPlatformFile.cpp
@@ -5,6 +5,7 @@
 =============================================================================*/
 
 #include "Android/AndroidPlatformFile.h"
+#include "HAL/IPlatformFileManagedStorageWrapper.h"
 
 #if USE_ANDROID_FILE
 #include "Misc/App.h"
@@ -2166,7 +2167,13 @@ private:
 
 IPlatformFile& IPlatformFile::GetPlatformPhysical()
 {
+#if PLATFORM_USE_PLATFORM_FILE_MANAGED_STORAGE_WRAPPER
+	FAndroidPlatformFile& AndroidPlatformFile = FAndroidPlatformFile::GetPlatformPhysical();
+	static FManagedStoragePlatformFile ManagedStoragePlatformFileSingleton(&AndroidPlatformFile);
+	return ManagedStoragePlatformFileSingleton;
+#else
 	return FAndroidPlatformFile::GetPlatformPhysical();
+#endif //PLATFORM_USE_PLATFORM_FILE_MANAGED_STORAGE_WRAPPER
 }
 
 IAndroidPlatformFile & IAndroidPlatformFile::GetPlatformPhysical()
