revert the Persistant Storage Manager change on Android it is causing imediate crashes on startup

diff --git a/Engine/Source/Runtime/Core/Private/Android/AndroidPlatformFile.cpp b/Engine/Source/Runtime/Core/Private/Android/AndroidPlatformFile.cpp
--- a/Engine/Source/Runtime/Core/Private/Android/AndroidPlatformFile.cpp
+++ b/Engine/Source/Runtime/Core/Private/Android/AndroidPlatformFile.cpp
@@ -2167,13 +2167,7 @@ private:
 
 IPlatformFile& IPlatformFile::GetPlatformPhysical()
 {
-#if PLATFORM_USE_PLATFORM_FILE_MANAGED_STORAGE_WRAPPER
-	FAndroidPlatformFile& AndroidPlatformFile = FAndroidPlatformFile::GetPlatformPhysical();
-	static FManagedStoragePlatformFile ManagedStoragePlatformFileSingleton(&AndroidPlatformFile);
-	return ManagedStoragePlatformFileSingleton;
-#else
 	return FAndroidPlatformFile::GetPlatformPhysical();
-#endif //PLATFORM_USE_PLATFORM_FILE_MANAGED_STORAGE_WRAPPER
 }
 
 IAndroidPlatformFile & IAndroidPlatformFile::GetPlatformPhysical()
