Added missing macros

#rnx

diff --git a/Engine/Source/Runtime/Core/Public/UObject/NameTypes.h b/Engine/Source/Runtime/Core/Public/UObject/NameTypes.h
--- a/Engine/Source/Runtime/Core/Public/UObject/NameTypes.h
+++ b/Engine/Source/Runtime/Core/Public/UObject/NameTypes.h
@@ -253,8 +253,8 @@ public:
 
 	CORE_API void Write(FArchive& Ar) const;
 
-	static int32 GetDataOffset();
-	struct FNameStringView MakeView(union FNameBuffer& OptionalDecodeBuffer) const;
+	static CORE_API int32 GetDataOffset();
+	struct CORE_API FNameStringView MakeView(union FNameBuffer& OptionalDecodeBuffer) const;
 private:
 	friend class FName;
 	friend struct FNameHelper;
