Include OVERLOADED_VIRTUAL_WARNINGS in the third party start/end block for all clang platforms
#rnx

diff --git a/Engine/Source/Runtime/Core/Public/Android/AndroidPlatformCompilerPreSetup.h b/Engine/Source/Runtime/Core/Public/Android/AndroidPlatformCompilerPreSetup.h
--- a/Engine/Source/Runtime/Core/Public/Android/AndroidPlatformCompilerPreSetup.h
+++ b/Engine/Source/Runtime/Core/Public/Android/AndroidPlatformCompilerPreSetup.h
@@ -12,12 +12,14 @@
 		PRAGMA_DISABLE_UNDEFINED_IDENTIFIER_WARNINGS \
 		PRAGMA_DISABLE_MISSING_VIRTUAL_DESTRUCTOR_WARNINGS \
 		PRAGMA_DISABLE_DEPRECATION_WARNINGS \
-        PRAGMA_DISABLE_UNINITIALIZED_CONST_REFERENCE_WARNINGS
+		PRAGMA_DISABLE_OVERLOADED_VIRTUAL_WARNINGS \
+		PRAGMA_DISABLE_UNINITIALIZED_CONST_REFERENCE_WARNINGS
 #endif
 
 #ifndef THIRD_PARTY_INCLUDES_END
 	#define THIRD_PARTY_INCLUDES_END \
-        PRAGMA_ENABLE_UNINITIALIZED_CONST_REFERENCE_WARNINGS \
+		PRAGMA_ENABLE_UNINITIALIZED_CONST_REFERENCE_WARNINGS \
+		PRAGMA_DISABLE_OVERLOADED_VIRTUAL_WARNINGS \
 		PRAGMA_ENABLE_DEPRECATION_WARNINGS \
 		PRAGMA_ENABLE_MISSING_VIRTUAL_DESTRUCTOR_WARNINGS \
 		PRAGMA_ENABLE_UNDEFINED_IDENTIFIER_WARNINGS \
diff --git a/Engine/Source/Runtime/Core/Public/Apple/ApplePlatformCompilerPreSetup.h b/Engine/Source/Runtime/Core/Public/Apple/ApplePlatformCompilerPreSetup.h
--- a/Engine/Source/Runtime/Core/Public/Apple/ApplePlatformCompilerPreSetup.h
+++ b/Engine/Source/Runtime/Core/Public/Apple/ApplePlatformCompilerPreSetup.h
@@ -12,12 +12,14 @@
 		PRAGMA_DISABLE_UNDEFINED_IDENTIFIER_WARNINGS \
 		PRAGMA_DISABLE_MISSING_VIRTUAL_DESTRUCTOR_WARNINGS \
 		PRAGMA_DISABLE_DEPRECATION_WARNINGS \
-        PRAGMA_DISABLE_UNINITIALIZED_CONST_REFERENCE_WARNINGS
+		PRAGMA_DISABLE_OVERLOADED_VIRTUAL_WARNINGS \
+		PRAGMA_DISABLE_UNINITIALIZED_CONST_REFERENCE_WARNINGS
 #endif
 
 #ifndef THIRD_PARTY_INCLUDES_END
 	#define THIRD_PARTY_INCLUDES_END \
-        PRAGMA_ENABLE_UNINITIALIZED_CONST_REFERENCE_WARNINGS \
+		PRAGMA_ENABLE_UNINITIALIZED_CONST_REFERENCE_WARNINGS \
+		PRAGMA_DISABLE_OVERLOADED_VIRTUAL_WARNINGS \
 		PRAGMA_ENABLE_DEPRECATION_WARNINGS \
 		PRAGMA_ENABLE_MISSING_VIRTUAL_DESTRUCTOR_WARNINGS \
 		PRAGMA_ENABLE_UNDEFINED_IDENTIFIER_WARNINGS \
diff --git a/Engine/Source/Runtime/Core/Public/Unix/UnixPlatformCompilerPreSetup.h b/Engine/Source/Runtime/Core/Public/Unix/UnixPlatformCompilerPreSetup.h
--- a/Engine/Source/Runtime/Core/Public/Unix/UnixPlatformCompilerPreSetup.h
+++ b/Engine/Source/Runtime/Core/Public/Unix/UnixPlatformCompilerPreSetup.h
@@ -41,14 +41,16 @@
 		PRAGMA_DISABLE_MISSING_VIRTUAL_DESTRUCTOR_WARNINGS \
 		PRAGMA_DISABLE_DEPRECATED_REGISTER_WARNINGS \
 		PRAGMA_DISABLE_DEPRECATION_WARNINGS \
+		PRAGMA_DISABLE_OVERLOADED_VIRTUAL_WARNINGS \
 		PRAGMA_DISABLE_MACRO_REDEFINED_WARNINGS \
-        PRAGMA_DISABLE_UNINITIALIZED_CONST_REFERENCE_WARNINGS
+		PRAGMA_DISABLE_UNINITIALIZED_CONST_REFERENCE_WARNINGS
 #endif
 
 #ifndef THIRD_PARTY_INCLUDES_END
 	#define THIRD_PARTY_INCLUDES_END \
-        PRAGMA_ENABLE_UNINITIALIZED_CONST_REFERENCE_WARNINGS \
+		PRAGMA_ENABLE_UNINITIALIZED_CONST_REFERENCE_WARNINGS \
 		PRAGMA_ENABLE_MACRO_REDEFINED_WARNINGS \
+		PRAGMA_DISABLE_OVERLOADED_VIRTUAL_WARNINGS \
 		PRAGMA_ENABLE_DEPRECATION_WARNINGS \
 		PRAGMA_ENABLE_DEPRECATED_REGISTER_WARNINGS \
 		PRAGMA_ENABLE_MISSING_VIRTUAL_DESTRUCTOR_WARNINGS \
