Compile fix, second attempt
#rnx

diff --git a/Engine/Source/Runtime/Core/Public/Containers/SpscQueue.h b/Engine/Source/Runtime/Core/Public/Containers/SpscQueue.h
--- a/Engine/Source/Runtime/Core/Public/Containers/SpscQueue.h
+++ b/Engine/Source/Runtime/Core/Public/Containers/SpscQueue.h
@@ -118,11 +118,18 @@ private:
 private:
 	// consumer part 
 	// accessed mainly by consumer, infrequently by producer 
+#ifndef PLATFORM_MAC // some projects are still built on macOS before v.10.14 that doesn't have aligned new/delete operators
 	alignas(PLATFORM_CACHE_LINE_SIZE) std::atomic<FNode*> Tail; // tail of the queue 
-
+#else
+	std::atomic<FNode*> Tail; // tail of the queue
+#endif
 	// producer part 
 	// accessed only by producer 
+#ifndef PLATFORM_MAC
 	alignas(PLATFORM_CACHE_LINE_SIZE) FNode* Head; // head of the queue 
+#else
+	FNode* Head; // head of the queue
+#endif
 	FNode* First; // last unused node (tail of node cache) 
 	FNode* TailCopy; // helper (points somewhere between First and Tail)
 };
\ No newline at end of file
