-tracenocache command line option allowing users to optionall disable the caching of important trace events if they so wish.

#rnx

diff --git a/Engine/Source/Runtime/Core/Private/ProfilingDebugging/TraceAuxiliary.cpp b/Engine/Source/Runtime/Core/Private/ProfilingDebugging/TraceAuxiliary.cpp
--- a/Engine/Source/Runtime/Core/Private/ProfilingDebugging/TraceAuxiliary.cpp
+++ b/Engine/Source/Runtime/Core/Private/ProfilingDebugging/TraceAuxiliary.cpp
@@ -695,9 +695,29 @@ void FTraceAuxiliary::Initialize(const TCHAR* CommandLine)
 		<< Session2.ConfigurationType(uint8(FApp::GetBuildConfiguration()))
 		<< Session2.TargetType(uint8(FApp::GetBuildTargetType()));
 
+	FString Parameter;
+
+	// Attempt to send trace data somewhere from the command line. It prehaps
+	// seems odd to do this before initialising Trace, but it is done this way
+	// to support disabling the "important" cache without losing any events.
+	if (FParse::Value(CommandLine, TEXT("-tracehost="), Parameter))
+	{
+		GTraceAuxiliary.Connect(ETraceConnectType::Network, *Parameter);
+	}
+	else if (FParse::Value(CommandLine, TEXT("-tracefile="), Parameter))
+	{
+		GTraceAuxiliary.SetTruncateFile(FParse::Param(CommandLine, TEXT("tracefiletrunc")));
+		GTraceAuxiliary.Connect(ETraceConnectType::File, *Parameter);
+	}
+	else if (FParse::Param(CommandLine, TEXT("tracefile")))
+	{
+		GTraceAuxiliary.Connect(ETraceConnectType::File, nullptr);
+	}
+
 	// Initialize Trace
 	UE::Trace::FInitializeDesc Desc;
 	Desc.bUseWorkerThread = FPlatformProcess::SupportsMultithreading();
+	Desc.bUseImportantCache = (FParse::Param(CommandLine, TEXT("tracenocache")) == false);
 	if (FParse::Value(CommandLine, TEXT("-tracetailmb="), Desc.TailSizeBytes))
 	{
 		Desc.TailSizeBytes <<= 20;
@@ -723,28 +743,12 @@ void FTraceAuxiliary::Initialize(const TCHAR* CommandLine)
 	});
 
 	// Extract an explicit channel set from the command line.
-	FString Parameter;
 	if (FParse::Value(CommandLine, TEXT("-trace="), Parameter, false))
 	{
 		GTraceAuxiliary.AddChannels(*Parameter);
 		GTraceAuxiliary.EnableChannels();
 	}
 
-	// Attempt to send trace data somewhere from the command line
-	if (FParse::Value(CommandLine, TEXT("-tracehost="), Parameter))
-	{
-		GTraceAuxiliary.Connect(ETraceConnectType::Network, *Parameter);
-	}
-	else if (FParse::Value(CommandLine, TEXT("-tracefile="), Parameter))
-	{
-		GTraceAuxiliary.SetTruncateFile(FParse::Param(CommandLine, TEXT("tracefiletrunc")));
-		GTraceAuxiliary.Connect(ETraceConnectType::File, *Parameter);
-	}
-	else if (FParse::Param(CommandLine, TEXT("tracefile")))
-	{
-		GTraceAuxiliary.Connect(ETraceConnectType::File, nullptr);
-	}
-
 	UE::Trace::ThreadRegister(TEXT("GameThread"), FPlatformTLS::GetCurrentThreadId(), -1);
 #endif
 }
