CompressedBuffer - decompress at offset

#rnx

diff --git a/Engine/Source/Runtime/Core/Private/Compression/CompressedBuffer.cpp b/Engine/Source/Runtime/Core/Private/Compression/CompressedBuffer.cpp
--- a/Engine/Source/Runtime/Core/Private/Compression/CompressedBuffer.cpp
+++ b/Engine/Source/Runtime/Core/Private/Compression/CompressedBuffer.cpp
@@ -133,7 +133,7 @@ class FDecoder
 {
 public:
 	virtual FCompositeBuffer Decompress(const FHeader& Header, const FCompositeBuffer& CompressedData) const = 0;
-	virtual bool TryDecompressTo(const FHeader& Header, const FCompositeBuffer& CompressedData, FMutableMemoryView RawView) const = 0;
+	virtual bool TryDecompressTo(const FHeader& Header, const FCompositeBuffer& CompressedData, FMutableMemoryView RawView, uint64 RawOffset) const = 0;
 	virtual uint64 GetHeaderSize(const FHeader& Header) const = 0;
 };
 
@@ -171,14 +171,14 @@ public:
 		return FCompositeBuffer();
 	}
 
-	bool TryDecompressTo(const FHeader& Header, const FCompositeBuffer& CompressedData, FMutableMemoryView RawView) const final
+	bool TryDecompressTo(const FHeader& Header, const FCompositeBuffer& CompressedData, FMutableMemoryView RawView, uint64 RawOffset) const final
 	{
 		if (Header.Method == EMethod::None &&
-			Header.TotalRawSize == RawView.GetSize() &&
+			RawOffset + RawView.GetSize() <= Header.TotalRawSize &&
 			Header.TotalCompressedSize == CompressedData.GetSize() &&
 			Header.TotalCompressedSize == Header.TotalRawSize + sizeof(FHeader))
 		{
-			CompressedData.CopyTo(RawView, sizeof(FHeader));
+			CompressedData.CopyTo(RawView, sizeof(FHeader) + RawOffset);
 			return true;
 		}
 		return false;
@@ -298,7 +298,7 @@ class FBlockDecoder : public FDecoder
 {
 public:
 	FCompositeBuffer Decompress(const FHeader& Header, const FCompositeBuffer& CompressedData) const final;
-	bool TryDecompressTo(const FHeader& Header, const FCompositeBuffer& CompressedData, FMutableMemoryView RawView) const final;
+	bool TryDecompressTo(const FHeader& Header, const FCompositeBuffer& CompressedData, FMutableMemoryView RawView, uint64 RawOffset) const final;
 
 	uint64 GetHeaderSize(const FHeader& Header) const final
 	{
@@ -322,7 +322,7 @@ FCompositeBuffer FBlockDecoder::Decompress(const FHeader& Header, const FComposi
 	if (!CompressedData.IsOwned() || Header.TotalRawSize == 0)
 	{
 		FUniqueBuffer Buffer = FUniqueBuffer::Alloc(Header.TotalRawSize);
-		return TryDecompressTo(Header, CompressedData, Buffer) ? FCompositeBuffer(Buffer.MoveToShared()) : FCompositeBuffer();
+		return TryDecompressTo(Header, CompressedData, Buffer, 0) ? FCompositeBuffer(Buffer.MoveToShared()) : FCompositeBuffer();
 	}
 
 	TArray64<uint32> CompressedBlockSizes;
@@ -422,49 +422,81 @@ FCompositeBuffer FBlockDecoder::Decompress(const FHeader& Header, const FComposi
 	return FCompositeBuffer(MoveTemp(Segments));
 }
 
-bool FBlockDecoder::TryDecompressTo(const FHeader& Header, const FCompositeBuffer& CompressedData, FMutableMemoryView RawView) const
+bool FBlockDecoder::TryDecompressTo(const FHeader& Header, const FCompositeBuffer& CompressedData, FMutableMemoryView RawView, uint64 RawOffset) const
 {
-	if (Header.TotalRawSize != RawView.GetSize() ||
-		Header.TotalCompressedSize != CompressedData.GetSize())
+	if (Header.TotalRawSize < RawOffset + RawView.GetSize() || Header.TotalCompressedSize != CompressedData.GetSize())
 	{
 		return false;
 	}
 
-	TArray64<uint32> CompressedBlockSizes;
-	CompressedBlockSizes.AddUninitialized(Header.BlockCount);
-	CompressedData.CopyTo(MakeMemoryView(CompressedBlockSizes), sizeof(FHeader));
-	Algo::ForEach(CompressedBlockSizes, [](uint32& Size) { Size = NETWORK_ORDER32(Size); });
+	const uint64 BlockSize = uint64(1) << Header.BlockSizeExponent;
+
+	FUniqueBuffer BlockSizeBuffer;
+	FMemoryView BlockSizeView = CompressedData.ViewOrCopyRange(sizeof(FHeader), Header.BlockCount * sizeof(uint32), BlockSizeBuffer);
+	TArrayView64<uint32 const> CompressedBlockSizes(reinterpret_cast<const uint32*>(BlockSizeView.GetData()), Header.BlockCount);
 
 	FUniqueBuffer CompressedBlockCopy;
-	const uint64 BlockSize = uint64(1) << Header.BlockSizeExponent;
-	uint64 CompressedOffset = sizeof(FHeader) + uint64(Header.BlockCount) * sizeof(uint32);
-	uint64 RemainingRawSize = Header.TotalRawSize;
-	uint64 RemainingCompressedSize = CompressedData.GetSize();
-	for (uint32 CompressedBlockSize : CompressedBlockSizes)
+	FUniqueBuffer UncompressedBlockCopy;
+
+	const uint64 FirstBlockIndex	= uint64(RawOffset / BlockSize);
+	const uint64 LastBlockIndex		= uint64((RawOffset + RawView.GetSize() - 1) / BlockSize);
+	const uint64 LastBlockSize		= BlockSize - ((Header.BlockCount * BlockSize) - Header.TotalRawSize);
+	uint64 OffsetInFirstBlock		= RawOffset % BlockSize;
+	uint64 CompressedOffset			= sizeof(FHeader) + uint64(Header.BlockCount) * sizeof(uint32);
+	uint64 RemainingRawSize			= RawView.GetSize();
+
+	for (uint64 BlockIndex = 0; BlockIndex < FirstBlockIndex; BlockIndex++)
 	{
-		if (RemainingCompressedSize < CompressedBlockSize)
-		{
-			return false;
-		}
+		const uint32 CompressedBlockSize = NETWORK_ORDER32(CompressedBlockSizes[BlockIndex]);
+		CompressedOffset += CompressedBlockSize;
+	}
 
-		const uint64 RawBlockSize = FMath::Min(RemainingRawSize, BlockSize);
-		if (RawBlockSize == CompressedBlockSize)
-		{
-			CompressedData.CopyTo(RawView.Left(RawBlockSize), CompressedOffset);
-		}
-		else
+	for (uint64 BlockIndex = FirstBlockIndex; BlockIndex <= LastBlockIndex; BlockIndex++)
+	{
+		const uint64 UncompressedBlockSize	= BlockIndex == Header.BlockCount - 1 ? LastBlockSize : BlockSize;
+		const uint32 CompressedBlockSize	= NETWORK_ORDER32(CompressedBlockSizes[BlockIndex]);
+		const bool IsCompressed				= CompressedBlockSize < UncompressedBlockSize;
+
+		const uint64 BytesToUncompress = OffsetInFirstBlock > 0 
+			? FMath::Min<uint64>(RawView.GetSize(), UncompressedBlockSize - OffsetInFirstBlock)
+			: FMath::Min<uint64>(RemainingRawSize, BlockSize);
+
+		FMemoryView CompressedBlock = CompressedData.ViewOrCopyRange(CompressedOffset, CompressedBlockSize, CompressedBlockCopy);
+
+		if (IsCompressed)
 		{
-			const FMemoryView CompressedBlock = CompressedData.ViewOrCopyRange(CompressedOffset, CompressedBlockSize, CompressedBlockCopy);
-			if (!DecompressBlock(RawView.Left(RawBlockSize), CompressedBlock))
+			FMutableMemoryView UncompressedBlock = RawView.Left(BytesToUncompress);
+
+			const bool bIsAligned = BytesToUncompress == UncompressedBlockSize;
+			if (!bIsAligned)
+			{
+				// Decompress to a temporary buffer when the first or the last block reads are not aligned with the block boundaries.
+				if (UncompressedBlockCopy.IsNull())
+				{
+					UncompressedBlockCopy = FUniqueBuffer::Alloc(BlockSize);
+				}
+				UncompressedBlock = UncompressedBlockCopy.GetView().Left(UncompressedBlockSize);
+			}
+
+			if (!DecompressBlock(UncompressedBlock, CompressedBlock))
 			{
 				return false;
 			}
+
+			if (!bIsAligned)
+			{
+				RawView.CopyFrom(UncompressedBlock.Mid(OffsetInFirstBlock, BytesToUncompress));
+			}
+		}
+		else
+		{
+			RawView.CopyFrom(CompressedBlock.Mid(OffsetInFirstBlock, BytesToUncompress));
 		}
 
-		RemainingCompressedSize -= CompressedBlockSize;
-		RemainingRawSize -= RawBlockSize;
+		OffsetInFirstBlock = 0;
+		RemainingRawSize -= BytesToUncompress;
 		CompressedOffset += CompressedBlockSize;
-		RawView += RawBlockSize;
+		RawView += BytesToUncompress;
 	}
 
 	return RemainingRawSize == 0;
@@ -751,7 +783,7 @@ bool FCompressedBuffer::TryGetCompressParameters(
 	return false;
 }
 
-bool FCompressedBuffer::TryDecompressTo(FMutableMemoryView RawView) const
+bool FCompressedBuffer::TryDecompressTo(FMutableMemoryView RawView, uint64 RawOffset) const
 {
 	using namespace UE::CompressedBuffer;
 	if (CompressedData)
@@ -759,13 +791,13 @@ bool FCompressedBuffer::TryDecompressTo(FMutableMemoryView RawView) const
 		const FHeader Header = FHeader::Read(CompressedData);
 		if (const FDecoder* const Decoder = GetDecoder(Header.Method))
 		{
-			return Decoder->TryDecompressTo(Header, CompressedData, RawView);
+			return Decoder->TryDecompressTo(Header, CompressedData, RawView, RawOffset);
 		}
 	}
 	return false;
 }
 
-FSharedBuffer FCompressedBuffer::Decompress() const
+FSharedBuffer FCompressedBuffer::Decompress(uint64 RawOffset, uint64 RawSize) const
 {
 	using namespace UE::CompressedBuffer;
 	if (CompressedData)
@@ -773,12 +805,9 @@ FSharedBuffer FCompressedBuffer::Decompress() const
 		const FHeader Header = FHeader::Read(CompressedData);
 		if (const FDecoder* const Decoder = GetDecoder(Header.Method))
 		{
-			if (Header.Method == EMethod::None)
-			{
-				return Decoder->Decompress(Header, CompressedData).ToShared();
-			}
-			FUniqueBuffer RawData = FUniqueBuffer::Alloc(Header.TotalRawSize);
-			if (Decoder->TryDecompressTo(Header, CompressedData, RawData))
+			const uint64 TotalRawSize = RawSize < ~uint64(0) ? RawSize : Header.TotalRawSize - RawOffset;
+			FUniqueBuffer RawData = FUniqueBuffer::Alloc(TotalRawSize);
+			if (Decoder->TryDecompressTo(Header, CompressedData, RawData, RawOffset))
 			{
 				return RawData.MoveToShared();
 			}
diff --git a/Engine/Source/Runtime/Core/Private/Tests/Compression/CompressedBufferTest.cpp b/Engine/Source/Runtime/Core/Private/Tests/Compression/CompressedBufferTest.cpp
new file mode 100644
--- /dev/null
+++ b/Engine/Source/Runtime/Core/Private/Tests/Compression/CompressedBufferTest.cpp
@@ -0,0 +1,137 @@
+// Copyright Epic Games, Inc. All Rights Reserved.
+
+#include "Compression/CompressedBuffer.h"
+#include "Compression/OodleDataCompression.h"
+
+#include "Misc/AutomationTest.h"
+
+PRAGMA_DISABLE_OPTIMIZATION
+///////////////////////////////////////////////////////////////////////////////////////////////////////////////////////
+
+IMPLEMENT_SIMPLE_AUTOMATION_TEST(FCompressedBufferDecompressTest, "System.Core.Compression.CompressedBufferDecompress", EAutomationTestFlags::ApplicationContextMask | EAutomationTestFlags::SmokeFilter)
+bool FCompressedBufferDecompressTest::RunTest(const FString& Parameters)
+{
+	auto GenerateData = [](int32 N) -> TArray<uint64> {
+		TArray<uint64> Data;
+		Data.SetNum(N);
+		for (int32 Idx = 0; Idx < Data.Num(); ++Idx)
+		{
+			Data[Idx] = uint64(Idx);
+		}
+		return Data;
+	};
+
+	auto ValidateData = [this](TArrayView<uint64 const> Values, TArrayView<uint64 const> ExpectedValues, int32 Offset) {
+		int32 ExpectedIndex = Offset;	
+		for (uint64 Value : Values)
+		{
+			const uint64 ExpectedValue = ExpectedValues[ExpectedIndex];
+			TestEqual("UncompressedValue", Value, ExpectedValue);
+			ExpectedIndex++;
+		}
+	};
+
+	// Test decompress with offset and size
+	{
+		auto UncompressAndValidate = [this, &ValidateData](
+			FCompressedBuffer Compressed,
+			int32 OffsetCount,
+			int32 Count,
+			const TArrayView<uint64 const>& ExpectedValues) {
+			
+			FSharedBuffer Uncompressed = Compressed.Decompress(OffsetCount * sizeof(uint64), Count * sizeof(uint64));
+			TArrayView<uint64 const> Values((const uint64*)Uncompressed.GetData(), int32(Uncompressed.GetSize() / sizeof(uint64)));
+			TestEqual("UncompressedCount", Values.Num(), Count);
+			ValidateData(Values, ExpectedValues, OffsetCount);
+		};
+
+		const uint64 BlockSize			= 64 * sizeof(uint64);
+		const int32 N					= 5000;
+		TArray<uint64> ExpectedValues	= GenerateData(N);
+		
+		FCompressedBuffer Compressed = FCompressedBuffer::Compress(
+			FSharedBuffer::MakeView(MakeMemoryView(ExpectedValues)),
+			FOodleDataCompression::ECompressor::Mermaid,
+			FOodleDataCompression::ECompressionLevel::Optimal4,
+			BlockSize);
+
+		UncompressAndValidate(Compressed, 0, N, ExpectedValues);
+		UncompressAndValidate(Compressed, 1, N - 1, ExpectedValues);
+		UncompressAndValidate(Compressed, N - 1, 1, ExpectedValues);
+		UncompressAndValidate(Compressed, 0, 1, ExpectedValues);
+		UncompressAndValidate(Compressed, 2, 4, ExpectedValues);
+		UncompressAndValidate(Compressed, 0, 512, ExpectedValues);
+		UncompressAndValidate(Compressed, 3, 514, ExpectedValues);
+		UncompressAndValidate(Compressed, 256, 512, ExpectedValues);
+		UncompressAndValidate(Compressed, 512, 512, ExpectedValues);
+		UncompressAndValidate(Compressed, 512, 512, ExpectedValues);
+		UncompressAndValidate(Compressed, 4993, 4, ExpectedValues);
+	}
+
+	// Decompress with offset only
+	{
+		const uint64 BlockSize			= 64 * sizeof(uint64);
+		const int32 N					= 1000;
+		TArray<uint64> ExpectedValues	= GenerateData(N);
+		
+		FCompressedBuffer Compressed = FCompressedBuffer::Compress(
+			FSharedBuffer::MakeView(MakeMemoryView(ExpectedValues)),
+			FOodleDataCompression::ECompressor::Mermaid,
+			FOodleDataCompression::ECompressionLevel::Optimal4,
+			BlockSize);
+		
+		const uint64				OffsetCount	 = 150;
+		FSharedBuffer				Uncompressed = Compressed.Decompress(OffsetCount * sizeof(uint64));
+		TArrayView<uint64 const>	Values((const uint64*)Uncompressed.GetData(), int32(Uncompressed.GetSize() / sizeof(uint64)));
+		ValidateData(Values, ExpectedValues, OffsetCount);
+	}
+
+	// Only one block
+	{
+		const uint64 BlockSize			= 256 * sizeof(uint64);
+		const int32 N					= 100;
+		TArray<uint64> ExpectedValues	= GenerateData(N);
+		
+		FCompressedBuffer Compressed = FCompressedBuffer::Compress(
+			FSharedBuffer::MakeView(MakeMemoryView(ExpectedValues)),
+			FOodleDataCompression::ECompressor::Mermaid,
+			FOodleDataCompression::ECompressionLevel::Optimal4,
+			BlockSize);
+		
+		const uint64 OffsetCount		= 2;
+		const uint64 Count				= 50;
+		FSharedBuffer Uncompressed		= Compressed.Decompress(OffsetCount * sizeof(uint64), Count * sizeof(uint64));
+		TArrayView<uint64 const> Values((const uint64*)Uncompressed.GetData(), int32(Uncompressed.GetSize() / sizeof(uint64)));
+		ValidateData(Values, ExpectedValues, OffsetCount);
+	}
+
+	// Uncompressed
+	{
+		const int32 N					= 4242;
+		TArray<uint64> ExpectedValues	= GenerateData(N);
+		
+		FCompressedBuffer Compressed = FCompressedBuffer::Compress(
+			FSharedBuffer::MakeView(MakeMemoryView(ExpectedValues)),
+			FOodleDataCompression::ECompressor::NotSet,
+			FOodleDataCompression::ECompressionLevel::None);
+		
+		{
+			const uint64 OffsetCount = 0;
+			const uint64 Count = N;
+			FSharedBuffer Uncompressed = Compressed.Decompress(OffsetCount * sizeof(uint64), Count * sizeof(uint64));
+			TArrayView<uint64 const> Values((const uint64*)Uncompressed.GetData(), int32(Uncompressed.GetSize() / sizeof(uint64)));
+			ValidateData(Values, ExpectedValues, OffsetCount);
+		}
+
+		{
+			const uint64 OffsetCount = 21;
+			const uint64 Count = 999;
+			FSharedBuffer Uncompressed = Compressed.Decompress(OffsetCount * sizeof(uint64), Count * sizeof(uint64));
+			TArrayView<uint64 const> Values((const uint64*)Uncompressed.GetData(), int32(Uncompressed.GetSize() / sizeof(uint64)));
+			ValidateData(Values, ExpectedValues, OffsetCount);
+		}
+	}
+
+	return true;
+}
+PRAGMA_ENABLE_OPTIMIZATION
diff --git a/Engine/Source/Runtime/Core/Public/Compression/CompressedBuffer.h b/Engine/Source/Runtime/Core/Public/Compression/CompressedBuffer.h
--- a/Engine/Source/Runtime/Core/Public/Compression/CompressedBuffer.h
+++ b/Engine/Source/Runtime/Core/Public/Compression/CompressedBuffer.h
@@ -109,16 +109,16 @@ public:
 		ECompressedBufferCompressionLevel& OutCompressionLevel) const;
 
 	/**
-	 * Decompress into a memory view that is exactly GetRawSize() bytes.
+	 * Decompress into a memory view that is less or equal to GetRawSize()
 	 */
-	[[nodiscard]] CORE_API bool TryDecompressTo(FMutableMemoryView RawView) const;
+	[[nodiscard]] CORE_API bool TryDecompressTo(FMutableMemoryView RawView, uint64 RawOffset = 0) const;
 
 	/**
 	 * Decompress into an owned buffer.
 	 *
 	 * @return An owned buffer containing the raw data, or null on error.
 	 */
-	[[nodiscard]] CORE_API FSharedBuffer Decompress() const;
+	[[nodiscard]] CORE_API FSharedBuffer Decompress(uint64 RawOffset = 0, uint64 RawSize = MAX_uint64) const;
 
 	/**
 	 * Decompress into an owned composite buffer.
