CompressedBuffer: Added additional basic tests

#rnx

diff --git a/Engine/Source/Runtime/Core/Private/Tests/Compression/CompressedBufferTest.cpp b/Engine/Source/Runtime/Core/Private/Tests/Compression/CompressedBufferTest.cpp
--- a/Engine/Source/Runtime/Core/Private/Tests/Compression/CompressedBufferTest.cpp
+++ b/Engine/Source/Runtime/Core/Private/Tests/Compression/CompressedBufferTest.cpp
@@ -1,17 +1,103 @@
 // Copyright Epic Games, Inc. All Rights Reserved.
 
 #include "Compression/CompressedBuffer.h"
-#include "Compression/OodleDataCompression.h"
 
+#include "Compression/OodleDataCompression.h"
+#include "Hash/Blake3.h"
 #include "Misc/AutomationTest.h"
 
-PRAGMA_DISABLE_OPTIMIZATION
-///////////////////////////////////////////////////////////////////////////////////////////////////////////////////////
+#if WITH_DEV_AUTOMATION_TESTS
+
+IMPLEMENT_SIMPLE_AUTOMATION_TEST(FCompressedBufferTest, "System.Core.Compression.CompressedBuffer", EAutomationTestFlags::ApplicationContextMask | EAutomationTestFlags::SmokeFilter)
+bool FCompressedBufferTest::RunTest(const FString& Parameters)
+{
+	const uint8 ZeroBuffer[1024]{};
+	const FBlake3Hash ZeroBufferHash = FBlake3::HashBuffer(MakeMemoryView(ZeroBuffer));
+
+	// Test Null
+	{
+		const FCompressedBuffer Buffer;
+		TestFalse(TEXT("FCompressedBuffer()"), (bool)Buffer);
+		TestTrue(TEXT("FCompressedBuffer().IsNull()"), Buffer.IsNull());
+		TestTrue(TEXT("FCompressedBuffer().IsOwned()"), Buffer.IsOwned());
+		TestEqual(TEXT("FCompressedBuffer().GetCompressedSize()"), Buffer.GetCompressedSize(), uint64(0));
+		TestEqual(TEXT("FCompressedBuffer().GetRawSize()"), Buffer.GetRawSize(), uint64(0));
+		TestEqual(TEXT("FCompressedBuffer().GetRawHash()"), Buffer.GetRawHash(), FBlake3Hash::Zero);
+		TestTrue(TEXT("FCompressedBuffer().Decompress()"), Buffer.Decompress().IsNull());
+		TestTrue(TEXT("FCompressedBuffer().DecompressToComposite()"), Buffer.DecompressToComposite().IsNull());
+	}
+
+	// Test Method None
+	{
+		ECompressedBufferCompressor Compressor;
+		ECompressedBufferCompressionLevel CompressionLevel;
+
+		const FCompressedBuffer Buffer = FCompressedBuffer::Compress(FSharedBuffer::MakeView(MakeMemoryView(ZeroBuffer)),
+			ECompressedBufferCompressor::NotSet, ECompressedBufferCompressionLevel::None);
+		TestTrue(TEXT("FCompressedBuffer::Compress(None)"), (bool)Buffer);
+		TestFalse(TEXT("FCompressedBuffer::Compress(None).IsNull()"), Buffer.IsNull());
+		TestTrue(TEXT("FCompressedBuffer::Compress(None).IsOwned()"), Buffer.IsOwned());
+		TestEqual(TEXT("FCompressedBuffer::Compress(None).GetCompressedSize()"), Buffer.GetCompressedSize(), sizeof(ZeroBuffer) + 64);
+		TestEqual(TEXT("FCompressedBuffer::Compress(None).GetRawSize()"), Buffer.GetRawSize(), sizeof(ZeroBuffer));
+		TestEqual(TEXT("FCompressedBuffer::Compress(None).GetRawHash()"), Buffer.GetRawHash(), ZeroBufferHash);
+		TestEqual(TEXT("FCompressedBuffer::Compress(None).Decompress()"), FBlake3::HashBuffer(Buffer.Decompress()), ZeroBufferHash);
+		TestEqual(TEXT("FCompressedBuffer::Compress(None).DecompressToComposite()"), FBlake3::HashBuffer(Buffer.DecompressToComposite()), Buffer.GetRawHash());
+		TestTrue(TEXT("FCompressedBuffer::Compress(None).TryGetCompressParameters()"), Buffer.TryGetCompressParameters(Compressor, CompressionLevel) &&
+			Compressor == ECompressedBufferCompressor::NotSet && CompressionLevel == ECompressedBufferCompressionLevel::None);
+
+		const FCompressedBuffer BufferCopy = FCompressedBuffer::FromCompressed(Buffer.GetCompressed());
+		TestTrue(TEXT("FCompressedBuffer::Compress(None, Copy)"), (bool)Buffer);
+		TestFalse(TEXT("FCompressedBuffer::Compress(None, Copy).IsNull()"), Buffer.IsNull());
+		TestTrue(TEXT("FCompressedBuffer::Compress(None, Copy).IsOwned()"), Buffer.IsOwned());
+		TestEqual(TEXT("FCompressedBuffer::Compress(None, Copy).GetCompressedSize()"), Buffer.GetCompressedSize(), sizeof(ZeroBuffer) + 64);
+		TestEqual(TEXT("FCompressedBuffer::Compress(None, Copy).GetRawSize()"), Buffer.GetRawSize(), sizeof(ZeroBuffer));
+		TestEqual(TEXT("FCompressedBuffer::Compress(None, Copy).GetRawHash()"), Buffer.GetRawHash(), ZeroBufferHash);
+		TestEqual(TEXT("FCompressedBuffer::Compress(None, Copy).Decompress()"), FBlake3::HashBuffer(Buffer.Decompress()), ZeroBufferHash);
+		TestEqual(TEXT("FCompressedBuffer::Compress(None, Copy).DecompressToComposite()"), FBlake3::HashBuffer(Buffer.DecompressToComposite()), Buffer.GetRawHash());
+		TestTrue(TEXT("FCompressedBuffer::Compress(None, Copy).TryGetCompressParameters()"), Buffer.TryGetCompressParameters(Compressor, CompressionLevel) &&
+			Compressor == ECompressedBufferCompressor::NotSet && CompressionLevel == ECompressedBufferCompressionLevel::None);
+	}
+
+	// Test Method Oodle
+	{
+		
+		ECompressedBufferCompressor Compressor;
+		ECompressedBufferCompressionLevel CompressionLevel;
+
+		const FCompressedBuffer Buffer = FCompressedBuffer::Compress(FSharedBuffer::MakeView(MakeMemoryView(ZeroBuffer)),
+			ECompressedBufferCompressor::Mermaid, ECompressedBufferCompressionLevel::VeryFast);
+		TestTrue(TEXT("FCompressedBuffer::Compress(Oodle)"), (bool)Buffer);
+		TestFalse(TEXT("FCompressedBuffer::Compress(Oodle).IsNull()"), Buffer.IsNull());
+		TestTrue(TEXT("FCompressedBuffer::Compress(Oodle).IsOwned()"), Buffer.IsOwned());
+		TestTrue(TEXT("FCompressedBuffer::Compress(Oodle).GetCompressedSize()"), Buffer.GetCompressedSize() < sizeof(ZeroBuffer));
+		TestEqual(TEXT("FCompressedBuffer::Compress(Oodle).GetRawSize()"), Buffer.GetRawSize(), sizeof(ZeroBuffer));
+		TestEqual(TEXT("FCompressedBuffer::Compress(Oodle).GetRawHash()"), Buffer.GetRawHash(), ZeroBufferHash);
+		TestEqual(TEXT("FCompressedBuffer::Compress(Oodle).Decompress()"), FBlake3::HashBuffer(Buffer.Decompress()), ZeroBufferHash);
+		TestEqual(TEXT("FCompressedBuffer::Compress(Oodle).DecompressToComposite()"), FBlake3::HashBuffer(Buffer.DecompressToComposite()), Buffer.GetRawHash());
+		TestTrue(TEXT("FCompressedBuffer::Compress(Oodle).TryGetCompressParameters()"), Buffer.TryGetCompressParameters(Compressor, CompressionLevel) &&
+			Compressor == ECompressedBufferCompressor::Mermaid && CompressionLevel == ECompressedBufferCompressionLevel::VeryFast);
+
+		const FCompressedBuffer BufferCopy = FCompressedBuffer::FromCompressed(Buffer.GetCompressed());
+		TestTrue(TEXT("FCompressedBuffer::Compress(Oodle, Copy)"), (bool)Buffer);
+		TestFalse(TEXT("FCompressedBuffer::Compress(Oodle, Copy).IsNull()"), Buffer.IsNull());
+		TestTrue(TEXT("FCompressedBuffer::Compress(Oodle, Copy).IsOwned()"), Buffer.IsOwned());
+		TestTrue(TEXT("FCompressedBuffer::Compress(Oodle, Copy).GetCompressedSize()"), Buffer.GetCompressedSize() < sizeof(ZeroBuffer));
+		TestEqual(TEXT("FCompressedBuffer::Compress(Oodle, Copy).GetRawSize()"), Buffer.GetRawSize(), sizeof(ZeroBuffer));
+		TestEqual(TEXT("FCompressedBuffer::Compress(Oodle, Copy).GetRawHash()"), Buffer.GetRawHash(), ZeroBufferHash);
+		TestEqual(TEXT("FCompressedBuffer::Compress(Oodle, Copy).Decompress()"), FBlake3::HashBuffer(Buffer.Decompress()), ZeroBufferHash);
+		TestEqual(TEXT("FCompressedBuffer::Compress(Oodle, Copy).DecompressToComposite()"), FBlake3::HashBuffer(Buffer.DecompressToComposite()), Buffer.GetRawHash());
+		TestTrue(TEXT("FCompressedBuffer::Compress(Oodle, Copy).TryGetCompressParameters()"), Buffer.TryGetCompressParameters(Compressor, CompressionLevel) &&
+			Compressor == ECompressedBufferCompressor::Mermaid && CompressionLevel == ECompressedBufferCompressionLevel::VeryFast);
+	}
+
+	return true;
+}
 
 IMPLEMENT_SIMPLE_AUTOMATION_TEST(FCompressedBufferDecompressTest, "System.Core.Compression.CompressedBufferDecompress", EAutomationTestFlags::ApplicationContextMask | EAutomationTestFlags::SmokeFilter)
 bool FCompressedBufferDecompressTest::RunTest(const FString& Parameters)
 {
-	auto GenerateData = [](int32 N) -> TArray<uint64> {
+	const auto GenerateData = [](int32 N) -> TArray<uint64>
+	{
 		TArray<uint64> Data;
 		Data.SetNum(N);
 		for (int32 Idx = 0; Idx < Data.Num(); ++Idx)
@@ -21,8 +107,9 @@ bool FCompressedBufferDecompressTest::RunTest(const FString& Parameters)
 		return Data;
 	};
 
-	auto ValidateData = [this](TArrayView<uint64 const> Values, TArrayView<uint64 const> ExpectedValues, int32 Offset) {
-		int32 ExpectedIndex = Offset;	
+	const auto ValidateData = [this](TConstArrayView<uint64> Values, TConstArrayView<uint64> ExpectedValues, int32 Offset)
+	{
+		int32 ExpectedIndex = Offset;
 		for (uint64 Value : Values)
 		{
 			const uint64 ExpectedValue = ExpectedValues[ExpectedIndex];
@@ -31,25 +118,30 @@ bool FCompressedBufferDecompressTest::RunTest(const FString& Parameters)
 		}
 	};
 
+	const auto CastToArrayView = [](FMemoryView View) -> TConstArrayView<uint64>
+	{
+		return MakeArrayView(static_cast<const uint64*>(View.GetData()), static_cast<int32>(View.GetSize() / sizeof(uint64)));
+	};
+
 	// Test decompress with offset and size
 	{
-		auto UncompressAndValidate = [this, &ValidateData](
-			FCompressedBuffer Compressed,
-			int32 OffsetCount,
-			int32 Count,
-			const TArrayView<uint64 const>& ExpectedValues) {
-			
-			FSharedBuffer Uncompressed = Compressed.Decompress(OffsetCount * sizeof(uint64), Count * sizeof(uint64));
-			TArrayView<uint64 const> Values((const uint64*)Uncompressed.GetData(), int32(Uncompressed.GetSize() / sizeof(uint64)));
+		const auto UncompressAndValidate = [this, &ValidateData, &CastToArrayView](
+			const FCompressedBuffer& Compressed,
+			const int32 OffsetCount,
+			const int32 Count,
+			const TConstArrayView<uint64> ExpectedValues)
+		{
+			const FSharedBuffer Uncompressed = Compressed.Decompress(OffsetCount * sizeof(uint64), Count * sizeof(uint64));
+			const TConstArrayView<uint64> Values = CastToArrayView(Uncompressed);
 			TestEqual("UncompressedCount", Values.Num(), Count);
 			ValidateData(Values, ExpectedValues, OffsetCount);
 		};
 
-		const uint64 BlockSize			= 64 * sizeof(uint64);
-		const int32 N					= 5000;
-		TArray<uint64> ExpectedValues	= GenerateData(N);
-		
-		FCompressedBuffer Compressed = FCompressedBuffer::Compress(
+		constexpr uint64 BlockSize = 64 * sizeof(uint64);
+		constexpr int32 N = 5000;
+		const TArray<uint64> ExpectedValues = GenerateData(N);
+
+		const FCompressedBuffer Compressed = FCompressedBuffer::Compress(
 			FSharedBuffer::MakeView(MakeMemoryView(ExpectedValues)),
 			FOodleDataCompression::ECompressor::Mermaid,
 			FOodleDataCompression::ECompressionLevel::Optimal4,
@@ -70,68 +162,69 @@ bool FCompressedBufferDecompressTest::RunTest(const FString& Parameters)
 
 	// Decompress with offset only
 	{
-		const uint64 BlockSize			= 64 * sizeof(uint64);
-		const int32 N					= 1000;
-		TArray<uint64> ExpectedValues	= GenerateData(N);
-		
-		FCompressedBuffer Compressed = FCompressedBuffer::Compress(
+		constexpr uint64 BlockSize = 64 * sizeof(uint64);
+		constexpr int32 N = 1000;
+		const TArray<uint64> ExpectedValues = GenerateData(N);
+
+		const FCompressedBuffer Compressed = FCompressedBuffer::Compress(
 			FSharedBuffer::MakeView(MakeMemoryView(ExpectedValues)),
 			FOodleDataCompression::ECompressor::Mermaid,
 			FOodleDataCompression::ECompressionLevel::Optimal4,
 			BlockSize);
 		
-		const uint64				OffsetCount	 = 150;
-		FSharedBuffer				Uncompressed = Compressed.Decompress(OffsetCount * sizeof(uint64));
-		TArrayView<uint64 const>	Values((const uint64*)Uncompressed.GetData(), int32(Uncompressed.GetSize() / sizeof(uint64)));
+		constexpr uint64 OffsetCount = 150;
+		const FSharedBuffer Uncompressed = Compressed.Decompress(OffsetCount * sizeof(uint64));
+		const TConstArrayView<uint64> Values = CastToArrayView(Uncompressed);
 		ValidateData(Values, ExpectedValues, OffsetCount);
 	}
 
 	// Only one block
 	{
-		const uint64 BlockSize			= 256 * sizeof(uint64);
-		const int32 N					= 100;
-		TArray<uint64> ExpectedValues	= GenerateData(N);
-		
-		FCompressedBuffer Compressed = FCompressedBuffer::Compress(
+		constexpr uint64 BlockSize = 256 * sizeof(uint64);
+		constexpr int32 N = 100;
+		const TArray<uint64> ExpectedValues = GenerateData(N);
+
+		const FCompressedBuffer Compressed = FCompressedBuffer::Compress(
 			FSharedBuffer::MakeView(MakeMemoryView(ExpectedValues)),
 			FOodleDataCompression::ECompressor::Mermaid,
 			FOodleDataCompression::ECompressionLevel::Optimal4,
 			BlockSize);
-		
-		const uint64 OffsetCount		= 2;
-		const uint64 Count				= 50;
-		FSharedBuffer Uncompressed		= Compressed.Decompress(OffsetCount * sizeof(uint64), Count * sizeof(uint64));
-		TArrayView<uint64 const> Values((const uint64*)Uncompressed.GetData(), int32(Uncompressed.GetSize() / sizeof(uint64)));
+
+		constexpr uint64 OffsetCount = 2;
+		constexpr uint64 Count = 50;
+		const FSharedBuffer Uncompressed = Compressed.Decompress(OffsetCount * sizeof(uint64), Count * sizeof(uint64));
+		const TConstArrayView<uint64> Values = CastToArrayView(Uncompressed);
 		ValidateData(Values, ExpectedValues, OffsetCount);
 	}
 
 	// Uncompressed
 	{
-		const int32 N					= 4242;
-		TArray<uint64> ExpectedValues	= GenerateData(N);
-		
-		FCompressedBuffer Compressed = FCompressedBuffer::Compress(
+		constexpr int32 N = 4242;
+		const TArray<uint64> ExpectedValues = GenerateData(N);
+
+		const FCompressedBuffer Compressed = FCompressedBuffer::Compress(
 			FSharedBuffer::MakeView(MakeMemoryView(ExpectedValues)),
 			FOodleDataCompression::ECompressor::NotSet,
 			FOodleDataCompression::ECompressionLevel::None);
-		
+
 		{
-			const uint64 OffsetCount = 0;
-			const uint64 Count = N;
-			FSharedBuffer Uncompressed = Compressed.Decompress(OffsetCount * sizeof(uint64), Count * sizeof(uint64));
-			TArrayView<uint64 const> Values((const uint64*)Uncompressed.GetData(), int32(Uncompressed.GetSize() / sizeof(uint64)));
+			constexpr uint64 OffsetCount = 0;
+			constexpr uint64 Count = N;
+			const FSharedBuffer Uncompressed = Compressed.Decompress(OffsetCount * sizeof(uint64), Count * sizeof(uint64));
+			const TConstArrayView<uint64> Values = CastToArrayView(Uncompressed);
 			ValidateData(Values, ExpectedValues, OffsetCount);
 		}
 
 		{
-			const uint64 OffsetCount = 21;
-			const uint64 Count = 999;
-			FSharedBuffer Uncompressed = Compressed.Decompress(OffsetCount * sizeof(uint64), Count * sizeof(uint64));
-			TArrayView<uint64 const> Values((const uint64*)Uncompressed.GetData(), int32(Uncompressed.GetSize() / sizeof(uint64)));
+			constexpr uint64 OffsetCount = 21;
+			constexpr uint64 Count = 999;
+			const FSharedBuffer Uncompressed = Compressed.Decompress(OffsetCount * sizeof(uint64), Count * sizeof(uint64));
+			const TConstArrayView<uint64> Values = CastToArrayView(Uncompressed);
 			ValidateData(Values, ExpectedValues, OffsetCount);
 		}
 	}
 
 	return true;
 }
-PRAGMA_ENABLE_OPTIMIZATION
+
+#endif // WITH_DEV_AUTOMATION_TESTS
