MessageDialogTesting - non-unity compile fix

diff --git a/Engine/Source/Runtime/Core/Private/Misc/MessageDialogTesting.cpp b/Engine/Source/Runtime/Core/Private/Misc/MessageDialogTesting.cpp
--- a/Engine/Source/Runtime/Core/Private/Misc/MessageDialogTesting.cpp
+++ b/Engine/Source/Runtime/Core/Private/Misc/MessageDialogTesting.cpp
@@ -2,6 +2,7 @@
 
 #include "CoreMinimal.h"
 #include "Misc/MessageDialog.h"
+#include "HAL/IConsoleManager.h"
 
 #if (WITH_EDITOR || IS_PROGRAM) && !UE_BUILD_SHIPPING
 
