tuned FTSTicker stress test as it fails on Mac

diff --git a/Engine/Source/Runtime/Core/Private/Tests/Containers/TickerTests.cpp b/Engine/Source/Runtime/Core/Private/Tests/Containers/TickerTests.cpp
--- a/Engine/Source/Runtime/Core/Private/Tests/Containers/TickerTests.cpp
+++ b/Engine/Source/Runtime/Core/Private/Tests/Containers/TickerTests.cpp
@@ -212,7 +212,7 @@ bool FTSTickerTest::RunTest(const FString& Parameters)
 		);
 
 		TArray<FTask> Tasks;
-		Tasks.Reserve(1000);
+		Tasks.Reserve(500);
 		for (int i = 0; i != 10; ++i)
 		{
 			Tasks.Add(Launch(UE_SOURCE_LOCATION,
@@ -242,7 +242,7 @@ bool FTSTickerTest::RunTest(const FString& Parameters)
 		FPlatformProcess::Sleep(1.0f); // let it run for a while
 		bQuit = true;
 		Tasks.Add(TickTask);
-		verify(Wait(Tasks, FTimespan::FromSeconds(3)));
+		verify(Wait(Tasks, FTimespan::FromSeconds(5)));
 	}
 
 	UE_BENCHMARK(5, TickerPerfTest<200, 200>);
