StringView: Updated the class comment based on feedback

#rnx

diff --git a/Engine/Source/Runtime/Core/Public/Containers/StringView.h b/Engine/Source/Runtime/Core/Public/Containers/StringView.h
--- a/Engine/Source/Runtime/Core/Public/Containers/StringView.h
+++ b/Engine/Source/Runtime/Core/Public/Containers/StringView.h
@@ -61,47 +61,49 @@ namespace StringViewPrivate
 }
 
 /**
- * String View
+ * A string view is a non-owning view of a range of characters.
  *
- * A string view is implicitly constructible from const char* style strings and
- * from compatible character ranges such as FString and TStringBuilderBase.
+ * Ensure that the underlying string is valid for the lifetime of the string view.
  *
- * A string view does not own any data nor does it attempt to control any lifetimes, it
- * merely points at a subrange of characters in some other string. It's up to the user
- * to ensure the underlying string stays valid for the lifetime of the string view.
+ * Be careful when constructing a string view from a temporary. Make a local copy if necessary.
  *
- * A string view is cheap to copy and is intended to be passed by value.
+ * FStringView View = Object->GetPathName(); // Invalid
  *
- * A string view does not represent a NUL terminated string and therefore you should
- * never pass in the pointer returned by GetData() into a C-string API accepting only a
- * pointer. You must either use a string builder to make a properly terminated string,
- * or use an API that accepts a length argument in addition to the C-string.
+ * FString PathName = Object->GetPathName(); // Valid
+ * FStringView View = PathName;
  *
- * String views are a good fit for arguments to functions which don't wish to care
- * which style of string construction is used by the caller. If you accept strings via
- * string views then the caller is free to use FString, FStringBuilder, raw C strings,
- * or any other type which can be converted into a string view.
+ * void ProcessPath(FStringView Path);       // Valid
+ * ProcessPath(Object->GetPathName());
+ *
+ * A string view is implicitly constructible from null-terminated strings, from contiguous ranges
+ * of characters such as FString and TStringBuilder, and from literals such as TEXT("..."_SV).
+ *
+ * A string view is cheap to copy and is meant to be passed by value. Avoid passing by reference.
+ *
+ * A string view is not guaranteed to represent a null-terminated string.
+ *
+ * Log or format a string view using UE_LOG(TEXT("%.*s"), View.Len(), View.GetData());
+ *
+ * A string view is a good fit for function parameters where the function has no requirements for
+ * how the string is stored. A caller may use FString, FStringView, TStringBuilder, a char array,
+ * a null-terminated string, or any other type which can convert to a string view.
  *
  * The UE::String namespace contains many functions that can operate on string views.
- * Most of these can be found in String/___.h in Core.
+ * Most of these functions can be found in String/___.h in Core.
  *
  * @code
- *	void DoFoo(FStringView InString);
- *
- *	// DoFoo may be called as:
- *	void MultiFoo()
- *	{
- *		FString MyFoo(TEXT("Zoo"));
- *		const TCHAR* MyFooStr = *MyFoo;
+ *	void UseString(FStringView InString);
  *
- *		TStringBuilder<64> BuiltFoo;
- *		BuiltFoo.Append(TEXT("ABC"));
+ *	FString String(TEXT("ABC"));
+ *	const TCHAR* CString = *String;
+ *	TStringBuilder<16> StringBuilder;
+ *	StringBuilder.Append(TEXT("ABC"));
  *
- *		DoFoo(MyFoo);
- *		DoFoo(MyFooStr);
- *		DoFoo(TEXT("ABC"));
- *		DoFoo(BuiltFoo);
- *	}
+ *	UseString(String);
+ *	UseString(CString);
+ *	UseString(StringBuilder);
+ *	UseString(TEXT("ABC"));
+ *	UseString(TEXT("ABC"_SV));
  * @endcode
  */
 template <typename CharType>
