CompactBinary: Minor corrections to documentation and quick edits to simplify code

#rnx

diff --git a/Engine/Source/Runtime/Core/Private/Serialization/CompactBinaryPackage.cpp b/Engine/Source/Runtime/Core/Private/Serialization/CompactBinaryPackage.cpp
--- a/Engine/Source/Runtime/Core/Private/Serialization/CompactBinaryPackage.cpp
+++ b/Engine/Source/Runtime/Core/Private/Serialization/CompactBinaryPackage.cpp
@@ -249,8 +249,7 @@ void FCbPackage::AddAttachment(const FCbAttachment& Attachment, FAttachmentResol
 
 int32 FCbPackage::RemoveAttachment(const FIoHash& Hash)
 {
-	const int32 Index = Algo::BinarySearchBy(Attachments, Hash,
-		[](const FCbAttachment& Attachment) -> FIoHash { return Attachment.GetHash(); });
+	const int32 Index = Algo::BinarySearchBy(Attachments, Hash, &FCbAttachment::GetHash);
 	if (Attachments.IsValidIndex(Index))
 	{
 		Attachments.RemoveAt(Index);
@@ -266,8 +265,7 @@ bool FCbPackage::Equals(const FCbPackage& Package) const
 
 const FCbAttachment* FCbPackage::FindAttachment(const FIoHash& Hash) const
 {
-	const int32 Index = Algo::BinarySearchBy(Attachments, Hash,
-		[](const FCbAttachment& Attachment) -> FIoHash { return Attachment.GetHash(); });
+	const int32 Index = Algo::BinarySearchBy(Attachments, Hash, &FCbAttachment::GetHash);
 	return Attachments.IsValidIndex(Index) ? &Attachments[Index] : nullptr;
 }
 
diff --git a/Engine/Source/Runtime/Core/Public/Serialization/CompactBinaryPackage.h b/Engine/Source/Runtime/Core/Public/Serialization/CompactBinaryPackage.h
--- a/Engine/Source/Runtime/Core/Public/Serialization/CompactBinaryPackage.h
+++ b/Engine/Source/Runtime/Core/Public/Serialization/CompactBinaryPackage.h
@@ -19,7 +19,17 @@ template <typename FuncType> class TFunctionRef;
 ///////////////////////////////////////////////////////////////////////////////////////////////////////////////////////
 
 /**
- * An attachment is either null, raw binary, compressed binary, or an object and is identified by its uncompressed hash.
+ * An attachment is either null, raw binary, compressed binary, or an object.
+ *
+ * Attachments are always identified by their raw hash, even when stored compressed.
+ *
+ * An attachment is serialized as a sequence of compact binary fields with no name. It is invalid
+ * to serialize a null attachment. A raw binary attachment is serialized as an empty Binary field
+ * when empty, and otherwise as a BinaryAttachment field containing the raw hash, followed by the
+ * raw binary data in a Binary field. A compressed binary attachment is serialized as Binary with
+ * a compressed buffer as the value. An object is serialized as an empty Object field when empty,
+ * and otherwise as an ObjectAttachment field containing the hash, followed by the object data in
+ * an Object field.
  */
 class FCbAttachment
 {
@@ -185,7 +195,7 @@ inline uint32 GetTypeHash(const FCbAttachment& Attachment)
  *
  * A package is a Merkle tree with an object as its root and non-leaf nodes, and either binary or
  * objects as its leaf nodes. Nodes reference their children from fields of type BinaryAttachment
- * or ObjectAttachment, which store the hash of the referenced data.
+ * or ObjectAttachment, which store the raw hash of the referenced attachment.
  *
  * It is invalid for a package to include attachments that are not referenced by its object or by
  * one of its referenced object attachments. This invariant needs to be maintained if attachments
@@ -196,12 +206,11 @@ inline uint32 GetTypeHash(const FCbAttachment& Attachment)
  * of the package consumer to have a mechanism for resolving those references when necessary.
  *
  * A package is serialized as a sequence of compact binary fields with no name. The object may be
- * both preceded and followed by attachments. The object itself is written as an Object field and
- * followed by its hash in a ObjectAttachment field, if the object is non-empty. A package
+ * both preceded and followed by attachments. The object is only serialized when it is non-empty,
+ * starting with its hash, in a Hash field, followed by the object, in an Object field. A package
  * ends with a Null field. The canonical order of components is the object and its hash, followed
  * by the attachments ordered by hash, followed by a Null field. It is valid for the a package to
- * have its components serialized in any order, provided there is at most one object and the null
- * field is written last, and the object hash immediately follows the object when present.
+ * have its attachments serialized in any order relative to each other and to the object.
  */
 class FCbPackage
 {
