Chaos: Serialization
- Increment file object version to match changes from 4.27-Chaos

diff --git a/Engine/Source/Runtime/Core/Public/UObject/ReleaseObjectVersion.h b/Engine/Source/Runtime/Core/Public/UObject/ReleaseObjectVersion.h
--- a/Engine/Source/Runtime/Core/Public/UObject/ReleaseObjectVersion.h
+++ b/Engine/Source/Runtime/Core/Public/UObject/ReleaseObjectVersion.h
@@ -142,6 +142,9 @@ struct CORE_API FReleaseObjectVersion
 		// Remove default resolution limit of 512 pixels for cubemaps generated from long-lat sources
 		LonglatTextureCubeDefaultMaxResolution,
 
+		// bake center of mass into chaos cache
+		GeometryCollectionCacheRemovesMassToLocal,
+
 		// -----<new versions can be added above this line>-------------------------------------------------
 		VersionPlusOne,
 		LatestVersion = VersionPlusOne - 1
