DepletableMpscQueue - a multi-producer/single-consumer concurrent atomic queue that consumes all items at once, atomically, and leave the queue in a usable empty state. Optimised performance by limiting a contention between consumer and producers. A typical example is a command queue where multiple threads enqueue and a single thread consumes the queue until there're no more items left.

diff --git a/Engine/Source/Runtime/Core/Private/Tests/Containers/ConcurrentQueuesTest.cpp b/Engine/Source/Runtime/Core/Private/Tests/Containers/ConcurrentQueuesTest.cpp
--- a/Engine/Source/Runtime/Core/Private/Tests/Containers/ConcurrentQueuesTest.cpp
+++ b/Engine/Source/Runtime/Core/Private/Tests/Containers/ConcurrentQueuesTest.cpp
@@ -9,13 +9,14 @@
 #include "Containers/SpscQueue.h"
 #include "Containers/MpscQueue.h"
 #include "Containers/ClosableMpscQueue.h"
+#include "Containers/DepletableMpscQueue.h"
 #include "Tests/Benchmark.h"
 
 #include <atomic>
 
 #if WITH_DEV_AUTOMATION_TESTS
 
-namespace ConcurrentQueuesTests
+namespace UE { namespace ConcurrentQueuesTests
 {
 	template<typename QueueType>
 	class TQueueAdapter : public QueueType
@@ -444,7 +445,7 @@ namespace ClosableMpscQueueTests
 			// `acquire` to "synchronise-with" the consumer
 			if (!bIsComplete.load(std::memory_order_acquire))
 			{
-				ensureAlways(Qeueue.Enqueue(MoveTemp(Value)));
+				ensureAlways(Queue.Enqueue(MoveTemp(Value)));
 				bRes = true;
 			}
 			// `release` to "synchonise-with" other threads
@@ -466,14 +467,14 @@ namespace ClosableMpscQueueTests
 			}
 
 			T Value;
-			while (Qeueue.Dequeue(Value))
+			while (Queue.Dequeue(Value))
 			{
 				Consumer(MoveTemp(Value));
 			}
 		}
 
 	private:
-		TQueue<T, EQueueMode::Mpsc> Qeueue;
+		TQueue<T, EQueueMode::Mpsc> Queue;
 		std::atomic<bool> bIsComplete{ false };
 		std::atomic<uint32> SubscribingThreadsNum{ 0 };
 	};
@@ -492,7 +493,7 @@ namespace ClosableMpscQueueTests
 			// `acquire` to "synchronise-with" the consumer
 			if (!bIsComplete.load(std::memory_order_acquire))
 			{
-				Qeueue.Enqueue(MoveTemp(Value));
+				Queue.Enqueue(MoveTemp(Value));
 				bRes = true;
 			}
 			// `release` to "synchonise-with" other threads
@@ -513,14 +514,14 @@ namespace ClosableMpscQueueTests
 				FPlatformProcess::Yield();
 			}
 
-			while (TOptional<T> Value = Qeueue.Dequeue())
+			while (TOptional<T> Value = Queue.Dequeue())
 			{
 				Consumer(MoveTemp(Value.GetValue()));
 			}
 		}
 
 	private:
-		TMpscQueue<T> Qeueue;
+		TMpscQueue<T> Queue;
 		std::atomic<bool> bIsComplete{ false };
 		std::atomic<uint32> SubscribingThreadsNum{ 0 };
 	};
@@ -559,14 +560,17 @@ namespace ClosableMpscQueueTests
 		}
 	}
 
-	template<uint32 Num, typename QueueT>
+	// launches a number of producer tasks that push incremental numbers into a queue until it's closed. then, on the main thread,
+	// the queue is closed, supposedly concurrently to producers trying pushing into it. after that it's checked that the same total sum was 
+	// pushed and popped from the queue
+	template<int32 Num, typename QueueT>
 	void TestCorrectness()
 	{
-		for (uint32 Run = 0; Run != Num; ++Run)
+		for (int32 Run = 0; Run != Num; ++Run)
 		{
 			QueueT Queue;
 
-			int32 ProducersNum = FPlatformMisc::NumberOfCoresIncludingHyperthreads() - 1;
+			int32 ProducersNum = FPlatformMisc::NumberOfWorkerThreadsToSpawn();
 
 			FGraphEventArray Producers;
 			Producers.Reserve(ProducersNum);
@@ -586,7 +590,7 @@ namespace ClosableMpscQueueTests
 
 						*NumProduced = i - 1;
 					}
-					));
+				));
 			}
 
 			FPlatformProcess::Yield();
@@ -603,7 +607,7 @@ namespace ClosableMpscQueueTests
 				Produced += N * (N + 1) / 2;
 			}
 
-			checkf(Produced == Consumed, TEXT("%d - %d"), Produced, Consumed);
+			checkf(Produced == Consumed, TEXT("%" INT64_FMT " - %" INT64_FMT ", %d run"), Produced, Consumed, Run);
 
 			if (Produced == 0)
 			{
@@ -726,18 +730,368 @@ namespace ClosableMpscQueueTests
 		UE_BENCHMARK(5, TestClosingEmptyQueue<10'000'000, TClosableMpscQueueOnTMpscQueue<void*>>);
 		UE_BENCHMARK(5, TestClosingEmptyQueue<10'000'000, TClosableLockFreePointerListUnorderedSingleConsumer_Adapter<void>>);
 
-		UE_BENCHMARK(5, TestCorrectness<1000, TClosableMpscQueue<void*>>);
-		UE_BENCHMARK(5, TestCorrectness<1000, TClosableMpscQueueOnTQueue<void*>>);
-		UE_BENCHMARK(5, TestCorrectness<1000, TClosableMpscQueueOnTMpscQueue<void*>>);
-		UE_BENCHMARK(5, TestCorrectness<1000, TClosableLockFreePointerListUnorderedSingleConsumer_Adapter<void>>);
+		UE_BENCHMARK(5, TestCorrectness<300, TClosableMpscQueue<void*>>);
+		UE_BENCHMARK(5, TestCorrectness<300, TClosableMpscQueueOnTQueue<void*>>);
+		UE_BENCHMARK(5, TestCorrectness<300, TClosableMpscQueueOnTMpscQueue<void*>>);
+		UE_BENCHMARK(5, TestCorrectness<300, TClosableLockFreePointerListUnorderedSingleConsumer_Adapter<void>>);
 
-		UE_BENCHMARK(5, TestPerf<1'000'000, TClosableMpscQueue<void*>>);
-		UE_BENCHMARK(5, TestPerf<1'000'000, TClosableMpscQueueOnTQueue<void*>>);
-		UE_BENCHMARK(5, TestPerf<1'000'000, TClosableMpscQueueOnTMpscQueue<void*>>);
-		UE_BENCHMARK(5, TestPerf<1'000'000, TClosableLockFreePointerListUnorderedSingleConsumer_Adapter<void>>);
+		UE_BENCHMARK(5, TestPerf<5'000'000, TClosableMpscQueue<void*>>);
+		UE_BENCHMARK(5, TestPerf<5'000'000, TClosableMpscQueueOnTQueue<void*>>);
+		UE_BENCHMARK(5, TestPerf<5'000'000, TClosableMpscQueueOnTMpscQueue<void*>>);
+		UE_BENCHMARK(5, TestPerf<5'000'000, TClosableLockFreePointerListUnorderedSingleConsumer_Adapter<void>>);
 
 		return true;
 	}
 }
 
+namespace DepletableMpscQueueTests
+{
+	// straightforward implementation over MPSC TQueue for comparison with TResettableMpscQueue
+	template<typename T>
+	class TDepletableMpscQueueOnTQueue
+	{
+		using FQueue = TQueue<T, EQueueMode::Mpsc>;
+
+	public:
+		~TDepletableMpscQueueOnTQueue()
+		{
+			checkf(SubscribingThreadsNum.load(std::memory_order_relaxed) == 0, TEXT("No other threads can use the object during its destruction"));
+			delete Queue.load(std::memory_order_relaxed);
+		}
+
+		void Enqueue(T Value)
+		{
+			// `acquire`/`release` to keep `Queue` usage inside the scope
+			verify(SubscribingThreadsNum.fetch_add(1, std::memory_order_acquire) >= 0);
+			Queue.load(std::memory_order_relaxed)->Enqueue(MoveTemp(Value));
+			// `release` to "synchonise-with" other threads
+			verify(SubscribingThreadsNum.fetch_sub(1, std::memory_order_release) >= 1);
+		}
+
+		template<typename F>
+		void Deplete(const F& Consumer)
+		{
+			// `acquire` to make it "happen-before" reading `SubscribingThreadsNum`
+			FQueue* LocalQueue = Queue.exchange(new FQueue, std::memory_order_acquire);
+
+			T Value;
+			bool bHasValue = false;
+			// it's possible that a queue is exchanged and depleted after a producer has gotten the queue but before it has finished enqueuing. 
+			// we use the following counter to make the consumer wait until all producers have finished enqueing to the current queue
+			do
+			{
+				bHasValue = LocalQueue->Dequeue(Value);
+				if (bHasValue)
+				{
+					Consumer(MoveTemp(Value));
+				}
+			} while (bHasValue || SubscribingThreadsNum.load(std::memory_order_relaxed) != 0);
+
+			delete LocalQueue;
+		}
+
+	private:
+		std::atomic<FQueue*> Queue{ new FQueue };
+		std::atomic<uint32> SubscribingThreadsNum{ 0 };
+	};
+
+	// straightforward implementation over MPSC TQueue for comparison with TResettableMpscQueue
+	template<typename T>
+	class TDepletableMpscQueueOnTMpscQueue
+	{
+	public:
+		~TDepletableMpscQueueOnTMpscQueue()
+		{
+			checkf(SubscribingThreadsNum.load(std::memory_order_relaxed) == 0, TEXT("No other threads can use the object during its destration"));
+			delete Queue.load(std::memory_order_relaxed);
+		}
+
+		void Enqueue(T Value)
+		{
+			// `acquire`/`release` to keep `Queue` usage inside the scope
+			verify(SubscribingThreadsNum.fetch_add(1, std::memory_order_acquire) >= 0);
+			Queue.load(std::memory_order_relaxed)->Enqueue(MoveTemp(Value));
+			// `release` to "synchonise-with" other threads
+			verify(SubscribingThreadsNum.fetch_sub(1, std::memory_order_release) >= 1);
+		}
+
+		template<typename F>
+		void Deplete(const F& Consumer)
+		{
+			// `acquire` to make it "happen-before" reading `SubscribingThreadsNum`
+			TMpscQueue<T>* LocalQueue = Queue.exchange(new TMpscQueue<T>, std::memory_order_acquire);
+
+			TOptional<T> Value;
+			// it's possible that a queue is exchanged and depleted after a producer has gotten the queue but before it has finished enqueuing. 
+			// we use the following counter to make the consumer wait until all producers have finished enqueing to the current queue
+			do
+			{
+				Value = LocalQueue->Dequeue();
+				if (Value.IsSet())
+				{
+					Consumer(MoveTemp(Value.GetValue()));
+				}
+			} while (Value.IsSet() || SubscribingThreadsNum.load(std::memory_order_relaxed) != 0);
+
+			delete LocalQueue;
+		}
+
+	private:
+		std::atomic<TMpscQueue<T>*> Queue{ new TMpscQueue<T> };
+		std::atomic<uint32> SubscribingThreadsNum{ 0 };
+	};
+
+	// wrapper over `TClosableLockFreePointerListUnorderedSingleConsumer`
+	template<typename T>
+	class TDepletableMpscQueueOnTClosableLockFreePointerListUnorderedSingleConsumer
+	{
+		using FQueue = TClosableLockFreePointerListUnorderedSingleConsumer<T, 0 /* as was used in FGraphEvent */>;
+
+	public:
+		~TDepletableMpscQueueOnTClosableLockFreePointerListUnorderedSingleConsumer()
+		{
+			checkf(SubscribingThreadsNum.load(std::memory_order_relaxed) == 0, TEXT("No other threads can use the object during its destration"));
+			delete Queue.load(std::memory_order_relaxed);
+		}
+
+		void Enqueue(T* Value)
+		{
+			// `acquire`/`release` to keep `Queue` usage inside the scope
+			verify(SubscribingThreadsNum.fetch_add(1, std::memory_order_acquire) >= 0);
+			ensureAlways(Queue.load(std::memory_order_relaxed)->PushIfNotClosed(MoveTemp(Value)));
+			// `release` to "synchonise-with" other threads
+			verify(SubscribingThreadsNum.fetch_sub(1, std::memory_order_release) >= 1);
+		}
+
+		template<typename F>
+		void Deplete(const F& Consumer)
+		{
+			// `acquire` to make it "happen-before" reading `SubscribingThreadsNum`
+			FQueue* LocalQueue = Queue.exchange(new FQueue, std::memory_order_acquire);
+
+			// wait for all producers to finish enqueuing to this queue instance before closing it, because enqueuing to a closed queue would fail
+			// `acquire` to make it "happen-before" dequeueing
+			while (SubscribingThreadsNum.load(std::memory_order_acquire) != 0)
+			{
+				FPlatformProcess::Yield();
+			}
+
+			TArray<T*> Items;
+			LocalQueue->PopAllAndClose(Items);
+			for (T* Item : Items)
+			{
+				Consumer(MoveTemp(Item));
+			}
+
+			delete LocalQueue;
+		}
+
+	private:
+		std::atomic<FQueue*> Queue{ new FQueue };
+		std::atomic<uint32> SubscribingThreadsNum{ 0 };
+	};
+
+	// wrapper over `TClosableMpscQueue`
+	template<typename T>
+	class TDepletableMpscQueueOnTClosableMpscQueue
+	{
+		using FQueue = TClosableMpscQueue<T>;
+
+	public:
+		~TDepletableMpscQueueOnTClosableMpscQueue()
+		{
+			checkf(SubscribingThreadsNum.load(std::memory_order_relaxed) == 0, TEXT("No other threads can use the object during its destration"));
+			delete Queue.load(std::memory_order_relaxed);
+		}
+
+		void Enqueue(T Value)
+		{
+			// `acquire`/`release` to keep `Queue` usage inside the scope
+			verify(SubscribingThreadsNum.fetch_add(1, std::memory_order_acquire) >= 0);
+			ensureAlways(Queue.load(std::memory_order_relaxed)->Enqueue(MoveTemp(Value)));
+			// `release` to "synchonise-with" other threads
+			verify(SubscribingThreadsNum.fetch_sub(1, std::memory_order_release) >= 1);
+		}
+
+		template<typename F>
+		void Deplete(const F& Consumer)
+		{
+			// `acquire` to make it "happen-before" reading `SubscribingThreadsNum`
+			FQueue* LocalQueue = Queue.exchange(new FQueue, std::memory_order_acquire);
+
+			// wait for all producers to finish enqueuing to this queue instance before closing it, because enqueuing to a closed queue would fail
+			// `acquire` to make it "happen-before" dequeueing
+			while (SubscribingThreadsNum.load(std::memory_order_acquire) != 0)
+			{
+				FPlatformProcess::Yield();
+			}
+
+			LocalQueue->Close(Consumer);
+
+			delete LocalQueue;
+		}
+
+	private:
+		std::atomic<FQueue*> Queue{ new FQueue };
+		std::atomic<uint32> SubscribingThreadsNum{ 0 };
+	};
+
+	template<uint32 Num, typename QueueT>
+	void TestConsumingEmptyQueue()
+	{
+		QueueT Q;
+		for (uint32 i = 0; i != Num; ++i)
+		{
+			Q.Deplete([](void*) { checkNoEntry(); });
+		}
+	}
+
+	template<uint64 Num, typename QueueT>
+	void TestPerf()
+	{
+		// spawns producers num = workers num. Every producer pushes `Num` items into the queue. The consumer continuously
+		// consumes queue until all produced items are consumed
+		QueueT Queue;
+
+		int32 ProducersNum = FPlatformMisc::NumberOfWorkerThreadsToSpawn();
+
+		FGraphEventArray Producers;
+		Producers.Reserve(ProducersNum);
+
+		for (int32 i = 0; i != ProducersNum; ++i)
+		{
+			Producers.Add(FFunctionGraphTask::CreateAndDispatchWhenReady
+			(
+				[&Queue]
+				{
+					uint64 Total = 0;
+					for (uint64 Product = 1; Product != Num + 1; ++Product)
+					{
+						Queue.Enqueue((void*)(intptr_t)Product);
+						Total += Product;
+					}
+					check(Total == Num * (Num + 1) / 2);
+				}
+			));
+		}
+
+		uint64 Produced = Num * (Num + 1) / 2 * ProducersNum;
+		uint64 Consumed = 0;
+		while (Consumed != Produced)
+		{
+			Queue.Deplete
+			(
+				[&Consumed](void* Value) 
+				{ 
+					uint64 IntValue = (uint64)(intptr_t)Value;
+					Consumed += IntValue;
+				}
+			);
+			checkf(Produced >= Consumed, TEXT("%llu - %llu"), Produced, Consumed);
+		}
+
+		checkf(Produced == Consumed, TEXT("%llu - %llu"), Produced, Consumed);
+
+		FTaskGraphInterface::Get().WaitUntilTasksComplete(Producers);
+
+		Queue.Deplete([](void*) { checkNoEntry(); }); // empty
+	}
+
+	IMPLEMENT_SIMPLE_AUTOMATION_TEST(FDepletableMpscQueueTest, "System.Core.Async.DepletableMpscQueueTest", EAutomationTestFlags::ApplicationContextMask | EAutomationTestFlags::EngineFilter);
+
+	bool FDepletableMpscQueueTest::RunTest(const FString& Parameters)
+	{
+		{
+			TDepletableMpscQueue<int> Q;
+			Q.Enqueue(1);
+			Q.Enqueue(2);
+			Q.Deplete([](int) {});
+			Q.Enqueue(3);
+			Q.Deplete([](int) {});
+			Q.Deplete([](int) {});
+			Q.Enqueue(4); // to destroy non-empty queue
+		}
+
+		{
+			TDepletableMpscQueue<int> Q;
+			Q.Deplete([](int) { check(false); });
+		}
+		{
+			TDepletableMpscQueue<int> Q;
+			Q.Enqueue(1);
+			bool Done = false;
+			Q.Deplete([&Done](int Item) { check(!Done && Item == 1); Done = true; });
+		}
+
+		{
+			TDepletableMpscQueue<int> Q;
+			Q.Enqueue(1);
+			Q.Enqueue(2);
+			int Step = 0;
+			Q.Deplete(
+				[&Step](int Item)
+				{
+					switch (Step)
+					{
+					case 0:
+						check(Item == 1);
+						++Step;
+						break;
+					case 1:
+						check(Item == 2);
+						++Step;
+						break;
+					default:
+						check(false);
+					}
+				}
+			);
+		}
+
+		//for (int i = 0; i != 1000000; ++i)
+		{
+			TDepletableMpscQueue<int> Q;
+			Q.Deplete([](int) { check(false); });
+			Q.Deplete([](int) { check(false); });
+		}
+
+		//for (int i = 0; i != 1000000; ++i)
+		{
+			int Res = 0;
+			TDepletableMpscQueue<int> Q;
+			Q.Enqueue(1);
+			Q.Deplete([&Res](int i) { Res += i; });
+			Q.Enqueue(2);
+			Q.Deplete([&Res](int i) { Res += i; });
+			checkf(Res == 3, TEXT("%d"), Res);
+		}
+
+		{	// EnqueueAndReturnWasEmpty
+			TDepletableMpscQueue<int> Q;
+			verify(Q.EnqueueAndReturnWasEmpty(1));
+			verify(!Q.EnqueueAndReturnWasEmpty(2));
+			Q.Deplete([](int) {});
+			verify(Q.EnqueueAndReturnWasEmpty(1));
+			verify(!Q.EnqueueAndReturnWasEmpty(2));
+			Q.Deplete([](int) {});
+			verify(Q.EnqueueAndReturnWasEmpty(1));
+			verify(!Q.EnqueueAndReturnWasEmpty(2));
+		}
+
+		UE_BENCHMARK(5, TestPerf<10'000, TDepletableMpscQueue<void*>>);
+		UE_BENCHMARK(5, TestPerf<10'000, TDepletableMpscQueueOnTQueue<void*>>);
+		UE_BENCHMARK(5, TestPerf<10'000, TDepletableMpscQueueOnTMpscQueue<void*>>);
+		UE_BENCHMARK(5, TestPerf<10'000, TDepletableMpscQueueOnTClosableLockFreePointerListUnorderedSingleConsumer<void>>);
+		UE_BENCHMARK(5, TestPerf<10'000, TDepletableMpscQueueOnTClosableMpscQueue<void*>>);
+
+		UE_BENCHMARK(5, TestConsumingEmptyQueue<10'000'000, TDepletableMpscQueue<void*>>);
+		UE_BENCHMARK(5, TestConsumingEmptyQueue<10'000'000, TDepletableMpscQueueOnTQueue<void*>>);
+		UE_BENCHMARK(5, TestConsumingEmptyQueue<10'000'000, TDepletableMpscQueueOnTMpscQueue<void*>>);
+		UE_BENCHMARK(5, TestConsumingEmptyQueue<10'000'000, TDepletableMpscQueueOnTClosableLockFreePointerListUnorderedSingleConsumer<void>>);
+		UE_BENCHMARK(5, TestConsumingEmptyQueue<10'000'000, TDepletableMpscQueueOnTClosableMpscQueue<void*>>);
+
+		return true;
+	}
+}}
+
 #endif // WITH_DEV_AUTOMATION_TESTS
diff --git a/Engine/Source/Runtime/Core/Public/Containers/DepletableMpscQueue.h b/Engine/Source/Runtime/Core/Public/Containers/DepletableMpscQueue.h
new file mode 100644
--- /dev/null
+++ b/Engine/Source/Runtime/Core/Public/Containers/DepletableMpscQueue.h
@@ -0,0 +1,142 @@
+// Copyright Epic Games, Inc. All Rights Reserved.
+
+#pragma once
+
+#include "CoreTypes.h"
+#include "Templates/TypeCompatibleBytes.h"
+#include <atomic>
+
+namespace UE
+{
+	/**
+	 * Multi-producer/single-consumer unbounded concurrent queue that is atomically consumed and is reset to its 
+	 * default empty state.
+	 * A typical use case is when the consumer doesn't stop until the queue is depleted.
+	 * Is faster than traditional MPSC queues, especially for consumer.
+	 */
+	template<typename T>
+	class TDepletableMpscQueue final
+	{
+	private:
+		struct FNode
+		{
+			std::atomic<FNode*> Next{ nullptr };
+			TTypeCompatibleBytes<T> Value;
+		};
+
+		FNode Sentinel; // `Sentinel.Next` is the head of the queue
+		std::atomic<FNode*> Tail{ &Sentinel };
+
+	public:
+		UE_NONCOPYABLE(TDepletableMpscQueue);
+
+		TDepletableMpscQueue() = default;
+
+		~TDepletableMpscQueue()
+		{
+			// delete remaining elements
+			FNode* Node = Sentinel.Next.load(std::memory_order_relaxed);
+			while (Node != nullptr)
+			{
+				DestructItem(Node->Value.GetTypedPtr());
+				FNode* Next = Node->Next.load(std::memory_order_relaxed);
+				delete Node;
+				Node = Next;
+			}
+		}
+
+		template <typename... ArgTypes>
+		void Enqueue(ArgTypes&&... Args)
+		{
+			EnqueueAndReturnWasEmpty(Forward<ArgTypes>(Args)...);
+		}
+
+		template <typename... ArgTypes>
+		bool EnqueueAndReturnWasEmpty(ArgTypes&&... Args)
+		{
+			FNode* New = new FNode;
+			new (&New->Value) T(Forward<ArgTypes>(Args)...);
+
+			// switch `Tail` to the new node and only then link the old tail to the new one. The list is not fully linked between these ops,
+			// this is explicitly handled by the consumer by waiting for the link
+
+			FNode* Prev = Tail.exchange(New, std::memory_order_release); // `release` to make sure the new node is fully constructed before
+			// it becomes visible to the consumer
+			check(Prev->Next.load(std::memory_order_relaxed) == nullptr);
+
+			Prev->Next.store(New, std::memory_order_relaxed);
+
+			return Prev == &Sentinel;
+		}
+
+		/**
+		 * Takes all items from the queue atomically and then consumes them.
+		 * @param Consumer: a functor with signature `AnyReturnType (T Value)`
+		 */
+		template<typename F>
+		void Deplete(const F& Consumer)
+		{
+			FNode* First = Sentinel.Next.load(std::memory_order_relaxed);
+			if (First == nullptr)
+			{
+				return; // empty
+			}
+
+			// reset the head so the next consumption can detect that the queue is empty
+			// `Sentinel.Next` is not touched by producers right now because it's already not null
+			Sentinel.Next.store(nullptr, std::memory_order_relaxed);
+
+			// reset the queue to the empty state. this redirects producers to start from `Sentinel` again.
+			// take note of the tail on resetting it because the list can be still not fully linked and so `Node.Next == nullptr` can't be 
+			// used to detect the end of the list
+			FNode* Last = Tail.exchange(&Sentinel, std::memory_order_release); // `release` to make `Sentinel.Next = nullptr` happen before
+			// modifying `Tail`
+			check(Last->Next.load(std::memory_order_relaxed) == nullptr);
+			// the previously queued items are detached from the instance (as a linked list, though potential not fully linked yet)
+
+			check(Last != &Sentinel); // can't be empty because of `First != nullptr` above
+
+			Consume(First, Last, Consumer);
+		}
+
+		// the result can be relied upon only in special cases, as the state can change concurrently. use with caution 
+		bool IsEmpty() const
+		{
+			return Tail.load(std::memory_order_relaxed) == &Sentinel;
+		}
+
+	private:
+		template<typename F>
+		static void Consume(FNode* First, FNode* Last, const F& Consumer)
+		{
+			auto GetNext = [](FNode* Node)
+			{
+				FNode* Next;
+				// producers can be still updating `Next`, wait until the link to the next element is established
+				do
+				{
+					Next = Node->Next.load(std::memory_order_relaxed);
+				} while (Next == nullptr);
+
+				return Next;
+			};
+
+			auto Consume = [&Consumer](FNode* Node)
+			{
+				T* ValuePtr = (T*)&Node->Value;
+				Consumer(MoveTemp(*ValuePtr));
+				DestructItem(ValuePtr);
+			};
+
+			while (First != Last)
+			{
+				Consume(First);
+				FNode* Next = GetNext(First);
+				delete First;
+				First = Next;
+			}
+
+			Consume(Last);
+		}
+	};
+}
