Reset cpu-scope-emitted flag after sending the end event.

If the Cpu channel is disabled between Start() and Stop() pair any subsequent calls to Stop() will continuously emit end events.

#rnx

diff --git a/Engine/Source/Runtime/Core/Public/Async/Fundamental/LocalQueue.h b/Engine/Source/Runtime/Core/Public/Async/Fundamental/LocalQueue.h
--- a/Engine/Source/Runtime/Core/Public/Async/Fundamental/LocalQueue.h
+++ b/Engine/Source/Runtime/Core/Public/Async/Fundamental/LocalQueue.h
@@ -198,6 +198,7 @@ public:
 				if (bCpuBeginEventEmitted)
 				{
 					FCpuProfilerTrace::OutputEndEvent();
+					bCpuBeginEventEmitted = false;
 				}
 #endif
 				NumWorkersLookingForWork.fetch_sub(1, std::memory_order_release);
