StaticClassProvider & StaticStructProvider concepts, and replaced functionality in two dependent files

diff --git a/Engine/Plugins/NetcodeUnitTest/NetcodeUnitTest/Source/NetcodeUnitTest/Public/NUTUtilReflection.h b/Engine/Plugins/NetcodeUnitTest/NetcodeUnitTest/Source/NetcodeUnitTest/Public/NUTUtilReflection.h
--- a/Engine/Plugins/NetcodeUnitTest/NetcodeUnitTest/Source/NetcodeUnitTest/Public/NUTUtilReflection.h
+++ b/Engine/Plugins/NetcodeUnitTest/NetcodeUnitTest/Source/NetcodeUnitTest/Public/NUTUtilReflection.h
@@ -3,11 +3,13 @@
 #pragma once
 
 #include "CoreMinimal.h"
+#include "Concepts/StaticStructProvider.h"
+#include "Templates/Models.h"
+#include "Templates/ValueOrError.h"
 #include "UObject/Class.h"
-#include "UObject/UnrealType.h"
 #include "UObject/Field.h"
 #include "UObject/StructOnScope.h"
-#include "Templates/ValueOrError.h"
+#include "UObject/UnrealType.h"
 
 class FStructOnScope;
 class FFuncReflection;
@@ -190,22 +192,6 @@ class FFuncReflection;
 //				that'd make for a nice feature.
 //				If you (eventually) add autocomplete to the unit test tool log window, this would be the next logical feature.
 
-
-// @todo #JohnB: Perhaps move up to Core, and add a UCLASS version checking StaticClass (tested before, and it works)
-/**
- * Determines at compile-time, whether or not a type is defined as a USTRUCT, through TIsUStruct<Type>::Value
- */
-template<typename T>
-class TIsUStruct
-{
-	template<typename A> static uint8	Check(decltype(&A::StaticStruct));
-	template<typename B> static uint16	Check(...);
-
-public:
-	enum { Value = sizeof(Check<T>(0)) == sizeof(uint8) };
-};
-
-
 /**
  * Used for specifying the warning level, for reflection helpers
  */
@@ -486,7 +472,7 @@ public:
 	template<typename T>
 	FORCEINLINE T* GetStructRef()
 	{
-		static_assert(TIsUStruct<T>::Value, "Only USTRUCT types can be passed to GetStructRef");
+		static_assert(TModels<CStaticStructProvider, T>::Value, "Only USTRUCT types can be passed to GetStructRef");
 
 		return (T*)(void*)(FVMReflection(*this)[TCHAR_TO_ANSI(*FString::Printf(TEXT("F%s"), *T::StaticStruct()->GetName()))]);
 	}
diff --git a/Engine/Source/Runtime/Core/Public/Concepts/StaticClassProvider.h b/Engine/Source/Runtime/Core/Public/Concepts/StaticClassProvider.h
new file mode 100644
--- /dev/null
+++ b/Engine/Source/Runtime/Core/Public/Concepts/StaticClassProvider.h
@@ -0,0 +1,16 @@
+﻿// Copyright Epic Games, Inc. All Rights Reserved.
+
+#pragma once
+
+class UClass;
+
+/**
+ * Describes a type with a StaticClass (static) member.
+ */
+struct CStaticClassProvider
+{
+	template <typename T>
+	auto Requires(UClass*& ClassRef) -> decltype(
+		ClassRef = T::StaticClass()
+	);
+};
diff --git a/Engine/Source/Runtime/Core/Public/Concepts/StaticStructProvider.h b/Engine/Source/Runtime/Core/Public/Concepts/StaticStructProvider.h
new file mode 100644
--- /dev/null
+++ b/Engine/Source/Runtime/Core/Public/Concepts/StaticStructProvider.h
@@ -0,0 +1,16 @@
+﻿// Copyright Epic Games, Inc. All Rights Reserved.
+
+#pragma once
+
+class UScriptStruct;
+
+/**
+ * Describes a type with a StaticStruct (static) member.
+ */
+struct CStaticStructProvider
+{
+	template <typename T>
+	auto Requires(UScriptStruct*& StructRef) -> decltype(
+		StructRef = T::StaticStruct()
+	);
+};
diff --git a/Engine/Source/Runtime/Core/Public/Serialization/MemoryLayout.h b/Engine/Source/Runtime/Core/Public/Serialization/MemoryLayout.h
--- a/Engine/Source/Runtime/Core/Public/Serialization/MemoryLayout.h
+++ b/Engine/Source/Runtime/Core/Public/Serialization/MemoryLayout.h
@@ -2,13 +2,16 @@
 
 #pragma once
 
+#include "Concepts/StaticClassProvider.h"
+#include "Concepts/StaticStructProvider.h"
 #include "Containers/EnumAsByte.h"
 #include "Containers/StringFwd.h"
+#include "Misc/DelayedAutoRegister.h"
+#include "Templates/EnableIf.h"
 #include "Templates/IsAbstract.h"
 #include "Templates/IsPolymorphic.h"
 #include "Templates/IsTriviallyDestructible.h"
-#include "Templates/EnableIf.h"
-#include "Misc/DelayedAutoRegister.h"
+#include "Templates/Models.h"
 
 class FHashedName;
 class FSHA1;
@@ -323,29 +326,10 @@ namespace Freeze
 	CORE_API FSHAHash HashLayout(const FTypeLayoutDesc& TypeDesc, const FPlatformTypeLayoutParameters& LayoutParams);
 }
 
-
-struct CProvidesStaticClass
-{
-	template<typename T>
-	auto Requires(const T&) -> decltype(T::StaticClass());
-};
-
-struct CProvidesStaticStruct
-{
-	template<typename T>
-	auto Requires(const T&) -> decltype(T::StaticStruct());
-};
-
 template<typename T>
 struct TUsePropertyFreezing
 {
-	static constexpr bool Value = (TModels<CProvidesStaticClass, T>::Value || TModels<CProvidesStaticStruct, T>::Value);
-};
-
-template<typename T>
-struct TProvidesStaticStruct
-{
-	static constexpr bool Value = TModels<CProvidesStaticStruct, T>::Value;
+	static constexpr bool Value = (TModels<CStaticClassProvider, T>::Value || TModels<CStaticStructProvider, T>::Value);
 };
 
 template <typename T, bool bUsePropertyFreezing=TUsePropertyFreezing<T>::Value>
@@ -354,7 +338,7 @@ struct TGetFreezeImageHelper
 	static FORCEINLINE FTypeLayoutDesc::FWriteFrozenMemoryImageFunc* Do() { return &Freeze::DefaultWriteMemoryImage; }
 };
 
-template <typename T, bool bProvidesStaticStruct=TProvidesStaticStruct<T>::Value>
+template <typename T, bool bProvidesStaticStruct=TModels<CStaticStructProvider, T>::Value>
 struct TGetFreezeImageFieldHelper
 {
 	static FORCEINLINE FFieldLayoutDesc::FWriteFrozenMemoryImageFunc* Do() { return &Freeze::DefaultWriteMemoryImageField; }
diff --git a/Engine/Source/Runtime/CoreUObject/Public/Serialization/ArchiveCrc32.h b/Engine/Source/Runtime/CoreUObject/Public/Serialization/ArchiveCrc32.h
--- a/Engine/Source/Runtime/CoreUObject/Public/Serialization/ArchiveCrc32.h
+++ b/Engine/Source/Runtime/CoreUObject/Public/Serialization/ArchiveCrc32.h
@@ -36,7 +36,7 @@ public:
 	 * @param Value The value to serialize.
 	 */
 	template <typename StructType>
-	FORCEINLINE friend typename TEnableIf<TProvidesStaticStruct<StructType>::Value, FArchive&>::Type operator <<(FArchive& Ar, const StructType& Value)
+	FORCEINLINE friend typename TEnableIf<TModels<CStaticStructProvider, StructType>::Value, FArchive&>::Type operator <<(FArchive& Ar, const StructType& Value)
 	{
 		StructType* MutableValue = const_cast<StructType*>(&Value);
 		StructType::StaticStruct()->SerializeItem(Ar, MutableValue, nullptr);
