Config: Convert each key to FName only once instead of three times

diff --git a/Engine/Source/Runtime/Core/Private/Misc/ConfigCacheIni.cpp b/Engine/Source/Runtime/Core/Private/Misc/ConfigCacheIni.cpp
--- a/Engine/Source/Runtime/Core/Private/Misc/ConfigCacheIni.cpp
+++ b/Engine/Source/Runtime/Core/Private/Misc/ConfigCacheIni.cpp
@@ -902,46 +902,47 @@ void FConfigFile::CombineFromBuffer(const FString& Buffer)
 					ProcessedValue = Value;
 				}
 
+				const FName Key(Start);
 				if (Cmd == '+')
 				{
 					// Add if not already present.
-					CurrentSection->HandleAddCommand( Start, MoveTemp(ProcessedValue), false );
+					CurrentSection->HandleAddCommand(Key, MoveTemp(ProcessedValue), false);
 				}
-				else if( Cmd=='-' )	
+				else if( Cmd=='-' )
 				{
 					// Remove if present.
-					CurrentSection->RemoveSingle( Start, ProcessedValue );
+					CurrentSection->RemoveSingle(Key, ProcessedValue);
 					CurrentSection->CompactStable();
 				}
 				else if ( Cmd=='.' )
 				{
-					CurrentSection->HandleAddCommand( Start, MoveTemp(ProcessedValue), true );
+					CurrentSection->HandleAddCommand(Key, MoveTemp(ProcessedValue), true);
 				}
 				else if( Cmd=='!' )
 				{
-					CurrentSection->Remove( Start );
+					CurrentSection->Remove(Key);
 				}
 				else if (Cmd == '@')
 				{
 					// track a key to show uniqueness for arrays of structs
-					CurrentSection->ArrayOfStructKeys.Add(Start, MoveTemp(ProcessedValue));
+					CurrentSection->ArrayOfStructKeys.Add(Key, MoveTemp(ProcessedValue));
 				}
 				else if (Cmd == '*')
 				{
 					// track a key to show uniqueness for arrays of structs
 					TMap<FName, FString>& POCKeys = FindOrAddHeterogeneous(PerObjectConfigArrayOfStructKeys, CurrentSectionName.ToView());
-					POCKeys.Add(Start, MoveTemp(ProcessedValue));
+					POCKeys.Add(Key, MoveTemp(ProcessedValue));
 				}
 				else
 				{
 					// First see if this can be processed as an array of keyed structs command
-					if (!CurrentSection->HandleArrayOfKeyedStructsCommand(Start, MoveTemp(ProcessedValue)))
+					if (!CurrentSection->HandleArrayOfKeyedStructsCommand(Key, MoveTemp(ProcessedValue)))
 					{
 						// Add if not present and replace if present.
-						FConfigValue* ConfigValue = CurrentSection->Find(Start);
+						FConfigValue* ConfigValue = CurrentSection->Find(Key);
 						if (!ConfigValue)
 						{
-							CurrentSection->Add(Start, MoveTemp(ProcessedValue));
+							CurrentSection->Add(Key, MoveTemp(ProcessedValue));
 						}
 						else
 						{
