Check GIgnoreDebugger for assert system debug break
Add command line arg to set GIgnoreDebugger
Remove GIgnoreDebugger from UE_BUILD_SHIPPING
Fix up naming convention and deprecate _DebugBreakAndPromptForRemote

=== Notes ===

Sometimes it's useful to run a process with a debugger attached, but knowing you get the same behavior as if it weren't, including how the assert system behaves

Such situations include:
* Tight timing of assert and exceptions in different threads
* When rarely reproducable ensure()s fire before or between breakpoints set in the debugger

=== Hashtags ===

diff --git a/Engine/Source/Runtime/Core/Private/HAL/ExceptionHandling.cpp b/Engine/Source/Runtime/Core/Private/HAL/ExceptionHandling.cpp
--- a/Engine/Source/Runtime/Core/Private/HAL/ExceptionHandling.cpp
+++ b/Engine/Source/Runtime/Core/Private/HAL/ExceptionHandling.cpp
@@ -22,9 +22,13 @@ CORE_API bool GAlwaysReportCrash = false;
 /** Whether to use ClientReportClient rather than the old AutoReporter. */
 CORE_API bool GUseCrashReportClient = true;
 
+#if !UE_BUILD_SHIPPING
+
 /** Whether we should ignore the attached debugger. */
 CORE_API bool GIgnoreDebugger = false;
 
+#endif // !UE_BUILD_SHIPPING
+
 CORE_API TCHAR MiniDumpFilenameW[1024] = TEXT("");
 
 
diff --git a/Engine/Source/Runtime/Core/Public/HAL/ExceptionHandling.h b/Engine/Source/Runtime/Core/Public/HAL/ExceptionHandling.h
--- a/Engine/Source/Runtime/Core/Public/HAL/ExceptionHandling.h
+++ b/Engine/Source/Runtime/Core/Public/HAL/ExceptionHandling.h
@@ -13,9 +13,13 @@ extern CORE_API bool GAlwaysReportCrash;
 /** Whether to use ClientReportClient rather than AutoReporter. */
 extern CORE_API bool GUseCrashReportClient;
 
+#if !UE_BUILD_SHIPPING
+
 /** Whether we should ignore the attached debugger. */
 extern CORE_API bool GIgnoreDebugger;
 
+#endif // !UE_BUILD_SHIPPING
+
 extern CORE_API TCHAR MiniDumpFilenameW[1024];
 
 // #CrashReport: 2014-09-11 Move to PlatformExceptionHandling
diff --git a/Engine/Source/Runtime/Core/Public/Logging/LogMacros.h b/Engine/Source/Runtime/Core/Public/Logging/LogMacros.h
--- a/Engine/Source/Runtime/Core/Public/Logging/LogMacros.h
+++ b/Engine/Source/Runtime/Core/Public/Logging/LogMacros.h
@@ -199,7 +199,7 @@ private:
 					UE_LOG_EXPAND_IS_FATAL(Verbosity, \
 						{ \
 							FMsg::Logf_Internal(UE_LOG_SOURCE_FILE(__FILE__), __LINE__, LCategoryName.GetCategoryName(), ELogVerbosity::Verbosity, LFormat, UE_LOG_Args...); \
-							_DebugBreakAndPromptForRemote(); \
+							UE_DEBUG_BREAK_AND_PROMPT_FOR_REMOTE(); \
 							FDebug::ProcessFatalError(PLATFORM_RETURN_ADDRESS()); \
 						}, \
 						{ \
@@ -256,7 +256,7 @@ private:
 				UE_LOG_EXPAND_IS_FATAL(Verbosity, \
 					{ \
 						FMsg::Logf_Internal(UE_LOG_SOURCE_FILE(__FILE__), __LINE__, CategoryName.GetCategoryName(), ELogVerbosity::Verbosity, Format,  ##__VA_ARGS__); \
-						_DebugBreakAndPromptForRemote(); \
+						UE_DEBUG_BREAK_AND_PROMPT_FOR_REMOTE(); \
 						FDebug::ProcessFatalError(PLATFORM_RETURN_ADDRESS()); \
 						CA_ASSUME(false); \
 					}, \
@@ -303,7 +303,7 @@ private:
 						UE_LOG_EXPAND_IS_FATAL(Verbosity, \
 							{ \
 								FMsg::Logf_Internal(UE_LOG_SOURCE_FILE(__FILE__), __LINE__, LCategoryName.GetCategoryName(), ELogVerbosity::Verbosity, LFormat, UE_LOG_Args...); \
-								_DebugBreakAndPromptForRemote(); \
+								UE_DEBUG_BREAK_AND_PROMPT_FOR_REMOTE(); \
 								FDebug::ProcessFatalError(PLATFORM_RETURN_ADDRESS()); \
 							}, \
 							{ \
diff --git a/Engine/Source/Runtime/Core/Public/Misc/AssertionMacros.h b/Engine/Source/Runtime/Core/Public/Misc/AssertionMacros.h
--- a/Engine/Source/Runtime/Core/Public/Misc/AssertionMacros.h
+++ b/Engine/Source/Runtime/Core/Public/Misc/AssertionMacros.h
@@ -170,10 +170,21 @@ RetType FORCENOINLINE UE_DEBUG_SECTION DispatchCheckVerify(InnerType&& Inner, Ar
 }
 
 #if !UE_BUILD_SHIPPING
-#define _DebugBreakAndPromptForRemote() \
+#define UE_DEBUG_BREAK_AND_PROMPT_FOR_REMOTE() \
 	if (!FPlatformMisc::IsDebuggerPresent()) { FPlatformMisc::PromptForRemoteDebugging(false); } UE_DEBUG_BREAK();
 #else
-	#define _DebugBreakAndPromptForRemote()
+	#define UE_DEBUG_BREAK_AND_PROMPT_FOR_REMOTE()
+#endif // !UE_BUILD_SHIPPING
+
+#define _DebugBreakAndPromptForRemote() \
+	DEPRECATED_MACRO(5.1, "Use UE_DEBUG_BREAK_AND_PROMPT_FOR_REMOTE.")
+
+#if !UE_BUILD_SHIPPING
+	extern CORE_API bool GIgnoreDebugger;
+	// This is named with PLATFORM_ prefix because UE_DEBUG_BREAK* are conditional on a debugger being detected, and PLATFORM_BREAK isn't
+	#define PLATFORM_BREAK_IF_DESIRED() if (LIKELY(!GIgnoreDebugger)) { PLATFORM_BREAK(); }
+#else
+	#define PLATFORM_BREAK_IF_DESIRED() PLATFORM_BREAK();
 #endif // !UE_BUILD_SHIPPING
 
 
@@ -203,7 +214,7 @@ RetType FORCENOINLINE UE_DEBUG_SECTION DispatchCheckVerify(InnerType&& Inner, Ar
 					} \
 				}; \
 				Impl::ExecCheckImplInternal(); \
-				PLATFORM_BREAK(); \
+				PLATFORM_BREAK_IF_DESIRED(); \
 				CA_ASSUME(false); \
 			} \
 		}
@@ -227,7 +238,7 @@ RetType FORCENOINLINE UE_DEBUG_SECTION DispatchCheckVerify(InnerType&& Inner, Ar
 				{ \
 					FDebug::CheckVerifyFailedImpl(#expr, __FILE__, __LINE__, PLATFORM_RETURN_ADDRESS(), format, ##__VA_ARGS__); \
 				}); \
-				PLATFORM_BREAK(); \
+				PLATFORM_BREAK_IF_DESIRED(); \
 				CA_ASSUME(false); \
 			} \
 		}
@@ -339,7 +350,7 @@ RetType FORCENOINLINE UE_DEBUG_SECTION DispatchCheckVerify(InnerType&& Inner, Ar
 				return true; \
 			} \
 			return false; \
-		}) && ([] () { PLATFORM_BREAK(); } (), false)))
+		}) && ([] () { PLATFORM_BREAK_IF_DESIRED(); } (), false)))
 
 	#define ensure(           InExpression                ) UE_ENSURE_IMPL( , false, InExpression, TEXT(""))
 	#define ensureMsgf(       InExpression, InFormat, ... ) UE_ENSURE_IMPL(&, false, InExpression, InFormat, ##__VA_ARGS__)
@@ -406,7 +417,7 @@ CORE_API void VARARGS LowLevelFatalErrorHandler(const ANSICHAR* File, int32 Line
 		{ \
 			void* ProgramCounter = PLATFORM_RETURN_ADDRESS(); \
 			LowLevelFatalErrorHandler(__FILE__, __LINE__, ProgramCounter, Format, ##__VA_ARGS__); \
-			_DebugBreakAndPromptForRemote(); \
+			UE_DEBUG_BREAK_AND_PROMPT_FOR_REMOTE(); \
 			FDebug::ProcessFatalError(ProgramCounter); \
 		}); \
 	}
diff --git a/Engine/Source/Runtime/CoreUObject/Public/UObject/UObjectGlobals.h b/Engine/Source/Runtime/CoreUObject/Public/UObject/UObjectGlobals.h
--- a/Engine/Source/Runtime/CoreUObject/Public/UObject/UObjectGlobals.h
+++ b/Engine/Source/Runtime/CoreUObject/Public/UObject/UObjectGlobals.h
@@ -2577,7 +2577,7 @@ struct FAssetMsg
 				FMsg::Logf_Internal(__FILE__, __LINE__, CategoryName.GetCategoryName(), ELogVerbosity::Verbosity, ASSET_LOG_FORMAT_STRING Format, *FormatPath, ##__VA_ARGS__); \
 				UE_LOG_EXPAND_IS_FATAL(Verbosity, \
 					{ \
-						_DebugBreakAndPromptForRemote(); \
+						UE_DEBUG_BREAK_AND_PROMPT_FOR_REMOTE(); \
 						FDebug::AssertFailed("", __FILE__, __LINE__, TEXT("%s: ") Format, *FormatPath, ##__VA_ARGS__); \
 						CA_ASSUME(false); \
 					}, \
diff --git a/Engine/Source/Runtime/Engine/Private/Analytics/EngineAnalyticsSessionSummary.cpp b/Engine/Source/Runtime/Engine/Private/Analytics/EngineAnalyticsSessionSummary.cpp
--- a/Engine/Source/Runtime/Engine/Private/Analytics/EngineAnalyticsSessionSummary.cpp
+++ b/Engine/Source/Runtime/Engine/Private/Analytics/EngineAnalyticsSessionSummary.cpp
@@ -366,9 +366,14 @@ bool FEngineAnalyticsSessionSummary::UpdateExternalProcessReporterState(bool bQu
 bool FEngineAnalyticsSessionSummary::UpdateDebuggerStates()
 {
 	bool bStateChanged = false;
+	bool bIgnoreDebugger = false;
+
+#if !UE_BUILD_SHIPPING
+	bIgnoreDebugger = GIgnoreDebugger;
+#endif // !UE_BUILD_SHIPPING
 
 	// Ignoring the debugger changes how IsDebuggerPresent() behave and masks the usage of the debugger if true.
-	if (!bDebuggerIgnored && GIgnoreDebugger)
+	if (!bDebuggerIgnored && bIgnoreDebugger)
 	{
 		bStateChanged = true;
 		bDebuggerIgnored = true;
diff --git a/Engine/Source/Runtime/Launch/Private/LaunchEngineLoop.cpp b/Engine/Source/Runtime/Launch/Private/LaunchEngineLoop.cpp
--- a/Engine/Source/Runtime/Launch/Private/LaunchEngineLoop.cpp
+++ b/Engine/Source/Runtime/Launch/Private/LaunchEngineLoop.cpp
@@ -1579,6 +1579,13 @@ int32 FEngineLoop::PreInitPreStartupScreen(const TCHAR* CmdLine)
 		FPlatformMisc::SetUTF8Output();
 	}
 
+#if !UE_BUILD_SHIPPING
+	if (FParse::Param(CmdLine, TEXT("IgnoreDebugger")))
+	{
+		GIgnoreDebugger = true;
+	}
+#endif // !UE_BUILD_SHIPPING
+
 	// Switch into executable's directory.
 	FPlatformProcess::SetCurrentWorkingDirectoryToBaseDir();
 
diff --git a/Engine/Source/Runtime/SlateCore/Private/Debugging/ConsoleSlateDebuggerBreak.cpp b/Engine/Source/Runtime/SlateCore/Private/Debugging/ConsoleSlateDebuggerBreak.cpp
--- a/Engine/Source/Runtime/SlateCore/Private/Debugging/ConsoleSlateDebuggerBreak.cpp
+++ b/Engine/Source/Runtime/SlateCore/Private/Debugging/ConsoleSlateDebuggerBreak.cpp
@@ -281,7 +281,7 @@ void FConsoleSlateDebuggerBreak::HandleWidgetInvalidated(const FSlateDebuggingIn
 		if (EnumHasAnyFlags(Args.InvalidateWidgetReason, FoundItem->Reason))
 		{
 			UE_LOG(LogSlateDebugger, Log, TEXT("Widget '%s' was invalidated."), *FReflectionMetaData::GetWidgetPath(Args.WidgetInvalidated));
-			_DebugBreakAndPromptForRemote();
+			UE_DEBUG_BREAK_AND_PROMPT_FOR_REMOTE();
 		}
 	}
 }
@@ -298,7 +298,7 @@ void FConsoleSlateDebuggerBreak::HandleBeginWidgetPaint(const SWidget* Widget, c
 	if (BreakOnBeginPaintElements.Contains(WidgetId))
 	{
 		UE_LOG(LogSlateDebugger, Log, TEXT("Widget '%s' was invalidated."), *FReflectionMetaData::GetWidgetPath(Widget));
-		_DebugBreakAndPromptForRemote();
+		UE_DEBUG_BREAK_AND_PROMPT_FOR_REMOTE();
 	}
 }
 
@@ -314,7 +314,7 @@ void FConsoleSlateDebuggerBreak::HandleEndWidgetPaint(const SWidget* Widget, con
 	if (BreakOnEndPaintElements.Contains(WidgetId))
 	{
 		UE_LOG(LogSlateDebugger, Log, TEXT("Widget '%s' was invalidated."), *FReflectionMetaData::GetWidgetPath(Widget));
-		_DebugBreakAndPromptForRemote();
+		UE_DEBUG_BREAK_AND_PROMPT_FOR_REMOTE();
 	}
 }
 
