Wrong CreateProcess() flag to pop out a debug console window.

#rnx

diff --git a/Engine/Source/Runtime/Core/Private/ProfilingDebugging/TraceAuxiliary.cpp b/Engine/Source/Runtime/Core/Private/ProfilingDebugging/TraceAuxiliary.cpp
--- a/Engine/Source/Runtime/Core/Private/ProfilingDebugging/TraceAuxiliary.cpp
+++ b/Engine/Source/Runtime/Core/Private/ProfilingDebugging/TraceAuxiliary.cpp
@@ -541,7 +541,7 @@ static void LaunchUnrealTraceInternal(const TCHAR* CommandLine)
 	uint32 CreateProcFlags = CREATE_BREAKAWAY_FROM_JOB;
 	if (FParse::Param(CommandLine, TEXT("traceshowstore")))
 	{
-		CreateProcFlags |= DETACHED_PROCESS;
+		CreateProcFlags |= CREATE_NEW_CONSOLE;
 	}
 	else
 	{
