Mac compilation: added "-faligned-new" if an old macOS is targetted to avoid "aligned new is not supported" error. Broken build: https://horde.devtools.epicgames.com/job/618160c0fc786a000148b7f4?step=ea45

preflight https://horde.devtools.epicgames.com/job/6182954f300d520001e9e1fb?step=1df5

diff --git a/Engine/Source/Programs/CrashReportClient/CrashReportClient.Target.cs b/Engine/Source/Programs/CrashReportClient/CrashReportClient.Target.cs
--- a/Engine/Source/Programs/CrashReportClient/CrashReportClient.Target.cs
+++ b/Engine/Source/Programs/CrashReportClient/CrashReportClient.Target.cs
@@ -49,9 +49,6 @@ public class CrashReportClientTarget : TargetRules
 		// Epic Games Launcher needs to run on OS X 10.9, so CrashReportClient needs this as well
 		bEnableOSX109Support = true;
 
-		// Disable automation tests to avoid compilation errors related to alignas, caused by setting bEnableOSX109Support to true above. Old OSX versions do not fully support this until C++17
-		bForceDisableAutomationTests = true;
-
 		// Need to disable the bundled version of dbghelp so that CrashDebugHelper can load dbgeng.dll.
 		WindowsPlatform.bUseBundledDbgHelp = false;
 
diff --git a/Engine/Source/Programs/EpicWebHelper/EpicWebHelper.Target.cs b/Engine/Source/Programs/EpicWebHelper/EpicWebHelper.Target.cs
--- a/Engine/Source/Programs/EpicWebHelper/EpicWebHelper.Target.cs
+++ b/Engine/Source/Programs/EpicWebHelper/EpicWebHelper.Target.cs
@@ -54,9 +54,6 @@ public class EpicWebHelperTarget : TargetRules
 		// Epic Games Launcher needs to run on OS X 10.9, so CEFSubProcess needs this as well
 		bEnableOSX109Support = true;
 
-		// Disable automation tests to avoid compilation errors related to alignas, caused by setting bEnableOSX109Support to true above. Old OSX versions do not fully support this until C++17
-		bForceDisableAutomationTests = true;
-
 		// Already a manifest specified through resource file
 		WindowsPlatform.ManifestFile = null;
 	}
diff --git a/Engine/Source/Programs/Mac/UnrealEditorServices/UnrealEditorServices.Target.cs b/Engine/Source/Programs/Mac/UnrealEditorServices/UnrealEditorServices.Target.cs
--- a/Engine/Source/Programs/Mac/UnrealEditorServices/UnrealEditorServices.Target.cs
+++ b/Engine/Source/Programs/Mac/UnrealEditorServices/UnrealEditorServices.Target.cs
@@ -25,8 +25,5 @@ public class UnrealEditorServicesTarget : TargetRules
 
 		// We still need to support old versions of the engine that are compatible with OS X 10.9
 		bEnableOSX109Support = true;
-
-		// Disable automation tests to avoid compilation errors related to alignas, caused by setting bEnableOSX109Support to true above. Old OSX versions do not fully support this until C++17
-		bForceDisableAutomationTests = true;
 	}
 }
diff --git a/Engine/Source/Programs/UnrealBuildTool/Platform/Mac/MacToolChain.cs b/Engine/Source/Programs/UnrealBuildTool/Platform/Mac/MacToolChain.cs
--- a/Engine/Source/Programs/UnrealBuildTool/Platform/Mac/MacToolChain.cs
+++ b/Engine/Source/Programs/UnrealBuildTool/Platform/Mac/MacToolChain.cs
@@ -240,6 +240,11 @@ namespace UnrealBuildTool
 				Result += " -Wundef" + (CompileEnvironment.bUndefinedIdentifierWarningsAsErrors ? "" : " -Wno-error=undef");
 			}
 
+			if (CompileEnvironment.bEnableOSX109Support)
+			{
+				Result += " -faligned-new"; // aligned operator new is supported only on macOS 10.14 and above
+			}
+
 			Result += " -c";
 
 			// Pass through architecture and OS info
diff --git a/Engine/Source/Runtime/Core/Public/Containers/SpscQueue.h b/Engine/Source/Runtime/Core/Public/Containers/SpscQueue.h
--- a/Engine/Source/Runtime/Core/Public/Containers/SpscQueue.h
+++ b/Engine/Source/Runtime/Core/Public/Containers/SpscQueue.h
@@ -118,18 +118,10 @@ private:
 private:
 	// consumer part 
 	// accessed mainly by consumer, infrequently by producer 
-#ifndef PLATFORM_MAC // some projects are still built on macOS before v.10.14 that doesn't have aligned new/delete operators
 	alignas(PLATFORM_CACHE_LINE_SIZE) std::atomic<FNode*> Tail; // tail of the queue 
-#else
-	std::atomic<FNode*> Tail; // tail of the queue
-#endif
 	// producer part 
 	// accessed only by producer 
-#ifndef PLATFORM_MAC
-	alignas(PLATFORM_CACHE_LINE_SIZE) FNode* Head; // head of the queue 
-#else
 	FNode* Head; // head of the queue
-#endif
 	FNode* First; // last unused node (tail of node cache) 
 	FNode* TailCopy; // helper (points somewhere between First and Tail)
 };
\ No newline at end of file
