Move GIgnoreDebugger to CoreGlobals.h/cpp
Remove extern-ed GIgnoreDebugger from various platforms and sources

### Hashtags

#rnx

diff --git a/Engine/Source/Runtime/Core/Private/Android/AndroidPlatformMisc.cpp b/Engine/Source/Runtime/Core/Private/Android/AndroidPlatformMisc.cpp
--- a/Engine/Source/Runtime/Core/Private/Android/AndroidPlatformMisc.cpp
+++ b/Engine/Source/Runtime/Core/Private/Android/AndroidPlatformMisc.cpp
@@ -2456,7 +2456,6 @@ JNI_METHOD void Java_com_epicgames_unreal_GameActivity_nativeSetAndroidStartupSt
 #if !UE_BUILD_SHIPPING
 bool FAndroidMisc::IsDebuggerPresent()
 {
-	extern CORE_API bool GIgnoreDebugger;
 	if (GIgnoreDebugger)
 	{
 		return false;
diff --git a/Engine/Source/Runtime/Core/Private/Apple/ApplePlatformMisc.cpp b/Engine/Source/Runtime/Core/Private/Apple/ApplePlatformMisc.cpp
--- a/Engine/Source/Runtime/Core/Private/Apple/ApplePlatformMisc.cpp
+++ b/Engine/Source/Runtime/Core/Private/Apple/ApplePlatformMisc.cpp
@@ -13,6 +13,25 @@
 #include "Apple/ApplePlatformDebugEvents.h"
 #include "Apple/ApplePlatformCrashContext.h"
 #include "FramePro/FrameProProfiler.h"
+#include "CoreGlobals.h"
+
+bool FApplePlatformMisc::IsDebuggerPresent()
+{
+	// Based on http://developer.apple.com/library/mac/#qa/qa1361/_index.html
+
+	if (GIgnoreDebugger)
+	{
+		return false;
+	}
+
+	struct kinfo_proc Info;
+	int32 Mib[] = { CTL_KERN, KERN_PROC, KERN_PROC_PID, getpid() };
+	SIZE_T Size = sizeof(Info);
+
+	sysctl(Mib, sizeof(Mib) / sizeof(*Mib), &Info, &Size, NULL, 0);
+
+	return (Info.kp_proc.p_flag & P_TRACED) != 0;
+}
 
 void FApplePlatformMisc::GetEnvironmentVariable(const TCHAR* VariableName, TCHAR* Result, int32 ResultLength)
 {
diff --git a/Engine/Source/Runtime/Core/Private/HAL/ExceptionHandling.cpp b/Engine/Source/Runtime/Core/Private/HAL/ExceptionHandling.cpp
--- a/Engine/Source/Runtime/Core/Private/HAL/ExceptionHandling.cpp
+++ b/Engine/Source/Runtime/Core/Private/HAL/ExceptionHandling.cpp
@@ -22,13 +22,6 @@ CORE_API bool GAlwaysReportCrash = false;
 /** Whether to use ClientReportClient rather than the old AutoReporter. */
 CORE_API bool GUseCrashReportClient = true;
 
-#if !UE_BUILD_SHIPPING
-
-/** Whether we should ignore the attached debugger. */
-CORE_API bool GIgnoreDebugger = false;
-
-#endif // !UE_BUILD_SHIPPING
-
 CORE_API TCHAR MiniDumpFilenameW[1024] = TEXT("");
 
 
diff --git a/Engine/Source/Runtime/Core/Private/Misc/CoreGlobals.cpp b/Engine/Source/Runtime/Core/Private/Misc/CoreGlobals.cpp
--- a/Engine/Source/Runtime/Core/Private/Misc/CoreGlobals.cpp
+++ b/Engine/Source/Runtime/Core/Private/Misc/CoreGlobals.cpp
@@ -421,6 +421,13 @@ bool					GEnableVREditorHacks = false;
 
 bool CORE_API			GIsGPUCrashed = false;
 
+#if !UE_BUILD_SHIPPING
+
+/** Whether we should ignore the attached debugger. */
+CORE_API bool			GIgnoreDebugger = false;
+
+#endif // #if !UE_BUILD_SHIPPING
+
 bool GetEmitDrawEvents()
 {
 	return GEmitDrawEvents;
diff --git a/Engine/Source/Runtime/Core/Public/Apple/ApplePlatformMisc.h b/Engine/Source/Runtime/Core/Public/Apple/ApplePlatformMisc.h
--- a/Engine/Source/Runtime/Core/Public/Apple/ApplePlatformMisc.h
+++ b/Engine/Source/Runtime/Core/Public/Apple/ApplePlatformMisc.h
@@ -22,8 +22,6 @@
 
 #define UE_DEBUG_BREAK_IMPL() PLATFORM_BREAK()
 
-extern CORE_API bool GIgnoreDebugger;
-
 #ifdef __OBJC__
 
 class FScopeAutoreleasePool
@@ -62,23 +60,7 @@ struct CORE_API FApplePlatformMisc : public FGenericPlatformMisc
 	static FString GetEnvironmentVariable(const TCHAR* VariableName);
 
 #if !UE_BUILD_SHIPPING
-	static bool IsDebuggerPresent()
-	{
-		// Based on http://developer.apple.com/library/mac/#qa/qa1361/_index.html
-
-		if (GIgnoreDebugger)
-		{
-			return false;
-		}
-
-		struct kinfo_proc Info;
-		int32 Mib[] = { CTL_KERN, KERN_PROC, KERN_PROC_PID, getpid() };
-		SIZE_T Size = sizeof(Info);
-
-		sysctl( Mib, sizeof( Mib ) / sizeof( *Mib ), &Info, &Size, NULL, 0 );
-
-		return ( Info.kp_proc.p_flag & P_TRACED ) != 0;
-	}
+	static bool IsDebuggerPresent();
 #endif
 
 	FORCEINLINE static void MemoryBarrier()
diff --git a/Engine/Source/Runtime/Core/Public/CoreGlobals.h b/Engine/Source/Runtime/Core/Public/CoreGlobals.h
--- a/Engine/Source/Runtime/Core/Public/CoreGlobals.h
+++ b/Engine/Source/Runtime/Core/Public/CoreGlobals.h
@@ -539,6 +539,13 @@ extern CORE_API bool GPumpingMessages;
 /** Enables various editor and HMD hacks that allow the experimental VR editor feature to work, perhaps at the expense of other systems */
 extern CORE_API bool GEnableVREditorHacks;
 
+#if !UE_BUILD_SHIPPING
+
+/** Whether we should ignore the attached debugger. */
+extern CORE_API bool GIgnoreDebugger;
+
+#endif // #if !UE_BUILD_SHIPPING
+
 enum class ETaskTag : int32
 {
 	ENone						= 0 << 0,
diff --git a/Engine/Source/Runtime/Core/Public/HAL/ExceptionHandling.h b/Engine/Source/Runtime/Core/Public/HAL/ExceptionHandling.h
--- a/Engine/Source/Runtime/Core/Public/HAL/ExceptionHandling.h
+++ b/Engine/Source/Runtime/Core/Public/HAL/ExceptionHandling.h
@@ -13,13 +13,6 @@ extern CORE_API bool GAlwaysReportCrash;
 /** Whether to use ClientReportClient rather than AutoReporter. */
 extern CORE_API bool GUseCrashReportClient;
 
-#if !UE_BUILD_SHIPPING
-
-/** Whether we should ignore the attached debugger. */
-extern CORE_API bool GIgnoreDebugger;
-
-#endif // !UE_BUILD_SHIPPING
-
 extern CORE_API TCHAR MiniDumpFilenameW[1024];
 
 // #CrashReport: 2014-09-11 Move to PlatformExceptionHandling
