Fix for shipping builds on Apple, move IsDebuggerPresent behind #if !UE_BUILD_SHIPPING

#rnx

diff --git a/Engine/Source/Runtime/Core/Private/Apple/ApplePlatformMisc.cpp b/Engine/Source/Runtime/Core/Private/Apple/ApplePlatformMisc.cpp
--- a/Engine/Source/Runtime/Core/Private/Apple/ApplePlatformMisc.cpp
+++ b/Engine/Source/Runtime/Core/Private/Apple/ApplePlatformMisc.cpp
@@ -15,6 +15,8 @@
 #include "FramePro/FrameProProfiler.h"
 #include "CoreGlobals.h"
 
+#if !UE_BUILD_SHIPPING
+
 bool FApplePlatformMisc::IsDebuggerPresent()
 {
 	// Based on http://developer.apple.com/library/mac/#qa/qa1361/_index.html
@@ -33,6 +35,8 @@ bool FApplePlatformMisc::IsDebuggerPresent()
 	return (Info.kp_proc.p_flag & P_TRACED) != 0;
 }
 
+#endif // !UE_BUILD_SHIPPING
+
 void FApplePlatformMisc::GetEnvironmentVariable(const TCHAR* VariableName, TCHAR* Result, int32 ResultLength)
 {
 	// Replace hyphens with underscores. Some legacy UE environment variables (eg. UE-SharedDataCachePath) are in widespread
