Don't spawn the trace server if -buildmachine is present on the command line.

#rnx

diff --git a/Engine/Source/Runtime/Core/Private/ProfilingDebugging/TraceAuxiliary.cpp b/Engine/Source/Runtime/Core/Private/ProfilingDebugging/TraceAuxiliary.cpp
--- a/Engine/Source/Runtime/Core/Private/ProfilingDebugging/TraceAuxiliary.cpp
+++ b/Engine/Source/Runtime/Core/Private/ProfilingDebugging/TraceAuxiliary.cpp
@@ -673,7 +673,7 @@ void FTraceAuxiliary::Initialize(const TCHAR* CommandLine)
 	TRACE_CPUPROFILER_EVENT_SCOPE(FTraceAux_Init);
 
 #if WITH_UNREAL_TRACE_LAUNCH
-	if (!FParse::Param(CommandLine, TEXT("notraceserver")))
+	if (!(FParse::Param(CommandLine, TEXT("notraceserver")) || FParse::Param(CommandLine, TEXT("buildmachine"))))
 	{
 		TRACE_CPUPROFILER_EVENT_SCOPE(FTraceAux_LaunchUnrealTrace);
 		LaunchUnrealTraceInternal(CommandLine);
