FThread supports -nothreading

diff --git a/Engine/Source/Runtime/Core/Private/HAL/Thread.cpp b/Engine/Source/Runtime/Core/Private/HAL/Thread.cpp
--- a/Engine/Source/Runtime/Core/Private/HAL/Thread.cpp
+++ b/Engine/Source/Runtime/Core/Private/HAL/Thread.cpp
@@ -9,19 +9,22 @@
 #include "Templates/UnrealTemplate.h"
 #include "Templates/UniquePtr.h"
 #include "Misc/Fork.h"
+#include "Misc/SingleThreadRunnable.h"
 
-class FThreadImpl final : public FRunnable
+class FThreadImpl final : public FRunnable, public FSingleThreadRunnable
 {
 public:
 	FThreadImpl(
 		TCHAR const* ThreadName,
 		TUniqueFunction<void()>&& InThreadFunction,
+		TUniqueFunction<void()>&& InSingleThreadTickFunction,
 		uint32 StackSize,
 		EThreadPriority ThreadPriority,
 		FThreadAffinity ThreadAffinity,
 		bool bIsForkable
-	) 
+	)
 		: ThreadFunction(MoveTemp(InThreadFunction))
+		, SingleThreadTickFunction(MoveTemp(InSingleThreadTickFunction))
 		, RunnableThread(bIsForkable ? FForkProcessHelper::CreateForkableThread(this, ThreadName, StackSize, ThreadPriority, ThreadAffinity.ThreadAffinityMask)
 			: FRunnableThread::Create(this, ThreadName, StackSize, ThreadPriority, ThreadAffinity.ThreadAffinityMask))
 	{
@@ -77,30 +80,55 @@ private:
 		Self.Reset();
 	}
 
+	virtual FSingleThreadRunnable* GetSingleThreadInterface() override
+	{
+		return this;
+	}
+
+	virtual void Tick()
+	{
+		SingleThreadTickFunction();
+	}
+
 private:
 	// Two strong references are hold for `FThreadImpl`, one in parent `FThread` and another here. The reference in `FThreadImpl` is released
 	// on `Detach`, as to physically detach `FThreadImpl` from `FThread`. The reference below (`Self`) is released before exiting the thread,
 	// as the work was done and `FThreadImpl` doesn't mind to be deleted. Releasing the last reference deletes the instance, so 
 	// no member access can be performed after that.
-	// This must be the declared before `RunnableThread` so it's already initialized when the thread is created, otherwise the thread can complete
+	// This must be declared before `RunnableThread` so it's already initialized when the thread is created, otherwise the thread can complete
 	// before `Self` is initialized.
 	TSharedPtr<FThreadImpl, ESPMode::ThreadSafe> Self;
 
 	TAtomic<bool> bIsInitialized{ false };
 
 	TUniqueFunction<void()> ThreadFunction;
+	TUniqueFunction<void()> SingleThreadTickFunction;
 	TUniquePtr<FRunnableThread> RunnableThread;
 };
 
 FThread::FThread(
 	TCHAR const* ThreadName,
 	TUniqueFunction<void()>&& ThreadFunction,
-	uint32 StackSize,
-	EThreadPriority ThreadPriority,
-	FThreadAffinity ThreadAffinity,
-	bool bIsForkable
+	uint32 StackSize/* = 0*/,
+	EThreadPriority ThreadPriority/* = TPri_Normal*/,
+	FThreadAffinity ThreadAffinity/* = FThreadAffinity()*/,
+	bool bIsForkable/* = false*/
+)
+	: Impl(MakeShared<FThreadImpl, ESPMode::ThreadSafe>(ThreadName, MoveTemp(ThreadFunction), [] {}, StackSize, ThreadPriority, ThreadAffinity, bIsForkable))
+{
+	Impl->Initialize(Impl);
+}
+
+FThread::FThread(
+	TCHAR const* ThreadName,
+	TUniqueFunction<void()>&& ThreadFunction,
+	TUniqueFunction<void()>&& SingleThreadTickFunction,
+	uint32 StackSize/* = 0*/,
+	EThreadPriority ThreadPriority/* = TPri_Normal*/,
+	FThreadAffinity ThreadAffinity/* = FThreadAffinity()*/,
+	bool bIsForkable/* = false*/
 )
-	: Impl(MakeShared<FThreadImpl, ESPMode::ThreadSafe>(ThreadName, MoveTemp(ThreadFunction), StackSize, ThreadPriority, ThreadAffinity, bIsForkable))
+	: Impl(MakeShared<FThreadImpl, ESPMode::ThreadSafe>(ThreadName, MoveTemp(ThreadFunction), MoveTemp(SingleThreadTickFunction), StackSize, ThreadPriority, ThreadAffinity, bIsForkable))
 {
 	Impl->Initialize(Impl);
 }
diff --git a/Engine/Source/Runtime/Core/Public/HAL/Thread.h b/Engine/Source/Runtime/Core/Public/HAL/Thread.h
--- a/Engine/Source/Runtime/Core/Public/HAL/Thread.h
+++ b/Engine/Source/Runtime/Core/Public/HAL/Thread.h
@@ -48,6 +48,18 @@ public:
 		bool bIsForkable = false
 	);
 
+	// with SingleThreadTickFunction that will be executed every frame if running with `-nothreading
+	// (FPlatformProcess::SupportsMultithreading() == false)
+	FThread(
+		TCHAR const* ThreadName,
+		TUniqueFunction<void()>&& ThreadFunction,
+		TUniqueFunction<void()>&& SingleThreadTickFunction,
+		uint32 StackSize = 0,
+		EThreadPriority ThreadPriority = TPri_Normal,
+		FThreadAffinity ThreadAffinity = FThreadAffinity(),
+		bool bIsForkable = false
+	);
+
 	// non-copyable
 	FThread(const FThread&) = delete;
 	FThread& operator=(const FThread&) = delete;
