config option to disable reserve workers. reserve workers don't provide any perf improvement on consoles so are disabled by default

diff --git a/Engine/Source/Runtime/Core/Private/Async/TaskGraph.cpp b/Engine/Source/Runtime/Core/Private/Async/TaskGraph.cpp
--- a/Engine/Source/Runtime/Core/Private/Async/TaskGraph.cpp
+++ b/Engine/Source/Runtime/Core/Private/Async/TaskGraph.cpp
@@ -32,6 +32,7 @@
 #include "ProfilingDebugging/MiscTrace.h"
 #include "ProfilingDebugging/ScopedTimers.h"
 #include "Misc/ConfigCacheIni.h"
+#include "Misc/CoreDelegates.h"
 
 #include "Async/Fundamental/Scheduler.h"
 #include "Async/Fundamental/ReserveScheduler.h"
@@ -108,8 +109,6 @@ static FAutoConsoleVariableRef CVarNumForegroundWorkers(
 	TEXT("Configures the number of foreground worker threads. Requires the scheduler to be restarted to have an affect")
 );
 
-bool GDisableReserveWorkers = 0;
-
 #if CREATE_HIPRI_TASK_THREADS || CREATE_BACKGROUND_TASK_THREADS
 	static void ThreadSwitchForABTest(const TArray<FString>& Args)
 	{
@@ -1881,6 +1880,8 @@ class FTaskGraphCompatibilityImplementation final : public FTaskGraphInterface
 	FThreadSafeCounter	ReentrancyCheck;
 	FAAArrayQueue<FBaseGraphTask>	QueuedBackgroundTasks;
 
+	std::atomic<bool> bReserveWorkersEnabled{ false };
+
 public:
 	FTaskGraphCompatibilityImplementation(int32 InNumWorkerThreads) 
 		: NumWorkerThreads(FForkProcessHelper::IsForkedMultithreadInstance() ? CVar_ForkedProcess_MaxWorkerThreads : InNumWorkerThreads)
@@ -1898,10 +1899,9 @@ public:
 			int32 NumForegroundWorkers =  FMath::Max(1, NumWorkerThreads - NumBackgroundWorkers);
 
 			LowLevelTasks::FScheduler::Get().StartWorkers(NumForegroundWorkers, NumBackgroundWorkers, FForkProcessHelper::IsForkedMultithreadInstance(), FPlatformAffinity::GetTaskThreadPriority(), FPlatformAffinity::GetTaskBPThreadPriority());
-			if (!GDisableReserveWorkers)
-			{
-				LowLevelTasks::FReserveScheduler::Get().StartWorkers(LowLevelTasks::FScheduler::Get(), NumForegroundWorkers + NumBackgroundWorkers, FForkProcessHelper::IsForkedMultithreadInstance(), FPlatformAffinity::GetTaskBPThreadPriority());
-			}
+
+			check(GConfig == nullptr); // otherwise reserve workers should be started right here
+			FCoreDelegates::ConfigReadyForUse.AddRaw(this, &FTaskGraphCompatibilityImplementation::StartReserveWorkers);
 
 			NumNamedThreads = ENamedThreads::ActualRenderingThread + 1;
 			ENamedThreads::bHasBackgroundThreads = 1;
@@ -1910,7 +1910,7 @@ public:
 		else
 		{
 			LowLevelTasks::FScheduler::Get().StopWorkers();
-			if (!GDisableReserveWorkers)
+			if (bReserveWorkersEnabled)
 			{
 				LowLevelTasks::FReserveScheduler::Get().StopWorkers();
 			}
@@ -1933,6 +1933,8 @@ public:
 
 	~FTaskGraphCompatibilityImplementation() override
 	{
+		FCoreDelegates::ConfigReadyForUse.RemoveAll(this);
+
 		for (auto& Callback : ShutdownCallbacks)
 		{
 			Callback();
@@ -1944,7 +1946,7 @@ public:
 			NamedThreads[ThreadIndex].bAttached = false;
 		}
 		LowLevelTasks::FScheduler::Get().StopWorkers();
-		if (!GDisableReserveWorkers)
+		if (bReserveWorkersEnabled)
 		{
 			LowLevelTasks::FReserveScheduler::Get().StopWorkers();
 		}
@@ -1982,7 +1984,7 @@ public:
 			LowLevelTasks::FScheduler::Get().StopWorkers();
 			LowLevelTasks::FScheduler::Get().StartWorkers(NumWorkers, NumBackgroundWorkers, FForkProcessHelper::IsForkedMultithreadInstance(), Pri, FPlatformAffinity::GetTaskBPThreadPriority());
 
-			if (!GDisableReserveWorkers)
+			if (bReserveWorkersEnabled)
 			{
 				LowLevelTasks::FReserveScheduler::Get().StopWorkers();
 				LowLevelTasks::FReserveScheduler::Get().StartWorkers(LowLevelTasks::FScheduler::Get(), NumWorkers + NumBackgroundWorkers, FForkProcessHelper::IsForkedMultithreadInstance(), FPlatformAffinity::GetTaskBPThreadPriority());
@@ -2161,7 +2163,7 @@ private:
 		{
 			// a worker thread gets blocked, involve a reserve worker to utilise an idle core
 			bool bSuccess = false; // has a reserve worker helped?
-			if (!GDisableReserveWorkers)
+			if (bReserveWorkersEnabled.load(std::memory_order_relaxed))
 			{
 				if (LowLevelTasks::DoReserveWorkUntil([Index(0), Tasks]() mutable
 				{
@@ -2319,6 +2321,28 @@ private:
 		checkThreadGraph(NamedThreads[Index].TaskGraphWorker->GetThreadId() == Index);
 		return *NamedThreads[Index].TaskGraphWorker;
 	}
+
+	void StartReserveWorkers()
+	{
+		if (bReserveWorkersEnabled)
+		{
+			return; // once enabled, reserve workers can't be disabled
+		}
+
+		int32 NumBackgroundWorkers = FMath::Max(1, NumWorkerThreads - FMath::Min<int>(GNumForegroundWorkers, NumWorkerThreads));
+		int32 NumForegroundWorkers = FMath::Max(1, NumWorkerThreads - NumBackgroundWorkers);
+
+		check(GConfig);
+		
+		bool bEnableReserveWorkers = true; // by default
+		GConfig->GetBool(TEXT("TaskGraph"), TEXT("EnableReserveWorkers"), bEnableReserveWorkers, GEngineIni);
+
+		if (bEnableReserveWorkers)
+		{
+			bReserveWorkersEnabled = true;
+			LowLevelTasks::FReserveScheduler::Get().StartWorkers(LowLevelTasks::FScheduler::Get(), NumForegroundWorkers + NumBackgroundWorkers, FForkProcessHelper::IsForkedMultithreadInstance(), FPlatformAffinity::GetTaskBPThreadPriority());
+		}
+	}
 };
 
 thread_local TArray<FBaseGraphTask*> FTaskGraphCompatibilityImplementation::NewTasks;
@@ -2344,15 +2368,6 @@ void FTaskGraphInterface::Startup(int32 NumThreads)
 		GUseNewTaskBackend = 1;
 	}
 
-	if (FParse::Param(FCommandLine::Get(), TEXT("TaskGraphDisableReserveWorkers")))
-	{
-		GDisableReserveWorkers = 1;
-	}
-	else if (FParse::Param(FCommandLine::Get(), TEXT("TaskGraphEnableReserveWorkers")))
-	{
-		GDisableReserveWorkers = 0;
-	}
-
 	if (GUseNewTaskBackend)
 	{
 		//We want to reduce the number of overall threads that UE uses so that there is are some 
