Fix error C2504: 'TAlignedBytes<1024,64>': base class undefined when making a shared FOnlineServicesEOS

diff --git a/Engine/Source/Runtime/Core/Public/Templates/TypeCompatibleBytes.h b/Engine/Source/Runtime/Core/Public/Templates/TypeCompatibleBytes.h
--- a/Engine/Source/Runtime/Core/Public/Templates/TypeCompatibleBytes.h
+++ b/Engine/Source/Runtime/Core/Public/Templates/TypeCompatibleBytes.h
@@ -45,6 +45,7 @@ struct TAlignedBytes<Size,1>
 #endif
 
 // Implement TAlignedBytes for these alignments.
+IMPLEMENT_ALIGNED_STORAGE(64);
 IMPLEMENT_ALIGNED_STORAGE(32);
 IMPLEMENT_ALIGNED_STORAGE(16);
 IMPLEMENT_ALIGNED_STORAGE(8);
