Csv profiler named event support

Add support for emitting named events from timing stats without compiling them in. This is controlled via new commandline options, which can be used in combination:
 -csvNamedEvents emits named events for exclusive stats
 -csvNamedEventsTiming emits named events for standard timing stats

Instrumentation overhead reduction
-Make it possible to disable CSVPROFILERTRACE_ENABLED, since this increases CSV profiler instrumentation overhead by about 15%
- Remove LLM scopes from stat functions and move into stat block allocation function. Reduces overhead slightly

diff --git a/Engine/Source/Runtime/Core/Private/ProfilingDebugging/CsvProfiler.cpp b/Engine/Source/Runtime/Core/Private/ProfilingDebugging/CsvProfiler.cpp
--- a/Engine/Source/Runtime/Core/Private/ProfilingDebugging/CsvProfiler.cpp
+++ b/Engine/Source/Runtime/Core/Private/ProfilingDebugging/CsvProfiler.cpp
@@ -2,7 +2,8 @@
 
 /**
 *
-* A lightweight multi-threaded CSV profiler which can be used for profiling in Test/Shipping builds
+* A lightweight multi-threaded profiler with very low instrumentation overhead. Suitable for Test or even final Shipping builds
+* Results are accumulated per-frame and emitted in CSV format
 */
 
 #include "ProfilingDebugging/CsvProfiler.h"
@@ -35,6 +36,11 @@
 
 #define CSV_PROFILER_INLINE FORCEINLINE
 
+#ifndef CSV_PROFILER_SUPPORT_NAMED_EVENTS
+// This doesn't actually enable named events. Use -csvNamedEvents or -csvNamedEventsTiming to enable for exclusive or normal timing stats respectively
+#define CSV_PROFILER_SUPPORT_NAMED_EVENTS ENABLE_NAMED_EVENTS
+#endif
+
 #define REPAIR_MARKER_STACKS 1
 #define CSV_THREAD_HIGH_PRI 0
 
@@ -137,7 +143,6 @@ static bool GGameThreadIsCsvProcessingThread = true;
 
 static uint32 GCsvProfilerFrameNumber = 0;
 
-
 static bool GCsvTrackWaitsOnAllThreads = false;
 static bool GCsvTrackWaitsOnGameThread = true;
 static bool GCsvTrackWaitsOnRenderThread = true;
@@ -160,6 +165,18 @@ static bool GCsvExitOnCompletion = false;
 
 static thread_local bool GCsvThreadLocalWaitsEnabled = false;
 
+
+#if CSV_PROFILER_SUPPORT_NAMED_EVENTS
+  #if PLATFORM_IMPLEMENTS_BeginNamedEventStatic
+    #define CSV_PROFILER_BeginNamedEvent(Color,Text) FPlatformMisc::BeginNamedEventStatic(Color, Text)
+  #else
+    #define CSV_PROFILER_BeginNamedEvent(Color,Text) FPlatformMisc::BeginNamedEvent(Color, Text)
+  #endif
+
+bool GCsvProfilerNamedEventsExclusive = false;
+bool GCsvProfilerNamedEventsTiming = false;
+#endif //CSV_PROFILER_SUPPORT_NAMED_EVENTS
+
 bool IsContinuousWriteEnabled(bool bGameThread)
 {
 	int CVarValue = -1;
@@ -747,6 +764,7 @@ public:
 private:
 	void AddTailBlock()
 	{
+		LLM_SCOPE(ELLMTag::CsvProfiler);
 		FBlock* NewTail = new FBlock;
 		if (TailBlock == nullptr)
 		{
@@ -1939,6 +1957,7 @@ public:
 
 	CSV_PROFILER_INLINE void PushWaitStatName(const char * WaitStatName)
 	{
+		LLM_SCOPE(ELLMTag::CsvProfiler);
 		WaitStatNameStack.Push(WaitStatName);
 	}
 	CSV_PROFILER_INLINE const char* PopWaitStatName()
@@ -3067,7 +3086,12 @@ void FCsvProfiler::BeginStat(const char * StatName, uint32 CategoryIndex)
 #if RECORD_TIMESTAMPS
 	if (GCsvProfilerIsCapturing && GCsvCategoriesEnabled[CategoryIndex])
 	{
-		LLM_SCOPE(ELLMTag::CsvProfiler);
+#if CSV_PROFILER_SUPPORT_NAMED_EVENTS
+		if (UNLIKELY(GCsvProfilerNamedEventsTiming))
+		{
+			CSV_PROFILER_BeginNamedEvent(FColor(255, 128, 255), StatName);
+		}
+#endif
 		FCsvProfilerThreadData::Get().AddTimestampBegin(StatName, CategoryIndex);
 	}
 #endif
@@ -3078,7 +3102,6 @@ void FCsvProfiler::BeginStat(const FName& StatName, uint32 CategoryIndex)
 #if RECORD_TIMESTAMPS
 	if (GCsvProfilerIsCapturing && GCsvCategoriesEnabled[CategoryIndex])
 	{
-		LLM_SCOPE(ELLMTag::CsvProfiler);
 		FCsvProfilerThreadData::Get().AddTimestampBegin(StatName, CategoryIndex);
 	}
 #endif
@@ -3089,8 +3112,13 @@ void FCsvProfiler::EndStat(const char * StatName, uint32 CategoryIndex)
 #if RECORD_TIMESTAMPS
 	if (GCsvProfilerIsCapturing && GCsvCategoriesEnabled[CategoryIndex])
 	{
-		LLM_SCOPE(ELLMTag::CsvProfiler);
 		FCsvProfilerThreadData::Get().AddTimestampEnd(StatName, CategoryIndex);
+#if CSV_PROFILER_SUPPORT_NAMED_EVENTS
+		if (UNLIKELY(GCsvProfilerNamedEventsTiming))
+		{
+			FPlatformMisc::EndNamedEvent();
+		}
+#endif
 	}
 #endif
 }
@@ -3100,7 +3128,6 @@ void FCsvProfiler::EndStat(const FName& StatName, uint32 CategoryIndex)
 #if RECORD_TIMESTAMPS
 	if (GCsvProfilerIsCapturing && GCsvCategoriesEnabled[CategoryIndex])
 	{
-		LLM_SCOPE(ELLMTag::CsvProfiler);
 		FCsvProfilerThreadData::Get().AddTimestampEnd(StatName, CategoryIndex);
 	}
 #endif
@@ -3111,7 +3138,12 @@ void FCsvProfiler::BeginExclusiveStat(const char * StatName)
 #if RECORD_TIMESTAMPS
 	if (GCsvProfilerIsCapturing && GCsvCategoriesEnabled[CSV_CATEGORY_INDEX(Exclusive)])
 	{
-		LLM_SCOPE(ELLMTag::CsvProfiler);
+#if CSV_PROFILER_SUPPORT_NAMED_EVENTS
+		if (UNLIKELY(GCsvProfilerNamedEventsExclusive))
+		{
+			CSV_PROFILER_BeginNamedEvent(FColor(255, 128, 128), StatName);
+		}
+#endif
 		FCsvProfilerThreadData::Get().AddTimestampExclusiveBegin(StatName);
 	}
 #endif
@@ -3122,8 +3154,14 @@ void FCsvProfiler::EndExclusiveStat(const char * StatName)
 #if RECORD_TIMESTAMPS
 	if (GCsvProfilerIsCapturing && GCsvCategoriesEnabled[CSV_CATEGORY_INDEX(Exclusive)])
 	{
-		LLM_SCOPE(ELLMTag::CsvProfiler);
 		FCsvProfilerThreadData::Get().AddTimestampExclusiveEnd(StatName);
+
+#if CSV_PROFILER_SUPPORT_NAMED_EVENTS
+		if (UNLIKELY(GCsvProfilerNamedEventsExclusive))
+		{
+			FPlatformMisc::EndNamedEvent();
+		}
+#endif
 	}
 #endif
 }
@@ -3134,7 +3172,6 @@ void FCsvProfiler::BeginSetWaitStat(const char * StatName)
 #if RECORD_TIMESTAMPS
 	if (GCsvProfilerIsCapturing && GCsvCategoriesEnabled[CSV_CATEGORY_INDEX(Exclusive)])
 	{
-		LLM_SCOPE(ELLMTag::CsvProfiler);
 		FCsvProfilerThreadData::Get().PushWaitStatName(StatName == nullptr ? GIgnoreWaitStatName : StatName);
 	}
 #endif
@@ -3155,10 +3192,15 @@ void FCsvProfiler::BeginWait()
 #if RECORD_TIMESTAMPS
 	if (GCsvProfilerIsCapturing && GCsvCategoriesEnabled[CSV_CATEGORY_INDEX(Exclusive)])
 	{
-		LLM_SCOPE(ELLMTag::CsvProfiler);
 		const char* WaitStatName = FCsvProfilerThreadData::Get().GetWaitStatName();
 		if (WaitStatName != GIgnoreWaitStatName)
 		{
+#if CSV_PROFILER_SUPPORT_NAMED_EVENTS
+			if (UNLIKELY(GCsvProfilerNamedEventsExclusive))
+			{
+				CSV_PROFILER_BeginNamedEvent(FColor(255, 128, 128), "CsvEventWait");
+			}
+#endif
 			FCsvProfilerThreadData::Get().AddTimestampExclusiveBegin(WaitStatName);
 		}
 	}
@@ -3170,11 +3212,16 @@ void FCsvProfiler::EndWait()
 #if RECORD_TIMESTAMPS
 	if (GCsvProfilerIsCapturing && GCsvCategoriesEnabled[CSV_CATEGORY_INDEX(Exclusive)])
 	{
-		LLM_SCOPE(ELLMTag::CsvProfiler);
 		const char* WaitStatName = FCsvProfilerThreadData::Get().GetWaitStatName();
 		if (WaitStatName != GIgnoreWaitStatName)
 		{
 			FCsvProfilerThreadData::Get().AddTimestampExclusiveEnd(FCsvProfilerThreadData::Get().GetWaitStatName());
+#if CSV_PROFILER_SUPPORT_NAMED_EVENTS
+			if (UNLIKELY(GCsvProfilerNamedEventsExclusive))
+			{
+				FPlatformMisc::EndNamedEvent();
+			}
+#endif
 		}
 	}
 #endif
@@ -3286,7 +3333,6 @@ void FCsvProfiler::RecordCustomStat(const char * StatName, uint32 CategoryIndex,
 {
 	if (GCsvProfilerIsCapturing && GCsvCategoriesEnabled[CategoryIndex])
 	{
-		LLM_SCOPE(ELLMTag::CsvProfiler);
 		FCsvProfilerThreadData::Get().AddCustomStat(StatName, CategoryIndex, Value, CustomStatOp);
 	}
 }
@@ -3301,7 +3347,6 @@ void FCsvProfiler::RecordCustomStat(const FName& StatName, uint32 CategoryIndex,
 {
 	if (GCsvProfilerIsCapturing && GCsvCategoriesEnabled[CategoryIndex])
 	{
-		LLM_SCOPE(ELLMTag::CsvProfiler);
 		FCsvProfilerThreadData::Get().AddCustomStat(StatName, CategoryIndex, Value, CustomStatOp);
 	}
 }
@@ -3317,7 +3362,6 @@ void FCsvProfiler::RecordCustomStat(const char * StatName, uint32 CategoryIndex,
 {
 	if (GCsvProfilerIsCapturing && GCsvCategoriesEnabled[CategoryIndex])
 	{
-		LLM_SCOPE(ELLMTag::CsvProfiler);
 		FCsvProfilerThreadData::Get().AddCustomStat(StatName, CategoryIndex, Value, CustomStatOp);
 	}
 }
@@ -3326,7 +3370,6 @@ void FCsvProfiler::RecordCustomStat(const FName& StatName, uint32 CategoryIndex,
 {
 	if (GCsvProfilerIsCapturing && GCsvCategoriesEnabled[CategoryIndex])
 	{
-		LLM_SCOPE(ELLMTag::CsvProfiler);
 		FCsvProfilerThreadData::Get().AddCustomStat(StatName, CategoryIndex, Value, CustomStatOp);
 	}
 }
@@ -3407,6 +3450,16 @@ void FCsvProfiler::Init()
 	{
 		GCsvUseProcessingThread = false;
 	}
+#if CSV_PROFILER_SUPPORT_NAMED_EVENTS
+	if (FParse::Param(FCommandLine::Get(), TEXT("csvNamedEvents")))
+	{
+		GCsvProfilerNamedEventsExclusive = true;
+	}
+	if (FParse::Param(FCommandLine::Get(), TEXT("csvNamedEventsTiming")))
+	{
+		GCsvProfilerNamedEventsTiming = true;
+	}
+#endif
 	if (FParse::Param(FCommandLine::Get(), TEXT("csvStatCounts")))
 	{
 		CVarCsvStatCounts.AsVariable()->Set(1);
@@ -3535,6 +3588,7 @@ float FCsvProfiler::ProcessStatData()
 
 #if CSV_PROFILER_ALLOW_DEBUG_FEATURES
 
+// Simple benchmarking and debugging tests for the csv profiler. Enable with -csvtest, e.g -csvtest -csvcaptureframes=400
 void CSVTest()
 {
 	uint32 FrameNumber = FCsvProfiler::Get()->GetCaptureFrameNumber();
@@ -3549,8 +3603,12 @@ void CSVTest()
 	}
 
 	{
+		// This stat measures the overhead of submitting 10k timing stat scopes in a frame. 
+		// Multiply the ms result by 100 to get the per-scope cost in ns
+		// (currently ~150ns/scope on last-gen consoles if CSVPROFILERTRACE_ENABLED is 0)
+		// Note that each scope emits two timestamps, so the per-timestamp cost is half this
 		CSV_SCOPED_TIMING_STAT(CsvTest, TimerStatTimer);
-		for (int i = 0; i < 100; i++)
+		for (int i = 0; i < 2500; i++)
 		{
 			CSV_SCOPED_TIMING_STAT(CsvTest, BeginEndbenchmarkInner0);
 			CSV_SCOPED_TIMING_STAT(CsvTest, BeginEndbenchmarkInner1);
diff --git a/Engine/Source/Runtime/Core/Public/ProfilingDebugging/CsvProfiler.h b/Engine/Source/Runtime/Core/Public/ProfilingDebugging/CsvProfiler.h
--- a/Engine/Source/Runtime/Core/Public/ProfilingDebugging/CsvProfiler.h
+++ b/Engine/Source/Runtime/Core/Public/ProfilingDebugging/CsvProfiler.h
@@ -20,13 +20,8 @@
 
 #include <atomic>
 
-
-
 #if CSV_PROFILER
 
-#define CSV_TIMING_STATS_EMIT_NAMED_EVENTS 0
-#define CSV_EXCLUSIVE_TIMING_STATS_EMIT_NAMED_EVENTS 0
-
 // Helpers
 #define CSV_CATEGORY_INDEX(CategoryName)						(_GCsvCategory_##CategoryName.Index)
 #define CSV_CATEGORY_INDEX_GLOBAL								(0)
@@ -116,6 +111,7 @@
 
 
 #if CSV_PROFILER
+
 class FCsvProfilerFrame;
 class FCsvProfilerThreadData;
 class FCsvProfilerProcessingThread;
@@ -359,16 +355,10 @@ public:
 		, CategoryIndex(InCategoryIndex)
 	{
 		FCsvProfiler::BeginStat(StatName, CategoryIndex);
-#if CSV_TIMING_STATS_EMIT_NAMED_EVENTS
-		FPlatformMisc::BeginNamedEvent(FColor(255, 128, 255), StatName);
-#endif
 	}
 
 	~FScopedCsvStat()
 	{
-#if CSV_TIMING_STATS_EMIT_NAMED_EVENTS
-		FPlatformMisc::EndNamedEvent();
-#endif
 		FCsvProfiler::EndStat(StatName, CategoryIndex);
 	}
 	const char * StatName;
@@ -382,16 +372,10 @@ public:
 		: StatName(InStatName)
 	{
 		FCsvProfiler::BeginExclusiveStat(StatName);
-#if CSV_EXCLUSIVE_TIMING_STATS_EMIT_NAMED_EVENTS
-		FPlatformMisc::BeginNamedEvent(FColor(255, 128, 128), StatName);
-#endif
 	}
 
 	~FScopedCsvStatExclusive()
 	{
-#if CSV_EXCLUSIVE_TIMING_STATS_EMIT_NAMED_EVENTS
-		FPlatformMisc::EndNamedEvent();
-#endif
 		FCsvProfiler::EndExclusiveStat(StatName);
 	}
 	const char * StatName;
@@ -407,9 +391,6 @@ public:
 		if (bCondition)
 		{
 			FCsvProfiler::BeginExclusiveStat(StatName);
-#if CSV_EXCLUSIVE_TIMING_STATS_EMIT_NAMED_EVENTS
-			FPlatformMisc::BeginNamedEvent(FColor(255,128,128),StatName);
-#endif
 		}
 	}
 
@@ -417,9 +398,6 @@ public:
 	{
 		if (bCondition)
 		{
-#if CSV_EXCLUSIVE_TIMING_STATS_EMIT_NAMED_EVENTS
-			FPlatformMisc::EndNamedEvent();
-#endif
 			FCsvProfiler::EndExclusiveStat(StatName);
 		}
 	}
@@ -436,9 +414,6 @@ public:
 		if (bCondition)
 		{
 			FCsvProfiler::BeginWait();
-#if CSV_EXCLUSIVE_TIMING_STATS_EMIT_NAMED_EVENTS
-			FPlatformMisc::BeginNamedEvent(FColor(255, 128, 128), EventWait);
-#endif
 		}
 	}
 
@@ -446,9 +421,6 @@ public:
 	{
 		if (bCondition)
 		{
-#if CSV_EXCLUSIVE_TIMING_STATS_EMIT_NAMED_EVENTS
-			FPlatformMisc::EndNamedEvent();
-#endif
 			FCsvProfiler::EndWait();
 		}
 	}
diff --git a/Engine/Source/Runtime/Core/Public/ProfilingDebugging/CsvProfilerTrace.h b/Engine/Source/Runtime/Core/Public/ProfilingDebugging/CsvProfilerTrace.h
--- a/Engine/Source/Runtime/Core/Public/ProfilingDebugging/CsvProfilerTrace.h
+++ b/Engine/Source/Runtime/Core/Public/ProfilingDebugging/CsvProfilerTrace.h
@@ -6,11 +6,14 @@
 #include "Trace/Config.h"
 #include "ProfilingDebugging/CsvProfilerConfig.h"
 
+// Note: Enabling this adds about 15% to the instrumentation overhead of the CSV profiler, even if Insights is not capturing
+#ifndef CSVPROFILERTRACE_ENABLED
 #if UE_TRACE_ENABLED && !UE_BUILD_SHIPPING
 #define CSVPROFILERTRACE_ENABLED (WITH_ENGINE && CSV_PROFILER)
 #else
 #define CSVPROFILERTRACE_ENABLED 0
 #endif
+#endif
 
 class FName;
 
diff --git a/Engine/Source/Runtime/RenderCore/Private/RenderGraphEvent.cpp b/Engine/Source/Runtime/RenderCore/Private/RenderGraphEvent.cpp
--- a/Engine/Source/Runtime/RenderCore/Private/RenderGraphEvent.cpp
+++ b/Engine/Source/Runtime/RenderCore/Private/RenderGraphEvent.cpp
@@ -350,18 +350,12 @@ inline void OnPushCSVStat(const FRDGCSVStatScope* Scope)
 {
 #if CSV_PROFILER
 	FCsvProfiler::BeginExclusiveStat(Scope->StatName);
-#if CSV_EXCLUSIVE_TIMING_STATS_EMIT_NAMED_EVENTS
-	FPlatformMisc::BeginNamedEvent(FColor(255, 128, 128), Scope->StatName);
-#endif
 #endif
 }
 
 inline void OnPopCSVStat(const FRDGCSVStatScope* Scope)
 {
 #if CSV_PROFILER
-#if CSV_EXCLUSIVE_TIMING_STATS_EMIT_NAMED_EVENTS
-	FPlatformMisc::EndNamedEvent();
-#endif
 	FCsvProfiler::EndExclusiveStat(Scope->StatName);
 #endif
 }
