TaskGraph: starting reserve workers regardless if GConfig is already initialized or not

diff --git a/Engine/Source/Runtime/Core/Private/Async/TaskGraph.cpp b/Engine/Source/Runtime/Core/Private/Async/TaskGraph.cpp
--- a/Engine/Source/Runtime/Core/Private/Async/TaskGraph.cpp
+++ b/Engine/Source/Runtime/Core/Private/Async/TaskGraph.cpp
@@ -1900,8 +1900,16 @@ public:
 
 			LowLevelTasks::FScheduler::Get().StartWorkers(NumForegroundWorkers, NumBackgroundWorkers, FForkProcessHelper::IsForkedMultithreadInstance(), FPlatformAffinity::GetTaskThreadPriority(), FPlatformAffinity::GetTaskBPThreadPriority());
 
-			check(GConfig == nullptr); // otherwise reserve workers should be started right here
-			FCoreDelegates::ConfigReadyForUse.AddRaw(this, &FTaskGraphCompatibilityImplementation::StartReserveWorkers);
+			check(IsInGameThread()); // otherwise we can have a race on starting reserve workers below
+			if (GConfig == nullptr)
+			{
+				// postpone starting reserve workers tillGConfig is initialized, to know if reserve workers are disabled
+				FCoreDelegates::ConfigReadyForUse.AddRaw(this, &FTaskGraphCompatibilityImplementation::StartReserveWorkers);
+			}
+			else
+			{
+				StartReserveWorkers();
+			}
 
 			NumNamedThreads = ENamedThreads::ActualRenderingThread + 1;
 			ENamedThreads::bHasBackgroundThreads = 1;
