Extend FResourceSizeEx to keep track of name for each memory size added.  This allows for verbose reporting of resource sizes.

This results in output like the following:

                                                                                                                            Object      NumKB      MaxKB   ResExcKB  ResExcDedSysKB  ResExcDedVidKB  ResExcUnkKB
StaticMesh /Game/AncientContent/Megascans/Surfaces/Rock/SM_RockRough1.SM_RockRough1       1.31          1.33     395.50           12.75                  0.00                 382.75
						                               None                                                 12.75
							               None                                                  0.79
							               TexcoordBuffer and TangentBuffer       66.35
							               PositionVertexBuffer                            5.53
							               ColorVertexBuffer                              22.12
							               IndexBuffer                                       95.99
							               ReversedDepthOnlyIndexBuffer           95.99
							               ReversedIndexBuffer                          95.99
							               WireframeIndexBuffer                          0.00

- Adding a `-verbose` flag to `obj list` which will dump each tracked named size.
- Adding additional `AddSharedSystemMemoryBytes(const FName& Tag, const SIZE_T InMemoryBytes)` function which accepts a named tag.
- Modify FStaticMeshLODResources::GetResourceSizeEx() to use the new interface.
- Modify FSkeletalMeshLODRenderData::GetResourceSizeEx() to use the new interface.
- Remove the "Shared" version of the memory tracking functions.  These are not used at all across the engine.  This simplifies the code and makes the output of obj list more readable.
- Move FResourceSizeEx into its own cpp file so modifying the header isn't a painful full recompile.

diff --git a/Engine/Source/Runtime/Core/Private/ProfilingDebugging/ResourceSize.cpp b/Engine/Source/Runtime/Core/Private/ProfilingDebugging/ResourceSize.cpp
new file mode 100644
--- /dev/null
+++ b/Engine/Source/Runtime/Core/Private/ProfilingDebugging/ResourceSize.cpp
@@ -0,0 +1,160 @@
+// Copyright Epic Games, Inc. All Rights Reserved.
+
+#include "ProfilingDebugging/ResourceSize.h"
+
+static FName UntrackedMemoryName(TEXT("Untracked Memory"));
+
+FResourceSizeEx::FResourceSizeEx()
+	: ResourceSizeMode(EResourceSizeMode::Exclusive)
+{
+}
+
+FResourceSizeEx::FResourceSizeEx(const EResourceSizeMode::Type InResourceSizeMode)
+	: ResourceSizeMode(InResourceSizeMode)
+{
+}
+
+FResourceSizeEx::FResourceSizeEx(const EResourceSizeMode::Type InResourceSizeMode, const SIZE_T InDedicatedSystemMemoryBytes, const SIZE_T InDedicatedVideoMemoryBytes)
+	: ResourceSizeMode(InResourceSizeMode)
+{
+	DedicatedSystemMemoryBytesMap.Add(UntrackedMemoryName, InDedicatedSystemMemoryBytes);
+	DedicatedVideoMemoryBytesMap.Add(UntrackedMemoryName, InDedicatedVideoMemoryBytes);
+}
+
+FResourceSizeEx::FResourceSizeEx(const EResourceSizeMode::Type InResourceSizeMode, const SIZE_T InUnknownMemoryBytes)
+	: ResourceSizeMode(InResourceSizeMode)
+{
+	UnknownMemoryBytesMap.Add(UntrackedMemoryName, InUnknownMemoryBytes);
+}
+
+void FResourceSizeEx::LogSummary(FOutputDevice& Ar) const
+{
+	auto PrintPair = [&Ar](const TPair<FName, SIZE_T>& Pair)
+	{
+		Ar.Logf(
+			TEXT("%140s %15.2f"),
+			*Pair.Key.ToString(),
+			static_cast<double>(Pair.Value) / 1024.0
+		);
+	};
+
+	for (const TPair<FName, SIZE_T>& Pair : DedicatedSystemMemoryBytesMap)
+	{
+		PrintPair(Pair);
+	}
+	for (const TPair<FName, SIZE_T>& Pair : DedicatedVideoMemoryBytesMap)
+	{
+		PrintPair(Pair);
+	}
+	for (const TPair<FName, SIZE_T>& Pair : UnknownMemoryBytesMap)
+	{
+		PrintPair(Pair);
+	}
+}
+
+EResourceSizeMode::Type FResourceSizeEx::GetResourceSizeMode() const
+{
+	return ResourceSizeMode;
+}
+
+FResourceSizeEx& FResourceSizeEx::AddDedicatedSystemMemoryBytes(const FName& Tag, const SIZE_T InMemoryBytes)
+{
+	SIZE_T& CurrentSize = DedicatedSystemMemoryBytesMap.FindOrAdd(Tag);
+	CurrentSize += InMemoryBytes;
+	return *this;
+}
+
+FResourceSizeEx& FResourceSizeEx::AddDedicatedSystemMemoryBytes(const SIZE_T InMemoryBytes)
+{
+	SIZE_T& CurrentSize = DedicatedSystemMemoryBytesMap.FindOrAdd(UntrackedMemoryName);
+	CurrentSize += InMemoryBytes;
+	return *this;
+}
+
+SIZE_T FResourceSizeEx::GetDedicatedSystemMemoryBytes() const
+{
+	SIZE_T Sum = 0;
+	TArray<SIZE_T> Values;
+	DedicatedSystemMemoryBytesMap.GenerateValueArray(Values);
+	for (const SIZE_T s : Values) Sum += s;
+	return Sum;
+}
+
+FResourceSizeEx& FResourceSizeEx::AddDedicatedVideoMemoryBytes(const FName& Tag, const SIZE_T InMemoryBytes)
+{
+	SIZE_T& CurrentSize = DedicatedVideoMemoryBytesMap.FindOrAdd(Tag);
+	CurrentSize += InMemoryBytes;
+	return *this;
+}
+
+FResourceSizeEx& FResourceSizeEx::AddDedicatedVideoMemoryBytes(const SIZE_T InMemoryBytes)
+{
+	SIZE_T& CurrentSize = DedicatedVideoMemoryBytesMap.FindOrAdd(UntrackedMemoryName);
+	CurrentSize += InMemoryBytes;
+	return *this;
+}
+
+SIZE_T FResourceSizeEx::GetDedicatedVideoMemoryBytes() const
+{
+	SIZE_T Sum = 0;
+	TArray<SIZE_T> Values;
+	DedicatedVideoMemoryBytesMap.GenerateValueArray(Values);
+	for (const SIZE_T s : Values) Sum += s;
+	return Sum;
+}
+
+FResourceSizeEx& FResourceSizeEx::AddUnknownMemoryBytes(const FName& Tag, const SIZE_T InMemoryBytes)
+{
+	SIZE_T& CurrentSize = UnknownMemoryBytesMap.FindOrAdd(Tag);
+	CurrentSize += InMemoryBytes;
+	return *this;
+}
+
+FResourceSizeEx& FResourceSizeEx::AddUnknownMemoryBytes(const SIZE_T InMemoryBytes)
+{
+	SIZE_T& CurrentSize = UnknownMemoryBytesMap.FindOrAdd(UntrackedMemoryName);
+	CurrentSize += InMemoryBytes;
+	return *this;
+}
+
+SIZE_T FResourceSizeEx::GetUnknownMemoryBytes() const
+{
+	SIZE_T Sum = 0;
+	TArray<SIZE_T> Values;
+	UnknownMemoryBytesMap.GenerateValueArray(Values);
+	for (const SIZE_T s : Values) Sum += s;
+	return Sum;
+}
+
+SIZE_T FResourceSizeEx::GetTotalMemoryBytes() const
+{
+	return GetDedicatedSystemMemoryBytes() + GetDedicatedVideoMemoryBytes() + GetUnknownMemoryBytes();
+}
+
+FResourceSizeEx& FResourceSizeEx::operator+=(const FResourceSizeEx& InRHS)
+{
+	ensureAlwaysMsgf(ResourceSizeMode == InRHS.ResourceSizeMode, TEXT("The two resource sizes use different counting modes. The result of adding them together may be incorrect."));
+
+	for (const TPair<FName, SIZE_T>& Pair : InRHS.DedicatedSystemMemoryBytesMap)
+	{
+		DedicatedSystemMemoryBytesMap.FindOrAdd(Pair.Key) += Pair.Value;
+	}
+
+	for (const TPair<FName, SIZE_T>& Pair : InRHS.DedicatedVideoMemoryBytesMap)
+	{
+		DedicatedVideoMemoryBytesMap.FindOrAdd(Pair.Key) += Pair.Value;
+	}
+
+	for (const TPair<FName, SIZE_T>& Pair : InRHS.UnknownMemoryBytesMap)
+	{
+		UnknownMemoryBytesMap.FindOrAdd(Pair.Key) += Pair.Value;
+	}
+
+	return *this;
+}
+
+FResourceSizeEx operator+(FResourceSizeEx InLHS, const FResourceSizeEx& InRHS)
+{
+	InLHS += InRHS;
+	return InLHS;
+}
diff --git a/Engine/Source/Runtime/Core/Public/ProfilingDebugging/ResourceSize.h b/Engine/Source/Runtime/Core/Public/ProfilingDebugging/ResourceSize.h
--- a/Engine/Source/Runtime/Core/Public/ProfilingDebugging/ResourceSize.h
+++ b/Engine/Source/Runtime/Core/Public/ProfilingDebugging/ResourceSize.h
@@ -2,8 +2,7 @@
 
 #pragma once
 
-#include "CoreTypes.h"
-#include "Misc/AssertionMacros.h"
+#include "CoreMinimal.h"
 
 /** Indicate what types of resources should be included for calculating used memory */
 namespace EResourceSizeMode
@@ -27,190 +26,87 @@ public:
 	/**
 	 * Default constructor. 
 	 */
-	explicit FResourceSizeEx()
-		: ResourceSizeMode(EResourceSizeMode::Exclusive)
-		, DedicatedSystemMemoryBytes(0)
-		, SharedSystemMemoryBytes(0)
-		, DedicatedVideoMemoryBytes(0)
-		, SharedVideoMemoryBytes(0)
-		, UnknownMemoryBytes(0)
-	{
-	}
+	CORE_API explicit FResourceSizeEx();
 
 	/**
 	 * Construct using a given mode. 
 	 */
-	explicit FResourceSizeEx(const EResourceSizeMode::Type InResourceSizeMode)
-		: ResourceSizeMode(InResourceSizeMode)
-		, DedicatedSystemMemoryBytes(0)
-		, SharedSystemMemoryBytes(0)
-		, DedicatedVideoMemoryBytes(0)
-		, SharedVideoMemoryBytes(0)
-		, UnknownMemoryBytes(0)
-	{
-	}
+	CORE_API explicit FResourceSizeEx(const EResourceSizeMode::Type InResourceSizeMode);
 
 	/**
 	 * Construct from known sizes. 
 	 */
-	FResourceSizeEx(const EResourceSizeMode::Type InResourceSizeMode, const SIZE_T InDedicatedSystemMemoryBytes, const SIZE_T InSharedSystemMemoryBytes, const SIZE_T InDedicatedVideoMemoryBytes, const SIZE_T InSharedVideoMemoryBytes)
-		: ResourceSizeMode(InResourceSizeMode)
-		, DedicatedSystemMemoryBytes(InDedicatedSystemMemoryBytes)
-		, SharedSystemMemoryBytes(InSharedSystemMemoryBytes)
-		, DedicatedVideoMemoryBytes(InDedicatedVideoMemoryBytes)
-		, SharedVideoMemoryBytes(InSharedVideoMemoryBytes)
-		, UnknownMemoryBytes(0)
-	{
-	}
+	CORE_API FResourceSizeEx(const EResourceSizeMode::Type InResourceSizeMode, const SIZE_T InDedicatedSystemMemoryBytes, const SIZE_T InDedicatedVideoMemoryBytes);
 
 	/**
 	 * Construct from legacy unknown size.
 	 * Deliberately explicit to avoid accidental use.
 	 */
-	FResourceSizeEx(const EResourceSizeMode::Type InResourceSizeMode, const SIZE_T InUnknownMemoryBytes)
-		: ResourceSizeMode(InResourceSizeMode)
-		, DedicatedSystemMemoryBytes(0)
-		, SharedSystemMemoryBytes(0)
-		, DedicatedVideoMemoryBytes(0)
-		, SharedVideoMemoryBytes(0)
-		, UnknownMemoryBytes(InUnknownMemoryBytes)
-	{
-	}
+	CORE_API FResourceSizeEx(const EResourceSizeMode::Type InResourceSizeMode, const SIZE_T InUnknownMemoryBytes);
+
+	CORE_API void LogSummary(FOutputDevice& Ar) const;
 
 	/**
 	 * Get the type of resource size held in this struct.
 	 */
-	EResourceSizeMode::Type GetResourceSizeMode() const
-	{
-		return ResourceSizeMode;
-	}
+	CORE_API EResourceSizeMode::Type GetResourceSizeMode() const;
+
+	CORE_API FResourceSizeEx& AddDedicatedSystemMemoryBytes(const FName& Tag, const SIZE_T InMemoryBytes);
 
 	/**
 	 * Add the given number of bytes to the dedicated system memory count.
 	 * @see DedicatedSystemMemoryBytes for a description of that memory type.
 	 */
-	FResourceSizeEx& AddDedicatedSystemMemoryBytes(const SIZE_T InMemoryBytes)
-	{
-		DedicatedSystemMemoryBytes += InMemoryBytes;
-		return *this;
-	}
+	CORE_API FResourceSizeEx& AddDedicatedSystemMemoryBytes(const SIZE_T InMemoryBytes);
 
 	/**
 	 * Get the number of bytes allocated from dedicated system memory.
 	 * @see DedicatedSystemMemoryBytes for a description of that memory type.
 	 */
-	SIZE_T GetDedicatedSystemMemoryBytes() const
-	{
-		return DedicatedSystemMemoryBytes;
-	}
+	CORE_API SIZE_T GetDedicatedSystemMemoryBytes() const;
 
-	/**
-	 * Add the given number of bytes to the shared system memory count.
-	 * @see SharedSystemMemoryBytes for a description of that memory type.
-	 */
-	FResourceSizeEx& AddSharedSystemMemoryBytes(const SIZE_T InMemoryBytes)
-	{
-		SharedSystemMemoryBytes += InMemoryBytes;
-		return *this;
-	}
-
-	/**
-	 * Get the number of bytes allocated from shared system memory.
-	 * @see SharedSystemMemoryBytes for a description of that memory type.
-	 */
-	SIZE_T GetSharedSystemMemoryBytes() const
-	{
-		return SharedSystemMemoryBytes;
-	}
+	CORE_API FResourceSizeEx& AddDedicatedVideoMemoryBytes(const FName& Tag, const SIZE_T InMemoryBytes);
 
 	/**
 	 * Add the given number of bytes to the dedicated video memory count.
 	 * @see DedicatedVideoMemoryBytes for a description of that memory type.
 	 */
-	FResourceSizeEx& AddDedicatedVideoMemoryBytes(const SIZE_T InMemoryBytes)
-	{
-		DedicatedVideoMemoryBytes += InMemoryBytes;
-		return *this;
-	}
+	CORE_API FResourceSizeEx& AddDedicatedVideoMemoryBytes(const SIZE_T InMemoryBytes);
 
 	/**
 	 * Get the number of bytes allocated from dedicated video memory.
 	 * @see DedicatedVideoMemoryBytes for a description of that memory type.
 	 */
-	SIZE_T GetDedicatedVideoMemoryBytes() const
-	{
-		return DedicatedVideoMemoryBytes;
-	}
+	CORE_API SIZE_T GetDedicatedVideoMemoryBytes() const;
 
-	/**
-	 * Add the given number of bytes to the shared video memory count.
-	 * @see SharedVideoMemoryBytes for a description of that memory type.
-	 */
-	FResourceSizeEx& AddSharedVideoMemoryBytes(const SIZE_T InMemoryBytes)
-	{
-		SharedVideoMemoryBytes += InMemoryBytes;
-		return *this;
-	}
-
-	/**
-	 * Get the number of bytes allocated from shared video memory.
-	 * @see SharedVideoMemoryBytes for a description of that memory type.
-	 */
-	SIZE_T GetSharedVideoMemoryBytes() const
-	{
-		return SharedVideoMemoryBytes;
-	}
+	CORE_API FResourceSizeEx& AddUnknownMemoryBytes(const FName& Tag, const SIZE_T InMemoryBytes);
 
 	/**
 	 * Add the given number of bytes to the unknown memory count.
 	 * @see UnknownMemoryBytes for a description of that memory type.
 	 */
-	FResourceSizeEx& AddUnknownMemoryBytes(const SIZE_T InMemoryBytes)
-	{
-		UnknownMemoryBytes += InMemoryBytes;
-		return *this;
-	}
+	CORE_API FResourceSizeEx& AddUnknownMemoryBytes(const SIZE_T InMemoryBytes);
 
 	/**
 	 * Get the number of bytes allocated from unknown memory.
 	 * @see UnknownMemoryBytes for a description of that memory type.
 	 */
-	SIZE_T GetUnknownMemoryBytes() const
-	{
-		return UnknownMemoryBytes;
-	}
+	CORE_API SIZE_T GetUnknownMemoryBytes() const;
 
 	/**
 	 * Get the total number of bytes allocated from any memory.
 	 */
-	SIZE_T GetTotalMemoryBytes() const
-	{
-		return DedicatedSystemMemoryBytes + SharedSystemMemoryBytes + DedicatedVideoMemoryBytes + SharedVideoMemoryBytes + UnknownMemoryBytes;
-	}
+	CORE_API SIZE_T GetTotalMemoryBytes() const;
 
 	/**
 	 * Add another FResourceSizeEx to this one.
 	 */
-	FResourceSizeEx& operator+=(const FResourceSizeEx& InRHS)
-	{
-		ensureAlwaysMsgf(ResourceSizeMode == InRHS.ResourceSizeMode, TEXT("The two resource sizes use different counting modes. The result of adding them together may be incorrect."));
-
-		DedicatedSystemMemoryBytes += InRHS.DedicatedSystemMemoryBytes;
-		SharedSystemMemoryBytes += InRHS.SharedSystemMemoryBytes;
-		DedicatedVideoMemoryBytes += InRHS.DedicatedVideoMemoryBytes;
-		SharedVideoMemoryBytes += InRHS.SharedVideoMemoryBytes;
-		UnknownMemoryBytes += InRHS.UnknownMemoryBytes;
-		return *this;
-	}
+	CORE_API FResourceSizeEx& operator+=(const FResourceSizeEx& InRHS);
 
 	/**
 	 * Add two FResourceSizeEx instances together and return a copy.
 	 */
-	friend FResourceSizeEx operator+(FResourceSizeEx InLHS, const FResourceSizeEx& InRHS)
-	{
-		InLHS += InRHS;
-		return InLHS;
-	}
+	friend FResourceSizeEx operator+(FResourceSizeEx InLHS, const FResourceSizeEx& InRHS);
 
 private:
 	/**
@@ -222,29 +118,17 @@ private:
 	 * The number of bytes of memory that this resource is using for CPU resources that have been allocated from dedicated system memory.
 	 * On platforms with unified memory, this typically refers to the things allocated in the preferred memory for CPU use.
 	 */
-	SIZE_T DedicatedSystemMemoryBytes;
-
-	/**
-	 * The number of bytes of memory that this resource is using for CPU resources that have been allocated from non-dedicated system memory (typically zero).
-	 * On platforms with unified memory, this typically refers to the things allocated in the memory outside the preferred memory for CPU use.
-	 */
-	SIZE_T SharedSystemMemoryBytes;
+	TMap<FName, SIZE_T> DedicatedSystemMemoryBytesMap;
 
 	/**
 	 * The number of bytes of memory that this resource is using for GPU resources that have been allocated from dedicated video memory.
 	 * On platforms with unified memory, this typically refers to the things allocated in the preferred memory for GPU use.
 	 */
-	SIZE_T DedicatedVideoMemoryBytes;
-
-	/**
-	 * The number of bytes of memory that this resource is using for GPU resources that have been allocated from non-dedicated video memory.
-	 * On platforms with unified memory, this typically refers to the things allocated in the memory outside the preferred memory for GPU use.
-	 */
-	SIZE_T SharedVideoMemoryBytes;
+	TMap<FName, SIZE_T> DedicatedVideoMemoryBytesMap;
 
 	/**
 	 * The number of bytes of memory that this resource is using from an unspecified section of memory.
 	 * This exists so that the legacy GetResourceSize(...) functions can still report back memory usage until they're updated to use FResourceSizeEx, and should not be used in new memory tracking code.
 	 */
-	SIZE_T UnknownMemoryBytes;
+	TMap<FName, SIZE_T> UnknownMemoryBytesMap;
 };
diff --git a/Engine/Source/Runtime/Engine/Private/SkeletalMeshLODRenderData.cpp b/Engine/Source/Runtime/Engine/Private/SkeletalMeshLODRenderData.cpp
--- a/Engine/Source/Runtime/Engine/Private/SkeletalMeshLODRenderData.cpp
+++ b/Engine/Source/Runtime/Engine/Private/SkeletalMeshLODRenderData.cpp
@@ -786,16 +786,16 @@ void FSkeletalMeshLODRenderData::GetResourceSizeEx(FResourceSizeEx& CumulativeRe
 		const FRawStaticIndexBuffer16or32Interface* IndexBuffer = MultiSizeIndexContainer.GetIndexBuffer();
 		if (IndexBuffer)
 		{
-			CumulativeResourceSize.AddUnknownMemoryBytes(IndexBuffer->GetResourceDataSize());
+			CumulativeResourceSize.AddUnknownMemoryBytes(TEXT("IndexBuffer"), IndexBuffer->GetResourceDataSize());
 		}
 	}
 
-	CumulativeResourceSize.AddUnknownMemoryBytes(StaticVertexBuffers.PositionVertexBuffer.GetNumVertices() * StaticVertexBuffers.PositionVertexBuffer.GetStride());
-	CumulativeResourceSize.AddUnknownMemoryBytes(StaticVertexBuffers.StaticMeshVertexBuffer.GetResourceSize());
-	CumulativeResourceSize.AddUnknownMemoryBytes(SkinWeightVertexBuffer.GetVertexDataSize());
-	CumulativeResourceSize.AddUnknownMemoryBytes(StaticVertexBuffers.ColorVertexBuffer.GetAllocatedSize());
-	CumulativeResourceSize.AddUnknownMemoryBytes(ClothVertexBuffer.GetVertexDataSize());
-	CumulativeResourceSize.AddUnknownMemoryBytes(SkinWeightProfilesData.GetResourcesSize());	
+	CumulativeResourceSize.AddUnknownMemoryBytes(TEXT("PositionVertexBuffer"), StaticVertexBuffers.PositionVertexBuffer.GetNumVertices() * StaticVertexBuffers.PositionVertexBuffer.GetStride());
+	CumulativeResourceSize.AddUnknownMemoryBytes(TEXT("TexcoordBuffer and TangentBuffer"), StaticVertexBuffers.StaticMeshVertexBuffer.GetResourceSize());
+	CumulativeResourceSize.AddUnknownMemoryBytes(TEXT("SkinWeightVertexBuffer"), SkinWeightVertexBuffer.GetVertexDataSize());
+	CumulativeResourceSize.AddUnknownMemoryBytes(TEXT("ColorVertexBuffer"), StaticVertexBuffers.ColorVertexBuffer.GetAllocatedSize());
+	CumulativeResourceSize.AddUnknownMemoryBytes(TEXT("ClothVertexBuffer"), ClothVertexBuffer.GetVertexDataSize());
+	CumulativeResourceSize.AddUnknownMemoryBytes(TEXT("SkinWeightProfilesData"), SkinWeightProfilesData.GetResourcesSize());	
 }
 
 SIZE_T FSkeletalMeshLODRenderData::GetCPUAccessMemoryOverhead() const
diff --git a/Engine/Source/Runtime/Engine/Private/StaticMesh.cpp b/Engine/Source/Runtime/Engine/Private/StaticMesh.cpp
--- a/Engine/Source/Runtime/Engine/Private/StaticMesh.cpp
+++ b/Engine/Source/Runtime/Engine/Private/StaticMesh.cpp
@@ -840,22 +840,20 @@ void FStaticMeshLODResources::Serialize(FArchive& Ar, UObject* Owner, int32 Inde
 
 void FStaticMeshLODResources::GetResourceSizeEx(FResourceSizeEx& CumulativeResourceSize) const
 {
-	const int32 VBSize = VertexBuffers.StaticMeshVertexBuffer.GetResourceSize() +
-		VertexBuffers.PositionVertexBuffer.GetStride() * VertexBuffers.PositionVertexBuffer.GetNumVertices() +
-		VertexBuffers.ColorVertexBuffer.GetStride() * VertexBuffers.ColorVertexBuffer.GetNumVertices();
+	CumulativeResourceSize.AddUnknownMemoryBytes(TEXT("TexcoordBuffer and TangentBuffer"), VertexBuffers.StaticMeshVertexBuffer.GetResourceSize());
+	CumulativeResourceSize.AddUnknownMemoryBytes(TEXT("PositionVertexBuffer"), VertexBuffers.PositionVertexBuffer.GetNumVertices());
+	CumulativeResourceSize.AddUnknownMemoryBytes(TEXT("ColorVertexBuffer"), VertexBuffers.ColorVertexBuffer.GetStride() * VertexBuffers.ColorVertexBuffer.GetNumVertices());
 
-	int32 NumIndicies = IndexBuffer.GetNumIndices();
+	const int32 IndexStride = (IndexBuffer.Is32Bit() ? 4 : 2);
+	CumulativeResourceSize.AddUnknownMemoryBytes(TEXT("IndexBuffer"), IndexBuffer.GetNumIndices() * IndexStride);
 
 	if (AdditionalIndexBuffers)
 	{
-		NumIndicies += AdditionalIndexBuffers->ReversedDepthOnlyIndexBuffer.GetNumIndices();
-		NumIndicies += AdditionalIndexBuffers->ReversedIndexBuffer.GetNumIndices();
-		NumIndicies += AdditionalIndexBuffers->WireframeIndexBuffer.GetNumIndices();
+		CumulativeResourceSize.AddUnknownMemoryBytes(TEXT("ReversedDepthOnlyIndexBuffer"), AdditionalIndexBuffers->ReversedDepthOnlyIndexBuffer.GetNumIndices() * IndexStride);
+		CumulativeResourceSize.AddUnknownMemoryBytes(TEXT("ReversedIndexBuffer"), AdditionalIndexBuffers->ReversedIndexBuffer.GetNumIndices() * IndexStride);
+		CumulativeResourceSize.AddUnknownMemoryBytes(TEXT("WireframeIndexBuffer"), AdditionalIndexBuffers->WireframeIndexBuffer.GetNumIndices() * IndexStride);
 	}
 
-	int32 IBSize = NumIndicies * (IndexBuffer.Is32Bit() ? 4 : 2);
-
-	CumulativeResourceSize.AddUnknownMemoryBytes(VBSize + IBSize);
 	CumulativeResourceSize.AddUnknownMemoryBytes(Sections.GetAllocatedSize());
 }
 
diff --git a/Engine/Source/Runtime/Engine/Private/UnrealEngine.cpp b/Engine/Source/Runtime/Engine/Private/UnrealEngine.cpp
--- a/Engine/Source/Runtime/Engine/Private/UnrealEngine.cpp
+++ b/Engine/Source/Runtime/Engine/Private/UnrealEngine.cpp
@@ -8147,15 +8147,13 @@ bool UEngine::HandleObjCommand( const TCHAR* Cmd, FOutputDevice& Ar )
 			}
 			else
 			{
-				Ar.Logf( TEXT("%32s %12s %12s %12s %12s %12s %12s %12s %12s %12i"), 
+				Ar.Logf( TEXT("%32s %12s %12s %12s %12s %12s %12s %12i"), 
 					*Class->GetName(), 
 					*FHierarchy::Size(ClassObjects.Max), 
 					*FHierarchy::Size(ClassObjects.Num), 
 					*FHierarchy::Size(ClassObjects.TrueResourceSize.GetTotalMemoryBytes()),
 					*FHierarchy::Size(ClassObjects.TrueResourceSize.GetDedicatedSystemMemoryBytes()),
-					*FHierarchy::Size(ClassObjects.TrueResourceSize.GetSharedSystemMemoryBytes()),
 					*FHierarchy::Size(ClassObjects.TrueResourceSize.GetDedicatedVideoMemoryBytes()),
-					*FHierarchy::Size(ClassObjects.TrueResourceSize.GetSharedVideoMemoryBytes()),
 					*FHierarchy::Size(ClassObjects.TrueResourceSize.GetUnknownMemoryBytes()),
 					ClassObjects.Count 
 				);
@@ -8171,30 +8169,26 @@ bool UEngine::HandleObjCommand( const TCHAR* Cmd, FOutputDevice& Ar )
 		if( Culled.Count > 0 )
 		{
 			Ar.Logf( TEXT("") );
-			Ar.Logf( TEXT("%32s %12s %12s %12s %12s %12s %12s %12s %12s %12i"), 
+			Ar.Logf( TEXT("%32s %12s %12s %12s %12s %12s %12s %12i"), 
 				TEXT("(Culled)"), 
 				*FHierarchy::Size(Culled.Max), 
 				*FHierarchy::Size(Culled.Num), 
 				*FHierarchy::Size(Culled.TrueResourceSize.GetTotalMemoryBytes()),
 				*FHierarchy::Size(Culled.TrueResourceSize.GetDedicatedSystemMemoryBytes()),
-				*FHierarchy::Size(Culled.TrueResourceSize.GetSharedSystemMemoryBytes()),
 				*FHierarchy::Size(Culled.TrueResourceSize.GetDedicatedVideoMemoryBytes()),
-				*FHierarchy::Size(Culled.TrueResourceSize.GetSharedVideoMemoryBytes()),
 				*FHierarchy::Size(Culled.TrueResourceSize.GetUnknownMemoryBytes()),
 				Culled.Count 
 			);
 		}
 
 		Ar.Logf( TEXT("") );
-		Ar.Logf( TEXT("%32s %12s %12s %12s %12s %12s %12s %12s %12s %12i"), 
+		Ar.Logf( TEXT("%32s %12s %12s %12s %12s %12s %12s %12i"), 
 			TEXT("Total"), 
 			*FHierarchy::Size(Total.Max), 
 			*FHierarchy::Size(Total.Num), 
 			*FHierarchy::Size(Total.TrueResourceSize.GetTotalMemoryBytes()),
 			*FHierarchy::Size(Total.TrueResourceSize.GetDedicatedSystemMemoryBytes()),
-			*FHierarchy::Size(Total.TrueResourceSize.GetSharedSystemMemoryBytes()),
 			*FHierarchy::Size(Total.TrueResourceSize.GetDedicatedVideoMemoryBytes()),
-			*FHierarchy::Size(Total.TrueResourceSize.GetSharedVideoMemoryBytes()),
 			*FHierarchy::Size(Total.TrueResourceSize.GetUnknownMemoryBytes()),
 			Total.Count 
 		);
@@ -8467,6 +8461,7 @@ bool UEngine::HandleObjCommand( const TCHAR* Cmd, FOutputDevice& Ar )
 		const bool bResourceSizeSort = FParse::Param(Cmd, TEXT("RESOURCESIZESORT"));
 		const bool bCSV = FParse::Param(Cmd, TEXT("CSV"));
 		const bool bShowFullClassName = FParse::Param(Cmd, TEXT("FULLCLASSNAME"));
+		const bool bVerboseObjectOutput = FParse::Param(Cmd, TEXT("VERBOSE"));
 
 		if( Objects.Num() )
 		{
@@ -8496,20 +8491,18 @@ bool UEngine::HandleObjCommand( const TCHAR* Cmd, FOutputDevice& Ar )
 
 			if (bCSV)
 			{
-				Ar.Logf(TEXT(", Object, NumKB, MaxKB, ResExcKB, ResExcDedSysKB, ResExcShrSysKB, ResExcDedVidKB, ResExcShrVidKB, ResExcUnkKB"));
+				Ar.Logf(TEXT(", Object, NumKB, MaxKB, ResExcKB, ResExcDedSysKB, ResExcDedVidKB, ResExcUnkKB"));
 			}
 			else
 			{
 				Ar.Logf(
-					TEXT("%140s %10s %10s %10s %15s %15s %15s %15s %15s"),
+					TEXT("%140s %10s %10s %10s %15s %15s %15s"),
 					TEXT("Object"),
 					TEXT("NumKB"),
 					TEXT("MaxKB"),
 					TEXT("ResExcKB"),
 					TEXT("ResExcDedSysKB"),
-					TEXT("ResExcShrSysKB"),
 					TEXT("ResExcDedVidKB"),
-					TEXT("ResExcShrVidKB"),
 					TEXT("ResExcUnkKB")
 				);
 			}
@@ -8519,15 +8512,13 @@ bool UEngine::HandleObjCommand( const TCHAR* Cmd, FOutputDevice& Ar )
 				if (bCSV)
 				{
 					Ar.Logf(
-						TEXT(", %s, %.2f, %.2f, %.2f, %.2f, %.2f, %.2f, %.2f, %.2f"),
+						TEXT(", %s, %.2f, %.2f, %.2f, %.2f, %.2f, %.2f"),
 						*ObjItem.Object->GetFullName(),
 						ObjItem.Num / 1024.0f,
 						ObjItem.Max / 1024.0f,
 						ObjItem.TrueResourceSize.GetTotalMemoryBytes() / 1024.0f,
 						ObjItem.TrueResourceSize.GetDedicatedSystemMemoryBytes() / 1024.0f,
-						ObjItem.TrueResourceSize.GetSharedSystemMemoryBytes() / 1024.0f,
 						ObjItem.TrueResourceSize.GetDedicatedVideoMemoryBytes() / 1024.0f,
-						ObjItem.TrueResourceSize.GetSharedVideoMemoryBytes() / 1024.0f,
 						ObjItem.TrueResourceSize.GetUnknownMemoryBytes() / 1024.0f
 					);
 
@@ -8535,17 +8526,20 @@ bool UEngine::HandleObjCommand( const TCHAR* Cmd, FOutputDevice& Ar )
 				else
 				{
 					Ar.Logf(
-						TEXT("%140s %10.2f %10.2f %10.2f %15.2f %15.2f %15.2f %15.2f %15.2f"),
+						TEXT("%140s %10.2f %10.2f %10.2f %15.2f %15.2f %15.2f"),
 						*ObjItem.Object->GetFullName(),
 						ObjItem.Num / 1024.0f,
 						ObjItem.Max / 1024.0f,
 						ObjItem.TrueResourceSize.GetTotalMemoryBytes() / 1024.0f,
 						ObjItem.TrueResourceSize.GetDedicatedSystemMemoryBytes() / 1024.0f,
-						ObjItem.TrueResourceSize.GetSharedSystemMemoryBytes() / 1024.0f,
 						ObjItem.TrueResourceSize.GetDedicatedVideoMemoryBytes() / 1024.0f,
-						ObjItem.TrueResourceSize.GetSharedVideoMemoryBytes() / 1024.0f,
 						ObjItem.TrueResourceSize.GetUnknownMemoryBytes() / 1024.0f
 					);
+
+					if (bVerboseObjectOutput)
+					{
+						ObjItem.TrueResourceSize.LogSummary(Ar);
+					}
 				}
 			}
 			Ar.Log(TEXT(""));
@@ -8569,21 +8563,19 @@ bool UEngine::HandleObjCommand( const TCHAR* Cmd, FOutputDevice& Ar )
 			
 			if (bCSV)
 			{
-				Ar.Logf(TEXT(", Class, Count, NumKB, MaxKB, ResExcKB, ResExcDedSysKB, ResExcShrSysKB, ResExcDedVidKB, ResExcShrVidKB, ResExcUnkKB"));
+				Ar.Logf(TEXT(", Class, Count, NumKB, MaxKB, ResExcKB, ResExcDedSysKB, ResExcDedVidKB, ResExcUnkKB"));
 			}
 			else
 			{
 				Ar.Logf(
-					TEXT(" %100s %8s %10s %10s %10s %15s %15s %15s %15s %15s"),
+					TEXT(" %100s %8s %10s %10s %10s %15s %15s %15s"),
 					TEXT("Class"),
 					TEXT("Count"),
 					TEXT("NumKB"),
 					TEXT("MaxKB"),
 					TEXT("ResExcKB"),
 					TEXT("ResExcDedSysKB"),
-					TEXT("ResExcShrSysKB"),
 					TEXT("ResExcDedVidKB"),
-					TEXT("ResExcShrVidKB"),
 					TEXT("ResExcUnkKB")
 				);
 			}
@@ -8592,31 +8584,27 @@ bool UEngine::HandleObjCommand( const TCHAR* Cmd, FOutputDevice& Ar )
 			{
 				if (bCSV)
 				{
-					Ar.Logf(TEXT(", %s, %i, %f, %f, %f, %f, %f, %f, %f, %f"),
+					Ar.Logf(TEXT(", %s, %i, %f, %f, %f, %f, %f, %f"),
 						bShowFullClassName ? *List[i].Class->GetFullName() : *List[i].Class->GetName(),
 						(int32)List[i].Count,
 						List[i].Num / 1024.0f,
 						List[i].Max / 1024.0f,
 						List[i].TrueResourceSize.GetTotalMemoryBytes() / 1024.0f,
 						List[i].TrueResourceSize.GetDedicatedSystemMemoryBytes() / 1024.0f,
-						List[i].TrueResourceSize.GetSharedSystemMemoryBytes() / 1024.0f,
 						List[i].TrueResourceSize.GetDedicatedVideoMemoryBytes() / 1024.0f,
-						List[i].TrueResourceSize.GetSharedVideoMemoryBytes() / 1024.0f,
 						List[i].TrueResourceSize.GetUnknownMemoryBytes() / 1024.0f
 					);
 				}
 				else
 				{
-					Ar.Logf(TEXT(" %100s %8i %10.2f %10.2f %10.2f %15.2f %15.2f %15.2f %15.2f %15.2f"),
+					Ar.Logf(TEXT(" %100s %8i %10.2f %10.2f %10.2f %15.2f %15.2f %15.2f"),
 						bShowFullClassName ? *List[i].Class->GetFullName() : *List[i].Class->GetName(),
 						(int32)List[i].Count,
 						List[i].Num / 1024.0f,
 						List[i].Max / 1024.0f,
 						List[i].TrueResourceSize.GetTotalMemoryBytes() / 1024.0f,
 						List[i].TrueResourceSize.GetDedicatedSystemMemoryBytes() / 1024.0f,
-						List[i].TrueResourceSize.GetSharedSystemMemoryBytes() / 1024.0f,
 						List[i].TrueResourceSize.GetDedicatedVideoMemoryBytes() / 1024.0f,
-						List[i].TrueResourceSize.GetSharedVideoMemoryBytes() / 1024.0f,
 						List[i].TrueResourceSize.GetUnknownMemoryBytes() / 1024.0f
 					);
 				}
@@ -8624,15 +8612,13 @@ bool UEngine::HandleObjCommand( const TCHAR* Cmd, FOutputDevice& Ar )
 			Ar.Log( TEXT("") );
 		}
 		Ar.Logf( 
-			TEXT("%i Objects (Total: %.3fM / Max: %.3fM / Res: %.3fM | ResDedSys: %.3fM / ResShrSys: %.3fM / ResDedVid: %.3fM / ResShrVid: %.3fM / ResUnknown: %.3fM)"),
+			TEXT("%i Objects (Total: %.3fM / Max: %.3fM / Res: %.3fM | ResDedSys: %.3fM / ResDedVid: %.3fM / ResUnknown: %.3fM)"),
 			Total.Count, 
 			(float)Total.Num/1024.0/1024.0, 
 			(float)Total.Max/1024.0/1024.0, 
 			(float)Total.TrueResourceSize.GetTotalMemoryBytes()/1024.0/1024.0,
 			(float)Total.TrueResourceSize.GetDedicatedSystemMemoryBytes()/1024.0/1024.0, 
-			(float)Total.TrueResourceSize.GetSharedSystemMemoryBytes()/1024.0/1024.0, 
 			(float)Total.TrueResourceSize.GetDedicatedVideoMemoryBytes()/1024.0/1024.0, 
-			(float)Total.TrueResourceSize.GetSharedVideoMemoryBytes()/1024.0/1024.0, 
 			(float)Total.TrueResourceSize.GetUnknownMemoryBytes()/1024.0/1024.0
 		);
 		return true;
