UE::String::Find functions: Fix crash when View to search is empty with Data==nullptr.
#rnx

diff --git a/Engine/Source/Runtime/Core/Private/String/Find.cpp b/Engine/Source/Runtime/Core/Private/String/Find.cpp
--- a/Engine/Source/Runtime/Core/Private/String/Find.cpp
+++ b/Engine/Source/Runtime/Core/Private/String/Find.cpp
@@ -29,7 +29,12 @@ static inline typename TStringView<CharType>::SizeType FindFirst(TStringView<Cha
 	}
 	const CharType* const SearchData = Search.GetData();
 	const CharType* const ViewBegin = View.GetData();
-	const CharType* const ViewEnd = ViewBegin + View.Len() - SearchLen;
+	const typename TStringView<CharType>::SizeType ViewLen = View.Len();
+	if (ViewLen < SearchLen)
+	{
+		return INDEX_NONE;
+	}
+	const CharType* const ViewEnd = ViewBegin + ViewLen - SearchLen;
 	if (SearchCase == ESearchCase::CaseSensitive)
 	{
 		for (const CharType* ViewIt = ViewBegin; ViewIt <= ViewEnd; ++ViewIt)
@@ -64,7 +69,12 @@ static inline typename TStringView<CharType>::SizeType FindLast(TStringView<Char
 	}
 	const CharType* const SearchData = Search.GetData();
 	const CharType* const ViewBegin = View.GetData();
-	const CharType* const ViewEnd = ViewBegin + View.Len() - SearchLen;
+	const typename TStringView<CharType>::SizeType ViewLen = View.Len();
+	if (ViewLen < SearchLen)
+	{
+		return INDEX_NONE;
+	}
+	const CharType* const ViewEnd = ViewBegin + ViewLen - SearchLen;
 	if (SearchCase == ESearchCase::CaseSensitive)
 	{
 		for (const CharType* ViewIt = ViewEnd; ViewIt >= ViewBegin; --ViewIt)
@@ -153,6 +163,10 @@ static inline typename TStringView<CharType>::SizeType FindLastOfAny(TStringView
 
 	const CharType* const ViewBegin = View.GetData();
 	const typename TStringView<CharType>::SizeType ViewLen = View.Len();
+	if (ViewLen == 0)
+	{
+		return INDEX_NONE;
+	}
 	for (typename TStringView<CharType>::SizeType ViewIndex = ViewLen - 1; ViewIndex >= 0; --ViewIndex)
 	{
 		const TStringView<CharType> RemainingView(ViewBegin + ViewIndex, ViewLen - ViewIndex);
@@ -199,6 +213,10 @@ static inline typename TStringView<CharType>::SizeType FindLastChar(TStringView<
 {
 	const CharType* const ViewBegin = View.GetData();
 	const CharType* const ViewEnd = ViewBegin + View.Len();
+	if (ViewEnd == ViewBegin)
+	{
+		return INDEX_NONE;
+	}
 	if (SearchCase == ESearchCase::CaseSensitive)
 	{
 		for (const CharType* ViewIt = ViewEnd - 1; ViewIt >= ViewBegin; --ViewIt)
@@ -283,6 +301,10 @@ static inline typename TStringView<CharType>::SizeType FindLastOfAnyChar(TString
 
 	const CharType* const ViewBegin = View.GetData();
 	const CharType* const ViewEnd = ViewBegin + View.Len();
+	if (ViewEnd == ViewBegin)
+	{
+		return INDEX_NONE;
+	}
 	if (SearchCase == ESearchCase::CaseSensitive)
 	{
 		for (const CharType* ViewIt = ViewEnd - 1; ViewIt >= ViewBegin; --ViewIt)
diff --git a/Engine/Source/Runtime/Core/Private/Tests/String/FindTest.cpp b/Engine/Source/Runtime/Core/Private/Tests/String/FindTest.cpp
--- a/Engine/Source/Runtime/Core/Private/Tests/String/FindTest.cpp
+++ b/Engine/Source/Runtime/Core/Private/Tests/String/FindTest.cpp
@@ -36,6 +36,10 @@ bool StringFindFirstTest::RunTest(const FString& Parameters)
 
 	TestEqual(TEXT("FindFirst"), UE::String::FindFirst("AbCABCAbCABC", "ABC"), 3);
 
+	TestEqual(TEXT("FindFirst"), UE::String::FindFirst(FStringView(nullptr, 0), TEXT("SearchTerm")), INDEX_NONE);
+	TestEqual(TEXT("FindFirst"), UE::String::FindFirst(FStringView(), TEXT("SearchTerm")), INDEX_NONE);
+	TestEqual(TEXT("FindFirst"), UE::String::FindFirst(FString(), TEXT("SearchTerm")), INDEX_NONE);
+
 	return true;
 }
 
@@ -67,6 +71,10 @@ bool StringFindLastTest::RunTest(const FString& Parameters)
 
 	TestEqual(TEXT("FindLast"), UE::String::FindLast("AbCABCAbCABC", "ABC"), 9);
 
+	TestEqual(TEXT("FindLast"), UE::String::FindLast(FStringView(nullptr, 0), TEXT("SearchTerm")), INDEX_NONE);
+	TestEqual(TEXT("FindLast"), UE::String::FindLast(FStringView(), TEXT("SearchTerm")), INDEX_NONE);
+	TestEqual(TEXT("FindLast"), UE::String::FindLast(FString(), TEXT("SearchTerm")), INDEX_NONE);
+
 	return true;
 }
 
@@ -100,6 +108,10 @@ bool StringFindFirstOfAnyTest::RunTest(const FString& Parameters)
 
 	TestEqual(TEXT("FindFirstOfAny"), UE::String::FindFirstOfAny("AbCABCAbCABC", {"CABc", "ABC"}), 3);
 
+	TestEqual(TEXT("FindFirstOfAny"), UE::String::FindFirstOfAny(FStringView(nullptr, 0), {TEXT("ABC"), TEXT("ABD")}), INDEX_NONE);
+	TestEqual(TEXT("FindFirstOfAny"), UE::String::FindFirstOfAny(FStringView(), {TEXT("ABC"), TEXT("ABD")}), INDEX_NONE);
+	TestEqual(TEXT("FindFirstOfAny"), UE::String::FindFirstOfAny(FString(), {TEXT("ABC"), TEXT("ABD")}), INDEX_NONE);
+
 	return true;
 }
 
@@ -133,6 +145,10 @@ bool StringFindLastOfAnyTest::RunTest(const FString& Parameters)
 
 	TestEqual(TEXT("FindLastOfAny"), UE::String::FindLastOfAny("AbCABCAbCABC", {"CABc", "ABC"}), 9);
 
+	TestEqual(TEXT("FindLastOfAny"), UE::String::FindLastOfAny(FStringView(nullptr, 0), { TEXT("ABC"), TEXT("ABD") }), INDEX_NONE);
+	TestEqual(TEXT("FindLastOfAny"), UE::String::FindLastOfAny(FStringView(), { TEXT("ABC"), TEXT("ABD") }), INDEX_NONE);
+	TestEqual(TEXT("FindLastOfAny"), UE::String::FindLastOfAny(FString(), { TEXT("ABC"), TEXT("ABD") }), INDEX_NONE);
+
 	return true;
 }
 
@@ -153,6 +169,10 @@ bool StringFindFirstCharTest::RunTest(const FString& Parameters)
 
 	TestEqual(TEXT("FindFirstChar"), UE::String::FindFirstChar("AbCABCAbCABC", 'B'), 4);
 
+	TestEqual(TEXT("FindFirstChar"), UE::String::FindFirstChar(FStringView(nullptr, 0), TEXT('A')), INDEX_NONE);
+	TestEqual(TEXT("FindFirstChar"), UE::String::FindFirstChar(FStringView(), TEXT('A')), INDEX_NONE);
+	TestEqual(TEXT("FindFirstChar"), UE::String::FindFirstChar(FString(), TEXT('A')), INDEX_NONE);
+
 	return true;
 }
 
@@ -173,6 +193,10 @@ bool StringFindLastCharTest::RunTest(const FString& Parameters)
 
 	TestEqual(TEXT("FindLastChar"), UE::String::FindLastChar("AbCABCAbCABC", 'B'), 10);
 
+	TestEqual(TEXT("FindLastChar"), UE::String::FindLastChar(FStringView(nullptr, 0), TEXT('A')), INDEX_NONE);
+	TestEqual(TEXT("FindLastChar"), UE::String::FindLastChar(FStringView(), TEXT('A')), INDEX_NONE);
+	TestEqual(TEXT("FindLastChar"), UE::String::FindLastChar(FString(), TEXT('A')), INDEX_NONE);
+
 	return true;
 }
 
@@ -193,6 +217,10 @@ bool StringFindFirstOfAnyCharTest::RunTest(const FString& Parameters)
 
 	TestEqual(TEXT("FindFirstOfAnyChar"), UE::String::FindFirstOfAnyChar("AbCABCAbcABC", {'c', 'B'}), 4);
 
+	TestEqual(TEXT("FindFirstOfAnyChar"), UE::String::FindFirstOfAnyChar(FStringView(nullptr, 0), { TEXT('A'), TEXT('B') }), INDEX_NONE);
+	TestEqual(TEXT("FindFirstOfAnyChar"), UE::String::FindFirstOfAnyChar(FStringView(), { TEXT('A'), TEXT('B') }), INDEX_NONE);
+	TestEqual(TEXT("FindFirstOfAnyChar"), UE::String::FindFirstOfAnyChar(FString(), { TEXT('A'), TEXT('B') }), INDEX_NONE);
+
 	return true;
 }
 
@@ -213,6 +241,10 @@ bool StringFindLastOfAnyCharTest::RunTest(const FString& Parameters)
 
 	TestEqual(TEXT("FindLastOfAnyChar"), UE::String::FindLastOfAnyChar("AbCABCAbcABC", {'c', 'B'}), 10);
 
+	TestEqual(TEXT("FindLastOfAnyChar"), UE::String::FindLastOfAnyChar(FStringView(nullptr, 0), { TEXT('A'), TEXT('B') }), INDEX_NONE);
+	TestEqual(TEXT("FindLastOfAnyChar"), UE::String::FindLastOfAnyChar(FStringView(), { TEXT('A'), TEXT('B') }), INDEX_NONE);
+	TestEqual(TEXT("FindLastOfAnyChar"), UE::String::FindLastOfAnyChar(FString(), { TEXT('A'), TEXT('B') }), INDEX_NONE);
+
 	return true;
 }
 
