lock-free lists: replaced `FNoncopyable` by `UE_NONCOPYABLE`, because it breaks empty base optimisation (https://fanael.github.io/breaking-empty-base-optimization.html)

diff --git a/Engine/Source/Runtime/Core/Public/Containers/LockFreeList.h b/Engine/Source/Runtime/Core/Public/Containers/LockFreeList.h
--- a/Engine/Source/Runtime/Core/Public/Containers/LockFreeList.h
+++ b/Engine/Source/Runtime/Core/Public/Containers/LockFreeList.h
@@ -266,8 +266,10 @@ struct FLockFreeLinkPolicy
 };
 
 template<int TPaddingForCacheContention, uint64 TABAInc = 1>
-class FLockFreePointerListLIFORoot : public FNoncopyable
+class FLockFreePointerListLIFORoot
 {
+	UE_NONCOPYABLE(FLockFreePointerListLIFORoot)
+
 	typedef FLockFreeLinkPolicy::TDoublePtr TDoublePtr;
 	typedef FLockFreeLinkPolicy::TLink TLink;
 	typedef FLockFreeLinkPolicy::TLinkPtr TLinkPtr;
@@ -420,12 +422,17 @@ private:
 };
 
 template<class T, int TPaddingForCacheContention, uint64 TABAInc = 1>
-class FLockFreePointerListLIFOBase : public FNoncopyable
+class FLockFreePointerListLIFOBase
 {
+	UE_NONCOPYABLE(FLockFreePointerListLIFOBase)
+
 	typedef FLockFreeLinkPolicy::TDoublePtr TDoublePtr;
 	typedef FLockFreeLinkPolicy::TLink TLink;
 	typedef FLockFreeLinkPolicy::TLinkPtr TLinkPtr;
+
 public:
+	FLockFreePointerListLIFOBase() = default;
+
 	void Reset()
 	{
 		RootList.Reset();
@@ -522,8 +529,10 @@ private:
 };
 
 template<class T, int TPaddingForCacheContention, uint64 TABAInc = 1>
-class FLockFreePointerFIFOBase : public FNoncopyable
+class FLockFreePointerFIFOBase
 {
+	UE_NONCOPYABLE(FLockFreePointerFIFOBase)
+
 	typedef FLockFreeLinkPolicy::TDoublePtr TDoublePtr;
 	typedef FLockFreeLinkPolicy::TLink TLink;
 	typedef FLockFreeLinkPolicy::TLinkPtr TLinkPtr;
@@ -662,8 +671,10 @@ private:
 
 
 template<class T, int TPaddingForCacheContention, int NumPriorities>
-class FStallingTaskQueue : public FNoncopyable
+class FStallingTaskQueue
 {
+	UE_NONCOPYABLE(FStallingTaskQueue)
+
 	typedef FLockFreeLinkPolicy::TDoublePtr TDoublePtr;
 	typedef FLockFreeLinkPolicy::TLink TLink;
 	typedef FLockFreeLinkPolicy::TLinkPtr TLinkPtr;
