LWC: ParseVector2D support for variants

diff --git a/Engine/Source/Runtime/Core/Private/Misc/DefaultValueHelper.cpp b/Engine/Source/Runtime/Core/Private/Misc/DefaultValueHelper.cpp
--- a/Engine/Source/Runtime/Core/Private/Misc/DefaultValueHelper.cpp
+++ b/Engine/Source/Runtime/Core/Private/Misc/DefaultValueHelper.cpp
@@ -430,7 +430,7 @@ bool FDefaultValueHelper::ParseVector(const FString& Source, FVector3f& OutVal)
 	FVector3d TempVec;
 	if (ParseVector(Source, TempVec))
 	{
-		OutVal = FVector3f((float)TempVec.X, (float)TempVec.Y, (float)TempVec.Z);
+		OutVal = FVector3f(TempVec);
 		return true;
 	}
 	return false;
@@ -474,15 +474,28 @@ bool FDefaultValueHelper::ParseVector(const FString& Source, FVector3d& OutVal)
 		return OutVal.InitFromString(Source);
 	}
 	
-	OutVal = FVector( 
-		(FVector::FReal)FCString::Atod(Start),
-		(FVector::FReal)FCString::Atod(FirstComma + 1),
-		(FVector::FReal)FCString::Atod(SecondComma + 1) );
+	OutVal = FVector3d( 
+		FCString::Atod(Start),
+		FCString::Atod(FirstComma + 1),
+		FCString::Atod(SecondComma + 1) );
 	return true;
 }
 
 
-bool FDefaultValueHelper::ParseVector2D(const FString& Source, FVector2D& OutVal)
+
+bool FDefaultValueHelper::ParseVector2D(const FString& Source, FVector2f& OutVal)
+{
+	// LWC_TODO: Perf pessimization, especially if LWC is disabled!
+	FVector2d TempVec;
+	if (ParseVector2D(Source, TempVec))
+	{
+		OutVal = FVector2f(TempVec);
+		return true;
+	}
+	return false;
+}
+
+bool FDefaultValueHelper::ParseVector2D(const FString& Source, FVector2d& OutVal)
 {
 	const bool bHasWhitespace = HasWhitespaces(Source);
 	const FString NoWhitespace = bHasWhitespace ? RemoveWhitespaces(Source) : FString();
@@ -507,9 +520,9 @@ bool FDefaultValueHelper::ParseVector2D(const FString& Source, FVector2D& OutVal
 		return OutVal.InitFromString(Source);
 	}
 
-	OutVal = FVector2D( 
-		FCString::Atof(Start),
-		FCString::Atof(FirstComma + 1) );
+	OutVal = FVector2d( 
+		FCString::Atod(Start),
+		FCString::Atod(FirstComma + 1) );
 	return true;
 }
 
@@ -520,7 +533,7 @@ bool FDefaultValueHelper::ParseVector4(const FString& Source, FVector4f& OutVal)
 	FVector4d TempVec;
 	if (ParseVector4(Source, TempVec))
 	{
-		OutVal = FVector4f((float)TempVec.X, (float)TempVec.Y, (float)TempVec.Z, (float)TempVec.W);
+		OutVal = FVector4f(TempVec);
 		return true;
 	}
 	return false;
@@ -552,7 +565,7 @@ bool FDefaultValueHelper::ParseRotator(const FString& Source, FRotator3f& OutVal
 	FRotator3d TempRotator;
 	if (ParseRotator(Source, TempRotator))
 	{
-		OutVal = FRotator3f((float)TempRotator.Pitch, (float)TempRotator.Yaw, (float)TempRotator.Roll);
+		OutVal = FRotator3f(TempRotator);
 		return true;
 	}
 	return false;
diff --git a/Engine/Source/Runtime/Core/Public/Misc/DefaultValueHelper.h b/Engine/Source/Runtime/Core/Public/Misc/DefaultValueHelper.h
--- a/Engine/Source/Runtime/Core/Public/Misc/DefaultValueHelper.h
+++ b/Engine/Source/Runtime/Core/Public/Misc/DefaultValueHelper.h
@@ -142,7 +142,8 @@ public:
 	 *
 	 * @return true if the conversion happened, false otherwise
 	 */
-	static bool ParseVector2D(const FString& Source, FVector2D& OutVal);
+	static bool ParseVector2D(const FString& Source, FVector2f& OutVal);
+	static bool ParseVector2D(const FString& Source, FVector2d& OutVal);
 
 	/**
 	* Converts a string into a FVector4. Accepted form: " %f, %f, %f, %f "
