[BuildFix] - Partially reverting my previous change (CL# 18259478) by putting the implementations back into the header.

diff --git a/Engine/Source/Runtime/Core/Private/ProfilingDebugging/ResourceSize.cpp b/Engine/Source/Runtime/Core/Private/ProfilingDebugging/ResourceSize.cpp
deleted file mode 100644
diff --git a/Engine/Source/Runtime/Core/Public/ProfilingDebugging/ResourceSize.h b/Engine/Source/Runtime/Core/Public/ProfilingDebugging/ResourceSize.h
--- a/Engine/Source/Runtime/Core/Public/ProfilingDebugging/ResourceSize.h
+++ b/Engine/Source/Runtime/Core/Public/ProfilingDebugging/ResourceSize.h
@@ -2,7 +2,8 @@
 
 #pragma once
 
-#include "CoreMinimal.h"
+#include "CoreTypes.h"
+#include "Misc/AssertionMacros.h"
 
 /** Indicate what types of resources should be included for calculating used memory */
 namespace EResourceSizeMode
@@ -23,90 +24,210 @@ namespace EResourceSizeMode
 struct FResourceSizeEx
 {
 public:
+
 	/**
 	 * Default constructor. 
 	 */
-	CORE_API explicit FResourceSizeEx();
+	explicit FResourceSizeEx()
+		: ResourceSizeMode(EResourceSizeMode::Exclusive)
+	{
+	}
 
 	/**
 	 * Construct using a given mode. 
 	 */
-	CORE_API explicit FResourceSizeEx(const EResourceSizeMode::Type InResourceSizeMode);
+	explicit FResourceSizeEx(const EResourceSizeMode::Type InResourceSizeMode)
+		: ResourceSizeMode(InResourceSizeMode)
+	{
+	}
 
 	/**
 	 * Construct from known sizes. 
 	 */
-	CORE_API FResourceSizeEx(const EResourceSizeMode::Type InResourceSizeMode, const SIZE_T InDedicatedSystemMemoryBytes, const SIZE_T InDedicatedVideoMemoryBytes);
+	FResourceSizeEx(const EResourceSizeMode::Type InResourceSizeMode, const SIZE_T InDedicatedSystemMemoryBytes, const SIZE_T InDedicatedVideoMemoryBytes)
+		: ResourceSizeMode(InResourceSizeMode)
+	{
+		DedicatedSystemMemoryBytesMap.Add(TEXT("Untracked Memory"), InDedicatedSystemMemoryBytes);
+		DedicatedVideoMemoryBytesMap.Add(TEXT("Untracked Memory"), InDedicatedVideoMemoryBytes);
+	}
 
 	/**
 	 * Construct from legacy unknown size.
 	 * Deliberately explicit to avoid accidental use.
 	 */
-	CORE_API FResourceSizeEx(const EResourceSizeMode::Type InResourceSizeMode, const SIZE_T InUnknownMemoryBytes);
+	FResourceSizeEx(const EResourceSizeMode::Type InResourceSizeMode, const SIZE_T InUnknownMemoryBytes)
+		: ResourceSizeMode(InResourceSizeMode)
+	{
+		UnknownMemoryBytesMap.Add(TEXT("Untracked Memory"), InUnknownMemoryBytes);
+	}
 
-	CORE_API void LogSummary(FOutputDevice& Ar) const;
+	void LogSummary(FOutputDevice& Ar) const
+	{
+		auto PrintPair = [&Ar](const TPair<FName, SIZE_T>& Pair)
+		{
+			Ar.Logf(
+				TEXT("%140s %15.2f"),
+				*Pair.Key.ToString(),
+				static_cast<double>(Pair.Value) / 1024.0
+			);
+		};
+
+		for (const TPair<FName, SIZE_T>& Pair : DedicatedSystemMemoryBytesMap)
+		{
+			PrintPair(Pair);
+		}
+		for (const TPair<FName, SIZE_T>& Pair : DedicatedVideoMemoryBytesMap)
+		{
+			PrintPair(Pair);
+		}
+		for (const TPair<FName, SIZE_T>& Pair : UnknownMemoryBytesMap)
+		{
+			PrintPair(Pair);
+		}
+	}
 
 	/**
 	 * Get the type of resource size held in this struct.
 	 */
-	CORE_API EResourceSizeMode::Type GetResourceSizeMode() const;
+	EResourceSizeMode::Type GetResourceSizeMode() const
+	{
+		return ResourceSizeMode;
+	}
 
-	CORE_API FResourceSizeEx& AddDedicatedSystemMemoryBytes(const FName& Tag, const SIZE_T InMemoryBytes);
+	FResourceSizeEx& AddDedicatedSystemMemoryBytes(const FName& Tag, const SIZE_T InMemoryBytes)
+	{
+		SIZE_T& CurrentSize = DedicatedSystemMemoryBytesMap.FindOrAdd(Tag);
+		CurrentSize += InMemoryBytes;
+		return *this;
+	}
 
 	/**
 	 * Add the given number of bytes to the dedicated system memory count.
 	 * @see DedicatedSystemMemoryBytes for a description of that memory type.
 	 */
-	CORE_API FResourceSizeEx& AddDedicatedSystemMemoryBytes(const SIZE_T InMemoryBytes);
+	FResourceSizeEx& AddDedicatedSystemMemoryBytes(const SIZE_T InMemoryBytes)
+	{
+		SIZE_T& CurrentSize = DedicatedSystemMemoryBytesMap.FindOrAdd(TEXT("Untracked Memory"));
+		CurrentSize += InMemoryBytes;
+		return *this;
+	}
 
 	/**
 	 * Get the number of bytes allocated from dedicated system memory.
 	 * @see DedicatedSystemMemoryBytes for a description of that memory type.
 	 */
-	CORE_API SIZE_T GetDedicatedSystemMemoryBytes() const;
-
-	CORE_API FResourceSizeEx& AddDedicatedVideoMemoryBytes(const FName& Tag, const SIZE_T InMemoryBytes);
+	SIZE_T GetDedicatedSystemMemoryBytes() const
+	{
+		SIZE_T Sum = 0;
+		TArray<SIZE_T> Values;
+		DedicatedSystemMemoryBytesMap.GenerateValueArray(Values);
+		for (const SIZE_T s : Values) Sum += s;
+		return Sum;
+	}
+
+	FResourceSizeEx& AddDedicatedVideoMemoryBytes(const FName& Tag, const SIZE_T InMemoryBytes)
+	{
+		SIZE_T& CurrentSize = DedicatedVideoMemoryBytesMap.FindOrAdd(Tag);
+		CurrentSize += InMemoryBytes;
+		return *this;
+	}
 
 	/**
 	 * Add the given number of bytes to the dedicated video memory count.
 	 * @see DedicatedVideoMemoryBytes for a description of that memory type.
 	 */
-	CORE_API FResourceSizeEx& AddDedicatedVideoMemoryBytes(const SIZE_T InMemoryBytes);
+	FResourceSizeEx& AddDedicatedVideoMemoryBytes(const SIZE_T InMemoryBytes)
+	{
+		SIZE_T& CurrentSize = DedicatedVideoMemoryBytesMap.FindOrAdd(TEXT("Untracked Memory"));
+		CurrentSize += InMemoryBytes;
+		return *this;
+	}
 
 	/**
 	 * Get the number of bytes allocated from dedicated video memory.
 	 * @see DedicatedVideoMemoryBytes for a description of that memory type.
 	 */
-	CORE_API SIZE_T GetDedicatedVideoMemoryBytes() const;
-
-	CORE_API FResourceSizeEx& AddUnknownMemoryBytes(const FName& Tag, const SIZE_T InMemoryBytes);
+	SIZE_T GetDedicatedVideoMemoryBytes() const
+	{
+		SIZE_T Sum = 0;
+		TArray<SIZE_T> Values;
+		DedicatedVideoMemoryBytesMap.GenerateValueArray(Values);
+		for (const SIZE_T s : Values) Sum += s;
+		return Sum;
+	}
+
+	FResourceSizeEx& AddUnknownMemoryBytes(const FName& Tag, const SIZE_T InMemoryBytes)
+	{
+		SIZE_T& CurrentSize = UnknownMemoryBytesMap.FindOrAdd(Tag);
+		CurrentSize += InMemoryBytes;
+		return *this;
+	}
 
 	/**
 	 * Add the given number of bytes to the unknown memory count.
 	 * @see UnknownMemoryBytes for a description of that memory type.
 	 */
-	CORE_API FResourceSizeEx& AddUnknownMemoryBytes(const SIZE_T InMemoryBytes);
+	FResourceSizeEx& AddUnknownMemoryBytes(const SIZE_T InMemoryBytes)
+	{
+		SIZE_T& CurrentSize = UnknownMemoryBytesMap.FindOrAdd(TEXT("Untracked Memory"));
+		CurrentSize += InMemoryBytes;
+		return *this;
+	}
 
 	/**
 	 * Get the number of bytes allocated from unknown memory.
 	 * @see UnknownMemoryBytes for a description of that memory type.
 	 */
-	CORE_API SIZE_T GetUnknownMemoryBytes() const;
+	SIZE_T GetUnknownMemoryBytes() const
+	{
+		SIZE_T Sum = 0;
+		TArray<SIZE_T> Values;
+		UnknownMemoryBytesMap.GenerateValueArray(Values);
+		for (const SIZE_T s : Values) Sum += s;
+		return Sum;
+	}
 
 	/**
 	 * Get the total number of bytes allocated from any memory.
 	 */
-	CORE_API SIZE_T GetTotalMemoryBytes() const;
+	SIZE_T GetTotalMemoryBytes() const
+	{
+		return GetDedicatedSystemMemoryBytes() + GetDedicatedVideoMemoryBytes() + GetUnknownMemoryBytes();
+	}
 
 	/**
 	 * Add another FResourceSizeEx to this one.
 	 */
-	CORE_API FResourceSizeEx& operator+=(const FResourceSizeEx& InRHS);
+	FResourceSizeEx& operator+=(const FResourceSizeEx& InRHS)
+	{
+		ensureAlwaysMsgf(ResourceSizeMode == InRHS.ResourceSizeMode, TEXT("The two resource sizes use different counting modes. The result of adding them together may be incorrect."));
+
+		for (const TPair<FName, SIZE_T>& Pair : InRHS.DedicatedSystemMemoryBytesMap)
+		{
+			DedicatedSystemMemoryBytesMap.FindOrAdd(Pair.Key) += Pair.Value;
+		}
+
+		for (const TPair<FName, SIZE_T>& Pair : InRHS.DedicatedVideoMemoryBytesMap)
+		{
+			DedicatedVideoMemoryBytesMap.FindOrAdd(Pair.Key) += Pair.Value;
+		}
+
+		for (const TPair<FName, SIZE_T>& Pair : InRHS.UnknownMemoryBytesMap)
+		{
+			UnknownMemoryBytesMap.FindOrAdd(Pair.Key) += Pair.Value;
+		}
+
+		return *this;
+	}
 
 	/**
 	 * Add two FResourceSizeEx instances together and return a copy.
 	 */
-	friend FResourceSizeEx operator+(FResourceSizeEx InLHS, const FResourceSizeEx& InRHS);
+	friend FResourceSizeEx operator+(FResourceSizeEx InLHS, const FResourceSizeEx& InRHS)
+	{
+		InLHS += InRHS;
+		return InLHS;
+	}
 
 private:
 	/**
