IoDispatcher: Trace counters cleanup pass, fixed double counting of bytes scattered and removed duplicate counters.
#rnx

diff --git a/Engine/Source/Runtime/Core/Internal/IO/IoDispatcherFileBackend.h b/Engine/Source/Runtime/Core/Internal/IO/IoDispatcherFileBackend.h
--- a/Engine/Source/Runtime/Core/Internal/IO/IoDispatcherFileBackend.h
+++ b/Engine/Source/Runtime/Core/Internal/IO/IoDispatcherFileBackend.h
@@ -26,7 +26,7 @@ struct FFileIoStoreCompressionContext
 class FFileIoStoreReader
 {
 public:
-	FFileIoStoreReader(IPlatformFileIoStore& InPlatformImpl);
+	FFileIoStoreReader(IPlatformFileIoStore& InPlatformImpl, FFileIoStoreStats& InStats);
 	~FFileIoStoreReader();
 	FIoStatus Initialize(const TCHAR* ContainerPath, int32 Order);
 	uint32 GetContainerInstanceId() const
@@ -51,9 +51,11 @@ public:
 
 private:
 	const FIoOffsetAndLength* FindChunkInternal(const FIoChunkId& ChunkId) const;
+	uint64 GetTocAllocatedSize() const;
 	
 	IPlatformFileIoStore& PlatformImpl;
-	
+	FFileIoStoreStats& Stats;
+
 	struct FPerfectHashMap
 	{
 		TArray<int32> TocChunkHashSeeds;
@@ -166,6 +168,7 @@ private:
 
 	uint64 ReadBufferSize = 0;
 	TSharedPtr<const FIoDispatcherBackendContext> BackendContext;
+	FFileIoStoreStats Stats;
 	FFileIoStoreBlockCache BlockCache;
 	FFileIoStoreBufferAllocator BufferAllocator;
 	FFileIoStoreRequestAllocator RequestAllocator;
diff --git a/Engine/Source/Runtime/Core/Internal/IO/IoDispatcherFileBackendTypes.h b/Engine/Source/Runtime/Core/Internal/IO/IoDispatcherFileBackendTypes.h
--- a/Engine/Source/Runtime/Core/Internal/IO/IoDispatcherFileBackendTypes.h
+++ b/Engine/Source/Runtime/Core/Internal/IO/IoDispatcherFileBackendTypes.h
@@ -17,7 +17,7 @@
 #define PLATFORM_IODISPATCHER_MODULE PREPROCESSOR_TO_STRING(PREPROCESSOR_JOIN(PLATFORM_HEADER_NAME, PlatformIoDispatcher))
 #endif
 
-#define IO_DISPATCHER_FILE_STATS (COUNTERSTRACE_ENABLED || CSV_PROFILER)
+#define UE_FILEIOSTORE_STATS_ENABLED (COUNTERSTRACE_ENABLED || CSV_PROFILER)
 
 struct FFileIoStoreCompressionContext;
 
@@ -520,15 +520,23 @@ private:
 	}
 };
 
+class FFileIoStoreStats;
+
 class FFileIoStoreBufferAllocator
 {
 public:
+	FFileIoStoreBufferAllocator(FFileIoStoreStats& InStats)
+		: Stats(InStats)
+	{
+	}
+
 	void Initialize(uint64 MemorySize, uint64 BufferSize, uint32 BufferAlignment);
 	FFileIoStoreBuffer* AllocBuffer();
 	void FreeBuffer(FFileIoStoreBuffer* Buffer);
 	uint64 GetBufferSize() const { return BufferSize; }
 
 private:
+	FFileIoStoreStats& Stats;
 	uint64 BufferSize = 0;
 	uint8* BufferMemory = nullptr;
 	FCriticalSection BuffersCritical;
@@ -538,7 +546,7 @@ private:
 class FFileIoStoreBlockCache
 {
 public:
-	FFileIoStoreBlockCache();
+	FFileIoStoreBlockCache(FFileIoStoreStats& Stats);
 	~FFileIoStoreBlockCache();
 
 	void Initialize(uint64 CacheMemorySize, uint64 ReadBufferSize);
@@ -554,6 +562,7 @@ private:
 		uint8* Buffer = nullptr;
 	};
 
+	FFileIoStoreStats& Stats;
 	uint8* CacheMemory = nullptr;
 	TMap<uint64, FCachedBlock*> CachedBlocks;
 	FCachedBlock CacheLruHead;
@@ -877,76 +886,127 @@ private:
 
 
 // Wrapper for sending stats to both insights and csv profiling from any platform-specific dispatcher.
-struct FFileIoStats
+#if UE_FILEIOSTORE_STATS_ENABLED
+class FFileIoStoreStats
 {
+public:
+	FFileIoStoreStats();
+	~FFileIoStoreStats();
+
+	// Called by the backend when underlying file system reads start
+	void OnFilesystemReadStarted(const FFileIoStoreReadRequest* Request);
+	void OnFilesystemReadsStarted(const FFileIoStoreReadRequestList& Requests);
+	// Called by the backend when underlying filesystem reads complete, possibly already decompressed on some systems
+	void OnFilesystemReadCompleted(const FFileIoStoreReadRequest* Request);
+	void OnFilesystemReadsCompleted(const FFileIoStoreReadRequestList& CompletedRequests);
+
+private:
+	friend class FFileIoStore;
+	friend class FFileIoStoreBlockCache;
+	friend class FFileIoStoreReader;
+	friend class FFileIoStoreBufferAllocator;
+
+#if COUNTERSTRACE_ENABLED
+	FCountersTrace::FCounterInt QueuedReadRequestsSizeCounter;
+	FCountersTrace::FCounterInt CompletedReadRequestsSizeCounter;
+	FCountersTrace::FCounterInt QueuedCompressedSizeCounter;
+	FCountersTrace::FCounterInt QueuedUncompressedSizeCounter;
+	FCountersTrace::FCounterInt CompletedCompressedSizeCounter;
+	FCountersTrace::FCounterInt CompletedUncompressedSizeCounter;
+	FCountersTrace::FCounterInt FileSystemSeeksTotalDistanceCounter;
+	FCountersTrace::FCounterInt FileSystemSeeksForwardCountCounter;
+	FCountersTrace::FCounterInt FileSystemSeeksBackwardCountCounter;
+	FCountersTrace::FCounterInt FileSystemSeeksChangeHandleCountCounter;
+	FCountersTrace::FCounterInt FileSystemCompletedRequestsSizeCounter;
+	FCountersTrace::FCounterInt BlockCacheStoredSizeCounter;
+	FCountersTrace::FCounterInt BlockCacheHitSizeCounter;
+	FCountersTrace::FCounterInt BlockCacheMissedSizeCounter;
+	FCountersTrace::FCounterInt ScatteredSizeCounter;
+	FCountersTrace::FCounterInt TocMemoryCounter;
+	FCountersTrace::FCounterInt AvailableBuffersCounter;
+#endif
+
 #if CSV_PROFILER
-	static uint64 QueuedFilesystemReadBytes;
-	static uint64 QueuedFilesystemReads;
+	uint64 QueuedFilesystemReadBytes = 0;
+	uint64 QueuedFilesystemReads = 0;
 
-	static uint64 QueuedUncompressBytesIn;
-	static uint64 QueuedUncompressBytesOut;
-	static uint64 QueuedUncompressBlocks;
+	uint64 QueuedUncompressBytesIn = 0;
+	uint64 QueuedUncompressBytesOut = 0;
+	uint64 QueuedUncompressBlocks = 0;
 #endif
 
-#if IO_DISPATCHER_FILE_STATS
 	// Used for seek tracking 
-	static uint64 LastHandle;
-	static uint64 LastOffset;
-
-	static uint32 FileIoStoreThreadId;
-	static uint32 IoDispatcherThreadId;
+	uint64 LastHandle = uint64(-1);
+	uint64 LastOffset = uint64(-1);
 
-	static FTSTicker::FDelegateHandle TickerHandle;
-#endif
+	FTSTicker::FDelegateHandle TickerHandle;
 
-	static float BytesToApproxMB(uint64 Bytes) {
+	float BytesToApproxMB(uint64 Bytes)
+	{
 		return float(double(Bytes) / 1024.0 / 1024.0);
 	}
-	static float BytesToApproxKB(uint64 Bytes) {
+
+	float BytesToApproxKB(uint64 Bytes)
+	{
 		return float(double(Bytes) / 1024.0);
 	}
 
-	static void Init();
-	static void Shutdown();
-	static bool CsvTick(float DeltaTime);
-
-	static bool IsInIoDispatcherThread();
-	static bool IsInFileIoStoreThread();
-	static void SetDispatcherThreadId(uint32 InIoDispatcherThreadId);
-	static void SetFileIoStoreThreadId(uint32 InFileIoStoreThreadId);
-
-	// Adds an underlying filesystem read to the queue to track num reads/bytes in the queue
-	// There is currently no separate stat for in-flight reads, those are still counted in the queued stat
-	static void OnFilesystemReadsQueued(const FFileIoStoreReadRequestList& Requests);
-	static void OnFilesystemReadsQueued(uint64 NumBytes, int32 NumReads);
-	// Called when the backend commits to the next read to check for estimated seek tracking
-	static void OnFilesystemReadStarted(uint64 Handle, uint64 Offset, uint64 NumBytes);
-	// Underlying filesystem read complete, BytesRead is bytes returned from filesystem, possibly already decompressed on some systems
-	static void OnFilesystemReadsComplete(const FFileIoStoreReadRequestList& CompletedRequests);
-	static void OnFilesystemReadsComplete(int64 BytesRead, int32 NumReads=1);
-	// Backend read complete, BytesRead is bytes returned from either block cache or filesystem
-	static void OnReadComplete(int64 BytesRead);
-	// NumBytes of data passed the decompression stage, even if the compression method was none
-	static void OnDecompressQueued(int64 CompressedBytes, int64 UncompressedBytes);
-	static void OnDecompressComplete(int64 CompressedBytes, int64 UncompressedBytes);
+	bool CsvTick(float DeltaTime);
+
+	void OnReadRequestsQueued(const FFileIoStoreReadRequestList& Requests);
+	void OnReadRequestsCompleted(const FFileIoStoreReadRequestList& Requests);
+	void OnDecompressQueued(const FFileIoStoreCompressedBlock* CompressedBlock);
+	void OnDecompressComplete(const FFileIoStoreCompressedBlock* CompressedBlock);
 
 	// Bytes were copied into the buffer provided to users of the io system
-	static void OnBytesScattered(int64 BytesScattered);
+	void OnBytesScattered(int64 BytesScattered);
 	// Record stats for block cache hit/miss rate & data throughput
-	static void OnBlockCacheStore(uint64 NumBytes);
-	static void OnBlockCacheHit(uint64 NumBytes);
-	static void OnBlockCacheMiss(uint64 NumBytes);
+	void OnBlockCacheStore(uint64 NumBytes);
+	void OnBlockCacheHit(uint64 NumBytes);
+	void OnBlockCacheMiss(uint64 NumBytes);
+
+	// A read was started without seeking
+	void OnSequentialRead();
+	// A read was started and was either a forward or reverse seek
+	void OnSeek(uint64 LastOffset, uint64 NewOffset);
+	// A read was started from a different file handle from the last one we used
+	void OnHandleChangeSeek();
+
+	void OnTocMounted(uint64 AllocatedSize);
+	void OnTocUnmounted(uint64 AllocatedSize);
 
-	static void OnCloseHandle(uint64 Handle);
+	void OnBufferReleased();
+	void OnBufferAllocated();
+};
+#else
+class FFileIoStoreStats
+{
+public:
+	void OnFilesystemReadStarted(const FFileIoStoreReadRequest* Request) {}
+	void OnFilesystemReadsStarted(const FFileIoStoreReadRequestList& Requests) {}
+	void OnFilesystemReadCompleted(const FFileIoStoreReadRequest* Request) {}
+	void OnFilesystemReadsCompleted(const FFileIoStoreReadRequestList& CompletedRequests) {}
 
 private:
-	// A read was started which the backend considers sequential
-	static void OnSequentialRead();
-	// A read was started which the backend considers a forward or reverse seek
-	static void OnSeek(uint64 LastOffset, uint64 NewOffset);
-	// A read was started from a different file handle from the last one we used
-	static void OnHandleChangeSeek();
+	friend class FFileIoStore;
+	friend class FFileIoStoreBlockCache;
+	friend class FFileIoStoreReader;
+	friend class FFileIoStoreBufferAllocator;
+
+	void OnReadRequestsQueued(const FFileIoStoreReadRequestList& Requests) {}
+	void OnReadRequestsCompleted(const FFileIoStoreReadRequestList& Requests) {}
+	void OnDecompressQueued(const FFileIoStoreCompressedBlock* CompressedBlock) {}
+	void OnDecompressComplete(const FFileIoStoreCompressedBlock* CompressedBlock) {}
+	void OnBytesScattered(int64 BytesScattered) {}
+	void OnBlockCacheStore(uint64 NumBytes) {}
+	void OnBlockCacheHit(uint64 NumBytes) {}
+	void OnBlockCacheMiss(uint64 NumBytes) {}
+	void OnTocMounted(uint64 AllocatedSize) {}
+	void OnTocUnmounted(uint64 AllocatedSize) {}
+	void OnBufferReleased() {}
+	void OnBufferAllocated() {}
 };
+#endif
 
 struct FInitializePlatformFileIoStoreParams
 {
@@ -954,6 +1014,7 @@ struct FInitializePlatformFileIoStoreParams
 	FFileIoStoreRequestAllocator* RequestAllocator = nullptr;
 	FFileIoStoreBufferAllocator* BufferAllocator = nullptr;
 	FFileIoStoreBlockCache* BlockCache = nullptr;
+	FFileIoStoreStats* Stats = nullptr;
 };
 
 class IPlatformFileIoStore
diff --git a/Engine/Source/Runtime/Core/Private/GenericPlatform/GenericPlatformIoDispatcher.cpp b/Engine/Source/Runtime/Core/Private/GenericPlatform/GenericPlatformIoDispatcher.cpp
--- a/Engine/Source/Runtime/Core/Private/GenericPlatform/GenericPlatformIoDispatcher.cpp
+++ b/Engine/Source/Runtime/Core/Private/GenericPlatform/GenericPlatformIoDispatcher.cpp
@@ -69,7 +69,6 @@ bool FGenericFileIoStoreImpl::OpenContainer(const TCHAR* ContainerFilePath, uint
 void FGenericFileIoStoreImpl::CloseContainer(uint64 ContainerFileHandle)
 {
 	check(ContainerFileHandle);
-	FFileIoStats::OnCloseHandle(ContainerFileHandle);
 	IFileHandle* FileHandle = reinterpret_cast<IFileHandle*>(ContainerFileHandle);
 	delete FileHandle;
 }
@@ -106,12 +105,11 @@ bool FGenericFileIoStoreImpl::StartRequests(FFileIoStoreRequestQueue& RequestQue
 	{
 		IFileHandle* FileHandle = reinterpret_cast<IFileHandle*>(static_cast<UPTRINT>(NextRequest->FileHandle));
 		 
-		FFileIoStats::OnFilesystemReadStarted(NextRequest->FileHandle, NextRequest->Offset, NextRequest->Size);
+		Stats->OnFilesystemReadStarted(NextRequest);
 		{
 			TRACE_CPUPROFILER_EVENT_SCOPE(ReadBlockFromFile);
 			NextRequest->bFailed = true;
 			int32 RetryCount = 0;
-			TRACE_COUNTER_INCREMENT(IoDispatcherFileBackendFileSystemRequests);
 			while (RetryCount++ < 10)
 			{
 				if (!FileHandle->Seek(NextRequest->Offset))
@@ -125,12 +123,11 @@ bool FGenericFileIoStoreImpl::StartRequests(FFileIoStoreRequestQueue& RequestQue
 					continue;
 				}
 				NextRequest->bFailed = false;
-				TRACE_COUNTER_ADD(IoDispatcherFileBackendFileSystemTotalBytesRead, NextRequest->Size);
+				Stats->OnFilesystemReadCompleted(NextRequest);
 				BlockCache->Store(NextRequest);
 				break;
 			}
 		}
-		FFileIoStats::OnFilesystemReadsComplete(NextRequest->Size);
 	}
 	{
 		FScopeLock _(&CompletedRequestsCritical);
diff --git a/Engine/Source/Runtime/Core/Private/GenericPlatform/GenericPlatformIoDispatcher.h b/Engine/Source/Runtime/Core/Private/GenericPlatform/GenericPlatformIoDispatcher.h
--- a/Engine/Source/Runtime/Core/Private/GenericPlatform/GenericPlatformIoDispatcher.h
+++ b/Engine/Source/Runtime/Core/Private/GenericPlatform/GenericPlatformIoDispatcher.h
@@ -33,6 +33,7 @@ public:
 		WakeUpDispatcherThreadDelegate = Params.WakeUpDispatcherThreadDelegate;
 		BufferAllocator = Params.BufferAllocator;
 		BlockCache = Params.BlockCache;
+		Stats = Params.Stats;
 	}
 	bool OpenContainer(const TCHAR* ContainerFilePath, uint64& ContainerFileHandle, uint64& ContainerFileSize) override;
 	void CloseContainer(uint64 ContainerFileHandle) override;
@@ -56,6 +57,7 @@ private:
 	FGenericFileIoStoreEventQueue EventQueue;
 	FFileIoStoreBufferAllocator* BufferAllocator = nullptr;
 	FFileIoStoreBlockCache* BlockCache = nullptr;
+	FFileIoStoreStats* Stats = nullptr;
 
 	FCriticalSection CompletedRequestsCritical;
 	FFileIoStoreReadRequestList CompletedRequests;
diff --git a/Engine/Source/Runtime/Core/Private/IO/IoDispatcher.cpp b/Engine/Source/Runtime/Core/Private/IO/IoDispatcher.cpp
--- a/Engine/Source/Runtime/Core/Private/IO/IoDispatcher.cpp
+++ b/Engine/Source/Runtime/Core/Private/IO/IoDispatcher.cpp
@@ -35,7 +35,128 @@ TUniquePtr<FIoDispatcher> GIoDispatcher;
 CSV_DEFINE_CATEGORY(IoDispatcher, true);
 CSV_DEFINE_STAT(IoDispatcher, PendingIoRequests);
 
-TRACE_DECLARE_INT_COUNTER(PendingIoRequests, TEXT("IoDispatcher/PendingIoRequests"));
+#if UE_IODISPATCHER_STATS_ENABLED
+class FIoRequestStats
+{
+public:
+	FIoRequestStats()
+#if COUNTERSTRACE_ENABLED
+		: PendingIoRequestsCounter(TEXT("IoDispatcher/PendingIoRequests"), TraceCounterDisplayHint_None)
+#endif
+	{
+		Categories.Reserve(4);
+		int32 PackageDataIndex = Categories.Emplace(TEXT("PackageData"));
+		int32 BulkDataIndex = Categories.Emplace(TEXT("BulkData"));
+		int32 ShadersIndex = Categories.Emplace(TEXT("Shaders"));
+		int32 MiscIndex = Categories.Emplace(TEXT("Misc"));
+
+		for (int32 Index = 0; Index < UE_ARRAY_COUNT(ChunkTypeToCategoryMap); ++Index)
+		{
+			ChunkTypeToCategoryMap[Index] = &Categories[MiscIndex];
+		}
+		ChunkTypeToCategoryMap[static_cast<int32>(EIoChunkType::ExportBundleData)] = &Categories[PackageDataIndex];
+		ChunkTypeToCategoryMap[static_cast<int32>(EIoChunkType::BulkData)] = &Categories[BulkDataIndex];
+		ChunkTypeToCategoryMap[static_cast<int32>(EIoChunkType::OptionalBulkData)] = &Categories[BulkDataIndex];
+		ChunkTypeToCategoryMap[static_cast<int32>(EIoChunkType::MemoryMappedBulkData)] = &Categories[BulkDataIndex];
+		ChunkTypeToCategoryMap[static_cast<int32>(EIoChunkType::ShaderCode)] = &Categories[ShadersIndex];
+
+#if CSV_PROFILER
+		TickerHandle = FTSTicker::GetCoreTicker().AddTicker(FTickerDelegate::CreateRaw(this, &FIoRequestStats::TickCsv));
+#endif
+	}
+
+	~FIoRequestStats()
+	{
+#if CSV_PROFILER
+		FTSTicker::GetCoreTicker().RemoveTicker(TickerHandle);
+#endif
+	}
+
+	void OnBatchIssued(FIoBatch& Batch)
+	{
+		double StartTime = FPlatformTime::Seconds();
+		FIoRequestImpl* Request = Batch.HeadRequest;
+		while (Request)
+		{
+			Request->StartTime = StartTime;
+			Request = Request->NextRequest;
+		}
+	}
+
+	void OnRequestStarted(FIoRequestImpl& Request)
+	{
+		++PendingIoRequests;
+#if COUNTERSTRACE_ENABLED
+		PendingIoRequestsCounter.Set(PendingIoRequests);
+#endif
+	}
+
+	void OnRequestCompleted(FIoRequestImpl& Request)
+	{
+		--PendingIoRequests;
+#if COUNTERSTRACE_ENABLED
+		PendingIoRequestsCounter.Set(PendingIoRequests);
+#endif
+		if (!Request.HasBuffer())
+		{
+			return;
+		}
+		FRequestCategory* Category = ChunkTypeToCategoryMap[static_cast<int32>(Request.ChunkId.GetChunkType())];
+		++Category->TotaRequestsCount;
+		Category->TotalRequestsTime += FPlatformTime::Seconds() - Request.StartTime;
+#if COUNTERSTRACE_ENABLED
+		Category->TotalLoadedCounter.Add(Request.GetBuffer().DataSize());
+		Category->AverageDurationCounter.Set(Category->TotalRequestsTime / double(Category->TotaRequestsCount));
+#endif
+	}
+
+private:
+	struct FRequestCategory
+	{
+		FRequestCategory(const TCHAR* Name)
+#if COUNTERSTRACE_ENABLED
+			: TotalLoadedCounter(*FString::Printf(TEXT("IoDispatcher/TotalLoaded (%s)"), Name), TraceCounterDisplayHint_Memory)
+			, AverageDurationCounter(*FString::Printf(TEXT("IoDispatcher/AverageDuration (%s)"), Name), TraceCounterDisplayHint_None)
+#endif
+		{
+
+		}
+
+#if COUNTERSTRACE_ENABLED
+		FCountersTrace::FCounterInt TotalLoadedCounter;
+		FCountersTrace::FCounterFloat AverageDurationCounter;
+#endif
+		uint64 TotaRequestsCount = 0;
+		double TotalRequestsTime = 0.0;
+	};
+
+#if CSV_PROFILER
+	bool TickCsv(float DeltaTime)
+	{
+		CSV_CUSTOM_STAT_DEFINED(PendingIoRequests, static_cast<int32>(PendingIoRequests), ECsvCustomStatOp::Set);
+		return true; // Keep ticking
+	}
+#endif
+
+	int64 PendingIoRequests = 0;
+#if COUNTERSTRACE_ENABLED
+	FCountersTrace::FCounterInt PendingIoRequestsCounter;
+#endif
+	TArray<FRequestCategory> Categories;
+	FRequestCategory* ChunkTypeToCategoryMap[static_cast<int32>(EIoChunkType::MAX)];
+#if CSV_PROFILER
+	FTSTicker::FDelegateHandle TickerHandle;
+#endif
+};
+#else
+class FIoRequestStats
+{
+public:
+	void OnBatchIssued(FIoBatch& Batch) {}
+	void OnRequestStarted(FIoRequestImpl& Request) {}
+	void OnRequestCompleted(FIoRequestImpl& Request) {}
+};
+#endif
 
 template <typename T, uint32 BlockSize = 128>
 class TBlockAllocator
@@ -201,16 +322,10 @@ public:
 			RequestAllocator->Trim();
 			BatchAllocator.Trim();
 		});
-#if CSV_PROFILER
-		TickerHandle = FTSTicker::GetCoreTicker().AddTicker(FTickerDelegate::CreateRaw(this, &FIoDispatcherImpl::TickCsv));
-#endif
 	}
 
 	~FIoDispatcherImpl()
 	{
-#if CSV_PROFILER
-		FTSTicker::GetCoreTicker().RemoveTicker(TickerHandle);
-#endif
 		delete Thread;
 		FPlatformProcess::ReturnSynchEventToPool(DispatcherEvent);
 		RequestAllocator->ReleaseRef();
@@ -229,7 +344,6 @@ public:
 			Backend->Initialize(BackendContext);
 		}
 		Thread = FRunnableThread::Create(this, TEXT("IoDispatcher"), 0, TPri_AboveNormal, FPlatformAffinity::GetIoDispatcherThreadMask());
-		FFileIoStats::SetDispatcherThreadId(Thread ? Thread->GetThreadID() : 0);
 		return true;
 	}
 
@@ -367,6 +481,9 @@ public:
 		{
 			BatchImpl->UnfinishedRequestsCount += RequestCount;
 		}
+
+		RequestStats.OnBatchIssued(Batch);
+		
 		{
 			FScopeLock _(&WaitingLock);
 			if (!WaitingRequestsHead)
@@ -449,7 +566,6 @@ private:
 				CompletedRequestsHead->ReleaseRef();
 				CompletedRequestsHead = NextRequest;
 				--PendingIoRequestsCount;
-				TRACE_COUNTER_SET(PendingIoRequests, PendingIoRequestsCount);
 			}
 		}
 	}
@@ -482,6 +598,8 @@ private:
 			return false;
 		}
 
+		RequestStats.OnRequestCompleted(*Request);
+
 		FIoBatchImpl* Batch = Request->Batch;
 		if (Request->Callback)
 		{
@@ -572,6 +690,7 @@ private:
 				RequestsToSubmitTail = nullptr;
 			}
 
+			RequestStats.OnRequestStarted(*Request);
 			if (Request->bCancelled)
 			{
 				CompleteRequest(Request, EIoErrorCode::Cancelled);
@@ -608,7 +727,6 @@ private:
 			}
 			
 			++PendingIoRequestsCount;
-			TRACE_COUNTER_SET(PendingIoRequests, PendingIoRequestsCount);	
 			ProcessCompletedRequests();
 		}
 	}
@@ -644,14 +762,6 @@ private:
 		DispatcherEvent->Trigger();
 	}
 
-#if CSV_PROFILER
-	bool TickCsv(float DeltaTime)
-	{
-		CSV_CUSTOM_STAT_DEFINED(PendingIoRequests, (int32)PendingIoRequestsCount, ECsvCustomStatOp::Set);
-		return true; // Keep ticking
-	}
-#endif
-
 	using FBatchAllocator = TBlockAllocator<FIoBatchImpl, 4096>;
 
 	TSharedRef<FIoDispatcherBackendContext> BackendContext;
@@ -670,9 +780,7 @@ private:
 	FIoDispatcher::FIoContainerUnmountedEvent ContainerUnmountedEvent;
 	uint64 PendingIoRequestsCount = 0;
 	int64 TotalLoaded = 0;
-#if CSV_PROFILER
-	FTSTicker::FDelegateHandle TickerHandle;
-#endif
+	FIoRequestStats RequestStats;
 };
 
 FIoDispatcher::FIoDispatcher()
diff --git a/Engine/Source/Runtime/Core/Private/IO/IoDispatcherFileBackend.cpp b/Engine/Source/Runtime/Core/Private/IO/IoDispatcherFileBackend.cpp
--- a/Engine/Source/Runtime/Core/Private/IO/IoDispatcherFileBackend.cpp
+++ b/Engine/Source/Runtime/Core/Private/IO/IoDispatcherFileBackend.cpp
@@ -24,24 +24,6 @@
 #include "IO/IoContainerHeader.h"
 #include "Serialization/MemoryReader.h"
 
-TRACE_DECLARE_INT_COUNTER(IoDispatcherLatencyCircuitBreaks, TEXT("IoDispatcher/LatencyCircuitBreaks"));
-TRACE_DECLARE_INT_COUNTER(IoDispatcherSeekDistanceCircuitBreaks, TEXT("IoDispatcher/SeekDistanceCircuitBreaks"));
-TRACE_DECLARE_INT_COUNTER(IoDispatcherNumPriorityQueues, TEXT("IoDispatcher/NumPriorityQueues"));
-TRACE_DECLARE_MEMORY_COUNTER(IoDispatcherFileBackendTotalBytesRead, TEXT("IoDispatcherFileBackend/TotalBytesRead"));
-TRACE_DECLARE_MEMORY_COUNTER(IoDispatcherFileBackendFileSystemTotalBytesRead, TEXT("IoDispatcherFileBackend/TotalBytesReadFromFileSystem"));
-TRACE_DECLARE_MEMORY_COUNTER(IoDispatcherFileBackendTotalBytesScattered, TEXT("IoDispatcherFileBackend/TotalBytesScattered"));
-TRACE_DECLARE_INT_COUNTER(IoDispatcherFileBackendFileSystemRequests, TEXT("IoDispatcherFileBackend/FileSystemRequests"));
-TRACE_DECLARE_INT_COUNTER(IoDispatcherFileBackendCacheHits, TEXT("IoDispatcherFileBackend/CacheHits"));
-TRACE_DECLARE_INT_COUNTER(IoDispatcherFileBackendCacheMisses, TEXT("IoDispatcherFileBackend/CacheMisses"));
-TRACE_DECLARE_MEMORY_COUNTER(IoDispatcherFileBackendTotalBytesCached, TEXT("IoDispatcherFileBackend/CacheTotalBytesStored"));
-TRACE_DECLARE_INT_COUNTER(IoDispatcherFileBackendQueueLength, TEXT("IoDispatcherFileBackend/QueueLength"));
-TRACE_DECLARE_INT_COUNTER(IoDispatcherFileBackendSequentialReads, TEXT("IoDispatcherFileBackend/SequentialReads"));
-TRACE_DECLARE_INT_COUNTER(IoDispatcherFileBackendForwardSeeks, TEXT("IoDispatcherFileBackend/ForwardSeeks"));
-TRACE_DECLARE_INT_COUNTER(IoDispatcherFileBackendBackwardSeeks, TEXT("IoDispatcherFileBackend/BackwardSeeks"));
-TRACE_DECLARE_INT_COUNTER(IoDispatcherFileBackendSwitchContainerSeeks, TEXT("IoDispatcherFileBackend/SwitchContainerSeeks"));
-TRACE_DECLARE_MEMORY_COUNTER(IoDispatcherFileBackendTotalSeekDistance, TEXT("IoDispatcherFileBackend/TotalSeekDistance"));
-TRACE_DECLARE_MEMORY_COUNTER(IoStoreTocMemory, TEXT("IoDispatcher/TocMemory"));
-
 //PRAGMA_DISABLE_OPTIMIZATION
 
 int32 GIoDispatcherBufferSizeKB = 256;
@@ -153,30 +135,39 @@ void FFileIoStoreBufferAllocator::Initialize(uint64 InMemorySize, uint64 InBuffe
 		Buffer->Memory = BufferMemory + BufferIndex * BufferSize;
 		Buffer->Next = FirstFreeBuffer;
 		FirstFreeBuffer = Buffer;
+		Stats.OnBufferReleased();
 	}
 }
 
 FFileIoStoreBuffer* FFileIoStoreBufferAllocator::AllocBuffer()
 {
-	FScopeLock Lock(&BuffersCritical);
-	FFileIoStoreBuffer* Buffer = FirstFreeBuffer;
-	if (Buffer)
+	FFileIoStoreBuffer* Buffer;
 	{
+		FScopeLock Lock(&BuffersCritical);
+		Buffer = FirstFreeBuffer;
+		if (!Buffer)
+		{
+			return nullptr;
+		}
 		FirstFreeBuffer = Buffer->Next;
-		return Buffer;
 	}
-	return nullptr;
+	Stats.OnBufferAllocated();
+	return Buffer;
 }
 
 void FFileIoStoreBufferAllocator::FreeBuffer(FFileIoStoreBuffer* Buffer)
 {
 	check(Buffer);
-	FScopeLock Lock(&BuffersCritical);
-	Buffer->Next = FirstFreeBuffer;
-	FirstFreeBuffer = Buffer;
+	{
+		FScopeLock Lock(&BuffersCritical);
+		Buffer->Next = FirstFreeBuffer;
+		FirstFreeBuffer = Buffer;
+	}
+	Stats.OnBufferReleased();
 }
 
-FFileIoStoreBlockCache::FFileIoStoreBlockCache()
+FFileIoStoreBlockCache::FFileIoStoreBlockCache(FFileIoStoreStats& InStats)
+	: Stats(InStats)
 {
 	CacheLruHead.LruNext = &CacheLruTail;
 	CacheLruTail.LruPrev = &CacheLruHead;
@@ -227,8 +218,7 @@ bool FFileIoStoreBlockCache::Read(FFileIoStoreReadRequest* Block)
 	FCachedBlock* CachedBlock = CachedBlocks.FindRef(Block->Key.Hash);
 	if (!CachedBlock)
 	{
-		FFileIoStats::OnBlockCacheMiss(ReadBufferSize);
-		TRACE_COUNTER_INCREMENT(IoDispatcherFileBackendCacheMisses);
+		Stats.OnBlockCacheMiss(ReadBufferSize);
 		return false;
 	}
 	
@@ -242,9 +232,8 @@ bool FFileIoStoreBlockCache::Read(FFileIoStoreReadRequest* Block)
 	CachedBlock->LruNext->LruPrev = CachedBlock;
 
 	check(CachedBlock->Buffer);
-	FFileIoStats::OnBlockCacheHit(ReadBufferSize);
+	Stats.OnBlockCacheHit(ReadBufferSize);
 	FMemory::Memcpy(Block->Buffer->Memory, CachedBlock->Buffer, ReadBufferSize);
-	TRACE_COUNTER_INCREMENT(IoDispatcherFileBackendCacheHits);
 
 	return true;
 }
@@ -278,11 +267,17 @@ void FFileIoStoreBlockCache::Store(const FFileIoStoreReadRequest* Block)
 	
 	check(BlockToReplace->Buffer);
 	FMemory::Memcpy(BlockToReplace->Buffer, Block->Buffer->Memory, ReadBufferSize);
-	FFileIoStats::OnBlockCacheStore(ReadBufferSize);
+	Stats.OnBlockCacheStore(ReadBufferSize);
 	CachedBlocks.Add(BlockToReplace->Key, BlockToReplace);
-	TRACE_COUNTER_ADD(IoDispatcherFileBackendTotalBytesCached, ReadBufferSize);
 }
 
+#define UE_FILEIOSTORE_DETAILED_QUEUE_COUNTERS_ENABLED 0
+#if UE_FILEIOSTORE_DETAILED_QUEUE_COUNTERS_ENABLED
+TRACE_DECLARE_INT_COUNTER(IoDispatcherLatencyCircuitBreaks, TEXT("IoDispatcher/LatencyCircuitBreaks"));
+TRACE_DECLARE_INT_COUNTER(IoDispatcherSeekDistanceCircuitBreaks, TEXT("IoDispatcher/SeekDistanceCircuitBreaks"));
+TRACE_DECLARE_INT_COUNTER(IoDispatcherNumPriorityQueues, TEXT("IoDispatcher/NumPriorityQueues"));
+#endif
+
 bool FFileIoStoreOffsetSortedRequestQueue::RequestSortPredicate(const FFileIoStoreReadRequestSortKey& A, const FFileIoStoreReadRequestSortKey& B)
 {
 	if (A.Handle == B.Handle)
@@ -343,10 +338,12 @@ FFileIoStoreReadRequest* FFileIoStoreOffsetSortedRequestQueue::GetNextInternal(F
 			uint64 ThresholdCycles = uint64((GIoDispatcherRequestLatencyCircuitBreakerMS * 1000.0) / FPlatformTime::GetSecondsPerCycle64());
 			bHeadRequestTooOld = (FPlatformTime::Cycles64() - RequestsBySequence.PeekHead()->CreationTime) >= ThresholdCycles;
 
+#if UE_FILEIOSTORE_DETAILED_QUEUE_COUNTERS_ENABLED
 			if (bPop)
 			{
 				TRACE_COUNTER_INCREMENT(IoDispatcherLatencyCircuitBreaks);
 			}
+#endif
 		}
 
 		const bool bChooseByOffset = 
@@ -369,10 +366,12 @@ FFileIoStoreReadRequest* FFileIoStoreOffsetSortedRequestQueue::GetNextInternal(F
 					// Large forward seek so switch back to the oldest outstanding request 
 					RequestIndex = INDEX_NONE;
 
+#if UE_FILEIOSTORE_DETAILED_QUEUE_COUNTERS_ENABLED
 					if (bPop)
 					{
 						TRACE_COUNTER_INCREMENT(IoDispatcherSeekDistanceCircuitBreaks);
 					}
+#endif
 				}
 			}
 		}
@@ -520,7 +519,9 @@ FFileIoStoreReadRequest* FFileIoStoreRequestQueue::Pop()
 		{
 			SortedPriorityQueues.Pop();
 			// SubQueue is invalid here
+#if UE_FILEIOSTORE_DETAILED_QUEUE_COUNTERS_ENABLED
 			TRACE_COUNTER_DECREMENT(IoDispatcherNumPriorityQueues);
+#endif
 		}
 	}
 	else
@@ -534,7 +535,6 @@ FFileIoStoreReadRequest* FFileIoStoreRequestQueue::Pop()
 	
 	check(Result->QueueStatus == FFileIoStoreReadRequest::QueueStatus_InQueue);
 	Result->QueueStatus = FFileIoStoreReadRequest::QueueStatus_Started;
-	TRACE_COUNTER_SET(IoDispatcherFileBackendQueueLength, Heap.Num());
 	return Result;
 }
 
@@ -544,7 +544,9 @@ void FFileIoStoreRequestQueue::PushToPriorityQueues(FFileIoStoreReadRequest* Req
 	if (!SortedPriorityQueues.IsValidIndex(QueueIndex) || SortedPriorityQueues[QueueIndex].GetPriority() != Request->Priority)
 	{
 		SortedPriorityQueues.Insert(FFileIoStoreOffsetSortedRequestQueue(Request->Priority), QueueIndex);
+#if UE_FILEIOSTORE_DETAILED_QUEUE_COUNTERS_ENABLED
 		TRACE_COUNTER_INCREMENT(IoDispatcherNumPriorityQueues);
+#endif
 	}
 	check(Algo::IsSortedBy(SortedPriorityQueues, QueuePriorityProjection, TLess<int32>()));
 	FFileIoStoreOffsetSortedRequestQueue& Queue = SortedPriorityQueues[QueueIndex];
@@ -568,7 +570,6 @@ void FFileIoStoreRequestQueue::Push(FFileIoStoreReadRequest& Request)
 	else
 	{
 		Heap.HeapPush(&Request, QueueSortFunc);
-		TRACE_COUNTER_SET(IoDispatcherFileBackendQueueLength, Heap.Num());
 	}
 }
 
@@ -592,8 +593,6 @@ void FFileIoStoreRequestQueue::Push(FFileIoStoreReadRequestList& Requests)
 			Heap.HeapPush(*It, QueueSortFunc);
 		}
 	}
-
-	TRACE_COUNTER_SET(IoDispatcherFileBackendQueueLength, Heap.Num());
 }
 
 void FFileIoStoreRequestQueue::UpdateOrder()
@@ -658,8 +657,9 @@ void FFileIoStoreRequestQueue::CancelRequestsWithFileHandle(const uint64 FileHan
 	}
 }
 
-FFileIoStoreReader::FFileIoStoreReader(IPlatformFileIoStore& InPlatformImpl)
+FFileIoStoreReader::FFileIoStoreReader(IPlatformFileIoStore& InPlatformImpl, FFileIoStoreStats& InStats)
 	: PlatformImpl(InPlatformImpl)
+	, Stats(InStats)
 {
 }
 
@@ -668,6 +668,16 @@ FFileIoStoreReader::~FFileIoStoreReader()
 	Close();
 }
 
+uint64 FFileIoStoreReader::GetTocAllocatedSize() const
+{
+	return TocImperfectHashMapFallback.GetAllocatedSize() +
+		PerfectHashMap.TocOffsetAndLengths.GetAllocatedSize() +
+		PerfectHashMap.TocChunkIds.GetAllocatedSize() +
+		PerfectHashMap.TocChunkHashSeeds.GetAllocatedSize() +
+		ContainerFile.CompressionBlocks.GetAllocatedSize() +
+		ContainerFile.BlockSignatureHashes.GetAllocatedSize();
+}
+
 FIoStatus FFileIoStoreReader::Initialize(const TCHAR* InContainerPath, int32 InOrder)
 {
 	IPlatformFile& Ipf = FPlatformFileManager::Get().GetPlatformFile();
@@ -737,13 +747,7 @@ FIoStatus FFileIoStoreReader::Initialize(const TCHAR* InContainerPath, int32 InO
 	ContainerFile.BlockSignatureHashes	= MoveTemp(TocResource.ChunkBlockSignatures);
 	ContainerFile.ContainerInstanceId	= ++GlobalContainerInstanceId;
 
-	TRACE_COUNTER_ADD(IoStoreTocMemory,
-		TocImperfectHashMapFallback.GetAllocatedSize() +
-		PerfectHashMap.TocOffsetAndLengths.GetAllocatedSize() +
-		PerfectHashMap.TocChunkIds.GetAllocatedSize() +
-		PerfectHashMap.TocChunkHashSeeds.GetAllocatedSize() +
-		ContainerFile.CompressionBlocks.GetAllocatedSize() +
-		ContainerFile.BlockSignatureHashes.GetAllocatedSize());
+	Stats.OnTocMounted(GetTocAllocatedSize());
 
 	ContainerId = TocResource.Header.ContainerId;
 	Order = InOrder;
@@ -762,6 +766,8 @@ FIoStatus FFileIoStoreReader::Close()
 		PlatformImpl.CloseContainer(Partition.FileHandle);
 	}
 
+	Stats.OnTocUnmounted(GetTocAllocatedSize());
+
 	PerfectHashMap.TocChunkHashSeeds.Empty();
 	PerfectHashMap.TocChunkIds.Empty();
 	PerfectHashMap.TocOffsetAndLengths.Empty();
@@ -1196,7 +1202,9 @@ int64 FFileIoStoreRequestTracker::GetLiveReadRequestsCount() const
 }
 
 FFileIoStore::FFileIoStore(TUniquePtr<IPlatformFileIoStore>&& InPlatformImpl)
-	: RequestTracker(RequestAllocator, RequestQueue)
+	: BlockCache(Stats)
+	, BufferAllocator(Stats)
+	, RequestTracker(RequestAllocator, RequestQueue)
 	, PlatformImpl(MoveTemp(InPlatformImpl))
 {
 }
@@ -1204,14 +1212,12 @@ FFileIoStore::FFileIoStore(TUniquePtr<IPlatformFileIoStore>&& InPlatformImpl)
 FFileIoStore::~FFileIoStore()
 {
 	delete Thread;
-	FFileIoStats::Shutdown();
 }
 
 void FFileIoStore::Initialize(TSharedRef<const FIoDispatcherBackendContext> InContext)
 {
 	check(!Thread);
 
-	FFileIoStats::Init();
 	BackendContext = InContext;
 	bIsMultithreaded = InContext->bIsMultiThreaded;
 
@@ -1229,7 +1235,8 @@ void FFileIoStore::Initialize(TSharedRef<const FIoDispatcherBackendContext> InCo
 		&BackendContext->WakeUpDispatcherThreadDelegate,
 		&RequestAllocator,
 		&BufferAllocator,
-		&BlockCache
+		&BlockCache,
+		&Stats
 	});
 
 	uint64 DecompressionContextCount = uint64(GIoDispatcherDecompressionWorkerCount > 0 ? GIoDispatcherDecompressionWorkerCount : 4);
@@ -1241,12 +1248,11 @@ void FFileIoStore::Initialize(TSharedRef<const FIoDispatcherBackendContext> InCo
 	}
 
 	Thread = FRunnableThread::Create(this, TEXT("IoService"), 0, TPri_AboveNormal);
-	FFileIoStats::SetFileIoStoreThreadId(Thread ? Thread->GetThreadID() : 0);
 }
 
 TIoStatusOr<FIoContainerHeader> FFileIoStore::Mount(const TCHAR* ContainerPath, int32 Order, const FGuid& EncryptionKeyGuid, const FAES::FAESKey& EncryptionKey)
 {
-	TUniquePtr<FFileIoStoreReader> Reader(new FFileIoStoreReader(*PlatformImpl));
+	TUniquePtr<FFileIoStoreReader> Reader(new FFileIoStoreReader(*PlatformImpl, Stats));
 	FIoStatus IoStatus = Reader->Initialize(ContainerPath, Order);
 	if (!IoStatus.IsOk())
 	{
@@ -1352,8 +1358,8 @@ bool FFileIoStore::Resolve(FIoRequestImpl* Request)
 				FFileIoStoreReadRequestList CustomRequests;
 				if (PlatformImpl->CreateCustomRequests(*ResolvedRequest, CustomRequests))
 				{
+					Stats.OnReadRequestsQueued(CustomRequests);
 					RequestTracker.AddReadRequestsToResolvedRequest(CustomRequests, *ResolvedRequest);
-					FFileIoStats::OnFilesystemReadsQueued(CustomRequests);
 					RequestQueue.Push(CustomRequests);
 					OnNewPendingRequestsAdded();
 				}
@@ -1566,7 +1572,7 @@ void FFileIoStore::CompleteDispatcherRequest(FFileIoStoreResolvedRequest* Resolv
 
 void FFileIoStore::FinalizeCompressedBlock(FFileIoStoreCompressedBlock* CompressedBlock)
 {
-	FFileIoStats::OnDecompressComplete(CompressedBlock->CompressedSize, CompressedBlock->UncompressedSize); 
+	Stats.OnDecompressComplete(CompressedBlock); 
 
 	if (CompressedBlock->RawBlocks.Num() > 1)
 	{
@@ -1598,8 +1604,7 @@ void FFileIoStore::FinalizeCompressedBlock(FFileIoStoreCompressedBlock* Compress
 	for (int32 ScatterIndex = 0, ScatterCount = CompressedBlock->ScatterList.Num(); ScatterIndex < ScatterCount; ++ScatterIndex)
 	{
 		FFileIoStoreBlockScatter& Scatter = CompressedBlock->ScatterList[ScatterIndex];
-		FFileIoStats::OnBytesScattered(Scatter.Size);
-		TRACE_COUNTER_ADD(IoDispatcherFileBackendTotalBytesScattered, Scatter.Size);
+		Stats.OnBytesScattered(Scatter.Size);
 		Scatter.Request->bFailed |= CompressedBlock->bFailed;
 		check(!CompressedBlock->bCancelled || !Scatter.Request->DispatcherRequest || Scatter.Request->DispatcherRequest->IsCancelled());
 		check(Scatter.Request->UnfinishedReadsCount > 0);
@@ -1626,13 +1631,11 @@ FIoRequestImpl* FFileIoStore::GetCompletedRequests()
 
 	FFileIoStoreReadRequestList CompletedRequests;
 	PlatformImpl->GetCompletedRequests(CompletedRequests);
+	Stats.OnReadRequestsCompleted(CompletedRequests);
 	for (auto It = CompletedRequests.Steal(); It; ++It)
 	{
 		FFileIoStoreReadRequest* CompletedRequest = *It;
 
-		FFileIoStats::OnReadComplete(CompletedRequest->Size);
-		TRACE_COUNTER_ADD(IoDispatcherFileBackendTotalBytesRead, CompletedRequest->Size);
-
 		if (!CompletedRequest->ImmediateScatter.Request)
 		{
 			check(CompletedRequest->Buffer || CompletedRequest->bCancelled);
@@ -1689,7 +1692,7 @@ FIoRequestImpl* FFileIoStore::GetCompletedRequests()
 				check(CompressedBlock->UnfinishedRawBlocksCount > 0);
 				if (--CompressedBlock->UnfinishedRawBlocksCount == 0)
 				{
-					FFileIoStats::OnDecompressQueued(CompressedBlock->CompressedSize, CompressedBlock->UncompressedSize);
+					Stats.OnDecompressQueued(CompressedBlock);
 					RequestTracker.RemoveCompressedBlock(CompressedBlock);
 					if (!ReadyForDecompressionTail)
 					{
@@ -1706,11 +1709,8 @@ FIoRequestImpl* FFileIoStore::GetCompletedRequests()
 		}
 		else
 		{
-			FFileIoStats::OnBytesScattered(CompletedRequest->ImmediateScatter.Size);
-			TRACE_COUNTER_ADD(IoDispatcherFileBackendTotalBytesScattered, CompletedRequest->ImmediateScatter.Size);
-
-
 			check(!CompletedRequest->Buffer);
+			Stats.OnBytesScattered(CompletedRequest->ImmediateScatter.Size);
 			FFileIoStoreResolvedRequest* CompletedResolvedRequest = CompletedRequest->ImmediateScatter.Request;
 			CompletedResolvedRequest->bFailed |= CompletedRequest->bFailed;
 			check(!CompletedRequest->bCancelled || !CompletedResolvedRequest->DispatcherRequest || CompletedResolvedRequest->DispatcherRequest->IsCancelled());
@@ -1952,7 +1952,7 @@ void FFileIoStore::ReadBlocks(FFileIoStoreResolvedRequest& ResolvedRequest)
 
 	if (!NewBlocks.IsEmpty())
 	{
-		FFileIoStats::OnFilesystemReadsQueued(NewBlocks);
+		Stats.OnReadRequestsQueued(NewBlocks);
 		RequestQueue.Push(NewBlocks);
 		OnNewPendingRequestsAdded();
 	}
@@ -2096,84 +2096,42 @@ CSV_DEFINE_STAT(IoDispatcherFileBackendVerbose,	FrameBlockCacheHitKB);
 CSV_DEFINE_STAT(IoDispatcherFileBackendVerbose,	FrameBlockCacheMisses);
 CSV_DEFINE_STAT(IoDispatcherFileBackendVerbose,	FrameBlockCacheMissKB);
 
-// These stats only go to Insights because they get very large and insights will automatically convert to friendly units
-// IoDispatcher thread
-// Outstanding bytes uncompressed (or decompressed by filesystem api) queued for reading
-TRACE_DECLARE_MEMORY_COUNTER(IoDispatcherFileBackendQueuedFilesystemReadBytes,	TEXT("IoDispatcherFileBackend/QueuedFilesystemReadBytes")); 
-
-// These are long term totals that grow forever and may be less useful?
-
-// IoDispatcher thread 
-// Total reads queued to be executed on the filesystem layer OR resolved from the block cache
-TRACE_DECLARE_INT_COUNTER(IoDispatcherFileBackendQueuedFilesystemReads,			TEXT("IoDispatcherFileBackend/QueuedFilesystemReads"));
-// Total number of compress operations queued
-TRACE_DECLARE_INT_COUNTER(IoDispatcherFileBackendQueuedUncompressBlocks,		TEXT("IoDispatcherFileBackend/QueuedUncompressBlocks"));
-// Total bytes queued to be consumed by decompression
-TRACE_DECLARE_MEMORY_COUNTER(IoDispatcherFileBackendQueuedUncompressBytesIn,	TEXT("IoDispatcherFileBackend/QueuedUncompressBytesIn"));
-// Total bytes queued to be produced by decompression
-TRACE_DECLARE_MEMORY_COUNTER(IoDispatcherFileBackendQueuedUncompressBytesOut,	TEXT("IoDispatcherFileBackend/QueuedUncompressBytesOut"));
-// Total bytes passing through the decompression stage, even if they were not compressed
-TRACE_DECLARE_MEMORY_COUNTER(IoDispatcherFileBackendTotalBytesUncompressedIn,	TEXT("IoDispatcherFileBackend/TotalBytesUncompressedIn"));
-// Total bytes passing through the decompression stage, even if they were not compressed
-TRACE_DECLARE_MEMORY_COUNTER(IoDispatcherFileBackendTotalBytesUncompressedOut,	TEXT("IoDispatcherFileBackend/TotalBytesUncompressedOut"));
-// FileIoStore thread
-// Total number of reads executed through filesystem API layer (useful for calculating ratios of seeks vs sequential reads)
-TRACE_DECLARE_MEMORY_COUNTER(IoDispatcherFileBackendTotalFilesystemReads,		TEXT("IoDispatcherFileBackend/TotalFilesystemReads"));		
-// Total bytes read from platform uncompressed or decompressed for us below the API layer
-TRACE_DECLARE_MEMORY_COUNTER(IoDispatcherFileBackendTotalFilesystemBytesRead,	TEXT("IoDispatcherFileBackend/TotalFilesystemBytesRead"));
-// Total bytes stored in block cache including replacements
-TRACE_DECLARE_MEMORY_COUNTER(IoDispatcherFileBackendTotalBlockCacheBytesStored,	TEXT("IoDispatcherFileBackend/TotalBlockCacheBytesStored"));
-// Total bytes retrieved from block cache
-TRACE_DECLARE_MEMORY_COUNTER(IoDispatcherFileBackendTotalBlockCacheBytesHit,	TEXT("IoDispatcherFileBackend/TotalBlockCacheBytesHit"));
-// Total bytes we failed to retrieve from block cache & had to read from the filesystem
-TRACE_DECLARE_MEMORY_COUNTER(IoDispatcherFileBackendTotalBlockCacheBytesMissed,	TEXT("IoDispatcherFileBackend/TotalBlockCacheBytesMissed"));
-// Total reads that the platform backend considered sequential 
-TRACE_DECLARE_INT_COUNTER(IoDispatcherFileBackendTotalSequentialReads,		TEXT("IoDispatcherFileBackend/TotalSequentialReads"));	
-TRACE_DECLARE_INT_COUNTER(IoDispatcherFileBackendTotalForwardSeeks,			TEXT("IoDispatcherFileBackend/TotalForwardSeeks"));		
-TRACE_DECLARE_INT_COUNTER(IoDispatcherFileBackendTotalBackwardSeeks,		TEXT("IoDispatcherFileBackend/TotalBackwardSeeks"));	
-TRACE_DECLARE_INT_COUNTER(IoDispatcherFileBackendTotalHandleChangeSeeks,	TEXT("IoDispatcherFileBackend/TotalHandleChangeSeeks"));
-// Total seeks including backwards and forwards
-TRACE_DECLARE_INT_COUNTER(IoDispatcherFileBackendTotalSeeks,			TEXT("IoDispatcherFileBackend/TotalSeeks"));
-TRACE_DECLARE_INT_COUNTER(IoDispatcherFileBackendTotalBlockCacheStores, TEXT("IoDispatcherFileBackend/TotalBlockCacheStores"));
-TRACE_DECLARE_INT_COUNTER(IoDispatcherFileBackendTotalBlockCacheHits,	TEXT("IoDispatcherFileBackend/TotalBlockCacheHits"));
-TRACE_DECLARE_INT_COUNTER(IoDispatcherFileBackendTotalBlockCacheMisses, TEXT("IoDispatcherFileBackend/TotalBlockCacheMisses"));
-
-
-#if CSV_PROFILER
-// IoDispatcher thread
-uint64 FFileIoStats::QueuedFilesystemReadBytes{ 0 };	
-uint64 FFileIoStats::QueuedFilesystemReads{ 0 };
-uint64 FFileIoStats::QueuedUncompressBytesIn{ 0 };
-uint64 FFileIoStats::QueuedUncompressBytesOut{ 0 };
-uint64 FFileIoStats::QueuedUncompressBlocks{ 0 };
+#if UE_FILEIOSTORE_STATS_ENABLED
+
+FFileIoStoreStats::FFileIoStoreStats()
+#if COUNTERSTRACE_ENABLED
+	: QueuedReadRequestsSizeCounter(TEXT("FileIoStore/QueuedReadRequestsSize"), TraceCounterDisplayHint_Memory)
+	, CompletedReadRequestsSizeCounter(TEXT("FileIoStore/CompletedReadRequestsSize"), TraceCounterDisplayHint_Memory)
+	, QueuedCompressedSizeCounter(TEXT("FileIoStore/QueuedCompressedSize"), TraceCounterDisplayHint_Memory)
+	, QueuedUncompressedSizeCounter(TEXT("FileIoStore/QueuedUncompressedSize"), TraceCounterDisplayHint_Memory)
+	, CompletedCompressedSizeCounter(TEXT("FileIoStore/CompletedCompressedSize"), TraceCounterDisplayHint_Memory)
+	, CompletedUncompressedSizeCounter(TEXT("FileIoStore/CompletedUncompressedSize"), TraceCounterDisplayHint_Memory)
+	, FileSystemSeeksTotalDistanceCounter(TEXT("FileIoStore/FileSystemSeeksTotalDistance"), TraceCounterDisplayHint_Memory)
+	, FileSystemSeeksForwardCountCounter(TEXT("FileIoStore/FileSystemSeeksForwardCount"), TraceCounterDisplayHint_None)
+	, FileSystemSeeksBackwardCountCounter(TEXT("FileIoStore/FileSystemSeeksBackwardCount"), TraceCounterDisplayHint_None)
+	, FileSystemSeeksChangeHandleCountCounter(TEXT("FileIoStore/FileSystemSeeksChangeHandleCount"), TraceCounterDisplayHint_None)
+	, FileSystemCompletedRequestsSizeCounter(TEXT("FileIoStore/FileSystemCompletedRequestsSize"), TraceCounterDisplayHint_Memory)
+	, BlockCacheStoredSizeCounter(TEXT("FileIoStore/BlockCacheStoredSize"), TraceCounterDisplayHint_Memory)
+	, BlockCacheHitSizeCounter(TEXT("FileIoStore/BlockCacheHitSize"), TraceCounterDisplayHint_Memory)
+	, BlockCacheMissedSizeCounter(TEXT("FileIoStore/BlockCacheMissedSize"), TraceCounterDisplayHint_Memory)
+	, ScatteredSizeCounter(TEXT("FileIoStore/ScatteredSize"), TraceCounterDisplayHint_Memory)
+	, TocMemoryCounter(TEXT("FileIoStore/TocMemory"), TraceCounterDisplayHint_Memory)
+	, AvailableBuffersCounter(TEXT("FileIoStore/AvailableBuffers"), TraceCounterDisplayHint_None)
 #endif
-
-#if IO_DISPATCHER_FILE_STATS
-// IoDispatcher thread
-uint64 FFileIoStats::LastHandle = 0;
-uint64 FFileIoStats::LastOffset = 0;
-
-uint32 FFileIoStats::FileIoStoreThreadId = 0;
-uint32 FFileIoStats::IoDispatcherThreadId = 0;
-
-FTSTicker::FDelegateHandle FFileIoStats::TickerHandle;
-#endif
-
-void FFileIoStats::Init()
 {
 #if CSV_PROFILER
-	TickerHandle = FTSTicker::GetCoreTicker().AddTicker(FTickerDelegate::CreateStatic(&FFileIoStats::CsvTick));
+	TickerHandle = FTSTicker::GetCoreTicker().AddTicker(FTickerDelegate::CreateRaw(this, &FFileIoStoreStats::CsvTick));
 #endif
 }
 
-void FFileIoStats::Shutdown()
+FFileIoStoreStats::~FFileIoStoreStats()
 {
 #if CSV_PROFILER
 	FTSTicker::GetCoreTicker().RemoveTicker(TickerHandle);
 #endif
 }
 
-bool FFileIoStats::CsvTick(float DeltaTime)
+bool FFileIoStoreStats::CsvTick(float DeltaTime)
 {
 #if CSV_PROFILER
 	CSV_CUSTOM_STAT_DEFINED(QueuedFilesystemReadMB, BytesToApproxMB(QueuedFilesystemReadBytes), ECsvCustomStatOp::Set);
@@ -2187,54 +2145,7 @@ bool FFileIoStats::CsvTick(float DeltaTime)
 	return true;
 }
 
-bool FFileIoStats::IsInIoDispatcherThread()
-{
-#if IO_DISPATCHER_FILE_STATS
-	if (FPlatformProcess::SupportsMultithreading())
-	{
-		return FPlatformTLS::GetCurrentThreadId() == IoDispatcherThreadId;
-	}
-#endif
-	return true;
-}
-
-bool FFileIoStats::IsInFileIoStoreThread()
-{
-#if IO_DISPATCHER_FILE_STATS
-	if (FPlatformProcess::SupportsMultithreading())
-	{
-		return FPlatformTLS::GetCurrentThreadId() == FileIoStoreThreadId;
-	}
-#endif
-	return true;
-}
-
-void FFileIoStats::SetDispatcherThreadId(uint32 InIoDispatcherThreadId)
-{
-#if IO_DISPATCHER_FILE_STATS
-	IoDispatcherThreadId = InIoDispatcherThreadId;
-#endif
-}
-
-void FFileIoStats::SetFileIoStoreThreadId(uint32 InFileIoStoreThreadId)
-{
-#if IO_DISPATCHER_FILE_STATS
-	FileIoStoreThreadId = InFileIoStoreThreadId;
-#endif
-}
-
-void FFileIoStats::OnCloseHandle(uint64 Handle)
-{
-#if IO_DISPATCHER_FILE_STATS
-	checkSlow(IsInIoDispatcherThread());
-	if( LastHandle == Handle )
-	{
-		LastHandle = 0;
-	}
-#endif
-}
-
-void FFileIoStats::OnFilesystemReadsQueued(const FFileIoStoreReadRequestList& Requests)
+void FFileIoStoreStats::OnReadRequestsQueued(const FFileIoStoreReadRequestList& Requests)
 {
 	uint64 TotalBytes = 0;
 	int32 NumReads = 0;
@@ -2244,151 +2155,163 @@ void FFileIoStats::OnFilesystemReadsQueued(const FFileIoStoreReadRequestList& Re
 		TotalBytes += Request->Size;
 	}
 
-	OnFilesystemReadsQueued(TotalBytes, NumReads);
-}
-
-void FFileIoStats::OnFilesystemReadsQueued(uint64 BytesToRead, int32 NumReads)
-{
-#if IO_DISPATCHER_FILE_STATS
-	checkSlow(IsInIoDispatcherThread());
-
 #if CSV_PROFILER
-	QueuedFilesystemReadBytes += BytesToRead;
+	QueuedFilesystemReadBytes += TotalBytes;
 	QueuedFilesystemReads += NumReads;
 #endif
 
-	TRACE_COUNTER_ADD(IoDispatcherFileBackendQueuedFilesystemReads, NumReads)
-	TRACE_COUNTER_ADD(IoDispatcherFileBackendQueuedFilesystemReadBytes, BytesToRead);
+#if COUNTERSTRACE_ENABLED
+	QueuedReadRequestsSizeCounter.Add(TotalBytes);
 #endif
 }
 
-void FFileIoStats::OnFilesystemReadStarted(uint64 Handle, uint64 Offset, uint64 NumBytes)
+void FFileIoStoreStats::OnFilesystemReadStarted(const FFileIoStoreReadRequest* Request)
 {
-#if IO_DISPATCHER_FILE_STATS
-	checkSlow(IsInFileIoStoreThread());
-
 	CSV_CUSTOM_STAT_DEFINED(FrameFilesystemReads, 1, ECsvCustomStatOp::Accumulate);
-	TRACE_COUNTER_INCREMENT(IoDispatcherFileBackendTotalFilesystemReads);
 
-	if (LastHandle != Handle)
+	if (LastHandle != Request->FileHandle)
 	{
 		OnHandleChangeSeek();
 	}
-	else if (LastOffset == Offset)
+	else if (LastOffset == Request->Offset)
 	{
 		OnSequentialRead();
 	}
 	else
 	{
-		OnSeek(LastOffset, Offset);
+		OnSeek(LastOffset, Request->Offset);
 	}
-	LastOffset = Offset + NumBytes;
-	LastHandle = Handle;
-#endif
+	LastOffset = Request->Offset + Request->Size;
+	LastHandle = Request->FileHandle;
+
+	CSV_CUSTOM_STAT_DEFINED(FrameFilesystemReads, 1, ECsvCustomStatOp::Accumulate);
 }
 
-void FFileIoStats::OnFilesystemReadsComplete(const FFileIoStoreReadRequestList& CompletedRequests)
+void FFileIoStoreStats::OnFilesystemReadsStarted(const FFileIoStoreReadRequestList& Requests)
 {
-#if IO_DISPATCHER_FILE_STATS
 	int64 TotalBytes = 0;
 	int32 NumReads = 0;
-	for (const FFileIoStoreReadRequest* Request : CompletedRequests)
+	for (const FFileIoStoreReadRequest* Request : Requests)
 	{
+		if (LastHandle != Request->FileHandle)
+		{
+			OnHandleChangeSeek();
+		}
+		else if (LastOffset == Request->Offset)
+		{
+			OnSequentialRead();
+		}
+		else
+		{
+			OnSeek(LastOffset, Request->Offset);
+		}
+
+		LastOffset = Request->Offset + Request->Size;
+		LastHandle = Request->FileHandle;
+
 		++NumReads;
 		TotalBytes += Request->Size;
 	}
 
-	FFileIoStats::OnFilesystemReadsComplete(TotalBytes, NumReads);
+	CSV_CUSTOM_STAT_DEFINED(FrameFilesystemReads, NumReads, ECsvCustomStatOp::Accumulate);
+}
+
+void FFileIoStoreStats::OnFilesystemReadCompleted(const FFileIoStoreReadRequest* CompletedRequest)
+{
+	CSV_CUSTOM_STAT_DEFINED(FrameFilesystemBytesReadKB, BytesToApproxKB(CompletedRequest->Size), ECsvCustomStatOp::Accumulate);
+#if COUNTERSTRACE_ENABLED
+	FileSystemCompletedRequestsSizeCounter.Add(CompletedRequest->Size);
 #endif
 }
 
-void FFileIoStats::OnFilesystemReadsComplete(int64 BytesRead, int32 NumReads)
+void FFileIoStoreStats::OnFilesystemReadsCompleted(const FFileIoStoreReadRequestList& CompletedRequests)
 {
-#if IO_DISPATCHER_FILE_STATS
-	checkSlow(IsInFileIoStoreThread());
-	CSV_CUSTOM_STAT_DEFINED(FrameFilesystemBytesReadKB, BytesToApproxKB(BytesRead), ECsvCustomStatOp::Accumulate);
+	int64 TotalBytes = 0;
+	int32 NumReads = 0;
+	for (const FFileIoStoreReadRequest* Request : CompletedRequests)
+	{
+		++NumReads;
+		TotalBytes += Request->Size;
+	}
 
-	TRACE_COUNTER_ADD(IoDispatcherFileBackendTotalFilesystemBytesRead, BytesRead);
+	CSV_CUSTOM_STAT_DEFINED(FrameFilesystemBytesReadKB, BytesToApproxKB(TotalBytes), ECsvCustomStatOp::Accumulate);
+#if COUNTERSTRACE_ENABLED
+	FileSystemCompletedRequestsSizeCounter.Add(TotalBytes);
 #endif
 }
 
-void FFileIoStats::OnReadComplete(int64 BytesRead)
+void FFileIoStoreStats::OnReadRequestsCompleted(const FFileIoStoreReadRequestList& CompletedRequests)
 {
-#if IO_DISPATCHER_FILE_STATS
-	checkSlow(IsInIoDispatcherThread());
+	int64 TotalBytes = 0;
+	int32 NumReads = 0;
+	for (const FFileIoStoreReadRequest* Request : CompletedRequests)
+	{
+		++NumReads;
+		TotalBytes += Request->Size;
+	}
+
 #if CSV_PROFILER
-	QueuedFilesystemReadBytes -= BytesRead;
-	--QueuedFilesystemReads;
+	QueuedFilesystemReadBytes -= TotalBytes;
+	QueuedFilesystemReads -= NumReads;
 #endif
 
-	CSV_CUSTOM_STAT_DEFINED(FrameBytesReadKB, BytesToApproxKB(BytesRead), ECsvCustomStatOp::Accumulate); // TODO: Send to insights if CSV is disabled
+	CSV_CUSTOM_STAT_DEFINED(FrameBytesReadKB, BytesToApproxKB(TotalBytes), ECsvCustomStatOp::Accumulate); // TODO: Send to insights if CSV is disabled
 
-	TRACE_COUNTER_SUBTRACT(IoDispatcherFileBackendQueuedFilesystemReadBytes, BytesRead);
-	TRACE_COUNTER_ADD(IoDispatcherFileBackendTotalBytesRead, BytesRead);
-	TRACE_COUNTER_DECREMENT(IoDispatcherFileBackendQueuedFilesystemReads);
+#if COUNTERSTRACE_ENABLED
+	QueuedReadRequestsSizeCounter.Subtract(TotalBytes);
+	CompletedReadRequestsSizeCounter.Add(TotalBytes);
 #endif
 }
 
-void FFileIoStats::OnDecompressQueued(int64 CompressedBytes, int64 UncompressedBytes)
+void FFileIoStoreStats::OnDecompressQueued(const FFileIoStoreCompressedBlock* CompressedBlock)
 {
-#if IO_DISPATCHER_FILE_STATS
-	checkSlow(IsInIoDispatcherThread());
 #if CSV_PROFILER
 	++QueuedUncompressBlocks;
-	QueuedUncompressBytesIn += CompressedBytes;
-	QueuedUncompressBytesOut += UncompressedBytes;
+	QueuedUncompressBytesIn += CompressedBlock->CompressedSize;
+	QueuedUncompressBytesOut += CompressedBlock->UncompressedSize;
 #endif
 
-	TRACE_COUNTER_INCREMENT(IoDispatcherFileBackendQueuedUncompressBlocks);
-	TRACE_COUNTER_ADD(IoDispatcherFileBackendQueuedUncompressBytesIn, CompressedBytes);
-	TRACE_COUNTER_ADD(IoDispatcherFileBackendQueuedUncompressBytesOut, UncompressedBytes);
+#if COUNTERSTRACE_ENABLED
+	QueuedCompressedSizeCounter.Add(CompressedBlock->CompressedSize);
+	QueuedUncompressedSizeCounter.Add(CompressedBlock->UncompressedSize);
 #endif
 }
 
-void FFileIoStats::OnDecompressComplete(int64 CompressedBytes, int64 UncompressedBytes)
+void FFileIoStoreStats::OnDecompressComplete(const FFileIoStoreCompressedBlock* CompressedBlock)
 {
-#if IO_DISPATCHER_FILE_STATS
-	checkSlow(IsInIoDispatcherThread());
 #if CSV_PROFILER
 	--QueuedUncompressBlocks;
-	QueuedUncompressBytesIn -= CompressedBytes;
-	QueuedUncompressBytesOut -= UncompressedBytes;
+	QueuedUncompressBytesIn -= CompressedBlock->CompressedSize;
+	QueuedUncompressBytesOut -= CompressedBlock->UncompressedSize;
 #endif
 
-	CSV_CUSTOM_STAT_DEFINED(FrameBytesUncompressedInKB, BytesToApproxKB(CompressedBytes), ECsvCustomStatOp::Accumulate); 
-	CSV_CUSTOM_STAT_DEFINED(FrameBytesUncompressedOutKB, BytesToApproxKB(UncompressedBytes), ECsvCustomStatOp::Accumulate);
+	CSV_CUSTOM_STAT_DEFINED(FrameBytesUncompressedInKB, BytesToApproxKB(CompressedBlock->CompressedSize), ECsvCustomStatOp::Accumulate); 
+	CSV_CUSTOM_STAT_DEFINED(FrameBytesUncompressedOutKB, BytesToApproxKB(CompressedBlock->UncompressedSize), ECsvCustomStatOp::Accumulate);
 
-	TRACE_COUNTER_DECREMENT(IoDispatcherFileBackendQueuedUncompressBlocks);
-	TRACE_COUNTER_SUBTRACT(IoDispatcherFileBackendQueuedUncompressBytesIn, CompressedBytes);
-	TRACE_COUNTER_SUBTRACT(IoDispatcherFileBackendQueuedUncompressBytesOut, UncompressedBytes);
+#if COUNTERSTRACE_ENABLED
+	QueuedCompressedSizeCounter.Subtract(CompressedBlock->CompressedSize);
+	QueuedUncompressedSizeCounter.Subtract(CompressedBlock->UncompressedSize);
 
-	TRACE_COUNTER_ADD(IoDispatcherFileBackendTotalBytesUncompressedIn, CompressedBytes);
-	TRACE_COUNTER_ADD(IoDispatcherFileBackendTotalBytesUncompressedOut, UncompressedBytes);
+	CompletedCompressedSizeCounter.Add(CompressedBlock->CompressedSize);
+	CompletedUncompressedSizeCounter.Add(CompressedBlock->UncompressedSize);
 #endif
 }
 
-void FFileIoStats::OnBytesScattered(int64 NumBytes)
+void FFileIoStoreStats::OnBytesScattered(int64 NumBytes)
 {
-#if IO_DISPATCHER_FILE_STATS
-	checkSlow(IsInIoDispatcherThread());
 	CSV_CUSTOM_STAT_DEFINED(FrameBytesScatteredKB, BytesToApproxKB(NumBytes), ECsvCustomStatOp::Accumulate);
-	TRACE_COUNTER_ADD(IoDispatcherFileBackendTotalBytesScattered, NumBytes);
+#if COUNTERSTRACE_ENABLED
+	ScatteredSizeCounter.Add(NumBytes);
 #endif
 }
 
-void FFileIoStats::OnSequentialRead()
+void FFileIoStoreStats::OnSequentialRead()
 {
-#if IO_DISPATCHER_FILE_STATS
-	checkSlow(IsInFileIoStoreThread());
 	CSV_CUSTOM_STAT_DEFINED(FrameSequentialReads, 1, ECsvCustomStatOp::Accumulate);
-	TRACE_COUNTER_INCREMENT(IoDispatcherFileBackendTotalSequentialReads);
-#endif
 }
 
-void FFileIoStats::OnSeek(uint64 PrevOffset, uint64 NewOffset)
+void FFileIoStoreStats::OnSeek(uint64 PrevOffset, uint64 NewOffset)
 {
-#if IO_DISPATCHER_FILE_STATS
-	checkSlow(IsInFileIoStoreThread());
 	if (NewOffset > PrevOffset)
 	{
 		int64 Delta = NewOffset - PrevOffset;
@@ -2396,8 +2319,10 @@ void FFileIoStats::OnSeek(uint64 PrevOffset, uint64 NewOffset)
 		CSV_CUSTOM_STAT_DEFINED(FrameForwardSeeks, 1, ECsvCustomStatOp::Accumulate);
 		CSV_CUSTOM_STAT_DEFINED(FrameSeekDistanceMB, BytesToApproxMB(Delta), ECsvCustomStatOp::Accumulate);
 
-		TRACE_COUNTER_ADD(IoDispatcherFileBackendTotalSeekDistance, Delta);
-		TRACE_COUNTER_INCREMENT(IoDispatcherFileBackendTotalForwardSeeks);
+#if COUNTERSTRACE_ENABLED
+		FileSystemSeeksTotalDistanceCounter.Add(Delta);
+		FileSystemSeeksForwardCountCounter.Increment();
+#endif
 	}
 	else
 	{
@@ -2405,59 +2330,81 @@ void FFileIoStats::OnSeek(uint64 PrevOffset, uint64 NewOffset)
 		CSV_CUSTOM_STAT_DEFINED(FrameBackwardSeeks, 1, ECsvCustomStatOp::Accumulate);
 		CSV_CUSTOM_STAT_DEFINED(FrameSeekDistanceMB, BytesToApproxMB(Delta), ECsvCustomStatOp::Accumulate);
 
-		TRACE_COUNTER_ADD(IoDispatcherFileBackendTotalSeekDistance, Delta);
-		TRACE_COUNTER_INCREMENT(IoDispatcherFileBackendTotalBackwardSeeks);
+#if COUNTERSTRACE_ENABLED
+		FileSystemSeeksTotalDistanceCounter.Add(Delta);
+		FileSystemSeeksBackwardCountCounter.Increment();
+#endif
 	}
 
 	CSV_CUSTOM_STAT_DEFINED(FrameSeeks, 1, ECsvCustomStatOp::Accumulate);
-	TRACE_COUNTER_INCREMENT(IoDispatcherFileBackendTotalSeeks);
-#endif
 }
 
-void FFileIoStats::OnHandleChangeSeek()
+void FFileIoStoreStats::OnHandleChangeSeek()
 {
-#if IO_DISPATCHER_FILE_STATS
-	checkSlow(IsInFileIoStoreThread());
 	CSV_CUSTOM_STAT_DEFINED(FrameHandleChangeSeeks, 1, ECsvCustomStatOp::Accumulate);
-	TRACE_COUNTER_INCREMENT(IoDispatcherFileBackendTotalHandleChangeSeeks);
-
 	CSV_CUSTOM_STAT_DEFINED(FrameSeeks, 1, ECsvCustomStatOp::Accumulate);
-	TRACE_COUNTER_INCREMENT(IoDispatcherFileBackendTotalSeeks);
+
+#if COUNTERSTRACE_ENABLED
+	FileSystemSeeksChangeHandleCountCounter.Increment();
 #endif
 }
 
-void FFileIoStats::OnBlockCacheStore(uint64 NumBytes)
+void FFileIoStoreStats::OnBlockCacheStore(uint64 NumBytes)
 {
-#if IO_DISPATCHER_FILE_STATS
-	checkSlow(IsInFileIoStoreThread());
 	CSV_CUSTOM_STAT_DEFINED(FrameBlockCacheStores, 1, ECsvCustomStatOp::Accumulate);
 	CSV_CUSTOM_STAT_DEFINED(FrameBlockCacheStoresKB, BytesToApproxKB(NumBytes), ECsvCustomStatOp::Accumulate);
 
-	TRACE_COUNTER_INCREMENT(IoDispatcherFileBackendTotalBlockCacheStores);
-	TRACE_COUNTER_ADD(IoDispatcherFileBackendTotalBlockCacheBytesStored, NumBytes);
+#if COUNTERSTRACE_ENABLED
+	BlockCacheStoredSizeCounter.Add(NumBytes);
 #endif
 }
 
-void FFileIoStats::OnBlockCacheHit(uint64 NumBytes)
+void FFileIoStoreStats::OnBlockCacheHit(uint64 NumBytes)
 {
-#if IO_DISPATCHER_FILE_STATS
-	checkSlow(IsInFileIoStoreThread());
 	CSV_CUSTOM_STAT_DEFINED(FrameBlockCacheHits, 1, ECsvCustomStatOp::Accumulate);
 	CSV_CUSTOM_STAT_DEFINED(FrameBlockCacheHitKB, BytesToApproxKB(NumBytes), ECsvCustomStatOp::Accumulate);
 
-	TRACE_COUNTER_INCREMENT(IoDispatcherFileBackendTotalBlockCacheHits);
-	TRACE_COUNTER_ADD(IoDispatcherFileBackendTotalBlockCacheBytesHit, NumBytes);
+#if COUNTERSTRACE_ENABLED
+	BlockCacheHitSizeCounter.Add(NumBytes);
 #endif
 }
 
-void FFileIoStats::OnBlockCacheMiss(uint64 NumBytes)
+void FFileIoStoreStats::OnBlockCacheMiss(uint64 NumBytes)
 {
-#if IO_DISPATCHER_FILE_STATS
-	checkSlow(IsInFileIoStoreThread());
 	CSV_CUSTOM_STAT_DEFINED(FrameBlockCacheMisses, 1, ECsvCustomStatOp::Accumulate);
 	CSV_CUSTOM_STAT_DEFINED(FrameBlockCacheMissKB, BytesToApproxKB(NumBytes), ECsvCustomStatOp::Accumulate);
 
-	TRACE_COUNTER_INCREMENT(IoDispatcherFileBackendTotalBlockCacheMisses);
-	TRACE_COUNTER_ADD(IoDispatcherFileBackendTotalBlockCacheBytesMissed, NumBytes);
+#if COUNTERSTRACE_ENABLED
+	BlockCacheMissedSizeCounter.Add(NumBytes);
 #endif
 }
+
+void FFileIoStoreStats::OnTocMounted(uint64 AllocatedSize)
+{
+#if COUNTERSTRACE_ENABLED
+	TocMemoryCounter.Add(AllocatedSize);
+#endif
+}
+
+void FFileIoStoreStats::OnTocUnmounted(uint64 AllocatedSize)
+{
+#if COUNTERSTRACE_ENABLED
+	TocMemoryCounter.Subtract(AllocatedSize);
+#endif
+}
+
+void FFileIoStoreStats::OnBufferReleased()
+{
+#if COUNTERSTRACE_ENABLED
+	AvailableBuffersCounter.Increment();
+#endif
+}
+
+void FFileIoStoreStats::OnBufferAllocated()
+{
+#if COUNTERSTRACE_ENABLED
+	AvailableBuffersCounter.Decrement();
+#endif
+}
+
+#endif
\ No newline at end of file
diff --git a/Engine/Source/Runtime/Core/Private/IO/IoStore.cpp b/Engine/Source/Runtime/Core/Private/IO/IoStore.cpp
--- a/Engine/Source/Runtime/Core/Private/IO/IoStore.cpp
+++ b/Engine/Source/Runtime/Core/Private/IO/IoStore.cpp
@@ -33,7 +33,7 @@ constexpr char FIoStoreTocHeader::TocMagicImg[];
 
 //////////////////////////////////////////////////////////////////////////
 
-TRACE_DECLARE_INT_COUNTER(IoStoreAvailableCompressionBuffers, TEXT("IoStore/AvailableCompressionBuffers"));
+//TRACE_DECLARE_INT_COUNTER(IoStoreAvailableCompressionBuffers, TEXT("IoStore/AvailableCompressionBuffers"));
 
 FString LexToString(const EIoChunkType Type)
 {
@@ -346,7 +346,7 @@ public:
 		{
 			AvailableCompressionBuffers.Add(new FIoBuffer(CompressionBufferSize));
 		}
-		TRACE_COUNTER_SET(IoStoreAvailableCompressionBuffers, AvailableCompressionBuffers.Num());
+		//TRACE_COUNTER_SET(IoStoreAvailableCompressionBuffers, AvailableCompressionBuffers.Num());
 
 		BeginCompressionThread = Async(EAsyncExecution::Thread, [this]() { BeginCompressionThreadFunc(); });
 		BeginEncryptionAndSigningThread = Async(EAsyncExecution::Thread, [this]() { BeginEncryptionAndSigningThreadFunc(); });
@@ -398,7 +398,7 @@ public:
 				if (AvailableCompressionBuffers.Num())
 				{
 					AllocatedBuffer = AvailableCompressionBuffers.Pop();
-					TRACE_COUNTER_DECREMENT(IoStoreAvailableCompressionBuffers);
+					//TRACE_COUNTER_DECREMENT(IoStoreAvailableCompressionBuffers);
 				}
 			}
 			if (!AllocatedBuffer)
@@ -422,7 +422,7 @@ public:
 			FScopeLock Lock(&AvailableCompressionBuffersCritical);
 			bTriggerEvent = AvailableCompressionBuffers.Num() == 0;
 			AvailableCompressionBuffers.Push(Buffer);
-			TRACE_COUNTER_INCREMENT(IoStoreAvailableCompressionBuffers);
+			//TRACE_COUNTER_INCREMENT(IoStoreAvailableCompressionBuffers);
 		}
 		if (bTriggerEvent)
 		{
diff --git a/Engine/Source/Runtime/Core/Public/IO/IoDispatcher.h b/Engine/Source/Runtime/Core/Public/IO/IoDispatcher.h
--- a/Engine/Source/Runtime/Core/Public/IO/IoDispatcher.h
+++ b/Engine/Source/Runtime/Core/Public/IO/IoDispatcher.h
@@ -885,6 +885,7 @@ class FIoBatch final
 {
 	friend class FIoDispatcher;
 	friend class FIoDispatcherImpl;
+	friend class FIoRequestStats;
 
 public:
 	CORE_API FIoBatch();
diff --git a/Engine/Source/Runtime/Core/Public/IO/IoDispatcherBackend.h b/Engine/Source/Runtime/Core/Public/IO/IoDispatcherBackend.h
--- a/Engine/Source/Runtime/Core/Public/IO/IoDispatcherBackend.h
+++ b/Engine/Source/Runtime/Core/Public/IO/IoDispatcherBackend.h
@@ -6,6 +6,10 @@
 #include "HAL/LowLevelMemTracker.h"
 #include "Misc/Optional.h"
 #include "ProfilingDebugging/TagTrace.h"
+#include "ProfilingDebugging/CountersTrace.h"
+#include "ProfilingDebugging/CsvProfiler.h"
+
+#define UE_IODISPATCHER_STATS_ENABLED (COUNTERSTRACE_ENABLED || CSV_PROFILER)
 
 struct FIoContainerHeader;
 
@@ -62,6 +66,7 @@ private:
 	friend class FIoDispatcherImpl;
 	friend class FIoRequest;
 	friend class FIoBatch;
+	friend class FIoRequestStats;
 
 	void AddRef()
 	{
@@ -81,6 +86,9 @@ private:
 	FIoRequestAllocator& Allocator;
 	struct IIoDispatcherBackend* Backend = nullptr;
 	FIoBatchImpl* Batch = nullptr;
+#if UE_IODISPATCHER_STATS_ENABLED
+	double StartTime = -1.0;
+#endif
 	TOptional<FIoBuffer> Buffer;
 	FIoReadCallback Callback;
 	TAtomic<uint32> RefCount{ 0 };
