reduced duration of TSTicker stresss test to avoid running out of memory, and enabled it back

diff --git a/Engine/Source/Runtime/Core/Private/Tests/Containers/TickerTests.cpp b/Engine/Source/Runtime/Core/Private/Tests/Containers/TickerTests.cpp
--- a/Engine/Source/Runtime/Core/Private/Tests/Containers/TickerTests.cpp
+++ b/Engine/Source/Runtime/Core/Private/Tests/Containers/TickerTests.cpp
@@ -9,7 +9,7 @@
 
 #if WITH_DEV_AUTOMATION_TESTS
 
-IMPLEMENT_SIMPLE_AUTOMATION_TEST(FTSTickerTest, "System.Core.Containers.TSTicker", EAutomationTestFlags::ApplicationContextMask | EAutomationTestFlags::EngineFilter | EAutomationTestFlags::Disabled);
+IMPLEMENT_SIMPLE_AUTOMATION_TEST(FTSTickerTest, "System.Core.Containers.TSTicker", EAutomationTestFlags::ApplicationContextMask | EAutomationTestFlags::EngineFilter);
 
 template<uint32 NumDelegates, uint32 NumTicks>
 void TickerPerfTest()
@@ -239,7 +239,7 @@ bool FTSTickerTest::RunTest(const FString& Parameters)
 			));
 		}
 
-		FPlatformProcess::Sleep(1.0f); // let it run for a while
+		FPlatformProcess::Sleep(0.3f); // let it run for a while
 		bQuit = true;
 		Tasks.Add(TickTask);
 		verify(Wait(Tasks, FTimespan::FromSeconds(5)));
