Provide a way for Android java to set UE's crash context.

diff --git a/Engine/Build/Android/Java/src/com/epicgames/unreal/GameActivity.java.template b/Engine/Build/Android/Java/src/com/epicgames/unreal/GameActivity.java.template
--- a/Engine/Build/Android/Java/src/com/epicgames/unreal/GameActivity.java.template
+++ b/Engine/Build/Android/Java/src/com/epicgames/unreal/GameActivity.java.template
@@ -2396,6 +2396,7 @@ public class GameActivity extends $${gameActivitySuperClass}$$ implements Surfac
 
 		// update memory stats every 10 seconds
 		memoryRunnable = new Runnable() {
+			boolean showPSSInfo = false;
 			int LastProcessMemoryValue = 0;
 			@Override
 			public void run()
@@ -2433,7 +2434,14 @@ public class GameActivity extends $${gameActivitySuperClass}$$ implements Surfac
 				{
 					Log.debug("memstat: "+memstatusline);
 				}
-				
+				if( showPSSInfo )
+				{
+					String[] smapsLines = FindLineFromFile("/proc/self/smaps_rollup",new String[] {"SwapPss:","Pss:"});
+				    for(String memstatusline : smapsLines)
+				    {
+					    Log.debug("memstat: "+memstatusline);
+				    }
+				}
 				synchronized(_activity)
 				{
 					if (ProcessMemory > 0)
@@ -6476,6 +6484,12 @@ public class GameActivity extends $${gameActivitySuperClass}$$ implements Surfac
 		return false;
 	}
 
+	public void SetCrashContextData(String Key, String Value)
+	{
+		nativeCrashContextSetStringKey(Key,Value);
+//$${gameActivitySetCrashContextData}$$
+	}
+
 	public native int nativeGetCPUFamily();
 	public native boolean nativeSupportsNEON();
 	public native void nativeSetAffinityInfo(boolean bEnableAffinity, int bigCoreMask, int littleCoreMask);
@@ -6518,6 +6532,12 @@ public class GameActivity extends $${gameActivitySuperClass}$$ implements Surfac
 	public native void nativeOnThermalStatusChangedListener(int status);
 	public native void nativeOnTrimMemory(int MemoryTrimValue);
 
+	public native void nativeCrashContextSetStringKey(String Key, String Value);
+	public native void nativeCrashContextSetBooleanKey(String Key, boolean Value);
+	public native void nativeCrashContextSetIntegerKey(String Key, int Value);
+	public native void nativeCrashContextSetFloatKey(String Key, float Value);
+	public native void nativeCrashContextSetDoubleKey(String Key, double Value);
+
 	static
 	{
 		System.loadLibrary("Unreal");
diff --git a/Engine/Source/Programs/UnrealBuildTool/Platform/Android/UEDeployAndroid.cs b/Engine/Source/Programs/UnrealBuildTool/Platform/Android/UEDeployAndroid.cs
--- a/Engine/Source/Programs/UnrealBuildTool/Platform/Android/UEDeployAndroid.cs
+++ b/Engine/Source/Programs/UnrealBuildTool/Platform/Android/UEDeployAndroid.cs
@@ -5000,6 +5000,7 @@ namespace UnrealBuildTool
 				{ "//$${gameActivityIapSetupServiceAdditions}$$", UPL.ProcessPluginNode(NDKArch, "gameActivityIapSetupServiceAdditions", "")},
 				{ "//$${gameActivityOnRestartApplicationAdditions}$$", UPL.ProcessPluginNode(NDKArch, "gameActivityOnRestartApplicationAdditions", "")},
 				{ "//$${gameActivityForceQuitAdditions}$$", UPL.ProcessPluginNode(NDKArch, "gameActivityForceQuitAdditions", "")},
+				{ "//$${gameActivitySetCrashContextData}$$", UPL.ProcessPluginNode(NDKArch, "gameActivitySetCrashContextData", "")},
 				{ "//$${soLoadLibrary}$$", UPL.ProcessPluginNode(NDKArch, "soLoadLibrary", LoadLibraryDefaults)},
 				{ "$${gameActivitySuperClass}$$", SuperClassDefault},
 			};
diff --git a/Engine/Source/Runtime/Core/Private/Android/AndroidPlatformCrashContext.cpp b/Engine/Source/Runtime/Core/Private/Android/AndroidPlatformCrashContext.cpp
--- a/Engine/Source/Runtime/Core/Private/Android/AndroidPlatformCrashContext.cpp
+++ b/Engine/Source/Runtime/Core/Private/Android/AndroidPlatformCrashContext.cpp
@@ -20,6 +20,9 @@
 #include "HAL/IConsoleManager.h"
 #include "HAL/ThreadManager.h"
 #include "HAL/RunnableThread.h"
+#include "GenericPlatform/GenericPlatformCrashContext.h"
+#include "Async/TaskGraphInterfaces.h"
+#include "Android/AndroidJavaEnv.h"
 
 static int64 GetAndroidLibraryBaseAddress();
 
@@ -27,6 +30,65 @@ extern FString AndroidThunkCpp_GetMetaDataString(const FString& Key);
 
 extern FString AndroidRelativeToAbsolutePath(bool bUseInternalBasePath, FString RelPath);
 
+/** Java to native crash context k/v setting API  */
+static void SetCrashContextOnGameThread(const FString KeyIn, const FString ValueIn)
+{
+	if (IsInGameThread())
+	{
+		FGenericCrashContext::SetGameData(KeyIn, ValueIn);
+	}
+	else if (FTaskGraphInterface::IsRunning())
+	{
+		FFunctionGraphTask::CreateAndDispatchWhenReady(
+			[Key = KeyIn, Value = ValueIn]()
+			{
+				FGenericCrashContext::SetGameData(Key, Value);
+			}
+		, TStatId(), NULL, ENamedThreads::GameThread);
+	}
+	else
+	{
+		UE_LOG(LogAndroid, Log, TEXT("Failed to set crash context `%s` = '%s'"), *KeyIn, *ValueIn);
+	}
+}
+
+JNI_METHOD void Java_com_epicgames_unreal_GameActivity_nativeCrashContextSetStringKey(JNIEnv* jenv, jobject thiz, jstring JavaKey, jstring JavaValue)
+{
+	FString Key, Value;
+	Key = FJavaHelper::FStringFromParam(jenv, JavaKey);
+	Value = FJavaHelper::FStringFromParam(jenv, JavaValue);
+
+	SetCrashContextOnGameThread(Key, Value);
+}
+
+JNI_METHOD void  Java_com_epicgames_unreal_GameActivity_nativeCrashContextSetBooleanKey(JNIEnv* jenv, jobject thiz, jstring JavaKey, jboolean JavaValue)
+{
+	FString Key;
+	Key = FJavaHelper::FStringFromParam(jenv, JavaKey);
+	SetCrashContextOnGameThread(Key, JavaValue ? TEXT("true") : TEXT("false"));
+}
+
+JNI_METHOD void  Java_com_epicgames_unreal_GameActivity_nativeCrashContextSetIntegerKey(JNIEnv* jenv, jobject thiz, jstring JavaKey, jint JavaValue)
+{
+	FString Key;
+	Key = FJavaHelper::FStringFromParam(jenv, JavaKey);
+	SetCrashContextOnGameThread(Key, *TTypeToString<int32>::ToString(JavaValue));
+}
+
+JNI_METHOD void  Java_com_epicgames_unreal_GameActivity_nativeCrashContextSetFloatKey(JNIEnv* jenv, jobject thiz, jstring JavaKey, jfloat JavaValue)
+{
+	FString Key;
+	Key = FJavaHelper::FStringFromParam(jenv, JavaKey);
+	SetCrashContextOnGameThread(Key, *TTypeToString<float>::ToString(JavaValue));
+}
+
+JNI_METHOD void  Java_com_epicgames_unreal_GameActivity_nativeCrashContextSetDoubleKey(JNIEnv* jenv, jobject thiz, jstring JavaKey, jdouble JavaValue)
+{
+	FString Key;
+	Key = FJavaHelper::FStringFromParam(jenv, JavaKey);
+	SetCrashContextOnGameThread(Key, *TTypeToString<double>::ToString(JavaValue));
+}
+
 /** Implement platform specific static cleanup function */
 void FGenericCrashContext::CleanupPlatformSpecificFiles()
 {
