UTF-8 fixes for FCString.
Resubmission of CL# 18223063 with fixes.

diff --git a/Engine/Source/Runtime/Core/Private/Misc/CString.cpp b/Engine/Source/Runtime/Core/Private/Misc/CString.cpp
--- a/Engine/Source/Runtime/Core/Private/Misc/CString.cpp
+++ b/Engine/Source/Runtime/Core/Private/Misc/CString.cpp
@@ -13,11 +13,50 @@ CORE_API const ANSICHAR TCStringSpcHelper<ANSICHAR>::SpcArray[MAX_SPACES + 1] =
 	"                                                               ";
 
 template <>
-CORE_API const WIDECHAR TCStringSpcHelper<WIDECHAR>::SpcArray[MAX_SPACES + 1] =
-	TEXT("                                                                ")
-	TEXT("                                                                ")
-	TEXT("                                                                ")
-	TEXT("                                                               ");
+CORE_API const WIDECHAR TCStringSpcHelper<WIDECHAR>::SpcArray[MAX_SPACES + 1] = WIDETEXT(
+	"                                                                "
+	"                                                                "
+	"                                                                "
+	"                                                               ");
+
+template <>
+CORE_API const UTF8CHAR TCStringSpcHelper<UTF8CHAR>::SpcArray[MAX_SPACES + 1] = {
+	UTF8TEXT(' '), UTF8TEXT(' '), UTF8TEXT(' '), UTF8TEXT(' '), UTF8TEXT(' '), UTF8TEXT(' '), UTF8TEXT(' '), UTF8TEXT(' '),
+	UTF8TEXT(' '), UTF8TEXT(' '), UTF8TEXT(' '), UTF8TEXT(' '), UTF8TEXT(' '), UTF8TEXT(' '), UTF8TEXT(' '), UTF8TEXT(' '),
+	UTF8TEXT(' '), UTF8TEXT(' '), UTF8TEXT(' '), UTF8TEXT(' '), UTF8TEXT(' '), UTF8TEXT(' '), UTF8TEXT(' '), UTF8TEXT(' '),
+	UTF8TEXT(' '), UTF8TEXT(' '), UTF8TEXT(' '), UTF8TEXT(' '), UTF8TEXT(' '), UTF8TEXT(' '), UTF8TEXT(' '), UTF8TEXT(' '),
+	UTF8TEXT(' '), UTF8TEXT(' '), UTF8TEXT(' '), UTF8TEXT(' '), UTF8TEXT(' '), UTF8TEXT(' '), UTF8TEXT(' '), UTF8TEXT(' '),
+	UTF8TEXT(' '), UTF8TEXT(' '), UTF8TEXT(' '), UTF8TEXT(' '), UTF8TEXT(' '), UTF8TEXT(' '), UTF8TEXT(' '), UTF8TEXT(' '),
+	UTF8TEXT(' '), UTF8TEXT(' '), UTF8TEXT(' '), UTF8TEXT(' '), UTF8TEXT(' '), UTF8TEXT(' '), UTF8TEXT(' '), UTF8TEXT(' '),
+	UTF8TEXT(' '), UTF8TEXT(' '), UTF8TEXT(' '), UTF8TEXT(' '), UTF8TEXT(' '), UTF8TEXT(' '), UTF8TEXT(' '), UTF8TEXT(' '),
+
+	UTF8TEXT(' '), UTF8TEXT(' '), UTF8TEXT(' '), UTF8TEXT(' '), UTF8TEXT(' '), UTF8TEXT(' '), UTF8TEXT(' '), UTF8TEXT(' '),
+	UTF8TEXT(' '), UTF8TEXT(' '), UTF8TEXT(' '), UTF8TEXT(' '), UTF8TEXT(' '), UTF8TEXT(' '), UTF8TEXT(' '), UTF8TEXT(' '),
+	UTF8TEXT(' '), UTF8TEXT(' '), UTF8TEXT(' '), UTF8TEXT(' '), UTF8TEXT(' '), UTF8TEXT(' '), UTF8TEXT(' '), UTF8TEXT(' '),
+	UTF8TEXT(' '), UTF8TEXT(' '), UTF8TEXT(' '), UTF8TEXT(' '), UTF8TEXT(' '), UTF8TEXT(' '), UTF8TEXT(' '), UTF8TEXT(' '),
+	UTF8TEXT(' '), UTF8TEXT(' '), UTF8TEXT(' '), UTF8TEXT(' '), UTF8TEXT(' '), UTF8TEXT(' '), UTF8TEXT(' '), UTF8TEXT(' '),
+	UTF8TEXT(' '), UTF8TEXT(' '), UTF8TEXT(' '), UTF8TEXT(' '), UTF8TEXT(' '), UTF8TEXT(' '), UTF8TEXT(' '), UTF8TEXT(' '),
+	UTF8TEXT(' '), UTF8TEXT(' '), UTF8TEXT(' '), UTF8TEXT(' '), UTF8TEXT(' '), UTF8TEXT(' '), UTF8TEXT(' '), UTF8TEXT(' '),
+	UTF8TEXT(' '), UTF8TEXT(' '), UTF8TEXT(' '), UTF8TEXT(' '), UTF8TEXT(' '), UTF8TEXT(' '), UTF8TEXT(' '), UTF8TEXT(' '),
+
+	UTF8TEXT(' '), UTF8TEXT(' '), UTF8TEXT(' '), UTF8TEXT(' '), UTF8TEXT(' '), UTF8TEXT(' '), UTF8TEXT(' '), UTF8TEXT(' '),
+	UTF8TEXT(' '), UTF8TEXT(' '), UTF8TEXT(' '), UTF8TEXT(' '), UTF8TEXT(' '), UTF8TEXT(' '), UTF8TEXT(' '), UTF8TEXT(' '),
+	UTF8TEXT(' '), UTF8TEXT(' '), UTF8TEXT(' '), UTF8TEXT(' '), UTF8TEXT(' '), UTF8TEXT(' '), UTF8TEXT(' '), UTF8TEXT(' '),
+	UTF8TEXT(' '), UTF8TEXT(' '), UTF8TEXT(' '), UTF8TEXT(' '), UTF8TEXT(' '), UTF8TEXT(' '), UTF8TEXT(' '), UTF8TEXT(' '),
+	UTF8TEXT(' '), UTF8TEXT(' '), UTF8TEXT(' '), UTF8TEXT(' '), UTF8TEXT(' '), UTF8TEXT(' '), UTF8TEXT(' '), UTF8TEXT(' '),
+	UTF8TEXT(' '), UTF8TEXT(' '), UTF8TEXT(' '), UTF8TEXT(' '), UTF8TEXT(' '), UTF8TEXT(' '), UTF8TEXT(' '), UTF8TEXT(' '),
+	UTF8TEXT(' '), UTF8TEXT(' '), UTF8TEXT(' '), UTF8TEXT(' '), UTF8TEXT(' '), UTF8TEXT(' '), UTF8TEXT(' '), UTF8TEXT(' '),
+	UTF8TEXT(' '), UTF8TEXT(' '), UTF8TEXT(' '), UTF8TEXT(' '), UTF8TEXT(' '), UTF8TEXT(' '), UTF8TEXT(' '), UTF8TEXT(' '),
+
+	UTF8TEXT(' '), UTF8TEXT(' '), UTF8TEXT(' '), UTF8TEXT(' '), UTF8TEXT(' '), UTF8TEXT(' '), UTF8TEXT(' '), UTF8TEXT(' '),
+	UTF8TEXT(' '), UTF8TEXT(' '), UTF8TEXT(' '), UTF8TEXT(' '), UTF8TEXT(' '), UTF8TEXT(' '), UTF8TEXT(' '), UTF8TEXT(' '),
+	UTF8TEXT(' '), UTF8TEXT(' '), UTF8TEXT(' '), UTF8TEXT(' '), UTF8TEXT(' '), UTF8TEXT(' '), UTF8TEXT(' '), UTF8TEXT(' '),
+	UTF8TEXT(' '), UTF8TEXT(' '), UTF8TEXT(' '), UTF8TEXT(' '), UTF8TEXT(' '), UTF8TEXT(' '), UTF8TEXT(' '), UTF8TEXT(' '),
+	UTF8TEXT(' '), UTF8TEXT(' '), UTF8TEXT(' '), UTF8TEXT(' '), UTF8TEXT(' '), UTF8TEXT(' '), UTF8TEXT(' '), UTF8TEXT(' '),
+	UTF8TEXT(' '), UTF8TEXT(' '), UTF8TEXT(' '), UTF8TEXT(' '), UTF8TEXT(' '), UTF8TEXT(' '), UTF8TEXT(' '), UTF8TEXT(' '),
+	UTF8TEXT(' '), UTF8TEXT(' '), UTF8TEXT(' '), UTF8TEXT(' '), UTF8TEXT(' '), UTF8TEXT(' '), UTF8TEXT(' '), UTF8TEXT(' '),
+	UTF8TEXT(' '), UTF8TEXT(' '), UTF8TEXT(' '), UTF8TEXT(' '), UTF8TEXT(' '), UTF8TEXT(' '), UTF8TEXT(' '), UTF8TEXT('\0')
+};
 
 template <>
 CORE_API const ANSICHAR TCStringSpcHelper<ANSICHAR>::TabArray[MAX_TABS + 1] = 
@@ -27,19 +66,65 @@ CORE_API const ANSICHAR TCStringSpcHelper<ANSICHAR>::TabArray[MAX_TABS + 1] =
 	"\t\t\t\t\t\t\t\t\t\t\t\t\t\t\t\t\t\t\t\t\t\t\t\t\t\t\t\t\t\t\t\t\t\t\t\t\t\t\t\t\t\t\t\t\t\t\t\t\t\t\t\t\t\t\t\t\t\t\t\t\t\t\t";
 
 template <>
-CORE_API const WIDECHAR TCStringSpcHelper<WIDECHAR>::TabArray[MAX_TABS + 1] = 
-	TEXT("\t\t\t\t\t\t\t\t\t\t\t\t\t\t\t\t\t\t\t\t\t\t\t\t\t\t\t\t\t\t\t\t\t\t\t\t\t\t\t\t\t\t\t\t\t\t\t\t\t\t\t\t\t\t\t\t\t\t\t\t\t\t\t\t")
-	TEXT("\t\t\t\t\t\t\t\t\t\t\t\t\t\t\t\t\t\t\t\t\t\t\t\t\t\t\t\t\t\t\t\t\t\t\t\t\t\t\t\t\t\t\t\t\t\t\t\t\t\t\t\t\t\t\t\t\t\t\t\t\t\t\t\t")
-	TEXT("\t\t\t\t\t\t\t\t\t\t\t\t\t\t\t\t\t\t\t\t\t\t\t\t\t\t\t\t\t\t\t\t\t\t\t\t\t\t\t\t\t\t\t\t\t\t\t\t\t\t\t\t\t\t\t\t\t\t\t\t\t\t\t\t")
-	TEXT("\t\t\t\t\t\t\t\t\t\t\t\t\t\t\t\t\t\t\t\t\t\t\t\t\t\t\t\t\t\t\t\t\t\t\t\t\t\t\t\t\t\t\t\t\t\t\t\t\t\t\t\t\t\t\t\t\t\t\t\t\t\t\t");
+CORE_API const WIDECHAR TCStringSpcHelper<WIDECHAR>::TabArray[MAX_TABS + 1] = WIDETEXT(
+	"\t\t\t\t\t\t\t\t\t\t\t\t\t\t\t\t\t\t\t\t\t\t\t\t\t\t\t\t\t\t\t\t\t\t\t\t\t\t\t\t\t\t\t\t\t\t\t\t\t\t\t\t\t\t\t\t\t\t\t\t\t\t\t\t"
+	"\t\t\t\t\t\t\t\t\t\t\t\t\t\t\t\t\t\t\t\t\t\t\t\t\t\t\t\t\t\t\t\t\t\t\t\t\t\t\t\t\t\t\t\t\t\t\t\t\t\t\t\t\t\t\t\t\t\t\t\t\t\t\t\t"
+	"\t\t\t\t\t\t\t\t\t\t\t\t\t\t\t\t\t\t\t\t\t\t\t\t\t\t\t\t\t\t\t\t\t\t\t\t\t\t\t\t\t\t\t\t\t\t\t\t\t\t\t\t\t\t\t\t\t\t\t\t\t\t\t\t"
+	"\t\t\t\t\t\t\t\t\t\t\t\t\t\t\t\t\t\t\t\t\t\t\t\t\t\t\t\t\t\t\t\t\t\t\t\t\t\t\t\t\t\t\t\t\t\t\t\t\t\t\t\t\t\t\t\t\t\t\t\t\t\t\t");
+
+template <>
+CORE_API const UTF8CHAR TCStringSpcHelper<UTF8CHAR>::TabArray[MAX_TABS + 1] = {
+	UTF8TEXT('\t'), UTF8TEXT('\t'), UTF8TEXT('\t'), UTF8TEXT('\t'), UTF8TEXT('\t'), UTF8TEXT('\t'), UTF8TEXT('\t'), UTF8TEXT('\t'),
+	UTF8TEXT('\t'), UTF8TEXT('\t'), UTF8TEXT('\t'), UTF8TEXT('\t'), UTF8TEXT('\t'), UTF8TEXT('\t'), UTF8TEXT('\t'), UTF8TEXT('\t'),
+	UTF8TEXT('\t'), UTF8TEXT('\t'), UTF8TEXT('\t'), UTF8TEXT('\t'), UTF8TEXT('\t'), UTF8TEXT('\t'), UTF8TEXT('\t'), UTF8TEXT('\t'),
+	UTF8TEXT('\t'), UTF8TEXT('\t'), UTF8TEXT('\t'), UTF8TEXT('\t'), UTF8TEXT('\t'), UTF8TEXT('\t'), UTF8TEXT('\t'), UTF8TEXT('\t'),
+	UTF8TEXT('\t'), UTF8TEXT('\t'), UTF8TEXT('\t'), UTF8TEXT('\t'), UTF8TEXT('\t'), UTF8TEXT('\t'), UTF8TEXT('\t'), UTF8TEXT('\t'),
+	UTF8TEXT('\t'), UTF8TEXT('\t'), UTF8TEXT('\t'), UTF8TEXT('\t'), UTF8TEXT('\t'), UTF8TEXT('\t'), UTF8TEXT('\t'), UTF8TEXT('\t'),
+	UTF8TEXT('\t'), UTF8TEXT('\t'), UTF8TEXT('\t'), UTF8TEXT('\t'), UTF8TEXT('\t'), UTF8TEXT('\t'), UTF8TEXT('\t'), UTF8TEXT('\t'),
+	UTF8TEXT('\t'), UTF8TEXT('\t'), UTF8TEXT('\t'), UTF8TEXT('\t'), UTF8TEXT('\t'), UTF8TEXT('\t'), UTF8TEXT('\t'), UTF8TEXT('\t'),
+
+	UTF8TEXT('\t'), UTF8TEXT('\t'), UTF8TEXT('\t'), UTF8TEXT('\t'), UTF8TEXT('\t'), UTF8TEXT('\t'), UTF8TEXT('\t'), UTF8TEXT('\t'),
+	UTF8TEXT('\t'), UTF8TEXT('\t'), UTF8TEXT('\t'), UTF8TEXT('\t'), UTF8TEXT('\t'), UTF8TEXT('\t'), UTF8TEXT('\t'), UTF8TEXT('\t'),
+	UTF8TEXT('\t'), UTF8TEXT('\t'), UTF8TEXT('\t'), UTF8TEXT('\t'), UTF8TEXT('\t'), UTF8TEXT('\t'), UTF8TEXT('\t'), UTF8TEXT('\t'),
+	UTF8TEXT('\t'), UTF8TEXT('\t'), UTF8TEXT('\t'), UTF8TEXT('\t'), UTF8TEXT('\t'), UTF8TEXT('\t'), UTF8TEXT('\t'), UTF8TEXT('\t'),
+	UTF8TEXT('\t'), UTF8TEXT('\t'), UTF8TEXT('\t'), UTF8TEXT('\t'), UTF8TEXT('\t'), UTF8TEXT('\t'), UTF8TEXT('\t'), UTF8TEXT('\t'),
+	UTF8TEXT('\t'), UTF8TEXT('\t'), UTF8TEXT('\t'), UTF8TEXT('\t'), UTF8TEXT('\t'), UTF8TEXT('\t'), UTF8TEXT('\t'), UTF8TEXT('\t'),
+	UTF8TEXT('\t'), UTF8TEXT('\t'), UTF8TEXT('\t'), UTF8TEXT('\t'), UTF8TEXT('\t'), UTF8TEXT('\t'), UTF8TEXT('\t'), UTF8TEXT('\t'),
+	UTF8TEXT('\t'), UTF8TEXT('\t'), UTF8TEXT('\t'), UTF8TEXT('\t'), UTF8TEXT('\t'), UTF8TEXT('\t'), UTF8TEXT('\t'), UTF8TEXT('\t'),
+
+	UTF8TEXT('\t'), UTF8TEXT('\t'), UTF8TEXT('\t'), UTF8TEXT('\t'), UTF8TEXT('\t'), UTF8TEXT('\t'), UTF8TEXT('\t'), UTF8TEXT('\t'),
+	UTF8TEXT('\t'), UTF8TEXT('\t'), UTF8TEXT('\t'), UTF8TEXT('\t'), UTF8TEXT('\t'), UTF8TEXT('\t'), UTF8TEXT('\t'), UTF8TEXT('\t'),
+	UTF8TEXT('\t'), UTF8TEXT('\t'), UTF8TEXT('\t'), UTF8TEXT('\t'), UTF8TEXT('\t'), UTF8TEXT('\t'), UTF8TEXT('\t'), UTF8TEXT('\t'),
+	UTF8TEXT('\t'), UTF8TEXT('\t'), UTF8TEXT('\t'), UTF8TEXT('\t'), UTF8TEXT('\t'), UTF8TEXT('\t'), UTF8TEXT('\t'), UTF8TEXT('\t'),
+	UTF8TEXT('\t'), UTF8TEXT('\t'), UTF8TEXT('\t'), UTF8TEXT('\t'), UTF8TEXT('\t'), UTF8TEXT('\t'), UTF8TEXT('\t'), UTF8TEXT('\t'),
+	UTF8TEXT('\t'), UTF8TEXT('\t'), UTF8TEXT('\t'), UTF8TEXT('\t'), UTF8TEXT('\t'), UTF8TEXT('\t'), UTF8TEXT('\t'), UTF8TEXT('\t'),
+	UTF8TEXT('\t'), UTF8TEXT('\t'), UTF8TEXT('\t'), UTF8TEXT('\t'), UTF8TEXT('\t'), UTF8TEXT('\t'), UTF8TEXT('\t'), UTF8TEXT('\t'),
+	UTF8TEXT('\t'), UTF8TEXT('\t'), UTF8TEXT('\t'), UTF8TEXT('\t'), UTF8TEXT('\t'), UTF8TEXT('\t'), UTF8TEXT('\t'), UTF8TEXT('\t'),
+
+	UTF8TEXT('\t'), UTF8TEXT('\t'), UTF8TEXT('\t'), UTF8TEXT('\t'), UTF8TEXT('\t'), UTF8TEXT('\t'), UTF8TEXT('\t'), UTF8TEXT('\t'),
+	UTF8TEXT('\t'), UTF8TEXT('\t'), UTF8TEXT('\t'), UTF8TEXT('\t'), UTF8TEXT('\t'), UTF8TEXT('\t'), UTF8TEXT('\t'), UTF8TEXT('\t'),
+	UTF8TEXT('\t'), UTF8TEXT('\t'), UTF8TEXT('\t'), UTF8TEXT('\t'), UTF8TEXT('\t'), UTF8TEXT('\t'), UTF8TEXT('\t'), UTF8TEXT('\t'),
+	UTF8TEXT('\t'), UTF8TEXT('\t'), UTF8TEXT('\t'), UTF8TEXT('\t'), UTF8TEXT('\t'), UTF8TEXT('\t'), UTF8TEXT('\t'), UTF8TEXT('\t'),
+	UTF8TEXT('\t'), UTF8TEXT('\t'), UTF8TEXT('\t'), UTF8TEXT('\t'), UTF8TEXT('\t'), UTF8TEXT('\t'), UTF8TEXT('\t'), UTF8TEXT('\t'),
+	UTF8TEXT('\t'), UTF8TEXT('\t'), UTF8TEXT('\t'), UTF8TEXT('\t'), UTF8TEXT('\t'), UTF8TEXT('\t'), UTF8TEXT('\t'), UTF8TEXT('\t'),
+	UTF8TEXT('\t'), UTF8TEXT('\t'), UTF8TEXT('\t'), UTF8TEXT('\t'), UTF8TEXT('\t'), UTF8TEXT('\t'), UTF8TEXT('\t'), UTF8TEXT('\t'),
+	UTF8TEXT('\t'), UTF8TEXT('\t'), UTF8TEXT('\t'), UTF8TEXT('\t'), UTF8TEXT('\t'), UTF8TEXT('\t'), UTF8TEXT('\t'), UTF8TEXT('\0')
+};
 
 bool FToBoolHelper::FromCStringAnsi( const ANSICHAR* String )
 {
-	return FToBoolHelper::FromCStringWide( ANSI_TO_TCHAR(String) );
+#if PLATFORM_TCHAR_IS_UTF8CHAR
+	return FToBoolHelper::FromCStringUtf8( StringCast<UTF8CHAR>(String).Get() );
+#else
+	return FToBoolHelper::FromCStringWide( StringCast<WIDECHAR>(String).Get() );
+#endif
 }
 
 bool FToBoolHelper::FromCStringWide( const WIDECHAR* String )
 {
+#if PLATFORM_TCHAR_IS_UTF8CHAR
+	return FToBoolHelper::FromCStringUtf8( StringCast<UTF8CHAR>(String).Get() );
+#else
 	const FCoreTexts& CoreTexts = FCoreTexts::Get();
 
 	if (
@@ -64,4 +149,37 @@ bool FToBoolHelper::FromCStringWide( const WIDECHAR* String )
 	{
 		return FCStringWide::Atoi(String) ? true : false;
 	}
+#endif
+}
+
+bool FToBoolHelper::FromCStringUtf8( const UTF8CHAR* String )
+{
+#if PLATFORM_TCHAR_IS_UTF8CHAR
+	const FCoreTexts& CoreTexts = FCoreTexts::Get();
+
+	if (
+			FCStringUtf8::Stricmp(String, TEXT("True"))==0
+		||	FCStringUtf8::Stricmp(String, TEXT("Yes"))==0
+		||	FCStringUtf8::Stricmp(String, TEXT("On"))==0
+		||	FCStringUtf8::Stricmp(String, *(CoreTexts.True.ToString()))==0
+		||	FCStringUtf8::Stricmp(String, *(CoreTexts.Yes.ToString()))==0)
+	{
+		return true;
+	}
+	else if(
+			FCStringUtf8::Stricmp(String, TEXT("False"))==0
+		||	FCStringUtf8::Stricmp(String, TEXT("No"))==0
+		||	FCStringUtf8::Stricmp(String, TEXT("Off"))==0
+		||	FCStringUtf8::Stricmp(String, *(CoreTexts.False.ToString()))==0
+		||	FCStringUtf8::Stricmp(String, *(CoreTexts.No.ToString()))==0)
+	{
+		return false;
+	}
+	else
+	{
+		return FCStringUtf8::Atoi(String) ? true : false;
+	}
+#else
+	return FToBoolHelper::FromCStringWide( StringCast<WIDECHAR>(String).Get() );
+#endif
 }
diff --git a/Engine/Source/Runtime/Core/Public/HAL/Platform.h b/Engine/Source/Runtime/Core/Public/HAL/Platform.h
--- a/Engine/Source/Runtime/Core/Public/HAL/Platform.h
+++ b/Engine/Source/Runtime/Core/Public/HAL/Platform.h
@@ -251,9 +251,11 @@
 #ifndef PLATFORM_WCHAR_IS_4_BYTES
 	#define PLATFORM_WCHAR_IS_4_BYTES			0
 #endif
+// The PLATFORM_TCHAR_IS_CHAR16 macro really represents PLATFORM_WIDECHAR_IS_CHAR16 - it should be deprecated
 #ifndef PLATFORM_TCHAR_IS_CHAR16
 	#define PLATFORM_TCHAR_IS_CHAR16			0
 #endif
+#define PLATFORM_WIDECHAR_IS_CHAR16 PLATFORM_TCHAR_IS_CHAR16
 #ifndef PLATFORM_TCHAR_IS_UTF8CHAR
 	#define PLATFORM_TCHAR_IS_UTF8CHAR			0
 #endif
@@ -1071,6 +1073,12 @@ namespace TypeTests
 	static_assert((!TAreTypesEqual<WIDECHAR, UCS2CHAR>::Value),  "WIDECHAR and UCS2CHAR should be different types.");
 	static_assert(TAreTypesEqual<TCHAR, ANSICHAR>::Value || TAreTypesEqual<TCHAR, WIDECHAR>::Value || TAreTypesEqual<TCHAR, UTF8CHAR>::Value, "TCHAR should either be ANSICHAR, WIDECHAR or UTF8CHAR.");
 
+#if PLATFORM_WIDECHAR_IS_CHAR16
+	static_assert(TAreTypesEqual<WIDECHAR, char16_t>::Value, "WIDECHAR should be char16_t");
+#else
+	static_assert(TAreTypesEqual<WIDECHAR, wchar_t>::Value, "WIDECHAR should be wchar_t");
+#endif
+
 	static_assert(sizeof(uint8) == 1, "uint8 type size test failed.");
 	static_assert(int32(uint8(-1)) == 0xFF, "uint8 type sign test failed.");
 
@@ -1122,16 +1130,18 @@ namespace TypeTests
 
 #define UTF8TEXT_PASTE(x)  u8 ## x
 #define UTF16TEXT_PASTE(x) u ## x
-#define WIDETEXT_PASTE(x)  L ## x
+#if PLATFORM_WIDECHAR_IS_CHAR16
+	#define WIDETEXT_PASTE(x)  UTF16TEXT_PASTE(x)
+#else
+	#define WIDETEXT_PASTE(x)  L ## x
+#endif
 
 // If we don't have a platform-specific define for the TEXT macro, define it now.
 #if !defined(TEXT) && !UE_BUILD_DOCS
-	#if PLATFORM_TCHAR_IS_CHAR16
-		#define TEXT_PASTE(x) UTF16TEXT_PASTE(x)
-	#elif PLATFORM_TCHAR_IS_UTF8CHAR
+	#if PLATFORM_TCHAR_IS_UTF8CHAR
 		#define TEXT_PASTE(x) UTF8TEXT(x)
 	#else
-		#define TEXT_PASTE(x) WIDETEXT_PASTE(x)
+		#define TEXT_PASTE(x) WIDETEXT(x)
 	#endif
 	#define TEXT(x) TEXT_PASTE(x)
 #endif
diff --git a/Engine/Source/Runtime/Core/Public/Misc/CString.h b/Engine/Source/Runtime/Core/Public/Misc/CString.h
--- a/Engine/Source/Runtime/Core/Public/Misc/CString.h
+++ b/Engine/Source/Runtime/Core/Public/Misc/CString.h
@@ -44,6 +44,7 @@ struct CORE_API FToBoolHelper
 {
 	static bool FromCStringAnsi( const ANSICHAR* String );
 	static bool FromCStringWide( const WIDECHAR* String );
+	static bool FromCStringUtf8( const UTF8CHAR* String );
 };
 
 /**
@@ -505,7 +506,7 @@ const typename TCString<T>::CharType* TCString<T>::Strifind( const CharType* Str
 	}
 	
 	bool Alnum  = 0;
-	CharType f = ( *Find < LITERAL(CharType, 'a') || *Find > LITERAL(CharType, 'z') ) ? (*Find) : (*Find + LITERAL(CharType,'A') - LITERAL(CharType,'a'));
+	CharType f = ( *Find < LITERAL(CharType, 'a') || *Find > LITERAL(CharType, 'z') ) ? (*Find) : (CharType)(*Find + LITERAL(CharType,'A') - LITERAL(CharType,'a'));
 	int32 Length = Strlen(Find++)-1;
 	CharType c = *Str++;
 	
@@ -910,7 +911,7 @@ inline int32 TCString<WIDECHAR>::SnprintfImpl(CharType* Dest, int32 DestSize, co
 }
 
 template <> 
-FORCEINLINE bool TCString<TCHAR>::ToBool(const WIDECHAR* Str)
+FORCEINLINE bool TCString<WIDECHAR>::ToBool(const WIDECHAR* Str)
 {
 	return FToBoolHelper::FromCStringWide(Str);
 }
@@ -951,3 +952,9 @@ FORCEINLINE bool TCString<ANSICHAR>::ToBool(const ANSICHAR* Str)
 {
 	return FToBoolHelper::FromCStringAnsi(Str);
 }
+
+template <> 
+FORCEINLINE bool TCString<UTF8CHAR>::ToBool(const UTF8CHAR* Str)
+{
+	return FToBoolHelper::FromCStringUtf8(Str);
+}
