Testing adding a new file with a pragma message to ensure horde jobs have synced and built the file.

diff --git a/Engine/Source/Runtime/Core/Private/ProfilingDebugging/ResourceSize.cpp b/Engine/Source/Runtime/Core/Private/ProfilingDebugging/ResourceSize.cpp
new file mode 100644
--- /dev/null
+++ b/Engine/Source/Runtime/Core/Private/ProfilingDebugging/ResourceSize.cpp
@@ -0,0 +1,5 @@
+// Copyright Epic Games, Inc. All Rights Reserved.
+
+#include "ProfilingDebugging/ResourceSize.h"
+
+#pragma message("Compiling ResourceSize.cpp")
