Moving implementations of FResourceSize into the cpp.

diff --git a/Engine/Source/Runtime/Core/Private/ProfilingDebugging/ResourceSize.cpp b/Engine/Source/Runtime/Core/Private/ProfilingDebugging/ResourceSize.cpp
--- a/Engine/Source/Runtime/Core/Private/ProfilingDebugging/ResourceSize.cpp
+++ b/Engine/Source/Runtime/Core/Private/ProfilingDebugging/ResourceSize.cpp
@@ -3,3 +3,159 @@
 #include "ProfilingDebugging/ResourceSize.h"
 
 #pragma message("Compiling ResourceSize.cpp")
+
+FResourceSizeEx::FResourceSizeEx()
+	: ResourceSizeMode(EResourceSizeMode::Exclusive)
+{
+}
+
+
+FResourceSizeEx::FResourceSizeEx(const EResourceSizeMode::Type InResourceSizeMode)
+	: ResourceSizeMode(InResourceSizeMode)
+{
+}
+
+FResourceSizeEx::FResourceSizeEx(const EResourceSizeMode::Type InResourceSizeMode, const SIZE_T InDedicatedSystemMemoryBytes, const SIZE_T InDedicatedVideoMemoryBytes)
+	: ResourceSizeMode(InResourceSizeMode)
+{
+	DedicatedSystemMemoryBytesMap.Add(TEXT("Untracked Memory"), InDedicatedSystemMemoryBytes);
+	DedicatedVideoMemoryBytesMap.Add(TEXT("Untracked Memory"), InDedicatedVideoMemoryBytes);
+}
+
+FResourceSizeEx::FResourceSizeEx(const EResourceSizeMode::Type InResourceSizeMode, const SIZE_T InUnknownMemoryBytes)
+	: ResourceSizeMode(InResourceSizeMode)
+{
+	UnknownMemoryBytesMap.Add(TEXT("Untracked Memory"), InUnknownMemoryBytes);
+}
+
+void FResourceSizeEx::LogSummary(FOutputDevice& Ar) const
+{
+	auto PrintPair = [&Ar](const TPair<FName, SIZE_T>& Pair)
+	{
+		Ar.Logf(
+			TEXT("%140s %15.2f"),
+			*Pair.Key.ToString(),
+			static_cast<double>(Pair.Value) / 1024.0
+		);
+	};
+
+	for (const TPair<FName, SIZE_T>& Pair : DedicatedSystemMemoryBytesMap)
+	{
+		PrintPair(Pair);
+	}
+	for (const TPair<FName, SIZE_T>& Pair : DedicatedVideoMemoryBytesMap)
+	{
+		PrintPair(Pair);
+	}
+	for (const TPair<FName, SIZE_T>& Pair : UnknownMemoryBytesMap)
+	{
+		PrintPair(Pair);
+	}
+}
+
+EResourceSizeMode::Type FResourceSizeEx::GetResourceSizeMode() const
+{
+	return ResourceSizeMode;
+}
+
+FResourceSizeEx& FResourceSizeEx::AddDedicatedSystemMemoryBytes(const FName& Tag, const SIZE_T InMemoryBytes)
+{
+	SIZE_T& CurrentSize = DedicatedSystemMemoryBytesMap.FindOrAdd(Tag);
+	CurrentSize += InMemoryBytes;
+	return *this;
+}
+
+FResourceSizeEx& FResourceSizeEx::AddDedicatedSystemMemoryBytes(const SIZE_T InMemoryBytes)
+{
+	SIZE_T& CurrentSize = DedicatedSystemMemoryBytesMap.FindOrAdd(TEXT("Untracked Memory"));
+	CurrentSize += InMemoryBytes;
+	return *this;
+}
+
+SIZE_T FResourceSizeEx::GetDedicatedSystemMemoryBytes() const
+{
+	SIZE_T Sum = 0;
+	TArray<SIZE_T> Values;
+	DedicatedSystemMemoryBytesMap.GenerateValueArray(Values);
+	for (const SIZE_T s : Values) Sum += s;
+	return Sum;
+}
+
+FResourceSizeEx& FResourceSizeEx::AddDedicatedVideoMemoryBytes(const FName& Tag, const SIZE_T InMemoryBytes)
+{
+	SIZE_T& CurrentSize = DedicatedVideoMemoryBytesMap.FindOrAdd(Tag);
+	CurrentSize += InMemoryBytes;
+	return *this;
+}
+
+FResourceSizeEx& FResourceSizeEx::AddDedicatedVideoMemoryBytes(const SIZE_T InMemoryBytes)
+{
+	SIZE_T& CurrentSize = DedicatedVideoMemoryBytesMap.FindOrAdd(TEXT("Untracked Memory"));
+	CurrentSize += InMemoryBytes;
+	return *this;
+}
+
+SIZE_T FResourceSizeEx::GetDedicatedVideoMemoryBytes() const
+{
+	SIZE_T Sum = 0;
+	TArray<SIZE_T> Values;
+	DedicatedVideoMemoryBytesMap.GenerateValueArray(Values);
+	for (const SIZE_T s : Values) Sum += s;
+	return Sum;
+}
+
+FResourceSizeEx& FResourceSizeEx::AddUnknownMemoryBytes(const FName& Tag, const SIZE_T InMemoryBytes)
+{
+	SIZE_T& CurrentSize = UnknownMemoryBytesMap.FindOrAdd(Tag);
+	CurrentSize += InMemoryBytes;
+	return *this;
+}
+
+FResourceSizeEx& FResourceSizeEx::AddUnknownMemoryBytes(const SIZE_T InMemoryBytes)
+{
+	SIZE_T& CurrentSize = UnknownMemoryBytesMap.FindOrAdd(TEXT("Untracked Memory"));
+	CurrentSize += InMemoryBytes;
+	return *this;
+}
+
+SIZE_T FResourceSizeEx::GetUnknownMemoryBytes() const
+{
+	SIZE_T Sum = 0;
+	TArray<SIZE_T> Values;
+	UnknownMemoryBytesMap.GenerateValueArray(Values);
+	for (const SIZE_T s : Values) Sum += s;
+	return Sum;
+}
+
+SIZE_T FResourceSizeEx::GetTotalMemoryBytes() const
+{
+	return GetDedicatedSystemMemoryBytes() + GetDedicatedVideoMemoryBytes() + GetUnknownMemoryBytes();
+}
+
+FResourceSizeEx& FResourceSizeEx::operator+=(const FResourceSizeEx& InRHS)
+{
+	ensureAlwaysMsgf(ResourceSizeMode == InRHS.ResourceSizeMode, TEXT("The two resource sizes use different counting modes. The result of adding them together may be incorrect."));
+
+	for (const TPair<FName, SIZE_T>& Pair : InRHS.DedicatedSystemMemoryBytesMap)
+	{
+		DedicatedSystemMemoryBytesMap.FindOrAdd(Pair.Key) += Pair.Value;
+	}
+
+	for (const TPair<FName, SIZE_T>& Pair : InRHS.DedicatedVideoMemoryBytesMap)
+	{
+		DedicatedVideoMemoryBytesMap.FindOrAdd(Pair.Key) += Pair.Value;
+	}
+
+	for (const TPair<FName, SIZE_T>& Pair : InRHS.UnknownMemoryBytesMap)
+	{
+		UnknownMemoryBytesMap.FindOrAdd(Pair.Key) += Pair.Value;
+	}
+
+	return *this;
+}
+
+FResourceSizeEx operator+(FResourceSizeEx InLHS, const FResourceSizeEx& InRHS)
+{
+	InLHS += InRHS;
+	return InLHS;
+}
diff --git a/Engine/Source/Runtime/Core/Public/ProfilingDebugging/ResourceSize.h b/Engine/Source/Runtime/Core/Public/ProfilingDebugging/ResourceSize.h
--- a/Engine/Source/Runtime/Core/Public/ProfilingDebugging/ResourceSize.h
+++ b/Engine/Source/Runtime/Core/Public/ProfilingDebugging/ResourceSize.h
@@ -21,213 +21,94 @@ namespace EResourceSizeMode
  * Struct used to count up the amount of memory used by a resource.
  * This is typically used for assets via UObject::GetResourceSizeEx(...).
  */
-struct FResourceSizeEx
+struct CORE_API FResourceSizeEx
 {
 public:
 
 	/**
 	 * Default constructor. 
 	 */
-	explicit FResourceSizeEx()
-		: ResourceSizeMode(EResourceSizeMode::Exclusive)
-	{
-	}
+	explicit FResourceSizeEx();
 
 	/**
 	 * Construct using a given mode. 
 	 */
-	explicit FResourceSizeEx(const EResourceSizeMode::Type InResourceSizeMode)
-		: ResourceSizeMode(InResourceSizeMode)
-	{
-	}
+	explicit FResourceSizeEx(const EResourceSizeMode::Type InResourceSizeMode);
 
 	/**
 	 * Construct from known sizes. 
 	 */
-	FResourceSizeEx(const EResourceSizeMode::Type InResourceSizeMode, const SIZE_T InDedicatedSystemMemoryBytes, const SIZE_T InDedicatedVideoMemoryBytes)
-		: ResourceSizeMode(InResourceSizeMode)
-	{
-		DedicatedSystemMemoryBytesMap.Add(TEXT("Untracked Memory"), InDedicatedSystemMemoryBytes);
-		DedicatedVideoMemoryBytesMap.Add(TEXT("Untracked Memory"), InDedicatedVideoMemoryBytes);
-	}
+	FResourceSizeEx(const EResourceSizeMode::Type InResourceSizeMode, const SIZE_T InDedicatedSystemMemoryBytes, const SIZE_T InDedicatedVideoMemoryBytes);
 
 	/**
 	 * Construct from legacy unknown size.
 	 * Deliberately explicit to avoid accidental use.
 	 */
-	FResourceSizeEx(const EResourceSizeMode::Type InResourceSizeMode, const SIZE_T InUnknownMemoryBytes)
-		: ResourceSizeMode(InResourceSizeMode)
-	{
-		UnknownMemoryBytesMap.Add(TEXT("Untracked Memory"), InUnknownMemoryBytes);
-	}
+	FResourceSizeEx(const EResourceSizeMode::Type InResourceSizeMode, const SIZE_T InUnknownMemoryBytes);
 
-	void LogSummary(FOutputDevice& Ar) const
-	{
-		auto PrintPair = [&Ar](const TPair<FName, SIZE_T>& Pair)
-		{
-			Ar.Logf(
-				TEXT("%140s %15.2f"),
-				*Pair.Key.ToString(),
-				static_cast<double>(Pair.Value) / 1024.0
-			);
-		};
-
-		for (const TPair<FName, SIZE_T>& Pair : DedicatedSystemMemoryBytesMap)
-		{
-			PrintPair(Pair);
-		}
-		for (const TPair<FName, SIZE_T>& Pair : DedicatedVideoMemoryBytesMap)
-		{
-			PrintPair(Pair);
-		}
-		for (const TPair<FName, SIZE_T>& Pair : UnknownMemoryBytesMap)
-		{
-			PrintPair(Pair);
-		}
-	}
+	void LogSummary(FOutputDevice& Ar) const;
 
 	/**
 	 * Get the type of resource size held in this struct.
 	 */
-	EResourceSizeMode::Type GetResourceSizeMode() const
-	{
-		return ResourceSizeMode;
-	}
+	EResourceSizeMode::Type GetResourceSizeMode() const;
 
-	FResourceSizeEx& AddDedicatedSystemMemoryBytes(const FName& Tag, const SIZE_T InMemoryBytes)
-	{
-		SIZE_T& CurrentSize = DedicatedSystemMemoryBytesMap.FindOrAdd(Tag);
-		CurrentSize += InMemoryBytes;
-		return *this;
-	}
+	FResourceSizeEx& AddDedicatedSystemMemoryBytes(const FName& Tag, const SIZE_T InMemoryBytes);
 
 	/**
 	 * Add the given number of bytes to the dedicated system memory count.
 	 * @see DedicatedSystemMemoryBytes for a description of that memory type.
 	 */
-	FResourceSizeEx& AddDedicatedSystemMemoryBytes(const SIZE_T InMemoryBytes)
-	{
-		SIZE_T& CurrentSize = DedicatedSystemMemoryBytesMap.FindOrAdd(TEXT("Untracked Memory"));
-		CurrentSize += InMemoryBytes;
-		return *this;
-	}
+	FResourceSizeEx& AddDedicatedSystemMemoryBytes(const SIZE_T InMemoryBytes);
 
 	/**
 	 * Get the number of bytes allocated from dedicated system memory.
 	 * @see DedicatedSystemMemoryBytes for a description of that memory type.
 	 */
-	SIZE_T GetDedicatedSystemMemoryBytes() const
-	{
-		SIZE_T Sum = 0;
-		TArray<SIZE_T> Values;
-		DedicatedSystemMemoryBytesMap.GenerateValueArray(Values);
-		for (const SIZE_T s : Values) Sum += s;
-		return Sum;
-	}
-
-	FResourceSizeEx& AddDedicatedVideoMemoryBytes(const FName& Tag, const SIZE_T InMemoryBytes)
-	{
-		SIZE_T& CurrentSize = DedicatedVideoMemoryBytesMap.FindOrAdd(Tag);
-		CurrentSize += InMemoryBytes;
-		return *this;
-	}
+	SIZE_T GetDedicatedSystemMemoryBytes() const;
+
+	FResourceSizeEx& AddDedicatedVideoMemoryBytes(const FName& Tag, const SIZE_T InMemoryBytes);
 
 	/**
 	 * Add the given number of bytes to the dedicated video memory count.
 	 * @see DedicatedVideoMemoryBytes for a description of that memory type.
 	 */
-	FResourceSizeEx& AddDedicatedVideoMemoryBytes(const SIZE_T InMemoryBytes)
-	{
-		SIZE_T& CurrentSize = DedicatedVideoMemoryBytesMap.FindOrAdd(TEXT("Untracked Memory"));
-		CurrentSize += InMemoryBytes;
-		return *this;
-	}
+	FResourceSizeEx& AddDedicatedVideoMemoryBytes(const SIZE_T InMemoryBytes);
 
 	/**
 	 * Get the number of bytes allocated from dedicated video memory.
 	 * @see DedicatedVideoMemoryBytes for a description of that memory type.
 	 */
-	SIZE_T GetDedicatedVideoMemoryBytes() const
-	{
-		SIZE_T Sum = 0;
-		TArray<SIZE_T> Values;
-		DedicatedVideoMemoryBytesMap.GenerateValueArray(Values);
-		for (const SIZE_T s : Values) Sum += s;
-		return Sum;
-	}
-
-	FResourceSizeEx& AddUnknownMemoryBytes(const FName& Tag, const SIZE_T InMemoryBytes)
-	{
-		SIZE_T& CurrentSize = UnknownMemoryBytesMap.FindOrAdd(Tag);
-		CurrentSize += InMemoryBytes;
-		return *this;
-	}
+	SIZE_T GetDedicatedVideoMemoryBytes() const;
+
+	FResourceSizeEx& AddUnknownMemoryBytes(const FName& Tag, const SIZE_T InMemoryBytes);
 
 	/**
 	 * Add the given number of bytes to the unknown memory count.
 	 * @see UnknownMemoryBytes for a description of that memory type.
 	 */
-	FResourceSizeEx& AddUnknownMemoryBytes(const SIZE_T InMemoryBytes)
-	{
-		SIZE_T& CurrentSize = UnknownMemoryBytesMap.FindOrAdd(TEXT("Untracked Memory"));
-		CurrentSize += InMemoryBytes;
-		return *this;
-	}
+	FResourceSizeEx& AddUnknownMemoryBytes(const SIZE_T InMemoryBytes);
 
 	/**
 	 * Get the number of bytes allocated from unknown memory.
 	 * @see UnknownMemoryBytes for a description of that memory type.
 	 */
-	SIZE_T GetUnknownMemoryBytes() const
-	{
-		SIZE_T Sum = 0;
-		TArray<SIZE_T> Values;
-		UnknownMemoryBytesMap.GenerateValueArray(Values);
-		for (const SIZE_T s : Values) Sum += s;
-		return Sum;
-	}
+	SIZE_T GetUnknownMemoryBytes() const;
 
 	/**
 	 * Get the total number of bytes allocated from any memory.
 	 */
-	SIZE_T GetTotalMemoryBytes() const
-	{
-		return GetDedicatedSystemMemoryBytes() + GetDedicatedVideoMemoryBytes() + GetUnknownMemoryBytes();
-	}
+	SIZE_T GetTotalMemoryBytes() const;
 
 	/**
 	 * Add another FResourceSizeEx to this one.
 	 */
-	FResourceSizeEx& operator+=(const FResourceSizeEx& InRHS)
-	{
-		ensureAlwaysMsgf(ResourceSizeMode == InRHS.ResourceSizeMode, TEXT("The two resource sizes use different counting modes. The result of adding them together may be incorrect."));
-
-		for (const TPair<FName, SIZE_T>& Pair : InRHS.DedicatedSystemMemoryBytesMap)
-		{
-			DedicatedSystemMemoryBytesMap.FindOrAdd(Pair.Key) += Pair.Value;
-		}
-
-		for (const TPair<FName, SIZE_T>& Pair : InRHS.DedicatedVideoMemoryBytesMap)
-		{
-			DedicatedVideoMemoryBytesMap.FindOrAdd(Pair.Key) += Pair.Value;
-		}
-
-		for (const TPair<FName, SIZE_T>& Pair : InRHS.UnknownMemoryBytesMap)
-		{
-			UnknownMemoryBytesMap.FindOrAdd(Pair.Key) += Pair.Value;
-		}
-
-		return *this;
-	}
+	FResourceSizeEx& operator+=(const FResourceSizeEx& InRHS);
 
 	/**
 	 * Add two FResourceSizeEx instances together and return a copy.
 	 */
-	friend FResourceSizeEx operator+(FResourceSizeEx InLHS, const FResourceSizeEx& InRHS)
-	{
-		InLHS += InRHS;
-		return InLHS;
-	}
+	friend FResourceSizeEx operator+(FResourceSizeEx InLHS, const FResourceSizeEx& InRHS);
 
 private:
 	/**
