Remove compiler message that was used to debug a potential build issue.

diff --git a/Engine/Source/Runtime/Core/Private/ProfilingDebugging/ResourceSize.cpp b/Engine/Source/Runtime/Core/Private/ProfilingDebugging/ResourceSize.cpp
--- a/Engine/Source/Runtime/Core/Private/ProfilingDebugging/ResourceSize.cpp
+++ b/Engine/Source/Runtime/Core/Private/ProfilingDebugging/ResourceSize.cpp
@@ -2,8 +2,6 @@
 
 #include "ProfilingDebugging/ResourceSize.h"
 
-#pragma message("Compiling ResourceSize.cpp")
-
 FResourceSizeEx::FResourceSizeEx()
 	: ResourceSizeMode(EResourceSizeMode::Exclusive)
 {
