Check there's a GWarn device before trying to dump stats to it.

RenderDoc and FTaskTagScope() do not play well together with the latter raising an ensure(). This ensure fires very early, earlier than all of the engine's machinery being reporting and logging are fully initialised, making it somewhat brittle.

#rnx

diff --git a/Engine/Source/Runtime/Core/Private/Windows/WindowsPlatformCrashContext.cpp b/Engine/Source/Runtime/Core/Private/Windows/WindowsPlatformCrashContext.cpp
--- a/Engine/Source/Runtime/Core/Private/Windows/WindowsPlatformCrashContext.cpp
+++ b/Engine/Source/Runtime/Core/Private/Windows/WindowsPlatformCrashContext.cpp
@@ -1022,7 +1022,12 @@ int32 ReportCrashUsingCrashReportClient(FWindowsPlatformCrashContext& InContext,
 		if (!bCrashReporterRan)
 		{
 			UE_LOG(LogWindows, Log, TEXT("Could not start crash report client using %s"), CrashReporterClientPath);
-			FPlatformMemory::DumpStats(*GWarn);
+
+			if (GWarn != nullptr)
+			{
+				FPlatformMemory::DumpStats(*GWarn);
+			}
+
 			if (!bNoDialog)
 			{
 				FText MessageTitle(FText::Format(
