Fixed build for non-lwc template

diff --git a/Engine/Source/Runtime/Core/Public/Misc/LargeWorldRenderPosition.h b/Engine/Source/Runtime/Core/Public/Misc/LargeWorldRenderPosition.h
--- a/Engine/Source/Runtime/Core/Public/Misc/LargeWorldRenderPosition.h
+++ b/Engine/Source/Runtime/Core/Public/Misc/LargeWorldRenderPosition.h
@@ -41,7 +41,8 @@ private:
 #else // !UE_LARGE_WORLD_COORDINATES_DISABLED
 
 public:
-	static float GetTileSize() { return 0.0f; }
+	CORE_API static float GetTileSize() { return 0.0f; }
+	CORE_API static FVector3f GetTileFor(FVector InPosition) { return FVector3f::ZeroVector; };
 	float GetTile() const { return 0.0f; }
 	float GetOffset() const { return Offset; }
 	double GetTileAsDouble() const { return 0.0; }
