TEXTVIEW macros added to replace TEXT("abc"_SV) and some example uses of that.
Some other TEXT('a') and TCHAR[] fixes.

diff --git a/Engine/Source/Runtime/Core/Private/HAL/LowLevelMemTracker.cpp b/Engine/Source/Runtime/Core/Private/HAL/LowLevelMemTracker.cpp
--- a/Engine/Source/Runtime/Core/Private/HAL/LowLevelMemTracker.cpp
+++ b/Engine/Source/Runtime/Core/Private/HAL/LowLevelMemTracker.cpp
@@ -3860,7 +3860,7 @@ namespace LLMPrivate
 		{
 			WriteTag(TagData);
 		}
-		Write(TEXT("\n"_SV));
+		Write(TEXTVIEW("\n"));
 
 		WriteCount++;
 
diff --git a/Engine/Source/Runtime/Core/Private/Misc/PathViews.cpp b/Engine/Source/Runtime/Core/Private/Misc/PathViews.cpp
--- a/Engine/Source/Runtime/Core/Private/Misc/PathViews.cpp
+++ b/Engine/Source/Runtime/Core/Private/Misc/PathViews.cpp
@@ -234,7 +234,7 @@ FString FPathViews::ChangeExtension(const FStringView& InPath, const FStringView
 		const FStringView::SizeType Pos = FStringView::SizeType(PathEndPos - InPath.GetData());
 		const FStringView FileWithoutExtension = InPath.Left(Pos);
 
-		if (!InNewExtension.IsEmpty() && !InNewExtension.StartsWith('.'))
+		if (!InNewExtension.IsEmpty() && !InNewExtension.StartsWith(TEXT('.')))
 		{
 			// The new extension lacks a period so we need to add it ourselves.
 			FString Result(FileWithoutExtension, InNewExtension.Len() + 1);
@@ -385,7 +385,7 @@ bool FPathViews::IsRelativePath(FStringView InPath)
 			// Path starts with /; it may be either /Foo or //Foo
 			return false;
 		}
-		else if (InPath[FirstLen - 1] == ':')
+		else if (InPath[FirstLen - 1] == TEXT(':'))
 		{
 			// InPath == Volume:/SomethingOrNothing
 			return false;
@@ -423,7 +423,7 @@ void FPathViews::ToAbsolutePathInline(FStringView BasePath, FStringBuilderBase&
 
 	if (IsRelativePath(InOutPath))
 	{
-		if (BasePath.EndsWith('/') || BasePath.EndsWith('\\'))
+		if (BasePath.EndsWith(TEXT('/')) || BasePath.EndsWith(TEXT('\\')))
 		{
 			InOutPath.Prepend(BasePath);
 		}
@@ -472,18 +472,18 @@ void FPathViews::ToAbsolutePath(FStringView InPath, FStringBuilderBase& OutPath)
 
 void FPathViews::NormalizeFilename(FStringBuilderBase& InOutPath)
 {
-	Algo::Replace(MakeArrayView(InOutPath), '\\', '/');
+	Algo::Replace(MakeArrayView(InOutPath), TEXT('\\'), TEXT('/'));
 	FPlatformMisc::NormalizePath(InOutPath);
 }
 
 static bool ShouldRemoveTrailingSlash(FStringView Dir)
 {
-	return Dir.EndsWith('/') && !Dir.EndsWith(TEXT("//"_SV), ESearchCase::CaseSensitive) && !Dir.EndsWith(TEXT(":/"_SV), ESearchCase::CaseSensitive);
+	return Dir.EndsWith(TEXT('/')) && !Dir.EndsWith(TEXTVIEW("//"), ESearchCase::CaseSensitive) && !Dir.EndsWith(TEXTVIEW(":/"), ESearchCase::CaseSensitive);
 }
 
 void FPathViews::NormalizeDirectoryName(FStringBuilderBase& InOutPath)
 {
-	Algo::Replace(MakeArrayView(InOutPath), '\\', '/');
+	Algo::Replace(MakeArrayView(InOutPath), TEXT('\\'), TEXT('/'));
 	InOutPath.RemoveSuffix(ShouldRemoveTrailingSlash(InOutPath.ToView()) ? 1 : 0);
 	FPlatformMisc::NormalizePath(InOutPath);
 }
@@ -523,8 +523,8 @@ bool FPathViews::CollapseRelativeDirectories(FStringBuilderBase& InOutPath)
 		FStringView Path = InOutPath.ToView();
 
 		// Consider paths which start with .. or /.. as invalid
-		constexpr FStringView ParentDir = TEXT("/.."_SV);
-		if (Path.StartsWith(TEXT(".."_SV)) || Path.StartsWith(ParentDir))
+		const FStringView ParentDir = TEXTVIEW("/..");
+		if (Path.StartsWith(TEXTVIEW("..")) || Path.StartsWith(ParentDir))
 		{
 			return false;
 		}
@@ -545,14 +545,14 @@ bool FPathViews::CollapseRelativeDirectories(FStringBuilderBase& InOutPath)
 
 		// Find previous directory, stop if we've found a directory that isn't "/./"
 		int32 PreviousSeparatorIndex = Index;
-		while (Path.Left(PreviousSeparatorIndex).FindLastChar('/', /* Out */ PreviousSeparatorIndex) &&
-				Path.Mid(PreviousSeparatorIndex + 1, 2) == TEXT("./"_SV));
+		while (Path.Left(PreviousSeparatorIndex).FindLastChar(TEXT('/'), /* Out */ PreviousSeparatorIndex) &&
+				Path.Mid(PreviousSeparatorIndex + 1, 2) == TEXTVIEW("./"));
 
 		PreviousSeparatorIndex = FMath::Max(0, PreviousSeparatorIndex);
 		
 		// If we're attempting to remove the drive letter, that's illegal
 		int32 RemoveLen = Index - PreviousSeparatorIndex + ParentDir.Len();
-		if (Contains(Path.Mid(PreviousSeparatorIndex, RemoveLen), ':'))
+		if (Contains(Path.Mid(PreviousSeparatorIndex, RemoveLen), TEXT(':')))
 		{
 			return false;
 		}
@@ -560,14 +560,14 @@ bool FPathViews::CollapseRelativeDirectories(FStringBuilderBase& InOutPath)
 		InOutPath.RemoveAt(PreviousSeparatorIndex, RemoveLen);
 	}
 
-	RemoveAll(InOutPath, TEXT("./"_SV));
+	RemoveAll(InOutPath, TEXTVIEW("./"));
 
 	return true;
 }
 
 void FPathViews::RemoveDuplicateSlashes(FStringBuilderBase& InOutPath)
 {
-	int32 DoubleSlashIdx = UE::String::FindFirst(InOutPath.ToView(), TEXT("//"_SV));
+	int32 DoubleSlashIdx = UE::String::FindFirst(InOutPath.ToView(), TEXTVIEW("//"));
 	if (DoubleSlashIdx == INDEX_NONE)
 	{
 		return;
@@ -575,15 +575,15 @@ void FPathViews::RemoveDuplicateSlashes(FStringBuilderBase& InOutPath)
 
 	TCHAR* WriteIt = InOutPath.GetData() + DoubleSlashIdx + 1;
 	TCHAR* ReadIt = WriteIt + 1;
-	for (TCHAR* End = InOutPath.GetData() + InOutPath.Len(), LastChar = '/'; ReadIt != End; LastChar = *ReadIt++)
+	for (TCHAR* End = InOutPath.GetData() + InOutPath.Len(), LastChar = TEXT('/'); ReadIt != End; LastChar = *ReadIt++)
 	{
-		if ((*ReadIt != '/') | (LastChar != '/'))
+		if ((*ReadIt != TEXT('/')) | (LastChar != TEXT('/')))
 		{
 			*WriteIt++ = *ReadIt;
 		}
 	}
 
-	*WriteIt = '\0';
+	*WriteIt = TEXT('\0');
 	InOutPath.RemoveSuffix(int32(ReadIt - WriteIt));
 }
 
@@ -622,7 +622,7 @@ void FPathViews::SplitFirstComponent(FStringView InPath, FStringView& OutFirstCo
 				OutRemainder = InPath.RightChop(1);
 			}
 		}
-		else if (InPath[FirstLen - 1] == ':')
+		else if (InPath[FirstLen - 1] == TEXT(':'))
 		{
 			// InPath == Volume:/SomethingOrNothing
 			// FirstComponent = Volume:/
diff --git a/Engine/Source/Runtime/Core/Private/Modules/ModuleManifest.cpp b/Engine/Source/Runtime/Core/Private/Modules/ModuleManifest.cpp
--- a/Engine/Source/Runtime/Core/Private/Modules/ModuleManifest.cpp
+++ b/Engine/Source/Runtime/Core/Private/Modules/ModuleManifest.cpp
@@ -57,14 +57,14 @@ bool FModuleManifest::TryRead(const FString& FileName, FModuleManifest& OutManif
 			return false;
 		}
 
-		if (Field.ToView() == TEXT("BuildId"_SV))
+		if (Field.ToView() == TEXTVIEW("BuildId"))
 		{
 			if (!FSimpleParse::ParseString(Ptr, OutManifest.BuildId))
 			{
 				return false;
 			}
 		}
-		else if (Field.ToView() == TEXT("Modules"_SV))
+		else if (Field.ToView() == TEXTVIEW("Modules"))
 		{
 			if (!FSimpleParse::MatchZeroOrMoreWhitespace(Ptr) || !FSimpleParse::MatchChar(Ptr, TEXT('{')) || !FSimpleParse::MatchZeroOrMoreWhitespace(Ptr))
 			{
diff --git a/Engine/Source/Runtime/Core/Private/Tests/Containers/StringViewTest.cpp b/Engine/Source/Runtime/Core/Private/Tests/Containers/StringViewTest.cpp
--- a/Engine/Source/Runtime/Core/Private/Tests/Containers/StringViewTest.cpp
+++ b/Engine/Source/Runtime/Core/Private/Tests/Containers/StringViewTest.cpp
@@ -99,9 +99,9 @@ bool FStringViewTestCtor::RunTest(const FString& Parameters)
 
 	// Create using string view literals
 	{
-		FStringView View = TEXT("Test"_SV);
+		FStringView View = TEXTVIEW("Test");
 		FAnsiStringView ViewAnsi = "Test"_ASV;
-		FWideStringView ViewWide = TEXT("Test"_WSV);
+		FWideStringView ViewWide = WIDETEXTVIEW("Test");
 	}
 
 	return true;
@@ -722,14 +722,14 @@ bool FStringViewTestMid::RunTest(const FString& Parameters)
 IMPLEMENT_SIMPLE_AUTOMATION_TEST(FStringViewTestTrimStartAndEnd, TEST_NAME_ROOT ".TrimStartAndEnd", TestFlags)
 bool FStringViewTestTrimStartAndEnd::RunTest(const FString& Parameters)
 {
-	TestTrue(TEXT("FStringView::TrimStartAndEnd(\"\")"), TEXT(""_SV).TrimStartAndEnd().IsEmpty());
-	TestTrue(TEXT("FStringView::TrimStartAndEnd(\" \")"), TEXT(" "_SV).TrimStartAndEnd().IsEmpty());
-	TestTrue(TEXT("FStringView::TrimStartAndEnd(\"  \")"), TEXT("  "_SV).TrimStartAndEnd().IsEmpty());
-	TestTrue(TEXT("FStringView::TrimStartAndEnd(\" \\t\\r\\n\")"), TEXT(" \t\r\n"_SV).TrimStartAndEnd().IsEmpty());
+	TestTrue(TEXT("FStringView::TrimStartAndEnd(\"\")"), TEXTVIEW("").TrimStartAndEnd().IsEmpty());
+	TestTrue(TEXT("FStringView::TrimStartAndEnd(\" \")"), TEXTVIEW(" ").TrimStartAndEnd().IsEmpty());
+	TestTrue(TEXT("FStringView::TrimStartAndEnd(\"  \")"), TEXTVIEW("  ").TrimStartAndEnd().IsEmpty());
+	TestTrue(TEXT("FStringView::TrimStartAndEnd(\" \\t\\r\\n\")"), TEXTVIEW(" \t\r\n").TrimStartAndEnd().IsEmpty());
 
-	TestEqual(TEXT("FStringView::TrimStartAndEnd(\"ABC123\")"), TEXT("ABC123"_SV).TrimStartAndEnd(), TEXT("ABC123"_SV));
-	TestEqual(TEXT("FStringView::TrimStartAndEnd(\"A \\t\\r\\nB\")"), TEXT("A \t\r\nB"_SV).TrimStartAndEnd(), TEXT("A \t\r\nB"_SV));
-	TestEqual(TEXT("FStringView::TrimStartAndEnd(\" \\t\\r\\nABC123\\n\\r\\t \")"), TEXT(" \t\r\nABC123\n\r\t "_SV).TrimStartAndEnd(), TEXT("ABC123"_SV));
+	TestEqual(TEXT("FStringView::TrimStartAndEnd(\"ABC123\")"), TEXTVIEW("ABC123").TrimStartAndEnd(), TEXTVIEW("ABC123"));
+	TestEqual(TEXT("FStringView::TrimStartAndEnd(\"A \\t\\r\\nB\")"), TEXTVIEW("A \t\r\nB").TrimStartAndEnd(), TEXTVIEW("A \t\r\nB"));
+	TestEqual(TEXT("FStringView::TrimStartAndEnd(\" \\t\\r\\nABC123\\n\\r\\t \")"), TEXTVIEW(" \t\r\nABC123\n\r\t ").TrimStartAndEnd(), TEXTVIEW("ABC123"));
 
 	return true;
 }
@@ -737,14 +737,14 @@ bool FStringViewTestTrimStartAndEnd::RunTest(const FString& Parameters)
 IMPLEMENT_SIMPLE_AUTOMATION_TEST(FStringViewTestTrimStart, TEST_NAME_ROOT ".TrimStart", TestFlags)
 bool FStringViewTestTrimStart::RunTest(const FString& Parameters)
 {
-	TestTrue(TEXT("FStringView::TrimStart(\"\")"), TEXT(""_SV).TrimStart().IsEmpty());
-	TestTrue(TEXT("FStringView::TrimStart(\" \")"), TEXT(" "_SV).TrimStart().IsEmpty());
-	TestTrue(TEXT("FStringView::TrimStart(\"  \")"), TEXT("  "_SV).TrimStart().IsEmpty());
-	TestTrue(TEXT("FStringView::TrimStart(\" \\t\\r\\n\")"), TEXT(" \t\r\n"_SV).TrimStart().IsEmpty());
+	TestTrue(TEXT("FStringView::TrimStart(\"\")"), TEXTVIEW("").TrimStart().IsEmpty());
+	TestTrue(TEXT("FStringView::TrimStart(\" \")"), TEXTVIEW(" ").TrimStart().IsEmpty());
+	TestTrue(TEXT("FStringView::TrimStart(\"  \")"), TEXTVIEW("  ").TrimStart().IsEmpty());
+	TestTrue(TEXT("FStringView::TrimStart(\" \\t\\r\\n\")"), TEXTVIEW(" \t\r\n").TrimStart().IsEmpty());
 
-	TestEqual(TEXT("FStringView::TrimStart(\"ABC123\")"), TEXT("ABC123"_SV).TrimStart(), TEXT("ABC123"_SV));
-	TestEqual(TEXT("FStringView::TrimStart(\"A \\t\\r\\nB\")"), TEXT("A \t\r\nB"_SV).TrimStart(), TEXT("A \t\r\nB"_SV));
-	TestEqual(TEXT("FStringView::TrimStart(\" \\t\\r\\nABC123\\n\\r\\t \")"), TEXT(" \t\r\nABC123\n\r\t "_SV).TrimStart(), TEXT("ABC123\n\r\t "_SV));
+	TestEqual(TEXT("FStringView::TrimStart(\"ABC123\")"), TEXTVIEW("ABC123").TrimStart(), TEXTVIEW("ABC123"));
+	TestEqual(TEXT("FStringView::TrimStart(\"A \\t\\r\\nB\")"), TEXTVIEW("A \t\r\nB").TrimStart(), TEXTVIEW("A \t\r\nB"));
+	TestEqual(TEXT("FStringView::TrimStart(\" \\t\\r\\nABC123\\n\\r\\t \")"), TEXTVIEW(" \t\r\nABC123\n\r\t ").TrimStart(), TEXTVIEW("ABC123\n\r\t "));
 
 	return true;
 }
@@ -752,14 +752,14 @@ bool FStringViewTestTrimStart::RunTest(const FString& Parameters)
 IMPLEMENT_SIMPLE_AUTOMATION_TEST(FStringViewTestTrimEnd, TEST_NAME_ROOT ".TrimEnd", TestFlags)
 bool FStringViewTestTrimEnd::RunTest(const FString& Parameters)
 {
-	TestTrue(TEXT("FStringView::TrimEnd(\"\")"), TEXT(""_SV).TrimEnd().IsEmpty());
-	TestTrue(TEXT("FStringView::TrimEnd(\" \")"), TEXT(" "_SV).TrimEnd().IsEmpty());
-	TestTrue(TEXT("FStringView::TrimEnd(\"  \")"), TEXT("  "_SV).TrimEnd().IsEmpty());
-	TestTrue(TEXT("FStringView::TrimEnd(\" \\t\\r\\n\")"), TEXT(" \t\r\n"_SV).TrimEnd().IsEmpty());
-
-	TestEqual(TEXT("FStringView::TrimEnd(\"ABC123\")"), TEXT("ABC123"_SV).TrimEnd(), TEXT("ABC123"_SV));
-	TestEqual(TEXT("FStringView::TrimEnd(\"A \\t\\r\\nB\")"), TEXT("A \t\r\nB"_SV).TrimEnd(), TEXT("A \t\r\nB"_SV));
-	TestEqual(TEXT("FStringView::TrimEnd(\" \\t\\r\\nABC123\\n\\r\\t \")"), TEXT(" \t\r\nABC123\n\r\t "_SV).TrimEnd(), TEXT(" \t\r\nABC123"_SV));
+	TestTrue(TEXT("FStringView::TrimEnd(\"\")"), TEXTVIEW("").TrimEnd().IsEmpty());
+	TestTrue(TEXT("FStringView::TrimEnd(\" \")"), TEXTVIEW(" ").TrimEnd().IsEmpty());
+	TestTrue(TEXT("FStringView::TrimEnd(\"  \")"), TEXTVIEW("  ").TrimEnd().IsEmpty());
+	TestTrue(TEXT("FStringView::TrimEnd(\" \\t\\r\\n\")"), TEXTVIEW(" \t\r\n").TrimEnd().IsEmpty());
+
+	TestEqual(TEXT("FStringView::TrimEnd(\"ABC123\")"), TEXTVIEW("ABC123").TrimEnd(), TEXTVIEW("ABC123"));
+	TestEqual(TEXT("FStringView::TrimEnd(\"A \\t\\r\\nB\")"), TEXTVIEW("A \t\r\nB").TrimEnd(), TEXTVIEW("A \t\r\nB"));
+	TestEqual(TEXT("FStringView::TrimEnd(\" \\t\\r\\nABC123\\n\\r\\t \")"), TEXTVIEW(" \t\r\nABC123\n\r\t ").TrimEnd(), TEXTVIEW(" \t\r\nABC123"));
 
 	return true;
 }
diff --git a/Engine/Source/Runtime/Core/Private/Tests/Hash/Blake3Test.cpp b/Engine/Source/Runtime/Core/Private/Tests/Hash/Blake3Test.cpp
--- a/Engine/Source/Runtime/Core/Private/Tests/Hash/Blake3Test.cpp
+++ b/Engine/Source/Runtime/Core/Private/Tests/Hash/Blake3Test.cpp
@@ -10,37 +10,38 @@
 IMPLEMENT_SIMPLE_AUTOMATION_TEST(FBlake3Test, "System.Core.Misc.Blake3", EAutomationTestFlags::ApplicationContextMask | EAutomationTestFlags::SmokeFilter)
 bool FBlake3Test::RunTest(const FString& Parameters)
 {
+	constexpr uint32 MaxInputLength = 31'744u;
+
 	struct FBlake3TestCase
 	{
 		uint32 InputLength;
 		FStringView Hash;
 	};
-	constexpr FBlake3TestCase TestCases[] =
+	const FBlake3TestCase TestCases[] =
 	{
-		{     0u, TEXT("af1349b9f5f9a1a6a0404dea36dcc9499bcb25c9adc112b7cc9a93cae41f3262"_SV)},
-		{     1u, TEXT("2d3adedff11b61f14c886e35afa036736dcd87a74d27b5c1510225d0f592e213"_SV)},
-		{ 1'023u, TEXT("10108970eeda3eb932baac1428c7a2163b0e924c9a9e25b35bba72b28f70bd11"_SV)},
-		{ 1'024u, TEXT("42214739f095a406f3fc83deb889744ac00df831c10daa55189b5d121c855af7"_SV)},
-		{ 1'025u, TEXT("d00278ae47eb27b34faecf67b4fe263f82d5412916c1ffd97c8cb7fb814b8444"_SV)},
-		{ 2'048u, TEXT("e776b6028c7cd22a4d0ba182a8bf62205d2ef576467e838ed6f2529b85fba24a"_SV)},
-		{ 2'049u, TEXT("5f4d72f40d7a5f82b15ca2b2e44b1de3c2ef86c426c95c1af0b6879522563030"_SV)},
-		{ 3'072u, TEXT("b98cb0ff3623be03326b373de6b9095218513e64f1ee2edd2525c7ad1e5cffd2"_SV)},
-		{ 3'073u, TEXT("7124b49501012f81cc7f11ca069ec9226cecb8a2c850cfe644e327d22d3e1cd3"_SV)},
-		{ 4'096u, TEXT("015094013f57a5277b59d8475c0501042c0b642e531b0a1c8f58d2163229e969"_SV)},
-		{ 4'097u, TEXT("9b4052b38f1c5fc8b1f9ff7ac7b27cd242487b3d890d15c96a1c25b8aa0fb995"_SV)},
-		{ 5'120u, TEXT("9cadc15fed8b5d854562b26a9536d9707cadeda9b143978f319ab34230535833"_SV)},
-		{ 5'121u, TEXT("628bd2cb2004694adaab7bbd778a25df25c47b9d4155a55f8fbd79f2fe154cff"_SV)},
-		{ 6'144u, TEXT("3e2e5b74e048f3add6d21faab3f83aa44d3b2278afb83b80b3c35164ebeca205"_SV)},
-		{ 6'145u, TEXT("f1323a8631446cc50536a9f705ee5cb619424d46887f3c376c695b70e0f0507f"_SV)},
-		{ 7'168u, TEXT("61da957ec2499a95d6b8023e2b0e604ec7f6b50e80a9678b89d2628e99ada77a"_SV)},
-		{ 7'169u, TEXT("a003fc7a51754a9b3c7fae0367ab3d782dccf28855a03d435f8cfe74605e7817"_SV)},
-		{ 8'192u, TEXT("aae792484c8efe4f19e2ca7d371d8c467ffb10748d8a5a1ae579948f718a2a63"_SV)},
-		{ 8'193u, TEXT("bab6c09cb8ce8cf459261398d2e7aef35700bf488116ceb94a36d0f5f1b7bc3b"_SV)},
-		{16'384u, TEXT("f875d6646de28985646f34ee13be9a576fd515f76b5b0a26bb324735041ddde4"_SV)},
-		{31'744u, TEXT("62b6960e1a44bcc1eb1a611a8d6235b6b4b78f32e7abc4fb4c6cdcce94895c47"_SV)},
+		{     0u,        TEXTVIEW("af1349b9f5f9a1a6a0404dea36dcc9499bcb25c9adc112b7cc9a93cae41f3262")},
+		{     1u,        TEXTVIEW("2d3adedff11b61f14c886e35afa036736dcd87a74d27b5c1510225d0f592e213")},
+		{ 1'023u,        TEXTVIEW("10108970eeda3eb932baac1428c7a2163b0e924c9a9e25b35bba72b28f70bd11")},
+		{ 1'024u,        TEXTVIEW("42214739f095a406f3fc83deb889744ac00df831c10daa55189b5d121c855af7")},
+		{ 1'025u,        TEXTVIEW("d00278ae47eb27b34faecf67b4fe263f82d5412916c1ffd97c8cb7fb814b8444")},
+		{ 2'048u,        TEXTVIEW("e776b6028c7cd22a4d0ba182a8bf62205d2ef576467e838ed6f2529b85fba24a")},
+		{ 2'049u,        TEXTVIEW("5f4d72f40d7a5f82b15ca2b2e44b1de3c2ef86c426c95c1af0b6879522563030")},
+		{ 3'072u,        TEXTVIEW("b98cb0ff3623be03326b373de6b9095218513e64f1ee2edd2525c7ad1e5cffd2")},
+		{ 3'073u,        TEXTVIEW("7124b49501012f81cc7f11ca069ec9226cecb8a2c850cfe644e327d22d3e1cd3")},
+		{ 4'096u,        TEXTVIEW("015094013f57a5277b59d8475c0501042c0b642e531b0a1c8f58d2163229e969")},
+		{ 4'097u,        TEXTVIEW("9b4052b38f1c5fc8b1f9ff7ac7b27cd242487b3d890d15c96a1c25b8aa0fb995")},
+		{ 5'120u,        TEXTVIEW("9cadc15fed8b5d854562b26a9536d9707cadeda9b143978f319ab34230535833")},
+		{ 5'121u,        TEXTVIEW("628bd2cb2004694adaab7bbd778a25df25c47b9d4155a55f8fbd79f2fe154cff")},
+		{ 6'144u,        TEXTVIEW("3e2e5b74e048f3add6d21faab3f83aa44d3b2278afb83b80b3c35164ebeca205")},
+		{ 6'145u,        TEXTVIEW("f1323a8631446cc50536a9f705ee5cb619424d46887f3c376c695b70e0f0507f")},
+		{ 7'168u,        TEXTVIEW("61da957ec2499a95d6b8023e2b0e604ec7f6b50e80a9678b89d2628e99ada77a")},
+		{ 7'169u,        TEXTVIEW("a003fc7a51754a9b3c7fae0367ab3d782dccf28855a03d435f8cfe74605e7817")},
+		{ 8'192u,        TEXTVIEW("aae792484c8efe4f19e2ca7d371d8c467ffb10748d8a5a1ae579948f718a2a63")},
+		{ 8'193u,        TEXTVIEW("bab6c09cb8ce8cf459261398d2e7aef35700bf488116ceb94a36d0f5f1b7bc3b")},
+		{16'384u,        TEXTVIEW("f875d6646de28985646f34ee13be9a576fd515f76b5b0a26bb324735041ddde4")},
+		{MaxInputLength, TEXTVIEW("62b6960e1a44bcc1eb1a611a8d6235b6b4b78f32e7abc4fb4c6cdcce94895c47")},
 	};
 
-	constexpr uint32 MaxInputLength = TestCases[UE_ARRAY_COUNT(TestCases) - 1].InputLength;
 	uint8 Input[MaxInputLength];
 	for (uint32 Index = 0; Index < MaxInputLength; ++Index)
 	{
diff --git a/Engine/Source/Runtime/Core/Private/Tests/Misc/PathTests.h b/Engine/Source/Runtime/Core/Private/Tests/Misc/PathTests.h
--- a/Engine/Source/Runtime/Core/Private/Tests/Misc/PathTests.h
+++ b/Engine/Source/Runtime/Core/Private/Tests/Misc/PathTests.h
@@ -126,21 +126,9 @@ struct FTestPair
 	FStringView Expected;
 };
 
-constexpr FStringView BaseDir = TEXT("/root"_SV);
+extern const FStringView BaseDir;
 
-constexpr FTestPair ExpectedRelativeToAbsolutePaths[] =
-{
-	{ TEXT(""_SV),					TEXT("/root/"_SV) },
-	{ TEXT("dir"_SV),				TEXT("/root/dir"_SV) },
-	{ TEXT("/groot"_SV),			TEXT("/groot"_SV) },
-	{ TEXT("/groot/"_SV),			TEXT("/groot/"_SV) },
-	{ TEXT("/r/dir"_SV),			TEXT("/r/dir"_SV) },
-	{ TEXT("/r/dir"_SV),			TEXT("/r/dir"_SV) },
-	{ TEXT("C:\\"_SV),				TEXT("C:/"_SV) },
-	{ TEXT("C:\\A\\B"_SV),			TEXT("C:/A/B"_SV) },
-	{ TEXT("a/b/../c"_SV),			TEXT("/root/a/c"_SV) },
-	{ TEXT("/a/b/../c"_SV),			TEXT("/a/c"_SV) },
-};
+extern const FTestPair ExpectedRelativeToAbsolutePaths[10];
 
 }
 
diff --git a/Engine/Source/Runtime/Core/Private/Tests/Misc/PathViewsTest.cpp b/Engine/Source/Runtime/Core/Private/Tests/Misc/PathViewsTest.cpp
--- a/Engine/Source/Runtime/Core/Private/Tests/Misc/PathViewsTest.cpp
+++ b/Engine/Source/Runtime/Core/Private/Tests/Misc/PathViewsTest.cpp
@@ -298,39 +298,39 @@ bool FPathViewsAppendTest::RunTest(const FString& InParameters)
 	TStringBuilder<256> Path;
 
 	FPathViews::Append(Path, TEXT("A"), TEXT(""));
-	TestEqual(TEXT("FPathViews::Append('A', '')"), FStringView(Path), TEXT("A/"_SV));
+	TestEqual(TEXT("FPathViews::Append('A', '')"), FStringView(Path), TEXTVIEW("A/"));
 	Path.Reset();
 
 	FPathViews::Append(Path, TEXT(""), TEXT("B"));
-	TestEqual(TEXT("FPathViews::Append('', 'B')"), FStringView(Path), TEXT("B"_SV));
+	TestEqual(TEXT("FPathViews::Append('', 'B')"), FStringView(Path), TEXTVIEW("B"));
 	Path.Reset();
 
 	FPathViews::Append(Path, TEXT("/"), TEXT("B"));
-	TestEqual(TEXT("FPathViews::Append('/', 'B')"), FStringView(Path), TEXT("/B"_SV));
+	TestEqual(TEXT("FPathViews::Append('/', 'B')"), FStringView(Path), TEXTVIEW("/B"));
 	Path.Reset();
 
 	FPathViews::Append(Path, TEXT("A"), TEXT("B"));
-	TestEqual(TEXT("FPathViews::Append('A', 'B')"), FStringView(Path), TEXT("A/B"_SV));
+	TestEqual(TEXT("FPathViews::Append('A', 'B')"), FStringView(Path), TEXTVIEW("A/B"));
 	Path.Reset();
 
 	FPathViews::Append(Path, TEXT("A/"), TEXT("B"));
-	TestEqual(TEXT("FPathViews::Append('A/', 'B')"), FStringView(Path), TEXT("A/B"_SV));
+	TestEqual(TEXT("FPathViews::Append('A/', 'B')"), FStringView(Path), TEXTVIEW("A/B"));
 	Path.Reset();
 
 	FPathViews::Append(Path, TEXT("A\\"), TEXT("B"));
-	TestEqual(TEXT("FPathViews::Append('A\\', 'B')"), FStringView(Path), TEXT("A\\B"_SV));
+	TestEqual(TEXT("FPathViews::Append('A\\', 'B')"), FStringView(Path), TEXTVIEW("A\\B"));
 	Path.Reset();
 
 	FPathViews::Append(Path, TEXT("A/B"), TEXT("C/D"));
-	TestEqual(TEXT("FPathViews::Append('A/B', 'C/D')"), FStringView(Path), TEXT("A/B/C/D"_SV));
+	TestEqual(TEXT("FPathViews::Append('A/B', 'C/D')"), FStringView(Path), TEXTVIEW("A/B/C/D"));
 	Path.Reset();
 
 	FPathViews::Append(Path, TEXT("A/"), TEXT("B"), TEXT("C/"), TEXT("D"));
-	TestEqual(TEXT("FPathViews::Append('A/', 'B', 'C/', 'D')"), FStringView(Path), TEXT("A/B/C/D"_SV));
+	TestEqual(TEXT("FPathViews::Append('A/', 'B', 'C/', 'D')"), FStringView(Path), TEXTVIEW("A/B/C/D"));
 	Path.Reset();
 
 	FPathViews::Append(Path, TEXT("A/"), 16, TEXT("B"));
-	TestEqual(TEXT("FPathViews::Append('A/', 16, 'B')"), FStringView(Path), TEXT("A/16/B"_SV));
+	TestEqual(TEXT("FPathViews::Append('A/', 16, 'B')"), FStringView(Path), TEXTVIEW("A/16/B"));
 	Path.Reset();
 
 	return true;
@@ -724,7 +724,7 @@ bool FPathViewsToAbsoluteTest::RunTest(const FString& InParameters)
 		FPathViews::ToAbsolutePathInline(BaseDir, ActualInline);	
 		TestEqual(TEXT("ToAbsolutePathInline"), ActualInline.ToView(), Pair.Expected);
 		
-		constexpr FStringView Original = TEXT("\\\\la/./.././la////"_SV);
+		const FStringView Original = TEXTVIEW("\\\\la/./.././la////");
 		TStringBuilder<64> ActualNondestructive;
 		ActualNondestructive << Original;
 		FPathViews::ToAbsolutePath(BaseDir, Pair.Input, ActualNondestructive);
diff --git a/Engine/Source/Runtime/Core/Private/Tests/Misc/PathsTest.cpp b/Engine/Source/Runtime/Core/Private/Tests/Misc/PathsTest.cpp
--- a/Engine/Source/Runtime/Core/Private/Tests/Misc/PathsTest.cpp
+++ b/Engine/Source/Runtime/Core/Private/Tests/Misc/PathsTest.cpp
@@ -9,6 +9,22 @@
 
 #if WITH_DEV_AUTOMATION_TESTS
 
+const FStringView PathTest::BaseDir = TEXTVIEW("/root");
+
+const PathTest::FTestPair PathTest::ExpectedRelativeToAbsolutePaths[10] =
+{
+	{ TEXTVIEW(""),					TEXTVIEW("/root/") },
+	{ TEXTVIEW("dir"),				TEXTVIEW("/root/dir") },
+	{ TEXTVIEW("/groot"),			TEXTVIEW("/groot") },
+	{ TEXTVIEW("/groot/"),			TEXTVIEW("/groot/") },
+	{ TEXTVIEW("/r/dir"),			TEXTVIEW("/r/dir") },
+	{ TEXTVIEW("/r/dir"),			TEXTVIEW("/r/dir") },
+	{ TEXTVIEW("C:\\"),				TEXTVIEW("C:/") },
+	{ TEXTVIEW("C:\\A\\B"),			TEXTVIEW("C:/A/B") },
+	{ TEXTVIEW("a/b/../c"),			TEXTVIEW("/root/a/c") },
+	{ TEXTVIEW("/a/b/../c"),		TEXTVIEW("/a/c") },
+};
+
 IMPLEMENT_SIMPLE_AUTOMATION_TEST(FPathTests, "System.Core.Misc.Paths", EAutomationTestFlags::ApplicationContextMask | EAutomationTestFlags::SmokeFilter)
 
 bool FPathTests::RunTest( const FString& Parameters )
diff --git a/Engine/Source/Runtime/Core/Private/Tests/Misc/StringBuilderTest.cpp b/Engine/Source/Runtime/Core/Private/Tests/Misc/StringBuilderTest.cpp
--- a/Engine/Source/Runtime/Core/Private/Tests/Misc/StringBuilderTest.cpp
+++ b/Engine/Source/Runtime/Core/Private/Tests/Misc/StringBuilderTest.cpp
@@ -47,11 +47,11 @@ bool FStringBuilderTestAppendString::RunTest(const FString& Parameters)
 		TStringBuilder<7> Builder;
 		Builder << TEXT('A') << TEXT('B') << TEXT('C');
 		Builder << 'D' << 'E' << 'F';
-		TestEqual(TEXT("Append Char"), FStringView(Builder), TEXT("ABCDEF"_SV));
+		TestEqual(TEXT("Append Char"), FStringView(Builder), TEXTVIEW("ABCDEF"));
 
 		TAnsiStringBuilder<4> AnsiBuilder;
 		AnsiBuilder << 'A' << 'B' << 'C';
-		TestEqual(TEXT("Append AnsiChar"), FAnsiStringView(AnsiBuilder), "ABC"_ASV);
+		TestEqual(TEXT("Append AnsiChar"), FAnsiStringView(AnsiBuilder), ANSITEXTVIEW("ABC"));
 	}
 
 	// Append C String
@@ -59,19 +59,19 @@ bool FStringBuilderTestAppendString::RunTest(const FString& Parameters)
 		TStringBuilder<7> Builder;
 		Builder << TEXT("ABC");
 		Builder << "DEF";
-		TestEqual(TEXT("Append C String"), FStringView(Builder), TEXT("ABCDEF"_SV));
+		TestEqual(TEXT("Append C String"), FStringView(Builder), TEXTVIEW("ABCDEF"));
 
 		TAnsiStringBuilder<4> AnsiBuilder;
 		AnsiBuilder << "ABC";
-		TestEqual(TEXT("Append Ansi C String"), FAnsiStringView(AnsiBuilder), "ABC"_ASV);
+		TestEqual(TEXT("Append Ansi C String"), FAnsiStringView(AnsiBuilder), ANSITEXTVIEW("ABC"));
 	}
 
 	// Append FStringView
 	{
 		TStringBuilder<7> Builder;
-		Builder << TEXT("ABC"_SV);
+		Builder << TEXTVIEW("ABC");
 		Builder << "DEF"_ASV;
-		TestEqual(TEXT("Append FStringView"), FStringView(Builder), TEXT("ABCDEF"_SV));
+		TestEqual(TEXT("Append FStringView"), FStringView(Builder), TEXTVIEW("ABCDEF"));
 
 		TAnsiStringBuilder<4> AnsiBuilder;
 		AnsiBuilder << "ABC"_ASV;
@@ -84,7 +84,7 @@ bool FStringBuilderTestAppendString::RunTest(const FString& Parameters)
 		Builder << TEXT("ABC");
 		TStringBuilder<4> BuilderCopy;
 		BuilderCopy << Builder;
-		TestEqual(TEXT("Append FStringBuilderBase"), FStringView(BuilderCopy), TEXT("ABC"_SV));
+		TestEqual(TEXT("Append FStringBuilderBase"), FStringView(BuilderCopy), TEXTVIEW("ABC"));
 
 		TAnsiStringBuilder<4> AnsiBuilder;
 		AnsiBuilder << "ABC";
@@ -97,15 +97,15 @@ bool FStringBuilderTestAppendString::RunTest(const FString& Parameters)
 	{
 		TStringBuilder<4> Builder;
 		Builder << FString(TEXT("ABC"));
-		TestEqual(TEXT("Append FString"), FStringView(Builder), TEXT("ABC"_SV));
+		TestEqual(TEXT("Append FString"), FStringView(Builder), TEXTVIEW("ABC"));
 	}
 
 	// Append Char Array
 	{
-		const TCHAR String[16] = TEXT("ABC");
+		const auto& String = TEXT("ABC");
 		TStringBuilder<4> Builder;
 		Builder << String;
-		TestEqual(TEXT("Append Char Array"), FStringView(Builder), TEXT("ABC"_SV));
+		TestEqual(TEXT("Append Char Array"), FStringView(Builder), TEXTVIEW("ABC"));
 
 		const ANSICHAR AnsiString[16] = "ABC";
 		TAnsiStringBuilder<4> AnsiBuilder;
diff --git a/Engine/Source/Runtime/Core/Public/Containers/StringView.h b/Engine/Source/Runtime/Core/Public/Containers/StringView.h
--- a/Engine/Source/Runtime/Core/Public/Containers/StringView.h
+++ b/Engine/Source/Runtime/Core/Public/Containers/StringView.h
@@ -379,10 +379,17 @@ constexpr inline auto GetNum(TStringView<CharType> String)
 
 //////////////////////////////////////////////////////////////////////////
 
+#if PLATFORM_TCHAR_IS_UTF8CHAR
+inline FStringView operator "" _SV(const ANSICHAR* String, size_t Size)
+{
+	return FStringView((const UTF8CHAR*)String, Size);
+}
+#else
 constexpr inline FStringView operator "" _SV(const TCHAR* String, size_t Size)
 {
 	return FStringView(String, Size);
 }
+#endif
 
 constexpr inline FAnsiStringView operator "" _ASV(const ANSICHAR* String, size_t Size)
 {
@@ -401,6 +408,15 @@ constexpr inline FWideStringView operator "" _WSV(const WIDECHAR* String, size_t
 	return FUtf8StringView(reinterpret_cast<const UTF8CHAR*>(String), Size);
 }
 
+#if PLATFORM_TCHAR_IS_UTF8CHAR
+	#define TEXTVIEW(str) (str##_SV)
+#else
+	#define TEXTVIEW(str) TEXT(str##_SV)
+#endif
+#define ANSITEXTVIEW(str) (str##_ASV)
+#define WIDETEXTVIEW(str) PREPROCESSOR_JOIN(WIDETEXT(str), _WSV)
+#define UTF8TEXTVIEW(str) (str##_U8SV)
+
 //////////////////////////////////////////////////////////////////////////
 
 /** Case insensitive string hash function. */
