StringBuilder: Changed Append and operator<< to accept a greater variety of character types

- Append now accepts any character type and converts it if necessary.
- AppendAnsi and AppendUtf8 will eventually be deprecated.
- AppendChar has been added and will eventually replace Append(CharType).
- MakeStringView and TStringView CTAD have been added to support Append(CharRangeType).
- WriteToString has been extended with other character types.
- TElementType has been made SFINAE-safe using void_t as suggested by [REDACTED_PERSON].

diff --git a/Engine/Source/Runtime/Core/Private/Tests/Containers/StringViewTest.cpp b/Engine/Source/Runtime/Core/Private/Tests/Containers/StringViewTest.cpp
--- a/Engine/Source/Runtime/Core/Private/Tests/Containers/StringViewTest.cpp
+++ b/Engine/Source/Runtime/Core/Private/Tests/Containers/StringViewTest.cpp
@@ -81,7 +81,7 @@ bool FStringViewTestCtor::RunTest(const FString& Parameters)
 		TestFalse(TEXT("View.IsEmpty"), View.IsEmpty());
 	}
 
-	// Create from a ansi literal
+	// Create from an ansi literal
 	{
 		FAnsiStringView View("Test Ctor");
 		TestEqual(TEXT("View length"), View.Len(), FCStringAnsi::Strlen("Test Ctor"));
@@ -101,7 +101,28 @@ bool FStringViewTestCtor::RunTest(const FString& Parameters)
 	{
 		FStringView View = TEXTVIEW("Test");
 		FAnsiStringView ViewAnsi = "Test"_ASV;
-		FWideStringView ViewWide = WIDETEXTVIEW("Test");
+		FWideStringView ViewWide = WIDETEXT("Test"_WSV);
+	}
+
+	// Verify that class template argument deduction is working
+	{
+		TStringView ViewAnsi((ANSICHAR*)"Test");
+		TStringView ViewWide((WIDECHAR*)WIDETEXT("Test"));
+		TStringView ViewUtf8((UTF8CHAR*)UTF8TEXT("Test"));
+	}
+	{
+		TStringView ViewAnsi((const ANSICHAR*)"Test");
+		TStringView ViewWide((const WIDECHAR*)WIDETEXT("Test"));
+		TStringView ViewUtf8((const UTF8CHAR*)UTF8TEXT("Test"));
+	}
+	{
+		TStringView ViewAnsi(WriteToAnsiString<16>("Test"));
+		TStringView ViewWide(WriteToWideString<16>(WIDETEXT("Test")));
+		TStringView ViewUtf8(WriteToUtf8String<16>(UTF8TEXT("Test")));
+	}
+	{
+		FString String(TEXT("Test"));
+		TStringView ViewString(String);
 	}
 
 	return true;
diff --git a/Engine/Source/Runtime/Core/Public/Containers/StringView.h b/Engine/Source/Runtime/Core/Public/Containers/StringView.h
--- a/Engine/Source/Runtime/Core/Public/Containers/StringView.h
+++ b/Engine/Source/Runtime/Core/Public/Containers/StringView.h
@@ -10,12 +10,11 @@
 #include "Misc/CString.h"
 #include "Templates/AndOrNot.h"
 #include "Templates/ChooseClass.h"
-#include "Templates/Decay.h"
 #include "Templates/EnableIf.h"
-#include "Templates/IsArray.h"
-#include "Templates/RemoveCV.h"
 #include "Templates/UnrealTemplate.h"
+#include "Traits/ElementType.h"
 #include "Traits/IsContiguousContainer.h"
+#include <type_traits>
 
 namespace StringViewPrivate
 {
@@ -52,12 +51,6 @@ namespace StringViewPrivate
 			typename TChooseClass<TIsConvertibleFromTo<T, FUtf8StringView>::Value, FUtf8StringView,
 			NotCompatible>::Result>::Result>::Result;
 	};
-
-	template <typename CharRangeType, typename ElementType>
-	struct TIsCompatibleRangeType
-	{
-		static constexpr bool Value = FPlatformString::IsCharEncodingCompatibleWith<typename TRemoveCV<typename TRemovePointer<decltype(StringViewPrivate::WrapGetData(DeclVal<CharRangeType&>()))>::Type>::Type, ElementType>();
-	};
 }
 
 /**
@@ -114,87 +107,57 @@ public:
 	using SizeType = int32;
 	using ViewType = TStringView<CharType>;
 
-private:
-	/** Trait testing whether a type is a contiguous range of CharType, and not CharType[]. */
-	template <typename CharRangeType>
-	using TIsCharRange = TAnd<
-		TIsContiguousContainer<CharRangeType>,
-		TNot<TIsArray<typename TRemoveReference<CharRangeType>::Type>>,
-		StringViewPrivate::TIsCompatibleRangeType<CharRangeType, ElementType>>;
-
 public:
 	/** Construct an empty view. */
 	constexpr TStringView() = default;
 
-	/**
-	 * Construct a view of the null-terminated string pointed to by InData.
-	 *
-	 * The caller is responsible for ensuring that the provided character range remains valid for the lifetime of the view.
-	 */
-	template <typename OtherCharType,
-		typename TEnableIf<
-			FPlatformString::IsCharEncodingCompatibleWith<OtherCharType, CharType>()
-		>::Type* = nullptr>
-	constexpr inline TStringView(const OtherCharType* InData)
-		: DataPtr((const CharType*)InData)
-		, Size(InData ? TCString<CharType>::Strlen((const CharType*)InData) : 0)
-	{
-	}
+	/** Construct a view of the null-terminated string pointed to by InData. */
 	constexpr inline TStringView(const CharType* InData)
 		: DataPtr(InData)
 		, Size(InData ? TCString<CharType>::Strlen(InData) : 0)
 	{
 	}
 
-	/**
-	 * Construct a view of InSize characters beginning at InData.
-	 *
-	 * The caller is responsible for ensuring that the provided character range remains valid for the lifetime of the view.
-	 */
+	/** Construct a view of InSize characters beginning at InData. */
+	constexpr inline TStringView(const CharType* InData, SizeType InSize)
+		: DataPtr(InData)
+		, Size(InSize)
+	{
+	}
+
+	/** Construct a view of the null-terminated string pointed to by InData. */
 	template <typename OtherCharType,
-		typename TEnableIf<
-			FPlatformString::IsCharEncodingCompatibleWith<OtherCharType, CharType>()
-		>::Type* = nullptr>
-	constexpr inline TStringView(const OtherCharType* InData, SizeType InSize)
+		std::enable_if_t<FPlatformString::IsCharEncodingCompatibleWith<OtherCharType, CharType>()>* = nullptr>
+	constexpr inline TStringView(const OtherCharType* InData)
 		: DataPtr((const CharType*)InData)
-		, Size(InSize)
+		, Size(InData ? TCString<CharType>::Strlen((const CharType*)InData) : 0)
 	{
 	}
-	constexpr inline TStringView(const CharType* InData, SizeType InSize)
-		: DataPtr(InData)
+
+	/** Construct a view of InSize characters beginning at InData. */
+	template <typename OtherCharType,
+		std::enable_if_t<FPlatformString::IsCharEncodingCompatibleWith<OtherCharType, CharType>()>* = nullptr>
+	constexpr inline TStringView(const OtherCharType* InData, SizeType InSize)
+		: DataPtr((const CharType*)InData)
 		, Size(InSize)
 	{
 	}
 
-	/**
-	 * Construct a view from a contiguous range of characters.
-	 *
-	 * The caller is responsible for ensuring that the provided character range remains valid for the lifetime of the view.
-	 */
+	/** Construct a view from a contiguous range of characters, such as FString or TStringBuilder. */
 	template <typename CharRangeType,
-		typename TEnableIf<TAnd<
-			TNot<TIsSame<ViewType, typename TDecay<CharRangeType>::Type>>,
-			TIsCharRange<CharRangeType>
-		>::Value>::Type* = nullptr>
-	constexpr inline TStringView(CharRangeType&& InRange)
+		std::enable_if_t<
+			TIsContiguousContainer<CharRangeType>::Value &&
+			TIsCharType<TElementType_T<CharRangeType>>::Value &&
+			FPlatformString::IsCharEncodingCompatibleWith<TElementType_T<CharRangeType>, CharType>() &&
+			!std::is_array_v<std::remove_reference_t<CharRangeType>> &&
+			!std::is_same_v<CharRangeType, ViewType>
+		>* = nullptr>
+	constexpr inline TStringView(const CharRangeType& InRange)
 		: DataPtr((const CharType*)StringViewPrivate::WrapGetData(InRange))
 		, Size(static_cast<SizeType>(GetNum(InRange)))
 	{
 	}
 
-	/**
-	 * Construct a view from a compatible view.
-	 */
-	template <typename OtherCharType,
-		typename TEnableIf<
-			!TIsSame<CharType, OtherCharType>::Value &&
-			FPlatformString::IsCharEncodingCompatibleWith<OtherCharType, CharType>()>::Type* = nullptr>
-	constexpr inline TStringView(TStringView<OtherCharType> InView)
-		: DataPtr(reinterpret_cast<const CharType*>(InView.GetData()))
-		, Size(InView.Len())
-	{
-	}
-
 	/** Access the character at the given index in the view. */
 	inline const CharType& operator[](SizeType Index) const;
 
@@ -371,6 +334,22 @@ protected:
 	SizeType Size = 0;
 };
 
+template <typename CharRangeType>
+TStringView(CharRangeType&& Range)
+	-> TStringView<TElementType_T<CharRangeType>>;
+
+template <typename CharPtrOrRangeType>
+constexpr inline auto MakeStringView(CharPtrOrRangeType&& CharPtrOrRange) -> decltype(TStringView(Forward<CharPtrOrRangeType>(CharPtrOrRange)))
+{
+	return TStringView(Forward<CharPtrOrRangeType>(CharPtrOrRange));
+}
+
+template <typename CharPtrType, typename SizeType>
+constexpr inline auto MakeStringView(CharPtrType&& CharPtr, SizeType&& Size) -> decltype(TStringView(Forward<CharPtrType>(CharPtr), Forward<SizeType>(Size)))
+{
+	return TStringView(Forward<CharPtrType>(CharPtr), Forward<SizeType>(Size));
+}
+
 template <typename CharType>
 constexpr inline auto GetNum(TStringView<CharType> String)
 {
diff --git a/Engine/Source/Runtime/Core/Public/Misc/StringBuilder.h b/Engine/Source/Runtime/Core/Public/Misc/StringBuilder.h
--- a/Engine/Source/Runtime/Core/Public/Misc/StringBuilder.h
+++ b/Engine/Source/Runtime/Core/Public/Misc/StringBuilder.h
@@ -11,6 +11,7 @@
 #include "Templates/IsArrayOrRefOfType.h"
 #include "Templates/IsValidVariadicFunctionArg.h"
 #include "Templates/UnrealTemplate.h"
+#include "Traits/IsCharType.h"
 #include "Traits/IsContiguousContainer.h"
 #include <type_traits>
 
@@ -141,95 +142,37 @@ public:
 		CurPos -= InCount;
 	}
 
-	inline BuilderType& Append(CharType Char)
+	template <typename OtherCharType,
+		std::enable_if_t<TIsCharType<OtherCharType>::Value>* = nullptr>
+	inline BuilderType& Append(const OtherCharType* const String, const int32 Length)
 	{
-		EnsureAdditionalCapacity(1);
-
-		*CurPos++ = Char;
-
+		int32 ConvertedLength = FPlatformString::ConvertedLength<CharType>(String, Length);
+		EnsureAdditionalCapacity(ConvertedLength);
+		CurPos = FPlatformString::Convert(CurPos, ConvertedLength, String, Length);
 		return *this;
 	}
 
-	inline BuilderType& AppendAnsi(const ANSICHAR* NulTerminatedString)
+	template <typename CharRangeType>
+	inline auto Append(CharRangeType&& Range) -> decltype(Append(MakeStringView(Forward<CharRangeType>(Range)).GetData(), int32(0)))
 	{
-		if (!NulTerminatedString)
-		{
-			return *this;
-		}
-
-		return AppendAnsi(NulTerminatedString, TCString<ANSICHAR>::Strlen(NulTerminatedString));
+		const TStringView View = MakeStringView(Forward<CharRangeType>(Range));
+		return Append(View.GetData(), View.Len());
 	}
 
-	inline BuilderType& AppendAnsi(const FAnsiStringView& AnsiString)
+	inline BuilderType& AppendChar(CharType Char)
 	{
-		return AppendAnsi(AnsiString.GetData(), AnsiString.Len());
-	}
-
-	inline BuilderType& AppendAnsi(const ANSICHAR* String, const int32 Length)
-	{
-		EnsureAdditionalCapacity(Length);
-
-		CharType* RESTRICT Dest = CurPos;
-		CurPos += Length;
-
-		for (int32 i = 0; i < Length; ++i)
-		{
-			Dest[i] = CharType(String[i]);
-		}
-
-		return *this;
-	}
-
-	inline BuilderType& AppendUtf8(const UTF8CHAR* NulTerminatedString)
-	{
-		if (!NulTerminatedString)
-		{
-			return *this;
-		}
-
-		return AppendUtf8(NulTerminatedString, TCString<UTF8CHAR>::Strlen(NulTerminatedString));
-	}
-
-	inline BuilderType& AppendUtf8(const FUtf8StringView& AnsiString)
-	{
-		return AppendUtf8(AnsiString.GetData(), AnsiString.Len());
-	}
-
-	inline BuilderType& AppendUtf8(const UTF8CHAR* String, const int32 Length)
-	{
-		int32 ConvertedBufferLength = FPlatformString::ConvertedLength<ElementType>(String, Length);
-		EnsureAdditionalCapacity(ConvertedBufferLength);
-
-		CurPos = FPlatformString::Convert(CurPos, ConvertedBufferLength, String, Length);
-
+		EnsureAdditionalCapacity(1);
+		*CurPos++ = Char;
 		return *this;
 	}
 
-	inline BuilderType& Append(const CharType* NulTerminatedString)
-	{
-		if (!NulTerminatedString)
-		{
-			return *this;
-		}
-
-		return Append(NulTerminatedString, TCString<CharType>::Strlen(NulTerminatedString));
-	}
-
-	inline BuilderType& Append(const ViewType& StringView)
-	{
-		return Append(StringView.GetData(), StringView.Len());
-	}
-
-	inline BuilderType& Append(const CharType* String, int32 Length)
-	{
-		EnsureAdditionalCapacity(Length);
-		CharType* RESTRICT Dest = CurPos;
-		CurPos += Length;
-
-		FMemory::Memcpy(Dest, String, Length * sizeof(CharType));
-
-		return *this;
-	}
+	inline BuilderType& Append(CharType Char) { return AppendChar(Char); }
+	inline BuilderType& AppendAnsi(const FAnsiStringView String) { return Append(String); }
+	inline BuilderType& AppendUtf8(const FUtf8StringView String) { return Append(String); }
+	inline BuilderType& AppendAnsi(const ANSICHAR* const String) { return Append(String); }
+	inline BuilderType& AppendUtf8(const UTF8CHAR* const String) { return Append(String); }
+	inline BuilderType& AppendAnsi(const ANSICHAR* const String, const int32 Length) { return Append(String, Length); }
+	inline BuilderType& AppendUtf8(const UTF8CHAR* const String, const int32 Length) { return Append(String, Length); }
 
 	/** Replace characters at given position and length with substring */
 	void ReplaceAt(int32 Pos, int32 RemoveLen, ViewType Str)
@@ -292,7 +235,7 @@ public:
 	 * @return The builder, to allow additional operations to be composed with this one.
 	 */
 	template <typename RangeType, typename DelimiterType,
-		typename = typename TEnableIf<TAnd<TCanAppendRange<RangeType&&>, TCanAppend<DelimiterType&&>>::Value>::Type>
+		std::enable_if_t<TAnd<TCanAppendRange<RangeType&&>, TCanAppend<DelimiterType&&>>::Value>* = nullptr>
 	inline BuilderType& Join(RangeType&& InRange, DelimiterType&& InDelimiter)
 	{
 		bool bFirst = true;
@@ -325,7 +268,7 @@ public:
 	 * @return The builder, to allow additional operations to be composed with this one.
 	 */
 	template <typename RangeType, typename DelimiterType, typename QuoteType,
-		typename = typename TEnableIf<TAnd<TCanAppendRange<RangeType>, TCanAppend<DelimiterType&&>, TCanAppend<QuoteType&&>>::Value>::Type>
+		std::enable_if_t<TAnd<TCanAppendRange<RangeType>, TCanAppend<DelimiterType&&>, TCanAppend<QuoteType&&>>::Value>* = nullptr>
 	inline BuilderType& JoinQuoted(RangeType&& InRange, DelimiterType&& InDelimiter, QuoteType&& InQuote)
 	{
 		bool bFirst = true;
@@ -349,8 +292,9 @@ public:
 	 *
 	 * @param Format A format string that specifies how to format the additional arguments. Refer to standard printf format.
 	 */
-	template <typename FmtType, typename... Types>
-	typename TEnableIf<TIsArrayOrRefOfType<FmtType, CharType>::Value, BuilderType&>::Type Appendf(const FmtType& Fmt, Types... Args)
+	template <typename FmtType, typename... Types,
+		std::enable_if_t<TIsArrayOrRefOfType<FmtType, CharType>::Value>* = nullptr>
+	BuilderType& Appendf(const FmtType& Fmt, Types... Args)
 	{
 		static_assert(TIsArrayOrRefOfType<FmtType, CharType>::Value, "Formatting string must be a character array.");
 		static_assert(TAnd<TIsValidVariadicFunctionArg<Types>...>::Value, "Invalid argument(s) passed to Appendf.");
@@ -430,57 +374,31 @@ private:
 
 // String Append Operators
 
-inline FAnsiStringBuilderBase&		operator<<(FAnsiStringBuilderBase& Builder, ANSICHAR Char)							{ return Builder.Append(Char); }
-inline FWideStringBuilderBase&		operator<<(FWideStringBuilderBase& Builder, ANSICHAR Char)							{ return Builder.Append(Char); }
-inline FWideStringBuilderBase&		operator<<(FWideStringBuilderBase& Builder, WIDECHAR Char)							{ return Builder.Append(Char); }
-inline FUtf8StringBuilderBase&		operator<<(FUtf8StringBuilderBase& Builder, ANSICHAR Char)							{ return Builder.Append(UTF8CHAR(Char)); }
-inline FUtf8StringBuilderBase&		operator<<(FUtf8StringBuilderBase& Builder, UTF8CHAR Char)							{ return Builder.Append(Char); }
-
-template <typename T>
-inline auto operator<<(FAnsiStringBuilderBase& Builder, T&& Str) -> decltype(Builder.Append(ImplicitConv<FAnsiStringView>(Forward<T>(Str))))
-{
-	return Builder.Append(ImplicitConv<FAnsiStringView>(Forward<T>(Str)));
-}
-
-template <typename T>
-inline auto operator<<(FWideStringBuilderBase& Builder, T&& Str) -> decltype(Builder.AppendUtf8(ImplicitConv<FUtf8StringView>(Forward<T>(Str))))
+template <typename CharType, typename CharRangeType>
+inline auto operator<<(TStringBuilderBase<CharType>& Builder, CharRangeType&& Str) -> decltype(Builder.Append(MakeStringView(Forward<CharRangeType>(Str))))
 {
 	// Anything convertible to an FAnsiStringView is also convertible to a FUtf8StringView, but FAnsiStringView is more efficient to convert
-	if constexpr (std::is_convertible_v<T, FAnsiStringView>)
+	if constexpr (std::is_convertible_v<CharRangeType, FAnsiStringView>)
 	{
-		return Builder.AppendAnsi(ImplicitConv<FAnsiStringView>(Forward<T>(Str)));
+		return Builder.Append(ImplicitConv<FAnsiStringView>(Forward<CharRangeType>(Str)));
 	}
 	else
 	{
-		return Builder.AppendUtf8(ImplicitConv<FUtf8StringView>(Forward<T>(Str)));
+		return Builder.Append(MakeStringView(Forward<CharRangeType>(Str)));
 	}
 }
 
-template <typename T>
-inline auto operator<<(FWideStringBuilderBase& Builder, T&& Str) -> decltype(Builder.Append(ImplicitConv<FWideStringView>(Forward<T>(Str))))
-{
-	return Builder.Append(ImplicitConv<FWideStringView>(Forward<T>(Str)));
-}
-
-template <typename T>
-inline auto operator<<(FUtf8StringBuilderBase& Builder, T&& Str) -> decltype(Builder.Append(ImplicitConv<FUtf8StringView>(Forward<T>(Str))))
-{
-	// Anything convertible to an FAnsiStringView is also convertible to a FUtf8StringView, but FAnsiStringView is more efficient to convert
-	if constexpr (std::is_convertible_v<T, FAnsiStringView>)
-	{
-		return Builder.AppendAnsi(ImplicitConv<FAnsiStringView>(Forward<T>(Str)));
-	}
-	else
-	{
-		return Builder.AppendUtf8(ImplicitConv<FUtf8StringView>(Forward<T>(Str)));
-	}
-}
+inline FAnsiStringBuilderBase&		operator<<(FAnsiStringBuilderBase& Builder, ANSICHAR Char)							{ return Builder.AppendChar(Char); }
+inline FWideStringBuilderBase&		operator<<(FWideStringBuilderBase& Builder, ANSICHAR Char)							{ return Builder.AppendChar(Char); }
+inline FWideStringBuilderBase&		operator<<(FWideStringBuilderBase& Builder, WIDECHAR Char)							{ return Builder.AppendChar(Char); }
+inline FUtf8StringBuilderBase&		operator<<(FUtf8StringBuilderBase& Builder, ANSICHAR Char)							{ return Builder.AppendChar(UTF8CHAR(Char)); }
+inline FUtf8StringBuilderBase&		operator<<(FUtf8StringBuilderBase& Builder, UTF8CHAR Char)							{ return Builder.AppendChar(Char); }
 
 // Prefer using << instead of += as operator+= is only intended for mechanical FString -> FStringView replacement.
-inline FStringBuilderBase&			operator+=(FStringBuilderBase& Builder, ANSICHAR Char)								{ return Builder.Append(Char); }
-inline FStringBuilderBase&			operator+=(FStringBuilderBase& Builder, WIDECHAR Char)								{ return Builder.Append(Char); }
+inline FStringBuilderBase&			operator+=(FStringBuilderBase& Builder, ANSICHAR Char)								{ return Builder.AppendChar(Char); }
+inline FStringBuilderBase&			operator+=(FStringBuilderBase& Builder, WIDECHAR Char)								{ return Builder.AppendChar(Char); }
 inline FStringBuilderBase&			operator+=(FStringBuilderBase& Builder, FWideStringView Str)						{ return Builder.Append(Str); }
-inline FStringBuilderBase&			operator+=(FStringBuilderBase& Builder, FUtf8StringView Str)						{ return Builder.AppendUtf8(Str); }
+inline FStringBuilderBase&			operator+=(FStringBuilderBase& Builder, FUtf8StringView Str)						{ return Builder.Append(Str); }
 
 // Integer Append Operators
 
@@ -520,12 +438,12 @@ inline FUtf8StringBuilderBase&		operator<<(FUtf8StringBuilderBase& Builder, uint
  * For void Action(FStringView) -> Action(WriteToString<64>(Arg1, Arg2));
  * For UE_LOG or checkf -> checkf(Condition, TEXT("%s"), *WriteToString<32>(Arg));
  */
-template <int32 BufferSize>
-class WriteToString : public TStringBuilder<BufferSize>
+template <typename CharType, int32 BufferSize>
+class TWriteToString : public TStringBuilderWithBuffer<CharType, BufferSize>
 {
 public:
 	template <typename... ArgTypes>
-	explicit WriteToString(ArgTypes&&... Args)
+	explicit TWriteToString(ArgTypes&&... Args)
 	{
 	#if PLATFORM_COMPILER_HAS_FOLD_EXPRESSIONS
 		(*this << ... << Forward<ArgTypes>(Args));
@@ -536,4 +454,13 @@ public:
 	}
 };
 
-template <int32 BufferSize> struct TIsContiguousContainer<WriteToString<BufferSize>> { static constexpr bool Value = true; };
+template <typename CharType, int32 BufferSize>
+struct TIsContiguousContainer<TWriteToString<CharType, BufferSize>>
+{
+	static constexpr bool Value = true;
+};
+
+template <int32 BufferSize> using WriteToString = TWriteToString<TCHAR, BufferSize>;
+template <int32 BufferSize> using WriteToAnsiString = TWriteToString<ANSICHAR, BufferSize>;
+template <int32 BufferSize> using WriteToWideString = TWriteToString<WIDECHAR, BufferSize>;
+template <int32 BufferSize> using WriteToUtf8String = TWriteToString<UTF8CHAR, BufferSize>;
diff --git a/Engine/Source/Runtime/Core/Public/Traits/ElementType.h b/Engine/Source/Runtime/Core/Public/Traits/ElementType.h
--- a/Engine/Source/Runtime/Core/Public/Traits/ElementType.h
+++ b/Engine/Source/Runtime/Core/Public/Traits/ElementType.h
@@ -2,16 +2,32 @@
 
 #pragma once
 
+#include "HAL/Platform.h"
 #include <initializer_list>
-#include "Templates/RemoveCV.h"
+#include <type_traits>
+
+namespace UE::Core::Private::ElementType
+{
+
+template <typename T, typename = void>
+struct TImpl
+{
+};
+
+template <typename T>
+struct TImpl<T, std::void_t<typename T::ElementType>>
+{
+	using Type = typename T::ElementType;
+};
+
+} // UE::Core::Private::ElementType
 
 /**
  * Traits class which gets the element type of a container.
  */
 template <typename T>
-struct TElementType
+struct TElementType : UE::Core::Private::ElementType::TImpl<T>
 {
-	using Type = typename T::ElementType;
 };
 
 template <typename T> struct TElementType<             T& > : TElementType<T> {};
@@ -34,5 +50,8 @@ template <typename T, size_t N> struct TElementType<const volatile T[N]> { using
 template <typename T>
 struct TElementType<std::initializer_list<T>>
 {
-	using Type = typename TRemoveCV<T>::Type;
+	using Type = std::remove_cv_t<T>;
 };
+
+template <typename T>
+using TElementType_T = typename TElementType<T>::Type;
