Fix FString temporaries when querying config file
  - Removes 211K temporaries from process start to end of skydiving section

diff --git a/Engine/Source/Runtime/Core/Private/Misc/RemoteConfigIni.cpp b/Engine/Source/Runtime/Core/Private/Misc/RemoteConfigIni.cpp
--- a/Engine/Source/Runtime/Core/Private/Misc/RemoteConfigIni.cpp
+++ b/Engine/Source/Runtime/Core/Private/Misc/RemoteConfigIni.cpp
@@ -306,7 +306,7 @@ bool FRemoteConfig::ShouldReadRemoteFile(const TCHAR* Filename)
 
 FRemoteConfigAsyncIOInfo* FRemoteConfig::FindConfig(const TCHAR* Filename)
 {
-	return ConfigBuffers.Find(FString(Filename));
+	return ConfigBuffers.FindByHash(GetTypeHash(Filename), Filename);
 }
 
 
