CsvProfiler: allow disabling categories through configuration by using [CsvProfiler]/DisabledCategories array.

diff --git a/Engine/Source/Runtime/Core/Private/ProfilingDebugging/CsvProfiler.cpp b/Engine/Source/Runtime/Core/Private/ProfilingDebugging/CsvProfiler.cpp
--- a/Engine/Source/Runtime/Core/Private/ProfilingDebugging/CsvProfiler.cpp
+++ b/Engine/Source/Runtime/Core/Private/ProfilingDebugging/CsvProfiler.cpp
@@ -29,6 +29,7 @@
 #include "Misc/Compression.h"
 #include "Misc/Fork.h"
 #include "Misc/Guid.h"
+#include "Misc/WildcardString.h"
 
 #include "HAL/PlatformMisc.h"
 
@@ -426,6 +427,26 @@ public:
 		return -1;
 	}
 
+	void UpdateCategoryFromConfig(int32 CategoryIndex)
+	{
+		for (FString const& DisabledCategory : CategoriesDisabledInConfig)
+		{
+			if (FWildcardString::IsMatch(*DisabledCategory, *CategoryNames[CategoryIndex]))
+			{
+				GCsvCategoriesEnabled[CategoryIndex] = false;
+			}
+		}
+	}
+
+	void UpdateCategoriesFromConfig()
+	{
+		GConfig->GetArray(TEXT("CsvProfiler"), TEXT("DisabledCategories"), CategoriesDisabledInConfig, GEngineIni);
+		for (int i = 0; i < GetCategoryCount(); ++i)
+		{
+			UpdateCategoryFromConfig(i);
+		}
+	}
+
 	int32 RegisterCategory(const FString& CategoryName, bool bEnableByDefault, bool bIsGlobal)
 	{
 		int32 Index = -1;
@@ -451,6 +472,7 @@ public:
 					GCsvCategoriesEnabled[Index] = bEnableByDefault;
 					CategoryNames[Index] = CategoryName;
 					CategoryNameToIndex.Add(CategoryName.ToLower(), Index);
+					UpdateCategoryFromConfig(Index);
 				}
 				TRACE_CSV_PROFILER_REGISTER_CATEGORY(Index, *CategoryName);
 			}
@@ -469,6 +491,7 @@ private:
 	mutable FCriticalSection CS;
 	TMap<FString, int32> CategoryNameToIndex;
 	TArray<FString> CategoryNames;
+	TArray<FString> CategoriesDisabledInConfig;
 
 	static FCsvCategoryData* Instance;
 };
@@ -551,6 +574,11 @@ static void CsvProfilerEndFrameRT()
 	FCsvProfiler::Get()->EndFrameRT();
 }
 
+static void CsvProfilerReadConfig()
+{
+	FCsvCategoryData::Get()->UpdateCategoriesFromConfig();
+}
+
 
 static FAutoConsoleCommand HandleCSVProfileCmd(
 	TEXT("CsvProfile"),
@@ -3499,6 +3527,15 @@ void FCsvProfiler::Init()
 	{
 		GCsvUseProcessingThread = false;
 	}
+
+	if (GConfig != nullptr && GConfig->IsReadyForUse())
+	{
+		CsvProfilerReadConfig();
+	}
+	else
+	{
+		FCoreDelegates::OnInit.AddStatic(CsvProfilerReadConfig);
+	}
 }
 
 bool FCsvProfiler::IsCapturing()
