StringView: Moved everything into StringView.h to ease maintenance and support additional character types

#rnx

diff --git a/Engine/Source/Runtime/Core/Private/Containers/StringView.cpp b/Engine/Source/Runtime/Core/Private/Containers/StringView.cpp
deleted file mode 100644
diff --git a/Engine/Source/Runtime/Core/Public/Containers/StringView.h b/Engine/Source/Runtime/Core/Public/Containers/StringView.h
--- a/Engine/Source/Runtime/Core/Public/Containers/StringView.h
+++ b/Engine/Source/Runtime/Core/Public/Containers/StringView.h
@@ -4,6 +4,7 @@
 
 #include "CoreTypes.h"
 #include "Containers/StringFwd.h"
+#include "HAL/UnrealMemory.h"
 #include "Math/NumericLimits.h"
 #include "Math/UnrealMathUtility.h"
 #include "Misc/Crc.h"
@@ -210,9 +211,9 @@ public:
 	/** Returns the middle part of the view between any whitespace at the start and end. */
 	UE_NODISCARD inline ViewType TrimStartAndEnd() const;
 	/** Returns the right part of the view after any whitespace at the start. */
-	UE_NODISCARD CORE_API ViewType TrimStart() const;
+	UE_NODISCARD inline ViewType TrimStart() const;
 	/** Returns the left part of the view before any whitespace at the end. */
-	UE_NODISCARD CORE_API ViewType TrimEnd() const;
+	UE_NODISCARD inline ViewType TrimEnd() const;
 
 	/** Modifies the view to be the given number of characters from the left. */
 	inline void LeftInline(SizeType CharCount) { *this = Left(CharCount); }
@@ -288,7 +289,7 @@ public:
 	 * @param SubStr			The text to search for
 	 * @param StartPosition		The start character position to search from
 	 */
-	CORE_API int32 Find(ViewType SubStr, SizeType StartPosition = 0) const;
+	inline int32 Find(ViewType SubStr, SizeType StartPosition = 0) const;
 
 	/**
 	 * Returns whether this view contains the specified substring.
@@ -309,7 +310,7 @@ public:
 	 *
 	 * @return Whether the character was found in the view.
 	 */
-	CORE_API bool FindChar(CharType InChar, SizeType& OutIndex) const;
+	inline bool FindChar(CharType InChar, SizeType& OutIndex) const;
 
 	/**
 	 * Search the view for the last occurrence of a character.
@@ -319,7 +320,7 @@ public:
 	 *
 	 * @return Whether the character was found in the view.
 	 */
-	CORE_API bool FindLastChar(CharType InChar, SizeType& OutIndex) const;
+	inline bool FindLastChar(CharType InChar, SizeType& OutIndex) const;
 
 public:
 	/**
@@ -406,4 +407,316 @@ FORCEINLINE uint32 GetTypeHash(TStringView<CharType> View)
 	return FCrc::Strihash_DEPRECATED(View.Len(), View.GetData());
 }
 
-#include "StringView.inl"
+template <typename CharType>
+inline const CharType& TStringView<CharType>::operator[](SizeType Index) const
+{
+	checkf(Index >= 0 && Index < Size, TEXT("Index out of bounds on StringView: index %i on a view with a length of %i"), Index, Size);
+	return DataPtr[Index];
+}
+
+template <typename CharType>
+inline typename TStringView<CharType>::SizeType TStringView<CharType>::CopyString(CharType* Dest, SizeType CharCount, SizeType Position) const
+{
+	const  SizeType CopyCount = FMath::Min(Size - Position, CharCount);
+	FMemory::Memcpy(Dest, DataPtr + Position, CopyCount * sizeof(CharType));
+	return CopyCount;
+}
+
+template <typename CharType>
+inline TStringView<CharType> TStringView<CharType>::Left(SizeType CharCount) const
+{
+	return ViewType(DataPtr, FMath::Clamp(CharCount, 0, Size));
+}
+
+template <typename CharType>
+inline TStringView<CharType> TStringView<CharType>::LeftChop(SizeType CharCount) const
+{
+	return ViewType(DataPtr, FMath::Clamp(Size - CharCount, 0, Size));
+}
+
+template <typename CharType>
+inline TStringView<CharType> TStringView<CharType>::Right(SizeType CharCount) const
+{
+	const SizeType OutLen = FMath::Clamp(CharCount, 0, Size);
+	return ViewType(DataPtr + Size - OutLen, OutLen);
+}
+
+template <typename CharType>
+inline TStringView<CharType> TStringView<CharType>::RightChop(SizeType CharCount) const
+{
+	const SizeType OutLen = FMath::Clamp(Size - CharCount, 0, Size);
+	return ViewType(DataPtr + Size - OutLen, OutLen);
+}
+
+template <typename CharType>
+inline TStringView<CharType> TStringView<CharType>::Mid(SizeType Position, SizeType CharCount) const
+{
+	const CharType* CurrentStart  = GetData();
+	const SizeType  CurrentLength = Len();
+
+	// Clamp minimum position at the start of the string, adjusting the length down if necessary
+	const SizeType NegativePositionOffset = (Position < 0) ? Position : 0;
+	CharCount += NegativePositionOffset;
+	Position  -= NegativePositionOffset;
+
+	// Clamp maximum position at the end of the string
+	Position = (Position > CurrentLength) ? CurrentLength : Position;
+
+	// Clamp count between 0 and the distance to the end of the string
+	CharCount = FMath::Clamp(CharCount, 0, (CurrentLength - Position));
+
+	ViewType Result = ViewType(CurrentStart + Position, CharCount);
+	return Result;
+}
+
+template <typename CharType>
+inline TStringView<CharType> TStringView<CharType>::TrimStartAndEnd() const
+{
+	return TrimStart().TrimEnd();
+}
+
+template <typename CharType>
+inline TStringView<CharType> TStringView<CharType>::TrimStart() const
+{
+	SizeType SpaceCount = 0;
+	for (ElementType Char : *this)
+	{
+		if (!TChar<ElementType>::IsWhitespace(Char))
+		{
+			break;
+		}
+		++SpaceCount;
+	}
+	return ViewType(DataPtr + SpaceCount, Size - SpaceCount);
+}
+
+template <typename CharType>
+inline TStringView<CharType> TStringView<CharType>::TrimEnd() const
+{
+	SizeType NewSize = Size;
+	while (NewSize && TChar<ElementType>::IsWhitespace(DataPtr[NewSize - 1]))
+	{
+		--NewSize;
+	}
+	return ViewType(DataPtr, NewSize);
+}
+
+template <typename CharType>
+template <typename OtherType, typename TEnableIf<StringViewPrivate::TIsConvertibleToStringView<OtherType>::Value>::Type*>
+inline bool TStringView<CharType>::Equals(OtherType&& Other, ESearchCase::Type SearchCase) const
+{
+	using OtherViewType = typename StringViewPrivate::TCompatibleStringViewType<OtherType>::Type;
+	const OtherViewType OtherView(Forward<OtherType>(Other));
+	return Len() == OtherView.Len() && Compare(OtherView, SearchCase) == 0;
+}
+
+template <typename CharType>
+template <typename OtherCharType>
+inline bool TStringView<CharType>::Equals(const OtherCharType* const Other, ESearchCase::Type SearchCase) const
+{
+	if (SearchCase == ESearchCase::CaseSensitive)
+	{
+		return FPlatformString::Strncmp(GetData(), Other, Len()) == 0 && !Other[Len()];
+	}
+	else
+	{
+		return FPlatformString::Strnicmp(GetData(), Other, Len()) == 0 && !Other[Len()];
+	}
+}
+
+template <typename CharType>
+template <typename OtherType, typename TEnableIf<StringViewPrivate::TIsConvertibleToStringView<OtherType>::Value>::Type*>
+inline int32 TStringView<CharType>::Compare(OtherType&& Other, ESearchCase::Type SearchCase) const
+{
+	using OtherViewType = typename StringViewPrivate::TCompatibleStringViewType<OtherType>::Type;
+	const OtherViewType OtherView(Forward<OtherType>(Other));
+
+	const SizeType SelfLen = Len();
+	const SizeType OtherLen = OtherView.Len();
+	const SizeType MinLen = FMath::Min(SelfLen, OtherLen);
+
+	int Result;
+	if (SearchCase == ESearchCase::CaseSensitive)
+	{
+		Result = FPlatformString::Strncmp(GetData(), OtherView.GetData(), MinLen);
+	}
+	else
+	{
+		Result = FPlatformString::Strnicmp(GetData(), OtherView.GetData(), MinLen);
+	}
+
+	if (Result != 0)
+	{
+		return Result;
+	}
+
+	return SelfLen - OtherLen;
+}
+
+template <typename CharType>
+template <typename OtherCharType>
+inline int32 TStringView<CharType>::Compare(const OtherCharType* const Other, ESearchCase::Type SearchCase) const
+{
+	int Result;
+	if (SearchCase == ESearchCase::CaseSensitive)
+	{
+		Result = FPlatformString::Strncmp(GetData(), Other, Len());
+	}
+	else
+	{
+		Result = FPlatformString::Strnicmp(GetData(), Other, Len());
+	}
+
+	if (Result != 0)
+	{
+		return Result;
+	}
+
+	// Equal if Other[Len()] is '\0' and less otherwise. !!Other[Len()] is 0 for '\0' and 1 otherwise.
+	return -!!Other[Len()];
+}
+
+template <typename CharType>
+inline bool TStringView<CharType>::StartsWith(ViewType Prefix, ESearchCase::Type SearchCase) const
+{
+	return Prefix.Equals(Left(Prefix.Len()), SearchCase);
+}
+
+template <typename CharType>
+inline bool TStringView<CharType>::EndsWith(ViewType Suffix, ESearchCase::Type SearchCase) const
+{
+	return Suffix.Equals(Right(Suffix.Len()), SearchCase);
+}
+
+template <typename CharType>
+inline int32 TStringView<CharType>::Find(ViewType SubStr, SizeType StartPosition) const
+{
+	SizeType MinSearchLen = StartPosition + SubStr.Len();
+	if (SubStr.IsEmpty() || Size < MinSearchLen)
+	{
+		return INDEX_NONE;
+	}
+
+	ElementType FirstSubChar = SubStr[0];
+	SizeType RestSubLen = SubStr.Len() - 1;
+	SizeType Maxi = Size - MinSearchLen;
+	for (SizeType i = StartPosition; i <= Maxi; ++i)
+	{
+		if (DataPtr[i] == FirstSubChar)
+		{
+			if (RestSubLen == 0 || TCString<CharType>::Strncmp(DataPtr + i + 1, SubStr.DataPtr + 1, RestSubLen) == 0)
+			{
+				return i;
+			}
+		}
+	}
+
+	return INDEX_NONE;
+}
+
+template <typename CharType>
+inline bool TStringView<CharType>::FindChar(const ElementType InChar, SizeType& OutIndex) const
+{
+	for (SizeType i = 0; i < Size; ++i)
+	{
+		if (DataPtr[i] == InChar)
+		{
+			OutIndex = i;
+			return true;
+		}
+	}
+
+	OutIndex = INDEX_NONE;
+	return false;
+}
+
+
+template <typename CharType>
+inline bool TStringView<CharType>::FindLastChar(const ElementType InChar, SizeType& OutIndex) const
+{
+	if (Size == 0)
+	{
+		OutIndex = INDEX_NONE;
+		return false;
+	}
+
+	for (SizeType i = Size - 1; i >= 0; --i)
+	{
+		if (DataPtr[i] == InChar)
+		{
+			OutIndex = i;
+			return true;
+		}
+	}
+
+	OutIndex = INDEX_NONE;
+	return false;
+}
+
+// Case-insensitive string view comparison operators
+
+template <typename CharType>
+inline bool operator==(TStringView<CharType> Lhs, TStringView<CharType> Rhs)
+{
+	return Lhs.Equals(Rhs, ESearchCase::IgnoreCase);
+}
+
+template <typename CharType>
+inline bool operator!=(TStringView<CharType> Lhs, TStringView<CharType> Rhs)
+{
+	return !(Lhs == Rhs);
+}
+
+// Case-insensitive character range comparison operators
+
+template <typename CharType, typename CharRangeType>
+inline auto operator==(TStringView<CharType> Lhs, CharRangeType&& Rhs)
+	-> decltype(Lhs.Equals(ImplicitConv<TStringView<CharType>>(Forward<CharRangeType>(Rhs)), ESearchCase::IgnoreCase))
+{
+	return Lhs.Equals(ImplicitConv<TStringView<CharType>>(Forward<CharRangeType>(Rhs)), ESearchCase::IgnoreCase);
+}
+
+template <typename CharType, typename CharRangeType>
+inline auto operator==(CharRangeType&& Lhs, TStringView<CharType> Rhs)
+	-> decltype(ImplicitConv<TStringView<CharType>>(Forward<CharRangeType>(Lhs)).Equals(Rhs, ESearchCase::IgnoreCase))
+{
+	return ImplicitConv<TStringView<CharType>>(Forward<CharRangeType>(Lhs)).Equals(Rhs, ESearchCase::IgnoreCase);
+}
+
+template <typename CharType, typename CharRangeType>
+inline auto operator!=(TStringView<CharType> Lhs, CharRangeType&& Rhs) -> decltype(!(Lhs == Forward<CharRangeType>(Rhs)))
+{
+	return !(Lhs == Forward<CharRangeType>(Rhs));
+}
+
+template <typename CharType, typename CharRangeType>
+inline auto operator!=(CharRangeType&& Lhs, TStringView<CharType> Rhs) -> decltype(!(Rhs == Forward<CharRangeType>(Lhs)))
+{
+	return !(Rhs == Forward<CharRangeType>(Lhs));
+}
+
+// Case-insensitive C-string comparison operators
+
+template <typename CharType>
+inline bool operator==(TStringView<CharType> Lhs, const CharType* Rhs)
+{
+	return TCString<CharType>::Strnicmp(Lhs.GetData(), Rhs, Lhs.Len()) == 0 && !Rhs[Lhs.Len()];
+}
+
+template <typename CharType>
+inline bool operator==(const CharType* Lhs, TStringView<CharType> Rhs)
+{
+	return Rhs == Lhs;
+}
+
+template <typename CharType>
+inline bool operator!=(TStringView<CharType> Lhs, const CharType* Rhs)
+{
+	return !(Lhs == Rhs);
+}
+
+template <typename CharType>
+inline bool operator!=(const CharType* Lhs, TStringView<CharType> Rhs)
+{
+	return !(Lhs == Rhs);
+}
diff --git a/Engine/Source/Runtime/Core/Public/Containers/StringView.inl b/Engine/Source/Runtime/Core/Public/Containers/StringView.inl
deleted file mode 100644
