NonUnity buildfix

diff --git a/Engine/Source/Runtime/Core/Private/ProfilingDebugging/CsvProfiler.cpp b/Engine/Source/Runtime/Core/Private/ProfilingDebugging/CsvProfiler.cpp
--- a/Engine/Source/Runtime/Core/Private/ProfilingDebugging/CsvProfiler.cpp
+++ b/Engine/Source/Runtime/Core/Private/ProfilingDebugging/CsvProfiler.cpp
@@ -27,6 +27,7 @@
 #include "Stats/Stats.h"
 #include "HAL/LowLevelMemTracker.h"
 #include "Misc/Compression.h"
+#include "Misc/ConfigCacheIni.h"
 #include "Misc/Fork.h"
 #include "Misc/Guid.h"
 #include "Misc/WildcardString.h"
