Remove spikes of untagged memory for iostore requests

diff --git a/Engine/Source/Runtime/Core/Private/IO/IoDispatcher.cpp b/Engine/Source/Runtime/Core/Private/IO/IoDispatcher.cpp
--- a/Engine/Source/Runtime/Core/Private/IO/IoDispatcher.cpp
+++ b/Engine/Source/Runtime/Core/Private/IO/IoDispatcher.cpp
@@ -739,6 +739,8 @@ private:
 	virtual uint32 Run()
 	{
 		FMemory::SetupTLSCachesOnCurrentThread();
+
+		LLM_SCOPE(ELLMTag::FileSystem);
 		while (!bStopRequested)
 		{
 			if (PendingIoRequestsCount)
