Use AnsiMalloc + aligned malloc when ASan is enabled

diff --git a/Engine/Source/Runtime/Core/Public/HAL/MallocAnsi.h b/Engine/Source/Runtime/Core/Public/HAL/MallocAnsi.h
--- a/Engine/Source/Runtime/Core/Public/HAL/MallocAnsi.h
+++ b/Engine/Source/Runtime/Core/Public/HAL/MallocAnsi.h
@@ -5,10 +5,12 @@
 #include "CoreTypes.h"
 #include "HAL/UnrealMemory.h"
 
-#if PLATFORM_MAC || PLATFORM_IOS
-	#define USE_ALIGNED_MALLOC 1
-#else
-	#define USE_ALIGNED_MALLOC 0
+#ifndef USE_ALIGNED_MALLOC
+	#if PLATFORM_MAC || PLATFORM_IOS
+		#define USE_ALIGNED_MALLOC 1
+	#else
+		#define USE_ALIGNED_MALLOC 0
+	#endif
 #endif
 
 CORE_API void* AnsiMalloc(SIZE_T Size, uint32 Alignment);
