Improve code comments for IFileManager::CreateFileReader and IFileManager::CreateFileWriter

#rnx

- The main point of the change is to tell the user which set of enum bitflags should be used with which flag parameter

diff --git a/Engine/Source/Runtime/Core/Public/HAL/FileManager.h b/Engine/Source/Runtime/Core/Public/HAL/FileManager.h
--- a/Engine/Source/Runtime/Core/Public/HAL/FileManager.h
+++ b/Engine/Source/Runtime/Core/Public/HAL/FileManager.h
@@ -74,10 +74,26 @@ public:
 	/** Returns whether the sandbox is enabled or not */
 	virtual bool IsSandboxEnabled() const = 0;
 
-	/** Creates file reader archive. */
+	/**
+	 * Opens a file for reading and create an FArchive which can be used to read from it.
+	 *
+	 * @param Filename	Path to the file to open
+	 * @param ReadFlags	An optional bitfield of optional flags. For flag values see @See EFileRead
+	 *
+	 * @return Returns a pointer to an FArchive when successful and a nullptr if the operation failed.
+	 * Note that it is up to the caller to delete the FArchive when done.
+	 */
 	virtual FArchive* CreateFileReader( const TCHAR* Filename, uint32 ReadFlags=0 )=0;
 
-	/** Creates file writer archive. */
+	/**
+	 * Opens a file for writing and create an FArchive which can be used to write to it.
+	 *
+	 * @param Filename		Path to the desired location of the file
+	 * @param WriteFlags	An optional bitfield of optional flags. For flag values see @See EFileWrite
+	 *
+	 * @return Returns a pointer to an FArchive when successful and a nullptr if the operation failed.
+	 * Note that it is up to the caller to delete the FArchive when done.
+	 */
 	virtual FArchive* CreateFileWriter( const TCHAR* Filename, uint32 WriteFlags=0 )=0;
 
 	// If you're writing to a debug file, you should use CreateDebugFileWriter, and wrap the calling code in #if ALLOW_DEBUG_FILES.
