[Insights] TraceLog stat group to visualize memory usage.

diff --git a/Engine/Source/Runtime/Core/Private/ProfilingDebugging/TraceAuxiliary.cpp b/Engine/Source/Runtime/Core/Private/ProfilingDebugging/TraceAuxiliary.cpp
--- a/Engine/Source/Runtime/Core/Private/ProfilingDebugging/TraceAuxiliary.cpp
+++ b/Engine/Source/Runtime/Core/Private/ProfilingDebugging/TraceAuxiliary.cpp
@@ -58,6 +58,13 @@ const TCHAR* GMemoryChannels = TEXT("memtag,memalloc,callstack,module");
 ////////////////////////////////////////////////////////////////////////////////
 CSV_DEFINE_CATEGORY(Trace, true);
 
+////////////////////////////////////////////////////////////////////////////////
+DECLARE_STATS_GROUP(TEXT("TraceLog"), STATGROUP_Trace, STATCAT_Advanced);
+DECLARE_MEMORY_STAT(TEXT("Memory used"), STAT_TraceMemoryUsed, STATGROUP_Trace);
+DECLARE_MEMORY_STAT(TEXT("Important event cache used"), STAT_TraceCacheUsed, STATGROUP_Trace);
+DECLARE_MEMORY_STAT(TEXT("Important event cache waste"), STAT_TraceCacheWaste, STATGROUP_Trace);
+DECLARE_MEMORY_STAT(TEXT("Sent"), STAT_TraceSent, STATGROUP_Trace);
+
 ////////////////////////////////////////////////////////////////////////////////
 enum class ETraceConnectType
 {
@@ -109,6 +116,7 @@ private:
 
 static FTraceAuxiliaryImpl GTraceAuxiliary;
 static FDelegateHandle GEndFrameDelegateHandle;
+static FDelegateHandle GEndFrameStatDelegateHandle;
 
 ////////////////////////////////////////////////////////////////////////////////
 void FTraceAuxiliaryImpl::AddChannels(const TCHAR* ChannelList)
@@ -423,6 +431,19 @@ void FTraceAuxiliaryImpl::StartEndFramePump()
 		// If the worker thread is disabled, pump the update from end frame
 		GEndFrameDelegateHandle = FCoreDelegates::OnEndFrame.AddStatic(UE::Trace::Update);
 	}
+	if (!GEndFrameStatDelegateHandle.IsValid())
+	{
+		// Update stats every frame
+		GEndFrameStatDelegateHandle = FCoreDelegates::OnEndFrame.AddLambda([]()
+		{
+			UE::Trace::FStatistics Stats;
+			UE::Trace::GetStatistics(Stats);
+			SET_MEMORY_STAT(STAT_TraceMemoryUsed, Stats.MemoryUsed);
+			SET_MEMORY_STAT(STAT_TraceCacheUsed, Stats.CacheUsed);
+			SET_MEMORY_STAT(STAT_TraceCacheWaste, Stats.CacheWaste);
+			SET_MEMORY_STAT(STAT_TraceSent, Stats.BytesSent);
+		});
+	}
 }
 
 
