Non-Unity compile fix

diff --git a/Engine/Source/Runtime/Core/Public/ProfilingDebugging/ResourceSize.h b/Engine/Source/Runtime/Core/Public/ProfilingDebugging/ResourceSize.h
--- a/Engine/Source/Runtime/Core/Public/ProfilingDebugging/ResourceSize.h
+++ b/Engine/Source/Runtime/Core/Public/ProfilingDebugging/ResourceSize.h
@@ -4,6 +4,9 @@
 
 #include "CoreTypes.h"
 #include "Misc/AssertionMacros.h"
+#include "Containers/Map.h"
+
+class FName;
 
 /** Indicate what types of resources should be included for calculating used memory */
 namespace EResourceSizeMode
