ANSICHAR, WIDECHAR, UCS2CHAR and UTF8CHAR support for FStringFormatter

diff --git a/Engine/Plugins/MovieScene/MovieRenderPipeline/Source/openexrRTTI/Private/OpenExrRTTIModule.cpp b/Engine/Plugins/MovieScene/MovieRenderPipeline/Source/openexrRTTI/Private/OpenExrRTTIModule.cpp
--- a/Engine/Plugins/MovieScene/MovieRenderPipeline/Source/openexrRTTI/Private/OpenExrRTTIModule.cpp
+++ b/Engine/Plugins/MovieScene/MovieRenderPipeline/Source/openexrRTTI/Private/OpenExrRTTIModule.cpp
@@ -31,8 +31,17 @@ public:
 			case FStringFormatArg::EType::String:
 				Attribute = new OPENEXR_IMF_INTERNAL_NAMESPACE::StringAttribute(std::string(TCHAR_TO_ANSI(*KVP.Value.StringValue)));
 				break;
-			case FStringFormatArg::EType::StringLiteral:
-				Attribute = new OPENEXR_IMF_INTERNAL_NAMESPACE::StringAttribute(std::string(TCHAR_TO_ANSI(KVP.Value.StringLiteralValue)));
+			case FStringFormatArg::EType::StringLiteralANSI:
+				Attribute = new OPENEXR_IMF_INTERNAL_NAMESPACE::StringAttribute(std::string(KVP.Value.StringLiteralANSIValue));
+				break;
+			case FStringFormatArg::EType::StringLiteralWIDE:
+				Attribute = new OPENEXR_IMF_INTERNAL_NAMESPACE::StringAttribute(std::string(StringCast<ANSICHAR>(KVP.Value.StringLiteralWIDEValue).Get()));
+				break;
+			case FStringFormatArg::EType::StringLiteralUCS2:
+				Attribute = new OPENEXR_IMF_INTERNAL_NAMESPACE::StringAttribute(std::string(StringCast<ANSICHAR>(KVP.Value.StringLiteralUCS2Value).Get()));
+				break;
+			case FStringFormatArg::EType::StringLiteralUTF8:
+				Attribute = new OPENEXR_IMF_INTERNAL_NAMESPACE::StringAttribute(std::string(StringCast<ANSICHAR>(KVP.Value.StringLiteralUTF8Value).Get()));
 				break;
 
 			case FStringFormatArg::EType::UInt:
diff --git a/Engine/Source/Runtime/Core/Private/Misc/StringFormatter.cpp b/Engine/Source/Runtime/Core/Private/Misc/StringFormatter.cpp
--- a/Engine/Source/Runtime/Core/Private/Misc/StringFormatter.cpp
+++ b/Engine/Source/Runtime/Core/Private/Misc/StringFormatter.cpp
@@ -13,7 +13,10 @@ FStringFormatArg::FStringFormatArg( const uint64 Value ) : Type(UInt), UIntValue
 FStringFormatArg::FStringFormatArg( const float Value ) : Type(Double), DoubleValue(Value) {}
 FStringFormatArg::FStringFormatArg( const double Value ) : Type(Double), DoubleValue(Value) {}
 FStringFormatArg::FStringFormatArg( FString Value ) : Type(String), StringValue(MoveTemp(Value)) {}
-FStringFormatArg::FStringFormatArg( const TCHAR* Value ) : Type(StringLiteral), StringLiteralValue(Value) {}
+FStringFormatArg::FStringFormatArg( const ANSICHAR* Value ) : Type(StringLiteralANSI), StringLiteralANSIValue(Value) {}
+FStringFormatArg::FStringFormatArg( const WIDECHAR* Value ) : Type(StringLiteralWIDE), StringLiteralWIDEValue(Value) {}
+FStringFormatArg::FStringFormatArg( const UCS2CHAR* Value ) : Type(StringLiteralUCS2), StringLiteralUCS2Value(Value) {}
+FStringFormatArg::FStringFormatArg( const UTF8CHAR* Value ) : Type(StringLiteralUTF8), StringLiteralUTF8Value(Value) {}
 FStringFormatArg::FStringFormatArg( const FStringFormatArg& RHS )
 {
 	Type = RHS.Type;
@@ -23,7 +26,10 @@ FStringFormatArg::FStringFormatArg( const FStringFormatArg& RHS )
 		case UInt: 				UIntValue = RHS.UIntValue; break;
 		case Double: 			IntValue = RHS.IntValue; break;
 		case String: 			StringValue = RHS.StringValue; break;
-		case StringLiteral: 	StringLiteralValue = RHS.StringLiteralValue; break;
+		case StringLiteralANSI: StringLiteralANSIValue = RHS.StringLiteralANSIValue; break;
+		case StringLiteralWIDE: StringLiteralWIDEValue = RHS.StringLiteralWIDEValue; break;
+		case StringLiteralUCS2: StringLiteralUCS2Value = RHS.StringLiteralUCS2Value; break;
+		case StringLiteralUTF8: StringLiteralUTF8Value = RHS.StringLiteralUTF8Value; break;
 	}
 }
 
@@ -31,11 +37,14 @@ void AppendToString(const FStringFormatArg& Arg, FString& StringToAppendTo)
 {
 	switch(Arg.Type)
 	{
-		case FStringFormatArg::Int: 			StringToAppendTo.Append(LexToString(Arg.IntValue)); break;
-		case FStringFormatArg::UInt: 			StringToAppendTo.Append(LexToString(Arg.UIntValue)); break;
-		case FStringFormatArg::Double: 			StringToAppendTo.Append(LexToString(Arg.DoubleValue)); break;
-		case FStringFormatArg::String: 			StringToAppendTo.AppendChars(*Arg.StringValue, Arg.StringValue.Len()); break;
-		case FStringFormatArg::StringLiteral: 	StringToAppendTo += Arg.StringLiteralValue; break;
+		case FStringFormatArg::Int: 				StringToAppendTo.Append(LexToString(Arg.IntValue)); break;
+		case FStringFormatArg::UInt: 				StringToAppendTo.Append(LexToString(Arg.UIntValue)); break;
+		case FStringFormatArg::Double: 				StringToAppendTo.Append(LexToString(Arg.DoubleValue)); break;
+		case FStringFormatArg::String: 				StringToAppendTo.AppendChars(*Arg.StringValue, Arg.StringValue.Len()); break;
+		case FStringFormatArg::StringLiteralANSI: 	StringToAppendTo.Append(Arg.StringLiteralANSIValue); break;
+		case FStringFormatArg::StringLiteralWIDE: 	StringToAppendTo.Append(Arg.StringLiteralWIDEValue); break;
+		case FStringFormatArg::StringLiteralUCS2: 	StringToAppendTo.Append(Arg.StringLiteralUCS2Value); break;
+		case FStringFormatArg::StringLiteralUTF8: 	StringToAppendTo.Append(Arg.StringLiteralUTF8Value); break;
 	}
 }
 
diff --git a/Engine/Source/Runtime/Core/Public/Misc/StringFormatArg.h b/Engine/Source/Runtime/Core/Public/Misc/StringFormatArg.h
--- a/Engine/Source/Runtime/Core/Public/Misc/StringFormatArg.h
+++ b/Engine/Source/Runtime/Core/Public/Misc/StringFormatArg.h
@@ -7,7 +7,7 @@
 /** An argument supplied to FString::Format */
 struct CORE_API FStringFormatArg
 {
-	enum EType { Int, UInt, Double, String, StringLiteral };
+	enum EType { Int, UInt, Double, String, StringLiteralANSI, StringLiteralWIDE, StringLiteralUCS2, StringLiteralUTF8 };
 
 	/** The type of this arg */
 	EType Type;
@@ -21,8 +21,14 @@ struct CORE_API FStringFormatArg
 		uint64 UIntValue;
 		/** Value as double */
 		double DoubleValue;
-		/** Value as a string literal */
-		const TCHAR* StringLiteralValue;
+		/** Value as an ANSI string literal */
+		const ANSICHAR* StringLiteralANSIValue;
+		/** Value as a WIDE string literal */
+		const WIDECHAR* StringLiteralWIDEValue;
+		/** Value as a UCS2 string literal */
+		const UCS2CHAR* StringLiteralUCS2Value;
+		/** Value as a UTF8 string literal */
+		const UTF8CHAR* StringLiteralUTF8Value;
 	};
 
 	/** Value as an FString */
@@ -35,7 +41,10 @@ struct CORE_API FStringFormatArg
 	FStringFormatArg( const float Value );
 	FStringFormatArg( const double Value );
 	FStringFormatArg( FString Value );
-	FStringFormatArg( const TCHAR* Value );
+	FStringFormatArg( const ANSICHAR* Value );
+	FStringFormatArg( const WIDECHAR* Value );
+	FStringFormatArg( const UCS2CHAR* Value );
+	FStringFormatArg( const UTF8CHAR* Value );
 
 	/** Copyable */
 	FStringFormatArg( const FStringFormatArg& RHS );
