IoStore:
Log toc signature hash when mounting containers
Add guard against loading toc files with a later version than we know about
#rnx

diff --git a/Engine/Source/Runtime/Core/Internal/IO/IoStore.h b/Engine/Source/Runtime/Core/Internal/IO/IoStore.h
--- a/Engine/Source/Runtime/Core/Internal/IO/IoStore.h
+++ b/Engine/Source/Runtime/Core/Internal/IO/IoStore.h
@@ -229,6 +229,8 @@ struct FIoStoreTocResource
 
 	TArray<FName> CompressionMethods;
 
+	FSHAHash SignatureHash;
+	
 	TArray<FSHAHash> ChunkBlockSignatures;
 
 	TArray<FIoStoreTocEntryMeta> ChunkMetas;
diff --git a/Engine/Source/Runtime/Core/Private/IO/IoDispatcherFileBackend.cpp b/Engine/Source/Runtime/Core/Private/IO/IoDispatcherFileBackend.cpp
--- a/Engine/Source/Runtime/Core/Private/IO/IoDispatcherFileBackend.cpp
+++ b/Engine/Source/Runtime/Core/Private/IO/IoDispatcherFileBackend.cpp
@@ -751,6 +751,8 @@ FIoStatus FFileIoStoreReader::Initialize(const TCHAR* InTocFilePath, int32 InOrd
 
 	Stats.OnTocMounted(GetTocAllocatedSize());
 
+	UE_LOG(LogIoDispatcher, Display, TEXT("Toc signature hash: %s"), *TocResource.SignatureHash.ToString());
+	
 	ContainerId = TocResource.Header.ContainerId;
 	Order = InOrder;
 	return FIoStatus::Ok;
diff --git a/Engine/Source/Runtime/Core/Private/IO/IoStore.cpp b/Engine/Source/Runtime/Core/Private/IO/IoStore.cpp
--- a/Engine/Source/Runtime/Core/Private/IO/IoStore.cpp
+++ b/Engine/Source/Runtime/Core/Private/IO/IoStore.cpp
@@ -2406,6 +2406,11 @@ FIoStatus FIoStoreTocResource::Read(const TCHAR* TocFilePath, EIoStoreTocReadOpt
 		return FIoStatusBuilder(EIoErrorCode::CorruptToc) << TEXT("Outdated TOC header version while reading '") << TocFilePath << TEXT("'");
 	}
 
+	if (Header.Version > static_cast<uint8>(EIoStoreTocVersion::Latest))
+	{
+		return FIoStatusBuilder(EIoErrorCode::CorruptToc) << TEXT("Too new TOC header version while reading '") << TocFilePath << TEXT("'");
+	}
+
 	const uint64 TotalTocSize = TocFileHandle->Size() - sizeof(FIoStoreTocHeader);
 	const uint64 TocMetaSize = Header.TocEntryCount * sizeof(FIoStoreTocEntryMeta);
 	const uint64 DefaultTocSize = TotalTocSize - Header.DirectoryIndexSize - TocMetaSize;
@@ -2495,6 +2500,10 @@ FIoStatus FIoStoreTocResource::Read(const TCHAR* TocFilePath, EIoStoreTocReadOpt
 		const int32* HashSize = reinterpret_cast<const int32*>(SignatureBuffer);
 		TArrayView<const uint8> TocSignature = MakeArrayView<const uint8>(reinterpret_cast<const uint8*>(HashSize + 1), *HashSize);
 		TArrayView<const uint8> BlockSignature = MakeArrayView<const uint8>(TocSignature.GetData() + *HashSize, *HashSize);
+		
+		TArrayView<const uint8> BothSignatures = MakeArrayView<const uint8>(TocSignature.GetData(), *HashSize * 2);
+		FSHA1::HashBuffer(BothSignatures.GetData(), BothSignatures.Num(), OutTocResource.SignatureHash.Hash);
+		
 		TArrayView<const FSHAHash> ChunkBlockSignatures = MakeArrayView<const FSHAHash>(reinterpret_cast<const FSHAHash*>(BlockSignature.GetData() + *HashSize), Header.TocCompressedBlockEntryCount);
 
 		// Adjust address to meta data
