IoDispatcher:
Fixed requests never being batched in single threaded mode
Always use the generic file backend in single threaded mode
#rnx

diff --git a/Engine/Source/Runtime/Core/Private/IO/IoDispatcher.cpp b/Engine/Source/Runtime/Core/Private/IO/IoDispatcher.cpp
--- a/Engine/Source/Runtime/Core/Private/IO/IoDispatcher.cpp
+++ b/Engine/Source/Runtime/Core/Private/IO/IoDispatcher.cpp
@@ -361,34 +361,34 @@ public:
 		{
 			DispatcherEvent->Trigger();
 		}
-		else
-		{
-			ProcessIncomingRequests();
-			while (PendingIoRequestsCount > 0)
-			{
-				ProcessCompletedRequests();
-			}
-		}
 	}
 
 	void Cancel(FIoRequestImpl* Request)
 	{
+		if (!BackendContext->bIsMultiThreaded)
+		{
+			return;
+		}
 		Request->AddRef();
 		{
 			FScopeLock _(&UpdateLock);
 			RequestsToCancel.Add(Request);
 		}
-		WakeUpDispatcherThread();
+		DispatcherEvent->Trigger();
 	}
 
 	void Reprioritize(FIoRequestImpl* Request)
 	{
+		if (!BackendContext->bIsMultiThreaded)
+		{
+			return;
+		}
 		Request->AddRef();
 		{
 			FScopeLock _(&UpdateLock);
 			RequestsToReprioritize.Add(Request);
 		}
-		WakeUpDispatcherThread();
+		DispatcherEvent->Trigger();
 	}
 
 	TIoStatusOr<FIoMappedRegion> OpenMapped(const FIoChunkId& ChunkId, const FIoReadOptions& Options)
@@ -497,7 +497,18 @@ public:
 			WaitingRequestsTail = Batch.TailRequest;
 		}
 		Batch.HeadRequest = Batch.TailRequest = nullptr;
-		WakeUpDispatcherThread();
+		if (BackendContext->bIsMultiThreaded)
+		{
+			DispatcherEvent->Trigger();
+		}
+		else
+		{
+			ProcessIncomingRequests();
+			while (PendingIoRequestsCount > 0)
+			{
+				ProcessCompletedRequests();
+			}
+		}
 	}
 
 	void IssueBatch(FIoBatch& Batch)
@@ -727,7 +738,6 @@ private:
 			}
 			
 			++PendingIoRequestsCount;
-			ProcessCompletedRequests();
 		}
 	}
 
diff --git a/Engine/Source/Runtime/Core/Private/IO/IoDispatcherFileBackend.cpp b/Engine/Source/Runtime/Core/Private/IO/IoDispatcherFileBackend.cpp
--- a/Engine/Source/Runtime/Core/Private/IO/IoDispatcherFileBackend.cpp
+++ b/Engine/Source/Runtime/Core/Private/IO/IoDispatcherFileBackend.cpp
@@ -2031,7 +2031,17 @@ uint32 FFileIoStore::Run()
 
 TSharedRef<IIoDispatcherFileBackend> CreateIoDispatcherFileBackend()
 {
-	const bool bCheckForPlatformImplementation = UE_BUILD_SHIPPING || !FParse::Param(FCommandLine::Get(), TEXT("forcegenericio"));
+	bool bCheckForPlatformImplementation = true;
+	if (!FGenericPlatformProcess::SupportsMultithreading())
+	{
+		bCheckForPlatformImplementation = false;
+	}
+#if !UE_BUILD_SHIPPING
+	if (FParse::Param(FCommandLine::Get(), TEXT("forcegenericio")))
+	{
+		bCheckForPlatformImplementation = false;
+	}
+#endif
 
 	if (bCheckForPlatformImplementation)
 	{
