CompressedBuffer: Added FCompressedBufferReader to store decoder buffers and allow streaming from seekable archives

#rnx

diff --git a/Engine/Source/Runtime/Core/Private/Compression/CompressedBuffer.cpp b/Engine/Source/Runtime/Core/Private/Compression/CompressedBuffer.cpp
--- a/Engine/Source/Runtime/Core/Private/Compression/CompressedBuffer.cpp
+++ b/Engine/Source/Runtime/Core/Private/Compression/CompressedBuffer.cpp
@@ -2,12 +2,14 @@
 
 #include "Compression/CompressedBuffer.h"
 
+#include "Algo/Accumulate.h"
 #include "Algo/ForEach.h"
 #include "Compression/OodleDataCompression.h"
 #include "Hash/Blake3.h"
 #include "Math/UnrealMathUtility.h"
 #include "Misc/ByteSwap.h"
 #include "Misc/Crc.h"
+#include "Misc/ScopeExit.h"
 #include "Serialization/Archive.h"
 #include "Templates/Function.h"
 
@@ -15,12 +17,13 @@ THIRD_PARTY_INCLUDES_START
 #include "Compression/lz4.h"
 THIRD_PARTY_INCLUDES_END
 
-namespace UE::CompressedBuffer
+namespace UE::CompressedBuffer::Private
 {
 
 ///////////////////////////////////////////////////////////////////////////////////////////////////////////////////////
 
 static constexpr uint64 DefaultBlockSize = 256 * 1024;
+static constexpr uint64 DefaultHeaderSize = 4 * 1024;
 
 ///////////////////////////////////////////////////////////////////////////////////////////////////////////////////////
 
@@ -117,17 +120,40 @@ struct FHeader
 		TotalRawSize = NETWORK_ORDER64(TotalRawSize);
 		TotalCompressedSize = NETWORK_ORDER64(TotalCompressedSize);
 	}
+
+	bool TryGetCompressParameters(
+		ECompressedBufferCompressor& OutCompressor,
+		ECompressedBufferCompressionLevel& OutCompressionLevel,
+		uint64& OutBlockSize) const
+	{
+		switch (Method)
+		{
+		case EMethod::None:
+			OutCompressor = ECompressedBufferCompressor::NotSet;
+			OutCompressionLevel = ECompressedBufferCompressionLevel::None;
+			OutBlockSize = 0;
+			return true;
+		case EMethod::Oodle:
+			OutCompressor = ECompressedBufferCompressor(Compressor);
+			OutCompressionLevel = ECompressedBufferCompressionLevel(CompressionLevel);
+			OutBlockSize = uint64(1) << BlockSizeExponent;
+			return true;
+		default:
+			return false;
+		}
+	}
 };
 
 static_assert(sizeof(FHeader) == 64, "FHeader is the wrong size.");
 
 ///////////////////////////////////////////////////////////////////////////////////////////////////////////////////////
 
-struct FDecompressionContext
+class FDecoderSource
 {
-	FUniqueBuffer CompressedBlockSizes;
-	FUniqueBuffer CompressedBlock;
-	FUniqueBuffer UncompressedBlock;
+public:
+	virtual bool Read(uint64 Offset, FMutableMemoryView Data) const = 0;
+	virtual FMemoryView ReadOrView(uint64 Offset, uint64 Size, FDecoderContext& Context) const = 0;
+	virtual FCompositeBuffer ReadToComposite(uint64 Offset, uint64 Size) const = 0;
 };
 
 class FEncoder
@@ -139,9 +165,10 @@ public:
 class FDecoder
 {
 public:
-	virtual FCompositeBuffer Decompress(const FHeader& Header, const FCompositeBuffer& CompressedData) const = 0;
-	virtual bool TryDecompressTo(FDecompressionContext& Context, const FHeader& Header, const FCompositeBuffer& CompressedData, FMutableMemoryView RawView, uint64 RawOffset) const = 0;
 	virtual uint64 GetHeaderSize(const FHeader& Header) const = 0;
+	virtual void DecodeHeader(FMutableMemoryView HeaderView) const = 0;
+	virtual bool TryDecompressTo(FDecoderContext& Context, const FDecoderSource& Source, const FHeader& Header, uint64 RawOffset, FMutableMemoryView RawView) const = 0;
+	virtual FCompositeBuffer DecompressToComposite(FDecoderContext& Context, const FDecoderSource& Source, const FHeader& Header, uint64 RawOffset, uint64 RawSize) const;
 };
 
 ///////////////////////////////////////////////////////////////////////////////////////////////////////////////////////
@@ -167,33 +194,44 @@ public:
 class FNoneDecoder final : public FDecoder
 {
 public:
-	FCompositeBuffer Decompress(const FHeader& Header, const FCompositeBuffer& CompressedData) const final
+	uint64 GetHeaderSize(const FHeader& Header) const final
 	{
-		if (Header.Method == EMethod::None &&
-			Header.TotalCompressedSize == CompressedData.GetSize() &&
-			Header.TotalCompressedSize == Header.TotalRawSize + sizeof(FHeader))
-		{
-			return CompressedData.Mid(sizeof(FHeader), Header.TotalRawSize).MakeOwned();
-		}
-		return FCompositeBuffer();
+		return sizeof(FHeader);
 	}
 
-	bool TryDecompressTo(FDecompressionContext&, const FHeader& Header, const FCompositeBuffer& CompressedData, FMutableMemoryView RawView, uint64 RawOffset) const final
+	void DecodeHeader(FMutableMemoryView HeaderView) const final
 	{
-		if (Header.Method == EMethod::None &&
-			RawOffset + RawView.GetSize() <= Header.TotalRawSize &&
-			Header.TotalCompressedSize == CompressedData.GetSize() &&
-			Header.TotalCompressedSize == Header.TotalRawSize + sizeof(FHeader))
+		FHeader& Header = *static_cast<FHeader*>(HeaderView.GetData());
+		Header.ByteSwap();
+	}
+
+	bool TryDecompressTo(
+		FDecoderContext& Context,
+		const FDecoderSource& Source,
+		const FHeader& Header,
+		uint64 RawOffset,
+		FMutableMemoryView RawView) const final
+	{
+		if (Header.Method == EMethod::None && RawOffset + RawView.GetSize() <= Header.TotalRawSize)
 		{
-			CompressedData.CopyTo(RawView, sizeof(FHeader) + RawOffset);
-			return true;
+			return Source.Read(sizeof(FHeader) + RawOffset, RawView);
 		}
 		return false;
 	}
 
-	uint64 GetHeaderSize(const FHeader& Header) const final
+	FCompositeBuffer DecompressToComposite(
+		FDecoderContext& Context,
+		const FDecoderSource& Source,
+		const FHeader& Header,
+		uint64 RawOffset,
+		uint64 RawSize) const final
 	{
-		return sizeof(FHeader);
+		if (Header.Method == EMethod::None &&
+			Header.TotalCompressedSize == Header.TotalRawSize + sizeof(FHeader))
+		{
+			return Source.ReadToComposite(sizeof(FHeader) + RawOffset, RawSize);
+		}
+		return FCompositeBuffer();
 	}
 };
 
@@ -304,211 +342,117 @@ FCompositeBuffer FBlockEncoder::Compress(const FCompositeBuffer& RawData, const
 class FBlockDecoder : public FDecoder
 {
 public:
-	FCompositeBuffer Decompress(const FHeader& Header, const FCompositeBuffer& CompressedData) const final;
-	bool TryDecompressTo(FDecompressionContext& Context, const FHeader& Header, const FCompositeBuffer& CompressedData, FMutableMemoryView RawView, uint64 RawOffset) const final;
-
 	uint64 GetHeaderSize(const FHeader& Header) const final
 	{
 		return sizeof(FHeader) + sizeof(uint32) * uint64(Header.BlockCount);
 	}
 
-protected:
-	virtual bool DecompressBlock(FMutableMemoryView RawData, FMemoryView CompressedData) const = 0;
-};
-
-FCompositeBuffer FBlockDecoder::Decompress(const FHeader& Header, const FCompositeBuffer& CompressedData) const
-{
-	if (Header.BlockCount == 0 ||
-		Header.TotalCompressedSize != CompressedData.GetSize())
-	{
-		return FCompositeBuffer();
-	}
-
-	// The raw data cannot reference the compressed data unless it is owned.
-	// An empty raw buffer requires an empty segment, which this path creates.
-	if (!CompressedData.IsOwned() || Header.TotalRawSize == 0)
-	{
-		FUniqueBuffer Buffer = FUniqueBuffer::Alloc(Header.TotalRawSize);
-		FDecompressionContext Context;
-		return TryDecompressTo(Context, Header, CompressedData, Buffer, 0) ? FCompositeBuffer(Buffer.MoveToShared()) : FCompositeBuffer();
-	}
-
-	TArray64<uint32> CompressedBlockSizes;
-	CompressedBlockSizes.AddUninitialized(Header.BlockCount);
-	CompressedData.CopyTo(MakeMemoryView(CompressedBlockSizes), sizeof(FHeader));
-	Algo::ForEach(CompressedBlockSizes, [](uint32& Size) { Size = NETWORK_ORDER32(Size); });
-
-	// Allocate the buffer for the raw blocks that were compressed.
-	FSharedBuffer RawData;
-	FMutableMemoryView RawDataView;
-	const uint64 BlockSize = uint64(1) << Header.BlockSizeExponent;
-	{
-		uint64 RawDataSize = 0;
-		uint64 RemainingRawSize = Header.TotalRawSize;
-		for (const uint32 CompressedBlockSize : CompressedBlockSizes)
-		{
-			const uint64 RawBlockSize = FMath::Min(RemainingRawSize, BlockSize);
-			if (CompressedBlockSize < BlockSize)
-			{
-				RawDataSize += RawBlockSize;
-			}
-			RemainingRawSize -= RawBlockSize;
-		}
-		FUniqueBuffer RawDataBuffer = FUniqueBuffer::Alloc(RawDataSize);
-		RawDataView = RawDataBuffer;
-		RawData = RawDataBuffer.MoveToShared();
-	}
-
-	// Decompress the compressed data in blocks and reference the uncompressed blocks.
-	uint64 PendingCompressedSegmentOffset = sizeof(FHeader) + uint64(Header.BlockCount) * sizeof(uint32);
-	uint64 PendingCompressedSegmentSize = 0;
-	uint64 PendingRawSegmentOffset = 0;
-	uint64 PendingRawSegmentSize = 0;
-	TArray<FSharedBuffer, TInlineAllocator<1>> Segments;
-
-	const auto CommitPendingCompressedSegment = [&PendingCompressedSegmentOffset, &PendingCompressedSegmentSize, &CompressedData, &Segments]
+	void DecodeHeader(FMutableMemoryView HeaderView) const final
 	{
-		if (PendingCompressedSegmentSize)
-		{
-			CompressedData.IterateRange(PendingCompressedSegmentOffset, PendingCompressedSegmentSize,
-				[&Segments](FMemoryView View, const FSharedBuffer& ViewOuter)
-				{
-					Segments.Add(FSharedBuffer::MakeView(View, ViewOuter));
-				});
-			PendingCompressedSegmentOffset += PendingCompressedSegmentSize;
-			PendingCompressedSegmentSize = 0;
-		}
-	};
-
-	const auto CommitPendingRawSegment = [&PendingRawSegmentOffset, &PendingRawSegmentSize, &RawData, &Segments]
-	{
-		if (PendingRawSegmentSize)
-		{
-			const FMemoryView PendingSegment = RawData.GetView().Mid(PendingRawSegmentOffset, PendingRawSegmentSize);
-			Segments.Add(FSharedBuffer::MakeView(PendingSegment, RawData));
-			PendingRawSegmentOffset += PendingRawSegmentSize;
-			PendingRawSegmentSize = 0;
-		}
-	};
-
-	FUniqueBuffer CompressedBlockCopy;
-	uint64 RemainingRawSize = Header.TotalRawSize;
-	uint64 RemainingCompressedSize = CompressedData.GetSize();
-	for (const uint32 CompressedBlockSize : CompressedBlockSizes)
-	{
-		if (RemainingCompressedSize < CompressedBlockSize)
-		{
-			return FCompositeBuffer();
-		}
-
-		const uint64 RawBlockSize = FMath::Min(RemainingRawSize, BlockSize);
-		if (RawBlockSize == CompressedBlockSize)
-		{
-			CommitPendingRawSegment();
-			PendingCompressedSegmentSize += RawBlockSize;
-		}
-		else
-		{
-			CommitPendingCompressedSegment();
-			const FMemoryView CompressedBlock = CompressedData.ViewOrCopyRange(PendingCompressedSegmentOffset, CompressedBlockSize, CompressedBlockCopy);
-			if (!DecompressBlock(RawDataView.Left(RawBlockSize), CompressedBlock))
-			{
-				return FCompositeBuffer();
-			}
-			PendingCompressedSegmentOffset += CompressedBlockSize;
-			PendingRawSegmentSize += RawBlockSize;
-			RawDataView += RawBlockSize;
-		}
-
-		RemainingCompressedSize -= CompressedBlockSize;
-		RemainingRawSize -= RawBlockSize;
+		FHeader& Header = *static_cast<FHeader*>(HeaderView.GetData());
+		Header.ByteSwap();
+		TArrayView<uint32> CompressedBlockSizes(reinterpret_cast<uint32*>(&Header + 1), Header.BlockCount);
+		Algo::ForEach(CompressedBlockSizes, [](uint32& Size) { Size = NETWORK_ORDER32(Size); });
 	}
 
-	CommitPendingCompressedSegment();
-	CommitPendingRawSegment();
+	bool TryDecompressTo(FDecoderContext& Context, const FDecoderSource& Source, const FHeader& Header, uint64 RawOffset, FMutableMemoryView RawView) const final;
+	FCompositeBuffer DecompressToComposite(FDecoderContext& Context, const FDecoderSource& Source, const FHeader& Header, uint64 RawOffset, uint64 RawSize) const final;
 
-	return FCompositeBuffer(MoveTemp(Segments));
-}
+protected:
+	virtual bool DecompressBlock(FMutableMemoryView RawData, FMemoryView CompressedData) const = 0;
+};
 
 bool FBlockDecoder::TryDecompressTo(
-	FDecompressionContext& Context,
+	FDecoderContext& Context,
+	const FDecoderSource& Source,
 	const FHeader& Header,
-	const FCompositeBuffer& CompressedData,
-	FMutableMemoryView RawView,
-	uint64 RawOffset) const
+	const uint64 RawOffset,
+	FMutableMemoryView RawView) const
 {
-	if (Header.TotalRawSize < RawOffset + RawView.GetSize() || Header.TotalCompressedSize != CompressedData.GetSize())
+	if (Header.TotalRawSize < RawOffset + RawView.GetSize())
 	{
 		return false;
 	}
 
 	const uint64 BlockSize = uint64(1) << Header.BlockSizeExponent;
+	const uint64 FirstBlockIndex = uint64(RawOffset / BlockSize);
+	const uint64 LastBlockIndex = uint64((RawOffset + RawView.GetSize() - 1) / BlockSize);
+	const uint64 LastBlockSize = BlockSize - ((Header.BlockCount * BlockSize) - Header.TotalRawSize);
 
-	FMemoryView BlockSizeView = CompressedData.ViewOrCopyRange(sizeof(FHeader), Header.BlockCount * sizeof(uint32), Context.CompressedBlockSizes);
-	TArrayView64<uint32 const> CompressedBlockSizes(reinterpret_cast<const uint32*>(BlockSizeView.GetData()), Header.BlockCount);
+	uint64 RawBlockOffset = RawOffset % BlockSize;
 
-	const uint64 FirstBlockIndex	= uint64(RawOffset / BlockSize);
-	const uint64 LastBlockIndex		= uint64((RawOffset + RawView.GetSize() - 1) / BlockSize);
-	const uint64 LastBlockSize		= BlockSize - ((Header.BlockCount * BlockSize) - Header.TotalRawSize);
-	uint64 OffsetInFirstBlock		= RawOffset % BlockSize;
-	uint64 CompressedOffset			= sizeof(FHeader) + uint64(Header.BlockCount) * sizeof(uint32);
-	uint64 RemainingRawSize			= RawView.GetSize();
-
-	for (uint64 BlockIndex = 0; BlockIndex < FirstBlockIndex; BlockIndex++)
-	{
-		const uint32 CompressedBlockSize = NETWORK_ORDER32(CompressedBlockSizes[BlockIndex]);
-		CompressedOffset += CompressedBlockSize;
-	}
+	const TConstArrayView64<uint32> CompressedBlockSizes(reinterpret_cast<const uint32*>(&Header + 1), Header.BlockCount);
+	uint64 CompressedOffset = sizeof(FHeader) + uint64(Header.BlockCount) * sizeof(uint32);
+	CompressedOffset += Algo::Accumulate(CompressedBlockSizes.Left(FirstBlockIndex), uint64(0));
 
 	for (uint64 BlockIndex = FirstBlockIndex; BlockIndex <= LastBlockIndex; BlockIndex++)
 	{
-		const uint64 UncompressedBlockSize	= BlockIndex == Header.BlockCount - 1 ? LastBlockSize : BlockSize;
-		const uint32 CompressedBlockSize	= NETWORK_ORDER32(CompressedBlockSizes[BlockIndex]);
-		const bool IsCompressed				= CompressedBlockSize < UncompressedBlockSize;
+		const uint64 RawBlockSize = BlockIndex == Header.BlockCount - 1 ? LastBlockSize : BlockSize;
+		const uint64 RawBlockReadSize = FMath::Min(RawView.GetSize(), RawBlockSize - RawBlockOffset);
+		const uint32 CompressedBlockSize = CompressedBlockSizes[BlockIndex];
+		const bool bIsCompressed = CompressedBlockSize < RawBlockSize;
 
-		const uint64 BytesToUncompress = OffsetInFirstBlock > 0 
-			? FMath::Min<uint64>(RawView.GetSize(), UncompressedBlockSize - OffsetInFirstBlock)
-			: FMath::Min<uint64>(RemainingRawSize, BlockSize);
-
-		FMemoryView CompressedBlock = CompressedData.ViewOrCopyRange(CompressedOffset, CompressedBlockSize, Context.CompressedBlock);
-
-		if (IsCompressed)
+		if (bIsCompressed)
 		{
-			FMutableMemoryView UncompressedBlock = RawView.Left(BytesToUncompress);
-
-			const bool bIsAligned = BytesToUncompress == UncompressedBlockSize;
-			if (!bIsAligned)
+			if (Context.RawBlockIndex == BlockIndex)
 			{
-				// Decompress to a temporary buffer when the first or the last block reads are not aligned with the block boundaries.
-				if (Context.UncompressedBlock.GetSize() != BlockSize)
+				RawView.Left(RawBlockReadSize).CopyFrom(Context.RawBlock.GetView().Mid(RawBlockOffset, RawBlockReadSize));
+			}
+			else
+			{
+				FMutableMemoryView RawBlock;
+				if (RawBlockReadSize == RawBlockSize)
 				{
-					Context.UncompressedBlock = FUniqueBuffer::Alloc(BlockSize);
+					RawBlock = RawView.Left(RawBlockSize);
+				}
+				else
+				{
+					if (Context.RawBlock.GetSize() < RawBlockSize)
+					{
+						Context.RawBlock = FUniqueBuffer::Alloc(BlockSize);
+					}
+					RawBlock = Context.RawBlock.GetView().Left(RawBlockSize);
+					Context.RawBlockIndex = BlockIndex;
 				}
-				UncompressedBlock = Context.UncompressedBlock.GetView().Left(UncompressedBlockSize);
-			}
 
-			if (!DecompressBlock(UncompressedBlock, CompressedBlock))
-			{
-				return false;
-			}
+				const FMemoryView CompressedBlock = Source.ReadOrView(CompressedOffset, CompressedBlockSize, Context);
+				if (CompressedBlock.IsEmpty() || !DecompressBlock(RawBlock, CompressedBlock))
+				{
+					return false;
+				}
 
-			if (!bIsAligned)
-			{
-				RawView.CopyFrom(UncompressedBlock.Mid(OffsetInFirstBlock, BytesToUncompress));
+				if (RawBlockReadSize != RawBlockSize)
+				{
+					RawView.CopyFrom(RawBlock.Mid(RawBlockOffset, RawBlockReadSize));
+				}
 			}
 		}
 		else
 		{
-			RawView.CopyFrom(CompressedBlock.Mid(OffsetInFirstBlock, BytesToUncompress));
+			Source.Read(CompressedOffset + RawBlockOffset, RawView.Left(RawBlockReadSize));
 		}
 
-		OffsetInFirstBlock = 0;
-		RemainingRawSize -= BytesToUncompress;
+		RawBlockOffset = 0;
 		CompressedOffset += CompressedBlockSize;
-		RawView += BytesToUncompress;
+		RawView += RawBlockReadSize;
 	}
 
-	return RemainingRawSize == 0;
+	return RawView.GetSize() == 0;
+}
+
+FCompositeBuffer FBlockDecoder::DecompressToComposite(
+	FDecoderContext& Context,
+	const FDecoderSource& Source,
+	const FHeader& Header,
+	const uint64 RawOffset,
+	const uint64 RawSize) const
+{
+	FUniqueBuffer Buffer = FUniqueBuffer::Alloc(RawSize);
+	if (TryDecompressTo(Context, Source, Header, RawOffset, Buffer))
+	{
+		return FCompositeBuffer(Buffer.MoveToShared());
+	}
+	return FCompositeBuffer();
 }
 
 ///////////////////////////////////////////////////////////////////////////////////////////////////////////////////////
@@ -659,7 +603,168 @@ bool FHeader::IsValid(const FCompositeBuffer& CompressedData)
 
 ///////////////////////////////////////////////////////////////////////////////////////////////////////////////////////
 
-} // UE::CompressedBuffer
+static const FHeader* TryReadHeader(FDecoderContext& Context, FArchive& Ar)
+{
+	if (Context.HeaderOffset != MAX_uint64)
+	{
+		return static_cast<const FHeader*>(Context.Header.GetData());
+	}
+
+	check(Ar.IsLoading());
+	const int64 Offset = Ar.Tell();
+	if (Offset == INDEX_NONE)
+	{
+		return nullptr;
+	}
+
+	FHeader Header;
+	Ar.Serialize(&Header, sizeof(FHeader));
+	Header.ByteSwap();
+
+	if (const FDecoder* const Decoder = GetDecoder(Header.Method); Decoder && Header.Magic == FHeader::ExpectedMagic)
+	{
+		const uint64 HeaderSize = Decoder->GetHeaderSize(Header);
+		if (Context.Header.GetSize() < HeaderSize)
+		{
+			Context.Header = FUniqueBuffer::Alloc(FMath::Max(FMath::RoundUpToPowerOfTwo64(HeaderSize), DefaultHeaderSize));
+		}
+
+		const FMutableMemoryView HeaderView = Context.Header.GetView().Left(HeaderSize);
+		const FMutableMemoryView HeaderTail = HeaderView.CopyFrom(MakeMemoryView(&Header, &Header + 1));
+		Ar.Serialize(HeaderTail.GetData(), static_cast<int64>(HeaderTail.GetSize()));
+
+		FHeader* const HeaderCopy = static_cast<FHeader*>(HeaderView.GetData());
+		HeaderCopy->ByteSwap();
+		if (Header.Crc32 == FHeader::CalculateCrc32(HeaderView))
+		{
+			Context.HeaderOffset = uint64(Offset);
+			Decoder->DecodeHeader(HeaderView);
+			return HeaderCopy;
+		}
+	}
+
+	return nullptr;
+}
+
+static const FHeader* TryReadHeader(FDecoderContext& Context, const FCompositeBuffer& Buffer)
+{
+	if (Context.HeaderOffset != MAX_uint64)
+	{
+		return static_cast<const FHeader*>(Context.Header.GetData());
+	}
+
+	if (Buffer.GetSize() < sizeof(FHeader))
+	{
+		return nullptr;
+	}
+
+	FHeader Header;
+	Buffer.CopyTo(MakeMemoryView(&Header, &Header + 1));
+	Header.ByteSwap();
+
+	if (const FDecoder* const Decoder = GetDecoder(Header.Method); Decoder && Header.Magic == FHeader::ExpectedMagic)
+	{
+		const uint64 HeaderSize = Decoder->GetHeaderSize(Header);
+		if (Context.Header.GetSize() < HeaderSize)
+		{
+			Context.Header = FUniqueBuffer::Alloc(FMath::Max(FMath::RoundUpToPowerOfTwo64(HeaderSize), DefaultHeaderSize));
+		}
+
+		const FMutableMemoryView HeaderView = Context.Header.GetView().Left(HeaderSize);
+		Buffer.CopyTo(HeaderView);
+
+		if (Header.Crc32 == FHeader::CalculateCrc32(HeaderView))
+		{
+			Context.HeaderOffset = 0;
+			Decoder->DecodeHeader(HeaderView);
+			return static_cast<FHeader*>(HeaderView.GetData());
+		}
+	}
+
+	return nullptr;
+}
+
+///////////////////////////////////////////////////////////////////////////////////////////////////////////////////////
+
+class FArchiveDecoderSource final : public FDecoderSource
+{
+public:
+	explicit FArchiveDecoderSource(FArchive& InArchive, const uint64 InBaseOffset)
+		: Archive(InArchive)
+		, BaseOffset(InBaseOffset)
+	{
+	}
+
+	bool Read(uint64 Offset, FMutableMemoryView Data) const final
+	{
+		Archive.Seek(int64(BaseOffset + Offset));
+		Archive.Serialize(Data.GetData(), int64(Data.GetSize()));
+		return !Archive.IsError();
+	}
+
+	FMemoryView ReadOrView(uint64 Offset, uint64 Size, FDecoderContext& Context) const final
+	{
+		if (Context.CompressedBlock.GetSize() < Size)
+		{
+			Context.CompressedBlock = FUniqueBuffer::Alloc(FMath::Max(FMath::RoundUpToPowerOfTwo64(Size), DefaultBlockSize));
+		}
+		const FMutableMemoryView View = Context.CompressedBlock.GetView().Left(Size);
+		return Read(Offset, View) ? View : FMemoryView();
+	}
+
+	FCompositeBuffer ReadToComposite(uint64 Offset, uint64 Size) const final
+	{
+		FUniqueBuffer Buffer = FUniqueBuffer::Alloc(Size);
+		if (Read(Offset, Buffer))
+		{
+			return FCompositeBuffer(Buffer.MoveToShared());
+		}
+		return FCompositeBuffer();
+	}
+
+private:
+	FArchive& Archive;
+	const uint64 BaseOffset;
+};
+
+class FBufferDecoderSource final : public FDecoderSource
+{
+public:
+	explicit FBufferDecoderSource(const FCompositeBuffer& InBuffer)
+		: Buffer(InBuffer)
+	{
+	}
+
+	bool Read(uint64 Offset, FMutableMemoryView Data) const final
+	{
+		if (Offset + Data.GetSize() <= Buffer.GetSize())
+		{
+			Buffer.CopyTo(Data, Offset);
+			return true;
+		}
+		return false;
+	}
+
+	FMemoryView ReadOrView(uint64 Offset, uint64 Size, FDecoderContext& Context) const final
+	{
+		return Buffer.ViewOrCopyRange(Offset, Size, Context.CompressedBlock, [](uint64 BufferSize) -> FUniqueBuffer
+		{
+			return FUniqueBuffer::Alloc(FMath::Max(FMath::RoundUpToPowerOfTwo64(BufferSize), DefaultBlockSize));
+		});
+	}
+
+	FCompositeBuffer ReadToComposite(uint64 Offset, uint64 Size) const final
+	{
+		return Buffer.Mid(Offset, Size).MakeOwned();
+	}
+
+private:
+	const FCompositeBuffer& Buffer;
+};
+
+///////////////////////////////////////////////////////////////////////////////////////////////////////////////////////
+
+} // UE::CompressedBuffer::Private
 
 FCompressedBuffer FCompressedBuffer::Compress(const FCompositeBuffer& RawData)
 {
@@ -677,7 +782,7 @@ FCompressedBuffer FCompressedBuffer::Compress(
 	ECompressedBufferCompressionLevel CompressionLevel,
 	uint64 BlockSize)
 {
-	using namespace UE::CompressedBuffer;
+	using namespace UE::CompressedBuffer::Private;
 
 	if (BlockSize == 0)
 	{
@@ -708,34 +813,34 @@ FCompressedBuffer FCompressedBuffer::Compress(
 FCompressedBuffer FCompressedBuffer::FromCompressed(const FCompositeBuffer& InCompressedData)
 {
 	FCompressedBuffer Local;
-	Local.CompressedData = UE::CompressedBuffer::ValidBufferOrEmpty(InCompressedData);
+	Local.CompressedData = UE::CompressedBuffer::Private::ValidBufferOrEmpty(InCompressedData);
 	return Local;
 }
 
 FCompressedBuffer FCompressedBuffer::FromCompressed(FCompositeBuffer&& InCompressedData)
 {
 	FCompressedBuffer Local;
-	Local.CompressedData = UE::CompressedBuffer::ValidBufferOrEmpty(MoveTemp(InCompressedData));
+	Local.CompressedData = UE::CompressedBuffer::Private::ValidBufferOrEmpty(MoveTemp(InCompressedData));
 	return Local;
 }
 
 FCompressedBuffer FCompressedBuffer::FromCompressed(const FSharedBuffer& InCompressedData)
 {
 	FCompressedBuffer Local;
-	Local.CompressedData = UE::CompressedBuffer::ValidBufferOrEmpty(InCompressedData);
+	Local.CompressedData = UE::CompressedBuffer::Private::ValidBufferOrEmpty(InCompressedData);
 	return Local;
 }
 
 FCompressedBuffer FCompressedBuffer::FromCompressed(FSharedBuffer&& InCompressedData)
 {
 	FCompressedBuffer Local;
-	Local.CompressedData = UE::CompressedBuffer::ValidBufferOrEmpty(MoveTemp(InCompressedData));
+	Local.CompressedData = UE::CompressedBuffer::Private::ValidBufferOrEmpty(MoveTemp(InCompressedData));
 	return Local;
 }
 
 FCompressedBuffer FCompressedBuffer::FromCompressed(FArchive& Ar)
 {
-	using namespace UE::CompressedBuffer;
+	using namespace UE::CompressedBuffer::Private;
 	check(Ar.IsLoading());
 
 	FHeader Header;
@@ -749,7 +854,7 @@ FCompressedBuffer FCompressedBuffer::FromCompressed(FArchive& Ar)
 		Header.ByteSwap();
 		const FMutableMemoryView MutableView = MutableBuffer.GetView().CopyFrom(MakeMemoryView(&Header, &Header + 1));
 		Ar.Serialize(MutableView.GetData(), static_cast<int64>(MutableView.GetSize()));
-		Local.CompressedData = UE::CompressedBuffer::ValidBufferOrEmpty(MutableBuffer.MoveToShared());
+		Local.CompressedData = ValidBufferOrEmpty(MutableBuffer.MoveToShared());
 	}
 	if (Local.IsNull())
 	{
@@ -758,98 +863,58 @@ FCompressedBuffer FCompressedBuffer::FromCompressed(FArchive& Ar)
 	return Local;
 }
 
+uint64 FCompressedBuffer::GetCompressedSize() const
+{
+	return CompressedData.GetSize();
+}
+
 uint64 FCompressedBuffer::GetRawSize() const
 {
-	return CompressedData ? UE::CompressedBuffer::FHeader::Read(CompressedData).TotalRawSize : 0;
+	return CompressedData ? UE::CompressedBuffer::Private::FHeader::Read(CompressedData).TotalRawSize : 0;
 }
 
 FBlake3Hash FCompressedBuffer::GetRawHash() const
 {
-	return CompressedData ? UE::CompressedBuffer::FHeader::Read(CompressedData).RawHash : FBlake3Hash();
+	return CompressedData ? UE::CompressedBuffer::Private::FHeader::Read(CompressedData).RawHash : FBlake3Hash();
 }
 
 bool FCompressedBuffer::TryGetCompressParameters(
 	ECompressedBufferCompressor& OutCompressor,
-	ECompressedBufferCompressionLevel& OutCompressionLevel) const
-{
-	using namespace UE::CompressedBuffer;
-	if (CompressedData)
-	{
-		switch (const FHeader Header = FHeader::Read(CompressedData); Header.Method)
-		{
-		case EMethod::None:
-			OutCompressor = ECompressedBufferCompressor::NotSet;
-			OutCompressionLevel = ECompressedBufferCompressionLevel::None;
-			return true;
-		case EMethod::Oodle:
-			OutCompressor = ECompressedBufferCompressor(Header.Compressor);
-			OutCompressionLevel = ECompressedBufferCompressionLevel(Header.CompressionLevel);
-			return true;
-		default:
-			break;
-		}
-	}
-	return false;
-}
-
-bool FCompressedBuffer::TryGetBlockSize(uint32 &OutBlockSize) const
+	ECompressedBufferCompressionLevel& OutCompressionLevel,
+	uint64& OutBlockSize) const
 {
-	using namespace UE::CompressedBuffer;
+	using namespace UE::CompressedBuffer::Private;
 	if (CompressedData)
 	{
-		const FHeader Header = FHeader::Read(CompressedData);
-		OutBlockSize = uint64(1) << Header.BlockSizeExponent;
-		return true;
+		return FHeader::Read(CompressedData).TryGetCompressParameters(OutCompressor, OutCompressionLevel, OutBlockSize);
 	}
 	return false;
 }
 
 bool FCompressedBuffer::TryDecompressTo(FMutableMemoryView RawView, uint64 RawOffset) const
 {
-	using namespace UE::CompressedBuffer;
 	if (CompressedData)
 	{
-		const FHeader Header = FHeader::Read(CompressedData);
-		if (const FDecoder* const Decoder = GetDecoder(Header.Method))
-		{
-			FDecompressionContext DecompressionContext;
-			return Decoder->TryDecompressTo(DecompressionContext, Header, CompressedData, RawView, RawOffset);
-		}
+		return FCompressedBufferReader(*this).TryDecompressTo(RawView, RawOffset);
 	}
 	return false;
 }
 
 FSharedBuffer FCompressedBuffer::Decompress(uint64 RawOffset, uint64 RawSize) const
 {
-	using namespace UE::CompressedBuffer;
 	if (CompressedData)
 	{
-		const FHeader Header = FHeader::Read(CompressedData);
-		if (const FDecoder* const Decoder = GetDecoder(Header.Method))
-		{
-			const uint64 TotalRawSize = RawSize < ~uint64(0) ? RawSize : Header.TotalRawSize - RawOffset;
-			FUniqueBuffer RawData = FUniqueBuffer::Alloc(TotalRawSize);
-			
-			FDecompressionContext Context;
-			if (Decoder->TryDecompressTo(Context, Header, CompressedData, RawData, RawOffset))
-			{
-				return RawData.MoveToShared();
-			}
-		}
+		return FCompressedBufferReader(*this).Decompress(RawOffset, RawSize);
 	}
 	return FSharedBuffer();
 }
 
 FCompositeBuffer FCompressedBuffer::DecompressToComposite() const
 {
-	using namespace UE::CompressedBuffer;
+	using namespace UE::CompressedBuffer::Private;
 	if (CompressedData)
 	{
-		const FHeader Header = FHeader::Read(CompressedData);
-		if (const FDecoder* const Decoder = GetDecoder(Header.Method))
-		{
-			return Decoder->Decompress(Header, CompressedData);
-		}
+		return FCompressedBufferReader(*this).DecompressToComposite();
 	}
 	return FCompositeBuffer();
 }
@@ -874,29 +939,222 @@ FArchive& operator<<(FArchive& Ar, FCompressedBuffer& Buffer)
 
 ///////////////////////////////////////////////////////////////////////////////////////////////////////////////////////
 
-class FCompressedBufferDecoder::FImpl
+FCompressedBufferReader::FCompressedBufferReader(FArchive& Archive)
 {
-public:
-	bool TryDecompressTo(const FCompressedBuffer& CompressedBuffer, FMutableMemoryView RawView, uint64 RawOffset = 0)
+	SetSource(Archive);
+}
+
+FCompressedBufferReader::FCompressedBufferReader(const FCompressedBuffer& Buffer)
+{
+	SetSource(Buffer);
+}
+
+void FCompressedBufferReader::ResetBuffers()
+{
+	using namespace UE::CompressedBuffer::Private;
+	if (SourceArchive && Context.HeaderOffset != MAX_uint64)
+	{
+		SourceArchive->Seek(int64(Context.HeaderOffset));
+	}
+	Context = FDecoderContext();
+}
+
+void FCompressedBufferReader::ResetSource()
+{
+	Context.HeaderOffset = MAX_uint64;
+	Context.RawBlockIndex = MAX_uint64;
+	SourceArchive = nullptr;
+	SourceBuffer = nullptr;
+}
+
+void FCompressedBufferReader::SetSource(FArchive& Archive)
+{
+	if (SourceArchive == &Archive)
 	{
-		using namespace UE::CompressedBuffer;
-		if (FCompositeBuffer CompressedData = CompressedBuffer.GetCompressed())
+		return;
+	}
+	Context.HeaderOffset = MAX_uint64;
+	Context.RawBlockIndex = MAX_uint64;
+	SourceArchive = &Archive;
+	SourceBuffer = nullptr;
+}
+
+void FCompressedBufferReader::SetSource(const FCompressedBuffer& Buffer)
+{
+	if (SourceBuffer == &Buffer)
+	{
+		return;
+	}
+	Context.HeaderOffset = MAX_uint64;
+	Context.RawBlockIndex = MAX_uint64;
+	SourceArchive = nullptr;
+	SourceBuffer = &Buffer;
+}
+
+uint64 FCompressedBufferReader::GetCompressedSize()
+{
+	using namespace UE::CompressedBuffer::Private;
+	if (const FCompressedBuffer* const Buffer = SourceBuffer)
+	{
+		return Buffer->GetCompressedSize();
+	}
+	if (const FHeader* const Header = TryReadHeader())
+	{
+		return Header->TotalCompressedSize;
+	}
+	return 0;
+}
+
+uint64 FCompressedBufferReader::GetRawSize()
+{
+	using namespace UE::CompressedBuffer::Private;
+	if (const FCompressedBuffer* const Buffer = SourceBuffer)
+	{
+		return Buffer->GetRawSize();
+	}
+	if (const FHeader* const Header = TryReadHeader())
+	{
+		return Header->TotalRawSize;
+	}
+	return 0;
+}
+
+FBlake3Hash FCompressedBufferReader::GetRawHash()
+{
+	using namespace UE::CompressedBuffer::Private;
+	if (const FCompressedBuffer* const Buffer = SourceBuffer)
+	{
+		return Buffer->GetRawHash();
+	}
+	if (const FHeader* const Header = TryReadHeader())
+	{
+		return Header->RawHash;
+	}
+	return FBlake3Hash();
+}
+
+bool FCompressedBufferReader::TryGetCompressParameters(
+	ECompressedBufferCompressor& OutCompressor,
+	ECompressedBufferCompressionLevel& OutCompressionLevel,
+	uint64& OutBlockSize)
+{
+	using namespace UE::CompressedBuffer::Private;
+	if (const FCompressedBuffer* const Buffer = SourceBuffer)
+	{
+		return Buffer->TryGetCompressParameters(OutCompressor, OutCompressionLevel, OutBlockSize);
+	}
+	if (const FHeader* const Header = TryReadHeader())
+	{
+		return Header->TryGetCompressParameters(OutCompressor, OutCompressionLevel, OutBlockSize);
+	}
+	return false;
+}
+
+bool FCompressedBufferReader::TryDecompressTo(FMutableMemoryView RawView, uint64 RawOffset)
+{
+	using namespace UE::CompressedBuffer::Private;
+	if (const FHeader* const Header = TryReadHeader())
+	{
+		const uint64 TotalRawSize = Header->TotalRawSize;
+		if (RawOffset <= TotalRawSize && RawView.GetSize() <= TotalRawSize - RawOffset)
 		{
-			const FHeader Header = FHeader::Read(CompressedData);
-			if (const FDecoder* const Decoder = GetDecoder(Header.Method))
+			if (const FDecoder* const Decoder = GetDecoder(Header->Method))
 			{
-				return Decoder->TryDecompressTo(DecompressionContext, Header, CompressedData, RawView, RawOffset);
+				if (Decoder->TryDecompressTo(Context,
+					SourceArchive
+						? ImplicitConv<const FDecoderSource&>(FArchiveDecoderSource(*SourceArchive, Context.HeaderOffset))
+						: ImplicitConv<const FDecoderSource&>(FBufferDecoderSource(SourceBuffer->GetCompressed())),
+					*Header, RawOffset, RawView))
+				{
+					return true;
+				}
+			}
+		}
+	}
+	return false;
+}
+
+FSharedBuffer FCompressedBufferReader::Decompress(const uint64 RawOffset, const uint64 RawSize)
+{
+	using namespace UE::CompressedBuffer::Private;
+	if (const FHeader* const Header = TryReadHeader())
+	{
+		const uint64 TotalRawSize = Header->TotalRawSize;
+		const uint64 RawSizeToCopy = RawSize == MAX_uint64 ? TotalRawSize - RawOffset : RawSize;
+		if (RawOffset <= TotalRawSize && RawSizeToCopy <= TotalRawSize - RawOffset)
+		{
+			if (const FDecoder* const Decoder = GetDecoder(Header->Method))
+			{
+				FUniqueBuffer RawData = FUniqueBuffer::Alloc(RawSizeToCopy);
+				if (Decoder->TryDecompressTo(Context,
+					SourceArchive
+						? ImplicitConv<const FDecoderSource&>(FArchiveDecoderSource(*SourceArchive, Context.HeaderOffset))
+						: ImplicitConv<const FDecoderSource&>(FBufferDecoderSource(SourceBuffer->GetCompressed())),
+					*Header, RawOffset, RawData))
+				{
+					return RawData.MoveToShared();
+				}
+			}
+		}
+	}
+	return FSharedBuffer();
+}
+
+FCompositeBuffer FCompressedBufferReader::DecompressToComposite(const uint64 RawOffset, const uint64 RawSize)
+{
+	using namespace UE::CompressedBuffer::Private;
+	if (const FHeader* const Header = TryReadHeader())
+	{
+		const uint64 TotalRawSize = Header->TotalRawSize;
+		const uint64 RawSizeToCopy = RawSize == MAX_uint64 ? TotalRawSize - RawOffset : RawSize;
+		if (RawOffset <= TotalRawSize && RawSizeToCopy <= TotalRawSize - RawOffset)
+		{
+			if (const FDecoder* const Decoder = GetDecoder(Header->Method))
+			{
+				return Decoder->DecompressToComposite(Context,
+					SourceArchive
+						? ImplicitConv<const FDecoderSource&>(FArchiveDecoderSource(*SourceArchive, Context.HeaderOffset))
+						: ImplicitConv<const FDecoderSource&>(FBufferDecoderSource(SourceBuffer->GetCompressed())),
+					*Header, RawOffset, RawSizeToCopy);
 			}
 		}
-		return false;
+	}
+	return FCompositeBuffer();
+}
+
+const UE::CompressedBuffer::Private::FHeader* FCompressedBufferReader::TryReadHeader()
+{
+	using namespace UE::CompressedBuffer;
+	if (FArchive* const Archive = SourceArchive)
+	{
+		return Private::TryReadHeader(Context, *Archive);
+	}
+	if (const FCompressedBuffer* const Buffer = SourceBuffer)
+	{
+		return Private::TryReadHeader(Context, Buffer->GetCompressed());
+	}
+	return nullptr;
+}
+
+///////////////////////////////////////////////////////////////////////////////////////////////////////////////////////
+
+class FCompressedBufferDecoder::FImpl
+{
+public:
+	bool TryDecompressTo(const FCompressedBuffer& CompressedBuffer, FMutableMemoryView RawView, uint64 RawOffset)
+	{
+		using namespace UE::CompressedBuffer::Private;
+		Reader.SetSource(CompressedBuffer);
+		ON_SCOPE_EXIT { Reader.ResetSource(); };
+		return Reader.TryDecompressTo(RawView, RawOffset);
 	}
 
 private:
-	UE::CompressedBuffer::FDecompressionContext DecompressionContext;
+	FCompressedBufferReader Reader;
 };
 
 FCompressedBufferDecoder::FCompressedBufferDecoder()
-: Impl(new FImpl()) 
+: Impl(new FImpl())
 {
 }
 
@@ -906,8 +1164,7 @@ FCompressedBufferDecoder::~FCompressedBufferDecoder()
 
 bool FCompressedBufferDecoder::TryDecompressTo(const FCompressedBuffer& CompressedBuffer, FMutableMemoryView RawView, uint64 RawOffset)
 {
-	return Impl->TryDecompressTo(CompressedBuffer,RawView, RawOffset);
+	return Impl->TryDecompressTo(CompressedBuffer, RawView, RawOffset);
 }
 
 ///////////////////////////////////////////////////////////////////////////////////////////////////////////////////////
-
diff --git a/Engine/Source/Runtime/Core/Private/Memory/CompositeBuffer.cpp b/Engine/Source/Runtime/Core/Private/Memory/CompositeBuffer.cpp
--- a/Engine/Source/Runtime/Core/Private/Memory/CompositeBuffer.cpp
+++ b/Engine/Source/Runtime/Core/Private/Memory/CompositeBuffer.cpp
@@ -68,10 +68,14 @@ FCompositeBuffer FCompositeBuffer::Mid(uint64 Offset, uint64 Size) const
 	return Buffer;
 }
 
-FMemoryView FCompositeBuffer::ViewOrCopyRange(uint64 Offset, uint64 Size, FUniqueBuffer& CopyBuffer) const
+FMemoryView FCompositeBuffer::ViewOrCopyRange(
+	const uint64 Offset,
+	const uint64 Size,
+	FUniqueBuffer& CopyBuffer,
+	TFunctionRef<FUniqueBuffer (uint64 Size)> Allocator) const
 {
 	FMemoryView View;
-	IterateRange(Offset, Size, [Size, &View, &CopyBuffer, WriteView = FMutableMemoryView()](FMemoryView Segment) mutable
+	IterateRange(Offset, Size, [Size, &View, &CopyBuffer, &Allocator, WriteView = FMutableMemoryView()](FMemoryView Segment) mutable
 		{
 			if (Size == Segment.GetSize())
 			{
@@ -83,7 +87,7 @@ FMemoryView FCompositeBuffer::ViewOrCopyRange(uint64 Offset, uint64 Size, FUniqu
 				{
 					if (CopyBuffer.GetSize() < Size)
 					{
-						CopyBuffer = FUniqueBuffer::Alloc(Size);
+						CopyBuffer = Allocator(Size);
 					}
 					View = WriteView = CopyBuffer.GetView().Left(Size);
 				}
diff --git a/Engine/Source/Runtime/Core/Private/Tests/Compression/CompressedBufferTest.cpp b/Engine/Source/Runtime/Core/Private/Tests/Compression/CompressedBufferTest.cpp
--- a/Engine/Source/Runtime/Core/Private/Tests/Compression/CompressedBufferTest.cpp
+++ b/Engine/Source/Runtime/Core/Private/Tests/Compression/CompressedBufferTest.cpp
@@ -5,6 +5,7 @@
 #include "Compression/OodleDataCompression.h"
 #include "Hash/Blake3.h"
 #include "Misc/AutomationTest.h"
+#include "Serialization/BufferReader.h"
 
 #if WITH_DEV_AUTOMATION_TESTS
 
@@ -29,8 +30,9 @@ bool FCompressedBufferTest::RunTest(const FString& Parameters)
 
 	// Test Method None
 	{
-		ECompressedBufferCompressor Compressor;
-		ECompressedBufferCompressionLevel CompressionLevel;
+		ECompressedBufferCompressor Compressor = ECompressedBufferCompressor::Kraken;
+		ECompressedBufferCompressionLevel CompressionLevel = ECompressedBufferCompressionLevel::Normal;
+		uint64 BlockSize = MAX_uint64;
 
 		const FCompressedBuffer Buffer = FCompressedBuffer::Compress(FSharedBuffer::MakeView(MakeMemoryView(ZeroBuffer)),
 			ECompressedBufferCompressor::NotSet, ECompressedBufferCompressionLevel::None);
@@ -42,8 +44,12 @@ bool FCompressedBufferTest::RunTest(const FString& Parameters)
 		TestEqual(TEXT("FCompressedBuffer::Compress(None).GetRawHash()"), Buffer.GetRawHash(), ZeroBufferHash);
 		TestEqual(TEXT("FCompressedBuffer::Compress(None).Decompress()"), FBlake3::HashBuffer(Buffer.Decompress()), ZeroBufferHash);
 		TestEqual(TEXT("FCompressedBuffer::Compress(None).DecompressToComposite()"), FBlake3::HashBuffer(Buffer.DecompressToComposite()), Buffer.GetRawHash());
-		TestTrue(TEXT("FCompressedBuffer::Compress(None).TryGetCompressParameters()"), Buffer.TryGetCompressParameters(Compressor, CompressionLevel) &&
-			Compressor == ECompressedBufferCompressor::NotSet && CompressionLevel == ECompressedBufferCompressionLevel::None);
+		TestTrue(TEXT("FCompressedBuffer::Compress(None).TryGetCompressParameters()"), Buffer.TryGetCompressParameters(Compressor, CompressionLevel, BlockSize) &&
+			Compressor == ECompressedBufferCompressor::NotSet && CompressionLevel == ECompressedBufferCompressionLevel::None && BlockSize == 0);
+
+		Compressor = ECompressedBufferCompressor::Kraken;
+		CompressionLevel = ECompressedBufferCompressionLevel::Normal;
+		BlockSize = MAX_uint64;
 
 		const FCompressedBuffer BufferCopy = FCompressedBuffer::FromCompressed(Buffer.GetCompressed());
 		TestTrue(TEXT("FCompressedBuffer::Compress(None, Copy)"), (bool)Buffer);
@@ -54,15 +60,15 @@ bool FCompressedBufferTest::RunTest(const FString& Parameters)
 		TestEqual(TEXT("FCompressedBuffer::Compress(None, Copy).GetRawHash()"), Buffer.GetRawHash(), ZeroBufferHash);
 		TestEqual(TEXT("FCompressedBuffer::Compress(None, Copy).Decompress()"), FBlake3::HashBuffer(Buffer.Decompress()), ZeroBufferHash);
 		TestEqual(TEXT("FCompressedBuffer::Compress(None, Copy).DecompressToComposite()"), FBlake3::HashBuffer(Buffer.DecompressToComposite()), Buffer.GetRawHash());
-		TestTrue(TEXT("FCompressedBuffer::Compress(None, Copy).TryGetCompressParameters()"), Buffer.TryGetCompressParameters(Compressor, CompressionLevel) &&
-			Compressor == ECompressedBufferCompressor::NotSet && CompressionLevel == ECompressedBufferCompressionLevel::None);
+		TestTrue(TEXT("FCompressedBuffer::Compress(None, Copy).TryGetCompressParameters()"), Buffer.TryGetCompressParameters(Compressor, CompressionLevel, BlockSize) &&
+			Compressor == ECompressedBufferCompressor::NotSet && CompressionLevel == ECompressedBufferCompressionLevel::None && BlockSize == 0);
 	}
 
 	// Test Method Oodle
 	{
-		
-		ECompressedBufferCompressor Compressor;
-		ECompressedBufferCompressionLevel CompressionLevel;
+		ECompressedBufferCompressor Compressor = ECompressedBufferCompressor::Kraken;
+		ECompressedBufferCompressionLevel CompressionLevel = ECompressedBufferCompressionLevel::Normal;
+		uint64 BlockSize = MAX_uint64;
 
 		const FCompressedBuffer Buffer = FCompressedBuffer::Compress(FSharedBuffer::MakeView(MakeMemoryView(ZeroBuffer)),
 			ECompressedBufferCompressor::Mermaid, ECompressedBufferCompressionLevel::VeryFast);
@@ -74,8 +80,12 @@ bool FCompressedBufferTest::RunTest(const FString& Parameters)
 		TestEqual(TEXT("FCompressedBuffer::Compress(Oodle).GetRawHash()"), Buffer.GetRawHash(), ZeroBufferHash);
 		TestEqual(TEXT("FCompressedBuffer::Compress(Oodle).Decompress()"), FBlake3::HashBuffer(Buffer.Decompress()), ZeroBufferHash);
 		TestEqual(TEXT("FCompressedBuffer::Compress(Oodle).DecompressToComposite()"), FBlake3::HashBuffer(Buffer.DecompressToComposite()), Buffer.GetRawHash());
-		TestTrue(TEXT("FCompressedBuffer::Compress(Oodle).TryGetCompressParameters()"), Buffer.TryGetCompressParameters(Compressor, CompressionLevel) &&
-			Compressor == ECompressedBufferCompressor::Mermaid && CompressionLevel == ECompressedBufferCompressionLevel::VeryFast);
+		TestTrue(TEXT("FCompressedBuffer::Compress(Oodle).TryGetCompressParameters()"), Buffer.TryGetCompressParameters(Compressor, CompressionLevel, BlockSize) &&
+			Compressor == ECompressedBufferCompressor::Mermaid && CompressionLevel == ECompressedBufferCompressionLevel::VeryFast && FMath::IsPowerOfTwo(BlockSize));
+
+		Compressor = ECompressedBufferCompressor::Kraken;
+		CompressionLevel = ECompressedBufferCompressionLevel::Normal;
+		BlockSize = MAX_uint64;
 
 		const FCompressedBuffer BufferCopy = FCompressedBuffer::FromCompressed(Buffer.GetCompressed());
 		TestTrue(TEXT("FCompressedBuffer::Compress(Oodle, Copy)"), (bool)Buffer);
@@ -86,8 +96,8 @@ bool FCompressedBufferTest::RunTest(const FString& Parameters)
 		TestEqual(TEXT("FCompressedBuffer::Compress(Oodle, Copy).GetRawHash()"), Buffer.GetRawHash(), ZeroBufferHash);
 		TestEqual(TEXT("FCompressedBuffer::Compress(Oodle, Copy).Decompress()"), FBlake3::HashBuffer(Buffer.Decompress()), ZeroBufferHash);
 		TestEqual(TEXT("FCompressedBuffer::Compress(Oodle, Copy).DecompressToComposite()"), FBlake3::HashBuffer(Buffer.DecompressToComposite()), Buffer.GetRawHash());
-		TestTrue(TEXT("FCompressedBuffer::Compress(Oodle, Copy).TryGetCompressParameters()"), Buffer.TryGetCompressParameters(Compressor, CompressionLevel) &&
-			Compressor == ECompressedBufferCompressor::Mermaid && CompressionLevel == ECompressedBufferCompressionLevel::VeryFast);
+		TestTrue(TEXT("FCompressedBuffer::Compress(Oodle, Copy).TryGetCompressParameters()"), Buffer.TryGetCompressParameters(Compressor, CompressionLevel, BlockSize) &&
+			Compressor == ECompressedBufferCompressor::Mermaid && CompressionLevel == ECompressedBufferCompressionLevel::VeryFast && FMath::IsPowerOfTwo(BlockSize));
 	}
 
 	return true;
@@ -123,18 +133,36 @@ bool FCompressedBufferDecompressTest::RunTest(const FString& Parameters)
 		return MakeArrayView(static_cast<const uint64*>(View.GetData()), static_cast<int32>(View.GetSize() / sizeof(uint64)));
 	};
 
+	FCompressedBufferReader Reader;
+
 	// Test decompress with offset and size
 	{
-		const auto UncompressAndValidate = [this, &ValidateData, &CastToArrayView](
+		const auto UncompressAndValidate = [this, &Reader, &ValidateData, &CastToArrayView](
 			const FCompressedBuffer& Compressed,
 			const int32 OffsetCount,
 			const int32 Count,
 			const TConstArrayView<uint64> ExpectedValues)
 		{
-			const FSharedBuffer Uncompressed = Compressed.Decompress(OffsetCount * sizeof(uint64), Count * sizeof(uint64));
-			const TConstArrayView<uint64> Values = CastToArrayView(Uncompressed);
-			TestEqual("UncompressedCount", Values.Num(), Count);
-			ValidateData(Values, ExpectedValues, OffsetCount);
+			Reader.SetSource(Compressed);
+			{
+				const FSharedBuffer Uncompressed = Reader.Decompress(OffsetCount * sizeof(uint64), Count * sizeof(uint64));
+				const TConstArrayView<uint64> Values = CastToArrayView(Uncompressed);
+				TestEqual("UncompressedCount", Values.Num(), Count);
+				ValidateData(Values, ExpectedValues, OffsetCount);
+			}
+			{
+				FUniqueBuffer Uncompressed = FUniqueBuffer::Alloc(Count * sizeof(uint64));
+				const bool bOk = Reader.TryDecompressTo(Uncompressed, OffsetCount * sizeof(uint64));
+				TestTrue("UncompressedOk", bOk);
+				const TConstArrayView<uint64> Values = CastToArrayView(Uncompressed);
+				ValidateData(Values, ExpectedValues, OffsetCount);
+			}
+			{
+				const FSharedBuffer Uncompressed = Compressed.Decompress(OffsetCount * sizeof(uint64), Count * sizeof(uint64));
+				const TConstArrayView<uint64> Values = CastToArrayView(Uncompressed);
+				TestEqual("UncompressedCount", Values.Num(), Count);
+				ValidateData(Values, ExpectedValues, OffsetCount);
+			}
 		};
 
 		constexpr uint64 BlockSize = 64 * sizeof(uint64);
@@ -173,9 +201,25 @@ bool FCompressedBufferDecompressTest::RunTest(const FString& Parameters)
 			BlockSize);
 		
 		constexpr uint64 OffsetCount = 150;
-		const FSharedBuffer Uncompressed = Compressed.Decompress(OffsetCount * sizeof(uint64));
-		const TConstArrayView<uint64> Values = CastToArrayView(Uncompressed);
-		ValidateData(Values, ExpectedValues, OffsetCount);
+		{
+			FSharedBuffer Buffer = Compressed.GetCompressed().ToShared();
+			FBufferReader Ar(const_cast<void*>(Buffer.GetData()), int64(Buffer.GetSize()), /*bFreeOnClose*/ false, /*bIsPersistent*/ true);
+			Reader.SetSource(Ar);
+			const FSharedBuffer Uncompressed = Reader.Decompress(OffsetCount * sizeof(uint64));
+			const TConstArrayView<uint64> Values = CastToArrayView(Uncompressed);
+			ValidateData(Values, ExpectedValues, OffsetCount);
+		}
+		{
+			Reader.SetSource(Compressed);
+			const FSharedBuffer Uncompressed = Reader.Decompress(OffsetCount * sizeof(uint64));
+			const TConstArrayView<uint64> Values = CastToArrayView(Uncompressed);
+			ValidateData(Values, ExpectedValues, OffsetCount);
+		}
+		{
+			const FSharedBuffer Uncompressed = Compressed.Decompress(OffsetCount * sizeof(uint64));
+			const TConstArrayView<uint64> Values = CastToArrayView(Uncompressed);
+			ValidateData(Values, ExpectedValues, OffsetCount);
+		}
 	}
 
 	// Only one block
@@ -192,9 +236,25 @@ bool FCompressedBufferDecompressTest::RunTest(const FString& Parameters)
 
 		constexpr uint64 OffsetCount = 2;
 		constexpr uint64 Count = 50;
-		const FSharedBuffer Uncompressed = Compressed.Decompress(OffsetCount * sizeof(uint64), Count * sizeof(uint64));
-		const TConstArrayView<uint64> Values = CastToArrayView(Uncompressed);
-		ValidateData(Values, ExpectedValues, OffsetCount);
+		{
+			FSharedBuffer Buffer = Compressed.GetCompressed().ToShared();
+			FBufferReader Ar(const_cast<void*>(Buffer.GetData()), int64(Buffer.GetSize()), /*bFreeOnClose*/ false, /*bIsPersistent*/ true);
+			Reader.SetSource(Ar);
+			const FSharedBuffer Uncompressed = Reader.Decompress(OffsetCount * sizeof(uint64), Count * sizeof(uint64));
+			const TConstArrayView<uint64> Values = CastToArrayView(Uncompressed);
+			ValidateData(Values, ExpectedValues, OffsetCount);
+		}
+		{
+			Reader.SetSource(Compressed);
+			const FSharedBuffer Uncompressed = Reader.Decompress(OffsetCount * sizeof(uint64), Count * sizeof(uint64));
+			const TConstArrayView<uint64> Values = CastToArrayView(Uncompressed);
+			ValidateData(Values, ExpectedValues, OffsetCount);
+		}
+		{
+			const FSharedBuffer Uncompressed = Compressed.Decompress(OffsetCount * sizeof(uint64), Count * sizeof(uint64));
+			const TConstArrayView<uint64> Values = CastToArrayView(Uncompressed);
+			ValidateData(Values, ExpectedValues, OffsetCount);
+		}
 	}
 
 	// Uncompressed
@@ -210,17 +270,33 @@ bool FCompressedBufferDecompressTest::RunTest(const FString& Parameters)
 		{
 			constexpr uint64 OffsetCount = 0;
 			constexpr uint64 Count = N;
-			const FSharedBuffer Uncompressed = Compressed.Decompress(OffsetCount * sizeof(uint64), Count * sizeof(uint64));
-			const TConstArrayView<uint64> Values = CastToArrayView(Uncompressed);
-			ValidateData(Values, ExpectedValues, OffsetCount);
+			{
+				Reader.SetSource(Compressed);
+				const FSharedBuffer Uncompressed = Reader.Decompress(OffsetCount * sizeof(uint64), Count * sizeof(uint64));
+				const TConstArrayView<uint64> Values = CastToArrayView(Uncompressed);
+				ValidateData(Values, ExpectedValues, OffsetCount);
+			}
+			{
+				const FSharedBuffer Uncompressed = Compressed.Decompress(OffsetCount * sizeof(uint64), Count * sizeof(uint64));
+				const TConstArrayView<uint64> Values = CastToArrayView(Uncompressed);
+				ValidateData(Values, ExpectedValues, OffsetCount);
+			}
 		}
 
 		{
 			constexpr uint64 OffsetCount = 21;
 			constexpr uint64 Count = 999;
-			const FSharedBuffer Uncompressed = Compressed.Decompress(OffsetCount * sizeof(uint64), Count * sizeof(uint64));
-			const TConstArrayView<uint64> Values = CastToArrayView(Uncompressed);
-			ValidateData(Values, ExpectedValues, OffsetCount);
+			{
+				Reader.SetSource(Compressed);
+				const FSharedBuffer Uncompressed = Reader.Decompress(OffsetCount * sizeof(uint64), Count * sizeof(uint64));
+				const TConstArrayView<uint64> Values = CastToArrayView(Uncompressed);
+				ValidateData(Values, ExpectedValues, OffsetCount);
+			}
+			{
+				const FSharedBuffer Uncompressed = Compressed.Decompress(OffsetCount * sizeof(uint64), Count * sizeof(uint64));
+				const TConstArrayView<uint64> Values = CastToArrayView(Uncompressed);
+				ValidateData(Values, ExpectedValues, OffsetCount);
+			}
 		}
 	}
 
diff --git a/Engine/Source/Runtime/Core/Public/Compression/CompressedBuffer.h b/Engine/Source/Runtime/Core/Public/Compression/CompressedBuffer.h
--- a/Engine/Source/Runtime/Core/Public/Compression/CompressedBuffer.h
+++ b/Engine/Source/Runtime/Core/Public/Compression/CompressedBuffer.h
@@ -4,6 +4,7 @@
 
 #include "Memory/CompositeBuffer.h"
 #include "Memory/MemoryFwd.h"
+#include "Memory/SharedBuffer.h"
 #include "Templates/UniquePtr.h"
 
 class FArchive;
@@ -12,6 +13,8 @@ struct FBlake3Hash;
 namespace FOodleDataCompression { enum class ECompressionLevel : int8; }
 namespace FOodleDataCompression { enum class ECompressor : uint8; }
 
+namespace UE::CompressedBuffer::Private { struct FHeader; }
+
 using ECompressedBufferCompressionLevel = FOodleDataCompression::ECompressionLevel;
 using ECompressedBufferCompressor = FOodleDataCompression::ECompressor;
 
@@ -87,8 +90,8 @@ public:
 	[[nodiscard]] inline const FCompositeBuffer& GetCompressed() const & { return CompressedData; }
 	[[nodiscard]] inline FCompositeBuffer GetCompressed() && { return MoveTemp(CompressedData); }
 
-	/** Returns the size of the compressed data. Zero if this is null. */
-	[[nodiscard]] inline uint64 GetCompressedSize() const { return CompressedData.GetSize(); }
+	/** Returns the size of the compressed data. Zero on error or if this is null. */
+	[[nodiscard]] CORE_API uint64 GetCompressedSize() const;
 
 	/** Returns the size of the raw data. Zero on error or if this is empty or null. */
 	[[nodiscard]] CORE_API uint64 GetRawSize() const;
@@ -107,15 +110,34 @@ public:
 	 */
 	[[nodiscard]] CORE_API bool TryGetCompressParameters(
 		ECompressedBufferCompressor& OutCompressor,
-		ECompressedBufferCompressionLevel& OutCompressionLevel) const;
+		ECompressedBufferCompressionLevel& OutCompressionLevel,
+		uint64& OutBlockSize) const;
 
-	/**
-	 * Returns the compression block size. Zero if the buffer is not compressed.
-	 */
-	[[nodiscard]] CORE_API bool TryGetBlockSize(uint32 &OutBlockSize) const;
+	[[nodiscard]] inline bool TryGetCompressParameters(
+		ECompressedBufferCompressor& OutCompressor,
+		ECompressedBufferCompressionLevel& OutCompressionLevel) const
+	{
+		uint64 BlockSize;
+		return TryGetCompressParameters(OutCompressor, OutCompressionLevel, BlockSize);
+	}
+
+	[[nodiscard]] inline bool TryGetBlockSize(uint32& OutBlockSize) const
+	{
+		ECompressedBufferCompressor Compressor;
+		ECompressedBufferCompressionLevel CompressionLevel;
+		uint64 BlockSize64;
+		if (TryGetCompressParameters(Compressor, CompressionLevel, BlockSize64))
+		{
+			OutBlockSize = uint32(BlockSize64);
+			return true;
+		}
+		return false;
+	}
 
 	/**
-	 * Decompress into a memory view that is less or equal to GetRawSize()
+	 * Decompress into a memory view that is less or equal to the raw size.
+	 *
+	 * @return True if the requested range was decompressed, otherwise false.
 	 */
 	[[nodiscard]] CORE_API bool TryDecompressTo(FMutableMemoryView RawView, uint64 RawOffset = 0) const;
 
@@ -144,6 +166,109 @@ inline const FCompressedBuffer FCompressedBuffer::Null;
 
 CORE_API FArchive& operator<<(FArchive& Ar, FCompressedBuffer& Buffer);
 
+namespace UE::CompressedBuffer::Private
+{
+
+/** A reusable context for the compressed buffer decoder. */
+struct FDecoderContext
+{
+	/** Offset at which the compressed buffer begins, otherwise MAX_uint64. */
+	uint64 HeaderOffset = MAX_uint64;
+	/** Index of the block stored in RawBlock, otherwise MAX_uint64. */
+	uint64 RawBlockIndex = MAX_uint64;
+
+	/** A buffer used to store the header when HeaderOffset is not MAX_uint64. */
+	FUniqueBuffer Header;
+	/** A buffer used to store a raw block when a partial block read is requested. */
+	FUniqueBuffer RawBlock;
+	/** A buffer used to store a compressed block when it was not in contiguous memory. */
+	FUniqueBuffer CompressedBlock;
+};
+
+} // UE::CompressedBuffer::Private
+
+/**
+ * A type that stores the state needed to decompress a compressed buffer.
+ *
+ * The compressed buffer can be in memory or can be loaded from a seekable archive.
+ *
+ * The reader can be reused across multiple source buffers, which allows its temporary buffers to
+ * be reused if they are the right size.
+ *
+ * It is only safe to use the reader from one thread at a time.
+ *
+ * @see FCompressedBuffer
+ */
+class FCompressedBufferReader
+{
+public:
+	/** Construct a reader with no source. */
+	FCompressedBufferReader() = default;
+
+	/** Construct a reader that will read from an archive as needed. */
+	CORE_API explicit FCompressedBufferReader(FArchive& Archive);
+
+	/** Construct a reader from an in-memory compressed buffer. */
+	CORE_API explicit FCompressedBufferReader(const FCompressedBuffer& Buffer);
+
+	/** Release any temporary buffers that have been allocated by the reader. */
+	CORE_API void ResetBuffers();
+
+	/** Clears the reference to the source without releasing temporary buffers. */
+	CORE_API void ResetSource();
+
+	CORE_API void SetSource(FArchive& Archive);
+	CORE_API void SetSource(const FCompressedBuffer& Buffer);
+
+	[[nodiscard]] CORE_API uint64 GetCompressedSize();
+	[[nodiscard]] CORE_API uint64 GetRawSize();
+
+	/** Returns the hash of the raw data. Zero on error or if this is null. */
+	[[nodiscard]] CORE_API FBlake3Hash GetRawHash();
+
+	/**
+	 * Returns the compressor and compression level used by this buffer.
+	 *
+	 * The compressor and compression level may differ from those specified when creating the buffer
+	 * because an incompressible buffer is stored with no compression. Parameters cannot be accessed
+	 * if this is null or uses a method other than Oodle, in which case this returns false.
+	 *
+	 * @return True if parameters were written, otherwise false.
+	 */
+	[[nodiscard]] CORE_API bool TryGetCompressParameters(
+		ECompressedBufferCompressor& OutCompressor,
+		ECompressedBufferCompressionLevel& OutCompressionLevel,
+		uint64& OutBlockSize);
+
+	/**
+	 * Decompress into a memory view that is less or equal to the raw size.
+	 *
+	 * @return True if the requested range was decompressed, otherwise false.
+	 */
+	[[nodiscard]] CORE_API bool TryDecompressTo(FMutableMemoryView RawView, uint64 RawOffset = 0);
+
+	/**
+	 * Decompress into an owned buffer.
+	 *
+	 * @return An owned buffer containing the raw data, or null on error.
+	 */
+	[[nodiscard]] CORE_API FSharedBuffer Decompress(uint64 RawOffset = 0, uint64 RawSize = MAX_uint64);
+
+	/**
+	 * Decompress into an owned composite buffer.
+	 *
+	 * @return An owned buffer containing the raw data, or null on error.
+	 */
+	[[nodiscard]] CORE_API FCompositeBuffer DecompressToComposite(uint64 RawOffset = 0, uint64 RawSize = MAX_uint64);
+
+private:
+	const UE::CompressedBuffer::Private::FHeader* TryReadHeader();
+
+	FArchive* SourceArchive = nullptr;
+	const FCompressedBuffer* SourceBuffer = nullptr;
+	UE::CompressedBuffer::Private::FDecoderContext Context;
+};
+
 /**
  * A stateful decoder that reuses temporary buffers between decompression calls.
  */
@@ -163,4 +288,3 @@ private:
 	class FImpl;
 	TUniquePtr<FImpl> Impl;
 };
-
diff --git a/Engine/Source/Runtime/Core/Public/Memory/CompositeBuffer.h b/Engine/Source/Runtime/Core/Public/Memory/CompositeBuffer.h
--- a/Engine/Source/Runtime/Core/Public/Memory/CompositeBuffer.h
+++ b/Engine/Source/Runtime/Core/Public/Memory/CompositeBuffer.h
@@ -81,8 +81,13 @@ public:
 	 * @param Offset       The byte offset in this buffer that the range starts at.
 	 * @param Size         The number of bytes in the range to view or copy.
 	 * @param CopyBuffer   The buffer to write the copy into if a copy is required.
+	 * @param Allocator    The optional allocator to use when the copy buffer is required.
 	 */
-	[[nodiscard]] CORE_API FMemoryView ViewOrCopyRange(uint64 Offset, uint64 Size, FUniqueBuffer& CopyBuffer) const;
+	[[nodiscard]] CORE_API FMemoryView ViewOrCopyRange(
+		uint64 Offset,
+		uint64 Size,
+		FUniqueBuffer& CopyBuffer,
+		TFunctionRef<FUniqueBuffer (uint64 Size)> Allocator = FUniqueBuffer::Alloc) const;
 
 	/**
 	 * Copies a range of the buffer to a contiguous region of memory.
