Add missing constexpr and UE_NODISCARD to UnrealMathUtility
#tests Tested in PIE, preflight on all platforms

diff --git a/Engine/Source/Runtime/Core/Public/Math/UnrealMathUtility.h b/Engine/Source/Runtime/Core/Public/Math/UnrealMathUtility.h
--- a/Engine/Source/Runtime/Core/Public/Math/UnrealMathUtility.h
+++ b/Engine/Source/Runtime/Core/Public/Math/UnrealMathUtility.h
@@ -182,14 +182,14 @@ struct FMath : public FPlatformMath
 	// Random Number Functions
 
 	/** Helper function for rand implementations. Returns a random number in [0..A) */
-	static FORCEINLINE int32 RandHelper(int32 A)
+	UE_NODISCARD static FORCEINLINE int32 RandHelper(int32 A)
 	{
 		// Note that on some platforms RAND_MAX is a large number so we cannot do ((rand()/(RAND_MAX+1)) * A)
 		// or else we may include the upper bound results, which should be excluded.
 		return A > 0 ? Min(TruncToInt(FRand() * (float)A), A - 1) : 0;
 	}
 
-	static FORCEINLINE int64 RandHelper64(int64 A)
+	UE_NODISCARD static FORCEINLINE int64 RandHelper64(int64 A)
 	{
 		// Note that on some platforms RAND_MAX is a large number so we cannot do ((rand()/(RAND_MAX+1)) * A)
 		// or else we may include the upper bound results, which should be excluded.
@@ -197,37 +197,37 @@ struct FMath : public FPlatformMath
 	}
 
 	/** Helper function for rand implementations. Returns a random number >= Min and <= Max */
-	static FORCEINLINE int32 RandRange(int32 Min, int32 Max)
+	UE_NODISCARD static FORCEINLINE int32 RandRange(int32 Min, int32 Max)
 	{
 		const int32 Range = (Max - Min) + 1;
 		return Min + RandHelper(Range);
 	}
 
-	static FORCEINLINE int64 RandRange(int64 Min, int64 Max)
+	UE_NODISCARD static FORCEINLINE int64 RandRange(int64 Min, int64 Max)
 	{
 		const int64 Range = (Max - Min) + 1;
 		return Min + RandHelper64(Range);
 	}
 
 	/** Util to generate a random number in a range. Overloaded to distinguish from int32 version, where passing a float is typically a mistake. */
-	static FORCEINLINE float RandRange(float InMin, float InMax)
+	UE_NODISCARD static FORCEINLINE float RandRange(float InMin, float InMax)
 	{
 		return FRandRange(InMin, InMax);
 	}
 
-	static FORCEINLINE double RandRange(double InMin, double InMax)
+	UE_NODISCARD static FORCEINLINE double RandRange(double InMin, double InMax)
 	{
 		return FRandRange(InMin, InMax);
 	}
 
 	/** Util to generate a random number in a range. */
-	static FORCEINLINE float FRandRange(float InMin, float InMax)
+	UE_NODISCARD static FORCEINLINE float FRandRange(float InMin, float InMax)
 	{
 		return InMin + (InMax - InMin) * FRand();
 	}
 
 	/** Util to generate a random number in a range. */
-	static FORCEINLINE double FRandRange(double InMin, double InMax)
+	UE_NODISCARD static FORCEINLINE double FRandRange(double InMin, double InMax)
 	{
 		return InMin + (InMax - InMin) * FRand();	// LWC_TODO: Implement FRandDbl() for increased precision
 	}
@@ -235,31 +235,31 @@ struct FMath : public FPlatformMath
 	RESOLVE_FLOAT_AMBIGUITY_2_ARGS(FRandRange);
 
 	/** Util to generate a random boolean. */
-	static FORCEINLINE bool RandBool()
+	UE_NODISCARD static FORCEINLINE bool RandBool()
 	{
 		return (RandRange(0,1) == 1) ? true : false;
 	}
 
 	/** Return a uniformly distributed random unit length vector = point on the unit sphere surface. */
-	static FVector VRand();
+	UE_NODISCARD static FVector VRand();
 	
 	/**
 	 * Returns a random unit vector, uniformly distributed, within the specified cone
 	 * ConeHalfAngleRad is the half-angle of cone, in radians.  Returns a normalized vector. 
 	 */
-	static CORE_API FVector VRandCone(FVector const& Dir, float ConeHalfAngleRad);
+	UE_NODISCARD static CORE_API FVector VRandCone(FVector const& Dir, float ConeHalfAngleRad);
 
 	/** 
 	 * This is a version of VRandCone that handles "squished" cones, i.e. with different angle limits in the Y and Z axes.
 	 * Assumes world Y and Z, although this could be extended to handle arbitrary rotations.
 	 */
-	static CORE_API FVector VRandCone(FVector const& Dir, float HorizontalConeHalfAngleRad, float VerticalConeHalfAngleRad);
+	UE_NODISCARD static CORE_API FVector VRandCone(FVector const& Dir, float HorizontalConeHalfAngleRad, float VerticalConeHalfAngleRad);
 
 	/** Returns a random point, uniformly distributed, within the specified radius */
-	static CORE_API FVector2D RandPointInCircle(float CircleRadius);
+	UE_NODISCARD static CORE_API FVector2D RandPointInCircle(float CircleRadius);
 
 	/** Returns a random point within the passed in bounding box */
-	static CORE_API FVector RandPointInBox(const FBox& Box);
+	UE_NODISCARD static CORE_API FVector RandPointInBox(const FBox& Box);
 
 	/** 
 	 * Given a direction vector and a surface normal, returns the vector reflected across the surface normal.
@@ -270,13 +270,13 @@ struct FMath : public FPlatformMath
 	 *
 	 * @returns Reflected vector.
 	 */
-	static CORE_API FVector GetReflectionVector(const FVector& Direction, const FVector& SurfaceNormal);
+	UE_NODISCARD static CORE_API FVector GetReflectionVector(const FVector& Direction, const FVector& SurfaceNormal);
 	
 	// Predicates
 
 	/** Checks if value is within a range, exclusive on MaxValue) */
 	template< class T, class U> 
-	static FORCEINLINE bool IsWithin(const T& TestValue, const U& MinValue, const U& MaxValue)
+	UE_NODISCARD static constexpr FORCEINLINE bool IsWithin(const T& TestValue, const U& MinValue, const U& MaxValue)
 	{
 		return ((TestValue >= MinValue) && (TestValue < MaxValue));
 	}
@@ -284,7 +284,7 @@ struct FMath : public FPlatformMath
 
 	/** Checks if value is within a range, inclusive on MaxValue) */
 	template< class T, class U> 
-	static FORCEINLINE bool IsWithinInclusive(const T& TestValue, const U& MinValue, const U& MaxValue)
+	UE_NODISCARD static constexpr FORCEINLINE bool IsWithinInclusive(const T& TestValue, const U& MinValue, const U& MaxValue)
 	{
 		return ((TestValue>=MinValue) && (TestValue <= MaxValue));
 	}
@@ -296,12 +296,12 @@ struct FMath : public FPlatformMath
 	 *	@param ErrorTolerance	Maximum allowed difference for considering them as 'nearly equal'
 	 *	@return					true if A and B are nearly equal
 	 */
-	static FORCEINLINE bool IsNearlyEqual(float A, float B, float ErrorTolerance = SMALL_NUMBER)
+	UE_NODISCARD static FORCEINLINE bool IsNearlyEqual(float A, float B, float ErrorTolerance = SMALL_NUMBER)
 	{
 		return Abs<float>( A - B ) <= ErrorTolerance;
 	}
 
-	static FORCEINLINE bool IsNearlyEqual(double A, double B, double ErrorTolerance = DOUBLE_SMALL_NUMBER)
+	UE_NODISCARD static FORCEINLINE bool IsNearlyEqual(double A, double B, double ErrorTolerance = DOUBLE_SMALL_NUMBER)
 	{
 		return Abs<double>(A - B) <= ErrorTolerance;
 	}
@@ -315,7 +315,7 @@ struct FMath : public FPlatformMath
 	 *	@param ErrorTolerance	Maximum allowed difference for considering Value as 'nearly zero'
 	 *	@return					true if Value is nearly zero
 	 */
-	static FORCEINLINE bool IsNearlyZero(float Value, float ErrorTolerance = SMALL_NUMBER)
+	UE_NODISCARD static FORCEINLINE bool IsNearlyZero(float Value, float ErrorTolerance = SMALL_NUMBER)
 	{
 		return Abs<float>( Value ) <= ErrorTolerance;
 	}
@@ -326,7 +326,7 @@ struct FMath : public FPlatformMath
 	 *	@param ErrorTolerance	Maximum allowed difference for considering Value as 'nearly zero'
 	 *	@return					true if Value is nearly zero
 	 */
-	static FORCEINLINE bool IsNearlyZero(double Value, double ErrorTolerance = DOUBLE_SMALL_NUMBER)
+	UE_NODISCARD static FORCEINLINE bool IsNearlyZero(double Value, double ErrorTolerance = DOUBLE_SMALL_NUMBER)
 	{
 		return Abs<double>( Value ) <= ErrorTolerance;
 	}
@@ -395,7 +395,7 @@ public:
 	 *	                        numbers are allowed to differ.
 	 *	@return					true if the two values are nearly equal.
 	 */
-	static FORCEINLINE bool IsNearlyEqualByULP(float A, float B, int32 MaxUlps = 4)
+	UE_NODISCARD static FORCEINLINE bool IsNearlyEqualByULP(float A, float B, int32 MaxUlps = 4)
 	{
 		return TIsNearlyEqualByULP<float, uint32, uint32(1U << 31)>(A, B, MaxUlps);
 	}
@@ -417,7 +417,7 @@ public:
 	 *	                        numbers are allowed to differ.
 	 *	@return					true if the two values are nearly equal.
 	 */
-	static FORCEINLINE bool IsNearlyEqualByULP(double A, double B, int32 MaxUlps = 4)
+	UE_NODISCARD static FORCEINLINE bool IsNearlyEqualByULP(double A, double B, int32 MaxUlps = 4)
 	{
 		return TIsNearlyEqualByULP<double, uint64, uint64(1ULL << 63)>(A, B, MaxUlps);
 	}
@@ -428,19 +428,19 @@ public:
 	 *	@return			true if Value is a power of two
 	 */
 	template <typename T>
-	static FORCEINLINE bool IsPowerOfTwo( T Value )
+	UE_NODISCARD static constexpr FORCEINLINE bool IsPowerOfTwo( T Value )
 	{
 		return ((Value & (Value - 1)) == (T)0);
 	}
 
 	/** Converts a float to a nearest less or equal integer. */
-	static FORCEINLINE float Floor(float F)
+	UE_NODISCARD static FORCEINLINE float Floor(float F)
 	{
 		return FloorToFloat(F);
 	}
 
 	/** Converts a double to a nearest less or equal integer. */
-	static FORCEINLINE double Floor(double F)
+	UE_NODISCARD static FORCEINLINE double Floor(double F)
 	{
 		return FloorToDouble(F);
 	}
@@ -453,7 +453,7 @@ public:
 	    typename IntegralType,
 	    typename TEnableIf<TIsIntegral<IntegralType>::Value>::Type* = nullptr
 	>
-	static FORCEINLINE IntegralType Floor(IntegralType I)
+	UE_NODISCARD static constexpr FORCEINLINE IntegralType Floor(IntegralType I)
 	{
 	    return I;
 	}
@@ -463,35 +463,35 @@ public:
 
 	/** Returns highest of 3 values */
 	template< class T > 
-	static FORCEINLINE T Max3( const T A, const T B, const T C )
+	UE_NODISCARD static constexpr FORCEINLINE T Max3( const T A, const T B, const T C )
 	{
 		return Max ( Max( A, B ), C );
 	}
 
 	/** Returns lowest of 3 values */
 	template< class T > 
-	static FORCEINLINE T Min3( const T A, const T B, const T C )
+	UE_NODISCARD static constexpr FORCEINLINE T Min3( const T A, const T B, const T C )
 	{
 		return Min ( Min( A, B ), C );
 	}
 
 	/** Multiples value by itself */
 	template< class T > 
-	static FORCEINLINE T Square( const T A )
+	UE_NODISCARD static constexpr FORCEINLINE T Square( const T A )
 	{
 		return A*A;
 	}
 
 	/** Cubes the value */
 	template< class T > 
-    static FORCEINLINE T Cube( const T A )
+    UE_NODISCARD static constexpr FORCEINLINE T Cube( const T A )
 	{
 		return A*A*A;
 	}
 
 	/** Clamps X to be between Min and Max, inclusive */
 	template< class T >
-	UE_NODISCARD static FORCEINLINE T Clamp(const T X, const T Min, const T Max)
+	UE_NODISCARD static constexpr FORCEINLINE T Clamp(const T X, const T Min, const T Max)
 	{
 		return (X < Min) ? Min : (X < Max) ? X : Max;
 	}
@@ -500,7 +500,7 @@ public:
 	/** Wraps X to be between Min and Max, inclusive. */
 	/** When X can wrap to both Min and Max, it will wrap to Min if it lies below the range and wrap to Max if it is above the range. */
 	template< class T >
-	static FORCEINLINE T Wrap(const T X, const T Min, const T Max)
+	UE_NODISCARD static constexpr FORCEINLINE T Wrap(const T X, const T Min, const T Max)
 	{
 		T Size = Max - Min;
 		if (Size == 0)
@@ -524,7 +524,7 @@ public:
 
 	/** Snaps a value to the nearest grid multiple */
 	template< class T >
-	static FORCEINLINE T GridSnap(T Location, T Grid)
+	UE_NODISCARD static constexpr FORCEINLINE T GridSnap(T Location, T Grid)
 	{
 		return (Grid == T{}) ? Location : (Floor((Location + (Grid/(T)2)) / Grid) * Grid);
 	}
@@ -532,21 +532,21 @@ public:
 
 	/** Divides two integers and rounds up */
 	template <class T>
-	static FORCEINLINE T DivideAndRoundUp(T Dividend, T Divisor)
+	UE_NODISCARD static constexpr FORCEINLINE T DivideAndRoundUp(T Dividend, T Divisor)
 	{
 		return (Dividend + Divisor - 1) / Divisor;
 	}
 
 	/** Divides two integers and rounds down */
 	template <class T>
-	static FORCEINLINE T DivideAndRoundDown(T Dividend, T Divisor)
+	UE_NODISCARD static constexpr FORCEINLINE T DivideAndRoundDown(T Dividend, T Divisor)
 	{
 		return Dividend / Divisor;
 	}
 
 	/** Divides two integers and rounds to nearest */
 	template <class T>
-	static FORCEINLINE T DivideAndRoundNearest(T Dividend, T Divisor)
+	UE_NODISCARD static constexpr FORCEINLINE T DivideAndRoundNearest(T Dividend, T Divisor)
 	{
 		return (Dividend >= 0)
 			? (Dividend + Divisor / 2) / Divisor
@@ -560,10 +560,10 @@ public:
 	 *
 	 * @return the base 2 log of the value
 	 */
-	static FORCEINLINE float Log2(float Value)
+	UE_NODISCARD static FORCEINLINE float Log2(float Value)
 	{
 		// Cached value for fast conversions
-		const float LogToLog2 = 1.44269502f; // 1.f / Loge(2.f)
+		constexpr float LogToLog2 = 1.44269502f; // 1.f / Loge(2.f)
 		// Do the platform specific log and convert using the cached value
 		return Loge(Value) * LogToLog2;
 	}
@@ -575,10 +575,10 @@ public:
 	 *
 	 * @return the base 2 log of the value
 	 */
-	static FORCEINLINE double Log2(double Value)
+	UE_NODISCARD static FORCEINLINE double Log2(double Value)
 	{
 		// Cached value for fast conversions
-		const double LogToLog2 = 1.4426950408889634; // 1.0 / Loge(2.0);
+		constexpr double LogToLog2 = 1.4426950408889634; // 1.0 / Loge(2.0);
 		// Do the platform specific log and convert using the cached value
 		return Loge(Value) * LogToLog2;
 	}
@@ -591,7 +591,7 @@ public:
 	* @param Value  input angles 
 	*/
 	template<typename T, TEMPLATE_REQUIRES(TIsFloatingPoint<T>::Value)>
-	static FORCEINLINE void SinCos(typename TDecay<T>::Type* ScalarSin, typename TDecay<T>::Type* ScalarCos, T  Value )
+	static constexpr FORCEINLINE void SinCos(typename TDecay<T>::Type* ScalarSin, typename TDecay<T>::Type* ScalarCos, T  Value )
 	{
 		// Map Value to y in [-pi,pi], x = 2*pi*quotient + remainder.
 		T quotient = (INV_PI*0.5f)*Value;
@@ -658,7 +658,7 @@ public:
 	* @param Value  input angle
 	* @return ASin of Value
 	*/
-	static FORCEINLINE float FastAsin(float Value)
+	UE_NODISCARD static FORCEINLINE float FastAsin(float Value)
 	{
 		// Clamp input to [-1,1].
 		bool nonnegative = (Value >= 0.0f);
@@ -677,7 +677,7 @@ public:
 	}
 #undef FASTASIN_HALF_PI
 
-	static FORCEINLINE double FastAsin(double Value)
+	UE_NODISCARD static FORCEINLINE double FastAsin(double Value)
 	{
 		// TODO: add fast approximation
 		return FMath::Asin(Value);
@@ -691,7 +691,7 @@ public:
 	 * @return					Value in degrees.
 	 */
 	template<class T>
-	static FORCEINLINE auto RadiansToDegrees(T const& RadVal) -> decltype(RadVal * (180.f / PI))
+	UE_NODISCARD static constexpr FORCEINLINE auto RadiansToDegrees(T const& RadVal) -> decltype(RadVal * (180.f / PI))
 	{
 		return RadVal * (180.f / PI);
 	}
@@ -705,7 +705,7 @@ public:
 	 * @return					Value in radians.
 	 */
 	template<class T>
-	static FORCEINLINE auto DegreesToRadians(T const& DegVal) -> decltype(DegVal * (PI / 180.f))
+	UE_NODISCARD static constexpr FORCEINLINE auto DegreesToRadians(T const& DegVal) -> decltype(DegVal * (PI / 180.f))
 	{
 		return DegVal * (PI / 180.f);
 	}
@@ -720,11 +720,11 @@ public:
 	 * @param MaxAngleDegrees	"to" angle that defines the end of the range of valid angles
 	 * @return Returns clamped angle in the range -180..180.
 	 */
-	static CORE_API float ClampAngle(float AngleDegrees, float MinAngleDegrees, float MaxAngleDegrees);
+	UE_NODISCARD static CORE_API float ClampAngle(float AngleDegrees, float MinAngleDegrees, float MaxAngleDegrees);
 
 	/** Find the smallest angle between two headings (in degrees) */
 	template<typename T, typename T2, TEMPLATE_REQUIRES(TOr<TIsFloatingPoint<T>, TIsFloatingPoint<T2>>::Value)>
-	static auto FindDeltaAngleDegrees(T A1, T2 A2) -> decltype(A1 * A2)
+	UE_NODISCARD static constexpr auto FindDeltaAngleDegrees(T A1, T2 A2) -> decltype(A1 * A2)
 	{
 		// Find the difference
 		auto Delta = A2 - A1;
@@ -748,7 +748,7 @@ public:
 
 	/** Find the smallest angle between two headings (in radians) */
 	template<typename T, typename T2, TEMPLATE_REQUIRES(TOr<TIsFloatingPoint<T>, TIsFloatingPoint<T2>>::Value)>
-	static auto FindDeltaAngleRadians(T A1, T2 A2) -> decltype(A1 * A2)
+	UE_NODISCARD static constexpr auto FindDeltaAngleRadians(T A1, T2 A2) -> decltype(A1 * A2)
 	{
 		// Find the difference
 		auto Delta = A2 - A1;
@@ -771,14 +771,14 @@ public:
 	}
 
 	UE_DEPRECATED(4.12, "Please use FindDeltaAngleRadians(float A1, float A2) instead of FindDeltaAngle(float A1, float A2).")
-	static float FindDeltaAngle(float A1, float A2)
+	UE_NODISCARD static constexpr float FindDeltaAngle(float A1, float A2)
 	{
 		return FindDeltaAngleRadians(A1, A2);
 	}
 
 	/** Given a heading which may be outside the +/- PI range, 'unwind' it back into that range. */
 	template<typename T, TEMPLATE_REQUIRES(TIsFloatingPoint<T>::Value)>
-	static T UnwindRadians(T A)
+	UE_NODISCARD static constexpr T UnwindRadians(T A)
 	{
 		while(A > PI)
 		{
@@ -795,7 +795,7 @@ public:
 
 	/** Utility to ensure angle is between +/- 180 degrees by unwinding. */
 	template<typename T, TEMPLATE_REQUIRES(TIsFloatingPoint<T>::Value)>
-	static T UnwindDegrees(T A)
+	UE_NODISCARD static constexpr T UnwindDegrees(T A)
 	{
 		while(A > 180.f)
 		{
@@ -813,8 +813,8 @@ public:
 	/** 
 	 * Given two angles in degrees, 'wind' the rotation in Angle1 so that it avoids >180 degree flips.
 	 * Good for winding rotations previously expressed as quaternions into a euler-angle representation.
-	 * @param	Angle0	The first angle that we wind relative to.
-	 * @param	Angle1	The second angle that we may wind relative to the first.
+	 * @param	InAngle0	The first angle that we wind relative to.
+	 * @param	InOutAngle1	The second angle that we may wind relative to the first.
 	 */
 	static CORE_API void WindRelativeAnglesDegrees(float InAngle0, float& InOutAngle1);
 	static CORE_API void WindRelativeAnglesDegrees(double InAngle0, double& InOutAngle1);
@@ -827,7 +827,7 @@ public:
 	 *
 	 * @return a new rotation component value
 	 */
-	static CORE_API float FixedTurn(float InCurrent, float InDesired, float InDeltaRate);
+	UE_NODISCARD static CORE_API float FixedTurn(float InCurrent, float InDesired, float InDeltaRate);
 
 	/** Converts given Cartesian coordinate pair to Polar coordinate system. */
 	template<typename T>
@@ -892,13 +892,13 @@ public:
 	 * @return	FVector2D	X = Azimuth angle (in radians) (-PI, +PI)
 	 *						Y = Elevation angle (in radians) (-PI/2, +PI/2)
 	 */
-	static CORE_API FVector2D GetAzimuthAndElevation(const FVector &Direction, const FVector &AxisX, const FVector &AxisY, const FVector &AxisZ);
+	UE_NODISCARD static CORE_API FVector2D GetAzimuthAndElevation(const FVector &Direction, const FVector &AxisX, const FVector &AxisY, const FVector &AxisZ);
 
 	// Interpolation Functions
 
 	/** Calculates the percentage along a line from MinValue to MaxValue that Value is. */
 	template<typename T, typename T2, TEMPLATE_REQUIRES(TIsFloatingPoint<T>::Value)>
-	static FORCEINLINE auto GetRangePct(T MinValue, T MaxValue, T2 Value)
+	UE_NODISCARD static constexpr FORCEINLINE auto GetRangePct(T MinValue, T MaxValue, T2 Value)
 	{
 		// Avoid Divide by Zero.
 		// But also if our range is a point, output whether Value is before or after.
@@ -914,21 +914,21 @@ public:
 
 	/** Same as above, but taking a 2d vector as the range. */
 	template<typename T, typename T2, TEMPLATE_REQUIRES(TIsFloatingPoint<T>::Value)>
-	static auto GetRangePct(UE::Math::TVector2<T> const& Range, T2 Value)
+	UE_NODISCARD static auto GetRangePct(UE::Math::TVector2<T> const& Range, T2 Value)
 	{
 		return GetRangePct(Range.X, Range.Y, Value);
 	}
 	
 	/** Basically a Vector2d version of Lerp. */
 	template<typename T, typename T2, TEMPLATE_REQUIRES(TIsFloatingPoint<T>::Value)>
-	static auto GetRangeValue(UE::Math::TVector2<T> const& Range, T2 Pct)
+	UE_NODISCARD static auto GetRangeValue(UE::Math::TVector2<T> const& Range, T2 Pct)
 	{
 		return Lerp(Range.X, Range.Y, Pct);
 	}
 
 	/** For the given Value clamped to the [Input:Range] inclusive, returns the corresponding percentage in [Output:Range] Inclusive. */
 	template<typename T, typename T2>
-	static FORCEINLINE auto GetMappedRangeValueClamped(const UE::Math::TVector2<T>& InputRange, const UE::Math::TVector2<T>& OutputRange, const T2 Value)
+	UE_NODISCARD static FORCEINLINE auto GetMappedRangeValueClamped(const UE::Math::TVector2<T>& InputRange, const UE::Math::TVector2<T>& OutputRange, const T2 Value)
 	{
 		using RangePctType = decltype(T() * T2());
 		const RangePctType ClampedPct = Clamp<RangePctType>(GetRangePct(InputRange, Value), 0.f, 1.f);
@@ -937,25 +937,25 @@ public:
 
 	/** Transform the given Value relative to the input range to the Output Range. */
 	template<typename T, typename T2>
-	static FORCEINLINE auto GetMappedRangeValueUnclamped(const UE::Math::TVector2<T>& InputRange, const UE::Math::TVector2<T>& OutputRange, const T2 Value)
+	UE_NODISCARD static FORCEINLINE auto GetMappedRangeValueUnclamped(const UE::Math::TVector2<T>& InputRange, const UE::Math::TVector2<T>& OutputRange, const T2 Value)
 	{
 		return GetRangeValue(OutputRange, GetRangePct(InputRange, Value));
 	}
 
 	template<class T>
-	static FORCEINLINE double GetRangePct(TRange<T> const& Range, T Value)
+	UE_NODISCARD static FORCEINLINE double GetRangePct(TRange<T> const& Range, T Value)
 	{
 		return GetRangePct(Range.GetLowerBoundValue(), Range.GetUpperBoundValue(), Value);
 	}
 
 	template<class T>
-	static FORCEINLINE T GetRangeValue(TRange<T> const& Range, T Pct)
+	UE_NODISCARD static FORCEINLINE T GetRangeValue(TRange<T> const& Range, T Pct)
 	{
 		return FMath::Lerp<T>(Range.GetLowerBoundValue(), Range.GetUpperBoundValue(), Pct);
 	}
 
 	template<class T>
-	static FORCEINLINE T GetMappedRangeValueClamped(const TRange<T>& InputRange, const TRange<T>& OutputRange, const T Value)
+	UE_NODISCARD static FORCEINLINE T GetMappedRangeValueClamped(const TRange<T>& InputRange, const TRange<T>& OutputRange, const T Value)
 	{
 		const T ClampedPct = FMath::Clamp<T>(GetRangePct(InputRange, Value), 0, 1);
 		return GetRangeValue(OutputRange, ClampedPct);
@@ -966,14 +966,14 @@ public:
 													TNot<TCustomLerp<T>>,
 													TOr<TIsFloatingPoint<U>, TIsSame<T, U>>
 													>::Value)>
-	static FORCEINLINE_DEBUGGABLE T Lerp( const T& A, const T& B, const U& Alpha )
+	UE_NODISCARD static constexpr FORCEINLINE_DEBUGGABLE T Lerp( const T& A, const T& B, const U& Alpha )
 	{
 		return (T)(A + Alpha * (B-A));
 	}
 
 	/** Custom lerps defined for those classes not suited to the default implemenation. */
 	template< class T, class U, TEMPLATE_REQUIRES(TCustomLerp<T>::Value) >
-	static FORCEINLINE_DEBUGGABLE T Lerp(const T& A, const T& B, const U& Alpha)
+	UE_NODISCARD static FORCEINLINE_DEBUGGABLE T Lerp(const T& A, const T& B, const U& Alpha)
 	{
 		return TCustomLerp<T>::Lerp(A, B, Alpha);
 	}
@@ -984,7 +984,7 @@ public:
 																		TNot<TCustomLerp<T1>>,
 																		TNot<TCustomLerp<T2>>
 																		>::Value)>
-	static auto Lerp( const T1& A, const T2& B, const T3& Alpha ) -> decltype(A * B)
+	UE_NODISCARD static auto Lerp( const T1& A, const T2& B, const T3& Alpha ) -> decltype(A * B)
 	{
 		using ABType = decltype(A * B);
 		return Lerp(ABType(A), ABType(B), Alpha);
@@ -992,21 +992,21 @@ public:
 
 	/** Performs a linear interpolation between two values, Alpha ranges from 0-1. Handles full numeric range of T */
 	template< class T > 
-	static FORCEINLINE_DEBUGGABLE T LerpStable( const T& A, const T& B, double Alpha )
+	UE_NODISCARD static constexpr FORCEINLINE_DEBUGGABLE T LerpStable( const T& A, const T& B, double Alpha )
 	{
 		return (T)((A * (1.0 - Alpha)) + (B * Alpha));
 	}
 
 	/** Performs a linear interpolation between two values, Alpha ranges from 0-1. Handles full numeric range of T */
 	template< class T >
-	static FORCEINLINE_DEBUGGABLE T LerpStable(const T& A, const T& B, float Alpha)
+	UE_NODISCARD static constexpr FORCEINLINE_DEBUGGABLE T LerpStable(const T& A, const T& B, float Alpha)
 	{
 		return (T)((A * (1.0f - Alpha)) + (B * Alpha));
 	}
 
 	// Allow passing of differing A/B types.
 	template<typename T1, typename T2, typename T3, TEMPLATE_REQUIRES(!TIsSame<T1, T2>::Value)>
-	static auto LerpStable( const T1& A, const T2& B, const T3& Alpha ) -> decltype(A * B)
+	UE_NODISCARD static auto LerpStable( const T1& A, const T2& B, const T3& Alpha ) -> decltype(A * B)
 	{
 		using ABType = decltype(A * B);
 		return LerpStable(ABType(A), ABType(B), Alpha);
@@ -1017,7 +1017,7 @@ public:
 													TNot<TCustomLerp<T>>,
 													TOr<TIsFloatingPoint<U>, TIsSame<T, U>>
 													>::Value)>
-	static FORCEINLINE_DEBUGGABLE T BiLerp(const T& P00,const T& P10,const T& P01,const T& P11, const U& FracX, const U& FracY)
+	UE_NODISCARD static constexpr FORCEINLINE_DEBUGGABLE T BiLerp(const T& P00,const T& P10,const T& P01,const T& P11, const U& FracX, const U& FracY)
 	{
 		return Lerp(
 			Lerp(P00,P10,FracX),
@@ -1028,7 +1028,7 @@ public:
 
 	/** Custom lerps defined for those classes not suited to the default implemenation. */
 	template< class T, class U, TEMPLATE_REQUIRES(TCustomLerp<T>::Value) >
-	static FORCEINLINE_DEBUGGABLE T BiLerp(const T& P00, const T& P10, const T& P01, const T& P11, const U& FracX, const U& FracY)
+	UE_NODISCARD static FORCEINLINE_DEBUGGABLE T BiLerp(const T& P00, const T& P10, const T& P01, const T& P11, const U& FracX, const U& FracY)
 	{
 		return TCustomLerp<T>::BiLerp(P00, P10, P01, P11, FracX, FracY);
 	}
@@ -1038,7 +1038,7 @@ public:
 	 *
 	 * @param  P - end points
 	 * @param  T - tangent directions at end points
-	 * @param  Alpha - distance along spline
+	 * @param  A - distance along spline
 	 *
 	 * @return  Interpolated value
 	 */
@@ -1046,7 +1046,7 @@ public:
 													TNot<TCustomLerp<T>>,
 													TOr<TIsFloatingPoint<U>, TIsSame<T, U>>
 													>::Value)>
-	static FORCEINLINE_DEBUGGABLE T CubicInterp( const T& P0, const T& T0, const T& P1, const T& T1, const U& A )
+	UE_NODISCARD static constexpr FORCEINLINE_DEBUGGABLE T CubicInterp( const T& P0, const T& T0, const T& P1, const T& T1, const U& A )
 	{
 		const U A2 = A * A;
 		const U A3 = A2 * A;
@@ -1056,7 +1056,7 @@ public:
 
 	/** Custom lerps defined for those classes not suited to the default implemenation. */
 	template< class T, class U, TEMPLATE_REQUIRES(TCustomLerp<T>::Value) >
-	static FORCEINLINE_DEBUGGABLE T CubicInterp(const T& P0, const T& T0, const T& P1, const T& T1, const U& A)
+	UE_NODISCARD static FORCEINLINE_DEBUGGABLE T CubicInterp(const T& P0, const T& T0, const T& P1, const T& T1, const U& A)
 	{
 		return TCustomLerp<T>::CubicInterp(P0, T0, P1, T1, A);
 	}
@@ -1066,12 +1066,12 @@ public:
 	 *
 	 * @param  P - end points
 	 * @param  T - tangent directions at end points
-	 * @param  Alpha - distance along spline
+	 * @param  A - distance along spline
 	 *
 	 * @return  Interpolated value
 	 */
 	template< class T, class U, TEMPLATE_REQUIRES(TIsFloatingPoint<U>::Value) >
-	static FORCEINLINE_DEBUGGABLE T CubicInterpDerivative( const T& P0, const T& T0, const T& P1, const T& T1, const U& A )
+	UE_NODISCARD static constexpr FORCEINLINE_DEBUGGABLE T CubicInterpDerivative( const T& P0, const T& T0, const T& P1, const T& T1, const U& A )
 	{
 		T a = 6.f*P0 + 3.f*T0 + 3.f*T1 - 6.f*P1;
 		T b = -6.f*P0 - 4.f*T0 - 2.f*T1 + 6.f*P1;
@@ -1087,12 +1087,12 @@ public:
 	 *
 	 * @param  P - end points
 	 * @param  T - tangent directions at end points
-	 * @param  Alpha - distance along spline
+	 * @param  A - distance along spline
 	 *
 	 * @return  Interpolated value
 	 */
 	template< class T, class U, TEMPLATE_REQUIRES(TIsFloatingPoint<U>::Value) >
-	static FORCEINLINE_DEBUGGABLE T CubicInterpSecondDerivative( const T& P0, const T& T0, const T& P1, const T& T1, const U& A )
+	UE_NODISCARD static constexpr FORCEINLINE_DEBUGGABLE T CubicInterpSecondDerivative( const T& P0, const T& T0, const T& P1, const T& T1, const U& A )
 	{
 		T a = 12.f*P0 + 6.f*T0 + 6.f*T1 - 12.f*P1;
 		T b = -6.f*P0 - 4.f*T0 - 2.f*T1 + 6.f*P1;
@@ -1102,7 +1102,7 @@ public:
 
 	/** Interpolate between A and B, applying an ease in function.  Exp controls the degree of the curve. */
 	template< class T >
-	static FORCEINLINE_DEBUGGABLE T InterpEaseIn(const T& A, const T& B, float Alpha, float Exp)
+	UE_NODISCARD static FORCEINLINE_DEBUGGABLE T InterpEaseIn(const T& A, const T& B, float Alpha, float Exp)
 	{
 		float const ModifiedAlpha = Pow(Alpha, Exp);
 		return Lerp<T>(A, B, ModifiedAlpha);
@@ -1110,7 +1110,7 @@ public:
 
 	/** Interpolate between A and B, applying an ease out function.  Exp controls the degree of the curve. */
 	template< class T >
-	static FORCEINLINE_DEBUGGABLE T InterpEaseOut(const T& A, const T& B, float Alpha, float Exp)
+	UE_NODISCARD static FORCEINLINE_DEBUGGABLE T InterpEaseOut(const T& A, const T& B, float Alpha, float Exp)
 	{
 		float const ModifiedAlpha = 1.f - Pow(1.f - Alpha, Exp);
 		return Lerp<T>(A, B, ModifiedAlpha);
@@ -1118,7 +1118,7 @@ public:
 
 	/** Interpolate between A and B, applying an ease in/out function.  Exp controls the degree of the curve. */
 	template< class T > 
-	static FORCEINLINE_DEBUGGABLE T InterpEaseInOut( const T& A, const T& B, float Alpha, float Exp )
+	UE_NODISCARD static FORCEINLINE_DEBUGGABLE T InterpEaseInOut( const T& A, const T& B, float Alpha, float Exp )
 	{
 		return Lerp<T>(A, B, (Alpha < 0.5f) ?
 			InterpEaseIn(0.f, 1.f, Alpha * 2.f, Exp) * 0.5f :
@@ -1127,7 +1127,7 @@ public:
 
 	/** Interpolation between A and B, applying a step function. */
 	template< class T >
-	static FORCEINLINE_DEBUGGABLE T InterpStep(const T& A, const T& B, float Alpha, int32 Steps)
+	UE_NODISCARD static constexpr FORCEINLINE_DEBUGGABLE T InterpStep(const T& A, const T& B, float Alpha, int32 Steps)
 	{
 		if (Steps <= 1 || Alpha <= 0)
 		{
@@ -1146,7 +1146,7 @@ public:
 
 	/** Interpolation between A and B, applying a sinusoidal in function. */
 	template< class T >
-	static FORCEINLINE_DEBUGGABLE T InterpSinIn(const T& A, const T& B, float Alpha)
+	UE_NODISCARD static FORCEINLINE_DEBUGGABLE T InterpSinIn(const T& A, const T& B, float Alpha)
 	{
 		float const ModifiedAlpha = -1.f * Cos(Alpha * HALF_PI) + 1.f;
 		return Lerp<T>(A, B, ModifiedAlpha);
@@ -1154,7 +1154,7 @@ public:
 	
 	/** Interpolation between A and B, applying a sinusoidal out function. */
 	template< class T >
-	static FORCEINLINE_DEBUGGABLE T InterpSinOut(const T& A, const T& B, float Alpha)
+	UE_NODISCARD static FORCEINLINE_DEBUGGABLE T InterpSinOut(const T& A, const T& B, float Alpha)
 	{
 		float const ModifiedAlpha = Sin(Alpha * HALF_PI);
 		return Lerp<T>(A, B, ModifiedAlpha);
@@ -1162,7 +1162,7 @@ public:
 
 	/** Interpolation between A and B, applying a sinusoidal in/out function. */
 	template< class T >
-	static FORCEINLINE_DEBUGGABLE T InterpSinInOut(const T& A, const T& B, float Alpha)
+	UE_NODISCARD static FORCEINLINE_DEBUGGABLE T InterpSinInOut(const T& A, const T& B, float Alpha)
 	{
 		return Lerp<T>(A, B, (Alpha < 0.5f) ?
 			InterpSinIn(0.f, 1.f, Alpha * 2.f) * 0.5f :
@@ -1171,7 +1171,7 @@ public:
 
 	/** Interpolation between A and B, applying an exponential in function. */
 	template< class T >
-	static FORCEINLINE_DEBUGGABLE T InterpExpoIn(const T& A, const T& B, float Alpha)
+	UE_NODISCARD static FORCEINLINE_DEBUGGABLE T InterpExpoIn(const T& A, const T& B, float Alpha)
 	{
 		float const ModifiedAlpha = (Alpha == 0.f) ? 0.f : Pow(2.f, 10.f * (Alpha - 1.f));
 		return Lerp<T>(A, B, ModifiedAlpha);
@@ -1179,7 +1179,7 @@ public:
 
 	/** Interpolation between A and B, applying an exponential out function. */
 	template< class T >
-	static FORCEINLINE_DEBUGGABLE T InterpExpoOut(const T& A, const T& B, float Alpha)
+	UE_NODISCARD static FORCEINLINE_DEBUGGABLE T InterpExpoOut(const T& A, const T& B, float Alpha)
 	{
 		float const ModifiedAlpha = (Alpha == 1.f) ? 1.f : -Pow(2.f, -10.f * Alpha) + 1.f;
 		return Lerp<T>(A, B, ModifiedAlpha);
@@ -1187,7 +1187,7 @@ public:
 
 	/** Interpolation between A and B, applying an exponential in/out function. */
 	template< class T >
-	static FORCEINLINE_DEBUGGABLE T InterpExpoInOut(const T& A, const T& B, float Alpha)
+	UE_NODISCARD static FORCEINLINE_DEBUGGABLE T InterpExpoInOut(const T& A, const T& B, float Alpha)
 	{
 		return Lerp<T>(A, B, (Alpha < 0.5f) ?
 			InterpExpoIn(0.f, 1.f, Alpha * 2.f) * 0.5f :
@@ -1196,7 +1196,7 @@ public:
 
 	/** Interpolation between A and B, applying a circular in function. */
 	template< class T >
-	static FORCEINLINE_DEBUGGABLE T InterpCircularIn(const T& A, const T& B, float Alpha)
+	UE_NODISCARD static FORCEINLINE_DEBUGGABLE T InterpCircularIn(const T& A, const T& B, float Alpha)
 	{
 		float const ModifiedAlpha = -1.f * (Sqrt(1.f - Alpha * Alpha) - 1.f);
 		return Lerp<T>(A, B, ModifiedAlpha);
@@ -1204,7 +1204,7 @@ public:
 
 	/** Interpolation between A and B, applying a circular out function. */
 	template< class T >
-	static FORCEINLINE_DEBUGGABLE T InterpCircularOut(const T& A, const T& B, float Alpha)
+	UE_NODISCARD static FORCEINLINE_DEBUGGABLE T InterpCircularOut(const T& A, const T& B, float Alpha)
 	{
 		Alpha -= 1.f;
 		float const ModifiedAlpha = Sqrt(1.f - Alpha  * Alpha);
@@ -1213,7 +1213,7 @@ public:
 
 	/** Interpolation between A and B, applying a circular in/out function. */
 	template< class T >
-	static FORCEINLINE_DEBUGGABLE T InterpCircularInOut(const T& A, const T& B, float Alpha)
+	UE_NODISCARD static FORCEINLINE_DEBUGGABLE T InterpCircularInOut(const T& A, const T& B, float Alpha)
 	{
 		return Lerp<T>(A, B, (Alpha < 0.5f) ?
 			InterpCircularIn(0.f, 1.f, Alpha * 2.f) * 0.5f :
@@ -1223,7 +1223,7 @@ public:
 	// Rotator specific interpolation
 	// Similar to Lerp, but does not take the shortest path. Allows interpolation over more than 180 degrees.
 	template< typename T, typename U >
-	static UE::Math::TRotator<T> LerpRange(const UE::Math::TRotator<T>& A, const UE::Math::TRotator<T>& B, U Alpha);
+	UE_NODISCARD static UE::Math::TRotator<T> LerpRange(const UE::Math::TRotator<T>& A, const UE::Math::TRotator<T>& B, U Alpha);
 
 	/*
 	 *	Cubic Catmull-Rom Spline interpolation. Based on http://www.cemyuksel.com/research/catmullrom_param/catmullrom.pdf 
@@ -1237,7 +1237,7 @@ public:
 	 *	T - The interpolation factor in the range 0 to 1. 0 returns P1. 1 returns P2.
 	 */
 	template< class U > 
-	static FORCEINLINE_DEBUGGABLE U CubicCRSplineInterp(const U& P0, const U& P1, const U& P2, const U& P3, const float T0, const float T1, const float T2, const float T3, const float T)
+	UE_NODISCARD static constexpr FORCEINLINE_DEBUGGABLE U CubicCRSplineInterp(const U& P0, const U& P1, const U& P2, const U& P3, const float T0, const float T1, const float T2, const float T3, const float T)
 	{
 		//Based on http://www.cemyuksel.com/research/catmullrom_param/catmullrom.pdf 
 		float InvT1MinusT0 = 1.0f / (T1 - T0);
@@ -1257,7 +1257,7 @@ public:
 
 	/* Same as CubicCRSplineInterp but with additional saftey checks. If the checks fail P1 is returned. **/
 	template< class U >
-	static FORCEINLINE_DEBUGGABLE U CubicCRSplineInterpSafe(const U& P0, const U& P1, const U& P2, const U& P3, const float T0, const float T1, const float T2, const float T3, const float T)
+	UE_NODISCARD static constexpr FORCEINLINE_DEBUGGABLE U CubicCRSplineInterpSafe(const U& P0, const U& P1, const U& P2, const U& P3, const float T0, const float T1, const float T2, const float T3, const float T)
 	{
 		//Based on http://www.cemyuksel.com/research/catmullrom_param/catmullrom.pdf 
 		float T1MinusT0 = (T1 - T0);
@@ -1290,29 +1290,29 @@ public:
 	// Special-case interpolation
 
 	/** Interpolate a normal vector Current to Target, by interpolating the angle between those vectors with constant step. */
-	static CORE_API FVector VInterpNormalRotationTo(const FVector& Current, const FVector& Target, float DeltaTime, float RotationSpeedDegrees);
+	UE_NODISCARD static CORE_API FVector VInterpNormalRotationTo(const FVector& Current, const FVector& Target, float DeltaTime, float RotationSpeedDegrees);
 
 	/** Interpolate vector from Current to Target with constant step */
-	static CORE_API FVector VInterpConstantTo(const FVector& Current, const FVector& Target, float DeltaTime, float InterpSpeed);
+	UE_NODISCARD static CORE_API FVector VInterpConstantTo(const FVector& Current, const FVector& Target, float DeltaTime, float InterpSpeed);
 
 	/** Interpolate vector from Current to Target. Scaled by distance to Target, so it has a strong start speed and ease out. */
-	static CORE_API FVector VInterpTo( const FVector& Current, const FVector& Target, float DeltaTime, float InterpSpeed );
+	UE_NODISCARD static CORE_API FVector VInterpTo( const FVector& Current, const FVector& Target, float DeltaTime, float InterpSpeed );
 	
 	/** Interpolate vector2D from Current to Target with constant step */
-	static CORE_API FVector2D Vector2DInterpConstantTo( const FVector2D& Current, const FVector2D& Target, float DeltaTime, float InterpSpeed );
+	UE_NODISCARD static CORE_API FVector2D Vector2DInterpConstantTo( const FVector2D& Current, const FVector2D& Target, float DeltaTime, float InterpSpeed );
 
 	/** Interpolate vector2D from Current to Target. Scaled by distance to Target, so it has a strong start speed and ease out. */
-	static CORE_API FVector2D Vector2DInterpTo( const FVector2D& Current, const FVector2D& Target, float DeltaTime, float InterpSpeed );
+	UE_NODISCARD static CORE_API FVector2D Vector2DInterpTo( const FVector2D& Current, const FVector2D& Target, float DeltaTime, float InterpSpeed );
 
 	/** Interpolate rotator from Current to Target with constant step */
-	static CORE_API FRotator RInterpConstantTo( const FRotator& Current, const FRotator& Target, float DeltaTime, float InterpSpeed);
+	UE_NODISCARD static CORE_API FRotator RInterpConstantTo( const FRotator& Current, const FRotator& Target, float DeltaTime, float InterpSpeed);
 
 	/** Interpolate rotator from Current to Target. Scaled by distance to Target, so it has a strong start speed and ease out. */
-	static CORE_API FRotator RInterpTo( const FRotator& Current, const FRotator& Target, float DeltaTime, float InterpSpeed);
+	UE_NODISCARD static CORE_API FRotator RInterpTo( const FRotator& Current, const FRotator& Target, float DeltaTime, float InterpSpeed);
 
 	/** Interpolate float from Current to Target with constant step */
 	template<typename T1, typename T2 = T1, typename T3 = T2, typename T4 = T3>
-	static auto FInterpConstantTo( T1 Current, T2 Target, T3 DeltaTime, T4 InterpSpeed )
+	UE_NODISCARD static auto FInterpConstantTo( T1 Current, T2 Target, T3 DeltaTime, T4 InterpSpeed )
 	{
 		using RetType = decltype(T1() * T2() * T3() * T4());
 	
@@ -1330,7 +1330,7 @@ public:
 
 	/** Interpolate float from Current to Target. Scaled by distance to Target, so it has a strong start speed and ease out. */
 	template<typename T1, typename T2 = T1, typename T3 = T2, typename T4 = T3>
-	static auto FInterpTo( T1  Current, T2 Target, T3 DeltaTime, T4 InterpSpeed )
+	UE_NODISCARD static auto FInterpTo( T1  Current, T2 Target, T3 DeltaTime, T4 InterpSpeed )
 	{
 		using RetType = decltype(T1() * T2() * T3() * T4());
 	
@@ -1356,15 +1356,15 @@ public:
 	}
 
 	/** Interpolate Linear Color from Current to Target. Scaled by distance to Target, so it has a strong start speed and ease out. */
-	static CORE_API FLinearColor CInterpTo(const FLinearColor& Current, const FLinearColor& Target, float DeltaTime, float InterpSpeed);
+	UE_NODISCARD static CORE_API FLinearColor CInterpTo(const FLinearColor& Current, const FLinearColor& Target, float DeltaTime, float InterpSpeed);
 
 	/** Interpolate quaternion from Current to Target with constant step (in radians) */
 	template< class T >
-	static CORE_API UE::Math::TQuat<T> QInterpConstantTo(const UE::Math::TQuat<T>& Current, const UE::Math::TQuat<T>& Target, float DeltaTime, float InterpSpeed);
+	UE_NODISCARD static CORE_API UE::Math::TQuat<T> QInterpConstantTo(const UE::Math::TQuat<T>& Current, const UE::Math::TQuat<T>& Target, float DeltaTime, float InterpSpeed);
 
 	/** Interpolate quaternion from Current to Target. Scaled by angle to Target, so it has a strong start speed and ease out. */
 	template< class T >
-	static CORE_API UE::Math::TQuat<T> QInterpTo(const UE::Math::TQuat<T>& Current, const UE::Math::TQuat<T>& Target, float DeltaTime, float InterpSpeed);
+	UE_NODISCARD static CORE_API UE::Math::TQuat<T> QInterpTo(const UE::Math::TQuat<T>& Current, const UE::Math::TQuat<T>& Target, float DeltaTime, float InterpSpeed);
 
 	/**
 	* Returns an approximation of Exp(-X) based on a Taylor expansion that has had the coefficients adjusted (using
@@ -1373,11 +1373,11 @@ public:
 	* to zero for large X). 
 	*/
 	template<class T>
-	static T InvExpApprox(T X)
+	UE_NODISCARD static constexpr T InvExpApprox(T X)
 	{
-		const T A(1.00746054); // 1 / 1! in Taylor series
-		const T B(0.45053901); // 1 / 2! in Taylor series
-		const T C(0.25724632); // 1 / 3! in Taylor series
+		constexpr T A(1.00746054f); // 1 / 1! in Taylor series
+		constexpr T B(0.45053901f); // 1 / 2! in Taylor series
+		constexpr T C(0.25724632f); // 1 / 3! in Taylor series
 		return 1 / (1 + A * X + B * X * X + C * X * X * X);		
 	}
 	
@@ -1392,7 +1392,7 @@ public:
 	* @param  InSmoothingTime Timescale over which to smooth. Larger values result in more smoothed behaviour. Can be zero.
 	*/
 	template< class T >
-    static void ExponentialSmoothingApprox(
+    static constexpr void ExponentialSmoothingApprox(
         T&          InOutValue,
         const T&    InTargetValue,
         const float InDeltaTime,
@@ -1428,7 +1428,7 @@ public:
 	 * @param  InSmoothingTime   Timescale over which to smooth. Larger values result in more smoothed behaviour. Can be zero.
 	 */
 	template< class T >
-	static void CriticallyDampedSmoothing(
+	static constexpr void CriticallyDampedSmoothing(
 		T&          InOutValue,
 		T&          InOutValueRate,
 		const T&    InTargetValue,
@@ -1625,7 +1625,7 @@ public:
 	 *
 	 * @return  Pulsating value (0.0-1.0)
 	 */
-	static float MakePulsatingValue( const double InCurrentTime, const float InPulsesPerSecond, const float InPhase = 0.0f )
+	UE_NODISCARD static float MakePulsatingValue( const double InCurrentTime, const float InPulsesPerSecond, const float InPhase = 0.0f )
 	{
 		return 0.5f + 0.5f * FMath::Sin( ( ( 0.25f + InPhase ) * TWO_PI ) + ( (float)InCurrentTime * TWO_PI ) * InPulsesPerSecond );
 	}
@@ -1643,7 +1643,7 @@ public:
 	 * @return The point of intersection between the ray and the plane.
 	 */
 	template<typename FReal>
-	static UE::Math::TVector<FReal> RayPlaneIntersection(const UE::Math::TVector<FReal>& RayOrigin, const UE::Math::TVector<FReal>& RayDirection, const UE::Math::TPlane<FReal>& Plane);
+	UE_NODISCARD static UE::Math::TVector<FReal> RayPlaneIntersection(const UE::Math::TVector<FReal>& RayOrigin, const UE::Math::TVector<FReal>& RayDirection, const UE::Math::TPlane<FReal>& Plane);
 
 	/**
 	 * Find the intersection of a line and an offset plane. Assumes that the
@@ -1658,7 +1658,7 @@ public:
 	 * @return The point of intersection between the line and the plane.
 	 */
 	template<typename FReal>
-	static UE::Math::TVector<FReal> LinePlaneIntersection(const UE::Math::TVector<FReal>& Point1, const UE::Math::TVector<FReal>& Point2, const UE::Math::TVector<FReal>& PlaneOrigin, const UE::Math::TVector<FReal>& PlaneNormal);
+	UE_NODISCARD static UE::Math::TVector<FReal> LinePlaneIntersection(const UE::Math::TVector<FReal>& Point1, const UE::Math::TVector<FReal>& Point2, const UE::Math::TVector<FReal>& PlaneOrigin, const UE::Math::TVector<FReal>& PlaneNormal);
 
 	/**
 	 * Find the intersection of a line and a plane. Assumes that the line and
@@ -1672,12 +1672,12 @@ public:
 	 * @return The point of intersection between the line and the plane.
 	 */
 	template<typename FReal>
-	static UE::Math::TVector<FReal> LinePlaneIntersection(const UE::Math::TVector<FReal>& Point1, const UE::Math::TVector<FReal>& Point2, const UE::Math::TPlane<FReal>& Plane);
+	UE_NODISCARD static UE::Math::TVector<FReal> LinePlaneIntersection(const UE::Math::TVector<FReal>& Point1, const UE::Math::TVector<FReal>& Point2, const UE::Math::TPlane<FReal>& Plane);
 
 
 	// @parma InOutScissorRect should be set to View.ViewRect before the call
 	// @return 0: light is not visible, 1:use scissor rect, 2: no scissor rect needed
-	static CORE_API uint32 ComputeProjectedSphereScissorRect(struct FIntRect& InOutScissorRect, FVector SphereOrigin, float Radius, FVector ViewOrigin, const FMatrix& ViewMatrix, const FMatrix& ProjMatrix);
+	UE_NODISCARD static CORE_API uint32 ComputeProjectedSphereScissorRect(struct FIntRect& InOutScissorRect, FVector SphereOrigin, float Radius, FVector ViewOrigin, const FMatrix& ViewMatrix, const FMatrix& ProjMatrix);
 
 	// @param ConeOrigin Cone origin
 	// @param ConeDirection Cone direction
@@ -1686,7 +1686,7 @@ public:
 	// @param SinConeAngle Sin of the cone angle
 	// @return Minimal bounding sphere encompassing given cone
 	template<typename FReal>
-	static UE::Math::TSphere<FReal> ComputeBoundingSphereForCone(UE::Math::TVector<FReal> const& ConeOrigin, UE::Math::TVector<FReal> const& ConeDirection, FReal ConeRadius, FReal CosConeAngle, FReal SinConeAngle);
+	UE_NODISCARD static UE::Math::TSphere<FReal> ComputeBoundingSphereForCone(UE::Math::TVector<FReal> const& ConeOrigin, UE::Math::TVector<FReal> const& ConeDirection, FReal ConeRadius, FReal CosConeAngle, FReal SinConeAngle);
 
 	/** 
 	 * Determine if a plane and an AABB intersect
@@ -1694,7 +1694,7 @@ public:
 	 * @param AABB - the axis aligned bounding box to test
 	 * @return if collision occurs
 	 */
-	static CORE_API bool PlaneAABBIntersection(const FPlane& P, const FBox& AABB);
+	UE_NODISCARD static CORE_API bool PlaneAABBIntersection(const FPlane& P, const FBox& AABB);
 
 	/**
 	 * Determine the position of an AABB relative to a plane:
@@ -1703,7 +1703,7 @@ public:
 	 * @param AABB - the axis aligned bounding box to test
 	 * @return -1 if below, 1 if above, 0 if intersects
 	 */
-	static CORE_API int32 PlaneAABBRelativePosition(const FPlane& P, const FBox& AABB);
+	UE_NODISCARD static CORE_API int32 PlaneAABBRelativePosition(const FPlane& P, const FBox& AABB);
 
 	/**
 	 * Performs a sphere vs box intersection test using Arvo's algorithm:
@@ -1712,62 +1712,62 @@ public:
 	 *		if (SphereCenter(i) < BoxMin(i)) d2 += (SphereCenter(i) - BoxMin(i)) ^ 2
 	 *		else if (SphereCenter(i) > BoxMax(i)) d2 += (SphereCenter(i) - BoxMax(i)) ^ 2
 	 *
-	 * @param Sphere the center of the sphere being tested against the AABB
+	 * @param SphereCenter the center of the sphere being tested against the AABB
 	 * @param RadiusSquared the size of the sphere being tested
 	 * @param AABB the box being tested against
 	 *
 	 * @return Whether the sphere/box intersect or not.
 	 */
 	template<typename FReal>
-	static bool SphereAABBIntersection(const UE::Math::TVector<FReal>& SphereCenter,const FReal RadiusSquared,const UE::Math::TBox<FReal>& AABB);
+	UE_NODISCARD static bool SphereAABBIntersection(const UE::Math::TVector<FReal>& SphereCenter,const FReal RadiusSquared,const UE::Math::TBox<FReal>& AABB);
 
 	/**
 	 * Converts a sphere into a point plus radius squared for the test above
 	 */
 	template<typename FReal>
-	static bool SphereAABBIntersection(const UE::Math::TSphere<FReal>& Sphere, const UE::Math::TBox<FReal>& AABB);
+	UE_NODISCARD static bool SphereAABBIntersection(const UE::Math::TSphere<FReal>& Sphere, const UE::Math::TBox<FReal>& AABB);
 
 	/** Determines whether a point is inside a box. */
 	template<typename FReal>
-	static bool PointBoxIntersection( const UE::Math::TVector<FReal>& Point, const UE::Math::TBox<FReal>& Box );
+	UE_NODISCARD static bool PointBoxIntersection( const UE::Math::TVector<FReal>& Point, const UE::Math::TBox<FReal>& Box );
 
 	/** Determines whether a line intersects a box. */
 	template<typename FReal>
-	static bool LineBoxIntersection( const UE::Math::TBox<FReal>& Box, const UE::Math::TVector<FReal>& Start, const UE::Math::TVector<FReal>& End, const UE::Math::TVector<FReal>& Direction );
+	UE_NODISCARD static bool LineBoxIntersection( const UE::Math::TBox<FReal>& Box, const UE::Math::TVector<FReal>& Start, const UE::Math::TVector<FReal>& End, const UE::Math::TVector<FReal>& Direction );
 
 	/** Determines whether a line intersects a box. This overload avoids the need to do the reciprocal every time. */
 	template<typename FReal>
-	static bool LineBoxIntersection( const UE::Math::TBox<FReal>& Box, const UE::Math::TVector<FReal>& Start, const UE::Math::TVector<FReal>& End, const UE::Math::TVector<FReal>& Direction, const UE::Math::TVector<FReal>& OneOverDirection );
+	UE_NODISCARD static bool LineBoxIntersection( const UE::Math::TBox<FReal>& Box, const UE::Math::TVector<FReal>& Start, const UE::Math::TVector<FReal>& End, const UE::Math::TVector<FReal>& Direction, const UE::Math::TVector<FReal>& OneOverDirection );
 
 	/* Swept-Box vs Box test */
 	static CORE_API bool LineExtentBoxIntersection(const FBox& inBox, const FVector& Start, const FVector& End, const FVector& Extent, FVector& HitLocation, FVector& HitNormal, float& HitTime);
 
 	/** Determines whether a line intersects a sphere. */
 	template<typename FReal>
-	static bool LineSphereIntersection(const UE::Math::TVector<FReal>& Start,const UE::Math::TVector<FReal>& Dir, FReal Length,const UE::Math::TVector<FReal>& Origin, FReal Radius);
+	UE_NODISCARD static bool LineSphereIntersection(const UE::Math::TVector<FReal>& Start,const UE::Math::TVector<FReal>& Dir, FReal Length,const UE::Math::TVector<FReal>& Origin, FReal Radius);
 
 	/**
 	 * Assumes the cone tip is at 0,0,0 (means the SphereCenter is relative to the cone tip)
 	 * @return true: cone and sphere do intersect, false otherwise
 	 */
-	static CORE_API bool SphereConeIntersection(const FVector& SphereCenter, float SphereRadius, const FVector& ConeAxis, float ConeAngleSin, float ConeAngleCos);
+	UE_NODISCARD static CORE_API bool SphereConeIntersection(const FVector& SphereCenter, float SphereRadius, const FVector& ConeAxis, float ConeAngleSin, float ConeAngleCos);
 
 	/** Find the point on the line segment from LineStart to LineEnd which is closest to Point */
-	static CORE_API FVector ClosestPointOnLine(const FVector& LineStart, const FVector& LineEnd, const FVector& Point);
+	UE_NODISCARD static CORE_API FVector ClosestPointOnLine(const FVector& LineStart, const FVector& LineEnd, const FVector& Point);
 
 	/** Find the point on the infinite line between two points (LineStart, LineEnd) which is closest to Point */
-	static CORE_API FVector ClosestPointOnInfiniteLine(const FVector& LineStart, const FVector& LineEnd, const FVector& Point);
+	UE_NODISCARD static CORE_API FVector ClosestPointOnInfiniteLine(const FVector& LineStart, const FVector& LineEnd, const FVector& Point);
 
 	/** Compute intersection point of three planes. Return 1 if valid, 0 if infinite. */
 	template<typename FReal>
-	static bool IntersectPlanes3(UE::Math::TVector<FReal>& I, const UE::Math::TPlane<FReal>& P1, const UE::Math::TPlane<FReal>& P2, const UE::Math::TPlane<FReal>& P3 );
+	UE_NODISCARD static bool IntersectPlanes3(UE::Math::TVector<FReal>& I, const UE::Math::TPlane<FReal>& P1, const UE::Math::TPlane<FReal>& P2, const UE::Math::TPlane<FReal>& P3 );
 
 	/**
 	 * Compute intersection point and direction of line joining two planes.
 	 * Return 1 if valid, 0 if infinite.
 	 */
 	template<typename FReal>
-	static bool IntersectPlanes2(UE::Math::TVector<FReal>& I, UE::Math::TVector<FReal>& D, const UE::Math::TPlane<FReal>& P1, const UE::Math::TPlane<FReal>& P2);
+	UE_NODISCARD static bool IntersectPlanes2(UE::Math::TVector<FReal>& I, UE::Math::TVector<FReal>& D, const UE::Math::TPlane<FReal>& P1, const UE::Math::TPlane<FReal>& P2);
 
 	/**
 	 * Calculates the distance of a given Point in world space to a given line,
@@ -1781,7 +1781,7 @@ public:
 	 * @return	distance of Point from line defined by (Origin, Direction)
 	 */
 	static CORE_API float PointDistToLine(const FVector &Point, const FVector &Direction, const FVector &Origin, FVector &OutClosestPoint);
-	static CORE_API float PointDistToLine(const FVector &Point, const FVector &Direction, const FVector &Origin);
+	UE_NODISCARD static CORE_API float PointDistToLine(const FVector &Point, const FVector &Direction, const FVector &Origin);
 
 	/**
 	 * Returns closest point on a segment to a given point.
@@ -1794,7 +1794,7 @@ public:
 	 *
 	 * @return	point on the segment defined by (StartPoint, EndPoint) that is closest to Point.
 	 */
-	static CORE_API FVector ClosestPointOnSegment(const FVector &Point, const FVector &StartPoint, const FVector &EndPoint);
+	UE_NODISCARD static CORE_API FVector ClosestPointOnSegment(const FVector &Point, const FVector &StartPoint, const FVector &EndPoint);
 
 	/**
 	* FVector2D version of ClosestPointOnSegment.
@@ -1808,7 +1808,7 @@ public:
 	*
 	* @return	point on the segment defined by (StartPoint, EndPoint) that is closest to Point.
 	*/
-	static CORE_API FVector2D ClosestPointOnSegment2D(const FVector2D &Point, const FVector2D &StartPoint, const FVector2D &EndPoint);
+	UE_NODISCARD static CORE_API FVector2D ClosestPointOnSegment2D(const FVector2D &Point, const FVector2D &StartPoint, const FVector2D &EndPoint);
 
 	/**
 	 * Returns distance from a point to the closest point on a segment.
@@ -1819,7 +1819,7 @@ public:
 	 *
 	 * @return	closest distance from Point to segment defined by (StartPoint, EndPoint).
 	 */
-	static CORE_API float PointDistToSegment(const FVector &Point, const FVector &StartPoint, const FVector &EndPoint);
+	UE_NODISCARD static CORE_API float PointDistToSegment(const FVector &Point, const FVector &StartPoint, const FVector &EndPoint);
 
 	/**
 	 * Returns square of the distance from a point to the closest point on a segment.
@@ -1830,7 +1830,7 @@ public:
 	 *
 	 * @return	square of the closest distance from Point to segment defined by (StartPoint, EndPoint).
 	 */
-	static CORE_API float PointDistToSegmentSquared(const FVector &Point, const FVector &StartPoint, const FVector &EndPoint);
+	UE_NODISCARD static CORE_API float PointDistToSegmentSquared(const FVector &Point, const FVector &StartPoint, const FVector &EndPoint);
 
 	/** 
 	 * Find closest points between 2 segments.
@@ -1864,7 +1864,7 @@ public:
 	 * @param Plane		- plane to intersect with
 	 * @return time(T) of intersection
 	 */
-	static CORE_API float GetTForSegmentPlaneIntersect(const FVector& StartPoint, const FVector& EndPoint, const FPlane& Plane);
+	UE_NODISCARD static CORE_API float GetTForSegmentPlaneIntersect(const FVector& StartPoint, const FVector& EndPoint, const FPlane& Plane);
 
 	/**
 	 * Returns true if there is an intersection between the segment specified by StartPoint and Endpoint, and
@@ -1885,7 +1885,7 @@ public:
 	* @param EndPoint   - end point of segment
 	* @param A, B, C	- points defining the triangle 
 	* @param OutIntersectPoint - out var for the point on the segment that intersects the triangle (if any)
-	* @param OutNormal - out var for the triangle normal
+	* @param OutTriangleNormal - out var for the triangle normal
 	* @return true if intersection occurred
 	*/
 	static CORE_API bool SegmentTriangleIntersection(const FVector& StartPoint, const FVector& EndPoint, const FVector& A, const FVector& B, const FVector& C, FVector& OutIntersectPoint, FVector& OutTriangleNormal);
@@ -1913,7 +1913,7 @@ public:
 	 *
 	 * @return	Point on triangle ABC closest to given point
 	 */
-	static CORE_API FVector ClosestPointOnTriangleToPoint(const FVector& Point, const FVector& A, const FVector& B, const FVector& C);
+	UE_NODISCARD static CORE_API FVector ClosestPointOnTriangleToPoint(const FVector& Point, const FVector& A, const FVector& B, const FVector& C);
 
 	/**
 	 * Returns closest point on a tetrahedron to a point.
@@ -1925,7 +1925,7 @@ public:
 	 *
 	 * @return	Point on tetrahedron ABCD closest to given point
 	 */
-	static CORE_API FVector ClosestPointOnTetrahedronToPoint(const FVector& Point, const FVector& A, const FVector& B, const FVector& C, const FVector& D);
+	UE_NODISCARD static CORE_API FVector ClosestPointOnTetrahedronToPoint(const FVector& Point, const FVector& A, const FVector& B, const FVector& C, const FVector& D);
 
 	/** 
 	 * Find closest point on a Sphere to a Line.
@@ -1944,13 +1944,13 @@ public:
 	 * @param Point - The Point in question
 	 * @param ConeStartPoint - the beginning of the cone (with the smallest radius)
 	 * @param ConeLine - the line out from the start point that ends at the largest radius point of the cone
-	 * @param radiusAtStart - the radius at the ConeStartPoint (0 for a 'proper' cone)
-	 * @param radiusAtEnd - the largest radius of the cone
-	 * @param percentageOut - output variable the holds how much within the cone the point is (1 = on center line, 0 = on exact edge or outside cone).
+	 * @param RadiusAtStart - the radius at the ConeStartPoint (0 for a 'proper' cone)
+	 * @param RadiusAtEnd - the largest radius of the cone
+	 * @param PercentageOut - output variable the holds how much within the cone the point is (1 = on center line, 0 = on exact edge or outside cone).
 	 *
 	 * @return true if the point is within the cone, false otherwise.
 	 */
-	static CORE_API bool GetDistanceWithinConeSegment(FVector Point, FVector ConeStartPoint, FVector ConeLine, float RadiusAtStart, float RadiusAtEnd, float &PercentageOut);
+	UE_NODISCARD static CORE_API bool GetDistanceWithinConeSegment(FVector Point, FVector ConeStartPoint, FVector ConeLine, float RadiusAtStart, float RadiusAtEnd, float &PercentageOut);
 
 	/**
 	 * Determines whether a given set of points are coplanar, with a tolerance. Any three points or less are always coplanar.
@@ -1960,52 +1960,52 @@ public:
 	 *
 	 * @return Whether the points are relatively coplanar, based on the tolerance
 	 */
-	static CORE_API bool PointsAreCoplanar(const TArray<FVector>& Points, const float Tolerance = 0.1f);
+	UE_NODISCARD static CORE_API bool PointsAreCoplanar(const TArray<FVector>& Points, const float Tolerance = 0.1f);
 
 	/**
 	 * Truncates a floating point number to half if closer than the given tolerance.
-	 * @param Value				Floating point number to truncate
+	 * @param F					Floating point number to truncate
 	 * @param Tolerance			Maximum allowed difference to 0.5 in order to truncate
 	 * @return					The truncated value
 	 */
-	static CORE_API float TruncateToHalfIfClose(float F, float Tolerance = SMALL_NUMBER);
-	static CORE_API double TruncateToHalfIfClose(double F, double Tolerance = SMALL_NUMBER);
+	UE_NODISCARD static CORE_API float TruncateToHalfIfClose(float F, float Tolerance = SMALL_NUMBER);
+	UE_NODISCARD static CORE_API double TruncateToHalfIfClose(double F, double Tolerance = SMALL_NUMBER);
 
 	/**
 	* Converts a floating point number to the nearest integer, equidistant ties go to the value which is closest to an even value: 1.5 becomes 2, 0.5 becomes 0
 	* @param F		Floating point value to convert
 	* @return		The rounded integer
 	*/
-	static CORE_API float RoundHalfToEven(float F);
-	static CORE_API double RoundHalfToEven(double F);
+	UE_NODISCARD static CORE_API float RoundHalfToEven(float F);
+	UE_NODISCARD static CORE_API double RoundHalfToEven(double F);
 
 	/**
 	* Converts a floating point number to the nearest integer, equidistant ties go to the value which is further from zero: -0.5 becomes -1.0, 0.5 becomes 1.0
 	* @param F		Floating point value to convert
 	* @return		The rounded integer
 	*/
-	static CORE_API float RoundHalfFromZero(float F);
-	static CORE_API double RoundHalfFromZero(double F);
+	UE_NODISCARD static CORE_API float RoundHalfFromZero(float F);
+	UE_NODISCARD static CORE_API double RoundHalfFromZero(double F);
 
 	/**
 	* Converts a floating point number to the nearest integer, equidistant ties go to the value which is closer to zero: -0.5 becomes 0, 0.5 becomes 0
 	* @param F		Floating point value to convert
 	* @return		The rounded integer
 	*/
-	static CORE_API float RoundHalfToZero(float F);
-	static CORE_API double RoundHalfToZero(double F);
+	UE_NODISCARD static CORE_API float RoundHalfToZero(float F);
+	UE_NODISCARD static CORE_API double RoundHalfToZero(double F);
 
 	/**
 	* Converts a floating point number to an integer which is further from zero, "larger" in absolute value: 0.1 becomes 1, -0.1 becomes -1
 	* @param F		Floating point value to convert
 	* @return		The rounded integer
 	*/
-	static FORCEINLINE float RoundFromZero(float F)
+	UE_NODISCARD static FORCEINLINE float RoundFromZero(float F)
 	{
 		return (F < 0.0f) ? FloorToFloat(F) : CeilToFloat(F);
 	}
 
-	static FORCEINLINE double RoundFromZero(double F)
+	UE_NODISCARD static FORCEINLINE double RoundFromZero(double F)
 	{
 		return (F < 0.0) ? FloorToDouble(F) : CeilToDouble(F);
 	}
@@ -2015,12 +2015,12 @@ public:
 	* @param F		Floating point value to convert
 	* @return		The rounded integer
 	*/
-	static FORCEINLINE float RoundToZero(float F)
+	UE_NODISCARD static FORCEINLINE float RoundToZero(float F)
 	{
 		return (F < 0.0f) ? CeilToFloat(F) : FloorToFloat(F);
 	}
 
-	static FORCEINLINE double RoundToZero(double F)
+	UE_NODISCARD static FORCEINLINE double RoundToZero(double F)
 	{
 		return (F < 0.0) ? CeilToDouble(F) : FloorToDouble(F);
 	}
@@ -2030,12 +2030,12 @@ public:
 	* @param F		Floating point value to convert
 	* @return		The rounded integer
 	*/
-	static FORCEINLINE float RoundToNegativeInfinity(float F)
+	UE_NODISCARD static FORCEINLINE float RoundToNegativeInfinity(float F)
 	{
 		return FloorToFloat(F);
 	}
 
-	static FORCEINLINE double RoundToNegativeInfinity(double F)
+	UE_NODISCARD static FORCEINLINE double RoundToNegativeInfinity(double F)
 	{
 		return FloorToDouble(F);
 	}
@@ -2045,12 +2045,12 @@ public:
 	* @param F		Floating point value to convert
 	* @return		The rounded integer
 	*/
-	static FORCEINLINE float RoundToPositiveInfinity(float F)
+	UE_NODISCARD static FORCEINLINE float RoundToPositiveInfinity(float F)
 	{
 		return CeilToFloat(F);
 	}
 
-	static FORCEINLINE double RoundToPositiveInfinity(double F)
+	UE_NODISCARD static FORCEINLINE double RoundToPositiveInfinity(double F)
 	{
 		return CeilToDouble(F);
 	}
@@ -2064,7 +2064,7 @@ public:
 	 * @param	Val		The value to use
 	 * @return	FString	The human readable string
 	 */
-	static CORE_API FString FormatIntToHumanReadable(int32 Val);
+	UE_NODISCARD static CORE_API FString FormatIntToHumanReadable(int32 Val);
 
 
 	// Utilities
@@ -2076,7 +2076,7 @@ public:
 	 * @param NumBytes		Number of bytes to test (will be rounded down to a multiple of 4)
 	 * @return				true if the memory region passed the test
 	 */
-	static CORE_API bool MemoryTest( void* BaseAddress, uint32 NumBytes );
+	UE_NODISCARD static CORE_API bool MemoryTest( void* BaseAddress, uint32 NumBytes );
 
 	/**
 	 * Evaluates a numerical equation.
@@ -2096,23 +2096,23 @@ public:
 	 * Computes the barycentric coordinates for a given point in a triangle, only considering the XY coordinates - simpler than the Compute versions
 	 *
 	 * @param	Point			point to convert to barycentric coordinates (in plane of ABC)
-	 * @param	A,B,C			three non-colinear points defining a triangle in CCW
+	 * @param	A,B,C			three non-collinear points defining a triangle in CCW
 	 * 
 	 * @return Vector containing the three weights a,b,c such that Point = a*A + b*B + c*C
 	 *							                                or Point = A + b*(B-A) + c*(C-A) = (1-b-c)*A + b*B + c*C
 	 */
-	static CORE_API FVector GetBaryCentric2D(const FVector& Point, const FVector& A, const FVector& B, const FVector& C);
+	UE_NODISCARD static CORE_API FVector GetBaryCentric2D(const FVector& Point, const FVector& A, const FVector& B, const FVector& C);
 
 	/**
 	 * Computes the barycentric coordinates for a given point in a triangle - simpler than the Compute versions
 	 *
 	 * @param	Point			point to convert to barycentric coordinates (in plane of ABC)
-	 * @param	A,B,C			three non-colinear points defining a triangle in CCW
+	 * @param	A,B,C			three non-collinear points defining a triangle in CCW
 	 *
 	 * @return Vector containing the three weights a,b,c such that Point = a*A + b*B + c*C
 	 *							                                or Point = A + b*(B-A) + c*(C-A) = (1-b-c)*A + b*B + c*C
 	 */
-	static CORE_API FVector GetBaryCentric2D(const FVector2D& Point, const FVector2D& A, const FVector2D& B, const FVector2D& C);
+	UE_NODISCARD static CORE_API FVector GetBaryCentric2D(const FVector2D& Point, const FVector2D& A, const FVector2D& B, const FVector2D& C);
 
 	/**
 	 * Computes the barycentric coordinates for a given point in a triangle
@@ -2123,7 +2123,7 @@ public:
 	 * @return Vector containing the three weights a,b,c such that Point = a*A + b*B + c*C
 	 *							                               or Point = A + b*(B-A) + c*(C-A) = (1-b-c)*A + b*B + c*C
 	 */
-	static CORE_API FVector ComputeBaryCentric2D(const FVector& Point, const FVector& A, const FVector& B, const FVector& C);
+	UE_NODISCARD static CORE_API FVector ComputeBaryCentric2D(const FVector& Point, const FVector& A, const FVector& B, const FVector& C);
 
 	/**
 	 * Computes the barycentric coordinates for a given point on a tetrahedron (3D)
@@ -2133,7 +2133,7 @@ public:
 	 *
 	 * @return Vector containing the four weights a,b,c,d such that Point = a*A + b*B + c*C + d*D
 	 */
-	static CORE_API FVector4 ComputeBaryCentric3D(const FVector& Point, const FVector& A, const FVector& B, const FVector& C, const FVector& D);
+	UE_NODISCARD static CORE_API FVector4 ComputeBaryCentric3D(const FVector& Point, const FVector& A, const FVector& B, const FVector& C, const FVector& D);
 
 	/** 32 bit values where BitFlag[x] == (1<<x) */
 	static CORE_API const uint32 BitFlag[32];
@@ -2149,7 +2149,7 @@ public:
 	 * @return Smoothed value between 0 and 1
 	 */
 	template<typename T>
-	static T SmoothStep(T A, T B, T X)
+	UE_NODISCARD static constexpr T SmoothStep(T A, T B, T X)
 	{
 		if (X < A)
 		{
@@ -2167,7 +2167,7 @@ public:
 	 * Get a bit in memory created from bitflags (uint32 Value:1), used for EngineShowFlags,
 	 * TestBitFieldFunctions() tests the implementation
 	 */
-	static inline bool ExtractBoolFromBitfield(uint8* Ptr, uint32 Index)
+	UE_NODISCARD static constexpr bool ExtractBoolFromBitfield(uint8* Ptr, uint32 Index)
 	{
 		uint8* BytePtr = Ptr + Index / 8;
 		uint8 Mask = (uint8)(1 << (Index & 0x7));
@@ -2179,7 +2179,7 @@ public:
 	 * Set a bit in memory created from bitflags (uint32 Value:1), used for EngineShowFlags,
 	 * TestBitFieldFunctions() tests the implementation
 	 */
-	static inline void SetBoolInBitField(uint8* Ptr, uint32 Index, bool bSet)
+	static constexpr void SetBoolInBitField(uint8* Ptr, uint32 Index, bool bSet)
 	{
 		uint8* BytePtr = Ptr + Index / 8;
 		uint8 Mask = (uint8)(1 << (Index & 0x7));
@@ -2202,7 +2202,7 @@ public:
 
 	// @param x assumed to be in this range: 0..1
 	// @return 0..255
-	static uint8 Quantize8UnsignedByte(float x)
+	UE_NODISCARD static uint8 Quantize8UnsignedByte(float x)
 	{
 		// 0..1 -> 0..255
 		int32 Ret = (int32)(x * 255.f + 0.5f);
@@ -2215,7 +2215,7 @@ public:
 	
 	// @param x assumed to be in this range: -1..1
 	// @return 0..255
-	static uint8 Quantize8SignedByte(float x)
+	UE_NODISCARD static uint8 Quantize8SignedByte(float x)
 	{
 		// -1..1 -> 0..1
 		float y = x * 0.5f + 0.5f;
@@ -2224,7 +2224,7 @@ public:
 	}
 
 	// Use the Euclidean method to find the GCD
-	static int32 GreatestCommonDivisor(int32 a, int32 b)
+	UE_NODISCARD static constexpr int32 GreatestCommonDivisor(int32 a, int32 b)
 	{
 		while (b != 0)
 		{
@@ -2237,7 +2237,7 @@ public:
 
 	// LCM = a/gcd * b
 	// a and b are the number we want to find the lcm
-	static int32 LeastCommonMultiplier(int32 a, int32 b)
+	UE_NODISCARD static constexpr int32 LeastCommonMultiplier(int32 a, int32 b)
 	{
 		int32 CurrentGcd = GreatestCommonDivisor(a, b);
 		return CurrentGcd == 0 ? 0 : (a / CurrentGcd) * b;
@@ -2250,7 +2250,7 @@ public:
 	 *
 	 * @return	Perlin noise in the range of -1.0 to 1.0
 	 */
-	static CORE_API float PerlinNoise1D(float Value);
+	UE_NODISCARD static CORE_API float PerlinNoise1D(float Value);
 
 	/**
 	* Generates a 1D Perlin noise sample at the given location.  Returns a continuous random value between -1.0 and 1.0.
@@ -2259,7 +2259,7 @@ public:
 	*
 	* @return	Perlin noise in the range of -1.0 to 1.0
 	*/
-	static CORE_API float PerlinNoise2D(const FVector2D& Location);
+	UE_NODISCARD static CORE_API float PerlinNoise2D(const FVector2D& Location);
 	 
 
 	/**
@@ -2269,7 +2269,7 @@ public:
 	*
 	* @return	Perlin noise in the range of -1.0 to 1.0
 	*/
-	static CORE_API float PerlinNoise3D(const FVector& Location);
+	UE_NODISCARD static CORE_API float PerlinNoise3D(const FVector& Location);
 
 	/**
 	 * Calculates the new value in a weighted moving average series using the previous value and the weight
@@ -2281,7 +2281,7 @@ public:
 	 * @return the next value in the series
 	 */
 	template<typename T>
-	static inline T WeightedMovingAverage(T CurrentSample, T PreviousSample, T Weight)
+	UE_NODISCARD static inline T WeightedMovingAverage(T CurrentSample, T PreviousSample, T Weight)
 	{
 		Weight = Clamp<T>(Weight, 0.f, 1.f);
 		T WAvg = (CurrentSample * Weight) + (PreviousSample * (1.f - Weight));
@@ -2302,7 +2302,7 @@ public:
 	 * @return the next value in the series
 	 */
 	template<typename T>
-	static  inline T DynamicWeightedMovingAverage(T CurrentSample, T PreviousSample, T MaxDistance, T MinWeight, T MaxWeight)
+	UE_NODISCARD static inline T DynamicWeightedMovingAverage(T CurrentSample, T PreviousSample, T MaxDistance, T MinWeight, T MaxWeight)
 	{
 		// We need the distance between samples to determine how much of each weight to use
 		const T Distance = Abs<T>(CurrentSample - PreviousSample);
@@ -2319,18 +2319,18 @@ public:
 	// Temporary support for ambiguities to simplify UE4 -> UE5 upgrades - falls back to float variant.
 	// A call is considered ambiguous if it passes no float types, any long double, or multiple mismatched float/double types.
 	template<typename T> UE_DEPRECATED(5.0, "Arguments cause function resolution ambiguity.")
-	static FORCEINLINE float Log2(T&& Value) { return Log2((float)Value); }
+	UE_NODISCARD static FORCEINLINE float Log2(T&& Value) { return Log2((float)Value); }
 	template<typename T1, typename T2, typename T3, TEMPLATE_REQUIRES(TIsAmbiguous<T1, T2, T3>)> UE_DEPRECATED(5.0, "Arguments cause function resolution ambiguity.")
-	static FORCEINLINE void PolarToCartesian(const T1 Rad, const T2 Ang, T3& OutX, T3& OutY) { float OX, OY; PolarToCartesian((float)Rad, (float)Ang, OX, OY); OutX = OX; OutY = OY; }
+	UE_NODISCARD static FORCEINLINE void PolarToCartesian(const T1 Rad, const T2 Ang, T3& OutX, T3& OutY) { float OX, OY; PolarToCartesian((float)Rad, (float)Ang, OX, OY); OutX = OX; OutY = OY; }
 	template<typename T1, typename T2, typename T3, TEMPLATE_REQUIRES(TIsAmbiguous<T1, T2, T3>)> UE_DEPRECATED(5.0, "Arguments cause function resolution ambiguity.")
-	static FORCEINLINE void CartesianToPolar(const T1 X, const T2 Y, T3& OutRad, T3& OutAng) { float ORad, OAng; CartesianToPolar((float)X, (float)Y, ORad, OAng); OutRad = ORad; OutAng = OAng; }
+	UE_NODISCARD static FORCEINLINE void CartesianToPolar(const T1 X, const T2 Y, T3& OutRad, T3& OutAng) { float ORad, OAng; CartesianToPolar((float)X, (float)Y, ORad, OAng); OutRad = ORad; OutAng = OAng; }
 
 	template<typename T1, typename T2, typename T3, TEMPLATE_REQUIRES(TIsAmbiguous<T1, T2, T3>)> UE_DEPRECATED(5.0, "Arguments cause function resolution ambiguity.")
-	static FORCEINLINE float SmoothStep(T1&& A, T2&& B, T3&& C) { return SmoothStep((float)A, (float)B, (float)C); }
+	UE_NODISCARD static FORCEINLINE float SmoothStep(T1&& A, T2&& B, T3&& C) { return SmoothStep((float)A, (float)B, (float)C); }
 	template<typename T1, typename T2, typename T3, TEMPLATE_REQUIRES(TIsAmbiguous<T1, T2, T3>)> UE_DEPRECATED(5.0, "Arguments cause function resolution ambiguity.")
-	static FORCEINLINE float WeightedMovingAverage(T1&& CurrentSample, T2&& PreviousSample, T3&& Weight) { return WeightedMovingAverage((float)CurrentSample, (float)PreviousSample, (float)Weight); }
+	UE_NODISCARD static FORCEINLINE float WeightedMovingAverage(T1&& CurrentSample, T2&& PreviousSample, T3&& Weight) { return WeightedMovingAverage((float)CurrentSample, (float)PreviousSample, (float)Weight); }
 	template<typename T1, typename T2, typename T3, typename T4, typename T5, TEMPLATE_REQUIRES(TIsAmbiguous<T1, T2, T3, T4, T5>)> UE_DEPRECATED(5.0, "Arguments cause function resolution ambiguity.")
-	static FORCEINLINE float DynamicWeightedMovingAverage(T1&& CurrentSample, T2&& PreviousSample, T3&& MaxDistance, T4&& MinWeight, T5&& MaxWeight) { return DynamicWeightedMovingAverage((float)CurrentSample, (float)PreviousSample, (float)MaxDistance, (float)MinWeight, (float)MaxWeight); }
+	UE_NODISCARD static FORCEINLINE float DynamicWeightedMovingAverage(T1&& CurrentSample, T2&& PreviousSample, T3&& MaxDistance, T4&& MinWeight, T5&& MaxWeight) { return DynamicWeightedMovingAverage((float)CurrentSample, (float)PreviousSample, (float)MaxDistance, (float)MinWeight, (float)MaxWeight); }
 };
 
 // LWC Conversion helpers - LWC_TODO: These are temporary and should be removed for UE5.0 final.
