Fix extra UE_NODISCARD added during ROBOMERGE

diff --git a/Engine/Source/Runtime/Core/Public/Math/UnrealMathUtility.h b/Engine/Source/Runtime/Core/Public/Math/UnrealMathUtility.h
--- a/Engine/Source/Runtime/Core/Public/Math/UnrealMathUtility.h
+++ b/Engine/Source/Runtime/Core/Public/Math/UnrealMathUtility.h
@@ -2321,9 +2321,9 @@ public:
 	template<typename T> UE_DEPRECATED(5.0, "Arguments cause function resolution ambiguity.")
 	UE_NODISCARD static FORCEINLINE float Log2(T&& Value) { return Log2((float)Value); }
 	template<typename T1, typename T2, typename T3, TEMPLATE_REQUIRES(TIsAmbiguous<T1, T2, T3>)> UE_DEPRECATED(5.0, "Arguments cause function resolution ambiguity.")
-	UE_NODISCARD static FORCEINLINE void PolarToCartesian(const T1 Rad, const T2 Ang, T3& OutX, T3& OutY) { float OX, OY; PolarToCartesian((float)Rad, (float)Ang, OX, OY); OutX = OX; OutY = OY; }
+	static FORCEINLINE void PolarToCartesian(const T1 Rad, const T2 Ang, T3& OutX, T3& OutY) { float OX, OY; PolarToCartesian((float)Rad, (float)Ang, OX, OY); OutX = OX; OutY = OY; }
 	template<typename T1, typename T2, typename T3, TEMPLATE_REQUIRES(TIsAmbiguous<T1, T2, T3>)> UE_DEPRECATED(5.0, "Arguments cause function resolution ambiguity.")
-	UE_NODISCARD static FORCEINLINE void CartesianToPolar(const T1 X, const T2 Y, T3& OutRad, T3& OutAng) { float ORad, OAng; CartesianToPolar((float)X, (float)Y, ORad, OAng); OutRad = ORad; OutAng = OAng; }
+	static FORCEINLINE void CartesianToPolar(const T1 X, const T2 Y, T3& OutRad, T3& OutAng) { float ORad, OAng; CartesianToPolar((float)X, (float)Y, ORad, OAng); OutRad = ORad; OutAng = OAng; }
 
 	template<typename T1, typename T2, typename T3, TEMPLATE_REQUIRES(TIsAmbiguous<T1, T2, T3>)> UE_DEPRECATED(5.0, "Arguments cause function resolution ambiguity.")
 	UE_NODISCARD static FORCEINLINE float SmoothStep(T1&& A, T2&& B, T3&& C) { return SmoothStep((float)A, (float)B, (float)C); }
