Made DecompressToComposite pure virtual as there was no specific definition for the FDecoder class.

diff --git a/Engine/Source/Runtime/Core/Private/Compression/CompressedBuffer.cpp b/Engine/Source/Runtime/Core/Private/Compression/CompressedBuffer.cpp
--- a/Engine/Source/Runtime/Core/Private/Compression/CompressedBuffer.cpp
+++ b/Engine/Source/Runtime/Core/Private/Compression/CompressedBuffer.cpp
@@ -168,7 +168,7 @@ public:
 	virtual uint64 GetHeaderSize(const FHeader& Header) const = 0;
 	virtual void DecodeHeader(FMutableMemoryView HeaderView) const = 0;
 	virtual bool TryDecompressTo(FDecoderContext& Context, const FDecoderSource& Source, const FHeader& Header, uint64 RawOffset, FMutableMemoryView RawView) const = 0;
-	virtual FCompositeBuffer DecompressToComposite(FDecoderContext& Context, const FDecoderSource& Source, const FHeader& Header, uint64 RawOffset, uint64 RawSize) const;
+	virtual FCompositeBuffer DecompressToComposite(FDecoderContext& Context, const FDecoderSource& Source, const FHeader& Header, uint64 RawOffset, uint64 RawSize) const =0;
 };
 
 ///////////////////////////////////////////////////////////////////////////////////////////////////////////////////////
