LLM: Add bootstrap tag declarations using function statics, for name-based tags that are used in scopes before c++ globals construction is completed.
#rnx

diff --git a/Engine/Source/Runtime/Core/Private/HAL/LowLevelMemTracker.cpp b/Engine/Source/Runtime/Core/Private/HAL/LowLevelMemTracker.cpp
--- a/Engine/Source/Runtime/Core/Private/HAL/LowLevelMemTracker.cpp
+++ b/Engine/Source/Runtime/Core/Private/HAL/LowLevelMemTracker.cpp
@@ -476,7 +476,13 @@ namespace LLMPrivate
 }
 }
 
-static FName TagName_CustomName(TEXT("CustomName"));
+static FName GetTagName_CustomName()
+{
+	// This FName can be read before c++ global constructors are complete, by LLM_SCOPE_BY_BOOTSTRAP_TAG
+	// so it needs to be a function static rather than a global.
+	static FName TagName_CustomName(TEXT("CustomName"));
+	return TagName_CustomName;
+}
 static FName TagName_Untagged(TEXT("Untagged"));
 
 FName LLMGetTagUniqueName(ELLMTag Tag)
@@ -1455,7 +1461,7 @@ UE::LLMPrivate::FTagData& FLowLevelMemTracker::RegisterTagData(FName Name, FName
 		LLMCheck(DisplayName.IsNone());
 		LLMCheck(ParentName.IsNone());
 		DisplayName = Name;
-		ParentName = TagName_CustomName;
+		ParentName = GetTagName_CustomName();
 
 		const TCHAR* Start = FCString::Strstr(NameBuffer.ToString(), TEXT("///"));
 		if (Start)
@@ -1500,7 +1506,7 @@ UE::LLMPrivate::FTagData& FLowLevelMemTracker::RegisterTagData(FName Name, FName
 		}
 		else if (ParentName.IsNone())
 		{
-			ParentName = TagName_CustomName;
+			ParentName = GetTagName_CustomName();
 		}
 
 		// Display name is set to the leaf /Name portion of the unique name, and is overridden if DisplayName is provided
@@ -1623,7 +1629,7 @@ void FLowLevelMemTracker::FinishConstruct(UE::LLMPrivate::FTagData* TagData, UE:
 				TagDataLock.ReadUnlock();
 				UE_LOG(LogHAL, Error, TEXT("LLM Parent tag %s was not available when child tag %s was used in %s"), *ParentName.ToString(), *TagData->GetName().ToString(), ToString(ReferenceSource));
 				TagDataLock.ReadLock();
-				ParentDataPtr = TagDataNameMap->Find(TagName_CustomName);
+				ParentDataPtr = TagDataNameMap->Find(GetTagName_CustomName());
 				LLMCheck(ParentDataPtr);
 			}
 			FTagData* ParentData = *ParentDataPtr;
diff --git a/Engine/Source/Runtime/Core/Private/Misc/MemStack.cpp b/Engine/Source/Runtime/Core/Private/Misc/MemStack.cpp
--- a/Engine/Source/Runtime/Core/Private/Misc/MemStack.cpp
+++ b/Engine/Source/Runtime/Core/Private/Misc/MemStack.cpp
@@ -9,7 +9,7 @@
 #include "Misc/CoreDelegates.h"
 #include <atomic>
 
-LLM_DEFINE_TAG(MemStack);
+LLM_DEFINE_BOOTSTRAP_TAG(MemStack);
 
 DECLARE_MEMORY_STAT(TEXT("MemStack Large Block"), STAT_MemStackLargeBLock,STATGROUP_Memory);
 DECLARE_MEMORY_STAT(TEXT("PageAllocator Free"), STAT_PageAllocatorFree, STATGROUP_Memory);
@@ -235,7 +235,7 @@ uint64 FPageAllocator::BytesFree()
 
 void *FPageAllocator::Alloc()
 {
-	LLM_SCOPE_BYTAG(MemStack); 
+	LLM_SCOPE_BY_BOOTSTRAP_TAG(MemStack);
 	void *Result = TheAllocator.Allocate();
 	STAT(UpdateStats());
 	return Result;
@@ -405,4 +405,4 @@ bool FMemStackBase::ContainsPointer(const void* Pointer) const
 	}
 
 	return false;
-}
\ No newline at end of file
+}
diff --git a/Engine/Source/Runtime/Core/Public/HAL/LowLevelMemTracker.h b/Engine/Source/Runtime/Core/Public/HAL/LowLevelMemTracker.h
--- a/Engine/Source/Runtime/Core/Public/HAL/LowLevelMemTracker.h
+++ b/Engine/Source/Runtime/Core/Public/HAL/LowLevelMemTracker.h
@@ -344,6 +344,22 @@ extern FName LLMGetTagStat(ELLMTag Tag);
 #define LLM_DECLARE_TAG(UniqueNameWithUnderscores) extern FLLMTagDeclaration PREPROCESSOR_JOIN(LLMTagDeclaration_, UniqueNameWithUnderscores)
 #define LLM_DECLARE_TAG_API(UniqueNameWithUnderscores, ModuleAPI) extern ModuleAPI FLLMTagDeclaration PREPROCESSOR_JOIN(LLMTagDeclaration_, UniqueNameWithUnderscores)
 
+/**
+ * The BootStrap versions of LLM_DEFINE_TAG, LLM_SCOPE_BYTAG, and LLM_DECLARE_TAG support use in scopes during global
+ * c++ constructors, before Main. These tags are slightly more expensive (even when LLM is disabled) in all
+ * configurations other than shipping, because they use a function static rather than a global variable.
+ */
+#define LLM_DEFINE_BOOTSTRAP_TAG(UniqueNameWithUnderscores, ...) \
+	FLLMTagDeclaration& PREPROCESSOR_JOIN(GetLLMTagDeclaration_, UniqueNameWithUnderscores)() \
+	{ \
+		static FLLMTagDeclaration PREPROCESSOR_JOIN(LLMTagDeclaration_, UniqueNameWithUnderscores)(TEXT(#UniqueNameWithUnderscores), ##__VA_ARGS__); \
+		return PREPROCESSOR_JOIN(LLMTagDeclaration_, UniqueNameWithUnderscores); \
+	}
+#define LLM_SCOPE_BY_BOOTSTRAP_TAG(TagDeclName) FLLMScope SCOPE_NAME(PREPROCESSOR_JOIN(GetLLMTagDeclaration_, TagDeclName)().GetUniqueName(), false /* bIsStatTag */, ELLMTagSet::None, ELLMTracker::Default);\
+												UE_MEMSCOPE(PREPROCESSOR_JOIN(GetLLMTagDeclaration_,TagDeclName)().GetUniqueName());
+#define LLM_DECLARE_BOOTSTRAP_TAG(UniqueNameWithUnderscores) extern FLLMTagDeclaration& PREPROCESSOR_JOIN(GetLLMTagDeclaration_, UniqueNameWithUnderscores)();
+#define LLM_DECLARE_BOOTSTRAP_TAG_API(UniqueNameWithUnderscores, ModuleAPI) extern ModuleAPI FLLMTagDeclaration& PREPROCESSOR_JOIN(GetLLMTagDeclaration_, UniqueNameWithUnderscores)();
+
 typedef void*(*LLMAllocFunction)(size_t);
 typedef void(*LLMFreeFunction)(void*, size_t);
 
@@ -769,5 +785,9 @@ protected:
 	#define LLM_DEFINE_TAG(...)
 	#define LLM_DECLARE_TAG(...)
 	#define LLM_DECLARE_TAG_API(...)
+	#define LLM_DEFINE_BOOTSTRAP_TAG(...)
+	#define LLM_SCOPE_BY_BOOTSTRAP_TAG(...)
+	#define LLM_DECLARE_BOOTSTRAP_TAG(...) 
+	#define LLM_DECLARE_BOOTSTRAP_TAG_API(...)
 
 #endif		// #if ENABLE_LOW_LEVEL_MEM_TRACKER
