Implement OnTestStart and OnTestEnd events to the AutomationTestFramework
+ add TestSamples plugin to centralize test examples.

diff --git a/Engine/Plugins/Tests/TestSamples/Source/TestSamples/Private/TestSamples.cpp b/Engine/Plugins/Tests/TestSamples/Source/TestSamples/Private/TestSamples.cpp
new file mode 100644
--- /dev/null
+++ b/Engine/Plugins/Tests/TestSamples/Source/TestSamples/Private/TestSamples.cpp
@@ -0,0 +1,42 @@
+// Copyright Epic Games, Inc. All Rights Reserved.
+
+#include "TestSamples.h"
+#include "Modules/ModuleManager.h"
+
+#define LOCTEXT_NAMESPACE "FTestSamplesModule"
+
+DEFINE_LOG_CATEGORY_STATIC(LogTestSamples, Log, All);
+
+void FTestSamplesModule::StartupModule()
+{
+	// Hook callbacks to Test Start and End event
+	FAutomationTestFramework::Get().OnTestStartEvent.AddRaw(this, &FTestSamplesModule::OnTestStart);
+	FAutomationTestFramework::Get().OnTestEndEvent.AddRaw(this, &FTestSamplesModule::OnTestEnd);
+}
+
+void FTestSamplesModule::ShutdownModule()
+{
+	// Remove all registered callbacks on shutdown
+	FAutomationTestFramework::Get().OnTestStartEvent.RemoveAll(this);
+	FAutomationTestFramework::Get().OnTestEndEvent.RemoveAll(this);
+}
+
+void FTestSamplesModule::OnTestStart(FAutomationTestBase* Test)
+{
+	if (Test != nullptr)
+	{
+		UE_LOG(LogTestSamples, Verbose, TEXT("Starting test: %s"), *Test->GetTestFullName())
+	}
+}
+void FTestSamplesModule::OnTestEnd(FAutomationTestBase* Test)
+{
+	if (Test != nullptr)
+	{
+		UE_LOG(LogTestSamples, Verbose, TEXT("Ending test: %s"), *Test->GetTestFullName())
+	}
+
+}
+
+#undef LOCTEXT_NAMESPACE
+	
+IMPLEMENT_MODULE(FTestSamplesModule, TestSamples)
diff --git a/Engine/Plugins/Tests/TestSamples/Source/TestSamples/Public/TestSamples.h b/Engine/Plugins/Tests/TestSamples/Source/TestSamples/Public/TestSamples.h
new file mode 100644
--- /dev/null
+++ b/Engine/Plugins/Tests/TestSamples/Source/TestSamples/Public/TestSamples.h
@@ -0,0 +1,21 @@
+// Copyright Epic Games, Inc. All Rights Reserved.
+
+#pragma once
+
+#include "CoreMinimal.h"
+#include "Modules/ModuleInterface.h"
+#include "Misc/AutomationTest.h"
+
+class FTestSamplesModule : public IModuleInterface
+{
+public:
+
+	/** IModuleInterface implementation */
+	virtual void StartupModule() override;
+	virtual void ShutdownModule() override;
+
+private:
+	void OnTestStart(FAutomationTestBase* Test);
+	void OnTestEnd(FAutomationTestBase* Test);
+
+};
diff --git a/Engine/Plugins/Tests/TestSamples/Source/TestSamples/TestSamples.Build.cs b/Engine/Plugins/Tests/TestSamples/Source/TestSamples/TestSamples.Build.cs
new file mode 100644
--- /dev/null
+++ b/Engine/Plugins/Tests/TestSamples/Source/TestSamples/TestSamples.Build.cs
@@ -0,0 +1,50 @@
+// Copyright Epic Games, Inc. All Rights Reserved.
+
+using UnrealBuildTool;
+
+public class TestSamples : ModuleRules
+{
+	public TestSamples(ReadOnlyTargetRules Target) : base(Target)
+	{
+		
+		PublicIncludePaths.AddRange(
+			new string[] {
+				// ... add public include paths required here ...
+			}
+			);
+				
+		
+		PrivateIncludePaths.AddRange(
+			new string[] {
+				// ... add other private include paths required here ...
+			}
+			);
+			
+		
+		PublicDependencyModuleNames.AddRange(
+			new string[]
+			{
+				"Core",
+				"CoreUObject",
+				// ... add other public dependencies that you statically link with here ...
+			}
+			);
+			
+		
+		PrivateDependencyModuleNames.AddRange(
+			new string[]
+			{
+				"Engine",
+				// ... add private dependencies that you statically link with here ...	
+			}
+			);
+		
+		
+		DynamicallyLoadedModuleNames.AddRange(
+			new string[]
+			{
+				// ... add any modules that your module loads dynamically here ...
+			}
+			);
+	}
+}
diff --git a/Engine/Plugins/Tests/TestSamples/TestSamples.uplugin b/Engine/Plugins/Tests/TestSamples/TestSamples.uplugin
new file mode 100644
--- /dev/null
+++ b/Engine/Plugins/Tests/TestSamples/TestSamples.uplugin
@@ -0,0 +1,24 @@
+{
+	"FileVersion": 3,
+	"Version": 1,
+	"VersionName": "1.0",
+	"FriendlyName": "TestSamples",
+	"Description": "",
+	"Category": "Testing",
+	"CreatedBy": "Epic Games, Inc.",
+	"CreatedByURL": "http://epicgames.com",
+	"DocsURL": "",
+	"MarketplaceURL": "",
+	"SupportURL": "",
+	"EnabledByDefault": false,
+	"CanContainContent": false,
+	"IsBetaVersion": false,
+	"Installed": false,
+	"Modules": [
+		{
+			"Name": "TestSamples",
+			"Type": "Editor",
+			"LoadingPhase": "Default"
+		}
+	]
+}
\ No newline at end of file
diff --git a/Engine/Source/Runtime/Core/Private/Misc/AutomationTest.cpp b/Engine/Source/Runtime/Core/Private/Misc/AutomationTest.cpp
--- a/Engine/Source/Runtime/Core/Private/Misc/AutomationTest.cpp
+++ b/Engine/Source/Runtime/Core/Private/Misc/AutomationTest.cpp
@@ -743,7 +743,7 @@ void FAutomationTestFramework::InternalStartTest( const FString& InTestToRun )
 
 		StartTime = FPlatformTime::Seconds();
 
-		//if non-
+		// If not a smoke test, log the test has started.
 		uint32 NonSmokeTestFlags = (EAutomationTestFlags::FilterMask & (~EAutomationTestFlags::SmokeFilter));
 		if (RequestedTestFilter & NonSmokeTestFlags)
 		{
@@ -752,6 +752,8 @@ void FAutomationTestFramework::InternalStartTest( const FString& InTestToRun )
 
 		CurrentTest->SetTestContext(Parameters);
 
+		OnTestStartEvent.Broadcast(CurrentTest);
+
 		{
 			TRACE_CPUPROFILER_EVENT_SCOPE_TEXT(*FString::Printf(TEXT("AutomationTest %s"), *CurrentTest->GetBeautifiedTestName()));
 			// Run the test!
@@ -765,14 +767,6 @@ bool FAutomationTestFramework::InternalStopTest(FAutomationTestExecutionInfo& Ou
 	check(GIsAutomationTesting);
 	check(LatentCommands.IsEmpty());
 
-	double EndTime = FPlatformTime::Seconds();
-	double TimeForTest = static_cast<float>(EndTime - StartTime);
-	uint32 NonSmokeTestFlags = (EAutomationTestFlags::FilterMask & (~EAutomationTestFlags::SmokeFilter));
-	if (RequestedTestFilter & NonSmokeTestFlags)
-	{
-		UE_LOG(LogAutomationTest, Log, TEXT("%s %s ran in %f"), *CurrentTest->GetBeautifiedTestName(), *Parameters, TimeForTest);
-	}
-
 	// Determine if the test was successful based on three criteria:
 	// 1) Did the test itself report success?
 	// 2) Did any errors occur and were logged by the feedback context during execution?++----
@@ -782,7 +776,17 @@ bool FAutomationTestFramework::InternalStopTest(FAutomationTestExecutionInfo& Ou
 	CurrentTest->ExpectedErrors.Empty();
 
 	// Set the success state of the test based on the above criteria
-	CurrentTest->SetSuccessState( bTestSuccessful );
+	CurrentTest->SetSuccessState(bTestSuccessful);
+
+	OnTestEndEvent.Broadcast(CurrentTest);
+
+	double EndTime = FPlatformTime::Seconds();
+	double TimeForTest = static_cast<float>(EndTime - StartTime);
+	uint32 NonSmokeTestFlags = (EAutomationTestFlags::FilterMask & (~EAutomationTestFlags::SmokeFilter));
+	if (RequestedTestFilter & NonSmokeTestFlags)
+	{
+		UE_LOG(LogAutomationTest, Log, TEXT("%s %s ran in %f"), *CurrentTest->GetBeautifiedTestName(), *Parameters, TimeForTest);
+	}
 
 	// Fill out the provided execution info with the info from the test
 	CurrentTest->GetExecutionInfo( OutExecutionInfo );
@@ -1134,6 +1138,11 @@ void FAutomationTestBase::SetSuccessState( bool bSuccessful )
 	ExecutionInfo.bSuccessful = bSuccessful;
 }
 
+bool FAutomationTestBase::GetSuccessState()
+{
+	return ExecutionInfo.bSuccessful;
+}
+
 void FAutomationTestBase::GetExecutionInfo( FAutomationTestExecutionInfo& OutInfo ) const
 {
 	OutInfo = ExecutionInfo;
diff --git a/Engine/Source/Runtime/Core/Public/Misc/AutomationTest.h b/Engine/Source/Runtime/Core/Public/Misc/AutomationTest.h
--- a/Engine/Source/Runtime/Core/Public/Misc/AutomationTest.h
+++ b/Engine/Source/Runtime/Core/Public/Misc/AutomationTest.h
@@ -753,16 +753,24 @@ DECLARE_MULTICAST_DELEGATE_TwoParams(FOnTestDataRetrieved, bool /*bWasNew*/, con
 
 DECLARE_MULTICAST_DELEGATE_TwoParams(FOnPerformanceDataRetrieved, bool /*bSuccess*/, const FString& /*ErrorMessage*/);
 
+DECLARE_MULTICAST_DELEGATE_OneParam(FOnTestEvent, FAutomationTestBase*);
+
 /** Class representing the main framework for running automation tests */
 class CORE_API FAutomationTestFramework
 {
 public:
-	/** Called right before unit testing is about to begin */
+	/** Called right before automated test is about to begin */
 	FSimpleMulticastDelegate PreTestingEvent;
-	
-	/** Called after all unit tests have completed */
+
+	/** Called after all automated tests have completed */
 	FSimpleMulticastDelegate PostTestingEvent;
 
+	/** Called when each automated test is starting */
+	FOnTestEvent OnTestStartEvent;
+
+	/** Called when each automated test is ending */
+	FOnTestEvent OnTestEndEvent;
+
 	/** Called when a screenshot comparison completes. */
 	FOnTestScreenshotComparisonComplete OnScreenshotCompared;
 
@@ -1309,6 +1317,11 @@ public:
 	 */
 	void SetSuccessState( bool bSuccessful );
 
+	/**
+	 * Return the test success state
+	 */
+	bool GetSuccessState();
+
 	/**
 	 * Populate the provided execution info object with the execution info contained within the test. Not particularly efficient,
 	 * but providing direct access to the test's private execution info could result in errors.
