Skip an ensure if RenderDoc has been detected.

Static initialisation will happen on a thread other than the game thread due to the way RenderDoc injects its DLL (details in the change's comment). Thus the ensure that validates that the game thread and static-init thread are one and the same always fails. Because te GameThread task tag scope is so early (before engine pre-init), and because there is a lot of machinery being a failed ensure (e.g. crash reporter client, dumping of memory stats), this failing ensure is very fragile and can often go undiscovered. So to put it to rest, avoid the ensure altogether if RenderDoc's around.

Note that the segfaulting ensure() is already fixed in 18280863.

diff --git a/Engine/Source/Runtime/Core/Private/HAL/ThreadingBase.cpp b/Engine/Source/Runtime/Core/Private/HAL/ThreadingBase.cpp
--- a/Engine/Source/Runtime/Core/Private/HAL/ThreadingBase.cpp
+++ b/Engine/Source/Runtime/Core/Private/HAL/ThreadingBase.cpp
@@ -15,6 +15,10 @@
 #include "Async/Fundamental/Scheduler.h"
 #include "Tasks/Pipe.h"
 
+#if PLATFORM_WINDOWS
+#include "Windows/AllowWindowsPlatformTypes.h"
+#endif
+
 #include <atomic>
 
 #ifndef IS_RUNNING_GAMETHREAD_ON_EXTERNAL_THREAD
@@ -75,7 +79,20 @@ FTaskTagScope::FTaskTagScope(bool InTagOnlyIfNone, ETaskTag InTag) : Tag(InTag),
 	{
 		checkf(Tag == ETaskTag::EGameThread, TEXT("The Gamethread can only be tagged on the inital thread of the application"));
 #if !IS_RUNNING_GAMETHREAD_ON_EXTERNAL_THREAD
-		ensureMsgf(IsRunningDuringStaticInit(), TEXT("Static initialization should have happened on the same thread as the main thread"));
+#	if PLATFORM_WINDOWS && DO_CHECK
+		// When RenderDoc injects its DLL it first creates this process in a suspended
+		// state where PE loading is not complete and no static DLL dependencies have
+		// been loaded (i.e. no static-init executed). RenderDoc then remotely creates
+		// a thread to load renderdoc.dll and it is this thread that then completes
+		// the PE loading process prior to calling the thread's entry point. This
+		// results static initialisation unexpectedly happening on RenderDoc's injection
+		// thread and the ensure below fails.
+		static bool bRenderDocDetected = (::LoadLibraryW(L"renderdoc.dll") != nullptr);
+		if (!bRenderDocDetected)
+#	endif
+		{
+			ensureMsgf(IsRunningDuringStaticInit(), TEXT("Static initialization should have happened on the same thread as the main thread"));
+		}
 #endif
 	}
 
@@ -1545,3 +1562,7 @@ FRunnableThread* FForkProcessHelper::CreateForkableThread(class FRunnable* InRun
 
 	return NewThread;
 }
+
+#if PLATFORM_WINDOWS
+#include "Windows/HideWindowsPlatformTypes.h"
+#endif
