[CSV] - adding SystemMaxMB to CSV reports to represent console retail memory limit on the produced graphs.

diff --git a/Engine/Source/Runtime/Core/Private/ProfilingDebugging/CsvProfiler.cpp b/Engine/Source/Runtime/Core/Private/ProfilingDebugging/CsvProfiler.cpp
--- a/Engine/Source/Runtime/Core/Private/ProfilingDebugging/CsvProfiler.cpp
+++ b/Engine/Source/Runtime/Core/Private/ProfilingDebugging/CsvProfiler.cpp
@@ -2853,18 +2853,24 @@ void FCsvProfiler::EndFrame()
 		FPlatformMemoryStats MemoryStats = FPlatformMemory::GetStats();
 		float PhysicalMBFree = float(MemoryStats.AvailablePhysical) / (1024.0f * 1024.0f);
 		float UsedExtendedMB = 0;
+		float PhysicalMBUsed = float(MemoryStats.UsedPhysical) / (1024.0f * 1024.0f);
+		float VirtualMBUsed = float(MemoryStats.UsedVirtual) / (1024.0f * 1024.0f);
+
+		// infer the max we can allocate
+		float TotalSystemMB = PhysicalMBFree + PhysicalMBUsed;
 #if !UE_BUILD_SHIPPING
 		// Subtract any extra development memory from physical free. This can result in negative values in cases where we would have crashed OOM
 		PhysicalMBFree -= float(FPlatformMemory::GetExtraDevelopmentMemorySize() / 1024ull / 1024ull);
 		UsedExtendedMB = PhysicalMBFree < 0.0f ? -PhysicalMBFree : 0;
+
+		TotalSystemMB -= float(FPlatformMemory::GetExtraDevelopmentMemorySize() / 1024ull / 1024ull);
 #endif
-		float PhysicalMBUsed = float(MemoryStats.UsedPhysical) / (1024.0f * 1024.0f);
-		float VirtualMBUsed = float(MemoryStats.UsedVirtual) / (1024.0f * 1024.0f);
 		
 		CSV_CUSTOM_STAT_GLOBAL(MemoryFreeMB, PhysicalMBFree, ECsvCustomStatOp::Set);
 		CSV_CUSTOM_STAT_GLOBAL(PhysicalUsedMB, PhysicalMBUsed, ECsvCustomStatOp::Set);
 		CSV_CUSTOM_STAT_GLOBAL(VirtualUsedMB, VirtualMBUsed, ECsvCustomStatOp::Set);
 		CSV_CUSTOM_STAT_GLOBAL(ExtendedUsedMB, UsedExtendedMB, ECsvCustomStatOp::Set);
+		CSV_CUSTOM_STAT_GLOBAL(SystemMaxMB, TotalSystemMB, ECsvCustomStatOp::Set);
 
 		// If we're single-threaded, process the stat data here
 		if (ProcessingThread == nullptr)
