UE5:  MEMORY:  Adjustment to container allocation policy to save memory on arrays with optional inline elements.  Make it so inline elements are ignored when computing the desired slack growth size for the first heap allocation.  The slack algorithm has a special case to save memory on the initial heap allocation, versus subsequent reallocations, and we don't want the inline elements to be treated as if they were the first heap allocation.  (also reviewed by [REDACTED_PERSON])

#rnx

diff --git a/Engine/Source/Runtime/Core/Public/Containers/ContainerAllocationPolicies.h b/Engine/Source/Runtime/Core/Public/Containers/ContainerAllocationPolicies.h
--- a/Engine/Source/Runtime/Core/Public/Containers/ContainerAllocationPolicies.h
+++ b/Engine/Source/Runtime/Core/Public/Containers/ContainerAllocationPolicies.h
@@ -733,9 +733,12 @@ public:
 		FORCEINLINE SizeType CalculateSlackGrow(SizeType NumElements, SizeType NumAllocatedElements, SIZE_T NumBytesPerElement) const
 		{
 			// If the elements use less space than the inline allocation, only use the inline allocation as slack.
+			// Also, when computing slack growth, don't count inline elements -- the slack algorithm has a special
+			// case to save memory on the initial heap allocation, versus subsequent reallocations, and we don't
+			// want the inline elements to be treated as if they were the first heap allocation.
 			return NumElements <= NumInlineElements ?
 				NumInlineElements :
-				SecondaryData.CalculateSlackGrow(NumElements, NumAllocatedElements, NumBytesPerElement);
+				SecondaryData.CalculateSlackGrow(NumElements, NumAllocatedElements <= NumInlineElements ? 0 : NumAllocatedElements, NumBytesPerElement);
 		}
 
 		SIZE_T GetAllocatedSize(SizeType NumAllocatedElements, SIZE_T NumBytesPerElement) const
