Tweaked MallocBinned2 bin sizes for Frosty to save 50MB on alignment waste in SmallBlockAllocator

diff --git a/Engine/Source/Runtime/Core/Private/HAL/MallocBinned2.cpp b/Engine/Source/Runtime/Core/Private/HAL/MallocBinned2.cpp
--- a/Engine/Source/Runtime/Core/Private/HAL/MallocBinned2.cpp
+++ b/Engine/Source/Runtime/Core/Private/HAL/MallocBinned2.cpp
@@ -96,12 +96,12 @@ int32 RecursionCounter = 0;
 // They must be 16-byte aligned as well.
 static uint16 SmallBlockSizes[] =
 {
-	16, 32, 48, 64, 80, 96, 112, 128,
+	16, 32, 48, 64, 80, 96, 128,
 	160, 192, 224, 256, 288, 320, 384, 448,
 	512, 576, 640, 704, 768, 896, 1024 - 16, 1168,
-	1360, 1632, 2048 - 16, 2336, 2720, 3264, 4096 - 16, 4368,
-	4672, 5040, 5456, 5952, 6544 - 16, 7280, 8192 - 16, 9360,
-	10912, 13104, 16384 - 16, 21840, 32768 - 16
+	1488, 1632, 2048 - 16, 2336, 2720, 3264, 4096 - 16, 4368,
+	5040, 5456, 5952, 6544 - 16, 7280, 8192 - 16, 9360,
+	10912, 13104, 16384 - 16, 19104, 21840, 27024, 32768 - 16
 };
 
 MS_ALIGN(PLATFORM_CACHE_LINE_SIZE) static uint8 UnusedAlignPadding[PLATFORM_CACHE_LINE_SIZE] GCC_ALIGN(PLATFORM_CACHE_LINE_SIZE) = { 0 };
