Reduce memory allocation done by output devices, since this is enabled for shipping builds.

diff --git a/Engine/Source/Runtime/Core/Private/Misc/DateTime.cpp b/Engine/Source/Runtime/Core/Private/Misc/DateTime.cpp
--- a/Engine/Source/Runtime/Core/Private/Misc/DateTime.cpp
+++ b/Engine/Source/Runtime/Core/Private/Misc/DateTime.cpp
@@ -234,7 +234,12 @@ FString FDateTime::ToString() const
 FString FDateTime::ToString(const TCHAR* Format) const
 {
 	TStringBuilder<32> Result;
+	ToString(Format, Result);
+	return Result.ToString();
+}
 
+void FDateTime::ToString(const TCHAR* Format, FStringBuilderBase& Result) const
+{
 	if (Format != nullptr)
 	{
 		while (*Format != TCHAR('\0'))
@@ -267,8 +272,6 @@ FString FDateTime::ToString(const TCHAR* Format) const
 			Format++;
 		}
 	}
-
-	return Result.ToString();
 }
 
 FString FDateTime::ToFormattedString(const TCHAR* Format) const
diff --git a/Engine/Source/Runtime/Core/Private/Misc/OutputDeviceHelper.cpp b/Engine/Source/Runtime/Core/Private/Misc/OutputDeviceHelper.cpp
--- a/Engine/Source/Runtime/Core/Private/Misc/OutputDeviceHelper.cpp
+++ b/Engine/Source/Runtime/Core/Private/Misc/OutputDeviceHelper.cpp
@@ -10,27 +10,31 @@
 FString FOutputDeviceHelper::FormatLogLine( ELogVerbosity::Type Verbosity, const class FName& Category, const TCHAR* Message /*= nullptr*/, ELogTimes::Type LogTime /*= ELogTimes::None*/, const double Time /*= -1.0*/, int32* OutCategoryIndex)
 {
 	const bool bShowCategory = GPrintLogCategory && Category != NAME_None;
-	FString Format;
+	TStringBuilder<512> Format;
 
 	switch (LogTime)
 	{
 		case ELogTimes::SinceGStartTime:
 		{																	
 			const double RealTime = (Time < 0.0) ? (FPlatformTime::Seconds() - GStartTime) : Time;
-			Format = FString::Printf( TEXT( "[%07.2f][%3llu]" ), RealTime, GFrameCounter % 1000);
+			Format.Append(*FString::Printf( TEXT( "[%07.2f][%3llu]" ), RealTime, GFrameCounter % 1000));
 			break;
 		}
 
 		case ELogTimes::UTC:
-			Format = FString::Printf(TEXT("[%s][%3llu]"), *FDateTime::UtcNow().ToString(TEXT("%Y.%m.%d-%H.%M.%S:%s")), GFrameCounter % 1000);
+			Format.Append(TEXT("["));
+			FDateTime::UtcNow().ToString(TEXT("%Y.%m.%d-%H.%M.%S:%s"), Format);
+			Format.Append(TEXT("]["));
+			Format.Append(*FString::Printf(TEXT("%3llu"), GFrameCounter % 1000));
+			Format.Append(TEXT("]"));
 			break;
 
 		case ELogTimes::Local:
-			Format = FString::Printf(TEXT("[%s][%3llu]"), *FDateTime::Now().ToString(TEXT("%Y.%m.%d-%H.%M.%S:%s")), GFrameCounter % 1000);
+			Format.Append(*FString::Printf(TEXT("[%s][%3llu]"), *FDateTime::Now().ToString(TEXT("%Y.%m.%d-%H.%M.%S:%s")), GFrameCounter % 1000));
 			break;
 
 		case ELogTimes::Timecode:
-			Format = FString::Printf(TEXT("[%s][%3llu]"), *FApp::GetTimecode().ToString(), GFrameCounter % 1000);
+			Format.Append(*FString::Printf(TEXT("[%s][%3llu]"), *FApp::GetTimecode().ToString(), GFrameCounter % 1000));
 			break;
 
 		default:
@@ -44,28 +48,28 @@ FString FOutputDeviceHelper::FormatLogLine( ELogVerbosity::Type Verbosity, const
 
 	if (bShowCategory)
 	{
-		Format += Category.ToString();
-		Format += TEXT(": ");
+		Category.AppendString(Format);
+		Format.Append(TEXT(": "));
 
 		if (GPrintLogVerbosity && Verbosity != ELogVerbosity::Log)
 		{
-			Format += ToString(Verbosity);
-			Format += TEXT(": ");
+			Format.Append(ToString(Verbosity));
+			Format.Append(TEXT(": "));
 		}
 	}
 	else if (GPrintLogVerbosity && Verbosity != ELogVerbosity::Log)
 	{
 #if !HACK_HEADER_GENERATOR
-		Format += ToString(Verbosity);
-		Format += TEXT(": ");
+		Format.Append(ToString(Verbosity));
+		Format.Append(TEXT(": "));
 #endif
 	}
 
 	if (Message)
 	{
-		Format += Message;
+		Format.Append(Message);
 	}
-	return Format;
+	return Format.ToString();
 }
 
 void FOutputDeviceHelper::FormatCastAndSerializeLine(FArchive& Output, const TCHAR* Data, ELogVerbosity::Type Verbosity, const class FName& Category, const double Time, bool bSuppressEventTag, bool bAutoEmitLineTerminator)
diff --git a/Engine/Source/Runtime/Core/Public/Misc/DateTime.h b/Engine/Source/Runtime/Core/Public/Misc/DateTime.h
--- a/Engine/Source/Runtime/Core/Public/Misc/DateTime.h
+++ b/Engine/Source/Runtime/Core/Public/Misc/DateTime.h
@@ -530,6 +530,7 @@ public:
 	 * @see Parse, ToIso8601
 	 */
 	CORE_API FString ToString(const TCHAR* Format) const;
+	CORE_API void ToString(const TCHAR* Format, FStringBuilderBase& Result) const;
 
 	/**
 	 * Returns the string representation of this date.
