[PlatMisc] - add RestartWithCmdLine to PlatformMisc, not implemented on all plats

diff --git a/Engine/Source/Runtime/Core/Private/GenericPlatform/GenericPlatformMisc.cpp b/Engine/Source/Runtime/Core/Private/GenericPlatform/GenericPlatformMisc.cpp
--- a/Engine/Source/Runtime/Core/Private/GenericPlatform/GenericPlatformMisc.cpp
+++ b/Engine/Source/Runtime/Core/Private/GenericPlatform/GenericPlatformMisc.cpp
@@ -911,6 +911,12 @@ void FGenericPlatformMisc::RequestExit( bool Force )
 	}
 }
 
+bool FGenericPlatformMisc::RestartApplicationWithCmdLine(const char* CmdLine)
+{
+	UE_LOG(LogInit, Display, TEXT("Restart application (with cmdnline) is not supported or implemented in current platform"));
+	return false;
+}
+
 bool FGenericPlatformMisc::RestartApplication()
 {
 	UE_LOG(LogInit, Display, TEXT("Restart application is not supported or implemented in current platform"));
diff --git a/Engine/Source/Runtime/Core/Public/GenericPlatform/GenericPlatformMisc.h b/Engine/Source/Runtime/Core/Public/GenericPlatform/GenericPlatformMisc.h
--- a/Engine/Source/Runtime/Core/Public/GenericPlatform/GenericPlatformMisc.h
+++ b/Engine/Source/Runtime/Core/Public/GenericPlatform/GenericPlatformMisc.h
@@ -922,6 +922,11 @@ public:
 	 */
 	static bool RestartApplication();
 
+    /**
+	 * Requests application to restart appending the passed in CmdLine to the restart
+	 */
+	static bool RestartApplicationWithCmdLine(const char* CmdLine);
+
 	/**
 	 * Returns the last system error code in string form.  NOTE: Only one return value is valid at a time!
 	 *
