UE5:  MEMORY:  Defer allocation of streaming buffer for Pakfile reader until it's actually used.  Saves 1 MB of memory in cases where it's never used (8 MB total, given per-thread cached Pakfile readers).  The streaming buffer is only used when sequential reads occur without seeking, or when the user explicitly requests a pre-cache.

#rnx

diff --git a/Engine/Source/Runtime/Core/Private/HAL/FileManagerGeneric.cpp b/Engine/Source/Runtime/Core/Private/HAL/FileManagerGeneric.cpp
--- a/Engine/Source/Runtime/Core/Private/HAL/FileManagerGeneric.cpp
+++ b/Engine/Source/Runtime/Core/Private/HAL/FileManagerGeneric.cpp
@@ -663,7 +663,6 @@ FArchiveFileReaderGeneric::FArchiveFileReaderGeneric( IFileHandle* InHandle, con
 	, bFirstReadAfterSeek(false)
 {
 	BufferSize = FMath::Min(FMath::RoundUpToPowerOfTwo64((int64)InBufferSize), (uint64)Size);
-	BufferArray.Reserve(BufferSize);
 	this->SetIsLoading(true);
 	this->SetIsPersistent(true);
 }
@@ -745,6 +744,14 @@ bool FArchiveFileReaderGeneric::Precache(int64 PrecacheOffset, int64 PrecacheSiz
 
 bool FArchiveFileReaderGeneric::InternalPrecache( int64 PrecacheOffset, int64 PrecacheSize )
 {
+	// We defer allocation of BufferArray until we actually need it.  If the client doesn't do sequential reads
+	// or explicitly request precaching via "FArchiveFileReaderGeneric::Precache", InternalPrecache might never
+	// be called, and the buffer never actually used, a waste of 1 MB per archive file.
+	if (BufferArray.Max() == 0)
+	{
+		BufferArray.Reserve(BufferSize);
+	}
+
 	// Only precache at current position and avoid work if precaching same offset twice.
 	if (Pos != PrecacheOffset)
 	{
@@ -777,7 +784,7 @@ bool FArchiveFileReaderGeneric::InternalPrecache( int64 PrecacheOffset, int64 Pr
 			// We need to read more to satisfy the precache request.
 			// Since Pos is within the buffer, the low-level read position is at the end of the buffer. 
 			// Copy the existing bytes after Pos out of the old buffer into the new buffer, and then read only the remaining bytes from the low-level handle into the rest of the new buffer.
-			if (BufferCount <= BufferArray.Max())
+			if (BufferCount <= BufferSize)
 			{
 				// We don't need to reallocate the buffer, so we can just move the RemainingBufferCount bytes from the end of the buffer to the beginning.
 				FMemory::Memmove(BufferArray.GetData(), BufferArray.GetData() + Pos - BufferBase, RemainingBufferCount);
