[CsvProfiler] - enable continuous writes

diff --git a/Engine/Source/Runtime/Core/Private/ProfilingDebugging/CsvProfiler.cpp b/Engine/Source/Runtime/Core/Private/ProfilingDebugging/CsvProfiler.cpp
--- a/Engine/Source/Runtime/Core/Private/ProfilingDebugging/CsvProfiler.cpp
+++ b/Engine/Source/Runtime/Core/Private/ProfilingDebugging/CsvProfiler.cpp
@@ -86,7 +86,7 @@ TAutoConsoleVariable<int32> CVarCsvBlockOnCaptureEnd(
 
 TAutoConsoleVariable<int32> CVarCsvContinuousWrites(
 	TEXT("csv.ContinuousWrites"),
-	0,
+	1,
 	TEXT("When 1, completed CSV rows are converted to CSV format strings and appended to the write buffer whilst the capture is in progress.\r\n")
 	TEXT("When 0, CSV rows are accumulated in memory as binary data, and only converted to strings and flushed to disk at the end of the capture."),
 	ECVF_Default
