GC optimization: delay call to FMemory::Trim after unhashing is done instead of doing it right after the reachability pass, and also don't wait on memory trimming task thread to finish.
This has the side-effect of releasing unused OS pages to the next GC (since they are released from the per-thread calls to FlushCurrentThreadCache, but doesn't has a noticeable impact on internal demo.
In practice, this makes FMemory::Trim detached from the GC reachability pass, and takes ~4ms instead of ~10-50ms depending on the other threads workload.
Original code by PJ.Kack.

diff --git a/Engine/Source/Runtime/Core/Private/Async/TaskGraph.cpp b/Engine/Source/Runtime/Core/Private/Async/TaskGraph.cpp
--- a/Engine/Source/Runtime/Core/Private/Async/TaskGraph.cpp
+++ b/Engine/Source/Runtime/Core/Private/Async/TaskGraph.cpp
@@ -2628,11 +2628,14 @@ void FTaskGraphInterface::BroadcastSlow_OnlyUseForSpecialPurposes(bool bDoTaskTh
 	double StartTime = FPlatformTime::Seconds();
 
 	QUICK_SCOPE_CYCLE_COUNTER(STAT_FTaskGraphInterface_BroadcastSlow_OnlyUseForSpecialPurposes);
+	TRACE_CPUPROFILER_EVENT_SCOPE(FTaskGraphInterface_BroadcastSlow);
 	check(FPlatformTLS::GetCurrentThreadId() == GGameThreadId);
+
+	Callback(ENamedThreads::GameThread_Local);
+
 	if (!TaskGraphImplementationSingleton)
 	{
 		// we aren't going yet
-		Callback(ENamedThreads::GameThread);
 		return;
 	}
 
@@ -2695,33 +2698,6 @@ void FTaskGraphInterface::BroadcastSlow_OnlyUseForSpecialPurposes(bool bDoTaskTh
 	}
 
 
-	FGraphEventArray Tasks;
-#if STATS
-PRAGMA_DISABLE_DEPRECATION_WARNINGS
-	if (FTaskGraphInterface::Get().IsThreadProcessingTasks(ENamedThreads::StatsThread))
-	{
-		Tasks.Add(TGraphTask<FBroadcastTask>::CreateTask().ConstructAndDispatchWhenReady(Callback, StartTime, TEXT("Stats"), ENamedThreads::SetTaskPriority(ENamedThreads::StatsThread, ENamedThreads::HighTaskPriority), nullptr, nullptr, nullptr));
-	}
-PRAGMA_ENABLE_DEPRECATION_WARNINGS
-#endif
-	if (IsRHIThreadRunning())
-	{
-		Tasks.Add(TGraphTask<FBroadcastTask>::CreateTask().ConstructAndDispatchWhenReady(Callback, StartTime, TEXT("RHIT"), ENamedThreads::SetTaskPriority(ENamedThreads::RHIThread, ENamedThreads::HighTaskPriority), nullptr, nullptr, nullptr));
-	}
-	ENamedThreads::Type RenderThread = ENamedThreads::GetRenderThread();
-	if (RenderThread != ENamedThreads::GameThread)
-	{
-		Tasks.Add(TGraphTask<FBroadcastTask>::CreateTask().ConstructAndDispatchWhenReady(Callback, StartTime, TEXT("RT"), ENamedThreads::SetTaskPriority(RenderThread, ENamedThreads::HighTaskPriority), nullptr, nullptr, nullptr));
-	}
-PRAGMA_DISABLE_DEPRECATION_WARNINGS
-	if (FTaskGraphInterface::Get().IsThreadProcessingTasks(ENamedThreads::AudioThread))
-	{
-		Tasks.Add(TGraphTask<FBroadcastTask>::CreateTask().ConstructAndDispatchWhenReady(Callback, StartTime, TEXT("AudioT"), ENamedThreads::SetTaskPriority(ENamedThreads::AudioThread, ENamedThreads::HighTaskPriority), nullptr, nullptr, nullptr));
-	}
-PRAGMA_ENABLE_DEPRECATION_WARNINGS
-
-	Callback(ENamedThreads::GameThread_Local);
-
 	if (bDoTaskThreads)
 	{
 		check(MyEvent);
@@ -2736,6 +2712,7 @@ PRAGMA_ENABLE_DEPRECATION_WARNINGS
 		{
 			const double StartTimeInner = FPlatformTime::Seconds();
 			QUICK_SCOPE_CYCLE_COUNTER(STAT_Broadcast_WaitForTaskThreads);
+			TRACE_CPUPROFILER_EVENT_SCOPE(Broadcast_WaitForTaskThreads);
 			FTaskGraphInterface::Get().WaitUntilTasksComplete(TaskThreadTasks, ENamedThreads::GameThread_Local);
 			{
 				const double ThisTime = FPlatformTime::Seconds() - StartTimeInner;
@@ -2746,36 +2723,31 @@ PRAGMA_ENABLE_DEPRECATION_WARNINGS
 			}
 		}
 	}
-	{
-		double StartTimeInner = FPlatformTime::Seconds();
-		QUICK_SCOPE_CYCLE_COUNTER(STAT_Broadcast_WaitForNamedThreads);
-		
-		// Wait for all tasks to be complete.  Spin and pump messages to avoid deadlocks when other threads send messages and block until messages are processed
-		while (true)
-		{
-			bool bAnyNotDone = false;
-			for (FGraphEventRef& Item : Tasks)
-			{
-				if (Item.GetReference() && !Item->IsComplete())
-				{
-					bAnyNotDone = true;
-					break;
-				}
-			}
-			if (!bAnyNotDone)
-			{
-				break;
-			}
 
-			FPlatformMisc::PumpMessagesOutsideMainLoop();
-		}
-		
-		const double EndTimeInner = FPlatformTime::Seconds() - StartTimeInner;
-		if (EndTimeInner > 0.02)
-		{
-			UE_CLOG(GPrintBroadcastWarnings, LogTaskGraph, Warning, TEXT("Task graph took %6.2fms to wait for named thread broadcast."), EndTimeInner * 1000.0);
-		}
+#if STATS
+PRAGMA_DISABLE_DEPRECATION_WARNINGS
+	if (FTaskGraphInterface::Get().IsThreadProcessingTasks(ENamedThreads::StatsThread))
+	{
+		TGraphTask<FBroadcastTask>::CreateTask().ConstructAndDispatchWhenReady(Callback, StartTime, TEXT("Stats"), ENamedThreads::SetTaskPriority(ENamedThreads::StatsThread, ENamedThreads::HighTaskPriority), nullptr, nullptr, nullptr);
 	}
+PRAGMA_ENABLE_DEPRECATION_WARNINGS
+#endif
+	if (IsRHIThreadRunning())
+	{
+		TGraphTask<FBroadcastTask>::CreateTask().ConstructAndDispatchWhenReady(Callback, StartTime, TEXT("RHIT"), ENamedThreads::SetTaskPriority(ENamedThreads::RHIThread, ENamedThreads::HighTaskPriority), nullptr, nullptr, nullptr);
+	}
+	ENamedThreads::Type RenderThread = ENamedThreads::GetRenderThread();
+	if (RenderThread != ENamedThreads::GameThread)
+	{
+		TGraphTask<FBroadcastTask>::CreateTask().ConstructAndDispatchWhenReady(Callback, StartTime, TEXT("RT"), ENamedThreads::SetTaskPriority(RenderThread, ENamedThreads::HighTaskPriority), nullptr, nullptr, nullptr);
+	}
+PRAGMA_DISABLE_DEPRECATION_WARNINGS
+	if (FTaskGraphInterface::Get().IsThreadProcessingTasks(ENamedThreads::AudioThread))
+	{
+		TGraphTask<FBroadcastTask>::CreateTask().ConstructAndDispatchWhenReady(Callback, StartTime, TEXT("AudioT"), ENamedThreads::SetTaskPriority(ENamedThreads::AudioThread, ENamedThreads::HighTaskPriority), nullptr, nullptr, nullptr);
+	}
+PRAGMA_ENABLE_DEPRECATION_WARNINGS
+
 	for (FEvent* TaskEvent : TaskEvents)
 	{
 		FPlatformProcess::ReturnSynchEventToPool(TaskEvent);
diff --git a/Engine/Source/Runtime/CoreUObject/Private/UObject/GarbageCollection.cpp b/Engine/Source/Runtime/CoreUObject/Private/UObject/GarbageCollection.cpp
--- a/Engine/Source/Runtime/CoreUObject/Private/UObject/GarbageCollection.cpp
+++ b/Engine/Source/Runtime/CoreUObject/Private/UObject/GarbageCollection.cpp
@@ -1590,7 +1590,7 @@ void IncrementalPurgeGarbage(bool bUseTimeLimit, double TimeLimit)
 		GObjCurrentPurgeObjectIndexNeedsReset = true;
 	}
 	// Early out if there is nothing to do.
-	if (!GObjPurgeIsRequired)
+	if (!GObjPurgeIsRequired && !GObjIncrementalPurgeIsInProgress)
 	{
 		return;
 	}
@@ -1622,6 +1622,13 @@ void IncrementalPurgeGarbage(bool bUseTimeLimit, double TimeLimit)
 
 	} ResetPurgeProgress(bCompleted);
 	
+	// if the purge was completed last tick, perform the trim to finalize this incremental purge
+	if (!GObjPurgeIsRequired)
+	{
+		FMemory::Trim();
+		bCompleted = true;
+	}
+	else
 	{
 		// Lock before settting GCStartTime as it could be slow to lock if async loading is in progress
 		// but we still want to perform some GC work otherwise we'd be keeping objects in memory for a long time
@@ -1645,6 +1652,9 @@ void IncrementalPurgeGarbage(bool bUseTimeLimit, double TimeLimit)
 		{
 			bCompleted = IncrementalDestroyGarbage(bUseTimeLimit, TimeLimit);
 		}
+
+		// when running incrementally using a time limit, add one last tick for the memory trim
+		bCompleted = bCompleted && !bUseTimeLimit;
 	}
 #endif // !UE_WITH_GC
 }
@@ -2178,7 +2188,10 @@ void CollectGarbageInternal(EObjectFlags KeepFlags, bool bPerformFullPurge)
 		if (GObjIncrementalPurgeIsInProgress || GObjPurgeIsRequired)
 		{
 			IncrementalPurgeGarbage(false);
-			FMemory::Trim();
+			if (!bPerformFullPurge)
+			{
+				FMemory::Trim();
+			}
 		}
 		check(!GObjIncrementalPurgeIsInProgress);
 		check(!GObjPurgeIsRequired);
@@ -2291,8 +2304,10 @@ void CollectGarbageInternal(EObjectFlags KeepFlags, bool bPerformFullPurge)
 		// Destroy all pending delete linkers
 		DeleteLoaders();
 
-		// Trim allocator memory
-		FMemory::Trim();
+		if (bPerformFullPurge)
+		{
+			FMemory::Trim();
+		}
 	}
 
 	// Route callbacks to verify GC assumptions
