Compile fix for ParallelForWithTaskContext using a non-default allocator, and also use ArrayView where appropriate to lower number of template arguments.

diff --git a/Engine/Source/Runtime/Core/Public/Async/ParallelFor.h b/Engine/Source/Runtime/Core/Public/Async/ParallelFor.h
--- a/Engine/Source/Runtime/Core/Public/Async/ParallelFor.h
+++ b/Engine/Source/Runtime/Core/Public/Async/ParallelFor.h
@@ -105,24 +105,24 @@ namespace ParallelForImpl
 		}
 
 		template<typename ContextType>
-		bool Process(TArray<ContextType>* InContexts, int32 TaskIndex, int32 TasksToSpawn, TSharedRef<TParallelForData, ESPMode::ThreadSafe>& Data, ENamedThreads::Type InDesiredThread, bool bMaster);
+		bool Process(const TArrayView<ContextType>& Contexts, int32 TaskIndex, int32 TasksToSpawn, TSharedRef<TParallelForData, ESPMode::ThreadSafe>& Data, ENamedThreads::Type InDesiredThread, bool bMaster);
 		
 		inline bool Process(int32 TaskIndex, int32 TasksToSpawn, TSharedRef<TParallelForData, ESPMode::ThreadSafe>& Data, ENamedThreads::Type InDesiredThread, bool bMaster)
 		{
-			return Process<nullptr_t>(nullptr, TaskIndex, TasksToSpawn, Data, InDesiredThread, bMaster);
+			return Process<nullptr_t>(TArrayView<nullptr_t>(), TaskIndex, TasksToSpawn, Data, InDesiredThread, bMaster);
 		}
 	};
 
 	template<typename FunctionType, typename ContextType = nullptr_t>
 	class TParallelForTask
 	{
-		TArray<ContextType>* Contexts;
+		TArrayView<ContextType> Contexts;
 		TSharedRef<TParallelForData<FunctionType>, ESPMode::ThreadSafe> Data;
 		ENamedThreads::Type DesiredThread;
 		int32 TaskIndex;
 		int32 TasksToSpawn;
 	public:
-		TParallelForTask(TArray<ContextType>* InContexts, int32 InTaskIndex, TSharedRef<TParallelForData<FunctionType>, ESPMode::ThreadSafe>& InData, ENamedThreads::Type InDesiredThread, int32 InTasksToSpawn = 0)
+		TParallelForTask(const TArrayView<ContextType>& InContexts, int32 InTaskIndex, TSharedRef<TParallelForData<FunctionType>, ESPMode::ThreadSafe>& InData, ENamedThreads::Type InDesiredThread, int32 InTasksToSpawn = 0)
 			: Contexts(InContexts)
 			, Data(InData) 
 			, DesiredThread(InDesiredThread)
@@ -131,8 +131,7 @@ namespace ParallelForImpl
 		{
 		}		
 		TParallelForTask(int32 InTaskIndex, TSharedRef<TParallelForData<FunctionType>, ESPMode::ThreadSafe>& InData, ENamedThreads::Type InDesiredThread, int32 InTasksToSpawn = 0)
-			: Contexts(nullptr)
-			, Data(InData) 
+			: Data(InData) 
 			, DesiredThread(InDesiredThread)
 			, TaskIndex(InTaskIndex)
 			, TasksToSpawn(InTasksToSpawn)
@@ -166,28 +165,21 @@ namespace ParallelForImpl
 
 	// Helper to call body with context reference
 	template <typename FunctionType, typename ContextType>
-	inline void CallBody(const FunctionType& Body, ContextType* Context, int32 Index)
+	inline void CallBody(const FunctionType& Body, const TArrayView<ContextType>& Contexts, int32 TaskIndex, int32 Index)
 	{
-		Body(*Context, Index);
+		Body(Contexts[TaskIndex], Index);
 	}
-	
+
 	// Helper specialization for "no context", which changes the assumed body call signature
 	template <typename FunctionType>
-	inline void CallBody(const FunctionType& Body, nullptr_t*, int32 Index)
+	inline void CallBody(const FunctionType& Body, const TArrayView<nullptr_t>&, int32, int32 Index)
 	{
 		Body(Index);
 	}
 
-	template <typename FunctionType, typename ContextType>
-	inline void CallBody(const FunctionType& Body, TArray<ContextType>* Contexts, int32 TaskIndex, int32 Index)
-	{
-		ContextType* Context = Contexts ? &(*Contexts)[TaskIndex] : nullptr;
-		CallBody(Body, Context, Index);
-	}
-
 	template<typename FunctionType> 
 	template<typename ContextType> 
-	inline bool TParallelForData<FunctionType>::Process(TArray<ContextType>* Contexts, int32 TaskIndex, int32 TasksToSpawn, TSharedRef<TParallelForData<FunctionType>, ESPMode::ThreadSafe>& Data, ENamedThreads::Type InDesiredThread, bool bMaster)
+	inline bool TParallelForData<FunctionType>::Process(const TArrayView<ContextType>& Contexts, int32 TaskIndex, int32 TasksToSpawn, TSharedRef<TParallelForData<FunctionType>, ESPMode::ThreadSafe>& Data, ENamedThreads::Type InDesiredThread, bool bMaster)
 	{
 		TRACE_CPUPROFILER_EVENT_SCOPE(TParallelForData::Process)
 
@@ -229,7 +221,7 @@ namespace ParallelForImpl
 			}
 			if (MyIndex < LocalNum)
 			{
-				check(Contexts == nullptr || TaskIndex < Contexts->Num());
+				check(Contexts.Num() == 0 || TaskIndex < Contexts.Num());
 
 				int32 ThisBlockSize = LocalBlockSize;
 				if (MyIndex == LocalNum - 1)
@@ -315,7 +307,7 @@ namespace ParallelForImpl
 	}
 
 	template<typename FunctionType, typename ContextType>
-	inline void ParallelForInternal(int32 Num, FunctionType Body, EParallelForFlags Flags, TArray<ContextType>* Contexts)
+	inline void ParallelForInternal(int32 Num, FunctionType Body, EParallelForFlags Flags, const TArrayView<ContextType>& Contexts)
 	{
 		SCOPE_CYCLE_COUNTER(STAT_ParallelFor);
 		check(Num >= 0);
@@ -370,7 +362,7 @@ namespace ParallelForImpl
 	template<typename FunctionType>
 	inline void ParallelForInternal(int32 Num, FunctionType Body, EParallelForFlags Flags)
 	{
-		ParallelForInternal<FunctionType, nullptr_t>(Num, Body, Flags, nullptr);
+		ParallelForInternal<FunctionType, nullptr_t>(Num, Body, Flags, TArrayView<nullptr_t>());
 	}
 	
 	/** 
@@ -535,7 +527,7 @@ inline void ParallelForWithTaskContext(TArray<ContextType, ContextAllocatorType>
 		{
 			new(&OutContexts[ContextIndex]) ContextType(ContextConstructor(ContextIndex, NumContexts));
 		}
-		ParallelForImpl::ParallelForInternal(Num, Body, Flags, &OutContexts);
+		ParallelForImpl::ParallelForInternal(Num, Body, Flags, TArrayView<ContextType>(OutContexts));
 	}
 }
 
@@ -557,6 +549,6 @@ inline void ParallelForWithTaskContext(TArray<ContextType, ContextAllocatorType>
 		const int32 NumContexts = ParallelForImpl::GetNumberOfThreadTasks(Num, Flags) + 1;
 		OutContexts.Reset();
 		OutContexts.AddDefaulted(NumContexts);
-		ParallelForImpl::ParallelForInternal(Num, Body, Flags, &OutContexts);
+		ParallelForImpl::ParallelForInternal(Num, Body, Flags, TArrayView<ContextType>(OutContexts));
 	}
 }
\ No newline at end of file
