CompressedBuffer: Removed partial decompression from FCompressedBuffer now that FCompressedBufferReader is available

Requiring the use of a separate reader type makes it more likely that readers will be reused, and makes it easier to audit reader usage going forward. Reusing readers is desirable to reduce the number of large temporary allocations made during partial decompression of a buffer.

- Added FCompressedBuffer::Save(FArchive&) and renamed FromCompressed(FArchive&) to Load(FArchive&).
- Added FCompressedBufferReaderSourceScope to set a buffer source within a scope.
- Added proper bounds checks to FNoneDecoder.
- Store the header checksum on the decoder context to allow raw blocks to be reused across sources.
- Decode the header on the fly to avoid a temporary header allocation when the header is in contiguous memory.

#rnx

diff --git a/Engine/Source/Developer/DerivedDataCache/Private/DerivedDataBuildJob.cpp b/Engine/Source/Developer/DerivedDataCache/Private/DerivedDataBuildJob.cpp
--- a/Engine/Source/Developer/DerivedDataCache/Private/DerivedDataBuildJob.cpp
+++ b/Engine/Source/Developer/DerivedDataCache/Private/DerivedDataBuildJob.cpp
@@ -1236,7 +1236,7 @@ void FBuildJob::ExportBuild() const
 			FPathViews::Append(ExportPath, FIoHash(Buffer.GetRawHash()));
 			if (TUniquePtr<FArchive> Ar{IFileManager::Get().CreateFileWriter(*ExportPath)})
 			{
-				*Ar << const_cast<FCompressedBuffer&>(Buffer);
+				Buffer.Save(*Ar);
 			}
 			ExportPath.RemoveSuffix(ExportPath.Len() - ExportRootLen);
 		});
diff --git a/Engine/Source/Developer/DerivedDataCache/Private/FileSystemDerivedDataBackend.cpp b/Engine/Source/Developer/DerivedDataCache/Private/FileSystemDerivedDataBackend.cpp
--- a/Engine/Source/Developer/DerivedDataCache/Private/FileSystemDerivedDataBackend.cpp
+++ b/Engine/Source/Developer/DerivedDataCache/Private/FileSystemDerivedDataBackend.cpp
@@ -1671,7 +1671,7 @@ void FFileSystemCacheStore::GetChunks(
 					FSharedBuffer Buffer;
 					if (Payload.HasData() && !bExistsOnly)
 					{
-						Buffer = Payload.GetData().Decompress(RawOffset, RawSize);
+						Buffer = FCompressedBufferReader(Payload.GetData()).Decompress(RawOffset, RawSize);
 					}
 					OnComplete({Chunk.Key, Chunk.Id, Chunk.RawOffset,
 						RawSize, Payload.GetRawHash(), MoveTemp(Buffer), PayloadStatus});
@@ -1953,7 +1953,7 @@ bool FFileSystemCacheStore::PutCacheContent(const FCompressedBuffer& Content, co
 	BuildCacheContentPath(RawHash, Path);
 	if (!FileExists(Path))
 	{
-		if (!SaveFile(Path, Context, [&Content](FArchive& Ar) { Ar << const_cast<FCompressedBuffer&>(Content); }))
+		if (!SaveFile(Path, Context, [&Content](FArchive& Ar) { Content.Save(Ar); }))
 		{
 			return false;
 		}
diff --git a/Engine/Source/Developer/DerivedDataCache/Private/MemoryDerivedDataBackend.cpp b/Engine/Source/Developer/DerivedDataCache/Private/MemoryDerivedDataBackend.cpp
--- a/Engine/Source/Developer/DerivedDataCache/Private/MemoryDerivedDataBackend.cpp
+++ b/Engine/Source/Developer/DerivedDataCache/Private/MemoryDerivedDataBackend.cpp
@@ -614,19 +614,29 @@ void FMemoryDerivedDataBackend::GetChunks(
 	IRequestOwner& Owner,
 	FOnCacheGetChunkComplete&& OnComplete)
 {
+	FPayload Payload;
+	FCacheKey PayloadKey;
+	FCompressedBufferReader Reader;
 	for (const FCacheChunkRequest& Chunk : Chunks)
 	{
 		const bool bExistsOnly = EnumHasAnyFlags(Chunk.Policy, ECachePolicy::SkipValue);
 		COOK_STAT(auto Timer = bExistsOnly ? UsageStats.TimeProbablyExists() : UsageStats.TimeGet());
-		FPayload Payload;
 		if (ShouldSimulateMiss(Chunk.Key))
 		{
 			UE_LOG(LogDerivedDataCache, Verbose, TEXT("%s: Simulated miss for get of %s from '%.*s'"),
 				*GetName(), *WriteToString<96>(Chunk.Key, '/', Chunk.Id), Context.Len(), Context.GetData());
 		}
+		else if (PayloadKey == Chunk.Key && Payload.GetId() == Chunk.Id)
+		{
+			// Payload matches the request.
+		}
 		else if (FWriteScopeLock ScopeLock(SynchronizationObject); const FCacheRecord* Record = CacheRecords.Find(Chunk.Key))
 		{
+			Reader.ResetSource();
+			Payload.Reset();
 			Payload = Record->GetAttachmentPayload(Chunk.Id);
+			PayloadKey = Chunk.Key;
+			Reader.SetSource(Payload.GetData());
 		}
 		if (Payload && Chunk.RawOffset <= Payload.GetRawSize())
 		{
@@ -634,15 +644,14 @@ void FMemoryDerivedDataBackend::GetChunks(
 			COOK_STAT(Timer.AddHit(RawSize));
 			if (OnComplete)
 			{
-				FUniqueBuffer Buffer;
+				FSharedBuffer Buffer;
 				if (Payload.HasData() && !bExistsOnly)
 				{
-					Buffer = FUniqueBuffer::Alloc(RawSize);
-					Payload.GetData().DecompressToComposite().CopyTo(Buffer, Chunk.RawOffset);
+					Buffer = Reader.Decompress(Chunk.RawOffset, RawSize);
 				}
 				const EStatus Status = bExistsOnly || Buffer ? EStatus::Ok : EStatus::Error;
 				OnComplete({Chunk.Key, Chunk.Id, Chunk.RawOffset,
-					RawSize, Payload.GetRawHash(), Buffer.MoveToShared(), Status});
+					RawSize, Payload.GetRawHash(), MoveTemp(Buffer), Status});
 			}
 		}
 		else
diff --git a/Engine/Source/Developer/DerivedDataCache/Private/PakFileDerivedDataBackend.cpp b/Engine/Source/Developer/DerivedDataCache/Private/PakFileDerivedDataBackend.cpp
--- a/Engine/Source/Developer/DerivedDataCache/Private/PakFileDerivedDataBackend.cpp
+++ b/Engine/Source/Developer/DerivedDataCache/Private/PakFileDerivedDataBackend.cpp
@@ -555,6 +555,7 @@ void FPakFileDerivedDataBackend::GetChunks(
 	SortedChunks.StableSort(TChunkLess());
 
 	FOptionalCacheRecord Record;
+	FCompressedBufferReader Reader;
 	for (const FCacheChunkRequest& Chunk : SortedChunks)
 	{
 		constexpr ECachePolicy SkipFlag = ECachePolicy::SkipValue;
@@ -585,7 +586,8 @@ void FPakFileDerivedDataBackend::GetChunks(
 					FSharedBuffer Buffer;
 					if (Payload.HasData() && !bExistsOnly)
 					{
-						Buffer = Payload.GetData().Decompress(RawOffset, RawSize);
+						FCompressedBufferReaderSourceScope Source(Reader, Payload.GetData());
+						Buffer = Reader.Decompress(RawOffset, RawSize);
 					}
 					OnComplete({Chunk.Key, Chunk.Id, Chunk.RawOffset,
 						RawSize, Payload.GetRawHash(), MoveTemp(Buffer), PayloadStatus});
@@ -839,7 +841,7 @@ bool FPakFileDerivedDataBackend::PutCacheContent(const FCompressedBuffer& Conten
 	FPathViews::Append(Path, TEXT("Content"), RawHash);
 	if (!FileExists(Path))
 	{
-		if (!SaveFile(Path, Context, [&Content](FArchive& Ar) { Ar << const_cast<FCompressedBuffer&>(Content); }))
+		if (!SaveFile(Path, Context, [&Content](FArchive& Ar) { Content.Save(Ar); }))
 		{
 			return false;
 		}
diff --git a/Engine/Source/Developer/DerivedDataCache/Private/ZenDerivedDataBackend.cpp b/Engine/Source/Developer/DerivedDataCache/Private/ZenDerivedDataBackend.cpp
--- a/Engine/Source/Developer/DerivedDataCache/Private/ZenDerivedDataBackend.cpp
+++ b/Engine/Source/Developer/DerivedDataCache/Private/ZenDerivedDataBackend.cpp
@@ -821,7 +821,7 @@ void FZenDerivedDataBackend::GetChunks(
 				else if (const FCbAttachment* Attachment = BatchResponse.FindAttachment(HashView.AsHash()))
 				{
 					const FCompressedBuffer& CompressedBuffer = Attachment->AsCompressedBinary();
-					FSharedBuffer Buffer = CompressedBuffer.Decompress(Chunk.RawOffset, Chunk.RawSize);
+					FSharedBuffer Buffer = FCompressedBufferReader(CompressedBuffer).Decompress(Chunk.RawOffset, Chunk.RawSize);
 					
 					UE_LOG(LogDerivedDataCache, Verbose, TEXT("%s: Cache hit for '%s' from '%.*s'"),
 						*GetName(), *WriteToString<96>(Chunk.Key, '/', Chunk.Id), Context.Len(), Context.GetData());
diff --git a/Engine/Source/Developer/IoStoreUtilities/Private/ZenStoreHttpClient.cpp b/Engine/Source/Developer/IoStoreUtilities/Private/ZenStoreHttpClient.cpp
--- a/Engine/Source/Developer/IoStoreUtilities/Private/ZenStoreHttpClient.cpp
+++ b/Engine/Source/Developer/IoStoreUtilities/Private/ZenStoreHttpClient.cpp
@@ -474,20 +474,22 @@ TIoStatusOr<FIoBuffer> FZenStoreHttpClient::ReadOpLogUri(FStringBuilderBase& Chu
 			uint64 CompressedOffset = 0;
 			if (Offset > 0)
 			{
-				uint32 BlockSize = 0;
-				if  (!Compressed.TryGetBlockSize(BlockSize))
+				uint64 BlockSize = 0;
+				ECompressedBufferCompressor Compressor;
+				ECompressedBufferCompressionLevel CompressionLevel;
+				if  (!Compressed.TryGetCompressParameters(Compressor, CompressionLevel, BlockSize))
 				{
 					return FIoStatus(EIoErrorCode::CompressionError);
 				}
 
 				if (BlockSize > 0)
 				{
-					CompressedOffset = Offset % uint64(BlockSize);
+					CompressedOffset = Offset % BlockSize;
 				}
 			}
 
 			FIoBuffer Decompressed(Compressed.GetRawSize());
-			if (!Compressed.TryDecompressTo(Decompressed.GetMutableView(), CompressedOffset))
+			if (!FCompressedBufferReader(Compressed).TryDecompressTo(Decompressed.GetMutableView(), CompressedOffset))
 			{
 				return FIoStatus(EIoErrorCode::CompressionError);
 			}
diff --git a/Engine/Source/Developer/Virtualization/Private/VirtualizationFileBackend.cpp b/Engine/Source/Developer/Virtualization/Private/VirtualizationFileBackend.cpp
--- a/Engine/Source/Developer/Virtualization/Private/VirtualizationFileBackend.cpp
+++ b/Engine/Source/Developer/Virtualization/Private/VirtualizationFileBackend.cpp
@@ -173,7 +173,7 @@ FCompressedBuffer FFileSystemBackend::PullData(const FPayloadId& Id)
 		return FCompressedBuffer();
 	}
 
-	return FCompressedBuffer::FromCompressed(*FileAr);
+	return FCompressedBuffer::Load(*FileAr);
 }
 
 bool FFileSystemBackend::DoesPayloadExist(const FPayloadId& Id)
diff --git a/Engine/Source/Developer/Virtualization/Private/VirtualizationSourceControlBackend.cpp b/Engine/Source/Developer/Virtualization/Private/VirtualizationSourceControlBackend.cpp
--- a/Engine/Source/Developer/Virtualization/Private/VirtualizationSourceControlBackend.cpp
+++ b/Engine/Source/Developer/Virtualization/Private/VirtualizationSourceControlBackend.cpp
@@ -159,7 +159,7 @@ EPushResult FSourceControlBackend::PushData(const FPayloadId& Id, const FCompres
 			return EPushResult::Failed;
 		}
 
-		*FileAr << const_cast<FCompressedBuffer&>(Payload);
+		Payload.Save(*FileAr);
 
 		if (!FileAr->Close())
 		{
diff --git a/Engine/Source/Editor/UnrealEd/Private/DerivedDataBuildLocalExecutor.cpp b/Engine/Source/Editor/UnrealEd/Private/DerivedDataBuildLocalExecutor.cpp
--- a/Engine/Source/Editor/UnrealEd/Private/DerivedDataBuildLocalExecutor.cpp
+++ b/Engine/Source/Editor/UnrealEd/Private/DerivedDataBuildLocalExecutor.cpp
@@ -252,7 +252,7 @@ public:
 				FPathViews::Append(Path, SandboxRoot, TEXT("Outputs"), FIoHash(Payload.GetRawHash()));
 				if (TUniquePtr<FArchive> Ar{IFileManager::Get().CreateFileReader(*Path, FILEREAD_Silent)})
 				{
-					BufferForPayload = FCompressedBuffer::FromCompressed(*Ar);
+					BufferForPayload = FCompressedBuffer::Load(*Ar);
 				}
 
 				if (BufferForPayload.IsNull())
diff --git a/Engine/Source/Programs/DerivedDataBuildWorker/Private/DerivedDataBuildWorker.cpp b/Engine/Source/Programs/DerivedDataBuildWorker/Private/DerivedDataBuildWorker.cpp
--- a/Engine/Source/Programs/DerivedDataBuildWorker/Private/DerivedDataBuildWorker.cpp
+++ b/Engine/Source/Programs/DerivedDataBuildWorker/Private/DerivedDataBuildWorker.cpp
@@ -305,7 +305,7 @@ void FBuildWorkerProgram::BuildComplete(FBuildCompleteParams&& Params) const
 		{
 			if (TUniquePtr<FArchive> Ar = OpenOutput(Output.GetName(), Payload.GetRawHash()))
 			{
-				*Ar << const_cast<FCompressedBuffer&>(Payload.GetData());
+				Payload.GetData().Save(*Ar);
 				if (Ar->Close())
 				{
 					continue;
diff --git a/Engine/Source/Runtime/Core/Private/Compression/CompressedBuffer.cpp b/Engine/Source/Runtime/Core/Private/Compression/CompressedBuffer.cpp
--- a/Engine/Source/Runtime/Core/Private/Compression/CompressedBuffer.cpp
+++ b/Engine/Source/Runtime/Core/Private/Compression/CompressedBuffer.cpp
@@ -9,9 +9,7 @@
 #include "Math/UnrealMathUtility.h"
 #include "Misc/ByteSwap.h"
 #include "Misc/Crc.h"
-#include "Misc/ScopeExit.h"
 #include "Serialization/Archive.h"
-#include "Templates/Function.h"
 
 THIRD_PARTY_INCLUDES_START
 #include "Compression/lz4.h"
@@ -166,9 +164,8 @@ class FDecoder
 {
 public:
 	virtual uint64 GetHeaderSize(const FHeader& Header) const = 0;
-	virtual void DecodeHeader(FMutableMemoryView HeaderView) const = 0;
-	virtual bool TryDecompressTo(FDecoderContext& Context, const FDecoderSource& Source, const FHeader& Header, uint64 RawOffset, FMutableMemoryView RawView) const = 0;
-	virtual FCompositeBuffer DecompressToComposite(FDecoderContext& Context, const FDecoderSource& Source, const FHeader& Header, uint64 RawOffset, uint64 RawSize) const =0;
+	virtual bool TryDecompressTo(FDecoderContext& Context, const FDecoderSource& Source, const FHeader& Header, FMemoryView HeaderView, uint64 RawOffset, FMutableMemoryView RawView) const = 0;
+	virtual FCompositeBuffer DecompressToComposite(FDecoderContext& Context, const FDecoderSource& Source, const FHeader& Header, const FMemoryView HeaderView, uint64 RawOffset, uint64 RawSize) const = 0;
 };
 
 ///////////////////////////////////////////////////////////////////////////////////////////////////////////////////////
@@ -199,20 +196,18 @@ public:
 		return sizeof(FHeader);
 	}
 
-	void DecodeHeader(FMutableMemoryView HeaderView) const final
-	{
-		FHeader& Header = *static_cast<FHeader*>(HeaderView.GetData());
-		Header.ByteSwap();
-	}
-
 	bool TryDecompressTo(
 		FDecoderContext& Context,
 		const FDecoderSource& Source,
 		const FHeader& Header,
-		uint64 RawOffset,
-		FMutableMemoryView RawView) const final
+		const FMemoryView HeaderView,
+		const uint64 RawOffset,
+		const FMutableMemoryView RawView) const final
 	{
-		if (Header.Method == EMethod::None && RawOffset + RawView.GetSize() <= Header.TotalRawSize)
+		if (Header.Method == EMethod::None &&
+			RawOffset <= Header.TotalRawSize &&
+			RawView.GetSize() <= Header.TotalRawSize - RawOffset &&
+			Header.TotalCompressedSize == Header.TotalRawSize + sizeof(FHeader))
 		{
 			return Source.Read(sizeof(FHeader) + RawOffset, RawView);
 		}
@@ -223,10 +218,13 @@ public:
 		FDecoderContext& Context,
 		const FDecoderSource& Source,
 		const FHeader& Header,
-		uint64 RawOffset,
-		uint64 RawSize) const final
+		const FMemoryView HeaderView,
+		const uint64 RawOffset,
+		const uint64 RawSize) const final
 	{
 		if (Header.Method == EMethod::None &&
+			RawOffset <= Header.TotalRawSize &&
+			RawSize <= Header.TotalRawSize - RawOffset &&
 			Header.TotalCompressedSize == Header.TotalRawSize + sizeof(FHeader))
 		{
 			return Source.ReadToComposite(sizeof(FHeader) + RawOffset, RawSize);
@@ -269,11 +267,11 @@ FCompositeBuffer FBlockEncoder::Compress(const FCompositeBuffer& RawData, const
 	FBlake3 RawHash;
 
 	const uint64 BlockCount = FMath::DivideAndRoundUp(RawSize, BlockSize);
-	checkf(BlockCount <= MAX_uint32, TEXT("Raw data of size %" UINT64_FMT " with block size %" UINT64_FMT " requires ")
+	checkf(BlockCount < MAX_uint32, TEXT("Raw data of size %" UINT64_FMT " with block size %" UINT64_FMT " requires ")
 		TEXT("%" UINT64_FMT " blocks, but the limit is %u."), RawSize, BlockSize, BlockCount, MAX_uint32);
 
 	// Allocate the buffer for the header, metadata, and compressed blocks.
-	const uint64 MetaSize = BlockCount * sizeof(uint32);
+	const uint64 MetaSize = sizeof(uint32) * BlockCount;
 	const uint64 CompressedDataSize = sizeof(FHeader) + MetaSize + GetCompressedBlocksBound(BlockCount, BlockSize, RawSize);
 	FUniqueBuffer CompressedData = FUniqueBuffer::Alloc(CompressedDataSize);
 
@@ -347,16 +345,8 @@ public:
 		return sizeof(FHeader) + sizeof(uint32) * uint64(Header.BlockCount);
 	}
 
-	void DecodeHeader(FMutableMemoryView HeaderView) const final
-	{
-		FHeader& Header = *static_cast<FHeader*>(HeaderView.GetData());
-		Header.ByteSwap();
-		TArrayView<uint32> CompressedBlockSizes(reinterpret_cast<uint32*>(&Header + 1), Header.BlockCount);
-		Algo::ForEach(CompressedBlockSizes, [](uint32& Size) { Size = NETWORK_ORDER32(Size); });
-	}
-
-	bool TryDecompressTo(FDecoderContext& Context, const FDecoderSource& Source, const FHeader& Header, uint64 RawOffset, FMutableMemoryView RawView) const final;
-	FCompositeBuffer DecompressToComposite(FDecoderContext& Context, const FDecoderSource& Source, const FHeader& Header, uint64 RawOffset, uint64 RawSize) const final;
+	bool TryDecompressTo(FDecoderContext& Context, const FDecoderSource& Source, const FHeader& Header, FMemoryView HeaderView, uint64 RawOffset, FMutableMemoryView RawView) const final;
+	FCompositeBuffer DecompressToComposite(FDecoderContext& Context, const FDecoderSource& Source, const FHeader& Header, const FMemoryView HeaderView, uint64 RawOffset, uint64 RawSize) const final;
 
 protected:
 	virtual bool DecompressBlock(FMutableMemoryView RawData, FMemoryView CompressedData) const = 0;
@@ -366,6 +356,7 @@ bool FBlockDecoder::TryDecompressTo(
 	FDecoderContext& Context,
 	const FDecoderSource& Source,
 	const FHeader& Header,
+	const FMemoryView HeaderView,
 	const uint64 RawOffset,
 	FMutableMemoryView RawView) const
 {
@@ -375,21 +366,22 @@ bool FBlockDecoder::TryDecompressTo(
 	}
 
 	const uint64 BlockSize = uint64(1) << Header.BlockSizeExponent;
-	const uint64 FirstBlockIndex = uint64(RawOffset / BlockSize);
-	const uint64 LastBlockIndex = uint64((RawOffset + RawView.GetSize() - 1) / BlockSize);
-	const uint64 LastBlockSize = BlockSize - ((Header.BlockCount * BlockSize) - Header.TotalRawSize);
+	const uint64 LastBlockSize = BlockSize - (BlockSize * Header.BlockCount - Header.TotalRawSize);
+	const uint32 FirstBlockIndex = uint32(RawOffset / BlockSize);
+	const uint32 LastBlockIndex = uint32((RawOffset + RawView.GetSize() - 1) / BlockSize);
 
 	uint64 RawBlockOffset = RawOffset % BlockSize;
 
-	const TConstArrayView64<uint32> CompressedBlockSizes(reinterpret_cast<const uint32*>(&Header + 1), Header.BlockCount);
-	uint64 CompressedOffset = sizeof(FHeader) + uint64(Header.BlockCount) * sizeof(uint32);
-	CompressedOffset += Algo::Accumulate(CompressedBlockSizes.Left(FirstBlockIndex), uint64(0));
+	const uint32* const CompressedBlockSizes = static_cast<const uint32*>(HeaderView.RightChop(sizeof(FHeader)).GetData());
+	uint64 CompressedOffset = sizeof(FHeader) + sizeof(uint32) * uint32(Header.BlockCount) +
+		Algo::TransformAccumulate(MakeArrayView(CompressedBlockSizes, FirstBlockIndex),
+			[](uint32 Size) -> uint64 { return NETWORK_ORDER32(Size); }, uint64(0));
 
-	for (uint64 BlockIndex = FirstBlockIndex; BlockIndex <= LastBlockIndex; BlockIndex++)
+	for (uint32 BlockIndex = FirstBlockIndex; BlockIndex <= LastBlockIndex; BlockIndex++)
 	{
 		const uint64 RawBlockSize = BlockIndex == Header.BlockCount - 1 ? LastBlockSize : BlockSize;
 		const uint64 RawBlockReadSize = FMath::Min(RawView.GetSize(), RawBlockSize - RawBlockOffset);
-		const uint32 CompressedBlockSize = CompressedBlockSizes[BlockIndex];
+		const uint32 CompressedBlockSize = NETWORK_ORDER32(CompressedBlockSizes[BlockIndex]);
 		const bool bIsCompressed = CompressedBlockSize < RawBlockSize;
 
 		if (bIsCompressed)
@@ -444,11 +436,12 @@ FCompositeBuffer FBlockDecoder::DecompressToComposite(
 	FDecoderContext& Context,
 	const FDecoderSource& Source,
 	const FHeader& Header,
+	const FMemoryView HeaderView,
 	const uint64 RawOffset,
 	const uint64 RawSize) const
 {
 	FUniqueBuffer Buffer = FUniqueBuffer::Alloc(RawSize);
-	if (TryDecompressTo(Context, Source, Header, RawOffset, Buffer))
+	if (TryDecompressTo(Context, Source, Header, HeaderView, RawOffset, Buffer))
 	{
 		return FCompositeBuffer(Buffer.MoveToShared());
 	}
@@ -603,21 +596,24 @@ bool FHeader::IsValid(const FCompositeBuffer& CompressedData)
 
 ///////////////////////////////////////////////////////////////////////////////////////////////////////////////////////
 
-static const FHeader* TryReadHeader(FDecoderContext& Context, FArchive& Ar)
+static bool TryReadHeader(FDecoderContext& Context, FArchive& Ar, FHeader& OutHeader, FMemoryView& OutHeaderView)
 {
 	if (Context.HeaderOffset != MAX_uint64)
 	{
-		return static_cast<const FHeader*>(Context.Header.GetData());
+		OutHeaderView = Context.Header.GetView().Left(Context.HeaderSize);
+		MakeMemoryView(&OutHeader, &OutHeader + 1).CopyFrom(OutHeaderView.Left(sizeof(FHeader)));
+		OutHeader.ByteSwap();
+		return true;
 	}
 
 	check(Ar.IsLoading());
 	const int64 Offset = Ar.Tell();
 	if (Offset == INDEX_NONE)
 	{
-		return nullptr;
+		return false;
 	}
 
-	FHeader Header;
+	FHeader& Header = OutHeader;
 	Ar.Serialize(&Header, sizeof(FHeader));
 	Header.ByteSwap();
 
@@ -631,57 +627,63 @@ static const FHeader* TryReadHeader(FDecoderContext& Context, FArchive& Ar)
 
 		const FMutableMemoryView HeaderView = Context.Header.GetView().Left(HeaderSize);
 		const FMutableMemoryView HeaderTail = HeaderView.CopyFrom(MakeMemoryView(&Header, &Header + 1));
-		Ar.Serialize(HeaderTail.GetData(), static_cast<int64>(HeaderTail.GetSize()));
+		Ar.Serialize(HeaderTail.GetData(), int64(HeaderTail.GetSize()));
 
 		FHeader* const HeaderCopy = static_cast<FHeader*>(HeaderView.GetData());
 		HeaderCopy->ByteSwap();
 		if (Header.Crc32 == FHeader::CalculateCrc32(HeaderView))
 		{
 			Context.HeaderOffset = uint64(Offset);
-			Decoder->DecodeHeader(HeaderView);
-			return HeaderCopy;
+			Context.HeaderSize = HeaderSize;
+			Context.HeaderCrc32 = Header.Crc32;
+			Context.RawBlockIndex = MAX_uint32;
+			OutHeaderView = HeaderView;
+			return true;
 		}
 	}
 
-	return nullptr;
+	return false;
 }
 
-static const FHeader* TryReadHeader(FDecoderContext& Context, const FCompositeBuffer& Buffer)
+static bool TryReadHeader(FDecoderContext& Context, const FCompositeBuffer& Buffer, FHeader& OutHeader, FMemoryView& OutHeaderView)
 {
 	if (Context.HeaderOffset != MAX_uint64)
 	{
-		return static_cast<const FHeader*>(Context.Header.GetData());
+		OutHeaderView = Buffer.ViewOrCopyRange(Context.HeaderOffset, Context.HeaderSize, Context.Header);
+		MakeMemoryView(&OutHeader, &OutHeader + 1).CopyFrom(OutHeaderView.Left(sizeof(FHeader)));
+		OutHeader.ByteSwap();
+		return true;
 	}
 
 	if (Buffer.GetSize() < sizeof(FHeader))
 	{
-		return nullptr;
+		return false;
 	}
 
-	FHeader Header;
+	FHeader& Header = OutHeader;
 	Buffer.CopyTo(MakeMemoryView(&Header, &Header + 1));
 	Header.ByteSwap();
 
 	if (const FDecoder* const Decoder = GetDecoder(Header.Method); Decoder && Header.Magic == FHeader::ExpectedMagic)
 	{
 		const uint64 HeaderSize = Decoder->GetHeaderSize(Header);
-		if (Context.Header.GetSize() < HeaderSize)
-		{
-			Context.Header = FUniqueBuffer::Alloc(FMath::Max(FMath::RoundUpToPowerOfTwo64(HeaderSize), DefaultHeaderSize));
-		}
-
-		const FMutableMemoryView HeaderView = Context.Header.GetView().Left(HeaderSize);
-		Buffer.CopyTo(HeaderView);
-
+		const FMemoryView HeaderView = Buffer.ViewOrCopyRange(0, HeaderSize, Context.Header,
+			[](uint64 Size) { return FUniqueBuffer::Alloc(FMath::Max(FMath::RoundUpToPowerOfTwo64(Size), DefaultHeaderSize)); });
 		if (Header.Crc32 == FHeader::CalculateCrc32(HeaderView))
 		{
 			Context.HeaderOffset = 0;
-			Decoder->DecodeHeader(HeaderView);
-			return static_cast<FHeader*>(HeaderView.GetData());
+			Context.HeaderSize = HeaderSize;
+			if (Context.HeaderCrc32 != Header.Crc32)
+			{
+				Context.HeaderCrc32 = Header.Crc32;
+				Context.RawBlockIndex = MAX_uint32;
+			}
+			OutHeaderView = HeaderView;
+			return true;
 		}
 	}
 
-	return nullptr;
+	return false;
 }
 
 ///////////////////////////////////////////////////////////////////////////////////////////////////////////////////////
@@ -838,7 +840,7 @@ FCompressedBuffer FCompressedBuffer::FromCompressed(FSharedBuffer&& InCompressed
 	return Local;
 }
 
-FCompressedBuffer FCompressedBuffer::FromCompressed(FArchive& Ar)
+FCompressedBuffer FCompressedBuffer::Load(FArchive& Ar)
 {
 	using namespace UE::CompressedBuffer::Private;
 	check(Ar.IsLoading());
@@ -863,6 +865,15 @@ FCompressedBuffer FCompressedBuffer::FromCompressed(FArchive& Ar)
 	return Local;
 }
 
+void FCompressedBuffer::Save(FArchive& Ar) const
+{
+	check(Ar.IsSaving());
+	for (const FSharedBuffer& Segment : CompressedData.GetSegments())
+	{
+		Ar.Serialize(const_cast<void*>(Segment.GetData()), int64(Segment.GetSize()));
+	}
+}
+
 uint64 FCompressedBuffer::GetCompressedSize() const
 {
 	return CompressedData.GetSize();
@@ -891,27 +902,26 @@ bool FCompressedBuffer::TryGetCompressParameters(
 	return false;
 }
 
-bool FCompressedBuffer::TryDecompressTo(FMutableMemoryView RawView, uint64 RawOffset) const
+bool FCompressedBuffer::TryDecompressTo(const FMutableMemoryView RawView) const
 {
-	if (CompressedData)
+	if (CompressedData && RawView.GetSize() == GetRawSize())
 	{
-		return FCompressedBufferReader(*this).TryDecompressTo(RawView, RawOffset);
+		return FCompressedBufferReader(*this).TryDecompressTo(RawView);
 	}
 	return false;
 }
 
-FSharedBuffer FCompressedBuffer::Decompress(uint64 RawOffset, uint64 RawSize) const
+FSharedBuffer FCompressedBuffer::Decompress() const
 {
 	if (CompressedData)
 	{
-		return FCompressedBufferReader(*this).Decompress(RawOffset, RawSize);
+		return FCompressedBufferReader(*this).Decompress();
 	}
 	return FSharedBuffer();
 }
 
 FCompositeBuffer FCompressedBuffer::DecompressToComposite() const
 {
-	using namespace UE::CompressedBuffer::Private;
 	if (CompressedData)
 	{
 		return FCompressedBufferReader(*this).DecompressToComposite();
@@ -925,14 +935,11 @@ FArchive& operator<<(FArchive& Ar, FCompressedBuffer& Buffer)
 {
 	if (Ar.IsLoading())
 	{
-		Buffer = FCompressedBuffer::FromCompressed(Ar);
+		Buffer = FCompressedBuffer::Load(Ar);
 	}
 	else
 	{
-		for (const FSharedBuffer& Segment : Buffer.GetCompressed().GetSegments())
-		{
-			Ar.Serialize(const_cast<void*>(Segment.GetData()), static_cast<int64>(Segment.GetSize()));
-		}
+		Buffer.Save(Ar);
 	}
 	return Ar;
 }
@@ -962,7 +969,6 @@ void FCompressedBufferReader::ResetBuffers()
 void FCompressedBufferReader::ResetSource()
 {
 	Context.HeaderOffset = MAX_uint64;
-	Context.RawBlockIndex = MAX_uint64;
 	SourceArchive = nullptr;
 	SourceBuffer = nullptr;
 }
@@ -974,7 +980,6 @@ void FCompressedBufferReader::SetSource(FArchive& Archive)
 		return;
 	}
 	Context.HeaderOffset = MAX_uint64;
-	Context.RawBlockIndex = MAX_uint64;
 	SourceArchive = &Archive;
 	SourceBuffer = nullptr;
 }
@@ -986,7 +991,6 @@ void FCompressedBufferReader::SetSource(const FCompressedBuffer& Buffer)
 		return;
 	}
 	Context.HeaderOffset = MAX_uint64;
-	Context.RawBlockIndex = MAX_uint64;
 	SourceArchive = nullptr;
 	SourceBuffer = &Buffer;
 }
@@ -994,13 +998,11 @@ void FCompressedBufferReader::SetSource(const FCompressedBuffer& Buffer)
 uint64 FCompressedBufferReader::GetCompressedSize()
 {
 	using namespace UE::CompressedBuffer::Private;
-	if (const FCompressedBuffer* const Buffer = SourceBuffer)
-	{
-		return Buffer->GetCompressedSize();
-	}
-	if (const FHeader* const Header = TryReadHeader())
+	FHeader Header;
+	FMemoryView HeaderView;
+	if (TryReadHeader(Header, HeaderView))
 	{
-		return Header->TotalCompressedSize;
+		return Header.TotalCompressedSize;
 	}
 	return 0;
 }
@@ -1008,13 +1010,11 @@ uint64 FCompressedBufferReader::GetCompressedSize()
 uint64 FCompressedBufferReader::GetRawSize()
 {
 	using namespace UE::CompressedBuffer::Private;
-	if (const FCompressedBuffer* const Buffer = SourceBuffer)
-	{
-		return Buffer->GetRawSize();
-	}
-	if (const FHeader* const Header = TryReadHeader())
+	FHeader Header;
+	FMemoryView HeaderView;
+	if (TryReadHeader(Header, HeaderView))
 	{
-		return Header->TotalRawSize;
+		return Header.TotalRawSize;
 	}
 	return 0;
 }
@@ -1022,13 +1022,11 @@ uint64 FCompressedBufferReader::GetRawSize()
 FBlake3Hash FCompressedBufferReader::GetRawHash()
 {
 	using namespace UE::CompressedBuffer::Private;
-	if (const FCompressedBuffer* const Buffer = SourceBuffer)
-	{
-		return Buffer->GetRawHash();
-	}
-	if (const FHeader* const Header = TryReadHeader())
+	FHeader Header;
+	FMemoryView HeaderView;
+	if (TryReadHeader(Header, HeaderView))
 	{
-		return Header->RawHash;
+		return Header.RawHash;
 	}
 	return FBlake3Hash();
 }
@@ -1039,32 +1037,32 @@ bool FCompressedBufferReader::TryGetCompressParameters(
 	uint64& OutBlockSize)
 {
 	using namespace UE::CompressedBuffer::Private;
-	if (const FCompressedBuffer* const Buffer = SourceBuffer)
-	{
-		return Buffer->TryGetCompressParameters(OutCompressor, OutCompressionLevel, OutBlockSize);
-	}
-	if (const FHeader* const Header = TryReadHeader())
+	FHeader Header;
+	FMemoryView HeaderView;
+	if (TryReadHeader(Header, HeaderView))
 	{
-		return Header->TryGetCompressParameters(OutCompressor, OutCompressionLevel, OutBlockSize);
+		return Header.TryGetCompressParameters(OutCompressor, OutCompressionLevel, OutBlockSize);
 	}
 	return false;
 }
 
-bool FCompressedBufferReader::TryDecompressTo(FMutableMemoryView RawView, uint64 RawOffset)
+bool FCompressedBufferReader::TryDecompressTo(const FMutableMemoryView RawView, const uint64 RawOffset)
 {
 	using namespace UE::CompressedBuffer::Private;
-	if (const FHeader* const Header = TryReadHeader())
+	FHeader Header;
+	FMemoryView HeaderView;
+	if (TryReadHeader(Header, HeaderView))
 	{
-		const uint64 TotalRawSize = Header->TotalRawSize;
+		const uint64 TotalRawSize = Header.TotalRawSize;
 		if (RawOffset <= TotalRawSize && RawView.GetSize() <= TotalRawSize - RawOffset)
 		{
-			if (const FDecoder* const Decoder = GetDecoder(Header->Method))
+			if (const FDecoder* const Decoder = GetDecoder(Header.Method))
 			{
 				if (Decoder->TryDecompressTo(Context,
 					SourceArchive
 						? ImplicitConv<const FDecoderSource&>(FArchiveDecoderSource(*SourceArchive, Context.HeaderOffset))
 						: ImplicitConv<const FDecoderSource&>(FBufferDecoderSource(SourceBuffer->GetCompressed())),
-					*Header, RawOffset, RawView))
+					Header, HeaderView, RawOffset, RawView))
 				{
 					return true;
 				}
@@ -1077,20 +1075,22 @@ bool FCompressedBufferReader::TryDecompressTo(FMutableMemoryView RawView, uint64
 FSharedBuffer FCompressedBufferReader::Decompress(const uint64 RawOffset, const uint64 RawSize)
 {
 	using namespace UE::CompressedBuffer::Private;
-	if (const FHeader* const Header = TryReadHeader())
+	FHeader Header;
+	FMemoryView HeaderView;
+	if (TryReadHeader(Header, HeaderView))
 	{
-		const uint64 TotalRawSize = Header->TotalRawSize;
+		const uint64 TotalRawSize = Header.TotalRawSize;
 		const uint64 RawSizeToCopy = RawSize == MAX_uint64 ? TotalRawSize - RawOffset : RawSize;
 		if (RawOffset <= TotalRawSize && RawSizeToCopy <= TotalRawSize - RawOffset)
 		{
-			if (const FDecoder* const Decoder = GetDecoder(Header->Method))
+			if (const FDecoder* const Decoder = GetDecoder(Header.Method))
 			{
 				FUniqueBuffer RawData = FUniqueBuffer::Alloc(RawSizeToCopy);
 				if (Decoder->TryDecompressTo(Context,
 					SourceArchive
 						? ImplicitConv<const FDecoderSource&>(FArchiveDecoderSource(*SourceArchive, Context.HeaderOffset))
 						: ImplicitConv<const FDecoderSource&>(FBufferDecoderSource(SourceBuffer->GetCompressed())),
-					*Header, RawOffset, RawData))
+					Header, HeaderView, RawOffset, RawData))
 				{
 					return RawData.MoveToShared();
 				}
@@ -1103,68 +1103,39 @@ FSharedBuffer FCompressedBufferReader::Decompress(const uint64 RawOffset, const
 FCompositeBuffer FCompressedBufferReader::DecompressToComposite(const uint64 RawOffset, const uint64 RawSize)
 {
 	using namespace UE::CompressedBuffer::Private;
-	if (const FHeader* const Header = TryReadHeader())
+	FHeader Header;
+	FMemoryView HeaderView;
+	if (TryReadHeader(Header, HeaderView))
 	{
-		const uint64 TotalRawSize = Header->TotalRawSize;
+		const uint64 TotalRawSize = Header.TotalRawSize;
 		const uint64 RawSizeToCopy = RawSize == MAX_uint64 ? TotalRawSize - RawOffset : RawSize;
 		if (RawOffset <= TotalRawSize && RawSizeToCopy <= TotalRawSize - RawOffset)
 		{
-			if (const FDecoder* const Decoder = GetDecoder(Header->Method))
+			if (const FDecoder* const Decoder = GetDecoder(Header.Method))
 			{
 				return Decoder->DecompressToComposite(Context,
 					SourceArchive
 						? ImplicitConv<const FDecoderSource&>(FArchiveDecoderSource(*SourceArchive, Context.HeaderOffset))
 						: ImplicitConv<const FDecoderSource&>(FBufferDecoderSource(SourceBuffer->GetCompressed())),
-					*Header, RawOffset, RawSizeToCopy);
+					Header, HeaderView, RawOffset, RawSizeToCopy);
 			}
 		}
 	}
 	return FCompositeBuffer();
 }
 
-const UE::CompressedBuffer::Private::FHeader* FCompressedBufferReader::TryReadHeader()
+bool FCompressedBufferReader::TryReadHeader(UE::CompressedBuffer::Private::FHeader& OutHeader, FMemoryView& OutHeaderView)
 {
 	using namespace UE::CompressedBuffer;
 	if (FArchive* const Archive = SourceArchive)
 	{
-		return Private::TryReadHeader(Context, *Archive);
+		return Private::TryReadHeader(Context, *Archive, OutHeader, OutHeaderView);
 	}
 	if (const FCompressedBuffer* const Buffer = SourceBuffer)
 	{
-		return Private::TryReadHeader(Context, Buffer->GetCompressed());
+		return Private::TryReadHeader(Context, Buffer->GetCompressed(), OutHeader, OutHeaderView);
 	}
-	return nullptr;
-}
-
-///////////////////////////////////////////////////////////////////////////////////////////////////////////////////////
-
-class FCompressedBufferDecoder::FImpl
-{
-public:
-	bool TryDecompressTo(const FCompressedBuffer& CompressedBuffer, FMutableMemoryView RawView, uint64 RawOffset)
-	{
-		using namespace UE::CompressedBuffer::Private;
-		Reader.SetSource(CompressedBuffer);
-		ON_SCOPE_EXIT { Reader.ResetSource(); };
-		return Reader.TryDecompressTo(RawView, RawOffset);
-	}
-
-private:
-	FCompressedBufferReader Reader;
-};
-
-FCompressedBufferDecoder::FCompressedBufferDecoder()
-: Impl(new FImpl())
-{
-}
-
-FCompressedBufferDecoder::~FCompressedBufferDecoder()
-{
-}
-
-bool FCompressedBufferDecoder::TryDecompressTo(const FCompressedBuffer& CompressedBuffer, FMutableMemoryView RawView, uint64 RawOffset)
-{
-	return Impl->TryDecompressTo(CompressedBuffer, RawView, RawOffset);
+	return false;
 }
 
 ///////////////////////////////////////////////////////////////////////////////////////////////////////////////////////
diff --git a/Engine/Source/Runtime/Core/Private/Tests/Compression/CompressedBufferTest.cpp b/Engine/Source/Runtime/Core/Private/Tests/Compression/CompressedBufferTest.cpp
--- a/Engine/Source/Runtime/Core/Private/Tests/Compression/CompressedBufferTest.cpp
+++ b/Engine/Source/Runtime/Core/Private/Tests/Compression/CompressedBufferTest.cpp
@@ -157,12 +157,6 @@ bool FCompressedBufferDecompressTest::RunTest(const FString& Parameters)
 				const TConstArrayView<uint64> Values = CastToArrayView(Uncompressed);
 				ValidateData(Values, ExpectedValues, OffsetCount);
 			}
-			{
-				const FSharedBuffer Uncompressed = Compressed.Decompress(OffsetCount * sizeof(uint64), Count * sizeof(uint64));
-				const TConstArrayView<uint64> Values = CastToArrayView(Uncompressed);
-				TestEqual("UncompressedCount", Values.Num(), Count);
-				ValidateData(Values, ExpectedValues, OffsetCount);
-			}
 		};
 
 		constexpr uint64 BlockSize = 64 * sizeof(uint64);
@@ -204,22 +198,17 @@ bool FCompressedBufferDecompressTest::RunTest(const FString& Parameters)
 		{
 			FSharedBuffer Buffer = Compressed.GetCompressed().ToShared();
 			FBufferReader Ar(const_cast<void*>(Buffer.GetData()), int64(Buffer.GetSize()), /*bFreeOnClose*/ false, /*bIsPersistent*/ true);
-			Reader.SetSource(Ar);
+			FCompressedBufferReaderSourceScope Source(Reader, Ar);
 			const FSharedBuffer Uncompressed = Reader.Decompress(OffsetCount * sizeof(uint64));
 			const TConstArrayView<uint64> Values = CastToArrayView(Uncompressed);
 			ValidateData(Values, ExpectedValues, OffsetCount);
 		}
 		{
-			Reader.SetSource(Compressed);
+			FCompressedBufferReaderSourceScope Source(Reader, Compressed);
 			const FSharedBuffer Uncompressed = Reader.Decompress(OffsetCount * sizeof(uint64));
 			const TConstArrayView<uint64> Values = CastToArrayView(Uncompressed);
 			ValidateData(Values, ExpectedValues, OffsetCount);
 		}
-		{
-			const FSharedBuffer Uncompressed = Compressed.Decompress(OffsetCount * sizeof(uint64));
-			const TConstArrayView<uint64> Values = CastToArrayView(Uncompressed);
-			ValidateData(Values, ExpectedValues, OffsetCount);
-		}
 	}
 
 	// Only one block
@@ -239,22 +228,17 @@ bool FCompressedBufferDecompressTest::RunTest(const FString& Parameters)
 		{
 			FSharedBuffer Buffer = Compressed.GetCompressed().ToShared();
 			FBufferReader Ar(const_cast<void*>(Buffer.GetData()), int64(Buffer.GetSize()), /*bFreeOnClose*/ false, /*bIsPersistent*/ true);
-			Reader.SetSource(Ar);
+			FCompressedBufferReaderSourceScope Source(Reader, Ar);
 			const FSharedBuffer Uncompressed = Reader.Decompress(OffsetCount * sizeof(uint64), Count * sizeof(uint64));
 			const TConstArrayView<uint64> Values = CastToArrayView(Uncompressed);
 			ValidateData(Values, ExpectedValues, OffsetCount);
 		}
 		{
-			Reader.SetSource(Compressed);
+			FCompressedBufferReaderSourceScope Source(Reader, Compressed);
 			const FSharedBuffer Uncompressed = Reader.Decompress(OffsetCount * sizeof(uint64), Count * sizeof(uint64));
 			const TConstArrayView<uint64> Values = CastToArrayView(Uncompressed);
 			ValidateData(Values, ExpectedValues, OffsetCount);
 		}
-		{
-			const FSharedBuffer Uncompressed = Compressed.Decompress(OffsetCount * sizeof(uint64), Count * sizeof(uint64));
-			const TConstArrayView<uint64> Values = CastToArrayView(Uncompressed);
-			ValidateData(Values, ExpectedValues, OffsetCount);
-		}
 	}
 
 	// Uncompressed
@@ -266,37 +250,22 @@ bool FCompressedBufferDecompressTest::RunTest(const FString& Parameters)
 			FSharedBuffer::MakeView(MakeMemoryView(ExpectedValues)),
 			FOodleDataCompression::ECompressor::NotSet,
 			FOodleDataCompression::ECompressionLevel::None);
+		Reader.SetSource(Compressed);
 
 		{
 			constexpr uint64 OffsetCount = 0;
 			constexpr uint64 Count = N;
-			{
-				Reader.SetSource(Compressed);
-				const FSharedBuffer Uncompressed = Reader.Decompress(OffsetCount * sizeof(uint64), Count * sizeof(uint64));
-				const TConstArrayView<uint64> Values = CastToArrayView(Uncompressed);
-				ValidateData(Values, ExpectedValues, OffsetCount);
-			}
-			{
-				const FSharedBuffer Uncompressed = Compressed.Decompress(OffsetCount * sizeof(uint64), Count * sizeof(uint64));
-				const TConstArrayView<uint64> Values = CastToArrayView(Uncompressed);
-				ValidateData(Values, ExpectedValues, OffsetCount);
-			}
+			const FSharedBuffer Uncompressed = Reader.Decompress(OffsetCount * sizeof(uint64), Count * sizeof(uint64));
+			const TConstArrayView<uint64> Values = CastToArrayView(Uncompressed);
+			ValidateData(Values, ExpectedValues, OffsetCount);
 		}
 
 		{
 			constexpr uint64 OffsetCount = 21;
 			constexpr uint64 Count = 999;
-			{
-				Reader.SetSource(Compressed);
-				const FSharedBuffer Uncompressed = Reader.Decompress(OffsetCount * sizeof(uint64), Count * sizeof(uint64));
-				const TConstArrayView<uint64> Values = CastToArrayView(Uncompressed);
-				ValidateData(Values, ExpectedValues, OffsetCount);
-			}
-			{
-				const FSharedBuffer Uncompressed = Compressed.Decompress(OffsetCount * sizeof(uint64), Count * sizeof(uint64));
-				const TConstArrayView<uint64> Values = CastToArrayView(Uncompressed);
-				ValidateData(Values, ExpectedValues, OffsetCount);
-			}
+			const FSharedBuffer Uncompressed = Reader.Decompress(OffsetCount * sizeof(uint64), Count * sizeof(uint64));
+			const TConstArrayView<uint64> Values = CastToArrayView(Uncompressed);
+			ValidateData(Values, ExpectedValues, OffsetCount);
 		}
 	}
 
diff --git a/Engine/Source/Runtime/Core/Public/Compression/CompressedBuffer.h b/Engine/Source/Runtime/Core/Public/Compression/CompressedBuffer.h
--- a/Engine/Source/Runtime/Core/Public/Compression/CompressedBuffer.h
+++ b/Engine/Source/Runtime/Core/Public/Compression/CompressedBuffer.h
@@ -5,7 +5,6 @@
 #include "Memory/CompositeBuffer.h"
 #include "Memory/MemoryFwd.h"
 #include "Memory/SharedBuffer.h"
-#include "Templates/UniquePtr.h"
 
 class FArchive;
 struct FBlake3Hash;
@@ -23,6 +22,10 @@ using ECompressedBufferCompressor = FOodleDataCompression::ECompressor;
  *
  * A buffer is self-contained in the sense that it can be decompressed without external knowledge
  * of the compression format or the size of the raw data.
+ *
+ * The buffer may be partially decompressed using FCompressedBufferReader.
+ *
+ * @see FCompressedBufferReader
  */
 class FCompressedBuffer
 {
@@ -68,7 +71,20 @@ public:
 	[[nodiscard]] CORE_API static FCompressedBuffer FromCompressed(FCompositeBuffer&& CompressedData);
 	[[nodiscard]] CORE_API static FCompressedBuffer FromCompressed(const FSharedBuffer& CompressedData);
 	[[nodiscard]] CORE_API static FCompressedBuffer FromCompressed(FSharedBuffer&& CompressedData);
-	[[nodiscard]] CORE_API static FCompressedBuffer FromCompressed(FArchive& Ar);
+
+	/**
+	 * Load a compressed buffer from an archive, as saved by Save().
+	 *
+	 * The entire compressed buffer will be loaded from the archive before this function returns.
+	 * Prefer to use FCompressedBufferReader to stream from an archive when the compressed buffer
+	 * does not need to be fully loaded into memory.
+	 *
+	 * @return A compressed buffer, or null on error. Ar.IsError() will be true on error.
+	 */
+	[[nodiscard]] CORE_API static FCompressedBuffer Load(FArchive& Ar);
+
+	/** Save the compressed buffer to an archive. */
+	CORE_API void Save(FArchive& Ar) const;
 
 	/** Reset this to null. */
 	inline void Reset() { CompressedData.Reset(); }
@@ -113,46 +129,19 @@ public:
 		ECompressedBufferCompressionLevel& OutCompressionLevel,
 		uint64& OutBlockSize) const;
 
-	[[nodiscard]] inline bool TryGetCompressParameters(
-		ECompressedBufferCompressor& OutCompressor,
-		ECompressedBufferCompressionLevel& OutCompressionLevel) const
-	{
-		uint64 BlockSize;
-		return TryGetCompressParameters(OutCompressor, OutCompressionLevel, BlockSize);
-	}
-
-	[[nodiscard]] inline bool TryGetBlockSize(uint32& OutBlockSize) const
-	{
-		ECompressedBufferCompressor Compressor;
-		ECompressedBufferCompressionLevel CompressionLevel;
-		uint64 BlockSize64;
-		if (TryGetCompressParameters(Compressor, CompressionLevel, BlockSize64))
-		{
-			OutBlockSize = uint32(BlockSize64);
-			return true;
-		}
-		return false;
-	}
-
 	/**
-	 * Decompress into a memory view that is less or equal to the raw size.
+	 * Decompress into a memory view that is exactly equal to the raw size.
 	 *
 	 * @return True if the requested range was decompressed, otherwise false.
 	 */
-	[[nodiscard]] CORE_API bool TryDecompressTo(FMutableMemoryView RawView, uint64 RawOffset = 0) const;
+	[[nodiscard]] CORE_API bool TryDecompressTo(FMutableMemoryView RawView) const;
 
 	/**
 	 * Decompress into an owned buffer.
 	 *
 	 * @return An owned buffer containing the raw data, or null on error.
 	 */
-	[[nodiscard]] CORE_API FSharedBuffer Decompress(uint64 RawOffset = 0, uint64 RawSize = MAX_uint64) const;
-
-	/**
-	 * Decompress into an owned composite buffer.
-	 *
-	 * @return An owned buffer containing the raw data, or null on error.
-	 */
+	[[nodiscard]] CORE_API FSharedBuffer Decompress() const;
 	[[nodiscard]] CORE_API FCompositeBuffer DecompressToComposite() const;
 
 	/** A null compressed buffer. */
@@ -172,10 +161,14 @@ namespace UE::CompressedBuffer::Private
 /** A reusable context for the compressed buffer decoder. */
 struct FDecoderContext
 {
-	/** Offset at which the compressed buffer begins, otherwise MAX_uint64. */
+	/** The offset in the source at which the compressed buffer begins, otherwise MAX_uint64. */
 	uint64 HeaderOffset = MAX_uint64;
-	/** Index of the block stored in RawBlock, otherwise MAX_uint64. */
-	uint64 RawBlockIndex = MAX_uint64;
+	/** The size of the header if known, otherwise 0. */
+	uint64 HeaderSize = 0;
+	/** The CRC-32 from the header, otherwise 0. */
+	uint32 HeaderCrc32 = 0;
+	/** Index of the block stored in RawBlock, otherwise MAX_uint32. */
+	uint32 RawBlockIndex = MAX_uint32;
 
 	/** A buffer used to store the header when HeaderOffset is not MAX_uint64. */
 	FUniqueBuffer Header;
@@ -220,10 +213,13 @@ public:
 	CORE_API void SetSource(FArchive& Archive);
 	CORE_API void SetSource(const FCompressedBuffer& Buffer);
 
+	/** Returns the size of the compressed data. Zero on error. */
 	[[nodiscard]] CORE_API uint64 GetCompressedSize();
+
+	/** Returns the size of the raw data. Zero on error. */
 	[[nodiscard]] CORE_API uint64 GetRawSize();
 
-	/** Returns the hash of the raw data. Zero on error or if this is null. */
+	/** Returns the hash of the raw data. Zero on error. */
 	[[nodiscard]] CORE_API FBlake3Hash GetRawHash();
 
 	/**
@@ -241,8 +237,10 @@ public:
 		uint64& OutBlockSize);
 
 	/**
-	 * Decompress into a memory view that is less or equal to the raw size.
+	 * Decompress into a memory view that is less than or equal to the available raw size.
 	 *
+	 * @param RawView     The view to write to. The size to read is equal to the view size.
+	 * @param RawOffset   The offset into the raw data from which to decompress.
 	 * @return True if the requested range was decompressed, otherwise false.
 	 */
 	[[nodiscard]] CORE_API bool TryDecompressTo(FMutableMemoryView RawView, uint64 RawOffset = 0);
@@ -250,41 +248,45 @@ public:
 	/**
 	 * Decompress into an owned buffer.
 	 *
-	 * @return An owned buffer containing the raw data, or null on error.
-	 */
-	[[nodiscard]] CORE_API FSharedBuffer Decompress(uint64 RawOffset = 0, uint64 RawSize = MAX_uint64);
-
-	/**
-	 * Decompress into an owned composite buffer.
+	 * RawOffset must be at most the raw buffer size. RawSize may be MAX_uint64 to read the whole
+	 * buffer from RawOffset, and must otherwise fit within the bounds of the buffer.
 	 *
+	 * @param RawOffset   The offset into the raw data from which to decompress.
+	 * @param RawSize     The size of the raw data to read from the offset.
 	 * @return An owned buffer containing the raw data, or null on error.
 	 */
+	[[nodiscard]] CORE_API FSharedBuffer Decompress(uint64 RawOffset = 0, uint64 RawSize = MAX_uint64);
 	[[nodiscard]] CORE_API FCompositeBuffer DecompressToComposite(uint64 RawOffset = 0, uint64 RawSize = MAX_uint64);
 
 private:
-	const UE::CompressedBuffer::Private::FHeader* TryReadHeader();
+	bool TryReadHeader(UE::CompressedBuffer::Private::FHeader& OutHeader, FMemoryView& OutHeaderView);
 
 	FArchive* SourceArchive = nullptr;
 	const FCompressedBuffer* SourceBuffer = nullptr;
 	UE::CompressedBuffer::Private::FDecoderContext Context;
 };
 
-/**
- * A stateful decoder that reuses temporary buffers between decompression calls.
- */
-class FCompressedBufferDecoder
+/** A type that sets the source of a reader upon construction and resets it upon destruction. */
+class FCompressedBufferReaderSourceScope
 {
-	UE_NONCOPYABLE(FCompressedBufferDecoder);
 public:
-	CORE_API FCompressedBufferDecoder();
-	CORE_API ~FCompressedBufferDecoder();
+	inline FCompressedBufferReaderSourceScope(FCompressedBufferReader& InReader, FArchive& InArchive)
+		: Reader(InReader)
+	{
+		Reader.SetSource(InArchive);
+	}
 
-	/**
-	 * Decompress into a memory view that is less or equal to GetRawSize()
-	 */
-	[[nodiscard]] CORE_API bool TryDecompressTo(const FCompressedBuffer& CompressedBuffer, FMutableMemoryView RawView, uint64 RawOffset = 0);
+	inline FCompressedBufferReaderSourceScope(FCompressedBufferReader& InReader, const FCompressedBuffer& InBuffer)
+		: Reader(InReader)
+	{
+		Reader.SetSource(InBuffer);
+	}
+
+	inline ~FCompressedBufferReaderSourceScope()
+	{
+		Reader.ResetSource();
+	}
 
 private:
-	class FImpl;
-	TUniquePtr<FImpl> Impl;
+	FCompressedBufferReader& Reader;
 };
diff --git a/Engine/Source/Runtime/CoreUObject/Private/Serialization/EditorDerivedData.cpp b/Engine/Source/Runtime/CoreUObject/Private/Serialization/EditorDerivedData.cpp
--- a/Engine/Source/Runtime/CoreUObject/Private/Serialization/EditorDerivedData.cpp
+++ b/Engine/Source/Runtime/CoreUObject/Private/Serialization/EditorDerivedData.cpp
@@ -87,9 +87,9 @@ void FEditorDerivedDataCompressedBuffer::Read(IoStore::FDerivedDataIoRequest Req
 	const uint64 RequestSize = FMath::Min(Request.GetSize(), RequestOffset <= DataSize ? DataSize - RequestOffset : 0);
 	if (RequestSize)
 	{
-		if (const FSharedBuffer Buffer = Data.Decompress())
+		const FMutableMemoryView View = Request.CreateBuffer(RequestSize);
+		if (FCompressedBufferReader(Data).TryDecompressTo(View, RequestOffset))
 		{
-			Request.CreateBuffer(RequestSize).CopyFrom(Buffer.GetView().Mid(RequestOffset, RequestSize));
 			Request.SetComplete();
 		}
 		else
diff --git a/Engine/Source/Runtime/CoreUObject/Private/Serialization/VirtualizedBulkData.cpp b/Engine/Source/Runtime/CoreUObject/Private/Serialization/VirtualizedBulkData.cpp
--- a/Engine/Source/Runtime/CoreUObject/Private/Serialization/VirtualizedBulkData.cpp
+++ b/Engine/Source/Runtime/CoreUObject/Private/Serialization/VirtualizedBulkData.cpp
@@ -190,7 +190,8 @@ FCompressionSettings::FCompressionSettings(const FCompressedBuffer& Buffer)
 {
 	// Note that if the buffer is using a non-oodle format we consider it
 	// as not set.
-	if (!Buffer.TryGetCompressParameters(Compressor, CompressionLevel))
+	uint64 BlockSize;
+	if (!Buffer.TryGetCompressParameters(Compressor, CompressionLevel, BlockSize))
 	{
 		Reset();
 	}
@@ -877,10 +878,10 @@ FCompressedBuffer FVirtualizedUntypedBulkData::LoadFromDisk() const
 			FCompressedBuffer AssetBuffer = LoadFromPackageFile();
 
 			FPayloadId SidecarId(SidecarBuffer.Decompress());
-			FPayloadId AssetId(SidecarBuffer.Decompress());
+			FPayloadId AssetId(AssetBuffer.Decompress());
 
 			UE_CLOG(SidecarId != PayloadContentId, LogVirtualization, Error, TEXT("Sidecar content did not hash correctly! Found '%s' Expected '%s'"), *SidecarId.ToString(), *PayloadContentId.ToString());
-			UE_CLOG(AssetId != PayloadContentId, LogVirtualization, Error, TEXT("Asset content did not hash correctly! Found '%s' Expected '%s'"), *SidecarId.ToString(), *PayloadContentId.ToString())
+			UE_CLOG(AssetId != PayloadContentId, LogVirtualization, Error, TEXT("Asset content did not hash correctly! Found '%s' Expected '%s'"), *AssetId.ToString(), *PayloadContentId.ToString())
 
 			return SidecarBuffer;
 		}
diff --git a/Engine/Source/Runtime/StorageServerClient/Private/StorageServerConnection.cpp b/Engine/Source/Runtime/StorageServerClient/Private/StorageServerConnection.cpp
--- a/Engine/Source/Runtime/StorageServerClient/Private/StorageServerConnection.cpp
+++ b/Engine/Source/Runtime/StorageServerClient/Private/StorageServerConnection.cpp
@@ -8,6 +8,7 @@
 #include "Misc/Paths.h"
 #include "IO/IoDispatcher.h"
 #include "Misc/StringBuilder.h"
+#include "Misc/ScopeExit.h"
 #include "Misc/ScopeLock.h"
 #include "Serialization/CompactBinary.h"
 #include "Serialization/CompactBinarySerialization.h"
@@ -51,11 +52,13 @@ static uint64 GetCompressedOffset(const FCompressedBuffer& Buffer, uint64 RawOff
 {
 	if (RawOffset > 0)
 	{
-		uint32 BlockSize = 0;
-		const bool bOk = Buffer.TryGetBlockSize(BlockSize);
+		uint64 BlockSize = 0;
+		ECompressedBufferCompressor Compressor;
+		ECompressedBufferCompressionLevel CompressionLevel;
+		const bool bOk = Buffer.TryGetCompressParameters(Compressor, CompressionLevel, BlockSize);
 		check(bOk);
 
-		return BlockSize > 0 ? RawOffset % uint64(BlockSize) : 0;
+		return BlockSize > 0 ? RawOffset % BlockSize : 0;
 	}
 
 	return 0;
@@ -342,10 +345,11 @@ int64 FStorageServerResponse::SerializeChunk(FStorageServerSerializationContext&
 
 		if (FCompressedBuffer Compressed = FCompressedBuffer::FromCompressed(FSharedBuffer::MakeView(Context.CompressedBuffer.GetData(), ContentLength)))
 		{
+			FCompressedBufferReaderSourceScope Source(Context.Decoder, Compressed);
 			const uint64 ChunkSize = FMath::Min(Compressed.GetRawSize(), RawSize);
 			OutChunk = TargetVa ? FIoBuffer(FIoBuffer::Wrap, TargetVa, ChunkSize) : FIoBuffer(ChunkSize);
 			const uint64 CompressedOffset = GetCompressedOffset(Compressed, RawOffset);
-			if (Context.Decoder.TryDecompressTo(Compressed, OutChunk.GetMutableView(), CompressedOffset))
+			if (Context.Decoder.TryDecompressTo(OutChunk.GetMutableView(), CompressedOffset))
 			{
 				return ChunkSize;
 			}
@@ -381,7 +385,7 @@ int64 FStorageServerResponse::SerializeChunkTo(FMutableMemoryView Memory, uint64
 	{
 		TRACE_CPUPROFILER_EVENT_SCOPE(ZenClient::SerializeCompressedChunk);
 		
-		TArray<uint8> CompressedBuffer;
+		TArray64<uint8> CompressedBuffer;
 		CompressedBuffer.Reset(ContentLength);
 		Serialize(CompressedBuffer.GetData(), ContentLength);
 		
@@ -389,7 +393,7 @@ int64 FStorageServerResponse::SerializeChunkTo(FMutableMemoryView Memory, uint64
 		{
 			FMutableMemoryView Dst = Memory.Left(FMath::Min(Memory.GetSize(), Compressed.GetRawSize()));
 			const uint64 CompressedOffset = GetCompressedOffset(Compressed, RawOffset);
-			if (Compressed.TryDecompressTo(Dst, CompressedOffset))
+			if (FCompressedBufferReader(Compressed).TryDecompressTo(Dst, CompressedOffset))
 			{
 				return Dst.GetSize();
 			}
diff --git a/Engine/Source/Runtime/StorageServerClient/Public/StorageServerConnection.h b/Engine/Source/Runtime/StorageServerClient/Public/StorageServerConnection.h
--- a/Engine/Source/Runtime/StorageServerClient/Public/StorageServerConnection.h
+++ b/Engine/Source/Runtime/StorageServerClient/Public/StorageServerConnection.h
@@ -63,8 +63,8 @@ inline EStorageServerContentType GetMimeType(const FAnsiStringView& ContentType)
 
 struct FStorageServerSerializationContext
 {
-	TArray<uint8> CompressedBuffer;
-	FCompressedBufferDecoder Decoder;
+	TArray64<uint8> CompressedBuffer;
+	FCompressedBufferReader Decoder;
 };
 
 class FStorageServerRequest
