CompressedBuffer: Added const to params on virtual functions now that it is not an error on certain platforms

#rnx

diff --git a/Engine/Source/Runtime/Core/Private/Compression/CompressedBuffer.cpp b/Engine/Source/Runtime/Core/Private/Compression/CompressedBuffer.cpp
--- a/Engine/Source/Runtime/Core/Private/Compression/CompressedBuffer.cpp
+++ b/Engine/Source/Runtime/Core/Private/Compression/CompressedBuffer.cpp
@@ -85,7 +85,7 @@ struct FHeader
 	 *
 	 * @param HeaderView   View of the header to write, including any method-specific header data.
 	 */
-	void Write(FMutableMemoryView HeaderView) const
+	void Write(const FMutableMemoryView HeaderView) const
 	{
 		FHeader Header = *this;
 		Header.ByteSwap();
@@ -97,13 +97,13 @@ struct FHeader
 	}
 
 	/** Calculate the CRC-32 from a view of a header including any method-specific header data. */
-	static uint32 CalculateCrc32(FMemoryView HeaderView)
+	static uint32 CalculateCrc32(const FMemoryView HeaderView)
 	{
 		uint32 Crc32 = 0;
 		constexpr uint64 MethodOffset = STRUCT_OFFSET(FHeader, Method);
 		for (FMemoryView View = HeaderView + MethodOffset; const uint64 ViewSize = View.GetSize();)
 		{
-			const int32 Size = static_cast<int32>(FMath::Min<uint64>(ViewSize, MAX_int32));
+			const int32 Size = int32(FMath::Min<uint64>(ViewSize, MAX_int32));
 			Crc32 = FCrc::MemCrc32(View.GetData(), Size, Crc32);
 			View += Size;
 		}
@@ -173,7 +173,7 @@ public:
 class FNoneEncoder final : public FEncoder
 {
 public:
-	FCompositeBuffer Compress(const FCompositeBuffer& RawData, uint64 BlockSize) const final
+	FCompositeBuffer Compress(const FCompositeBuffer& RawData, const uint64 BlockSize) const final
 	{
 		FHeader Header;
 		Header.Method = EMethod::None;
@@ -248,7 +248,7 @@ protected:
 	virtual bool CompressBlock(FMutableMemoryView& CompressedData, FMemoryView RawData) const = 0;
 
 private:
-	uint64 GetCompressedBlocksBound(uint64 BlockCount, uint64 BlockSize, uint64 RawSize) const
+	uint64 GetCompressedBlocksBound(const uint64 BlockCount, const uint64 BlockSize, const uint64 RawSize) const
 	{
 		switch (BlockCount)
 		{
@@ -277,7 +277,7 @@ FCompositeBuffer FBlockEncoder::Compress(const FCompositeBuffer& RawData, const
 
 	// Compress the raw data in blocks and store the raw data for incompressible blocks.
 	TArray64<uint32> CompressedBlockSizes;
-	CompressedBlockSizes.Reserve(static_cast<uint32>(BlockCount));
+	CompressedBlockSizes.Reserve(uint32(BlockCount));
 	uint64 CompressedSize = 0;
 	{
 		FUniqueBuffer RawBlockCopy;
@@ -305,7 +305,7 @@ FCompositeBuffer FBlockEncoder::Compress(const FCompositeBuffer& RawData, const
 				CompressedBlocksView += CompressedBlockSize;
 			}
 
-			CompressedBlockSizes.Add(static_cast<uint32>(CompressedBlockSize));
+			CompressedBlockSizes.Add(uint32(CompressedBlockSize));
 			CompressedSize += CompressedBlockSize;
 			RawOffset += RawBlockSize;
 		}
@@ -326,8 +326,8 @@ FCompositeBuffer FBlockEncoder::Compress(const FCompositeBuffer& RawData, const
 	Header.Method = GetMethod();
 	Header.Compressor = GetCompressor();
 	Header.CompressionLevel = GetCompressionLevel();
-	Header.BlockSizeExponent = static_cast<uint8>(FMath::FloorLog2_64(BlockSize));
-	Header.BlockCount = static_cast<uint32>(BlockCount);
+	Header.BlockSizeExponent = uint8(FMath::FloorLog2_64(BlockSize));
+	Header.BlockCount = uint32(BlockCount);
 	Header.TotalRawSize = RawSize;
 	Header.TotalCompressedSize = sizeof(FHeader) + MetaSize + CompressedSize;
 	Header.RawHash = RawHash.Finalize();
@@ -461,21 +461,21 @@ public:
 
 protected:
 	EMethod GetMethod() const final { return EMethod::Oodle; }
-	uint8 GetCompressor() const final { return static_cast<uint8>(Compressor); }
-	uint8 GetCompressionLevel() const final { return static_cast<uint8>(CompressionLevel); }
+	uint8 GetCompressor() const final { return uint8(Compressor); }
+	uint8 GetCompressionLevel() const final { return uint8(CompressionLevel); }
 
 	uint64 CompressBlockBound(uint64 RawSize) const final
 	{
-		return static_cast<uint64>(FOodleDataCompression::CompressedBufferSizeNeeded(static_cast<int64>(RawSize)));
+		return uint64(FOodleDataCompression::CompressedBufferSizeNeeded(int64(RawSize)));
 	}
 
-	bool CompressBlock(FMutableMemoryView& CompressedData, FMemoryView RawData) const final
+	bool CompressBlock(FMutableMemoryView& CompressedData, const FMemoryView RawData) const final
 	{
 		const int64 Size = FOodleDataCompression::Compress(
-			CompressedData.GetData(), static_cast<uint64>(CompressedData.GetSize()),
-			RawData.GetData(), static_cast<int64>(RawData.GetSize()),
+			CompressedData.GetData(), uint64(CompressedData.GetSize()),
+			RawData.GetData(), int64(RawData.GetSize()),
 			Compressor, CompressionLevel);
-		CompressedData.LeftInline(static_cast<uint64>(Size));
+		CompressedData.LeftInline(uint64(Size));
 		return Size > 0;
 	}
 
@@ -487,11 +487,11 @@ private:
 class FOodleDecoder final : public FBlockDecoder
 {
 protected:
-	bool DecompressBlock(FMutableMemoryView RawData, FMemoryView CompressedData) const final
+	bool DecompressBlock(const FMutableMemoryView RawData, const FMemoryView CompressedData) const final
 	{
 		return FOodleDataCompression::Decompress(
-			RawData.GetData(), static_cast<int64>(RawData.GetSize()),
-			CompressedData.GetData(), static_cast<int64>(CompressedData.GetSize()));
+			RawData.GetData(), int64(RawData.GetSize()),
+			CompressedData.GetData(), int64(CompressedData.GetSize()));
 	}
 };
 
@@ -508,19 +508,19 @@ protected:
 	{
 		if (RawSize <= LZ4_MAX_INPUT_SIZE)
 		{
-			return static_cast<uint64>(LZ4_compressBound(static_cast<int>(RawSize)));
+			return uint64(LZ4_compressBound(int(RawSize)));
 		}
 		return 0;
 	}
 
-	bool CompressBlock(FMutableMemoryView& CompressedData, FMemoryView RawData) const final
+	bool CompressBlock(FMutableMemoryView& CompressedData, const FMemoryView RawData) const final
 	{
 		if (RawData.GetSize() <= LZ4_MAX_INPUT_SIZE)
 		{
 			const int Size = LZ4_compress_default(
 				static_cast<const char*>(RawData.GetData()), static_cast<char*>(CompressedData.GetData()),
-				static_cast<int>(RawData.GetSize()), static_cast<int>(FMath::Min<uint64>(CompressedData.GetSize(), MAX_int32)));
-			CompressedData.LeftInline(static_cast<uint64>(Size));
+				int(RawData.GetSize()), int(FMath::Min<uint64>(CompressedData.GetSize(), MAX_int32)));
+			CompressedData.LeftInline(uint64(Size));
 			return Size > 0;
 		}
 		return false;
@@ -530,16 +530,16 @@ protected:
 class FLZ4Decoder final : public FBlockDecoder
 {
 protected:
-	bool DecompressBlock(FMutableMemoryView RawData, FMemoryView CompressedData) const final
+	bool DecompressBlock(const FMutableMemoryView RawData, const FMemoryView CompressedData) const final
 	{
 		if (CompressedData.GetSize() <= MAX_int32)
 		{
 			const int Size = LZ4_decompress_safe(
 				static_cast<const char*>(CompressedData.GetData()),
 				static_cast<char*>(RawData.GetData()),
-				static_cast<int>(CompressedData.GetSize()),
-				static_cast<int>(FMath::Min<uint64>(RawData.GetSize(), LZ4_MAX_INPUT_SIZE)));
-			return static_cast<uint64>(Size) == RawData.GetSize();
+				int(CompressedData.GetSize()),
+				int(FMath::Min<uint64>(RawData.GetSize(), LZ4_MAX_INPUT_SIZE)));
+			return uint64(Size) == RawData.GetSize();
 		}
 		return false;
 	}
@@ -547,7 +547,7 @@ protected:
 
 ///////////////////////////////////////////////////////////////////////////////////////////////////////////////////////
 
-static const FDecoder* GetDecoder(EMethod Method)
+static const FDecoder* GetDecoder(const EMethod Method)
 {
 	static FNoneDecoder None;
 	static FOodleDecoder Oodle;
@@ -704,7 +704,7 @@ public:
 		return !Archive.IsError();
 	}
 
-	FMemoryView ReadOrView(uint64 Offset, uint64 Size, FDecoderContext& Context) const final
+	FMemoryView ReadOrView(const uint64 Offset, const uint64 Size, FDecoderContext& Context) const final
 	{
 		if (Context.CompressedBlock.GetSize() < Size)
 		{
@@ -714,7 +714,7 @@ public:
 		return Read(Offset, View) ? View : FMemoryView();
 	}
 
-	FCompositeBuffer ReadToComposite(uint64 Offset, uint64 Size) const final
+	FCompositeBuffer ReadToComposite(const uint64 Offset, const uint64 Size) const final
 	{
 		FUniqueBuffer Buffer = FUniqueBuffer::Alloc(Size);
 		if (Read(Offset, Buffer))
@@ -737,7 +737,7 @@ public:
 	{
 	}
 
-	bool Read(uint64 Offset, FMutableMemoryView Data) const final
+	bool Read(const uint64 Offset, const FMutableMemoryView Data) const final
 	{
 		if (Offset + Data.GetSize() <= Buffer.GetSize())
 		{
@@ -747,7 +747,7 @@ public:
 		return false;
 	}
 
-	FMemoryView ReadOrView(uint64 Offset, uint64 Size, FDecoderContext& Context) const final
+	FMemoryView ReadOrView(const uint64 Offset, const uint64 Size, FDecoderContext& Context) const final
 	{
 		return Buffer.ViewOrCopyRange(Offset, Size, Context.CompressedBlock, [](uint64 BufferSize) -> FUniqueBuffer
 		{
@@ -755,7 +755,7 @@ public:
 		});
 	}
 
-	FCompositeBuffer ReadToComposite(uint64 Offset, uint64 Size) const final
+	FCompositeBuffer ReadToComposite(const uint64 Offset, const uint64 Size) const final
 	{
 		return Buffer.Mid(Offset, Size).MakeOwned();
 	}
@@ -780,8 +780,8 @@ FCompressedBuffer FCompressedBuffer::Compress(const FSharedBuffer& RawData)
 
 FCompressedBuffer FCompressedBuffer::Compress(
 	const FCompositeBuffer& RawData,
-	ECompressedBufferCompressor Compressor,
-	ECompressedBufferCompressionLevel CompressionLevel,
+	const ECompressedBufferCompressor Compressor,
+	const ECompressedBufferCompressionLevel CompressionLevel,
 	uint64 BlockSize)
 {
 	using namespace UE::CompressedBuffer::Private;
@@ -805,9 +805,9 @@ FCompressedBuffer FCompressedBuffer::Compress(
 
 FCompressedBuffer FCompressedBuffer::Compress(
 	const FSharedBuffer& RawData,
-	ECompressedBufferCompressor Compressor,
-	ECompressedBufferCompressionLevel CompressionLevel,
-	uint64 BlockSize)
+	const ECompressedBufferCompressor Compressor,
+	const ECompressedBufferCompressionLevel CompressionLevel,
+	const uint64 BlockSize)
 {
 	return Compress(FCompositeBuffer(RawData), Compressor, CompressionLevel, BlockSize);
 }
@@ -855,7 +855,7 @@ FCompressedBuffer FCompressedBuffer::Load(FArchive& Ar)
 		FUniqueBuffer MutableBuffer = FUniqueBuffer::Alloc(Header.TotalCompressedSize);
 		Header.ByteSwap();
 		const FMutableMemoryView MutableView = MutableBuffer.GetView().CopyFrom(MakeMemoryView(&Header, &Header + 1));
-		Ar.Serialize(MutableView.GetData(), static_cast<int64>(MutableView.GetSize()));
+		Ar.Serialize(MutableView.GetData(), int64(MutableView.GetSize()));
 		Local.CompressedData = ValidBufferOrEmpty(MutableBuffer.MoveToShared());
 	}
 	if (Local.IsNull())
