cosmetic: replaced FThread constructor parameter `bool bIsForkable` by more explicit enum `EForkable`

diff --git a/Engine/Source/Runtime/Core/Private/Async/Fundamental/ReserveScheduler.cpp b/Engine/Source/Runtime/Core/Private/Async/Fundamental/ReserveScheduler.cpp
--- a/Engine/Source/Runtime/Core/Private/Async/Fundamental/ReserveScheduler.cpp
+++ b/Engine/Source/Runtime/Core/Private/Async/Fundamental/ReserveScheduler.cpp
@@ -8,7 +8,7 @@ namespace LowLevelTasks
 
 FReserveScheduler FReserveScheduler::Singleton;
 
-TUniquePtr<FThread> FReserveScheduler::CreateWorker(bool bIsForkable, FSchedulerTls::FLocalQueueType* WorkerLocalQueue, EThreadPriority Priority)
+TUniquePtr<FThread> FReserveScheduler::CreateWorker(FThread::EForkable IsForkable, FSchedulerTls::FLocalQueueType* WorkerLocalQueue, EThreadPriority Priority)
 {
 	FYieldedWork* ReserveEvent = new FYieldedWork;
 	ReserveEvents.Emplace(ReserveEvent);
@@ -41,7 +41,7 @@ TUniquePtr<FThread> FReserveScheduler::CreateWorker(bool bIsForkable, FScheduler
 			FSchedulerTls::WorkerType = EWorkerType::None;
 			FSchedulerTls::ActiveScheduler = nullptr;
 			FSchedulerTls::LocalQueue = nullptr;
-		}, 0, Priority, FThreadAffinity{ FPlatformAffinity::GetTaskGraphThreadMask(), 0 }, bIsForkable
+		}, 0, Priority, FThreadAffinity{ FPlatformAffinity::GetTaskGraphThreadMask(), 0 }, IsForkable
 	);
 }
 
@@ -58,7 +58,7 @@ bool FReserveScheduler::DoReserveWorkUntil(FConditional&& Condition)
 	return false;
 }
 
-void FReserveScheduler::StartWorkers(FScheduler& MainScheduler, uint32 NumWorkers, bool bIsForkable, EThreadPriority WorkerPriority)
+void FReserveScheduler::StartWorkers(FScheduler& MainScheduler, uint32 NumWorkers, FThread::EForkable IsForkable, EThreadPriority WorkerPriority)
 {
 	if (NumWorkers == 0)
 	{
@@ -78,7 +78,7 @@ void FReserveScheduler::StartWorkers(FScheduler& MainScheduler, uint32 NumWorker
 		for (uint32 WorkerId = 0; WorkerId < NumWorkers; ++WorkerId)
 		{
 			WorkerLocalQueues.Emplace(MainScheduler.GetQueueRegistry(), ELocalQueueType::EBusyWait, nullptr);
-			WorkerThreads.Add(CreateWorker(bIsForkable, &WorkerLocalQueues.Last(), WorkerPriority));
+			WorkerThreads.Add(CreateWorker(IsForkable, &WorkerLocalQueues.Last(), WorkerPriority));
 		}
 		UE::Trace::ThreadGroupEnd();
 	}
diff --git a/Engine/Source/Runtime/Core/Private/Async/Fundamental/Scheduler.cpp b/Engine/Source/Runtime/Core/Private/Async/Fundamental/Scheduler.cpp
--- a/Engine/Source/Runtime/Core/Private/Async/Fundamental/Scheduler.cpp
+++ b/Engine/Source/Runtime/Core/Private/Async/Fundamental/Scheduler.cpp
@@ -40,7 +40,7 @@ namespace LowLevelTasks
 		}
 	}
 
-	TUniquePtr<FThread> FScheduler::CreateWorker(bool bPermitBackgroundWork, bool bIsForkable, FSleepEvent* ExternalWorkerEvent, FSchedulerTls::FLocalQueueType* ExternalWorkerLocalQueue, EThreadPriority Priority, uint64 InAffinity)
+	TUniquePtr<FThread> FScheduler::CreateWorker(bool bPermitBackgroundWork, FThread::EForkable IsForkable, FSleepEvent* ExternalWorkerEvent, FSchedulerTls::FLocalQueueType* ExternalWorkerLocalQueue, EThreadPriority Priority, uint64 InAffinity)
 	{
 		uint32 WorkerId = NextWorkerId++;
 		const uint32 WaitTimes[8] = { 719, 991, 1361, 1237, 1597, 953, 587, 1439 };
@@ -84,11 +84,11 @@ namespace LowLevelTasks
 			[this, ExternalWorkerEvent, ExternalWorkerLocalQueue, WaitTime, bPermitBackgroundWork]
 			{ 
 				WorkerMain(ExternalWorkerEvent, ExternalWorkerLocalQueue, WaitTime, bPermitBackgroundWork);
-			}, 0, Priority, FThreadAffinity{ ThreadAffinityMask & ProcessorGroups.ThreadAffinities[CpuGroup], CpuGroup }, bIsForkable
+			}, 0, Priority, FThreadAffinity{ ThreadAffinityMask & ProcessorGroups.ThreadAffinities[CpuGroup], CpuGroup }, IsForkable
 		);
 	}
 
-	void FScheduler::StartWorkers(uint32 NumForegroundWorkers, uint32 NumBackgroundWorkers, bool bIsForkable, EThreadPriority InWorkerPriority,  EThreadPriority InBackgroundPriority, uint64 InWorkerAffinity, uint64 InBackgroundAffinity)
+	void FScheduler::StartWorkers(uint32 NumForegroundWorkers, uint32 NumBackgroundWorkers, FThread::EForkable IsForkable, EThreadPriority InWorkerPriority,  EThreadPriority InBackgroundPriority, uint64 InWorkerAffinity, uint64 InBackgroundAffinity)
 	{
 		if (NumForegroundWorkers == 0 && NumBackgroundWorkers == 0)
 		{
@@ -127,7 +127,7 @@ namespace LowLevelTasks
 			{
 				WorkerEvents.Emplace();
 				WorkerLocalQueues.Emplace(QueueRegistry, ELocalQueueType::EForeground, &WorkerEvents.Last());
-				WorkerThreads.Add(CreateWorker(false, bIsForkable, &WorkerEvents.Last(), &WorkerLocalQueues.Last(), WorkerPriority, WorkerAffinity));
+				WorkerThreads.Add(CreateWorker(false, IsForkable, &WorkerEvents.Last(), &WorkerLocalQueues.Last(), WorkerPriority, WorkerAffinity));
 			}
 			UE::Trace::ThreadGroupEnd();
 			UE::Trace::ThreadGroupBegin(TEXT("Background Workers"));
@@ -135,7 +135,7 @@ namespace LowLevelTasks
 			{
 				WorkerEvents.Emplace();
 				WorkerLocalQueues.Emplace(QueueRegistry, ELocalQueueType::EBackground, &WorkerEvents.Last());
-				WorkerThreads.Add(CreateWorker(true, bIsForkable, &WorkerEvents.Last(), &WorkerLocalQueues.Last(), BackgroundPriority, BackgroundAffinity));
+				WorkerThreads.Add(CreateWorker(true, IsForkable, &WorkerEvents.Last(), &WorkerLocalQueues.Last(), BackgroundPriority, BackgroundAffinity));
 			}
 			UE::Trace::ThreadGroupEnd();
 		}
@@ -168,12 +168,12 @@ namespace LowLevelTasks
 			}
 		}
 	}
-	void FScheduler::RestartWorkers(uint32 NumForegroundWorkers, uint32 NumBackgroundWorkers, bool bIsForkable, EThreadPriority InWorkerPriority, EThreadPriority InBackgroundPriority, uint64 InWorkerAffinity, uint64 InBackgroundAffinity)
+	void FScheduler::RestartWorkers(uint32 NumForegroundWorkers, uint32 NumBackgroundWorkers, FThread::EForkable IsForkable, EThreadPriority InWorkerPriority, EThreadPriority InBackgroundPriority, uint64 InWorkerAffinity, uint64 InBackgroundAffinity)
 	{
 		FScopeLock Lock(&WorkerThreadsCS);
 		TemporaryShutdown.store(true, std::memory_order_release);
 		StopWorkers(false);
-		StartWorkers(NumForegroundWorkers, NumBackgroundWorkers, bIsForkable, InWorkerPriority, InBackgroundPriority, InWorkerAffinity, InBackgroundAffinity);
+		StartWorkers(NumForegroundWorkers, NumBackgroundWorkers, IsForkable, InWorkerPriority, InBackgroundPriority, InWorkerAffinity, InBackgroundAffinity);
 		TemporaryShutdown.store(false, std::memory_order_release);
 	}
 
diff --git a/Engine/Source/Runtime/Core/Private/Async/TaskGraph.cpp b/Engine/Source/Runtime/Core/Private/Async/TaskGraph.cpp
--- a/Engine/Source/Runtime/Core/Private/Async/TaskGraph.cpp
+++ b/Engine/Source/Runtime/Core/Private/Async/TaskGraph.cpp
@@ -1898,7 +1898,7 @@ public:
 			int32 NumBackgroundWorkers = FMath::Max(1, NumWorkerThreads - FMath::Min<int>(GNumForegroundWorkers, NumWorkerThreads));
 			int32 NumForegroundWorkers =  FMath::Max(1, NumWorkerThreads - NumBackgroundWorkers);
 
-			LowLevelTasks::FScheduler::Get().StartWorkers(NumForegroundWorkers, NumBackgroundWorkers, FForkProcessHelper::IsForkedMultithreadInstance(), FPlatformAffinity::GetTaskThreadPriority(), FPlatformAffinity::GetTaskBPThreadPriority());
+			LowLevelTasks::FScheduler::Get().StartWorkers(NumForegroundWorkers, NumBackgroundWorkers, FForkProcessHelper::IsForkedMultithreadInstance() ? FThread::Forkable : FThread::NonForkable, FPlatformAffinity::GetTaskThreadPriority(), FPlatformAffinity::GetTaskBPThreadPriority());
 
 			check(IsInGameThread()); // otherwise we can have a race on starting reserve workers below
 			if (GConfig == nullptr)
@@ -1990,12 +1990,12 @@ public:
 			int32 NumWorkers =  FMath::Max(1, NumWorkerThreads - NumBackgroundWorkers);
 
 			LowLevelTasks::FScheduler::Get().StopWorkers();
-			LowLevelTasks::FScheduler::Get().StartWorkers(NumWorkers, NumBackgroundWorkers, FForkProcessHelper::IsForkedMultithreadInstance(), Pri, FPlatformAffinity::GetTaskBPThreadPriority());
+			LowLevelTasks::FScheduler::Get().StartWorkers(NumWorkers, NumBackgroundWorkers, FForkProcessHelper::IsForkedMultithreadInstance() ? FThread::Forkable : FThread::NonForkable, Pri, FPlatformAffinity::GetTaskBPThreadPriority());
 
 			if (bReserveWorkersEnabled)
 			{
 				LowLevelTasks::FReserveScheduler::Get().StopWorkers();
-				LowLevelTasks::FReserveScheduler::Get().StartWorkers(LowLevelTasks::FScheduler::Get(), NumWorkers + NumBackgroundWorkers, FForkProcessHelper::IsForkedMultithreadInstance(), FPlatformAffinity::GetTaskBPThreadPriority());
+				LowLevelTasks::FReserveScheduler::Get().StartWorkers(LowLevelTasks::FScheduler::Get(), NumWorkers + NumBackgroundWorkers, FForkProcessHelper::IsForkedMultithreadInstance() ? FThread::Forkable : FThread::NonForkable, FPlatformAffinity::GetTaskBPThreadPriority());
 			}
 		}
 	}
@@ -2348,7 +2348,7 @@ private:
 		if (bEnableReserveWorkers)
 		{
 			bReserveWorkersEnabled = true;
-			LowLevelTasks::FReserveScheduler::Get().StartWorkers(LowLevelTasks::FScheduler::Get(), NumForegroundWorkers + NumBackgroundWorkers, FForkProcessHelper::IsForkedMultithreadInstance(), FPlatformAffinity::GetTaskBPThreadPriority());
+			LowLevelTasks::FReserveScheduler::Get().StartWorkers(LowLevelTasks::FScheduler::Get(), NumForegroundWorkers + NumBackgroundWorkers, FForkProcessHelper::IsForkedMultithreadInstance() ? FThread::Forkable : FThread::NonForkable, FPlatformAffinity::GetTaskBPThreadPriority());
 		}
 	}
 };
diff --git a/Engine/Source/Runtime/Core/Private/HAL/Thread.cpp b/Engine/Source/Runtime/Core/Private/HAL/Thread.cpp
--- a/Engine/Source/Runtime/Core/Private/HAL/Thread.cpp
+++ b/Engine/Source/Runtime/Core/Private/HAL/Thread.cpp
@@ -21,11 +21,11 @@ public:
 		uint32 StackSize,
 		EThreadPriority ThreadPriority,
 		FThreadAffinity ThreadAffinity,
-		bool bIsForkable
+		FThread::EForkable IsForkable
 	)
 		: ThreadFunction(MoveTemp(InThreadFunction))
 		, SingleThreadTickFunction(MoveTemp(InSingleThreadTickFunction))
-		, RunnableThread(bIsForkable ? FForkProcessHelper::CreateForkableThread(this, ThreadName, StackSize, ThreadPriority, ThreadAffinity.ThreadAffinityMask)
+		, RunnableThread(IsForkable == FThread::Forkable ? FForkProcessHelper::CreateForkableThread(this, ThreadName, StackSize, ThreadPriority, ThreadAffinity.ThreadAffinityMask)
 			: FRunnableThread::Create(this, ThreadName, StackSize, ThreadPriority, ThreadAffinity.ThreadAffinityMask))
 	{
 		RunnableThread->SetThreadAffinity(ThreadAffinity);
@@ -112,9 +112,9 @@ FThread::FThread(
 	uint32 StackSize/* = 0*/,
 	EThreadPriority ThreadPriority/* = TPri_Normal*/,
 	FThreadAffinity ThreadAffinity/* = FThreadAffinity()*/,
-	bool bIsForkable/* = false*/
+	EForkable IsForkable/* = NonForkable*/
 )
-	: Impl(MakeShared<FThreadImpl, ESPMode::ThreadSafe>(ThreadName, MoveTemp(ThreadFunction), [] {}, StackSize, ThreadPriority, ThreadAffinity, bIsForkable))
+	: Impl(MakeShared<FThreadImpl, ESPMode::ThreadSafe>(ThreadName, MoveTemp(ThreadFunction), [] {}, StackSize, ThreadPriority, ThreadAffinity, IsForkable))
 {
 	Impl->Initialize(Impl);
 }
@@ -126,9 +126,9 @@ FThread::FThread(
 	uint32 StackSize/* = 0*/,
 	EThreadPriority ThreadPriority/* = TPri_Normal*/,
 	FThreadAffinity ThreadAffinity/* = FThreadAffinity()*/,
-	bool bIsForkable/* = false*/
+	EForkable IsForkable/* = NonForkable*/
 )
-	: Impl(MakeShared<FThreadImpl, ESPMode::ThreadSafe>(ThreadName, MoveTemp(ThreadFunction), MoveTemp(SingleThreadTickFunction), StackSize, ThreadPriority, ThreadAffinity, bIsForkable))
+	: Impl(MakeShared<FThreadImpl, ESPMode::ThreadSafe>(ThreadName, MoveTemp(ThreadFunction), MoveTemp(SingleThreadTickFunction), StackSize, ThreadPriority, ThreadAffinity, IsForkable))
 {
 	Impl->Initialize(Impl);
 }
diff --git a/Engine/Source/Runtime/Core/Public/Async/Fundamental/ReserveScheduler.h b/Engine/Source/Runtime/Core/Public/Async/Fundamental/ReserveScheduler.h
--- a/Engine/Source/Runtime/Core/Public/Async/Fundamental/ReserveScheduler.h
+++ b/Engine/Source/Runtime/Core/Public/Async/Fundamental/ReserveScheduler.h
@@ -30,14 +30,14 @@ namespace LowLevelTasks
 		FORCEINLINE_DEBUGGABLE static FReserveScheduler& Get();
 
 		//start number of reserve workers where 0 is the system default
-		CORE_API void StartWorkers(FScheduler& MainScheduler, uint32 ReserveWorkers = 0, bool bIsForkable = false, EThreadPriority WorkerPriority = EThreadPriority::TPri_Normal);
+		CORE_API void StartWorkers(FScheduler& MainScheduler, uint32 ReserveWorkers = 0, FThread::EForkable IsForkable = FThread::NonForkable, EThreadPriority WorkerPriority = EThreadPriority::TPri_Normal);
 		CORE_API void StopWorkers();
 
 		//tries to yield this thread using the YieldEvent and do busywork on a reserve worker.
 		CORE_API bool DoReserveWorkUntil(FConditional&& Condition);
 
 	private: 
-		TUniquePtr<FThread> CreateWorker(bool bIsForkable = false, FSchedulerTls::FLocalQueueType* WorkerLocalQueue = nullptr, EThreadPriority Priority = EThreadPriority::TPri_Normal);
+		TUniquePtr<FThread> CreateWorker(FThread::EForkable IsForkable = FThread::NonForkable, FSchedulerTls::FLocalQueueType* WorkerLocalQueue = nullptr, EThreadPriority Priority = EThreadPriority::TPri_Normal);
 
 	private:
 		template<typename ElementType>
diff --git a/Engine/Source/Runtime/Core/Public/Async/Fundamental/Scheduler.h b/Engine/Source/Runtime/Core/Public/Async/Fundamental/Scheduler.h
--- a/Engine/Source/Runtime/Core/Public/Async/Fundamental/Scheduler.h
+++ b/Engine/Source/Runtime/Core/Public/Async/Fundamental/Scheduler.h
@@ -156,9 +156,9 @@ namespace LowLevelTasks
 		FORCEINLINE_DEBUGGABLE static FScheduler& Get();
 
 		//start number of workers where 0 is the system default
-		CORE_API void StartWorkers(uint32 NumForegroundWorkers = 0, uint32 NumBackgroundWorkers = 0, bool bIsForkable = false, EThreadPriority InWorkerPriority = EThreadPriority::TPri_Normal, EThreadPriority InBackgroundPriority = EThreadPriority::TPri_BelowNormal, uint64 InWorkerAffinity = 0, uint64 InBackgroundAffinity = 0);
+		CORE_API void StartWorkers(uint32 NumForegroundWorkers = 0, uint32 NumBackgroundWorkers = 0, FThread::EForkable IsForkable = FThread::NonForkable, EThreadPriority InWorkerPriority = EThreadPriority::TPri_Normal, EThreadPriority InBackgroundPriority = EThreadPriority::TPri_BelowNormal, uint64 InWorkerAffinity = 0, uint64 InBackgroundAffinity = 0);
 		CORE_API void StopWorkers(bool DrainGlobalQueue = true);
-		CORE_API void RestartWorkers(uint32 NumForegroundWorkers = 0, uint32 NumBackgroundWorkers = 0, bool bIsForkable = false, EThreadPriority WorkerPriority = EThreadPriority::TPri_Normal, EThreadPriority BackgroundPriority = EThreadPriority::TPri_BelowNormal, uint64 InWorkerAffinity = 0, uint64 InBackgroundAffinity = 0);
+		CORE_API void RestartWorkers(uint32 NumForegroundWorkers = 0, uint32 NumBackgroundWorkers = 0, FThread::EForkable IsForkable = FThread::NonForkable, EThreadPriority WorkerPriority = EThreadPriority::TPri_Normal, EThreadPriority BackgroundPriority = EThreadPriority::TPri_BelowNormal, uint64 InWorkerAffinity = 0, uint64 InBackgroundAffinity = 0);
 
 		//try to launch the task, the return value will specify if the task was in the ready state and has been launhced
 		inline bool TryLaunch(FTask& Task, EQueuePreference QueuePreference = EQueuePreference::DefaultPreference, bool bWakeUpWorker = true);	
@@ -190,7 +190,7 @@ namespace LowLevelTasks
 		~FScheduler();
 
 	private: 
-		TUniquePtr<FThread> CreateWorker(bool bPermitBackgroundWork = false, bool bIsForkable = false, FSleepEvent* ExternalWorkerEvent = nullptr, FSchedulerTls::FLocalQueueType* ExternalWorkerLocalQueue = nullptr, EThreadPriority Priority = EThreadPriority::TPri_Normal, uint64 InAffinity = 0);
+		TUniquePtr<FThread> CreateWorker(bool bPermitBackgroundWork = false, FThread::EForkable IsForkable = FThread::NonForkable, FSleepEvent* ExternalWorkerEvent = nullptr, FSchedulerTls::FLocalQueueType* ExternalWorkerLocalQueue = nullptr, EThreadPriority Priority = EThreadPriority::TPri_Normal, uint64 InAffinity = 0);
 		void WorkerMain(struct FSleepEvent* WorkerEvent, FSchedulerTls::FLocalQueueType* ExternalWorkerLocalQueue, uint32 WaitCycles, bool bPermitBackgroundWork);
 		CORE_API void LaunchInternal(FTask& Task, EQueuePreference QueuePreference, bool bWakeUpWorker);
 		CORE_API void BusyWaitInternal(const FConditional& Conditional, bool ForceAllowBackgroundWork);
diff --git a/Engine/Source/Runtime/Core/Public/HAL/Thread.h b/Engine/Source/Runtime/Core/Public/HAL/Thread.h
--- a/Engine/Source/Runtime/Core/Public/HAL/Thread.h
+++ b/Engine/Source/Runtime/Core/Public/HAL/Thread.h
@@ -23,6 +23,9 @@
 class CORE_API FThread final
 {
 public:
+	// indicates if the thread should be forked in case the owning process is forked
+	enum EForkable { Forkable, NonForkable };
+
 	/**
 	 * Creates new "empty" thread object that doesn't represent a system thread
 	 */
@@ -37,7 +40,7 @@ public:
 	* @param StackSize The size of the stack to create. 0 means use the current thread's stack size
 	* @param ThreadPriority Tells the thread whether it needs to adjust its priority or not. Defaults to normal priority
 	* @param ThreadAffinity Tells the thread whether it needs to adjust its affinity or not. Defaults to no affinity
-	* @param bIsForkable Tells the thread whether it can be forked. Defaults to false
+	* @param IsForkable Tells the thread whether it can be forked. Defaults to NonForkable
 	*/
 	FThread(
 		TCHAR const* ThreadName,
@@ -45,7 +48,7 @@ public:
 		uint32 StackSize = 0,
 		EThreadPriority ThreadPriority = TPri_Normal,
 		FThreadAffinity ThreadAffinity = FThreadAffinity(),
-		bool bIsForkable = false
+		EForkable IsForkable = NonForkable
 	);
 
 	// with SingleThreadTickFunction that will be executed every frame if running with `-nothreading
@@ -57,7 +60,7 @@ public:
 		uint32 StackSize = 0,
 		EThreadPriority ThreadPriority = TPri_Normal,
 		FThreadAffinity ThreadAffinity = FThreadAffinity(),
-		bool bIsForkable = false
+		EForkable IsForkable = NonForkable
 	);
 
 	// non-copyable
diff --git a/Engine/Source/Runtime/Engine/Private/UnrealEngine.cpp b/Engine/Source/Runtime/Engine/Private/UnrealEngine.cpp
--- a/Engine/Source/Runtime/Engine/Private/UnrealEngine.cpp
+++ b/Engine/Source/Runtime/Engine/Private/UnrealEngine.cpp
@@ -16896,7 +16896,7 @@ static void SetupThreadConfig(const TArray<FString>& Args)
 	int32 NumWorkerThreads = FPlatformMisc::NumberOfWorkerThreadsToSpawn();
 	int32 NumBackgroundWorkers = FMath::Max(1, NumWorkerThreads - FMath::Min<int>(2, NumWorkerThreads));
 	int32 NumForegroundWorkers = FMath::Max(1, NumWorkerThreads - NumBackgroundWorkers);
-	LowLevelTasks::FScheduler::Get().RestartWorkers(NumForegroundWorkers, NumBackgroundWorkers, FForkProcessHelper::IsForkedMultithreadInstance(), TaskThreadConfig.Priority, TaskBPThreadConfig.Priority, TaskThreadConfig.Affinity, TaskBPThreadConfig.Affinity);
+	LowLevelTasks::FScheduler::Get().RestartWorkers(NumForegroundWorkers, NumBackgroundWorkers, FForkProcessHelper::IsForkedMultithreadInstance() ? FThread::Forkable : FThread::NonForkable, TaskThreadConfig.Priority, TaskBPThreadConfig.Priority, TaskThreadConfig.Affinity, TaskBPThreadConfig.Affinity);
 
 	UE_LOG(LogConsoleResponse, Display, TEXT("Task Priority %s"), *ThreadPriorityToString(TaskThreadConfig.Priority));
 	UE_LOG(LogConsoleResponse, Display, TEXT("Task Affinity %016llX"), TaskThreadConfig.Affinity);
