Adds InitFromCompactString() to FVector to unify Getter/Setter of LevelBoundsLocation/LevelBoundsExtent tags.

#changelist validated

diff --git a/Engine/Source/Runtime/Core/Private/Tests/Math/UnrealMathTest.cpp b/Engine/Source/Runtime/Core/Private/Tests/Math/UnrealMathTest.cpp
--- a/Engine/Source/Runtime/Core/Private/Tests/Math/UnrealMathTest.cpp
+++ b/Engine/Source/Runtime/Core/Private/Tests/Math/UnrealMathTest.cpp
@@ -3995,4 +3995,38 @@ bool FMathWrapTest::RunTest(const FString& Parameters)
 	return true;
 }
 
+IMPLEMENT_SIMPLE_AUTOMATION_TEST(FInitVectorTest, "System.Core.Math.InitVector", EAutomationTestFlags::ApplicationContextMask | EAutomationTestFlags::SmokeFilter)
+bool FInitVectorTest::RunTest(const FString& Parameters)
+{
+	auto TestInitFromCompactString = [this](const FString& InTestName, const FVector& InExpected)
+	{
+		FVector Actual(13.37, 13.37, 13.37);
+		const bool bIsInitialized = Actual.InitFromCompactString(InExpected.ToCompactString());
+
+		TestTrue(*(InTestName + " return value"), bIsInitialized);
+		TestEqual(*InTestName, Actual, InExpected, KINDA_SMALL_NUMBER);
+	};
+
+	TestInitFromCompactString(TEXT("InitFromCompactString Simple"), FVector(1.2, 2.3, 3.4));
+	TestInitFromCompactString(TEXT("InitFromCompactString Zero"), FVector(0, 0, 0));
+	TestInitFromCompactString(TEXT("InitFromCompactString Int"), FVector(1, 2, 3));
+	
+	TestInitFromCompactString(TEXT("InitFromCompactString X == 0"), FVector(0, 2, 3));
+	TestInitFromCompactString(TEXT("InitFromCompactString Y == 0"), FVector(1.3, 0, 3.7));
+	TestInitFromCompactString(TEXT("InitFromCompactString Z == 0"), FVector(1.2, 2.5, 0));
+	
+	TestInitFromCompactString(TEXT("InitFromCompactString X < 0"), FVector(-433.2, 6.5, 0));
+	TestInitFromCompactString(TEXT("InitFromCompactString Y < 0"), FVector(43.2, -6.5, 98));
+	TestInitFromCompactString(TEXT("InitFromCompactString Z < 0"), FVector(33.8, 0, -76));
+	
+	TestInitFromCompactString(TEXT("InitFromCompactString X == 0 && Y == 0"), FVector(0, 0, 32.8));
+	TestInitFromCompactString(TEXT("InitFromCompactString X == 0 && Z == 0"), FVector(0, 61.3, 0));
+	TestInitFromCompactString(TEXT("InitFromCompactString Y == 0 && Z == 0"), FVector(65.3, 0, 0));
+
+	TestFalse(TEXT("InitFromCompactString BadString1"), FVector().InitFromCompactString(TEXT("W(0)")));
+	TestFalse(TEXT("InitFromCompactString BadString2"), FVector().InitFromCompactString(TEXT("V(XYZ)")));
+	
+	return true;
+}
+
 #endif //WITH_DEV_AUTOMATION_TESTS
\ No newline at end of file
diff --git a/Engine/Source/Runtime/Core/Public/Math/Vector.h b/Engine/Source/Runtime/Core/Public/Math/Vector.h
--- a/Engine/Source/Runtime/Core/Public/Math/Vector.h
+++ b/Engine/Source/Runtime/Core/Public/Math/Vector.h
@@ -860,6 +860,16 @@ public:
      */
     bool InitFromString(const FString& InSourceString);
 
+	/**
+	 * Initialize this Vector based on an FString. The String is expected to contain V(0)
+	 * or at least one value X=, Y=, Z=, previously produced by ToCompactString()
+	 * The TVector<T> will be bogus when InitFromString returns false.
+	 *
+	 * @param	InSourceString	FString containing the vector values.
+	 * @return true if any of the X,Y,Z values were read successfully; false otherwise.
+	 */
+	bool InitFromCompactString(const FString& InSourceString);
+
     /** 
      * Converts a Cartesian unit vector into spherical coordinates on the unit sphere.
      * @return Output Theta will be in the range [0, PI], and output Phi will be in the range [-PI, PI]. 
@@ -2266,15 +2276,32 @@ FORCEINLINE FString TVector<T>::ToCompactString() const
     return ReturnString;
 }
 
+template<typename T>
+FORCEINLINE bool TVector<T>::InitFromCompactString(const FString& InSourceString)
+{
+	bool bAxisFound = false;
+	
+	X = Y = Z = 0;
+
+	if (FCString::Strifind(*InSourceString, TEXT("V(0)")) != nullptr)
+	{
+		return true;
+	}
+
+	const bool bSuccessful = FParse::Value(*InSourceString, TEXT("X="), X) | FParse::Value(*InSourceString, TEXT("Y="), Y) | FParse::Value(*InSourceString, TEXT("Z="), Z);
+
+	return bSuccessful;
+}
+
 template<typename T>
 FORCEINLINE bool TVector<T>::InitFromString(const FString& InSourceString)
 {
-    X = Y = Z = 0;
+	X = Y = Z = 0;
 
-    // The initialization is only successful if the X, Y, and Z values can all be parsed from the string
-    const bool bSuccessful = FParse::Value(*InSourceString, TEXT("X=") , X) && FParse::Value(*InSourceString, TEXT("Y="), Y) && FParse::Value(*InSourceString, TEXT("Z="), Z);
+	// The initialization is only successful if the X, Y, and Z values can all be parsed from the string
+	const bool bSuccessful = FParse::Value(*InSourceString, TEXT("X=") , X) && FParse::Value(*InSourceString, TEXT("Y="), Y) && FParse::Value(*InSourceString, TEXT("Z="), Z);
 
-    return bSuccessful;
+	return bSuccessful;
 }
 
 template<typename T>
diff --git a/Engine/Source/Runtime/Engine/Private/Level.cpp b/Engine/Source/Runtime/Engine/Private/Level.cpp
--- a/Engine/Source/Runtime/Engine/Private/Level.cpp
+++ b/Engine/Source/Runtime/Engine/Private/Level.cpp
@@ -2003,8 +2003,8 @@ bool ULevel::GetLevelBoundsFromPackage(FName LevelPackage, FBox& OutLevelBounds)
 		{
 			FVector LevelBoundsLocation;
 			FVector LevelBoundsExtent;
-			if (LevelBoundsLocation.InitFromString(LevelBoundsLocationStr) &&
-				LevelBoundsExtent.InitFromString(LevelBoundsExtentStr))
+			if (LevelBoundsLocation.InitFromCompactString(LevelBoundsLocationStr) &&
+				LevelBoundsExtent.InitFromCompactString(LevelBoundsExtentStr))
 			{
 				OutLevelBounds = FBox(LevelBoundsLocation - LevelBoundsExtent, LevelBoundsLocation + LevelBoundsExtent);
 				return true;
