Compilation warning fix. Adds explicit float specifier to init vector tests.

#changelist validated

diff --git a/Engine/Source/Runtime/Core/Private/Tests/Math/UnrealMathTest.cpp b/Engine/Source/Runtime/Core/Private/Tests/Math/UnrealMathTest.cpp
--- a/Engine/Source/Runtime/Core/Private/Tests/Math/UnrealMathTest.cpp
+++ b/Engine/Source/Runtime/Core/Private/Tests/Math/UnrealMathTest.cpp
@@ -4000,28 +4000,28 @@ bool FInitVectorTest::RunTest(const FString& Parameters)
 {
 	auto TestInitFromCompactString = [this](const FString& InTestName, const FVector& InExpected)
 	{
-		FVector Actual(13.37, 13.37, 13.37);
+		FVector Actual(13.37f, 13.37f, 13.37f);
 		const bool bIsInitialized = Actual.InitFromCompactString(InExpected.ToCompactString());
 
 		TestTrue(*(InTestName + " return value"), bIsInitialized);
 		TestEqual(*InTestName, Actual, InExpected, KINDA_SMALL_NUMBER);
 	};
 
-	TestInitFromCompactString(TEXT("InitFromCompactString Simple"), FVector(1.2, 2.3, 3.4));
+	TestInitFromCompactString(TEXT("InitFromCompactString Simple"), FVector(1.2f, 2.3f, 3.4f));
 	TestInitFromCompactString(TEXT("InitFromCompactString Zero"), FVector(0, 0, 0));
 	TestInitFromCompactString(TEXT("InitFromCompactString Int"), FVector(1, 2, 3));
 	
 	TestInitFromCompactString(TEXT("InitFromCompactString X == 0"), FVector(0, 2, 3));
-	TestInitFromCompactString(TEXT("InitFromCompactString Y == 0"), FVector(1.3, 0, 3.7));
-	TestInitFromCompactString(TEXT("InitFromCompactString Z == 0"), FVector(1.2, 2.5, 0));
+	TestInitFromCompactString(TEXT("InitFromCompactString Y == 0"), FVector(1.3f, 0, 3.7f));
+	TestInitFromCompactString(TEXT("InitFromCompactString Z == 0"), FVector(1.2f, 2.5f, 0));
 	
-	TestInitFromCompactString(TEXT("InitFromCompactString X < 0"), FVector(-433.2, 6.5, 0));
-	TestInitFromCompactString(TEXT("InitFromCompactString Y < 0"), FVector(43.2, -6.5, 98));
-	TestInitFromCompactString(TEXT("InitFromCompactString Z < 0"), FVector(33.8, 0, -76));
+	TestInitFromCompactString(TEXT("InitFromCompactString X < 0"), FVector(-433.2f, 6.5f, 0));
+	TestInitFromCompactString(TEXT("InitFromCompactString Y < 0"), FVector(43.2f, -6.5f, 98));
+	TestInitFromCompactString(TEXT("InitFromCompactString Z < 0"), FVector(33.8f, 0, -76));
 	
-	TestInitFromCompactString(TEXT("InitFromCompactString X == 0 && Y == 0"), FVector(0, 0, 32.8));
-	TestInitFromCompactString(TEXT("InitFromCompactString X == 0 && Z == 0"), FVector(0, 61.3, 0));
-	TestInitFromCompactString(TEXT("InitFromCompactString Y == 0 && Z == 0"), FVector(65.3, 0, 0));
+	TestInitFromCompactString(TEXT("InitFromCompactString X == 0 && Y == 0"), FVector(0, 0, 32.8f));
+	TestInitFromCompactString(TEXT("InitFromCompactString X == 0 && Z == 0"), FVector(0, 61.3f, 0));
+	TestInitFromCompactString(TEXT("InitFromCompactString Y == 0 && Z == 0"), FVector(65.3f, 0, 0));
 
 	TestFalse(TEXT("InitFromCompactString BadString1"), FVector().InitFromCompactString(TEXT("W(0)")));
 	TestFalse(TEXT("InitFromCompactString BadString2"), FVector().InitFromCompactString(TEXT("V(XYZ)")));
