string helper functions for oodle enums.

diff --git a/Engine/Source/Runtime/Core/Private/Compression/OodleDataCompression.cpp b/Engine/Source/Runtime/Core/Private/Compression/OodleDataCompression.cpp
--- a/Engine/Source/Runtime/Core/Private/Compression/OodleDataCompression.cpp
+++ b/Engine/Source/Runtime/Core/Private/Compression/OodleDataCompression.cpp
@@ -13,6 +13,83 @@ extern ICompressionFormat * CreateOodleDataCompressionFormat();
 namespace FOodleDataCompression
 {
 
+static struct { ECompressor Compressor; const TCHAR* Name; } CompressorNameMap[] = 
+{
+	{ECompressor::NotSet, TEXT("Not Set")},
+	{ECompressor::Selkie, TEXT("Selkie")},
+	{ECompressor::Mermaid, TEXT("Mermaid")},
+	{ECompressor::Kraken, TEXT("Kraken")},
+	{ECompressor::Leviathan, TEXT("Leviathan")}
+};
+
+bool ECompressorToString(ECompressor InCompressor, const TCHAR** OutName)
+{
+	if ((SIZE_T)InCompressor >= sizeof(CompressorNameMap) / sizeof(CompressorNameMap[0]))
+	{
+		return false;
+	}
+
+	*OutName = CompressorNameMap[(uint32)InCompressor].Name;
+	return true;
+}
+
+
+bool ECompressorFromString(const FString& InName, ECompressor& OutCompressor)
+{
+	for (SIZE_T i = 0; i < sizeof(CompressorNameMap) / sizeof(CompressorNameMap[0]); i++)
+	{
+		if (InName == CompressorNameMap[i].Name)
+		{
+			OutCompressor = CompressorNameMap[i].Compressor;
+			return true;
+		}
+	}
+	return false;
+}
+
+static struct { ECompressionLevel Level; const TCHAR* Name; } CompressionLevelNameMap[] =
+{
+	{ECompressionLevel::HyperFast4, TEXT("HyperFast4")},
+	{ECompressionLevel::HyperFast3, TEXT("HyperFast3")},
+	{ECompressionLevel::HyperFast2, TEXT("HyperFast2")},
+	{ECompressionLevel::HyperFast1, TEXT("HyperFast1")},
+	{ECompressionLevel::None, TEXT("None")},
+	{ECompressionLevel::SuperFast, TEXT("SuperFast")},
+	{ECompressionLevel::VeryFast, TEXT("VeryFast")},
+	{ECompressionLevel::Fast, TEXT("Fast")},
+	{ECompressionLevel::Normal, TEXT("Normal")},
+	{ECompressionLevel::Optimal1, TEXT("Optimal1")},
+	{ECompressionLevel::Optimal2, TEXT("Optimal2")},
+	{ECompressionLevel::Optimal3, TEXT("Optimal3")},
+	{ECompressionLevel::Optimal4, TEXT("Optimal4")}
+};
+
+bool ECompressionLevelToString(ECompressionLevel InLevel, const TCHAR** OutName)
+{
+	for (SIZE_T i = 0; i < sizeof(CompressionLevelNameMap) / sizeof(CompressionLevelNameMap[0]); i++)
+	{
+		if (CompressionLevelNameMap[i].Level == InLevel)
+		{
+			*OutName = CompressionLevelNameMap[i].Name;
+			return true;
+		}
+	}
+	return false;
+}
+
+bool ECompressionLevelFromValue(int8 InValue, ECompressionLevel& OutLevel)
+{
+	if (InValue < OodleLZ_CompressionLevel_Min ||
+		InValue > OodleLZ_CompressionLevel_Max)
+	{
+		return false;
+	}
+
+	OutLevel = (ECompressionLevel)InValue;
+	return true;
+}
+
+
 
 static OodleLZ_Compressor CompressorToOodleLZ_Compressor(ECompressor Compressor)
 {
diff --git a/Engine/Source/Runtime/Core/Public/Compression/OodleDataCompression.h b/Engine/Source/Runtime/Core/Public/Compression/OodleDataCompression.h
--- a/Engine/Source/Runtime/Core/Public/Compression/OodleDataCompression.h
+++ b/Engine/Source/Runtime/Core/Public/Compression/OodleDataCompression.h
@@ -48,9 +48,12 @@ enum class ECompressor : uint8
 	Selkie = 1,
 	Mermaid = 2,
 	Kraken  = 3,
-	Leviathan = 4
+	Leviathan = 4 // if another added update CompressorNameMap
 };
 
+// If the input is invalid, returns false, and doesn't touch output.
+bool CORE_API ECompressorToString(ECompressor InCompressor, const TCHAR** OutName);
+bool CORE_API ECompressorFromString(const class FString& InName, ECompressor& OutCompressor);
 
 /**
  * ECompressionLevel : Choose the Oodle Compression Level
@@ -97,9 +100,13 @@ enum class ECompressionLevel : int8
 	Optimal1 = 5,
 	Optimal2 = 6,
 	Optimal3 = 7,
-	Optimal4 = 8,
+	Optimal4 = 8, // if another added update CompressionLevelNameMap
 };
 
+// If the input is invalid, returns false, and doesn't touch output.
+bool CORE_API ECompressionLevelFromValue(int8 InValue, ECompressionLevel& OutLevel);
+bool CORE_API ECompressionLevelToString(ECompressionLevel InLevel, const TCHAR** OutName);
+
 enum class ECompressionCommonUsage : uint8
 {
 	Default = 0,
