Add missing copyright

diff --git a/Engine/Source/Runtime/Core/Public/Misc/MemStackUtility.h b/Engine/Source/Runtime/Core/Public/Misc/MemStackUtility.h
--- a/Engine/Source/Runtime/Core/Public/Misc/MemStackUtility.h
+++ b/Engine/Source/Runtime/Core/Public/Misc/MemStackUtility.h
@@ -1,3 +1,4 @@
+// Copyright Epic Games, Inc. All Rights Reserved.
 #pragma once
 
 #include "CoreMinimal.h"
