Save about 25% of total cook memory usage by releasing the read buffer cache in FlushCache
  - 210GB to 154GB Peak Used virtual
  - 198GB to 149GB Peak Used physical

diff --git a/Engine/Source/Runtime/Core/Private/HAL/FileManagerGeneric.cpp b/Engine/Source/Runtime/Core/Private/HAL/FileManagerGeneric.cpp
--- a/Engine/Source/Runtime/Core/Private/HAL/FileManagerGeneric.cpp
+++ b/Engine/Source/Runtime/Core/Private/HAL/FileManagerGeneric.cpp
@@ -906,6 +906,11 @@ void FArchiveFileReaderGeneric::FlushCache()
 	{
 		Handle->ShrinkBuffers();
 	}
+
+	BufferArray.Empty();
+	// Make sure the underlying system file has the right position for any read that comes after a flushed cache.
+	SeekLowLevel(Pos);
+	BufferBase = Pos;
 }
 
 FArchiveFileWriterGeneric::FArchiveFileWriterGeneric( IFileHandle* InHandle, const TCHAR* InFilename, int64 InPos, uint32 InBufferSize, uint32 InFlags )
