[Insights] Add special tag for Trace.

 * Adds special value reserved to Trace allocations
 * Adds checks to avoid collisions between FName indices and ELLMTags.

diff --git a/Engine/Source/Runtime/Core/Private/ProfilingDebugging/TagTrace.cpp b/Engine/Source/Runtime/Core/Private/ProfilingDebugging/TagTrace.cpp
--- a/Engine/Source/Runtime/Core/Private/ProfilingDebugging/TagTrace.cpp
+++ b/Engine/Source/Runtime/Core/Private/ProfilingDebugging/TagTrace.cpp
@@ -2,10 +2,8 @@
 #include "ProfilingDebugging/TagTrace.h"
 
 #include "Experimental/Containers/GrowOnlyLockFreeHash.h"
-#include "Containers/Set.h"
 #include "CoreTypes.h"
 #include "ProfilingDebugging/MemoryTrace.h"
-#include "Misc/ScopeLock.h"
 #include "HAL/LowLevelMemTracker.h"
 #include "Trace/Trace.inl"
 #include "UObject/NameTypes.h"
@@ -125,16 +123,19 @@ class FTagTrace
 {
 public:
 					FTagTrace(FMalloc* InMalloc);
-	void			AnnounceGenericTags();
+	void			AnnounceGenericTags() const;
 	void 			AnnounceTagDeclarations();
+					void AnnounceSpecialTags() const;
 #if ENABLE_LOW_LEVEL_MEM_TRACKER
 	static void 	OnAnnounceTagDeclaration(FLLMTagDeclaration& TagDeclaration);
 #endif
-	int32			AnnounceCustomTag(int32 Tag, int32 ParentTag, const ANSICHAR* Display);
+	int32			AnnounceCustomTag(int32 Tag, int32 ParentTag, const ANSICHAR* Display) const;
 	int32 			AnnounceFNameTag(const FName& TagName);
 
 private:
 
+	static constexpr int32 FNAME_INDEX_OFFSET = 512;
+	
 	struct FTagNameSetEntry
 	{
 		std::atomic_int32_t Data;
@@ -166,10 +167,11 @@ FTagTrace::FTagTrace(FMalloc* InMalloc)
 	AnnouncedNames.Reserve(1024);
 	AnnounceGenericTags();
 	AnnounceTagDeclarations();
+	AnnounceSpecialTags();
 }
 
 ////////////////////////////////////////////////////////////////////////////////
-void FTagTrace::AnnounceGenericTags()
+void FTagTrace::AnnounceGenericTags() const
 {
 #if ENABLE_LOW_LEVEL_MEM_TRACKER
 	#define TRACE_TAG_SPEC(Enum,Str,Stat,Group,ParentTag)\
@@ -199,6 +201,22 @@ void FTagTrace::AnnounceTagDeclarations()
 #endif
 }
 
+////////////////////////////////////////////////////////////////////////////////
+void FTagTrace::AnnounceSpecialTags() const
+{
+	auto EmitTag = [](const TCHAR* DisplayString, int32 Tag, int32 ParentTag)
+	{
+		const int32 DisplayLen = FCString::Strlen(DisplayString);
+		UE_TRACE_LOG(Memory, TagSpec, MemAllocChannel, DisplayLen * sizeof(ANSICHAR))
+			 << TagSpec.Tag(Tag)
+			 << TagSpec.Parent(ParentTag)
+			 << TagSpec.Display(DisplayString, DisplayLen);
+	};
+
+	EmitTag(TEXT("Trace"), TRACE_TAG, -1);
+
+}
+
 ////////////////////////////////////////////////////////////////////////////////
 #if ENABLE_LOW_LEVEL_MEM_TRACKER
 void FTagTrace::OnAnnounceTagDeclaration(FLLMTagDeclaration& TagDeclaration)
@@ -211,14 +229,7 @@ void FTagTrace::OnAnnounceTagDeclaration(FLLMTagDeclaration& TagDeclaration)
 ////////////////////////////////////////////////////////////////////////////////
 int32 FTagTrace::AnnounceFNameTag(const FName& Name)
 {
-	const int32 NameIndex = Name.GetDisplayIndex().ToUnstableInt();
-
-	// Don't announce NAME_None, if we happen to get that passed in.  The "AnnouncedNames" container
-	// is not allowed to hold "NAME_None", as zero represents an invalid key.
-	if (!NameIndex)
-	{
-		return NameIndex;
-	}
+	const int32 NameIndex = Name.GetDisplayIndex().ToUnstableInt() + FNAME_INDEX_OFFSET;
 
 	// Find or add the item
 	bool bAlreadyInTable;
@@ -235,7 +246,7 @@ int32 FTagTrace::AnnounceFNameTag(const FName& Name)
 }
 
 ////////////////////////////////////////////////////////////////////////////////
-int32 FTagTrace::AnnounceCustomTag(int32 Tag, int32 ParentTag, const ANSICHAR* Display)
+int32 FTagTrace::AnnounceCustomTag(int32 Tag, int32 ParentTag, const ANSICHAR* Display) const
 {		
 	const uint32 DisplayLen = FCStringAnsi::Strlen(Display);
 	UE_TRACE_LOG(Memory, TagSpec, MemAllocChannel, DisplayLen * sizeof(ANSICHAR))
diff --git a/Engine/Source/Runtime/Core/Public/ProfilingDebugging/TagTrace.h b/Engine/Source/Runtime/Core/Public/ProfilingDebugging/TagTrace.h
--- a/Engine/Source/Runtime/Core/Public/ProfilingDebugging/TagTrace.h
+++ b/Engine/Source/Runtime/Core/Public/ProfilingDebugging/TagTrace.h
@@ -72,7 +72,9 @@ private:
 };
 
 ////////////////////////////////////////////////////////////////////////////////
+constexpr int32 TRACE_TAG = 257;
 
+////////////////////////////////////////////////////////////////////////////////
 #define UE_MEMSCOPE(InTag)				FMemScope PREPROCESSOR_JOIN(MemScope,__LINE__)(InTag);
 #define UE_MEMSCOPE_PTR(InPtr)			FMemScopePtr PREPROCESSOR_JOIN(MemPtrScope,__LINE__)((uint64)InPtr);
 
