fixed natvis for TTypeCompatiibleBytes<lambda>

diff --git a/Engine/Source/Runtime/Core/Public/Templates/TypeCompatibleBytes.h b/Engine/Source/Runtime/Core/Public/Templates/TypeCompatibleBytes.h
--- a/Engine/Source/Runtime/Core/Public/Templates/TypeCompatibleBytes.h
+++ b/Engine/Source/Runtime/Core/Public/Templates/TypeCompatibleBytes.h
@@ -62,6 +62,7 @@ struct TTypeCompatibleBytes :
 		alignof(ElementType)
 		>
 {
+	using ElementTypeAlias_NatVisHelper = ElementType;
 	ElementType*		GetTypedPtr()		{ return (ElementType*)this;  }
 	const ElementType*	GetTypedPtr() const	{ return (const ElementType*)this; }
 };
