String: Extended ParseTokens and ParseLines

- Added option flags to control skipping empty tokens/lines, trimming tokens/lines, and ignoring delimiter case.
- Added overloads that add to an output object, such as an array.

diff --git a/Engine/Source/Runtime/Core/Private/String/ParseLines.cpp b/Engine/Source/Runtime/Core/Private/String/ParseLines.cpp
--- a/Engine/Source/Runtime/Core/Private/String/ParseLines.cpp
+++ b/Engine/Source/Runtime/Core/Private/String/ParseLines.cpp
@@ -2,14 +2,13 @@
 
 #include "String/ParseLines.h"
 
-#include "Containers/StringView.h"
-
-namespace UE
-{
-namespace String
+namespace UE::String
 {
 
-void ParseLines(const FStringView& View, const TFunctionRef<void(FStringView)>& Visitor)
+void ParseLines(
+	const FStringView View,
+	const TFunctionRef<void (FStringView)> Visitor,
+	const EParseLinesOptions Options)
 {
 	const TCHAR* ViewIt = View.GetData();
 	const TCHAR* ViewEnd = ViewIt + View.Len();
@@ -36,10 +35,17 @@ void ParseLines(const FStringView& View, const TFunctionRef<void(FStringView)>&
 			}
 		}
 
-		Visitor(FStringView(LineStart, static_cast<FStringView::SizeType>(LineEnd - LineStart)));
+		FStringView Line(LineStart, static_cast<FStringView::SizeType>(LineEnd - LineStart));
+		if (EnumHasAnyFlags(Options, EParseLinesOptions::Trim))
+		{
+			Line = Line.TrimStartAndEnd();
+		}
+		if (!EnumHasAnyFlags(Options, EParseLinesOptions::SkipEmpty) || !Line.IsEmpty())
+		{
+			Visitor(Line);
+		}
 	}
 	while (ViewIt != ViewEnd);
 }
 
-}
-}
+} // UE::String
diff --git a/Engine/Source/Runtime/Core/Private/String/ParseTokens.cpp b/Engine/Source/Runtime/Core/Private/String/ParseTokens.cpp
--- a/Engine/Source/Runtime/Core/Private/String/ParseTokens.cpp
+++ b/Engine/Source/Runtime/Core/Private/String/ParseTokens.cpp
@@ -5,46 +5,100 @@
 #include "Algo/AllOf.h"
 #include "Algo/Find.h"
 #include "Algo/NoneOf.h"
+#include "Algo/Transform.h"
+#include "Containers/Array.h"
 #include "Containers/BitArray.h"
-#include "Containers/StringView.h"
 
-namespace UE
+namespace UE::String
 {
-namespace String
+
+inline static void ParseTokensVisitToken(
+	const TFunctionRef<void (FStringView)>& Visitor,
+	const EParseTokensOptions Options,
+	FStringView Token)
 {
+	if (EnumHasAnyFlags(Options, EParseTokensOptions::Trim))
+	{
+		Token = Token.TrimStartAndEnd();
+	}
+	if (!EnumHasAnyFlags(Options, EParseTokensOptions::SkipEmpty) || !Token.IsEmpty())
+	{
+		Visitor(Token);
+	}
+}
 
 /** Parse tokens with one single-character delimiter. */
-inline static void ParseTokens1Delim1Char(const FStringView& View, const TCHAR Delimiter, TFunctionRef<void(FStringView)> Visitor)
+inline static void ParseTokens1Delim1Char(
+	const FStringView View,
+	const TCHAR Delimiter,
+	const TFunctionRef<void (FStringView)> Visitor,
+	const EParseTokensOptions Options)
 {
 	const TCHAR* ViewIt = View.GetData();
 	const TCHAR* const ViewEnd = ViewIt + View.Len();
 	const TCHAR* NextToken = ViewIt;
 
-	for (;;)
+	if (EnumHasAnyFlags(Options, EParseTokensOptions::IgnoreCase))
 	{
-		if (ViewIt == ViewEnd)
+		const TCHAR LowerDelimiter = FChar::ToLower(Delimiter);
+		for (;;)
 		{
-			break;
+			if (ViewIt == ViewEnd)
+			{
+				break;
+			}
+			if (FChar::ToLower(*ViewIt) != Delimiter)
+			{
+				++ViewIt;
+				continue;
+			}
+			ParseTokensVisitToken(Visitor, Options, FStringView(NextToken, static_cast<FStringView::SizeType>(ViewIt - NextToken)));
+			NextToken = ++ViewIt;
 		}
-		if (*ViewIt != Delimiter)
+	}
+	else
+	{
+		for (;;)
 		{
-			++ViewIt;
-			continue;
+			if (ViewIt == ViewEnd)
+			{
+				break;
+			}
+			if (*ViewIt != Delimiter)
+			{
+				++ViewIt;
+				continue;
+			}
+			ParseTokensVisitToken(Visitor, Options, FStringView(NextToken, static_cast<FStringView::SizeType>(ViewIt - NextToken)));
+			NextToken = ++ViewIt;
 		}
-		Visitor(FStringView(NextToken, static_cast<FStringView::SizeType>(ViewIt - NextToken)));
-		NextToken = ++ViewIt;
 	}
 
-	Visitor(FStringView(NextToken, static_cast<FStringView::SizeType>(ViewIt - NextToken)));
+	ParseTokensVisitToken(Visitor, Options, FStringView(NextToken, static_cast<FStringView::SizeType>(ViewIt - NextToken)));
 }
 
 /** Parse tokens with multiple single-character Basic Latin delimiters. */
-inline static void ParseTokensNDelim1CharBasicLatin(const FStringView& View, TArrayView<const TCHAR> Delimiters, TFunctionRef<void(FStringView)> Visitor)
+inline static void ParseTokensNDelim1CharBasicLatin(
+	const FStringView View,
+	const TConstArrayView<TCHAR> Delimiters,
+	const TFunctionRef<void (FStringView)> Visitor,
+	const EParseTokensOptions Options)
 {
 	TBitArray<> DelimiterMask(false, 128);
-	for (uint32 Delimiter : Delimiters)
+	if (EnumHasAnyFlags(Options, EParseTokensOptions::IgnoreCase))
+	{
+		for (TCHAR Delimiter : Delimiters)
+		{
+			DelimiterMask[FChar::ToUnsigned(FChar::ToLower(Delimiter))] = true;
+			DelimiterMask[FChar::ToUnsigned(FChar::ToUpper(Delimiter))] = true;
+		}
+	}
+	else
 	{
-		DelimiterMask[Delimiter] = true;
+		for (TCHAR Delimiter : Delimiters)
+		{
+			DelimiterMask[FChar::ToUnsigned(Delimiter)] = true;
+		}
 	}
 
 	const TCHAR* ViewIt = View.GetData();
@@ -63,45 +117,75 @@ inline static void ParseTokensNDelim1CharBasicLatin(const FStringView& View, TAr
 			++ViewIt;
 			continue;
 		}
-		Visitor(FStringView(NextToken, static_cast<FStringView::SizeType>(ViewIt - NextToken)));
+		ParseTokensVisitToken(Visitor, Options, FStringView(NextToken, static_cast<FStringView::SizeType>(ViewIt - NextToken)));
 		NextToken = ++ViewIt;
 	}
 
-	Visitor(FStringView(NextToken, static_cast<FStringView::SizeType>(ViewIt - NextToken)));
+	ParseTokensVisitToken(Visitor, Options, FStringView(NextToken, static_cast<FStringView::SizeType>(ViewIt - NextToken)));
 }
 
 /** Parse tokens with multiple single-character delimiters in the Basic Multilingual Plane. */
-inline static void ParseTokensNDelim1Char(const FStringView& View, TArrayView<const TCHAR> Delimiters, TFunctionRef<void(FStringView)> Visitor)
+inline static void ParseTokensNDelim1Char(
+	const FStringView View,
+	const TConstArrayView<TCHAR> Delimiters,
+	TFunctionRef<void (FStringView)> Visitor,
+	const EParseTokensOptions Options)
 {
 	if (Algo::AllOf(Delimiters, [](TCHAR Delimiter) { return Delimiter < 128; }))
 	{
-		return ParseTokensNDelim1CharBasicLatin(View, Delimiters, Visitor);
+		return ParseTokensNDelim1CharBasicLatin(View, Delimiters, MoveTemp(Visitor), Options);
 	}
 
 	const TCHAR* ViewIt = View.GetData();
 	const TCHAR* const ViewEnd = ViewIt + View.Len();
 	const TCHAR* NextToken = ViewIt;
 
-	for (;;)
+	if (EnumHasAnyFlags(Options, EParseTokensOptions::IgnoreCase))
 	{
-		if (ViewIt == ViewEnd)
+		TArray<TCHAR, TInlineAllocator<16>> LowerDelimiters;
+		Algo::Transform(Delimiters, LowerDelimiters, FChar::ToLower);
+		for (;;)
 		{
-			break;
+			if (ViewIt == ViewEnd)
+			{
+				break;
+			}
+			if (!Algo::Find(Delimiters, FChar::ToLower(*ViewIt)))
+			{
+				++ViewIt;
+				continue;
+			}
+			ParseTokensVisitToken(Visitor, Options, FStringView(NextToken, static_cast<FStringView::SizeType>(ViewIt - NextToken)));
+			NextToken = ++ViewIt;
 		}
-		if (!Algo::Find(Delimiters, *ViewIt))
+	}
+	else
+	{
+		for (;;)
 		{
-			++ViewIt;
-			continue;
+			if (ViewIt == ViewEnd)
+			{
+				break;
+			}
+			if (!Algo::Find(Delimiters, *ViewIt))
+			{
+				++ViewIt;
+				continue;
+			}
+			ParseTokensVisitToken(Visitor, Options, FStringView(NextToken, static_cast<FStringView::SizeType>(ViewIt - NextToken)));
+			NextToken = ++ViewIt;
 		}
-		Visitor(FStringView(NextToken, static_cast<FStringView::SizeType>(ViewIt - NextToken)));
-		NextToken = ++ViewIt;
 	}
 
-	Visitor(FStringView(NextToken, static_cast<FStringView::SizeType>(ViewIt - NextToken)));
+	ParseTokensVisitToken(Visitor, Options, FStringView(NextToken, static_cast<FStringView::SizeType>(ViewIt - NextToken)));
 }
 
 /** Parse tokens with multiple multi-character delimiters. */
-inline static void ParseTokensNDelimNChar(const FStringView& View, TArrayView<const FStringView> Delimiters, TFunctionRef<void(FStringView)> Visitor)
+inline static void ParseTokensNDelimNChar(
+	const FStringView View,
+	const TConstArrayView<FStringView> Delimiters,
+	const TFunctionRef<void (FStringView)> Visitor,
+	const EParseTokensOptions Options)
 {
 	// This is a naive implementation that takes time proportional to View.Len() * TotalDelimiterLen.
 	// If this function becomes a bottleneck, it can be specialized separately for one and many delimiters.
@@ -110,13 +194,14 @@ inline static void ParseTokensNDelimNChar(const FStringView& View, TArrayView<co
 	const FStringView::SizeType ViewLen = View.Len();
 	FStringView::SizeType NextTokenIndex = 0;
 
+	const ESearchCase::Type SearchCase = EnumHasAnyFlags(Options, EParseTokensOptions::IgnoreCase) ? ESearchCase::IgnoreCase : ESearchCase::CaseSensitive;
 	for (FStringView::SizeType ViewIndex = 0; ViewIndex != ViewLen;)
 	{
 		const FStringView RemainingView(View.GetData() + ViewIndex, ViewLen - ViewIndex);
-		auto MatchDelimiter = [&RemainingView](const FStringView& Delimiter) { return RemainingView.StartsWith(Delimiter, ESearchCase::CaseSensitive); };
+		auto MatchDelimiter = [RemainingView, SearchCase](FStringView Delimiter) { return RemainingView.StartsWith(Delimiter, SearchCase); };
 		if (const FStringView* Delimiter = Algo::FindByPredicate(Delimiters, MatchDelimiter))
 		{
-			Visitor(FStringView(View.GetData() + NextTokenIndex, ViewIndex - NextTokenIndex));
+			ParseTokensVisitToken(Visitor, Options, FStringView(View.GetData() + NextTokenIndex, ViewIndex - NextTokenIndex));
 			ViewIndex += Delimiter->Len();
 			NextTokenIndex = ViewIndex;
 		}
@@ -126,22 +211,26 @@ inline static void ParseTokensNDelimNChar(const FStringView& View, TArrayView<co
 		}
 	}
 
-	Visitor(FStringView(View.GetData() + NextTokenIndex, ViewLen - NextTokenIndex));
+	ParseTokensVisitToken(Visitor, Options, FStringView(View.GetData() + NextTokenIndex, ViewLen - NextTokenIndex));
 }
 
-void ParseTokensMultiple(const FStringView& View, TArrayView<const FStringView> Delimiters, TFunctionRef<void(FStringView)> Visitor)
+void ParseTokensMultiple(
+	const FStringView View,
+	const TConstArrayView<FStringView> Delimiters,
+	TFunctionRef<void (FStringView)> Visitor,
+	const EParseTokensOptions Options)
 {
 	check(Algo::NoneOf(Delimiters, &FStringView::IsEmpty));
 	switch (Delimiters.Num())
 	{
 	case 0:
-		return Visitor(View);
+		return ParseTokensVisitToken(Visitor, Options, View);
 	case 1:
 		if (Delimiters[0].Len() == 1)
 		{
-			return ParseTokens1Delim1Char(View, Delimiters[0][0], MoveTemp(Visitor));
+			return ParseTokens1Delim1Char(View, Delimiters[0][0], MoveTemp(Visitor), Options);
 		}
-		return ParseTokensNDelimNChar(View, Delimiters, MoveTemp(Visitor));
+		return ParseTokensNDelimNChar(View, Delimiters, MoveTemp(Visitor), Options);
 	default:
 		if (Algo::AllOf(Delimiters, [](const FStringView& Delimiter) { return Delimiter.Len() == 1; }))
 		{
@@ -151,27 +240,52 @@ void ParseTokensMultiple(const FStringView& View, TArrayView<const FStringView>
 			{
 				DelimiterChars.Add(Delimiter[0]);
 			}
-			return ParseTokensNDelim1Char(View, DelimiterChars, MoveTemp(Visitor));
+			return ParseTokensNDelim1Char(View, DelimiterChars, MoveTemp(Visitor), Options);
 		}
 		else
 		{
-			return ParseTokensNDelimNChar(View, Delimiters, MoveTemp(Visitor));
+			return ParseTokensNDelimNChar(View, Delimiters, MoveTemp(Visitor), Options);
 		}
 	}
 }
 
-void ParseTokensMultiple(const FStringView& View, TArrayView<const TCHAR> Delimiters, TFunctionRef<void(FStringView)> Visitor)
+void ParseTokensMultiple(
+	const FStringView View,
+	const TConstArrayView<TCHAR> Delimiters,
+	TFunctionRef<void (FStringView)> Visitor,
+	const EParseTokensOptions Options)
 {
 	switch (Delimiters.Num())
 	{
 	case 0:
-		return Visitor(View);
+		return ParseTokensVisitToken(Visitor, Options, View);
 	case 1:
-		return ParseTokens1Delim1Char(View, Delimiters[0], MoveTemp(Visitor));
+		return ParseTokens1Delim1Char(View, Delimiters[0], MoveTemp(Visitor), Options);
 	default:
-		return ParseTokensNDelim1Char(View, Delimiters, MoveTemp(Visitor));
+		return ParseTokensNDelim1Char(View, Delimiters, MoveTemp(Visitor), Options);
 	}
 }
 
+void ParseTokens(
+	const FStringView View,
+	const FStringView Delimiter,
+	TFunctionRef<void (FStringView)> Visitor,
+	const EParseTokensOptions Options)
+{
+	if (Delimiter.Len() == 1)
+	{
+		return ParseTokens1Delim1Char(View, Delimiter[0], MoveTemp(Visitor), Options);
+	}
+	return ParseTokensNDelimNChar(View, MakeArrayView(&Delimiter, 1), MoveTemp(Visitor), Options);
 }
+
+void ParseTokens(
+	const FStringView View,
+	const TCHAR Delimiter,
+	TFunctionRef<void (FStringView)> Visitor,
+	const EParseTokensOptions Options)
+{
+	return ParseTokens1Delim1Char(View, Delimiter, MoveTemp(Visitor), Options);
 }
+
+} // UE::String
diff --git a/Engine/Source/Runtime/Core/Private/Tests/String/ParseLinesTest.cpp b/Engine/Source/Runtime/Core/Private/Tests/String/ParseLinesTest.cpp
--- a/Engine/Source/Runtime/Core/Private/Tests/String/ParseLinesTest.cpp
+++ b/Engine/Source/Runtime/Core/Private/Tests/String/ParseLinesTest.cpp
@@ -13,40 +13,65 @@
 IMPLEMENT_SIMPLE_AUTOMATION_TEST(StringParseLinesTest, "System.Core.String.ParseLines", EAutomationTestFlags::ApplicationContextMask | EAutomationTestFlags::EngineFilter)
 bool StringParseLinesTest::RunTest(const FString& Parameters)
 {
-	auto RunParseLinesTest = [this](const FStringView& View, std::initializer_list<FStringView> ExpectedLines)
+	using EOptions = UE::String::EParseLinesOptions;
+
+	auto RunParseLinesTest = [this](FStringView View, std::initializer_list<FStringView> ExpectedLines, EOptions Options = EOptions::None)
 	{
 		TArray<FStringView, TInlineAllocator<8>> ResultLines;
-		UE::String::ParseLines(View, [&ResultLines](FStringView Line) { ResultLines.Add(Line); });
+		UE::String::ParseLines(View, ResultLines, Options);
 		if (!Algo::CompareByPredicate(ResultLines, ExpectedLines, TEqualTo<>()))
 		{
 			TStringBuilder<512> Error;
-			Error << TEXT("UE::String::ParseLines failed to parse \"") << FString(View).ReplaceCharWithEscapedChar() << TEXT("\" result {");
-			Error.JoinQuoted(ResultLines, TEXT(", "_SV), TEXT("\""_SV));
-			Error << TEXT("} expected {");
-			Error.JoinQuoted(ExpectedLines, TEXT(", "_SV), TEXT("\""_SV));
-			Error << TEXT("}");
+			Error << TEXTVIEW("UE::String::ParseLines failed to parse \"") << FString(View).ReplaceCharWithEscapedChar() << TEXTVIEW("\" result {");
+			Error.JoinQuoted(ResultLines, TEXTVIEW(", "), TEXTVIEW("\""));
+			Error << TEXTVIEW("} expected {");
+			Error.JoinQuoted(ExpectedLines, TEXTVIEW(", "), TEXTVIEW("\""));
+			Error << TEXTVIEW("}");
 			AddError(Error.ToString());
 		}
 	};
 
-	RunParseLinesTest(TEXT(""_SV),                        {TEXT(""_SV)});
-	RunParseLinesTest(TEXT("\n"_SV),                      {TEXT(""_SV)});
-	RunParseLinesTest(TEXT("\r"_SV),                      {TEXT(""_SV)});
-	RunParseLinesTest(TEXT("\r\n"_SV),                    {TEXT(""_SV)});
-	RunParseLinesTest(TEXT("\n\n"_SV),                    {TEXT(""_SV), TEXT(""_SV)});
-	RunParseLinesTest(TEXT("\r\r"_SV),                    {TEXT(""_SV), TEXT(""_SV)});
-	RunParseLinesTest(TEXT("\r\n\r\n"_SV),                {TEXT(""_SV), TEXT(""_SV)});
-	RunParseLinesTest(TEXT("\r\nABC"_SV).Left(2),         {TEXT(""_SV)});
-	RunParseLinesTest(TEXT("\r\nABC\r\nDEF"_SV).Left(5),  {TEXT(""_SV), TEXT("ABC"_SV)});
-	RunParseLinesTest(TEXT("ABC DEF"_SV),                 {TEXT("ABC DEF"_SV)});
-	RunParseLinesTest(TEXT("\nABC DEF\n"_SV),             {TEXT(""_SV), TEXT("ABC DEF"_SV)});
-	RunParseLinesTest(TEXT("\rABC DEF\r"_SV),             {TEXT(""_SV), TEXT("ABC DEF"_SV)});
-	RunParseLinesTest(TEXT("\r\nABC DEF\r\n"_SV),         {TEXT(""_SV), TEXT("ABC DEF"_SV)});
-	RunParseLinesTest(TEXT("\r\n\r\nABC DEF\r\n\r\n"_SV), {TEXT(""_SV), TEXT(""_SV), TEXT("ABC DEF"_SV), TEXT(""_SV)});
-	RunParseLinesTest(TEXT("ABC\nDEF"_SV),                {TEXT("ABC"_SV), TEXT("DEF"_SV)});
-	RunParseLinesTest(TEXT("ABC\rDEF"_SV),                {TEXT("ABC"_SV), TEXT("DEF"_SV)});
-	RunParseLinesTest(TEXT("\r\nABC\r\nDEF\r\n"_SV),      {TEXT(""_SV), TEXT("ABC"_SV), TEXT("DEF"_SV)});
-	RunParseLinesTest(TEXT("\r\nABC\r\n\r\nDEF\r\n"_SV),  {TEXT(""_SV), TEXT("ABC"_SV), TEXT(""_SV), TEXT("DEF"_SV)});
+	constexpr EOptions KeepEmpty = EOptions::None;
+	constexpr EOptions SkipEmpty = EOptions::SkipEmpty;
+	constexpr EOptions Trim = EOptions::Trim;
+
+	RunParseLinesTest(TEXTVIEW(""), {}, SkipEmpty);
+	RunParseLinesTest(TEXTVIEW(""), {TEXTVIEW("")}, KeepEmpty);
+	RunParseLinesTest(TEXTVIEW("\n"), {}, SkipEmpty);
+	RunParseLinesTest(TEXTVIEW("\n"), {TEXTVIEW("")}, KeepEmpty);
+	RunParseLinesTest(TEXTVIEW("\r"), {}, SkipEmpty);
+	RunParseLinesTest(TEXTVIEW("\r"), {TEXTVIEW("")}, KeepEmpty);
+	RunParseLinesTest(TEXTVIEW("\r\n"), {}, SkipEmpty);
+	RunParseLinesTest(TEXTVIEW("\r\n"), {TEXTVIEW("")}, KeepEmpty);
+	RunParseLinesTest(TEXTVIEW("\n\n"), {}, SkipEmpty);
+	RunParseLinesTest(TEXTVIEW("\n\n"), {TEXTVIEW(""), TEXTVIEW("")}, KeepEmpty);
+	RunParseLinesTest(TEXTVIEW("\r\r"), {}, SkipEmpty);
+	RunParseLinesTest(TEXTVIEW("\r\r"), {TEXTVIEW(""), TEXTVIEW("")}, KeepEmpty);
+	RunParseLinesTest(TEXTVIEW("\r\n\r\n"), {}, SkipEmpty);
+	RunParseLinesTest(TEXTVIEW("\r\n\r\n"), {TEXTVIEW(""), TEXTVIEW("")}, KeepEmpty);
+	RunParseLinesTest(TEXTVIEW("\r\nABC").Left(2), {}, SkipEmpty);
+	RunParseLinesTest(TEXTVIEW("\r\nABC").Left(2), {TEXTVIEW("")}, KeepEmpty);
+	RunParseLinesTest(TEXTVIEW("\r\nABC\r\nDEF").Left(5), {TEXTVIEW("ABC")}, SkipEmpty);
+	RunParseLinesTest(TEXTVIEW("\r\nABC\r\nDEF").Left(5), {TEXTVIEW(""), TEXTVIEW("ABC")}, KeepEmpty);
+	RunParseLinesTest(TEXTVIEW("ABC DEF"), {TEXTVIEW("ABC DEF")}, SkipEmpty);
+	RunParseLinesTest(TEXTVIEW("\nABC DEF\n"), {TEXTVIEW("ABC DEF")}, SkipEmpty);
+	RunParseLinesTest(TEXTVIEW("\nABC DEF\n"), {TEXTVIEW(""), TEXTVIEW("ABC DEF")}, KeepEmpty);
+	RunParseLinesTest(TEXTVIEW("\rABC DEF\r"), {TEXTVIEW("ABC DEF")}, SkipEmpty);
+	RunParseLinesTest(TEXTVIEW("\rABC DEF\r"), {TEXTVIEW(""), TEXTVIEW("ABC DEF")}, KeepEmpty);
+	RunParseLinesTest(TEXTVIEW("\r\nABC DEF\r\n"), {TEXTVIEW("ABC DEF")}, SkipEmpty);
+	RunParseLinesTest(TEXTVIEW("\r\nABC DEF\r\n"), {TEXTVIEW(""), TEXTVIEW("ABC DEF")}, KeepEmpty);
+	RunParseLinesTest(TEXTVIEW("\r\n\r\nABC DEF\r\n\r\n"), {TEXTVIEW("ABC DEF")}, SkipEmpty);
+	RunParseLinesTest(TEXTVIEW("\r\n\r\nABC DEF\r\n\r\n"), {TEXTVIEW(""), TEXTVIEW(""), TEXTVIEW("ABC DEF"), TEXTVIEW("")}, KeepEmpty);
+	RunParseLinesTest(TEXTVIEW("ABC\nDEF"), {TEXTVIEW("ABC"), TEXTVIEW("DEF")}, SkipEmpty);
+	RunParseLinesTest(TEXTVIEW("ABC\rDEF"), {TEXTVIEW("ABC"), TEXTVIEW("DEF")}, SkipEmpty);
+	RunParseLinesTest(TEXTVIEW("\r\nABC\r\nDEF\r\n"), {TEXTVIEW("ABC"), TEXTVIEW("DEF")}, SkipEmpty);
+	RunParseLinesTest(TEXTVIEW("\r\nABC\r\nDEF\r\n"), {TEXTVIEW(""), TEXTVIEW("ABC"), TEXTVIEW("DEF")}, KeepEmpty);
+	RunParseLinesTest(TEXTVIEW("\r\nABC\r\n\r\nDEF\r\n"), {TEXTVIEW("ABC"), TEXTVIEW("DEF")}, SkipEmpty);
+	RunParseLinesTest(TEXTVIEW("\r\nABC\r\n\r\nDEF\r\n"), {TEXTVIEW(""), TEXTVIEW("ABC"), TEXTVIEW(""), TEXTVIEW("DEF")}, KeepEmpty);
+	RunParseLinesTest(TEXTVIEW(" \t\r\n\t ABC \t\r\n\t \t\r\n\t DEF \t\r\n"), {TEXTVIEW(" \t"), TEXTVIEW("\t ABC \t"), TEXTVIEW("\t \t"), TEXTVIEW("\t DEF \t")}, SkipEmpty);
+	RunParseLinesTest(TEXTVIEW(" \t\r\n\t ABC \t\r\n\t \t\r\n\t DEF \t\r\n"), {TEXTVIEW("ABC"), TEXTVIEW("DEF")}, SkipEmpty | Trim);
+	RunParseLinesTest(TEXTVIEW(" \t\r\n\t ABC \t\r\n\t \t\r\n\t DEF \t\r\n"), {TEXTVIEW(" \t"), TEXTVIEW("\t ABC \t"), TEXTVIEW("\t \t"), TEXTVIEW("\t DEF \t")}, KeepEmpty);
+	RunParseLinesTest(TEXTVIEW(" \t\r\n\t ABC \t\r\n\t \t\r\n\t DEF \t\r\n"), {TEXTVIEW(""), TEXTVIEW("ABC"), TEXTVIEW(""), TEXTVIEW("DEF")}, KeepEmpty | Trim);
 
 	return true;
 }
diff --git a/Engine/Source/Runtime/Core/Private/Tests/String/ParseTokensTest.cpp b/Engine/Source/Runtime/Core/Private/Tests/String/ParseTokensTest.cpp
--- a/Engine/Source/Runtime/Core/Private/Tests/String/ParseTokensTest.cpp
+++ b/Engine/Source/Runtime/Core/Private/Tests/String/ParseTokensTest.cpp
@@ -12,42 +12,77 @@
 IMPLEMENT_SIMPLE_AUTOMATION_TEST(StringParseTokensByStringTest, "System.Core.String.ParseTokens.ByString", EAutomationTestFlags::ApplicationContextMask | EAutomationTestFlags::EngineFilter)
 bool StringParseTokensByStringTest::RunTest(const FString& Parameters)
 {
-	auto RunParseTokensTest = [this](const FStringView& View, std::initializer_list<FStringView> Delimiters, std::initializer_list<FStringView> ExpectedTokens)
+	using EOptions = UE::String::EParseTokensOptions;
+
+	auto RunParseTokensTest = [this](
+		FStringView View,
+		std::initializer_list<FStringView> Delimiters,
+		std::initializer_list<FStringView> ExpectedTokens,
+		EOptions Options = EOptions::None)
 	{
 		TArray<FStringView, TInlineAllocator<8>> ResultTokens;
-		UE::String::ParseTokensMultiple(View, Delimiters, [&ResultTokens](FStringView Token) { ResultTokens.Add(Token); });
+		if (GetNum(Delimiters) == 1)
+		{
+			UE::String::ParseTokens(View, GetData(Delimiters)[0], ResultTokens, Options);
+		}
+		else
+		{
+			UE::String::ParseTokensMultiple(View, Delimiters, ResultTokens, Options);
+		}
 		if (!Algo::CompareByPredicate(ResultTokens, ExpectedTokens, TEqualTo<>()))
 		{
 			TStringBuilder<512> Error;
-			Error << TEXT("UE::String::ParseTokensMultiple failed to parse \"") << View << TEXT("\" with delimiters {");
-			Error.JoinQuoted(Delimiters, TEXT(", "_SV), TEXT("\""_SV));
-			Error << TEXT("} result {");
-			Error.JoinQuoted(ResultTokens, TEXT(", "_SV), TEXT("\""_SV));
-			Error << TEXT("} expected {");
-			Error.JoinQuoted(ExpectedTokens, TEXT(", "_SV), TEXT("\""_SV));
-			Error << TEXT("}");
+			Error << TEXTVIEW("UE::String::ParseTokens[Multiple] failed to parse \"") << View << TEXTVIEW("\" with delimiters {");
+			Error.JoinQuoted(Delimiters, TEXTVIEW(", "), TEXTVIEW("\""));
+			Error << TEXTVIEW("} result {");
+			Error.JoinQuoted(ResultTokens, TEXTVIEW(", "), TEXTVIEW("\""));
+			Error << TEXTVIEW("} expected {");
+			Error.JoinQuoted(ExpectedTokens, TEXTVIEW(", "), TEXTVIEW("\""));
+			Error << TEXTVIEW("}");
 			AddError(Error.ToString());
 		}
 	};
 
-	RunParseTokensTest(TEXT(""),         {},                       {TEXT("")});
+	constexpr EOptions KeepEmpty = EOptions::None;
+	constexpr EOptions SkipEmpty = EOptions::SkipEmpty;
+	constexpr EOptions IgnoreCase = EOptions::IgnoreCase;
+	constexpr EOptions Trim = EOptions::Trim;
+
+	RunParseTokensTest(TEXT(""),         {},                       {},                                 SkipEmpty);
+	RunParseTokensTest(TEXT(""),         {},                       {TEXT("")},                         KeepEmpty);
 	RunParseTokensTest(TEXT("ABC"),      {},                       {TEXT("ABC")});
 
-	RunParseTokensTest(TEXT(""),         {TEXT(",")},              {TEXT("")});
-	RunParseTokensTest(TEXT(","),        {TEXT(",")},              {TEXT(""), TEXT("")});
-	RunParseTokensTest(TEXT(",,"),       {TEXT(",")},              {TEXT(""), TEXT(""), TEXT("")});
+	RunParseTokensTest(TEXT(""),         {TEXT(",")},              {},                                 SkipEmpty);
+	RunParseTokensTest(TEXT(""),         {TEXT(",")},              {TEXT("")},                         KeepEmpty);
+	RunParseTokensTest(TEXT(","),        {TEXT(",")},              {},                                 SkipEmpty);
+	RunParseTokensTest(TEXT(","),        {TEXT(",")},              {TEXT(""), TEXT("")},               KeepEmpty);
+	RunParseTokensTest(TEXT(",,"),       {TEXT(",")},              {},                                 SkipEmpty);
+	RunParseTokensTest(TEXT(",,"),       {TEXT(",")},              {TEXT(""), TEXT(""), TEXT("")},     KeepEmpty);
+	RunParseTokensTest(TEXT(", ,"),      {TEXT(",")},              {TEXT(" ")},                        SkipEmpty);
+	RunParseTokensTest(TEXT(", ,"),      {TEXT(",")},              {},                                 SkipEmpty | Trim);
+	RunParseTokensTest(TEXT(", ,"),      {TEXT(",")},              {TEXT(""), TEXT(" "), TEXT("")},    KeepEmpty);
+	RunParseTokensTest(TEXT(", ,"),      {TEXT(",")},              {TEXT(""), TEXT(""), TEXT("")},    KeepEmpty | Trim);
 	RunParseTokensTest(TEXT("ABC"),      {TEXT(",")},              {TEXT("ABC")});
-	RunParseTokensTest(TEXT("A,,C"),     {TEXT(",")},              {TEXT("A"), TEXT(""), TEXT("C")});
-	RunParseTokensTest(TEXT("A,B,C"),    {TEXT(",")},              {TEXT("A"), TEXT("B"), TEXT("C")});
-	RunParseTokensTest(TEXT(",A,B,C,"),  {TEXT(",")},              {TEXT(""), TEXT("A"), TEXT("B"), TEXT("C"), TEXT("")});
+	RunParseTokensTest(TEXT("A,,C"),     {TEXT(",")},              {TEXT("A"), TEXT("C")},             SkipEmpty);
+	RunParseTokensTest(TEXT("A,,C"),     {TEXT(",")},              {TEXT("A"), TEXT(""), TEXT("C")},   KeepEmpty);
+	RunParseTokensTest(TEXT("A,\tB\t,C"), {TEXT(",")},             {TEXT("A"), TEXT("\tB\t"), TEXT("C")});
+	RunParseTokensTest(TEXT(",A, B ,C,"), {TEXT(",")},             {TEXT("A"), TEXT(" B "), TEXT("C")},                     SkipEmpty);
+	RunParseTokensTest(TEXT(",A, B ,C,"), {TEXT(",")},             {TEXT(""), TEXT("A"), TEXT(" B "), TEXT("C"), TEXT("")}, KeepEmpty);
 	RunParseTokensTest(TEXT("A\u2022B\u2022C"), {TEXT("\u2022")},  {TEXT("A"), TEXT("B"), TEXT("C")});
 
-	RunParseTokensTest(TEXT("ABCDABCD"), {TEXT("AB")},             {TEXT(""), TEXT("CD"), TEXT("CD")});
-	RunParseTokensTest(TEXT("ABCDABCD"), {TEXT("ABCD")},           {TEXT(""), TEXT(""), TEXT("")});
+	RunParseTokensTest(TEXT("ABCDABCD"), {TEXT("AB")},             {TEXT("CD"), TEXT("CD")},           SkipEmpty);
+	RunParseTokensTest(TEXT("ABCDABCD"), {TEXT("AB")},             {TEXT(""), TEXT("CD"), TEXT("CD")}, KeepEmpty);
+	RunParseTokensTest(TEXT("ABCDABCD"), {TEXT("ABCD")},           {},                                 SkipEmpty);
+	RunParseTokensTest(TEXT("ABCDABCD"), {TEXT("ABCD")},           {TEXT(""), TEXT(""), TEXT("")},     KeepEmpty);
 	RunParseTokensTest(TEXT("ABCDABCD"), {TEXT("DA")},             {TEXT("ABC"), TEXT("BCD")});
 
-	RunParseTokensTest(TEXT("ABCDABCD"), {TEXT("B"),  TEXT("D")},  {TEXT("A"), TEXT("C"), TEXT("A"), TEXT("C"), TEXT("")});
-	RunParseTokensTest(TEXT("ABCDABCD"), {TEXT("BC"), TEXT("DA")}, {TEXT("A"), TEXT(""), TEXT(""), TEXT("D")});
+	RunParseTokensTest(TEXT("ABCDABCD"), {TEXT("B"),  TEXT("D")},  {TEXT("A"), TEXT("C"), TEXT("A"), TEXT("C")},           SkipEmpty);
+	RunParseTokensTest(TEXT("ABCDABCD"), {TEXT("B"),  TEXT("D")},  {TEXT("A"), TEXT("C"), TEXT("A"), TEXT("C"), TEXT("")}, KeepEmpty);
+	RunParseTokensTest(TEXT("ABCDABCD"), {TEXT("BC"), TEXT("DA")}, {TEXT("A"), TEXT("D")},                                 SkipEmpty);
+	RunParseTokensTest(TEXT("ABCDABCD"), {TEXT("BC"), TEXT("DA")}, {TEXT("A"), TEXT(""), TEXT(""), TEXT("D")},             KeepEmpty);
+
+	RunParseTokensTest(TEXT("AbCdaBcDAbCd"), {TEXT("Bc"), TEXT("da")}, {TEXT("AbC"), TEXT("DAbCd")}, SkipEmpty);
+	RunParseTokensTest(TEXT("AbCdaBcDAbCd"), {TEXT("Bc"), TEXT("da")}, {TEXT("A"), TEXT("d")}, SkipEmpty | IgnoreCase);
 
 	RunParseTokensTest(TEXT("A\u2022\u2022B,,C"), {TEXT(",,"), TEXT("\u2022\u2022")},                      {TEXT("A"), TEXT("B"), TEXT("C")});
 	RunParseTokensTest(TEXT("A\u2022\u2022B\u0085\u0085C"), {TEXT("\u0085\u0085"), TEXT("\u2022\u2022")},  {TEXT("A"), TEXT("B"), TEXT("C")});
@@ -58,40 +93,76 @@ bool StringParseTokensByStringTest::RunTest(const FString& Parameters)
 IMPLEMENT_SIMPLE_AUTOMATION_TEST(StringParseTokensByCharTest, "System.Core.String.ParseTokens.ByChar", EAutomationTestFlags::ApplicationContextMask | EAutomationTestFlags::EngineFilter)
 bool StringParseTokensByCharTest::RunTest(const FString& Parameters)
 {
-	auto RunParseTokensTest = [this](const FStringView& View, std::initializer_list<TCHAR> Delimiters, std::initializer_list<FStringView> ExpectedTokens)
+	using EOptions = UE::String::EParseTokensOptions;
+
+	auto RunParseTokensTest = [this](
+		FStringView View,
+		std::initializer_list<TCHAR> Delimiters,
+		std::initializer_list<FStringView> ExpectedTokens,
+		EOptions Options = EOptions::None)
 	{
 		TArray<FStringView, TInlineAllocator<8>> ResultTokens;
-		UE::String::ParseTokensMultiple(View, Delimiters, [&ResultTokens](FStringView Token) { ResultTokens.Add(Token); });
+		if (GetNum(Delimiters) == 1)
+		{
+			UE::String::ParseTokens(View, GetData(Delimiters)[0], ResultTokens, Options);
+		}
+		else
+		{
+			UE::String::ParseTokensMultiple(View, Delimiters, ResultTokens, Options);
+		}
 		if (!Algo::CompareByPredicate(ResultTokens, ExpectedTokens, TEqualTo<>()))
 		{
 			TStringBuilder<512> Error;
-			Error << TEXT("UE::String::ParseTokensMultiple failed to parse \"") << View << TEXT("\" with delimiters {");
-			Error.JoinQuoted(Delimiters, TEXT(", "_SV), TEXT("'"_SV));
-			Error << TEXT("} result {");
-			Error.JoinQuoted(ResultTokens, TEXT(", "_SV), TEXT("\""_SV));
-			Error << TEXT("} expected {");
-			Error.JoinQuoted(ExpectedTokens, TEXT(", "_SV), TEXT("\""_SV));
-			Error << TEXT("}");
+			Error << TEXTVIEW("UE::String::ParseTokens[Multiple] failed to parse \"") << View << TEXTVIEW("\" with delimiters {");
+			Error.JoinQuoted(Delimiters, TEXTVIEW(", "), TEXTVIEW("'"));
+			Error << TEXTVIEW("} result {");
+			Error.JoinQuoted(ResultTokens, TEXTVIEW(", "), TEXTVIEW("\""));
+			Error << TEXTVIEW("} expected {");
+			Error.JoinQuoted(ExpectedTokens, TEXTVIEW(", "), TEXTVIEW("\""));
+			Error << TEXTVIEW("}");
 			AddError(Error.ToString());
 		}
 	};
 
-	RunParseTokensTest(TEXT(""),         {},                       {TEXT("")});
+	constexpr EOptions KeepEmpty = EOptions::None;
+	constexpr EOptions SkipEmpty = EOptions::SkipEmpty;
+	constexpr EOptions IgnoreCase = EOptions::IgnoreCase;
+	constexpr EOptions Trim = EOptions::Trim;
+
+	RunParseTokensTest(TEXT(""),         {},                       {},                                 SkipEmpty);
+	RunParseTokensTest(TEXT(""),         {},                       {TEXT("")},                         KeepEmpty);
 	RunParseTokensTest(TEXT("ABC"),      {},                       {TEXT("ABC")});
 
-	RunParseTokensTest(TEXT(""),         {TEXT(',')},              {TEXT("")});
-	RunParseTokensTest(TEXT(","),        {TEXT(',')},              {TEXT(""), TEXT("")});
-	RunParseTokensTest(TEXT(",,"),       {TEXT(',')},              {TEXT(""), TEXT(""), TEXT("")});
+	RunParseTokensTest(TEXT(""),         {TEXT(',')},              {},                                 SkipEmpty);
+	RunParseTokensTest(TEXT(""),         {TEXT(',')},              {TEXT("")},                         KeepEmpty);
+	RunParseTokensTest(TEXT(","),        {TEXT(',')},              {},                                 SkipEmpty);
+	RunParseTokensTest(TEXT(","),        {TEXT(',')},              {TEXT(""), TEXT("")},               KeepEmpty);
+	RunParseTokensTest(TEXT(",,"),       {TEXT(',')},              {},                                 SkipEmpty);
+	RunParseTokensTest(TEXT(",,"),       {TEXT(',')},              {TEXT(""), TEXT(""), TEXT("")},     KeepEmpty);
+	RunParseTokensTest(TEXT(", ,"),      {TEXT(',')},              {TEXT(" ")},                        SkipEmpty);
+	RunParseTokensTest(TEXT(", ,"),      {TEXT(',')},              {},                                 SkipEmpty | Trim);
+	RunParseTokensTest(TEXT(", ,"),      {TEXT(',')},              {TEXT(""), TEXT(" "), TEXT("")},    KeepEmpty);
+	RunParseTokensTest(TEXT(", ,"),      {TEXT(',')},              {TEXT(""), TEXT(""), TEXT("")},    KeepEmpty | Trim);
 	RunParseTokensTest(TEXT("ABC"),      {TEXT(',')},              {TEXT("ABC")});
-	RunParseTokensTest(TEXT("A,,C"),     {TEXT(',')},              {TEXT("A"), TEXT(""), TEXT("C")});
-	RunParseTokensTest(TEXT("A,B,C"),    {TEXT(',')},              {TEXT("A"), TEXT("B"), TEXT("C")});
-	RunParseTokensTest(TEXT(",A,B,C,"),  {TEXT(',')},              {TEXT(""), TEXT("A"), TEXT("B"), TEXT("C"), TEXT("")});
+	RunParseTokensTest(TEXT("A,,C"),     {TEXT(',')},              {TEXT("A"), TEXT("C")},             SkipEmpty);
+	RunParseTokensTest(TEXT("A,,C"),     {TEXT(',')},              {TEXT("A"), TEXT(""), TEXT("C")},   KeepEmpty);
+	RunParseTokensTest(TEXT("A,\tB\t,C"), {TEXT(',')},             {TEXT("A"), TEXT("\tB\t"), TEXT("C")});
+	RunParseTokensTest(TEXT(",A, B ,C,"), {TEXT(',')},             {TEXT("A"), TEXT(" B "), TEXT("C")},                     SkipEmpty);
+	RunParseTokensTest(TEXT(",A, B ,C,"), {TEXT(',')},             {TEXT(""), TEXT("A"), TEXT(" B "), TEXT("C"), TEXT("")}, KeepEmpty);
 	RunParseTokensTest(TEXT("A\u2022B\u2022C"), {TEXT('\u2022')},  {TEXT("A"), TEXT("B"), TEXT("C")});
 
-	RunParseTokensTest(TEXT("ABCDABCD"), {TEXT('B'),  TEXT('D')},                  {TEXT("A"), TEXT("C"), TEXT("A"), TEXT("C"), TEXT("")});
+	RunParseTokensTest(TEXT("ABCDABCD"), {TEXT('B'),  TEXT('D')},                  {TEXT("A"), TEXT("C"), TEXT("A"), TEXT("C")},           SkipEmpty);
+	RunParseTokensTest(TEXT("ABCDABCD"), {TEXT('B'),  TEXT('D')},                  {TEXT("A"), TEXT("C"), TEXT("A"), TEXT("C"), TEXT("")}, KeepEmpty);
 	RunParseTokensTest(TEXT("A\u2022B,C"), {TEXT(','), TEXT('\u2022')},            {TEXT("A"), TEXT("B"), TEXT("C")});
 	RunParseTokensTest(TEXT("A\u2022B\u0085C"), {TEXT('\u0085'), TEXT('\u2022')},  {TEXT("A"), TEXT("B"), TEXT("C")});
 
+	RunParseTokensTest(TEXT("ABC"), {TEXT('b')}, {TEXT("ABC")}, SkipEmpty);
+	RunParseTokensTest(TEXT("ABC"), {TEXT('b')}, {TEXT("A"), TEXT("C")}, SkipEmpty | IgnoreCase);
+	RunParseTokensTest(TEXT("AbCdaBcD"), {TEXT('B'),  TEXT('d')}, {TEXT("AbC"), TEXT("A"), TEXT("cD")}, SkipEmpty);
+	RunParseTokensTest(TEXT("AbCdaBcD"), {TEXT('B'),  TEXT('d')}, {TEXT("A"), TEXT("C"), TEXT("a"), TEXT("c")}, SkipEmpty | IgnoreCase);
+	RunParseTokensTest(TEXT("A\u2022B\u2022C"), {TEXT('\u2022'), TEXT('b')}, {TEXT("A"), TEXT("B"), TEXT("C")}, SkipEmpty);
+	RunParseTokensTest(TEXT("A\u2022B\u2022C"), {TEXT('\u2022'), TEXT('b')}, {TEXT("A"), TEXT("C")}, SkipEmpty | IgnoreCase);
+
 	return true;
 }
 
diff --git a/Engine/Source/Runtime/Core/Public/String/ParseLines.h b/Engine/Source/Runtime/Core/Public/String/ParseLines.h
--- a/Engine/Source/Runtime/Core/Public/String/ParseLines.h
+++ b/Engine/Source/Runtime/Core/Public/String/ParseLines.h
@@ -3,23 +3,56 @@
 #pragma once
 
 #include "CoreTypes.h"
-#include "Containers/StringFwd.h"
+#include "Containers/StringView.h"
+#include "Misc/EnumClassFlags.h"
 #include "Templates/Function.h"
 
-namespace UE
+namespace UE::String
 {
-namespace String
+
+enum class EParseLinesOptions
 {
+	/** Use the default options when parsing lines. */
+	None = 0,
+	/** Skip lines that are empty or, if trimming, only whitespace. */
+	SkipEmpty = 1 << 0,
+	/** Trim whitespace from each parsed line. */
+	Trim = 1 << 1,
+};
+
+ENUM_CLASS_FLAGS(EParseLinesOptions);
 
 /**
- * Visit every line in the input view as terminated by any of CRLF, CR, LF.
+ * Visit every line in the input string as terminated by any of CRLF, CR, LF.
  *
- * Empty lines are visited by this function.
+ * By default, empty lines are visited.
  *
- * @param View The string view to split into lines.
- * @param Visitor A function that is called for each line.
+ * @param View      A view of the string to split into lines.
+ * @param Visitor   A function that is called for each line.
+ * @param Options   Flags to modify the default behavior.
  */
-CORE_API void ParseLines(const FStringView& View, const TFunctionRef<void(FStringView)>& Visitor);
+CORE_API void ParseLines(
+	FStringView View,
+	TFunctionRef<void (FStringView)> Visitor,
+	EParseLinesOptions Options = EParseLinesOptions::None);
 
+/**
+ * Parse lines in the input string as terminated by any of CRLF, CR, LF.
+ *
+ * Output strings are sub-views of the input view and have the same lifetime as the input view.
+ * By default, empty lines are collected.
+ *
+ * @param View      A view of the string to split into lines.
+ * @param Output    The output to add parsed lines to by calling Output.Add(FStringView).
+ * @param Options   Flags to modify the default behavior.
+ */
+template <typename OutputType>
+inline void ParseLines(
+	const FStringView View,
+	OutputType& Output,
+	const EParseLinesOptions Options = EParseLinesOptions::None)
+{
+	ParseLines(View, [&Output](FStringView Line) { Output.Add(Line); }, Options);
 }
-}
+
+} // UE::String
diff --git a/Engine/Source/Runtime/Core/Public/String/ParseTokens.h b/Engine/Source/Runtime/Core/Public/String/ParseTokens.h
--- a/Engine/Source/Runtime/Core/Public/String/ParseTokens.h
+++ b/Engine/Source/Runtime/Core/Public/String/ParseTokens.h
@@ -4,65 +4,177 @@
 
 #include "CoreTypes.h"
 #include "Containers/ArrayView.h"
-#include "Containers/StringFwd.h"
+#include "Containers/StringView.h"
+#include "Misc/EnumClassFlags.h"
 #include "Templates/Function.h"
 
-namespace UE
+namespace UE::String
 {
-namespace String
+
+enum class EParseTokensOptions
 {
+	/** Use the default options when parsing tokens. */
+	None = 0,
+	/** Ignore case when comparing delimiters against the string. */
+	IgnoreCase = 1 << 0,
+	/** Skip tokens that are empty or, if trimming, only whitespace. */
+	SkipEmpty = 1 << 1,
+	/** Trim whitespace from each parsed token. */
+	Trim = 1 << 2,
+};
+
+ENUM_CLASS_FLAGS(EParseTokensOptions);
 
 /**
- * Visit every token in the input view as separated by any of the delimiters.
+ * Visit every token in the input string, as separated by the delimiter.
  *
- * Comparisons with delimiters are case-sensitive and empty tokens are visited.
- * Behavior is undefined when delimiters overlap each other, such as the delimiters
- * ("AB, "BC") and the input view "1ABC2".
+ * By default, comparisons with the delimiter are case-sensitive and empty tokens are visited.
  *
- * @param View The string view to split into tokens.
- * @param Delimiters An array of non-overlapping non-empty delimiters to split the input string on.
- * @param Visitor A function that is called for each token.
+ * @param View        A view of the string to split into tokens.
+ * @param Delimiter   A delimiter character to split on.
+ * @param Visitor     A function that is called for each token.
+ * @param Options     Flags to modify the default behavior.
  */
-CORE_API void ParseTokensMultiple(const FStringView& View, TArrayView<const FStringView> Delimiters, TFunctionRef<void(FStringView)> Visitor);
+CORE_API void ParseTokens(
+	FStringView View,
+	TCHAR Delimiter,
+	TFunctionRef<void (FStringView)> Visitor,
+	EParseTokensOptions Options = EParseTokensOptions::None);
 
 /**
- * Visit every token in the input view as separated by any of the delimiters.
+ * Parses every token in the input string, as separated by the delimiter.
  *
- * Comparisons with delimiters are case-sensitive and empty tokens are visited.
+ * Output strings are sub-views of the input view and have the same lifetime as the input view.
+ * By default, comparisons with the delimiter are case-sensitive and empty tokens are visited.
  *
- * @param View The string view to split into tokens.
- * @param Delimiters An array of delimiter characters to split the input string on.
- * @param Visitor A function that is called for each token.
+ * @param View        A view of the string to split into tokens.
+ * @param Delimiter   A delimiter character to split on.
+ * @param Output      The output to add parsed tokens to by calling Output.Add(FStringView).
+ * @param Options     Flags to modify the default behavior.
  */
-CORE_API void ParseTokensMultiple(const FStringView& View, TArrayView<const TCHAR> Delimiters, TFunctionRef<void(FStringView)> Visitor);
+template <typename OutputType>
+inline void ParseTokens(
+	const FStringView View,
+	const TCHAR Delimiter,
+	OutputType& Output,
+	const EParseTokensOptions Options = EParseTokensOptions::None)
+{
+	return ParseTokens(View, Delimiter, [&Output](FStringView Token) { Output.Add(Token); }, Options);
+}
 
 /**
- * Visit every token in the input view as separated by the delimiter.
+ * Visit every token in the input string, as separated by the delimiter.
  *
- * Comparisons with the delimiter are case-sensitive and empty tokens are visited.
+ * By default, comparisons with the delimiter are case-sensitive and empty tokens are visited.
  *
- * @param View The string view to split into tokens.
- * @param Delimiter The non-empty delimiter to split the input string on.
- * @param Visitor A function that is called for each token.
+ * @param View        A view of the string to split into tokens.
+ * @param Delimiter   A non-empty delimiter to split on.
+ * @param Visitor     A function that is called for each token.
+ * @param Options     Flags to modify the default behavior.
  */
-inline void ParseTokens(const FStringView& View, const FStringView& Delimiter, TFunctionRef<void(FStringView)> Visitor)
+CORE_API void ParseTokens(
+	FStringView View,
+	FStringView Delimiter,
+	TFunctionRef<void (FStringView)> Visitor,
+	EParseTokensOptions Options = EParseTokensOptions::None);
+
+/**
+ * Parses every token in the input string, as separated by the delimiter.
+ *
+ * Output strings are sub-views of the input view and have the same lifetime as the input view.
+ * By default, comparisons with the delimiter are case-sensitive and empty tokens are visited.
+ *
+ * @param View        A view of the string to split into tokens.
+ * @param Delimiter   A non-empty delimiter to split on.
+ * @param Output      The output to add parsed tokens to by calling Output.Add(FStringView).
+ * @param Options     Flags to modify the default behavior.
+ */
+template <typename OutputType>
+inline void ParseTokens(
+	const FStringView View,
+	const FStringView Delimiter,
+	OutputType& Output,
+	const EParseTokensOptions Options = EParseTokensOptions::None)
 {
-	return ParseTokensMultiple(View, MakeArrayView(&Delimiter, 1), Visitor);
+	return ParseTokens(View, Delimiter, [&Output](FStringView Token) { Output.Add(Token); }, Options);
 }
 
 /**
- * Visit every token in the input view as separated by the delimiter.
+ * Visit every token in the input string, as separated by any of the delimiters.
+ *
+ * By default, comparisons with delimiters are case-sensitive and empty tokens are visited.
+ *
+ * @param View         A view of the string to split into tokens.
+ * @param Delimiters   An array of delimiter characters to split on.
+ * @param Visitor      A function that is called for each token.
+ * @param Options      Flags to modify the default behavior.
+ */
+CORE_API void ParseTokensMultiple(
+	FStringView View,
+	TConstArrayView<TCHAR> Delimiters,
+	TFunctionRef<void (FStringView)> Visitor,
+	EParseTokensOptions Options = EParseTokensOptions::None);
+
+/**
+ * Parses every token in the input string, as separated by any of the delimiters.
  *
- * Comparisons with the delimiter are case-sensitive and empty tokens are visited.
+ * Output strings are sub-views of the input view and have the same lifetime as the input view.
+ * By default, comparisons with delimiters are case-sensitive and empty tokens are visited.
  *
- * @param View The string view to split into tokens.
- * @param Delimiter The delimiter character to split the input string on.
- * @param Visitor A function that is called for each token.
+ * @param View         A view of the string to split into tokens.
+ * @param Delimiters   An array of delimiter characters to split on.
+ * @param Output       The output to add parsed tokens to by calling Output.Add(FStringView).
+ * @param Options      Flags to modify the default behavior.
  */
-inline void ParseTokens(const FStringView& View, TCHAR Delimiter, TFunctionRef<void(FStringView)> Visitor)
+template <typename OutputType>
+inline void ParseTokensMultiple(
+	const FStringView View,
+	const TConstArrayView<TCHAR> Delimiters,
+	OutputType& Output,
+	const EParseTokensOptions Options = EParseTokensOptions::None)
 {
-	return ParseTokensMultiple(View, MakeArrayView(&Delimiter, 1), Visitor);
+	ParseTokensMultiple(View, Delimiters, [&Output](FStringView Token) { Output.Add(Token); }, Options);
 }
 
+/**
+ * Visit every token in the input string, as separated by any of the delimiters.
+ *
+ * By default, comparisons with delimiters are case-sensitive and empty tokens are visited.
+ * Behavior is undefined when delimiters overlap each other, such as the delimiters
+ * ("AB, "BC") and the input string "1ABC2".
+ *
+ * @param View         A view of the string to split into tokens.
+ * @param Delimiters   An array of non-overlapping non-empty delimiters to split on.
+ * @param Visitor      A function that is called for each token.
+ * @param Options      Flags to modify the default behavior.
+ */
+CORE_API void ParseTokensMultiple(
+	FStringView View,
+	TConstArrayView<FStringView> Delimiters,
+	TFunctionRef<void (FStringView)> Visitor,
+	EParseTokensOptions Options = EParseTokensOptions::None);
+
+/**
+ * Parses every token in the input string, as separated by any of the delimiters.
+ *
+ * Output strings are sub-views of the input view and have the same lifetime as the input view.
+ * By default, comparisons with delimiters are case-sensitive and empty tokens are visited.
+ * Behavior is undefined when delimiters overlap each other, such as the delimiters
+ * ("AB, "BC") and the input string "1ABC2".
+ *
+ * @param View         A view of the string to split into tokens.
+ * @param Delimiters   An array of non-overlapping non-empty delimiters to split on.
+ * @param Output       The output to add parsed tokens to by calling Output.Add(FStringView).
+ * @param Options      Flags to modify the default behavior.
+ */
+template <typename OutputType>
+inline void ParseTokensMultiple(
+	const FStringView View,
+	const TConstArrayView<FStringView> Delimiters,
+	OutputType& Output,
+	const EParseTokensOptions Options = EParseTokensOptions::None)
+{
+	ParseTokensMultiple(View, Delimiters, [&Output](FStringView Token) { Output.Add(Token); }, Options);
 }
-}
+
+} // UE::String
