Respect FILEREAD_Silent in FArchiveFileReaderGeneric

diff --git a/Engine/Source/Runtime/Core/Private/HAL/FileManagerGeneric.cpp b/Engine/Source/Runtime/Core/Private/HAL/FileManagerGeneric.cpp
--- a/Engine/Source/Runtime/Core/Private/HAL/FileManagerGeneric.cpp
+++ b/Engine/Source/Runtime/Core/Private/HAL/FileManagerGeneric.cpp
@@ -654,12 +654,13 @@ void FFileManagerGeneric::FindFilesRecursiveInternal( TArray<FString>& FileNames
 	GetLowLevel().IterateDirectoryRecursively(StartDirectory, FileMatch);
 }
 
-FArchiveFileReaderGeneric::FArchiveFileReaderGeneric( IFileHandle* InHandle, const TCHAR* InFilename, int64 InSize, uint32 InBufferSize )
+FArchiveFileReaderGeneric::FArchiveFileReaderGeneric( IFileHandle* InHandle, const TCHAR* InFilename, int64 InSize, uint32 InBufferSize, uint32 InFlags )
 	: Filename( InFilename )
 	, Size( InSize )
 	, Pos( 0 )
 	, BufferBase( 0 )
 	, Handle( InHandle )
+	, Flags( InFlags )
 	, bFirstReadAfterSeek(false)
 {
 	BufferSize = FMath::Min(FMath::RoundUpToPowerOfTwo64((int64)InBufferSize), (uint64)Size);
@@ -690,7 +691,7 @@ void FArchiveFileReaderGeneric::Seek( int64 InPos )
 	{
 		TCHAR ErrorBuffer[1024];
 		SetError();
-		UE_LOG(LogFileManager, Error, TEXT("SetFilePointer on %s Failed %lld/%lld: %lld %s"), *Filename, InPos, Size, Pos, FPlatformMisc::GetSystemErrorMessage(ErrorBuffer, 1024, 0));
+		UE_CLOG( !IsSilent(), LogFileManager, Error, TEXT("SetFilePointer on %s Failed %lld/%lld: %lld %s"), *Filename, InPos, Size, Pos, FPlatformMisc::GetSystemErrorMessage(ErrorBuffer, 1024, 0));
 	}
 
 	Pos = InPos;
@@ -821,7 +822,7 @@ bool FArchiveFileReaderGeneric::InternalPrecache( int64 PrecacheOffset, int64 Pr
 	if (Count!=ReadCount)
 	{
 		TCHAR ErrorBuffer[1024];
-		UE_LOG( LogFileManager, Warning, TEXT( "ReadFile failed: Count=%lld ReadCount=%lld Error=%s" ), Count, ReadCount, FPlatformMisc::GetSystemErrorMessage( ErrorBuffer, 1024, 0 ) );
+		UE_CLOG( !IsSilent(), LogFileManager, Warning, TEXT( "ReadFile failed: Count=%lld ReadCount=%lld Error=%s" ), Count, ReadCount, FPlatformMisc::GetSystemErrorMessage( ErrorBuffer, 1024, 0 ) );
 		BufferCount = WriteOffset + Count;
 		BufferArray.SetNumUninitialized(BufferCount);
 		// The read failed, but we do not SetError or return false just because a precache read fails.
@@ -836,7 +837,7 @@ void FArchiveFileReaderGeneric::Serialize( void* V, int64 Length )
 	if (Pos + Length > Size)
 	{
 		SetError();
-		UE_LOG(LogFileManager, Error, TEXT("Requested read of %d bytes when %d bytes remain (file=%s, size=%d)"), Length, Size-Pos, *Filename, Size);
+		UE_CLOG( !IsSilent(), LogFileManager, Error, TEXT("Requested read of %d bytes when %d bytes remain (file=%s, size=%d)"), Length, Size-Pos, *Filename, Size);
 		return;
 	}
 
@@ -869,7 +870,7 @@ void FArchiveFileReaderGeneric::Serialize( void* V, int64 Length )
 				{
 					TCHAR ErrorBuffer[1024];
 					SetError();
-					UE_LOG( LogFileManager, Warning, TEXT( "ReadFile failed: Count=%lld Length=%lld Error=%s for file %s" ), 
+					UE_CLOG( !IsSilent(), LogFileManager, Warning, TEXT( "ReadFile failed: Count=%lld Length=%lld Error=%s for file %s" ), 
 						Count, Length, FPlatformMisc::GetSystemErrorMessage( ErrorBuffer, 1024, 0 ), *Filename );
 				}
 				Pos += Length;
@@ -878,14 +879,14 @@ void FArchiveFileReaderGeneric::Serialize( void* V, int64 Length )
 			if (!InternalPrecache(Pos, MAX_int32))
 			{
 				SetError();
-				UE_LOG( LogFileManager, Warning, TEXT( "ReadFile failed during precaching for file %s" ),*Filename );
+				UE_CLOG( !IsSilent(), LogFileManager, Warning, TEXT( "ReadFile failed during precaching for file %s" ),*Filename );
 				return;
 			}
 			Copy = FMath::Min( Length, BufferBase+BufferArray.Num()-Pos );
 			if( Copy<=0 )
 			{
 				SetError();
-				UE_LOG( LogFileManager, Error, TEXT( "ReadFile beyond EOF %lld+%lld/%lld for file %s" ), 
+				UE_CLOG( !IsSilent(), LogFileManager, Error, TEXT( "ReadFile beyond EOF %lld+%lld/%lld for file %s" ), 
 					Pos, Length, Size, *Filename );
 			}
 			if( IsError() )
diff --git a/Engine/Source/Runtime/Core/Public/HAL/FileManagerGeneric.h b/Engine/Source/Runtime/Core/Public/HAL/FileManagerGeneric.h
--- a/Engine/Source/Runtime/Core/Public/HAL/FileManagerGeneric.h
+++ b/Engine/Source/Runtime/Core/Public/HAL/FileManagerGeneric.h
@@ -224,7 +224,7 @@ private:
 class CORE_API FArchiveFileReaderGeneric : public FArchive
 {
 public:
-	FArchiveFileReaderGeneric( IFileHandle* InHandle, const TCHAR* InFilename, int64 InSize, uint32 InBufferSize = PLATFORM_FILE_READER_BUFFER_SIZE );
+	FArchiveFileReaderGeneric( IFileHandle* InHandle, const TCHAR* InFilename, int64 InSize, uint32 InBufferSize = PLATFORM_FILE_READER_BUFFER_SIZE, uint32 InFlags = FILEREAD_None );
 	~FArchiveFileReaderGeneric();
 
 	virtual void Seek( int64 InPos ) final;
@@ -264,6 +264,12 @@ protected:
 	**/
 	virtual void ReadLowLevel(uint8* Dest, int64 CountToRead, int64& OutBytesRead);
 
+	/** Returns true if the archive should suppress logging in case of error */
+	bool IsSilent() const
+	{
+		return !!(Flags & FILEREAD_Silent);
+	}
+
 	/** Filename for debugging purposes. */
 	FString Filename;
 	int64 Size;
@@ -276,6 +282,7 @@ protected:
 	 */
 	TArray64<uint8> BufferArray;
 	int64 BufferSize;
+	uint32 Flags;
 	bool bFirstReadAfterSeek;
 
 	enum
