CompactBinary: Added a missing call to Ar.SetError() when loading attachments from an archive

#rnx

diff --git a/Engine/Source/Runtime/Core/Private/Serialization/CompactBinaryPackage.cpp b/Engine/Source/Runtime/Core/Private/Serialization/CompactBinaryPackage.cpp
--- a/Engine/Source/Runtime/Core/Private/Serialization/CompactBinaryPackage.cpp
+++ b/Engine/Source/Runtime/Core/Private/Serialization/CompactBinaryPackage.cpp
@@ -57,6 +57,7 @@ static bool TryLoadAttachmentFromArchive(FCbAttachment& OutAttachment, FCbField&
 			return true;
 		}
 	}
+	Ar.SetError();
 	return false;
 }
 
