CompactBinary: Added validation functions for FCb[Type] for ease of use

#rnx

diff --git a/Engine/Source/Runtime/Core/Private/Serialization/CompactBinaryValidation.cpp b/Engine/Source/Runtime/Core/Private/Serialization/CompactBinaryValidation.cpp
--- a/Engine/Source/Runtime/Core/Private/Serialization/CompactBinaryValidation.cpp
+++ b/Engine/Source/Runtime/Core/Private/Serialization/CompactBinaryValidation.cpp
@@ -6,6 +6,7 @@
 #include "IO/IoHash.h"
 #include "Misc/ByteSwap.h"
 #include "Serialization/CompactBinary.h"
+#include "Serialization/CompactBinaryPackage.h"
 #include "Serialization/VarInt.h"
 
 ///////////////////////////////////////////////////////////////////////////////////////////////////////////////////////
@@ -529,6 +530,17 @@ static FIoHash ValidateCbPackageObject(FCbFieldView& Value, FMemoryView& View, E
 	return FIoHash();
 }
 
+static ECbValidateError ValidateCbFieldView(FCbFieldView Value, ECbValidateMode Mode)
+{
+	struct FCopy : public FCbFieldView
+	{
+		using FCbFieldView::GetType;
+		using FCbFieldView::GetViewNoType;
+	};
+	const FCopy& Copy = static_cast<const FCopy&>(Value);
+	return ValidateCompactBinary(Copy.GetViewNoType(), Mode, Copy.GetType());
+}
+
 ECbValidateError ValidateCompactBinary(FMemoryView View, ECbValidateMode Mode, ECbFieldType Type)
 {
 	ECbValidateError Error = ECbValidateError::None;
@@ -634,4 +646,74 @@ ECbValidateError ValidateCompactBinaryPackage(FMemoryView View, ECbValidateMode
 	return Error;
 }
 
+ECbValidateError ValidateCompactBinary(const FCbField& Value, ECbValidateMode Mode)
+{
+	return ValidateCbFieldView(Value, Mode);
+}
+
+ECbValidateError ValidateCompactBinary(const FCbArray& Value, ECbValidateMode Mode)
+{
+	return ValidateCbFieldView(Value.AsFieldView(), Mode);
+}
+
+ECbValidateError ValidateCompactBinary(const FCbObject& Value, ECbValidateMode Mode)
+{
+	return ValidateCbFieldView(Value.AsFieldView(), Mode);
+}
+
+ECbValidateError ValidateCompactBinary(const FCbPackage& Value, ECbValidateMode Mode)
+{
+	ECbValidateError Error = ValidateCompactBinary(Value.GetObject(), Mode);
+
+	for (const FCbAttachment& Attachment : Value.GetAttachments())
+	{
+		Error |= ValidateCompactBinary(Attachment, Mode);
+	}
+
+	if (EnumHasAnyFlags(Mode, ECbValidateMode::PackageHash) && (Value.GetObjectHash() != Value.GetObject().GetHash()))
+	{
+		AddError(Error, ECbValidateError::InvalidPackageHash);
+	}
+
+	return Error;
+}
+
+ECbValidateError ValidateCompactBinary(const FCbAttachment& Value, ECbValidateMode Mode)
+{
+	ECbValidateError Error = ECbValidateError::None;
+
+	if (Value.IsObject())
+	{
+		Error |= ValidateCompactBinary(Value.AsObject(), Mode);
+	}
+
+	if (EnumHasAnyFlags(Mode, ECbValidateMode::PackageHash))
+	{
+		if (Value.IsObject())
+		{
+			if (Value.GetHash() != Value.AsObject().GetHash())
+			{
+				AddError(Error, ECbValidateError::InvalidPackageHash);
+			}
+		}
+		else if (Value.IsBinary())
+		{
+			if (Value.GetHash() != FIoHash::HashBuffer(Value.AsCompositeBinary()))
+			{
+				AddError(Error, ECbValidateError::InvalidPackageHash);
+			}
+		}
+		else if (Value.IsCompressedBinary())
+		{
+			const FCompressedBuffer& Buffer = Value.AsCompressedBinary();
+			if (Buffer.GetRawHash() != FIoHash::HashBuffer(Buffer.DecompressToComposite()))
+			{
+				AddError(Error, ECbValidateError::InvalidPackageHash);
+			}
+		}
+	}
+
+	return Error;
+}
+
 ///////////////////////////////////////////////////////////////////////////////////////////////////////////////////////
diff --git a/Engine/Source/Runtime/Core/Public/Serialization/CompactBinaryValidation.h b/Engine/Source/Runtime/Core/Public/Serialization/CompactBinaryValidation.h
--- a/Engine/Source/Runtime/Core/Public/Serialization/CompactBinaryValidation.h
+++ b/Engine/Source/Runtime/Core/Public/Serialization/CompactBinaryValidation.h
@@ -7,6 +7,9 @@
 #include "Misc/EnumClassFlags.h"
 #include "Serialization/CompactBinary.h"
 
+class FCbAttachment;
+class FCbPackage;
+
 ///////////////////////////////////////////////////////////////////////////////////////////////////////////////////////
 
 /** Flags for validating compact binary data. */
@@ -184,4 +187,17 @@ CORE_API ECbValidateError ValidateCompactBinaryAttachment(FMemoryView View, ECbV
  */
 CORE_API ECbValidateError ValidateCompactBinaryPackage(FMemoryView View, ECbValidateMode Mode);
 
+/**
+ * Validate the compact binary data for one value as specified by the mode flags.
+ *
+ * Validation recurses into attachments, objects, arrays, and fields within the top-level value.
+ *
+ * @return None on success, otherwise the flags for the types of errors that were detected.
+ */
+CORE_API ECbValidateError ValidateCompactBinary(const FCbField& Value, ECbValidateMode Mode);
+CORE_API ECbValidateError ValidateCompactBinary(const FCbArray& Value, ECbValidateMode Mode);
+CORE_API ECbValidateError ValidateCompactBinary(const FCbObject& Value, ECbValidateMode Mode);
+CORE_API ECbValidateError ValidateCompactBinary(const FCbPackage& Value, ECbValidateMode Mode);
+CORE_API ECbValidateError ValidateCompactBinary(const FCbAttachment& Value, ECbValidateMode Mode);
+
 ///////////////////////////////////////////////////////////////////////////////////////////////////////////////////////
