CompressedBuffer: Added FCompressedBufferReader::HasSource()

#rnx

diff --git a/Engine/Source/Runtime/Core/Public/Compression/CompressedBuffer.h b/Engine/Source/Runtime/Core/Public/Compression/CompressedBuffer.h
--- a/Engine/Source/Runtime/Core/Public/Compression/CompressedBuffer.h
+++ b/Engine/Source/Runtime/Core/Public/Compression/CompressedBuffer.h
@@ -213,6 +213,8 @@ public:
 	CORE_API void SetSource(FArchive& Archive);
 	CORE_API void SetSource(const FCompressedBuffer& Buffer);
 
+	[[nodiscard]] inline bool HasSource() const { return SourceArchive || SourceBuffer; }
+
 	/** Returns the size of the compressed data. Zero on error. */
 	[[nodiscard]] CORE_API uint64 GetCompressedSize();
 
