[Insights] TraceMalloc should not use it's own root heap

 * TraceMalloc should not use it's own root heap. Instead it will use a regular heap in system memory and make fake heaps (since we can't track virtual allocations for trace's own allocations).
 * The first named was miscategorized as a root heap.

diff --git a/Engine/Source/Runtime/Core/Private/ProfilingDebugging/MemoryAllocationTrace.cpp b/Engine/Source/Runtime/Core/Private/ProfilingDebugging/MemoryAllocationTrace.cpp
--- a/Engine/Source/Runtime/Core/Private/ProfilingDebugging/MemoryAllocationTrace.cpp
+++ b/Engine/Source/Runtime/Core/Private/ProfilingDebugging/MemoryAllocationTrace.cpp
@@ -134,8 +134,6 @@ void FAllocationTrace::Initialize() const
 	check(SystemRootHeap == EMemoryTraceRootHeap::SystemMemory);
 	const HeapId VideoRootHeap = RootHeapSpec(TEXT("Video memory"));
 	check(VideoRootHeap == EMemoryTraceRootHeap::VideoMemory);
-	const HeapId TraceRootHeap = RootHeapSpec(TEXT("Trace"));
-	check(TraceRootHeap == EMemoryTraceRootHeap::TraceMemory);
 
 	static_assert((1 << SizeShift) - 1 <= MIN_ALIGNMENT, "Not enough bits to pack size fields");
 }
@@ -299,7 +297,7 @@ void FAllocationTrace::ReallocFree(void* Address, HeapId RootHeap)
 ////////////////////////////////////////////////////////////////////////////////
 HeapId FAllocationTrace::HeapSpec(HeapId ParentId, const TCHAR* Name, EMemoryTraceHeapFlags Flags) const
 {
-	static std::atomic<HeapId> HeapIdCount(15); //Reserve indexes for root heaps
+	static std::atomic<HeapId> HeapIdCount(EMemoryTraceRootHeap::EndReserved + 1); //Reserve indexes for root heaps
 	const HeapId Id = HeapIdCount.fetch_add(1);
 	const uint32 NameLen = FCString::Strlen(Name);
 	const uint32 DataSize = NameLen * sizeof(TCHAR);
@@ -382,10 +380,12 @@ void* FTraceMalloc::Malloc(SIZE_T Count, uint32 Alignment)
 	const uint32 AlignmentPow2 = uint32(FPlatformMath::CountTrailingZeros(Alignment));
 	const uint32 Alignment_SizeLower = (AlignmentPow2 << SizeShift) | uint32(Size & ((1 << SizeShift) - 1));
 
+	UE_MEMSCOPE(TRACE_TAG);
+	
 	UE_TRACE_LOG(Memory, Alloc, MemAllocChannel)
 		<< Alloc.CallstackId(0)
 		<< Alloc.Address(uint64(NewPtr))
-		<< Alloc.RootHeap(uint8(EMemoryTraceRootHeap::TraceMemory))
+		<< Alloc.RootHeap(uint8(EMemoryTraceRootHeap::SystemMemory))
 		<< Alloc.Size(uint32(Size >> SizeShift))
 		<< Alloc.AlignmentPow2_SizeLower(uint8(Alignment_SizeLower));
 #endif //UE_MEMORY_TRACE_ENABLED
@@ -409,12 +409,14 @@ void* FTraceMalloc::Realloc(void* Original, SIZE_T Count, uint32 Alignment)
 
 	UE_TRACE_LOG(Memory, ReallocFree, MemAllocChannel)
 		<< ReallocFree.Address(uint64(Original))
-		<< ReallocFree.RootHeap(uint8(EMemoryTraceRootHeap::TraceMemory));
+		<< ReallocFree.RootHeap(uint8(EMemoryTraceRootHeap::SystemMemory));
 
+	UE_MEMSCOPE(TRACE_TAG);
+	
 	UE_TRACE_LOG(Memory, ReallocAlloc, MemAllocChannel)
 		<< ReallocAlloc.CallstackId(0)
 		<< ReallocAlloc.Address(uint64(NewPtr))
-		<< ReallocAlloc.RootHeap(uint8(EMemoryTraceRootHeap::TraceMemory))
+		<< ReallocAlloc.RootHeap(uint8(EMemoryTraceRootHeap::SystemMemory))
 		<< ReallocAlloc.Size(uint32(Size >> SizeShift))
 		<< ReallocAlloc.AlignmentPow2_SizeLower(uint8(Alignment_SizeLower));
 #endif //UE_MEMORY_TRACE_ENABLED
@@ -433,7 +435,7 @@ void FTraceMalloc::Free(void* Original)
 #if UE_MEMORY_TRACE_ENABLED
 	UE_TRACE_LOG(Memory, Free, MemAllocChannel)
 		<< Free.Address(uint64(Original))
-		<< Free.RootHeap(uint8(EMemoryTraceRootHeap::TraceMemory));
+		<< Free.RootHeap(uint8(EMemoryTraceRootHeap::SystemMemory));
 #endif //UE_MEMORY_TRACE_ENABLED
 }
 
diff --git a/Engine/Source/Runtime/Core/Public/ProfilingDebugging/MemoryTrace.h b/Engine/Source/Runtime/Core/Public/ProfilingDebugging/MemoryTrace.h
--- a/Engine/Source/Runtime/Core/Public/ProfilingDebugging/MemoryTrace.h
+++ b/Engine/Source/Runtime/Core/Public/ProfilingDebugging/MemoryTrace.h
@@ -32,8 +32,7 @@ enum EMemoryTraceRootHeap : uint8
 {
 	SystemMemory, // RAM
 	VideoMemory, // VRAM
-	TraceMemory, // Memory used by Trace itself
-	EndHardcoded = TraceMemory,
+	EndHardcoded = VideoMemory,
 	EndReserved = 15
 };
 
@@ -140,7 +139,7 @@ CORE_API void MemoryTrace_ReallocAlloc(uint64 Address, uint64 NewSize, uint32 Al
 #else
 
 inline HeapId MemoryTrace_RootHeapSpec(const TCHAR* Name, EMemoryTraceHeapFlags Flags = EMemoryTraceHeapFlags::None) { return ~0; };
-inline HeapId MemoryTrace_HeapSpec(const TCHAR* Name, EMemoryTraceHeapFlags Flags = EMemoryTraceHeapFlags::None) { return ~0; }
+inline HeapId MemoryTrace_HeapSpec(HeapId ParentId, const TCHAR* Name, EMemoryTraceHeapFlags Flags = EMemoryTraceHeapFlags::None) { return ~0; }
 inline void MemoryTrace_MarkAllocAsHeap(uint64 Address, HeapId Heap) {}
 inline void MemoryTrace_UnmarkAllocAsHeap(uint64 Address, HeapId Heap) {}
 inline void MemoryTrace_Alloc(uint64 Address, uint64 Size, uint32 Alignment, HeapId RootHeap = EMemoryTraceRootHeap::SystemMemory) {}
