Fixed a compilation issue with stall detector when COUNTERSTRACE_ENABLED is 0

diff --git a/Engine/Source/Runtime/Core/Public/ProfilingDebugging/StallDetector.h b/Engine/Source/Runtime/Core/Public/ProfilingDebugging/StallDetector.h
--- a/Engine/Source/Runtime/Core/Public/ProfilingDebugging/StallDetector.h
+++ b/Engine/Source/Runtime/Core/Public/ProfilingDebugging/StallDetector.h
@@ -15,11 +15,11 @@
 #include <atomic>
 
 #ifndef STALL_DETECTOR
- #if WITH_EDITOR && PLATFORM_WINDOWS && !UE_BUILD_SHIPPING
+ #if WITH_EDITOR && PLATFORM_WINDOWS && !UE_BUILD_SHIPPING && COUNTERSTRACE_ENABLED
   #define STALL_DETECTOR 1
- #else // WITH_EDITOR && PLATFORM_WINDOWS && !UE_BUILD_SHIPPING
+ #else // WITH_EDITOR && PLATFORM_WINDOWS && !UE_BUILD_SHIPPING && COUNTERSTRACE_ENABLED
   #define STALL_DETECTOR 0
- #endif // WITH_EDITOR && PLATFORM_WINDOWS && !UE_BUILD_SHIPPING
+ #endif // WITH_EDITOR && PLATFORM_WINDOWS && !UE_BUILD_SHIPPING && COUNTERSTRACE_ENABLED
 #endif // STALL_DETECTOR
 
 #if STALL_DETECTOR
