removed FPlatformProcess::ReadCycleCounter() as it can't be implemented cross-platform. it wasn't released yet so no need for deprecation

diff --git a/Engine/Source/Runtime/Core/Public/GenericPlatform/GenericPlatformProcess.h b/Engine/Source/Runtime/Core/Public/GenericPlatform/GenericPlatformProcess.h
--- a/Engine/Source/Runtime/Core/Public/GenericPlatform/GenericPlatformProcess.h
+++ b/Engine/Source/Runtime/Core/Public/GenericPlatform/GenericPlatformProcess.h
@@ -750,58 +750,15 @@ struct CORE_API FGenericPlatformProcess
 #endif
 	}
 
-	/**
-	* Read the cycle counter of the CPU
-	*/
-	static FORCEINLINE uint64 ReadCycleCounter()
-	{
-#if PLATFORM_APPLE
-		return mach_absolute_time();
-#elif defined(_MSC_VER) && PLATFORM_CPU_X86_FAMILY
-		return __rdtsc();
-#elif defined(_MSC_VER) && defined(_M_ARM)
-		return __rdpmccntr64();
-#elif defined(_MSC_VER) && defined(_M_ARM64)
-		// Duplicate define of ARM64_PMCCNTR_EL0 from winnt.h to avoid including the entire header
-#		define UE_ARM64_SYSREG(op0, op1, crn, crm, op2) \
-		( ((op0 & 1) << 14) | \
-		  ((op1 & 7) << 11) | \
-		  ((crn & 15) << 7) | \
-		  ((crm & 15) << 3) | \
-		  ((op2 & 7) << 0) )
-#		define UE_ARM64_PMCCNTR_EL0 UE_ARM64_SYSREG(3,3, 9,13,0)  // Cycle Count Register [CP15_PMCCNTR]	
-
-		return _ReadStatusReg(UE_ARM64_PMCCNTR_EL0);
-
-#		undef UE_ARM64_PMCCNTR_EL0
-#		undef UE_ARM64_SYSREG
-#elif __has_builtin(__builtin_readcyclecounter)
-		return __builtin_readcyclecounter();
-#else
-#	error Unsupported architecture!
-#endif
-	}
-
 	/**
 	* Tells the processor to pause for at least the amount of cycles given. Is used for spin-loops to improve the speed at 
 	* which the code detects the release of the lock and power-consumption.
 	*/
-	static FORCEINLINE void YieldCycles(uint64 cycles)
+	static FORCEINLINE void YieldCycles(uint64 Cycles)
 	{
-#if PLATFORM_ANDROID_ARM64 || (PLATFORM_APPLE && !PLATFORM_MAC_X86) || PLATFORM_LINUXARM64
-		// We can't read cycle counter from user mode on these platform
-		for (uint64 i = 0; i < cycles; i++)
+		for (uint64 i = 0; i < Cycles; i++)
 		{
 			Yield();
 		}
-#else
-		uint64 start = ReadCycleCounter();
-		//some 32bit implementations return 0 for __builtin_readcyclecounter just to be on the safe side we protect against this.
-		cycles = start != 0 ? cycles : 0;
-		do
-		{
-			Yield();
-		} while((ReadCycleCounter() - start) < cycles);
-#endif
 	}
 };
