Fix unity issue.

diff --git a/Engine/Source/Runtime/Core/Private/ProfilingDebugging/MemoryAllocationTrace.cpp b/Engine/Source/Runtime/Core/Private/ProfilingDebugging/MemoryAllocationTrace.cpp
--- a/Engine/Source/Runtime/Core/Private/ProfilingDebugging/MemoryAllocationTrace.cpp
+++ b/Engine/Source/Runtime/Core/Private/ProfilingDebugging/MemoryAllocationTrace.cpp
@@ -2,6 +2,7 @@
 
 #include "ProfilingDebugging/MemoryAllocationTrace.h"
 #include "HAL/PlatformTime.h"
+#include "ProfilingDebugging/TagTrace.h"
 #include "ProfilingDebugging/TraceMalloc.h"
 
 #if UE_TRACE_ENABLED
