[Insights] Added Trace.status console command

Added console command to print status of trace to in game console. In order to list the available channels a enumeration method was added.

diff --git a/Engine/Source/Runtime/Core/Private/ProfilingDebugging/TraceAuxiliary.cpp b/Engine/Source/Runtime/Core/Private/ProfilingDebugging/TraceAuxiliary.cpp
--- a/Engine/Source/Runtime/Core/Private/ProfilingDebugging/TraceAuxiliary.cpp
+++ b/Engine/Source/Runtime/Core/Private/ProfilingDebugging/TraceAuxiliary.cpp
@@ -533,6 +533,78 @@ static void TraceAuxiliaryResume()
 	GTraceAuxiliary.EnableChannels();
 }
 
+////////////////////////////////////////////////////////////////////////////////
+static void TraceAuxiliaryStatus()
+{
+	UE_LOG(LogConsoleResponse, Display, TEXT("Trace status -------------------------------"));
+	
+	// Status of data connection
+	const TCHAR* Dest = GTraceAuxiliary.GetDest();
+	TStringBuilder<256> ConnectionStr;
+	if (Dest && FCString::Strlen(Dest) > 0)
+	{
+		ConnectionStr.Appendf(TEXT("Tracing to '%s'"), Dest);
+	}
+	else
+	{
+		ConnectionStr.Append(TEXT("Not tracing"));
+	}
+	UE_LOG(LogConsoleResponse, Display, TEXT("- Connection: %.*s"), ConnectionStr.ToView().Len(), ConnectionStr.ToView().GetData());
+	
+	// Stats
+	UE::Trace::FStatistics Stats;
+	UE::Trace::GetStatistics(Stats);
+	constexpr float MbDiv = 1024*1024;
+	UE_LOG(LogConsoleResponse, Display, TEXT("- Memory used: %.02f Mb"),
+		float(Stats.MemoryUsed) / MbDiv);
+	UE_LOG(LogConsoleResponse, Display, TEXT("- Important Events cache: %.02f Mb, waste: %0.02f Mb"),
+		float(Stats.CacheUsed) / MbDiv,
+		float(Stats.CacheWaste) / MbDiv);
+	UE_LOG(LogConsoleResponse, Display, TEXT("- Sent: %.2f Mb"),
+		float(Stats.BytesSent) / MbDiv);
+	
+	// Channels
+	struct EnumerateType
+	{
+		TStringBuilder<128> EnabledChannelsStr;
+		uint32 EnabledChannelsLastLineBreak = 0;
+		TStringBuilder<1024> DisabledChannelsStr;
+		uint32 DisabledChannelsLastLineBreak = 0;
+
+		void InsertLineBreaks()
+		{
+			if ((EnabledChannelsStr.Len() - EnabledChannelsLastLineBreak) > 100)
+			{
+				EnabledChannelsStr << TEXT("\n");
+				EnabledChannelsLastLineBreak = DisabledChannelsStr.Len();
+			}
+			if ((DisabledChannelsStr.Len() - DisabledChannelsLastLineBreak) > 100)
+			{
+				DisabledChannelsStr << TEXT("\n");
+				DisabledChannelsLastLineBreak = DisabledChannelsStr.Len();
+			}
+		}
+	} EnumerateUserData;
+	UE::Trace::EnumerateChannels([](const ANSICHAR* Name, bool bEnabled, void* User)
+	{
+		EnumerateType* EnumerateUser = static_cast<EnumerateType*>(User);
+		FAnsiStringView NameView = FAnsiStringView(Name).LeftChop(7); // Remove "Channel" suffix
+		if (bEnabled)
+		{
+			EnumerateUser->EnabledChannelsStr << NameView << TEXT(", ");
+		}
+		else
+		{
+			EnumerateUser->DisabledChannelsStr << NameView << TEXT(", ");
+		}
+		EnumerateUser->InsertLineBreaks();
+	}, &EnumerateUserData);
+	UE_LOG(LogConsoleResponse, Display, TEXT("- Enabled channels:\n %s"), EnumerateUserData.EnabledChannelsStr.ToString());
+	UE_LOG(LogConsoleResponse, Display, TEXT("- Available channels:\n %s"), EnumerateUserData.DisabledChannelsStr.ToString());
+	
+	UE_LOG(LogConsoleResponse, Display, TEXT("-------------------------------------------"));
+}
+
 ////////////////////////////////////////////////////////////////////////////////
 static FAutoConsoleCommand TraceAuxiliarySendCmd(
 	TEXT("Trace.Send"),
@@ -571,6 +643,12 @@ static FAutoConsoleCommand TraceAuxiliaryResumeCmd(
 	FConsoleCommandDelegate::CreateStatic(TraceAuxiliaryResume)
 );
 
+////////////////////////////////////////////////////////////////////////////////
+static FAutoConsoleCommand TraceAuxiliaryStatusCmd(
+	TEXT("Trace.Status"),
+	TEXT("Print Trace status to console"),
+	FConsoleCommandDelegate::CreateStatic(TraceAuxiliaryStatus)
+);
 #endif // UE_TRACE_ENABLED
 
 
diff --git a/Engine/Source/Runtime/TraceLog/Private/Trace/Channel.cpp b/Engine/Source/Runtime/TraceLog/Private/Trace/Channel.cpp
--- a/Engine/Source/Runtime/TraceLog/Private/Trace/Channel.cpp
+++ b/Engine/Source/Runtime/TraceLog/Private/Trace/Channel.cpp
@@ -217,7 +217,7 @@ FChannel* FChannel::FindChannel(const ANSICHAR* ChannelName)
 	};
 	for (FChannel* Channel : ChannelLists)
 	{
-		for (; Channel != nullptr; Channel = (FChannel*)(Channel->Next))
+		for (; Channel != nullptr; Channel = Channel->Next)
 		{
 			if (Channel->Name.Hash == ChannelNameHash)
 			{
@@ -229,6 +229,24 @@ FChannel* FChannel::FindChannel(const ANSICHAR* ChannelName)
 	return nullptr;
 }
 
+///////////////////////////////////////////////////////////////////////////////
+void FChannel::EnumerateChannels(ChannelIterFunc Func, void* User) 
+{
+	using namespace Private;
+	FChannel* ChannelLists[] =
+	{
+		AtomicLoadAcquire(&GNewChannelList),
+		AtomicLoadAcquire(&GHeadChannel),
+	};
+	for (FChannel* Channel : ChannelLists)
+	{
+		for (; Channel != nullptr; Channel = Channel->Next)
+		{
+			Func(Channel->Name.Ptr, Channel->IsEnabled(), User);
+		}
+	}
+}
+
 ///////////////////////////////////////////////////////////////////////////////
 bool FChannel::Toggle(bool bEnabled)
 {
diff --git a/Engine/Source/Runtime/TraceLog/Private/Trace/Trace.cpp b/Engine/Source/Runtime/TraceLog/Private/Trace/Trace.cpp
--- a/Engine/Source/Runtime/TraceLog/Private/Trace/Trace.cpp
+++ b/Engine/Source/Runtime/TraceLog/Private/Trace/Trace.cpp
@@ -123,7 +123,13 @@ bool ToggleChannel(const TCHAR* ChannelName, bool bEnabled)
 	ToAnsiCheap(ChannelNameA, ChannelName);
 	return FChannel::Toggle(ChannelNameA, bEnabled);
 }
-	
+
+////////////////////////////////////////////////////////////////////////////////
+void EnumerateChannels(ChannelIterFunc IterFunc, void* User)
+{
+	FChannel::EnumerateChannels(IterFunc, User);
+}
+
 ////////////////////////////////////////////////////////////////////////////////
 void StartWorkerThread()
 {
diff --git a/Engine/Source/Runtime/TraceLog/Public/Trace/Detail/Channel.h b/Engine/Source/Runtime/TraceLog/Public/Trace/Detail/Channel.h
--- a/Engine/Source/Runtime/TraceLog/Public/Trace/Detail/Channel.h
+++ b/Engine/Source/Runtime/TraceLog/Public/Trace/Detail/Channel.h
@@ -10,6 +10,8 @@
 
 namespace UE {
 namespace Trace {
+	
+typedef void ChannelIterFunc(const ANSICHAR*, bool, void*);	
 
 /*
 	A named channel which can be used to filter trace events. Channels can be
@@ -50,6 +52,7 @@ public:
 	static bool			Toggle(const ANSICHAR* ChannelName, bool bEnabled);
 	static void			ToggleAll(bool bEnabled);
 	static FChannel*	FindChannel(const ANSICHAR* ChannelName);
+	static void			EnumerateChannels(ChannelIterFunc Func, void* User);
 	bool				Toggle(bool bEnabled);
 	bool				IsEnabled() const;
 	explicit			operator bool () const;
diff --git a/Engine/Source/Runtime/TraceLog/Public/Trace/Trace.h b/Engine/Source/Runtime/TraceLog/Public/Trace/Trace.h
--- a/Engine/Source/Runtime/TraceLog/Public/Trace/Trace.h
+++ b/Engine/Source/Runtime/TraceLog/Public/Trace/Trace.h
@@ -30,6 +30,7 @@ struct FInitializeDesc
 
 typedef void*		AllocFunc(SIZE_T, uint32);
 typedef void		FreeFunc(void*, SIZE_T);
+typedef void		ChannelIterFunc(const ANSICHAR*, bool, void*);	
 
 struct FStatistics
 {
@@ -52,6 +53,7 @@ UE_TRACE_API bool	IsTracing() UE_TRACE_IMPL(false);
 UE_TRACE_API bool	Stop() UE_TRACE_IMPL(false);
 UE_TRACE_API bool	IsChannel(const TCHAR* ChanneName) UE_TRACE_IMPL(false);
 UE_TRACE_API bool	ToggleChannel(const TCHAR* ChannelName, bool bEnabled) UE_TRACE_IMPL(false);
+UE_TRACE_API void	EnumerateChannels(ChannelIterFunc IterFunc, void* User);
 UE_TRACE_API void	ThreadRegister(const TCHAR* Name, uint32 SystemId, int32 SortHint) UE_TRACE_IMPL();
 UE_TRACE_API void	ThreadGroupBegin(const TCHAR* Name) UE_TRACE_IMPL();
 UE_TRACE_API void	ThreadGroupEnd() UE_TRACE_IMPL();
