[Insights] Memory units IEC compatible, use ToString to log stringbuilder.

diff --git a/Engine/Source/Runtime/Core/Private/ProfilingDebugging/TraceAuxiliary.cpp b/Engine/Source/Runtime/Core/Private/ProfilingDebugging/TraceAuxiliary.cpp
--- a/Engine/Source/Runtime/Core/Private/ProfilingDebugging/TraceAuxiliary.cpp
+++ b/Engine/Source/Runtime/Core/Private/ProfilingDebugging/TraceAuxiliary.cpp
@@ -549,18 +549,18 @@ static void TraceAuxiliaryStatus()
 	{
 		ConnectionStr.Append(TEXT("Not tracing"));
 	}
-	UE_LOG(LogConsoleResponse, Display, TEXT("- Connection: %.*s"), ConnectionStr.ToView().Len(), ConnectionStr.ToView().GetData());
+	UE_LOG(LogConsoleResponse, Display, TEXT("- Connection: %s"), ConnectionStr.ToString());
 	
 	// Stats
 	UE::Trace::FStatistics Stats;
 	UE::Trace::GetStatistics(Stats);
 	constexpr float MbDiv = 1024*1024;
-	UE_LOG(LogConsoleResponse, Display, TEXT("- Memory used: %.02f Mb"),
+	UE_LOG(LogConsoleResponse, Display, TEXT("- Memory used: %.02f MiB"),
 		float(Stats.MemoryUsed) / MbDiv);
-	UE_LOG(LogConsoleResponse, Display, TEXT("- Important Events cache: %.02f Mb, waste: %0.02f Mb"),
+	UE_LOG(LogConsoleResponse, Display, TEXT("- Important Events cache: %.02f MiB, waste: %0.02f MiB"),
 		float(Stats.CacheUsed) / MbDiv,
 		float(Stats.CacheWaste) / MbDiv);
-	UE_LOG(LogConsoleResponse, Display, TEXT("- Sent: %.2f Mb"),
+	UE_LOG(LogConsoleResponse, Display, TEXT("- Sent: %.2f MiB"),
 		float(Stats.BytesSent) / MbDiv);
 	
 	// Channels
