Managed storage wrapper - dump categories and sizes to log at startup.

diff --git a/Engine/Source/Runtime/Core/Private/HAL/IPlatformFileManagedStorageWrapper.cpp b/Engine/Source/Runtime/Core/Private/HAL/IPlatformFileManagedStorageWrapper.cpp
--- a/Engine/Source/Runtime/Core/Private/HAL/IPlatformFileManagedStorageWrapper.cpp
+++ b/Engine/Source/Runtime/Core/Private/HAL/IPlatformFileManagedStorageWrapper.cpp
@@ -315,6 +315,9 @@ FPersistentStorageManager::FCategoryInfo FPersistentStorageManager::InitCategori
 
 			int64 Quota = (QuotaInMB >= 0) ? QuotaInMB * 1024 * 1024 : -1;	// Quota being negative means infinite quota
 			int64 OptionalQuota = (OptionalQuotaMB >= 0) ? OptionalQuotaMB * 1024 * 1024 : 0;
+
+			UE_LOG(LogPlatformFileManagedStorage, Display, TEXT("Using Category %s : %" INT64_FMT " MB"), *CategoryName, QuotaInMB);
+
 			Categories.Emplace(MoveTemp(CategoryName), MoveTemp(Directories), Quota, OptionalQuota);
 		}
 	}
