Fixed thread affinity for task graph

diff --git a/Engine/Source/Runtime/Core/Private/Android/AndroidPlatformProcess.cpp b/Engine/Source/Runtime/Core/Private/Android/AndroidPlatformProcess.cpp
--- a/Engine/Source/Runtime/Core/Private/Android/AndroidPlatformProcess.cpp
+++ b/Engine/Source/Runtime/Core/Private/Android/AndroidPlatformProcess.cpp
@@ -25,8 +25,8 @@
 	#endif
 #endif
 
-int64 FAndroidAffinity::GameThreadMask = FPlatformAffinity::GetNoAffinityMask();
-int64 FAndroidAffinity::RenderingThreadMask = FPlatformAffinity::GetNoAffinityMask();
+uint64 FAndroidAffinity::GameThreadMask = FPlatformAffinity::GetNoAffinityMask();
+uint64 FAndroidAffinity::RenderingThreadMask = FPlatformAffinity::GetNoAffinityMask();
 
 void* FAndroidPlatformProcess::GetDllHandle(const TCHAR* Filename)
 {
diff --git a/Engine/Source/Runtime/Core/Public/Android/AndroidPlatformAffinity.h b/Engine/Source/Runtime/Core/Public/Android/AndroidPlatformAffinity.h
--- a/Engine/Source/Runtime/Core/Public/Android/AndroidPlatformAffinity.h
+++ b/Engine/Source/Runtime/Core/Public/Android/AndroidPlatformAffinity.h
@@ -12,7 +12,7 @@ class FAndroidAffinity : public FGenericPlatformAffinity
 {
 private:
 	static uint64 GetLittleCoreMask();
-	const static uint64 AllCores = 0xFFFFFFFFFF;
+	const static uint64 AllCores = 0xFFFFFFFFFFFFFFFF;
 public:
 	static const CORE_API uint64 GetMainGameMask()
 	{
@@ -41,7 +41,7 @@ public:
 
 	static const CORE_API uint64 GetTaskGraphThreadMask()
 	{
-		return GetLittleCoreMask();
+		return AllCores;
 	}
 
 	static const CORE_API uint64 GetStatsThreadMask()
@@ -80,8 +80,8 @@ public:
 	}
 
 public:
-	static int64 GameThreadMask;
-	static int64 RenderingThreadMask;
+	static uint64 GameThreadMask;
+	static uint64 RenderingThreadMask;
 };
 
 typedef FAndroidAffinity FPlatformAffinity;
