GrowableAllocator optimization
Optimize the chunk search so it's no longer O(n^2) in most cases. Instead of searching all free entries, we cache the size of the largest entry and use that for trivial checks. We only have to search all entries now if the cached chunk size is dirty (due to a previous allocation), or we can't trivially determine if an allocation fits due to padding.

The chunk search time is now 20x faster, and overall allocator overhead is reduced by 40%. This is particularly beneficial for allocator configurations with large numbers of chunks.

Also disable allocation histogram in code by default, since is basically dead code and adds overhead

diff --git a/Engine/Source/Runtime/Core/Public/Misc/GenericGrowableAllocator.h b/Engine/Source/Runtime/Core/Public/Misc/GenericGrowableAllocator.h
--- a/Engine/Source/Runtime/Core/Public/Misc/GenericGrowableAllocator.h
+++ b/Engine/Source/Runtime/Core/Public/Misc/GenericGrowableAllocator.h
@@ -21,7 +21,8 @@
 // - Can maybe just make virtuals in FGrowableAllocationBase ??
 /////////////
 
-#define ALLOCATION_HISTOGRAM
+#define ALLOCATION_HISTOGRAM			(!UE_BUILD_SHIPPING && !UE_BUILD_TEST) && 1
+#define ALLOCATION_HISTOGRAM_DETAILED	(!UE_BUILD_SHIPPING && !UE_BUILD_TEST) && 0
 
 struct FGrowableAllocationBase
 {
@@ -129,6 +130,7 @@ public:
 		: MemoryType(Type)
 		, HeapSize(InSize)
 		, UsedMemorySize(0)
+		, MaxFreeEntrySize(MaxFreeEntrySizeDirty)
 		, CriticalSection(InCriticalSection)
 	{
 	}
@@ -139,6 +141,7 @@ public:
 		HeapSize = CreateInternalMemory(HeapSize);
 		// entire chunk is free
 		FreeList = new FFreeEntry(NULL, 0, HeapSize);
+		MaxFreeEntrySize = FreeList->BlockSize;
 	}
 
 	/**
@@ -162,8 +165,30 @@ public:
 	*/
 	bool CanFitEntry(uint32 Size, uint32 Alignment)
 	{
+		// Compute MaxFreeEntrySize if necessary (should only happen if this chunk was just allocated from)
+		if (MaxFreeEntrySize == MaxFreeEntrySizeDirty)
+		{
+			MaxFreeEntrySize = 0;
+			for (FFreeEntry* Entry = FreeList; Entry; Entry = Entry->Next)
+			{
+				MaxFreeEntrySize = FMath::Max(Entry->BlockSize, MaxFreeEntrySize);
+			}
+		}
+
+		// Return false if we trivially don't fit
+		if (Size > MaxFreeEntrySize)
+		{
+			return false;
+		}
+
+		// Return true if we trivially do fit
+		if (Size + Alignment - 1 <= MaxFreeEntrySize)
+		{
+			return true;
+		}
+
+		// Slow method - search the free entries for a free chunk
 		bool bResult = false;
-		// look for a good free chunk
 		for (FFreeEntry *Entry = FreeList; Entry; Entry = Entry->Next)
 		{
 			if (Entry->CanFit(Size, Alignment))
@@ -200,37 +225,50 @@ public:
 
 		// look for a good free chunk
 		FFreeEntry* Prev = NULL;
+		FFreeEntry* FindEntry = NULL;
+
 		for (FFreeEntry *Entry = FreeList; Entry; Entry = Entry->Next)
 		{
 			if (Entry->CanFit(AlignedSize, Alignment))
 			{
-				// Use it, leaving over any unused space
-				UsedMemorySize += AlignedSize;
-				bool bDelete;
-				uint32 Padding;
-				uint32 Offset = Entry->Split(AlignedSize, Alignment, bDelete, Padding, MinAllocationSize);
-				if (bDelete)
-				{
-					FFreeEntry*& PrevRef = Prev ? Prev->Next : FreeList;
-					PrevRef = Entry->Next;
-					delete Entry;
-				}
+				FindEntry = Entry;
+				break;
+			}
+			Prev = Entry;
+		}
+
+		if (FindEntry != NULL)
+		{
+			if (FindEntry->BlockSize == MaxFreeEntrySize)
+			{
+				// We're probably about to just split our largest entry, so mark the max size as dirty to indicate it needs recomputing
+				MaxFreeEntrySize = MaxFreeEntrySizeDirty;
+			}
+			// Use it, leaving over any unused space
+			UsedMemorySize += AlignedSize;
+			bool bDelete;
+			uint32 Padding;
+			uint32 Offset = FindEntry->Split(AlignedSize, Alignment, bDelete, Padding, MinAllocationSize);
+			if (bDelete)
+			{
+				FFreeEntry*& PrevRef = Prev ? Prev->Next : FreeList;
+				PrevRef = FindEntry->Next;
+				delete FindEntry;
+			}
 
-				FGrowableAllocationBase* Allocation = CreateAllocationStruct();
-				Allocation->Size = AlignedSize;
-				Allocation->Padding = Padding;
-				Allocation->Offset = Offset;
+			FGrowableAllocationBase* Allocation = CreateAllocationStruct();
+			Allocation->Size = AlignedSize;
+			Allocation->Padding = Padding;
+			Allocation->Offset = Offset;
 #if !UE_BUILD_SHIPPING
-				Allocation->OwnerType = OwnerType;
+			Allocation->OwnerType = OwnerType;
 #endif
-				LLM(FLowLevelMemTracker::Get().OnLowLevelAlloc(ELLMTracker::Default, GetAddressForTracking(Offset), Size));
-				MemoryTrace_Alloc(uint64(GetAddressForTracking(Offset)), Size, Alignment, EMemoryTraceRootHeap::SystemMemory);
+			LLM(FLowLevelMemTracker::Get().OnLowLevelAlloc(ELLMTracker::Default, GetAddressForTracking(Offset), Size));
+			MemoryTrace_Alloc(uint64(GetAddressForTracking(Offset)), Size, Alignment, EMemoryTraceRootHeap::SystemMemory);
 
-				// let the implementation fill in any more
-				InitializeAllocationStruct(Allocation);
-				return Allocation;
-			}
-			Prev = Entry;
+			// let the implementation fill in any more
+			InitializeAllocationStruct(Allocation);
+			return Allocation;
 		}
 
 		// if no suitable blocks were found, we must fail
@@ -257,7 +295,7 @@ public:
 		UsedMemorySize -= Size;
 		CurrentAllocs--;
 
-		// Search for where a place to insert a new free entry.
+		// Search for a place to insert a new free entry
 		FFreeEntry* Prev = NULL;
 		FFreeEntry* Entry = FreeList;
 		while (Entry && Offset > Entry->Location)
@@ -278,8 +316,13 @@ public:
 			{
 				Prev->BlockSize += Entry->BlockSize;
 				Prev->Next = Entry->Next;
+				MaxFreeEntrySize = FMath::Max(MaxFreeEntrySize, Prev->BlockSize);
 				delete Entry;
 			}
+			else
+			{
+				MaxFreeEntrySize = FMath::Max(MaxFreeEntrySize, Entry->BlockSize);
+			}
 			return true;
 		}
 
@@ -296,6 +339,7 @@ public:
 				Prev->Next = Entry->Next;
 				delete Entry;
 			}
+			MaxFreeEntrySize = FMath::Max(MaxFreeEntrySize, Prev->BlockSize);
 			return true;
 		}
 
@@ -303,6 +347,7 @@ public:
 		FFreeEntry* NewFree = new FFreeEntry(Entry, Offset - Padding, AllocationSize);
 		FFreeEntry*& PrevRef = Prev ? Prev->Next : FreeList;
 		PrevRef = NewFree;
+		MaxFreeEntrySize = FMath::Max(MaxFreeEntrySize, NewFree->BlockSize);
 		return true;
 	}
 
@@ -443,11 +488,16 @@ public:
 	/** Size of used memory */
 	uint64 UsedMemorySize;
 
+	/** Size of the largest free entry (will be MaxFreeEntrySizeDirty if unknown) */
+	uint32 MaxFreeEntrySize;
+
 	/** List of free blocks */
 	FFreeEntry* FreeList;
 
 	/** Shared critical section */
 	FCriticalSection* CriticalSection;
+
+	static const uint32 MaxFreeEntrySizeDirty = 0xffffffff;
 };
 
 
@@ -556,7 +606,7 @@ public:
 			return nullptr;
 		}
 
-#if !(UE_BUILD_SHIPPING || UE_BUILD_TEST)
+#if ALLOCATION_HISTOGRAM
 		TotalAllocationsHistogram.FindOrAdd(AlignedSize)++;
 		PeakAllocationsHistogram.FindOrAdd(AlignedSize)++;
 
@@ -565,12 +615,12 @@ public:
 
 		TypeAllocInfo.Counts.Allocations++;
 		TypeAllocInfo.TotalAllocated += Result->Size;
-#ifdef ALLOCATION_HISTOGRAM
+#if ALLOCATION_HISTOGRAM_DETAILED
 		TypeAllocInfo.AllocationHistogram.FindOrAdd(Result->Size).Allocations++;
 #endif
-#endif // #if !(UE_BUILD_SHIPPING || UE_BUILD_TEST)
+#endif // ALLOCATION_HISTOGRAM 
 
-#if !UE_BUILD_SHIPPING
+#if STATS
 		if (OwnerTypeToStatIdMap)
 		{
 			INC_MEMORY_STAT_BY_FName(OwnerTypeToStatIdMap[Result->OwnerType], Result->Size);
@@ -596,7 +646,7 @@ public:
 			ChunkAllocatorType* Chunk = AllocChunks[ChunkIndex];
 			if (Chunk && Chunk->DoesChunkContainAllocation(Memory))
 			{
-#if !(UE_BUILD_SHIPPING || UE_BUILD_TEST)
+#if ALLOCATION_HISTOGRAM
 				PeakAllocationsHistogram[Memory->Size]--;
 
 				// untrack per type allocation info
@@ -604,12 +654,12 @@ public:
 
 				TypeAllocInfo.TotalAllocated -= Memory->Size;
 				TypeAllocInfo.Counts.Frees++;
-#ifdef ALLOCATION_HISTOGRAM
+#if ALLOCATION_HISTOGRAM_DETAILED
 				TypeAllocInfo.AllocationHistogram[Memory->Size].Frees++;
 #endif
-#endif // !(UE_BUILD_SHIPPING || UE_BUILD_TEST)
+#endif // ALLOCATION_HISTOGRAM
 
-#if !UE_BUILD_SHIPPING
+#if STATS
 				if (OwnerTypeToStatIdMap)
 				{
 					DEC_MEMORY_STAT_BY_FName(OwnerTypeToStatIdMap[Memory->OwnerType], Memory->Size);
@@ -671,18 +721,19 @@ public:
 		// multi-thread protection
 		FScopeLock ScopeLock(&CriticalSection);
 
-#if !(UE_BUILD_SHIPPING || UE_BUILD_TEST)
+#if ALLOCATION_HISTOGRAM
 		FPlatformMisc::LowLevelOutputDebugStringf(TEXT("Total Allocations Histogram:\n"));
 		for (auto It = TotalAllocationsHistogram.CreateIterator(); It; ++It)
 		{
 			FPlatformMisc::LowLevelOutputDebugStringf(TEXT(" %d, %d\n"), It.Key(), It.Value());
 		}
-#endif // !(UE_BUILD_SHIPPING || UE_BUILD_TEST
+#endif // ALLOCATION_HISTOGRAM
 
 		int32 NumChunks = 0;
 		// pass off to individual alloc chunks
 		for (int32 ChunkIndex = 0; ChunkIndex < AllocChunks.Num(); ChunkIndex++)
 		{
+			FPlatformMisc::LowLevelOutputDebugStringf(TEXT("\n-----------------\nChunk %d\n"),ChunkIndex);
 			ChunkAllocatorType* Chunk = AllocChunks[ChunkIndex];
 			if (Chunk)
 			{
@@ -693,7 +744,7 @@ public:
 
 	void DumpMemoryInfo()
 	{
-#if !(UE_BUILD_SHIPPING || UE_BUILD_TEST)
+#if ALLOCATION_HISTOGRAM
 		// we use LowLevel here because we are usually dumping this out while 
 		for (auto InfoIt = PerTypeAllocationInfo.CreateConstIterator(); InfoIt; ++InfoIt)
 		{
@@ -707,14 +758,14 @@ public:
 			{
 				FPlatformMisc::LowLevelOutputDebugStringf(TEXT("      'OwnerType %d': %lld Allocs: %d Frees: %d\n"), InfoIt.Key(), Info.TotalAllocated, Info.Counts.Allocations, Info.Counts.Frees);
 			}
-#ifdef ALLOCATION_HISTOGRAM
+#if ALLOCATION_HISTOGRAM_DETAILED
 			for (auto HistoIt = Info.AllocationHistogram.CreateConstIterator(); HistoIt; ++HistoIt)
 			{
 				FPlatformMisc::LowLevelOutputDebugStringf(TEXT("           %d, %lld, %lld\n"), HistoIt.Key(), HistoIt.Value().Allocations, HistoIt.Value().Frees);
 			}
 #endif
 		}
-#endif // !(UE_BUILD_SHIPPING || UE_BUILD_TEST)
+#endif // ALLOCATION_HISTOGRAM
 	}
 
 private:
@@ -816,7 +867,7 @@ private:
 	 // a critical section used to coordinate all access in this instance of the allocator and it's chunks
 	 FCriticalSection CriticalSection;
 
-#if !(UE_BUILD_SHIPPING || UE_BUILD_TEST)
+#if ALLOCATION_HISTOGRAM
 	 struct AllocFreeCounts
 	 {
 		 uint32 Allocations;
@@ -825,7 +876,7 @@ private:
 	 struct AllocationInfo
 	 {
 		 uint64 TotalAllocated;
-#ifdef ALLOCATION_HISTOGRAM
+#if ALLOCATION_HISTOGRAM_DETAILED
 		 TMap<uint32, AllocFreeCounts> AllocationHistogram;
 #endif
 		 AllocFreeCounts Counts;
@@ -836,7 +887,6 @@ private:
 	 TMap<uint64, uint32> TotalAllocationsHistogram;
 	 TMap<uint64, uint32> OutstandingAllocationsHistogram;
 	 TMap<uint64, uint32> PeakAllocationsHistogram;
-#endif // !(UE_BUILD_SHIPPING || UE_BUILD_TEST)
-
+#endif // ALLOCATION_HISTOGRAM
 };
 
