Fix for TRobinHoodHashMap::Remove() not compiling.
Fix for TRobinHoodHashMap::Num() data loss warning.

diff --git a/Engine/Source/Runtime/Core/Public/Experimental/Containers/RobinHoodHashTable.h b/Engine/Source/Runtime/Core/Public/Experimental/Containers/RobinHoodHashTable.h
--- a/Engine/Source/Runtime/Core/Public/Experimental/Containers/RobinHoodHashTable.h
+++ b/Engine/Source/Runtime/Core/Public/Experimental/Containers/RobinHoodHashTable.h
@@ -768,7 +768,7 @@ namespace RobinHoodHashTable_Private
 
 		inline int32 Num() const
 		{
-			return KeyValueData.Num();
+			return (int32)KeyValueData.Num();
 		}
 
 		inline IndexType GetMaxIndex() const
@@ -920,7 +920,7 @@ namespace RobinHoodHashTable_Private
 		bool Remove(const KeyType& Key)
 		{
 			const FHashType HashValue = ComputeHash(Key);
-			return RemoveByHash(HashValue);
+			return RemoveByHash(HashValue, Key);
 		}
 
 		inline bool RemoveByElementId(FHashElementId Id)
