* Added TIteratorAdapter.
* Added TDeque.
* Added Catch 2 Container unit test project.
* Added TTestToken for container construction/destruction testing.
* Added 48 unit tests for TDeque.

#test passed TDeque unit tests: All tests passed (1221 assertions in 48 test cases)

diff --git a/Engine/Source/Runtime/Core/Public/Containers/Deque.h b/Engine/Source/Runtime/Core/Public/Containers/Deque.h
new file mode 100644
--- /dev/null
+++ b/Engine/Source/Runtime/Core/Public/Containers/Deque.h
@@ -0,0 +1,612 @@
+// Copyright Epic Games, Inc. All Rights Reserved.
+
+#pragma once
+
+#include "CoreTypes.h"
+#include "GenericPlatform/GenericPlatformMath.h"
+#include "IteratorAdapter.h"
+#include "Misc/AssertionMacros.h"
+#include "Templates/ChooseClass.h"
+#include "Templates/MemoryOps.h"
+
+#include <initializer_list>
+
+template <typename InElementType, typename InAllocatorType = FDefaultAllocator>
+class TDeque;
+
+namespace UE
+{
+namespace Deque
+{
+namespace Private
+{
+/**
+ * This implementation assumes that the Index value is never going to exceed twice the Range value. This way we can
+ * avoid the modulo operator or a power of 2 range value requirement and have an efficient wrap around function.
+ */
+template <typename InSizeType>
+FORCEINLINE InSizeType WrapAround(InSizeType Index, InSizeType Range)
+{
+	return (Index < Range) ? Index : Index - Range;
+}
+
+/**
+ * TDeque iterator class.
+ */
+template <typename InElementType, typename InSizeType>
+class TIteratorBase
+{
+public:
+	using ElementType = InElementType;
+	using SizeType = InSizeType;
+
+	TIteratorBase() = default;
+
+protected:
+	/**
+	 * Internal parameter constructor (used exclusively by the container).
+	 */
+	TIteratorBase(ElementType* Data, SizeType Range, SizeType Offset)
+		: Data(Data), Range(Range), Offset(Offset)
+	{
+	}
+
+	FORCEINLINE ElementType& Dereference() const
+	{
+		return *(Data + WrapAround(Offset, Range));
+	}
+
+	FORCEINLINE void Increment()
+	{
+		checkSlow(Offset + 1 < Range * 2);	// See WrapAround
+		Offset++;
+	}
+
+	FORCEINLINE bool Equals(const TIteratorBase& Other) const
+	{
+		return Data + Offset == Other.Data + Other.Offset;
+	}
+
+private:
+	ElementType* Data = nullptr;
+	SizeType Range = 0;
+	SizeType Offset = 0;
+};
+
+template <typename InElementType, typename InSizeType>
+using TIterator = TIteratorAdapter<TIteratorBase<InElementType, InSizeType>>;
+}  // namespace Private
+}  // namespace Deque
+}  // namespace UE
+
+/**
+ * Sequential double-ended queue (deque) container class.
+ *
+ * A dynamically sized sequential queue of arbitrary size.
+ **/
+template <typename InElementType, typename InAllocatorType>
+class TDeque
+{
+	template <typename AnyElementType, typename AnyAllocatorType>
+	friend class TDeque;
+
+public:
+	using AllocatorType = InAllocatorType;
+	using SizeType = typename InAllocatorType::SizeType;
+	using ElementType = InElementType;
+
+	using ElementAllocatorType = typename TChooseClass<
+		AllocatorType::NeedsElementType,
+		typename AllocatorType::template ForElementType<ElementType>,
+		typename AllocatorType::ForAnyElementType>::Result;
+
+	using ConstIteratorType = UE::Deque::Private::TIterator<const ElementType, SizeType>;
+	using IteratorType = UE::Deque::Private::TIterator<ElementType, SizeType>;
+
+	TDeque() : Capacity(Storage.GetInitialCapacity())
+	{
+	}
+
+	TDeque(TDeque&& Other)
+	{
+		MoveUnchecked(MoveTemp(Other));
+	}
+
+	TDeque(const TDeque& Other)
+	{
+		CopyUnchecked(Other);
+	}
+
+	TDeque(std::initializer_list<ElementType> InList)
+	{
+		CopyUnchecked(InList);
+	}
+
+	~TDeque()
+	{
+		Empty();
+	}
+
+	TDeque& operator=(TDeque&& Other)
+	{
+		if (this != &Other)
+		{
+			Reset();
+			MoveUnchecked(MoveTemp(Other));
+		}
+		return *this;
+	}
+
+	TDeque& operator=(const TDeque& Other)
+	{
+		if (this != &Other)
+		{
+			Reset();
+			CopyUnchecked(Other);
+		}
+		return *this;
+	}
+
+	TDeque& operator=(std::initializer_list<ElementType> InList)
+	{
+		Reset();
+		CopyUnchecked(InList);
+		return *this;
+	}
+
+	FORCEINLINE const ElementType& operator[](SizeType Index) const
+	{
+		CheckValidIndex(Index);
+		return GetData()[UE::Deque::Private::WrapAround(Head + Index, Capacity)];
+	}
+
+	FORCEINLINE ElementType& operator[](SizeType Index)
+	{
+		CheckValidIndex(Index);
+		return GetData()[UE::Deque::Private::WrapAround(Head + Index, Capacity)];
+	}
+
+	FORCEINLINE const ElementType& Last() const
+	{
+		CheckValidIndex(0);
+		return GetData()[UE::Deque::Private::WrapAround(Tail + Capacity - 1, Capacity)];
+	}
+
+	FORCEINLINE ElementType& Last()
+	{
+		CheckValidIndex(0);
+		return GetData()[UE::Deque::Private::WrapAround(Tail + Capacity - 1, Capacity)];
+	}
+
+	FORCEINLINE const ElementType& First() const
+	{
+		CheckValidIndex(0);
+		return GetData()[Head];
+	}
+
+	FORCEINLINE ElementType& First()
+	{
+		CheckValidIndex(0);
+		return GetData()[Head];
+	}
+
+	FORCEINLINE bool IsEmpty() const
+	{
+		return !Count;
+	}
+
+	FORCEINLINE SizeType Max() const
+	{
+		return Capacity;
+	}
+
+	FORCEINLINE SizeType Num() const
+	{
+		return Count;
+	}
+
+	/*
+	 * Constructs an element in place using the parameter arguments and adds it at the back of the queue.
+	 * This method returns a reference to the constructed element.
+	 */
+	template <typename... ArgsType>
+	ElementType& EmplaceLast(ArgsType&&... Args)
+	{
+		GrowIfRequired();
+		ElementType* Target = GetData() + Tail;
+		new (Target) ElementType(Forward<ArgsType>(Args)...);
+		Tail = UE::Deque::Private::WrapAround(Tail + 1, Capacity);
+		Count++;
+		return *Target;
+	}
+
+	/*
+	 * Constructs an element in place using the parameter arguments and adds it at the front of the queue.
+	 * This method returns a reference to the constructed element.
+	 */
+	template <typename... ArgsType>
+	ElementType& EmplaceFirst(ArgsType&&... Args)
+	{
+		GrowIfRequired();
+		Head = UE::Deque::Private::WrapAround(Head + Capacity - 1, Capacity);
+		ElementType* Target = GetData() + Head;
+		new (Target) ElementType(Forward<ArgsType>(Args)...);
+		Count++;
+		return *Target;
+	}
+
+	/*
+	 * Adds the parameter element at the back of the queue.
+	 */
+	FORCEINLINE void PushLast(const ElementType& Element)
+	{
+		EmplaceLast(Element);
+	}
+
+	FORCEINLINE void PushLast(ElementType&& Element)
+	{
+		EmplaceLast(MoveTempIfPossible(Element));
+	}
+
+	/*
+	 * Adds the parameter element at the front of the queue.
+	 */
+	FORCEINLINE void PushFirst(const ElementType& Element)
+	{
+		EmplaceFirst(Element);
+	}
+
+	FORCEINLINE void PushFirst(ElementType&& Element)
+	{
+		EmplaceFirst(MoveTempIfPossible(Element));
+	}
+
+	/*
+	 * Removes the element at the back of the queue.
+	 * This method requires a non-empty queue.
+	 */
+	void PopLast()
+	{
+		CheckValidIndex(0);
+		const SizeType NextTail = UE::Deque::Private::WrapAround(Tail + Capacity - 1, Capacity);
+		DestructItem(GetData() + NextTail);
+		Tail = NextTail;
+		Count--;
+	}
+
+	/*
+	 * Removes the element at the front of the queue.
+	 * This method requires a non-empty queue.
+	 */
+	void PopFirst()
+	{
+		CheckValidIndex(0);
+		DestructItem(GetData() + Head);
+		Head = UE::Deque::Private::WrapAround(Head + 1, Capacity);
+		Count--;
+	}
+
+	/*
+	 * Removes the element at the back of the queue if existent after copying it to the parameter
+	 * out value.
+	 */
+	bool TryPopLast(ElementType& OutValue)
+	{
+		if (IsEmpty())
+		{
+			return false;
+		}
+		const SizeType NextTail = UE::Deque::Private::WrapAround(Tail + Capacity - 1, Capacity);
+		OutValue = MoveTempIfPossible(GetData()[NextTail]);
+		DestructItem(GetData() + NextTail);
+		Tail = NextTail;
+		Count--;
+		return true;
+	}
+
+	/*
+	 * Removes the element at the front of the queue if existent after copying it to the parameter
+	 * out value.
+	 */
+	bool TryPopFirst(ElementType& OutValue)
+	{
+		if (IsEmpty())
+		{
+			return false;
+		}
+		OutValue = MoveTempIfPossible(GetData()[Head]);
+		DestructItem(GetData() + Head);
+		Head = UE::Deque::Private::WrapAround(Head + 1, Capacity);
+		Count--;
+		return true;
+	}
+
+	/*
+	 * Destroys all contained elements but doesn't release the container's storage.
+	 */
+	void Reset()
+	{
+		if (Count)
+		{
+			if (Head < Tail)
+			{
+				DestructItems(GetData() + Head, Count);
+			}
+			else
+			{
+				DestructItems(GetData(), Tail);
+				DestructItems(GetData() + Head, Capacity - Head);
+			}
+		}
+		Head = Tail = Count = 0;
+	}
+
+	/*
+	 * Empties the queue effectively destroying all contained elements and releases any acquired storage.
+	 */
+	void Empty()
+	{
+		Reset();
+		if (Capacity)
+		{
+			Storage.ResizeAllocation(0, 0, sizeof(ElementType));
+			Capacity = 0;
+		}
+	}
+
+	/*
+	 * Reserves storage for the parameter element count.
+	 */
+	void Reserve(SizeType InCount)
+	{
+		if (Capacity < InCount)
+		{
+			Grow(Storage.CalculateSlackReserve(InCount, sizeof(ElementType)));
+		}
+	}
+
+	/*
+	 * STL IteratorType model compliance methods.
+	 */
+	FORCEINLINE ConstIteratorType begin() const
+	{
+		return GetIterator(Head);
+	}
+	FORCEINLINE IteratorType begin()
+	{
+		return GetIterator(Head);
+	}
+	FORCEINLINE ConstIteratorType end() const
+	{
+		return GetIterator(Head + Count);
+	}
+	FORCEINLINE IteratorType end()
+	{
+		return GetIterator(Head + Count);
+	}
+
+private:
+	ElementAllocatorType Storage;
+	SizeType Capacity = 0;
+	SizeType Count = 0;
+	SizeType Head = 0;
+	SizeType Tail = 0;
+
+	FORCEINLINE const ElementType* GetData() const
+	{
+		return (const ElementType*)Storage.GetAllocation();
+	}
+
+	FORCEINLINE ElementType* GetData()
+	{
+		return (ElementType*)Storage.GetAllocation();
+	}
+
+	FORCEINLINE ConstIteratorType GetIterator(SizeType HeadOffset) const
+	{
+		return ConstIteratorType(InPlace, GetData(), Max(), HeadOffset);
+	}
+
+	FORCEINLINE IteratorType GetIterator(SizeType HeadOffset)
+	{
+		return IteratorType(InPlace, GetData(), Max(), HeadOffset);
+	}
+
+	/**
+	 * Grows the container's storage to the parameter capacity value.
+	 * This method preserves existing data.
+	 */
+	void Grow(SizeType InCapacity)
+	{
+		checkSlow(Capacity < InCapacity);
+		if (Count)
+		{
+			Linearize();
+		}
+		Storage.ResizeAllocation(Count, InCapacity, sizeof(ElementType));
+		Capacity = InCapacity;
+		Head = 0;
+		Tail = Count;
+	}
+
+	/**
+	 * Grows the container to the next capacity value (determined by the storage allocator) if full.
+	 * This method preserves existing data.
+	 */
+	void GrowIfRequired()
+	{
+		if (Count == Capacity)
+		{
+			Grow(Storage.CalculateSlackGrow(Count + 1, Capacity, sizeof(ElementType)));
+		}
+	}
+
+	/*
+	 * Copies the parameter container into this one.
+	 * This method assumes no previously existing content.
+	 */
+	void CopyUnchecked(const TDeque& Other)
+	{
+		checkSlow(!Count);
+		if (Other.Count)
+		{
+			Grow(Storage.CalculateSlackReserve(Other.Count, sizeof(ElementType)));
+			CopyElements(Other);
+		}
+	}
+
+	void CopyUnchecked(std::initializer_list<ElementType> InList)
+	{
+		const SizeType InCount = static_cast<SizeType>(InList.size());
+		checkSlow(!Count);
+		if (InCount)
+		{
+			Grow(Storage.CalculateSlackReserve(InCount, sizeof(ElementType)));
+			ConstructItems<ElementType>(GetData(), &*InList.begin(), InCount);
+			Tail = Count = InCount;
+		}
+	}
+
+	/*
+	 * Moves the parameter container into this one.
+	 * This method assumes no previously existing content.
+	 */
+	void MoveUnchecked(TDeque&& Other)
+	{
+		checkSlow(!Count);
+		if (Other.Count)
+		{
+			Storage.MoveToEmpty(Other.Storage);
+			Capacity = Other.Capacity;
+			Count = Other.Count;
+			Head = Other.Head;
+			Tail = Other.Tail;
+			Other.Capacity = Other.Storage.GetInitialCapacity();
+			Other.Count = Other.Head = Other.Tail = 0;
+		}
+	}
+
+	/*
+	 * Copies the parameter container elements into this one.
+	 * The copied range is linearized.
+	 */
+	void CopyElements(const TDeque& Other)
+	{
+		if (Other.Head < Other.Tail)
+		{
+			ConstructItems<ElementType>(GetData(), Other.GetData() + Other.Head, Other.Count);
+		}
+		else
+		{
+			const SizeType HeadToEndOffset = Other.Capacity - Other.Head;
+			ConstructItems<ElementType>(GetData(), Other.GetData() + Other.Head, HeadToEndOffset);
+			ConstructItems<ElementType>(GetData() + HeadToEndOffset, Other.GetData(), Other.Tail);
+		}
+		Head = 0;
+		Tail = Count = Other.Count;
+	}
+
+	/*
+	 * Moves the parameter container elements into this one.
+	 * The moved range is linearized.
+	 */
+	void MoveElements(TDeque& Other)
+	{
+		if (Other.Head < Other.Tail)
+		{
+			RelocateConstructItems<ElementType>(GetData(), Other.GetData(), Other.Count);
+		}
+		else
+		{
+			const SizeType HeadToEndOffset = Other.Capacity - Other.Head;
+			RelocateConstructItems<ElementType>(GetData(), Other.GetData() + Other.Head, HeadToEndOffset);
+			RelocateConstructItems<ElementType>(GetData() + HeadToEndOffset, Other.GetData(), Other.Tail);
+		}
+		Head = 0;
+		Tail = Count = Other.Count;
+		Other.Head = Other.Tail = Other.Count = 0;
+	}
+
+	/**
+	 * Shifts the contained range to the beginning of the storage so it's linear.
+	 * This method is faster than a full range rotation but requires a temporary extra storage whenever the tail is
+	 * wrapped around.
+	 */
+	void Linearize()
+	{
+		if (Head < Tail)
+		{
+			ShiftElementsLeft(Count);
+		}
+		else
+		{
+			ElementAllocatorType TempStorage;
+			TempStorage.ResizeAllocation(0, Tail, sizeof(ElementType));
+			RelocateConstructItems<ElementType>(TempStorage.GetAllocation(), GetData(), Tail);
+			const SizeType HeadToEndOffset = Capacity - Head;
+			ShiftElementsLeft(HeadToEndOffset);
+			RelocateConstructItems<ElementType>(GetData() + HeadToEndOffset, TempStorage.GetAllocation(), Tail);
+		}
+	}
+
+	/**
+	 * Moves the parameter number of elements to the left shifting the head to the beginning of the storage.
+	 */
+	void ShiftElementsLeft(SizeType InCount)
+	{
+		if (Head == 0)
+		{
+			return;
+		}
+		SizeType Offset = 0;
+		while (Offset < InCount)
+		{
+			const SizeType Step = ::FMath::Min(Head, InCount - Offset);
+			RelocateConstructItems<ElementType>(GetData() + Offset, GetData() + Head + Offset, Step);
+			Offset += Step;
+		}
+	}
+
+	FORCEINLINE void CheckValidIndex(SizeType Index) const
+	{
+		checkSlow((Count >= 0) & (Capacity >= Count));
+		if constexpr (AllocatorType::RequireRangeCheck)
+		{
+			checkf((Index >= 0) & (Index < Count), TEXT("Parameter index %d exceeds container size %d"), Index, Count);
+		}
+	}
+};
+
+//---------------------------------------------------------------------------------------------------------------------
+// TDeque comparison operators
+//---------------------------------------------------------------------------------------------------------------------
+
+template <typename InElementType, typename InAllocatorType>
+inline bool operator==(
+	const TDeque<InElementType, InAllocatorType>& Left,
+	const TDeque<InElementType, InAllocatorType>& Right)
+{
+	if (Left.Num() == Right.Num())
+	{
+		auto EndIt = Left.end();
+		auto LeftIt = Left.begin();
+		auto RightIt = Right.begin();
+		while (LeftIt != EndIt)
+		{
+			if (*LeftIt++ != *RightIt++)
+			{
+				return false;
+			}
+		}
+		return true;
+	}
+	return false;
+}
+
+template <typename InElementType, typename InAllocatorType>
+inline bool operator!=(
+	const TDeque<InElementType, InAllocatorType>& Left,
+	const TDeque<InElementType, InAllocatorType>& Right)
+{
+	return !(Left == Right);
+}
diff --git a/Engine/Source/Runtime/Core/Public/Containers/IteratorAdapter.h b/Engine/Source/Runtime/Core/Public/Containers/IteratorAdapter.h
new file mode 100644
--- /dev/null
+++ b/Engine/Source/Runtime/Core/Public/Containers/IteratorAdapter.h
@@ -0,0 +1,125 @@
+// Copyright Epic Games, Inc. All Rights Reserved.
+
+#pragma once
+
+#include "CoreTypes.h"
+#include "Templates/UnrealTemplate.h"
+
+template <typename Base>
+class TIteratorAdapter;
+
+/**
+ * Adapter class for iterator implementation.
+ *
+ * This class provides an interface to implement UE compatible iterators with a minimum set of implementation
+ * requirements. This provides the user with an easier way to implement custom iterators. The adapter uses a CRTP
+ * pattern requiring its base class to implement the following:
+ *
+ * Types:
+ *   - ElementType: the element type used in the container.
+ *   - SizeType: the type use for the container's element indexing and iterator's offset.
+ * Methods:
+ *   - bool          Equals(const IteratorType& other)
+ *   - ElementType&  Dereference() const;
+ *   - void          Increment();
+ * Optional:
+ *   - void          Decrement();
+ */
+template <typename Base>
+class TIteratorAdapter : public Base
+{
+public:
+	using BaseType = Base;
+	using ThisType = TIteratorAdapter<Base>;
+	using ElementType = typename BaseType::ElementType;
+	using SizeType = typename BaseType::SizeType;
+
+	TIteratorAdapter() = default;
+
+	/**
+	 * Perfect forwarding constructor to the iterator class.
+	 */
+	template < typename... Args>
+	TIteratorAdapter(EInPlace, Args&&... InArgs) : BaseType(Forward<Args>(InArgs)...)
+	{
+	}
+
+	ElementType& operator*() const
+	{
+		return this->Dereference();
+	}
+
+	ElementType* operator->() const
+	{
+		return &this->Dereference();
+	}
+
+	ThisType& operator++()
+	{
+		this->Increment();
+		return *this;
+	}
+
+	ThisType operator++(int)
+	{
+		ThisType Temp = *this;
+		++*this;
+		return Temp;
+	}
+
+	ThisType& operator--()
+	{
+		this->Decrement();
+		return *this;
+	}
+
+	ThisType operator--(int)
+	{
+		ThisType Temp = *this;
+		--*this;
+		return Temp;
+	}
+
+	ThisType& operator+=(SizeType Offset)
+	{
+		this->Increment(Offset);
+		return *this;
+	}
+
+	ThisType& operator-=(SizeType Offset)
+	{
+		this->Decrement(Offset);
+		return *this;
+	}
+
+	ThisType operator+(SizeType Offset) const
+	{
+		ThisType Temp = *this;
+		return Temp += Offset;
+	}
+
+	ThisType operator-(SizeType Offset) const
+	{
+		ThisType Temp = *this;
+		return Temp -= Offset;
+	}
+
+protected:
+	// BaseType MUST implement at least forward iteration which is enforced with the "using" declaration.
+	using BaseType::Dereference;
+	using BaseType::Equals;
+	using BaseType::Increment;
+
+	template <typename AnyElementType>
+	friend class TIteratorAdapter;
+
+	friend bool operator==(const TIteratorAdapter& Left, const TIteratorAdapter& Right)
+	{
+		return Left.Equals(Right);
+	}
+
+	friend bool operator!=(const TIteratorAdapter& Left, const TIteratorAdapter& Right)
+	{
+		return !(Left == Right);
+	}
+};
