PR #8637: Do not use Generic string implementation on HoloLens (Contributed by [REDACTED_PERSON])

diff --git a/Engine/Source/Runtime/Core/Public/HoloLens/HoloLensPlatform.h b/Engine/Source/Runtime/Core/Public/HoloLens/HoloLensPlatform.h
--- a/Engine/Source/Runtime/Core/Public/HoloLens/HoloLensPlatform.h
+++ b/Engine/Source/Runtime/Core/Public/HoloLens/HoloLensPlatform.h
@@ -73,6 +73,7 @@ typedef FHoloLensTypes FPlatformTypes;
 #define PLATFORM_COMPILER_HAS_EXPLICIT_OPERATORS			1
 //#define EXCEPTION_EXECUTE_HANDLER                           1
 //#define PLATFORM_COMPILER_HAS_DEFAULT_FUNCTION_TEMPLATE_ARGUMENTS	1
+#define PLATFORM_USE_GENERIC_STRING_IMPLEMENTATION			0  // GenericWidePlatformString will fail on HoloLens in UE5
 
 #define PLATFORM_BREAK() (__nop(), __debugbreak())
 
@@ -120,4 +121,4 @@ typedef FHoloLensTypes FPlatformTypes;
 #define DLLIMPORT __declspec(dllimport)
 
 // disable this now as it is annoying for generic platform implementations
-#pragma warning(disable : 4100) // unreferenced formal parameter
\ No newline at end of file
+#pragma warning(disable : 4100) // unreferenced formal parameter
