Fix an issue where ServerContext wasn't being correctly enabled when UnrealEditor.exe is ran with -server

diff --git a/Engine/Source/Runtime/Core/Private/Misc/AutomationTest.cpp b/Engine/Source/Runtime/Core/Private/Misc/AutomationTest.cpp
--- a/Engine/Source/Runtime/Core/Private/Misc/AutomationTest.cpp
+++ b/Engine/Source/Runtime/Core/Private/Misc/AutomationTest.cpp
@@ -528,6 +528,7 @@ void FAutomationTestFramework::GetValidTestNames( TArray<FAutomationTestInfo>& T
 	// Determine required application type (Editor, Game, or Commandlet)
 	const bool bRunningEditor = GIsEditor && !IsRunningCommandlet();
 	const bool bRunningGame = !GIsEditor || IsRunningGame();
+	const bool bRunningServer = !GIsEditor || IsRunningDedicatedServer();
 	const bool bRunningCommandlet = IsRunningCommandlet();
 
 	//application flags
@@ -540,6 +541,10 @@ void FAutomationTestFramework::GetValidTestNames( TArray<FAutomationTestInfo>& T
 	{
 		ApplicationSupportFlags |= EAutomationTestFlags::ClientContext;
 	}
+	if ( bRunningServer )
+	{
+		ApplicationSupportFlags |= EAutomationTestFlags::ServerContext;
+	}
 	if ( bRunningCommandlet )
 	{
 		ApplicationSupportFlags |= EAutomationTestFlags::CommandletContext;
