Fix static analysis on loop:
  error: loop variable 'OwnerName' of type 'const FName' creates a copy from type 'const FName'

diff --git a/Engine/Source/Runtime/Core/Private/Misc/PathPermissionList.cpp b/Engine/Source/Runtime/Core/Private/Misc/PathPermissionList.cpp
--- a/Engine/Source/Runtime/Core/Private/Misc/PathPermissionList.cpp
+++ b/Engine/Source/Runtime/Core/Private/Misc/PathPermissionList.cpp
@@ -247,7 +247,7 @@ bool FPathPermissionList::UnregisterOwners(const TArray<FName>& OwnerNames)
 	{
 		TGuardValue<bool> Guard(bSuppressOnFilterChanged, true);
 
-		for (FName OwnerName : OwnerNames)
+		for (const FName& OwnerName : OwnerNames)
 		{
 			bFilterChanged |= UnregisterOwner(OwnerName);
 		}
@@ -308,7 +308,7 @@ FPathPermissionList FPathPermissionList::CombinePathFilters(const FPathPermissio
 
 	for (const TPair<FString, FPermissionListOwners>& It : GetDenyList())
 	{
-		for (const FName OwnerName : It.Value)
+		for (const FName& OwnerName : It.Value)
 		{
 			Result.AddDenyListItem(OwnerName, It.Key);
 		}
@@ -316,7 +316,7 @@ FPathPermissionList FPathPermissionList::CombinePathFilters(const FPathPermissio
 
 	for (const TPair<FString, FPermissionListOwners>& It : OtherFilter.GetDenyList())
 	{
-		for (const FName OwnerName : It.Value)
+		for (const FName& OwnerName : It.Value)
 		{
 			Result.AddDenyListItem(OwnerName, It.Key);
 		}
@@ -329,7 +329,7 @@ FPathPermissionList FPathPermissionList::CombinePathFilters(const FPathPermissio
 			const FString& Path = It.Key;
 			if (OtherFilter.PassesStartsWithFilter(Path, true))
 			{
-				for (const FName OwnerName : It.Value)
+				for (const FName& OwnerName : It.Value)
 				{
 					Result.AddAllowListItem(OwnerName, Path);
 				}
@@ -341,7 +341,7 @@ FPathPermissionList FPathPermissionList::CombinePathFilters(const FPathPermissio
 			const FString& Path = It.Key;
 			if (PassesStartsWithFilter(Path, true))
 			{
-				for (const FName OwnerName : It.Value)
+				for (const FName& OwnerName : It.Value)
 				{
 					Result.AddAllowListItem(OwnerName, Path);
 				}
