disabled busy waiting for audio FDecodeHandle tasks to avoid it picking long running tasks and hitching, added an option for ThreadPool tasks to opt out of busy waiting

diff --git a/Engine/Source/Runtime/AudioMixer/Private/AudioMixerSourceDecode.cpp b/Engine/Source/Runtime/AudioMixer/Private/AudioMixerSourceDecode.cpp
--- a/Engine/Source/Runtime/AudioMixer/Private/AudioMixerSourceDecode.cpp
+++ b/Engine/Source/Runtime/AudioMixer/Private/AudioMixerSourceDecode.cpp
@@ -199,7 +199,7 @@ public:
 	{
 		if (Task)
 		{
-			Task->EnsureCompletion();
+			Task->EnsureCompletion(/*bIsLatencySensitive =*/ true);
 			delete Task;
 		}
 	}
@@ -217,7 +217,7 @@ public:
 	{
 		if (Task)
 		{
-			Task->EnsureCompletion();
+			Task->EnsureCompletion(/*bIsLatencySensitive =*/ true);
 		}
 	}
 
@@ -228,7 +228,7 @@ public:
 			// If Cancel returns false, it means we weren't able to cancel. So lets then fallback to ensure complete.
 			if (!Task->Cancel())
 			{
-				Task->EnsureCompletion();
+				Task->EnsureCompletion(/*bIsLatencySensitive =*/ true);
 			}
 		}
 	}
diff --git a/Engine/Source/Runtime/Core/Public/Async/AsyncWork.h b/Engine/Source/Runtime/Core/Public/Async/AsyncWork.h
--- a/Engine/Source/Runtime/Core/Public/Async/AsyncWork.h
+++ b/Engine/Source/Runtime/Core/Public/Async/AsyncWork.h
@@ -358,8 +358,9 @@ class FAsyncTask
 
 	/** 
 	* Internal call to synchronize completion between threads, never called from a pool thread
+	* @param bIsLatencySensitive specifies if waiting for the task should return as soon as possible even if this delays other tasks
 	**/
-	void SyncCompletion()
+	void SyncCompletion(bool bIsLatencySensitive)
 	{
 		TRACE_CPUPROFILER_EVENT_SCOPE(FAsyncTask::SyncCompletion);
 
@@ -369,7 +370,7 @@ class FAsyncTask
 			FScopeCycleCounter Scope( Task.GetStatId() );
 			DECLARE_SCOPE_CYCLE_COUNTER( TEXT( "FAsyncTask::SyncCompletion" ), STAT_FAsyncTask_SyncCompletion, STATGROUP_ThreadPoolAsyncTasks );
 
-			if (LowLevelTasks::FScheduler::Get().IsWorkerThread())
+			if (LowLevelTasks::FScheduler::Get().IsWorkerThread() && !bIsLatencySensitive)
 			{
 				LowLevelTasks::BusyWaitUntil([this]() { return IsWorkDone(); });
 			}
@@ -462,8 +463,9 @@ public:
 	/** 
 	* Wait until the job is complete
 	* @param bDoWorkOnThisThreadIfNotStarted if true and the work has not been started, retract the async task and do it now on this thread
+	* @param specifies if waiting for the task should return as soon as possible even if this delays other tasks
 	**/
-	void EnsureCompletion(bool bDoWorkOnThisThreadIfNotStarted = true)
+	void EnsureCompletion(bool bDoWorkOnThisThreadIfNotStarted = true, bool bIsLatencySensitive = false)
 	{
 		bool DoSyncCompletion = true;
 		if (bDoWorkOnThisThreadIfNotStarted)
@@ -486,7 +488,7 @@ public:
 		}
 		if (DoSyncCompletion)
 		{
-			SyncCompletion();
+			SyncCompletion(bIsLatencySensitive);
 		}
 		CheckIdle(); // Must have had bDoWorkOnThisThreadIfNotStarted == false and needed it to be true for a synchronous job
 	}
@@ -571,7 +573,7 @@ public:
 		{
 			return false;
 		}
-		SyncCompletion();
+		SyncCompletion(/*bIsLatencySensitive = */false);
 		return true;
 	}
 
