REL - allow to not run the task continuation during cancellation, as it can be beneficial for some implementations to handle it with a manual Launch.

diff --git a/Engine/Source/Runtime/Core/Public/Async/Fundamental/Task.h b/Engine/Source/Runtime/Core/Public/Async/Fundamental/Task.h
--- a/Engine/Source/Runtime/Core/Public/Async/Fundamental/Task.h
+++ b/Engine/Source/Runtime/Core/Public/Async/Fundamental/Task.h
@@ -133,7 +133,7 @@ namespace LowLevelTasks
 			uintptr_t bAllowBusyWaiting : 1;
 			
 		public:
-			FPackedData() : FPackedData(nullptr,  ETaskPriority::Default, ETaskState::Completed, true)
+			FPackedData() : FPackedData(nullptr,  ETaskPriority::Count, ETaskState::Completed, true)
 			{
 				static_assert(!PLATFORM_32BITS, "32bit Platforms are not supported");
 				static_assert(uintptr_t(ETaskPriority::Count) <= (1ull << 3), "Not enough bits to store ETaskPriority");
@@ -225,8 +225,8 @@ namespace LowLevelTasks
 		}
 
 		
-		//try to cancel the task if it has not been launched yet the continuation will run immediately.
-		inline bool TryCancel();
+		//try to cancel the task if it has not been launched yet and ExecuteTaskOnSuccess is true the continuation will run immediately.
+		inline bool TryCancel(bool ExecuteTaskOnSuccess = true);
 
 		//try to execute the task if it has not been launched yet the task will execute immediately.
 		inline bool TryExecute();
@@ -347,7 +347,7 @@ namespace LowLevelTasks
 			|| PackedData.compare_exchange_strong(CanceledAndReadyState, FPackedData(LocalPackedData, ETaskState::Canceled), std::memory_order_acquire, std::memory_order_relaxed);
 	}
 
-	inline bool FTask::TryCancel()
+	inline bool FTask::TryCancel(bool ExecuteTaskOnSuccess)
 	{
 		FPackedData LocalPackedData = PackedData.load(std::memory_order_relaxed);
 		FPackedData ReadyState(LocalPackedData, ETaskState::Ready);
@@ -357,7 +357,7 @@ namespace LowLevelTasks
 		bool WasCanceled = PackedData.compare_exchange_strong(ReadyState, FPackedData(LocalPackedData, ETaskState::CanceledAndReady), std::memory_order_relaxed)
 			|| PackedData.compare_exchange_strong(ScheduledState, FPackedData(LocalPackedData, ETaskState::Canceled), std::memory_order_relaxed);
 
-		if(WasCanceled && TryPrepareLaunch())
+		if(ExecuteTaskOnSuccess && WasCanceled && TryPrepareLaunch())
 		{
 			ExecuteTask();
 			return true;
