Fixed usage of TCHAR_TO_UTF8 when running CrashReportClient on Mac.

diff --git a/Engine/Source/Runtime/Core/Private/Mac/MacPlatformMisc.cpp b/Engine/Source/Runtime/Core/Private/Mac/MacPlatformMisc.cpp
--- a/Engine/Source/Runtime/Core/Private/Mac/MacPlatformMisc.cpp
+++ b/Engine/Source/Runtime/Core/Private/Mac/MacPlatformMisc.cpp
@@ -2339,8 +2339,11 @@ void FMacCrashContext::GenerateCrashInfoAndLaunchReporter() const
 	{
 		// create a crash-specific directory
 		FString CrashInfoFolder = FString::Printf(TEXT("%s/CrashReport-UE4-%s-pid-%d-%s"), UTF8_TO_TCHAR(GMacAppInfo.CrashReportPath), UTF8_TO_TCHAR(GMacAppInfo.AppNameUTF8), (int32)getpid(), *GMacAppInfo.RunUUID.ToString());
-		
-		GenerateInfoInFolder(TCHAR_TO_UTF8(*CrashInfoFolder));
+
+		// Do not inline this! The lifetime of this object needs to extend over the usage of Argv in posix_spawn() call below.
+		auto CrashInfoFolderUTF8 = TStringConversion<FTCHARToUTF8_Convert>(*CrashInfoFolder);
+
+		GenerateInfoInFolder(CrashInfoFolderUTF8.Get());
 
 		CrashInfoFolder += TEXT("/");
 
@@ -2350,24 +2353,24 @@ void FMacCrashContext::GenerateCrashInfoAndLaunchReporter() const
 		int32		Argc		= 0;
 
 		Argv[Argc++] = "CrashReportClient";
-		Argv[Argc++] = TCHAR_TO_UTF8(*CrashInfoFolder);
+		Argv[Argc++] = CrashInfoFolderUTF8.Get();
 
-			if (bImplicitSend)
-			{
+		if (bImplicitSend)
+		{
 			Argv[Argc++] = "-Unattended";
 			Argv[Argc++] = "-ImplicitSend";
-			}
-			else if(GMacAppInfo.bIsUnattended)
-			{
+		}
+		else if(GMacAppInfo.bIsUnattended)
+		{
 			Argv[Argc++] ="-Unattended";
-			}
-			else
-			{
+		}
+		else
+		{
 			if (!bSendUsageData)
-				{
+			{
 				Argv[Argc++] = "-NoAnalytics";
 			}
-				}
+		}
 
 		posix_spawn_file_actions_t FileActions;
 		posix_spawn_file_actions_init(&FileActions);
@@ -2375,10 +2378,10 @@ void FMacCrashContext::GenerateCrashInfoAndLaunchReporter() const
 		posix_spawnattr_t SpawnAttr;
 		posix_spawnattr_init(&SpawnAttr);
 
-				{
+		{
 			uint32 SpawnFlags = POSIX_SPAWN_SETPGROUP;
 			posix_spawnattr_setflags(&SpawnAttr, SpawnFlags);
-				}
+		}
 
 		extern char **environ; // provided by libc
 
@@ -2391,8 +2394,8 @@ void FMacCrashContext::GenerateCrashInfoAndLaunchReporter() const
 		if (Status != 0)
 		{
 			UE_LOG(LogHAL, Fatal, TEXT("FMacPlatformMisc::GenerateCrashInfoAndLaunchReporter: posix_spawn() failed (%d, %s)"), Status, UTF8_TO_TCHAR(strerror(Status)));
-			}
 		}
+	}
 
 	// Sandboxed applications re-raise the signal to trampoline into the system crash reporter as suppressing it may fall foul of Apple's Mac App Store rules.
 	// @todo Submit an application to the MAS & see whether Apple's reviewers highlight our crash reporting or trampolining to the system reporter.
