Fixed usage of TCHAR_TO_ANSI.

diff --git a/Engine/Source/Runtime/Core/Private/FramePro/FrameProUE4.cpp b/Engine/Source/Runtime/Core/Private/FramePro/FrameProUE4.cpp
--- a/Engine/Source/Runtime/Core/Private/FramePro/FrameProUE4.cpp
+++ b/Engine/Source/Runtime/Core/Private/FramePro/FrameProUE4.cpp
@@ -482,13 +482,13 @@
 			}
 
 			FString ProcessName = FPaths::GetCleanFilename(ProcessNameOrPath);
-			const char* p_process_name = TCHAR_TO_ANSI(*ProcessName);
+			const auto p_process_name = StringCast<ANSICHAR>(*ProcessName);
 //@EPIC END
-			size_t length = strlen(p_process_name);
+			size_t length = strlen(p_process_name.Get());
 			size_t max_length = max_name_length - 1;
 
 			int copy_length = length < max_length ? length : max_length;
-			FCStringAnsi::Strncpy(p_name, p_process_name, copy_length);
+			FCStringAnsi::Strncpy(p_name, p_process_name.Get(), copy_length);
 
 			return true;
 		}
@@ -750,10 +750,10 @@
 		void Platform::GetRecordingFolder(char* p_path, int max_path_length)
 		{
 			FString path = FPaths::ProfilingDir() + TEXT("FramePro/");
-			char* p_ansi_path = TCHAR_TO_ANSI(*path);
-			int length = strlen(p_ansi_path) + 1;
+			const auto p_ansi_path = StringCast<ANSICHAR>(*path);
+			int length = strlen(p_ansi_path.Get()) + 1;
 			FRAMEPRO_ASSERT(length <= max_path_length);
-			memcpy(p_path, p_ansi_path, length);
+			memcpy(p_path, p_ansi_path.Get(), length);
 		}
 
 		//------------------------------------------------------------------------
