Replaced use of TStringView::SizeType with int32 or equivalent casts

SizeType will be deprecated in a subsequent change.

#rnx

diff --git a/Engine/Plugins/Editor/AssetReferenceRestrictions/Source/AssetReferenceRestrictions/Private/AssetReferencingDomains.cpp b/Engine/Plugins/Editor/AssetReferenceRestrictions/Source/AssetReferenceRestrictions/Private/AssetReferencingDomains.cpp
--- a/Engine/Plugins/Editor/AssetReferenceRestrictions/Source/AssetReferenceRestrictions/Private/AssetReferencingDomains.cpp
+++ b/Engine/Plugins/Editor/AssetReferenceRestrictions/Source/AssetReferenceRestrictions/Private/AssetReferencingDomains.cpp
@@ -60,7 +60,7 @@ public:
 	
 	void AddDomain(TSharedPtr<FDomainData> Domain, FStringView RemainingPath)
 	{
-		FStringView::SizeType DirectorySeparatorIndex;
+		int32 DirectorySeparatorIndex;
 		if (RemainingPath.FindChar(TEXT('/'), /*out*/ DirectorySeparatorIndex))
 		{
 			check(DirectorySeparatorIndex > 0);
diff --git a/Engine/Source/Developer/CollectionManager/Private/Collection.cpp b/Engine/Source/Developer/CollectionManager/Private/Collection.cpp
--- a/Engine/Source/Developer/CollectionManager/Private/Collection.cpp
+++ b/Engine/Source/Developer/CollectionManager/Private/Collection.cpp
@@ -105,7 +105,7 @@ bool FCollection::Load(FText& OutError)
 			break;
 		}
 
-		FStringView::SizeType Offset;
+		int32 Offset;
 		if (Line.FindChar(TEXT(':'), Offset))
 		{
 			FString Key(Line.Left(Offset));
diff --git a/Engine/Source/Developer/Zen/Internal/ZenServerInterface.cpp b/Engine/Source/Developer/Zen/Internal/ZenServerInterface.cpp
--- a/Engine/Source/Developer/Zen/Internal/ZenServerInterface.cpp
+++ b/Engine/Source/Developer/Zen/Internal/ZenServerInterface.cpp
@@ -269,7 +269,7 @@ FServiceSettings::ReadFromURL(FStringView InstanceURL)
 		InstanceURL.RightChopInline(7);
 	}
 
-	FStringView::SizeType PortDelimIndex = INDEX_NONE;
+	int32 PortDelimIndex = INDEX_NONE;
 	InstanceURL.FindChar(TEXT(':'), PortDelimIndex);
 	if (PortDelimIndex != INDEX_NONE)
 	{
diff --git a/Engine/Source/Runtime/Core/Private/HAL/LowLevelMemTracker.cpp b/Engine/Source/Runtime/Core/Private/HAL/LowLevelMemTracker.cpp
--- a/Engine/Source/Runtime/Core/Private/HAL/LowLevelMemTracker.cpp
+++ b/Engine/Source/Runtime/Core/Private/HAL/LowLevelMemTracker.cpp
@@ -1476,7 +1476,7 @@ UE::LLMPrivate::FTagData& FLowLevelMemTracker::RegisterTagData(FName Name, FName
 			const TCHAR* End = FCString::Strstr(Start, TEXT("///"));
 			if (End)
 			{
-				DisplayName = FName(FStringView(Start, static_cast<FStringView::SizeType>(End - Start)));
+				DisplayName = FName(FStringView(Start, UE_PTRDIFF_TO_INT32(End - Start)));
 			}
 		}
 	}
@@ -1501,7 +1501,7 @@ UE::LLMPrivate::FTagData& FLowLevelMemTracker::RegisterTagData(FName Name, FName
 		LLMCheckf(LeafStart[0] != '\0', TEXT("Invalid LLM custom name tag '%s'. Tag names must not end with /."), NameBuffer.ToString());
 		if (LeafStart != NameBuffer.ToString())
 		{
-			FName ParsedParentName = FName(FStringView(NameBuffer.ToString(), static_cast<FStringView::SizeType>(LeafStart - 1 - NameBuffer.ToString())));
+			FName ParsedParentName = FName(FStringView(NameBuffer.ToString(), UE_PTRDIFF_TO_INT32(LeafStart - 1 - NameBuffer.ToString())));
 			if (!ParentName.IsNone() && ParentName != ParsedParentName)
 			{
 				TStringBuilder<128> ParentBuffer;
diff --git a/Engine/Source/Runtime/Core/Private/Misc/PathViews.cpp b/Engine/Source/Runtime/Core/Private/Misc/PathViews.cpp
--- a/Engine/Source/Runtime/Core/Private/Misc/PathViews.cpp
+++ b/Engine/Source/Runtime/Core/Private/Misc/PathViews.cpp
@@ -165,7 +165,7 @@ FStringView FPathViews::GetPath(const FStringView& InPath)
 {
 	if (const TCHAR* EndPos = Algo::FindLastByPredicate(InPath, UE4PathViews_Private::IsSlashOrBackslash))
 	{
-		return InPath.Left((FStringView::SizeType)(EndPos - InPath.GetData()));
+		return InPath.Left(UE_PTRDIFF_TO_INT32(EndPos - InPath.GetData()));
 	}
 	return FStringView();
 }
@@ -175,7 +175,7 @@ FStringView FPathViews::GetExtension(const FStringView& InPath, bool bIncludeDot
 	if (const TCHAR* Dot = Algo::FindLast(GetCleanFilename(InPath), TEXT('.')))
 	{
 		const TCHAR* Extension = bIncludeDot ? Dot : Dot + 1;
-		return FStringView(Extension, (FStringView::SizeType)(InPath.GetData() + InPath.Len() - Extension));
+		return FStringView(Extension, UE_PTRDIFF_TO_INT32(InPath.GetData() + InPath.Len() - Extension));
 	}
 	return FStringView();
 }
@@ -185,7 +185,7 @@ FStringView FPathViews::GetPathLeaf(const FStringView& InPath)
 	if (const TCHAR* EndPos = Algo::FindLastByPredicate(InPath, UE4PathViews_Private::IsNotSlashOrBackslash))
 	{
 		++EndPos;
-		return GetCleanFilename(InPath.Left((FStringView::SizeType)(EndPos - InPath.GetData())));
+		return GetCleanFilename(InPath.Left(UE_PTRDIFF_TO_INT32(EndPos - InPath.GetData())));
 	}
 	return FStringView();
 }
@@ -199,7 +199,7 @@ bool FPathViews::IsPathLeaf(FStringView InPath)
 	{
 		return true;
 	}
-	if (Algo::FindByPredicate(InPath.RightChop((FStringView::SizeType)(FirstSlash - InPath.GetData())), IsNotSlashOrBackslash) == nullptr)
+	if (Algo::FindByPredicate(InPath.RightChop(UE_PTRDIFF_TO_INT32(FirstSlash - InPath.GetData())), IsNotSlashOrBackslash) == nullptr)
 	{
 		// The first slash is after the last non-slash
 		// This means it is either required for e.g. // or D:/ or /, or it is a redundant terminating slash D:/Foo/
@@ -219,7 +219,7 @@ void FPathViews::Split(const FStringView& InPath, FStringView& OutPath, FStringV
 {
 	const FStringView CleanName = GetCleanFilename(InPath);
 	const TCHAR* DotPos = Algo::FindLast(CleanName, TEXT('.'));
-	const FStringView::SizeType NameLen = DotPos ? (FStringView::SizeType)(DotPos - CleanName.GetData()) : CleanName.Len();
+	const int32 NameLen = DotPos ? UE_PTRDIFF_TO_INT32(DotPos - CleanName.GetData()) : CleanName.Len();
 	OutPath = InPath.LeftChop(CleanName.Len() + 1);
 	OutName = CleanName.Left(NameLen);
 	OutExt = CleanName.RightChop(NameLen + 1);
@@ -231,7 +231,7 @@ FString FPathViews::ChangeExtension(const FStringView& InPath, const FStringView
 	const TCHAR* PathEndPos = Algo::FindLastByPredicate(InPath, UE4PathViews_Private::IsSlashOrBackslashOrPeriod);
 	if (PathEndPos != nullptr && *PathEndPos == TEXT('.'))
 	{
-		const FStringView::SizeType Pos = FStringView::SizeType(PathEndPos - InPath.GetData());
+		const int32 Pos = UE_PTRDIFF_TO_INT32(PathEndPos - InPath.GetData());
 		const FStringView FileWithoutExtension = InPath.Left(Pos);
 
 		if (!InNewExtension.IsEmpty() && !InNewExtension.StartsWith(TEXT('.')))
@@ -379,7 +379,7 @@ bool FPathViews::IsRelativePath(FStringView InPath)
 
 	if (const TCHAR* EndPos = Algo::FindByPredicate(InPath, IsSlashOrBackslash))
 	{
-		FStringView::SizeType FirstLen = static_cast<FStringView::SizeType>(EndPos - InPath.GetData());
+		const int32 FirstLen = UE_PTRDIFF_TO_INT32(EndPos - InPath.GetData());
 		if (FirstLen == 0)
 		{
 			// Path starts with /; it may be either /Foo or //Foo
@@ -593,7 +593,7 @@ void FPathViews::SplitFirstComponent(FStringView InPath, FStringView& OutFirstCo
 
 	if (const TCHAR* EndPos = Algo::FindByPredicate(InPath, IsSlashOrBackslash))
 	{
-		FStringView::SizeType FirstLen = static_cast<FStringView::SizeType>(EndPos - InPath.GetData());
+		const int32 FirstLen = UE_PTRDIFF_TO_INT32(EndPos - InPath.GetData());
 		if (FirstLen == 0)
 		{
 			// Path starts with /; it may be either /Foo or //Foo
diff --git a/Engine/Source/Runtime/Core/Private/Serialization/CompactBinaryWriter.cpp b/Engine/Source/Runtime/Core/Private/Serialization/CompactBinaryWriter.cpp
--- a/Engine/Source/Runtime/Core/Private/Serialization/CompactBinaryWriter.cpp
+++ b/Engine/Source/Runtime/Core/Private/Serialization/CompactBinaryWriter.cpp
@@ -206,8 +206,7 @@ FUtf8StringView FCbWriter::GetActiveName() const
 		const uint8* const EncodedName = Data.GetData() + State.Offset + sizeof(ECbFieldType);
 		uint32 NameLenByteCount;
 		const uint64 NameLen = ReadVarUInt(EncodedName, NameLenByteCount);
-		using NameLenSizeType = FUtf8StringView::SizeType;
-		const NameLenSizeType ClampedNameLen = NameLenSizeType(FMath::Clamp<uint64>(NameLen, 0, ~NameLenSizeType(0)));
+		const int32 ClampedNameLen = int32(FMath::Clamp<uint64>(NameLen, 0, MAX_int32));
 		return FUtf8StringView(reinterpret_cast<const UTF8CHAR*>(EncodedName + NameLenByteCount), ClampedNameLen);
 	}
 	return FUtf8StringView();
diff --git a/Engine/Source/Runtime/Core/Private/String/Find.cpp b/Engine/Source/Runtime/Core/Private/String/Find.cpp
--- a/Engine/Source/Runtime/Core/Private/String/Find.cpp
+++ b/Engine/Source/Runtime/Core/Private/String/Find.cpp
@@ -6,12 +6,7 @@
 #include "Algo/Find.h"
 #include "Algo/NoneOf.h"
 
-namespace UE
-{
-namespace String
-{
-
-namespace Private
+namespace UE::String::Private
 {
 
 // These are naive implementations that take time proportional to View.Len() * TotalSearchLen.
@@ -19,17 +14,17 @@ namespace Private
 // There are algorithms for each that are linear or sub-linear in the length of the string to search.
 
 template <typename CharType>
-static inline typename TStringView<CharType>::SizeType FindFirst(TStringView<CharType> View, TStringView<CharType> Search, ESearchCase::Type SearchCase)
+static inline int32 FindFirst(TStringView<CharType> View, TStringView<CharType> Search, ESearchCase::Type SearchCase)
 {
 	check(!Search.IsEmpty());
-	const typename TStringView<CharType>::SizeType SearchLen = Search.Len();
+	const int32 SearchLen = Search.Len();
 	if (SearchLen == 1)
 	{
 		return String::FindFirstChar(View, Search[0], SearchCase);
 	}
 	const CharType* const SearchData = Search.GetData();
 	const CharType* const ViewBegin = View.GetData();
-	const typename TStringView<CharType>::SizeType ViewLen = View.Len();
+	const int32 ViewLen = View.Len();
 	if (ViewLen < SearchLen)
 	{
 		return INDEX_NONE;
@@ -41,7 +36,7 @@ static inline typename TStringView<CharType>::SizeType FindFirst(TStringView<Cha
 		{
 			if (TCString<CharType>::Strncmp(ViewIt, SearchData, SearchLen) == 0)
 			{
-				return typename TStringView<CharType>::SizeType(ViewIt - ViewBegin);
+				return int32(ViewIt - ViewBegin);
 			}
 		}
 	}
@@ -51,7 +46,7 @@ static inline typename TStringView<CharType>::SizeType FindFirst(TStringView<Cha
 		{
 			if (TCString<CharType>::Strnicmp(ViewIt, SearchData, SearchLen) == 0)
 			{
-				return typename TStringView<CharType>::SizeType(ViewIt - ViewBegin);
+				return int32(ViewIt - ViewBegin);
 			}
 		}
 	}
@@ -59,17 +54,17 @@ static inline typename TStringView<CharType>::SizeType FindFirst(TStringView<Cha
 }
 
 template <typename CharType>
-static inline typename TStringView<CharType>::SizeType FindLast(TStringView<CharType> View, TStringView<CharType> Search, ESearchCase::Type SearchCase)
+static inline int32 FindLast(TStringView<CharType> View, TStringView<CharType> Search, ESearchCase::Type SearchCase)
 {
 	check(!Search.IsEmpty());
-	const typename TStringView<CharType>::SizeType SearchLen = Search.Len();
+	const int32 SearchLen = Search.Len();
 	if (SearchLen == 1)
 	{
 		return String::FindLastChar(View, Search[0], SearchCase);
 	}
 	const CharType* const SearchData = Search.GetData();
 	const CharType* const ViewBegin = View.GetData();
-	const typename TStringView<CharType>::SizeType ViewLen = View.Len();
+	const int32 ViewLen = View.Len();
 	if (ViewLen < SearchLen)
 	{
 		return INDEX_NONE;
@@ -81,7 +76,7 @@ static inline typename TStringView<CharType>::SizeType FindLast(TStringView<Char
 		{
 			if (TCString<CharType>::Strncmp(ViewIt, SearchData, SearchLen) == 0)
 			{
-				return typename TStringView<CharType>::SizeType(ViewIt - ViewBegin);
+				return int32(ViewIt - ViewBegin);
 			}
 		}
 	}
@@ -91,7 +86,7 @@ static inline typename TStringView<CharType>::SizeType FindLast(TStringView<Char
 		{
 			if (TCString<CharType>::Strnicmp(ViewIt, SearchData, SearchLen) == 0)
 			{
-				return typename TStringView<CharType>::SizeType(ViewIt - ViewBegin);
+				return int32(ViewIt - ViewBegin);
 			}
 		}
 	}
@@ -99,7 +94,7 @@ static inline typename TStringView<CharType>::SizeType FindLast(TStringView<Char
 }
 
 template <typename CharType>
-static inline typename TStringView<CharType>::SizeType FindFirstOfAny(TStringView<CharType> View, TConstArrayView<TStringView<CharType>> Search, ESearchCase::Type SearchCase)
+static inline int32 FindFirstOfAny(TStringView<CharType> View, TConstArrayView<TStringView<CharType>> Search, ESearchCase::Type SearchCase)
 {
 	check(Algo::NoneOf(Search, &TStringView<CharType>::IsEmpty));
 	switch (Search.Num())
@@ -123,8 +118,8 @@ static inline typename TStringView<CharType>::SizeType FindFirstOfAny(TStringVie
 	}
 
 	const CharType* const ViewBegin = View.GetData();
-	const typename TStringView<CharType>::SizeType ViewLen = View.Len();
-	for (typename TStringView<CharType>::SizeType ViewIndex = 0; ViewIndex != ViewLen; ++ViewIndex)
+	const int32 ViewLen = View.Len();
+	for (int32 ViewIndex = 0; ViewIndex != ViewLen; ++ViewIndex)
 	{
 		const TStringView<CharType> RemainingView(ViewBegin + ViewIndex, ViewLen - ViewIndex);
 		auto MatchPattern = [&RemainingView, SearchCase](const TStringView<CharType>& Pattern) { return RemainingView.StartsWith(Pattern, SearchCase); };
@@ -137,7 +132,7 @@ static inline typename TStringView<CharType>::SizeType FindFirstOfAny(TStringVie
 }
 
 template <typename CharType>
-static inline typename TStringView<CharType>::SizeType FindLastOfAny(TStringView<CharType> View, TConstArrayView<TStringView<CharType>> Search, ESearchCase::Type SearchCase)
+static inline int32 FindLastOfAny(TStringView<CharType> View, TConstArrayView<TStringView<CharType>> Search, ESearchCase::Type SearchCase)
 {
 	check(Algo::NoneOf(Search, &TStringView<CharType>::IsEmpty));
 	switch (Search.Num())
@@ -162,12 +157,12 @@ static inline typename TStringView<CharType>::SizeType FindLastOfAny(TStringView
 
 
 	const CharType* const ViewBegin = View.GetData();
-	const typename TStringView<CharType>::SizeType ViewLen = View.Len();
+	const int32 ViewLen = View.Len();
 	if (ViewLen == 0)
 	{
 		return INDEX_NONE;
 	}
-	for (typename TStringView<CharType>::SizeType ViewIndex = ViewLen - 1; ViewIndex >= 0; --ViewIndex)
+	for (int32 ViewIndex = ViewLen - 1; ViewIndex >= 0; --ViewIndex)
 	{
 		const TStringView<CharType> RemainingView(ViewBegin + ViewIndex, ViewLen - ViewIndex);
 		auto MatchPattern = [&RemainingView, SearchCase](const TStringView<CharType>& Pattern) { return RemainingView.StartsWith(Pattern, SearchCase); };
@@ -180,7 +175,7 @@ static inline typename TStringView<CharType>::SizeType FindLastOfAny(TStringView
 }
 
 template <typename CharType>
-static inline typename TStringView<CharType>::SizeType FindFirstChar(TStringView<CharType> View, CharType Search, ESearchCase::Type SearchCase)
+static inline int32 FindFirstChar(TStringView<CharType> View, CharType Search, ESearchCase::Type SearchCase)
 {
 	const CharType* const ViewBegin = View.GetData();
 	const CharType* const ViewEnd = ViewBegin + View.Len();
@@ -190,7 +185,7 @@ static inline typename TStringView<CharType>::SizeType FindFirstChar(TStringView
 		{
 			if (*ViewIt == Search)
 			{
-				return typename TStringView<CharType>::SizeType(ViewIt - ViewBegin);
+				return int32(ViewIt - ViewBegin);
 			}
 		}
 	}
@@ -201,7 +196,7 @@ static inline typename TStringView<CharType>::SizeType FindFirstChar(TStringView
 		{
 			if (TChar<CharType>::ToUpper(*ViewIt) == SearchUpper)
 			{
-				return typename TStringView<CharType>::SizeType(ViewIt - ViewBegin);
+				return int32(ViewIt - ViewBegin);
 			}
 		}
 	}
@@ -209,7 +204,7 @@ static inline typename TStringView<CharType>::SizeType FindFirstChar(TStringView
 }
 
 template <typename CharType>
-static inline typename TStringView<CharType>::SizeType FindLastChar(TStringView<CharType> View, CharType Search, ESearchCase::Type SearchCase)
+static inline int32 FindLastChar(TStringView<CharType> View, CharType Search, ESearchCase::Type SearchCase)
 {
 	const CharType* const ViewBegin = View.GetData();
 	const CharType* const ViewEnd = ViewBegin + View.Len();
@@ -223,7 +218,7 @@ static inline typename TStringView<CharType>::SizeType FindLastChar(TStringView<
 		{
 			if (*ViewIt == Search)
 			{
-				return typename TStringView<CharType>::SizeType(ViewIt - ViewBegin);
+				return int32(ViewIt - ViewBegin);
 			}
 		}
 	}
@@ -234,7 +229,7 @@ static inline typename TStringView<CharType>::SizeType FindLastChar(TStringView<
 		{
 			if (TChar<CharType>::ToUpper(*ViewIt) == SearchUpper)
 			{
-				return typename TStringView<CharType>::SizeType(ViewIt - ViewBegin);
+				return int32(ViewIt - ViewBegin);
 			}
 		}
 	}
@@ -242,7 +237,7 @@ static inline typename TStringView<CharType>::SizeType FindLastChar(TStringView<
 }
 
 template <typename CharType>
-static inline typename TStringView<CharType>::SizeType FindFirstOfAnyChar(TStringView<CharType> View, TConstArrayView<CharType> Search, ESearchCase::Type SearchCase)
+static inline int32 FindFirstOfAnyChar(TStringView<CharType> View, TConstArrayView<CharType> Search, ESearchCase::Type SearchCase)
 {
 	switch (Search.Num())
 	{
@@ -262,7 +257,7 @@ static inline typename TStringView<CharType>::SizeType FindFirstOfAnyChar(TStrin
 		{
 			if (Algo::Find(Search, *ViewIt))
 			{
-				return typename TStringView<CharType>::SizeType(ViewIt - ViewBegin);
+				return int32(ViewIt - ViewBegin);
 			}
 		}
 		return INDEX_NONE;
@@ -279,7 +274,7 @@ static inline typename TStringView<CharType>::SizeType FindFirstOfAnyChar(TStrin
 		{
 			if (Algo::Find(SearchUpper, TChar<CharType>::ToUpper(*ViewIt)))
 			{
-				return typename TStringView<CharType>::SizeType(ViewIt - ViewBegin);
+				return int32(ViewIt - ViewBegin);
 			}
 		}
 		return INDEX_NONE;
@@ -287,7 +282,7 @@ static inline typename TStringView<CharType>::SizeType FindFirstOfAnyChar(TStrin
 }
 
 template <typename CharType>
-static inline typename TStringView<CharType>::SizeType FindLastOfAnyChar(TStringView<CharType> View, TConstArrayView<CharType> Search, ESearchCase::Type SearchCase)
+static inline int32 FindLastOfAnyChar(TStringView<CharType> View, TConstArrayView<CharType> Search, ESearchCase::Type SearchCase)
 {
 	switch (Search.Num())
 	{
@@ -311,7 +306,7 @@ static inline typename TStringView<CharType>::SizeType FindLastOfAnyChar(TString
 		{
 			if (Algo::Find(Search, *ViewIt))
 			{
-				return typename TStringView<CharType>::SizeType(ViewIt - ViewBegin);
+				return int32(ViewIt - ViewBegin);
 			}
 		}
 		return INDEX_NONE;
@@ -328,94 +323,96 @@ static inline typename TStringView<CharType>::SizeType FindLastOfAnyChar(TString
 		{
 			if (Algo::Find(SearchUpper, TChar<CharType>::ToUpper(*ViewIt)))
 			{
-				return typename TStringView<CharType>::SizeType(ViewIt - ViewBegin);
+				return int32(ViewIt - ViewBegin);
 			}
 		}
 		return INDEX_NONE;
 	}
 }
 
-}
+} // UE::String::Private
 
-FAnsiStringView::SizeType FindFirst(FAnsiStringView View, FAnsiStringView Search, ESearchCase::Type SearchCase)
+namespace UE::String
+{
+
+int32 FindFirst(FAnsiStringView View, FAnsiStringView Search, ESearchCase::Type SearchCase)
 {
 	return Private::FindFirst(View, Search, SearchCase);
 }
 
-FWideStringView::SizeType FindFirst(FWideStringView View, FWideStringView Search, ESearchCase::Type SearchCase)
+int32 FindFirst(FWideStringView View, FWideStringView Search, ESearchCase::Type SearchCase)
 {
 	return Private::FindFirst(View, Search, SearchCase);
 }
 
-FAnsiStringView::SizeType FindLast(FAnsiStringView View, FAnsiStringView Search, ESearchCase::Type SearchCase)
+int32 FindLast(FAnsiStringView View, FAnsiStringView Search, ESearchCase::Type SearchCase)
 {
 	return Private::FindLast(View, Search, SearchCase);
 }
 
-FWideStringView::SizeType FindLast(FWideStringView View, FWideStringView Search, ESearchCase::Type SearchCase)
+int32 FindLast(FWideStringView View, FWideStringView Search, ESearchCase::Type SearchCase)
 {
 	return Private::FindLast(View, Search, SearchCase);
 }
 
-FAnsiStringView::SizeType FindFirstOfAny(FAnsiStringView View, TConstArrayView<FAnsiStringView> Search, ESearchCase::Type SearchCase)
+int32 FindFirstOfAny(FAnsiStringView View, TConstArrayView<FAnsiStringView> Search, ESearchCase::Type SearchCase)
 {
 	return Private::FindFirstOfAny(View, Search, SearchCase);
 }
 
-FWideStringView::SizeType FindFirstOfAny(FWideStringView View, TConstArrayView<FWideStringView> Search, ESearchCase::Type SearchCase)
+int32 FindFirstOfAny(FWideStringView View, TConstArrayView<FWideStringView> Search, ESearchCase::Type SearchCase)
 {
 	return Private::FindFirstOfAny(View, Search, SearchCase);
 }
 
-FAnsiStringView::SizeType FindLastOfAny(FAnsiStringView View, TConstArrayView<FAnsiStringView> Search, ESearchCase::Type SearchCase)
+int32 FindLastOfAny(FAnsiStringView View, TConstArrayView<FAnsiStringView> Search, ESearchCase::Type SearchCase)
 {
 	return Private::FindLastOfAny(View, Search, SearchCase);
 }
 
-FWideStringView::SizeType FindLastOfAny(FWideStringView View, TConstArrayView<FWideStringView> Search, ESearchCase::Type SearchCase)
+int32 FindLastOfAny(FWideStringView View, TConstArrayView<FWideStringView> Search, ESearchCase::Type SearchCase)
 {
 	return Private::FindLastOfAny(View, Search, SearchCase);
 }
 
-FAnsiStringView::SizeType FindFirstChar(FAnsiStringView View, ANSICHAR Search, ESearchCase::Type SearchCase)
+int32 FindFirstChar(FAnsiStringView View, ANSICHAR Search, ESearchCase::Type SearchCase)
 {
 	return Private::FindFirstChar(View, Search, SearchCase);
 }
 
-FWideStringView::SizeType FindFirstChar(FWideStringView View, WIDECHAR Search, ESearchCase::Type SearchCase)
+int32 FindFirstChar(FWideStringView View, WIDECHAR Search, ESearchCase::Type SearchCase)
 {
 	return Private::FindFirstChar(View, Search, SearchCase);
 }
 
-FAnsiStringView::SizeType FindLastChar(FAnsiStringView View, ANSICHAR Search, ESearchCase::Type SearchCase)
+int32 FindLastChar(FAnsiStringView View, ANSICHAR Search, ESearchCase::Type SearchCase)
 {
 	return Private::FindLastChar(View, Search, SearchCase);
 }
 
-FWideStringView::SizeType FindLastChar(FWideStringView View, WIDECHAR Search, ESearchCase::Type SearchCase)
+int32 FindLastChar(FWideStringView View, WIDECHAR Search, ESearchCase::Type SearchCase)
 {
 	return Private::FindLastChar(View, Search, SearchCase);
 }
 
-FAnsiStringView::SizeType FindFirstOfAnyChar(FAnsiStringView View, TConstArrayView<ANSICHAR> Search, ESearchCase::Type SearchCase)
+int32 FindFirstOfAnyChar(FAnsiStringView View, TConstArrayView<ANSICHAR> Search, ESearchCase::Type SearchCase)
 {
 	return Private::FindFirstOfAnyChar(View, Search, SearchCase);
 }
 
-FWideStringView::SizeType FindFirstOfAnyChar(FWideStringView View, TConstArrayView<WIDECHAR> Search, ESearchCase::Type SearchCase)
+int32 FindFirstOfAnyChar(FWideStringView View, TConstArrayView<WIDECHAR> Search, ESearchCase::Type SearchCase)
 {
 	return Private::FindFirstOfAnyChar(View, Search, SearchCase);
 }
 
-FAnsiStringView::SizeType FindLastOfAnyChar(FAnsiStringView View, TConstArrayView<ANSICHAR> Search, ESearchCase::Type SearchCase)
+int32 FindLastOfAnyChar(FAnsiStringView View, TConstArrayView<ANSICHAR> Search, ESearchCase::Type SearchCase)
 {
 	return Private::FindLastOfAnyChar(View, Search, SearchCase);
 }
 
-FWideStringView::SizeType FindLastOfAnyChar(FWideStringView View, TConstArrayView<WIDECHAR> Search, ESearchCase::Type SearchCase)
+int32 FindLastOfAnyChar(FWideStringView View, TConstArrayView<WIDECHAR> Search, ESearchCase::Type SearchCase)
 {
 	return Private::FindLastOfAnyChar(View, Search, SearchCase);
 }
 
-}
-}
+} // UE::String
diff --git a/Engine/Source/Runtime/Core/Private/String/ParseLines.cpp b/Engine/Source/Runtime/Core/Private/String/ParseLines.cpp
--- a/Engine/Source/Runtime/Core/Private/String/ParseLines.cpp
+++ b/Engine/Source/Runtime/Core/Private/String/ParseLines.cpp
@@ -35,7 +35,7 @@ void ParseLines(
 			}
 		}
 
-		FStringView Line(LineStart, static_cast<FStringView::SizeType>(LineEnd - LineStart));
+		FStringView Line(LineStart, UE_PTRDIFF_TO_INT32(LineEnd - LineStart));
 		if (EnumHasAnyFlags(Options, EParseLinesOptions::Trim))
 		{
 			Line = Line.TrimStartAndEnd();
diff --git a/Engine/Source/Runtime/Core/Private/String/ParseTokens.cpp b/Engine/Source/Runtime/Core/Private/String/ParseTokens.cpp
--- a/Engine/Source/Runtime/Core/Private/String/ParseTokens.cpp
+++ b/Engine/Source/Runtime/Core/Private/String/ParseTokens.cpp
@@ -52,7 +52,7 @@ inline static void ParseTokens1Delim1Char(
 				++ViewIt;
 				continue;
 			}
-			ParseTokensVisitToken(Visitor, Options, FStringView(NextToken, static_cast<FStringView::SizeType>(ViewIt - NextToken)));
+			ParseTokensVisitToken(Visitor, Options, FStringView(NextToken, UE_PTRDIFF_TO_INT32(ViewIt - NextToken)));
 			NextToken = ++ViewIt;
 		}
 	}
@@ -69,12 +69,12 @@ inline static void ParseTokens1Delim1Char(
 				++ViewIt;
 				continue;
 			}
-			ParseTokensVisitToken(Visitor, Options, FStringView(NextToken, static_cast<FStringView::SizeType>(ViewIt - NextToken)));
+			ParseTokensVisitToken(Visitor, Options, FStringView(NextToken, UE_PTRDIFF_TO_INT32(ViewIt - NextToken)));
 			NextToken = ++ViewIt;
 		}
 	}
 
-	ParseTokensVisitToken(Visitor, Options, FStringView(NextToken, static_cast<FStringView::SizeType>(ViewIt - NextToken)));
+	ParseTokensVisitToken(Visitor, Options, FStringView(NextToken, UE_PTRDIFF_TO_INT32(ViewIt - NextToken)));
 }
 
 /** Parse tokens with multiple single-character Basic Latin delimiters. */
@@ -117,11 +117,11 @@ inline static void ParseTokensNDelim1CharBasicLatin(
 			++ViewIt;
 			continue;
 		}
-		ParseTokensVisitToken(Visitor, Options, FStringView(NextToken, static_cast<FStringView::SizeType>(ViewIt - NextToken)));
+		ParseTokensVisitToken(Visitor, Options, FStringView(NextToken, UE_PTRDIFF_TO_INT32(ViewIt - NextToken)));
 		NextToken = ++ViewIt;
 	}
 
-	ParseTokensVisitToken(Visitor, Options, FStringView(NextToken, static_cast<FStringView::SizeType>(ViewIt - NextToken)));
+	ParseTokensVisitToken(Visitor, Options, FStringView(NextToken, UE_PTRDIFF_TO_INT32(ViewIt - NextToken)));
 }
 
 /** Parse tokens with multiple single-character delimiters in the Basic Multilingual Plane. */
@@ -155,7 +155,7 @@ inline static void ParseTokensNDelim1Char(
 				++ViewIt;
 				continue;
 			}
-			ParseTokensVisitToken(Visitor, Options, FStringView(NextToken, static_cast<FStringView::SizeType>(ViewIt - NextToken)));
+			ParseTokensVisitToken(Visitor, Options, FStringView(NextToken, UE_PTRDIFF_TO_INT32(ViewIt - NextToken)));
 			NextToken = ++ViewIt;
 		}
 	}
@@ -172,12 +172,12 @@ inline static void ParseTokensNDelim1Char(
 				++ViewIt;
 				continue;
 			}
-			ParseTokensVisitToken(Visitor, Options, FStringView(NextToken, static_cast<FStringView::SizeType>(ViewIt - NextToken)));
+			ParseTokensVisitToken(Visitor, Options, FStringView(NextToken, UE_PTRDIFF_TO_INT32(ViewIt - NextToken)));
 			NextToken = ++ViewIt;
 		}
 	}
 
-	ParseTokensVisitToken(Visitor, Options, FStringView(NextToken, static_cast<FStringView::SizeType>(ViewIt - NextToken)));
+	ParseTokensVisitToken(Visitor, Options, FStringView(NextToken, UE_PTRDIFF_TO_INT32(ViewIt - NextToken)));
 }
 
 /** Parse tokens with multiple multi-character delimiters. */
@@ -191,11 +191,11 @@ inline static void ParseTokensNDelimNChar(
 	// If this function becomes a bottleneck, it can be specialized separately for one and many delimiters.
 	// There are algorithms for each are linear or sub-linear in the length of string to search.
 
-	const FStringView::SizeType ViewLen = View.Len();
-	FStringView::SizeType NextTokenIndex = 0;
+	const int32 ViewLen = View.Len();
+	int32 NextTokenIndex = 0;
 
 	const ESearchCase::Type SearchCase = EnumHasAnyFlags(Options, EParseTokensOptions::IgnoreCase) ? ESearchCase::IgnoreCase : ESearchCase::CaseSensitive;
-	for (FStringView::SizeType ViewIndex = 0; ViewIndex != ViewLen;)
+	for (int32 ViewIndex = 0; ViewIndex != ViewLen;)
 	{
 		const FStringView RemainingView(View.GetData() + ViewIndex, ViewLen - ViewIndex);
 		auto MatchDelimiter = [RemainingView, SearchCase](FStringView Delimiter) { return RemainingView.StartsWith(Delimiter, SearchCase); };
diff --git a/Engine/Source/Runtime/Core/Public/Serialization/CompactBinaryValue.h b/Engine/Source/Runtime/Core/Public/Serialization/CompactBinaryValue.h
--- a/Engine/Source/Runtime/Core/Public/Serialization/CompactBinaryValue.h
+++ b/Engine/Source/Runtime/Core/Public/Serialization/CompactBinaryValue.h
@@ -250,7 +250,7 @@ FORCEINLINE FCbCustomByName FCbValue::AsCustomByName() const
 	FCbCustomByName Value;
 	Value.Name = FUtf8StringView(
 		reinterpret_cast<const UTF8CHAR*>(Bytes),
-		static_cast<FUtf8StringView::SizeType>(TypeNameLen));
+		static_cast<int32>(TypeNameLen));
 	Value.Data = MakeMemoryView(Bytes + TypeNameLen, DataSize - TypeNameLen - TypeNameLenByteCount);
 	return Value;
 }
diff --git a/Engine/Source/Runtime/Core/Public/String/Find.h b/Engine/Source/Runtime/Core/Public/String/Find.h
--- a/Engine/Source/Runtime/Core/Public/String/Find.h
+++ b/Engine/Source/Runtime/Core/Public/String/Find.h
@@ -6,9 +6,7 @@
 #include "Containers/ArrayView.h"
 #include "Containers/StringView.h"
 
-namespace UE
-{
-namespace String
+namespace UE::String
 {
 
 /**
@@ -19,8 +17,8 @@ namespace String
  * @param SearchCase Whether the comparison should ignore case.
  * @return The position at which the search string was found, or INDEX_NONE if not found.
  */
-CORE_API FAnsiStringView::SizeType FindFirst(FAnsiStringView View, FAnsiStringView Search, ESearchCase::Type SearchCase = ESearchCase::CaseSensitive);
-CORE_API FWideStringView::SizeType FindFirst(FWideStringView View, FWideStringView Search, ESearchCase::Type SearchCase = ESearchCase::CaseSensitive);
+CORE_API int32 FindFirst(FAnsiStringView View, FAnsiStringView Search, ESearchCase::Type SearchCase = ESearchCase::CaseSensitive);
+CORE_API int32 FindFirst(FWideStringView View, FWideStringView Search, ESearchCase::Type SearchCase = ESearchCase::CaseSensitive);
 
 /**
  * Search the view for the last occurrence of the search string.
@@ -30,8 +28,8 @@ CORE_API FWideStringView::SizeType FindFirst(FWideStringView View, FWideStringVi
  * @param SearchCase Whether the comparison should ignore case.
  * @return The position at which the search string was found, or INDEX_NONE if not found.
  */
-CORE_API FAnsiStringView::SizeType FindLast(FAnsiStringView View, FAnsiStringView Search, ESearchCase::Type SearchCase = ESearchCase::CaseSensitive);
-CORE_API FWideStringView::SizeType FindLast(FWideStringView View, FWideStringView Search, ESearchCase::Type SearchCase = ESearchCase::CaseSensitive);
+CORE_API int32 FindLast(FAnsiStringView View, FAnsiStringView Search, ESearchCase::Type SearchCase = ESearchCase::CaseSensitive);
+CORE_API int32 FindLast(FWideStringView View, FWideStringView Search, ESearchCase::Type SearchCase = ESearchCase::CaseSensitive);
 
 /**
  * Search the view for the first occurrence of any search string.
@@ -41,8 +39,8 @@ CORE_API FWideStringView::SizeType FindLast(FWideStringView View, FWideStringVie
  * @param SearchCase Whether the comparison should ignore case.
  * @return The position at which any search string was found, or INDEX_NONE if not found.
  */
-CORE_API FAnsiStringView::SizeType FindFirstOfAny(FAnsiStringView View, TConstArrayView<FAnsiStringView> Search, ESearchCase::Type SearchCase = ESearchCase::CaseSensitive);
-CORE_API FWideStringView::SizeType FindFirstOfAny(FWideStringView View, TConstArrayView<FWideStringView> Search, ESearchCase::Type SearchCase = ESearchCase::CaseSensitive);
+CORE_API int32 FindFirstOfAny(FAnsiStringView View, TConstArrayView<FAnsiStringView> Search, ESearchCase::Type SearchCase = ESearchCase::CaseSensitive);
+CORE_API int32 FindFirstOfAny(FWideStringView View, TConstArrayView<FWideStringView> Search, ESearchCase::Type SearchCase = ESearchCase::CaseSensitive);
 
 /**
  * Search the view for the last occurrence of any search string.
@@ -52,8 +50,8 @@ CORE_API FWideStringView::SizeType FindFirstOfAny(FWideStringView View, TConstAr
  * @param SearchCase Whether the comparison should ignore case.
  * @return The position at which any search string was found, or INDEX_NONE if not found.
  */
-CORE_API FAnsiStringView::SizeType FindLastOfAny(FAnsiStringView View, TConstArrayView<FAnsiStringView> Search, ESearchCase::Type SearchCase = ESearchCase::CaseSensitive);
-CORE_API FWideStringView::SizeType FindLastOfAny(FWideStringView View, TConstArrayView<FWideStringView> Search, ESearchCase::Type SearchCase = ESearchCase::CaseSensitive);
+CORE_API int32 FindLastOfAny(FAnsiStringView View, TConstArrayView<FAnsiStringView> Search, ESearchCase::Type SearchCase = ESearchCase::CaseSensitive);
+CORE_API int32 FindLastOfAny(FWideStringView View, TConstArrayView<FWideStringView> Search, ESearchCase::Type SearchCase = ESearchCase::CaseSensitive);
 
 /**
  * Search the view for the first occurrence of the search character.
@@ -63,8 +61,8 @@ CORE_API FWideStringView::SizeType FindLastOfAny(FWideStringView View, TConstArr
  * @param SearchCase Whether the comparison should ignore case.
  * @return The position at which the search character was found, or INDEX_NONE if not found.
  */
-CORE_API FAnsiStringView::SizeType FindFirstChar(FAnsiStringView View, ANSICHAR Search, ESearchCase::Type SearchCase = ESearchCase::CaseSensitive);
-CORE_API FWideStringView::SizeType FindFirstChar(FWideStringView View, WIDECHAR Search, ESearchCase::Type SearchCase = ESearchCase::CaseSensitive);
+CORE_API int32 FindFirstChar(FAnsiStringView View, ANSICHAR Search, ESearchCase::Type SearchCase = ESearchCase::CaseSensitive);
+CORE_API int32 FindFirstChar(FWideStringView View, WIDECHAR Search, ESearchCase::Type SearchCase = ESearchCase::CaseSensitive);
 
 /**
  * Search the view for the last occurrence of the search character.
@@ -74,8 +72,8 @@ CORE_API FWideStringView::SizeType FindFirstChar(FWideStringView View, WIDECHAR
  * @param SearchCase Whether the comparison should ignore case.
  * @return The position at which the search character was found, or INDEX_NONE if not found.
  */
-CORE_API FAnsiStringView::SizeType FindLastChar(FAnsiStringView View, ANSICHAR Search, ESearchCase::Type SearchCase = ESearchCase::CaseSensitive);
-CORE_API FWideStringView::SizeType FindLastChar(FWideStringView View, WIDECHAR Search, ESearchCase::Type SearchCase = ESearchCase::CaseSensitive);
+CORE_API int32 FindLastChar(FAnsiStringView View, ANSICHAR Search, ESearchCase::Type SearchCase = ESearchCase::CaseSensitive);
+CORE_API int32 FindLastChar(FWideStringView View, WIDECHAR Search, ESearchCase::Type SearchCase = ESearchCase::CaseSensitive);
 
 /**
  * Search the view for the first occurrence of any search character.
@@ -85,8 +83,8 @@ CORE_API FWideStringView::SizeType FindLastChar(FWideStringView View, WIDECHAR S
  * @param SearchCase Whether the comparison should ignore case.
  * @return The position at which any search character was found, or INDEX_NONE if not found.
  */
-CORE_API FAnsiStringView::SizeType FindFirstOfAnyChar(FAnsiStringView View, TConstArrayView<ANSICHAR> Search, ESearchCase::Type SearchCase = ESearchCase::CaseSensitive);
-CORE_API FWideStringView::SizeType FindFirstOfAnyChar(FWideStringView View, TConstArrayView<WIDECHAR> Search, ESearchCase::Type SearchCase = ESearchCase::CaseSensitive);
+CORE_API int32 FindFirstOfAnyChar(FAnsiStringView View, TConstArrayView<ANSICHAR> Search, ESearchCase::Type SearchCase = ESearchCase::CaseSensitive);
+CORE_API int32 FindFirstOfAnyChar(FWideStringView View, TConstArrayView<WIDECHAR> Search, ESearchCase::Type SearchCase = ESearchCase::CaseSensitive);
 
 /**
  * Search the view for the last occurrence of any search character.
@@ -96,8 +94,7 @@ CORE_API FWideStringView::SizeType FindFirstOfAnyChar(FWideStringView View, TCon
  * @param SearchCase Whether the comparison should ignore case.
  * @return The position at which any search character was found, or INDEX_NONE if not found.
  */
-CORE_API FAnsiStringView::SizeType FindLastOfAnyChar(FAnsiStringView View, TConstArrayView<ANSICHAR> Search, ESearchCase::Type SearchCase = ESearchCase::CaseSensitive);
-CORE_API FWideStringView::SizeType FindLastOfAnyChar(FWideStringView View, TConstArrayView<WIDECHAR> Search, ESearchCase::Type SearchCase = ESearchCase::CaseSensitive);
+CORE_API int32 FindLastOfAnyChar(FAnsiStringView View, TConstArrayView<ANSICHAR> Search, ESearchCase::Type SearchCase = ESearchCase::CaseSensitive);
+CORE_API int32 FindLastOfAnyChar(FWideStringView View, TConstArrayView<WIDECHAR> Search, ESearchCase::Type SearchCase = ESearchCase::CaseSensitive);
 
-}
-}
+} // UE::String
diff --git a/Engine/Source/Runtime/Core/Public/String/ParseLines.h b/Engine/Source/Runtime/Core/Public/String/ParseLines.h
--- a/Engine/Source/Runtime/Core/Public/String/ParseLines.h
+++ b/Engine/Source/Runtime/Core/Public/String/ParseLines.h
@@ -10,7 +10,7 @@
 namespace UE::String
 {
 
-enum class EParseLinesOptions
+enum class EParseLinesOptions : uint32
 {
 	/** Use the default options when parsing lines. */
 	None = 0,
diff --git a/Engine/Source/Runtime/Core/Public/String/ParseTokens.h b/Engine/Source/Runtime/Core/Public/String/ParseTokens.h
--- a/Engine/Source/Runtime/Core/Public/String/ParseTokens.h
+++ b/Engine/Source/Runtime/Core/Public/String/ParseTokens.h
@@ -11,7 +11,7 @@
 namespace UE::String
 {
 
-enum class EParseTokensOptions
+enum class EParseTokensOptions : uint32
 {
 	/** Use the default options when parsing tokens. */
 	None = 0,
diff --git a/Engine/Source/Runtime/CoreUObject/Private/Misc/PackageName.cpp b/Engine/Source/Runtime/CoreUObject/Private/Misc/PackageName.cpp
--- a/Engine/Source/Runtime/CoreUObject/Private/Misc/PackageName.cpp
+++ b/Engine/Source/Runtime/CoreUObject/Private/Misc/PackageName.cpp
@@ -811,7 +811,7 @@ bool FPackageName::DoesPackageNameContainInvalidCharacters(FStringView InLongPac
 	TStringBuilder<32> MatchedInvalidChars;
 	for (const TCHAR* InvalidCharacters = INVALID_LONGPACKAGE_CHARACTERS; *InvalidCharacters; ++InvalidCharacters)
 	{
-		FStringView::SizeType OutIndex;
+		int32 OutIndex;
 		if (InLongPackageName.FindChar(*InvalidCharacters, OutIndex))
 		{
 			MatchedInvalidChars += *InvalidCharacters;
diff --git a/Engine/Source/Runtime/CoreUObject/Private/UObject/ObjectPathId.cpp b/Engine/Source/Runtime/CoreUObject/Private/UObject/ObjectPathId.cpp
--- a/Engine/Source/Runtime/CoreUObject/Private/UObject/ObjectPathId.cpp
+++ b/Engine/Source/Runtime/CoreUObject/Private/UObject/ObjectPathId.cpp
@@ -168,10 +168,10 @@ FName FObjectPathId::MakeImportPathIdAndPackageName(const FObjectImport& Import,
 }
 
 #if !(UE_BUILD_SHIPPING || UE_BUILD_TEST)
-template <typename T>
+template <typename CharType>
 struct TStringViewNamePathProducer
 {
-	TStringViewNamePathProducer(T StringPath)
+	TStringViewNamePathProducer(TStringView<CharType> StringPath)
 	: CurrentStringView(StringPath)
 	{}
 
@@ -182,7 +182,7 @@ struct TStringViewNamePathProducer
 			return NAME_None;
 		}
 
-		typename T::SizeType FoundIndex = INDEX_NONE;
+		int32 FoundIndex = INDEX_NONE;
 		if (!FindLastSeparator(FoundIndex))
 		{
 			FName RetVal(CurrentStringView);
@@ -195,27 +195,27 @@ struct TStringViewNamePathProducer
 		return RetVal;
 	}
 private:
-	static inline bool IsPathIdSeparator(typename T::ElementType Char) { return Char == '.' || Char == ':'; }
+	static inline bool IsPathIdSeparator(CharType Char) { return Char == '.' || Char == ':'; }
 
-	bool FindLastSeparator(typename T::SizeType& OutIndex)
+	bool FindLastSeparator(int32& OutIndex)
 	{
-		if (const typename T::ElementType* Separator = Algo::FindLastByPredicate(CurrentStringView, IsPathIdSeparator))
+		if (const CharType* Separator = Algo::FindLastByPredicate(CurrentStringView, IsPathIdSeparator))
 		{
-			OutIndex = static_cast<typename T::SizeType>(Separator - CurrentStringView.GetData());
+			OutIndex = UE_PTRDIFF_TO_INT32(Separator - CurrentStringView.GetData());
 			return true;
 		}
 		OutIndex = INDEX_NONE;
 		return false;
 	}
 
-	T CurrentStringView;
+	TStringView<CharType> CurrentStringView;
 };
 
 FObjectPathId::FObjectPathId(FWideStringView StringPath)
 {
 	// @TODO: OBJPTR: Need to handle redirects.  FCoreRedirectObjectName could be used, but it doesn't fit
 	//		conveniently with the FName walk approach that is currently here.
-	TStringViewNamePathProducer<FWideStringView> NamePathProducer(StringPath);
+	TStringViewNamePathProducer<WIDECHAR> NamePathProducer(StringPath);
 
 	StoreObjectPathId(NamePathProducer, static_cast<uint64>(EPathId::FlagSimple), PathId);
 }
@@ -224,7 +224,7 @@ FObjectPathId::FObjectPathId(FAnsiStringView StringPath)
 {
 	// @TODO: OBJPTR: Need to handle redirects.  FCoreRedirectObjectName could be used, but it doesn't fit
 	//		conveniently with the FName walk approach that is currently here.
-	TStringViewNamePathProducer<FAnsiStringView> NamePathProducer(StringPath);
+	TStringViewNamePathProducer<ANSICHAR> NamePathProducer(StringPath);
 
 	StoreObjectPathId(NamePathProducer, static_cast<uint64>(EPathId::FlagSimple), PathId);
 }
