SharedBuffer: Improved the memory order used for the reference counts and flags

Materialization needs to use acquire/release to guarantee that one thread can see the Data and Size values materialized by another thread.

Any other operations on flags can be relaxed because there is no expectation that they have related dependencies.

#rnx

diff --git a/Engine/Source/Runtime/Core/Public/Memory/SharedBuffer.h b/Engine/Source/Runtime/Core/Public/Memory/SharedBuffer.h
--- a/Engine/Source/Runtime/Core/Public/Memory/SharedBuffer.h
+++ b/Engine/Source/Runtime/Core/Public/Memory/SharedBuffer.h
@@ -612,7 +612,7 @@ inline bool FBufferOwner::IsOwned() const
 
 inline void FBufferOwner::SetIsOwned()
 {
-	ReferenceCountsAndFlags.fetch_or(SetFlags(EBufferOwnerFlags::Owned));
+	ReferenceCountsAndFlags.fetch_or(SetFlags(EBufferOwnerFlags::Owned), std::memory_order_relaxed);
 }
 
 inline bool FBufferOwner::IsImmutable() const
@@ -622,7 +622,7 @@ inline bool FBufferOwner::IsImmutable() const
 
 inline void FBufferOwner::SetIsImmutable()
 {
-	ReferenceCountsAndFlags.fetch_or(SetFlags(EBufferOwnerFlags::Immutable));
+	ReferenceCountsAndFlags.fetch_or(SetFlags(EBufferOwnerFlags::Immutable), std::memory_order_relaxed);
 }
 
 inline void FBufferOwner::Materialize()
@@ -636,12 +636,12 @@ inline void FBufferOwner::Materialize()
 
 inline bool FBufferOwner::IsMaterialized() const
 {
-	return EnumHasAnyFlags(GetFlags(ReferenceCountsAndFlags.load(std::memory_order_relaxed)), EBufferOwnerFlags::Materialized);
+	return EnumHasAnyFlags(GetFlags(ReferenceCountsAndFlags.load(std::memory_order_acquire)), EBufferOwnerFlags::Materialized);
 }
 
 inline void FBufferOwner::SetIsMaterialized()
 {
-	ReferenceCountsAndFlags.fetch_or(SetFlags(EBufferOwnerFlags::Materialized));
+	ReferenceCountsAndFlags.fetch_or(SetFlags(EBufferOwnerFlags::Materialized), std::memory_order_release);
 }
 
 inline uint32 FBufferOwner::GetTotalRefCount() const
