CompressedBuffer: Validate the size of the compressed buffer when validating the header

The size check was removed from the decoders during the addition of FCompressedBufferReader, but is needed to avoid asserts when attempting to decompress a buffer that is too small.

#rnx

diff --git a/Engine/Source/Runtime/Core/Private/Compression/CompressedBuffer.cpp b/Engine/Source/Runtime/Core/Private/Compression/CompressedBuffer.cpp
--- a/Engine/Source/Runtime/Core/Private/Compression/CompressedBuffer.cpp
+++ b/Engine/Source/Runtime/Core/Private/Compression/CompressedBuffer.cpp
@@ -575,10 +575,11 @@ inline FCompositeBuffer ValidBufferOrEmpty(BufferType&& CompressedData)
 
 bool FHeader::IsValid(const FCompositeBuffer& CompressedData)
 {
-	if (sizeof(FHeader) <= CompressedData.GetSize())
+	const uint64 CompressedDataSize = CompressedData.GetSize();
+	if (sizeof(FHeader) <= CompressedDataSize)
 	{
 		const FHeader Header = Read(CompressedData);
-		if (Header.Magic == FHeader::ExpectedMagic)
+		if (Header.Magic == FHeader::ExpectedMagic && Header.TotalCompressedSize <= CompressedDataSize)
 		{
 			if (const FDecoder* const Decoder = GetDecoder(Header.Method))
 			{
diff --git a/Engine/Source/Runtime/Core/Private/Tests/Compression/CompressedBufferTest.cpp b/Engine/Source/Runtime/Core/Private/Tests/Compression/CompressedBufferTest.cpp
--- a/Engine/Source/Runtime/Core/Private/Tests/Compression/CompressedBufferTest.cpp
+++ b/Engine/Source/Runtime/Core/Private/Tests/Compression/CompressedBufferTest.cpp
@@ -209,6 +209,15 @@ bool FCompressedBufferDecompressTest::RunTest(const FString& Parameters)
 			const TConstArrayView<uint64> Values = CastToArrayView(Uncompressed);
 			ValidateData(Values, ExpectedValues, OffsetCount);
 		}
+
+		// Short Buffer
+		{
+			const FCompressedBuffer CompressedShort =
+				FCompressedBuffer::FromCompressed(Compressed.GetCompressed().Mid(0, Compressed.GetCompressedSize() - 128));
+			Reader.SetSource(CompressedShort);
+			const FSharedBuffer Uncompressed = Reader.Decompress();
+			TestTrue(TEXT("FCompressedBufferReader::Decompress(Oodle, Short)"), Uncompressed.IsNull());
+		}
 	}
 
 	// Only one block
@@ -267,6 +276,15 @@ bool FCompressedBufferDecompressTest::RunTest(const FString& Parameters)
 			const TConstArrayView<uint64> Values = CastToArrayView(Uncompressed);
 			ValidateData(Values, ExpectedValues, OffsetCount);
 		}
+
+		// Short Buffer
+		{
+			const FCompressedBuffer CompressedShort =
+				FCompressedBuffer::FromCompressed(Compressed.GetCompressed().Mid(0, Compressed.GetCompressedSize() - 128));
+			Reader.SetSource(CompressedShort);
+			const FSharedBuffer Uncompressed = Reader.Decompress();
+			TestTrue(TEXT("FCompressedBufferReader::Decompress(None, Short)"), Uncompressed.IsNull());
+		}
 	}
 
 	return true;
