Backout Changelist 18400446 as it seems to be the cause significant crashing on Switch. Part of the same backout of 18400230 that was previously submitted in CL 18458103
#test none

diff --git a/Engine/Source/Runtime/Core/Public/Misc/GenericGrowableAllocator.h b/Engine/Source/Runtime/Core/Public/Misc/GenericGrowableAllocator.h
--- a/Engine/Source/Runtime/Core/Public/Misc/GenericGrowableAllocator.h
+++ b/Engine/Source/Runtime/Core/Public/Misc/GenericGrowableAllocator.h
@@ -9,7 +9,6 @@
 #include "CoreMinimal.h"
 #include "HAL/LowLevelMemTracker.h"
 #include "Misc/ScopeLock.h"
-#include "ProfilingDebugging/MemoryTrace.h"
 #include "Stats/Stats.h"
 
 //////////////
@@ -21,8 +20,7 @@
 // - Can maybe just make virtuals in FGrowableAllocationBase ??
 /////////////
 
-#define ALLOCATION_HISTOGRAM			(!UE_BUILD_SHIPPING && !UE_BUILD_TEST) && 1
-#define ALLOCATION_HISTOGRAM_DETAILED	(!UE_BUILD_SHIPPING && !UE_BUILD_TEST) && 0
+#define ALLOCATION_HISTOGRAM
 
 struct FGrowableAllocationBase
 {
@@ -227,7 +225,7 @@ public:
 		FFreeEntry* Prev = NULL;
 		FFreeEntry* FindEntry = NULL;
 
-		for (FFreeEntry *Entry = FreeList; Entry; Entry = Entry->Next)
+		for (FFreeEntry* Entry = FreeList; Entry; Entry = Entry->Next)
 		{
 			if (Entry->CanFit(AlignedSize, Alignment))
 			{
@@ -264,7 +262,6 @@ public:
 			Allocation->OwnerType = OwnerType;
 #endif
 			LLM(FLowLevelMemTracker::Get().OnLowLevelAlloc(ELLMTracker::Default, GetAddressForTracking(Offset), Size));
-			MemoryTrace_Alloc(uint64(GetAddressForTracking(Offset)), Size, Alignment, EMemoryTraceRootHeap::SystemMemory);
 
 			// let the implementation fill in any more
 			InitializeAllocationStruct(Allocation);
@@ -287,7 +284,6 @@ public:
 		uint32 Offset = Memory->Offset;
 
 		LLM(FLowLevelMemTracker::Get().OnLowLevelFree(ELLMTracker::Default, GetAddressForTracking(Offset)));
-		MemoryTrace_Free(uint64(GetAddressForTracking(Offset)), EMemoryTraceRootHeap::SystemMemory);
 
 		// we are now done with the Allocation object
 		DestroyAllocationStruct(Memory);
@@ -656,7 +652,7 @@ public:
 				TypeAllocInfo.Counts.Frees++;
 #if ALLOCATION_HISTOGRAM_DETAILED
 				TypeAllocInfo.AllocationHistogram[Memory->Size].Frees++;
-#endif
+#endif 
 #endif // ALLOCATION_HISTOGRAM
 
 #if STATS
