StringView

- Deprecated TStringView::SizeType in favor of int32. We do not expect strings over 2 GiB and do not support those sizes with TStringBuilder or FString.
- Added IntCastChecked to the constructor for safety when constructing from larger range types.
- Switched to STL type traits to improve consistency and, theoretically, compile times.
- Switched the implementation of Find, FindChar, FindLastChar to use equivalent functions in UE::String.
- Removed TIsConvertibleToStringView and TCompatibleStringViewType by relying on deduction of the view type, which we can do with C++17.
- Enable the use of TStringView with any char type, not only the three that had hard-coded support.

diff --git a/Engine/Source/Runtime/Core/Private/Tests/Containers/StringViewTest.cpp b/Engine/Source/Runtime/Core/Private/Tests/Containers/StringViewTest.cpp
--- a/Engine/Source/Runtime/Core/Private/Tests/Containers/StringViewTest.cpp
+++ b/Engine/Source/Runtime/Core/Private/Tests/Containers/StringViewTest.cpp
@@ -20,13 +20,19 @@ static_assert(TIsContiguousContainer<FStringView>::Value, "FStringView must be a
 static_assert(TIsContiguousContainer<FAnsiStringView>::Value, "FAnsiStringView must be a contiguous container.");
 static_assert(TIsContiguousContainer<FWideStringView>::Value, "FWideStringView must be a contiguous container.");
 
-static_assert(StringViewPrivate::TIsConvertibleToStringView<FString>::Value, "FString must be convertible to FStringView.");
-static_assert(StringViewPrivate::TIsConvertibleToStringView<FAnsiStringBuilderBase>::Value, "FAnsiStringBuilderBase must be convertible to FAnsiStringView.");
-static_assert(StringViewPrivate::TIsConvertibleToStringView<FWideStringBuilderBase>::Value, "FWideStringBuilderBase must be convertible to FWideStringView.");
+namespace UE::String::Private::TestArgumentDependentLookup
+{
+
+struct FTestType
+{
+	using ElementType = TCHAR;
+};
+
+const TCHAR* GetData(const FTestType&) { return TEXT("ABC"); }
+int32 GetNum(const FTestType&) { return 3; }
+} // UE::String::Private::TestArgumentDependentLookup
 
-static_assert(TIsSame<FStringView, typename StringViewPrivate::TCompatibleStringViewType<FString>::Type>::Value, "FString must be convertible to FStringView.");
-static_assert(TIsSame<FAnsiStringView, typename StringViewPrivate::TCompatibleStringViewType<FAnsiStringBuilderBase>::Type>::Value, "FAnsiStringBuilderBase must be convertible to FAnsiStringView.");
-static_assert(TIsSame<FWideStringView, typename StringViewPrivate::TCompatibleStringViewType<FWideStringBuilderBase>::Type>::Value, "FWideStringBuilderBase must be convertible to FWideStringView.");
+template <> struct TIsContiguousContainer<UE::String::Private::TestArgumentDependentLookup::FTestType> { static constexpr bool Value = true; };
 
 #define TEST_NAME_ROOT "System.Core.StringView"
 constexpr const uint32 TestFlags = EAutomationTestFlags::ApplicationContextMask | EAutomationTestFlags::SmokeFilter;
@@ -125,6 +131,13 @@ bool FStringViewTestCtor::RunTest(const FString& Parameters)
 		TStringView ViewString(String);
 	}
 
+	// Verify that argument-dependent lookup is working for GetData and GetNum
+	{
+		UE::String::Private::TestArgumentDependentLookup::FTestType Test;
+		FStringView View(Test);
+		TestTrue(TEXT("StringView ADL"), View.Equals(TEXTVIEW("ABC")));
+	}
+
 	return true;
 }
 
diff --git a/Engine/Source/Runtime/Core/Public/Containers/StringView.h b/Engine/Source/Runtime/Core/Public/Containers/StringView.h
--- a/Engine/Source/Runtime/Core/Public/Containers/StringView.h
+++ b/Engine/Source/Runtime/Core/Public/Containers/StringView.h
@@ -9,50 +9,22 @@
 #include "Math/UnrealMathUtility.h"
 #include "Misc/Crc.h"
 #include "Misc/CString.h"
-#include "Templates/AndOrNot.h"
-#include "Templates/ChooseClass.h"
-#include "Templates/EnableIf.h"
+#include "String/Find.h"
 #include "Templates/UnrealTemplate.h"
 #include "Traits/ElementType.h"
+#include "Traits/IsCharType.h"
 #include "Traits/IsContiguousContainer.h"
 #include <type_traits>
 
-namespace StringViewPrivate
+namespace UE::Core::Private
 {
-	/** Wrapper to allow GetData to resolve from within a scope with another overload. */
+	/** Allow GetData to called unqualified from a scope with its own overload of GetData. */
 	template <typename... ArgTypes>
-	constexpr inline auto WrapGetData(ArgTypes&&... Args) -> decltype(GetData(Forward<ArgTypes>(Args)...))
+	constexpr inline auto StringViewGetData(ArgTypes&&... Args) -> decltype(GetData(Forward<ArgTypes>(Args)...))
 	{
 		return GetData(Forward<ArgTypes>(Args)...);
 	}
-
-	template <typename From, typename To>
-	struct TIsConvertibleFromTo
-	{
-		static constexpr bool Value = __is_convertible_to(From, To);
-	};
-
-	template <typename T>
-	struct TIsConvertibleToStringView
-	{
-		static constexpr bool Value = TOr<
-			TIsConvertibleFromTo<T, FAnsiStringView>,
-			TIsConvertibleFromTo<T, FWideStringView>,
-			TIsConvertibleFromTo<T, FUtf8StringView>
-			>::Value;
-	};
-
-	template <typename T>
-	struct TCompatibleStringViewType
-	{
-		struct NotCompatible;
-		using Type =
-			typename TChooseClass<TIsConvertibleFromTo<T, FAnsiStringView>::Value, FAnsiStringView,
-			typename TChooseClass<TIsConvertibleFromTo<T, FWideStringView>::Value, FWideStringView,
-			typename TChooseClass<TIsConvertibleFromTo<T, FUtf8StringView>::Value, FUtf8StringView,
-			NotCompatible>::Result>::Result>::Result;
-	};
-}
+} // UE::Core::Private
 
 /**
  * A string view is a non-owning view of a range of characters.
@@ -105,7 +77,7 @@ class TStringView
 {
 public:
 	using ElementType = CharType;
-	using SizeType = int32;
+	using SizeType UE_DEPRECATED(5.0, "Use int32 instead.") = int32;
 	using ViewType = TStringView<CharType>;
 
 public:
@@ -120,7 +92,7 @@ public:
 	}
 
 	/** Construct a view of InSize characters beginning at InData. */
-	constexpr inline TStringView(const CharType* InData, SizeType InSize)
+	constexpr inline TStringView(const CharType* InData, int32 InSize)
 		: DataPtr(InData)
 		, Size(InSize)
 	{
@@ -138,7 +110,7 @@ public:
 	/** Construct a view of InSize characters beginning at InData. */
 	template <typename OtherCharType,
 		std::enable_if_t<FPlatformString::IsCharEncodingCompatibleWith<OtherCharType, CharType>()>* = nullptr>
-	constexpr inline TStringView(const OtherCharType* InData, SizeType InSize)
+	constexpr inline TStringView(const OtherCharType* InData, int32 InSize)
 		: DataPtr((const CharType*)InData)
 		, Size(InSize)
 	{
@@ -154,33 +126,33 @@ public:
 			!std::is_same_v<CharRangeType, ViewType>
 		>* = nullptr>
 	constexpr inline TStringView(const CharRangeType& InRange)
-		: DataPtr((const CharType*)StringViewPrivate::WrapGetData(InRange))
-		, Size(static_cast<SizeType>(GetNum(InRange)))
+		: DataPtr((const CharType*)UE::Core::Private::StringViewGetData(InRange))
+		, Size(IntCastChecked<int32>(GetNum(InRange)))
 	{
 	}
 
 	/** Access the character at the given index in the view. */
-	inline const CharType& operator[](SizeType Index) const;
+	inline const CharType& operator[](int32 Index) const;
 
 	/** Returns a pointer to the start of the view. This is NOT guaranteed to be null-terminated! */
-	UE_NODISCARD constexpr inline const CharType* GetData() const { return DataPtr; }
+	[[nodiscard]] constexpr inline const CharType* GetData() const { return DataPtr; }
 
 	// Capacity
 
 	/** Returns the length of the string view. */
-	UE_NODISCARD constexpr inline SizeType Len() const { return Size; }
+	[[nodiscard]] constexpr inline int32 Len() const { return Size; }
 
 	/** Returns whether the string view is empty. */
-	UE_NODISCARD constexpr inline bool IsEmpty() const { return Size == 0; }
+	[[nodiscard]] constexpr inline bool IsEmpty() const { return Size == 0; }
 
 	// Modifiers
 
 	/** Modifies the view to remove the given number of characters from the start. */
-	inline void		RemovePrefix(SizeType CharCount)	{ DataPtr += CharCount; Size -= CharCount; }
+	inline void RemovePrefix(int32 CharCount) { DataPtr += CharCount; Size -= CharCount; }
 	/** Modifies the view to remove the given number of characters from the end. */
-	inline void		RemoveSuffix(SizeType CharCount)	{ Size -= CharCount; }
+	inline void RemoveSuffix(int32 CharCount) { Size -= CharCount; }
 	/** Resets to an empty view */
-	inline void		Reset()								{ DataPtr = nullptr; Size = 0; }
+	inline void Reset() { DataPtr = nullptr; Size = 0; }
 
 	// Operations
 
@@ -193,38 +165,38 @@ public:
 	 *
 	 * @return The number of characters written to the destination buffer.
 	 */
-	inline SizeType CopyString(CharType* Dest, SizeType CharCount, SizeType Position = 0) const;
+	inline int32 CopyString(CharType* Dest, int32 CharCount, int32 Position = 0) const;
 
 	/** Alias for Mid. */
-	UE_NODISCARD inline ViewType SubStr(SizeType Position, SizeType CharCount) const { return Mid(Position, CharCount); }
+	[[nodiscard]] inline ViewType SubStr(int32 Position, int32 CharCount) const { return Mid(Position, CharCount); }
 
 	/** Returns the left-most part of the view by taking the given number of characters from the left. */
-	UE_NODISCARD inline ViewType Left(SizeType CharCount) const;
+	[[nodiscard]] inline ViewType Left(int32 CharCount) const;
 	/** Returns the left-most part of the view by chopping the given number of characters from the right. */
-	UE_NODISCARD inline ViewType LeftChop(SizeType CharCount) const;
+	[[nodiscard]] inline ViewType LeftChop(int32 CharCount) const;
 	/** Returns the right-most part of the view by taking the given number of characters from the right. */
-	UE_NODISCARD inline ViewType Right(SizeType CharCount) const;
+	[[nodiscard]] inline ViewType Right(int32 CharCount) const;
 	/** Returns the right-most part of the view by chopping the given number of characters from the left. */
-	UE_NODISCARD inline ViewType RightChop(SizeType CharCount) const;
+	[[nodiscard]] inline ViewType RightChop(int32 CharCount) const;
 	/** Returns the middle part of the view by taking up to the given number of characters from the given position. */
-	UE_NODISCARD inline ViewType Mid(SizeType Position, SizeType CharCount = TNumericLimits<SizeType>::Max()) const;
+	[[nodiscard]] inline ViewType Mid(int32 Position, int32 CharCount = MAX_int32) const;
 	/** Returns the middle part of the view between any whitespace at the start and end. */
-	UE_NODISCARD inline ViewType TrimStartAndEnd() const;
+	[[nodiscard]] inline ViewType TrimStartAndEnd() const;
 	/** Returns the right part of the view after any whitespace at the start. */
-	UE_NODISCARD inline ViewType TrimStart() const;
+	[[nodiscard]] inline ViewType TrimStart() const;
 	/** Returns the left part of the view before any whitespace at the end. */
-	UE_NODISCARD inline ViewType TrimEnd() const;
+	[[nodiscard]] inline ViewType TrimEnd() const;
 
 	/** Modifies the view to be the given number of characters from the left. */
-	inline void LeftInline(SizeType CharCount) { *this = Left(CharCount); }
+	inline void LeftInline(int32 CharCount) { *this = Left(CharCount); }
 	/** Modifies the view by chopping the given number of characters from the right. */
-	inline void LeftChopInline(SizeType CharCount) { *this = LeftChop(CharCount); }
+	inline void LeftChopInline(int32 CharCount) { *this = LeftChop(CharCount); }
 	/** Modifies the view to be the given number of characters from the right. */
-	inline void RightInline(SizeType CharCount) { *this = Right(CharCount); }
+	inline void RightInline(int32 CharCount) { *this = Right(CharCount); }
 	/** Modifies the view by chopping the given number of characters from the left. */
-	inline void RightChopInline(SizeType CharCount) { *this = RightChop(CharCount); }
+	inline void RightChopInline(int32 CharCount) { *this = RightChop(CharCount); }
 	/** Modifies the view to be the middle part by taking up to the given number of characters from the given position. */
-	inline void MidInline(SizeType Position, SizeType CharCount = TNumericLimits<SizeType>::Max()) { *this = Mid(Position, CharCount); }
+	inline void MidInline(int32 Position, int32 CharCount = MAX_int32) { *this = Mid(Position, CharCount); }
 	/** Modifies the view to be the middle part between any whitespace at the start and end. */
 	inline void TrimStartAndEndInline() { *this = TrimStartAndEnd(); }
 	/** Modifies the view to be the right part after any whitespace at the start. */
@@ -235,92 +207,111 @@ public:
 	// Comparison
 
 	/**
-	 * Check whether this view is lexicographically equivalent to another view.
+	 * Check whether this view is equivalent to a character range.
 	 *
-	 * @param SearchCase Whether the comparison should ignore case.
+	 * @param Other        A character range that is comparable with the character type of this view.
+	 * @param SearchCase   Whether the comparison should ignore case.
 	 */
-	template <typename OtherType, typename TEnableIf<StringViewPrivate::TIsConvertibleToStringView<OtherType>::Value>::Type* = nullptr>
-	UE_NODISCARD inline bool Equals(OtherType&& Other, ESearchCase::Type SearchCase = ESearchCase::CaseSensitive) const;
+	template <typename OtherRangeType, decltype(MakeStringView(DeclVal<OtherRangeType>()))* = nullptr>
+	[[nodiscard]] inline bool Equals(OtherRangeType&& Other, ESearchCase::Type SearchCase = ESearchCase::CaseSensitive) const
+	{
+		const auto OtherView = MakeStringView(Forward<OtherRangeType>(Other));
+		return Len() == OtherView.Len() && Compare(OtherView, SearchCase) == 0;
+	}
 
 	/**
-	 * Check whether this view is lexicographically equivalent to another view.
+	 * Check whether this view is equivalent to a string view.
 	 *
-	 * @param SearchCase Whether the comparison should ignore case.
+	 * @param Other        A string that is comparable with the character type of this view.
+	 * @param SearchCase   Whether the comparison should ignore case.
 	 */
-	template <typename OtherCharType>
-	UE_NODISCARD inline bool Equals(const OtherCharType* Other, ESearchCase::Type SearchCase = ESearchCase::CaseSensitive) const;
+	template <typename OtherCharType,
+		std::enable_if_t<TIsCharType<OtherCharType>::Value>* = nullptr>
+	[[nodiscard]] inline bool Equals(const OtherCharType* Other, ESearchCase::Type SearchCase = ESearchCase::CaseSensitive) const;
 
 	/**
-	 * Compare this view lexicographically with another view.
-	 *
-	 * @param SearchCase Whether the comparison should ignore case.
+	 * Compare this view with a character range.
 	 *
+	 * @param Other        A character range that is comparable with the character type of this view.
+	 * @param SearchCase   Whether the comparison should ignore case.
 	 * @return 0 is equal, negative if this view is less, positive if this view is greater.
 	 */
-	template <typename OtherType, typename TEnableIf<StringViewPrivate::TIsConvertibleToStringView<OtherType>::Value>::Type* = nullptr>
-	UE_NODISCARD inline int32 Compare(OtherType&& Other, ESearchCase::Type SearchCase = ESearchCase::CaseSensitive) const;
+	template <typename OtherRangeType, decltype(MakeStringView(DeclVal<OtherRangeType>()))* = nullptr>
+	[[nodiscard]] inline int32 Compare(OtherRangeType&& Other, ESearchCase::Type SearchCase = ESearchCase::CaseSensitive) const
+	{
+		return Compare(MakeStringView(Forward<OtherRangeType>(Other)), SearchCase);
+	}
 
 	/**
-	 * Compare this view lexicographically with another view.
+	 * Compare this view with a string view.
 	 *
-	 * @param SearchCase Whether the comparison should ignore case.
+	 * @param Other        A string view that is comparable with the character type of this view.
+	 * @param SearchCase   Whether the comparison should ignore case.
+	 * @return 0 is equal, negative if this view is less, positive if this view is greater.
+	 */
+	template <typename OtherCharType,
+		std::enable_if_t<TIsCharType<OtherCharType>::Value>* = nullptr>
+		[[nodiscard]] inline int32 Compare(TStringView<OtherCharType> Other, ESearchCase::Type SearchCase = ESearchCase::CaseSensitive) const;
+
+	/**
+	 * Compare this view with a null-terminated string.
 	 *
+	 * @param Other        A null-terminated string that is comparable with the character type of this view.
+	 * @param SearchCase   Whether the comparison should ignore case.
 	 * @return 0 is equal, negative if this view is less, positive if this view is greater.
 	 */
 	template <typename OtherCharType>
-	UE_NODISCARD inline int32 Compare(const OtherCharType* Other, ESearchCase::Type SearchCase = ESearchCase::CaseSensitive) const;
+	[[nodiscard]] inline int32 Compare(const OtherCharType* Other, ESearchCase::Type SearchCase = ESearchCase::CaseSensitive) const;
 
 	/** Returns whether this view starts with the prefix character compared case-sensitively. */
-	UE_NODISCARD inline bool StartsWith(CharType Prefix) const { return Size >= 1 && DataPtr[0] == Prefix; }
+	[[nodiscard]] inline bool StartsWith(CharType Prefix) const { return Size >= 1 && DataPtr[0] == Prefix; }
 	/** Returns whether this view starts with the prefix with optional case sensitivity. */
-	UE_NODISCARD inline bool StartsWith(ViewType Prefix, ESearchCase::Type SearchCase = ESearchCase::IgnoreCase) const;
+	[[nodiscard]] inline bool StartsWith(ViewType Prefix, ESearchCase::Type SearchCase = ESearchCase::IgnoreCase) const;
 
 	/** Returns whether this view ends with the suffix character compared case-sensitively. */
-	UE_NODISCARD inline bool EndsWith(CharType Suffix) const { return Size >= 1 && DataPtr[Size-1] == Suffix; }
+	[[nodiscard]] inline bool EndsWith(CharType Suffix) const { return Size >= 1 && DataPtr[Size-1] == Suffix; }
 	/** Returns whether this view ends with the suffix with optional case sensitivity. */
-	UE_NODISCARD inline bool EndsWith(ViewType Suffix, ESearchCase::Type SearchCase = ESearchCase::IgnoreCase) const;
+	[[nodiscard]] inline bool EndsWith(ViewType Suffix, ESearchCase::Type SearchCase = ESearchCase::IgnoreCase) const;
 
 	// Searching/Finding
 
 	/**
-	 * Searches the view for a substring, and returns index into this view
-	 * of the first found instance, or INDEX_NONE if not found.
+	 * Search the view for the first occurrence of a search string.
 	 *
-	 * @param SubStr			The text to search for
-	 * @param StartPosition		The start character position to search from
+	 * @param Search          The string to search for. Comparison is lexicographic.
+	 * @param StartPosition   The character position to start searching from.
+	 * @return The index of the first occurrence of the search string if found, otherwise INDEX_NONE.
 	 */
-	inline int32 Find(ViewType SubStr, SizeType StartPosition = 0) const;
+	[[nodiscard]] inline int32 Find(ViewType Search, int32 StartPosition = 0) const;
 
 	/**
 	 * Returns whether this view contains the specified substring.
 	 *
-	 * @param SubStr			Text to search for
-	 * @return					Returns whether the view contains the substring
-	 **/
-	UE_NODISCARD FORCEINLINE bool Contains(ViewType SubStr) const
+	 * @param Search   Text to search for
+	 * @return True if the view contains the search string, otherwise false.
+	 */
+	[[nodiscard]] inline bool Contains(ViewType Search) const
 	{
-		return Find(SubStr) != INDEX_NONE;
+		return Find(Search) != INDEX_NONE;
 	}
 
 	/**
 	 * Search the view for the first occurrence of a character.
 	 *
-	 * @param InChar The character to search for. Comparison is lexicographic.
-	 * @param OutIndex [out] The position at which the character was found, or INDEX_NONE if not found.
-	 *
-	 * @return Whether the character was found in the view.
+	 * @param Search           The character to search for. Comparison is lexicographic.
+	 * @param OutIndex [out]   The position at which the character was found, or INDEX_NONE if not found.
+	 * @return True if the character was found in the view, otherwise false.
 	 */
-	inline bool FindChar(CharType InChar, SizeType& OutIndex) const;
+	inline bool FindChar(CharType Search, int32& OutIndex) const;
 
 	/**
 	 * Search the view for the last occurrence of a character.
 	 *
-	 * @param InChar The character to search for. Comparison is lexicographic.
-	 * @param OutIndex [out] The position at which the character was found, or INDEX_NONE if not found.
-	 *
-	 * @return Whether the character was found in the view.
+	 * @param Search           The character to search for. Comparison is lexicographic.
+	 * @param OutIndex [out]   The position at which the character was found, or INDEX_NONE if not found.
+	 * @return True if the character was found in the view, otherwise false.
 	 */
-	inline bool FindLastChar(CharType InChar, SizeType& OutIndex) const;
+	inline bool FindLastChar(CharType Search, int32& OutIndex) const;
 
 public:
 	/**
@@ -332,7 +323,7 @@ public:
 
 protected:
 	const CharType* DataPtr = nullptr;
-	SizeType Size = 0;
+	int32 Size = 0;
 };
 
 template <typename CharRangeType>
@@ -408,54 +399,54 @@ FORCEINLINE uint32 GetTypeHash(TStringView<CharType> View)
 }
 
 template <typename CharType>
-inline const CharType& TStringView<CharType>::operator[](SizeType Index) const
+inline const CharType& TStringView<CharType>::operator[](int32 Index) const
 {
 	checkf(Index >= 0 && Index < Size, TEXT("Index out of bounds on StringView: index %i on a view with a length of %i"), Index, Size);
 	return DataPtr[Index];
 }
 
 template <typename CharType>
-inline typename TStringView<CharType>::SizeType TStringView<CharType>::CopyString(CharType* Dest, SizeType CharCount, SizeType Position) const
+inline int32 TStringView<CharType>::CopyString(CharType* Dest, int32 CharCount, int32 Position) const
 {
-	const  SizeType CopyCount = FMath::Min(Size - Position, CharCount);
+	const int32 CopyCount = FMath::Min(Size - Position, CharCount);
 	FMemory::Memcpy(Dest, DataPtr + Position, CopyCount * sizeof(CharType));
 	return CopyCount;
 }
 
 template <typename CharType>
-inline TStringView<CharType> TStringView<CharType>::Left(SizeType CharCount) const
+inline TStringView<CharType> TStringView<CharType>::Left(int32 CharCount) const
 {
 	return ViewType(DataPtr, FMath::Clamp(CharCount, 0, Size));
 }
 
 template <typename CharType>
-inline TStringView<CharType> TStringView<CharType>::LeftChop(SizeType CharCount) const
+inline TStringView<CharType> TStringView<CharType>::LeftChop(int32 CharCount) const
 {
 	return ViewType(DataPtr, FMath::Clamp(Size - CharCount, 0, Size));
 }
 
 template <typename CharType>
-inline TStringView<CharType> TStringView<CharType>::Right(SizeType CharCount) const
+inline TStringView<CharType> TStringView<CharType>::Right(int32 CharCount) const
 {
-	const SizeType OutLen = FMath::Clamp(CharCount, 0, Size);
+	const int32 OutLen = FMath::Clamp(CharCount, 0, Size);
 	return ViewType(DataPtr + Size - OutLen, OutLen);
 }
 
 template <typename CharType>
-inline TStringView<CharType> TStringView<CharType>::RightChop(SizeType CharCount) const
+inline TStringView<CharType> TStringView<CharType>::RightChop(int32 CharCount) const
 {
-	const SizeType OutLen = FMath::Clamp(Size - CharCount, 0, Size);
+	const int32 OutLen = FMath::Clamp(Size - CharCount, 0, Size);
 	return ViewType(DataPtr + Size - OutLen, OutLen);
 }
 
 template <typename CharType>
-inline TStringView<CharType> TStringView<CharType>::Mid(SizeType Position, SizeType CharCount) const
+inline TStringView<CharType> TStringView<CharType>::Mid(int32 Position, int32 CharCount) const
 {
 	const CharType* CurrentStart  = GetData();
-	const SizeType  CurrentLength = Len();
+	const int32 CurrentLength = Len();
 
 	// Clamp minimum position at the start of the string, adjusting the length down if necessary
-	const SizeType NegativePositionOffset = (Position < 0) ? Position : 0;
+	const int32 NegativePositionOffset = (Position < 0) ? Position : 0;
 	CharCount += NegativePositionOffset;
 	Position  -= NegativePositionOffset;
 
@@ -478,10 +469,10 @@ inline TStringView<CharType> TStringView<CharType>::TrimStartAndEnd() const
 template <typename CharType>
 inline TStringView<CharType> TStringView<CharType>::TrimStart() const
 {
-	SizeType SpaceCount = 0;
-	for (ElementType Char : *this)
+	int32 SpaceCount = 0;
+	for (CharType Char : *this)
 	{
-		if (!TChar<ElementType>::IsWhitespace(Char))
+		if (!TChar<CharType>::IsWhitespace(Char))
 		{
 			break;
 		}
@@ -493,8 +484,8 @@ inline TStringView<CharType> TStringView<CharType>::TrimStart() const
 template <typename CharType>
 inline TStringView<CharType> TStringView<CharType>::TrimEnd() const
 {
-	SizeType NewSize = Size;
-	while (NewSize && TChar<ElementType>::IsWhitespace(DataPtr[NewSize - 1]))
+	int32 NewSize = Size;
+	while (NewSize && TChar<CharType>::IsWhitespace(DataPtr[NewSize - 1]))
 	{
 		--NewSize;
 	}
@@ -502,16 +493,8 @@ inline TStringView<CharType> TStringView<CharType>::TrimEnd() const
 }
 
 template <typename CharType>
-template <typename OtherType, typename TEnableIf<StringViewPrivate::TIsConvertibleToStringView<OtherType>::Value>::Type*>
-inline bool TStringView<CharType>::Equals(OtherType&& Other, ESearchCase::Type SearchCase) const
-{
-	using OtherViewType = typename StringViewPrivate::TCompatibleStringViewType<OtherType>::Type;
-	const OtherViewType OtherView(Forward<OtherType>(Other));
-	return Len() == OtherView.Len() && Compare(OtherView, SearchCase) == 0;
-}
-
-template <typename CharType>
-template <typename OtherCharType>
+template <typename OtherCharType,
+	std::enable_if_t<TIsCharType<OtherCharType>::Value>*>
 inline bool TStringView<CharType>::Equals(const OtherCharType* const Other, ESearchCase::Type SearchCase) const
 {
 	if (SearchCase == ESearchCase::CaseSensitive)
@@ -525,15 +508,13 @@ inline bool TStringView<CharType>::Equals(const OtherCharType* const Other, ESea
 }
 
 template <typename CharType>
-template <typename OtherType, typename TEnableIf<StringViewPrivate::TIsConvertibleToStringView<OtherType>::Value>::Type*>
-inline int32 TStringView<CharType>::Compare(OtherType&& Other, ESearchCase::Type SearchCase) const
+template <typename OtherCharType,
+	std::enable_if_t<TIsCharType<OtherCharType>::Value>*>
+inline int32 TStringView<CharType>::Compare(TStringView<OtherCharType> OtherView, ESearchCase::Type SearchCase) const
 {
-	using OtherViewType = typename StringViewPrivate::TCompatibleStringViewType<OtherType>::Type;
-	const OtherViewType OtherView(Forward<OtherType>(Other));
-
-	const SizeType SelfLen = Len();
-	const SizeType OtherLen = OtherView.Len();
-	const SizeType MinLen = FMath::Min(SelfLen, OtherLen);
+	const int32 SelfLen = Len();
+	const int32 OtherLen = OtherView.Len();
+	const int32 MinLen = FMath::Min(SelfLen, OtherLen);
 
 	int Result;
 	if (SearchCase == ESearchCase::CaseSensitive)
@@ -589,68 +570,25 @@ inline bool TStringView<CharType>::EndsWith(ViewType Suffix, ESearchCase::Type S
 }
 
 template <typename CharType>
-inline int32 TStringView<CharType>::Find(ViewType SubStr, SizeType StartPosition) const
+inline int32 TStringView<CharType>::Find(const ViewType Search, const int32 StartPosition) const
 {
-	SizeType MinSearchLen = StartPosition + SubStr.Len();
-	if (SubStr.IsEmpty() || Size < MinSearchLen)
-	{
-		return INDEX_NONE;
-	}
-
-	ElementType FirstSubChar = SubStr[0];
-	SizeType RestSubLen = SubStr.Len() - 1;
-	SizeType Maxi = Size - MinSearchLen;
-	for (SizeType i = StartPosition; i <= Maxi; ++i)
-	{
-		if (DataPtr[i] == FirstSubChar)
-		{
-			if (RestSubLen == 0 || TCString<CharType>::Strncmp(DataPtr + i + 1, SubStr.DataPtr + 1, RestSubLen) == 0)
-			{
-				return i;
-			}
-		}
-	}
-
-	return INDEX_NONE;
+	const int32 Index = UE::String::FindFirst(RightChop(StartPosition), Search);
+	return Index == INDEX_NONE ? INDEX_NONE : Index + StartPosition;
 }
 
 template <typename CharType>
-inline bool TStringView<CharType>::FindChar(const ElementType InChar, SizeType& OutIndex) const
+inline bool TStringView<CharType>::FindChar(const CharType Search, int32& OutIndex) const
 {
-	for (SizeType i = 0; i < Size; ++i)
-	{
-		if (DataPtr[i] == InChar)
-		{
-			OutIndex = i;
-			return true;
-		}
-	}
-
-	OutIndex = INDEX_NONE;
-	return false;
+	OutIndex = UE::String::FindFirstChar(*this, Search);
+	return OutIndex != INDEX_NONE;
 }
 
 
 template <typename CharType>
-inline bool TStringView<CharType>::FindLastChar(const ElementType InChar, SizeType& OutIndex) const
+inline bool TStringView<CharType>::FindLastChar(const CharType Search, int32& OutIndex) const
 {
-	if (Size == 0)
-	{
-		OutIndex = INDEX_NONE;
-		return false;
-	}
-
-	for (SizeType i = Size - 1; i >= 0; --i)
-	{
-		if (DataPtr[i] == InChar)
-		{
-			OutIndex = i;
-			return true;
-		}
-	}
-
-	OutIndex = INDEX_NONE;
-	return false;
+	OutIndex = UE::String::FindLastChar(*this, Search);
+	return OutIndex != INDEX_NONE;
 }
 
 // Case-insensitive string view comparison operators
