Add ability to use TArrayView with Append in TSet.

diff --git a/Engine/Source/Runtime/Core/Public/Containers/Set.h b/Engine/Source/Runtime/Core/Public/Containers/Set.h
--- a/Engine/Source/Runtime/Core/Public/Containers/Set.h
+++ b/Engine/Source/Runtime/Core/Public/Containers/Set.h
@@ -756,6 +756,16 @@ public:
 		return ElementId;
 	}
 
+	template<typename ViewSizeType>
+	void Append(const TArrayView<ElementType, ViewSizeType>& InElements)
+	{
+		Reserve(Elements.Num() + InElements.Num());
+		for (const ElementType& Element : InElements)
+		{
+			Add(Element);
+		}
+	}
+
 	template<typename ArrayAllocator>
 	void Append(const TArray<ElementType, ArrayAllocator>& InElements)
 	{
