Add time based cleanup for Hazard pointers to avoid over accumulation of buffers on higher core count
  - Saves 2MB on 3 cores during replay run
  - Saves 20MB+ on 16 cores during replay run
  - Saves 500MB+ on 64 cores since those values never stabilized

Contributors to this CL: Arne.Schober, Martin.Sevigny

diff --git a/Engine/Source/Runtime/Core/Private/Experimental/Containers/HazardPointer.cpp b/Engine/Source/Runtime/Core/Private/Experimental/Containers/HazardPointer.cpp
--- a/Engine/Source/Runtime/Core/Private/Experimental/Containers/HazardPointer.cpp
+++ b/Engine/Source/Runtime/Core/Private/Experimental/Containers/HazardPointer.cpp
@@ -1,16 +1,13 @@
 // Copyright Epic Games, Inc. All Rights Reserved.
 
 #include "Experimental/Containers/HazardPointer.h"
+#include "Misc/App.h"
 #include "Misc/ScopeLock.h"
 
 FHazardPointerCollection::~FHazardPointerCollection()
 {
-	for (TArray<HazardPointer_Impl::FHazardDeleter>* Deleters : AllTlsVariables)
+	for (FTlsData* Deleters : AllTlsVariables)
 	{
-		for (HazardPointer_Impl::FHazardDeleter& Deleter : *Deleters)
-		{
-			Deleter.Delete();
-		}
 		delete Deleters;
 	}
 	FPlatformTLS::FreeTlsSlot(CollectablesTlsSlot);
@@ -51,24 +48,27 @@ template CORE_API FHazardPointerCollection::FHazardRecord* FHazardPointerCollect
 
 void FHazardPointerCollection::Delete(const HazardPointer_Impl::FHazardDeleter& Ptr, int32 CollectLimit)
 {
-	TArray<HazardPointer_Impl::FHazardDeleter>* Collectables = (TArray<HazardPointer_Impl::FHazardDeleter>*)FPlatformTLS::GetTlsValue(CollectablesTlsSlot);
+	FTlsData* Collectables = (FTlsData*)FPlatformTLS::GetTlsValue(CollectablesTlsSlot);
 	if (Collectables == nullptr)
 	{
 		FScopeLock Lock(&AllTlsVariablesCS);
-		Collectables = new TArray<HazardPointer_Impl::FHazardDeleter>();
+		Collectables = new FTlsData();
 		AllTlsVariables.Add(Collectables);
 		FPlatformTLS::SetTlsValue(CollectablesTlsSlot, Collectables);
 	}
-	checkSlow(!Collectables->Contains(Ptr));
+	checkSlow(!Collectables->ReclamationList.Contains(Ptr));
 
 	//add to be deleted pointer to thread local List
-	Collectables->Add(Ptr);
+	Collectables->ReclamationList.Add(Ptr);
 
 	//maybe scan the list
-	int32 DeleteMetric = CollectLimit <= 0 ? ((5 * TotalNumHazardRecords) / 4) : CollectLimit;
-	if (Collectables->Num() >= DeleteMetric)
+	const double CurrentTime = FApp::GetGameTime();
+	const bool TimeLimit = (CurrentTime - Collectables->TimeOfLastCollection) > 1.0;
+	const int32 DeleteMetric = CollectLimit <= 0 ? ((5 * TotalNumHazardRecords) / 4) : CollectLimit;
+	if (TimeLimit || Collectables->ReclamationList.Num() >= DeleteMetric)
 	{
-		Collect(*Collectables);
+		Collectables->TimeOfLastCollection = CurrentTime;
+		Collect(Collectables->ReclamationList);
 	}
 }
 
@@ -150,4 +150,4 @@ void FHazardPointerCollection::Collect(TArray<HazardPointer_Impl::FHazardDeleter
 	{
 		DeletedCollectable.Delete();
 	}
-}
\ No newline at end of file
+}
diff --git a/Engine/Source/Runtime/Core/Public/Experimental/Containers/HazardPointer.h b/Engine/Source/Runtime/Core/Public/Experimental/Containers/HazardPointer.h
--- a/Engine/Source/Runtime/Core/Public/Experimental/Containers/HazardPointer.h
+++ b/Engine/Source/Runtime/Core/Public/Experimental/Containers/HazardPointer.h
@@ -53,6 +53,21 @@ class FHazardPointerCollection
 	template<typename, bool>
 	friend class THazardPointer;
 
+	struct FTlsData
+	{
+		TArray<HazardPointer_Impl::FHazardDeleter> ReclamationList;
+		double TimeOfLastCollection = .0;
+
+		~FTlsData()
+		{
+			for (HazardPointer_Impl::FHazardDeleter& Deleter : ReclamationList)
+			{
+				Deleter.Delete();
+			}
+			ReclamationList.Empty();
+		}
+	};
+
 	class alignas(PLATFORM_CACHE_LINE_SIZE * 2) FHazardRecord
 	{
 		friend class FHazardPointerCollection;
@@ -125,7 +140,7 @@ class FHazardPointerCollection
 
 	FCriticalSection AllTlsVariablesCS;
 	FCriticalSection HazardRecordBlocksCS;
-	TArray<TArray<HazardPointer_Impl::FHazardDeleter>*> AllTlsVariables;
+	TArray<FTlsData*> AllTlsVariables;
 	TArray<FHazardRecordChunk*> HazardRecordBlocks;
 
 	uint32 CollectablesTlsSlot = 0;
