Fixing Bad Merge in CL 18459564. Incorrectly failed to backout the previous CL 18400633.
#tests none

diff --git a/Engine/Source/Runtime/Core/Public/Misc/GenericGrowableAllocator.h b/Engine/Source/Runtime/Core/Public/Misc/GenericGrowableAllocator.h
--- a/Engine/Source/Runtime/Core/Public/Misc/GenericGrowableAllocator.h
+++ b/Engine/Source/Runtime/Core/Public/Misc/GenericGrowableAllocator.h
@@ -9,6 +9,7 @@
 #include "CoreMinimal.h"
 #include "HAL/LowLevelMemTracker.h"
 #include "Misc/ScopeLock.h"
+#include "ProfilingDebugging/MemoryTrace.h"
 #include "Stats/Stats.h"
 
 //////////////
@@ -128,7 +129,6 @@ public:
 		: MemoryType(Type)
 		, HeapSize(InSize)
 		, UsedMemorySize(0)
-		, MaxFreeEntrySize(MaxFreeEntrySizeDirty)
 		, CriticalSection(InCriticalSection)
 	{
 	}
@@ -139,7 +139,6 @@ public:
 		HeapSize = CreateInternalMemory(HeapSize);
 		// entire chunk is free
 		FreeList = new FFreeEntry(NULL, 0, HeapSize);
-		MaxFreeEntrySize = FreeList->BlockSize;
 	}
 
 	/**
@@ -163,30 +162,8 @@ public:
 	*/
 	bool CanFitEntry(uint32 Size, uint32 Alignment)
 	{
-		// Compute MaxFreeEntrySize if necessary (should only happen if this chunk was just allocated from)
-		if (MaxFreeEntrySize == MaxFreeEntrySizeDirty)
-		{
-			MaxFreeEntrySize = 0;
-			for (FFreeEntry* Entry = FreeList; Entry; Entry = Entry->Next)
-			{
-				MaxFreeEntrySize = FMath::Max(Entry->BlockSize, MaxFreeEntrySize);
-			}
-		}
-
-		// Return false if we trivially don't fit
-		if (Size > MaxFreeEntrySize)
-		{
-			return false;
-		}
-
-		// Return true if we trivially do fit
-		if (Size + Alignment - 1 <= MaxFreeEntrySize)
-		{
-			return true;
-		}
-
-		// Slow method - search the free entries for a free chunk
 		bool bResult = false;
+		// look for a good free chunk
 		for (FFreeEntry *Entry = FreeList; Entry; Entry = Entry->Next)
 		{
 			if (Entry->CanFit(Size, Alignment))
@@ -223,49 +200,37 @@ public:
 
 		// look for a good free chunk
 		FFreeEntry* Prev = NULL;
-		FFreeEntry* FindEntry = NULL;
-
-		for (FFreeEntry* Entry = FreeList; Entry; Entry = Entry->Next)
+		for (FFreeEntry *Entry = FreeList; Entry; Entry = Entry->Next)
 		{
 			if (Entry->CanFit(AlignedSize, Alignment))
 			{
-				FindEntry = Entry;
-				break;
-			}
-			Prev = Entry;
-		}
-
-		if (FindEntry != NULL)
-		{
-			if (FindEntry->BlockSize == MaxFreeEntrySize)
-			{
-				// We're probably about to just split our largest entry, so mark the max size as dirty to indicate it needs recomputing
-				MaxFreeEntrySize = MaxFreeEntrySizeDirty;
-			}
-			// Use it, leaving over any unused space
-			UsedMemorySize += AlignedSize;
-			bool bDelete;
-			uint32 Padding;
-			uint32 Offset = FindEntry->Split(AlignedSize, Alignment, bDelete, Padding, MinAllocationSize);
-			if (bDelete)
-			{
-				FFreeEntry*& PrevRef = Prev ? Prev->Next : FreeList;
-				PrevRef = FindEntry->Next;
-				delete FindEntry;
-			}
+				// Use it, leaving over any unused space
+				UsedMemorySize += AlignedSize;
+				bool bDelete;
+				uint32 Padding;
+				uint32 Offset = Entry->Split(AlignedSize, Alignment, bDelete, Padding, MinAllocationSize);
+				if (bDelete)
+				{
+					FFreeEntry*& PrevRef = Prev ? Prev->Next : FreeList;
+					PrevRef = Entry->Next;
+					delete Entry;
+				}
 
-			FGrowableAllocationBase* Allocation = CreateAllocationStruct();
-			Allocation->Size = AlignedSize;
-			Allocation->Padding = Padding;
-			Allocation->Offset = Offset;
+				FGrowableAllocationBase* Allocation = CreateAllocationStruct();
+				Allocation->Size = AlignedSize;
+				Allocation->Padding = Padding;
+				Allocation->Offset = Offset;
 #if !UE_BUILD_SHIPPING
-			Allocation->OwnerType = OwnerType;
+				Allocation->OwnerType = OwnerType;
 #endif
-			LLM(FLowLevelMemTracker::Get().OnLowLevelAlloc(ELLMTracker::Default, GetAddressForTracking(Offset), Size));
+				LLM(FLowLevelMemTracker::Get().OnLowLevelAlloc(ELLMTracker::Default, GetAddressForTracking(Offset), Size));
+				MemoryTrace_Alloc(uint64(GetAddressForTracking(Offset)), Size, Alignment, EMemoryTraceRootHeap::SystemMemory);
 
-			// let the implementation fill in any more
-			InitializeAllocationStruct(Allocation);
-			return Allocation;
+				// let the implementation fill in any more
+				InitializeAllocationStruct(Allocation);
+				return Allocation;
+			}
+			Prev = Entry;
 		}
 
 		// if no suitable blocks were found, we must fail
@@ -284,6 +249,7 @@ public:
 		uint32 Offset = Memory->Offset;
 
 		LLM(FLowLevelMemTracker::Get().OnLowLevelFree(ELLMTracker::Default, GetAddressForTracking(Offset)));
+		MemoryTrace_Free(uint64(GetAddressForTracking(Offset)), EMemoryTraceRootHeap::SystemMemory);
 
 		// we are now done with the Allocation object
 		DestroyAllocationStruct(Memory);
@@ -291,7 +257,7 @@ public:
 		UsedMemorySize -= Size;
 		CurrentAllocs--;
 
-		// Search for a place to insert a new free entry
+		// Search for where a place to insert a new free entry.
 		FFreeEntry* Prev = NULL;
 		FFreeEntry* Entry = FreeList;
 		while (Entry && Offset > Entry->Location)
@@ -312,13 +278,8 @@ public:
 			{
 				Prev->BlockSize += Entry->BlockSize;
 				Prev->Next = Entry->Next;
-				MaxFreeEntrySize = FMath::Max(MaxFreeEntrySize, Prev->BlockSize);
 				delete Entry;
 			}
-			else
-			{
-				MaxFreeEntrySize = FMath::Max(MaxFreeEntrySize, Entry->BlockSize);
-			}
 			return true;
 		}
 
@@ -335,7 +296,6 @@ public:
 				Prev->Next = Entry->Next;
 				delete Entry;
 			}
-			MaxFreeEntrySize = FMath::Max(MaxFreeEntrySize, Prev->BlockSize);
 			return true;
 		}
 
@@ -343,7 +303,6 @@ public:
 		FFreeEntry* NewFree = new FFreeEntry(Entry, Offset - Padding, AllocationSize);
 		FFreeEntry*& PrevRef = Prev ? Prev->Next : FreeList;
 		PrevRef = NewFree;
-		MaxFreeEntrySize = FMath::Max(MaxFreeEntrySize, NewFree->BlockSize);
 		return true;
 	}
 
@@ -484,16 +443,11 @@ public:
 	/** Size of used memory */
 	uint64 UsedMemorySize;
 
-	/** Size of the largest free entry (will be MaxFreeEntrySizeDirty if unknown) */
-	uint32 MaxFreeEntrySize;
-
 	/** List of free blocks */
 	FFreeEntry* FreeList;
 
 	/** Shared critical section */
 	FCriticalSection* CriticalSection;
-
-	static const uint32 MaxFreeEntrySizeDirty = 0xffffffff;
 };
 
 
@@ -602,7 +556,7 @@ public:
 			return nullptr;
 		}
 
-#if ALLOCATION_HISTOGRAM
+#if !(UE_BUILD_SHIPPING || UE_BUILD_TEST)
 		TotalAllocationsHistogram.FindOrAdd(AlignedSize)++;
 		PeakAllocationsHistogram.FindOrAdd(AlignedSize)++;
 
@@ -611,12 +565,12 @@ public:
 
 		TypeAllocInfo.Counts.Allocations++;
 		TypeAllocInfo.TotalAllocated += Result->Size;
-#if ALLOCATION_HISTOGRAM_DETAILED
+#ifdef ALLOCATION_HISTOGRAM
 		TypeAllocInfo.AllocationHistogram.FindOrAdd(Result->Size).Allocations++;
 #endif
-#endif // ALLOCATION_HISTOGRAM 
+#endif // #if !(UE_BUILD_SHIPPING || UE_BUILD_TEST)
 
-#if STATS
+#if !UE_BUILD_SHIPPING
 		if (OwnerTypeToStatIdMap)
 		{
 			INC_MEMORY_STAT_BY_FName(OwnerTypeToStatIdMap[Result->OwnerType], Result->Size);
@@ -642,7 +596,7 @@ public:
 			ChunkAllocatorType* Chunk = AllocChunks[ChunkIndex];
 			if (Chunk && Chunk->DoesChunkContainAllocation(Memory))
 			{
-#if ALLOCATION_HISTOGRAM
+#if !(UE_BUILD_SHIPPING || UE_BUILD_TEST)
 				PeakAllocationsHistogram[Memory->Size]--;
 
 				// untrack per type allocation info
@@ -650,12 +604,12 @@ public:
 
 				TypeAllocInfo.TotalAllocated -= Memory->Size;
 				TypeAllocInfo.Counts.Frees++;
-#if ALLOCATION_HISTOGRAM_DETAILED
+#ifdef ALLOCATION_HISTOGRAM
 				TypeAllocInfo.AllocationHistogram[Memory->Size].Frees++;
-#endif 
-#endif // ALLOCATION_HISTOGRAM
+#endif
+#endif // !(UE_BUILD_SHIPPING || UE_BUILD_TEST)
 
-#if STATS
+#if !UE_BUILD_SHIPPING
 				if (OwnerTypeToStatIdMap)
 				{
 					DEC_MEMORY_STAT_BY_FName(OwnerTypeToStatIdMap[Memory->OwnerType], Memory->Size);
@@ -717,19 +671,18 @@ public:
 		// multi-thread protection
 		FScopeLock ScopeLock(&CriticalSection);
 
-#if ALLOCATION_HISTOGRAM
+#if !(UE_BUILD_SHIPPING || UE_BUILD_TEST)
 		FPlatformMisc::LowLevelOutputDebugStringf(TEXT("Total Allocations Histogram:\n"));
 		for (auto It = TotalAllocationsHistogram.CreateIterator(); It; ++It)
 		{
 			FPlatformMisc::LowLevelOutputDebugStringf(TEXT(" %d, %d\n"), It.Key(), It.Value());
 		}
-#endif // ALLOCATION_HISTOGRAM
+#endif // !(UE_BUILD_SHIPPING || UE_BUILD_TEST
 
 		int32 NumChunks = 0;
 		// pass off to individual alloc chunks
 		for (int32 ChunkIndex = 0; ChunkIndex < AllocChunks.Num(); ChunkIndex++)
 		{
-			FPlatformMisc::LowLevelOutputDebugStringf(TEXT("\n-----------------\nChunk %d\n"),ChunkIndex);
 			ChunkAllocatorType* Chunk = AllocChunks[ChunkIndex];
 			if (Chunk)
 			{
@@ -740,7 +693,7 @@ public:
 
 	void DumpMemoryInfo()
 	{
-#if ALLOCATION_HISTOGRAM
+#if !(UE_BUILD_SHIPPING || UE_BUILD_TEST)
 		// we use LowLevel here because we are usually dumping this out while 
 		for (auto InfoIt = PerTypeAllocationInfo.CreateConstIterator(); InfoIt; ++InfoIt)
 		{
@@ -754,14 +707,14 @@ public:
 			{
 				FPlatformMisc::LowLevelOutputDebugStringf(TEXT("      'OwnerType %d': %lld Allocs: %d Frees: %d\n"), InfoIt.Key(), Info.TotalAllocated, Info.Counts.Allocations, Info.Counts.Frees);
 			}
-#if ALLOCATION_HISTOGRAM_DETAILED
+#ifdef ALLOCATION_HISTOGRAM
 			for (auto HistoIt = Info.AllocationHistogram.CreateConstIterator(); HistoIt; ++HistoIt)
 			{
 				FPlatformMisc::LowLevelOutputDebugStringf(TEXT("           %d, %lld, %lld\n"), HistoIt.Key(), HistoIt.Value().Allocations, HistoIt.Value().Frees);
 			}
 #endif
 		}
-#endif // ALLOCATION_HISTOGRAM
+#endif // !(UE_BUILD_SHIPPING || UE_BUILD_TEST)
 	}
 
 private:
@@ -863,7 +816,7 @@ private:
 	 // a critical section used to coordinate all access in this instance of the allocator and it's chunks
 	 FCriticalSection CriticalSection;
 
-#if ALLOCATION_HISTOGRAM
+#if !(UE_BUILD_SHIPPING || UE_BUILD_TEST)
 	 struct AllocFreeCounts
 	 {
 		 uint32 Allocations;
@@ -872,7 +825,7 @@ private:
 	 struct AllocationInfo
 	 {
 		 uint64 TotalAllocated;
-#if ALLOCATION_HISTOGRAM_DETAILED
+#ifdef ALLOCATION_HISTOGRAM
 		 TMap<uint32, AllocFreeCounts> AllocationHistogram;
 #endif
 		 AllocFreeCounts Counts;
@@ -883,6 +836,7 @@ private:
 	 TMap<uint64, uint32> TotalAllocationsHistogram;
 	 TMap<uint64, uint32> OutstandingAllocationsHistogram;
 	 TMap<uint64, uint32> PeakAllocationsHistogram;
-#endif // ALLOCATION_HISTOGRAM
+#endif // !(UE_BUILD_SHIPPING || UE_BUILD_TEST)
+
 };
 
