REL - Tasks can be canceled/retraced while they are initzialized.

diff --git a/Engine/Source/Runtime/Core/Public/Experimental/Async/AwaitableTask.h b/Engine/Source/Runtime/Core/Public/Experimental/Async/AwaitableTask.h
--- a/Engine/Source/Runtime/Core/Public/Experimental/Async/AwaitableTask.h
+++ b/Engine/Source/Runtime/Core/Public/Experimental/Async/AwaitableTask.h
@@ -463,19 +463,19 @@ public:
 
 	//initialize and Launch a Task with given Priority
 	template<typename CallableT>
-	void InitAndLaunch(const TCHAR* DebugName, LowLevelTasks::ETaskPriority Priority, CallableT&& Callable)
+	bool InitAndLaunch(const TCHAR* DebugName, LowLevelTasks::ETaskPriority Priority, CallableT&& Callable)
 	{
 		Init(DebugName, Priority, Forward<CallableT>(Callable));
 		using CallableType = std::decay_t<CallableT>;
 		using PromiseVTableType = TPromiseVTable<TPromise<ReturnType, CallableType>>;
-		verify(GetVtable<PromiseVTableType>()->TryLaunch());
+		return GetVtable<PromiseVTableType>()->TryLaunch();
 	}
 
 	//initialize and Launch a Task
 	template<typename CallableT>
-	inline void InitAndLaunch(const TCHAR* DebugName, CallableT&& Callable)
+	inline bool InitAndLaunch(const TCHAR* DebugName, CallableT&& Callable)
 	{
-		InitAndLaunch(DebugName, LowLevelTasks::ETaskPriority::Default, Forward<CallableT>(Callable));
+		return InitAndLaunch(DebugName, LowLevelTasks::ETaskPriority::Default, Forward<CallableT>(Callable));
 	}
 
 	//Abandon the Task and freeing its memory if this is the last reference to it.
