Added missing FMath::CeilToInt64.

diff --git a/Engine/Source/Runtime/Core/Public/GenericPlatform/GenericPlatformMath.h b/Engine/Source/Runtime/Core/Public/GenericPlatform/GenericPlatformMath.h
--- a/Engine/Source/Runtime/Core/Public/GenericPlatform/GenericPlatformMath.h
+++ b/Engine/Source/Runtime/Core/Public/GenericPlatform/GenericPlatformMath.h
@@ -345,6 +345,13 @@ struct FGenericPlatformMath
 
 	RESOLVE_FLOAT_TO_TYPE_AMBIGUITY(CeilToInt, int32);
 
+	static FORCEINLINE int64 CeilToInt64(double F)
+	{
+		int64 I = (int64)F;
+		I += ((double)I < F);
+		return I;
+	}
+
 	/**
 	* Converts a float to the nearest greater or equal integer.
 	* @param F		Floating point value to convert
diff --git a/Engine/Source/Runtime/Core/Public/Math/UnrealPlatformMathSSE.h b/Engine/Source/Runtime/Core/Public/Math/UnrealPlatformMathSSE.h
--- a/Engine/Source/Runtime/Core/Public/Math/UnrealPlatformMathSSE.h
+++ b/Engine/Source/Runtime/Core/Public/Math/UnrealPlatformMathSSE.h
@@ -255,6 +255,12 @@ struct TUnrealPlatformMathSSEBase : public Base
 		return UE4::SSE::CeilToInt(F);
 	}
 
+	template<typename T>
+	static FORCEINLINE int64 CeilToInt64(T F)
+	{
+		return UE4::SSE::CeilToInt64(F);
+	}
+
 	template<typename T>
 	static FORCEINLINE T InvSqrt(T F)
 	{
