Added FBox2D::Overlap (similar to FBox::Overlap).

diff --git a/Engine/Source/Runtime/Core/Public/Math/Box2D.h b/Engine/Source/Runtime/Core/Public/Math/Box2D.h
--- a/Engine/Source/Runtime/Core/Public/Math/Box2D.h
+++ b/Engine/Source/Runtime/Core/Public/Math/Box2D.h
@@ -286,6 +286,14 @@ public:
 		bIsValid = false;
 	}
 
+	/**
+	 * Returns the overlap box of two boxes
+	 *
+	 * @param Other The bounding box to test overlap
+	 * @return the overlap box. Result will be invalid if they don't overlap
+	 */
+	TBox2<T> Overlap( const TBox2<T>& Other ) const;
+
 	/**
 	 * Checks whether the given box intersects this box.
 	 *
@@ -463,6 +471,28 @@ FORCEINLINE TVector2<T> TBox2<T>::GetClosestPointTo( const TVector2<T>& Point )
 	return ClosestPoint;
 }
 
+template<typename T>
+TBox2<T> TBox2<T>::Overlap(const TBox2<T>& Other) const
+{
+	if (Intersect(Other) == false)
+	{
+		static TBox2<T> EmptyBox(ForceInit);
+		return EmptyBox;
+	}
+
+	// otherwise they overlap
+	// so find overlapping box
+	TVector2<T> MinVector, MaxVector;
+
+	MinVector.X = FMath::Max(Min.X, Other.Min.X);
+	MaxVector.X = FMath::Min(Max.X, Other.Max.X);
+
+	MinVector.Y = FMath::Max(Min.Y, Other.Min.Y);
+	MaxVector.Y = FMath::Min(Max.Y, Other.Max.Y);
+
+	return TBox2<T>(MinVector, MaxVector);
+}
+
 template<typename T>
 FORCEINLINE bool TBox2<T>::Intersect( const TBox2<T> & Other ) const
 {
