FHashBuilder now uses HashCombineFast, so it gets faster when it's properly implemented.

diff --git a/Engine/Source/Runtime/Core/Public/Misc/HashBuilder.h b/Engine/Source/Runtime/Core/Public/Misc/HashBuilder.h
--- a/Engine/Source/Runtime/Core/Public/Misc/HashBuilder.h
+++ b/Engine/Source/Runtime/Core/Public/Misc/HashBuilder.h
@@ -42,7 +42,7 @@ public:
 	template <typename T>
 	typename TEnableIf<TModels<CGetTypeHashable, T>::Value, FHashBuilder&>::Type Append(const T& InData)
 	{
-		Hash = HashCombine(Hash, GetTypeHash(InData));
+		Hash = HashCombineFast(Hash, GetTypeHash(InData));
 		return *this;
 	}
 
