Fix Typo in deprecation upgrade comment.

diff --git a/Engine/Source/Runtime/Core/Public/Containers/Ticker.h b/Engine/Source/Runtime/Core/Public/Containers/Ticker.h
--- a/Engine/Source/Runtime/Core/Public/Containers/Ticker.h
+++ b/Engine/Source/Runtime/Core/Public/Containers/Ticker.h
@@ -168,7 +168,7 @@ private:
  * Migration guide:
  * As FTicker is deprecated and replaced by thread-safe FTSTicker class, if you see a deprecation warning please migrate your code. Change
  * `FTicker` to `FTSTicker` and change the type of the delegate handle returned by `AddTicker` methods from `FDelegateHandle` to
- * `FTSTicker::DelegateHandle`
+ * `FTSTicker::FDelegateHandle`
  */
 class UE_DEPRECATED(5.0, "Please use thread-safe FTSTicker, see the migration guide in the code comment above") FTicker
 {
