- Added ability to get a CVar's value for another platform (IConsoleVariable::GetPlatformValueVariable()), which gets another IConsoleVariable that contains the value
- Added console command syntax for viewing the value from the editor (<CVarName>[at]<PlatformName>)
- Converted FShaderPlatformCachedIniValue to use the new system (which will give a better possible value on the target platform). Eventually, this can be replaced with direct calls to GetPlatformValueVariable(ShaderPlatformToPlatformName(SHaderPlatform))
#p4v-preflight-copy 18268687

diff --git a/Engine/Source/Runtime/Core/Private/HAL/ConsoleManager.cpp b/Engine/Source/Runtime/Core/Private/HAL/ConsoleManager.cpp
--- a/Engine/Source/Runtime/Core/Private/HAL/ConsoleManager.cpp
+++ b/Engine/Source/Runtime/Core/Private/HAL/ConsoleManager.cpp
@@ -307,6 +307,175 @@ void OnCVarChange(T& Dst, const T& Src, EConsoleVariableFlags Flags, EConsoleVar
 	}
 }
 
+#if ALLOW_OTHER_PLATFORM_CONFIG
+
+static void ExpandScalabilityCVar(FConfigCacheIni* ConfigSystem, const FString& CVarKey, const FString& CVarValue, TMap<FString, FString>& ExpandedCVars, bool bOverwriteExistingValue)
+{
+	// load scalability settings directly from ini instead of using scalability system, so as not to inadvertantly mess anything up
+	// if the DP had sg.ResolutionQuality=3, we would read [ResolutionQuality@3]
+	FString SectionName = FString::Printf(TEXT("%s@%s"), *CVarKey.Mid(3), *CVarValue);
+	// walk over the scalability section and add them in, unless already done
+	FConfigSection* ScalabilitySection = ConfigSystem->GetSectionPrivate(*SectionName, false, true, GScalabilityIni);
+	if (ScalabilitySection != nullptr)
+	{
+		for (const auto& Pair : *ScalabilitySection)
+		{
+			FString ScalabilityKey = Pair.Key.ToString();
+			if (bOverwriteExistingValue || !ExpandedCVars.Contains(ScalabilityKey))
+			{
+				ExpandedCVars.Add(ScalabilityKey, Pair.Value.GetValue());
+			}
+		}
+	}
+}
+
+bool IConsoleManager::VisitPlatformCVarsForEmulation(FName PlatformName, bool bVisitPlatformDeviceProfile, TFunctionRef<void(const FString& CVarName, const FString& CVarValue, EConsoleVariableFlags SetBy)> Visit)
+{
+	// get the config system for the platform the DP uses
+	FConfigCacheIni* ConfigSystem = FConfigCacheIni::ForPlatform(*PlatformName.ToString());
+	if (ConfigSystem == nullptr)
+	{
+		return false;
+	}
+
+	// ECVF_SetByConstructor:
+	//   doesn't come from ini
+
+	// ECVF_SetByScalability:
+	//	 skipped, i believe this is not really loaded as a normal layer per-se, it's up to the other sections to set with this one
+
+	// ECVF_SetByGameSetting:
+	//   skipped, since we don't have a user
+
+
+	// ECVF_SetByProjectSetting:
+	// ECVF_SetBySystemSettingsIni:
+	// ECVF_SetByDeviceProfile: 
+	//   ONLY using the base DP named for the platform, as the DPManager needs to be used for proper DeviceProfiles
+
+	const TCHAR* DeviceProfileTag = TEXT("_NamedDeviceProfile");
+	struct FSectionPair
+	{
+		const TCHAR* Name; EConsoleVariableFlags SetBy;
+	} Sections[] =
+	{
+		{ TEXT("/Script/Engine.RendererSettings"), ECVF_SetByProjectSetting },
+		{ TEXT("/Script/Engine.RendererOverrideSettings"), ECVF_SetByProjectSetting },
+		{ TEXT("/Script/Engine.StreamingSettings"), ECVF_SetByProjectSetting },
+		{ TEXT("/Script/Engine.GarbageCollectionSettings"), ECVF_SetByProjectSetting },
+		{ TEXT("/Script/Engine.NetworkSettings"), ECVF_SetByProjectSetting },
+
+		{ TEXT("SystemSettings"), ECVF_SetBySystemSettingsIni },
+		{ TEXT("ConsoleVariables"), ECVF_SetBySystemSettingsIni },
+
+		{ DeviceProfileTag, ECVF_SetByDeviceProfile },
+	};
+
+	// now walk up the stack getting current values
+	for (const FSectionPair& SectionPair : Sections)
+	{
+		FConfigSection* Section;
+		if (FCString::Strcmp(SectionPair.Name, DeviceProfileTag) == 0)
+		{
+			// skip this if we don't want it
+			if (bVisitPlatformDeviceProfile == false)
+			{
+				continue;
+			}
+			FString DPSectionName = PlatformName.ToString() + TEXT(" DeviceProfile");
+			Section = ConfigSystem->GetSectionPrivate(*DPSectionName, false, true, GDeviceProfilesIni);
+		}
+		else
+		{
+			Section = ConfigSystem->GetSectionPrivate(SectionPair.Name, false, true, GEngineIni);
+		}
+
+		if (Section != nullptr)
+		{
+			// add the cvars from the section
+			for (const auto& Pair : *Section)
+			{
+				FString Key = Pair.Key.ToString();
+				FString Value = Pair.Value.GetValue();
+
+				// DPs have +CVars= prefix for the key, and the real cvar KVP is in the Value
+				if (Key.StartsWith(TEXT("CVars")))
+				{
+					FString ValueCopy = Value;
+					// expect a second = in the value, skip any that don't
+					if (!ValueCopy.Split(TEXT("="), &Key, &Value))
+					{
+						continue;
+					}
+				}
+
+				if (Key.StartsWith(TEXT("sg.")))
+				{
+					// @todo ini: If anything in here was already set, overwrite it or skip it?
+					// the priorities may cause runtime to fail to set a cvar that this will set blindly, since we are ignoring
+					// priority by doing them "in order". Scalablity is one of the lowest priorities, so should almost never be allowed?
+
+					TMap<FString, FString> ScalabilityCVars;
+					ExpandScalabilityCVar(ConfigSystem, Key, Value, ScalabilityCVars, true);
+
+					for (const auto& ScalabilityPair : ScalabilityCVars)
+					{
+						Visit(ScalabilityPair.Key, ScalabilityPair.Value, ECVF_SetByScalability);
+					}
+				}
+				else
+				{
+					Visit(Key, Value, SectionPair.SetBy);
+				}
+			}
+		}
+	}
+
+	// ECVF_SetByConsoleVariablesIni:
+	//   maybe skip this? it's a weird one, but maybe?
+
+	// ECVF_SetByCommandline:
+	//   skip as this would not be expected to apply to emulation
+
+	// ECVF_SetByCode:
+	//   skip because it cannot be set by code
+
+	// ECVF_SetByConsole
+	//   we could have this if we made a per-platform CVar, not just the shared default value
+
+
+	return true;
+}
+
+
+static bool GetConfigValueFromRuntimeSources(FName PlatformName, IConsoleVariable* CVar, FString& OutValue, EConsoleVariableFlags& LastSetBy)
+{
+
+	FString VariableName = IConsoleManager::Get().FindConsoleObjectName(CVar);
+	if (VariableName.Len() == 0)
+	{
+		return false;
+	}
+
+	OutValue = CVar->GetDefaultValueVariable()->GetString();
+	LastSetBy = ECVF_SetByConstructor;
+
+	IConsoleManager::VisitPlatformCVarsForEmulation(PlatformName, true, 
+		[&VariableName, &OutValue, &LastSetBy](const FString& CVarName, const FString& CVarValue, EConsoleVariableFlags SetBy)
+		{
+			// if this key is the variable, set it at the current SetBy level
+			if (CVarName == VariableName)
+			{
+				OutValue = CVarValue;
+				LastSetBy = SetBy;
+			}
+		});
+
+	// by this point, the value will be set, with the DefaultValue if nothing else from the visit function
+	return true;
+}
+#endif
+
 // T: bool, int32, float, FString
 template <class T>
 class FConsoleVariable : public FConsoleVariableBase
@@ -318,7 +487,7 @@ public:
 #if ALLOW_OTHER_PLATFORM_CONFIG
 		if (bSaveDefault)
 		{
-			PlatformIndependentDefault = TSharedPtr<IConsoleVariable>(new FConsoleVariable<T>(DefaultValue, TEXT(""), Flags, false));
+			PlatformIndependentDefault = TSharedPtr<IConsoleVariable>(new FConsoleVariable<T>(DefaultValue, TEXT(""), (EConsoleVariableFlags)(Flags | ECVF_ReadOnly), false));
 		}
 #endif
 	}
@@ -344,7 +513,7 @@ public:
 	virtual float GetFloat() const override;
 	virtual FString GetString() const override;
 
-virtual bool IsVariableBool() const override { return false; }
+	virtual bool IsVariableBool() const override { return false; }
 	virtual bool IsVariableInt() const override { return false; }
 	virtual bool IsVariableFloat() const override { return false; }
 	virtual bool IsVariableString() const override { return false; }
@@ -380,6 +549,37 @@ private: // ----------------------------------------------------
 	{
 		return PlatformIndependentDefault.Get();
 	}
+
+	TMap<FName, TSharedPtr<IConsoleVariable> > PlatformValues;
+	virtual IConsoleVariable* GetPlatformValueVariable(FName PlatformName)
+	{
+		check(IsInGameThread());
+		if (!PlatformValues.Contains(PlatformName))
+		{
+			FString ConfigValue;
+			EConsoleVariableFlags LastSetBy;
+			if (!GetConfigValueFromRuntimeSources(PlatformName, this, ConfigValue, LastSetBy))
+			{
+				return nullptr;
+			}
+			
+			T TypedValue;
+			TTypeFromString<T>::FromString(TypedValue, *ConfigValue);
+
+			// clear the existing setby mask
+			int NewFlags = GetFlags() & ~ECVF_SetFlagMask;
+			// add in new LastSetBy and make it readonly
+			NewFlags |= LastSetBy | ECVF_ReadOnly;
+
+			// make a new cvar to hold the value from the other platform
+			TSharedPtr<IConsoleVariable> PlatformCVar = TSharedPtr<IConsoleVariable>(new FConsoleVariable<T>(TypedValue, TEXT(""), (EConsoleVariableFlags)NewFlags, false));
+
+			PlatformValues.Add(PlatformName, PlatformCVar);
+		}
+
+		return PlatformValues.FindRef(PlatformName).Get();
+	}
+
 #endif
 
 };
@@ -1370,6 +1570,20 @@ bool FConsoleManager::ProcessUserConsoleInput(const TCHAR* InInput, FOutputDevic
 		Param1.MidInline(0, Param1.Len() - 1, false);
 	}
 
+	// look for the <cvar>@<platform> syntax
+	FName PlatformName;
+	if (Param1.Contains(TEXT("@")))
+	{
+		FString Left, Right;
+		Param1.Split(TEXT("@"), &Left, &Right);
+
+		if (Left.Len() && Right.Len())
+		{
+			Param1 = Left;
+			PlatformName = *Right;
+		}
+	}
+
 	IConsoleObject* CObj = FindConsoleObject(*Param1);
 	if(!CObj)
 	{
@@ -1394,6 +1608,28 @@ bool FConsoleManager::ProcessUserConsoleInput(const TCHAR* InInput, FOutputDevic
 	IConsoleCommand* CCmd = CObj->AsCommand();
 	IConsoleVariable* CVar = CObj->AsVariable();
 
+	if (PlatformName != NAME_None)
+	{
+		if (CVar == nullptr)
+		{
+			Ar.Logf(TEXT("Ignoring platform portion (@%s), which is only valid for looking up CVars"), *PlatformName.ToString());
+		}
+		else
+		{
+#if ALLOW_OTHER_PLATFORM_CONFIG
+			CVar = CVar->GetPlatformValueVariable(PlatformName);
+			if (!CVar)
+			{
+				Ar.Logf(TEXT("Unable find CVar %s for platform %s (possibly invalid platform name?)"), *Param1, *PlatformName.ToString());
+				return false;
+			}
+#else
+			Ar.Logf(TEXT("Unable to lookup a CVar value on another platform in this build"));
+			return false;
+#endif
+		}
+	}
+
 	if( CCmd )
 	{
 		// Process command
@@ -1450,9 +1686,13 @@ bool FConsoleManager::ProcessUserConsoleInput(const TCHAR* InInput, FOutputDevic
 			}
 			else
 			{
-				if(bReadOnly)
+				if (PlatformName != NAME_None)
+				{
+					Ar.Logf(TEXT("Error: Unable to set a value for %s another platform!"), *Param1);
+				}
+				else if(bReadOnly)
 				{
-					Ar.Logf(TEXT("Error: %s is read only!"), *Param1, *CVar->GetString());
+					Ar.Logf(TEXT("Error: %s is read only!"), *Param1);
 				}
 				else
 				{
diff --git a/Engine/Source/Runtime/Core/Public/HAL/ConsoleManager.h b/Engine/Source/Runtime/Core/Public/HAL/ConsoleManager.h
--- a/Engine/Source/Runtime/Core/Public/HAL/ConsoleManager.h
+++ b/Engine/Source/Runtime/Core/Public/HAL/ConsoleManager.h
@@ -36,7 +36,7 @@ public:
 	bool IsThreadPropagationThread();
 
 	/** @param InVar must not be 0 */
-	FString FindConsoleObjectName(const IConsoleObject* Obj) const;
+	virtual FString FindConsoleObjectName(const IConsoleObject* Obj) const override;
 
 	/** Can be moved out into some automated testing system */
 	void Test();
diff --git a/Engine/Source/Runtime/Core/Public/HAL/IConsoleManager.h b/Engine/Source/Runtime/Core/Public/HAL/IConsoleManager.h
--- a/Engine/Source/Runtime/Core/Public/HAL/IConsoleManager.h
+++ b/Engine/Source/Runtime/Core/Public/HAL/IConsoleManager.h
@@ -460,6 +460,13 @@ public:
 		return nullptr;
 	}
 
+#if ALLOW_OTHER_PLATFORM_CONFIG
+	virtual IConsoleVariable* GetPlatformValueVariable(FName PlatformName)
+	{
+		return nullptr;
+	}
+#endif
+
 	// convenience methods
 
 	/** Set the internal value from the specified bool. */
@@ -852,6 +859,13 @@ struct CORE_API IConsoleManager
 	*/
 	virtual IConsoleObject* FindConsoleObject(const TCHAR* Name, bool bTrackFrequentCalls = true) const = 0;
 
+	/**
+	 * Lookup the name of a console object by its pointer 
+	 * @param Object to lookup
+	 * @return Name of the object, or an empty string if the object can't be found
+	 */
+	virtual FString FindConsoleObjectName(const IConsoleObject* Obj) const = 0;
+
 	/**
 	 * Find a typed console variable (faster access to the value, no virtual function call)
 	 * @param Name must not be 0
@@ -934,6 +948,20 @@ struct CORE_API IConsoleManager
 		return *Singleton;
 	}
 
+#if ALLOW_OTHER_PLATFORM_CONFIG
+	/**
+	 * Walks over the best set of ini sections for another platform that can be used to emulate cvars as the platform
+	 * will have set - WITH THE NOTABLE EXCEPTION OF specific DeviceProfiles. Use the DeviceProfileManager code to get a DP
+	 * for a platform, as this code cannot easily access DeviceProfile logic.
+	 * It also won't include any UserSettings or ConsoleVariables.ini settings.
+	 * 
+	 * @param PlatformName The platform name (the ini name, so Windows, not Win64)
+	 * @param bVisitPlatformDeviceProfile Set to true if you want Visit called with CVars found in the base DeviceProfile named PlatformName 
+	 * @param Visit the callback to run for each CVar found
+	 */
+	static bool VisitPlatformCVarsForEmulation(FName PlatformName, bool bVisitPlatformDeviceProfile, TFunctionRef<void(const FString& CVarName, const FString& CVarValue, EConsoleVariableFlags SetBy)> Visit);
+#endif
+
 protected:
 	virtual ~IConsoleManager() { }
 
diff --git a/Engine/Source/Runtime/Engine/Private/DeviceProfiles/DeviceProfileManager.cpp b/Engine/Source/Runtime/Engine/Private/DeviceProfiles/DeviceProfileManager.cpp
--- a/Engine/Source/Runtime/Engine/Private/DeviceProfiles/DeviceProfileManager.cpp
+++ b/Engine/Source/Runtime/Engine/Private/DeviceProfiles/DeviceProfileManager.cpp
@@ -571,79 +571,18 @@ void UDeviceProfileManager::InitializeCVarsForActiveDeviceProfile()
 #if ALLOW_OTHER_PLATFORM_CONFIG
 void UDeviceProfileManager::ExpandDeviceProfileCVars(UDeviceProfile* DeviceProfile)
 {
-	// get the config system for the platform the DP uses
-	FConfigCacheIni* ConfigSystem = FConfigCacheIni::ForPlatform(*DeviceProfile->DeviceType);
-	check(ConfigSystem);
-
-	// walk up the chain cvar setBys and emulate what would happen on the target platform
-	FString Platform = DeviceProfile->DeviceType;
-
-	// now walk up the stack getting current values
-
-	// ECVF_SetByConstructor:
-	//   in PlatformIndependentDefault, used if getting a var but not in this DP
-
-	// ECVF_SetByScalability:
-	//	 skipped, i believe this is not really loaded as a normal layer per-se, it's up to the other sections to set with this one
-
-	// ECVF_SetByGameSetting:
-	//   skipped, since we don't have a user
-
-	const TCHAR* SectionNames[] = {
-		// ECVF_SetByProjectSetting:
-		TEXT("/Script/Engine.RendererSettings"),
-		TEXT("/Script/Engine.RendererOverrideSettings"),
-		TEXT("/Script/Engine.StreamingSettings"),
-		TEXT("/Script/Engine.GarbageCollectionSettings"),
-		TEXT("/Script/Engine.NetworkSettings"),
-
-		// ECVF_SetBySystemSettingsIni:
-		TEXT("SystemSettings"),
-		TEXT("ConsoleVariables"),
-	};
-
-	// go through possible cvar sections that the target platform would load and read all cvars in them
 	TMap<FString, FString> CVarsToAdd;
-	for (const TCHAR* SectionName : SectionNames)
-	{
-		FConfigSection* Section = ConfigSystem->GetSectionPrivate(SectionName, false, true, GEngineIni);
-		if (Section != nullptr)
+	IConsoleManager::VisitPlatformCVarsForEmulation(*DeviceProfile->DeviceType, false,
+		[&CVarsToAdd](const FString& CVarName, const FString& CVarValue, EConsoleVariableFlags SetBy)
 		{
-			// add the cvars from the section
-			for (const auto& Pair : *Section)
-			{
-				FString Key = Pair.Key.ToString();
-				FString Value = Pair.Value.GetValue();
-				if (Key.StartsWith(TEXT("sg.")))
-				{
-					// @todo ini: If anything in here was already set, overwrite it or skip it?
-					// the priorities may cause runtime to fail to set a cvar that this will set blindly, since we are ignoring
-					// priority by doing them "in order". Scalablity is one of the lowest priorities, so should almost never be allowed?
-					ExpandScalabilityCVar(ConfigSystem, Key, Value, CVarsToAdd, true);
-				}
-				CVarsToAdd.Add(Pair.Key.ToString(), Pair.Value.GetValue());
-			}
-		}
-	}
+			CVarsToAdd.Add(CVarName, CVarValue);
+		});
+
+	// we now have all the cvars for this platform without _any_ device profiles, so now add those into the map
 	DeviceProfile->AddExpandedCVars(CVarsToAdd);
 
-	// ECVF_SetByDeviceProfile:
+	// and finalize them into a cache
 	ProcessDeviceProfileIniSettings(DeviceProfile->GetName(), EDeviceProfileMode::DPM_CacheValues);
-
-	// ECVF_SetByGameOverride:
-	//   skipped, since we don't have a user
-
-	// ECVF_SetByConsoleVariablesIni:
-	//   maybe skip this? it's a weird one, but maybe?
-
-	// ECVF_SetByCommandline:
-	//   skip as this would not be expected to apply to emulation
-
-	// ECVF_SetByCode:
-	//   skip because it cannot be set by code
-
-	// ECVF_SetByConsole
-	//   we could have this if we made a per-platform CVar, not just the shared default value
 }
 #endif
 
@@ -1372,7 +1311,7 @@ void UDeviceProfileManager::GetAllPossibleParentProfiles(const UDeviceProfile* C
 
 
 #if ALLOW_OTHER_PLATFORM_CONFIG
-static bool GetCVarForPlatform( FOutputDevice& Ar, FString DPName, FString CVarName)
+static bool GetCVarForDeviceProfile( FOutputDevice& Ar, FString DPName, FString CVarName)
 {
 	UDeviceProfile* DeviceProfile = UDeviceProfileManager::Get().FindProfile(DPName, false);
 	if (DeviceProfile == nullptr)
@@ -1419,7 +1358,7 @@ public:
 				return false;
 			}
 
-			return GetCVarForPlatform(Ar, DPName, CVarName);
+			return GetCVarForDeviceProfile(Ar, DPName, CVarName);
 		}
 		else if (FParse::Command(&Cmd, TEXT("dpdump")))
 		{
@@ -1462,4 +1401,4 @@ public:
 } GPlatformCVarExec;
 
 
-#endif
+#endif
\ No newline at end of file
diff --git a/Engine/Source/Runtime/RenderCore/Private/RenderUtils.cpp b/Engine/Source/Runtime/RenderCore/Private/RenderUtils.cpp
--- a/Engine/Source/Runtime/RenderCore/Private/RenderUtils.cpp
+++ b/Engine/Source/Runtime/RenderCore/Private/RenderUtils.cpp
@@ -1214,64 +1214,6 @@ RENDERCORE_API bool SupportsTSR(const FStaticShaderPlatform Platform)
 	return FDataDrivenShaderPlatformInfo::GetSupportsGen5TemporalAA(Platform);
 }
 
-template<typename Type>
-Type FShaderPlatformCachedIniValue<Type>::Get(EShaderPlatform ShaderPlatform)
-{
-	Type Value = {};
-
-#if WITH_EDITOR
-	Type* ExistingEntry = CachedValues.Find(ShaderPlatform);
-	if (ExistingEntry != nullptr)
-	{
-		return *ExistingEntry;
-	}
-
-	bool bTestedIni = TestedIni.Contains(ShaderPlatform);
-	if (!bTestedIni)
-	{
-		TestedIni.Add(ShaderPlatform);
-		FConfigCacheIni* PlatformIni = FConfigCacheIni::ForPlatform(ShaderPlatformToPlatformName(ShaderPlatform));
-
-		if (PlatformIni != nullptr)
-		{
-			// Rendering CVars can be set from a number of sections - see UDeviceProfileManager::ExpandDeviceProfileCVars for their ordering.
-			// (Here they are in reverse order because we will use the first found setting).
-			const TCHAR* ConsoleVariablesSections[] =
-			{
-				TEXT("ConsoleVariables"),
-				TEXT("SystemSettings"),
-				TEXT("/Script/Engine.RendererOverrideSettings"),
-				TEXT("/Script/Engine.RendererSettings")
-			};
-
-			for (const TCHAR* SectionThatCanSetAVar : ConsoleVariablesSections)
-			{
-				if (PlatformIni->GetValue(SectionThatCanSetAVar, Key, Value, TEXT("Engine")))
-				{
-					CachedValues.Add(ShaderPlatform, Value);
-					return Value;
-				}
-			}
-		}
-	}
-#endif
-	// If it's not in the ini file, don't cache, always return the CVar's current value
-	if (CVar==nullptr)
-	{
-		CVar = IConsoleManager::Get().FindConsoleVariable(Key);
-	}
-
-	if (CVar!=nullptr)
-	{
-		CVar->GetValue(Value);
-	}
-	return Value;
-}
-template struct FShaderPlatformCachedIniValue<FString>;
-template struct FShaderPlatformCachedIniValue<int32>;
-template struct FShaderPlatformCachedIniValue<float>;
-template struct FShaderPlatformCachedIniValue<bool>;
-
 RENDERCORE_API int32 GUseForwardShading = 0;
 static FAutoConsoleVariableRef CVarForwardShading(
 	TEXT("r.ForwardShading"),
diff --git a/Engine/Source/Runtime/RenderCore/Public/RenderUtils.h b/Engine/Source/Runtime/RenderCore/Public/RenderUtils.h
--- a/Engine/Source/Runtime/RenderCore/Public/RenderUtils.h
+++ b/Engine/Source/Runtime/RenderCore/Public/RenderUtils.h
@@ -515,27 +515,46 @@ RENDERCORE_API bool PlatformSupportsVelocityRendering(const FStaticShaderPlatfor
 
 RENDERCORE_API bool IsUsingDBuffers(const FStaticShaderPlatform Platform);
 
-
-/* Simple cache for RendererSettings ini lookup per shader platform. */
 template<typename Type>
 struct RENDERCORE_API FShaderPlatformCachedIniValue
 {
-	FShaderPlatformCachedIniValue(const TCHAR* InKey)
-		: Key(InKey)
+	FShaderPlatformCachedIniValue(const TCHAR* /*InSection*/, const TCHAR* InCVarName)
+		: CVarName(InCVarName)
 		, CVar(nullptr)
-	{}
+	{
+	}
+
+	FShaderPlatformCachedIniValue(IConsoleVariable* InCVar)
+		: CVar(InCVar)
+	{
+	}
 
-	Type Get(EShaderPlatform ShaderPlatform);
+	Type Get(EShaderPlatform ShaderPlatform)
+	{
+		Type Value {};
+#if ALLOW_OTHER_PLATFORM_CONFIG
+		// find/create a cvar if needed
+		if (CVar == nullptr)
+		{
+			CVar = IConsoleManager::Get().FindConsoleVariable(*CVarName);
+			if (CVar == nullptr)
+			{
+				// this could be a cvar that only exists on the target platform so create a dummy one
+				CVar = IConsoleManager::Get().RegisterConsoleVariable(*CVarName, Type(), TEXT(""), ECVF_ReadOnly);
+			}
+		}
+
+		// now get the value from ther platform that makes sense for this shader platform
+		CVar->GetPlatformValueVariable(ShaderPlatformToPlatformName(ShaderPlatform))->GetValue(Value);
+#else
+		unimplemented();
+#endif
+		return Value;
+	}
 
 private:
-	const TCHAR* Key;
+	FString CVarName;
 	IConsoleVariable* CVar;
-#if WITH_EDITOR
-	/** Set that holds shader platforms for which the ini check was done. */
-	TSet<EShaderPlatform> TestedIni;
-	/** Cached values from the inis. */
-	TMap<EShaderPlatform, Type> CachedValues;
-#endif
 };
 
 /** Returns if ForwardShading is enabled. Only valid for the current platform (otherwise call ITargetPlatform::UsesForwardShading()). */
