- Parallelized the other-platform config reading so early reading by one platform doesn't need to wait on the rest (this was causing a delay in cooking)
#p4v-preflight-copy 18444236

diff --git a/Engine/Source/Runtime/Core/Private/Misc/ConfigCacheIni.cpp b/Engine/Source/Runtime/Core/Private/Misc/ConfigCacheIni.cpp
--- a/Engine/Source/Runtime/Core/Private/Misc/ConfigCacheIni.cpp
+++ b/Engine/Source/Runtime/Core/Private/Misc/ConfigCacheIni.cpp
@@ -4318,6 +4318,10 @@ void FConfigCacheIni::InitializeConfigSystem()
 	// create GConfig
 	GConfig = new FConfigCacheIni(EConfigCacheType::DiskBacked);
 
+#if WITH_EDITOR
+	AsyncInitializeConfigForPlatforms();
+#endif
+
 	// load the main .ini files (unless we're running a program or a gameless UnrealEditor.exe, DefaultEngine.ini is required).
 	const bool bIsGamelessExe = !FApp::HasProjectName();
 	const bool bDefaultEngineIniRequired = !bIsGamelessExe && (GIsGameAgnosticExe || FApp::IsProjectNameEmpty());
@@ -4369,14 +4373,10 @@ void FConfigCacheIni::InitializeConfigSystem()
 	FConfigCacheIni::LoadGlobalIniFile(GLightmassIni, TEXT("Lightmass"));
 #endif
 
-
 	// now we can make use of GConfig
 	GConfig->bIsReadyForUse = true;
 	FCoreDelegates::ConfigReadyForUse.Broadcast();
 
-#if WITH_EDITOR
-	AsyncInitializeConfigForPlatforms();
-#endif
 }
 
 const FString& FConfigCacheIni::GetCustomConfigString()
@@ -5435,72 +5435,79 @@ void DeleteRecordedConfigReadsFromIni()
 
 
 #if ALLOW_OTHER_PLATFORM_CONFIG
-static TMap<FName, FConfigCacheIni> GConfigForPlatform;
+// these are knowingly leaked
+static TMap<FName, FConfigCacheIni*> GConfigForPlatform;
 static FCriticalSection GConfigForPlatformLock;
 
-static TFuture<void>& GetConfigFuture()
-{
-	static TFuture<void> ConfigFuture;
-	return ConfigFuture;
-}
-#endif
-
 #if WITH_EDITOR
+static TMap<FName, TFuture<void>> GPlatformConfigFutures;
+
 void FConfigCacheIni::AsyncInitializeConfigForPlatforms()
 {
-	GetConfigFuture() = Async(EAsyncExecution::ThreadPool, []
+	// make sure we call this super early before anyone else would be calling ForPlatform()
+	check(!GConfig->bIsReadyForUse);
+
+	// pre-create all platforms so that the loop below doesn't reallocate anything in the map
+	for (auto Pair : FDataDrivenPlatformInfoRegistry::GetAllPlatformInfos())
 	{
-		double Start = FPlatformTime::Seconds();
+		GPlatformConfigFutures.Emplace(Pair.Key);
+		GConfigForPlatform.Add(Pair.Key, new FConfigCacheIni(EConfigCacheType::Temporary));
+	}
 
-		for (auto Pair : FDataDrivenPlatformInfoRegistry::GetAllPlatformInfos())
+	for (auto Pair : FDataDrivenPlatformInfoRegistry::GetAllPlatformInfos())
+	{
+		FName PlatformName = Pair.Key;
+		GPlatformConfigFutures[PlatformName] = Async(EAsyncExecution::ThreadPool, [PlatformName]
 		{
-			FName PlatformName = Pair.Key;
-
-			FConfigCacheIni& NewConfig = GConfigForPlatform.Emplace(PlatformName, EConfigCacheType::Temporary);
-
-			NewConfig.InitializeKnownConfigFiles(*PlatformName.ToString(), false);
-		}
+			double Start = FPlatformTime::Seconds();
 
-		UE_LOG(LogConfig, Display, TEXT("Loading platform ini files took %.2f seconds"), FPlatformTime::Seconds() - Start);
-	});
+			FConfigCacheIni* NewConfig = GConfigForPlatform.FindChecked(PlatformName);
+			NewConfig->InitializeKnownConfigFiles(*PlatformName.ToString(), false);
+	
+			UE_LOG(LogConfig, Display, TEXT("Loading %s ini files took %.2f seconds"), *PlatformName.ToString(), FPlatformTime::Seconds() - Start);
+		});
+	}
 }
 
 #endif
+#endif
 
 
 FConfigCacheIni* FConfigCacheIni::ForPlatform(FName PlatformName)
 {
 #if ALLOW_OTHER_PLATFORM_CONFIG
+	check(GConfig->bIsReadyForUse);
+
 	// use GConfig when no platform is specified
 	if (PlatformName == NAME_None)
 	{
 		return GConfig;
 	}
 
-	FScopeLock Lock(&GConfigForPlatformLock);
-
+#if WITH_EDITOR
 	// they are likely already loaded, but just block to make sure
-	if (GetConfigFuture().IsValid())
 	{
-		double Start = FPlatformTime::Seconds();
-
-		GetConfigFuture().Get();
-		GetConfigFuture() = TFuture<void>();
-
-		UE_LOG(LogConfig, Display, TEXT("Blocking on platform ini files took %.2f seconds"), FPlatformTime::Seconds() - Start);
+		GPlatformConfigFutures[PlatformName].Get();
 	}
+#endif
 
-	// look for existing config
-	FConfigCacheIni* PlatformConfig = GConfigForPlatform.Find(PlatformName);
+	// protect against other threads clearing the array, or two threads trying to read in a missing platform at the same time
+	FScopeLock Lock(&GConfigForPlatformLock);
+	FConfigCacheIni* PlatformConfig = GConfigForPlatform.FindRef(PlatformName);
 
-	// read any missing platform configs now, on demand
+	// read any missing platform configs now, on demand (this will happen when WITH_EDITOR is 0)
 	if (PlatformConfig == nullptr)
 	{
-		PlatformConfig = &GConfigForPlatform.Emplace(PlatformName, EConfigCacheType::Temporary);
+		double Start = FPlatformTime::Seconds();
+		
+		PlatformConfig = GConfigForPlatform.Add(PlatformName, new FConfigCacheIni(EConfigCacheType::Temporary));
 		PlatformConfig->InitializeKnownConfigFiles(*PlatformName.ToString(), false);
+
+		UE_LOG(LogConfig, Display, TEXT("Read in platform %s ini files took %.2f seconds"), *PlatformName.ToString(), FPlatformTime::Seconds() - Start);
 	}
 
-	return GConfigForPlatform.Find(PlatformName);
+	return GConfigForPlatform.FindRef(PlatformName);
+
 #else
 	UE_LOG(LogConfig, Error, TEXT("FConfigCacheIni::ForPlatform cannot be called when not in a developer tool"));
 	return nullptr;
@@ -5510,11 +5517,8 @@ FConfigCacheIni* FConfigCacheIni::ForPlatform(FName PlatformName)
 void FConfigCacheIni::ClearOtherPlatformConfigs()
 {
 #if ALLOW_OTHER_PLATFORM_CONFIG
-
+	// this will read in on next call to ForPlatform()
 	FScopeLock Lock(&GConfigForPlatformLock);
-
-	checkf(!GetConfigFuture().IsValid(), TEXT("Trying to clear platform configs while still reading them in is unexpected"));
-
 	GConfigForPlatform.Empty();
 #endif
 }
