Fix crash when exiting the engine (destroying FConfigCacheIni futures).

diff --git a/Engine/Source/Runtime/Core/Private/Misc/ConfigCacheIni.cpp b/Engine/Source/Runtime/Core/Private/Misc/ConfigCacheIni.cpp
--- a/Engine/Source/Runtime/Core/Private/Misc/ConfigCacheIni.cpp
+++ b/Engine/Source/Runtime/Core/Private/Misc/ConfigCacheIni.cpp
@@ -2243,6 +2243,9 @@ void FConfigFile::ProcessPropertyAndWriteForDefaults(int IniCombineThreshold, co
 /*-----------------------------------------------------------------------------
 	FConfigCacheIni
 -----------------------------------------------------------------------------*/
+#if WITH_EDITOR
+static TMap<FName, TFuture<void>> GPlatformConfigFutures;
+#endif
 
 FConfigCacheIni::FConfigCacheIni(EConfigCacheType InType)
 	: bAreFileOperationsDisabled(false)
@@ -2956,6 +2959,14 @@ bool FConfigCacheIni::GetPerObjectConfigSections( const FString& Filename, const
 void FConfigCacheIni::Exit()
 {
 	Flush( 1 );
+
+#if WITH_EDITOR
+	for (auto& PlatformConfigFuture : GPlatformConfigFutures)
+	{
+		PlatformConfigFuture.Value.Get();
+	}
+	GPlatformConfigFutures.Empty();
+#endif
 }
 
 static void DumpFile(FOutputDevice& Ar, const FString& Filename, const FConfigFile& File)
@@ -5440,8 +5451,6 @@ static TMap<FName, FConfigCacheIni*> GConfigForPlatform;
 static FCriticalSection GConfigForPlatformLock;
 
 #if WITH_EDITOR
-static TMap<FName, TFuture<void>> GPlatformConfigFutures;
-
 void FConfigCacheIni::AsyncInitializeConfigForPlatforms()
 {
 	// make sure we call this super early before anyone else would be calling ForPlatform()
