Undo range MallocBinned3 from revision 2 to revision #head because they break alignment on PS5

diff --git a/Engine/Source/Runtime/Core/Private/HAL/MallocBinned3.cpp b/Engine/Source/Runtime/Core/Private/HAL/MallocBinned3.cpp
--- a/Engine/Source/Runtime/Core/Private/HAL/MallocBinned3.cpp
+++ b/Engine/Source/Runtime/Core/Private/HAL/MallocBinned3.cpp
@@ -958,7 +958,7 @@ void* FMallocBinned3::MallocExternal(SIZE_T Size, uint32 Alignment)
 
 	// Only allocate from the small pools if the size is small enough and the alignment isn't crazy large.
 	// With large alignments, we'll waste a lot of memory allocating an entire page, but such alignments are highly unlikely in practice.
-	if ((Size <= BINNED3_MAX_SMALL_POOL_SIZE) )
+	if ((Size <= BINNED3_MAX_SMALL_POOL_SIZE) & (Alignment <= BINNED3_MINIMUM_ALIGNMENT)) // one branch, not two
 	{
 		uint32 PoolIndex = BoundSizeToPoolIndex(Size);
 		FPerThreadFreeBlockLists* Lists = GMallocBinned3PerThreadCaches ? FPerThreadFreeBlockLists::Get() : nullptr;
@@ -1116,7 +1116,7 @@ void* FMallocBinned3::ReallocExternal(void* Ptr, SIZE_T NewSize, uint32 Alignmen
 		check(Ptr); // null is an OS allocation because it will not fall in our VM block
 		uint32 BlockSize = PoolIndexToBlockSize(PoolIndex);
 		if (
-			((NewSize <= BlockSize)) && // one branch, not two
+			((NewSize <= BlockSize) & (Alignment <= BINNED3_MINIMUM_ALIGNMENT)) && // one branch, not two
 			(PoolIndex == 0 || NewSize > PoolIndexToBlockSize(PoolIndex - 1)))
 		{
 #if BINNED3_ALLOCATOR_STATS
@@ -1150,7 +1150,7 @@ void* FMallocBinned3::ReallocExternal(void* Ptr, SIZE_T NewSize, uint32 Alignmen
 	uint32 PoolOSRequestedBytes = Pool->GetOSRequestedBytes();
 	checkf(PoolOSRequestedBytes <= PoolOsBytes, TEXT("FMallocBinned3::ReallocExternal %d %d"), int32(PoolOSRequestedBytes), int32(PoolOsBytes));
 	if (NewSize > PoolOsBytes || // can't fit in the old block
-		(NewSize <= BINNED3_MAX_SMALL_POOL_SIZE) || // can switch to the small block allocator
+		(NewSize <= BINNED3_MAX_SMALL_POOL_SIZE && Alignment <= BINNED3_MINIMUM_ALIGNMENT) || // can switch to the small block allocator
 		Align(NewSize, OsAllocationGranularity) < PoolOsBytes) // we can get some pages back
 	{
 		Mutex.Unlock();
diff --git a/Engine/Source/Runtime/Core/Public/HAL/MallocBinned3.h b/Engine/Source/Runtime/Core/Public/HAL/MallocBinned3.h
--- a/Engine/Source/Runtime/Core/Public/HAL/MallocBinned3.h
+++ b/Engine/Source/Runtime/Core/Public/HAL/MallocBinned3.h
@@ -27,13 +27,11 @@
 #define BINNED3_BASE_PAGE_SIZE				4096			// Minimum "page size" for binned3
 #define BINNED3_MINIMUM_ALIGNMENT_SHIFT		4				// Alignment of blocks, expressed as a shift
 #define BINNED3_MINIMUM_ALIGNMENT			16				// Alignment of blocks
-#ifndef BINNED3_MAX_SMALL_POOL_SIZE
 #if USE_CACHED_PAGE_ALLOCATOR_FOR_LARGE_ALLOCS
 #define BINNED3_MAX_SMALL_POOL_SIZE			(BINNEDCOMMON_MAX_LISTED_SMALL_POOL_SIZE)	// Maximum medium block size
 #else
 #define BINNED3_MAX_SMALL_POOL_SIZE			(128 * 1024)	// Maximum medium block size
 #endif
-#endif
 #define BINNED3_SMALL_POOL_COUNT			(BINNEDCOMMON_NUM_LISTED_SMALL_POOLS + (BINNED3_MAX_SMALL_POOL_SIZE - BINNEDCOMMON_MAX_LISTED_SMALL_POOL_SIZE) / BINNED3_BASE_PAGE_SIZE)
 
 #if USE_512MB_MAX_MEMORY_PER_BLOCK_SIZE
@@ -71,10 +69,10 @@
 #define BINNED3_ALLOW_RUNTIME_TWEAKING 0
 #if BINNED3_ALLOW_RUNTIME_TWEAKING
 	extern CORE_API int32 GMallocBinned3PerThreadCaches;
-	extern CORE_API int32 GMallocBinned3BundleSize;
-	extern CORE_API int32 GMallocBinned3BundleCount;
-	extern CORE_API int32 GMallocBinned3MaxBundlesBeforeRecycle;
-	extern CORE_API int32 GMallocBinned3AllocExtra;
+	extern CORE_API int32 GMallocBinned3BundleSize = DEFAULT_GMallocBinned3BundleSize;
+	extern CORE_API int32 GMallocBinned3BundleCount = DEFAULT_GMallocBinned3BundleCount;
+	extern CORE_API int32 GMallocBinned3MaxBundlesBeforeRecycle = BINNED3_MAX_GMallocBinned3MaxBundlesBeforeRecycle;
+	extern CORE_API int32 GMallocBinned3AllocExtra = DEFAULT_GMallocBinned3AllocExtra;
 #else
 	#define GMallocBinned3PerThreadCaches DEFAULT_GMallocBinned3PerThreadCaches
 	#define GMallocBinned3BundleSize DEFAULT_GMallocBinned3BundleSize
@@ -528,18 +526,10 @@ public:
 	FORCEINLINE virtual void* Malloc(SIZE_T Size, uint32 Alignment) override
 	{
 		void* Result = nullptr;
-
-
-		if (Alignment > BINNED3_MINIMUM_ALIGNMENT)
-		{
-			Size = Align(Size, Alignment);
-		}
-
-
-
+	
 		// Only allocate from the small pools if the size is small enough and the alignment isn't crazy large.
 		// With large alignments, we'll waste a lot of memory allocating an entire page, but such alignments are highly unlikely in practice.
-		if ((Size <= BINNED3_MAX_SMALL_POOL_SIZE) )
+		if ((Size <= BINNED3_MAX_SMALL_POOL_SIZE) & (Alignment <= BINNED3_MINIMUM_ALIGNMENT)) // one branch, not two
 		{
 			FPerThreadFreeBlockLists* Lists = GMallocBinned3PerThreadCaches ? FPerThreadFreeBlockLists::Get() : nullptr;
 			if (Lists)
@@ -565,12 +555,7 @@ public:
 	}
 	FORCEINLINE virtual void* Realloc(void* Ptr, SIZE_T NewSize, uint32 Alignment) override
 	{
-		if (Alignment > BINNED3_MINIMUM_ALIGNMENT)
-		{
-			NewSize = Align(NewSize, Alignment);
-		}
-
-		if (NewSize <= BINNED3_MAX_SMALL_POOL_SIZE )
+		if (NewSize <= BINNED3_MAX_SMALL_POOL_SIZE && Alignment <= BINNED3_MINIMUM_ALIGNMENT) // one branch, not two
 		{
 			FPerThreadFreeBlockLists* Lists = GMallocBinned3PerThreadCaches ? FPerThreadFreeBlockLists::Get() : nullptr;
 
@@ -665,15 +650,10 @@ public:
 
 	FORCEINLINE virtual SIZE_T QuantizeSize(SIZE_T Count, uint32 Alignment) override
 	{
-		if (Alignment > BINNED3_MINIMUM_ALIGNMENT)
-		{
-			Count = Align(Count, Alignment);
-		}
-
 		static_assert(DEFAULT_ALIGNMENT <= BINNED3_MINIMUM_ALIGNMENT, "DEFAULT_ALIGNMENT is assumed to be zero"); // used below
 		checkSlow((Alignment & (Alignment - 1)) == 0); // Check the alignment is a power of two
 		SIZE_T SizeOut;
-		if ((Count <= BINNED3_MAX_SMALL_POOL_SIZE))
+		if ((Count <= BINNED3_MAX_SMALL_POOL_SIZE) & (Alignment <= BINNED3_MINIMUM_ALIGNMENT)) // one branch, not two
 		{
 			SizeOut = PoolIndexToBlockSize(BoundSizeToPoolIndex(Count));
 		}
