REL - extra debug functionallity for hashtable

diff --git a/Engine/Source/Runtime/Core/Public/Experimental/Containers/RobinHoodHashTable.h b/Engine/Source/Runtime/Core/Public/Experimental/Containers/RobinHoodHashTable.h
--- a/Engine/Source/Runtime/Core/Public/Experimental/Containers/RobinHoodHashTable.h
+++ b/Engine/Source/Runtime/Core/Public/Experimental/Containers/RobinHoodHashTable.h
@@ -190,6 +190,19 @@ namespace RobinHoodHashTable_Private
 			}
 		}
 
+		bool Contains(IndexType Index) const
+		{
+			for (int i = FreeList.Num() - 1; i > -0; i--)
+			{
+				const FSpan& Span = FreeList[i];
+				if (Index >= Span.Start && Index <= Span.End)
+				{
+					return true;
+				}
+			}
+			return false;
+		}
+
 		void Empty()
 		{
 			FreeList.Reset(1);
@@ -346,6 +359,9 @@ namespace RobinHoodHashTable_Private
 				FreeList.Push(Index);
 				Hashes[Index] = FHashType();
 				KeyVals[Index].~KeyValueType();
+#if RUN_HASHTABLE_CONCURENCY_CHECKS
+				memset(&KeyVals[Index], 0, sizeof(KeyValueType));
+#endif
 				CHECK_CONCURRENT_ACCESS(FPlatformAtomics::InterlockedDecrement(&ConcurrentReaders) == 0);
 				CHECK_CONCURRENT_ACCESS(FPlatformAtomics::InterlockedDecrement(&ConcurrentWriters) == 0);
 				//TODO shrink KeyValue Array.
@@ -355,13 +371,28 @@ namespace RobinHoodHashTable_Private
 				// - remove the entry from the FreeList (probably remove the entire span at the end of iteration, as its tail would need to be moved)
 			}
 
+			inline bool Contains(IndexType Index) const
+			{
+				return Index < (IndexType)KeyVals.Num() && !FreeList.Contains(Index);
+			}
+
 			inline const KeyValueType& Get(IndexType Index) const
 			{
+				CHECK_CONCURRENT_ACCESS(ConcurrentWriters == 0);
+				CHECK_CONCURRENT_ACCESS(FPlatformAtomics::InterlockedIncrement(&ConcurrentReaders) >= 1);
+				CHECK_CONCURRENT_ACCESS(Contains(Index));
+				CHECK_CONCURRENT_ACCESS(FPlatformAtomics::InterlockedDecrement(&ConcurrentReaders) >= 0);
+				CHECK_CONCURRENT_ACCESS(ConcurrentWriters == 0);
 				return KeyVals[Index];
 			}
 
 			inline KeyValueType& Get(IndexType Index)
 			{
+				CHECK_CONCURRENT_ACCESS(ConcurrentWriters == 0);
+				CHECK_CONCURRENT_ACCESS(FPlatformAtomics::InterlockedIncrement(&ConcurrentReaders) >= 1);
+				CHECK_CONCURRENT_ACCESS(Contains(Index));
+				CHECK_CONCURRENT_ACCESS(FPlatformAtomics::InterlockedDecrement(&ConcurrentReaders) >= 0);
+				CHECK_CONCURRENT_ACCESS(ConcurrentWriters == 0);
 				return KeyVals[Index];
 			}
 
@@ -786,6 +817,11 @@ namespace RobinHoodHashTable_Private
 			return KeyValueData.Get(Id.GetIndex()).GetElement();
 		}
 
+		inline bool ContainsElementId(FHashElementId Id) const
+		{
+			return KeyValueData.Contains(Id.GetIndex());
+		}
+
 		inline FHashElementId FindIdByHash(const FHashType HashValue, const KeyType& ComparableKey) const
 		{
 			CHECK_CONCURRENT_ACCESS(ConcurrentWriters == 0);
