REL - partially undoing removal of Readcyclecounter as it caused a large pessimization on x86 platforms as workers now spin for several milliseconds

diff --git a/Engine/Source/Runtime/Core/Public/GenericPlatform/GenericPlatformProcess.h b/Engine/Source/Runtime/Core/Public/GenericPlatform/GenericPlatformProcess.h
--- a/Engine/Source/Runtime/Core/Public/GenericPlatform/GenericPlatformProcess.h
+++ b/Engine/Source/Runtime/Core/Public/GenericPlatform/GenericPlatformProcess.h
@@ -756,9 +756,33 @@ struct CORE_API FGenericPlatformProcess
 	*/
 	static FORCEINLINE void YieldCycles(uint64 Cycles)
 	{
+#if PLATFORM_CPU_X86_FAMILY
+		auto ReadCycleCounter = []()
+		{
+#if defined(_MSC_VER)
+			return __rdtsc();
+#elif PLATFORM_APPLE
+			return mach_absolute_time();
+#elif __has_builtin(__builtin_readcyclecounter)
+			return __builtin_readcyclecounter();
+#else
+#	error Unsupported architecture!
+#endif
+		};
+
+		uint64 start = ReadCycleCounter();
+		//some 32bit implementations return 0 for __builtin_readcyclecounter just to be on the safe side we protect against this.
+		Cycles = start != 0 ? Cycles : 0;
+		do
+		{
+			Yield();
+		} while((ReadCycleCounter() - start) < Cycles);
+#else
+		// We can't read cycle counter from user mode on these platform
 		for (uint64 i = 0; i < Cycles; i++)
 		{
 			Yield();
 		}
+#endif
 	}
 };
