REL - Automatically Inherit Task Properties

diff --git a/Engine/Source/Runtime/Core/Private/Async/Fundamental/Scheduler.cpp b/Engine/Source/Runtime/Core/Private/Async/Fundamental/Scheduler.cpp
--- a/Engine/Source/Runtime/Core/Private/Async/Fundamental/Scheduler.cpp
+++ b/Engine/Source/Runtime/Core/Private/Async/Fundamental/Scheduler.cpp
@@ -241,15 +241,27 @@ namespace LowLevelTasks
 		return ~0;
 	}
 
-	void FTask::PropagateUserData()
+	void FTask::InheritParentData(const TCHAR*& DebugName, ETaskPriority& Priority)
 	{
 		const FTask* ActiveTask = FSchedulerTls::GetActiveTask();
 		if (ActiveTask != nullptr)
 		{
+			if(DebugName == nullptr)
+			{
+				DebugName = ActiveTask->GetDebugName();
+			}
+			if (Priority == ETaskPriority::Inherit)
+			{
+				Priority = ActiveTask->GetPriority();
+			}
 			UserData = ActiveTask->GetUserData();
 		}
 		else
 		{
+			if (Priority == ETaskPriority::Inherit)
+			{
+				Priority = ETaskPriority::Default;
+			}
 			UserData = nullptr;
 		}
 	}
diff --git a/Engine/Source/Runtime/Core/Public/Async/Fundamental/Task.h b/Engine/Source/Runtime/Core/Public/Async/Fundamental/Task.h
--- a/Engine/Source/Runtime/Core/Public/Async/Fundamental/Task.h
+++ b/Engine/Source/Runtime/Core/Public/Async/Fundamental/Task.h
@@ -20,7 +20,8 @@ namespace LowLevelTasks
 		BackgroundHigh = ForegroundCount,
 		BackgroundNormal,
 		BackgroundLow,
-		Count
+		Inherit,
+		Count = Inherit,
 	};
 
 	inline const TCHAR* ToString(ETaskPriority Priority)
@@ -133,7 +134,7 @@ namespace LowLevelTasks
 			uintptr_t bAllowBusyWaiting : 1;
 			
 		public:
-			FPackedData() : FPackedData(nullptr,  ETaskPriority::Count, ETaskState::Completed, true)
+			FPackedData() : FPackedData(nullptr,  ETaskPriority::Inherit, ETaskState::Completed, true)
 			{
 				static_assert(!PLATFORM_32BITS, "32bit Platforms are not supported");
 				static_assert(uintptr_t(ETaskPriority::Count) <= (1ull << 3), "Not enough bits to store ETaskPriority");
@@ -261,7 +262,7 @@ namespace LowLevelTasks
 		inline bool TryPrepareLaunch();
 		//after calling this function the task can be considered dead
 		inline void ExecuteTask();
-		CORE_API void PropagateUserData();
+		CORE_API void InheritParentData(const TCHAR*& DebugName, ETaskPriority& Priority);
 		inline bool AllowBusyWaiting() const;
 	};
 
@@ -292,7 +293,7 @@ namespace LowLevelTasks
 				return true;
 			}
 		};
-		PropagateUserData();
+		InheritParentData(InDebugName, InPriority);
 		PackedData.store(FPackedData(InDebugName, InPriority, ETaskState::Ready, bAllowBusyWaiting), std::memory_order_release);
 	}
 
@@ -317,7 +318,7 @@ namespace LowLevelTasks
 				return true;
 			}
 		};
-		PropagateUserData();
+		InheritParentData(InDebugName, InPriority);
 		PackedData.store(FPackedData(InDebugName, InPriority, ETaskState::Ready, bAllowBusyWaiting), std::memory_order_release);
 	}
 
