REL - fix Mac redefinition error

diff --git a/Engine/Source/Runtime/Core/Public/HAL/MallocBinned3.h b/Engine/Source/Runtime/Core/Public/HAL/MallocBinned3.h
--- a/Engine/Source/Runtime/Core/Public/HAL/MallocBinned3.h
+++ b/Engine/Source/Runtime/Core/Public/HAL/MallocBinned3.h
@@ -27,11 +27,13 @@
 #define BINNED3_BASE_PAGE_SIZE				4096			// Minimum "page size" for binned3
 #define BINNED3_MINIMUM_ALIGNMENT_SHIFT		4				// Alignment of blocks, expressed as a shift
 #define BINNED3_MINIMUM_ALIGNMENT			16				// Alignment of blocks
+#ifndef BINNED3_MAX_SMALL_POOL_SIZE
 #if USE_CACHED_PAGE_ALLOCATOR_FOR_LARGE_ALLOCS
 #define BINNED3_MAX_SMALL_POOL_SIZE			(BINNEDCOMMON_MAX_LISTED_SMALL_POOL_SIZE)	// Maximum medium block size
 #else
 #define BINNED3_MAX_SMALL_POOL_SIZE			(128 * 1024)	// Maximum medium block size
 #endif
+#endif
 #define BINNED3_SMALL_POOL_COUNT			(BINNEDCOMMON_NUM_LISTED_SMALL_POOLS + (BINNED3_MAX_SMALL_POOL_SIZE - BINNEDCOMMON_MAX_LISTED_SMALL_POOL_SIZE) / BINNED3_BASE_PAGE_SIZE)
 
 #if USE_512MB_MAX_MEMORY_PER_BLOCK_SIZE
