Quick workaround for check triggering when loading certain levels

#rnx

diff --git a/Engine/Source/Runtime/Core/Private/HAL/ConsoleManager.cpp b/Engine/Source/Runtime/Core/Private/HAL/ConsoleManager.cpp
--- a/Engine/Source/Runtime/Core/Private/HAL/ConsoleManager.cpp
+++ b/Engine/Source/Runtime/Core/Private/HAL/ConsoleManager.cpp
@@ -553,7 +553,7 @@ private: // ----------------------------------------------------
 	TMap<FName, TSharedPtr<IConsoleVariable> > PlatformValues;
 	virtual IConsoleVariable* GetPlatformValueVariable(FName PlatformName)
 	{
-		check(IsInGameThread());
+		//check(IsInGameThread());
 		if (!PlatformValues.Contains(PlatformName))
 		{
 			FString ConfigValue;
