- Better fix for CVars crashing when looking up other platform value in rendering thread

diff --git a/Engine/Source/Runtime/Core/Private/HAL/ConsoleManager.cpp b/Engine/Source/Runtime/Core/Private/HAL/ConsoleManager.cpp
--- a/Engine/Source/Runtime/Core/Private/HAL/ConsoleManager.cpp
+++ b/Engine/Source/Runtime/Core/Private/HAL/ConsoleManager.cpp
@@ -6,6 +6,7 @@ ConsoleManager.cpp: console command handling
 
 #include "HAL/ConsoleManager.h"
 #include "Misc/ScopeLock.h"
+#include "Misc/ScopeRWLock.h"
 #include "Misc/Paths.h"
 #include "Stats/Stats.h"
 #include "Misc/ConfigCacheIni.h"
@@ -551,9 +552,11 @@ private: // ----------------------------------------------------
 	}
 
 	TMap<FName, TSharedPtr<IConsoleVariable> > PlatformValues;
+	FRWLock PlatformValuesLock;
 	virtual IConsoleVariable* GetPlatformValueVariable(FName PlatformName)
 	{
-		//check(IsInGameThread());
+		// cheap lock here, contention is very rare
+		FRWScopeLock Lock(PlatformValuesLock, SLT_Write);
 		if (!PlatformValues.Contains(PlatformName))
 		{
 			FString ConfigValue;
