Speculative fix for how a FFileHandleWindows might change from valid to invalid. The unwritten contract is that FFileHandleWindows instances returned from OpenRead are either null or are IsValid() == true.
Also fix FArchiveFileReaderGeneric::FlushCache from possibly dereferning a null handle to do the seek.
#rnx

diff --git a/Engine/Source/Runtime/Core/Private/HAL/FileManagerGeneric.cpp b/Engine/Source/Runtime/Core/Private/HAL/FileManagerGeneric.cpp
--- a/Engine/Source/Runtime/Core/Private/HAL/FileManagerGeneric.cpp
+++ b/Engine/Source/Runtime/Core/Private/HAL/FileManagerGeneric.cpp
@@ -903,15 +903,16 @@ void FArchiveFileReaderGeneric::Serialize( void* V, int64 Length )
 
 void FArchiveFileReaderGeneric::FlushCache()
 {
+	BufferArray.Empty();
+	// After clearing BufferArray, we need to set BufferBase and LowLevel to Pos. They may have been up to BufferSize away.
+	BufferBase = Pos;
+
 	if (Handle.IsValid())
 	{
 		Handle->ShrinkBuffers();
+		// After clearing BufferArray, we need to set BufferBase and LowLevel to Pos. They may have been up to BufferSize away.
+		SeekLowLevel(Pos);
 	}
-
-	BufferArray.Empty();
-	// Make sure the underlying system file has the right position for any read that comes after a flushed cache.
-	SeekLowLevel(Pos);
-	BufferBase = Pos;
 }
 
 FArchiveFileWriterGeneric::FArchiveFileWriterGeneric( IFileHandle* InHandle, const TCHAR* InFilename, int64 InPos, uint32 InBufferSize, uint32 InFlags )
diff --git a/Engine/Source/Runtime/Core/Private/Windows/WindowsPlatformFile.cpp b/Engine/Source/Runtime/Core/Private/Windows/WindowsPlatformFile.cpp
--- a/Engine/Source/Runtime/Core/Private/Windows/WindowsPlatformFile.cpp
+++ b/Engine/Source/Runtime/Core/Private/Windows/WindowsPlatformFile.cpp
@@ -503,6 +503,12 @@ public:
 			TRACE_PLATFORMFILE_BEGIN_REOPEN(Handle);
 			HANDLE NewFileHandle = ReOpenFile(Handle, DesiredAccess, ShareMode, Flags);
 			TRACE_PLATFORMFILE_END_REOPEN(NewFileHandle);
+			if (NewFileHandle == INVALID_HANDLE_VALUE)
+			{
+				// We are not allowed to change this->IsValid() from true to false.
+				// If the Reopen fails, keep the old FileHandle
+				return;
+			}
 
 			TRACE_PLATFORMFILE_BEGIN_CLOSE(Handle);
 			BOOL CloseResult = CloseHandle(Handle);
@@ -519,6 +525,7 @@ public:
 			(void)CloseResult;
 #endif
 			Handle = NewFileHandle;
+			check(IsValid());
 		}
 	}
 };
@@ -757,6 +764,7 @@ public:
 		if (Seek(NewSize) && UpdatedNonOverlappedPos() && SetEndOfFile(FileHandle) != 0)
 		{
 			UpdateFileSize();
+			check(IsValid());
 			return true;
 		}
 		return false;
@@ -768,6 +776,12 @@ public:
 			TRACE_PLATFORMFILE_BEGIN_REOPEN(FileHandle);
 			HANDLE NewFileHandle = ReOpenFile(FileHandle, DesiredAccess, ShareMode, Flags);
 			TRACE_PLATFORMFILE_END_REOPEN(NewFileHandle);
+			if (NewFileHandle == INVALID_HANDLE_VALUE)
+			{
+				// We are not allowed to change this->IsValid() from true to false.
+				// If the Reopen fails, keep the old FileHandle
+				return;
+			}
 
 			TRACE_PLATFORMFILE_BEGIN_CLOSE(FileHandle);
 			BOOL CloseResult = CloseHandle(FileHandle);
@@ -784,6 +798,7 @@ public:
 			(void)CloseResult;
 #endif
 			FileHandle = NewFileHandle;
+			check(IsValid());
 		}
 	}
 };
