REL - reduce allocation pressure during local queue un/registration

diff --git a/Engine/Source/Runtime/Core/Public/Async/Fundamental/LocalQueue.h b/Engine/Source/Runtime/Core/Public/Async/Fundamental/LocalQueue.h
--- a/Engine/Source/Runtime/Core/Public/Async/Fundamental/LocalQueue.h
+++ b/Engine/Source/Runtime/Core/Public/Async/Fundamental/LocalQueue.h
@@ -221,10 +221,6 @@ private:
 		TArray<TLocalQueue*, TInlineAllocator<32>> LocalQueues;
 		TLocalQueue* RemovedQueue = nullptr;
 
-		FLocalQueueCollection(FLocalQueueCollection* Previous) : LocalQueues(Previous->LocalQueues)
-		{
-		}
-
 		FLocalQueueCollection() = default;
 
 		~FLocalQueueCollection()
@@ -408,14 +404,14 @@ private:
 			NumActiveWorkers[QueueType == ELocalQueueType::EBackground].fetch_add(1, std::memory_order_relaxed);
 		}
 
+		FLocalQueueCollection* Copy = new FLocalQueueCollection();
 		while(true)
 		{
 			FLocalQueueCollection* Previous = Hazard.Get();
-			FLocalQueueCollection* Copy = new FLocalQueueCollection(Previous);
+			Copy->LocalQueues = Previous->LocalQueues;	
 			Copy->LocalQueues.Add(QueueToAdd);
 			if (!QueueCollection.compare_exchange_strong(Previous, Copy, std::memory_order_release, std::memory_order_relaxed))
 			{
-				delete Copy;
 				continue;
 			}
 			HazardsCollection.Delete(Previous);
@@ -432,15 +428,14 @@ private:
 			NumActiveWorkers[QueueType == ELocalQueueType::EBackground].fetch_sub(1, std::memory_order_relaxed);
 		}
 
+		FLocalQueueCollection* Copy = new FLocalQueueCollection();
 		while(true)
 		{
 			FLocalQueueCollection* Previous = Hazard.Get();
-			FLocalQueueCollection* Copy = new FLocalQueueCollection(Previous);
-			int NumRemoved = Copy->LocalQueues.Remove(QueueToRemove);
-			checkSlow(NumRemoved == 1);
+			Copy->LocalQueues = Previous->LocalQueues;	
+			verifySlow(Copy->LocalQueues.Remove(QueueToRemove) == 1);
 			if (!QueueCollection.compare_exchange_strong(Previous, Copy, std::memory_order_release, std::memory_order_relaxed))
 			{
-				delete Copy;
 				continue;
 			}
 			if (!WorkerOwned)
