REL - Allocator support for ChunkedArray

diff --git a/Engine/Source/Runtime/Core/Public/Containers/ChunkedArray.h b/Engine/Source/Runtime/Core/Public/Containers/ChunkedArray.h
--- a/Engine/Source/Runtime/Core/Public/Containers/ChunkedArray.h
+++ b/Engine/Source/Runtime/Core/Public/Containers/ChunkedArray.h
@@ -47,7 +47,7 @@ namespace UE4ChunkedArray_Private
 }
 
 /** An array that uses multiple allocations to avoid allocation failure due to fragmentation. */
-template<typename InElementType, uint32 TargetBytesPerChunk = 16384 >
+template<typename InElementType, uint32 TargetBytesPerChunk = 16384, typename AllocatorType = FDefaultAllocator >
 class TChunkedArray
 {
 	using ElementType = InElementType;
@@ -270,7 +270,7 @@ public:
 
 protected:
 
-	friend struct TContainerTraits<TChunkedArray<ElementType, TargetBytesPerChunk>>;
+	friend struct TContainerTraits<TChunkedArray<ElementType, TargetBytesPerChunk, AllocatorType>>;
 
 	enum { NumElementsPerChunk = TargetBytesPerChunk / sizeof(ElementType) };
 
@@ -282,7 +282,7 @@ protected:
 	};
 
 	/** The chunks of the array's elements. */
-	typedef TIndirectArray<FChunk> ChunksType;
+	typedef TIndirectArray<FChunk, AllocatorType> ChunksType;
 	ChunksType Chunks;
 
 	/** The number of elements in the array. */
@@ -329,14 +329,15 @@ public:
 };
 
 
-template <typename ElementType, uint32 TargetBytesPerChunk>
-struct TContainerTraits<TChunkedArray<ElementType, TargetBytesPerChunk> > : public TContainerTraitsBase<TChunkedArray<ElementType, TargetBytesPerChunk> >
+template <typename ElementType, uint32 TargetBytesPerChunk, typename AllocatorType>
+struct TContainerTraits<TChunkedArray<ElementType, TargetBytesPerChunk, AllocatorType> > : public TContainerTraitsBase<TChunkedArray<ElementType, TargetBytesPerChunk, AllocatorType> >
 {
-	enum { MoveWillEmptyContainer = TContainerTraits<typename TChunkedArray<ElementType, TargetBytesPerChunk>::ChunksType>::MoveWillEmptyContainer };
+	enum { MoveWillEmptyContainer = TContainerTraits<typename TChunkedArray<ElementType, TargetBytesPerChunk, AllocatorType>::ChunksType>::MoveWillEmptyContainer };
 };
 
 
-template <typename T,uint32 TargetBytesPerChunk> void* operator new( size_t Size, TChunkedArray<T,TargetBytesPerChunk>& ChunkedArray )
+template <typename T,uint32 TargetBytesPerChunk, typename AllocatorType> 
+void* operator new( size_t Size, TChunkedArray<T,TargetBytesPerChunk, AllocatorType>& ChunkedArray )
 {
 	check(Size == sizeof(T));
 	const int32 Index = ChunkedArray.Add(1);
