- Fixed a race condition with FPaths::ProjectDir being called the first time from a thread (it caches value on first run)

diff --git a/Engine/Source/Runtime/Core/Private/Misc/ConfigCacheIni.cpp b/Engine/Source/Runtime/Core/Private/Misc/ConfigCacheIni.cpp
--- a/Engine/Source/Runtime/Core/Private/Misc/ConfigCacheIni.cpp
+++ b/Engine/Source/Runtime/Core/Private/Misc/ConfigCacheIni.cpp
@@ -5453,6 +5453,10 @@ static FCriticalSection GConfigForPlatformLock;
 #if WITH_EDITOR
 void FConfigCacheIni::AsyncInitializeConfigForPlatforms()
 {
+	// this function has an internal cache that is created on first call, and if that happens on a thread, crashes can happen,
+	// so the trivial solution is to call it now to enforce the cache is created before any threads read from it
+	FPaths::ProjectDir();
+
 	// make sure we call this super early before anyone else would be calling ForPlatform()
 	check(!GConfig->bIsReadyForUse);
 
