Use of "Allow/Hide Windows Types" headers can unexpectedly break Unity builds.

Some third party includes will conditionally define TRUE and FALSE to appropriate values. ICU is one such example that does this. Consider the case where such a library is included in a unity blob both before and after an unrelated .cpp file that makes use of an Allow/HideWindowsTypes.h pair. The second include will be skipped due to header guards. However TRUE and FALSE have been undefined and subtle compile errors occur.

#rnx

diff --git a/Engine/Source/Runtime/Core/Private/HAL/ThreadingBase.cpp b/Engine/Source/Runtime/Core/Private/HAL/ThreadingBase.cpp
--- a/Engine/Source/Runtime/Core/Private/HAL/ThreadingBase.cpp
+++ b/Engine/Source/Runtime/Core/Private/HAL/ThreadingBase.cpp
@@ -15,10 +15,6 @@
 #include "Async/Fundamental/Scheduler.h"
 #include "Tasks/Pipe.h"
 
-#if PLATFORM_WINDOWS
-#include "Windows/AllowWindowsPlatformTypes.h"
-#endif
-
 #include <atomic>
 
 #ifndef IS_RUNNING_GAMETHREAD_ON_EXTERNAL_THREAD
@@ -1562,7 +1558,3 @@ FRunnableThread* FForkProcessHelper::CreateForkableThread(class FRunnable* InRun
 
 	return NewThread;
 }
-
-#if PLATFORM_WINDOWS
-#include "Windows/HideWindowsPlatformTypes.h"
-#endif
