Non-unity fix.

#rnx

diff --git a/Engine/Source/Runtime/Core/Private/HAL/ThreadingBase.cpp b/Engine/Source/Runtime/Core/Private/HAL/ThreadingBase.cpp
--- a/Engine/Source/Runtime/Core/Private/HAL/ThreadingBase.cpp
+++ b/Engine/Source/Runtime/Core/Private/HAL/ThreadingBase.cpp
@@ -15,6 +15,10 @@
 #include "Async/Fundamental/Scheduler.h"
 #include "Tasks/Pipe.h"
 
+#if PLATFORM_WINDOWS
+#include "Microsoft/MinimalWindowsApi.h"
+#endif
+
 #include <atomic>
 
 #ifndef IS_RUNNING_GAMETHREAD_ON_EXTERNAL_THREAD
@@ -83,7 +87,7 @@ FTaskTagScope::FTaskTagScope(bool InTagOnlyIfNone, ETaskTag InTag) : Tag(InTag),
 		// the PE loading process prior to calling the thread's entry point. This
 		// results static initialisation unexpectedly happening on RenderDoc's injection
 		// thread and the ensure below fails.
-		static bool bRenderDocDetected = (::LoadLibraryW(L"renderdoc.dll") != nullptr);
+		static bool bRenderDocDetected = (Windows::LoadLibraryW(L"renderdoc.dll") != nullptr);
 		if (!bRenderDocDetected)
 #	endif
 		{
