Allow non-desktop platforms the option to have default project directories relative to the executable, when one is not specified on the command line.

#rnx

diff --git a/Engine/Source/Runtime/Core/Private/GenericPlatform/GenericPlatformMisc.cpp b/Engine/Source/Runtime/Core/Private/GenericPlatform/GenericPlatformMisc.cpp
--- a/Engine/Source/Runtime/Core/Private/GenericPlatform/GenericPlatformMisc.cpp
+++ b/Engine/Source/Runtime/Core/Private/GenericPlatform/GenericPlatformMisc.cpp
@@ -44,6 +44,10 @@
 	THIRD_PARTY_INCLUDES_END
 #endif
 
+#if !defined(PLATFORM_PROJECT_DIR_RELATIVE_TO_EXECUTABLE)
+	#define PLATFORM_PROJECT_DIR_RELATIVE_TO_EXECUTABLE PLATFORM_DESKTOP
+#endif
+
 DEFINE_LOG_CATEGORY_STATIC(LogGenericPlatformMisc, Log, All);
 
 
@@ -1282,7 +1286,7 @@ const TCHAR* FGenericPlatformMisc::ProjectDir()
 				}
 				else
 				{
-#if !PLATFORM_DESKTOP
+#if !PLATFORM_PROJECT_DIR_RELATIVE_TO_EXECUTABLE
 					ProjectDir = FString::Printf(TEXT("../../../%s/"), FApp::GetProjectName());
 #else
 					// This assumes the game executable is in <GAME>/Binaries/<PLATFORM>
