- Delayed the threaded config loading until its safer to do so. The old solution wasn't good enough
#okforgiyhub public

diff --git a/Engine/Source/Runtime/Core/Private/Misc/ConfigCacheIni.cpp b/Engine/Source/Runtime/Core/Private/Misc/ConfigCacheIni.cpp
--- a/Engine/Source/Runtime/Core/Private/Misc/ConfigCacheIni.cpp
+++ b/Engine/Source/Runtime/Core/Private/Misc/ConfigCacheIni.cpp
@@ -4174,7 +4174,11 @@ bool FConfigCacheIni::InitializeKnownConfigFiles(const TCHAR* PlatformName, bool
 		}
 	}
 
-	bIsReadyForUse = true;
+	// Gconfig set itself ready for use later on
+	if (this != GConfig)
+	{
+		bIsReadyForUse = true;
+	}
 
 	return bEngineConfigCreated;
 }
@@ -4329,10 +4333,6 @@ void FConfigCacheIni::InitializeConfigSystem()
 	// create GConfig
 	GConfig = new FConfigCacheIni(EConfigCacheType::DiskBacked);
 
-#if WITH_EDITOR
-	AsyncInitializeConfigForPlatforms();
-#endif
-
 	// load the main .ini files (unless we're running a program or a gameless UnrealEditor.exe, DefaultEngine.ini is required).
 	const bool bIsGamelessExe = !FApp::HasProjectName();
 	const bool bDefaultEngineIniRequired = !bIsGamelessExe && (GIsGameAgnosticExe || FApp::IsProjectNameEmpty());
@@ -4386,8 +4386,12 @@ void FConfigCacheIni::InitializeConfigSystem()
 
 	// now we can make use of GConfig
 	GConfig->bIsReadyForUse = true;
-	FCoreDelegates::ConfigReadyForUse.Broadcast();
 
+#if WITH_EDITOR
+	AsyncInitializeConfigForPlatforms();
+#endif
+
+	FCoreDelegates::ConfigReadyForUse.Broadcast();
 }
 
 const FString& FConfigCacheIni::GetCustomConfigString()
@@ -5453,13 +5457,6 @@ static FCriticalSection GConfigForPlatformLock;
 #if WITH_EDITOR
 void FConfigCacheIni::AsyncInitializeConfigForPlatforms()
 {
-	// this function has an internal cache that is created on first call, and if that happens on a thread, crashes can happen,
-	// so the trivial solution is to call it now to enforce the cache is created before any threads read from it
-	FPaths::ProjectDir();
-
-	// make sure we call this super early before anyone else would be calling ForPlatform()
-	check(!GConfig->bIsReadyForUse);
-
 	// pre-create all platforms so that the loop below doesn't reallocate anything in the map
 	for (auto Pair : FDataDrivenPlatformInfoRegistry::GetAllPlatformInfos())
 	{
