- Further fix for ProjectDir crash

diff --git a/Engine/Source/Runtime/Core/Private/Misc/ConfigCacheIni.cpp b/Engine/Source/Runtime/Core/Private/Misc/ConfigCacheIni.cpp
--- a/Engine/Source/Runtime/Core/Private/Misc/ConfigCacheIni.cpp
+++ b/Engine/Source/Runtime/Core/Private/Misc/ConfigCacheIni.cpp
@@ -5457,6 +5457,9 @@ static FCriticalSection GConfigForPlatformLock;
 #if WITH_EDITOR
 void FConfigCacheIni::AsyncInitializeConfigForPlatforms()
 {
+	// since GConfig just became ready, this will cache something new on first use, and we don't want two threads going at it
+	FPaths::ProjectDir();
+
 	// pre-create all platforms so that the loop below doesn't reallocate anything in the map
 	for (auto Pair : FDataDrivenPlatformInfoRegistry::GetAllPlatformInfos())
 	{
