- Fixed up some issues with other platform cvars:
  - allow for FConsoleVariableRefs to have other platforms
  - dpreload now clears out cached cvar values and also expanded cvars inside the deviceprofiles

diff --git a/Engine/Source/Runtime/Core/Private/HAL/ConsoleManager.cpp b/Engine/Source/Runtime/Core/Private/HAL/ConsoleManager.cpp
--- a/Engine/Source/Runtime/Core/Private/HAL/ConsoleManager.cpp
+++ b/Engine/Source/Runtime/Core/Private/HAL/ConsoleManager.cpp
@@ -475,22 +475,60 @@ static bool GetConfigValueFromRuntimeSources(FName PlatformName, IConsoleVariabl
 	// by this point, the value will be set, with the DefaultValue if nothing else from the visit function
 	return true;
 }
+
 #endif
 
+
+// an intermediate class between specific typed CVars and FConsoleVariableBase to handle looking up other platform's cvars and caching them
+// if ALLOW_OTHER_PLATFORM_CONFIG is 0, then this is a pass-through class that does nothing
+template <class T>
+class FOtherPlatformValueHelper : public FConsoleVariableBase
+{
+public:
+
+	FOtherPlatformValueHelper(const T& DefaultValue, const TCHAR* Help, EConsoleVariableFlags Flags, bool bSaveDefault);
+
+#if ALLOW_OTHER_PLATFORM_CONFIG
+	// remember the default value of this cvar, before anything else can assign to it - this way we can know in the editor what the value
+	// would be for a cvar on another platform, if no .ini file on that platform sets it
+	virtual IConsoleVariable* GetDefaultValueVariable()
+	{
+		return PlatformIndependentDefault.Get();
+	}
+
+	virtual TSharedPtr<IConsoleVariable> GetPlatformValueVariable(FName PlatformName) override;
+	virtual void ClearPlatformVariables(FName PlatformName) override;
+
+protected:
+	// cache of the default value - we need to remember this when calculating the value on another platform if that platform doesn't override it in any ini file
+	// even if this platform does (in which case the default is lost)
+	TSharedPtr<IConsoleVariable> PlatformIndependentDefault;
+
+	// cache of the values of this cvar on other platforms
+	TMap<FName, TSharedPtr<IConsoleVariable> > PlatformValues;
+	FRWLock PlatformValuesLock;
+
+#endif
+};
+
+
+
+
+
 // T: bool, int32, float, FString
 template <class T>
-class FConsoleVariable : public FConsoleVariableBase
+class FConsoleVariable : public FOtherPlatformValueHelper<T>
 {
+// help find functions without needing this-> prefixes
+	using FConsoleVariableBase::GetShadowIndex;
+	using FConsoleVariableBase::CanChange;
+	using FConsoleVariableBase::Flags;
+
 public:
 	FConsoleVariable(T DefaultValue, const TCHAR* Help, EConsoleVariableFlags Flags, bool bSaveDefault=true) 
-		: FConsoleVariableBase(Help, Flags), Data(DefaultValue)
+		: FOtherPlatformValueHelper<T>(DefaultValue, Help, Flags, bSaveDefault)
+		, Data(DefaultValue)
 	{
-#if ALLOW_OTHER_PLATFORM_CONFIG
-		if (bSaveDefault)
-		{
-			PlatformIndependentDefault = TSharedPtr<IConsoleVariable>(new FConsoleVariable<T>(DefaultValue, TEXT(""), (EConsoleVariableFlags)(Flags | ECVF_ReadOnly), false));
-		}
-#endif
 	}
 
 	// interface IConsoleVariable ----------------------------------- 
@@ -541,50 +579,6 @@ private: // ----------------------------------------------------
 		OnCVarChange(Data.ShadowedValue[1], Data.ShadowedValue[0], Flags, SetBy);
 		FConsoleVariableBase::OnChanged(SetBy);
 	}
-
-#if ALLOW_OTHER_PLATFORM_CONFIG
-	// remember the dfefault value of this cvar, before anything else can assign to it - this way we can know in the editor what the value
-	// would be for a cvar on another platform, if no .ini file on that platform sets it
-	TSharedPtr<IConsoleVariable> PlatformIndependentDefault;
-	virtual IConsoleVariable* GetDefaultValueVariable() override
-	{
-		return PlatformIndependentDefault.Get();
-	}
-
-	TMap<FName, TSharedPtr<IConsoleVariable> > PlatformValues;
-	FRWLock PlatformValuesLock;
-	virtual IConsoleVariable* GetPlatformValueVariable(FName PlatformName)
-	{
-		// cheap lock here, contention is very rare
-		FRWScopeLock Lock(PlatformValuesLock, SLT_Write);
-		if (!PlatformValues.Contains(PlatformName))
-		{
-			FString ConfigValue;
-			EConsoleVariableFlags LastSetBy;
-			if (!GetConfigValueFromRuntimeSources(PlatformName, this, ConfigValue, LastSetBy))
-			{
-				return nullptr;
-			}
-			
-			T TypedValue;
-			TTypeFromString<T>::FromString(TypedValue, *ConfigValue);
-
-			// clear the existing setby mask
-			int NewFlags = GetFlags() & ~ECVF_SetFlagMask;
-			// add in new LastSetBy and make it readonly
-			NewFlags |= LastSetBy | ECVF_ReadOnly;
-
-			// make a new cvar to hold the value from the other platform
-			TSharedPtr<IConsoleVariable> PlatformCVar = TSharedPtr<IConsoleVariable>(new FConsoleVariable<T>(TypedValue, TEXT(""), (EConsoleVariableFlags)NewFlags, false));
-
-			PlatformValues.Add(PlatformName, PlatformCVar);
-		}
-
-		return PlatformValues.FindRef(PlatformName).Get();
-	}
-
-#endif
-
 };
 
 // specialization for all
@@ -714,15 +708,87 @@ template<> TConsoleVariableData<FString>* FConsoleVariable<FString>::AsVariableS
 	return &Data;
 }
 
+
+
+template<class T>
+FOtherPlatformValueHelper<T>::FOtherPlatformValueHelper(const T& DefaultValue, const TCHAR* Help, EConsoleVariableFlags Flags, bool bSaveDefault)
+	: FConsoleVariableBase(Help, Flags)
+{
+#if ALLOW_OTHER_PLATFORM_CONFIG
+	if (bSaveDefault)
+	{
+		PlatformIndependentDefault = TSharedPtr<IConsoleVariable>(new FConsoleVariable<T>(DefaultValue, TEXT(""), (EConsoleVariableFlags)(Flags | ECVF_ReadOnly), false));
+	}
+#endif
+}
+
+#if ALLOW_OTHER_PLATFORM_CONFIG
+template<class T>
+TSharedPtr<IConsoleVariable> FOtherPlatformValueHelper<T>::GetPlatformValueVariable(FName PlatformName)
+{
+	// cheap lock here, contention is very rare
+	FRWScopeLock Lock(PlatformValuesLock, SLT_Write);
+	if (!PlatformValues.Contains(PlatformName))
+	{
+		FString ConfigValue;
+		EConsoleVariableFlags LastSetBy;
+		if (!GetConfigValueFromRuntimeSources(PlatformName, this, ConfigValue, LastSetBy))
+		{
+			return nullptr;
+		}
+
+		T TypedValue;
+		TTypeFromString<T>::FromString(TypedValue, *ConfigValue);
+
+		// clear the existing setby mask
+		int NewFlags = GetFlags() & ~ECVF_SetFlagMask;
+		// add in new LastSetBy and make it readonly
+		NewFlags |= LastSetBy | ECVF_ReadOnly;
+
+		// make a new cvar to hold the value from the other platform
+		TSharedPtr<IConsoleVariable> PlatformCVar = TSharedPtr<IConsoleVariable>(new FConsoleVariable<T>(TypedValue, TEXT(""), (EConsoleVariableFlags)NewFlags, false));
+
+		PlatformValues.Add(PlatformName, PlatformCVar);
+	}
+
+	return PlatformValues.FindRef(PlatformName);
+}
+
+template<class T>
+void FOtherPlatformValueHelper<T>::ClearPlatformVariables(FName PlatformName)
+{
+	FRWScopeLock Lock(PlatformValuesLock, SLT_Write);
+
+	if (PlatformName == NAME_None)
+	{
+		PlatformValues.Empty();
+	}
+	else
+	{
+		PlatformValues.Remove(PlatformName);
+	}
+}
+
+#endif
+
+
+
 // ----
 
 // T: int32, float, bool
 template <class T>
-class FConsoleVariableRef : public FConsoleVariableBase
+class FConsoleVariableRef : public FOtherPlatformValueHelper<T>
 {
+	// help find functions without needing this-> prefixes
+	using FConsoleVariableBase::GetShadowIndex;
+	using FConsoleVariableBase::CanChange;
+	using FConsoleVariableBase::Flags;
+
 public:
 	FConsoleVariableRef(T& InRefValue, const TCHAR* Help, EConsoleVariableFlags Flags) 
-		: FConsoleVariableBase(Help, Flags), RefValue(InRefValue), MainValue(InRefValue)
+		: FOtherPlatformValueHelper<T>(InRefValue, Help, Flags, true)
+		, RefValue(InRefValue)
+		, MainValue(InRefValue)
 	{
 	}
 
@@ -1610,6 +1676,7 @@ bool FConsoleManager::ProcessUserConsoleInput(const TCHAR* InInput, FOutputDevic
 
 	IConsoleCommand* CCmd = CObj->AsCommand();
 	IConsoleVariable* CVar = CObj->AsVariable();
+	TSharedPtr<IConsoleVariable> PlatformCVar;
 
 	if (PlatformName != NAME_None)
 	{
@@ -1620,7 +1687,8 @@ bool FConsoleManager::ProcessUserConsoleInput(const TCHAR* InInput, FOutputDevic
 		else
 		{
 #if ALLOW_OTHER_PLATFORM_CONFIG
-			CVar = CVar->GetPlatformValueVariable(PlatformName);
+			PlatformCVar = CVar->GetPlatformValueVariable(PlatformName);
+			CVar = PlatformCVar.Get();
 			if (!CVar)
 			{
 				Ar.Logf(TEXT("Unable find CVar %s for platform %s (possibly invalid platform name?)"), *Param1, *PlatformName.ToString());
diff --git a/Engine/Source/Runtime/Core/Public/HAL/IConsoleManager.h b/Engine/Source/Runtime/Core/Public/HAL/IConsoleManager.h
--- a/Engine/Source/Runtime/Core/Public/HAL/IConsoleManager.h
+++ b/Engine/Source/Runtime/Core/Public/HAL/IConsoleManager.h
@@ -461,9 +461,23 @@ public:
 	}
 
 #if ALLOW_OTHER_PLATFORM_CONFIG
-	virtual IConsoleVariable* GetPlatformValueVariable(FName PlatformName)
+
+	/**
+	 * Get a CVar opject that matches this cvar, but contains the value of the platform given.
+	 */
+	virtual TSharedPtr<IConsoleVariable> GetPlatformValueVariable(FName PlatformName)
+	{
+		// this could be called for some special subclass like FConsoleVariableBitRef that don't implement this (yet)
+		unimplemented();
+		return TSharedPtr<IConsoleVariable>();
+	}
+
+	/**
+	 * Used only for debugging/iterating, this will clear all of the other platform's cvar objects, which will 
+	 * force a fresh lookup (with fresh ini files, likely) on next call to GetPlatformValueVariable
+	 */
+	virtual void ClearPlatformVariables(FName PlatformName=NAME_None)
 	{
-		return nullptr;
 	}
 #endif
 
diff --git a/Engine/Source/Runtime/Engine/Classes/DeviceProfiles/DeviceProfile.h b/Engine/Source/Runtime/Engine/Classes/DeviceProfiles/DeviceProfile.h
--- a/Engine/Source/Runtime/Engine/Classes/DeviceProfiles/DeviceProfile.h
+++ b/Engine/Source/Runtime/Engine/Classes/DeviceProfiles/DeviceProfile.h
@@ -164,6 +164,7 @@ private:
 public:
 	const TMap<FString, FString>& GetAllExpandedCVars();
 	void AddExpandedCVars(const TMap<FString, FString>& CVarsToMerge);
+	void ClearAllExpandedCVars();
 private:
 	/** Resolved CVars, including expanded scalability cvars used to properly emulate one platform on another */
 	TMap<FString, FString> AllExpandedCVars;
diff --git a/Engine/Source/Runtime/Engine/Private/DeviceProfiles/DeviceProfile.cpp b/Engine/Source/Runtime/Engine/Private/DeviceProfiles/DeviceProfile.cpp
--- a/Engine/Source/Runtime/Engine/Private/DeviceProfiles/DeviceProfile.cpp
+++ b/Engine/Source/Runtime/Engine/Private/DeviceProfiles/DeviceProfile.cpp
@@ -395,5 +395,11 @@ void UDeviceProfile::AddExpandedCVars(const TMap<FString, FString>& CVarsToMerge
 	}
 }
 
+void UDeviceProfile::ClearAllExpandedCVars()
+{
+	AllExpandedCVars.Empty();
+}
+
+
 
 #endif
\ No newline at end of file
diff --git a/Engine/Source/Runtime/Engine/Private/DeviceProfiles/DeviceProfileManager.cpp b/Engine/Source/Runtime/Engine/Private/DeviceProfiles/DeviceProfileManager.cpp
--- a/Engine/Source/Runtime/Engine/Private/DeviceProfiles/DeviceProfileManager.cpp
+++ b/Engine/Source/Runtime/Engine/Private/DeviceProfiles/DeviceProfileManager.cpp
@@ -1386,8 +1386,24 @@ public:
 		}
 		else if (FParse::Command(&Cmd, TEXT("dpreload")))
 		{
+			// clear cached other-platform CVars
+			IConsoleManager::Get().ForEachConsoleObjectThatStartsWith(FConsoleObjectVisitor::CreateLambda(
+				[](const TCHAR* Key, IConsoleObject* ConsoleObject)
+				{
+					if (IConsoleVariable* AsVariable = ConsoleObject->AsVariable())
+					{
+						AsVariable->ClearPlatformVariables();
+					}
+
+				}));
+
+			// clear some cached cvars in other-platform expansions
+			for (TObjectPtr<UDeviceProfile> DeviceProfile : UDeviceProfileManager::Get().Profiles)
+			{
+				DeviceProfile->ClearAllExpandedCVars();
+			}
+
 			FConfigCacheIni::ClearOtherPlatformConfigs();
-			// @todo ini clear out all DPs AllExpandedCVars
 		}
 		else if (FParse::Command(&Cmd, TEXT("dpreapply")))
 		{
diff --git a/Engine/Source/Runtime/RenderCore/Public/RenderUtils.h b/Engine/Source/Runtime/RenderCore/Public/RenderUtils.h
--- a/Engine/Source/Runtime/RenderCore/Public/RenderUtils.h
+++ b/Engine/Source/Runtime/RenderCore/Public/RenderUtils.h
@@ -544,13 +544,17 @@ struct RENDERCORE_API FShaderPlatformCachedIniValue
 			}
 		}
 
-		if (CVar != nullptr)
+		// now get the value from the platform that makes sense for this shader platform
+		TSharedPtr<IConsoleVariable> OtherPlatformVar = CVar->GetPlatformValueVariable(ShaderPlatformToPlatformName(ShaderPlatform));
+		ensureMsgf(OtherPlatformVar.IsValid(), TEXT("Failed to get another platform's version of a cvar (possible name: '%s'). It is probably an esoteric subclass that needs to implement GetPlatformValueVariable."), *CVarName);
+		if (OtherPlatformVar.IsValid())
 		{
-			// now get the value from the platform that makes sense for this shader platform
-			if (IConsoleVariable* ConsoleVariable = CVar->GetPlatformValueVariable(ShaderPlatformToPlatformName(ShaderPlatform)))
-			{
-				ConsoleVariable->GetValue(Value);
-			}
+			OtherPlatformVar->GetValue(Value);
+		}
+		else
+		{
+			// get this platform's value, even tho it could be wrong
+			CVar->GetValue(Value);
 		}
 #else
 		unimplemented();
