[Insights] Instrument core classes with memory tracing.

diff --git a/Engine/Source/Runtime/Core/Private/Windows/WindowsRunnableThread.h b/Engine/Source/Runtime/Core/Private/Windows/WindowsRunnableThread.h
--- a/Engine/Source/Runtime/Core/Private/Windows/WindowsRunnableThread.h
+++ b/Engine/Source/Runtime/Core/Private/Windows/WindowsRunnableThread.h
@@ -143,6 +143,10 @@ protected:
 			LLM_PLATFORM_SCOPE(ELLMTag::ThreadStackPlatform);
 			// add in the thread size, since it's allocated in a black box we can't track
 			// note: I don't see any accounting for this when threads are destroyed
+			const uint64 FakeAddress = uint64(InRunnable) | (uint64(1u) << 48);
+			MemoryTrace_Alloc(FakeAddress, InStackSize, 4);
+			MemoryTrace_MarkAllocAsHeap(FakeAddress, EMemoryTraceRootHeap::SystemMemory);
+			MemoryTrace_Alloc(FakeAddress, InStackSize, 4);
 			LLM(FLowLevelMemTracker::Get().OnLowLevelAlloc(ELLMTracker::Default, nullptr, InStackSize));
 			LLM(FLowLevelMemTracker::Get().OnLowLevelAlloc(ELLMTracker::Platform, nullptr, InStackSize));
 
diff --git a/Engine/Source/Runtime/RHI/Private/GPUDefragAllocator.cpp b/Engine/Source/Runtime/RHI/Private/GPUDefragAllocator.cpp
--- a/Engine/Source/Runtime/RHI/Private/GPUDefragAllocator.cpp
+++ b/Engine/Source/Runtime/RHI/Private/GPUDefragAllocator.cpp
@@ -223,6 +223,7 @@ void* FGPUDefragAllocator::Allocate(int64 AllocationSize, int32 Alignment, TStat
 
 	check(IsAligned(AllocatedChunk->Base, Alignment));
 
+	MemoryTrace_Alloc((uint64)AllocatedChunk->Base, AllocationSize, Alignment);
 	LLM(FLowLevelMemTracker::Get().OnLowLevelAlloc(ELLMTracker::Default, AllocatedChunk->Base, AllocationSize));
 
 	return AllocatedChunk->Base;
@@ -354,6 +355,7 @@ void* GBestFitAllocatorFreePointer = nullptr;
 */
 void FGPUDefragAllocator::Free(void* Pointer)
 {
+	MemoryTrace_Free((uint64)Pointer);
 #if ENABLE_LOW_LEVEL_MEM_TRACKER
 	if (Pointer)
 	{
diff --git a/Engine/Source/Runtime/RHI/Public/GPUDefragAllocator.h b/Engine/Source/Runtime/RHI/Public/GPUDefragAllocator.h
--- a/Engine/Source/Runtime/RHI/Public/GPUDefragAllocator.h
+++ b/Engine/Source/Runtime/RHI/Public/GPUDefragAllocator.h
@@ -792,6 +792,8 @@ protected:
 	*/
 	void Relocate(FRelocationStats& Stats, FMemoryChunk* Dest, int64 DestOffset, const void* Source, int64 Size, void* UserPayload)
 	{
+		MemoryTrace_ReallocFree((uint64)Source);
+		MemoryTrace_ReallocAlloc((uint64) Dest, Size, 4);
 		LLM(FLowLevelMemTracker::Get().OnLowLevelAllocMoved(ELLMTracker::Default, Dest->Base, Source));
 
 		uint8* DestAddr = Dest->Base + DestOffset;
