REL - allow to pull all Init Parameters in bulk

diff --git a/Engine/Source/Runtime/Core/Public/Async/Fundamental/Task.h b/Engine/Source/Runtime/Core/Public/Async/Fundamental/Task.h
--- a/Engine/Source/Runtime/Core/Public/Async/Fundamental/Task.h
+++ b/Engine/Source/Runtime/Core/Public/Async/Fundamental/Task.h
@@ -249,6 +249,14 @@ namespace LowLevelTasks
 		inline bool IsBackgroundTask() const;
 		inline bool AllowBusyWaiting() const;
 
+		struct FInitData
+		{
+			const TCHAR* DebugName;
+			ETaskPriority Priority;
+			bool bAllowBusyWaiting;
+		};
+		inline FInitData GetInitData() const;
+
 		void* GetUserData() const { return UserData; }
 		void SetUserData(void* NewUserData) const { UserData = NewUserData; }
 
@@ -409,6 +417,12 @@ namespace LowLevelTasks
 		return PackedData.load(std::memory_order_relaxed).AllowBusyWaiting();
 	}
 
+	inline FTask::FInitData FTask::GetInitData() const
+	{
+		FPackedData LocalPackedData = PackedData.load(std::memory_order_relaxed);
+		return { LocalPackedData.GetDebugName(), LocalPackedData.GetPriority(), LocalPackedData.AllowBusyWaiting() };
+	}
+
 	enum class ESleepState
 	{
 		Affinity,
