Clarify TBox comment for TransformProjectBy and added UE_NODISCARD to methods returning a computed/transformed box.

diff --git a/Engine/Source/Editor/UnrealEd/Private/AssetSelection.cpp b/Engine/Source/Editor/UnrealEd/Private/AssetSelection.cpp
--- a/Engine/Source/Editor/UnrealEd/Private/AssetSelection.cpp
+++ b/Engine/Source/Editor/UnrealEd/Private/AssetSelection.cpp
@@ -503,7 +503,7 @@ namespace ActorPlacementUtils
 			FVector ExpandedScale = FVector(1.0f + (GetDefault<ULevelEditorMiscSettings>()->PercentageThresholdForPrompt / 100.0f));
 			FTransform ExpandedScaleTransform = FTransform::Identity;
 			ExpandedScaleTransform.SetScale3D(ExpandedScale);
-			CurrentLevelBounds.TransformBy(ExpandedScaleTransform);
+			CurrentLevelBounds = CurrentLevelBounds.TransformBy(ExpandedScaleTransform);
 			for (int32 ActorTransformIndex = 0; ActorTransformIndex < InActorTransforms.Num(); ++ActorTransformIndex)
 			{
 				FTransform ActorTransform = InActorTransforms[ActorTransformIndex];
diff --git a/Engine/Source/Runtime/Core/Public/Math/Box.h b/Engine/Source/Runtime/Core/Public/Math/Box.h
--- a/Engine/Source/Runtime/Core/Public/Math/Box.h
+++ b/Engine/Source/Runtime/Core/Public/Math/Box.h
@@ -364,7 +364,7 @@ public:
 	 * @param Other The bounding box to test overlap
 	 * @return the overlap box. It can be 0 if they don't overlap
 	 */
-	TBox<T> Overlap( const TBox<T>& Other ) const;
+	UE_NODISCARD TBox<T> Overlap( const TBox<T>& Other ) const;
 
 	/**
 	  * Gets a bounding volume transformed by an inverted FTransform object.
@@ -372,7 +372,7 @@ public:
 	  * @param M The transformation object to perform the inversely transform this box with.
 	  * @return	The transformed box.
 	  */
-	TBox<T> InverseTransformBy( const FTransform& M ) const;
+	UE_NODISCARD TBox<T> InverseTransformBy( const FTransform& M ) const;
 
 	/** 
 	 * Checks whether the given location is inside this box.
@@ -451,7 +451,7 @@ public:
 	 * @return The transformed box.
 	 * @see TransformProjectBy
 	 */
-	TBox<T> TransformBy( const TMatrix<T>& M ) const;
+	UE_NODISCARD TBox<T> TransformBy( const TMatrix<T>& M ) const;
 
 	/**
 	 * Gets a bounding volume transformed by a FTransform object.
@@ -460,16 +460,16 @@ public:
 	 * @return The transformed box.
 	 * @see TransformProjectBy
 	 */
-	TBox<T> TransformBy( const FTransform& M ) const;
+	UE_NODISCARD TBox<T> TransformBy( const FTransform& M ) const;
 
 	/** 
-	 * Transforms and projects a world bounding box to screen space
+	 * Returns the current world bounding box transformed and projected to screen space
 	 *
 	 * @param ProjM The projection matrix.
 	 * @return The transformed box.
 	 * @see TransformBy
 	 */
-	TBox<T> TransformProjectBy( const TMatrix<T>& ProjM ) const;
+	UE_NODISCARD TBox<T> TransformProjectBy( const TMatrix<T>& ProjM ) const;
 
 	/**
 	 * Get a textual representation of this box.
