Cleaned up namespaces in Core

#rnx

diff --git a/Engine/Source/Runtime/Core/Private/String/BytesToHex.cpp b/Engine/Source/Runtime/Core/Private/String/BytesToHex.cpp
--- a/Engine/Source/Runtime/Core/Private/String/BytesToHex.cpp
+++ b/Engine/Source/Runtime/Core/Private/String/BytesToHex.cpp
@@ -4,9 +4,7 @@
 
 #include "Misc/StringBuilder.h"
 
-namespace UE
-{
-namespace String
+namespace UE::String
 {
 
 template <typename CharType, ANSICHAR LetterA>
@@ -88,5 +86,4 @@ void BytesToHexLower(TConstArrayView<uint8> Bytes, FUtf8StringBuilderBase& Build
 	BytesToHexImpl<UTF8CHAR, 'a'>(Bytes, GetData(Builder) + Offset);
 }
 
-}
-}
+} // UE::String
diff --git a/Engine/Source/Runtime/Core/Private/String/HexToBytes.cpp b/Engine/Source/Runtime/Core/Private/String/HexToBytes.cpp
--- a/Engine/Source/Runtime/Core/Private/String/HexToBytes.cpp
+++ b/Engine/Source/Runtime/Core/Private/String/HexToBytes.cpp
@@ -5,9 +5,7 @@
 #include "Containers/StringView.h"
 #include "Containers/UnrealString.h"
 
-namespace UE
-{
-namespace String
+namespace UE::String
 {
 
 template <typename CharType>
@@ -44,5 +42,4 @@ int32 HexToBytes(FUtf8StringView Hex, uint8* OutBytes)
 	return HexToBytesImpl<UTF8CHAR>(Hex, OutBytes);
 }
 
-}
-}
+} // UE::String
diff --git a/Engine/Source/Runtime/Core/Public/String/BytesToHex.h b/Engine/Source/Runtime/Core/Public/String/BytesToHex.h
--- a/Engine/Source/Runtime/Core/Public/String/BytesToHex.h
+++ b/Engine/Source/Runtime/Core/Public/String/BytesToHex.h
@@ -6,9 +6,7 @@
 #include "Containers/ArrayView.h"
 #include "Containers/StringFwd.h"
 
-namespace UE
-{
-namespace String
+namespace UE::String
 {
 
 /**
@@ -37,5 +35,4 @@ CORE_API void BytesToHexLower(TConstArrayView<uint8> Bytes, FAnsiStringBuilderBa
 CORE_API void BytesToHexLower(TConstArrayView<uint8> Bytes, FWideStringBuilderBase& Builder);
 CORE_API void BytesToHexLower(TConstArrayView<uint8> Bytes, FUtf8StringBuilderBase& Builder);
 
-}
-}
+} // UE::String
diff --git a/Engine/Source/Runtime/Core/Public/String/HexToBytes.h b/Engine/Source/Runtime/Core/Public/String/HexToBytes.h
--- a/Engine/Source/Runtime/Core/Public/String/HexToBytes.h
+++ b/Engine/Source/Runtime/Core/Public/String/HexToBytes.h
@@ -5,9 +5,7 @@
 #include "CoreTypes.h"
 #include "Containers/StringFwd.h"
 
-namespace UE
-{
-namespace String
+namespace UE::String
 {
 
 /** 
@@ -22,5 +20,4 @@ CORE_API int32 HexToBytes(FAnsiStringView Hex, uint8* OutBytes);
 CORE_API int32 HexToBytes(FWideStringView Hex, uint8* OutBytes);
 CORE_API int32 HexToBytes(FUtf8StringView Hex, uint8* OutBytes);
 
-}
-}
+} // UE::String
