- OtherPlatform CVars will now allow "True" type values for "int32" type cvars
#p4v-preflight-copy 18518878

diff --git a/Engine/Source/Runtime/Core/Private/HAL/ConsoleManager.cpp b/Engine/Source/Runtime/Core/Private/HAL/ConsoleManager.cpp
--- a/Engine/Source/Runtime/Core/Private/HAL/ConsoleManager.cpp
+++ b/Engine/Source/Runtime/Core/Private/HAL/ConsoleManager.cpp
@@ -738,7 +738,7 @@ TSharedPtr<IConsoleVariable> FOtherPlatformValueHelper<T>::GetPlatformValueVaria
 		}
 
 		T TypedValue;
-		TTypeFromString<T>::FromString(TypedValue, *ConfigValue);
+		TTypeFromString<T>::FromString(TypedValue, ConvertValueFromHumanFriendlyValue(*ConfigValue));
 
 		// clear the existing setby mask
 		int NewFlags = GetFlags() & ~ECVF_SetFlagMask;
