lock-free lists:
* replaced padding by alignof to minimise memory usage
* replaced deriving from FNoncopyable by UE_NONCOPYABLE to "enable empty base optimisation"

diff --git a/Engine/Source/Runtime/Core/Private/Tests/Async/TaskGraphTest.cpp b/Engine/Source/Runtime/Core/Private/Tests/Async/TaskGraphTest.cpp
--- a/Engine/Source/Runtime/Core/Private/Tests/Async/TaskGraphTest.cpp
+++ b/Engine/Source/Runtime/Core/Private/Tests/Async/TaskGraphTest.cpp
@@ -346,15 +346,15 @@ namespace OldTaskGraphTests
 	struct FTestRigFIFO
 	{
 		FLockFreePointerFIFOBase<FTestStruct, PLATFORM_CACHE_LINE_SIZE> Test1;
-		FLockFreePointerFIFOBase<FTestStruct, 1> Test2;
-		FLockFreePointerFIFOBase<FTestStruct, 1, 1 << 4> Test3;
+		FLockFreePointerFIFOBase<FTestStruct, 8> Test2;
+		FLockFreePointerFIFOBase<FTestStruct, 8, 1 << 4> Test3;
 	};
 
 	struct FTestRigLIFO
 	{
 		FLockFreePointerListLIFOBase<FTestStruct, PLATFORM_CACHE_LINE_SIZE> Test1;
-		FLockFreePointerListLIFOBase<FTestStruct, 1> Test2;
-		FLockFreePointerListLIFOBase<FTestStruct, 1, 1 << 4> Test3;
+		FLockFreePointerListLIFOBase<FTestStruct, 8> Test2;
+		FLockFreePointerListLIFOBase<FTestStruct, 8, 1 << 4> Test3;
 	};
 
 	static void TestLockFree(int32 OuterIters = 3)
diff --git a/Engine/Source/Runtime/Core/Public/Containers/LockFreeList.h b/Engine/Source/Runtime/Core/Public/Containers/LockFreeList.h
--- a/Engine/Source/Runtime/Core/Public/Containers/LockFreeList.h
+++ b/Engine/Source/Runtime/Core/Public/Containers/LockFreeList.h
@@ -50,17 +50,6 @@ CORE_API void LockFreeFreeLinks(SIZE_T AllocSize, void* Ptr);
 #define MAX_LOCK_FREE_LINKS_AS_BITS (26)
 #define MAX_LOCK_FREE_LINKS (1 << 26)
 
-template<int TPaddingForCacheContention>
-struct FPaddingForCacheContention
-{
-	uint8 PadToAvoidContention[TPaddingForCacheContention];
-};
-
-template<>
-struct FPaddingForCacheContention<0>
-{
-};
-
 template<class T, unsigned int MaxTotalItems, unsigned int ItemsPerPage>
 class TLockFreeAllocOnceIndexedAllocator
 {
@@ -128,11 +117,8 @@ private:
 		return (void*)(Pages[BlockIndex] + SubIndex);
 	}
 
-	uint8 PadToAvoidContention0[PLATFORM_CACHE_LINE_SIZE];
-	FThreadSafeCounter NextIndex;
-	uint8 PadToAvoidContention1[PLATFORM_CACHE_LINE_SIZE];
-	T* Pages[MaxBlocks];
-	uint8 PadToAvoidContention2[PLATFORM_CACHE_LINE_SIZE];
+	alignas(PLATFORM_CACHE_LINE_SIZE) FThreadSafeCounter NextIndex;
+	alignas(PLATFORM_CACHE_LINE_SIZE) T* Pages[MaxBlocks];
 };
 
 
@@ -415,10 +401,7 @@ public:
 	}
 
 private:
-
-	FPaddingForCacheContention<TPaddingForCacheContention> PadToAvoidContention1;
-	TDoublePtr Head;
-	FPaddingForCacheContention<TPaddingForCacheContention> PadToAvoidContention2;
+	alignas(TPaddingForCacheContention) TDoublePtr Head;
 };
 
 template<class T, int TPaddingForCacheContention, uint64 TABAInc = 1>
@@ -661,12 +644,8 @@ public:
 	}
 
 private:
-
-	FPaddingForCacheContention<TPaddingForCacheContention> PadToAvoidContention1;
-	TDoublePtr Head;
-	FPaddingForCacheContention<TPaddingForCacheContention> PadToAvoidContention2;
-	TDoublePtr Tail;
-	FPaddingForCacheContention<TPaddingForCacheContention> PadToAvoidContention3;
+	alignas(TPaddingForCacheContention) TDoublePtr Head;
+	alignas(TPaddingForCacheContention) TDoublePtr Tail;
 };
 
 
@@ -816,8 +795,7 @@ private:
 
 	FLockFreePointerFIFOBase<T, TPaddingForCacheContention> PriorityQueues[NumPriorities];
 	// not a pointer to anything, rather tracks the stall state of all threads servicing this queue.
-	TDoublePtr MasterState;
-	FPaddingForCacheContention<TPaddingForCacheContention> PadToAvoidContention1;
+	alignas(TPaddingForCacheContention) TDoublePtr MasterState;
 };
 
 
