removed `check(IsInGameThread())` in TLockFreeFixedSizeAllocator_TLSCacheBase constructor as it's too limiting and I don't see any reason for this requirement

diff --git a/Engine/Source/Runtime/Core/Public/Containers/LockFreeFixedSizeAllocator.h b/Engine/Source/Runtime/Core/Public/Containers/LockFreeFixedSizeAllocator.h
--- a/Engine/Source/Runtime/Core/Public/Containers/LockFreeFixedSizeAllocator.h
+++ b/Engine/Source/Runtime/Core/Public/Containers/LockFreeFixedSizeAllocator.h
@@ -25,12 +25,12 @@ class TLockFreeFixedSizeAllocator_TLSCacheBase : public FNoncopyable
 		SIZE_PER_BUNDLE = 65536,
 		NUM_PER_BUNDLE = SIZE_PER_BUNDLE / SIZE
 	};
+
 public:
 
 	TLockFreeFixedSizeAllocator_TLSCacheBase()
 	{
 		static_assert(SIZE >= sizeof(void*) && SIZE % sizeof(void*) == 0, "Blocks in TLockFreeFixedSizeAllocator must be at least the size of a pointer.");
-		check(IsInGameThread());
 		TlsSlot = FPlatformTLS::AllocTlsSlot();
 		check(FPlatformTLS::IsValidTlsSlot(TlsSlot));
 	}
