QualifiedFrameTime: add helper functions to create an FTimecode from an FQualifiedFrameTime

diff --git a/Engine/Source/Runtime/Core/Public/Misc/QualifiedFrameTime.h b/Engine/Source/Runtime/Core/Public/Misc/QualifiedFrameTime.h
--- a/Engine/Source/Runtime/Core/Public/Misc/QualifiedFrameTime.h
+++ b/Engine/Source/Runtime/Core/Public/Misc/QualifiedFrameTime.h
@@ -53,6 +53,31 @@ public:
 		return  FFrameRate::TransformTime(Time, Rate, DesiredRate);
 	}
 
+	/**
+	 * Create an FTimecode from this qualified frame time.
+	 *
+	 * Any subframe value in the frame time will be ignored.
+	 * Whether or not the returned timecode is a drop-frame timecode will be determined by the qualified frame time's frame rate
+	 * and the CVar specifying whether to generate drop-frame timecodes by default for supported frame rates.
+	 */
+	FTimecode ToTimecode() const
+	{
+		return FTimecode::FromFrameNumber(Time.FloorToFrame(), Rate);
+	}
+
+	/**
+	 * Create an FTimecode from this qualified frame time. Optionally supports creating a drop-frame timecode,
+	 * which drops certain timecode display numbers to help account for NTSC frame rates which are fractional.
+	 *
+	 * @param bDropFrame    If true, the returned timecode will drop the first two frames on every minute (except when Minute % 10 == 0).
+	 *     This is only valid for NTSC framerates (29.97, 59.94) and will assert if you try to create a drop-frame timecode
+	 *     from a non-valid framerate. All framerates can be represented by non-drop timecode.
+	 */
+	FTimecode ToTimecode(bool bDropFrame) const
+	{
+		return FTimecode::FromFrameNumber(Time.FloorToFrame(), Rate, bDropFrame);
+	}
+
 public:
 
 	/** The frame time */
