Add flag to FPlatformProcess::CreatePipe to indicate the pipe will be used for stdin redirection rather than stdout redirection.
Windows needs to set permissions accordingly.

diff --git a/Engine/Source/Runtime/Core/Private/GenericPlatform/GenericPlatformProcess.cpp b/Engine/Source/Runtime/Core/Private/GenericPlatform/GenericPlatformProcess.cpp
--- a/Engine/Source/Runtime/Core/Private/GenericPlatform/GenericPlatformProcess.cpp
+++ b/Engine/Source/Runtime/Core/Private/GenericPlatform/GenericPlatformProcess.cpp
@@ -557,7 +557,7 @@ void FGenericPlatformProcess::ClosePipe( void* ReadPipe, void* WritePipe )
 	UE_LOG(LogHAL, Fatal, TEXT("FGenericPlatformProcess::ClosePipe not implemented on this platform"));
 }
 
-bool FGenericPlatformProcess::CreatePipe( void*& ReadPipe, void*& WritePipe )
+bool FGenericPlatformProcess::CreatePipe(void*& ReadPipe, void*& WritePipe, bool bWritePipeLocal)
 {
 	UE_LOG(LogHAL, Fatal, TEXT("FGenericPlatformProcess::CreatePipe not implemented on this platform"));
 	return false;
diff --git a/Engine/Source/Runtime/Core/Private/Mac/MacPlatformProcess.cpp b/Engine/Source/Runtime/Core/Private/Mac/MacPlatformProcess.cpp
--- a/Engine/Source/Runtime/Core/Private/Mac/MacPlatformProcess.cpp
+++ b/Engine/Source/Runtime/Core/Private/Mac/MacPlatformProcess.cpp
@@ -1309,7 +1309,7 @@ void FMacPlatformProcess::ClosePipe( void* ReadPipe, void* WritePipe )
 	}
 }
 
-bool FMacPlatformProcess::CreatePipe( void*& ReadPipe, void*& WritePipe )
+bool FMacPlatformProcess::CreatePipe( void*& ReadPipe, void*& WritePipe, bool bWritePipeLocal )
 {
 	SCOPED_AUTORELEASE_POOL;
 	int pipefd[2];
diff --git a/Engine/Source/Runtime/Core/Private/Unix/UnixPlatformProcess.cpp b/Engine/Source/Runtime/Core/Private/Unix/UnixPlatformProcess.cpp
--- a/Engine/Source/Runtime/Core/Private/Unix/UnixPlatformProcess.cpp
+++ b/Engine/Source/Runtime/Core/Private/Unix/UnixPlatformProcess.cpp
@@ -513,7 +513,7 @@ void FUnixPlatformProcess::ClosePipe( void* ReadPipe, void* WritePipe )
 	}
 }
 
-bool FUnixPlatformProcess::CreatePipe( void*& ReadPipe, void*& WritePipe )
+bool FUnixPlatformProcess::CreatePipe(void*& ReadPipe, void*& WritePipe, bool bWritePipeLocal)
 {
 	int PipeFd[2];
 	if (-1 == pipe(PipeFd))
diff --git a/Engine/Source/Runtime/Core/Private/Windows/WindowsPlatformProcess.cpp b/Engine/Source/Runtime/Core/Private/Windows/WindowsPlatformProcess.cpp
--- a/Engine/Source/Runtime/Core/Private/Windows/WindowsPlatformProcess.cpp
+++ b/Engine/Source/Runtime/Core/Private/Windows/WindowsPlatformProcess.cpp
@@ -1471,7 +1471,7 @@ void FWindowsPlatformProcess::ClosePipe( void* ReadPipe, void* WritePipe )
 	}
 }
 
-bool FWindowsPlatformProcess::CreatePipe( void*& ReadPipe, void*& WritePipe )
+bool FWindowsPlatformProcess::CreatePipe( void*& ReadPipe, void*& WritePipe, bool bWritePipeLocal )
 {
 	SECURITY_ATTRIBUTES Attr = { sizeof(SECURITY_ATTRIBUTES), NULL, true };
 	
@@ -1480,7 +1480,7 @@ bool FWindowsPlatformProcess::CreatePipe( void*& ReadPipe, void*& WritePipe )
 		return false;
 	}
 
-	if (!::SetHandleInformation(ReadPipe, HANDLE_FLAG_INHERIT, 0))
+	if (!::SetHandleInformation(bWritePipeLocal ? WritePipe : ReadPipe, HANDLE_FLAG_INHERIT, 0))
 	{
 		return false;
 	}
diff --git a/Engine/Source/Runtime/Core/Public/GenericPlatform/GenericPlatformProcess.h b/Engine/Source/Runtime/Core/Public/GenericPlatform/GenericPlatformProcess.h
--- a/Engine/Source/Runtime/Core/Public/GenericPlatform/GenericPlatformProcess.h
+++ b/Engine/Source/Runtime/Core/Public/GenericPlatform/GenericPlatformProcess.h
@@ -614,10 +614,11 @@ struct CORE_API FGenericPlatformProcess
 	 *
 	 * @param ReadPipe Will hold the handle to the read end of the pipe.
 	 * @param WritePipe Will hold the handle to the write end of the pipe.
+	 * @parm bWritePipeLocal indicates that the write pipe end will be used locally, instead of the read pipe
 	 * @return true on success, false otherwise.
 	 * @see ClosePipe, ReadPipe
 	 */
-	static bool CreatePipe( void*& ReadPipe, void*& WritePipe );
+	static bool CreatePipe(void*& ReadPipe, void*& WritePipe, bool bWritePipeLocal = false);
 
 	/**
 	 * Reads all pending data from an anonymous pipe, such as STDOUT or STDERROR of a process.
diff --git a/Engine/Source/Runtime/Core/Public/Mac/MacPlatformProcess.h b/Engine/Source/Runtime/Core/Public/Mac/MacPlatformProcess.h
--- a/Engine/Source/Runtime/Core/Public/Mac/MacPlatformProcess.h
+++ b/Engine/Source/Runtime/Core/Public/Mac/MacPlatformProcess.h
@@ -258,7 +258,7 @@ struct CORE_API FMacPlatformProcess : public FGenericPlatformProcess
 	static void ExploreFolder( const TCHAR* FilePath );
 	static FRunnableThread* CreateRunnableThread();
 	static void ClosePipe( void* ReadPipe, void* WritePipe );
-	static bool CreatePipe( void*& ReadPipe, void*& WritePipe );
+	static bool CreatePipe(void*& ReadPipe, void*& WritePipe, bool bWritePipeLocal = false);
 	static FString ReadPipe( void* ReadPipe );
 	static bool ReadPipeToArray(void* ReadPipe, TArray<uint8> & Output);
 	static bool WritePipe(void* WritePipe, const FString& Message, FString* OutWritten = nullptr);
diff --git a/Engine/Source/Runtime/Core/Public/Unix/UnixPlatformProcess.h b/Engine/Source/Runtime/Core/Public/Unix/UnixPlatformProcess.h
--- a/Engine/Source/Runtime/Core/Public/Unix/UnixPlatformProcess.h
+++ b/Engine/Source/Runtime/Core/Public/Unix/UnixPlatformProcess.h
@@ -270,7 +270,7 @@ struct CORE_API FUnixPlatformProcess : public FGenericPlatformProcess
 	static const TCHAR* GetModulePrefix();
 	static const TCHAR* GetModuleExtension();
 	static void ClosePipe( void* ReadPipe, void* WritePipe );
-	static bool CreatePipe( void*& ReadPipe, void*& WritePipe );
+	static bool CreatePipe(void*& ReadPipe, void*& WritePipe, bool bWritePipeLocal = false);
 	static FString ReadPipe( void* ReadPipe );
 	static bool ReadPipeToArray(void* ReadPipe, TArray<uint8> & Output);
 	static bool WritePipe(void* WritePipe, const FString& Message, FString* OutWritten = nullptr);
diff --git a/Engine/Source/Runtime/Core/Public/Windows/WindowsPlatformProcess.h b/Engine/Source/Runtime/Core/Public/Windows/WindowsPlatformProcess.h
--- a/Engine/Source/Runtime/Core/Public/Windows/WindowsPlatformProcess.h
+++ b/Engine/Source/Runtime/Core/Public/Windows/WindowsPlatformProcess.h
@@ -182,7 +182,7 @@ public:
 	static class FEvent* CreateSynchEvent(bool bIsManualReset = false);
 	static class FRunnableThread* CreateRunnableThread();
 	static void ClosePipe( void* ReadPipe, void* WritePipe );
-	static bool CreatePipe( void*& ReadPipe, void*& WritePipe );
+	static bool CreatePipe(void*& ReadPipe, void*& WritePipe, bool bWritePipeLocal = false);
 	static FString ReadPipe( void* ReadPipe );
 	static bool ReadPipeToArray(void* ReadPipe, TArray<uint8> & Output);
 	static bool WritePipe(void* WritePipe, const FString& Message, FString* OutWritten = nullptr);
