[Insights] Trace console commands improvments.

* Added 'Trace.Enable' and 'Trace.Disable' console commands to control channels.
* 'Trace.Start' is marked as deprecated in favor of Trace.File.
* For 'Trace.File' both arguments are now optional, and heuristic for detecting a file argument implemented.
* 'Trace.Stop' no longer disables channels.
* 'Trace.Pause/Resume' now keeps a list of the current active channels which will be restored. Logic for enabling channels has been separated from EnableChannels which is used in initialization sequence.
* Changed FChannel toggling to directly write the state. This change was done to reduce confusion around enabling/disabling channels and preset that overlaps.
* Refactored some methods internal to FTraceAuxilliaryImpl

diff --git a/Engine/Source/Runtime/Core/Private/ProfilingDebugging/TraceAuxiliary.cpp b/Engine/Source/Runtime/Core/Private/ProfilingDebugging/TraceAuxiliary.cpp
--- a/Engine/Source/Runtime/Core/Private/ProfilingDebugging/TraceAuxiliary.cpp
+++ b/Engine/Source/Runtime/Core/Private/ProfilingDebugging/TraceAuxiliary.cpp
@@ -49,6 +49,7 @@
 #include "ProfilingDebugging/CsvProfiler.h"
 #include "String/ParseTokens.h"
 #include "Templates/UnrealTemplate.h"
+#include "Templates/Invoke.h"
 #include "Trace/Trace.inl"
 
 ////////////////////////////////////////////////////////////////////////////////
@@ -77,12 +78,16 @@ class FTraceAuxiliaryImpl
 {
 public:
 	const TCHAR*			GetDest() const;
-	template <class T> void	ReadChannels(T&& Callback) const;
-	void					AddChannels(const TCHAR* ChannelList);
+	bool					IsConnected() const;
+	void					GetActiveChannelsString(FStringBuilderBase& String) const;
+	void					AddCommandlineChannels(const TCHAR* ChannelList);
+	void					EnableChannels(const TCHAR* ChannelList);
+	void					DisableChannels(const TCHAR* ChannelList);
 	bool					Connect(ETraceConnectType Type, const TCHAR* Parameter);
 	bool					Stop();
-	void					EnableChannels();
-	void					DisableChannels();
+	void					ResumeChannels();
+	void					PauseChannels();
+	void					EnableCommandlineChannels();
 	void					SetTruncateFile(bool bTruncateFile);
 	void					UpdateCsvStats() const;
 	void					StartWorkerThread();
@@ -102,16 +107,20 @@ private:
 	};
 
 	void					AddChannel(const TCHAR* Name);
-	void					AddChannels(const TCHAR* Name, bool bResolvePresets);
-	void					EnableChannel(FChannel& Channel);
+	void					RemoveChannel(const TCHAR* Name);
+	template <class T> void ForEachChannel(const TCHAR* ChannelList, bool bResolvePresets, T Callable);
+	static uint32			HashChannelName(const TCHAR* Name);
+	bool					EnableChannel(const TCHAR* Channel);
+	void					DisableChannel(const TCHAR* Channel);
 	bool					SendToHost(const TCHAR* Host);
 	bool					WriteToFile(const TCHAR* Path=nullptr);
 
-	TMap<uint32, FChannel>	Channels;
+	TMap<uint32, FChannel>	CommandlineChannels;
 	FString					TraceDest;
 	EState					State = EState::Stopped;
 	bool					bTruncateFile = false;
 	bool					bWorkerThreadStarted = false;
+	FString					PausedPreset;
 };
 
 static FTraceAuxiliaryImpl GTraceAuxiliary;
@@ -119,15 +128,28 @@ static FDelegateHandle GEndFrameDelegateHandle;
 static FDelegateHandle GEndFrameStatDelegateHandle;
 
 ////////////////////////////////////////////////////////////////////////////////
-void FTraceAuxiliaryImpl::AddChannels(const TCHAR* ChannelList)
+void FTraceAuxiliaryImpl::AddCommandlineChannels(const TCHAR* ChannelList)
 {
-	AddChannels(ChannelList, true);
+	ForEachChannel(ChannelList, true, &FTraceAuxiliaryImpl::AddChannel);
 }
 
 ////////////////////////////////////////////////////////////////////////////////
-void FTraceAuxiliaryImpl::AddChannels(const TCHAR* ChannelList, bool bResolvePresets)
+void FTraceAuxiliaryImpl::EnableChannels(const TCHAR* ChannelList)
 {
-	UE::String::ParseTokens(ChannelList, TEXT(","), [this, bResolvePresets] (const FStringView& Token)
+	ForEachChannel(ChannelList, true, &FTraceAuxiliaryImpl::EnableChannel);
+}
+
+////////////////////////////////////////////////////////////////////////////////
+void FTraceAuxiliaryImpl::DisableChannels(const TCHAR* ChannelList)
+{
+	ForEachChannel(ChannelList, true, &FTraceAuxiliaryImpl::DisableChannel);
+}
+
+////////////////////////////////////////////////////////////////////////////////
+template<typename T>
+void FTraceAuxiliaryImpl::ForEachChannel(const TCHAR* ChannelList, bool bResolvePresets, T Callable)
+{
+	UE::String::ParseTokens(ChannelList, TEXT(","), [this, bResolvePresets, Callable] (const FStringView& Token)
 	{
 		TCHAR Name[80];
 		const size_t ChannelNameSize = Token.CopyString(Name, UE_ARRAY_COUNT(Name) - 1);
@@ -139,26 +161,26 @@ void FTraceAuxiliaryImpl::AddChannels(const TCHAR* ChannelList, bool bResolvePre
 			// Check against hard coded presets
 			if(FCString::Stricmp(Name, TEXT("default")) == 0)
 			{
-				AddChannels(GDefaultChannels, false);
+				ForEachChannel(GDefaultChannels, false, Callable);
 			}
 			else if(FCString::Stricmp(Name, TEXT("memory"))== 0)
 			{
-				AddChannels(GMemoryChannels, false);
+				ForEachChannel(GMemoryChannels, false, Callable);
 			}
 			// Check against data driven presets (if available)
 			else if (GConfig && GConfig->GetString(TEXT("Trace.ChannelPresets"), Name, Value, GEngineIni))
 			{
-				AddChannels(*Value, false);
+				ForEachChannel(*Value, false, Callable);
 				return;
 			}
 		}
 
-		AddChannel(Name);
+		Invoke(Callable, this, Name);
 	});
 }
 
 ////////////////////////////////////////////////////////////////////////////////
-void FTraceAuxiliaryImpl::AddChannel(const TCHAR* Name)
+uint32 FTraceAuxiliaryImpl::HashChannelName(const TCHAR* Name)
 {
 	uint32 Hash = 5381;
 	for (const TCHAR* c = Name; *c; ++c)
@@ -166,18 +188,43 @@ void FTraceAuxiliaryImpl::AddChannel(const TCHAR* Name)
 		uint32 LowerC = *c | 0x20;
         Hash = ((Hash << 5) + Hash) + LowerC;
 	}
+	return Hash;
+}
 
-	if (Channels.Find(Hash) != nullptr)
+////////////////////////////////////////////////////////////////////////////////
+void FTraceAuxiliaryImpl::AddChannel(const TCHAR* Name)
+{
+	uint32 Hash = HashChannelName(Name);
+
+	if (CommandlineChannels.Find(Hash) != nullptr)
 	{
 		return;
 	}
 
-	FChannel& Value = Channels.Add(Hash, {});
+	FChannel& Value = CommandlineChannels.Add(Hash, {});
 	Value.Name = Name;
 
-	if (State >= EState::Tracing)
+	if (State >= EState::Tracing && !Value.bActive)
+	{
+		Value.bActive = EnableChannel(*Value.Name);
+	}
+}
+
+////////////////////////////////////////////////////////////////////////////////
+void FTraceAuxiliaryImpl::RemoveChannel(const TCHAR* Name)
+{
+	uint32 Hash = HashChannelName(Name);
+
+	FChannel Channel;
+	if (!CommandlineChannels.RemoveAndCopyValue(Hash, Channel))
 	{
-		EnableChannel(Value);
+		return;
+	}
+
+	if (State >= EState::Tracing && Channel.bActive)
+	{
+		DisableChannel(*Channel.Name);
+		Channel.bActive = false;
 	}
 }
 
@@ -222,12 +269,12 @@ bool FTraceAuxiliaryImpl::Connect(ETraceConnectType Type, const TCHAR* Parameter
 
 	// We're now connected. If we don't appear to have any channels we'll set
 	// some defaults for the user. Less futzing.
-	if (!Channels.Num())
+	if (!CommandlineChannels.Num())
 	{
-		AddChannels(GDefaultChannels);
+		AddCommandlineChannels(GDefaultChannels);
 	}
 
-	EnableChannels();
+	ResumeChannels();
 
 	State = EState::Tracing;
 	return true;
@@ -241,63 +288,81 @@ bool FTraceAuxiliaryImpl::Stop()
 		return false;
 	}
 
-	DisableChannels();
 	State = EState::Stopped;
 	TraceDest.Reset();
 	return true;
 }
 
 ////////////////////////////////////////////////////////////////////////////////
-void FTraceAuxiliaryImpl::EnableChannel(FChannel& Channel)
+bool FTraceAuxiliaryImpl::EnableChannel(const TCHAR* Channel)
 {
-	if (Channel.bActive)
+	// Channel names have been provided by the user and may not exist yet. As
+	// we want to maintain bActive accurately (channels toggles are reference
+	// counted), we will first check Trace knows of the channel.
+	if (!UE::Trace::IsChannel(Channel))
 	{
-		return;
+		return false;
 	}
 
+	EPlatformEvent Event = PlatformEvents_GetEvent(Channel);
+	if (Event != EPlatformEvent::None)
+	{
+		PlatformEvents_Enable(Event);
+	}
+
+	return UE::Trace::ToggleChannel(Channel, true);
+}
+
+////////////////////////////////////////////////////////////////////////////////
+void FTraceAuxiliaryImpl::DisableChannel(const TCHAR* Channel)
+{
 	// Channel names have been provided by the user and may not exist yet. As
 	// we want to maintain bActive accurately (channels toggles are reference
 	// counted), we will first check Trace knows of the channel.
-	if (!UE::Trace::IsChannel(*Channel.Name))
+	if (!UE::Trace::IsChannel(Channel))
 	{
 		return;
 	}
 
-	EPlatformEvent Event = PlatformEvents_GetEvent(Channel.Name);
+	EPlatformEvent Event = PlatformEvents_GetEvent(Channel);
 	if (Event != EPlatformEvent::None)
 	{
-		PlatformEvents_Enable(Event);
+		PlatformEvents_Disable(Event);
 	}
 
-	UE::Trace::ToggleChannel(*Channel.Name, true);
-	Channel.bActive = true;
+	UE::Trace::ToggleChannel(Channel, false);
 }
 
 ////////////////////////////////////////////////////////////////////////////////
-void FTraceAuxiliaryImpl::EnableChannels()
+void FTraceAuxiliaryImpl::ResumeChannels()
 {
-	for (auto& ChannelPair : Channels)
+	// There was paused preset apply that
+	if (!PausedPreset.IsEmpty())
 	{
-		EnableChannel(ChannelPair.Value);
+		ForEachChannel(*PausedPreset, false, &FTraceAuxiliaryImpl::EnableChannel);
 	}
 }
 
 ////////////////////////////////////////////////////////////////////////////////
-void FTraceAuxiliaryImpl::DisableChannels()
+void FTraceAuxiliaryImpl::PauseChannels()
+{
+	// Collect current set of active channels
+	TStringBuilder<128> EnabledChannels;
+	GetActiveChannelsString(EnabledChannels);
+	PausedPreset = EnabledChannels.ToString();
+
+	// Disable those channels
+	ForEachChannel(*PausedPreset, true, &FTraceAuxiliaryImpl::DisableChannel);
+}
+
+////////////////////////////////////////////////////////////////////////////////
+void FTraceAuxiliaryImpl::EnableCommandlineChannels()
 {
-	for (auto& ChannelPair : Channels)
+	for (auto& ChannelPair : CommandlineChannels)
 	{
-		FChannel& Channel = ChannelPair.Value;
-		if (Channel.bActive)
+		if (!ChannelPair.Value.bActive)
 		{
-			UE::Trace::ToggleChannel(*Channel.Name, false);
-			Channel.bActive = false;
-
-			EPlatformEvent Event = PlatformEvents_GetEvent(Channel.Name);
-			if (Event != EPlatformEvent::None)
-			{
-				PlatformEvents_Disable(Event);
-			}
+			ChannelPair.Value.bActive = EnableChannel(*ChannelPair.Value.Name);
 		}
 	}
 }
@@ -384,13 +449,23 @@ const TCHAR* FTraceAuxiliaryImpl::GetDest() const
 }
 
 ////////////////////////////////////////////////////////////////////////////////
-template <class T>
-void FTraceAuxiliaryImpl::ReadChannels(T&& Callback) const
+bool FTraceAuxiliaryImpl::IsConnected() const
+{
+	return State == EState::Tracing;
+}
+
+////////////////////////////////////////////////////////////////////////////////
+void FTraceAuxiliaryImpl::GetActiveChannelsString(FStringBuilderBase& String) const
 {
-	for (const auto& ChannelPair : Channels)
+	UE::Trace::EnumerateChannels([](const ANSICHAR* Name, bool bEnabled, void* User)
 	{
-		Callback(*(ChannelPair.Value.Name));
-	}
+		FStringBuilderBase& EnabledChannelsStr = *static_cast<FStringBuilderBase*>(User);
+		if (bEnabled)
+		{
+			FAnsiStringView NameView = FAnsiStringView(Name).LeftChop(7); // Remove "Channel" suffix
+			EnabledChannelsStr << NameView << TEXT(",");
+		}
+	}, &String);
 }
 
 ////////////////////////////////////////////////////////////////////////////////
@@ -463,18 +538,11 @@ static void TraceAuxiliaryConnectEpilogue()
 	}
 
 	// Give the user some feedback that everything's underway.
-	FString Channels;
-	GTraceAuxiliary.ReadChannels([&Channels] (const TCHAR* Channel)
-	{
-		if (Channels.Len())
-		{
-			Channels += TEXT(",");
-		}
-
-		Channels += Channel;
-	});
+	TStringBuilder<128> Channels;
+	GTraceAuxiliary.GetActiveChannelsString(Channels);
+	
 	UE_LOG(LogConsoleResponse, Log, TEXT("Tracing to: %s"), TraceDest);
-	UE_LOG(LogConsoleResponse, Log, TEXT("Trace channels: %s"), *Channels);
+	UE_LOG(LogConsoleResponse, Log, TEXT("Trace channels: %s"), Channels.ToString());
 }
 
 ////////////////////////////////////////////////////////////////////////////////
@@ -498,12 +566,32 @@ static void TraceAuxiliarySend(const TArray<FString>& Args)
 }
 
 ////////////////////////////////////////////////////////////////////////////////
-static void TraceAuxiliaryStart(const TArray<FString>& Args)
+static void TraceAuxiliaryFile(const TArray<FString>& Args)
 {
-	const TCHAR* Channels = Args.Num() > 0 ? *Args[0] : nullptr;
+	const TCHAR* Channels = nullptr;
+	const TCHAR* File = nullptr;
+	if (Args.Num() == 2)
+	{
+		File = *Args[0];
+		Channels = *Args[1];
+	}
+	else if (Args.Num() == 1)
+	{
+		// Try to detect if the first argument is a file or path
+		if (!FCString::Strchr(*Args[0], TEXT('/')) && !FCString::Strchr(*Args[0], TEXT('\\'))
+			&& !FCString::Strchr(*Args[0], TEXT('.')) && !FCString::Strchr(*Args[0], TEXT(':')) )
+		{
+			Channels = *Args[0];
+		}
+		else
+		{
+			File = *Args[0];
+		}
+	}
+	
 	FTraceAuxiliary::Options Opts;
 	Opts.bNoWorkerThread = true;
-	if (!FTraceAuxiliary::Start(FTraceAuxiliary::EConnectionType::File, nullptr, Channels, &Opts))
+	if (!FTraceAuxiliary::Start(FTraceAuxiliary::EConnectionType::File, File, Channels, &Opts))
 	{
 		UE_LOG(LogConsoleResponse, Warning, TEXT("Failed to start tracing to a file"));
 		return;
@@ -512,6 +600,13 @@ static void TraceAuxiliaryStart(const TArray<FString>& Args)
 	TraceAuxiliaryConnectEpilogue();
 }
 
+////////////////////////////////////////////////////////////////////////////////
+static void TraceAuxiliaryStart(const TArray<FString>& Args)
+{
+	UE_LOG(LogConsoleResponse, Warning, TEXT("'Trace.Start' is being deprecated in favor of 'Trace.File'."));
+	TraceAuxiliaryFile(Args);
+}
+
 ////////////////////////////////////////////////////////////////////////////////
 static void TraceAuxiliaryStop()
 {
@@ -523,14 +618,14 @@ static void TraceAuxiliaryStop()
 static void TraceAuxiliaryPause()
 {
 	UE_LOG(LogConsoleResponse, Log, TEXT("Tracing paused"));
-	GTraceAuxiliary.DisableChannels();
+	GTraceAuxiliary.PauseChannels();
 }
 
 ////////////////////////////////////////////////////////////////////////////////
 static void TraceAuxiliaryResume()
 {
 	UE_LOG(LogConsoleResponse, Log, TEXT("Tracing resumed"));
-	GTraceAuxiliary.EnableChannels();
+	GTraceAuxiliary.ResumeChannels();
 }
 
 ////////////////////////////////////////////////////////////////////////////////
@@ -605,10 +700,43 @@ static void TraceAuxiliaryStatus()
 	UE_LOG(LogConsoleResponse, Display, TEXT("-------------------------------------------"));
 }
 
+////////////////////////////////////////////////////////////////////////////////
+static void TraceAuxiliaryEnableChannels(const TArray<FString>& Args)
+{
+	if (Args.Num() == 0)
+	{
+		UE_LOG(LogConsoleResponse, Warning, TEXT("Need to provide at least one channel."));
+		return;
+	}
+	GTraceAuxiliary.EnableChannels(*Args[0]);
+
+	TStringBuilder<128> EnabledChannels;
+	GTraceAuxiliary.GetActiveChannelsString(EnabledChannels);
+	UE_LOG(LogConsoleResponse, Display, TEXT("Enabled channels: %s"), EnabledChannels.ToString());
+}
+
+////////////////////////////////////////////////////////////////////////////////
+static void TraceAuxiliaryDisableChannels(const TArray<FString>& Args)
+{
+	if (Args.Num() == 0)
+	{
+		GTraceAuxiliary.PauseChannels();
+	}
+	else
+	{
+		GTraceAuxiliary.DisableChannels(*Args[0]);
+	}
+	
+	TStringBuilder<128> EnabledChannels;
+	GTraceAuxiliary.GetActiveChannelsString(EnabledChannels);
+	UE_LOG(LogConsoleResponse, Display, TEXT("Enabled channels: %s"), EnabledChannels.ToString());
+}
+
 ////////////////////////////////////////////////////////////////////////////////
 static FAutoConsoleCommand TraceAuxiliarySendCmd(
 	TEXT("Trace.Send"),
-	TEXT("Send trace data to the trace store; Trace.Send <Host> [ChannelSet]"),
+	TEXT("Send trace data to the trace store; Trace.Send <Host> [ChannelSet]"
+		" where ChannelSet is either comma-separated list of trace channels, a Config/Trace.ChannelPresets key, or optional."),
 	FConsoleCommandWithArgsDelegate::CreateStatic(TraceAuxiliarySend)
 );
 
@@ -616,12 +744,23 @@ static FAutoConsoleCommand TraceAuxiliarySendCmd(
 static FAutoConsoleCommand TraceAuxiliaryStartCmd(
 	TEXT("Trace.Start"),
 	TEXT(
-		"Begin tracing profiling events to a file; Trace.Start [ChannelSet]"
+		"(Deprecated: Use Trace.File) Begin tracing profiling events to a file; Trace.Start [ChannelSet]"
 		" where ChannelSet is either comma-separated list of trace channels, a Config/Trace.ChannelPresets key, or optional."
 	),
 	FConsoleCommandWithArgsDelegate::CreateStatic(TraceAuxiliaryStart)
 );
 
+////////////////////////////////////////////////////////////////////////////////
+static FAutoConsoleCommand TraceAuxiliaryFileCmd(
+	TEXT("Trace.File"),
+	TEXT(
+		"Begin tracing profiling events to a file; Trace.File [Filepath.utrace] [ChannelSet]"
+		" where ChannelSet is either comma-separated list of trace channels, a Config/Trace.ChannelPresets key, or optional."
+		" Either Filepath or ChannelSet can be excluded."
+	),
+	FConsoleCommandWithArgsDelegate::CreateStatic(TraceAuxiliaryFile)
+);
+
 ////////////////////////////////////////////////////////////////////////////////
 static FAutoConsoleCommand TraceAuxiliaryStopCmd(
 	TEXT("Trace.Stop"),
@@ -649,6 +788,24 @@ static FAutoConsoleCommand TraceAuxiliaryStatusCmd(
 	TEXT("Print Trace status to console"),
 	FConsoleCommandDelegate::CreateStatic(TraceAuxiliaryStatus)
 );
+
+////////////////////////////////////////////////////////////////////////////////
+static FAutoConsoleCommand TraceAuxiliaryChannelEnableCmd(
+	TEXT("Trace.Enable"),
+	TEXT("Enable a set of channels; Trace.Enable [ChannelSet]"
+		" where ChannelSet is either comma-separated list of trace channels, a Config/Trace.ChannelPresets key, or optional."),
+	FConsoleCommandWithArgsDelegate::CreateStatic(TraceAuxiliaryEnableChannels)
+);
+
+////////////////////////////////////////////////////////////////////////////////
+static FAutoConsoleCommand TraceAuxiliaryChannelDisableCmd(
+	TEXT("Trace.Disable"),
+	TEXT("Disable a set of channels; Trace.Disable [ChannelSet]"
+		" where ChannelSet is either comma-separated list of trace channels, a Config/Trace.ChannelPresets key, or optional."
+		" If no channels are specified all channels are disabled."),
+	FConsoleCommandWithArgsDelegate::CreateStatic(TraceAuxiliaryDisableChannels)
+);
+
 #endif // UE_TRACE_ENABLED
 
 
@@ -898,6 +1055,12 @@ static void SetupChannelsFromCommandline(const TCHAR* CommandLine)
 bool FTraceAuxiliary::Start(EConnectionType Type, const TCHAR* Target, const TCHAR* Channels, Options* Options)
 {
 #if UE_TRACE_ENABLED
+	if (GTraceAuxiliary.IsConnected())
+	{
+		UE_LOG(LogCore, Error, TEXT("Unable to start trace, already tracing to %s"), GTraceAuxiliary.GetDest());
+		return false;
+	}
+	
 	// Make sure the worker thread is started unless explicitly opt out.
 	if (!Options || !Options->bNoWorkerThread)
 	{
@@ -907,8 +1070,8 @@ bool FTraceAuxiliary::Start(EConnectionType Type, const TCHAR* Target, const TCH
 	if (Channels)
 	{
 		UE_LOG(LogCore, Display, TEXT("Trace channels: '%s'"), Channels);
-		GTraceAuxiliary.AddChannels(Channels);
-		GTraceAuxiliary.EnableChannels();
+		GTraceAuxiliary.AddCommandlineChannels(Channels);
+		GTraceAuxiliary.EnableCommandlineChannels();
 	}
 	
 	if (Options)
@@ -946,7 +1109,7 @@ bool FTraceAuxiliary::Stop()
 bool FTraceAuxiliary::Pause()
 {
 #if UE_TRACE_ENABLED
-	GTraceAuxiliary.DisableChannels();
+	GTraceAuxiliary.PauseChannels();
 #endif
 	return true;
 }
@@ -955,7 +1118,7 @@ bool FTraceAuxiliary::Pause()
 bool FTraceAuxiliary::Resume()
 {
 #if UE_TRACE_ENABLED
-	GTraceAuxiliary.EnableChannels();
+	GTraceAuxiliary.ResumeChannels();
 #endif
 	return true;
 }
@@ -1057,7 +1220,7 @@ void FTraceAuxiliary::Initialize(const TCHAR* CommandLine)
 	{
 		if (Reason == EModuleChangeReason::ModuleLoaded)
 		{
-			GTraceAuxiliary.EnableChannels();
+			GTraceAuxiliary.EnableCommandlineChannels();
 		}
 	});
 	
@@ -1077,8 +1240,8 @@ void FTraceAuxiliary::InitializePresets(const TCHAR* CommandLine)
 	FString Parameter;
 	if (FParse::Value(CommandLine, TEXT("-trace="), Parameter, false))
 	{
-		GTraceAuxiliary.AddChannels(*Parameter);
-		GTraceAuxiliary.EnableChannels();
+		GTraceAuxiliary.AddCommandlineChannels(*Parameter);
+		GTraceAuxiliary.EnableCommandlineChannels();
 	}
 #endif 
 }
@@ -1097,7 +1260,7 @@ void FTraceAuxiliary::Shutdown()
 void FTraceAuxiliary::EnableChannels()
 {
 #if UE_TRACE_ENABLED
-	GTraceAuxiliary.EnableChannels();
+	GTraceAuxiliary.EnableCommandlineChannels();
 #endif
 }
 
diff --git a/Engine/Source/Runtime/TraceLog/Private/Trace/Channel.cpp b/Engine/Source/Runtime/TraceLog/Private/Trace/Channel.cpp
--- a/Engine/Source/Runtime/TraceLog/Private/Trace/Channel.cpp
+++ b/Engine/Source/Runtime/TraceLog/Private/Trace/Channel.cpp
@@ -251,7 +251,7 @@ void FChannel::EnumerateChannels(ChannelIterFunc Func, void* User)
 bool FChannel::Toggle(bool bEnabled)
 {
 	using namespace Private;
-	AtomicAddRelaxed(&Enabled, bEnabled ? 1 : -1);
+	AtomicStoreRelaxed(&Enabled, bEnabled ? 1 : -1);
 
 	UE_TRACE_LOG(Trace, ChannelToggle, TraceLogChannel)
 		<< ChannelToggle.Id(Name.Hash)
