[Insights] Fix compile issue and bug where commandline trace channels would be lost.

diff --git a/Engine/Source/Runtime/Core/Private/ProfilingDebugging/TraceAuxiliary.cpp b/Engine/Source/Runtime/Core/Private/ProfilingDebugging/TraceAuxiliary.cpp
--- a/Engine/Source/Runtime/Core/Private/ProfilingDebugging/TraceAuxiliary.cpp
+++ b/Engine/Source/Runtime/Core/Private/ProfilingDebugging/TraceAuxiliary.cpp
@@ -88,6 +88,7 @@ public:
 	void					ResumeChannels();
 	void					PauseChannels();
 	void					EnableCommandlineChannels();
+	void					EnableCommandlineChannelsPostInitialize();
 	void					SetTruncateFile(bool bTruncateFile);
 	void					UpdateCsvStats() const;
 	void					StartWorkerThread();
@@ -367,6 +368,16 @@ void FTraceAuxiliaryImpl::EnableCommandlineChannels()
 	}
 }
 
+////////////////////////////////////////////////////////////////////////////////
+void FTraceAuxiliaryImpl::EnableCommandlineChannelsPostInitialize()
+{
+	for (auto& ChannelPair : CommandlineChannels)
+	{
+		// Intentionally enable channel without checking current state.
+		ChannelPair.Value.bActive = EnableChannel(*ChannelPair.Value.Name);
+	}
+}
+
 ////////////////////////////////////////////////////////////////////////////////
 void FTraceAuxiliaryImpl::SetTruncateFile(bool bNewTruncateFileState)
 {
@@ -1021,36 +1032,11 @@ static bool StartFromCommandlineArguments(const TCHAR* CommandLine)
 	}
 	
 	// Finally start tracing to the requested connection
-	return FTraceAuxiliary::Start(Type, Target, nullptr, &Opts);
+	return FTraceAuxiliary::Start(Type, Target, *Channels, &Opts);
 #endif
 	return false;
 }
 
-////////////////////////////////////////////////////////////////////////////////
-static void SetupChannelsFromCommandline(const TCHAR* CommandLine)
-{
-#if UE_TRACE_ENABLED
-	// Get active channels
-	FString Channels;
-	if (FParse::Value(CommandLine, TEXT("-trace="), Channels, false))
-	{
-	}
-	else if (FParse::Param(CommandLine, TEXT("trace")))
-	{
-		Channels = TEXT("default");
-	}
-	else
-	{
-		return;
-	}
-	
-	GTraceAuxiliary.AddChannels(*Channels);
-	GTraceAuxiliary.EnableChannels();
-	
-	UE_LOG(LogCore, Display, TEXT("Trace channels: %s"), *Channels);
-#endif
-}
-
 ////////////////////////////////////////////////////////////////////////////////
 bool FTraceAuxiliary::Start(EConnectionType Type, const TCHAR* Target, const TCHAR* Channels, Options* Options)
 {
@@ -1225,9 +1211,10 @@ void FTraceAuxiliary::Initialize(const TCHAR* CommandLine)
 	});
 	
 	UE::Trace::ThreadRegister(TEXT("GameThread"), FPlatformTLS::GetCurrentThreadId(), -1);
-
-	SetupChannelsFromCommandline(CommandLine);
-
+	
+	// Workaround for the fact that even if StartFromCommandlineArguments will enable channels
+	// specified by the commandline, UE::Trace::Initialize will reset all channels.
+	GTraceAuxiliary.EnableCommandlineChannelsPostInitialize();
 #endif
 }
 
