Fixed pointer to unexpected type beeing passed to CaptureStackBackTrace() (Contributed by [REDACTED_PERSON])

#8769

diff --git a/Engine/Source/Runtime/Core/Private/HAL/MallocCallstackHandler.cpp b/Engine/Source/Runtime/Core/Private/HAL/MallocCallstackHandler.cpp
--- a/Engine/Source/Runtime/Core/Private/HAL/MallocCallstackHandler.cpp
+++ b/Engine/Source/Runtime/Core/Private/HAL/MallocCallstackHandler.cpp
@@ -146,11 +146,10 @@ int32 FMallocCallstackHandler::GetCallStackIndex()
 	// Capture callstack and create FCallStackMapKey.
 	uint64 FullCallStack[MaxCallStackDepth + CallStackEntriesToSkipCount] = { 0 };
 	// CRC is filled in by the CaptureStackBackTrace, not all platforms calculate this for you.
-	uint32 CRC = 0;
-	FPlatformStackWalk::CaptureStackBackTrace(&FullCallStack[0], MaxCallStackDepth + CallStackEntriesToSkipCount, &CRC);
+	FPlatformStackWalk::CaptureStackBackTrace(&FullCallStack[0], MaxCallStackDepth + CallStackEntriesToSkipCount);
 	// Skip first n entries as they are inside the allocator.
 	uint64* CallStack = &FullCallStack[CallStackEntriesToSkipCount];
-	FCallStackMapKey CallStackMapKey(CRC, CallStack);
+	FCallStackMapKey CallStackMapKey(CallStack);
 	RWLock.ReadLock();
 	int32* IndexPtr = CallStackMapKeyToCallStackIndexMap.Find(CallStackMapKey);
 	if (IndexPtr)
diff --git a/Engine/Source/Runtime/Core/Public/HAL/MallocCallstackHandler.h b/Engine/Source/Runtime/Core/Public/HAL/MallocCallstackHandler.h
--- a/Engine/Source/Runtime/Core/Public/HAL/MallocCallstackHandler.h
+++ b/Engine/Source/Runtime/Core/Public/HAL/MallocCallstackHandler.h
@@ -8,6 +8,7 @@
 #include "HAL/MemoryBase.h"
 #include "Containers/Array.h"
 #include "Containers/Map.h"
+#include "Misc/Crc.h"
 
 class CORE_API FMallocCallstackHandler : public FMalloc
 {
@@ -142,13 +143,13 @@ public:
 	/** Used as a key in our current allocations/freed allocations maps*/
 	struct FCallStackMapKey
 	{
-		uint32 CRC; /* Generated by the platforms capture call stack function, if the platform you are working with doesn't generate this for you crc32 the call stack array manually */
-		uint64 *CallStack;
+		uint32 CRC;
+		uint64* CallStack;
 
-		FCallStackMapKey(uint32 InCRC, uint64* InCallStack)
-			: CRC(InCRC),
-			CallStack(InCallStack)
+		FCallStackMapKey(uint64* InCallStack)
+			: CallStack(InCallStack)
 		{
+			CRC = FCrc::MemCrc32(InCallStack, MaxCallStackDepth * sizeof(uint64));
 		}
 
 		friend bool operator==(const FCallStackMapKey& A, const FCallStackMapKey& B)
