FrameRate: allow creating FFrameRates from NTSC fractional FPS value strings

This enables parsing an NTSC frame rate as an FPS string (i.e. "23.976", "29.97", or "59.94")
into the appropriate FFrameRate value (24000/1001, 30000/1001, or 60000/1001, respectively).

Attempting to parse a string with any other fractional FPS value will still result in an error
as it did previously.

diff --git a/Engine/Source/Runtime/Core/Private/Misc/FrameRate.cpp b/Engine/Source/Runtime/Core/Private/Misc/FrameRate.cpp
--- a/Engine/Source/Runtime/Core/Private/Misc/FrameRate.cpp
+++ b/Engine/Source/Runtime/Core/Private/Misc/FrameRate.cpp
@@ -15,6 +15,15 @@ namespace
 
 	const TCHAR* const FFramesPerSecond::Moniker = TEXT("fps");
 	const TCHAR* const FSeconds::Moniker         = TEXT("s");
+
+	const double NTSC_24_FPS = 23.976;
+	const double NTSC_30_FPS = 29.97;
+	const double NTSC_60_FPS = 59.94;
+
+	const uint32 NTSC_24_Numerator = 24000u;
+	const uint32 NTSC_30_Numerator = 30000u;
+	const uint32 NTSC_60_Numerator = 60000u;
+	const uint32 NTSC_Denominator  = 1001u;
 }
 
 DEFINE_EXPRESSION_NODE_TYPE(FFrameRate, 0x4EDAA92F, 0xB75E4B9E, 0xB7E0ABC2, 0x8D981FCB)
@@ -137,6 +146,20 @@ private:
 		double RoundedFPS = FMath::RoundToDouble(InFPS);
 		if (RoundedFPS != InFPS)
 		{
+			// Allow creating FFrameRates for fractional FPS values only if they match an NTSC frame rate.
+			if (FMath::IsNearlyEqual(InFPS, NTSC_24_FPS, DOUBLE_KINDA_SMALL_NUMBER))
+			{
+				return MakeValue(FFrameRate(NTSC_24_Numerator, NTSC_Denominator));
+			}
+			else if (FMath::IsNearlyEqual(InFPS, NTSC_30_FPS, DOUBLE_KINDA_SMALL_NUMBER))
+			{
+				return MakeValue(FFrameRate(NTSC_30_Numerator, NTSC_Denominator));
+			}
+			else if (FMath::IsNearlyEqual(InFPS, NTSC_60_FPS, DOUBLE_KINDA_SMALL_NUMBER))
+			{
+				return MakeValue(FFrameRate(NTSC_60_Numerator, NTSC_Denominator));
+			}
+
 			return MakeError(FText::Format(LOCTEXT("FractionalFrameRate_Format", "Fractional FPS specified: {0}.\nPlease use x/y notation to define such framerates."), InFPS));
 		}
 
