Fix the ability to append preview platforms to existing data driven platforms.

diff --git a/Engine/Source/Runtime/Core/Private/Misc/DataDrivenPlatformInfoRegistry.cpp b/Engine/Source/Runtime/Core/Private/Misc/DataDrivenPlatformInfoRegistry.cpp
--- a/Engine/Source/Runtime/Core/Private/Misc/DataDrivenPlatformInfoRegistry.cpp
+++ b/Engine/Source/Runtime/Core/Private/Misc/DataDrivenPlatformInfoRegistry.cpp
@@ -216,6 +216,17 @@ static void DDPIGetString(const FConfigFile& IniFile, const TCHAR* Key, FString&
 	}
 }
 
+static void DDPIGetGuid(const FConfigFile& IniFile, const TCHAR* Key, FGuid& OutGuid)
+{
+	FString StringData = DDPITryRedirect(IniFile, Key);
+
+	// if we ended up with a string, convert it, otherwise leave it alone
+	if (StringData.Len() > 0)
+	{
+		OutGuid = FGuid(StringData);
+	}
+}
+
 static void DDPIGetStringArray(const FConfigFile& IniFile, const TCHAR* Key, TArray<FString>& OutArray)
 {
 	// we don't support redirecting arrays
@@ -299,9 +310,7 @@ static void LoadDDPIIniSettings(const FConfigFile& IniFile, FDataDrivenPlatformI
 	DDPIGetBool(IniFile, TEXT("Freezing_bAlignBases"), Info.Freezing_bAlignBases);
 	DDPIGetBool(IniFile, TEXT("Freezing_bWithRayTracing"), Info.Freezing_bWithRayTracing);
 
-	FString GuidString;
-	DDPIGetString(IniFile, TEXT("GlobalIdentifier"), GuidString);
-	Info.GlobalIdentifier = FGuid(GuidString);
+	DDPIGetGuid(IniFile, TEXT("GlobalIdentifier"), Info.GlobalIdentifier);
 	checkf(Info.GlobalIdentifier != FGuid(), TEXT("Platform %s didn't have a valid GlobalIdentifier set in DataDrivenPlatformInfo.ini"), *PlatformName.ToString());
 
 	// NOTE: add more settings here!
