Add Linux -crashhandlerstacksize command line option to set crash handler stack size

Also move stack to mmap'd memory instead of alloc'ing in heap

Should allow us to set & reduce size of 200k stack significantly via command line.

Brandon.Schaefer

diff --git a/Engine/Source/Runtime/Core/Private/Unix/UnixPlatformCrashContext.cpp b/Engine/Source/Runtime/Core/Private/Unix/UnixPlatformCrashContext.cpp
--- a/Engine/Source/Runtime/Core/Private/Unix/UnixPlatformCrashContext.cpp
+++ b/Engine/Source/Runtime/Core/Private/Unix/UnixPlatformCrashContext.cpp
@@ -32,6 +32,8 @@
 #include "HAL/ThreadHeartBeat.h"
 #include "BuildSettings.h"
 
+#include <sys/mman.h>
+
 #include <atomic>
 
 extern CORE_API bool GIsGPUCrashed;
@@ -1008,8 +1010,39 @@ void FUnixPlatformMisc::SetGracefulTerminationHandler()
 	sigaction(SIGHUP, &Action, nullptr);	//  this should actually cause the server to just re-read configs (restart?)
 }
 
-// reserve stack for the main thread in BSS
-char FRunnableThreadUnix::MainThreadSignalHandlerStack[FRunnableThreadUnix::EConstants::CrashHandlerStackSize];
+// Stack pointer for the main thread
+void *FRunnableThreadUnix::MainThreadSignalHandlerStack = nullptr;
+
+void *FRunnableThreadUnix::AllocCrashHandlerStack()
+{
+	uint64 StackBufferSize = GetCrashHandlerStackSize();
+
+	return mmap(nullptr, StackBufferSize, PROT_READ | PROT_WRITE, MAP_PRIVATE | MAP_ANON, -1, 0);
+}
+
+void FRunnableThreadUnix::FreeCrashHandlerStack(void *StackBuffer)
+{
+	if (StackBuffer)
+	{
+		munmap(StackBuffer, GetCrashHandlerStackSize());
+	}
+}
+
+// Defined in UnixPlatformMemory, set via -crashhandlerstacksize command line.
+extern uint64 GCrashHandlerStackSize;
+
+uint64 FRunnableThreadUnix::GetCrashHandlerStackSize()
+{
+	if (GCrashHandlerStackSize == 0)
+	{
+		GCrashHandlerStackSize = EConstants::CrashHandlerStackSize;
+	}
+	else if (GCrashHandlerStackSize < EConstants::CrashHandlerStackSizeMin)
+	{
+		GCrashHandlerStackSize = EConstants::CrashHandlerStackSizeMin;
+	}
+	return GCrashHandlerStackSize;
+}
 
 // Defined in UnixPlatformMemory.cpp. Allows settings a specific signal to maintain its default handler rather then ignoring it
 extern int32 GSignalToDefault;
@@ -1096,5 +1129,10 @@ void FUnixPlatformMisc::SetCrashHandler(void (* CrashHandler)(const FGenericCras
 
 	checkf(IsInGameThread(), TEXT("Crash handler for the game thread should be set from the game thread only."));
 
-	FRunnableThreadUnix::SetupSignalHandlerStack(FRunnableThreadUnix::MainThreadSignalHandlerStack, sizeof(FRunnableThreadUnix::MainThreadSignalHandlerStack), nullptr);
+	if (!FRunnableThreadUnix::MainThreadSignalHandlerStack)
+	{
+		FRunnableThreadUnix::MainThreadSignalHandlerStack = FRunnableThreadUnix::AllocCrashHandlerStack();
+	}
+
+	FRunnableThreadUnix::SetupSignalHandlerStack(FRunnableThreadUnix::MainThreadSignalHandlerStack, FRunnableThreadUnix::GetCrashHandlerStackSize(), nullptr);
 }
diff --git a/Engine/Source/Runtime/Core/Private/Unix/UnixPlatformMemory.cpp b/Engine/Source/Runtime/Core/Private/Unix/UnixPlatformMemory.cpp
--- a/Engine/Source/Runtime/Core/Private/Unix/UnixPlatformMemory.cpp
+++ b/Engine/Source/Runtime/Core/Private/Unix/UnixPlatformMemory.cpp
@@ -64,6 +64,9 @@ bool CORE_API GTimeEnsures = true;
 // Allows settings a specific signal to maintain its default handler rather then ignoring the signal
 int32 CORE_API GSignalToDefault = 0;
 
+// Allows setting crash handler stack size
+uint64 CORE_API GCrashHandlerStackSize = 0;
+
 // Due to dotnet not allowing any files marked as LOCK_EX to be opened for read only or copied, this allows us to
 // to disable the locking mechanics. [REDACTED_PROFILE_URL]/runtime/issues/34126
 // Default to true, can be disabled with -noexclusivelockonwrite
@@ -235,6 +238,12 @@ class FMalloc* FUnixPlatformMemory::BaseAllocator()
 					GSignalToDefault = FMath::Max(SignalToDefault, 0);
 				}
 
+				const char CrashHandlerStackSize[] = "-crashhandlerstacksize=";
+				if (const char* Cmd = FCStringAnsi::Stristr(Arg, CrashHandlerStackSize))
+				{
+					GCrashHandlerStackSize = FCStringAnsi::Atoi64(Cmd + sizeof(CrashHandlerStackSize) - 1);
+				}
+
 				const char FileMapCacheCmd[] = "-filemapcachesize=";
 				if (const char* Cmd = FCStringAnsi::Stristr(Arg, FileMapCacheCmd))
 				{
diff --git a/Engine/Source/Runtime/Core/Private/Unix/UnixPlatformMisc.cpp b/Engine/Source/Runtime/Core/Private/Unix/UnixPlatformMisc.cpp
--- a/Engine/Source/Runtime/Core/Private/Unix/UnixPlatformMisc.cpp
+++ b/Engine/Source/Runtime/Core/Private/Unix/UnixPlatformMisc.cpp
@@ -121,6 +121,7 @@ void UnixPlatform_UpdateCacheLineSize()
 // Defined in UnixPlatformMemory
 extern bool GUseKSM;
 extern bool GKSMMergeAllPages;
+extern uint64 GCrashHandlerStackSize;
 
 static void UnixPlatForm_CheckIfKSMUsable()
 {
@@ -282,6 +283,7 @@ void FUnixPlatformMisc::PlatformInit()
 	UE_LOG(LogInit, Log, TEXT(" -ksmmergeall - marks all mmap'd memory pages suitable for KSM (%s)"), GKSMMergeAllPages ? TEXT("ON") : TEXT("OFF"));
 	UE_LOG(LogInit, Log, TEXT(" -preloadmodulesymbols - Loads the main module symbols file into memory (%s)"), bPreloadedModuleSymbolFile ? TEXT("ON") : TEXT("OFF"));
 	UE_LOG(LogInit, Log, TEXT(" -sigdfl=SIGNAL - Allows a specific signal to be set to its default handler rather then ignoring the signal"));
+	UE_LOG(LogInit, Log, TEXT(" -crashhandlerstacksize - Allows setting crash handler stack sizes (%lu)"), GCrashHandlerStackSize);
 	UE_LOG(LogInit, Log, TEXT(" -noexclusivelockonwrite - disables marking files created by the engine as exclusive locked while the engine has them opened"));
 
 	// [RCL] FIXME: this should be printed in specific modules, if at all
diff --git a/Engine/Source/Runtime/Core/Public/Unix/UnixPlatformRunnableThread.h b/Engine/Source/Runtime/Core/Public/Unix/UnixPlatformRunnableThread.h
--- a/Engine/Source/Runtime/Core/Public/Unix/UnixPlatformRunnableThread.h
+++ b/Engine/Source/Runtime/Core/Public/Unix/UnixPlatformRunnableThread.h
@@ -28,12 +28,14 @@ class FRunnableThreadUnix : public FRunnableThreadPThread
 		enum
 		{
 			UnixThreadNameLimit = 15,			// the limit for thread name is just 15 chars :( http://man7.org/linux/man-pages/man3/pthread_setname_np.3.html
-			CrashHandlerStackSize = SIGSTKSZ + 192 * 1024	// should be at least SIGSTKSIZE, plus 192K because we do logging and symbolication in crash handler
+
+			CrashHandlerStackSize = SIGSTKSZ + 192 * 1024,	// should be at least SIGSTKSIZE, plus 192K because we do logging and symbolication in crash handler
+			CrashHandlerStackSizeMin = SIGSTKSZ + 8 * 1024  // minimum allowed stack size
 		};
 	};
 
 	/** Each thread needs a separate stack for the signal handler, so possible stack overflows in the thread are handled */
-	char ThreadCrashHandlingStack[EConstants::CrashHandlerStackSize];
+	void* ThreadCrashHandlingStack;
 
 	/** Address of stack guard page - if nullptr, the page wasn't set */
 	void* StackGuardPageAddress;
@@ -47,10 +49,15 @@ class FRunnableThreadUnix : public FRunnableThreadPThread
 public:
 
 	/** Separate stack for the signal handler (so possible stack overflows don't go unnoticed), for the main thread specifically. */
-	static char MainThreadSignalHandlerStack[EConstants::CrashHandlerStackSize];
+	static void *MainThreadSignalHandlerStack;
+
+	static void *AllocCrashHandlerStack();
+	static void FreeCrashHandlerStack(void *StackBuffer);
+	static uint64 GetCrashHandlerStackSize();
 
 	FRunnableThreadUnix()
 		:	FRunnableThreadPThread()
+		,	ThreadCrashHandlingStack(nullptr)
 		,	StackGuardPageAddress(nullptr)
 		,	BaselineNiceValue(0)
 		,	bGotBaselineNiceValue(false)
@@ -77,6 +84,13 @@ public:
 	 */
 	static bool SetupSignalHandlerStack(void* StackBuffer, const size_t StackBufferSize, void** OutStackGuardPageAddress)
 	{
+		if (!StackBuffer)
+		{
+			return false;
+		}
+
+		// Added by CL 11188846 for ASan, TSan, and UBSan
+		// #jira UE-62784 UE-62803 UE-62804
 		if (FORCE_ANSI_ALLOCATOR)
 		{
 			return true;
@@ -254,7 +268,9 @@ private:
 		}
 
 		// set the alternate stack for handling crashes due to stack overflow
-		SetupSignalHandlerStack(ThreadCrashHandlingStack, sizeof(ThreadCrashHandlingStack), &StackGuardPageAddress);
+		check(ThreadCrashHandlingStack == nullptr);
+		ThreadCrashHandlingStack = AllocCrashHandlerStack();
+		SetupSignalHandlerStack(ThreadCrashHandlingStack, FRunnableThreadUnix::GetCrashHandlerStackSize(), &StackGuardPageAddress);
 	}
 
 	virtual void PostRun() override
@@ -271,6 +287,12 @@ private:
 
 			StackGuardPageAddress = nullptr;
 		}
+
+		if (ThreadCrashHandlingStack != nullptr)
+		{
+			FreeCrashHandlerStack(ThreadCrashHandlingStack);
+			ThreadCrashHandlingStack = nullptr;
+		}
 	}
 
 	/**
