REL - Move the Inhert prioiry hint after Cout because Count is miss used to encode another meaning.

diff --git a/Engine/Source/Runtime/Core/Public/Async/Fundamental/Task.h b/Engine/Source/Runtime/Core/Public/Async/Fundamental/Task.h
--- a/Engine/Source/Runtime/Core/Public/Async/Fundamental/Task.h
+++ b/Engine/Source/Runtime/Core/Public/Async/Fundamental/Task.h
@@ -20,8 +20,8 @@ namespace LowLevelTasks
 		BackgroundHigh = ForegroundCount,
 		BackgroundNormal,
 		BackgroundLow,
-		Inherit,
-		Count = Inherit,
+		Count,
+		Inherit, //Inherit the TaskPriority from the launching Task or the Default Priority if not launched from a Task.
 	};
 
 	inline const TCHAR* ToString(ETaskPriority Priority)
@@ -134,7 +134,7 @@ namespace LowLevelTasks
 			uintptr_t bAllowBusyWaiting : 1;
 			
 		public:
-			FPackedData() : FPackedData(nullptr,  ETaskPriority::Inherit, ETaskState::Completed, true)
+			FPackedData() : FPackedData(nullptr,  ETaskPriority::Default, ETaskState::Completed, true)
 			{
 				static_assert(!PLATFORM_32BITS, "32bit Platforms are not supported");
 				static_assert(uintptr_t(ETaskPriority::Count) <= (1ull << 3), "Not enough bits to store ETaskPriority");
