Ship robustness: Add a fatal assert if FName text is too long in HashLowerCase. This can occur with adversarial corruption of AssetRegistry.bin and could be taken advantage of for injected code execution.
#rnx

diff --git a/Engine/Source/Runtime/Core/Private/UObject/UnrealNames.cpp b/Engine/Source/Runtime/Core/Private/UObject/UnrealNames.cpp
--- a/Engine/Source/Runtime/Core/Private/UObject/UnrealNames.cpp
+++ b/Engine/Source/Runtime/Core/Private/UObject/UnrealNames.cpp
@@ -596,6 +596,9 @@ template<class CharType>
 FORCENOINLINE FNameHash HashLowerCase(const CharType* Str, uint32 Len)
 {
 	CharType LowerStr[NAME_SIZE];
+	UE_CLOG(Len > NAME_SIZE, LogCore, Fatal,
+		TEXT("FName length is too long; HashLowerCase value is undefined. Length = %d, MaxLength = %d, Name=%.*s..."),
+		Len, NAME_SIZE, NAME_SIZE, StringCast<TCHAR>(Str, NAME_SIZE).Get());
 	for (uint32 I = 0; I < Len; ++I)
 	{
 		LowerStr[I] = TChar<CharType>::ToLower(Str[I]);
