Add a warning for when the FileCache block size doesn't match any iostore container compressed block sizes.

diff --git a/Engine/Source/Runtime/Core/Private/FileCache/FileCache.cpp b/Engine/Source/Runtime/Core/Private/FileCache/FileCache.cpp
--- a/Engine/Source/Runtime/Core/Private/FileCache/FileCache.cpp
+++ b/Engine/Source/Runtime/Core/Private/FileCache/FileCache.cpp
@@ -283,6 +283,52 @@ private:
 
 ///////////////
 
+#if !UE_BUILD_SHIPPING
+
+static std::atomic_bool bFileCacheInitialized = false;
+static std::atomic_uint GIoStoreCompressionBlockSize = 0; // 0 means no iostore has reported any.
+static std::atomic_bool bIoStoreCompressionBlockSizeMultiple = false; // if we get different ones, we by definition can't match.
+
+//
+// This exists to log warnings when projects are misconfigured to have FileCache block sizes
+// different than iostore compression block sizes. The difficulty is that iostore containers
+// are mounted _very early_ - before the cvars are resolved - so we don't know what the final
+// FileCache block size will be. However, once in the FileCache constructor, there's no code api
+// to get us access to the underlying iostore containers - file iostore containers are above core.
+// So we have to have that code call over to us when it's loaded, and we just keep track of what
+// compression blocks we've seen.
+//
+// Once we initialize, we check and log as necessary.
+//
+// And of course, iostore containers can be mounted _after_ the file cache is initialized,
+// so we have to check for that and log immediately in that case.
+//
+void FileCache_PostIoStoreCompressionBlockSize(uint32 InCompressionBlockSize, FString const& InContainerFilePath)
+{
+	if (bFileCacheInitialized.load())
+	{
+		// We can direct check since CacheSlotID is correct (cvars resolved).
+		if (InCompressionBlockSize != CacheSlotID::GetSize())
+		{
+			UE_LOG(LogStreamingFileCache, Warning, TEXT("IoStore container %s has a different block sizes than FileCache (%d vs %d)!"), *InContainerFilePath, InCompressionBlockSize, CacheSlotID::GetSize());
+			UE_LOG(LogStreamingFileCache, Warning, TEXT("	Check your IoStore compression block size (Project Settings -> 'Package Compression Commandline Options'"));
+			UE_LOG(LogStreamingFileCache, Warning, TEXT("	and your File Cache block size (fc.BlockSize cvar). They should match!"));
+		}
+		return;
+	}
+
+	// otherwise, save off the value.
+	uint32 LastBlockSize = GIoStoreCompressionBlockSize.exchange(InCompressionBlockSize);
+	if (LastBlockSize && // we had a value
+		LastBlockSize != InCompressionBlockSize) // and it was different
+	{
+		// Mark that we are dealing with more than one compression blocks size.
+		bIoStoreCompressionBlockSizeMultiple.store(true);
+	}
+}
+
+#endif //!UE_BUILD_SHIPPING
+
 FFileCache::FFileCache(int32 NumSlots)
 	: EvictFileCacheCommand(TEXT("r.VT.EvictFileCache"), TEXT("Evict all the file caches in the VT system."),
 		FConsoleCommandDelegate::CreateRaw(this, &FFileCache::EvictFileCacheFromConsole))
@@ -291,6 +337,26 @@ FFileCache::FFileCache(int32 NumSlots)
 {
 	LLM_SCOPE(ELLMTag::FileSystem);
 
+#if !UE_BUILD_SHIPPING
+	//
+	// If we aren't shipping, check and see if we have a mismatch vs the compression
+	// block size for any mounted iostore containers.
+	//
+	uint32 IoStoreCompressionBlockSize = GIoStoreCompressionBlockSize.load();
+	if (bIoStoreCompressionBlockSizeMultiple.load())
+	{
+		UE_LOG(LogStreamingFileCache, Warning, TEXT("IoStore containers have multiple compression block sizes! This means the FileCache block size must be misaligned with at least one!"));
+		UE_LOG(LogStreamingFileCache, Warning, TEXT("	Check your IoStore compression block size (Project Settings -> 'Package Compression Commandline Options'"));
+		UE_LOG(LogStreamingFileCache, Warning, TEXT("	and your File Cache block size (fc.BlockSize cvar). They should match!"));
+	}
+	else if (IoStoreCompressionBlockSize != CacheSlotID::GetSize())
+	{
+		UE_LOG(LogStreamingFileCache, Warning, TEXT("IoStore containers have a different block sizes than FileCache (%d vs %d)!"), IoStoreCompressionBlockSize, CacheSlotID::GetSize());
+		UE_LOG(LogStreamingFileCache, Warning, TEXT("	Check your IoStore compression block size (Project Settings -> 'Package Compression Commandline Options'"));
+		UE_LOG(LogStreamingFileCache, Warning, TEXT("	and your File Cache block size (fc.BlockSize cvar). They should match!"));
+	}
+#endif // !UE_BUILD_SHIPPING
+
 	Memory = (uint8*)FMemory::Malloc(SizeInBytes);
 
 	SlotInfo.AddUninitialized(NumSlots + 1);
@@ -307,6 +373,10 @@ FFileCache::FFileCache(int32 NumSlots)
 	// list is circular
 	SlotInfo[0].PrevSlotIndex = NumSlots;
 	SlotInfo[NumSlots].NextSlotIndex = 0;
+
+#if !UE_BUILD_SHIPPING
+	bFileCacheInitialized.store(true);
+#endif
 }
 
 CacheSlotID FFileCache::AcquireAndLockSlot(FFileCacheHandle* InHandle, CacheLineID InLineID)
diff --git a/Engine/Source/Runtime/Core/Public/FileCache/FileCache.h b/Engine/Source/Runtime/Core/Public/FileCache/FileCache.h
--- a/Engine/Source/Runtime/Core/Public/FileCache/FileCache.h
+++ b/Engine/Source/Runtime/Core/Public/FileCache/FileCache.h
@@ -64,3 +64,11 @@ public:
 	 */
 	virtual void WaitAll() = 0;
 };
+
+#if !UE_BUILD_SHIPPING
+//
+// Called by the IoStore system to inform the FileCache of used compression block sizes,
+// for usage warning purposes. (FFileIoStoreReader::ReadContainerHeader)
+//
+void CORE_API FileCache_PostIoStoreCompressionBlockSize(uint32 InCompressionBlockSize, FString const& InContainerFilePath);
+#endif // !UE_BUILD_SHIPPING
\ No newline at end of file
diff --git a/Engine/Source/Runtime/PakFile/Private/IoDispatcherFileBackend.cpp b/Engine/Source/Runtime/PakFile/Private/IoDispatcherFileBackend.cpp
--- a/Engine/Source/Runtime/PakFile/Private/IoDispatcherFileBackend.cpp
+++ b/Engine/Source/Runtime/PakFile/Private/IoDispatcherFileBackend.cpp
@@ -23,6 +23,7 @@
 #include "Modules/ModuleManager.h"
 #include "IO/IoContainerHeader.h"
 #include "Serialization/MemoryReader.h"
+#include "FileCache/FileCache.h"
 
 //PRAGMA_DISABLE_OPTIMIZATION
 
@@ -896,6 +897,10 @@ TIoStatusOr<FIoContainerHeader> FFileIoStoreReader::ReadContainerHeader() const
 	const FFileIoStoreContainerFilePartition& Partition = ContainerFile.Partitions[PartitionIndex];
 	const uint64 RawOffset = CompressionBlockEntry->GetOffset() % ContainerFile.PartitionSize;
 
+#if !UE_BUILD_SHIPPING
+	FileCache_PostIoStoreCompressionBlockSize(CompressionBlockSize, Partition.FilePath);
+#endif
+
 	FIoBuffer IoBuffer(Align(Size, FAES::AESBlockSize));
 	IPlatformFile& Ipf = FPlatformFileManager::Get().GetPlatformFile();
 	TUniquePtr<IFileHandle> ContainerFileHandle(Ipf.OpenRead(*Partition.FilePath));
