Some derived Windows platforms can override Windows OS Version functionality.

 - GetVersionEx WinAPI function does not work as expected on Manifested Applications.

#rnx

diff --git a/Engine/Source/Runtime/Core/Private/Windows/WindowsPlatformMisc.cpp b/Engine/Source/Runtime/Core/Private/Windows/WindowsPlatformMisc.cpp
--- a/Engine/Source/Runtime/Core/Private/Windows/WindowsPlatformMisc.cpp
+++ b/Engine/Source/Runtime/Core/Private/Windows/WindowsPlatformMisc.cpp
@@ -301,6 +301,7 @@ int32 GetOSVersionsHelper( TCHAR* OutOSVersionLabel, int32 OSVersionLabelLength,
 	return ErrorCode;
 }
 
+#if WINDOWS_USE_DEFAULT_OSVERSIONHELPER
 int32 FWindowsOSVersionHelper::GetOSVersions( FString& OutOSVersionLabel, FString& OutOSSubVersionLabel )
 {
 	TCHAR OSVersionLabel[128];
@@ -316,6 +317,7 @@ int32 FWindowsOSVersionHelper::GetOSVersions( FString& OutOSVersionLabel, FStrin
 
 	return Result;
 }
+#endif //WINDOWS_USE_DEFAULT_OSVERSIONHELPER
 
 namespace
 {
diff --git a/Engine/Source/Runtime/Core/Public/Windows/WindowsPlatform.h b/Engine/Source/Runtime/Core/Public/Windows/WindowsPlatform.h
--- a/Engine/Source/Runtime/Core/Public/Windows/WindowsPlatform.h
+++ b/Engine/Source/Runtime/Core/Public/Windows/WindowsPlatform.h
@@ -108,6 +108,7 @@ typedef FWindowsPlatformTypes FPlatformTypes;
 #define WINDOWS_USE_FEATURE_PLATFORMHTTP_CLASS				1
 #define WINDOWS_USE_FEATURE_LAUNCH							1
 #define WINDOWS_USE_FEATURE_DYNAMIC_RHI						1
+#define WINDOWS_USE_DEFAULT_OSVERSIONHELPER					1
 
 // Q: Why is there a __nop() before __debugbreak()?
 // A: VS' debug engine has a bug where it will silently swallow explicit
