Fix for UE4 references in UE5 codebase

diff --git a/Engine/Source/Runtime/Core/Private/Serialization/StructuredArchive.cpp b/Engine/Source/Runtime/Core/Private/Serialization/StructuredArchive.cpp
--- a/Engine/Source/Runtime/Core/Private/Serialization/StructuredArchive.cpp
+++ b/Engine/Source/Runtime/Core/Private/Serialization/StructuredArchive.cpp
@@ -86,7 +86,7 @@ FStructuredArchiveSlot FStructuredArchive::Open()
 	check(!CurrentSlotElementId.IsValid());
 
 	RootElementId = ElementIdGenerator.Generate();
-	CurrentScope.Emplace(RootElementId, UE4StructuredArchive_Private::EElementType::Root);
+	CurrentScope.Emplace(RootElementId, StructuredArchive_Private::EElementType::Root);
 
 	CurrentSlotElementId = ElementIdGenerator.Generate();
 
@@ -95,16 +95,16 @@ FStructuredArchiveSlot FStructuredArchive::Open()
 
 void FStructuredArchive::Close()
 {
-	SetScope(UE4StructuredArchive_Private::FSlotPosition(0, RootElementId));
+	SetScope(StructuredArchive_Private::FSlotPosition(0, RootElementId));
 }
 
-void FStructuredArchive::EnterSlot(UE4StructuredArchive_Private::FSlotPosition Slot, bool bEnteringAttributedValue)
+void FStructuredArchive::EnterSlot(StructuredArchive_Private::FSlotPosition Slot, bool bEnteringAttributedValue)
 {
 	int32                                    ParentDepth = Slot.Depth;
-	UE4StructuredArchive_Private::FElementId ElementId   = Slot.ElementId;
+	StructuredArchive_Private::FElementId ElementId   = Slot.ElementId;
 
 	// If the slot being entered has attributes, enter the value slot first.
-	if (ParentDepth + 1 < CurrentScope.Num() && CurrentScope[ParentDepth + 1].Id == ElementId && CurrentScope[ParentDepth + 1].Type == UE4StructuredArchive_Private::EElementType::AttributedValue)
+	if (ParentDepth + 1 < CurrentScope.Num() && CurrentScope[ParentDepth + 1].Id == ElementId && CurrentScope[ParentDepth + 1].Type == StructuredArchive_Private::EElementType::AttributedValue)
 	{
 #if DO_STRUCTURED_ARCHIVE_CONTAINER_CHECKS
 		checkf(!CurrentSlotElementId.IsValid() && !CurrentContainer.Top()->bAttributedValueWritten, TEXT("Attempt to serialize data into an invalid slot"));
@@ -113,15 +113,15 @@ void FStructuredArchive::EnterSlot(UE4StructuredArchive_Private::FSlotPosition S
 		checkf(!CurrentSlotElementId.IsValid(), TEXT("Attempt to serialize data into an invalid slot"));
 #endif
 
-		SetScope(UE4StructuredArchive_Private::FSlotPosition(ParentDepth + 1, ElementId));
+		SetScope(StructuredArchive_Private::FSlotPosition(ParentDepth + 1, ElementId));
 		Formatter.EnterAttributedValueValue();
 	}
 	else if (!bEnteringAttributedValue && Formatter.TryEnterAttributedValueValue())
 	{
-		int32 NewDepth = EnterSlotAsType(UE4StructuredArchive_Private::FSlotPosition(ParentDepth, ElementId), UE4StructuredArchive_Private::EElementType::AttributedValue);
+		int32 NewDepth = EnterSlotAsType(StructuredArchive_Private::FSlotPosition(ParentDepth, ElementId), StructuredArchive_Private::EElementType::AttributedValue);
 		check(NewDepth == ParentDepth + 1);
-		UE4StructuredArchive_Private::FElementId AttributedValueId = CurrentScope[NewDepth].Id;
-		SetScope(UE4StructuredArchive_Private::FSlotPosition(NewDepth, AttributedValueId));
+		StructuredArchive_Private::FElementId AttributedValueId = CurrentScope[NewDepth].Id;
+		SetScope(StructuredArchive_Private::FSlotPosition(NewDepth, AttributedValueId));
 #if DO_STRUCTURED_ARCHIVE_CONTAINER_CHECKS
 		CurrentContainer.Emplace(0);
 #endif
@@ -132,12 +132,12 @@ void FStructuredArchive::EnterSlot(UE4StructuredArchive_Private::FSlotPosition S
 		CurrentSlotElementId.Reset();
 	}
 
-	CurrentEnteringAttributeState = UE4StructuredArchive_Private::EEnteringAttributeState::NotEnteringAttribute;
+	CurrentEnteringAttributeState = StructuredArchive_Private::EEnteringAttributeState::NotEnteringAttribute;
 }
 
-int32 FStructuredArchive::EnterSlotAsType(UE4StructuredArchive_Private::FSlotPosition Slot, UE4StructuredArchive_Private::EElementType ElementType)
+int32 FStructuredArchive::EnterSlotAsType(StructuredArchive_Private::FSlotPosition Slot, StructuredArchive_Private::EElementType ElementType)
 {
-	EnterSlot(Slot, ElementType == UE4StructuredArchive_Private::EElementType::AttributedValue);
+	EnterSlot(Slot, ElementType == StructuredArchive_Private::EElementType::AttributedValue);
 
 	int32 NewSlotDepth = Slot.Depth + 1;
 
@@ -146,8 +146,8 @@ int32 FStructuredArchive::EnterSlotAsType(UE4StructuredArchive_Private::FSlotPos
 	//
 	// We don't need to do adjust for attributes, because entering the attribute slot will bump the depth anyway.
 	if (NewSlotDepth < CurrentScope.Num() &&
-		CurrentScope[NewSlotDepth].Type == UE4StructuredArchive_Private::EElementType::AttributedValue &&
-		CurrentEnteringAttributeState == UE4StructuredArchive_Private::EEnteringAttributeState::NotEnteringAttribute)
+		CurrentScope[NewSlotDepth].Type == StructuredArchive_Private::EElementType::AttributedValue &&
+		CurrentEnteringAttributeState == StructuredArchive_Private::EEnteringAttributeState::NotEnteringAttribute)
 	{
 		++NewSlotDepth;
 	}
@@ -162,32 +162,32 @@ void FStructuredArchive::LeaveSlot()
 	{
 		switch (CurrentScope.Top().Type)
 		{
-		case UE4StructuredArchive_Private::EElementType::Record:
+		case StructuredArchive_Private::EElementType::Record:
 			Formatter.LeaveField();
 			break;
-		case UE4StructuredArchive_Private::EElementType::Array:
+		case StructuredArchive_Private::EElementType::Array:
 			Formatter.LeaveArrayElement();
 #if DO_STRUCTURED_ARCHIVE_CONTAINER_CHECKS
 			CurrentContainer.Top()->Index++;
 #endif
 			break;
-		case UE4StructuredArchive_Private::EElementType::Stream:
+		case StructuredArchive_Private::EElementType::Stream:
 			Formatter.LeaveStreamElement();
 			break;
-		case UE4StructuredArchive_Private::EElementType::Map:
+		case StructuredArchive_Private::EElementType::Map:
 			Formatter.LeaveMapElement();
 #if DO_STRUCTURED_ARCHIVE_CONTAINER_CHECKS
 			CurrentContainer.Top()->Index++;
 #endif
 			break;
-		case UE4StructuredArchive_Private::EElementType::AttributedValue:
+		case StructuredArchive_Private::EElementType::AttributedValue:
 			Formatter.LeaveAttribute();
 			break;
 		}
 	}
 }
 
-void FStructuredArchive::SetScope(UE4StructuredArchive_Private::FSlotPosition Slot)
+void FStructuredArchive::SetScope(StructuredArchive_Private::FSlotPosition Slot)
 {
 	// Make sure the scope is valid
 	checkf(Slot.Depth < CurrentScope.Num() && CurrentScope[Slot.Depth].Id == Slot.ElementId, TEXT("Invalid scope for writing to archive"));
@@ -202,13 +202,13 @@ void FStructuredArchive::SetScope(UE4StructuredArchive_Private::FSlotPosition Sl
 			const FElement& Element = CurrentScope[CurrentDepth];
 			switch (Element.Type)
 			{
-			case UE4StructuredArchive_Private::EElementType::Record:
+			case StructuredArchive_Private::EElementType::Record:
 				Formatter.LeaveRecord();
 #if DO_STRUCTURED_ARCHIVE_CONTAINER_CHECKS
 				CurrentContainer.Pop(false);
 #endif
 				break;
-			case UE4StructuredArchive_Private::EElementType::Array:
+			case StructuredArchive_Private::EElementType::Array:
 #if DO_STRUCTURED_ARCHIVE_CONTAINER_CHECKS
 				checkf(GetUnderlyingArchive().IsLoading() || CurrentContainer.Top()->Index == CurrentContainer.Top()->Count, TEXT("Incorrect number of elements serialized in array"));
 #endif
@@ -217,10 +217,10 @@ void FStructuredArchive::SetScope(UE4StructuredArchive_Private::FSlotPosition Sl
 				CurrentContainer.Pop(false);
 #endif
 				break;
-			case UE4StructuredArchive_Private::EElementType::Stream:
+			case StructuredArchive_Private::EElementType::Stream:
 				Formatter.LeaveStream();
 				break;
-			case UE4StructuredArchive_Private::EElementType::Map:
+			case StructuredArchive_Private::EElementType::Map:
 #if DO_STRUCTURED_ARCHIVE_CONTAINER_CHECKS
 				checkf(CurrentContainer.Top()->Index == CurrentContainer.Top()->Count, TEXT("Incorrect number of elements serialized in map"));
 #endif
@@ -229,7 +229,7 @@ void FStructuredArchive::SetScope(UE4StructuredArchive_Private::FSlotPosition Sl
 				CurrentContainer.Pop(false);
 #endif
 				break;
-			case UE4StructuredArchive_Private::EElementType::AttributedValue:
+			case StructuredArchive_Private::EElementType::AttributedValue:
 				Formatter.LeaveAttributedValue();
 #if DO_STRUCTURED_ARCHIVE_CONTAINER_CHECKS
 				CurrentContainer.Pop(false);
@@ -255,7 +255,7 @@ void FStructuredArchive::SetScope(UE4StructuredArchive_Private::FSlotPosition Sl
 
 FStructuredArchiveRecord FStructuredArchiveSlot::EnterRecord()
 {
-	int32 NewDepth = Ar.EnterSlotAsType(*this, UE4StructuredArchive_Private::EElementType::Record);
+	int32 NewDepth = Ar.EnterSlotAsType(*this, StructuredArchive_Private::EElementType::Record);
 
 #if DO_STRUCTURED_ARCHIVE_CONTAINER_CHECKS
 	Ar.CurrentContainer.Emplace(0);
@@ -268,7 +268,7 @@ FStructuredArchiveRecord FStructuredArchiveSlot::EnterRecord()
 
 FStructuredArchiveRecord FStructuredArchiveSlot::EnterRecord_TextOnly(TArray<FString>& OutFieldNames)
 {
-	int32 NewDepth = Ar.EnterSlotAsType(*this, UE4StructuredArchive_Private::EElementType::Record);
+	int32 NewDepth = Ar.EnterSlotAsType(*this, StructuredArchive_Private::EElementType::Record);
 
 #if DO_STRUCTURED_ARCHIVE_CONTAINER_CHECKS
 	Ar.CurrentContainer.Emplace(0);
@@ -281,7 +281,7 @@ FStructuredArchiveRecord FStructuredArchiveSlot::EnterRecord_TextOnly(TArray<FSt
 
 FStructuredArchiveArray FStructuredArchiveSlot::EnterArray(int32& Num)
 {
-	int32 NewDepth = Ar.EnterSlotAsType(*this, UE4StructuredArchive_Private::EElementType::Array);
+	int32 NewDepth = Ar.EnterSlotAsType(*this, StructuredArchive_Private::EElementType::Array);
 
 	Ar.Formatter.EnterArray(Num);
 
@@ -294,7 +294,7 @@ FStructuredArchiveArray FStructuredArchiveSlot::EnterArray(int32& Num)
 
 FStructuredArchiveStream FStructuredArchiveSlot::EnterStream()
 {
-	int32 NewDepth = Ar.EnterSlotAsType(*this, UE4StructuredArchive_Private::EElementType::Stream);
+	int32 NewDepth = Ar.EnterSlotAsType(*this, StructuredArchive_Private::EElementType::Stream);
 
 	Ar.Formatter.EnterStream();
 
@@ -303,7 +303,7 @@ FStructuredArchiveStream FStructuredArchiveSlot::EnterStream()
 
 FStructuredArchiveStream FStructuredArchiveSlot::EnterStream_TextOnly(int32& OutNumElements)
 {
-	int32 NewDepth = Ar.EnterSlotAsType(*this, UE4StructuredArchive_Private::EElementType::Stream);
+	int32 NewDepth = Ar.EnterSlotAsType(*this, StructuredArchive_Private::EElementType::Stream);
 
 	Ar.Formatter.EnterStream_TextOnly(OutNumElements);
 
@@ -312,7 +312,7 @@ FStructuredArchiveStream FStructuredArchiveSlot::EnterStream_TextOnly(int32& Out
 
 FStructuredArchiveMap FStructuredArchiveSlot::EnterMap(int32& Num)
 {
-	int32 NewDepth = Ar.EnterSlotAsType(*this, UE4StructuredArchive_Private::EElementType::Map);
+	int32 NewDepth = Ar.EnterSlotAsType(*this, StructuredArchive_Private::EElementType::Map);
 
 	Ar.Formatter.EnterMap(Num);
 
@@ -328,9 +328,9 @@ FStructuredArchiveSlot FStructuredArchiveSlot::EnterAttribute(FArchiveFieldName
 	check(Ar.CurrentScope.Num() > 0);
 
 	int32 NewDepth = Depth + 1;
-	if (NewDepth >= Ar.CurrentScope.Num() || Ar.CurrentScope[NewDepth].Id != ElementId || Ar.CurrentScope[NewDepth].Type != UE4StructuredArchive_Private::EElementType::AttributedValue)
+	if (NewDepth >= Ar.CurrentScope.Num() || Ar.CurrentScope[NewDepth].Id != ElementId || Ar.CurrentScope[NewDepth].Type != StructuredArchive_Private::EElementType::AttributedValue)
 	{
-		int32 NewDepthCheck = Ar.EnterSlotAsType(*this, UE4StructuredArchive_Private::EElementType::AttributedValue);
+		int32 NewDepthCheck = Ar.EnterSlotAsType(*this, StructuredArchive_Private::EElementType::AttributedValue);
 		checkSlow(NewDepth == NewDepthCheck);
 
 		Ar.Formatter.EnterAttributedValue();
@@ -340,11 +340,11 @@ FStructuredArchiveSlot FStructuredArchiveSlot::EnterAttribute(FArchiveFieldName
 #endif
 	}
 
-	Ar.CurrentEnteringAttributeState = UE4StructuredArchive_Private::EEnteringAttributeState::NotEnteringAttribute;
+	Ar.CurrentEnteringAttributeState = StructuredArchive_Private::EEnteringAttributeState::NotEnteringAttribute;
 
-	UE4StructuredArchive_Private::FElementId AttributedValueId = Ar.CurrentScope[NewDepth].Id;
+	StructuredArchive_Private::FElementId AttributedValueId = Ar.CurrentScope[NewDepth].Id;
 
-	Ar.SetScope(UE4StructuredArchive_Private::FSlotPosition(NewDepth, AttributedValueId));
+	Ar.SetScope(StructuredArchive_Private::FSlotPosition(NewDepth, AttributedValueId));
 
 	Ar.CurrentSlotElementId = Ar.ElementIdGenerator.Generate();
 
@@ -369,9 +369,9 @@ TOptional<FStructuredArchiveSlot> FStructuredArchiveSlot::TryEnterAttribute(FArc
 	check(Ar.CurrentScope.Num() > 0);
 
 	int32 NewDepth = Depth + 1;
-	if (NewDepth >= Ar.CurrentScope.Num() || Ar.CurrentScope[NewDepth].Id != ElementId || Ar.CurrentScope[NewDepth].Type != UE4StructuredArchive_Private::EElementType::AttributedValue)
+	if (NewDepth >= Ar.CurrentScope.Num() || Ar.CurrentScope[NewDepth].Id != ElementId || Ar.CurrentScope[NewDepth].Type != StructuredArchive_Private::EElementType::AttributedValue)
 	{
-		int32 NewDepthCheck = Ar.EnterSlotAsType(*this, UE4StructuredArchive_Private::EElementType::AttributedValue);
+		int32 NewDepthCheck = Ar.EnterSlotAsType(*this, StructuredArchive_Private::EElementType::AttributedValue);
 		checkSlow(NewDepth == NewDepthCheck);
 
 		Ar.Formatter.EnterAttributedValue();
@@ -392,13 +392,13 @@ TOptional<FStructuredArchiveSlot> FStructuredArchiveSlot::TryEnterAttribute(FArc
 #endif
 #endif
 
-	UE4StructuredArchive_Private::FElementId AttributedValueId = Ar.CurrentScope[NewDepth].Id;
+	StructuredArchive_Private::FElementId AttributedValueId = Ar.CurrentScope[NewDepth].Id;
 
-	Ar.SetScope(UE4StructuredArchive_Private::FSlotPosition(NewDepth, AttributedValueId));
+	Ar.SetScope(StructuredArchive_Private::FSlotPosition(NewDepth, AttributedValueId));
 
 	if (Ar.Formatter.TryEnterAttribute(AttributeName, bEnterWhenWriting))
 	{
-		Ar.CurrentEnteringAttributeState = UE4StructuredArchive_Private::EEnteringAttributeState::NotEnteringAttribute;
+		Ar.CurrentEnteringAttributeState = StructuredArchive_Private::EEnteringAttributeState::NotEnteringAttribute;
 
 		Ar.CurrentSlotElementId = Ar.ElementIdGenerator.Generate();
 
diff --git a/Engine/Source/Runtime/Core/Public/Serialization/StructuredArchive.h b/Engine/Source/Runtime/Core/Public/Serialization/StructuredArchive.h
--- a/Engine/Source/Runtime/Core/Public/Serialization/StructuredArchive.h
+++ b/Engine/Source/Runtime/Core/Public/Serialization/StructuredArchive.h
@@ -126,7 +126,7 @@ class FStructuredArchiveArray;
 class FStructuredArchiveStream;
 class FStructuredArchiveMap;
 
-namespace UE4StructuredArchive_Private
+namespace StructuredArchive_Private
 {
 	struct FElementId
 	{
@@ -229,7 +229,7 @@ namespace UE4StructuredArchive_Private
  * and can merely have a value serialized into it. That value may be a literal (eg. int, float) or compound object
  * (eg. object, array, map).
  */
-class CORE_API FStructuredArchiveSlot final : public UE4StructuredArchive_Private::FSlotBase
+class CORE_API FStructuredArchiveSlot final : public StructuredArchive_Private::FSlotBase
 {
 public:
 	FStructuredArchiveRecord EnterRecord();
@@ -313,14 +313,14 @@ private:
 	friend FStructuredArchiveStream;
 	friend FStructuredArchiveMap;
 
-	using UE4StructuredArchive_Private::FSlotBase::FSlotBase;
+	using StructuredArchive_Private::FSlotBase::FSlotBase;
 };
 
 /**
  * Represents a record in the structured archive. An object contains slots that are identified by FArchiveName,
  * which may be compiled out with binary-only archives.
  */
-class CORE_API FStructuredArchiveRecord final : public UE4StructuredArchive_Private::FSlotBase
+class CORE_API FStructuredArchiveRecord final : public StructuredArchive_Private::FSlotBase
 {
 public:
 	FStructuredArchiveSlot EnterField(FArchiveFieldName Name);
@@ -344,14 +344,14 @@ private:
 	friend FStructuredArchive;
 	friend FStructuredArchiveSlot;
 
-	using UE4StructuredArchive_Private::FSlotBase::FSlotBase;
+	using StructuredArchive_Private::FSlotBase::FSlotBase;
 };
 
 /**
  * Represents an array in the structured archive. An object contains slots that are identified by a FArchiveFieldName,
  * which may be compiled out with binary-only archives.
  */
-class CORE_API FStructuredArchiveArray final : public UE4StructuredArchive_Private::FSlotBase
+class CORE_API FStructuredArchiveArray final : public StructuredArchive_Private::FSlotBase
 {
 public:
 	FStructuredArchiveSlot EnterElement();
@@ -367,13 +367,13 @@ private:
 	friend FStructuredArchive;
 	friend FStructuredArchiveSlot;
 
-	using UE4StructuredArchive_Private::FSlotBase::FSlotBase;
+	using StructuredArchive_Private::FSlotBase::FSlotBase;
 };
 
 /**
  * Represents an unsized sequence of slots in the structured archive (similar to an array, but without a known size).
  */
-class CORE_API FStructuredArchiveStream final : public UE4StructuredArchive_Private::FSlotBase
+class CORE_API FStructuredArchiveStream final : public StructuredArchive_Private::FSlotBase
 {
 public:
 	FStructuredArchiveSlot EnterElement();
@@ -389,14 +389,14 @@ private:
 	friend FStructuredArchive;
 	friend FStructuredArchiveSlot;
 
-	using UE4StructuredArchive_Private::FSlotBase::FSlotBase;
+	using StructuredArchive_Private::FSlotBase::FSlotBase;
 };
 
 /**
  * Represents a map in the structured archive. A map is similar to a record, but keys can be read back out from an archive.
  * (This is an important distinction for binary archives).
  */
-class CORE_API FStructuredArchiveMap final : public UE4StructuredArchive_Private::FSlotBase
+class CORE_API FStructuredArchiveMap final : public StructuredArchive_Private::FSlotBase
 {
 public:
 	FStructuredArchiveSlot EnterElement(FString& Name);
@@ -406,7 +406,7 @@ private:
 	friend FStructuredArchive;
 	friend FStructuredArchiveSlot;
 
-	using UE4StructuredArchive_Private::FSlotBase::FSlotBase;
+	using StructuredArchive_Private::FSlotBase::FSlotBase;
 };
 
 /**
@@ -488,10 +488,10 @@ private:
 
 	struct FElement
 	{
-		UE4StructuredArchive_Private::FElementId Id;
-		UE4StructuredArchive_Private::EElementType Type;
+		StructuredArchive_Private::FElementId Id;
+		StructuredArchive_Private::EElementType Type;
 
-		FElement(UE4StructuredArchive_Private::FElementId InId, UE4StructuredArchive_Private::EElementType InType)
+		FElement(StructuredArchive_Private::FElementId InId, StructuredArchive_Private::EElementType InType)
 			: Id(InId)
 			, Type(InType)
 		{
@@ -500,9 +500,9 @@ private:
 
 	struct FIdGenerator
 	{
-		UE4StructuredArchive_Private::FElementId Generate()
+		StructuredArchive_Private::FElementId Generate()
 		{
-			return UE4StructuredArchive_Private::FElementId(NextId++);
+			return StructuredArchive_Private::FElementId(NextId++);
 		}
 
 	private:
@@ -517,12 +517,12 @@ private:
 	/**
 	 * The ID of the root element.
 	 */
-	UE4StructuredArchive_Private::FElementId RootElementId;
+	StructuredArchive_Private::FElementId RootElementId;
 
 	/**
 	 * The element ID assigned for the current slot. Slots are transient, and only exist as placeholders until something is written into them. This is reset to 0 when something is created in a slot, and the created item can assume the element id.
 	 */
-	UE4StructuredArchive_Private::FElementId CurrentSlotElementId;
+	StructuredArchive_Private::FElementId CurrentSlotElementId;
 
 	/**
 	 * Tracks the current stack of objects being written. Used by SetScope() to ensure that scopes are always closed correctly in the underlying formatter,
@@ -542,19 +542,19 @@ private:
 	/**
 	 * Whether or not we've just entered an attribute
 	 */
-	UE4StructuredArchive_Private::EEnteringAttributeState CurrentEnteringAttributeState = UE4StructuredArchive_Private::EEnteringAttributeState::NotEnteringAttribute;
+	StructuredArchive_Private::EEnteringAttributeState CurrentEnteringAttributeState = StructuredArchive_Private::EEnteringAttributeState::NotEnteringAttribute;
 
 	/**
 	 * Enters the current slot for serializing a value. Asserts if the archive is not in a state about to write to an empty-slot.
 	 */
-	void EnterSlot(UE4StructuredArchive_Private::FSlotPosition Slot, bool bEnteringAttributedValue = false);
+	void EnterSlot(StructuredArchive_Private::FSlotPosition Slot, bool bEnteringAttributedValue = false);
 
 	/**
 	 * Enters the current slot, adding an element onto the stack. Asserts if the archive is not in a state about to write to an empty-slot.
 	 *
 	 * @return  The depth of the newly-entered slot.
 	 */
-	int32 EnterSlotAsType(UE4StructuredArchive_Private::FSlotPosition Slot, UE4StructuredArchive_Private::EElementType ElementType);
+	int32 EnterSlotAsType(StructuredArchive_Private::FSlotPosition Slot, StructuredArchive_Private::EElementType ElementType);
 
 	/**
 	 * Leaves slot at the top of the current scope
@@ -564,16 +564,16 @@ private:
 	/**
 	 * Switches to the scope for the given slot.
 	 */
-	void SetScope(UE4StructuredArchive_Private::FSlotPosition Slot);
+	void SetScope(StructuredArchive_Private::FSlotPosition Slot);
 #endif
 };
 
-FORCEINLINE FArchive& UE4StructuredArchive_Private::FSlotBase::GetUnderlyingArchive() const
+FORCEINLINE FArchive& StructuredArchive_Private::FSlotBase::GetUnderlyingArchive() const
 {
 	return Ar.GetUnderlyingArchive();
 }
 
-FORCEINLINE const FArchiveState& UE4StructuredArchive_Private::FSlotBase::GetArchiveState() const
+FORCEINLINE const FArchiveState& StructuredArchive_Private::FSlotBase::GetArchiveState() const
 {
 	return Ar.GetArchiveState();
 }
