REL - fast ConcurrentLinear Allocator that can be used for temporary allocations

diff --git a/Engine/Source/Runtime/Core/Private/HAL/LowLevelMemTracker.cpp b/Engine/Source/Runtime/Core/Private/HAL/LowLevelMemTracker.cpp
--- a/Engine/Source/Runtime/Core/Private/HAL/LowLevelMemTracker.cpp
+++ b/Engine/Source/Runtime/Core/Private/HAL/LowLevelMemTracker.cpp
@@ -105,6 +105,7 @@ DECLARE_LLM_MEMORY_STAT(TEXT("GenericPlatformMallocCrash"), STAT_GenericPlatform
 DECLARE_LLM_MEMORY_STAT(TEXT("GenericPlatformMallocCrash"), STAT_GenericPlatformMallocCrashPlatformLLM, STATGROUP_LLMPlatform);
 DECLARE_LLM_MEMORY_STAT(TEXT("Engine Misc"), STAT_EngineMiscLLM, STATGROUP_LLMFULL);
 DECLARE_LLM_MEMORY_STAT(TEXT("TaskGraph Misc Tasks"), STAT_TaskGraphTasksMiscLLM, STATGROUP_LLMFULL);
+DECLARE_LLM_MEMORY_STAT(TEXT("Linear Allocator"), STAT_LinearAllocatorLLM, STATGROUP_LLMFULL);
 DECLARE_LLM_MEMORY_STAT(TEXT("Audio"), STAT_AudioLLM, STATGROUP_LLMFULL);
 DECLARE_LLM_MEMORY_STAT(TEXT("AudioMisc"), STAT_AudioMiscLLM, STATGROUP_LLMFULL);
 DECLARE_LLM_MEMORY_STAT(TEXT("AudioSoundWaves"), STAT_AudioSoundWavesLLM, STATGROUP_LLMFULL);
diff --git a/Engine/Source/Runtime/Core/Public/Experimental/ConcurrentLinearAllocator.h b/Engine/Source/Runtime/Core/Public/Experimental/ConcurrentLinearAllocator.h
new file mode 100644
--- /dev/null
+++ b/Engine/Source/Runtime/Core/Public/Experimental/ConcurrentLinearAllocator.h
@@ -0,0 +1,510 @@
+// Copyright Epic Games, Inc. All Rights Reserved.
+
+#pragma once
+#include <atomic>
+#include "CoreMinimal.h"
+#include "HAL/UnrealMemory.h"
+#include "HAL/LowLevelMemTracker.h"
+#include "Templates/UniquePtr.h"
+
+#if __has_include(<sanitizer/asan_interface.h>)
+#include <sanitizer/asan_interface.h>
+#else
+#define ASAN_POISON_MEMORY_REGION(addr, size) ((void)(addr), (void)(size))
+#define ASAN_UNPOISON_MEMORY_REGION(addr, size) ((void)(addr), (void)(size))
+#endif
+
+
+struct FDefaultBlockAllocationTag 
+{
+	typedef void* (*MallocType)(SIZE_T Size);
+	typedef void (*FreeType)(void* Pointer);
+
+	static constexpr MallocType MallocFunction = &FMemory::SystemMalloc;
+	static constexpr FreeType FreeFunction = &FMemory::SystemFree;
+	static constexpr uint32 BlockSize = 64 * 1024;		// Blocksize used to allocate from
+	static constexpr bool AllowOversizedBlocks = true;  // The allocator supports oversized Blocks and will store them in a seperate Block with counter 1
+	static constexpr const TCHAR* TagName = TEXT("DefaultLinear");
+};
+
+// This concurrent fast linear Allocator can be used for temporary allocations that are usually produced and consumed on different threads and within the lifetime of a frame.
+// Although the lifetime of any individual Allocation is not hard tied to a frame (tracking is done using the FBlockHeader::NumAllocations atomic variable) the Application will eventually run OOM if Allocations are not cleaned up in a timely manner.
+// The Allocator works by allocating a larger block in TLS which has a Header at the front which contains the Atomic and all allocations are than allocated from this block:
+//
+// ------------------------------------------------------------------------------------------------------------------------------------------------------------------------------------------------------------
+// | FBlockHeader(atomic counter etc.) | Alignment Waste | FAllocationHeader(size) | Memory used for Allocation | Alignment Waste | FAllocationHeader(size) | Memory used for Allocation | FreeSpace ... 
+// ------------------------------------------------------------------------------------------------------------------------------------------------------------------------------------------------------------
+template<typename BlockAllocationTag>
+class TConcurrentLinearAllocator
+{
+	struct FBlockHeader;
+
+	class FAllocationHeader
+	{
+	public:
+		FORCEINLINE FAllocationHeader(FBlockHeader* InBlockHeader, SIZE_T InAllocationSize) 
+		{
+			uintptr_t Offset = uintptr_t(this) - uintptr_t(InBlockHeader);
+			checkSlow(Offset < UINT32_MAX);
+			BlockHeaderOffset = uint32(Offset);
+
+			checkSlow(InAllocationSize < UINT32_MAX);
+			AllocationSize = uint32(InAllocationSize);
+		}
+
+		FORCEINLINE FBlockHeader* GetBlockHeader() const
+		{
+			return reinterpret_cast<FBlockHeader*>(uintptr_t(this) - BlockHeaderOffset);
+		}
+
+		FORCEINLINE SIZE_T GetAllocationSize() const
+		{
+			return SIZE_T(AllocationSize);
+		}
+
+	private:
+		uint32 BlockHeaderOffset; //this is the negative offset from the allocation to the BlockHeader.
+		uint32 AllocationSize; //this is the size of the allocation following the AllocationHeader.
+	};
+
+	static FORCEINLINE FAllocationHeader* GetAllocationHeader(void* Pointer)
+	{
+		return reinterpret_cast<FAllocationHeader*>(Pointer) - 1;
+	}
+
+	struct FBlockHeader
+	{
+		FORCEINLINE FBlockHeader() 
+			//We need to reserve space for at least one BlockHeader as well as one AllocationHeader
+			: NextAllocationPtr(uintptr_t(this) + sizeof(FBlockHeader) + sizeof(FAllocationHeader))
+		{
+		}
+
+		std::atomic_uint NumAllocations { UINT_MAX }; //this is shared between threads and tracks the number of live allocations (plus UINT_MAX) and we will fix this up when we close a block.
+		uint8 Padding[PLATFORM_CACHE_LINE_SIZE - sizeof(std::atomic_uint)]; //avoid false sharing
+		uintptr_t NextAllocationPtr; //this is the next address we are trying to allocate from.
+		unsigned int Num = 0; //this tracks the TLS local Number of allocation from a given Block
+	};
+
+public:
+	template<uint32 Alignment> //this can help the compiler folding away the alignment when it is statically known
+	static FORCEINLINE_DEBUGGABLE void* Malloc(SIZE_T Size)
+	{
+		return Malloc(Size, Alignment);
+	}
+
+	static FORCEINLINE_DEBUGGABLE void* Malloc(SIZE_T Size, uint32 Alignment)
+	{
+		checkSlow(Alignment >= 1 && FMath::IsPowerOfTwo(Alignment));
+		Alignment = (Alignment < alignof(FAllocationHeader)) ? alignof(FAllocationHeader) : Alignment;
+
+		static thread_local FBlockHeader* Header;
+		if (Header == nullptr)
+		{
+			static_assert(BlockAllocationTag::BlockSize >= sizeof(FBlockHeader) + sizeof(FAllocationHeader));
+			Header = new (BlockAllocationTag::MallocFunction(BlockAllocationTag::BlockSize)) FBlockHeader;
+			checkSlow(IsAligned(Header, alignof(FBlockHeader)));
+			ASAN_POISON_MEMORY_REGION( Header + 1, BlockAllocationTag::BlockSize - sizeof(FBlockHeader) );
+		}
+
+		for (;;)
+		{
+			uintptr_t AlignedOffset = Align(Header->NextAllocationPtr, Alignment);
+			if (AlignedOffset + Size <= uintptr_t(Header) + BlockAllocationTag::BlockSize)
+			{
+				Header->NextAllocationPtr = AlignedOffset + Size + sizeof(FAllocationHeader);
+				Header->Num++;
+				
+				FAllocationHeader* AllocationHeader = reinterpret_cast<FAllocationHeader*>(AlignedOffset) - 1;
+				ASAN_UNPOISON_MEMORY_REGION( AllocationHeader, sizeof(FAllocationHeader) + Size );
+				new (AllocationHeader) FAllocationHeader(Header, Size);
+				ASAN_POISON_MEMORY_REGION( AllocationHeader, sizeof(FAllocationHeader) );
+
+				LLM_IF_ENABLED( FLowLevelMemTracker::Get().OnLowLevelAlloc(ELLMTracker::Default, reinterpret_cast<void*>(AlignedOffset), Size, ELLMTag::LinearAllocator) );
+				return reinterpret_cast<void*>(AlignedOffset);
+			}
+
+			// The cold path of the code starts here
+			constexpr SIZE_T HeaderSize = sizeof(FBlockHeader) + sizeof(FAllocationHeader);
+			if constexpr (BlockAllocationTag::AllowOversizedBlocks)
+			{
+				//support for oversized Blocks
+				if (HeaderSize + Size + Alignment > BlockAllocationTag::BlockSize)
+				{
+					FBlockHeader* LargeHeader = new (BlockAllocationTag::MallocFunction(HeaderSize + Size + Alignment)) FBlockHeader;
+					checkSlow(IsAligned(LargeHeader, alignof(FBlockHeader)));
+
+					uintptr_t LargeAlignedOffset = Align(LargeHeader->NextAllocationPtr, Alignment);
+					LargeHeader->NextAllocationPtr = uintptr_t(LargeHeader) + HeaderSize + Size + Alignment;
+					LargeHeader->NumAllocations.store(1, std::memory_order_release);
+
+					checkSlow(LargeAlignedOffset + Size <= LargeHeader->NextAllocationPtr);
+					FAllocationHeader* AllocationHeader = new (reinterpret_cast<FAllocationHeader*>(LargeAlignedOffset) - 1) FAllocationHeader(LargeHeader, Size);
+					ASAN_POISON_MEMORY_REGION( AllocationHeader, sizeof(FAllocationHeader) );
+
+					LLM_IF_ENABLED( FLowLevelMemTracker::Get().OnLowLevelAlloc(ELLMTracker::Default, reinterpret_cast<void*>(LargeAlignedOffset), Size, ELLMTag::LinearAllocator) );
+					return reinterpret_cast<void*>(LargeAlignedOffset);
+				}
+			}
+			else
+			{
+				check(HeaderSize + Size + Alignment <= BlockAllocationTag::BlockSize);
+			}
+
+			//allocation of a new Block
+			Header->NextAllocationPtr = uintptr_t(Header) + BlockAllocationTag::BlockSize;
+			const uint32 DeltaCount = UINT_MAX - Header->Num; 
+
+			// on the allocating side we only need to do a single atomic to reduce contention with the deletions
+			// this will leave the atomic in a state where it only counts the number of live allocations (before it was based of UINT_MAX)
+			if (Header->NumAllocations.fetch_sub(DeltaCount, std::memory_order_release) == DeltaCount)
+			{
+				//if all allocations are already freed we can reuse the Block again
+				Header->~FBlockHeader();
+				Header = new (Header) FBlockHeader;
+				ASAN_POISON_MEMORY_REGION( Header + 1, BlockAllocationTag::BlockSize - sizeof(FBlockHeader) );
+				continue;
+			}
+
+			Header = new (BlockAllocationTag::MallocFunction(BlockAllocationTag::BlockSize)) FBlockHeader;
+			checkSlow(IsAligned(Header, alignof(FBlockHeader)));
+			ASAN_POISON_MEMORY_REGION( Header + 1, BlockAllocationTag::BlockSize - sizeof(FBlockHeader) );
+		}
+	}
+
+	static FORCEINLINE_DEBUGGABLE void Free(void* Pointer)
+	{
+		if(Pointer != nullptr)
+		{
+			LLM_IF_ENABLED(FLowLevelMemTracker::Get().OnLowLevelFree(ELLMTracker::Default, reinterpret_cast<void*>(Pointer)));
+
+			FAllocationHeader* AllocationHeader = GetAllocationHeader(Pointer);
+			ASAN_UNPOISON_MEMORY_REGION( AllocationHeader, sizeof(FAllocationHeader) );
+			FBlockHeader* Header = AllocationHeader->GetBlockHeader();
+			ASAN_POISON_MEMORY_REGION( AllocationHeader, sizeof(FAllocationHeader) + AllocationHeader->GetAllocationSize() );
+
+			// this deletes complete blocks when the last allocation is freed
+			if (Header->NumAllocations.fetch_sub(1, std::memory_order_acq_rel) == 1)
+			{
+				Header->~FBlockHeader();
+				ASAN_UNPOISON_MEMORY_REGION( Header, Header->NextAllocationPtr - uintptr_t(Header) );
+				BlockAllocationTag::FreeFunction(Header);
+			}
+		}
+	}
+
+	static FORCEINLINE_DEBUGGABLE SIZE_T GetAllocationSize(void* Pointer)
+	{
+		if (Pointer)
+		{
+			FAllocationHeader* AllocationHeader = GetAllocationHeader(Pointer);
+			ASAN_UNPOISON_MEMORY_REGION( AllocationHeader, sizeof(FAllocationHeader) );
+			SIZE_T Size = AllocationHeader->GetAllocationSize();
+			ASAN_POISON_MEMORY_REGION( AllocationHeader, sizeof(FAllocationHeader) );
+			return Size;
+		}
+		return 0;
+	}
+
+	static FORCEINLINE_DEBUGGABLE void* Realloc(void* Old, SIZE_T Size, uint32 Alignment)
+	{
+		void* New = nullptr;
+		if(Size != 0)
+		{
+			New = Malloc(Size, Alignment);
+			SIZE_T OldSize = GetAllocationSize(Old);
+			if(OldSize != 0)
+			{
+				memcpy(New, Old, Size < OldSize ? Size : OldSize);
+			}
+		}
+		Free(Old);
+		return New;
+	}
+};
+
+using FConcurrentLinearAllocator = TConcurrentLinearAllocator<FDefaultBlockAllocationTag>;
+
+template<typename ObjectType, typename BlockAllocationTag = FDefaultBlockAllocationTag>
+class TConcurrentLinearObject
+{
+public:
+	FORCEINLINE_DEBUGGABLE static void* operator new(size_t Size)
+	{
+		return TConcurrentLinearAllocator<BlockAllocationTag>::template Malloc<alignof(ObjectType)>(Size);
+	}
+
+	FORCEINLINE_DEBUGGABLE static void* operator new[](size_t Size)
+	{
+		return TConcurrentLinearAllocator<BlockAllocationTag>::template Malloc<alignof(ObjectType)>(Size);
+	}
+	
+	FORCEINLINE_DEBUGGABLE static void* operator new(size_t Size, std::align_val_t Align)
+	{
+		checkSlow(size_t(Align) == alignof(ObjectType));
+		return TConcurrentLinearAllocator<BlockAllocationTag>::template Malloc<alignof(ObjectType)>(Size);
+	}
+
+	FORCEINLINE_DEBUGGABLE static void* operator new[](size_t Size, std::align_val_t Align)
+	{
+		checkSlow(size_t(Align)  == alignof(ObjectType));
+		return TConcurrentLinearAllocator<BlockAllocationTag>::template Malloc<alignof(ObjectType)>(Size);
+	}
+
+	FORCEINLINE_DEBUGGABLE static void operator delete(void* Ptr)
+	{
+		return TConcurrentLinearAllocator<BlockAllocationTag>::Free(Ptr);
+	}
+
+	FORCEINLINE_DEBUGGABLE static  void operator delete[](void* Ptr)
+	{
+		return TConcurrentLinearAllocator<BlockAllocationTag>::Free(Ptr);
+	}
+};
+
+template<typename BlockAllocationTag>
+class TConcurrentLinearArrayAllocator
+{
+public:
+	using SizeType = int32;
+
+	enum { NeedsElementType = true };
+	enum { RequireRangeCheck = true };
+
+	template<typename ElementType>
+	class ForElementType
+	{
+	public:
+		ForElementType() = default;
+		~ForElementType()
+		{
+			if (Data)
+			{
+				TConcurrentLinearAllocator<BlockAllocationTag>::Free(Data);
+			}
+		}
+		/**
+		* Moves the state of another allocator into this one.
+		* Assumes that the allocator is currently empty, i.e. memory may be allocated but any existing elements have already been destructed (if necessary).
+		* @param Other - The allocator to move the state from.  This allocator should be left in a valid empty state.
+		*/
+		inline void MoveToEmpty(ForElementType& Other)
+		{
+			checkSlow(this != &Other);
+
+			if (Data)
+			{
+				TConcurrentLinearAllocator<BlockAllocationTag>::Free(Data);
+			}
+
+			Data = Other.Data;
+			Other.Data = nullptr;
+		}
+		inline ElementType* GetAllocation() const
+		{
+			return Data;
+		}
+		void ResizeAllocation(SizeType PreviousNumElements, SizeType NumElements, SIZE_T NumBytesPerElement)
+		{
+			Data = (ElementType*)TConcurrentLinearAllocator<BlockAllocationTag>::Realloc(Data, NumElements * NumBytesPerElement, alignof(ElementType));
+		}
+		SizeType CalculateSlackReserve(SizeType NumElements, SIZE_T NumBytesPerElement) const
+		{
+			return DefaultCalculateSlackReserve(NumElements, NumBytesPerElement, false);
+		}
+		SizeType CalculateSlackShrink(SizeType NumElements, SizeType NumAllocatedElements, SIZE_T NumBytesPerElement) const
+		{
+			return DefaultCalculateSlackShrink(NumElements, NumAllocatedElements, NumBytesPerElement, false);
+		}
+		SizeType CalculateSlackGrow(SizeType NumElements, SizeType NumAllocatedElements, SIZE_T NumBytesPerElement) const
+		{
+			return DefaultCalculateSlackGrow(NumElements, NumAllocatedElements, NumBytesPerElement, false);
+		}
+
+		SIZE_T GetAllocatedSize(SizeType NumAllocatedElements, SIZE_T NumBytesPerElement) const
+		{
+			return NumAllocatedElements * NumBytesPerElement;
+		}
+
+		bool HasAllocation() const
+		{
+			return !!Data;
+		}
+
+		SizeType GetInitialCapacity() const
+		{
+			return 0;
+		}
+
+
+	private:
+		ForElementType(const ForElementType&) = delete;
+		ForElementType& operator=(const ForElementType&) = delete;
+		ElementType* Data = nullptr;
+	};
+
+	using ForAnyElementType = ForElementType<FScriptContainerElement>;
+};
+
+template <typename BlockAllocationTag>
+struct TAllocatorTraits<TConcurrentLinearArrayAllocator<BlockAllocationTag>> : TAllocatorTraitsBase<TConcurrentLinearArrayAllocator<BlockAllocationTag>>
+{
+	enum { SupportsMove = true };
+	enum { IsZeroConstruct = true };
+};
+
+template <typename BlockAllocationTag>
+using TConcurrentLinearBitArrayAllocator = TInlineAllocator<4, TConcurrentLinearArrayAllocator<BlockAllocationTag>>;
+
+template <typename BlockAllocationTag>
+using TConcurrentLinearSparseArrayAllocator = TSparseArrayAllocator<TConcurrentLinearArrayAllocator<BlockAllocationTag>, TConcurrentLinearBitArrayAllocator<BlockAllocationTag>>;
+
+template <typename BlockAllocationTag>
+using TConcurrentLinearSetAllocator = TSetAllocator<TConcurrentLinearSparseArrayAllocator<BlockAllocationTag>, TInlineAllocator<1, TConcurrentLinearBitArrayAllocator<BlockAllocationTag>>>;
+
+using FConcurrentLinearArrayAllocator = TConcurrentLinearArrayAllocator<FDefaultBlockAllocationTag>;
+using FConcurrentLinearBitArrayAllocator = TConcurrentLinearBitArrayAllocator<FDefaultBlockAllocationTag>;
+using FConcurrentLinearSparseArrayAllocator = TConcurrentLinearSparseArrayAllocator<FDefaultBlockAllocationTag>;
+using FConcurrentLinearSetAllocator = TConcurrentLinearSetAllocator<FDefaultBlockAllocationTag>;
+
+//The BulkObjectAllocator can be used to atomically destroy all allocated Objects, it will properly call every destructor before deleting the memory as well
+template<typename BlockAllocationTag>
+class TConcurrentLinearBulkObjectAllocator
+{
+	using ThisType = TConcurrentLinearBulkObjectAllocator<BlockAllocationTag>;
+	struct FAllocation
+	{
+		virtual ~FAllocation() = default;
+		FAllocation* Next = nullptr;
+	};
+
+	template <typename T>
+	struct TObject final : FAllocation
+	{
+		TObject() = default;
+		virtual ~TObject() override
+		{
+			T* Alloc = reinterpret_cast<T*>(uintptr_t(this) + Align(sizeof(TObject<T>), alignof(T)));
+			checkSlow(IsAligned(Alloc, alignof(T)));
+
+			// We need a typedef here because VC won't compile the destructor call below if ElementType itself has a member called ElementType
+			using DestructorType = T;
+			Alloc->DestructorType::~T();
+		}
+	};
+
+	template <typename T>
+	struct TObjectArray final : FAllocation
+	{
+		inline TObjectArray(SIZE_T InNum) 
+			: Num(InNum) 
+		{
+		}
+
+		virtual ~TObjectArray() override
+		{
+			T* Alloc = reinterpret_cast<T*>(uintptr_t(this) + Align(sizeof(TObjectArray<T>), alignof(T)));
+			checkSlow(IsAligned(Alloc, alignof(T)));
+
+			for (SIZE_T i = 0; i < Num; i++)
+			{
+				// We need a typedef here because VC won't compile the destructor call below if ElementType itself has a member called ElementType
+				using DestructorType = T;
+				Alloc[i].DestructorType::~T();
+			}
+		}
+		SIZE_T Num;
+	};
+
+	std::atomic<FAllocation*> Next { nullptr };
+
+public:
+	TConcurrentLinearBulkObjectAllocator() = default;
+	~TConcurrentLinearBulkObjectAllocator()
+	{
+		BulkDelete();
+	}
+
+	inline void BulkDelete()
+	{
+		FAllocation* Allocation = Next.exchange(nullptr, std::memory_order_release);
+		while (Allocation != nullptr)
+		{
+			FAllocation* NextAllocation = Allocation->Next;
+			Allocation->~FAllocation();
+			TConcurrentLinearAllocator<BlockAllocationTag>::Free(Allocation);
+			Allocation = NextAllocation;
+		}
+	}
+
+	inline void* Malloc(SIZE_T Size, uint32 Alignment)
+	{
+		SIZE_T TotalSize = Align(sizeof(FAllocation), Alignment) + Size;
+		FAllocation* Allocation = new (TConcurrentLinearAllocator<BlockAllocationTag>::Malloc(TotalSize, FMath::Max<size_t>(alignof(FAllocation), Alignment))) FAllocation();
+
+		void* Alloc = Align(Allocation + 1, Alignment);
+		checkSlow(uintptr_t(Alloc) + Size - uintptr_t(Allocation) <= TotalSize);
+
+		Allocation->Next = Next.load(std::memory_order_relaxed);
+		while (!Next.compare_exchange_strong(Allocation->Next, Allocation, std::memory_order_acquire, std::memory_order_relaxed))
+		{
+		}
+		return Alloc;
+	}
+
+	template<typename T, typename... TArgs>
+	inline T* Create(TArgs&&... Args)
+	{
+		T* Alloc = CreateNoInit<T>();
+		new ((void*)Alloc) T(Forward<TArgs>(Args)...);
+		return Alloc;
+	}
+
+	template<typename T, typename... TArgs>
+	inline T* CreateArray(SIZE_T Num, const TArgs&... Args)
+	{
+		T* Alloc = CreateArrayNoInit<T>(Num);
+		for (SIZE_T i = 0; i < Num; i++)
+		{
+			new ((void*)(&Alloc[i])) T(Args...);
+		}
+		return Alloc;
+	}
+
+private: //Kept private, there is also deliberately no new operator overloads because those will not be able to send the ObjectType into the allocator defeating its purpose 
+	template<typename T>
+	inline T* CreateNoInit()
+	{
+		SIZE_T TotalSize = Align(sizeof(TObject<T>), alignof(T)) + sizeof(T);
+		TObject<T>* Object = new (TConcurrentLinearAllocator<BlockAllocationTag>::template Malloc<FMath::Max(alignof(TObject<T>), alignof(T))>(TotalSize)) TObject<T>();
+
+		T* Alloc = reinterpret_cast<T*>(uintptr_t(Object) + Align(sizeof(TObject<T>), alignof(T)));
+		checkSlow(IsAligned(Alloc, alignof(T)));
+		checkSlow(uintptr_t(Alloc + 1) - uintptr_t(Object) <= TotalSize);
+
+		Object->Next = Next.load(std::memory_order_relaxed);
+		while (!Next.compare_exchange_weak(Object->Next, Object, std::memory_order_acquire, std::memory_order_relaxed))
+		{
+		}
+		return Alloc;
+	}
+
+	template<typename T>
+	inline T* CreateArrayNoInit(SIZE_T Num)
+	{
+		SIZE_T TotalSize = Align(sizeof(TObjectArray<T>), alignof(T)) + (sizeof(T) * Num);
+		TObjectArray<T>* Array = new (TConcurrentLinearAllocator<BlockAllocationTag>::template Malloc<FMath::Max(alignof(TObjectArray<T>), alignof(T))>(TotalSize)) TObjectArray<T>(Num);
+
+		T* Alloc = reinterpret_cast<T*>(uintptr_t(Array) + Align(sizeof(TObjectArray<T>), alignof(T)));
+		checkSlow(IsAligned(Alloc, alignof(T)));
+		checkSlow(uintptr_t(Alloc + Num) - uintptr_t(Array) <= TotalSize);
+
+		Array->Next = Next.load(std::memory_order_relaxed);
+		while (!Next.compare_exchange_weak(Array->Next, Array, std::memory_order_acquire, std::memory_order_relaxed))
+		{
+		}
+		return Alloc;
+	}
+};
+
+using FConcurrentLinearBulkObjectAllocator = TConcurrentLinearBulkObjectAllocator<FDefaultBlockAllocationTag>;
diff --git a/Engine/Source/Runtime/Core/Public/HAL/LowLevelMemTracker.h b/Engine/Source/Runtime/Core/Public/HAL/LowLevelMemTracker.h
--- a/Engine/Source/Runtime/Core/Public/HAL/LowLevelMemTracker.h
+++ b/Engine/Source/Runtime/Core/Public/HAL/LowLevelMemTracker.h
@@ -148,6 +148,7 @@ enum class ELLMTagSet : uint8
 	macro(EngineMisc,							"EngineMisc",					GET_STATFNAME(STAT_EngineMiscLLM),							GET_STATFNAME(STAT_EngineSummaryLLM),			-1)\
 	/* Any task kicked off from the task graph that doesn't have its own category. Should be fairly low. */ \
 	macro(TaskGraphTasksMisc,					"TaskGraphMiscTasks",			GET_STATFNAME(STAT_TaskGraphTasksMiscLLM),					GET_STATFNAME(STAT_EngineSummaryLLM),			-1)\
+	macro(LinearAllocator,						"LinearAllocator",				GET_STATFNAME(STAT_LinearAllocatorLLM),						NAME_None,										-1)\
 	macro(Audio,								"Audio",						GET_STATFNAME(STAT_AudioLLM),								GET_STATFNAME(STAT_AudioSummaryLLM),			-1)\
 	macro(AudioMisc,							"AudioMisc",					GET_STATFNAME(STAT_AudioMiscLLM),							GET_STATFNAME(STAT_AudioSummaryLLM),			ELLMTag::Audio)\
 	macro(AudioSoundWaves,						"AudioSoundWaves",				GET_STATFNAME(STAT_AudioSoundWavesLLM),						GET_STATFNAME(STAT_AudioSummaryLLM),			ELLMTag::Audio)\
@@ -248,7 +249,7 @@ enum class ELLMTag : LLM_TAG_TYPE
 
 	//------------------------------
 	// Platform tags
-	PlatformTagStart = 107,
+	PlatformTagStart = 108,
 	PlatformTagEnd = 149,
 
 	//------------------------------
