Allow FUnixPlatformProcess::CreateProc to locate executables within $PATH, rather than requiring an absolute path.

Changed underlying API from posix_spawn to posix_spawnp, which only differs in behavior if the provided path fragment does not contain any path separators. Matches within the base directory are still prioritized if they exist. Failures are no longer treated as fatal.

diff --git a/Engine/Source/Runtime/Core/Private/Unix/UnixPlatformProcess.cpp b/Engine/Source/Runtime/Core/Private/Unix/UnixPlatformProcess.cpp
--- a/Engine/Source/Runtime/Core/Private/Unix/UnixPlatformProcess.cpp
+++ b/Engine/Source/Runtime/Core/Private/Unix/UnixPlatformProcess.cpp
@@ -766,28 +766,46 @@ FProcHandle FUnixPlatformProcess::CreateProc(const TCHAR* URL, const TCHAR* Parm
 FProcHandle FUnixPlatformProcess::CreateProc(const TCHAR* URL, const TCHAR* Parms, bool bLaunchDetached, bool bLaunchHidden, bool bLaunchReallyHidden, uint32* OutProcessID, int32 PriorityModifier, const TCHAR* OptionalWorkingDirectory, void* PipeWriteChild, void* PipeReadChild, void* PipeStdErrChild)
 {
 	// @TODO bLaunchHidden bLaunchReallyHidden are not handled
-	// We need an absolute path to executable
+
 	FString ProcessPath = URL;
-	if (*URL != TEXT('/'))
-	{
-		ProcessPath = FPaths::ConvertRelativePathToFull(ProcessPath);
-	}
 
-	if (!FPaths::FileExists(ProcessPath))
+	// - If the first character is /, we use the provided absolute path as-is.
+	// - If no leading slash, prefer the result of FPaths::ConvertRelativePathToFull if it exists. This is roughly
+	//   morally equivalent to prepending the basedir to $PATH (and in turn, Windows' cwd (==basedir) priority).
+	// - If there was a (non-leading) slash, and ConvertRelativePathToFull missed, return failure.
+	// - If there were no path separators in the input string, allow posix_spawnp to search the $PATH.
+
+	int32 PathSepIdx = INDEX_NONE;
+	const bool bInputHasPathSep = ProcessPath.FindChar(TEXT('/'), PathSepIdx);
+	bool bAbsolutePath = PathSepIdx == 0;
+
+	if (!bAbsolutePath)
 	{
-		return FProcHandle();
+		const FString CandidatePath = FPaths::ConvertRelativePathToFull(ProcessPath);
+		if (bInputHasPathSep || FPaths::FileExists(CandidatePath))
+		{
+			ProcessPath = CandidatePath;
+			bAbsolutePath = true;
+		}
 	}
 
-	// check if it's worth attemptting to execute the file
-	if (!UnixPlatformProcess::AttemptToMakeExecIfNotAlready(ProcessPath))
+	// Even if we weren't passed an absolute path, we may have expanded to one above.
+	if (bAbsolutePath)
 	{
-		return FProcHandle();
-	}
+		if (!FPaths::FileExists(ProcessPath))
+		{
+			UE_LOG(LogHAL, Error, TEXT("FUnixPlatformProcess::CreateProc: File does not exist (%s)"), *ProcessPath);
+			return FProcHandle();
+		}
 
-	FString Commandline = FString::Printf(TEXT("\"%s\""), *ProcessPath);
-	Commandline += TEXT(" ");
-	Commandline += Parms;
+		if (!UnixPlatformProcess::AttemptToMakeExecIfNotAlready(ProcessPath))
+		{
+			UE_LOG(LogHAL, Error, TEXT("FUnixPlatformProcess::CreateProc: File not executable (%s)"), *ProcessPath);
+			return FProcHandle();
+		}
+	}
 
+	const FString Commandline = FString::Printf(TEXT("\"%s\" %s"), *ProcessPath, Parms);
 	UE_LOG(LogHAL, Verbose, TEXT("FUnixPlatformProcess::CreateProc: '%s'"), *Commandline);
 
 	TArray<FString> ArgvArray;
@@ -955,7 +973,7 @@ FProcHandle FUnixPlatformProcess::CreateProc(const TCHAR* URL, const TCHAR* Parm
 		}
 
 		posix_spawnattr_setflags(&SpawnAttr, SpawnFlags);
-		PosixSpawnErrNo = posix_spawn(&ChildPid, TCHAR_TO_UTF8(*ProcessPath), &FileActions, &SpawnAttr, Argv, environ);
+		PosixSpawnErrNo = posix_spawnp(&ChildPid, TCHAR_TO_UTF8(*ProcessPath), &FileActions, &SpawnAttr, Argv, environ);
 		posix_spawn_file_actions_destroy(&FileActions);
 	}
 	else
@@ -970,14 +988,14 @@ FProcHandle FUnixPlatformProcess::CreateProc(const TCHAR* URL, const TCHAR* Parm
 		SpawnFlags |= POSIX_SPAWN_USEVFORK;
 
 		posix_spawnattr_setflags(&SpawnAttr, SpawnFlags);
-		PosixSpawnErrNo = posix_spawn(&ChildPid, TCHAR_TO_UTF8(*ProcessPath), nullptr, &SpawnAttr, Argv, environ);
+		PosixSpawnErrNo = posix_spawnp(&ChildPid, TCHAR_TO_UTF8(*ProcessPath), nullptr, &SpawnAttr, Argv, environ);
 	}
 	posix_spawnattr_destroy(&SpawnAttr);
 
 	if (PosixSpawnErrNo != 0)
 	{
-		UE_LOG(LogHAL, Fatal, TEXT("FUnixPlatformProcess::CreateProc: posix_spawn() failed (%d, %s)"), PosixSpawnErrNo, UTF8_TO_TCHAR(strerror(PosixSpawnErrNo)));
-		return FProcHandle();	// produce knowingly invalid handle if for some reason Fatal log (above) returns
+		UE_LOG(LogHAL, Error, TEXT("FUnixPlatformProcess::CreateProc: posix_spawnp() failed (%d, %s)"), PosixSpawnErrNo, UTF8_TO_TCHAR(strerror(PosixSpawnErrNo)));
+		return FProcHandle();
 	}
 
 	// renice the child (subject to race condition).
