Fix crash with editor on shutdown in debug editor configuration

diff --git a/Engine/Source/Runtime/Core/Private/Misc/ConfigCacheIni.cpp b/Engine/Source/Runtime/Core/Private/Misc/ConfigCacheIni.cpp
--- a/Engine/Source/Runtime/Core/Private/Misc/ConfigCacheIni.cpp
+++ b/Engine/Source/Runtime/Core/Private/Misc/ConfigCacheIni.cpp
@@ -2244,7 +2244,11 @@ void FConfigFile::ProcessPropertyAndWriteForDefaults(int IniCombineThreshold, co
 	FConfigCacheIni
 -----------------------------------------------------------------------------*/
 #if WITH_EDITOR
-static TMap<FName, TFuture<void>> GPlatformConfigFutures;
+static TMap<FName, TFuture<void>>& GetPlatformConfigFutures()
+{
+	static TMap<FName, TFuture<void>> Futures;
+	return Futures;
+}
 #endif
 
 FConfigCacheIni::FConfigCacheIni(EConfigCacheType InType)
@@ -2961,11 +2965,11 @@ void FConfigCacheIni::Exit()
 	Flush( 1 );
 
 #if WITH_EDITOR
-	for (auto& PlatformConfigFuture : GPlatformConfigFutures)
+	for (auto& PlatformConfigFuture : GetPlatformConfigFutures())
 	{
 		PlatformConfigFuture.Value.Get();
 	}
-	GPlatformConfigFutures.Empty();
+	GetPlatformConfigFutures().Empty();
 #endif
 }
 
@@ -5463,14 +5467,14 @@ void FConfigCacheIni::AsyncInitializeConfigForPlatforms()
 	// pre-create all platforms so that the loop below doesn't reallocate anything in the map
 	for (auto Pair : FDataDrivenPlatformInfoRegistry::GetAllPlatformInfos())
 	{
-		GPlatformConfigFutures.Emplace(Pair.Key);
+		GetPlatformConfigFutures().Emplace(Pair.Key);
 		GConfigForPlatform.Add(Pair.Key, new FConfigCacheIni(EConfigCacheType::Temporary));
 	}
 
 	for (auto Pair : FDataDrivenPlatformInfoRegistry::GetAllPlatformInfos())
 	{
 		FName PlatformName = Pair.Key;
-		GPlatformConfigFutures[PlatformName] = Async(EAsyncExecution::ThreadPool, [PlatformName]
+		GetPlatformConfigFutures()[PlatformName] = Async(EAsyncExecution::ThreadPool, [PlatformName]
 		{
 			double Start = FPlatformTime::Seconds();
 
@@ -5500,7 +5504,7 @@ FConfigCacheIni* FConfigCacheIni::ForPlatform(FName PlatformName)
 #if WITH_EDITOR
 	// they are likely already loaded, but just block to make sure
 	{
-		GPlatformConfigFutures[PlatformName].Get();
+		GetPlatformConfigFutures()[PlatformName].Get();
 	}
 #endif
 
