Fixed: FPlatformMemoryStats.AvailableVirtual is always zero on iOS. Set it to the same value as AvailablePhysical

diff --git a/Engine/Source/Runtime/Core/Private/Apple/ApplePlatformMemory.cpp b/Engine/Source/Runtime/Core/Private/Apple/ApplePlatformMemory.cpp
--- a/Engine/Source/Runtime/Core/Private/Apple/ApplePlatformMemory.cpp
+++ b/Engine/Source/Runtime/Core/Private/Apple/ApplePlatformMemory.cpp
@@ -353,6 +353,7 @@ FPlatformMemoryStats FApplePlatformMemory::GetStats()
     FreeMem = (Stats.free_count + Stats.inactive_count) * MemoryConstants.PageSize;
 #endif
 	MemoryStats.AvailablePhysical = FreeMem;
+	MemoryStats.AvailableVirtual = FreeMem;
 	
 	// Just get memory information for the process and report the working set instead
 	mach_task_basic_info_data_t TaskInfo;
