Fixed FNameMap::GetMinimalName not being aware of case sensitive FNames

diff --git a/Engine/Source/Runtime/Core/Public/Serialization/MappedName.h b/Engine/Source/Runtime/Core/Public/Serialization/MappedName.h
--- a/Engine/Source/Runtime/Core/Public/Serialization/MappedName.h
+++ b/Engine/Source/Runtime/Core/Public/Serialization/MappedName.h
@@ -137,8 +137,9 @@ public:
 		check(MappedName.GetType() == NameMapType);
 		check(MappedName.GetIndex() < uint32(NameEntries.Num()));
 
-		FNameEntryId NameEntry = NameEntries[MappedName.GetIndex()];
-		return FMinimalName(NameEntry, MappedName.GetNumber());
+		FNameEntryId DisplayId = NameEntries[MappedName.GetIndex()];
+		FNameEntryId ComparisonId = FName::GetComparisonIdFromDisplayId(DisplayId);
+		return FMinimalName(ComparisonId, MappedName.GetNumber());
 	}
 
 private:
