Removed left-over comment from a previous fix that is not relevant anymore.

diff --git a/Engine/Source/Runtime/Core/Private/HAL/MallocCallstackHandler.cpp b/Engine/Source/Runtime/Core/Private/HAL/MallocCallstackHandler.cpp
--- a/Engine/Source/Runtime/Core/Private/HAL/MallocCallstackHandler.cpp
+++ b/Engine/Source/Runtime/Core/Private/HAL/MallocCallstackHandler.cpp
@@ -145,7 +145,6 @@ int32 FMallocCallstackHandler::GetCallStackIndex()
 
 	// Capture callstack and create FCallStackMapKey.
 	uint64 FullCallStack[MaxCallStackDepth + CallStackEntriesToSkipCount] = { 0 };
-	// CRC is filled in by the CaptureStackBackTrace, not all platforms calculate this for you.
 	FPlatformStackWalk::CaptureStackBackTrace(&FullCallStack[0], MaxCallStackDepth + CallStackEntriesToSkipCount);
 	// Skip first n entries as they are inside the allocator.
 	uint64* CallStack = &FullCallStack[CallStackEntriesToSkipCount];
