#ue5 - LWC: fix compile errors when using non-vectorized FTransform.

Some code was using the wrong #define to toggle (PLATFORM_ENABLE_VECTORINTRINSICS instead of ENABLE_VECTORIZED_TRANSFORM).

diff --git a/Engine/Plugins/Animation/ControlRig/Source/ControlRig/Private/Sequencer/MovieSceneControlRigParameterSection.cpp b/Engine/Plugins/Animation/ControlRig/Source/ControlRig/Private/Sequencer/MovieSceneControlRigParameterSection.cpp
--- a/Engine/Plugins/Animation/ControlRig/Source/ControlRig/Private/Sequencer/MovieSceneControlRigParameterSection.cpp
+++ b/Engine/Plugins/Animation/ControlRig/Source/ControlRig/Private/Sequencer/MovieSceneControlRigParameterSection.cpp
@@ -181,7 +181,11 @@ struct FParameterVectorChannelEditorData
 				{
 					const FRigControlValue::FTransform_Float Storage = ControlRig->GetHierarchy()
 						->GetControlValue(ControlElement, ERigControlValueType::Current).Get<FRigControlValue::FTransform_Float>();
+#if ENABLE_VECTORIZED_TRANSFORM
 					return FVector4(Storage.TranslationX, Storage.TranslationY, Storage.TranslationZ, Storage.TranslationW);
+#else
+					return FVector4(Storage.TranslationX, Storage.TranslationY, Storage.TranslationZ, 0.f);
+#endif
 				}
 			}
 		}
diff --git a/Engine/Source/Runtime/Core/Private/Math/Transform.cpp b/Engine/Source/Runtime/Core/Private/Math/Transform.cpp
--- a/Engine/Source/Runtime/Core/Private/Math/Transform.cpp
+++ b/Engine/Source/Runtime/Core/Private/Math/Transform.cpp
@@ -34,7 +34,7 @@ void TTransform<T>::DebugPrint() const
 template<typename T>
 FString TTransform<T>::ToHumanReadableString() const
 {
-	FRotator R(GetRotation());
+	TRotator<T> R(GetRotation());
 	TVector<T> TT(GetTranslation());
 	TVector<T> S(GetScale3D());
 
@@ -49,7 +49,7 @@ FString TTransform<T>::ToHumanReadableString() const
 template<typename T>
 FString TTransform<T>::ToString() const
 {
-	const FRotator R(Rotator());
+	const TRotator<T> R(Rotator());
 	const TVector<T> TT(GetTranslation());
 	const TVector<T> S(GetScale3D());
 
@@ -75,7 +75,7 @@ bool TTransform<T>::InitFromString( const FString& Source )
 	}
 
 	// Rotation
-	FRotator ParsedRotation = FRotator::ZeroRotator;
+	TRotator<T> ParsedRotation = TRotator<T>::ZeroRotator;
 	if( !FDefaultValueHelper::ParseRotator(ComponentStrings[1], ParsedRotation) )
 	{
 		return false;
