Added FFileHelper::LoadFileToStringWithLineVisitor to stream a text file and visit its lines

diff --git a/Engine/Source/Runtime/Core/Private/Misc/FileHelper.cpp b/Engine/Source/Runtime/Core/Private/Misc/FileHelper.cpp
--- a/Engine/Source/Runtime/Core/Private/Misc/FileHelper.cpp
+++ b/Engine/Source/Runtime/Core/Private/Misc/FileHelper.cpp
@@ -4,6 +4,8 @@
 #include "Containers/StringConv.h"
 #include "Logging/LogMacros.h"
 #include "CoreGlobals.h"
+#include "Memory/MemoryView.h"
+#include "Memory/SharedBuffer.h"
 #include "Misc/ByteSwap.h"
 #include "Misc/CoreMisc.h"
 #include "Misc/Paths.h"
@@ -340,6 +342,223 @@ bool FFileHelper::LoadFileToStringArrayWithPredicate(TArray<FString>& Result, co
 	return LoadFileToStringArrayWithPredicate(Result, Filename, Predicate);
 }
 
+namespace UE::FileHelper::Private
+{
+
+enum class EEncoding
+{
+	Unknown,
+	UTF8,
+	UTF16BE,
+	UTF16LE,
+};
+
+static EEncoding ParseEncoding(FMutableMemoryView& Buffer, const uint64 TotalSize)
+{
+	const uint64 Size = Buffer.GetSize();
+	const uint8* const Bytes = static_cast<const uint8*>(Buffer.GetData());
+	if (!(TotalSize & 1) && Size >= 2 && Bytes[0] == 0xff && Bytes[1] == 0xfe)
+	{
+		Buffer += 2;
+		return EEncoding::UTF16LE;
+	}
+	if (!(TotalSize & 1) && Size >= 2 && Bytes[0] == 0xfe && Bytes[1] == 0xff)
+	{
+		Buffer += 2;
+		return EEncoding::UTF16BE;
+	}
+	if (Size >= 3 && Bytes[0] == 0xef && Bytes[1] == 0xbb && Bytes[2] == 0xbf)
+	{
+		Buffer += 3;
+	}
+	return EEncoding::UTF8;
+}
+
+static FMutableMemoryView ParseLinesUTF8(FMutableMemoryView Buffer, TFunctionRef<void(FStringView Line)> Visitor, bool bLastBuffer)
+{
+	const uint8* Pos = static_cast<const uint8*>(Buffer.GetData());
+	const uint8* const End = static_cast<const uint8*>(Buffer.GetDataEnd());
+	while (Pos < End)
+	{
+		const uint8* const Line = Pos;
+
+		while (Pos < End && *Pos != '\r' && *Pos != '\n')
+		{
+			++Pos;
+		}
+
+		if (!bLastBuffer && (Pos == End || (Pos + 1 == End && *Pos == '\r')))
+		{
+			return Buffer.Right(static_cast<uint64>(End - Line));
+		}
+
+		Visitor(FUTF8ToTCHAR(reinterpret_cast<const UTF8CHAR*>(Line), UE_PTRDIFF_TO_INT32(Pos - Line)));
+
+		if (Pos < End && *Pos == '\r')
+		{
+			++Pos;
+		}
+		if (Pos < End && *Pos == '\n')
+		{
+			++Pos;
+		}
+	}
+	return Buffer.Right(0);
+}
+
+static FMutableMemoryView ParseLinesUTF16BE(FMutableMemoryView Buffer, TFunctionRef<void(FStringView Line)> Visitor, bool bLastBuffer)
+{
+	uint16* Pos = static_cast<uint16*>(Buffer.GetData());
+	uint16* const End = static_cast<uint16*>(Buffer.GetDataEnd());
+	while (Pos < End)
+	{
+		uint16* const Line = Pos;
+
+		while (Pos < End)
+		{
+			if (const uint16 CodeUnit = NETWORK_ORDER16(*Pos); CodeUnit == '\r' || CodeUnit == '\n')
+			{
+				break;
+			}
+			++Pos;
+		}
+
+		if (!bLastBuffer && (Pos == End || (Pos + 1 == End && NETWORK_ORDER16(*Pos) == '\r')))
+		{
+			return Buffer.Right(static_cast<uint64>(End - Line) * sizeof(uint16));
+		}
+
+		if constexpr (PLATFORM_LITTLE_ENDIAN)
+		{
+			for (uint16* SwapPos = Line; SwapPos < Pos; ++SwapPos)
+			{
+				*SwapPos = NETWORK_ORDER16(*SwapPos);
+			}
+		}
+
+		Visitor(FUTF16ToTCHAR(reinterpret_cast<const UTF16CHAR*>(Line), UE_PTRDIFF_TO_INT32(Pos - Line)));
+
+		if (Pos < End && NETWORK_ORDER16(*Pos) == '\r')
+		{
+			++Pos;
+		}
+		if (Pos < End && NETWORK_ORDER16(*Pos) == '\n')
+		{
+			++Pos;
+		}
+	}
+	return Buffer.Right(0);
+}
+
+static FMutableMemoryView ParseLinesUTF16LE(FMutableMemoryView Buffer, TFunctionRef<void(FStringView Line)> Visitor, bool bLastBuffer)
+{
+	uint16* Pos = static_cast<uint16*>(Buffer.GetData());
+	uint16* const End = static_cast<uint16*>(Buffer.GetDataEnd());
+	while (Pos < End)
+	{
+		uint16* const Line = Pos;
+
+		while (Pos < End)
+		{
+			if (const uint16 CodeUnit = INTEL_ORDER16(*Pos); CodeUnit == '\r' || CodeUnit == '\n')
+			{
+				break;
+			}
+			++Pos;
+		}
+
+		if (!bLastBuffer && (Pos == End || (Pos + 1 == End && INTEL_ORDER16(*Pos) == '\r')))
+		{
+			return Buffer.Right(static_cast<uint64>(End - Line) * sizeof(uint16));
+		}
+
+		if constexpr (!PLATFORM_LITTLE_ENDIAN)
+		{
+			for (uint16* SwapPos = Line; SwapPos < Pos; ++SwapPos)
+			{
+				*SwapPos = INTEL_ORDER16(*SwapPos);
+			}
+		}
+
+		Visitor(FUTF16ToTCHAR(reinterpret_cast<const UTF16CHAR*>(Line), UE_PTRDIFF_TO_INT32(Pos - Line)));
+
+		if (Pos < End && INTEL_ORDER16(*Pos) == '\r')
+		{
+			++Pos;
+		}
+		if (Pos < End && INTEL_ORDER16(*Pos) == '\n')
+		{
+			++Pos;
+		}
+	}
+	return Buffer.Right(0);
+}
+
+} // UE::FileHelper::Private
+
+bool FFileHelper::LoadFileToStringWithLineVisitor(const TCHAR* Filename, TFunctionRef<void(FStringView Line)> Visitor)
+{
+	using namespace UE::FileHelper::Private;
+
+	FScopedLoadingState ScopedLoadingState(Filename);
+	TUniquePtr<FArchive> Ar(IFileManager::Get().CreateFileReader(Filename, FILEREAD_Silent));
+	if (!Ar)
+	{
+		return false;
+	}
+
+	const int64 TotalSize = Ar->TotalSize();
+	FUniqueBuffer Buffer = FUniqueBuffer::Alloc(FMath::Min<int64>(TotalSize, 1024 * 1024));
+
+	EEncoding Encoding = EEncoding::Unknown;
+	FMutableMemoryView BufferTail = Buffer;
+	for (int64 RemainingSize = TotalSize; RemainingSize > 0;)
+	{
+		const FMutableMemoryView SerializeBuffer = BufferTail.Left(RemainingSize);
+		Ar->Serialize(SerializeBuffer.GetData(), static_cast<int64>(SerializeBuffer.GetSize()));
+		RemainingSize -= SerializeBuffer.GetSize();
+
+		FMutableMemoryView ParseBuffer = Buffer.GetView().LeftChop(BufferTail.GetSize() - SerializeBuffer.GetSize());
+
+		if (Encoding == EEncoding::Unknown)
+		{
+			Encoding = ParseEncoding(ParseBuffer, TotalSize);
+		}
+
+		switch (Encoding)
+		{
+		case EEncoding::UTF8:
+			ParseBuffer = ParseLinesUTF8(ParseBuffer, Visitor, RemainingSize == 0);
+			break;
+		case EEncoding::UTF16BE:
+			ParseBuffer = ParseLinesUTF16BE(ParseBuffer, Visitor, RemainingSize == 0);
+			break;
+		case EEncoding::UTF16LE:
+			ParseBuffer = ParseLinesUTF16LE(ParseBuffer, Visitor, RemainingSize == 0);
+			break;
+		default:
+			checkNoEntry();
+			return false;
+		}
+
+		if (Buffer.GetSize() == ParseBuffer.GetSize())
+		{
+			// No line endings were found. Double the buffer size and try again.
+			FUniqueBuffer NewBuffer = FUniqueBuffer::Alloc(2 * Buffer.GetSize());
+			BufferTail = NewBuffer.GetView().CopyFrom(ParseBuffer);
+			Buffer = MoveTemp(NewBuffer);
+		}
+		else
+		{
+			// At least one line ending was found. Move any partial line to the front of the buffer and continue.
+			FMemory::Memmove(Buffer.GetData(), ParseBuffer.GetData(), ParseBuffer.GetSize());
+			BufferTail = Buffer.GetView() + ParseBuffer.GetSize();
+		}
+	}
+
+	return Ar->Close();
+}
+
 /**
  * Save a binary array to a file.
  */
diff --git a/Engine/Source/Runtime/Core/Public/Misc/FileHelper.h b/Engine/Source/Runtime/Core/Public/Misc/FileHelper.h
--- a/Engine/Source/Runtime/Core/Public/Misc/FileHelper.h
+++ b/Engine/Source/Runtime/Core/Public/Misc/FileHelper.h
@@ -117,6 +117,14 @@ struct CORE_API FFileHelper
 	UE_DEPRECATED(4.26, "LoadFileToStringArrayWithPredicate no longer supports VerifyFlags. You can use UE::String::ParseLines to split up a string loaded with LoadFileToString")
 	static bool LoadFileToStringArrayWithPredicate(TArray<FString>& Result, const TCHAR* Filename, TFunctionRef<bool(const FString&)> Predicate, EHashOptions VerifyFlags);
 
+	/**
+	 * Load a text file and invoke a visitor for each line. Supports all combination of ANSI/Unicode files and platforms.
+	 *
+	 * @param Filename     Name of the file to load
+	 * @param Visitor      Visitor to invoke for each non-empty line in the file
+	 */
+	static bool LoadFileToStringWithLineVisitor(const TCHAR* Filename, TFunctionRef<void(FStringView Line)> Visitor);
+
 	/**
 	 * Save a binary array to a file.
 	 */
