Deleted redundant functions

These no longer override a base function, but they weren't marked with override so it wasn't caught by the compiler.

#rnx

diff --git a/Engine/Source/Runtime/Core/Private/Internationalization/TextHistory.h b/Engine/Source/Runtime/Core/Private/Internationalization/TextHistory.h
--- a/Engine/Source/Runtime/Core/Private/Internationalization/TextHistory.h
+++ b/Engine/Source/Runtime/Core/Private/Internationalization/TextHistory.h
@@ -257,11 +257,6 @@ public:
 	virtual void UpdateDisplayString() override;
 	//~ End FTextHistory Interface
 
-protected:
-	//~ Begin FTextHistory Interface
-	virtual bool CanRebuildLocalizedDisplayString() { return false; }
-	//~ End FTextHistory Interface
-
 private:
 	/** The ID for an FText (if any) */
 	FTextId TextId;
@@ -638,11 +633,6 @@ public:
 
 	void GetTableIdAndKey(FName& OutTableId, FTextKey& OutKey) const;
 
-protected:
-	//~ Begin FTextHistory Interface
-	virtual bool CanRebuildLocalizedDisplayString() { return false; }
-	//~ End FTextHistory Interface
-
 private:
 	enum class EStringTableLoadingPhase : uint8
 	{
