Hide/Allow header pair for disabling TCHAR in third-party headers in UTF-8 mode, where our definition of TCHAR diffs from theirs.

diff --git a/Engine/Source/Runtime/Core/Public/HAL/AllowTCHAR.h b/Engine/Source/Runtime/Core/Public/HAL/AllowTCHAR.h
new file mode 100644
--- /dev/null
+++ b/Engine/Source/Runtime/Core/Public/HAL/AllowTCHAR.h
@@ -0,0 +1,26 @@
+// Copyright Epic Games, Inc. All Rights Reserved.
+
+// This header is expected to be included after HAL/HideTCHAR.h - see the comments in that header
+// for more details.
+
+#ifdef UE_HIDETCHAR_INCLUDED
+	#undef UE_HIDETCHAR_INCLUDED
+#else
+	#error "HAL/AllowTCHAR.h included without a corresponding include of HAL/AllowTCHAR.h - check your includes."
+#endif
+
+#if !defined(UE_HIDETCHAR_INCLUDED_FIRST)
+	#define UE_HIDETCHAR_INCLUDED_FIRST
+
+	// We want UE_SYSTEM_TCHAR to exist after the first definition of TCHAR - even in
+	// !PLATFORM_TCHAR_IS_UTF8CHAR mode.
+	//
+	// TCHAR is expected to be defined inside the first includes we see being wrapped by our
+	// HideTCHAR/AllowTCHAR pair, so we will define our type in terms of it, before we undo our hack,
+	// if it exists.
+	using UE_SYSTEM_TCHAR = TCHAR;
+#endif
+
+#if PLATFORM_TCHAR_IS_UTF8CHAR
+	#undef TCHAR
+#endif
diff --git a/Engine/Source/Runtime/Core/Public/HAL/HideTCHAR.h b/Engine/Source/Runtime/Core/Public/HAL/HideTCHAR.h
new file mode 100644
--- /dev/null
+++ b/Engine/Source/Runtime/Core/Public/HAL/HideTCHAR.h
@@ -0,0 +1,33 @@
+// Copyright Epic Games, Inc. All Rights Reserved.
+
+// When third parties includes define or reference the 'system' TCHAR, it will clash with our
+// typedef in Platform.h when we're compiling in PLATFORM_TCHAR_IS_UTF8CHAR mode.
+//
+// By including this header before any third party include that refers to TCHAR, we will force
+// those headers to refer to that symbol instead.  After those includes, "HAL/AllowTCHAR.h" should
+// be included, e.g.:
+//
+// #include "HAL/HideTCHAR.h"
+// #include <tchar.h>
+// #include <other_third_party_header_that_refers_to_tchar.h>
+// #include "HAL/AllowTCHAR.h"
+//
+// Linking shouldn't be affected by this 'override' because the UE_SYSTEM_TCHAR type will mangle
+// in the same way.
+//
+// Preprocessor pasting of the token will be affected but that is rare and will need a bespoke
+// solution.
+//
+// If we ever need to deliberately refer to the system TCHAR rather than our Platform.h TCHAR, we
+// can use the UE_SYSTEM_TCHAR type instead.
+
+#ifndef UE_HIDETCHAR_INCLUDED
+	#define UE_HIDETCHAR_INCLUDED
+#else
+	#error "HAL/HideTCHAR.h is being multiply included - an include of HAL/AllowTCHAR.h is likely missing."
+#endif
+
+#if PLATFORM_TCHAR_IS_UTF8CHAR
+	// Override the TCHAR name.
+	#define TCHAR UE_SYSTEM_TCHAR
+#endif
diff --git a/Engine/Source/Runtime/Core/Public/Windows/WindowsSystemIncludes.h b/Engine/Source/Runtime/Core/Public/Windows/WindowsSystemIncludes.h
--- a/Engine/Source/Runtime/Core/Public/Windows/WindowsSystemIncludes.h
+++ b/Engine/Source/Runtime/Core/Public/Windows/WindowsSystemIncludes.h
@@ -17,7 +17,9 @@ THIRD_PARTY_INCLUDES_START
 #include <intrin.h>
 
 #include <stdint.h>
+#include "HAL/HideTCHAR.h"
 #include <tchar.h>
+#include "HAL/AllowTCHAR.h"
 
 // When compiling under Windows, these headers cause us particular problems.  We need to make sure they're included before we pull in our 
 // 'DoNotUseOldUE4Type' namespace.  This is because these headers will redeclare various numeric typedefs, but under the Clang and Visual
