REL - remove dangerous and broken cancellation logic, if they are not launched they will assert which is ok

diff --git a/Engine/Source/Runtime/Core/Public/Experimental/Async/AwaitableTask.h b/Engine/Source/Runtime/Core/Public/Experimental/Async/AwaitableTask.h
--- a/Engine/Source/Runtime/Core/Public/Experimental/Async/AwaitableTask.h
+++ b/Engine/Source/Runtime/Core/Public/Experimental/Async/AwaitableTask.h
@@ -258,12 +258,7 @@ namespace AwaitableTask_Detail
 		inline void Finish()
 		{
 			int LocalCounter = ReferenceCounter.fetch_sub(1, std::memory_order_release);
-			if (LocalCounter == 2 && !IsLaunched()) //this is 2 because fetch_sub returns the value before the decrement (value is actually 1)
-			{
-				//Cancel the Task and Try to run the Continuation if we are the last reference
-				verify(Task.TryCancel());
-			}
-			else if (LocalCounter == 1) //this is 1 because fetch_sub returns the value before the decrement (value is actually 0)
+			if (LocalCounter == 1) //this is 1 because fetch_sub returns the value before the decrement (value is actually 0)
 			{
 				delete this;
 			}
