TArrayView::Mid - Corrected a pointer to const vs const pointer which was forcing the return type to TArrayView<const ElementType>

#rnx

diff --git a/Engine/Source/Runtime/Core/Public/Containers/ArrayView.h b/Engine/Source/Runtime/Core/Public/Containers/ArrayView.h
--- a/Engine/Source/Runtime/Core/Public/Containers/ArrayView.h
+++ b/Engine/Source/Runtime/Core/Public/Containers/ArrayView.h
@@ -371,7 +371,7 @@ public:
 	/** Returns the middle part of the view by taking up to the given number of elements from the given position. */
 	[[nodiscard]] inline TArrayView Mid(SizeType Index, SizeType Count = TNumericLimits<SizeType>::Max()) const
 	{
-		const ElementType* CurrentStart  = GetData();
+		ElementType* const CurrentStart  = GetData();
 		const SizeType     CurrentLength = Num();
 
 		// Clamp minimum index at the start of the range, adjusting the length down if necessary
