[Insights] Clarify HeapId argument for mark as heap command.

diff --git a/Engine/Source/Runtime/Core/Public/ProfilingDebugging/MemoryTrace.h b/Engine/Source/Runtime/Core/Public/ProfilingDebugging/MemoryTrace.h
--- a/Engine/Source/Runtime/Core/Public/ProfilingDebugging/MemoryTrace.h
+++ b/Engine/Source/Runtime/Core/Public/ProfilingDebugging/MemoryTrace.h
@@ -92,7 +92,7 @@ CORE_API HeapId MemoryTrace_RootHeapSpec(const TCHAR* Name, EMemoryTraceHeapFlag
 /**
  * Mark a traced allocation as being a heap.
  * @param Address Address of the allocation
- * @param Heap Heap id, see /ref MemoryTrace_HeapSpec
+ * @param Heap Heap id, see /ref MemoryTrace_HeapSpec. If no specific heap spec has been created the correct root heap needs to be given.
  * @param Flags Additional properties of the heap allocation. Note that \ref EMemoryTraceHeapAllocationFlags::Heap is implicit.
  */
 CORE_API void MemoryTrace_MarkAllocAsHeap(uint64 Address, HeapId Heap, EMemoryTraceHeapAllocationFlags Flags = EMemoryTraceHeapAllocationFlags::None);
