[Insights] Only do memory allocation for callstack tracing when used.

Move code that reserved memory for callstack hashes used in callstack tracing to execute when used. Additionally add configuration options to allow platform specific values and behaviour.

diff --git a/Engine/Source/Runtime/Core/Private/ProfilingDebugging/CallstackTracePrivate.h b/Engine/Source/Runtime/Core/Private/ProfilingDebugging/CallstackTracePrivate.h
--- a/Engine/Source/Runtime/Core/Private/ProfilingDebugging/CallstackTracePrivate.h
+++ b/Engine/Source/Runtime/Core/Private/ProfilingDebugging/CallstackTracePrivate.h
@@ -9,6 +9,23 @@
 
 #if UE_CALLSTACK_TRACE_ENABLED
 
+////////////////////////////////////////////////////////////////////////////////
+#if !defined(UE_CALLSTACK_TRACE_RESERVE_MB)
+	// Initial size of the known set of callstacks
+	#if WITH_EDITOR
+		#define UE_CALLSTACK_TRACE_RESERVE_MB 16 // ~1M callstacks
+	#else
+		#define UE_CALLSTACK_TRACE_RESERVE_MB 8 // ~500k callstacks
+	#endif
+#endif
+
+#if !defined(UE_CALLSTACK_TRACE_RESERVE_GROWABLE)
+	// If disabled the known set will not grow. New callstacks will not be
+	// reported if the set is full
+	#define UE_CALLSTACK_TRACE_RESERVE_GROWABLE 1
+#endif
+
+
 ////////////////////////////////////////////////////////////////////////////////
 UE_TRACE_CHANNEL_EXTERN(CallstackChannel)
 
@@ -35,7 +52,6 @@ namespace {
 			: KnownSet(InMalloc)
 			, CallstackIdCounter(1) // 0 is reserved for "unknown callstack"
 		{
-			KnownSet.Reserve(1024 * 1024 * 2);
 		}
 		
 		uint32 AddCallstack(const FBacktraceEntry& Entry)
@@ -50,7 +66,19 @@ namespace {
 			KnownSet.Find(Hash, &Id, &bAlreadyAdded);
 			if (!bAlreadyAdded)
 			{
-				Id = CallstackIdCounter.fetch_add(1, std::memory_order_relaxed); 
+				Id = CallstackIdCounter.fetch_add(1, std::memory_order_relaxed);
+				// On the first callstack reserve memory up front
+				if (Id == 1)
+				{
+					KnownSet.Reserve(InitialReserveCount);
+				}
+#if !UE_CALLSTACK_TRACE_RESERVE_GROWABLE
+				// If configured as not growable, start returning unknown id's when full.
+				if (Id >= InitialReserveCount)
+				{
+					return 0;
+				}
+#endif
 				KnownSet.Emplace(Hash, Id);
 				UE_TRACE_LOG(Memory, CallstackSpec, CallstackChannel)
 					<< CallstackSpec.CallstackId(Id)
@@ -81,6 +109,9 @@ namespace {
 		};
 		typedef TGrowOnlyLockFreeHash<FEncounteredCallstackSetEntry, uint64, uint32> FEncounteredCallstackSet;
 
+		constexpr static uint32 InitialReserveBytes = UE_CALLSTACK_TRACE_RESERVE_MB * 1024 * 1024;
+		constexpr static uint32 InitialReserveCount = InitialReserveBytes / sizeof(FEncounteredCallstackSetEntry);
+		
 		FEncounteredCallstackSet 	KnownSet;
 		std::atomic_uint32_t		CallstackIdCounter;
 	};
