[Insights] Move UE_MEMORY_TAGS_TRACE_ENABLED to build file.

We need to set this define in other platform build files, and since some of them are included on windows (editor) builds as well, we can end up with warnings of redefinition of variable.

diff --git a/Engine/Source/Runtime/Core/Core.Build.cs b/Engine/Source/Runtime/Core/Core.Build.cs
--- a/Engine/Source/Runtime/Core/Core.Build.cs
+++ b/Engine/Source/Runtime/Core/Core.Build.cs
@@ -87,6 +87,11 @@ public class Core : ModuleRules
 			{
 				PrivateDependencyModuleNames.Add("WinPixEventRuntime");
 			}
+
+			if (Target.Configuration != UnrealTargetConfiguration.Shipping && Target.Type != TargetType.Program)
+			{
+				PublicDefinitions.Add("UE_MEMORY_TAGS_TRACE_ENABLED=1");
+			}
 		}
 		else if ((Target.Platform == UnrealTargetPlatform.HoloLens))
 		{
diff --git a/Engine/Source/Runtime/Core/Private/Windows/WindowsRunnableThread.h b/Engine/Source/Runtime/Core/Private/Windows/WindowsRunnableThread.h
--- a/Engine/Source/Runtime/Core/Private/Windows/WindowsRunnableThread.h
+++ b/Engine/Source/Runtime/Core/Private/Windows/WindowsRunnableThread.h
@@ -11,6 +11,7 @@
 #include "HAL/Runnable.h"
 #include "HAL/RunnableThread.h"
 #include "HAL/ThreadManager.h"
+#include "ProfilingDebugging/MemoryTrace.h"
 #include "Windows/WindowsHWrapper.h"
 
 class FRunnable;
diff --git a/Engine/Source/Runtime/Core/Public/ProfilingDebugging/TagTrace.h b/Engine/Source/Runtime/Core/Public/ProfilingDebugging/TagTrace.h
--- a/Engine/Source/Runtime/Core/Public/ProfilingDebugging/TagTrace.h
+++ b/Engine/Source/Runtime/Core/Public/ProfilingDebugging/TagTrace.h
@@ -1,7 +1,7 @@
 // Copyright Epic Games, Inc. All Rights Reserved.
 #pragma once
 #include "CoreTypes.h"
-#include "ProfilingDebugging/MemoryTrace.h"
+#include "Trace/Config.h"
 #include "Trace/Detail/LogScope.h"
 
 ////////////////////////////////////////////////////////////////////////////////
@@ -20,18 +20,11 @@ CORE_API int32	MemoryTrace_AnnounceFNameTag(const class FName& TagName);
 CORE_API int32	MemoryTrace_GetActiveTag();
 
 ////////////////////////////////////////////////////////////////////////////////
-
-#if !defined(UE_MEMORY_TAGS_TRACE_ENABLED) && UE_TRACE_ENABLED && UE_MEMORY_TRACE_ENABLED
-	#if PLATFORM_WINDOWS 
-		#define UE_MEMORY_TAGS_TRACE_ENABLED !UE_BUILD_SHIPPING
-	#endif
-#endif
-
 #if !defined(UE_MEMORY_TAGS_TRACE_ENABLED)
 	#define UE_MEMORY_TAGS_TRACE_ENABLED 0
 #endif
 
-#if UE_MEMORY_TAGS_TRACE_ENABLED
+#if UE_MEMORY_TAGS_TRACE_ENABLED && UE_TRACE_ENABLED
 
 ////////////////////////////////////////////////////////////////////////////////
 
diff --git a/Engine/Source/Runtime/RHI/Public/GPUDefragAllocator.h b/Engine/Source/Runtime/RHI/Public/GPUDefragAllocator.h
--- a/Engine/Source/Runtime/RHI/Public/GPUDefragAllocator.h
+++ b/Engine/Source/Runtime/RHI/Public/GPUDefragAllocator.h
@@ -9,6 +9,7 @@
 #include "HAL/IConsoleManager.h"
 #include "Containers/List.h"
 #include "HAL/LowLevelMemTracker.h"
+#include "ProfilingDebugging/MemoryTrace.h"
 
 #define LOG_EVERY_ALLOCATION			0
 #define DUMP_ALLOC_FREQUENCY			0 // 100
