[Insights] Additional memory tracing for many platforms

Implement additional memory tracing for many platforms to bring up parity with LLMs tracking.

diff --git a/Engine/Source/Runtime/Core/Private/ProfilingDebugging/Microsoft/WindowsModuleDiagnostics.cpp b/Engine/Source/Runtime/Core/Private/ProfilingDebugging/Microsoft/WindowsModuleDiagnostics.cpp
--- a/Engine/Source/Runtime/Core/Private/ProfilingDebugging/Microsoft/WindowsModuleDiagnostics.cpp
+++ b/Engine/Source/Runtime/Core/Private/ProfilingDebugging/Microsoft/WindowsModuleDiagnostics.cpp
@@ -6,6 +6,9 @@
 #include "Trace/Trace.h"
 #include "Trace/Trace.inl"
 #include "ProfilingDebugging/ModuleDiagnostics.h"
+#include "ProfilingDebugging/MemoryTrace.h"
+#include "ProfilingDebugging/TagTrace.h"
+#include "HAL/LowLevelMemTracker.h"
 
 #include COMPILED_PLATFORM_HEADER(PlatformModuleDiagnostics.h)
 
@@ -49,6 +52,7 @@ private:
 	static FModuleTrace*		Instance;
 	SubscriberArray				Subscribers;
 	void*						CallbackCookie;
+	HeapId						ProgramHeapId;
 };
 
 ////////////////////////////////////////////////////////////////////////////////
@@ -77,6 +81,11 @@ FModuleTrace* FModuleTrace::Get()
 void FModuleTrace::Initialize()
 {
 	using namespace UE::Trace;
+	
+	ProgramHeapId = MemoryTrace_HeapSpec(
+		SystemMemory, 
+		TEXT("Program"), 
+		EMemoryTraceHeapFlags::NeverFrees);
 
 	UE_TRACE_LOG(Diagnostics, ModuleInit, ModuleChannel, sizeof(char) * 3)
 		<< ModuleInit.SymbolFormat("pdb", 3)
@@ -192,6 +201,15 @@ void FModuleTrace::OnDllLoaded(const UNICODE_STRING& Name, UPTRINT Base)
 		<< ModuleLoad.Size(OptionalHeader.SizeOfImage)
 		<< ModuleLoad.ImageId(ImageId.GetData(), ImageId.Num());
 
+	{
+#if ENABLE_LOW_LEVEL_MEM_TRACKER
+		UE_MEMSCOPE(ELLMTag::ProgramSize);
+#endif
+		MemoryTrace_Alloc(Base, OptionalHeader.SizeOfImage, 4 * 1024);
+		MemoryTrace_MarkAllocAsHeap(Base, ProgramHeapId);
+		MemoryTrace_Alloc(Base, OptionalHeader.SizeOfImage, 4 * 1024);
+	}
+	
 	for (SubscribeFunc Subscriber : Subscribers)
 	{
 		Subscriber(true, (void*)Base, Name.Buffer);
