OodleDataCompression: clean up common error cases in OodleDecode

diff --git a/Engine/Source/Runtime/Core/Private/Compression/OodleDataCompression.cpp b/Engine/Source/Runtime/Core/Private/Compression/OodleDataCompression.cpp
--- a/Engine/Source/Runtime/Core/Private/Compression/OodleDataCompression.cpp
+++ b/Engine/Source/Runtime/Core/Private/Compression/OodleDataCompression.cpp
@@ -253,7 +253,7 @@ struct OodleScratchBuffers
 				}
 				else
 				{
-					UE_LOG(OodleDataCompression, Error, TEXT("FOodleDataCompressionFormat - shutting down while in use?"));
+					UE_LOG(OodleDataCompression, Error, TEXT("OodleDataCompression - shutting down while in use?"));
 				}
 			}
 
@@ -267,17 +267,17 @@ struct OodleScratchBuffers
 		OO_SINTa InCompBufSize = IntCastChecked<OO_SINTa>(InCompBufSize64);
 		OO_SINTa InRawLen = IntCastChecked<OO_SINTa>(InRawLen64);
 
-		// find the minimum size needed for this decode, OodleScratchMemorySize may be larger
-		OodleLZ_Compressor CurCompressor = OodleLZ_GetChunkCompressor(InCompBuf, InCompBufSize, NULL);
-		if( CurCompressor == OodleLZ_Compressor_Invalid )
+		if ( InCompBufSize <= 0 )
 		{
-			UE_LOG(OodleDataCompression, Error, TEXT("FOodleDataCompressionFormat - no Oodle compressor found!"));
+			UE_LOG(OodleDataCompression, Error, TEXT("OodleDecode: no compressed buffer size to decode!"));
 			return OODLELZ_FAILED;
 		}
-			
-		OO_SINTa DecoderMemorySize = OodleLZDecoder_MemorySizeNeeded(CurCompressor, InRawLen);
-		void * DecoderMemory = NULL;
-
+		else if ( static_cast<const OO_U8 *>(InCompBuf)[0] == 0 )
+		{
+			UE_LOG(OodleDataCompression, Error, TEXT("OodleDecode: compressed buffer starts with zero byte; invalid or corrupt compressed stream."));
+			return OODLELZ_FAILED;
+		}
+					
 		// try to take a mutex for one of the pre-allocated decode buffers
 		for (int i = 0; i < OodleScratchBufferCount; ++i)
 		{
@@ -289,7 +289,7 @@ struct OodleScratchBuffers
 					OodleScratches[i].OodleScratchMemory = FMemory::Malloc(OodleScratchMemorySize);
 					if (OodleScratches[i].OodleScratchMemory == nullptr)
 					{
-						UE_LOG(OodleDataCompression, Error, TEXT("FOodleDataCompressionFormat - Failed to allocate scratch buffer %d bytes!"), OodleScratchMemorySize);
+						UE_LOG(OodleDataCompression, Error, TEXT("OodleDecode: Failed to allocate scratch buffer %d bytes!"), OodleScratchMemorySize);
 						return OODLELZ_FAILED;
 					}
 				}
@@ -304,14 +304,32 @@ struct OodleScratchBuffers
 				return (int64) Result;
 			}
 		}
-			
+		
+		// couldn't lock a shared scratch buffer
+		//  allocate decoder mem on the heap to avoid waiting
+		 
 		//UE_LOG(OodleDataCompression, Display, TEXT("Decode with malloc : %d -> %d"),InCompBufSize,InRawLen );
+		
+		// find the minimum size needed for this decode, OodleScratchMemorySize may be larger
+		OodleLZ_Compressor CurCompressor = OodleLZ_GetChunkCompressor(InCompBuf, InCompBufSize, NULL);
+		if( CurCompressor == OodleLZ_Compressor_Invalid )
+		{
+			UE_LOG(OodleDataCompression, Error, TEXT("OodleDecode : no Oodle compressor found!; likely an invalid or corrupt compressed stream."));
+			return OODLELZ_FAILED;
+		}
+
+		OO_SINTa DecoderMemorySize = OodleLZDecoder_MemorySizeNeeded(CurCompressor, InRawLen);
+		if ( DecoderMemorySize <= 0 )
+		{
+			UE_LOG(OodleDataCompression, Error, TEXT("OodleDecode : MemorySizeNeeded invalid; likely an invalid or corrupt compressed stream."));
+			return OODLELZ_FAILED;
+		}
 
 		// allocate memory for the decoder so that Oodle doesn't allocate anything internally
-		DecoderMemory = FMemory::Malloc(DecoderMemorySize);
+		void * DecoderMemory = FMemory::Malloc(DecoderMemorySize);
 		if (DecoderMemory == NULL) 
 		{
-			UE_LOG(OodleDataCompression, Error, TEXT("FOodleDataCompressionFormat::OodleDecode - Failed to allocate %d!"), DecoderMemorySize);
+			UE_LOG(OodleDataCompression, Error, TEXT("OodleDecode - Failed to allocate %d!"), (int)DecoderMemorySize);
 			return OODLELZ_FAILED;
 		}
 
@@ -336,7 +354,7 @@ struct OodleScratchBuffers
 
 		if ( InCompressedBufferSize < (int64) OodleLZ_GetCompressedBufferSizeNeeded(LZCompressor,IntCastChecked<OO_SINTa>(InUncompressedSize)) )
 		{
-			UE_LOG(OodleDataCompression,Error,TEXT("OutCompressedSize too small\n"));		
+			UE_LOG(OodleDataCompression,Error,TEXT("OodleEncode: OutCompressedSize too small\n"));		
 			return OODLELZ_FAILED;
 		}
 		
@@ -351,7 +369,7 @@ struct OodleScratchBuffers
 					OodleScratches[i].OodleScratchMemory = FMemory::Malloc(OodleScratchMemorySize);
 					if (OodleScratches[i].OodleScratchMemory == nullptr)
 					{
-						UE_LOG(OodleDataCompression, Error, TEXT("FOodleDataCompressionFormat - Failed to allocate scratch buffer %d bytes!"), OodleScratchMemorySize);
+						UE_LOG(OodleDataCompression, Error, TEXT("OodleEncode: Failed to allocate scratch buffer %d bytes!"), OodleScratchMemorySize);
 						return OODLELZ_FAILED;
 					}
 				}
