LLT compilation errors

diff --git a/Engine/Source/Runtime/Core/Private/ProfilingDebugging/TagTrace.cpp b/Engine/Source/Runtime/Core/Private/ProfilingDebugging/TagTrace.cpp
--- a/Engine/Source/Runtime/Core/Private/ProfilingDebugging/TagTrace.cpp
+++ b/Engine/Source/Runtime/Core/Private/ProfilingDebugging/TagTrace.cpp
@@ -8,7 +8,7 @@
 #include "Trace/Trace.inl"
 #include "UObject/NameTypes.h"
 
-#if UE_MEMORY_TAGS_TRACE_ENABLED
+#if UE_MEMORY_TAGS_TRACE_ENABLED && UE_TRACE_ENABLED
 
 ////////////////////////////////////////////////////////////////////////////////
 
@@ -262,7 +262,7 @@ int32 FTagTrace::AnnounceCustomTag(int32 Tag, int32 ParentTag, const ANSICHAR* D
 ////////////////////////////////////////////////////////////////////////////////
 void MemoryTrace_InitTags(FMalloc* InMalloc)
 {
-#if UE_MEMORY_TAGS_TRACE_ENABLED
+#if UE_MEMORY_TAGS_TRACE_ENABLED && UE_TRACE_ENABLED
 	GTagTrace = (FTagTrace*)InMalloc->Malloc(sizeof(FTagTrace), alignof(FTagTrace));
 	new (GTagTrace) FTagTrace(InMalloc);
 #endif
@@ -271,7 +271,7 @@ void MemoryTrace_InitTags(FMalloc* InMalloc)
 ////////////////////////////////////////////////////////////////////////////////
 int32 MemoryTrace_AnnounceCustomTag(int32 Tag, int32 ParentTag, const TCHAR* Display)
 {
-#if UE_MEMORY_TAGS_TRACE_ENABLED
+#if UE_MEMORY_TAGS_TRACE_ENABLED && UE_TRACE_ENABLED
 	//todo: How do we check if tag trace is active?
 	if (GTagTrace)
 	{
@@ -284,7 +284,7 @@ int32 MemoryTrace_AnnounceCustomTag(int32 Tag, int32 ParentTag, const TCHAR* Dis
 ////////////////////////////////////////////////////////////////////////////////
 int32 MemoryTrace_AnnounceFNameTag(const FName& TagName)
 {
-#if UE_MEMORY_TAGS_TRACE_ENABLED
+#if UE_MEMORY_TAGS_TRACE_ENABLED && UE_TRACE_ENABLED
 	if (GTagTrace)
 	{
 		return GTagTrace->AnnounceFNameTag(TagName);
@@ -296,7 +296,7 @@ int32 MemoryTrace_AnnounceFNameTag(const FName& TagName)
 ////////////////////////////////////////////////////////////////////////////////
 int32 MemoryTrace_GetActiveTag()
 {
-#if UE_MEMORY_TAGS_TRACE_ENABLED
+#if UE_MEMORY_TAGS_TRACE_ENABLED && UE_TRACE_ENABLED
 	return GActiveTag;
 #else
 	return -1;
