Added FUnixPlatformFile::CopyFile() that differs from the generic implementation in that it also copies the file permission bits.

This is similar to CL 2421073, 2424491 for FApplePlatformFile.

diff --git a/Engine/Source/Runtime/Core/Private/Unix/UnixPlatformFile.cpp b/Engine/Source/Runtime/Core/Private/Unix/UnixPlatformFile.cpp
--- a/Engine/Source/Runtime/Core/Private/Unix/UnixPlatformFile.cpp
+++ b/Engine/Source/Runtime/Core/Private/Unix/UnixPlatformFile.cpp
@@ -1148,6 +1148,21 @@ bool FUnixPlatformFile::IterateDirectoryCommon(const TCHAR* Directory, const TFu
 	return Result;
 }
 
+bool FUnixPlatformFile::CopyFile(const TCHAR* To, const TCHAR* From, EPlatformFileRead ReadFlags, EPlatformFileWrite WriteFlags)
+{
+	bool Result = IPlatformFile::CopyFile(To, From, ReadFlags, WriteFlags);
+	if (Result)
+	{
+		struct stat FileInfo;
+		if (stat(TCHAR_TO_UTF8(*NormalizeFilename(From, false)), &FileInfo) == 0)
+		{
+			FileInfo.st_mode |= S_IWUSR;
+			chmod(TCHAR_TO_UTF8(*NormalizeFilename(To, true)), FileInfo.st_mode);
+		}
+	}
+	return Result;
+}
+
 bool FUnixPlatformFile::CreateDirectoriesFromPath(const TCHAR* Path)
 {
 	// if the file already exists, then directories exist.
diff --git a/Engine/Source/Runtime/Core/Public/Unix/UnixPlatformFile.h b/Engine/Source/Runtime/Core/Public/Unix/UnixPlatformFile.h
--- a/Engine/Source/Runtime/Core/Public/Unix/UnixPlatformFile.h
+++ b/Engine/Source/Runtime/Core/Public/Unix/UnixPlatformFile.h
@@ -54,6 +54,8 @@ public:
 	virtual bool IterateDirectory(const TCHAR* Directory, FDirectoryVisitor& Visitor) override;
 	virtual bool IterateDirectoryStat(const TCHAR* Directory, FDirectoryStatVisitor& Visitor) override;
 
+	virtual bool CopyFile(const TCHAR* To, const TCHAR* From, EPlatformFileRead ReadFlags = EPlatformFileRead::None, EPlatformFileWrite WriteFlags = EPlatformFileWrite::None) override;
+
 protected:
 	bool IterateDirectoryCommon(const TCHAR* Directory, const TFunctionRef<bool(struct dirent*)>& Visitor);
 
