Fix default copy constructor in FBitWriter to be const. This makes it consistent with super/subclasses and fixes compile errors if instances are used in TMaps.

diff --git a/Engine/Source/Runtime/Core/Public/Serialization/BitWriter.h b/Engine/Source/Runtime/Core/Public/Serialization/BitWriter.h
--- a/Engine/Source/Runtime/Core/Public/Serialization/BitWriter.h
+++ b/Engine/Source/Runtime/Core/Public/Serialization/BitWriter.h
@@ -28,7 +28,7 @@ public:
 	 */
 	FBitWriter( int64 InMaxBits, bool AllowResize = false );
 
-	FBitWriter(FBitWriter&) = default;
+	FBitWriter(const FBitWriter&) = default;
     FBitWriter& operator=(const FBitWriter&) = default;
     FBitWriter(FBitWriter&&) = default;
     FBitWriter& operator=(FBitWriter&&) = default;
