[Insights] Reset commandline channels at every start.

* Whenever starting a trace command (from the api or console command), reset the current command line list of channels.
* Don't assume default channels should be enabled if none are given. For command line arguments this will still be the case, but not console commands or the API.
* Emit log message when auto connecting.

diff --git a/Engine/Plugins/Animation/GameplayInsights/Source/GameplayInsights/Private/GameplayInsightsModule.cpp b/Engine/Plugins/Animation/GameplayInsights/Source/GameplayInsights/Private/GameplayInsightsModule.cpp
--- a/Engine/Plugins/Animation/GameplayInsights/Source/GameplayInsights/Private/GameplayInsightsModule.cpp
+++ b/Engine/Plugins/Animation/GameplayInsights/Source/GameplayInsights/Private/GameplayInsightsModule.cpp
@@ -175,6 +175,7 @@ void FGameplayInsightsModule::StartupModule()
 			if (!UnrealInsightsModule.GetStoreClient())
 			{
 #if WITH_TRACE_STORE
+				UE_LOG(LogCore, Display, TEXT("GameplayInsights module auto-connecting to internal trace server..."));
 				// Create the Store Service.
 				FString StoreDir = FPaths::ProjectSavedDir() / TEXT("TraceSessions");
 				UE::Trace::FStoreService::FDesc StoreServiceDesc;
@@ -191,6 +192,7 @@ void FGameplayInsightsModule::StartupModule()
 				UnrealInsightsModule.ConnectToStore(TEXT("localhost"), StoreService->GetPort());
 				UE::Trace::SendTo(TEXT("localhost"), StoreService->GetRecorderPort());
 #else
+				UE_LOG(LogCore, Display, TEXT("GameplayInsights module auto-connecting to local trace server..."));
 				UnrealInsightsModule.ConnectToStore(TEXT("127.0.0.1"));
 				const bool bConnected = FTraceAuxiliary::Start(
 					FTraceAuxiliary::EConnectionType::Network,
diff --git a/Engine/Plugins/Runtime/NetworkPredictionInsights/Source/NetworkPredictionInsights/Private/NetworkPredictionInsightsModule.cpp b/Engine/Plugins/Runtime/NetworkPredictionInsights/Source/NetworkPredictionInsights/Private/NetworkPredictionInsightsModule.cpp
--- a/Engine/Plugins/Runtime/NetworkPredictionInsights/Source/NetworkPredictionInsights/Private/NetworkPredictionInsightsModule.cpp
+++ b/Engine/Plugins/Runtime/NetworkPredictionInsights/Source/NetworkPredictionInsights/Private/NetworkPredictionInsightsModule.cpp
@@ -102,6 +102,7 @@ void FNetworkPredictionInsightsModule::StartupModule()
 			if (!UnrealInsightsModule.GetStoreClient())
 			{
 #if WITH_TRACE_STORE
+				UE_LOG(LogCore, Display, TEXT("NetworkPredictionInsights module auto-connecting to internal trace server..."));
 				// Create the Store Service.
 				FString StoreDir = FPaths::ProjectSavedDir() / TEXT("TraceSessions");
 				UE::Trace::FStoreService::FDesc StoreServiceDesc;
@@ -118,6 +119,7 @@ void FNetworkPredictionInsightsModule::StartupModule()
 				ensure(UnrealInsightsModule.ConnectToStore(TEXT("localhost"), StoreService->GetPort()));
 				UE::Trace::SendTo(TEXT("localhost"), StoreService->GetRecorderPort());
 #else
+				UE_LOG(LogCore, Display, TEXT("NetworkPredictionInsights module auto-connecting to local trace server..."));
 				UnrealInsightsModule.ConnectToStore(TEXT("127.0.0.1"));
 				const bool bConnected = FTraceAuxiliary::Start(
 					FTraceAuxiliary::EConnectionType::Network,
diff --git a/Engine/Source/Runtime/Core/Private/ProfilingDebugging/TraceAuxiliary.cpp b/Engine/Source/Runtime/Core/Private/ProfilingDebugging/TraceAuxiliary.cpp
--- a/Engine/Source/Runtime/Core/Private/ProfilingDebugging/TraceAuxiliary.cpp
+++ b/Engine/Source/Runtime/Core/Private/ProfilingDebugging/TraceAuxiliary.cpp
@@ -81,6 +81,7 @@ public:
 	bool					IsConnected() const;
 	void					GetActiveChannelsString(FStringBuilderBase& String) const;
 	void					AddCommandlineChannels(const TCHAR* ChannelList);
+	void					ResetCommandlineChannels();
 	void					EnableChannels(const TCHAR* ChannelList);
 	void					DisableChannels(const TCHAR* ChannelList);
 	bool					Connect(ETraceConnectType Type, const TCHAR* Parameter);
@@ -134,6 +135,12 @@ void FTraceAuxiliaryImpl::AddCommandlineChannels(const TCHAR* ChannelList)
 	ForEachChannel(ChannelList, true, &FTraceAuxiliaryImpl::AddChannel);
 }
 
+////////////////////////////////////////////////////////////////////////////////
+void FTraceAuxiliaryImpl::ResetCommandlineChannels()
+{
+	CommandlineChannels.Reset();
+}
+
 ////////////////////////////////////////////////////////////////////////////////
 void FTraceAuxiliaryImpl::EnableChannels(const TCHAR* ChannelList)
 {
@@ -268,13 +275,6 @@ bool FTraceAuxiliaryImpl::Connect(ETraceConnectType Type, const TCHAR* Parameter
 		return false;
 	}
 
-	// We're now connected. If we don't appear to have any channels we'll set
-	// some defaults for the user. Less futzing.
-	if (!CommandlineChannels.Num())
-	{
-		AddCommandlineChannels(GDefaultChannels);
-	}
-
 	ResumeChannels();
 
 	State = EState::Tracing;
@@ -1056,6 +1056,7 @@ bool FTraceAuxiliary::Start(EConnectionType Type, const TCHAR* Target, const TCH
 	if (Channels)
 	{
 		UE_LOG(LogCore, Display, TEXT("Trace channels: '%s'"), Channels);
+		GTraceAuxiliary.ResetCommandlineChannels();
 		GTraceAuxiliary.AddCommandlineChannels(Channels);
 		GTraceAuxiliary.EnableCommandlineChannels();
 	}
@@ -1260,12 +1261,13 @@ void FTraceAuxiliary::TryAutoConnect()
 	if (bShouldAutoConnect)
 	{
 	#if PLATFORM_WINDOWS
-	// If we can detect a named event it means UnrealInsights (Browser Mode) is running.
-	// In this case, we try to auto-connect with the Trace Server.
+		// If we can detect a named event it means UnrealInsights (Browser Mode) is running.
+		// In this case, we try to auto-connect with the Trace Server.
 		HANDLE KnownEvent = ::OpenEvent(EVENT_ALL_ACCESS, false, TEXT("Local\\UnrealInsightsBrowser"));
 		if (KnownEvent != nullptr)
 		{
-			Start(EConnectionType::Network, TEXT("127.0.0.1"), nullptr, nullptr);
+			UE_LOG(LogCore, Display, TEXT("Unreal Insights instance detected, auto-connecting to local trace server..."));
+			Start(EConnectionType::Network, TEXT("127.0.0.1"), TEXT("default"), nullptr);
 			::CloseHandle(KnownEvent);
 		}
 	#endif
