#ue5 - LWC: changes for Chaos propagating incorrect or NaN rotations caused by a few issues, mainly MSVC generating bad code with /arch:AVX enabled, but also due to API holes in PMatrix.

Main compiler bug workaround was in FRotation3::Conjugate() used by CalculateAngularVelocity1 (and unused CalculateAngularVelocity2). MSVC generated incorrect optimized code when complied with /arch:AVX, which breaks down on negating the individual elements and reassembling them. Workaround is to use the vectorized formulation already in the base TQuat (or disabling optimizations around the original code). Without this change, BodySolver::SetW() would assign (but now asserts) NaN values generated from CalculateAngularVelocity with 2 valid input rotations.

Includes various fixes for PMatrix<FReal, 3, 3> causing last element to take a value not equal to 1, which caused scaling of results. Due to missing operator*= and operator+= overrides (which also fixed bad Determinant() as well as bad transformations).

Also fixed missing definitions for Zero and Identity for recently refactored PMatrix, since only the FMatrix one was defined in Utilities.cpp

Refactored some massive inlined Matrix and Quat functions to separate cpp files, because the inlined assembly would go to sometimes over 400 lines per matrix mutliply. This made debugging difficult and was providing minimum value (and there are outstanding MSVC bug as well with code gen issues in overly large functions so best to avoid those).

Re-enabled AVX intrinsics on previously disabled platforms (from CL 18456125).
Note that for testing on PC, you can set VCToolChain.cs to enable /arch:AVX (which is the build flag that surfaces the issue). With that flag off but other avx/sse remaining on, bug does not occur. Toggling these will cause a full rebuild.

diff --git a/Engine/Source/Runtime/Core/Private/Math/UnrealMath.cpp b/Engine/Source/Runtime/Core/Private/Math/UnrealMath.cpp
--- a/Engine/Source/Runtime/Core/Private/Math/UnrealMath.cpp
+++ b/Engine/Source/Runtime/Core/Private/Math/UnrealMath.cpp
@@ -698,6 +698,58 @@ namespace Math
 		return RotatorFromQuat;
 	}
 
+	template<typename T>
+	TMatrix<T> TQuat<T>::operator*(const TMatrix<T>& M) const
+	{
+		TMatrix<T> Result;
+#if PLATFORM_ENABLE_VECTORINTRINSICS
+		const TQuat<T>::QuatVectorRegister This = VectorLoadAligned(this);
+		const TQuat<T>::QuatVectorRegister Inv = VectorQuaternionInverse(This);
+		for (int32 I = 0; I < 4; ++I)
+		{
+			TQuat<T>::QuatVectorRegister VT, VR;
+			const TQuat<T>::QuatVectorRegister VQ = VectorLoad(M.M[I]);
+			VT = VectorQuaternionMultiply2(This, VQ);
+			VR = VectorQuaternionMultiply2(VT, Inv);
+			VectorStore(VR, Result.M[I]);
+		}
+#else
+		const TQuat<T> Inv = Inverse();
+		for (int32 I = 0; I < 4; ++I)
+		{
+			TQuat<T> VT, VR;
+			const TQuat<T> VQ(M.M[I][0], M.M[I][1], M.M[I][2], M.M[I][3]);
+			VectorQuaternionMultiply(&VT, this, &VQ);
+			VectorQuaternionMultiply(&VR, &VT, &Inv);
+			Result.M[I][0] = VR.X;
+			Result.M[I][1] = VR.Y;
+			Result.M[I][2] = VR.Z;
+			Result.M[I][3] = VR.W;
+		}
+#endif
+
+		return Result;
+	}
+
+	template<typename T>
+	void TQuat<T>::ToMatrix(TMatrix<T>& R) const
+	{
+		// Note: copied and modified from TQuatRotationTranslationMatrix<> mainly to avoid circular dependency.
+#if !(UE_BUILD_SHIPPING || UE_BUILD_TEST) && WITH_EDITORONLY_DATA
+		// Make sure Quaternion is normalized
+		ensure(IsNormalized());
+#endif
+		const T x2 = X + X;    const T y2 = Y + Y;    const T z2 = Z + Z;
+		const T xx = X * x2;   const T xy = X * y2;   const T xz = X * z2;
+		const T yy = Y * y2;   const T yz = Y * z2;   const T zz = Z * z2;
+		const T wx = W * x2;   const T wy = W * y2;   const T wz = W * z2;
+
+		R.M[0][0] = 1.0f - (yy + zz);	R.M[1][0] = xy - wz;				R.M[2][0] = xz + wy;			R.M[3][0] = 0.0f;
+		R.M[0][1] = xy + wz;			R.M[1][1] = 1.0f - (xx + zz);		R.M[2][1] = yz - wx;			R.M[3][1] = 0.0f;
+		R.M[0][2] = xz - wy;			R.M[1][2] = yz + wx;				R.M[2][2] = 1.0f - (xx + yy);	R.M[3][2] = 0.0f;
+		R.M[0][3] = 0.0f;				R.M[1][3] = 0.0f;					R.M[2][3] = 0.0f;				R.M[3][3] = 1.0f;
+	}
+
 	template<typename T>
 	TQuat<T> TQuat<T>::MakeFromEuler(const TVector<T>& Euler)
 	{
diff --git a/Engine/Source/Runtime/Core/Private/Math/UnrealMathSSE.cpp b/Engine/Source/Runtime/Core/Private/Math/UnrealMathSSE.cpp
new file mode 100644
--- /dev/null
+++ b/Engine/Source/Runtime/Core/Private/Math/UnrealMathSSE.cpp
@@ -0,0 +1,301 @@
+// Copyright Epic Games, Inc. All Rights Reserved.
+
+#include "Math/UnrealMathSSE.h"
+#include "Math/Matrix.h"
+
+// This file implements non-inlined functions from UnrealMathSSE.h
+#if UE_USING_UNREALMATH_SSE
+
+void VectorMatrixMultiply(FMatrix44f* Result, const FMatrix44f* Matrix1, const FMatrix44f* Matrix2)
+{
+	const VectorRegister4Float* A = (const VectorRegister4Float*)Matrix1;
+	const VectorRegister4Float* B = (const VectorRegister4Float*)Matrix2;
+	VectorRegister4Float* R = (VectorRegister4Float*)Result;
+	VectorRegister4Float Temp, R0, R1, R2;
+
+	// First row of result (Matrix1[0] * Matrix2).
+	Temp = VectorMultiply(VectorReplicate(A[0], 0), B[0]);
+	Temp = VectorMultiplyAdd(VectorReplicate(A[0], 1), B[1], Temp);
+	Temp = VectorMultiplyAdd(VectorReplicate(A[0], 2), B[2], Temp);
+	R0 = VectorMultiplyAdd(VectorReplicate(A[0], 3), B[3], Temp);
+
+	// Second row of result (Matrix1[1] * Matrix2).
+	Temp = VectorMultiply(VectorReplicate(A[1], 0), B[0]);
+	Temp = VectorMultiplyAdd(VectorReplicate(A[1], 1), B[1], Temp);
+	Temp = VectorMultiplyAdd(VectorReplicate(A[1], 2), B[2], Temp);
+	R1 = VectorMultiplyAdd(VectorReplicate(A[1], 3), B[3], Temp);
+
+	// Third row of result (Matrix1[2] * Matrix2).
+	Temp = VectorMultiply(VectorReplicate(A[2], 0), B[0]);
+	Temp = VectorMultiplyAdd(VectorReplicate(A[2], 1), B[1], Temp);
+	Temp = VectorMultiplyAdd(VectorReplicate(A[2], 2), B[2], Temp);
+	R2 = VectorMultiplyAdd(VectorReplicate(A[2], 3), B[3], Temp);
+
+	// Fourth row of result (Matrix1[3] * Matrix2).
+	Temp = VectorMultiply(VectorReplicate(A[3], 0), B[0]);
+	Temp = VectorMultiplyAdd(VectorReplicate(A[3], 1), B[1], Temp);
+	Temp = VectorMultiplyAdd(VectorReplicate(A[3], 2), B[2], Temp);
+	Temp = VectorMultiplyAdd(VectorReplicate(A[3], 3), B[3], Temp);
+
+	// Store result. Must not be done during steps above in case source and destination are the same.
+	R[0] = R0;
+	R[1] = R1;
+	R[2] = R2;
+	R[3] = Temp;
+
+	// TODO: Should change to DiagnosticCheckNaN call, though that's only in Debug.
+#if !(UE_BUILD_SHIPPING || UE_BUILD_TEST)
+	ensure(!Result->ContainsNaN());
+#endif
+}
+
+
+void VectorMatrixMultiply(FMatrix44d* Result, const FMatrix44d* Matrix1, const FMatrix44d* Matrix2)
+{
+	// Warning: FMatrix44d alignment may not match VectorRegister4Double, so you can't just cast to VectorRegister4Double*.
+	typedef double Double4x4[4][4];
+	const Double4x4& AMRows = *((const Double4x4*)Matrix1->M);
+	const Double4x4& BMRows = *((const Double4x4*)Matrix2->M);
+	VectorRegister4Double R0, R1, R2, R3;
+
+	VectorRegister4Double B[4];
+	B[0] = VectorLoad(BMRows[0]);
+	B[1] = VectorLoad(BMRows[1]);
+	B[2] = VectorLoad(BMRows[2]);
+	B[3] = VectorLoad(BMRows[3]);
+
+	// First row of result (Matrix1[0] * Matrix2).
+	{
+		const VectorRegister4Double ARow = VectorLoad(AMRows[0]);
+		R0 = VectorMultiply(VectorReplicate(ARow, 0), B[0]);
+		R0 = VectorMultiplyAdd(VectorReplicate(ARow, 1), B[1], R0);
+		R0 = VectorMultiplyAdd(VectorReplicate(ARow, 2), B[2], R0);
+		R0 = VectorMultiplyAdd(VectorReplicate(ARow, 3), B[3], R0);
+	}
+
+	// Second row of result (Matrix1[1] * Matrix2).
+	{
+		const VectorRegister4Double ARow = VectorLoad(AMRows[1]);
+		R1 = VectorMultiply(VectorReplicate(ARow, 0), B[0]);
+		R1 = VectorMultiplyAdd(VectorReplicate(ARow, 1), B[1], R1);
+		R1 = VectorMultiplyAdd(VectorReplicate(ARow, 2), B[2], R1);
+		R1 = VectorMultiplyAdd(VectorReplicate(ARow, 3), B[3], R1);
+	}
+
+	// Third row of result (Matrix1[2] * Matrix2).
+	{
+		const VectorRegister4Double ARow = VectorLoad(AMRows[2]);
+		R2 = VectorMultiply(VectorReplicate(ARow, 0), B[0]);
+		R2 = VectorMultiplyAdd(VectorReplicate(ARow, 1), B[1], R2);
+		R2 = VectorMultiplyAdd(VectorReplicate(ARow, 2), B[2], R2);
+		R2 = VectorMultiplyAdd(VectorReplicate(ARow, 3), B[3], R2);
+	}
+
+	// Fourth row of result (Matrix1[3] * Matrix2).
+	{
+		const VectorRegister4Double ARow = VectorLoad(AMRows[3]);
+		R3 = VectorMultiply(VectorReplicate(ARow, 0), B[0]);
+		R3 = VectorMultiplyAdd(VectorReplicate(ARow, 1), B[1], R3);
+		R3 = VectorMultiplyAdd(VectorReplicate(ARow, 2), B[2], R3);
+		R3 = VectorMultiplyAdd(VectorReplicate(ARow, 3), B[3], R3);
+	}
+
+	// Store result. Must not be done during steps above in case source and destination are the same.
+	Double4x4& ResultDst = *((Double4x4*)Result->M);
+	VectorStore(R0, ResultDst[0]);
+	VectorStore(R1, ResultDst[1]);
+	VectorStore(R2, ResultDst[2]);
+	VectorStore(R3, ResultDst[3]);
+
+	// TODO: Should change to DiagnosticCheckNaN call, though that's only in Debug.
+#if !(UE_BUILD_SHIPPING || UE_BUILD_TEST)
+	ensure(!Result->ContainsNaN());
+#endif
+}
+
+
+/**
+ * Calculate the inverse of an FMatrix.
+ *
+ * @param DstMatrix		FMatrix pointer to where the result should be stored
+ * @param SrcMatrix		FMatrix pointer to the Matrix to be inversed
+ */
+void VectorMatrixInverse(FMatrix44f* DstMatrix, const FMatrix44f* SrcMatrix)
+{
+	typedef float Float4x4[4][4];
+	const Float4x4& M = *((const Float4x4*)SrcMatrix);
+	Float4x4 Result;
+	float Det[4];
+	Float4x4 Tmp;
+
+	Tmp[0][0] = M[2][2] * M[3][3] - M[2][3] * M[3][2];
+	Tmp[0][1] = M[1][2] * M[3][3] - M[1][3] * M[3][2];
+	Tmp[0][2] = M[1][2] * M[2][3] - M[1][3] * M[2][2];
+
+	Tmp[1][0] = M[2][2] * M[3][3] - M[2][3] * M[3][2];
+	Tmp[1][1] = M[0][2] * M[3][3] - M[0][3] * M[3][2];
+	Tmp[1][2] = M[0][2] * M[2][3] - M[0][3] * M[2][2];
+
+	Tmp[2][0] = M[1][2] * M[3][3] - M[1][3] * M[3][2];
+	Tmp[2][1] = M[0][2] * M[3][3] - M[0][3] * M[3][2];
+	Tmp[2][2] = M[0][2] * M[1][3] - M[0][3] * M[1][2];
+
+	Tmp[3][0] = M[1][2] * M[2][3] - M[1][3] * M[2][2];
+	Tmp[3][1] = M[0][2] * M[2][3] - M[0][3] * M[2][2];
+	Tmp[3][2] = M[0][2] * M[1][3] - M[0][3] * M[1][2];
+
+	Det[0] = M[1][1] * Tmp[0][0] - M[2][1] * Tmp[0][1] + M[3][1] * Tmp[0][2];
+	Det[1] = M[0][1] * Tmp[1][0] - M[2][1] * Tmp[1][1] + M[3][1] * Tmp[1][2];
+	Det[2] = M[0][1] * Tmp[2][0] - M[1][1] * Tmp[2][1] + M[3][1] * Tmp[2][2];
+	Det[3] = M[0][1] * Tmp[3][0] - M[1][1] * Tmp[3][1] + M[2][1] * Tmp[3][2];
+
+	const float Determinant = M[0][0] * Det[0] - M[1][0] * Det[1] + M[2][0] * Det[2] - M[3][0] * Det[3];
+	const float	RDet = 1.0f / Determinant;
+
+	Result[0][0] = RDet * Det[0];
+	Result[0][1] = -RDet * Det[1];
+	Result[0][2] = RDet * Det[2];
+	Result[0][3] = -RDet * Det[3];
+	Result[1][0] = -RDet * (M[1][0] * Tmp[0][0] - M[2][0] * Tmp[0][1] + M[3][0] * Tmp[0][2]);
+	Result[1][1] = RDet * (M[0][0] * Tmp[1][0] - M[2][0] * Tmp[1][1] + M[3][0] * Tmp[1][2]);
+	Result[1][2] = -RDet * (M[0][0] * Tmp[2][0] - M[1][0] * Tmp[2][1] + M[3][0] * Tmp[2][2]);
+	Result[1][3] = RDet * (M[0][0] * Tmp[3][0] - M[1][0] * Tmp[3][1] + M[2][0] * Tmp[3][2]);
+	Result[2][0] = RDet * (
+		M[1][0] * (M[2][1] * M[3][3] - M[2][3] * M[3][1]) -
+		M[2][0] * (M[1][1] * M[3][3] - M[1][3] * M[3][1]) +
+		M[3][0] * (M[1][1] * M[2][3] - M[1][3] * M[2][1])
+		);
+	Result[2][1] = -RDet * (
+		M[0][0] * (M[2][1] * M[3][3] - M[2][3] * M[3][1]) -
+		M[2][0] * (M[0][1] * M[3][3] - M[0][3] * M[3][1]) +
+		M[3][0] * (M[0][1] * M[2][3] - M[0][3] * M[2][1])
+		);
+	Result[2][2] = RDet * (
+		M[0][0] * (M[1][1] * M[3][3] - M[1][3] * M[3][1]) -
+		M[1][0] * (M[0][1] * M[3][3] - M[0][3] * M[3][1]) +
+		M[3][0] * (M[0][1] * M[1][3] - M[0][3] * M[1][1])
+		);
+	Result[2][3] = -RDet * (
+		M[0][0] * (M[1][1] * M[2][3] - M[1][3] * M[2][1]) -
+		M[1][0] * (M[0][1] * M[2][3] - M[0][3] * M[2][1]) +
+		M[2][0] * (M[0][1] * M[1][3] - M[0][3] * M[1][1])
+		);
+	Result[3][0] = -RDet * (
+		M[1][0] * (M[2][1] * M[3][2] - M[2][2] * M[3][1]) -
+		M[2][0] * (M[1][1] * M[3][2] - M[1][2] * M[3][1]) +
+		M[3][0] * (M[1][1] * M[2][2] - M[1][2] * M[2][1])
+		);
+	Result[3][1] = RDet * (
+		M[0][0] * (M[2][1] * M[3][2] - M[2][2] * M[3][1]) -
+		M[2][0] * (M[0][1] * M[3][2] - M[0][2] * M[3][1]) +
+		M[3][0] * (M[0][1] * M[2][2] - M[0][2] * M[2][1])
+		);
+	Result[3][2] = -RDet * (
+		M[0][0] * (M[1][1] * M[3][2] - M[1][2] * M[3][1]) -
+		M[1][0] * (M[0][1] * M[3][2] - M[0][2] * M[3][1]) +
+		M[3][0] * (M[0][1] * M[1][2] - M[0][2] * M[1][1])
+		);
+	Result[3][3] = RDet * (
+		M[0][0] * (M[1][1] * M[2][2] - M[1][2] * M[2][1]) -
+		M[1][0] * (M[0][1] * M[2][2] - M[0][2] * M[2][1]) +
+		M[2][0] * (M[0][1] * M[1][2] - M[0][2] * M[1][1])
+		);
+
+	memcpy(DstMatrix, &Result, sizeof(Result));
+
+	// TODO: Should change to DiagnosticCheckNaN call, though that's only in Debug.
+#if !(UE_BUILD_SHIPPING || UE_BUILD_TEST)
+	ensure(!DstMatrix->ContainsNaN());
+#endif
+}
+
+void VectorMatrixInverse(FMatrix44d* DstMatrix, const FMatrix44d* SrcMatrix)
+{
+	typedef double Double4x4[4][4];
+	const Double4x4& M = *((const Double4x4*)SrcMatrix->M);
+	Double4x4 Result;
+	double Det[4];
+	Double4x4 Tmp;
+
+	Tmp[0][0] = M[2][2] * M[3][3] - M[2][3] * M[3][2];
+	Tmp[0][1] = M[1][2] * M[3][3] - M[1][3] * M[3][2];
+	Tmp[0][2] = M[1][2] * M[2][3] - M[1][3] * M[2][2];
+
+	Tmp[1][0] = M[2][2] * M[3][3] - M[2][3] * M[3][2];
+	Tmp[1][1] = M[0][2] * M[3][3] - M[0][3] * M[3][2];
+	Tmp[1][2] = M[0][2] * M[2][3] - M[0][3] * M[2][2];
+
+	Tmp[2][0] = M[1][2] * M[3][3] - M[1][3] * M[3][2];
+	Tmp[2][1] = M[0][2] * M[3][3] - M[0][3] * M[3][2];
+	Tmp[2][2] = M[0][2] * M[1][3] - M[0][3] * M[1][2];
+
+	Tmp[3][0] = M[1][2] * M[2][3] - M[1][3] * M[2][2];
+	Tmp[3][1] = M[0][2] * M[2][3] - M[0][3] * M[2][2];
+	Tmp[3][2] = M[0][2] * M[1][3] - M[0][3] * M[1][2];
+
+	Det[0] = M[1][1] * Tmp[0][0] - M[2][1] * Tmp[0][1] + M[3][1] * Tmp[0][2];
+	Det[1] = M[0][1] * Tmp[1][0] - M[2][1] * Tmp[1][1] + M[3][1] * Tmp[1][2];
+	Det[2] = M[0][1] * Tmp[2][0] - M[1][1] * Tmp[2][1] + M[3][1] * Tmp[2][2];
+	Det[3] = M[0][1] * Tmp[3][0] - M[1][1] * Tmp[3][1] + M[2][1] * Tmp[3][2];
+
+	const double Determinant = M[0][0] * Det[0] - M[1][0] * Det[1] + M[2][0] * Det[2] - M[3][0] * Det[3];
+	const double	RDet = 1.0 / Determinant;
+
+	Result[0][0] = RDet * Det[0];
+	Result[0][1] = -RDet * Det[1];
+	Result[0][2] = RDet * Det[2];
+	Result[0][3] = -RDet * Det[3];
+	Result[1][0] = -RDet * (M[1][0] * Tmp[0][0] - M[2][0] * Tmp[0][1] + M[3][0] * Tmp[0][2]);
+	Result[1][1] = RDet * (M[0][0] * Tmp[1][0] - M[2][0] * Tmp[1][1] + M[3][0] * Tmp[1][2]);
+	Result[1][2] = -RDet * (M[0][0] * Tmp[2][0] - M[1][0] * Tmp[2][1] + M[3][0] * Tmp[2][2]);
+	Result[1][3] = RDet * (M[0][0] * Tmp[3][0] - M[1][0] * Tmp[3][1] + M[2][0] * Tmp[3][2]);
+	Result[2][0] = RDet * (
+		M[1][0] * (M[2][1] * M[3][3] - M[2][3] * M[3][1]) -
+		M[2][0] * (M[1][1] * M[3][3] - M[1][3] * M[3][1]) +
+		M[3][0] * (M[1][1] * M[2][3] - M[1][3] * M[2][1])
+		);
+	Result[2][1] = -RDet * (
+		M[0][0] * (M[2][1] * M[3][3] - M[2][3] * M[3][1]) -
+		M[2][0] * (M[0][1] * M[3][3] - M[0][3] * M[3][1]) +
+		M[3][0] * (M[0][1] * M[2][3] - M[0][3] * M[2][1])
+		);
+	Result[2][2] = RDet * (
+		M[0][0] * (M[1][1] * M[3][3] - M[1][3] * M[3][1]) -
+		M[1][0] * (M[0][1] * M[3][3] - M[0][3] * M[3][1]) +
+		M[3][0] * (M[0][1] * M[1][3] - M[0][3] * M[1][1])
+		);
+	Result[2][3] = -RDet * (
+		M[0][0] * (M[1][1] * M[2][3] - M[1][3] * M[2][1]) -
+		M[1][0] * (M[0][1] * M[2][3] - M[0][3] * M[2][1]) +
+		M[2][0] * (M[0][1] * M[1][3] - M[0][3] * M[1][1])
+		);
+	Result[3][0] = -RDet * (
+		M[1][0] * (M[2][1] * M[3][2] - M[2][2] * M[3][1]) -
+		M[2][0] * (M[1][1] * M[3][2] - M[1][2] * M[3][1]) +
+		M[3][0] * (M[1][1] * M[2][2] - M[1][2] * M[2][1])
+		);
+	Result[3][1] = RDet * (
+		M[0][0] * (M[2][1] * M[3][2] - M[2][2] * M[3][1]) -
+		M[2][0] * (M[0][1] * M[3][2] - M[0][2] * M[3][1]) +
+		M[3][0] * (M[0][1] * M[2][2] - M[0][2] * M[2][1])
+		);
+	Result[3][2] = -RDet * (
+		M[0][0] * (M[1][1] * M[3][2] - M[1][2] * M[3][1]) -
+		M[1][0] * (M[0][1] * M[3][2] - M[0][2] * M[3][1]) +
+		M[3][0] * (M[0][1] * M[1][2] - M[0][2] * M[1][1])
+		);
+	Result[3][3] = RDet * (
+		M[0][0] * (M[1][1] * M[2][2] - M[1][2] * M[2][1]) -
+		M[1][0] * (M[0][1] * M[2][2] - M[0][2] * M[2][1]) +
+		M[2][0] * (M[0][1] * M[1][2] - M[0][2] * M[1][1])
+		);
+
+	memcpy(DstMatrix, &Result, sizeof(Result));
+
+	// TODO: Should change to DiagnosticCheckNaN call, though that's only in Debug.
+#if !(UE_BUILD_SHIPPING || UE_BUILD_TEST)
+	ensure(!DstMatrix->ContainsNaN());
+#endif
+}
+
+#endif // UE_USING_UNREALMATH_SSE
diff --git a/Engine/Source/Runtime/Core/Private/Tests/Math/UnrealMathTest.cpp b/Engine/Source/Runtime/Core/Private/Tests/Math/UnrealMathTest.cpp
--- a/Engine/Source/Runtime/Core/Private/Tests/Math/UnrealMathTest.cpp
+++ b/Engine/Source/Runtime/Core/Private/Tests/Math/UnrealMathTest.cpp
@@ -55,7 +55,7 @@ FORCENOINLINE void ResetPassing()
  *
  * @return true if equal
  */
-bool TestVectorsEqualBitwise( VectorRegister4Float Vec0, VectorRegister4Float Vec1)
+FORCENOINLINE bool TestVectorsEqualBitwise( const VectorRegister4Float& Vec0, const VectorRegister4Float& Vec1)
 {
 	VectorStoreAligned( Vec0, GScratch + 0 );
 	VectorStoreAligned( Vec1, GScratch + 4 );
@@ -67,7 +67,7 @@ bool TestVectorsEqualBitwise( VectorRegister4Float Vec0, VectorRegister4Float Ve
 	return Passed;
 }
 
-bool TestVectorsEqualBitwise(VectorRegister4Double Vec0, VectorRegister4Double Vec1)
+FORCENOINLINE bool TestVectorsEqualBitwise(const VectorRegister4Double& Vec0, const VectorRegister4Double& Vec1)
 {
 	VectorStoreAligned(Vec0, GScratchDouble + 0);
 	VectorStoreAligned(Vec1, GScratchDouble + 4);
@@ -89,7 +89,7 @@ bool TestVectorsEqualBitwise(VectorRegister4Double Vec0, VectorRegister4Double V
  *
  * @return true if equal(ish)
  */
-bool TestVectorsEqual( VectorRegister4Float Vec0, VectorRegister4Float Vec1, float Tolerance = 0.0f)
+FORCENOINLINE bool TestVectorsEqual(const VectorRegister4Float& Vec0, const VectorRegister4Float& Vec1, float Tolerance = 0.0f)
 {
 	VectorStoreAligned( Vec0, GScratch + 0 );
 	VectorStoreAligned( Vec1, GScratch + 4 );
@@ -103,7 +103,7 @@ bool TestVectorsEqual( VectorRegister4Float Vec0, VectorRegister4Float Vec1, flo
 	return GSum <= Tolerance;
 }
 
-bool TestVectorsEqual(VectorRegister4Double Vec0, VectorRegister4Double Vec1, double Tolerance = 0.0)
+FORCENOINLINE bool TestVectorsEqual(const VectorRegister4Double& Vec0, const VectorRegister4Double& Vec1, double Tolerance = 0.0)
 {
 	VectorStoreAligned(Vec0, GScratchDouble + 0);
 	VectorStoreAligned(Vec1, GScratchDouble + 4);
@@ -120,7 +120,7 @@ bool TestVectorsEqual(VectorRegister4Double Vec0, VectorRegister4Double Vec1, do
 /**
  * Enforce tolerance per-component, not summed error
  */
-bool TestVectorsEqual_ComponentWiseError(VectorRegister4Float Vec0, VectorRegister4Float Vec1, float Tolerance = 0.0f)
+FORCENOINLINE bool TestVectorsEqual_ComponentWiseError(const VectorRegister4Float& Vec0, const VectorRegister4Float& Vec1, float Tolerance = 0.0f)
 {
 	VectorStoreAligned(Vec0, GScratch + 0);
 	VectorStoreAligned(Vec1, GScratch + 4);
@@ -137,7 +137,7 @@ bool TestVectorsEqual_ComponentWiseError(VectorRegister4Float Vec0, VectorRegist
 	return bPassing;
 }
 
-bool TestVectorsEqual_ComponentWiseError(VectorRegister4Double Vec0, VectorRegister4Double Vec1, double Tolerance = 0.0f)
+FORCENOINLINE bool TestVectorsEqual_ComponentWiseError(const VectorRegister4Double& Vec0, const VectorRegister4Double& Vec1, double Tolerance = 0.0f)
 {
 	VectorStoreAligned(Vec0, GScratchDouble + 0);
 	VectorStoreAligned(Vec1, GScratchDouble + 4);
@@ -164,7 +164,7 @@ bool TestVectorsEqual_ComponentWiseError(VectorRegister4Double Vec0, VectorRegis
  *
  * @return true if equal(ish)
  */
-bool TestVectorsEqual3( VectorRegister4Float Vec0, VectorRegister4Float Vec1, float Tolerance = 0.0f)
+FORCENOINLINE bool TestVectorsEqual3(const VectorRegister4Float& Vec0, const VectorRegister4Float& Vec1, float Tolerance = 0.0f)
 {
 	VectorStoreAligned( Vec0, GScratch + 0 );
 	VectorStoreAligned( Vec1, GScratch + 4 );
@@ -177,7 +177,7 @@ bool TestVectorsEqual3( VectorRegister4Float Vec0, VectorRegister4Float Vec1, fl
 	return GSum <= Tolerance;
 }
 
-bool TestVectorsEqual3(VectorRegister4Double Vec0, VectorRegister4Double Vec1, double Tolerance = 0.0)
+FORCENOINLINE bool TestVectorsEqual3(const VectorRegister4Double& Vec0, const VectorRegister4Double& Vec1, double Tolerance = 0.0)
 {
 	VectorStoreAligned(Vec0, GScratchDouble + 0);
 	VectorStoreAligned(Vec1, GScratchDouble + 4);
@@ -199,7 +199,7 @@ bool TestVectorsEqual3(VectorRegister4Double Vec0, VectorRegister4Double Vec1, d
  *
  * @return true if equal(ish)
  */
-bool TestFVector3Equal( const FVector3f& Vec0, const FVector3f& Vec1, float Tolerance = 0.0f)
+FORCENOINLINE bool TestFVector3Equal( const FVector3f& Vec0, const FVector3f& Vec1, float Tolerance = 0.0f)
 {
 	GScratch[0] = Vec0.X;
 	GScratch[1] = Vec0.Y;
@@ -219,7 +219,7 @@ bool TestFVector3Equal( const FVector3f& Vec0, const FVector3f& Vec1, float Tole
 	return GSum <= Tolerance;
 }
 
-bool TestQuatsEqual(const FQuat4f& Q0, const FQuat4f& Q1, float Tolerance)
+FORCENOINLINE bool TestQuatsEqual(const FQuat4f& Q0, const FQuat4f& Q1, float Tolerance)
 {
 	GScratch[0] = Q0.X;
 	GScratch[1] = Q0.Y;
@@ -236,7 +236,7 @@ bool TestQuatsEqual(const FQuat4f& Q0, const FQuat4f& Q1, float Tolerance)
 	return bEqual;
 }
 
-bool TestQuatsEqual(const FQuat4d& Q0, const FQuat4d& Q1, double Tolerance)
+FORCENOINLINE bool TestQuatsEqual(const FQuat4d& Q0, const FQuat4d& Q1, double Tolerance)
 {
 	GScratchDouble[0] = Q0.X;
 	GScratchDouble[1] = Q0.Y;
@@ -261,7 +261,7 @@ bool TestQuatsEqual(const FQuat4d& Q0, const FQuat4d& Q1, double Tolerance)
  *
  * @return true if normalized(ish)
  */
-bool TestFVector3Normalized(const FVector3f& Vec0, float Tolerance)
+FORCENOINLINE bool TestFVector3Normalized(const FVector3f& Vec0, float Tolerance)
 {
 	GScratch[0] = Vec0.X;
 	GScratch[1] = Vec0.Y;
@@ -286,7 +286,7 @@ bool TestFVector3Normalized(const FVector3f& Vec0, float Tolerance)
  *
  * @return true if normalized(ish)
  */
-bool TestQuatNormalized(const FQuat4f& Q0, float Tolerance)
+FORCENOINLINE bool TestQuatNormalized(const FQuat4f& Q0, float Tolerance)
 {
 	GScratch[0] = Q0.X;
 	GScratch[1] = Q0.Y;
@@ -312,7 +312,7 @@ bool TestQuatNormalized(const FQuat4f& Q0, float Tolerance)
  *
  * @return true if equal(ish)
  */
-bool TestMatricesEqual( FMatrix44f &Mat0, FMatrix44f &Mat1, float Tolerance = 0.0f)
+FORCENOINLINE bool TestMatricesEqual( const FMatrix44f &Mat0, const FMatrix44f &Mat1, float Tolerance = 0.0f)
 {
 	for (int32 Row = 0; Row < 4; ++Row ) 
 	{
@@ -331,7 +331,7 @@ bool TestMatricesEqual( FMatrix44f &Mat0, FMatrix44f &Mat1, float Tolerance = 0.
 	return true;
 }
 
-bool TestMatricesEqual(FMatrix44d& Mat0, FMatrix44d& Mat1, double Tolerance = 0.0f)
+FORCENOINLINE bool TestMatricesEqual(const FMatrix44d& Mat0, const FMatrix44d& Mat1, double Tolerance = 0.0)
 {
 	for (int32 Row = 0; Row < 4; ++Row)
 	{
@@ -612,20 +612,21 @@ MATHTEST_INLINE FVector3f TestQuaternionRotateVectorRegister(const FQuat4f& Quat
 * @param Quat1	Pointer to the first quaternion (must not be the destination)
 * @param Quat2	Pointer to the second quaternion (must not be the destination)
 */
-void TestVectorQuaternionMultiply( FQuat4f* Result, const FQuat4f* Quat1, const FQuat4f* Quat2)
+template<typename FloatType>
+void TestVectorQuaternionMultiply( UE::Math::TQuat<FloatType>* Result, const UE::Math::TQuat<FloatType>* Quat1, const UE::Math::TQuat<FloatType>* Quat2)
 {
-	typedef float Float4[4];
+	typedef FloatType Float4[4];
 	const Float4& A = *((const Float4*) Quat1);
 	const Float4& B = *((const Float4*) Quat2);
 	Float4 & R = *((Float4*) Result);
 
 	// store intermediate results in temporaries
-	const float TX = A[3]*B[0] + A[0]*B[3] + A[1]*B[2] - A[2]*B[1];
-	const float TY = A[3]*B[1] - A[0]*B[2] + A[1]*B[3] + A[2]*B[0];
-	const float TZ = A[3]*B[2] + A[0]*B[1] - A[1]*B[0] + A[2]*B[3];
-	const float TW = A[3]*B[3] - A[0]*B[0] - A[1]*B[1] - A[2]*B[2];
+	const FloatType TX = A[3]*B[0] + A[0]*B[3] + A[1]*B[2] - A[2]*B[1];
+	const FloatType TY = A[3]*B[1] - A[0]*B[2] + A[1]*B[3] + A[2]*B[0];
+	const FloatType TZ = A[3]*B[2] + A[0]*B[1] - A[1]*B[0] + A[2]*B[3];
+	const FloatType TW = A[3]*B[3] - A[0]*B[0] - A[1]*B[1] - A[2]*B[2];
 
-	// copy intermediate result to *this
+	// copy intermediate result to R
 	R[0] = TX;
 	R[1] = TY;
 	R[2] = TZ;
@@ -1951,53 +1952,55 @@ bool RunDoubleVectorTest()
 	MaskBits = VectorMaskBits(V0);
 	LogTest<double>(TEXT("VectorMaskBits"), MaskBits == 5);
 
-
-	FMatrix44d	M0, M1, M2, M3;
-	FVector3d Eye, LookAt, Up;
-	// Create Look at Matrix
-	Eye = FVector3d(1024.0f, -512.0f, -2048.0f);
-	LookAt = FVector3d(0.0f, 0.0f, 0.0f);
-	Up = FVector3d(0.0f, 1.0f, 0.0f);
-	M0 = FLookAtMatrix(Eye, LookAt, Up);
-
-	// Create GL ortho projection matrix
-	const double Width = 1920.0f;
-	const double Height = 1080.0f;
-	const double Left = 0.0f;
-	const double Right = Left + Width;
-	const double Top = 0.0f;
-	const double Bottom = Top + Height;
-	const double ZNear = -100.0f;
-	const double ZFar = 100.0f;
-
-	M1 = FMatrix44d(FPlane4d(2.0f / (Right - Left), 0, 0, 0),
-		FPlane4d(0, 2.0f / (Top - Bottom), 0, 0),
-		FPlane4d(0, 0, 1 / (ZNear - ZFar), 0),
-		FPlane4d((Left + Right) / (Left - Right), (Top + Bottom) / (Bottom - Top), ZNear / (ZNear - ZFar), 1));
-
-	VectorMatrixMultiply(&M2, &M0, &M1);
-	TestVectorMatrixMultiply(&M3, &M0, &M1);
-	LogTest<double>(TEXT("VectorMatrixMultiply"), TestMatricesEqual(M2, M3, 0.000001f));
-
-	VectorMatrixInverse(&M2, &M1);
-	TestVectorMatrixInverse(&M3, &M1);
-	LogTest<double>(TEXT("VectorMatrixInverse"), TestMatricesEqual(M2, M3, 0.000001f));
-	
-	// 	FTransform Transform;
-	// 	Transform.SetFromMatrix(M1);
-	// 	FTransform InvTransform = Transform.Inverse();
-	// 	FTransform InvTransform2 = FTransform(Transform.ToMatrixWithScale().Inverse());
-	// 	LogTest<double>( TEXT("FTransform Inverse"), InvTransform.Equals(InvTransform2, 1e-3f ) );
-
-	V0 = MakeVectorRegister(100.0f, -100.0f, 200.0f, 1.0f);
-	V1 = VectorTransformVector(V0, &M0);
-	V2 = TestVectorTransformVector(V0, &M0);
-	LogTest<double>(TEXT("VectorTransformVector"), TestVectorsEqual(V1, V2, 1e-8f));
-
-	V0 = MakeVectorRegister(32768.0f, 131072.0f, -8096.0f, 1.0f);
-	V1 = VectorTransformVector(V0, &M1);
-	V2 = TestVectorTransformVector(V0, &M1);
-	LogTest<double>(TEXT("VectorTransformVector"), TestVectorsEqual(V1, V2, 1e-8f));
+	// Matrix multiplications and transformations
+	{
+		FMatrix44d	M0, M1, M2, M3;
+		FVector3d Eye, LookAt, Up;
+		// Create Look at Matrix
+		Eye = FVector3d(1024.0f, -512.0f, -2048.0f);
+		LookAt = FVector3d(0.0f, 0.0f, 0.0f);
+		Up = FVector3d(0.0f, 1.0f, 0.0f);
+		M0 = FLookAtMatrix(Eye, LookAt, Up);
+
+		// Create GL ortho projection matrix
+		const double Width = 1920.0f;
+		const double Height = 1080.0f;
+		const double Left = 0.0f;
+		const double Right = Left + Width;
+		const double Top = 0.0f;
+		const double Bottom = Top + Height;
+		const double ZNear = -100.0f;
+		const double ZFar = 100.0f;
+
+		M1 = FMatrix44d(FPlane4d(2.0f / (Right - Left), 0, 0, 0),
+			FPlane4d(0, 2.0f / (Top - Bottom), 0, 0),
+			FPlane4d(0, 0, 1 / (ZNear - ZFar), 0),
+			FPlane4d((Left + Right) / (Left - Right), (Top + Bottom) / (Bottom - Top), ZNear / (ZNear - ZFar), 1));
+
+		VectorMatrixMultiply(&M2, &M0, &M1);
+		TestVectorMatrixMultiply(&M3, &M0, &M1);
+		LogTest<double>(TEXT("VectorMatrixMultiply"), TestMatricesEqual(M2, M3, 0.000001));
+
+		VectorMatrixInverse(&M2, &M1);
+		TestVectorMatrixInverse(&M3, &M1);
+		LogTest<double>(TEXT("VectorMatrixInverse"), TestMatricesEqual(M2, M3, 0.000001));
+
+		// 	FTransform Transform;
+		// 	Transform.SetFromMatrix(M1);
+		// 	FTransform InvTransform = Transform.Inverse();
+		// 	FTransform InvTransform2 = FTransform(Transform.ToMatrixWithScale().Inverse());
+		// 	LogTest<double>( TEXT("FTransform Inverse"), InvTransform.Equals(InvTransform2, 1e-3f ) );
+
+		V0 = MakeVectorRegister(100.0f, -100.0f, 200.0f, 1.0f);
+		V1 = VectorTransformVector(V0, &M0);
+		V2 = TestVectorTransformVector(V0, &M0);
+		LogTest<double>(TEXT("VectorTransformVector"), TestVectorsEqual(V1, V2, 1e-8));
+
+		V0 = MakeVectorRegister(32768.0f, 131072.0f, -8096.0f, 1.0f);
+		V1 = VectorTransformVector(V0, &M1);
+		V2 = TestVectorTransformVector(V0, &M1);
+		LogTest<double>(TEXT("VectorTransformVector"), TestVectorsEqual(V1, V2, 1e-8));
+	}
 
 
 	// Quat / Rotator conversion to vectors, matrices
@@ -2063,6 +2066,49 @@ bool RunDoubleVectorTest()
 	// Exp, Log tests
 	TestVectorExpLogFunctions<double, VectorRegister4Double>();
 
+	// Quat multiplication, Matrix conversion
+	{
+		FQuat4d Q0, Q1, Q2, Q3;
+		FMatrix44d M0, M1;
+		FMatrix44d IdentityInverse = FMatrix44d::Identity.Inverse();
+
+		Q0 = FQuat4d(FRotator3d(30.0f, -45.0f, 90.0f));
+		Q1 = FQuat4d(FRotator3d(45.0f, 60.0f, 120.0f));
+		VectorQuaternionMultiply(&Q2, &Q0, &Q1);
+		TestVectorQuaternionMultiply(&Q3, &Q0, &Q1);
+		LogTest<double>(TEXT("VectorQuaternionMultiply"), TestQuatsEqual(Q2, Q3, 1e-6));
+		V0 = VectorLoadAligned(&Q0);
+		V1 = VectorLoadAligned(&Q1);
+		V2 = VectorQuaternionMultiply2(V0, V1);
+		V3 = VectorLoadAligned(&Q3);
+		LogTest<double>(TEXT("VectorQuaternionMultiply2"), TestVectorsEqual(V2, V3, 1e-6));
+
+		M0 = Q0 * FMatrix44d::Identity;
+		M1 = UE::Math::TQuatRotationMatrix<double>(Q0);
+		LogTest<double>(TEXT("QuaterionMatrixConversion Q0"), TestMatricesEqual(M0, M1, 0.000001));
+		M0 = Q1.ToMatrix();
+		M1 = UE::Math::TQuatRotationMatrix<double>(Q1);
+		LogTest<double>(TEXT("QuaterionMatrixConversion Q1"), TestMatricesEqual(M0, M1, 0.000001));
+
+		Q0 = FQuat4d(FRotator3d(0.0f, 180.0f, 45.0f));
+		Q1 = FQuat4d(FRotator3d(-120.0f, -90.0f, 0.0f));
+		VectorQuaternionMultiply(&Q2, &Q0, &Q1);
+		TestVectorQuaternionMultiply(&Q3, &Q0, &Q1);
+		LogTest<double>(TEXT("VectorQuaternionMultiply"), TestQuatsEqual(Q2, Q3, 1e-6));
+		V0 = VectorLoadAligned(&Q0);
+		V1 = VectorLoadAligned(&Q1);
+		V2 = VectorQuaternionMultiply2(V0, V1);
+		V3 = VectorLoadAligned(&Q3);
+		LogTest<double>(TEXT("VectorQuaternionMultiply2"), TestVectorsEqual(V2, V3, 1e-6));
+
+		M0 = (-Q0) * FMatrix44d::Identity;
+		M1 = (-Q0).ToMatrix();
+		LogTest<double>(TEXT("QuaterionMatrixConversion Q0"), TestMatricesEqual(M0, M1, 0.000001));
+		M0 = (Q1.Inverse().Inverse()) * IdentityInverse;
+		M1 = Q1.ToMatrix();
+		LogTest<double>(TEXT("QuaterionMatrixConversion Q1"), TestMatricesEqual(M0, M1, 0.000001));
+	}
+
 	return GPassing;
 }
 
@@ -2641,53 +2687,55 @@ bool FVectorRegisterAbstractionTest::RunTest(const FString& Parameters)
 	MaskBits = VectorMaskBits(V0);
 	LogTest<float>(TEXT("VectorMaskBits"), MaskBits == 5);
 
+	// Matrix multiplications and transformations
+	{
+		FMatrix44f	M0, M1, M2, M3;
+		FVector3f Eye, LookAt, Up;	
+		// Create Look at Matrix
+		Eye    = FVector3f(1024.0f, -512.0f, -2048.0f);
+		LookAt = FVector3f(0.0f,		  0.0f,     0.0f);
+		Up     = FVector3f(0.0f,       1.0f,    0.0f);
+		M0	= FLookAtMatrix(Eye, LookAt, Up);		
+
+		// Create GL ortho projection matrix
+		const float Width = 1920.0f;
+		const float Height = 1080.0f;
+		const float Left = 0.0f;
+		const float Right = Left+Width;
+		const float Top = 0.0f;
+		const float Bottom = Top+Height;
+		const float ZNear = -100.0f;
+		const float ZFar = 100.0f;
+
+		M1 = FMatrix44f(FPlane4f(2.0f/(Right-Left),	0,							0,					0 ),
+			FPlane4f(0,							2.0f/(Top-Bottom),			0,					0 ),
+			FPlane4f(0,							0,							1/(ZNear-ZFar),		0 ),
+			FPlane4f((Left+Right)/(Left-Right),	(Top+Bottom)/(Bottom-Top),	ZNear/(ZNear-ZFar), 1 ) );
+
+		VectorMatrixMultiply( &M2, &M0, &M1 );
+		TestVectorMatrixMultiply( &M3, &M0, &M1 );
+		LogTest<float>( TEXT("VectorMatrixMultiply"), TestMatricesEqual( M2, M3, 0.000001f ) );
+
+		VectorMatrixInverse( &M2, &M1 );
+		TestVectorMatrixInverse( &M3, &M1 );
+		LogTest<float>( TEXT("VectorMatrixInverse"), TestMatricesEqual( M2, M3, 0.000001f) );
 
-	FMatrix44f	M0, M1, M2, M3;
-	FVector3f Eye, LookAt, Up;	
-	// Create Look at Matrix
-	Eye    = FVector3f(1024.0f, -512.0f, -2048.0f);
-	LookAt = FVector3f(0.0f,		  0.0f,     0.0f);
-	Up     = FVector3f(0.0f,       1.0f,    0.0f);
-	M0	= FLookAtMatrix(Eye, LookAt, Up);		
-
-	// Create GL ortho projection matrix
-	const float Width = 1920.0f;
-	const float Height = 1080.0f;
-	const float Left = 0.0f;
-	const float Right = Left+Width;
-	const float Top = 0.0f;
-	const float Bottom = Top+Height;
-	const float ZNear = -100.0f;
-	const float ZFar = 100.0f;
-
-	M1 = FMatrix44f(FPlane4f(2.0f/(Right-Left),	0,							0,					0 ),
-		FPlane4f(0,							2.0f/(Top-Bottom),			0,					0 ),
-		FPlane4f(0,							0,							1/(ZNear-ZFar),		0 ),
-		FPlane4f((Left+Right)/(Left-Right),	(Top+Bottom)/(Bottom-Top),	ZNear/(ZNear-ZFar), 1 ) );
-
-	VectorMatrixMultiply( &M2, &M0, &M1 );
-	TestVectorMatrixMultiply( &M3, &M0, &M1 );
-	LogTest<float>( TEXT("VectorMatrixMultiply"), TestMatricesEqual( M2, M3, 0.000001f ) );
-
-	VectorMatrixInverse( &M2, &M1 );
-	TestVectorMatrixInverse( &M3, &M1 );
-	LogTest<float>( TEXT("VectorMatrixInverse"), TestMatricesEqual( M2, M3, 0.000001f) );
-
-// 	FTransform Transform;
-// 	Transform.SetFromMatrix(M1);
-// 	FTransform InvTransform = Transform.Inverse();
-// 	FTransform InvTransform2 = FTransform(Transform.ToMatrixWithScale().Inverse());
-// 	LogTest<float>( TEXT("FTransform Inverse"), InvTransform.Equals(InvTransform2, 1e-3f ) );
-
-	V0 = MakeVectorRegister( 100.0f, -100.0f, 200.0f, 1.0f );
-	V1 = VectorTransformVector(V0, &M0);
-	V2 = TestVectorTransformVector(V0, &M0);
-	LogTest<float>( TEXT("VectorTransformVector"), TestVectorsEqual( V1, V2, 1e-8f ) );
-
-	V0 = MakeVectorRegister( 32768.0f,131072.0f, -8096.0f, 1.0f );
-	V1 = VectorTransformVector(V0, &M1);
-	V2 = TestVectorTransformVector(V0, &M1);
-	LogTest<float>( TEXT("VectorTransformVector"), TestVectorsEqual( V1, V2, 1e-8f ) );
+	// 	FTransform Transform;
+	// 	Transform.SetFromMatrix(M1);
+	// 	FTransform InvTransform = Transform.Inverse();
+	// 	FTransform InvTransform2 = FTransform(Transform.ToMatrixWithScale().Inverse());
+	// 	LogTest<float>( TEXT("FTransform Inverse"), InvTransform.Equals(InvTransform2, 1e-3f ) );
+
+		V0 = MakeVectorRegister( 100.0f, -100.0f, 200.0f, 1.0f );
+		V1 = VectorTransformVector(V0, &M0);
+		V2 = TestVectorTransformVector(V0, &M0);
+		LogTest<float>( TEXT("VectorTransformVector"), TestVectorsEqual( V1, V2, 1e-8f ) );
+
+		V0 = MakeVectorRegister( 32768.0f,131072.0f, -8096.0f, 1.0f );
+		V1 = VectorTransformVector(V0, &M1);
+		V2 = TestVectorTransformVector(V0, &M1);
+		LogTest<float>( TEXT("VectorTransformVector"), TestVectorsEqual( V1, V2, 1e-8f ) );
+	}
 
 	// NaN / Inf tests
 	SetScratch(0.0f, 0.0f, 0.0f, 0.0f);
@@ -3315,8 +3363,11 @@ bool FVectorRegisterAbstractionTest::RunTest(const FString& Parameters)
 		}
 	}
 
-	// Quat multiplication
+	// Quat multiplication, Matrix conversion
 	{
+		FMatrix44f M0, M1;
+		FMatrix44f IdentityInverse = FMatrix44f::Identity.Inverse();
+
 		Q0 = FQuat4f(FRotator3f(30.0f, -45.0f, 90.0f));
 		Q1 = FQuat4f(FRotator3f(45.0f, 60.0f, 120.0f));
 		VectorQuaternionMultiply(&Q2, &Q0, &Q1);
@@ -3328,16 +3379,30 @@ bool FVectorRegisterAbstractionTest::RunTest(const FString& Parameters)
 		V3 = VectorLoadAligned(&Q3);
 		LogTest<float>(TEXT("VectorQuaternionMultiply2"), TestVectorsEqual(V2, V3, 1e-6f));
 
+		M0 = Q0 * FMatrix44f::Identity;
+		M1 = UE::Math::TQuatRotationMatrix<float>(Q0);
+		LogTest<float>(TEXT("QuaterionMatrixConversion Q0"), TestMatricesEqual(M0, M1, 0.000001f));
+		M0 = Q1.ToMatrix();
+		M1 = UE::Math::TQuatRotationMatrix<float>(Q1);
+		LogTest<float>(TEXT("QuaterionMatrixConversion Q1"), TestMatricesEqual(M0, M1, 0.000001f));
+
 		Q0 = FQuat4f(FRotator3f(0.0f, 180.0f, 45.0f));
 		Q1 = FQuat4f(FRotator3f(-120.0f, -90.0f, 0.0f));
 		VectorQuaternionMultiply(&Q2, &Q0, &Q1);
 		TestVectorQuaternionMultiply(&Q3, &Q0, &Q1);
-		LogTest<float>(TEXT("VectorMatrixInverse"), TestQuatsEqual(Q2, Q3, 1e-6f));
+		LogTest<float>(TEXT("VectorQuaternionMultiply"), TestQuatsEqual(Q2, Q3, 1e-6f));
 		V0 = VectorLoadAligned(&Q0);
 		V1 = VectorLoadAligned(&Q1);
 		V2 = VectorQuaternionMultiply2(V0, V1);
 		V3 = VectorLoadAligned(&Q3);
 		LogTest<float>(TEXT("VectorQuaternionMultiply2"), TestVectorsEqual(V2, V3, 1e-6f));
+
+		M0 = (-Q0) * FMatrix44f::Identity;
+		M1 = (-Q0).ToMatrix();
+		LogTest<float>(TEXT("QuaterionMatrixConversion Q0"), TestMatricesEqual(M0, M1, 0.000001f));
+		M0 = (Q1.Inverse().Inverse()) * IdentityInverse;
+		M1 = Q1.ToMatrix();
+		LogTest<float>(TEXT("QuaterionMatrixConversion Q1"), TestMatricesEqual(M0, M1, 0.000001f));
 	}
 
 	if (!GPassing)
diff --git a/Engine/Source/Runtime/Core/Public/Math/Matrix.h b/Engine/Source/Runtime/Core/Public/Math/Matrix.h
--- a/Engine/Source/Runtime/Core/Public/Math/Matrix.h
+++ b/Engine/Source/Runtime/Core/Public/Math/Matrix.h
@@ -35,7 +35,7 @@ namespace Math
 {
 
 template<typename T>
-struct TMatrix
+struct alignas(16) TMatrix
 {
 	static_assert(std::is_floating_point_v<T>, "T must be floating point");
 
diff --git a/Engine/Source/Runtime/Core/Public/Math/Matrix.inl b/Engine/Source/Runtime/Core/Public/Math/Matrix.inl
--- a/Engine/Source/Runtime/Core/Public/Math/Matrix.inl
+++ b/Engine/Source/Runtime/Core/Public/Math/Matrix.inl
@@ -84,10 +84,10 @@ FORCEINLINE TMatrix<T>	TMatrix<T>::operator+(const TMatrix<T>& Other) const
 
 	for (int32 X = 0; X < 4; X++)
 	{
-		for (int32 Y = 0; Y < 4; Y++)
-		{
-			ResultMat.M[X][Y] = M[X][Y] + Other.M[X][Y];
-		}
+		ResultMat.M[X][0] = M[X][0] + Other.M[X][0];
+		ResultMat.M[X][1] = M[X][1] + Other.M[X][1];
+		ResultMat.M[X][2] = M[X][2] + Other.M[X][2];
+		ResultMat.M[X][3] = M[X][3] + Other.M[X][3];
 	}
 
 	ResultMat.DiagnosticCheckNaN();
@@ -108,10 +108,10 @@ FORCEINLINE TMatrix<T> TMatrix<T>::operator*(T Other) const
 
 	for (int32 X = 0; X < 4; X++)
 	{
-		for (int32 Y = 0; Y < 4; Y++)
-		{
-			ResultMat.M[X][Y] = M[X][Y] * Other;
-		}
+		ResultMat.M[X][0] = M[X][0] * Other;
+		ResultMat.M[X][1] = M[X][1] * Other;
+		ResultMat.M[X][2] = M[X][2] * Other;
+		ResultMat.M[X][3] = M[X][3] * Other;
 	}
 
 	ResultMat.DiagnosticCheckNaN();
diff --git a/Engine/Source/Runtime/Core/Public/Math/Quat.h b/Engine/Source/Runtime/Core/Public/Math/Quat.h
--- a/Engine/Source/Runtime/Core/Public/Math/Quat.h
+++ b/Engine/Source/Runtime/Core/Public/Math/Quat.h
@@ -240,7 +240,7 @@ public:
 	 * @param M Matrix to multiply by.
 	 * @return Matrix result after multiplication.
 	 */	
-	TMatrix<T> operator*(const TMatrix<T>& M) const;
+	CORE_API TMatrix<T> operator*(const TMatrix<T>& M) const;
 	
 	/**
 	 * Multiply this quaternion by a scaling factor.
@@ -503,6 +503,12 @@ public:
 	/** Get the TRotator<T> representation of this Quaternion. */
 	CORE_API TRotator<T> Rotator() const;
 
+	/** Get the TMatrix<T> representation of this Quaternion. */
+	FORCEINLINE TMatrix<T> ToMatrix() const;
+
+	/** Get the TMatrix<T> representation of this Quaternion and store it in Mat */
+	CORE_API void ToMatrix(TMatrix<T>& Mat) const;
+
 	/**
 	 * Get the axis of rotation of the Quaternion.
 	 * This is the axis around which rotation occurs to transform the canonical coordinate system to the target orientation.
@@ -805,27 +811,6 @@ FORCEINLINE TVector<T> TQuat<T>::operator*(const TVector<T>& V) const
 	return RotateVector(V);
 }
 
-template<typename T>
-inline UE::Math::TMatrix<T> TQuat<T>::operator*(const UE::Math::TMatrix<T>& M) const
-{
-	UE::Math::TMatrix<T> Result;
-	TQuat<T> VT, VR;
-	TQuat<T> Inv = Inverse();
-	for (int32 I=0; I<4; ++I)
-	{
-		TQuat<T> VQ(M.M[I][0], M.M[I][1], M.M[I][2], M.M[I][3]);
-		VectorQuaternionMultiply(&VT, this, &VQ);
-		VectorQuaternionMultiply(&VR, &VT, &Inv);
-		Result.M[I][0] = VR.X;
-		Result.M[I][1] = VR.Y;
-		Result.M[I][2] = VR.Z;
-		Result.M[I][3] = VR.W;
-	}
-
-	return Result;
-}
-
-
 /* TQuat inline functions
  *****************************************************************************/
 
@@ -951,7 +936,11 @@ FORCEINLINE TQuat<T> TQuat<T>::operator-(const TQuat<T>& Q) const
 template<typename T>
 FORCEINLINE TQuat<T> TQuat<T>::operator-() const
 {
+#if PLATFORM_ENABLE_VECTORINTRINSICS
+	return TQuat(VectorNegate(VectorLoadAligned(this)));
+#else
 	return TQuat(-X, -Y, -Z, -W);
+#endif
 }
 
 template<typename T>
@@ -1289,6 +1278,13 @@ FORCEINLINE TVector<T> TQuat<T>::Vector() const
 	return GetAxisX();
 }
 
+template<typename T>
+FORCEINLINE TMatrix<T> TQuat<T>::ToMatrix() const
+{
+	TMatrix<T> R;
+	ToMatrix(R);
+	return R;
+}
 
 template<typename T>
 FORCEINLINE T TQuat<T>::Error(const TQuat<T>& Q1, const TQuat<T>& Q2)
diff --git a/Engine/Source/Runtime/Core/Public/Math/UnrealMathSSE.h b/Engine/Source/Runtime/Core/Public/Math/UnrealMathSSE.h
--- a/Engine/Source/Runtime/Core/Public/Math/UnrealMathSSE.h
+++ b/Engine/Source/Runtime/Core/Public/Math/UnrealMathSSE.h
@@ -2,6 +2,9 @@
 
 #pragma once
 
+// Define to pull in UnrealMathSSE.cpp, since conditions for this file are complex in VectorRegister.h
+#define UE_USING_UNREALMATH_SSE 1
+
 // We require SSE2
 #include <emmintrin.h>
 
@@ -698,6 +701,27 @@ FORCEINLINE VectorRegister4Double VectorSetFloat1(double D)
 	return Result;
 }
 
+/**
+ * Stores a vector to memory (aligned or unaligned).
+ *
+ * @param Vec	Vector to store
+ * @param Ptr	Memory pointer
+ */
+FORCEINLINE void VectorStore(const VectorRegister4Float& Vec, float* Ptr)
+{
+	_mm_storeu_ps(Ptr, Vec);
+}
+
+FORCEINLINE void VectorStore(const VectorRegister4Double& Vec, double* Dst)
+{
+#if !UE_PLATFORM_MATH_USE_AVX
+	_mm_storeu_pd(Dst, Vec.XY);
+	_mm_storeu_pd(Dst + 2, Vec.ZW);
+#else
+	_mm256_storeu_pd(Dst, Vec);
+#endif
+}
+
 /**
  * Stores a vector to aligned memory.
  *
@@ -744,27 +768,6 @@ FORCEINLINE void VectorStoreAlignedStreamed(const VectorRegister4Double& Vec, do
 #endif
 }
 
-/**
- * Stores a vector to memory (aligned or unaligned).
- *
- * @param Vec	Vector to store
- * @param Ptr	Memory pointer
- */
-FORCEINLINE void VectorStore(const VectorRegister4Float& Vec, float* Ptr)
-{
-	_mm_storeu_ps(Ptr, Vec);
-}
-
-FORCEINLINE void VectorStore(const VectorRegister4Double& Vec, double* Dst)
-{
-#if !UE_PLATFORM_MATH_USE_AVX
-	_mm_storeu_pd(Dst, Vec.XY);
-	_mm_storeu_pd(Dst + 2, Vec.ZW);
-#else
-	_mm256_storeu_pd(Dst, Vec);
-#endif
-}
-
 /**
  * Stores the XYZ components of a vector to unaligned memory.
  *
@@ -879,7 +882,7 @@ namespace SSEPermuteHelpers
 	// On AVX1, permute2f128 can be quite slow, so look for alternatives (extract + insert). On AVX2, permute2f128 is more efficient and should equal or beat (extract + insert).
 	// Sources: https://www.agner.org/optimize/instruction_tables.pdf, https://uops.info/table.html
 	template<> FORCEINLINE VectorRegister4Double PermuteLanes<0, 0>(const VectorRegister4Double& Vec) { return _mm256_insertf128_pd(Vec, Vec.GetXY(), 1); } // copy XY to lane 1
-	template<> FORCEINLINE VectorRegister4Double PermuteLanes<1, 0>(const VectorRegister4Double& Vec) { return _mm256_insertf128_pd(_mm256_castpd128_pd256(Vec.GetZW()), Vec.GetXY(), 1); } // copy XY to lane 1
+	template<> FORCEINLINE VectorRegister4Double PermuteLanes<1, 0>(const VectorRegister4Double& Vec) { return _mm256_setr_m128d(Vec.GetZW(), Vec.GetXY()); } // swap XY and ZW
 	template<> FORCEINLINE VectorRegister4Double PermuteLanes<1, 1>(const VectorRegister4Double& Vec) { return _mm256_insertf128_pd(Vec, Vec.GetZW(), 0); } // copy ZW to lane 0
 #endif // !AVX2
 
@@ -1027,7 +1030,7 @@ namespace SSEPermuteHelpers
 	// On AVX1, permute2f128 can be quite slow, so look for alternatives (extract + insert). On AVX2, permute2f128 is more efficient and should equal or beat (extract + insert).
 	// Sources: https://www.agner.org/optimize/instruction_tables.pdf, https://uops.info/table.html
 	template<> FORCEINLINE VectorRegister4Double ShuffleLanes<0, 0>(const VectorRegister4Double& Vec1, const VectorRegister4Double& Vec2) { return _mm256_insertf128_pd(Vec1, Vec2.GetXY(), 1); } // copy XY to lane 1
-	template<> FORCEINLINE VectorRegister4Double ShuffleLanes<1, 0>(const VectorRegister4Double& Vec1, const VectorRegister4Double& Vec2) { return _mm256_insertf128_pd(_mm256_castpd128_pd256(Vec1.GetZW()), Vec2.GetXY(), 1); } // copy XY to lane 1
+	template<> FORCEINLINE VectorRegister4Double ShuffleLanes<1, 0>(const VectorRegister4Double& Vec1, const VectorRegister4Double& Vec2) { return _mm256_setr_m128d(Vec1.GetZW(), Vec2.GetXY()); } // swap XY and ZW
 	template<> FORCEINLINE VectorRegister4Double ShuffleLanes<1, 1>(const VectorRegister4Double& Vec1, const VectorRegister4Double& Vec2) { return _mm256_insertf128_pd(Vec2, Vec1.GetZW(), 0); } // copy ZW to lane 0
 #endif // !AVX2
 
@@ -2114,101 +2117,8 @@ FORCEINLINE VectorRegister4Double VectorSet_W1(const VectorRegister4Double& Vec)
  * @param Matrix1	Pointer to the first matrix
  * @param Matrix2	Pointer to the second matrix
  */
-FORCEINLINE void VectorMatrixMultiply(FMatrix44f* Result, const FMatrix44f* Matrix1, const FMatrix44f* Matrix2)
-{
-	const VectorRegister4Float* A	= (const VectorRegister4Float*) Matrix1;
-	const VectorRegister4Float* B	= (const VectorRegister4Float*) Matrix2;
-	VectorRegister4Float* R		= (VectorRegister4Float*) Result;
-	VectorRegister4Float Temp, R0, R1, R2;
-
-	// First row of result (Matrix1[0] * Matrix2).
-	Temp	= VectorMultiply( VectorReplicate( A[0], 0 ), B[0] );
-	Temp	= VectorMultiplyAdd( VectorReplicate( A[0], 1 ), B[1], Temp );
-	Temp	= VectorMultiplyAdd( VectorReplicate( A[0], 2 ), B[2], Temp );
-	R0		= VectorMultiplyAdd( VectorReplicate( A[0], 3 ), B[3], Temp );
-
-	// Second row of result (Matrix1[1] * Matrix2).
-	Temp	= VectorMultiply( VectorReplicate( A[1], 0 ), B[0] );
-	Temp	= VectorMultiplyAdd( VectorReplicate( A[1], 1 ), B[1], Temp );
-	Temp	= VectorMultiplyAdd( VectorReplicate( A[1], 2 ), B[2], Temp );
-	R1		= VectorMultiplyAdd( VectorReplicate( A[1], 3 ), B[3], Temp );
-
-	// Third row of result (Matrix1[2] * Matrix2).
-	Temp	= VectorMultiply( VectorReplicate( A[2], 0 ), B[0] );
-	Temp	= VectorMultiplyAdd( VectorReplicate( A[2], 1 ), B[1], Temp );
-	Temp	= VectorMultiplyAdd( VectorReplicate( A[2], 2 ), B[2], Temp );
-	R2		= VectorMultiplyAdd( VectorReplicate( A[2], 3 ), B[3], Temp );
-
-	// Fourth row of result (Matrix1[3] * Matrix2).
-	Temp	= VectorMultiply( VectorReplicate( A[3], 0 ), B[0] );
-	Temp	= VectorMultiplyAdd( VectorReplicate( A[3], 1 ), B[1], Temp );
-	Temp	= VectorMultiplyAdd( VectorReplicate( A[3], 2 ), B[2], Temp );
-	Temp	= VectorMultiplyAdd( VectorReplicate( A[3], 3 ), B[3], Temp );
-
-	// Store result
-	R[0] = R0;
-	R[1] = R1;
-	R[2] = R2;
-	R[3] = Temp;
-}
-
-
-FORCEINLINE void VectorMatrixMultiply(FMatrix44d* Result, const FMatrix44d* Matrix1, const FMatrix44d* Matrix2)
-{
-	// Warning: FMatrix44d alignment may not match VectorRegister4Double, so you can't just cast to VectorRegister4Double*.
-	typedef double Double4x4[4][4];
-	const Double4x4& AMRows = *((const Double4x4*)Matrix1);
-	const Double4x4& BMRows = *((const Double4x4*)Matrix2);
-	Double4x4& ResultDst = *((Double4x4*)Result);
-	VectorRegister4Double Temp;
-
-	VectorRegister4Double B[4];
-	B[0] = VectorLoad(BMRows[0]);
-	B[1] = VectorLoad(BMRows[1]);
-	B[2] = VectorLoad(BMRows[2]);
-	B[3] = VectorLoad(BMRows[3]);
-
-	// First row of result (Matrix1[0] * Matrix2).
-	{
-		VectorRegister4Double ARow = VectorLoad(AMRows[0]);
-		Temp = VectorMultiply(VectorReplicate(ARow, 0), B[0]);
-		Temp = VectorMultiplyAdd(VectorReplicate(ARow, 1), B[1], Temp);
-		Temp = VectorMultiplyAdd(VectorReplicate(ARow, 2), B[2], Temp);
-		Temp = VectorMultiplyAdd(VectorReplicate(ARow, 3), B[3], Temp);
-		VectorStore(Temp, ResultDst[0]);
-	}
-
-	// Second row of result (Matrix1[1] * Matrix2).
-	{
-		VectorRegister4Double ARow = VectorLoad(AMRows[1]);
-		Temp = VectorMultiply(VectorReplicate(ARow, 0), B[0]);
-		Temp = VectorMultiplyAdd(VectorReplicate(ARow, 1), B[1], Temp);
-		Temp = VectorMultiplyAdd(VectorReplicate(ARow, 2), B[2], Temp);
-		Temp = VectorMultiplyAdd(VectorReplicate(ARow, 3), B[3], Temp);
-		VectorStore(Temp, ResultDst[1]);
-	}
-
-	// Third row of result (Matrix1[2] * Matrix2).
-	{
-		VectorRegister4Double ARow = VectorLoad(AMRows[2]);
-		Temp = VectorMultiply(VectorReplicate(ARow, 0), B[0]);
-		Temp = VectorMultiplyAdd(VectorReplicate(ARow, 1), B[1], Temp);
-		Temp = VectorMultiplyAdd(VectorReplicate(ARow, 2), B[2], Temp);
-		Temp = VectorMultiplyAdd(VectorReplicate(ARow, 3), B[3], Temp);
-		VectorStore(Temp, ResultDst[2]);
-	}
-
-	// Fourth row of result (Matrix1[3] * Matrix2).
-	{
-		VectorRegister4Double ARow = VectorLoad(AMRows[3]);
-		Temp = VectorMultiply(VectorReplicate(ARow, 0), B[0]);
-		Temp = VectorMultiplyAdd(VectorReplicate(ARow, 1), B[1], Temp);
-		Temp = VectorMultiplyAdd(VectorReplicate(ARow, 2), B[2], Temp);
-		Temp = VectorMultiplyAdd(VectorReplicate(ARow, 3), B[3], Temp);
-		VectorStore(Temp, ResultDst[3]);
-	}
-}
-
+CORE_API void VectorMatrixMultiply(FMatrix44f* Result, const FMatrix44f* Matrix1, const FMatrix44f* Matrix2);
+CORE_API void VectorMatrixMultiply(FMatrix44d* Result, const FMatrix44d* Matrix1, const FMatrix44d* Matrix2);
 
 /**
  * Calculate the inverse of an FMatrix.
@@ -2216,173 +2126,8 @@ FORCEINLINE void VectorMatrixMultiply(FMatrix44d* Result, const FMatrix44d* Matr
  * @param DstMatrix		FMatrix pointer to where the result should be stored
  * @param SrcMatrix		FMatrix pointer to the Matrix to be inversed
  */
-FORCEINLINE void VectorMatrixInverse(FMatrix44d* DstMatrix, const FMatrix44d* SrcMatrix)
-{
-	typedef double Double4x4[4][4];
-	const Double4x4& M = *((const Double4x4*)SrcMatrix);
-	Double4x4 Result;
-	double Det[4];
-	Double4x4 Tmp;
-
-	Tmp[0][0]	= M[2][2] * M[3][3] - M[2][3] * M[3][2];
-	Tmp[0][1]	= M[1][2] * M[3][3] - M[1][3] * M[3][2];
-	Tmp[0][2]	= M[1][2] * M[2][3] - M[1][3] * M[2][2];
-
-	Tmp[1][0]	= M[2][2] * M[3][3] - M[2][3] * M[3][2];
-	Tmp[1][1]	= M[0][2] * M[3][3] - M[0][3] * M[3][2];
-	Tmp[1][2]	= M[0][2] * M[2][3] - M[0][3] * M[2][2];
-
-	Tmp[2][0]	= M[1][2] * M[3][3] - M[1][3] * M[3][2];
-	Tmp[2][1]	= M[0][2] * M[3][3] - M[0][3] * M[3][2];
-	Tmp[2][2]	= M[0][2] * M[1][3] - M[0][3] * M[1][2];
-
-	Tmp[3][0]	= M[1][2] * M[2][3] - M[1][3] * M[2][2];
-	Tmp[3][1]	= M[0][2] * M[2][3] - M[0][3] * M[2][2];
-	Tmp[3][2]	= M[0][2] * M[1][3] - M[0][3] * M[1][2];
-
-	Det[0]		= M[1][1]*Tmp[0][0] - M[2][1]*Tmp[0][1] + M[3][1]*Tmp[0][2];
-	Det[1]		= M[0][1]*Tmp[1][0] - M[2][1]*Tmp[1][1] + M[3][1]*Tmp[1][2];
-	Det[2]		= M[0][1]*Tmp[2][0] - M[1][1]*Tmp[2][1] + M[3][1]*Tmp[2][2];
-	Det[3]		= M[0][1]*Tmp[3][0] - M[1][1]*Tmp[3][1] + M[2][1]*Tmp[3][2];
-
-	const double Determinant = M[0][0]*Det[0] - M[1][0]*Det[1] + M[2][0]*Det[2] - M[3][0]*Det[3];
-	const double	RDet = 1.0 / Determinant;
-
-	Result[0][0] =  RDet * Det[0];
-	Result[0][1] = -RDet * Det[1];
-	Result[0][2] =  RDet * Det[2];
-	Result[0][3] = -RDet * Det[3];
-	Result[1][0] = -RDet * (M[1][0]*Tmp[0][0] - M[2][0]*Tmp[0][1] + M[3][0]*Tmp[0][2]);
-	Result[1][1] =  RDet * (M[0][0]*Tmp[1][0] - M[2][0]*Tmp[1][1] + M[3][0]*Tmp[1][2]);
-	Result[1][2] = -RDet * (M[0][0]*Tmp[2][0] - M[1][0]*Tmp[2][1] + M[3][0]*Tmp[2][2]);
-	Result[1][3] =  RDet * (M[0][0]*Tmp[3][0] - M[1][0]*Tmp[3][1] + M[2][0]*Tmp[3][2]);
-	Result[2][0] =  RDet * (
-		M[1][0] * (M[2][1] * M[3][3] - M[2][3] * M[3][1]) -
-		M[2][0] * (M[1][1] * M[3][3] - M[1][3] * M[3][1]) +
-		M[3][0] * (M[1][1] * M[2][3] - M[1][3] * M[2][1])
-		);
-	Result[2][1] = -RDet * (
-		M[0][0] * (M[2][1] * M[3][3] - M[2][3] * M[3][1]) -
-		M[2][0] * (M[0][1] * M[3][3] - M[0][3] * M[3][1]) +
-		M[3][0] * (M[0][1] * M[2][3] - M[0][3] * M[2][1])
-		);
-	Result[2][2] =  RDet * (
-		M[0][0] * (M[1][1] * M[3][3] - M[1][3] * M[3][1]) -
-		M[1][0] * (M[0][1] * M[3][3] - M[0][3] * M[3][1]) +
-		M[3][0] * (M[0][1] * M[1][3] - M[0][3] * M[1][1])
-		);
-	Result[2][3] = -RDet * (
-		M[0][0] * (M[1][1] * M[2][3] - M[1][3] * M[2][1]) -
-		M[1][0] * (M[0][1] * M[2][3] - M[0][3] * M[2][1]) +
-		M[2][0] * (M[0][1] * M[1][3] - M[0][3] * M[1][1])
-		);
-	Result[3][0] = -RDet * (
-		M[1][0] * (M[2][1] * M[3][2] - M[2][2] * M[3][1]) -
-		M[2][0] * (M[1][1] * M[3][2] - M[1][2] * M[3][1]) +
-		M[3][0] * (M[1][1] * M[2][2] - M[1][2] * M[2][1])
-		);
-	Result[3][1] =  RDet * (
-		M[0][0] * (M[2][1] * M[3][2] - M[2][2] * M[3][1]) -
-		M[2][0] * (M[0][1] * M[3][2] - M[0][2] * M[3][1]) +
-		M[3][0] * (M[0][1] * M[2][2] - M[0][2] * M[2][1])
-		);
-	Result[3][2] = -RDet * (
-		M[0][0] * (M[1][1] * M[3][2] - M[1][2] * M[3][1]) -
-		M[1][0] * (M[0][1] * M[3][2] - M[0][2] * M[3][1]) +
-		M[3][0] * (M[0][1] * M[1][2] - M[0][2] * M[1][1])
-		);
-	Result[3][3] =  RDet * (
-		M[0][0] * (M[1][1] * M[2][2] - M[1][2] * M[2][1]) -
-		M[1][0] * (M[0][1] * M[2][2] - M[0][2] * M[2][1]) +
-		M[2][0] * (M[0][1] * M[1][2] - M[0][2] * M[1][1])
-		);
-
-	memcpy( DstMatrix, &Result, sizeof(Result) );
-}
-
-FORCEINLINE void VectorMatrixInverse(FMatrix44f* DstMatrix, const FMatrix44f* SrcMatrix)
-{
-	typedef float Float4x4[4][4];
-	const Float4x4& M = *((const Float4x4*)SrcMatrix);
-	Float4x4 Result;
-	float Det[4];
-	Float4x4 Tmp;
-
-	Tmp[0][0] = M[2][2] * M[3][3] - M[2][3] * M[3][2];
-	Tmp[0][1] = M[1][2] * M[3][3] - M[1][3] * M[3][2];
-	Tmp[0][2] = M[1][2] * M[2][3] - M[1][3] * M[2][2];
-
-	Tmp[1][0] = M[2][2] * M[3][3] - M[2][3] * M[3][2];
-	Tmp[1][1] = M[0][2] * M[3][3] - M[0][3] * M[3][2];
-	Tmp[1][2] = M[0][2] * M[2][3] - M[0][3] * M[2][2];
-
-	Tmp[2][0] = M[1][2] * M[3][3] - M[1][3] * M[3][2];
-	Tmp[2][1] = M[0][2] * M[3][3] - M[0][3] * M[3][2];
-	Tmp[2][2] = M[0][2] * M[1][3] - M[0][3] * M[1][2];
-
-	Tmp[3][0] = M[1][2] * M[2][3] - M[1][3] * M[2][2];
-	Tmp[3][1] = M[0][2] * M[2][3] - M[0][3] * M[2][2];
-	Tmp[3][2] = M[0][2] * M[1][3] - M[0][3] * M[1][2];
-
-	Det[0] = M[1][1] * Tmp[0][0] - M[2][1] * Tmp[0][1] + M[3][1] * Tmp[0][2];
-	Det[1] = M[0][1] * Tmp[1][0] - M[2][1] * Tmp[1][1] + M[3][1] * Tmp[1][2];
-	Det[2] = M[0][1] * Tmp[2][0] - M[1][1] * Tmp[2][1] + M[3][1] * Tmp[2][2];
-	Det[3] = M[0][1] * Tmp[3][0] - M[1][1] * Tmp[3][1] + M[2][1] * Tmp[3][2];
-
-	float Determinant = M[0][0] * Det[0] - M[1][0] * Det[1] + M[2][0] * Det[2] - M[3][0] * Det[3];
-	const float	RDet = 1.0f / Determinant;
-
-	Result[0][0] = RDet * Det[0];
-	Result[0][1] = -RDet * Det[1];
-	Result[0][2] = RDet * Det[2];
-	Result[0][3] = -RDet * Det[3];
-	Result[1][0] = -RDet * (M[1][0] * Tmp[0][0] - M[2][0] * Tmp[0][1] + M[3][0] * Tmp[0][2]);
-	Result[1][1] = RDet * (M[0][0] * Tmp[1][0] - M[2][0] * Tmp[1][1] + M[3][0] * Tmp[1][2]);
-	Result[1][2] = -RDet * (M[0][0] * Tmp[2][0] - M[1][0] * Tmp[2][1] + M[3][0] * Tmp[2][2]);
-	Result[1][3] = RDet * (M[0][0] * Tmp[3][0] - M[1][0] * Tmp[3][1] + M[2][0] * Tmp[3][2]);
-	Result[2][0] = RDet * (
-		M[1][0] * (M[2][1] * M[3][3] - M[2][3] * M[3][1]) -
-		M[2][0] * (M[1][1] * M[3][3] - M[1][3] * M[3][1]) +
-		M[3][0] * (M[1][1] * M[2][3] - M[1][3] * M[2][1])
-		);
-	Result[2][1] = -RDet * (
-		M[0][0] * (M[2][1] * M[3][3] - M[2][3] * M[3][1]) -
-		M[2][0] * (M[0][1] * M[3][3] - M[0][3] * M[3][1]) +
-		M[3][0] * (M[0][1] * M[2][3] - M[0][3] * M[2][1])
-		);
-	Result[2][2] = RDet * (
-		M[0][0] * (M[1][1] * M[3][3] - M[1][3] * M[3][1]) -
-		M[1][0] * (M[0][1] * M[3][3] - M[0][3] * M[3][1]) +
-		M[3][0] * (M[0][1] * M[1][3] - M[0][3] * M[1][1])
-		);
-	Result[2][3] = -RDet * (
-		M[0][0] * (M[1][1] * M[2][3] - M[1][3] * M[2][1]) -
-		M[1][0] * (M[0][1] * M[2][3] - M[0][3] * M[2][1]) +
-		M[2][0] * (M[0][1] * M[1][3] - M[0][3] * M[1][1])
-		);
-	Result[3][0] = -RDet * (
-		M[1][0] * (M[2][1] * M[3][2] - M[2][2] * M[3][1]) -
-		M[2][0] * (M[1][1] * M[3][2] - M[1][2] * M[3][1]) +
-		M[3][0] * (M[1][1] * M[2][2] - M[1][2] * M[2][1])
-		);
-	Result[3][1] = RDet * (
-		M[0][0] * (M[2][1] * M[3][2] - M[2][2] * M[3][1]) -
-		M[2][0] * (M[0][1] * M[3][2] - M[0][2] * M[3][1]) +
-		M[3][0] * (M[0][1] * M[2][2] - M[0][2] * M[2][1])
-		);
-	Result[3][2] = -RDet * (
-		M[0][0] * (M[1][1] * M[3][2] - M[1][2] * M[3][1]) -
-		M[1][0] * (M[0][1] * M[3][2] - M[0][2] * M[3][1]) +
-		M[3][0] * (M[0][1] * M[1][2] - M[0][2] * M[1][1])
-		);
-	Result[3][3] = RDet * (
-		M[0][0] * (M[1][1] * M[2][2] - M[1][2] * M[2][1]) -
-		M[1][0] * (M[0][1] * M[2][2] - M[0][2] * M[2][1]) +
-		M[2][0] * (M[0][1] * M[1][2] - M[0][2] * M[1][1])
-		);
-
-	memcpy(DstMatrix, &Result, sizeof(Result));
-}
+CORE_API void VectorMatrixInverse(FMatrix44f* DstMatrix, const FMatrix44f* SrcMatrix);
+CORE_API void VectorMatrixInverse(FMatrix44d* DstMatrix, const FMatrix44d* SrcMatrix);
 
 
 
diff --git a/Engine/Source/Runtime/Core/Public/Math/UnrealMathVectorCommon.h b/Engine/Source/Runtime/Core/Public/Math/UnrealMathVectorCommon.h
--- a/Engine/Source/Runtime/Core/Public/Math/UnrealMathVectorCommon.h
+++ b/Engine/Source/Runtime/Core/Public/Math/UnrealMathVectorCommon.h
@@ -853,7 +853,6 @@ FORCEINLINE VectorRegister4Double VectorFractional(const VectorRegister4Double&
 
 FORCEINLINE void VectorQuaternionMultiply(FQuat4f* RESTRICT Result, const FQuat4f* RESTRICT Quat1, const FQuat4f* RESTRICT Quat2)
 {
-	//*((VectorRegister4Float*)Result) = VectorQuaternionMultiply2(*((const VectorRegister4Float*)Quat1), *((const VectorRegister4Float*)Quat2));
 	const VectorRegister4Float Q1 = VectorLoadAligned(Quat1);
 	const VectorRegister4Float Q2 = VectorLoadAligned(Quat2);
 	const VectorRegister4Float QResult = VectorQuaternionMultiply2(Q1, Q2);
diff --git a/Engine/Source/Runtime/Experimental/Chaos/Private/Chaos/Collision/CollisionSolver.cpp b/Engine/Source/Runtime/Experimental/Chaos/Private/Chaos/Collision/CollisionSolver.cpp
--- a/Engine/Source/Runtime/Experimental/Chaos/Private/Chaos/Collision/CollisionSolver.cpp
+++ b/Engine/Source/Runtime/Experimental/Chaos/Private/Chaos/Collision/CollisionSolver.cpp
@@ -311,7 +311,7 @@ namespace Chaos
 			}
 			else
 			{
-				const FRotation3& PlaneQ = (ManifoldPoint.ContactPoint.ContactNormalOwnerIndex == 0) ? Body0.Q() : Body1.Q();
+				const FRotation3 PlaneQ = (ManifoldPoint.ContactPoint.ContactNormalOwnerIndex == 0) ? Body0.Q() : Body1.Q();
 				ContactNormal = PlaneQ * ManifoldPoint.ManifoldContactNormal;
 			}
 
diff --git a/Engine/Source/Runtime/Experimental/Chaos/Private/Chaos/Evolution/SolverBody.cpp b/Engine/Source/Runtime/Experimental/Chaos/Private/Chaos/Evolution/SolverBody.cpp
--- a/Engine/Source/Runtime/Experimental/Chaos/Private/Chaos/Evolution/SolverBody.cpp
+++ b/Engine/Source/Runtime/Experimental/Chaos/Private/Chaos/Evolution/SolverBody.cpp
@@ -2,6 +2,8 @@
 #include "Chaos/Evolution/SolverBody.h"
 #include "Chaos/Particle/ParticleUtilities.h"
 
+//PRAGMA_DISABLE_OPTIMIZATION
+
 namespace Chaos
 {
 
@@ -80,4 +82,63 @@ namespace Chaos
 			State.InvI = Utilities::ComputeWorldSpaceInertia(State.Q, State.InvILocal);
 		}
 	}
-}
\ No newline at end of file
+
+	void SolverQuaternionNormalizeApprox(FRotation3& InOutQ)
+	{
+		//constexpr FReal Tolerance = FReal(2.107342e-08);	// For full double-precision accuracy
+		constexpr FReal ToleranceF = FReal(0.001);
+
+#if PLATFORM_ENABLE_VECTORINTRINSICS
+		using QuatVectorRegister = FRotation3::QuatVectorRegister;
+
+		constexpr VectorRegister Tolerance = MakeVectorRegisterConstant(ToleranceF, ToleranceF, ToleranceF, ToleranceF);
+		constexpr VectorRegister One = MakeVectorRegisterConstant(FReal(1), FReal(1), FReal(1), FReal(1));
+		constexpr VectorRegister Two = MakeVectorRegisterConstant(FReal(2), FReal(2), FReal(2), FReal(2));
+
+		//const FReal QSq = SizeSquared();
+		const QuatVectorRegister Q = VectorLoadAligned(&InOutQ);
+		const QuatVectorRegister QSq = VectorDot4(Q, Q);
+
+		// if (FMath::Abs(FReal(1) - QSq) < Tolerance)
+		const QuatVectorRegister ToleranceCheck = VectorAbs(VectorSubtract(One, QSq));
+		if (VectorMaskBits(VectorCompareLE(ToleranceCheck, Tolerance)) != 0)
+		{
+			// Q * (FReal(2) / (FReal(1) + QSq))
+			const QuatVectorRegister Denom = VectorAdd(One, QSq);
+			const QuatVectorRegister Mult = VectorDivide(Two, Denom);
+			const QuatVectorRegister Result = VectorMultiply(Q, Mult);
+			VectorStoreAligned(Result, &InOutQ);
+		}
+		else
+		{
+			// Q / QLen
+			// @todo(chaos): with doubles VectorReciprocalSqrt does twice as many sqrts as we need and also has a divide
+			const QuatVectorRegister Result = VectorNormalize(Q);
+			VectorStoreAligned(Result, &InOutQ);
+		}
+#else
+		const FReal QSq = InOutQ.SizeSquared();
+		if (FMath::Abs(FReal(1) - QSq) < ToleranceF)
+		{
+			InOutQ *= (FReal(2) / (FReal(1) + QSq));
+		}
+		else
+		{
+			InOutQ.Normalize();
+		}
+#endif
+
+#if !(UE_BUILD_SHIPPING || UE_BUILD_TEST)
+		ensure(!InOutQ.ContainsNaN());
+		ensure(InOutQ.IsNormalized());
+#endif
+	}
+
+	FRotation3 SolverQuaternionApplyAngularDeltaApprox(const FRotation3& InQ0, const FVec3& InDR)
+	{
+		FRotation3 Q1 = InQ0 + (FRotation3::FromElements(InDR, FReal(0)) * InQ0) * FReal(0.5);
+		SolverQuaternionNormalizeApprox(Q1);
+		return Q1;
+	}
+}
+
diff --git a/Engine/Source/Runtime/Experimental/Chaos/Private/Chaos/Joint/JointSolverConstraints.cpp b/Engine/Source/Runtime/Experimental/Chaos/Private/Chaos/Joint/JointSolverConstraints.cpp
--- a/Engine/Source/Runtime/Experimental/Chaos/Private/Chaos/Joint/JointSolverConstraints.cpp
+++ b/Engine/Source/Runtime/Experimental/Chaos/Private/Chaos/Joint/JointSolverConstraints.cpp
@@ -160,7 +160,7 @@ namespace Chaos
 		const FJointSolverJointState& JointState,
 		const FPBDJointSettings& JointSettings)
 	{
-		const FRotation3& R0 = JointState.Rs[0];
+		const FRotation3 R0 = JointState.Rs[0];
 		const FVec3& X0 = JointState.Xs[0];
 		const FVec3& X1 = JointState.Xs[1];
 
@@ -331,8 +331,8 @@ namespace Chaos
 		const FJointSolverJointState& JointState,
 		const FPBDJointSettings& JointSettings)
 	{
-		const FRotation3& R0 = JointState.Rs[0];
-		const FRotation3& R1 = JointState.Rs[1];
+		const FRotation3 R0 = JointState.Rs[0];
+		const FRotation3 R1 = JointState.Rs[1];
 		const FRotation3 R01 = R0.Inverse() * R1;
 		FRotation3 R01Twist = FRotation3::FromIdentity();
 		FRotation3 R01Swing = FRotation3::FromIdentity();
diff --git a/Engine/Source/Runtime/Experimental/Chaos/Private/Chaos/MassProperties.cpp b/Engine/Source/Runtime/Experimental/Chaos/Private/Chaos/MassProperties.cpp
--- a/Engine/Source/Runtime/Experimental/Chaos/Private/Chaos/MassProperties.cpp
+++ b/Engine/Source/Runtime/Experimental/Chaos/Private/Chaos/MassProperties.cpp
@@ -7,6 +7,9 @@
 #include "Chaos/TriangleMesh.h"
 #include "ChaosCheck.h"
 
+
+//PRAGMA_DISABLE_OPTIMIZATION
+
 namespace Chaos
 {
 	FRotation3 TransformToLocalSpace(FMatrix33& Inertia)
@@ -200,7 +203,7 @@ namespace Chaos
 		for (const FMassProperties& Child : MPArray)
 		{
 			NewMP.Volume += Child.Volume;
-			const FMatrix33 ChildRI = Child.RotationOfMass * FMatrix::Identity;
+			const FMatrix33 ChildRI = Child.RotationOfMass.ToMatrix();
 			const FMatrix33 ChildWorldSpaceI = ChildRI.GetTransposed() * Child.InertiaTensor * ChildRI;
 			NewMP.InertiaTensor += ChildWorldSpaceI;
 			NewMP.CenterOfMass += Child.CenterOfMass * Child.Mass;
diff --git a/Engine/Source/Runtime/Experimental/Chaos/Private/Chaos/PBDRigidDynamicSpringConstraints.cpp b/Engine/Source/Runtime/Experimental/Chaos/Private/Chaos/PBDRigidDynamicSpringConstraints.cpp
--- a/Engine/Source/Runtime/Experimental/Chaos/Private/Chaos/PBDRigidDynamicSpringConstraints.cpp
+++ b/Engine/Source/Runtime/Experimental/Chaos/Private/Chaos/PBDRigidDynamicSpringConstraints.cpp
@@ -38,8 +38,8 @@ void FPBDRigidDynamicSpringConstraints::UpdatePositionBasedState(const FReal Dt)
 			continue;
 		}
 
-		const FRotation3& Q0 = bIsRigidDynamic0 ? PBDRigid0->Q() : Static0->R();
-		const FRotation3& Q1 = bIsRigidDynamic1 ? PBDRigid1->Q() : Static1->R();
+		const FRotation3 Q0 = bIsRigidDynamic0 ? PBDRigid0->Q() : Static0->R();
+		const FRotation3 Q1 = bIsRigidDynamic1 ? PBDRigid1->Q() : Static1->R();
 		const FVec3& P0 = bIsRigidDynamic0 ? PBDRigid0->P() : Static0->X();
 		const FVec3& P1 = bIsRigidDynamic1 ? PBDRigid1->P() : Static1->X();
 
diff --git a/Engine/Source/Runtime/Experimental/Chaos/Private/Chaos/Utilities.cpp b/Engine/Source/Runtime/Experimental/Chaos/Private/Chaos/Utilities.cpp
--- a/Engine/Source/Runtime/Experimental/Chaos/Private/Chaos/Utilities.cpp
+++ b/Engine/Source/Runtime/Experimental/Chaos/Private/Chaos/Utilities.cpp
@@ -1,10 +1,148 @@
 // Copyright Epic Games, Inc. All Rights Reserved.
 #include "Chaos/Utilities.h"
 #include "Chaos/Matrix.h"
+#include "Chaos/Rotation.h"
 
 namespace Chaos
 {
 	// @todo(ccaulfield): should be in ChaosCore, but that can't actually include its own headers at the moment (e.g., Matrix.h includes headers from Chaos)
-	const FMatrix33 FMatrix33::Zero = FMatrix33(0, 0, 0);
-	const FMatrix33 FMatrix33::Identity = FMatrix33(1, 1, 1);
+	const PMatrix<FRealSingle, 3, 3> PMatrix<FRealSingle, 3, 3>::Zero = PMatrix<FRealSingle, 3, 3>(0, 0, 0);
+	const PMatrix<FRealSingle, 3, 3> PMatrix<FRealSingle, 3, 3>::Identity = PMatrix<FRealSingle, 3, 3>(1, 1, 1);
+
+	const PMatrix<FRealDouble, 3, 3> PMatrix<FRealDouble, 3, 3>::Zero = PMatrix<FRealDouble, 3, 3>(0, 0, 0);
+	const PMatrix<FRealDouble, 3, 3> PMatrix<FRealDouble, 3, 3>::Identity = PMatrix<FRealDouble, 3, 3>(1, 1, 1);
+
+	/**
+	 * Return the complex conjugate of the rotation
+	 */
+	TRotation<FReal, 3> TRotation<FReal, 3>::Conjugate(const ::Chaos::TRotation<FReal, 3>& InR)
+	{		
+#if 0
+		// Disabled: Inlining this block can generate bad code in MSVC optimized builds with /arch:AVX, causing NaNs in the output.
+		TRotation<FReal, 3> R;
+		R.X = -InR.X;
+		R.Y = -InR.Y;
+		R.Z = -InR.Z;
+		R.W = InR.W;
+		return R;
+#else
+		return InR.Inverse();
+#endif
+	}
+
+	/**
+	 * Negate all values of the quaternion (note: not the inverse rotation. See Conjugate)
+	 */
+	TRotation<FReal, 3> TRotation<FReal, 3>::Negate(const ::Chaos::TRotation<FReal, 3>& InR)
+	{
+		return -InR;
+	}
+
+	/**
+	 * Create a rotation about an axis V/|V| by angle |V| in radians
+	 */
+	TRotation<FReal, 3> TRotation<FReal, 3>::FromVector(const ::Chaos::TVector<FReal, 3>& V)
+	{
+		using FQuatReal = BaseQuat::FReal;
+		TRotation<FReal, 3> Rot;
+		FReal HalfSize = 0.5f * V.Size();
+		FReal sinc = (FMath::Abs(HalfSize) > 1e-8) ? FMath::Sin(HalfSize) / HalfSize : 1;
+		auto RotV = 0.5f * sinc * V;
+		Rot.X = (FQuatReal)RotV.X;
+		Rot.Y = (FQuatReal)RotV.Y;
+		Rot.Z = (FQuatReal)RotV.Z;
+		Rot.W = (FQuatReal)FMath::Cos(HalfSize);
+		return Rot;
+	}
+
+	/**
+	 * Generate a Rotation that would rotate vector InitialVector to FinalVector
+	 */
+	TRotation<FReal, 3> TRotation<FReal, 3>::FromRotatedVector(
+		const ::Chaos::TVector<FReal, 3>& InitialVector,
+		const ::Chaos::TVector<FReal, 3>& FinalVector)
+	{
+		typedef Chaos::TVector<FReal, 3> TV;
+		checkSlow(FMath::Abs(InitialVector.Size() - 1.0) < KINDA_SMALL_NUMBER);
+		checkSlow(FMath::Abs(FinalVector.Size() - 1.0) < KINDA_SMALL_NUMBER);
+
+		const double CosTheta = FMath::Clamp<FReal>(TV::DotProduct(InitialVector, FinalVector), -1., 1.);
+
+		TV V = TV::CrossProduct(InitialVector, FinalVector);
+		const FReal VMagnitude = V.Size();
+		if (VMagnitude == 0)
+		{
+			return TRotation<FReal, 3>::FromElements(InitialVector, 0.f);
+		}
+
+		const FRealDouble SSquared = .5 * (1.0 + CosTheta); // Uses the half angle formula
+		const FRealDouble VMagnitudeDesired = sqrt(1.0 - SSquared);
+		V *= static_cast<FReal>(VMagnitudeDesired / VMagnitude);
+
+		return TRotation<FReal, 3>::FromElements(V, static_cast<FReal>(sqrt(SSquared)));
+	}
+
+	//PRAGMA_DISABLE_OPTIMIZATION
+
+	/**
+	 * Calculate the angular velocity required to take an object with orientation R0 to orientation R1 in time Dt.
+	 *
+	 * Uses the relation: DQ/DT = (W * Q)/2
+	 */
+	TVector<FReal, 3> TRotation<FReal, 3>::CalculateAngularVelocity1(const TRotation<FReal, 3>& InR0, const TRotation<FReal, 3>& InR1, const FReal InDt)
+	{
+		if (!ensure(InDt > SMALL_NUMBER))
+		{
+			return TVector<FReal, 3>(0);
+		}
+
+		const TRotation<FReal, 3> R0 = InR0;
+		TRotation<FReal, 3> R1 = InR1;
+		R1.EnforceShortestArcWith(R0);
+
+		// W = 2 * dQ/dT * Qinv
+		const TRotation<FReal, 3> DRDt = (R1 - R0) / InDt;
+		const TRotation<FReal, 3> RInv = Conjugate(R0);
+		const TRotation<FReal, 3> W = (DRDt * RInv) * 2.0f;
+
+		return TVector<FReal, 3>(W.X, W.Y, W.Z);
+	}
+
+	//PRAGMA_ENABLE_OPTIMIZATION
+
+
+	/**
+	 * Calculate the angular velocity required to take an object with orientation R0 to orientation R1 in time Dt.
+	 *
+	 * Uses the Quaternion to Axis/Angle method.
+	 */
+	TVector<FReal, 3> TRotation<FReal, 3>::CalculateAngularVelocity2(const TRotation<FReal, 3>& InR0, const TRotation<FReal, 3>& InR1, const FReal InDt)
+	{
+		// @todo(ccaulfield): ToAxisAndAngle starts to return increasingly random, non-normalized axes for very small angles. This 
+		// underestimates the angular velocity magnitude and randomizes direction.
+		if (!ensure(InDt > SMALL_NUMBER))
+		{
+			return TVector<FReal, 3>(0);
+		}
+
+		const TRotation<FReal, 3> R0 = InR0;
+		TRotation<FReal, 3> R1 = InR1;
+		R1.EnforceShortestArcWith(R0);
+
+		const TRotation<FReal, 3> DR = R1 * Conjugate(R0);
+		FReal Angle = DR.GetAngle();
+		TVector<FReal, 3> Axis = DR.GetRotationAxis();
+		return Axis * (Angle / InDt);
+	}
+
+	/**
+	 * Return a new rotation equal to the input rotation with angular velocity W applied over time Dt.
+	 *
+	 * Uses the relation: DQ/DT = (W * Q)/2
+	 */
+	TRotation<FReal, 3> TRotation<FReal, 3>::IntegrateRotationWithAngularVelocity(const TRotation<FReal, 3>& InR0, const TVector<FReal, 3>& InW, const FReal InDt)
+	{
+		TRotation<FReal, 3> R1 = InR0 + (TRotation<FReal, 3>::FromElements(InW.X, InW.Y, InW.Z, 0.f) * InR0) * decltype(InR0.X)(InDt * 0.5f);
+		return R1.GetNormalized();
+	}
 }
diff --git a/Engine/Source/Runtime/Experimental/Chaos/Public/Chaos/Collision/PBDCollisionConstraint.h b/Engine/Source/Runtime/Experimental/Chaos/Public/Chaos/Collision/PBDCollisionConstraint.h
--- a/Engine/Source/Runtime/Experimental/Chaos/Public/Chaos/Collision/PBDCollisionConstraint.h
+++ b/Engine/Source/Runtime/Experimental/Chaos/Public/Chaos/Collision/PBDCollisionConstraint.h
@@ -101,7 +101,7 @@ namespace Chaos
 			: bDisabled(false)
 			, Normal(0)
 			, Location(0)
-			, Phi(FLT_MAX)
+			, Phi(TNumericLimits<FReal>::Max())
 			, Friction(0)
 			, AngularFriction(0)
 			, Restitution(0)
@@ -139,7 +139,7 @@ namespace Chaos
 		void Reset()
 		{
 			bDisabled = false;
-			Phi = FLT_MAX;
+			Phi = TNumericLimits<FReal>::Max();
 			RestitutionPadding = 0;
 		}
 
diff --git a/Engine/Source/Runtime/Experimental/Chaos/Public/Chaos/Evolution/SolverBody.h b/Engine/Source/Runtime/Experimental/Chaos/Public/Chaos/Evolution/SolverBody.h
--- a/Engine/Source/Runtime/Experimental/Chaos/Public/Chaos/Evolution/SolverBody.h
+++ b/Engine/Source/Runtime/Experimental/Chaos/Public/Chaos/Evolution/SolverBody.h
@@ -4,6 +4,14 @@
 #include "Chaos/Core.h"
 #include "Math/Quat.h"
 
+
+#if DO_CHECK && !(UE_BUILD_SHIPPING || UE_BUILD_TEST)
+#define ChaosSolverEnsure( condition )		ensure( condition )
+#else
+#define ChaosSolverEnsure(...)
+#endif
+
+
 // Set to 1 to force single precision in the constraint solver, even if the default numeric type is double
 #define CHAOS_CONSTRAINTSOLVER_LOWPRECISION 1
 
@@ -50,59 +58,9 @@ namespace Chaos
 	 * does not seem to cost much.
 	 * 
 	*/
-	FORCEINLINE void SolverQuaternionNormalizeApprox(FRotation3& InOutQ)
-	{
-		//constexpr FReal Tolerance = FReal(2.107342e-08);	// For full double-precision accuracy
-		constexpr FReal ToleranceF = FReal(0.001);
-
-#if PLATFORM_ENABLE_VECTORINTRINSICS
-		using QuatVectorRegister = FRotation3::QuatVectorRegister;
-
-		constexpr VectorRegister Tolerance = MakeVectorRegisterConstant(ToleranceF, ToleranceF, ToleranceF, ToleranceF);
-		constexpr VectorRegister One = MakeVectorRegisterConstant(FReal(1), FReal(1), FReal(1), FReal(1));
-		constexpr VectorRegister Two = MakeVectorRegisterConstant(FReal(2), FReal(2), FReal(2), FReal(2));
-
-		//const FReal QSq = SizeSquared();
-		const QuatVectorRegister Q = VectorLoadAligned(&InOutQ);
-		const QuatVectorRegister QSq = VectorDot4(Q, Q);
-
-		// if (FMath::Abs(FReal(1) - QSq) < Tolerance)
-		const QuatVectorRegister ToleranceCheck = VectorAbs(VectorSubtract(One, QSq));
-		if (VectorMaskBits(VectorCompareLE(ToleranceCheck, Tolerance)) != 0)
-		{
-			// Q * (FReal(2) / (FReal(1) + QSq))
-			const QuatVectorRegister Denom = VectorAdd(One, QSq);
-			const QuatVectorRegister Mult = VectorDivide(Two, Denom);
-			const QuatVectorRegister Result = VectorMultiply(Q, Mult);
-			VectorStoreAligned(Result, &InOutQ);
-		}
-		else
-		{
-			// Q / QLen
-			// @todo(chaos): with doubles VectorReciprocalSqrt does twice as many sqrts as we need and also has a divide
-			const QuatVectorRegister Result = VectorNormalize(Q);
-			VectorStoreAligned(Result, &InOutQ);
-		}
+	CHAOS_API void SolverQuaternionNormalizeApprox(FRotation3& InOutQ);
 
-#else
-		const FReal QSq = InOutQ.SizeSquared();
-		if (FMath::Abs(FReal(1) - QSq) < ToleranceF)
-		{
-			InOutQ *= (FReal(2) / (FReal(1) + QSq));
-		}
-		else
-		{
-			InOutQ *= FMath::InvSqrt(QSq);
-		}
-#endif
-	}
-
-	FORCEINLINE FRotation3 SolverQuaternionApplyAngularDeltaApprox(const FRotation3& InQ0, const FVec3& InDR)
-	{
-		FRotation3 Q1 = InQ0 + (FRotation3::FromElements(InDR, FReal(0)) * InQ0) * FReal(0.5);
-		SolverQuaternionNormalizeApprox(Q1);
-		return Q1;
-	}
+	CHAOS_API FRotation3 SolverQuaternionApplyAngularDeltaApprox(const FRotation3& InQ0, const FVec3& InDR);
 
 
 	/**
@@ -141,7 +99,7 @@ namespace Chaos
 			{
 				const FSolverReal InvDt = FSolverReal(1) / FSolverReal(Dt);
 				State.V += State.DP * InvDt;
-				State.W += State.DQ * InvDt;
+				SetW(State.W + State.DQ * InvDt);
 			}
 		}
 
@@ -195,7 +153,11 @@ namespace Chaos
 		 * @brief Pre-integration world-space center of mass rotation
 		*/
 		inline const FRotation3& R() const { return State.R; }
-		inline void SetR(const FRotation3& InR) { State.R = InR; }
+		inline void SetR(const FRotation3& InR)
+		{
+			ChaosSolverEnsure(!InR.ContainsNaN());
+			State.R = InR;
+		}
 
 		/**
 		 * @brief Predicted (post-integrate) world-space center of mass position
@@ -211,7 +173,11 @@ namespace Chaos
 		 * @see DQ(), CorrectedQ()
 		*/
 		inline const FRotation3& Q() const { return State.Q; }
-		inline void SetQ(const FRotation3& InQ) { State.Q = InQ; }
+		inline void SetQ(const FRotation3& InQ)
+		{
+			ChaosSolverEnsure(!InQ.ContainsNaN());
+			State.Q = InQ;
+		}
 
 		/**
 		 * @brief World-space center of mass velocity
@@ -223,7 +189,11 @@ namespace Chaos
 		 * @brief World-space center of mass angular velocity
 		*/
 		inline const FVec3& W() const { return State.W; }
-		inline void SetW(const FVec3& InW) { State.W = InW; }
+		inline void SetW(const FVec3& InW)
+		{
+			ChaosSolverEnsure(!InW.ContainsNaN());
+			State.W = InW;
+		}
 
 		inline const FVec3& CoM() const { return State.CoM; }
 		inline void SetCoM(const FVec3& InCoM) { State.CoM = InCoM; }
@@ -337,7 +307,7 @@ namespace Chaos
 		*/
 		inline void ApplyAngularVelocityDelta(const FSolverVec3& DW)
 		{
-			State.W += DW;
+			SetW(State.W + DW);
 		}
 
 		/**
diff --git a/Engine/Source/Runtime/Experimental/Chaos/Public/Chaos/Rotation.h b/Engine/Source/Runtime/Experimental/Chaos/Public/Chaos/Rotation.h
--- a/Engine/Source/Runtime/Experimental/Chaos/Public/Chaos/Rotation.h
+++ b/Engine/Source/Runtime/Experimental/Chaos/Public/Chaos/Rotation.h
@@ -55,9 +55,11 @@ namespace Chaos
 		TRotation(const FMatrix& Matrix)
 		    : BaseQuat(Matrix) {}
 
-		PMatrix<FReal, 3, 3> ToMatrix() const
+		inline PMatrix<FReal, 3, 3> ToMatrix() const
 		{
-			return (FMatrix44f)FRotationMatrix::Make(*this);
+			PMatrix<FReal, 3, 3> R;
+			BaseQuat::ToMatrix(R);
+			return R;
 		}
 
 		/**
@@ -133,28 +135,12 @@ namespace Chaos
 		/**
 		 * Return the complex conjugate of the rotation
 		 */
-		static inline TRotation<FReal, 3> Conjugate(const ::Chaos::TRotation<FReal, 3>& InR)
-		{
-			TRotation<FReal, 3> R;
-			R.X = -InR.X;
-			R.Y = -InR.Y;
-			R.Z = -InR.Z;
-			R.W = InR.W;
-			return R;
-		}
+		CHAOS_API static TRotation<FReal, 3> Conjugate(const ::Chaos::TRotation<FReal, 3>& InR);
 
 		/**
 		 * Negate all values of the quaternion (note: not the inverse rotation. See Conjugate)
 		 */
-		static inline TRotation<FReal, 3> Negate(const ::Chaos::TRotation<FReal, 3>& InR)
-		{
-			TRotation<FReal, 3> R;
-			R.X = -InR.X;
-			R.Y = -InR.Y;
-			R.Z = -InR.Z;
-			R.W = -InR.W;
-			return R;
-		}
+		CHAOS_API static TRotation<FReal, 3> Negate(const ::Chaos::TRotation<FReal, 3>& InR);
 
 		/**
 		 * Create an identity rotation
@@ -192,94 +178,28 @@ namespace Chaos
 		/**
 		 * Create a rotation about an axis V/|V| by angle |V| in radians
 		 */
-		static inline TRotation<FReal, 3> FromVector(const ::Chaos::TVector<FReal, 3>& V)
-		{
-			using FQuatReal = BaseQuat::FReal;
-			TRotation<FReal, 3> Rot;
-			FReal HalfSize = 0.5f * V.Size();
-			FReal sinc = (FMath::Abs(HalfSize) > 1e-8) ? FMath::Sin(HalfSize) / HalfSize : 1;
-			auto RotV = 0.5f * sinc * V;
-			Rot.X = (FQuatReal)RotV.X;
-			Rot.Y = (FQuatReal)RotV.Y;
-			Rot.Z = (FQuatReal)RotV.Z;
-			Rot.W = (FQuatReal)FMath::Cos(HalfSize);
-			return Rot;
-		}
+		CHAOS_API static TRotation<FReal, 3> FromVector(const ::Chaos::TVector<FReal, 3>& V);
 
 		/**
 		 * Generate a Rotation that would rotate vector InitialVector to FinalVector
 		 */
-		static TRotation<FReal, 3> FromRotatedVector(
-		    const ::Chaos::TVector<FReal, 3>& InitialVector,
-		    const ::Chaos::TVector<FReal, 3>& FinalVector)
-		{
-			typedef Chaos::TVector<FReal, 3> TV;
-			checkSlow(FMath::Abs(InitialVector.Size() - 1.0) < KINDA_SMALL_NUMBER);
-			checkSlow(FMath::Abs(FinalVector.Size() - 1.0) < KINDA_SMALL_NUMBER);
-
-			const double CosTheta = FMath::Clamp<FReal>(TV::DotProduct(InitialVector, FinalVector), -1., 1.);
-
-			TV V = TV::CrossProduct(InitialVector, FinalVector);
-			const FReal VMagnitude = V.Size();
-			if(VMagnitude == 0)
-			{
-				return TRotation<FReal, 3>::FromElements(InitialVector, 0.f);
-			}
-
-			const FRealDouble SSquared = .5 * (1.0 + CosTheta); // Uses the half angle formula
-			const FRealDouble VMagnitudeDesired = sqrt(1.0 - SSquared);
-			V *= static_cast<FReal>(VMagnitudeDesired / VMagnitude);
-
-			return TRotation<FReal, 3>::FromElements(V, static_cast<FReal>(sqrt(SSquared)));
-		}
+		CHAOS_API static TRotation<FReal, 3> FromRotatedVector(
+			const ::Chaos::TVector<FReal, 3>& InitialVector,
+			const ::Chaos::TVector<FReal, 3>& FinalVector);
 
 		/**
 		 * Calculate the angular velocity required to take an object with orientation R0 to orientation R1 in time Dt.
 		 *
 		 * Uses the relation: DQ/DT = (W * Q)/2
 		 */
-		static inline TVector<FReal, 3> CalculateAngularVelocity1(const TRotation<FReal, 3>& InR0, const TRotation<FReal, 3>& InR1, const FReal InDt)
-		{
-			if (!ensure(InDt > SMALL_NUMBER))
-			{
-				return TVector<FReal, 3>(0);
-			}
-
-			const TRotation<FReal, 3>& R0 = InR0;
-			TRotation<FReal, 3> R1 = InR1;
-			R1.EnforceShortestArcWith(R0);
-
-			// W = 2 * dQ/dT * Qinv
-			const TRotation<FReal, 3> DRDt = (R1 - R0) / (decltype(R0.X))InDt;
-			const TRotation<FReal, 3> RInv = Conjugate(R0);
-			const TRotation<FReal, 3> W = (DRDt * RInv) * 2.0f;
-
-			return TVector<FReal, 3>(W.X, W.Y, W.Z);
-		}
+		CHAOS_API static TVector<FReal, 3> CalculateAngularVelocity1(const TRotation<FReal, 3>& R0, const TRotation<FReal, 3>& InR1, const FReal InDt);
 
 		/**
 		 * Calculate the angular velocity required to take an object with orientation R0 to orientation R1 in time Dt.
 		 *
 		 * Uses the Quaternion to Axis/Angle method.
 		 */
-		static inline TVector<FReal, 3> CalculateAngularVelocity2(const TRotation<FReal, 3>& InR0, const TRotation<FReal, 3>& InR1, const FReal InDt)
-		{
-			// @todo(ccaulfield): ToAxisAndAngle starts to return increasingly random, non-normalized axes for very small angles. This 
-			// underestimates the angular velocity magnitude and randomizes direction.
-			if (!ensure(InDt > SMALL_NUMBER))
-			{
-				return TVector<FReal, 3>(0);
-			}
-
-			const TRotation<FReal, 3>& R0 = InR0;
-			TRotation<FReal, 3> R1 = InR1;
-			R1.EnforceShortestArcWith(R0);
-
-			const TRotation<FReal, 3> DR = R1 * Conjugate(R0);
-			FReal Angle = DR.GetAngle();
-			TVector<FReal, 3> Axis = DR.GetRotationAxis();
-			return Axis * (Angle / InDt);
-		}
+		CHAOS_API static TVector<FReal, 3> CalculateAngularVelocity2(const TRotation<FReal, 3>& R0, const TRotation<FReal, 3>& InR1, const FReal InDt);
 
 		/**
 		 * Calculate the angular velocity required to take an object with orientation R0 to orientation R1 in time Dt.
@@ -306,12 +226,7 @@ namespace Chaos
 		 *
 		 * Uses the relation: DQ/DT = (W * Q)/2
 		 */
-		static inline TRotation<FReal, 3> IntegrateRotationWithAngularVelocity(const TRotation<FReal, 3>& InR0, const TVector<FReal, 3>& InW, const FReal InDt)
-		{
-			TRotation<FReal, 3> R1 = InR0 + (TRotation<FReal, 3>::FromElements(InW.X, InW.Y, InW.Z, 0.f) * InR0) * decltype(InR0.X)(InDt * 0.5f);
-			return R1.GetNormalized();
-		}
-
+		CHAOS_API static TRotation<FReal, 3> IntegrateRotationWithAngularVelocity(const TRotation<FReal, 3>& InR0, const TVector<FReal, 3>& InW, const FReal InDt);
 
 		/**
 		 * Check that two rotations are approximately equal. Assumes the quaternions are normalized and in the same hemisphere.
diff --git a/Engine/Source/Runtime/Experimental/ChaosCore/Public/Chaos/Matrix.h b/Engine/Source/Runtime/Experimental/ChaosCore/Public/Chaos/Matrix.h
--- a/Engine/Source/Runtime/Experimental/ChaosCore/Public/Chaos/Matrix.h
+++ b/Engine/Source/Runtime/Experimental/ChaosCore/Public/Chaos/Matrix.h
@@ -307,12 +307,12 @@ namespace Chaos
 			M[1][3] = 0;
 			M[2][3] = 0;
 		}
-#if COMPILE_WITHOUT_UNREAL_SUPPORT
-		PMatrix<FReal, 3, 3> GetTransposed()
+
+		PMatrix<FRealDouble, 3, 3> GetTransposed() const
 		{
 			return PMatrix<FRealDouble, 3, 3>(M[0][0], M[0][1], M[0][2], M[1][0], M[1][1], M[1][2], M[2][0], M[2][1], M[2][2]);
 		}
-		FReal Determinant()
+		FRealDouble Determinant() const
 		{
 			return M[0][0] * (M[1][1] * M[2][2] - M[1][2] * M[2][1]) - M[0][1] * (M[1][0] * M[2][2] - M[1][2] * M[2][0]) + M[0][2] * (M[1][0] * M[2][1] - M[1][1] * M[2][0]);
 		}
@@ -329,7 +329,7 @@ namespace Chaos
 			M[2][2] += Other.M[2][2];
 			return *this;
 		}
-#endif
+
 		// TDOD(mlentine): This should really be a vector multiply and sum for each entry using sse
 		TVector<FRealDouble, 3> operator*(const TVector<FRealDouble, 3>& Other) const
 		{
@@ -385,6 +385,7 @@ namespace Chaos
 		{
 			return static_cast<const UE::Math::TMatrix<FRealDouble>*>(this)->operator*(static_cast<const UE::Math::TMatrix<FRealDouble>&>(Other));
 		}
+		// Needs to be overridden because base version multiplies M[3][3]
 		PMatrix<FRealDouble, 3, 3> operator*(const FRealDouble Other) const
 		{
 			return PMatrix<FRealDouble, 3, 3>(
@@ -398,6 +399,20 @@ namespace Chaos
 			    M[1][2] * Other,
 			    M[2][2] * Other);
 		}
+		// Needs to be overridden because base version multiplies M[3][3]
+		PMatrix<FRealDouble, 3, 3> operator*=(const FRealDouble Other)
+		{
+			M[0][0] *= Other;
+			M[0][1] *= Other;
+			M[0][2] *= Other;
+			M[1][0] *= Other;
+			M[1][1] *= Other;
+			M[1][2] *= Other;
+			M[2][0] *= Other;
+			M[2][1] *= Other;
+			M[2][2] *= Other;
+			return *this;
+		}
 		friend PMatrix<FRealDouble, 3, 3> operator*(const FRealDouble OtherF, const PMatrix<FRealDouble, 3, 3>& OtherM)
 		{
 			return OtherM * OtherF;
@@ -563,19 +578,6 @@ namespace Chaos
 			    M[1][0] * Other.M[0][2] + M[1][1] * Other.M[1][2] + M[1][2] * Other.M[2][2],
 			    M[2][0] * Other.M[0][2] + M[2][1] * Other.M[1][2] + M[2][2] * Other.M[2][2]);
 		}
-		PMatrix<FRealDouble, 3, 3> operator*(const FRealDouble Scalar)
-		{
-			return PMatrix<FRealDouble, 3, 3>(
-			    M[0][0] * Scalar + M[0][1] * Scalar + M[0][2] * Scalar,
-			    M[1][0] * Scalar + M[1][1] * Scalar + M[1][2] * Scalar,
-			    M[2][0] * Scalar + M[2][1] * Scalar + M[2][2] * Scalar,
-			    M[0][0] * Scalar + M[0][1] * Scalar + M[0][2] * Scalar,
-			    M[1][0] * Scalar + M[1][1] * Scalar + M[1][2] * Scalar,
-			    M[2][0] * Scalar + M[2][1] * Scalar + M[2][2] * Scalar,
-			    M[0][0] * Scalar + M[0][1] * Scalar + M[0][2] * Scalar,
-			    M[1][0] * Scalar + M[1][1] * Scalar + M[1][2] * Scalar,
-			    M[2][0] * Scalar + M[2][1] * Scalar + M[2][2] * Scalar);
-		}
 #endif
 		inline bool Equals(const PMatrix<FRealDouble, 3, 3>& Other, FRealDouble Tolerance = KINDA_SMALL_NUMBER) const
 		{
@@ -712,12 +714,12 @@ namespace Chaos
 			M[1][3] = 0;
 			M[2][3] = 0;
 		}
-#if COMPILE_WITHOUT_UNREAL_SUPPORT
-		PMatrix<FRealSingle, 3, 3> GetTransposed()
+
+		PMatrix<FRealSingle, 3, 3> GetTransposed() const
 		{
 			return PMatrix<FRealSingle, 3, 3>(M[0][0], M[0][1], M[0][2], M[1][0], M[1][1], M[1][2], M[2][0], M[2][1], M[2][2]);
 		}
-		FRealSingle Determinant()
+		FRealSingle Determinant() const
 		{
 			return M[0][0] * (M[1][1] * M[2][2] - M[1][2] * M[2][1]) - M[0][1] * (M[1][0] * M[2][2] - M[1][2] * M[2][0]) + M[0][2] * (M[1][0] * M[2][1] - M[1][1] * M[2][0]);
 		}
@@ -734,7 +736,7 @@ namespace Chaos
 			M[2][2] += Other.M[2][2];
 			return *this;
 		}
-#endif
+
 		// TDOD(mlentine): This should really be a vector multiply and sum for each entry using sse
 		TVector<FRealSingle, 3> operator*(const TVector<FRealSingle, 3>& Other) const
 		{
@@ -790,6 +792,7 @@ namespace Chaos
 		{
 			return static_cast<const UE::Math::TMatrix<FRealSingle>*>(this)->operator*(static_cast<const UE::Math::TMatrix<FRealSingle>&>(Other));
 		}
+		// Needs to be overridden because base version multiplies M[3][3]
 		PMatrix<FRealSingle, 3, 3> operator*(const FRealSingle Other) const
 		{
 			return PMatrix<FRealSingle, 3, 3>(
@@ -803,6 +806,20 @@ namespace Chaos
 				M[1][2] * Other,
 				M[2][2] * Other);
 		}
+		// Needs to be overridden because base version multiplies M[3][3]
+		PMatrix<FRealSingle, 3, 3> operator*=(const FRealSingle Other)
+		{
+			M[0][0] *= Other;
+			M[0][1] *= Other;
+			M[0][2] *= Other;
+			M[1][0] *= Other;
+			M[1][1] *= Other;
+			M[1][2] *= Other;
+			M[2][0] *= Other;
+			M[2][1] *= Other;
+			M[2][2] *= Other;
+			return *this;
+		}
 		friend PMatrix<FRealSingle, 3, 3> operator*(const FRealSingle OtherF, const PMatrix<FRealSingle, 3, 3>& OtherM)
 		{
 			return OtherM * OtherF;
@@ -968,19 +985,6 @@ namespace Chaos
 				M[1][0] * Other.M[0][2] + M[1][1] * Other.M[1][2] + M[1][2] * Other.M[2][2],
 				M[2][0] * Other.M[0][2] + M[2][1] * Other.M[1][2] + M[2][2] * Other.M[2][2]);
 		}
-		PMatrix<FRealSingle, 3, 3> operator*(const FRealSingle Scalar)
-		{
-			return PMatrix<FRealSingle, 3, 3>(
-				M[0][0] * Scalar + M[0][1] * Scalar + M[0][2] * Scalar,
-				M[1][0] * Scalar + M[1][1] * Scalar + M[1][2] * Scalar,
-				M[2][0] * Scalar + M[2][1] * Scalar + M[2][2] * Scalar,
-				M[0][0] * Scalar + M[0][1] * Scalar + M[0][2] * Scalar,
-				M[1][0] * Scalar + M[1][1] * Scalar + M[1][2] * Scalar,
-				M[2][0] * Scalar + M[2][1] * Scalar + M[2][2] * Scalar,
-				M[0][0] * Scalar + M[0][1] * Scalar + M[0][2] * Scalar,
-				M[1][0] * Scalar + M[1][1] * Scalar + M[1][2] * Scalar,
-				M[2][0] * Scalar + M[2][1] * Scalar + M[2][2] * Scalar);
-		}
 #endif
 		inline bool Equals(const PMatrix<FRealSingle, 3, 3>& Other, FRealSingle Tolerance = KINDA_SMALL_NUMBER) const
 		{
