FEM - Fixing some issues with the ConcurrentLinearAllocator: wrong atomics, thread cleanup, and extra pagecache (optional)

diff --git a/Engine/Source/Runtime/Core/Public/Experimental/ConcurrentLinearAllocator.h b/Engine/Source/Runtime/Core/Public/Experimental/ConcurrentLinearAllocator.h
--- a/Engine/Source/Runtime/Core/Public/Experimental/ConcurrentLinearAllocator.h
+++ b/Engine/Source/Runtime/Core/Public/Experimental/ConcurrentLinearAllocator.h
@@ -14,17 +14,61 @@
 #define ASAN_UNPOISON_MEMORY_REGION(addr, size) ((void)(addr), (void)(size))
 #endif
 
+template<uint32 BlockSize>
+class TBlockAllocationCache
+{
+	static constexpr uint32 NumRecycledSlots = 4;
+
+	FORCEINLINE static std::atomic<void*>& GetSlot(uint32 i)
+	{
+		static std::atomic<void*> Head[NumRecycledSlots];
+		return Head[i];
+	}
+
+public:
+	FORCEINLINE static void* Malloc(SIZE_T Size)
+	{
+		if (Size == BlockSize)
+		{
+			for (uint32 i = 0; i < NumRecycledSlots; i++)
+			{
+				void* Pointer = GetSlot(i).exchange(nullptr, std::memory_order_relaxed);
+				if (Pointer != nullptr)
+				{
+					return Pointer;
+				}	
+			}
+		}
+		return FPlatformMemory::BinnedAllocFromOS(Size);
+	}
+
+	FORCEINLINE static void Free(void* Pointer, SIZE_T Size)
+	{
+		if (Size == BlockSize)
+		{
+			for (uint32 i = 0; i < NumRecycledSlots; i++)
+			{
+				Pointer = GetSlot(i).exchange(Pointer, std::memory_order_relaxed);
+				if (Pointer == nullptr)
+				{
+					return;
+				}	
+			}
+		}
+		return FPlatformMemory::BinnedFreeToOS(Pointer, Size);
+	}
+};
 
 struct FDefaultBlockAllocationTag 
 {
-	typedef void* (*MallocType)(SIZE_T Size);
-	typedef void (*FreeType)(void* Pointer);
-
-	static constexpr MallocType MallocFunction = &FMemory::SystemMalloc;
-	static constexpr FreeType FreeFunction = &FMemory::SystemFree;
 	static constexpr uint32 BlockSize = 64 * 1024;		// Blocksize used to allocate from
 	static constexpr bool AllowOversizedBlocks = true;  // The allocator supports oversized Blocks and will store them in a seperate Block with counter 1
 	static constexpr const TCHAR* TagName = TEXT("DefaultLinear");
+
+	typedef void* (*MallocType)(SIZE_T Size);
+	typedef void (*FreeType)(void* Pointer, SIZE_T Size);
+	static constexpr MallocType MallocFunction = &FPlatformMemory::BinnedAllocFromOS;
+	static constexpr FreeType FreeFunction = &FPlatformMemory::BinnedFreeToOS;
 };
 
 // This concurrent fast linear Allocator can be used for temporary allocations that are usually produced and consumed on different threads and within the lifetime of a frame.
@@ -86,6 +130,36 @@ class TConcurrentLinearAllocator
 		unsigned int Num = 0; //this tracks the TLS local Number of allocation from a given Block
 	};
 
+	struct FTLSCleanup
+	{
+		FBlockHeader* Header = nullptr;
+
+		~FTLSCleanup()
+		{
+			if(Header)
+			{
+				Header->NextAllocationPtr = uintptr_t(Header) + BlockAllocationTag::BlockSize;
+				const uint32 DeltaCount = UINT_MAX - Header->Num; 
+
+				// on the allocating side we only need to do a single atomic to reduce contention with the deletions
+				// this will leave the atomic in a state where it only counts the number of live allocations (before it was based of UINT_MAX)
+				if (Header->NumAllocations.fetch_sub(DeltaCount, std::memory_order_release) == DeltaCount)
+				{
+					//if all allocations are already freed we can reuse the Block again
+					Header->~FBlockHeader();
+					ASAN_UNPOISON_MEMORY_REGION( Header, BlockAllocationTag::BlockSize );
+					BlockAllocationTag::FreeFunction(Header, BlockAllocationTag::BlockSize);				
+				}
+			}
+		}
+	};
+
+	FORCENOINLINE static void RegisterCleanup(FBlockHeader* Header)
+	{
+		static thread_local FTLSCleanup Cleanup;
+		new (&Cleanup) FTLSCleanup { Header };
+	}
+
 public:
 	template<uint32 Alignment> //this can help the compiler folding away the alignment when it is statically known
 	static FORCEINLINE_DEBUGGABLE void* Malloc(SIZE_T Size)
@@ -101,85 +175,82 @@ public:
 		static thread_local FBlockHeader* Header;
 		if (Header == nullptr)
 		{
+		AllocateNewBlock:
 			static_assert(BlockAllocationTag::BlockSize >= sizeof(FBlockHeader) + sizeof(FAllocationHeader));
 			Header = new (BlockAllocationTag::MallocFunction(BlockAllocationTag::BlockSize)) FBlockHeader;
 			checkSlow(IsAligned(Header, alignof(FBlockHeader)));
 			ASAN_POISON_MEMORY_REGION( Header + 1, BlockAllocationTag::BlockSize - sizeof(FBlockHeader) );
+			RegisterCleanup(Header);
 		}
 
-		for (;;)
+	AllocateNewItem:
+		uintptr_t AlignedOffset = Align(Header->NextAllocationPtr, Alignment);
+		if (AlignedOffset + Size <= uintptr_t(Header) + BlockAllocationTag::BlockSize)
 		{
-			uintptr_t AlignedOffset = Align(Header->NextAllocationPtr, Alignment);
-			if (AlignedOffset + Size <= uintptr_t(Header) + BlockAllocationTag::BlockSize)
-			{
-				Header->NextAllocationPtr = AlignedOffset + Size + sizeof(FAllocationHeader);
-				Header->Num++;
-				
-				FAllocationHeader* AllocationHeader = reinterpret_cast<FAllocationHeader*>(AlignedOffset) - 1;
-				ASAN_UNPOISON_MEMORY_REGION( AllocationHeader, sizeof(FAllocationHeader) + Size );
-				new (AllocationHeader) FAllocationHeader(Header, Size);
-				ASAN_POISON_MEMORY_REGION( AllocationHeader, sizeof(FAllocationHeader) );
+			Header->NextAllocationPtr = AlignedOffset + Size + sizeof(FAllocationHeader);
+			Header->Num++;
 
-				LLM_IF_ENABLED( FLowLevelMemTracker::Get().OnLowLevelAlloc(ELLMTracker::Default, reinterpret_cast<void*>(AlignedOffset), Size, ELLMTag::LinearAllocator) );
-				return reinterpret_cast<void*>(AlignedOffset);
-			}
+			FAllocationHeader* AllocationHeader = reinterpret_cast<FAllocationHeader*>(AlignedOffset) - 1;
+			ASAN_UNPOISON_MEMORY_REGION( AllocationHeader, sizeof(FAllocationHeader) + Size );
+			new (AllocationHeader) FAllocationHeader(Header, Size);
+			ASAN_POISON_MEMORY_REGION( AllocationHeader, sizeof(FAllocationHeader) );
 
-			// The cold path of the code starts here
-			constexpr SIZE_T HeaderSize = sizeof(FBlockHeader) + sizeof(FAllocationHeader);
-			if constexpr (BlockAllocationTag::AllowOversizedBlocks)
+			LLM_IF_ENABLED( FLowLevelMemTracker::Get().OnLowLevelAlloc(ELLMTracker::Default, AllocationHeader, Size + sizeof(FAllocationHeader), ELLMTag::LinearAllocator) );
+			return reinterpret_cast<void*>(AlignedOffset);
+		}
+
+		// The cold path of the code starts here
+		constexpr SIZE_T HeaderSize = sizeof(FBlockHeader) + sizeof(FAllocationHeader);
+		if constexpr (BlockAllocationTag::AllowOversizedBlocks)
+		{
+			//support for oversized Blocks
+			if (HeaderSize + Size + Alignment > BlockAllocationTag::BlockSize)
 			{
-				//support for oversized Blocks
-				if (HeaderSize + Size + Alignment > BlockAllocationTag::BlockSize)
-				{
-					FBlockHeader* LargeHeader = new (BlockAllocationTag::MallocFunction(HeaderSize + Size + Alignment)) FBlockHeader;
-					checkSlow(IsAligned(LargeHeader, alignof(FBlockHeader)));
+				FBlockHeader* LargeHeader = new (BlockAllocationTag::MallocFunction(HeaderSize + Size + Alignment)) FBlockHeader;
+				checkSlow(IsAligned(LargeHeader, alignof(FBlockHeader)));
 
-					uintptr_t LargeAlignedOffset = Align(LargeHeader->NextAllocationPtr, Alignment);
-					LargeHeader->NextAllocationPtr = uintptr_t(LargeHeader) + HeaderSize + Size + Alignment;
-					LargeHeader->NumAllocations.store(1, std::memory_order_release);
+				uintptr_t LargeAlignedOffset = Align(LargeHeader->NextAllocationPtr, Alignment);
+				LargeHeader->NextAllocationPtr = uintptr_t(LargeHeader) + HeaderSize + Size + Alignment;
+				LargeHeader->NumAllocations.store(1, std::memory_order_release);
 
-					checkSlow(LargeAlignedOffset + Size <= LargeHeader->NextAllocationPtr);
-					FAllocationHeader* AllocationHeader = new (reinterpret_cast<FAllocationHeader*>(LargeAlignedOffset) - 1) FAllocationHeader(LargeHeader, Size);
-					ASAN_POISON_MEMORY_REGION( AllocationHeader, sizeof(FAllocationHeader) );
+				checkSlow(LargeAlignedOffset + Size <= LargeHeader->NextAllocationPtr);
+				FAllocationHeader* AllocationHeader = new (reinterpret_cast<FAllocationHeader*>(LargeAlignedOffset) - 1) FAllocationHeader(LargeHeader, Size);
+				ASAN_POISON_MEMORY_REGION( AllocationHeader, sizeof(FAllocationHeader) );
 
-					LLM_IF_ENABLED( FLowLevelMemTracker::Get().OnLowLevelAlloc(ELLMTracker::Default, reinterpret_cast<void*>(LargeAlignedOffset), Size, ELLMTag::LinearAllocator) );
-					return reinterpret_cast<void*>(LargeAlignedOffset);
-				}
-			}
-			else
-			{
-				check(HeaderSize + Size + Alignment <= BlockAllocationTag::BlockSize);
+				LLM_IF_ENABLED( FLowLevelMemTracker::Get().OnLowLevelAlloc(ELLMTracker::Default, AllocationHeader, Size + sizeof(FAllocationHeader), ELLMTag::LinearAllocator) );
+				return reinterpret_cast<void*>(LargeAlignedOffset);
 			}
+		}
+		else
+		{
+			check(HeaderSize + Size + Alignment <= BlockAllocationTag::BlockSize);
+		}
 
-			//allocation of a new Block
-			Header->NextAllocationPtr = uintptr_t(Header) + BlockAllocationTag::BlockSize;
-			const uint32 DeltaCount = UINT_MAX - Header->Num; 
-
-			// on the allocating side we only need to do a single atomic to reduce contention with the deletions
-			// this will leave the atomic in a state where it only counts the number of live allocations (before it was based of UINT_MAX)
-			if (Header->NumAllocations.fetch_sub(DeltaCount, std::memory_order_release) == DeltaCount)
-			{
-				//if all allocations are already freed we can reuse the Block again
-				Header->~FBlockHeader();
-				Header = new (Header) FBlockHeader;
-				ASAN_POISON_MEMORY_REGION( Header + 1, BlockAllocationTag::BlockSize - sizeof(FBlockHeader) );
-				continue;
-			}
+		//allocation of a new Block
+		Header->NextAllocationPtr = uintptr_t(Header) + BlockAllocationTag::BlockSize;
+		const uint32 DeltaCount = UINT_MAX - Header->Num; 
 
-			Header = new (BlockAllocationTag::MallocFunction(BlockAllocationTag::BlockSize)) FBlockHeader;
-			checkSlow(IsAligned(Header, alignof(FBlockHeader)));
+		// on the allocating side we only need to do a single atomic to reduce contention with the deletions
+		// this will leave the atomic in a state where it only counts the number of live allocations (before it was based of UINT_MAX)
+		if (Header->NumAllocations.fetch_sub(DeltaCount, std::memory_order_release) == DeltaCount)
+		{
+			//if all allocations are already freed we can reuse the Block again
+			Header->~FBlockHeader();
+			Header = new (Header) FBlockHeader;
 			ASAN_POISON_MEMORY_REGION( Header + 1, BlockAllocationTag::BlockSize - sizeof(FBlockHeader) );
+			goto AllocateNewItem;
 		}
+
+		goto AllocateNewBlock;
 	}
 
 	static FORCEINLINE_DEBUGGABLE void Free(void* Pointer)
 	{
 		if(Pointer != nullptr)
 		{
-			LLM_IF_ENABLED(FLowLevelMemTracker::Get().OnLowLevelFree(ELLMTracker::Default, reinterpret_cast<void*>(Pointer)));
-
 			FAllocationHeader* AllocationHeader = GetAllocationHeader(Pointer);
 			ASAN_UNPOISON_MEMORY_REGION( AllocationHeader, sizeof(FAllocationHeader) );
+			LLM_IF_ENABLED(FLowLevelMemTracker::Get().OnLowLevelFree(ELLMTracker::Default, AllocationHeader));
 			FBlockHeader* Header = AllocationHeader->GetBlockHeader();
 			ASAN_POISON_MEMORY_REGION( AllocationHeader, sizeof(FAllocationHeader) + AllocationHeader->GetAllocationSize() );
 
@@ -188,7 +259,7 @@ public:
 			{
 				Header->~FBlockHeader();
 				ASAN_UNPOISON_MEMORY_REGION( Header, Header->NextAllocationPtr - uintptr_t(Header) );
-				BlockAllocationTag::FreeFunction(Header);
+				BlockAllocationTag::FreeFunction(Header, Header->NextAllocationPtr - uintptr_t(Header));
 			}
 		}
 	}
@@ -427,7 +498,7 @@ public:
 
 	inline void BulkDelete()
 	{
-		FAllocation* Allocation = Next.exchange(nullptr, std::memory_order_release);
+		FAllocation* Allocation = Next.exchange(nullptr, std::memory_order_acquire);
 		while (Allocation != nullptr)
 		{
 			FAllocation* NextAllocation = Allocation->Next;
@@ -446,7 +517,7 @@ public:
 		checkSlow(uintptr_t(Alloc) + Size - uintptr_t(Allocation) <= TotalSize);
 
 		Allocation->Next = Next.load(std::memory_order_relaxed);
-		while (!Next.compare_exchange_strong(Allocation->Next, Allocation, std::memory_order_acquire, std::memory_order_relaxed))
+		while (!Next.compare_exchange_strong(Allocation->Next, Allocation, std::memory_order_release, std::memory_order_relaxed))
 		{
 		}
 		return Alloc;
@@ -483,7 +554,7 @@ private: //Kept private, there is also deliberately no new operator overloads be
 		checkSlow(uintptr_t(Alloc + 1) - uintptr_t(Object) <= TotalSize);
 
 		Object->Next = Next.load(std::memory_order_relaxed);
-		while (!Next.compare_exchange_weak(Object->Next, Object, std::memory_order_acquire, std::memory_order_relaxed))
+		while (!Next.compare_exchange_weak(Object->Next, Object, std::memory_order_release, std::memory_order_relaxed))
 		{
 		}
 		return Alloc;
@@ -500,7 +571,7 @@ private: //Kept private, there is also deliberately no new operator overloads be
 		checkSlow(uintptr_t(Alloc + Num) - uintptr_t(Array) <= TotalSize);
 
 		Array->Next = Next.load(std::memory_order_relaxed);
-		while (!Next.compare_exchange_weak(Array->Next, Array, std::memory_order_acquire, std::memory_order_relaxed))
+		while (!Next.compare_exchange_weak(Array->Next, Array, std::memory_order_release, std::memory_order_relaxed))
 		{
 		}
 		return Alloc;
