#ue5 - LWC: fix #if check in case symbol is not defined.

diff --git a/Engine/Source/Runtime/Core/Private/Math/UnrealMathSSE.cpp b/Engine/Source/Runtime/Core/Private/Math/UnrealMathSSE.cpp
--- a/Engine/Source/Runtime/Core/Private/Math/UnrealMathSSE.cpp
+++ b/Engine/Source/Runtime/Core/Private/Math/UnrealMathSSE.cpp
@@ -1,7 +1,7 @@
 // Copyright Epic Games, Inc. All Rights Reserved.
 
 // This file implements non-inlined functions from UnrealMathSSE.h
-#if UE_USING_UNREALMATH_SSE
+#if defined(UE_USING_UNREALMATH_SSE) && (UE_USING_UNREALMATH_SSE != 0)
 
 #include "Math/UnrealMathSSE.h"
 #include "Math/Matrix.h"
