Add a log message for loading an executable with a compiled in custom config to help diagnose load issues

diff --git a/Engine/Source/Runtime/Core/Private/Misc/ConfigCacheIni.cpp b/Engine/Source/Runtime/Core/Private/Misc/ConfigCacheIni.cpp
--- a/Engine/Source/Runtime/Core/Private/Misc/ConfigCacheIni.cpp
+++ b/Engine/Source/Runtime/Core/Private/Misc/ConfigCacheIni.cpp
@@ -4414,7 +4414,12 @@ const FString& FConfigCacheIni::GetCustomConfigString()
 		{
 			UE_LOG(LogConfig, Log, TEXT("Overriding CustomConfig from %s to %s"), TEXT(CUSTOM_CONFIG), *CustomConfigString);
 		}
+		else
 #endif
+		if (!CustomConfigString.IsEmpty())
+		{
+			UE_LOG(LogConfig, Log, TEXT("Using compiled CustomConfig %s"), *CustomConfigString);
+		}
 	}
 	return CustomConfigString;
 }
