FEM - Fix small GraphEventLeak

diff --git a/Engine/Source/Runtime/Core/Private/Internationalization/TextLocalizationManager.cpp b/Engine/Source/Runtime/Core/Private/Internationalization/TextLocalizationManager.cpp
--- a/Engine/Source/Runtime/Core/Private/Internationalization/TextLocalizationManager.cpp
+++ b/Engine/Source/Runtime/Core/Private/Internationalization/TextLocalizationManager.cpp
@@ -449,6 +449,7 @@ void EndInitGameTextLocalization()
 	if (InitGameTextLocalizationTask)
 	{
 		FTaskGraphInterface::Get().WaitUntilTaskCompletes(InitGameTextLocalizationTask);
+		InitGameTextLocalizationTask.SafeRelease();
 	}
 }
 
