FEM - Fix memory oder of the Schedulers Sleep Stack

diff --git a/Engine/Source/Runtime/Core/Public/Async/Fundamental/Scheduler.h b/Engine/Source/Runtime/Core/Public/Async/Fundamental/Scheduler.h
--- a/Engine/Source/Runtime/Core/Public/Async/Fundamental/Scheduler.h
+++ b/Engine/Source/Runtime/Core/Public/Async/Fundamental/Scheduler.h
@@ -36,7 +36,7 @@ namespace LowLevelTasks
 	public:
 		NodeType* Pop()
 		{
-			FTopNode LocalTop = Top.load(std::memory_order_relaxed);
+			FTopNode LocalTop = Top.load(std::memory_order_acquire);
 			while (true) 
 			{			
 				if (LocalTop.Address == 0)
@@ -46,8 +46,8 @@ namespace LowLevelTasks
 #if DO_CHECK
 				int64 LastRevision = int64(LocalTop.Revision); 
 #endif
-				NodeType* Item = reinterpret_cast<NodeType*>(LocalTop.Address << 3);
-				if (Top.compare_exchange_weak(LocalTop, FTopNode { reinterpret_cast<uintptr_t>(Item->Next) >> 3, uintptr_t(LocalTop.Revision + 1) }, std::memory_order_acquire, std::memory_order_relaxed))
+				NodeType* Item = reinterpret_cast<NodeType*>(LocalTop.Address << 3);																 //acquire on failure because we read Item->Next next itteration
+				if (Top.compare_exchange_weak(LocalTop, FTopNode { reinterpret_cast<uintptr_t>(Item->Next) >> 3, uintptr_t(LocalTop.Revision + 1) }, std::memory_order_relaxed, std::memory_order_acquire)) 
 				{
 					Item->Next = nullptr;
 					return Item;
