[Insights] Append default file name for paths that end in directories.

Append the default file name when the trace target path ends with a slash, for examplethe path "foo/bar/". The target trace file will then be "foo/bar/defaultname.utrace". Note that "foo/bar" will still target "foo/bar.utrace".

diff --git a/Engine/Source/Runtime/Core/Private/ProfilingDebugging/TraceAuxiliary.cpp b/Engine/Source/Runtime/Core/Private/ProfilingDebugging/TraceAuxiliary.cpp
--- a/Engine/Source/Runtime/Core/Private/ProfilingDebugging/TraceAuxiliary.cpp
+++ b/Engine/Source/Runtime/Core/Private/ProfilingDebugging/TraceAuxiliary.cpp
@@ -400,9 +400,12 @@ bool FTraceAuxiliaryImpl::SendToHost(const TCHAR* Host)
 ////////////////////////////////////////////////////////////////////////////////
 bool FTraceAuxiliaryImpl::WriteToFile(const TCHAR* Path)
 {
+	// Default file name functor 
+	auto GetDefaultName = [] {return FDateTime::Now().ToString(TEXT("%Y%m%d_%H%M%S.utrace"));};
+	
 	if (Path == nullptr || *Path == '\0')
 	{
-		FString Name = FDateTime::Now().ToString(TEXT("%Y%m%d_%H%M%S.utrace"));
+		const FString Name = GetDefaultName();
 		return WriteToFile(*Name);
 	}
 
@@ -419,6 +422,12 @@ bool FTraceAuxiliaryImpl::WriteToFile(const TCHAR* Path)
 		WritePath = Path;
 	}
 
+	// If a directory is specified, add the default trace file name
+	if (WritePath.EndsWith(TEXT("/")) || WritePath.EndsWith(TEXT("\\")))
+	{
+		WritePath += GetDefaultName();
+	}
+
 	// The user may not have provided a suitable extension
 	if (!WritePath.EndsWith(".utrace"))
 	{
