[Insights] Allow log output in TraceAux to be routed to arbitrary category.

In order to be able to output detailed error and status messages to both core log or console feedback log depending on situation, add argument to set the target category. Also removed now redundant logs in console commands, since more detailed information is now emitted on the console feedback.

diff --git a/Engine/Source/Runtime/Core/Private/ProfilingDebugging/TraceAuxiliary.cpp b/Engine/Source/Runtime/Core/Private/ProfilingDebugging/TraceAuxiliary.cpp
--- a/Engine/Source/Runtime/Core/Private/ProfilingDebugging/TraceAuxiliary.cpp
+++ b/Engine/Source/Runtime/Core/Private/ProfilingDebugging/TraceAuxiliary.cpp
@@ -84,7 +84,7 @@ public:
 	void					ResetCommandlineChannels();
 	void					EnableChannels(const TCHAR* ChannelList);
 	void					DisableChannels(const TCHAR* ChannelList);
-	bool					Connect(ETraceConnectType Type, const TCHAR* Parameter);
+	bool					Connect(ETraceConnectType Type, const TCHAR* Parameter, const FLogCategoryBase& LogCategory);
 	bool					Stop();
 	void					ResumeChannels();
 	void					PauseChannels();
@@ -114,8 +114,8 @@ private:
 	static uint32			HashChannelName(const TCHAR* Name);
 	bool					EnableChannel(const TCHAR* Channel);
 	void					DisableChannel(const TCHAR* Channel);
-	bool					SendToHost(const TCHAR* Host);
-	bool					WriteToFile(const TCHAR* Path=nullptr);
+	bool					SendToHost(const TCHAR* Host, const FLogCategoryBase& LogCategory);
+	bool					WriteToFile(const TCHAR* Path, const FLogCategoryBase& LogCategory);
 
 	TMap<uint32, FChannel>	CommandlineChannels;
 	FString					TraceDest;
@@ -237,7 +237,7 @@ void FTraceAuxiliaryImpl::RemoveChannel(const TCHAR* Name)
 }
 
 ////////////////////////////////////////////////////////////////////////////////
-bool FTraceAuxiliaryImpl::Connect(ETraceConnectType Type, const TCHAR* Parameter)
+bool FTraceAuxiliaryImpl::Connect(ETraceConnectType Type, const TCHAR* Parameter, const FLogCategoryBase& LogCategory)
 {
 	// Connect/write to file. But only if we're not already sending/writing
 	bool bConnected = UE::Trace::IsTracing();
@@ -245,27 +245,27 @@ bool FTraceAuxiliaryImpl::Connect(ETraceConnectType Type, const TCHAR* Parameter
 	{
 		if (Type == ETraceConnectType::Network)
 		{
-			bConnected = SendToHost(Parameter);
+			bConnected = SendToHost(Parameter, LogCategory);
 			if (bConnected)
 			{
-				UE_LOG(LogCore, Display, TEXT("Trace started (connected to trace server %s)."), GetDest());
+				UE_LOG_REF(LogCategory, Display, TEXT("Trace started (connected to trace server %s)."), GetDest());
 			}
 			else
 			{
-				UE_LOG(LogCore, Error, TEXT("Trace failed to connect (trace server: %s)!"), Parameter ? Parameter : TEXT(""));
+				UE_LOG_REF(LogCategory, Error, TEXT("Trace failed to connect (trace server: %s)!"), Parameter ? Parameter : TEXT(""));
 			}	
 		}
 
 		else if (Type == ETraceConnectType::File)
 		{
-			bConnected = WriteToFile(Parameter);
+			bConnected = WriteToFile(Parameter, LogCategory);
 			if (bConnected)
 			{
-				UE_LOG(LogCore, Display, TEXT("Trace started (writing to file \"%s\")."), GetDest());
+				UE_LOG_REF(LogCategory, Display, TEXT("Trace started (writing to file \"%s\")."), GetDest());
 			}
 			else
 			{
-				UE_LOG(LogCore, Error, TEXT("Trace failed to connect (file: \"%s\")!"), Parameter ? Parameter : TEXT(""));
+				UE_LOG_REF(LogCategory, Error, TEXT("Trace failed to connect (file: \"%s\")!"), Parameter ? Parameter : TEXT(""));
 			}	
 		}
 	}
@@ -385,11 +385,11 @@ void FTraceAuxiliaryImpl::SetTruncateFile(bool bNewTruncateFileState)
 }
 
 ////////////////////////////////////////////////////////////////////////////////
-bool FTraceAuxiliaryImpl::SendToHost(const TCHAR* Host)
+bool FTraceAuxiliaryImpl::SendToHost(const TCHAR* Host, const FLogCategoryBase& LogCategory)
 {
 	if (!UE::Trace::SendTo(Host))
 	{
-		UE_LOG(LogCore, Warning, TEXT("Unable to trace to host '%s'"), Host);
+		UE_LOG_REF(LogCategory, Warning, TEXT("Unable to trace to host '%s'"), Host);
 		return false;
 	}
 
@@ -398,7 +398,7 @@ bool FTraceAuxiliaryImpl::SendToHost(const TCHAR* Host)
 }
 
 ////////////////////////////////////////////////////////////////////////////////
-bool FTraceAuxiliaryImpl::WriteToFile(const TCHAR* Path)
+bool FTraceAuxiliaryImpl::WriteToFile(const TCHAR* Path, const FLogCategoryBase& LogCategory)
 {
 	// Default file name functor 
 	auto GetDefaultName = [] {return FDateTime::Now().ToString(TEXT("%Y%m%d_%H%M%S.utrace"));};
@@ -406,7 +406,7 @@ bool FTraceAuxiliaryImpl::WriteToFile(const TCHAR* Path)
 	if (Path == nullptr || *Path == '\0')
 	{
 		const FString Name = GetDefaultName();
-		return WriteToFile(*Name);
+		return WriteToFile(*Name, LogCategory);
 	}
 
 	FString WritePath;
@@ -437,16 +437,16 @@ bool FTraceAuxiliaryImpl::WriteToFile(const TCHAR* Path)
 	IFileManager& FileManager = IFileManager::Get();
 
 	// Ensure we can write the trace file appropriately
-	FString WriteDir = FPaths::GetPath(WritePath);
+	const FString WriteDir = FPaths::GetPath(WritePath);
 	if (!FileManager.MakeDirectory(*WriteDir, true))
 	{
-		UE_LOG(LogCore, Warning, TEXT("Failed to create directory '%s'"), *WriteDir);
+		UE_LOG_REF(LogCategory, Warning, TEXT("Failed to create directory '%s'"), *WriteDir);
 		return false;
 	}
 
 	if (!bTruncateFile && FileManager.FileExists(*WritePath))
 	{
-		UE_LOG(LogCore, Warning, TEXT("Trace file '%s' already exists"), *WritePath);
+		UE_LOG_REF(LogCategory, Warning, TEXT("Trace file '%s' already exists"), *WritePath);
 		return false;
 	}
 
@@ -454,7 +454,7 @@ bool FTraceAuxiliaryImpl::WriteToFile(const TCHAR* Path)
 	FString NativePath = FileManager.ConvertToAbsolutePathForExternalAppForWrite(*WritePath);
 	if (!UE::Trace::WriteTo(*NativePath))
 	{
-		UE_LOG(LogCore, Warning, TEXT("Unable to trace to file '%s'"), *WritePath);
+		UE_LOG_REF(LogCategory, Warning, TEXT("Unable to trace to file '%s'"), *WritePath);
 		return false;
 	}
 
@@ -561,8 +561,7 @@ static void TraceAuxiliaryConnectEpilogue()
 	TStringBuilder<128> Channels;
 	GTraceAuxiliary.GetActiveChannelsString(Channels);
 	
-	UE_LOG(LogConsoleResponse, Log, TEXT("Tracing to: %s"), TraceDest);
-	UE_LOG(LogConsoleResponse, Log, TEXT("Trace channels: %s"), Channels.ToString());
+	UE_LOG(LogConsoleResponse, Log, TEXT("Active channels: %s"), Channels.ToString());
 }
 
 ////////////////////////////////////////////////////////////////////////////////
@@ -576,13 +575,10 @@ static void TraceAuxiliarySend(const TArray<FString>& Args)
 
 	const TCHAR* Target = *Args[0];
 	const TCHAR* Channels = Args.Num() > 1 ? *Args[1] : nullptr;
-	if (!FTraceAuxiliary::Start(FTraceAuxiliary::EConnectionType::Network, Target, Channels))
+	if (FTraceAuxiliary::Start(FTraceAuxiliary::EConnectionType::Network, Target, Channels, nullptr, LogConsoleResponse))
 	{
-		UE_LOG(LogConsoleResponse, Warning, TEXT("Failed to start tracing to '%s'"), *Args[0]);
-		return;
+		TraceAuxiliaryConnectEpilogue();
 	}
-
-	TraceAuxiliaryConnectEpilogue();
 }
 
 ////////////////////////////////////////////////////////////////////////////////
@@ -611,13 +607,10 @@ static void TraceAuxiliaryFile(const TArray<FString>& Args)
 	
 	FTraceAuxiliary::Options Opts;
 	Opts.bNoWorkerThread = true;
-	if (!FTraceAuxiliary::Start(FTraceAuxiliary::EConnectionType::File, File, Channels, &Opts))
+	if (FTraceAuxiliary::Start(FTraceAuxiliary::EConnectionType::File, File, Channels, &Opts, LogConsoleResponse))
 	{
-		UE_LOG(LogConsoleResponse, Warning, TEXT("Failed to start tracing to a file"));
-		return;
+		TraceAuxiliaryConnectEpilogue();
 	}
-
-	TraceAuxiliaryConnectEpilogue();
 }
 
 ////////////////////////////////////////////////////////////////////////////////
@@ -1047,12 +1040,12 @@ static bool StartFromCommandlineArguments(const TCHAR* CommandLine)
 }
 
 ////////////////////////////////////////////////////////////////////////////////
-bool FTraceAuxiliary::Start(EConnectionType Type, const TCHAR* Target, const TCHAR* Channels, Options* Options)
+bool FTraceAuxiliary::Start(EConnectionType Type, const TCHAR* Target, const TCHAR* Channels, Options* Options, const FLogCategoryBase& LogCategory)
 {
 #if UE_TRACE_ENABLED
 	if (GTraceAuxiliary.IsConnected())
 	{
-		UE_LOG(LogCore, Error, TEXT("Unable to start trace, already tracing to %s"), GTraceAuxiliary.GetDest());
+		UE_LOG_REF(LogCategory, Error, TEXT("Unable to start trace, already tracing to %s"), GTraceAuxiliary.GetDest());
 		return false;
 	}
 	
@@ -1064,7 +1057,7 @@ bool FTraceAuxiliary::Start(EConnectionType Type, const TCHAR* Target, const TCH
 
 	if (Channels)
 	{
-		UE_LOG(LogCore, Display, TEXT("Trace channels: '%s'"), Channels);
+		UE_LOG_REF(LogCategory, Display, TEXT("Requested channels: '%s'"), Channels);
 		GTraceAuxiliary.ResetCommandlineChannels();
 		GTraceAuxiliary.AddCommandlineChannels(Channels);
 		GTraceAuxiliary.EnableCommandlineChannels();
@@ -1081,11 +1074,11 @@ bool FTraceAuxiliary::Start(EConnectionType Type, const TCHAR* Target, const TCH
 
 	if (Type == EConnectionType::File)
 	{
-		return GTraceAuxiliary.Connect(ETraceConnectType::File, Target);
+		return GTraceAuxiliary.Connect(ETraceConnectType::File, Target, LogCategory);
 	}
 	else if(Type == EConnectionType::Network)
 	{
-		return GTraceAuxiliary.Connect(ETraceConnectType::Network, Target);	
+		return GTraceAuxiliary.Connect(ETraceConnectType::Network, Target, LogCategory);	
 	}
 #endif
 	return false;
diff --git a/Engine/Source/Runtime/Core/Public/ProfilingDebugging/TraceAuxiliary.h b/Engine/Source/Runtime/Core/Public/ProfilingDebugging/TraceAuxiliary.h
--- a/Engine/Source/Runtime/Core/Public/ProfilingDebugging/TraceAuxiliary.h
+++ b/Engine/Source/Runtime/Core/Public/ProfilingDebugging/TraceAuxiliary.h
@@ -41,9 +41,10 @@ public:
 	 * @param Target String to use for connection. See /ref EConnectionType for details.
 	 * @param Channels Comma separated list of channels to enable. If the pointer is null the default channels will be active.
 	 * @param Options Optional additional tracing options.
+	 * @param LogCategory Log channel to output messages to. Default set to 'Core'.
 	 * @return True when successfully starting the trace, false if the data connection could not be made.
 	 */
-	static bool Start(EConnectionType Type, const TCHAR* Target, const TCHAR* Channels, Options* Options = nullptr);
+	static bool Start(EConnectionType Type, const TCHAR* Target, const TCHAR* Channels, Options* Options = nullptr, const FLogCategoryBase& LogCategory = LogCore);
 
 	/**
 	 * Stop tracing.
