[Insights] Fix for no logging situations

diff --git a/Engine/Source/Runtime/Core/Private/ProfilingDebugging/TraceAuxiliary.cpp b/Engine/Source/Runtime/Core/Private/ProfilingDebugging/TraceAuxiliary.cpp
--- a/Engine/Source/Runtime/Core/Private/ProfilingDebugging/TraceAuxiliary.cpp
+++ b/Engine/Source/Runtime/Core/Private/ProfilingDebugging/TraceAuxiliary.cpp
@@ -1040,7 +1040,7 @@ static bool StartFromCommandlineArguments(const TCHAR* CommandLine)
 }
 
 ////////////////////////////////////////////////////////////////////////////////
-bool FTraceAuxiliary::Start(EConnectionType Type, const TCHAR* Target, const TCHAR* Channels, Options* Options, const FLogCategoryBase& LogCategory)
+bool FTraceAuxiliary::Start(EConnectionType Type, const TCHAR* Target, const TCHAR* Channels, Options* Options, const FLogCategoryAlias& LogCategory)
 {
 #if UE_TRACE_ENABLED
 	if (GTraceAuxiliary.IsConnected())
diff --git a/Engine/Source/Runtime/Core/Public/ProfilingDebugging/TraceAuxiliary.h b/Engine/Source/Runtime/Core/Public/ProfilingDebugging/TraceAuxiliary.h
--- a/Engine/Source/Runtime/Core/Public/ProfilingDebugging/TraceAuxiliary.h
+++ b/Engine/Source/Runtime/Core/Public/ProfilingDebugging/TraceAuxiliary.h
@@ -4,11 +4,22 @@
 
 #include "HAL/Platform.h"
 
+
+////////////////////////////////////////////////////////////////////////////////
+	
 ////////////////////////////////////////////////////////////////////////////////
 class CORE_API FTraceAuxiliary
 {
 public:
 
+	// In no logging configurations all log categories are of type FNoLoggingCategory, which has no relation with
+	// FLogCategoryBase. In order to not need to conditionally set the argument alias the type here.
+#if NO_LOGGING
+	typedef FNoLoggingCategory& FLogCategoryAlias;
+#else
+	typedef FLogCategoryBase FLogCategoryAlias;
+#endif
+	
 	enum class EConnectionType
 	{
 		/**
@@ -44,7 +55,7 @@ public:
 	 * @param LogCategory Log channel to output messages to. Default set to 'Core'.
 	 * @return True when successfully starting the trace, false if the data connection could not be made.
 	 */
-	static bool Start(EConnectionType Type, const TCHAR* Target, const TCHAR* Channels, Options* Options = nullptr, const FLogCategoryBase& LogCategory = LogCore);
+	static bool Start(EConnectionType Type, const TCHAR* Target, const TCHAR* Channels, Options* Options = nullptr, const FLogCategoryAlias& LogCategory = LogCore);
 
 	/**
 	 * Stop tracing.
