Implemented FCrc::Strihash_DEPRECATED for UTF8CHAR

diff --git a/Engine/Source/Runtime/Core/Public/Misc/Crc.h b/Engine/Source/Runtime/Core/Public/Misc/Crc.h
--- a/Engine/Source/Runtime/Core/Public/Misc/Crc.h
+++ b/Engine/Source/Runtime/Core/Public/Misc/Crc.h
@@ -192,3 +192,35 @@ inline uint32 FCrc::Strihash_DEPRECATED(const int32 DataLen, const WIDECHAR* Dat
 	}
 	return Hash;
 }
+
+template <>
+inline uint32 FCrc::Strihash_DEPRECATED(const UTF8CHAR* Data)
+{
+	// make sure table is initialized
+	check(CRCTable_DEPRECATED[1] != 0);
+
+	uint32 Hash=0;
+	while( *Data )
+	{
+		UTF8CHAR Ch = TChar<UTF8CHAR>::ToUpper(*Data++);
+		uint8 B  = Ch;
+		Hash = ((Hash >> 8) & 0x00FFFFFF) ^ CRCTable_DEPRECATED[(Hash ^ B) & 0x000000FF];
+	}
+	return Hash;
+}
+
+template <>
+inline uint32 FCrc::Strihash_DEPRECATED(const int32 DataLen, const UTF8CHAR* Data)
+{
+	// make sure table is initialized
+	check(CRCTable_DEPRECATED[1] != 0);
+
+	uint32 Hash=0;
+	for (int32 Idx = 0; Idx < DataLen; ++Idx)
+	{
+		UTF8CHAR Ch = TChar<UTF8CHAR>::ToUpper(*Data++);
+		uint8 B  = Ch;
+		Hash = ((Hash >> 8) & 0x00FFFFFF) ^ CRCTable_DEPRECATED[(Hash ^ B) & 0x000000FF];
+	}
+	return Hash;
+}
