[Insights] Nonunity build fixes.

diff --git a/Engine/Source/Runtime/Core/Public/ProfilingDebugging/TraceAuxiliary.h b/Engine/Source/Runtime/Core/Public/ProfilingDebugging/TraceAuxiliary.h
--- a/Engine/Source/Runtime/Core/Public/ProfilingDebugging/TraceAuxiliary.h
+++ b/Engine/Source/Runtime/Core/Public/ProfilingDebugging/TraceAuxiliary.h
@@ -2,8 +2,9 @@
 
 #pragma once
 
+#include "CoreGlobals.h"
 #include "HAL/Platform.h"
-
+#include "Logging/LogMacros.h"
 
 ////////////////////////////////////////////////////////////////////////////////
 	
