Disable stackwalking on Android of XOM check is detected

diff --git a/Engine/Source/Runtime/Core/Private/Android/AndroidPlatformStackWalk.cpp b/Engine/Source/Runtime/Core/Private/Android/AndroidPlatformStackWalk.cpp
--- a/Engine/Source/Runtime/Core/Private/Android/AndroidPlatformStackWalk.cpp
+++ b/Engine/Source/Runtime/Core/Private/Android/AndroidPlatformStackWalk.cpp
@@ -11,6 +11,8 @@
 #include <dlfcn.h>
 #include <cxxabi.h>
 #include <stdio.h>
+#include <signal.h>
+#include <setjmp.h>
 #include <android/log.h>
 #include "Android/AndroidSignals.h"
 
@@ -32,24 +34,63 @@
 #include "HAL/PlatformProcess.h"
 #include "Misc/ScopeExit.h"
 
-void FAndroidPlatformStackWalk::NotifyPlatformVersionInit()
+// Some devices with Android 10 have XOM security feature and walking stack might crash
+// We need to verify first that it's safe to read callstack in a call toInitStackWalking
+// Otherwise stack walking will be disabled
+static sigjmp_buf XomJmp;
+static bool DisableStackBacktracing = true;
+static bool StackWalkingInitialized = false;
+static void XomSignalHandler(int Sig)
 {
+	siglongjmp(XomJmp, Sig);
+}
+
+bool FAndroidPlatformStackWalk::InitStackWalking()
+{
+	if (StackWalkingInitialized)
+	{
+		return true;
+	}
+
 #if HAS_LIBUNWIND
 	// Without this stack walk might touch executable memory and ASan will terminate the app on that.
-	// Xom protection presents the same issue and is enabled on android 10 devices when the targetsdk is 29 or higher.
+	// Xom protection presents the same issue and is enabled on android 10 devices
 	// see https://source.android.com/devices/tech/debug/execute-only-memory
-	if (RUNNING_WITH_ASAN || (FAndroidMisc::GetTargetSDKVersion() >= 29 && FAndroidMisc::GetAndroidMajorVersion() == 10))
-	{
-		// prevent libunwind attempting to deref IP during signal frame test. (this will make backtrace called from a signal less useful.)
-		unw_disable_signal_frame_test(1);
-	}
-	else
+	// Please note that XOM is enabled on some devices even when building with TargetSDK < 29 (Oculus Quest 2)
+	// and not enabled on some other devices at all, like Pixel 4
+	sigset_t SignalSet;
+	sigemptyset(&SignalSet);
+	sigaddset(&SignalSet, SIGSEGV);
+
+	struct sigaction SigAction;
+	struct sigaction OldSigAction;
+	sigset_t OldSignalSet;
+	memset(&SigAction, 0, sizeof(SigAction));
+	SigAction.sa_handler = XomSignalHandler;
+	SigAction.sa_mask = SignalSet;
+
+	sigprocmask(SIG_SETMASK, &SigAction.sa_mask, &OldSignalSet);
+	sigaction(SIGSEGV, &SigAction, &OldSigAction);
+
+	if (sigsetjmp(XomJmp, 1) == 0)
 	{
 		// first call to unw_backtrace will trigger some initial large allocations and if it happens during stack capturing on an exception we might get another out of memory exception
 		const uint32 Depth = 16;
 		void* Stack[Depth];
 		unw_backtrace((void**)Stack, Depth);
+		DisableStackBacktracing = false;
 	}
+	else
+	{
+		unw_disable_signal_frame_test(1);
+		__android_log_print(ANDROID_LOG_DEBUG, "UE", "XOM has been detected");
+	}
+
+	sigaction(SIGSEGV, &OldSigAction, nullptr);
+	sigprocmask(SIG_SETMASK, &OldSignalSet, nullptr);
+
+	StackWalkingInitialized = true;
+	return true;
 #endif
 }
 
@@ -172,20 +213,10 @@ extern int32 unwind_backtrace_signal(void* sigcontext, uint64* Backtrace, int32
 
 uint32 FAndroidPlatformStackWalk::CaptureStackBackTrace(uint64* BackTrace, uint32 MaxDepth, void* Context)
 {
-#if PLATFORM_ANDROID_ARM64
-	if (FAndroidMisc::GetTargetSDKVersion() >= 29 && FAndroidMisc::GetAndroidMajorVersion() == 10)
+	if (DisableStackBacktracing)
 	{
-		// UE-103382
-		// due to execute-only memory (xom) we cannot currently walk the stack on Android 10 devices when targeting Android 29 or greater.
-		static int32 OnceOnly = 0;
-		if (Context == nullptr && OnceOnly == 0)
-		{
-			__android_log_print(ANDROID_LOG_DEBUG, "UE", "FAndroidPlatformStackWalk::CaptureStackBackTrace disabled on Android 10 with TargetSDK >= 29 due to XOM.");
-			OnceOnly = 1;
-		}
 		return 0;
 	}
-#endif
 
 	// Make sure we have place to store the information
 	if (BackTrace == NULL || MaxDepth == 0)
diff --git a/Engine/Source/Runtime/Core/Public/Android/AndroidPlatformStackWalk.h b/Engine/Source/Runtime/Core/Public/Android/AndroidPlatformStackWalk.h
--- a/Engine/Source/Runtime/Core/Public/Android/AndroidPlatformStackWalk.h
+++ b/Engine/Source/Runtime/Core/Public/Android/AndroidPlatformStackWalk.h
@@ -23,8 +23,7 @@ struct CORE_API FAndroidPlatformStackWalk : public FGenericPlatformStackWalk
 
 	static void HandleBackTraceSignal(siginfo* Info, void* Context);
 
-	// called when android version information is set.
-	static void NotifyPlatformVersionInit();
+	static bool InitStackWalking();
 };
 
 typedef FAndroidPlatformStackWalk FPlatformStackWalk;
diff --git a/Engine/Source/Runtime/Launch/Private/Android/LaunchAndroid.cpp b/Engine/Source/Runtime/Launch/Private/Android/LaunchAndroid.cpp
--- a/Engine/Source/Runtime/Launch/Private/Android/LaunchAndroid.cpp
+++ b/Engine/Source/Runtime/Launch/Private/Android/LaunchAndroid.cpp
@@ -434,6 +434,8 @@ int32 AndroidMain(struct android_app* state)
 		ValidGamepadKeyCodes.Add(ValidGamepadKeyCodesList[i]);
 	}
 
+	FAndroidPlatformStackWalk::InitStackWalking();
+
 	// wait for java activity onCreate to finish
 	{
 		SCOPED_BOOT_TIMING("Wait for GResumeMainInit");
@@ -1624,7 +1626,6 @@ JNI_METHOD void Java_com_epicgames_unreal_GameActivity_nativeSetAndroidVersionIn
 	auto UEOSLanguage = FJavaHelper::FStringFromParam(jenv, osLanguage);
 
 	FAndroidMisc::SetVersionInfo(UEAndroidVersion, targetSDKversion, UEPhoneMake, UEPhoneModel, UEPhoneBuildNumber, UEOSLanguage);
-	FAndroidPlatformStackWalk::NotifyPlatformVersionInit();
 }
 
 //This function is declared in the Java-defined class, GameActivity.java: "public native void nativeOnInitialDownloadStarted();
