StringView: Added operator< matching the behavior of FString

diff --git a/Engine/Source/Runtime/Core/Public/Containers/StringView.h b/Engine/Source/Runtime/Core/Public/Containers/StringView.h
--- a/Engine/Source/Runtime/Core/Public/Containers/StringView.h
+++ b/Engine/Source/Runtime/Core/Public/Containers/StringView.h
@@ -605,6 +605,12 @@ inline bool operator!=(TStringView<CharType> Lhs, TStringView<CharType> Rhs)
 	return !(Lhs == Rhs);
 }
 
+template <typename CharType>
+inline bool operator<(TStringView<CharType> Lhs, TStringView<CharType> Rhs)
+{
+	return Lhs.Compare(Rhs, ESearchCase::IgnoreCase) < 0;
+}
+
 // Case-insensitive character range comparison operators
 
 template <typename CharType, typename CharRangeType>
@@ -633,6 +639,20 @@ inline auto operator!=(CharRangeType&& Lhs, TStringView<CharType> Rhs) -> declty
 	return !(Rhs == Forward<CharRangeType>(Lhs));
 }
 
+template <typename CharType, typename CharRangeType>
+inline auto operator<(TStringView<CharType> Lhs, CharRangeType&& Rhs)
+	-> decltype(Lhs.Compare(ImplicitConv<TStringView<CharType>>(Forward<CharRangeType>(Rhs)), ESearchCase::IgnoreCase) < 0)
+{
+	return Lhs.Compare(ImplicitConv<TStringView<CharType>>(Forward<CharRangeType>(Rhs)), ESearchCase::IgnoreCase) < 0;
+}
+
+template <typename CharType, typename CharRangeType>
+inline auto operator<(CharRangeType&& Lhs, TStringView<CharType> Rhs)
+	-> decltype(ImplicitConv<TStringView<CharType>>(Forward<CharRangeType>(Lhs)).Compare(Rhs, ESearchCase::IgnoreCase) < 0)
+{
+	return ImplicitConv<TStringView<CharType>>(Forward<CharRangeType>(Lhs)).Compare(Rhs, ESearchCase::IgnoreCase) < 0;
+}
+
 // Case-insensitive C-string comparison operators
 
 template <typename CharType>
