FEM - optimize Concurrent Linear Allocator TLS Cache

diff --git a/Engine/Source/Runtime/Core/Public/Experimental/ConcurrentLinearAllocator.h b/Engine/Source/Runtime/Core/Public/Experimental/ConcurrentLinearAllocator.h
--- a/Engine/Source/Runtime/Core/Public/Experimental/ConcurrentLinearAllocator.h
+++ b/Engine/Source/Runtime/Core/Public/Experimental/ConcurrentLinearAllocator.h
@@ -17,12 +17,25 @@
 template<uint32 BlockSize>
 class TBlockAllocationCache
 {
-	static constexpr uint32 NumRecycledSlots = 4;
+	struct FTLSCleanup
+	{
+		void* Block = nullptr;
 
-	FORCEINLINE static std::atomic<void*>& GetSlot(uint32 i)
+		~FTLSCleanup()
+		{
+			if(Block)
+			{
+				FPlatformMemory::BinnedFreeToOS(Block, BlockSize);
+			}
+		}
+	};
+
+	FORCEINLINE static void* SwapBlock(void* NewBlock)
 	{
-		static std::atomic<void*> Head[NumRecycledSlots];
-		return Head[i];
+		static thread_local FTLSCleanup Tls;
+		void* Ret = Tls.Block;
+		Tls.Block = NewBlock;
+		return Ret;
 	}
 
 public:
@@ -30,13 +43,10 @@ public:
 	{
 		if (Size == BlockSize)
 		{
-			for (uint32 i = 0; i < NumRecycledSlots; i++)
+			void* Pointer = SwapBlock(nullptr);
+			if (Pointer != nullptr)
 			{
-				void* Pointer = GetSlot(i).exchange(nullptr, std::memory_order_relaxed);
-				if (Pointer != nullptr)
-				{
-					return Pointer;
-				}	
+				return Pointer;
 			}
 		}
 		return FPlatformMemory::BinnedAllocFromOS(Size);
@@ -46,13 +56,10 @@ public:
 	{
 		if (Size == BlockSize)
 		{
-			for (uint32 i = 0; i < NumRecycledSlots; i++)
+			Pointer = SwapBlock(Pointer);
+			if (Pointer == nullptr)
 			{
-				Pointer = GetSlot(i).exchange(Pointer, std::memory_order_relaxed);
-				if (Pointer == nullptr)
-				{
-					return;
-				}	
+				return;
 			}
 		}
 		return FPlatformMemory::BinnedFreeToOS(Pointer, Size);
@@ -154,8 +161,13 @@ class TConcurrentLinearAllocator
 		}
 	};
 
-	FORCENOINLINE static void RegisterCleanup(FBlockHeader* Header)
+	FORCENOINLINE static void AllocateBlock(FBlockHeader*& Header)
 	{
+		static_assert(BlockAllocationTag::BlockSize >= sizeof(FBlockHeader) + sizeof(FAllocationHeader));
+		Header = new (BlockAllocationTag::MallocFunction(BlockAllocationTag::BlockSize)) FBlockHeader;
+		checkSlow(IsAligned(Header, alignof(FBlockHeader)));
+		ASAN_POISON_MEMORY_REGION( Header + 1, BlockAllocationTag::BlockSize - sizeof(FBlockHeader) );
+
 		static thread_local FTLSCleanup Cleanup;
 		new (&Cleanup) FTLSCleanup { Header };
 	}
@@ -176,11 +188,7 @@ public:
 		if (Header == nullptr)
 		{
 		AllocateNewBlock:
-			static_assert(BlockAllocationTag::BlockSize >= sizeof(FBlockHeader) + sizeof(FAllocationHeader));
-			Header = new (BlockAllocationTag::MallocFunction(BlockAllocationTag::BlockSize)) FBlockHeader;
-			checkSlow(IsAligned(Header, alignof(FBlockHeader)));
-			ASAN_POISON_MEMORY_REGION( Header + 1, BlockAllocationTag::BlockSize - sizeof(FBlockHeader) );
-			RegisterCleanup(Header);
+			AllocateBlock(Header);
 		}
 
 	AllocateNewItem:
