FEM - adding Fastpath option fo the Concurrent Linear Allocator when certain requirements are fulfilled

diff --git a/Engine/Source/Runtime/Core/Public/Experimental/ConcurrentLinearAllocator.h b/Engine/Source/Runtime/Core/Public/Experimental/ConcurrentLinearAllocator.h
--- a/Engine/Source/Runtime/Core/Public/Experimental/ConcurrentLinearAllocator.h
+++ b/Engine/Source/Runtime/Core/Public/Experimental/ConcurrentLinearAllocator.h
@@ -3,18 +3,57 @@
 #pragma once
 #include <atomic>
 #include "CoreMinimal.h"
+#include "HAL/MallocAnsi.h"
 #include "HAL/UnrealMemory.h"
 #include "HAL/LowLevelMemTracker.h"
 #include "Templates/UniquePtr.h"
 
 #if __has_include(<sanitizer/asan_interface.h>)
 #include <sanitizer/asan_interface.h>
+#if defined(__SANITIZE_ADDRESS__)
+#define IS_ASAN_ENABLED 1
+#elif __has_feature(address_sanitizer)
+#define IS_ASAN_ENABLED 1
+#else
+#define IS_ASAN_ENABLED 0
+#endif
 #else
 #define ASAN_POISON_MEMORY_REGION(addr, size) ((void)(addr), (void)(size))
 #define ASAN_UNPOISON_MEMORY_REGION(addr, size) ((void)(addr), (void)(size))
+#define IS_ASAN_ENABLED 0
 #endif
 
-template<uint32 BlockSize>
+struct FOsAllocator
+{
+	static constexpr bool SupportsAlignment = false;
+
+	FORCEINLINE static void* Malloc(SIZE_T Size, uint32 Alignment)
+	{
+		return FPlatformMemory::BinnedAllocFromOS(Size);
+	}
+
+	FORCEINLINE static void Free(void* Pointer, SIZE_T Size)
+	{
+		return FPlatformMemory::BinnedFreeToOS(Pointer, Size);
+	}
+};
+
+struct FAlignedAllocator
+{
+	static constexpr bool SupportsAlignment = true;
+
+	FORCEINLINE static void* Malloc(SIZE_T Size, uint32 Alignment)
+	{
+		return AnsiMalloc(Size, Alignment);
+	}
+
+	FORCEINLINE static void Free(void* Pointer, SIZE_T Size)
+	{
+		return AnsiFree(Pointer);
+	}
+};
+
+template<uint32 BlockSize, typename Allocator = FOsAllocator>
 class TBlockAllocationCache
 {
 	struct FTLSCleanup
@@ -25,7 +64,7 @@ class TBlockAllocationCache
 		{
 			if(Block)
 			{
-				FPlatformMemory::BinnedFreeToOS(Block, BlockSize);
+				Allocator::Free(Block, BlockSize);
 			}
 		}
 	};
@@ -39,7 +78,9 @@ class TBlockAllocationCache
 	}
 
 public:
-	FORCEINLINE static void* Malloc(SIZE_T Size)
+	static constexpr bool SupportsAlignment = Allocator::SupportsAlignment;
+
+	FORCEINLINE static void* Malloc(SIZE_T Size, uint32 Alignment)
 	{
 		if (Size == BlockSize)
 		{
@@ -49,7 +90,7 @@ public:
 				return Pointer;
 			}
 		}
-		return FPlatformMemory::BinnedAllocFromOS(Size);
+		return Allocator::Malloc(Size, Alignment);
 	}
 
 	FORCEINLINE static void Free(void* Pointer, SIZE_T Size)
@@ -62,7 +103,7 @@ public:
 				return;
 			}
 		}
-		return FPlatformMemory::BinnedFreeToOS(Pointer, Size);
+		return Allocator::Free(Pointer, Size);
 	}
 };
 
@@ -70,24 +111,29 @@ struct FDefaultBlockAllocationTag
 {
 	static constexpr uint32 BlockSize = 64 * 1024;		// Blocksize used to allocate from
 	static constexpr bool AllowOversizedBlocks = true;  // The allocator supports oversized Blocks and will store them in a seperate Block with counter 1
+	static constexpr bool RequiresAccurateSize = true;  // GetAllocationSize returning the accurate size of the allocation otherwise it could be relaxed to return the size to the end of the Block
+	static constexpr bool InlineBlockAllocation = false;  // Inline or Noinline the BlockAllocation which can have an impact on Performance
 	static constexpr const TCHAR* TagName = TEXT("DefaultLinear");
 
-	typedef void* (*MallocType)(SIZE_T Size);
-	typedef void (*FreeType)(void* Pointer, SIZE_T Size);
-	static constexpr MallocType MallocFunction = &FPlatformMemory::BinnedAllocFromOS;
-	static constexpr FreeType FreeFunction = &FPlatformMemory::BinnedFreeToOS;
+	using Allocator = FOsAllocator;
 };
 
 // This concurrent fast linear Allocator can be used for temporary allocations that are usually produced and consumed on different threads and within the lifetime of a frame.
 // Although the lifetime of any individual Allocation is not hard tied to a frame (tracking is done using the FBlockHeader::NumAllocations atomic variable) the Application will eventually run OOM if Allocations are not cleaned up in a timely manner.
+// THere is a fastpath version of the allocator that skips Allocationheaders by aligning the BlockHeader with the Blocksize so that it easily can be found by AligingDOwn the Adress of the Allocation itself. 
 // The Allocator works by allocating a larger block in TLS which has a Header at the front which contains the Atomic and all allocations are than allocated from this block:
 //
 // ------------------------------------------------------------------------------------------------------------------------------------------------------------------------------------------------------------
-// | FBlockHeader(atomic counter etc.) | Alignment Waste | FAllocationHeader(size) | Memory used for Allocation | Alignment Waste | FAllocationHeader(size) | Memory used for Allocation | FreeSpace ... 
+// | FBlockHeader(atomic counter etc.) | Alignment Waste | FAllocationHeader(size, optional) | Memory used for Allocation | Alignment Waste | FAllocationHeader(size, optional) | Memory used for Allocation | FreeSpace ... 
 // ------------------------------------------------------------------------------------------------------------------------------------------------------------------------------------------------------------
 template<typename BlockAllocationTag>
 class TConcurrentLinearAllocator
 {
+	static constexpr bool SupportsFastPath = ((BlockAllocationTag::BlockSize <= (64 * 1024)) || !USE_MALLOC_BINNED3) //MallocBinned3 only supports Alignment up to (64 * 1024)
+		&& FMath::IsPowerOfTwo(BlockAllocationTag::BlockSize) //Aligndown only works with Pow2
+		&& !IS_ASAN_ENABLED && !BlockAllocationTag::RequiresAccurateSize // Only enabled when not using ASAN or when the Size of an allocation is not required 
+		&& BlockAllocationTag::Allocator::SupportsAlignment; //The allocator needs to align the BlockHeader by Blocksize to find the BlockHeader
+
 	struct FBlockHeader;
 
 	class FAllocationHeader
@@ -120,7 +166,14 @@ class TConcurrentLinearAllocator
 
 	static FORCEINLINE FAllocationHeader* GetAllocationHeader(void* Pointer)
 	{
-		return reinterpret_cast<FAllocationHeader*>(Pointer) - 1;
+		if(SupportsFastPath)
+		{
+			return nullptr;
+		}
+		else
+		{
+			return reinterpret_cast<FAllocationHeader*>(Pointer) - 1;
+		}
 	}
 
 	struct FBlockHeader
@@ -155,7 +208,7 @@ class TConcurrentLinearAllocator
 					//if all allocations are already freed we can reuse the Block again
 					Header->~FBlockHeader();
 					ASAN_UNPOISON_MEMORY_REGION( Header, BlockAllocationTag::BlockSize );
-					BlockAllocationTag::FreeFunction(Header, BlockAllocationTag::BlockSize);				
+					BlockAllocationTag::Allocator::Free(Header, BlockAllocationTag::BlockSize);				
 				}
 			}
 		}
@@ -163,11 +216,16 @@ class TConcurrentLinearAllocator
 
 	FORCENOINLINE static void AllocateBlock(FBlockHeader*& Header)
 	{
-		static_assert(BlockAllocationTag::BlockSize >= sizeof(FBlockHeader) + sizeof(FAllocationHeader));
-		Header = new (BlockAllocationTag::MallocFunction(BlockAllocationTag::BlockSize)) FBlockHeader;
-		checkSlow(IsAligned(Header, alignof(FBlockHeader)));
-		ASAN_POISON_MEMORY_REGION( Header + 1, BlockAllocationTag::BlockSize - sizeof(FBlockHeader) );
-
+		if constexpr (!BlockAllocationTag::InlineBlockAllocation)
+		{
+			static_assert(BlockAllocationTag::BlockSize >= sizeof(FBlockHeader) + sizeof(FAllocationHeader));
+			Header = new (BlockAllocationTag::Allocator::Malloc(BlockAllocationTag::BlockSize, BlockAllocationTag::BlockSize)) FBlockHeader;
+			checkSlow(IsAligned(Header, alignof(FBlockHeader)));
+			if constexpr (!SupportsFastPath)
+			{
+				ASAN_POISON_MEMORY_REGION( Header + 1, BlockAllocationTag::BlockSize - sizeof(FBlockHeader) );
+			}
+		}
 		static thread_local FTLSCleanup Cleanup;
 		new (&Cleanup) FTLSCleanup { Header };
 	}
@@ -182,55 +240,102 @@ public:
 	static FORCEINLINE_DEBUGGABLE void* Malloc(SIZE_T Size, uint32 Alignment)
 	{
 		checkSlow(Alignment >= 1 && FMath::IsPowerOfTwo(Alignment));
-		Alignment = (Alignment < alignof(FAllocationHeader)) ? alignof(FAllocationHeader) : Alignment;
+		if constexpr (!SupportsFastPath)
+		{
+			Alignment = (Alignment < alignof(FAllocationHeader)) ? alignof(FAllocationHeader) : Alignment;
+		}
 
 		static thread_local FBlockHeader* Header;
 		if (Header == nullptr)
 		{
 		AllocateNewBlock:
+			if constexpr (BlockAllocationTag::InlineBlockAllocation)
+			{
+				static_assert(BlockAllocationTag::BlockSize >= sizeof(FBlockHeader) + sizeof(FAllocationHeader));
+				Header = new (BlockAllocationTag::Allocator::Malloc(BlockAllocationTag::BlockSize, BlockAllocationTag::BlockSize)) FBlockHeader;
+				checkSlow(IsAligned(Header, alignof(FBlockHeader)));
+				if constexpr (!SupportsFastPath)
+				{
+					ASAN_POISON_MEMORY_REGION( Header + 1, BlockAllocationTag::BlockSize - sizeof(FBlockHeader) );
+				}
+			}
 			AllocateBlock(Header);
 		}
 
 	AllocateNewItem:
-		uintptr_t AlignedOffset = Align(Header->NextAllocationPtr, Alignment);
-		if (AlignedOffset + Size <= uintptr_t(Header) + BlockAllocationTag::BlockSize)
+		if constexpr (SupportsFastPath)
 		{
-			Header->NextAllocationPtr = AlignedOffset + Size + sizeof(FAllocationHeader);
-			Header->Num++;
-
-			FAllocationHeader* AllocationHeader = reinterpret_cast<FAllocationHeader*>(AlignedOffset) - 1;
-			ASAN_UNPOISON_MEMORY_REGION( AllocationHeader, sizeof(FAllocationHeader) + Size );
-			new (AllocationHeader) FAllocationHeader(Header, Size);
-			ASAN_POISON_MEMORY_REGION( AllocationHeader, sizeof(FAllocationHeader) );
+			uintptr_t AlignedOffset = Align(Header->NextAllocationPtr, Alignment);
+			if (AlignedOffset + Size <= uintptr_t(Header) + BlockAllocationTag::BlockSize)
+			{
+				Header->NextAllocationPtr = AlignedOffset + Size;
+				Header->Num++;
 
-			LLM_IF_ENABLED( FLowLevelMemTracker::Get().OnLowLevelAlloc(ELLMTracker::Default, AllocationHeader, Size + sizeof(FAllocationHeader), ELLMTag::LinearAllocator) );
-			return reinterpret_cast<void*>(AlignedOffset);
-		}
+				LLM_IF_ENABLED( FLowLevelMemTracker::Get().OnLowLevelAlloc(ELLMTracker::Default, reinterpret_cast<void*>(AlignedOffset), Size, ELLMTag::LinearAllocator) );
+				return reinterpret_cast<void*>(AlignedOffset);
+			}
 
-		// The cold path of the code starts here
-		constexpr SIZE_T HeaderSize = sizeof(FBlockHeader) + sizeof(FAllocationHeader);
-		if constexpr (BlockAllocationTag::AllowOversizedBlocks)
-		{
-			//support for oversized Blocks
-			if (HeaderSize + Size + Alignment > BlockAllocationTag::BlockSize)
+			// The cold path of the code starts here
+			constexpr SIZE_T HeaderSize = sizeof(FBlockHeader);
+			if constexpr (BlockAllocationTag::AllowOversizedBlocks)
 			{
-				FBlockHeader* LargeHeader = new (BlockAllocationTag::MallocFunction(HeaderSize + Size + Alignment)) FBlockHeader;
-				checkSlow(IsAligned(LargeHeader, alignof(FBlockHeader)));
+				//support for oversized Blocks
+				if (HeaderSize + Size + Alignment > BlockAllocationTag::BlockSize)
+				{
+					FBlockHeader* LargeHeader = new (BlockAllocationTag::Allocator::Malloc(HeaderSize + Size + Alignment, BlockAllocationTag::BlockSize)) FBlockHeader;
+					checkSlow(IsAligned(LargeHeader, alignof(FBlockHeader)));
 
-				uintptr_t LargeAlignedOffset = Align(LargeHeader->NextAllocationPtr, Alignment);
-				LargeHeader->NextAllocationPtr = uintptr_t(LargeHeader) + HeaderSize + Size + Alignment;
-				LargeHeader->NumAllocations.store(1, std::memory_order_release);
+					uintptr_t LargeAlignedOffset = Align(LargeHeader->NextAllocationPtr, Alignment);
+					LargeHeader->NextAllocationPtr = uintptr_t(LargeHeader) + HeaderSize + Size + Alignment;
+					LargeHeader->NumAllocations.store(1, std::memory_order_release);
 
-				checkSlow(LargeAlignedOffset + Size <= LargeHeader->NextAllocationPtr);
-				FAllocationHeader* AllocationHeader = new (reinterpret_cast<FAllocationHeader*>(LargeAlignedOffset) - 1) FAllocationHeader(LargeHeader, Size);
-				ASAN_POISON_MEMORY_REGION( AllocationHeader, sizeof(FAllocationHeader) );
+					checkSlow(LargeAlignedOffset + Size <= LargeHeader->NextAllocationPtr);
 
-				LLM_IF_ENABLED( FLowLevelMemTracker::Get().OnLowLevelAlloc(ELLMTracker::Default, AllocationHeader, Size + sizeof(FAllocationHeader), ELLMTag::LinearAllocator) );
-				return reinterpret_cast<void*>(LargeAlignedOffset);
+					LLM_IF_ENABLED( FLowLevelMemTracker::Get().OnLowLevelAlloc(ELLMTracker::Default, reinterpret_cast<void*>(LargeAlignedOffset), Size, ELLMTag::LinearAllocator) );
+					return reinterpret_cast<void*>(LargeAlignedOffset);
+				}
 			}
+			check(HeaderSize + Size + Alignment <= BlockAllocationTag::BlockSize);
 		}
 		else
 		{
+			uintptr_t AlignedOffset = Align(Header->NextAllocationPtr, Alignment);
+			if (AlignedOffset + Size <= uintptr_t(Header) + BlockAllocationTag::BlockSize)
+			{
+				Header->NextAllocationPtr = AlignedOffset + Size + sizeof(FAllocationHeader);
+				Header->Num++;
+
+				FAllocationHeader* AllocationHeader = reinterpret_cast<FAllocationHeader*>(AlignedOffset) - 1;
+				ASAN_UNPOISON_MEMORY_REGION( AllocationHeader, sizeof(FAllocationHeader) + Size );
+				new (AllocationHeader) FAllocationHeader(Header, Size);
+				ASAN_POISON_MEMORY_REGION( AllocationHeader, sizeof(FAllocationHeader) );
+
+				LLM_IF_ENABLED( FLowLevelMemTracker::Get().OnLowLevelAlloc(ELLMTracker::Default, AllocationHeader, Size + sizeof(FAllocationHeader), ELLMTag::LinearAllocator) );
+				return reinterpret_cast<void*>(AlignedOffset);
+			}
+
+			// The cold path of the code starts here
+			constexpr SIZE_T HeaderSize = sizeof(FBlockHeader) + sizeof(FAllocationHeader);
+			if constexpr (BlockAllocationTag::AllowOversizedBlocks)
+			{
+				//support for oversized Blocks
+				if (HeaderSize + Size + Alignment > BlockAllocationTag::BlockSize)
+				{
+					FBlockHeader* LargeHeader = new (BlockAllocationTag::Allocator::Malloc(HeaderSize + Size + Alignment, BlockAllocationTag::BlockSize)) FBlockHeader;
+					checkSlow(IsAligned(LargeHeader, alignof(FBlockHeader)));
+
+					uintptr_t LargeAlignedOffset = Align(LargeHeader->NextAllocationPtr, Alignment);
+					LargeHeader->NextAllocationPtr = uintptr_t(LargeHeader) + HeaderSize + Size + Alignment;
+					LargeHeader->NumAllocations.store(1, std::memory_order_release);
+
+					checkSlow(LargeAlignedOffset + Size <= LargeHeader->NextAllocationPtr);
+					FAllocationHeader* AllocationHeader = new (reinterpret_cast<FAllocationHeader*>(LargeAlignedOffset) - 1) FAllocationHeader(LargeHeader, Size);
+					ASAN_POISON_MEMORY_REGION( AllocationHeader, sizeof(FAllocationHeader) );
+
+					LLM_IF_ENABLED( FLowLevelMemTracker::Get().OnLowLevelAlloc(ELLMTracker::Default, AllocationHeader, Size + sizeof(FAllocationHeader), ELLMTag::LinearAllocator) );
+					return reinterpret_cast<void*>(LargeAlignedOffset);
+				}
+			}
 			check(HeaderSize + Size + Alignment <= BlockAllocationTag::BlockSize);
 		}
 
@@ -256,18 +361,32 @@ public:
 	{
 		if(Pointer != nullptr)
 		{
-			FAllocationHeader* AllocationHeader = GetAllocationHeader(Pointer);
-			ASAN_UNPOISON_MEMORY_REGION( AllocationHeader, sizeof(FAllocationHeader) );
-			LLM_IF_ENABLED(FLowLevelMemTracker::Get().OnLowLevelFree(ELLMTracker::Default, AllocationHeader));
-			FBlockHeader* Header = AllocationHeader->GetBlockHeader();
-			ASAN_POISON_MEMORY_REGION( AllocationHeader, sizeof(FAllocationHeader) + AllocationHeader->GetAllocationSize() );
+			if constexpr (SupportsFastPath)
+			{
+				FBlockHeader* Header = reinterpret_cast<FBlockHeader*>(AlignDown(Pointer, BlockAllocationTag::BlockSize));
 
-			// this deletes complete blocks when the last allocation is freed
-			if (Header->NumAllocations.fetch_sub(1, std::memory_order_acq_rel) == 1)
+				// this deletes complete blocks when the last allocation is freed
+				if (Header->NumAllocations.fetch_sub(1, std::memory_order_acq_rel) == 1)
+				{
+					Header->~FBlockHeader();
+					BlockAllocationTag::Allocator::Free(Header, Header->NextAllocationPtr - uintptr_t(Header));
+				}
+			}
+			else
 			{
-				Header->~FBlockHeader();
-				ASAN_UNPOISON_MEMORY_REGION( Header, Header->NextAllocationPtr - uintptr_t(Header) );
-				BlockAllocationTag::FreeFunction(Header, Header->NextAllocationPtr - uintptr_t(Header));
+				FAllocationHeader* AllocationHeader = GetAllocationHeader(Pointer);
+				ASAN_UNPOISON_MEMORY_REGION( AllocationHeader, sizeof(FAllocationHeader) );
+				LLM_IF_ENABLED(FLowLevelMemTracker::Get().OnLowLevelFree(ELLMTracker::Default, AllocationHeader));
+				FBlockHeader* Header = AllocationHeader->GetBlockHeader();
+				ASAN_POISON_MEMORY_REGION( AllocationHeader, sizeof(FAllocationHeader) + AllocationHeader->GetAllocationSize() );
+
+				// this deletes complete blocks when the last allocation is freed
+				if (Header->NumAllocations.fetch_sub(1, std::memory_order_acq_rel) == 1)
+				{
+					Header->~FBlockHeader();
+					ASAN_UNPOISON_MEMORY_REGION( Header, Header->NextAllocationPtr - uintptr_t(Header) );
+					BlockAllocationTag::Allocator::Free(Header, Header->NextAllocationPtr - uintptr_t(Header));
+				}
 			}
 		}
 	}
@@ -276,11 +395,18 @@ public:
 	{
 		if (Pointer)
 		{
-			FAllocationHeader* AllocationHeader = GetAllocationHeader(Pointer);
-			ASAN_UNPOISON_MEMORY_REGION( AllocationHeader, sizeof(FAllocationHeader) );
-			SIZE_T Size = AllocationHeader->GetAllocationSize();
-			ASAN_POISON_MEMORY_REGION( AllocationHeader, sizeof(FAllocationHeader) );
-			return Size;
+			if constexpr (SupportsFastPath)
+			{
+				return Align(uintptr_t(Pointer), BlockAllocationTag::BlockSize) - uintptr_t(Pointer);
+			}
+			else
+			{
+				FAllocationHeader* AllocationHeader = GetAllocationHeader(Pointer);
+				ASAN_UNPOISON_MEMORY_REGION( AllocationHeader, sizeof(FAllocationHeader) );
+				SIZE_T Size = AllocationHeader->GetAllocationSize();
+				ASAN_POISON_MEMORY_REGION( AllocationHeader, sizeof(FAllocationHeader) );
+				return Size;
+			}
 		}
 		return 0;
 	}
@@ -317,7 +443,7 @@ public:
 	{
 		return TConcurrentLinearAllocator<BlockAllocationTag>::template Malloc<alignof(ObjectType)>(Size);
 	}
-	
+
 	FORCEINLINE_DEBUGGABLE static void* operator new(size_t Size, std::align_val_t Align)
 	{
 		checkSlow(size_t(Align) == alignof(ObjectType));
