Fix for Expose_TFormatSpecifier to allow it to work properly in UTF-8 mode.

diff --git a/Engine/Source/Runtime/Core/Public/Templates/UnrealTypeTraits.h b/Engine/Source/Runtime/Core/Public/Templates/UnrealTypeTraits.h
--- a/Engine/Source/Runtime/Core/Public/Templates/UnrealTypeTraits.h
+++ b/Engine/Source/Runtime/Core/Public/Templates/UnrealTypeTraits.h
@@ -114,10 +114,9 @@ struct TFormatSpecifier
 template<> \
 struct TFormatSpecifier<type> \
 {  \
-	FORCEINLINE static const TCHAR (&GetFormatSpecifier())[5] \
+	FORCEINLINE static decltype(auto) GetFormatSpecifier() \
 	{ \
-		static const TCHAR Spec[5] = TEXT(format); \
-		return Spec; \
+		return TEXT(format); \
 	} \
 };
 
