Fix non-unity builds due to missing <Misc/Guid.h> include.

diff --git a/Engine/Source/Runtime/Core/Public/Logging/TokenizedMessage.h b/Engine/Source/Runtime/Core/Public/Logging/TokenizedMessage.h
--- a/Engine/Source/Runtime/Core/Public/Logging/TokenizedMessage.h
+++ b/Engine/Source/Runtime/Core/Public/Logging/TokenizedMessage.h
@@ -10,6 +10,7 @@
 #include "Delegates/Delegate.h"
 #include "Internationalization/Text.h"
 #include "Misc/Attribute.h"
+#include "Misc/Guid.h"
 
 /** The severity of the message type */
 namespace EMessageSeverity
