FEM - validate propper use of TConcurrentLinearObject

diff --git a/Engine/Source/Runtime/Core/Public/Experimental/ConcurrentLinearAllocator.h b/Engine/Source/Runtime/Core/Public/Experimental/ConcurrentLinearAllocator.h
--- a/Engine/Source/Runtime/Core/Public/Experimental/ConcurrentLinearAllocator.h
+++ b/Engine/Source/Runtime/Core/Public/Experimental/ConcurrentLinearAllocator.h
@@ -7,6 +7,7 @@
 #include "HAL/UnrealMemory.h"
 #include "HAL/LowLevelMemTracker.h"
 #include "Templates/UniquePtr.h"
+#include "Templates/UnrealTypeTraits.h"
 
 #if __has_include(<sanitizer/asan_interface.h>)
 #include <sanitizer/asan_interface.h>
@@ -436,6 +437,7 @@ class TConcurrentLinearObject
 public:
 	FORCEINLINE_DEBUGGABLE static void* operator new(size_t Size)
 	{
+		static_assert(TIsDerivedFrom<ObjectType, TConcurrentLinearObject<ObjectType, BlockAllocationTag>>::Value, "TConcurrentLinearObject must be base of it's ObjectType (see CRTP)");
 		return TConcurrentLinearAllocator<BlockAllocationTag>::template Malloc<alignof(ObjectType)>(Size);
 	}
 
