Define log categories even in NO_LOGGING configurations.

diff --git a/Engine/Source/Runtime/Core/Public/Logging/LogMacros.h b/Engine/Source/Runtime/Core/Public/Logging/LogMacros.h
--- a/Engine/Source/Runtime/Core/Public/Logging/LogMacros.h
+++ b/Engine/Source/Runtime/Core/Public/Logging/LogMacros.h
@@ -132,7 +132,7 @@ private:
 	#define UE_GET_LOG_VERBOSITY(...)		(ELogVerbosity::NoLogging)
 	#define UE_SET_LOG_VERBOSITY(...)
 	#define DECLARE_LOG_CATEGORY_EXTERN(CategoryName, DefaultVerbosity, CompileTimeVerbosity) extern FNoLoggingCategory CategoryName;
-	#define DEFINE_LOG_CATEGORY(...)
+	#define DEFINE_LOG_CATEGORY(CategoryName, ...) FNoLoggingCategory CategoryName;
 	#define DEFINE_LOG_CATEGORY_STATIC(...)
 	#define DECLARE_LOG_CATEGORY_CLASS(...)
 	#define DEFINE_LOG_CATEGORY_CLASS(...)
