UTF8CHAR overloads for FCrc::Strihash_DEPRECATED.

diff --git a/Engine/Source/Runtime/Core/Public/Misc/Crc.h b/Engine/Source/Runtime/Core/Public/Misc/Crc.h
--- a/Engine/Source/Runtime/Core/Public/Misc/Crc.h
+++ b/Engine/Source/Runtime/Core/Public/Misc/Crc.h
@@ -196,31 +196,42 @@ inline uint32 FCrc::Strihash_DEPRECATED(const int32 DataLen, const WIDECHAR* Dat
 template <>
 inline uint32 FCrc::Strihash_DEPRECATED(const UTF8CHAR* Data)
 {
-	// make sure table is initialized
-	check(CRCTable_DEPRECATED[1] != 0);
+	// We can't utilize StringConv.h here due to circular includes, so do the conversion manually
 
-	uint32 Hash=0;
-	while( *Data )
-	{
-		UTF8CHAR Ch = TChar<UTF8CHAR>::ToUpper(*Data++);
-		uint8 B  = Ch;
-		Hash = ((Hash >> 8) & 0x00FFFFFF) ^ CRCTable_DEPRECATED[(Hash ^ B) & 0x000000FF];
-	}
-	return Hash;
+	int32 Len          = FPlatformString::Strlen(Data);
+	int32 ConvertedLen = FPlatformString::ConvertedLength<WIDECHAR>(Data, Len);
+
+	WIDECHAR* Temp = new WIDECHAR[ConvertedLen];
+
+	WIDECHAR* TempEnd = FPlatformString::Convert(Temp, ConvertedLen, Data, Len);
+	checkf(TempEnd, TEXT("String conversion unsuccessful"));
+
+	// This doesn't work for strings containing characters outside the BMP, but
+	// then neither does the WIDECHAR overload.
+	uint32 Result = FCrc::Strihash_DEPRECATED(ConvertedLen, Temp);
+
+	delete [] Temp;
+
+	return Result;
 }
 
 template <>
 inline uint32 FCrc::Strihash_DEPRECATED(const int32 DataLen, const UTF8CHAR* Data)
 {
-	// make sure table is initialized
-	check(CRCTable_DEPRECATED[1] != 0);
+	// We can't utilize StringConv.h here due to circular includes, so do the conversion manually
 
-	uint32 Hash=0;
-	for (int32 Idx = 0; Idx < DataLen; ++Idx)
-	{
-		UTF8CHAR Ch = TChar<UTF8CHAR>::ToUpper(*Data++);
-		uint8 B  = Ch;
-		Hash = ((Hash >> 8) & 0x00FFFFFF) ^ CRCTable_DEPRECATED[(Hash ^ B) & 0x000000FF];
-	}
-	return Hash;
+	int32 ConvertedLen = FPlatformString::ConvertedLength<WIDECHAR>(Data, DataLen);
+
+	WIDECHAR* Temp = new WIDECHAR[ConvertedLen];
+
+	WIDECHAR* TempEnd = FPlatformString::Convert(Temp, ConvertedLen, Data, DataLen);
+	checkf(TempEnd, TEXT("String conversion unsuccessful"));
+
+	// This doesn't work for strings containing characters outside the BMP, but
+	// then neither does the WIDECHAR overload.
+	uint32 Result = FCrc::Strihash_DEPRECATED(ConvertedLen, Temp);
+
+	delete [] Temp;
+
+	return Result;
 }
