[Insights] Fix copy paste and build error

 * An extra '&' token was removed on the type alias in the no logging case.
 * In the case where trace is enabled, but logging is disabled, made a identical solution in TraceAuxImpl.

diff --git a/Engine/Source/Runtime/Core/Private/ProfilingDebugging/TraceAuxiliary.cpp b/Engine/Source/Runtime/Core/Private/ProfilingDebugging/TraceAuxiliary.cpp
--- a/Engine/Source/Runtime/Core/Private/ProfilingDebugging/TraceAuxiliary.cpp
+++ b/Engine/Source/Runtime/Core/Private/ProfilingDebugging/TraceAuxiliary.cpp
@@ -84,7 +84,7 @@ public:
 	void					ResetCommandlineChannels();
 	void					EnableChannels(const TCHAR* ChannelList);
 	void					DisableChannels(const TCHAR* ChannelList);
-	bool					Connect(ETraceConnectType Type, const TCHAR* Parameter, const FLogCategoryBase& LogCategory);
+	bool					Connect(ETraceConnectType Type, const TCHAR* Parameter, const FTraceAuxiliary::FLogCategoryAlias& LogCategory);
 	bool					Stop();
 	void					ResumeChannels();
 	void					PauseChannels();
@@ -114,8 +114,8 @@ private:
 	static uint32			HashChannelName(const TCHAR* Name);
 	bool					EnableChannel(const TCHAR* Channel);
 	void					DisableChannel(const TCHAR* Channel);
-	bool					SendToHost(const TCHAR* Host, const FLogCategoryBase& LogCategory);
-	bool					WriteToFile(const TCHAR* Path, const FLogCategoryBase& LogCategory);
+	bool					SendToHost(const TCHAR* Host, const FTraceAuxiliary::FLogCategoryAlias& LogCategory);
+	bool					WriteToFile(const TCHAR* Path, const FTraceAuxiliary::FLogCategoryAlias& LogCategory);
 
 	TMap<uint32, FChannel>	CommandlineChannels;
 	FString					TraceDest;
@@ -237,7 +237,7 @@ void FTraceAuxiliaryImpl::RemoveChannel(const TCHAR* Name)
 }
 
 ////////////////////////////////////////////////////////////////////////////////
-bool FTraceAuxiliaryImpl::Connect(ETraceConnectType Type, const TCHAR* Parameter, const FLogCategoryBase& LogCategory)
+bool FTraceAuxiliaryImpl::Connect(ETraceConnectType Type, const TCHAR* Parameter, const FTraceAuxiliary::FLogCategoryAlias& LogCategory)
 {
 	// Connect/write to file. But only if we're not already sending/writing
 	bool bConnected = UE::Trace::IsTracing();
@@ -385,7 +385,7 @@ void FTraceAuxiliaryImpl::SetTruncateFile(bool bNewTruncateFileState)
 }
 
 ////////////////////////////////////////////////////////////////////////////////
-bool FTraceAuxiliaryImpl::SendToHost(const TCHAR* Host, const FLogCategoryBase& LogCategory)
+bool FTraceAuxiliaryImpl::SendToHost(const TCHAR* Host, const FTraceAuxiliary::FLogCategoryAlias& LogCategory)
 {
 	if (!UE::Trace::SendTo(Host))
 	{
@@ -398,7 +398,7 @@ bool FTraceAuxiliaryImpl::SendToHost(const TCHAR* Host, const FLogCategoryBase&
 }
 
 ////////////////////////////////////////////////////////////////////////////////
-bool FTraceAuxiliaryImpl::WriteToFile(const TCHAR* Path, const FLogCategoryBase& LogCategory)
+bool FTraceAuxiliaryImpl::WriteToFile(const TCHAR* Path, const FTraceAuxiliary::FLogCategoryAlias& LogCategory)
 {
 	// Default file name functor 
 	auto GetDefaultName = [] {return FDateTime::Now().ToString(TEXT("%Y%m%d_%H%M%S.utrace"));};
diff --git a/Engine/Source/Runtime/Core/Public/ProfilingDebugging/TraceAuxiliary.h b/Engine/Source/Runtime/Core/Public/ProfilingDebugging/TraceAuxiliary.h
--- a/Engine/Source/Runtime/Core/Public/ProfilingDebugging/TraceAuxiliary.h
+++ b/Engine/Source/Runtime/Core/Public/ProfilingDebugging/TraceAuxiliary.h
@@ -16,7 +16,7 @@ public:
 	// In no logging configurations all log categories are of type FNoLoggingCategory, which has no relation with
 	// FLogCategoryBase. In order to not need to conditionally set the argument alias the type here.
 #if NO_LOGGING
-	typedef FNoLoggingCategory& FLogCategoryAlias;
+	typedef FNoLoggingCategory FLogCategoryAlias;
 #else
 	typedef FLogCategoryBase FLogCategoryAlias;
 #endif
