REL - trying to add an mfence for Hazardpointer registration to sync the atomic in the HarzardList with the Atomic of the Datastructure that uses Hazardpointers.
this has not been verified to fix issues but as I were not able to reproduce any issues this is worth a try.

diff --git a/Engine/Source/Runtime/Core/Public/Experimental/Containers/HazardPointer.h b/Engine/Source/Runtime/Core/Public/Experimental/Containers/HazardPointer.h
--- a/Engine/Source/Runtime/Core/Public/Experimental/Containers/HazardPointer.h
+++ b/Engine/Source/Runtime/Core/Public/Experimental/Containers/HazardPointer.h
@@ -85,15 +85,17 @@ class FHazardPointerCollection
 		}
 
 		//assign hazard pointer once acquired
-		inline void SetHazard(void* InHazard)
+		UE_NODISCARD inline void* SetHazard(void* InHazard)
 		{
-			Hazard.store(reinterpret_cast<uintptr_t>(InHazard), std::memory_order_release);
+			Hazard.store(reinterpret_cast<uintptr_t>(InHazard), std::memory_order_relaxed);
+			std::atomic_thread_fence(std::memory_order_seq_cst);
+			return reinterpret_cast<void*>(Hazard.load(std::memory_order_relaxed));
 		}
 
 		//this thread wants to re-use the slot but does not want to hold onto the pointer
 		inline void Retire()
 		{
-			Hazard.store(0, std::memory_order_relaxed);
+			Hazard.store(0, std::memory_order_release);
 		}
 
 		//good to reuse by another thread
@@ -300,8 +302,7 @@ public:
 		H* HazardPointer;
 		do
 		{
-			Record->SetHazard(Hazard->load(std::memory_order_relaxed));
-			HazardPointer = (H*)Record->GetHazard();
+			HazardPointer = (H*)Record->SetHazard(Hazard->load(std::memory_order_relaxed));
 		} while (HazardPointer != Hazard->load(std::memory_order_relaxed));
 		return HazardPointer;
 	}
