REL - readding release aquire semantics for ARM

diff --git a/Engine/Source/Runtime/Core/Public/Experimental/Containers/HazardPointer.h b/Engine/Source/Runtime/Core/Public/Experimental/Containers/HazardPointer.h
--- a/Engine/Source/Runtime/Core/Public/Experimental/Containers/HazardPointer.h
+++ b/Engine/Source/Runtime/Core/Public/Experimental/Containers/HazardPointer.h
@@ -87,9 +87,9 @@ class FHazardPointerCollection
 		//assign hazard pointer once acquired
 		UE_NODISCARD inline void* SetHazard(void* InHazard)
 		{
-			Hazard.store(reinterpret_cast<uintptr_t>(InHazard), std::memory_order_relaxed);
+			Hazard.store(reinterpret_cast<uintptr_t>(InHazard), std::memory_order_release);
 			std::atomic_thread_fence(std::memory_order_seq_cst);
-			return reinterpret_cast<void*>(Hazard.load(std::memory_order_relaxed));
+			return reinterpret_cast<void*>(Hazard.load(std::memory_order_acquire));
 		}
 
 		//this thread wants to re-use the slot but does not want to hold onto the pointer
