Allow FBoxSphereBounds to use different types for position/extent
FPrimitiveSceneInfoCompact uses float extent for bounds

diff --git a/Engine/Source/Runtime/Core/Public/CoreFwd.h b/Engine/Source/Runtime/Core/Public/CoreFwd.h
--- a/Engine/Source/Runtime/Core/Public/CoreFwd.h
+++ b/Engine/Source/Runtime/Core/Public/CoreFwd.h
@@ -29,10 +29,30 @@ UE_DECLARE_LWC_TYPE(Transform, 3);
 UE_DECLARE_LWC_TYPE(Vector2,, FVector2D);
 UE_DECLARE_LWC_TYPE(Vector, 3);
 UE_DECLARE_LWC_TYPE(Vector4);
-UE_DECLARE_LWC_TYPE(BoxSphereBounds, 3);
 struct FIntPoint;
 struct FIntRect;
 
+namespace UE
+{
+namespace Math
+{
+template<typename T, typename TExtent = T>
+struct TBoxSphereBounds;
+}
+}
+using FBoxSphereBounds3f = UE::Math::TBoxSphereBounds<float, float>;
+using FBoxSphereBounds3d = UE::Math::TBoxSphereBounds<double, double>;
+// FCompactBoxSphereBounds always stores float extents
+using FCompactBoxSphereBounds3d = UE::Math::TBoxSphereBounds<double, float>;
+
+#if UE_LARGE_WORLD_COORDINATES_DISABLED
+using FBoxSphereBounds = FBoxSphereBounds3f;
+using FCompactBoxSphereBounds = FBoxSphereBounds3f;
+#else
+using FBoxSphereBounds = FBoxSphereBounds3d;
+using FCompactBoxSphereBounds = FCompactBoxSphereBounds3d;
+#endif
+
 // Misc
 struct FResourceSizeEx;
 class IConsoleVariable;
diff --git a/Engine/Source/Runtime/Core/Public/Math/BoxSphereBounds.h b/Engine/Source/Runtime/Core/Public/Math/BoxSphereBounds.h
--- a/Engine/Source/Runtime/Core/Public/Math/BoxSphereBounds.h
+++ b/Engine/Source/Runtime/Core/Public/Math/BoxSphereBounds.h
@@ -19,7 +19,7 @@ namespace UE
 namespace Math
 {
 
-template<typename T>
+template<typename T, typename TExtent>
 struct TBoxSphereBounds
 {
 	using FReal = T;
@@ -28,10 +28,10 @@ struct TBoxSphereBounds
 	TVector<T>	Origin;
 
 	/** Holds the extent of the bounding box. */
-	TVector<T> BoxExtent;
+	TVector<TExtent> BoxExtent;
 
 	/** Holds the radius of the bounding sphere. */
-	T SphereRadius;
+	TExtent SphereRadius;
 
 public:
 
@@ -43,7 +43,7 @@ public:
 	 *
 	 * @param EForceInit Force Init Enum.
 	 */
-	explicit FORCEINLINE TBoxSphereBounds( EForceInit ) 
+	explicit FORCEINLINE TBoxSphereBounds( EForceInit )
 		: Origin(ForceInit)
 		, BoxExtent(ForceInit)
 		, SphereRadius(0)
@@ -58,7 +58,7 @@ public:
 	 * @param InBoxExtent half size of box.
 	 * @param InSphereRadius radius of the sphere.
 	 */
-	TBoxSphereBounds( const TVector<T>& InOrigin, const TVector<T>& InBoxExtent, T InSphereRadius )
+	TBoxSphereBounds( const TVector<T>& InOrigin, const TVector<TExtent>& InBoxExtent, TExtent InSphereRadius )
 		: Origin(InOrigin)
 		, BoxExtent(InBoxExtent)
 		, SphereRadius(InSphereRadius)
@@ -74,8 +74,11 @@ public:
 	 */
 	TBoxSphereBounds( const TBox<T>& Box, const TSphere<T>& Sphere )
 	{
-		Box.GetCenterAndExtents(Origin,BoxExtent);
-		SphereRadius = FMath::Min(BoxExtent.Size(), (Sphere.Center - Origin).Size() + Sphere.W);
+		TVector<T> LocalExtent;
+		Box.GetCenterAndExtents(Origin, LocalExtent);
+		BoxExtent = TVector<TExtent>(LocalExtent);
+
+		SphereRadius = (TExtent)FMath::Min(LocalExtent.Size(), (Sphere.Center - Origin).Size() + Sphere.W);
 
 		DiagnosticCheckNaN();
 	}
@@ -89,7 +92,10 @@ public:
 	 */
 	TBoxSphereBounds( const TBox<T>& Box )
 	{
-		Box.GetCenterAndExtents(Origin,BoxExtent);
+		TVector<T> LocalExtent;
+		Box.GetCenterAndExtents(Origin, LocalExtent);
+		BoxExtent = TVector<TExtent>(LocalExtent);
+
 		SphereRadius = BoxExtent.Size();
 
 		DiagnosticCheckNaN();
@@ -101,8 +107,8 @@ public:
 	TBoxSphereBounds( const TSphere<T>& Sphere )
 	{
 		Origin = Sphere.Center;
-		BoxExtent = TVector<T>(Sphere.W);
 		SphereRadius = Sphere.W;
+		BoxExtent = TVector<TExtent>(SphereRadius);
 
 		DiagnosticCheckNaN();
 	}
@@ -118,8 +124,8 @@ public:
 	TBoxSphereBounds( const TVector<T>* Points, uint32 NumPoints );
 
 	// Conversion to other type.
-	template<typename FArg, TEMPLATE_REQUIRES(!TIsSame<T, FArg>::Value)>
-	explicit TBoxSphereBounds(const TBoxSphereBounds<FArg>& From) : TBoxSphereBounds<T>(TVector<T>(From.Origin), TVector<T>(From.BoxExtent), (T)From.SphereRadius) {}
+	template<typename TFrom, typename TExtentFrom, TEMPLATE_REQUIRES(!TAnd<TIsSame<T, TFrom>, TIsSame<TExtent, TExtentFrom>>::Value)>
+	explicit TBoxSphereBounds(const TBoxSphereBounds<TFrom, TExtentFrom>& From) : TBoxSphereBounds<T, TExtent>(TVector<T>(From.Origin), TVector<TExtent>(From.BoxExtent), (TExtent)From.SphereRadius) {}
 
 public:
 
@@ -132,7 +138,7 @@ public:
 	 * @param Other The other bounding volume.
 	 * @return The combined bounding volume.
 	 */
-	FORCEINLINE TBoxSphereBounds<T> operator+( const TBoxSphereBounds<T>& Other ) const;
+	FORCEINLINE TBoxSphereBounds<T, TExtent> operator+( const TBoxSphereBounds<T, TExtent>& Other ) const;
 
 	/**
 	 * Compare bounding volume this and Other.
@@ -140,7 +146,7 @@ public:
 	 * @param Other The other bounding volume.
 	 * @return true of they match.
 	 */
-	FORCEINLINE bool operator==(const TBoxSphereBounds<T>& Other) const;
+	FORCEINLINE bool operator==(const TBoxSphereBounds<T, TExtent>& Other) const;
 	
 	/**
 	 * Compare bounding volume this and Other.
@@ -148,7 +154,7 @@ public:
 	 * @param Other The other bounding volume.
 	 * @return true of they do not match.
 	 */	
-	FORCEINLINE bool operator!=(const TBoxSphereBounds<T>& Other) const;
+	FORCEINLINE bool operator!=(const TBoxSphereBounds<T, TExtent>& Other) const;
 
 public:
 
@@ -174,9 +180,9 @@ public:
 	 * @param  Tolerance Error tolerance added to test distance.
 	 * @return true if spheres intersect, false otherwise.
 	 */
-	FORCEINLINE static bool SpheresIntersect(const TBoxSphereBounds<T>& A, const TBoxSphereBounds<T>& B, T Tolerance = KINDA_SMALL_NUMBER)
+	FORCEINLINE static bool SpheresIntersect(const TBoxSphereBounds<T, TExtent>& A, const TBoxSphereBounds<T, TExtent>& B, TExtent Tolerance = KINDA_SMALL_NUMBER)
 	{
-		return (A.Origin - B.Origin).SizeSquared() <= FMath::Square(FMath::Max<T>(0.f, A.SphereRadius + B.SphereRadius + Tolerance));
+		return (A.Origin - B.Origin).SizeSquared() <= FMath::Square(FMath::Max<TExtent>(0, A.SphereRadius + B.SphereRadius + Tolerance));
 	}
 
 	/**
@@ -186,7 +192,7 @@ public:
 	 * @param  B Second BoxSphereBounds to test.
 	 * @return true if boxes intersect, false otherwise.
 	 */
-	FORCEINLINE static bool BoxesIntersect(const TBoxSphereBounds<T>& A, const TBoxSphereBounds<T>& B)
+	FORCEINLINE static bool BoxesIntersect(const TBoxSphereBounds<T, TExtent>& A, const TBoxSphereBounds<T, TExtent>& B)
 	{
 		return A.GetBox().Intersect(B.GetBox());
 	}
@@ -233,7 +239,7 @@ public:
 	 * @param ExpandAmount The size to increase by.
 	 * @return A new box with the expanded size.
 	 */
-	FORCEINLINE TBoxSphereBounds<T> ExpandBy( T ExpandAmount ) const
+	FORCEINLINE TBoxSphereBounds<T, TExtent> ExpandBy(TExtent ExpandAmount ) const
 	{
 		return TBoxSphereBounds(Origin, BoxExtent + ExpandAmount, SphereRadius + ExpandAmount);
 	}
@@ -244,7 +250,7 @@ public:
 	 * @param M The matrix.
 	 * @return The transformed volume.
 	 */
-	TBoxSphereBounds<T> TransformBy( const TMatrix<T>& M ) const;
+	TBoxSphereBounds<T, TExtent> TransformBy( const TMatrix<T>& M ) const;
 
 	/**
 	 * Gets a bounding volume transformed by a FTransform object.
@@ -252,7 +258,7 @@ public:
 	 * @param M The FTransform object.
 	 * @return The transformed volume.
 	 */
-	TBoxSphereBounds<T> TransformBy( const TTransform<T>& M ) const;
+	TBoxSphereBounds<T, TExtent> TransformBy( const TTransform<T>& M ) const;
 
 	/**
 	 * Get a textual representation of this bounding box.
@@ -266,7 +272,7 @@ public:
 	 *
 	 * This is a legacy version of the function used to compute primitive bounds, to avoid the need to rebuild lighting after the change.
 	 */
-	friend TBoxSphereBounds<T> Union( const TBoxSphereBounds<T>& A,const TBoxSphereBounds<T>& B )
+	friend TBoxSphereBounds<T, TExtent> Union( const TBoxSphereBounds<T, TExtent>& A,const TBoxSphereBounds<T, TExtent>& B )
 	{
 		return A + B;
 	}
@@ -301,7 +307,7 @@ public:
 };
 
 
-/* TBoxSphereBounds<T> inline functions
+/* TBoxSphereBounds<T, TExtent> inline functions
  *****************************************************************************/
 
 /**
@@ -311,7 +317,7 @@ public:
  * @param Bounds The bounding volume to serialize.
  * @return The archive..
  */
-inline FArchive& operator<<(FArchive& Ar, TBoxSphereBounds<float>& Bounds)
+inline FArchive& operator<<(FArchive& Ar, TBoxSphereBounds<float, float>& Bounds)
 {
 	Ar << Bounds.Origin << Bounds.BoxExtent << Bounds.SphereRadius;
 	return Ar;
@@ -324,7 +330,7 @@ inline FArchive& operator<<(FArchive& Ar, TBoxSphereBounds<float>& Bounds)
  * @param Bounds The bounding volume to serialize.
  * @return The archive..
  */
-inline FArchive& operator<<(FArchive& Ar, TBoxSphereBounds<double>& Bounds)
+inline FArchive& operator<<(FArchive& Ar, TBoxSphereBounds<double, double>& Bounds)
 {
 	Ar << Bounds.Origin << Bounds.BoxExtent;
 	// LWC_TODO: Serializer
@@ -345,8 +351,8 @@ inline FArchive& operator<<(FArchive& Ar, TBoxSphereBounds<double>& Bounds)
 	return Ar;
 }
 
-template<typename T>
-FORCEINLINE TBoxSphereBounds<T>::TBoxSphereBounds( const TVector<T>* Points, uint32 NumPoints )
+template<typename T, typename TExtent>
+FORCEINLINE TBoxSphereBounds<T, TExtent>::TBoxSphereBounds( const TVector<T>* Points, uint32 NumPoints )
 {
 	TBox<T> BoundingBox(ForceInit);
 
@@ -356,14 +362,16 @@ FORCEINLINE TBoxSphereBounds<T>::TBoxSphereBounds( const TVector<T>* Points, uin
 		BoundingBox += Points[PointIndex];
 	}
 
-	BoundingBox.GetCenterAndExtents(Origin, BoxExtent);
+	TVector<T> LocalExtent;
+	BoundingBox.GetCenterAndExtents(Origin, LocalExtent);
+	BoxExtent = TVector<TExtent>(LocalExtent);
 
 	// using the center of the bounding box as the origin of the sphere, find the radius of the bounding sphere.
-	T SquaredSphereRadius = 0.0f;
+	TExtent SquaredSphereRadius = 0;
 
 	for (uint32 PointIndex = 0; PointIndex < NumPoints; PointIndex++)
 	{
-		SquaredSphereRadius = FMath::Max(SquaredSphereRadius, (Points[PointIndex] - Origin).SizeSquared());	// LWC_TODO: Precision loss
+		SquaredSphereRadius = FMath::Max<TExtent>(SquaredSphereRadius, (Points[PointIndex] - Origin).SizeSquared());	// LWC_TODO: Precision loss
 	}
 
 	SphereRadius = FMath::Sqrt(SquaredSphereRadius);
@@ -371,8 +379,8 @@ FORCEINLINE TBoxSphereBounds<T>::TBoxSphereBounds( const TVector<T>* Points, uin
 	DiagnosticCheckNaN();
 }
 
-template<typename T>
-FORCEINLINE TBoxSphereBounds<T> TBoxSphereBounds<T>::operator+( const TBoxSphereBounds<T>& Other ) const
+template<typename T, typename TExtent>
+FORCEINLINE TBoxSphereBounds<T, TExtent> TBoxSphereBounds<T, TExtent>::operator+( const TBoxSphereBounds<T, TExtent>& Other ) const
 {
 	TBox<T> BoundingBox(ForceInit);
 
@@ -382,41 +390,41 @@ FORCEINLINE TBoxSphereBounds<T> TBoxSphereBounds<T>::operator+( const TBoxSphere
 	BoundingBox += (Other.Origin + Other.BoxExtent);
 
 	// build a bounding sphere from the bounding box's origin and the radii of A and B.
-	TBoxSphereBounds<T> Result(BoundingBox);
+	TBoxSphereBounds<T, TExtent> Result(BoundingBox);
 
-	Result.SphereRadius = FMath::Min(Result.SphereRadius, FMath::Max((Origin - Result.Origin).Size() + SphereRadius, (Other.Origin - Result.Origin).Size() + Other.SphereRadius));
+	Result.SphereRadius = FMath::Min<TExtent>(Result.SphereRadius, FMath::Max<TExtent>((Origin - Result.Origin).Size() + SphereRadius, (Other.Origin - Result.Origin).Size() + Other.SphereRadius));
 	Result.DiagnosticCheckNaN();
 
 	return Result;
 }
 
-template<typename T>
-FORCEINLINE bool TBoxSphereBounds<T>::operator==(const TBoxSphereBounds<T>& Other) const
+template<typename T, typename TExtent>
+FORCEINLINE bool TBoxSphereBounds<T, TExtent>::operator==(const TBoxSphereBounds<T, TExtent>& Other) const
 {
 	return Origin == Other.Origin && BoxExtent == Other.BoxExtent &&  SphereRadius == Other.SphereRadius;
 }
 
-template<typename T>
-FORCEINLINE bool TBoxSphereBounds<T>::operator!=(const TBoxSphereBounds<T>& Other) const
+template<typename T, typename TExtent>
+FORCEINLINE bool TBoxSphereBounds<T, TExtent>::operator!=(const TBoxSphereBounds<T, TExtent>& Other) const
 {
 	return !(*this == Other);
 }
 
-template<typename T>
-FORCEINLINE bool TBoxSphereBounds<T>::Serialize(FArchive &Ar)
+template<typename T, typename TExtent>
+FORCEINLINE bool TBoxSphereBounds<T, TExtent>::Serialize(FArchive &Ar)
 {
 	Ar << *this;
 	return true;
 }
 
-template<typename T>
-FORCEINLINE FString TBoxSphereBounds<T>::ToString() const
+template<typename T, typename TExtent>
+FORCEINLINE FString TBoxSphereBounds<T, TExtent>::ToString() const
 {
 	return FString::Printf(TEXT("Origin=%s, BoxExtent=(%s), SphereRadius=(%f)"), *Origin.ToString(), *BoxExtent.ToString(), SphereRadius);
 }
 
-template<typename T>
-TBoxSphereBounds<T> TBoxSphereBounds<T>::TransformBy(const TMatrix<T>& M) const
+template<typename T, typename TExtent>
+TBoxSphereBounds<T, TExtent> TBoxSphereBounds<T, TExtent>::TransformBy(const TMatrix<T>& M) const
 {
 #if ENABLE_NAN_DIAGNOSTIC
 	if (M.ContainsNaN())
@@ -459,7 +467,7 @@ TBoxSphereBounds<T> TBoxSphereBounds<T>::TransformBy(const TMatrix<T>& M) const
 	Result.SphereRadius = FMath::Min(Result.SphereRadius, BoxExtentMagnitude);
 
 	Result.DiagnosticCheckNaN();
-	return Result;
+	return TBoxSphereBounds<T, TExtent>(Result);
 }
 
 /**
@@ -468,28 +476,26 @@ TBoxSphereBounds<T> TBoxSphereBounds<T>::TransformBy(const TMatrix<T>& M) const
  * @param M The FTransform object.
  * @return The transformed volume.
  */
- template<typename T>
-TBoxSphereBounds<T> TBoxSphereBounds<T>::TransformBy(const TTransform<T>& M) const
+ template<typename T, typename TExtent>
+TBoxSphereBounds<T, TExtent> TBoxSphereBounds<T, TExtent>::TransformBy(const TTransform<T>& M) const
 {
 #if ENABLE_NAN_DIAGNOSTIC
 	M.DiagnosticCheckNaN_All();
 #endif
 
 	const TMatrix<T> Mat = M.ToMatrixWithScale();
-	TBoxSphereBounds<T> Result = TransformBy(Mat);
+	TBoxSphereBounds<T, TExtent> Result = TransformBy(Mat);
 	return Result;
 }
 
-
 } // namespace UE::Math
 } // namespace UE
 
-UE_DECLARE_LWC_TYPE(BoxSphereBounds, 3);
-
 template <> struct TIsPODType<FBoxSphereBounds3f> { enum { Value = true }; };
 template <> struct TIsPODType<FBoxSphereBounds3d> { enum { Value = true }; };
 template <> struct TIsUECoreVariant<FBoxSphereBounds3f> { enum { Value = true }; };
 template <> struct TIsUECoreVariant<FBoxSphereBounds3d> { enum { Value = true }; };
+template <> struct TIsUECoreVariant<FCompactBoxSphereBounds3d> { enum { Value = true }; };
 
 template<>
 inline bool FBoxSphereBounds3f::SerializeFromMismatchedTag(FName StructTag, FArchive& Ar)
@@ -501,5 +507,4 @@ template<>
 inline bool FBoxSphereBounds3d::SerializeFromMismatchedTag(FName StructTag, FArchive& Ar)
 {
 	return UE_SERIALIZE_VARIANT_FROM_MISMATCHED_TAG(Ar, BoxSphereBounds, BoxSphereBounds3d, BoxSphereBounds3f);
-
-}
\ No newline at end of file
+}
diff --git a/Engine/Source/Runtime/Core/Public/Math/GenericOctree.h b/Engine/Source/Runtime/Core/Public/Math/GenericOctree.h
--- a/Engine/Source/Runtime/Core/Public/Math/GenericOctree.h
+++ b/Engine/Source/Runtime/Core/Public/Math/GenericOctree.h
@@ -20,6 +20,8 @@ class FOctreeChildNodeRef;
 class FBoxCenterAndExtent
 {
 public:
+	using FReal = FVector4::FReal;
+
 	FVector4 Center;	// LWC_TODO: Case for FVector4d
 	FVector4 Extent;
 
@@ -42,10 +44,11 @@ public:
 	}
 
 	/** FBoxSphereBounds conversion constructor. */
-	explicit FBoxCenterAndExtent(const FBoxSphereBounds& BoxSphere)
+	template<typename TExtent>
+	explicit FBoxCenterAndExtent(const UE::Math::TBoxSphereBounds<FReal, TExtent>& BoxSphere)
 	{
 		Center = BoxSphere.Origin;
-		Extent = BoxSphere.BoxExtent;
+		Extent = FVector(BoxSphere.BoxExtent);
 		Center.W = Extent.W = 0;
 	}
 
@@ -84,7 +87,8 @@ public:
 	 * Warning: this operates on the W of the bounds positions!
 	 * @return true if the boxes intersect, or false.
 	 */
-	friend FORCEINLINE bool Intersect(const FBoxSphereBounds& A,const FBoxCenterAndExtent& B)
+	template<typename TExtent>
+	friend FORCEINLINE bool Intersect(const UE::Math::TBoxSphereBounds<FReal, TExtent>& A,const FBoxCenterAndExtent& B)
 	{
 		// CenterDifference is the vector between the centers of the bounding boxes.
 		const VectorRegister CenterDifference = VectorAbs(VectorSubtract(VectorLoadFloat3_W0(&A.Origin),VectorLoadAligned(&B.Center)));
diff --git a/Engine/Source/Runtime/Renderer/Private/LightSceneInfo.cpp b/Engine/Source/Runtime/Renderer/Private/LightSceneInfo.cpp
--- a/Engine/Source/Runtime/Renderer/Private/LightSceneInfo.cpp
+++ b/Engine/Source/Runtime/Renderer/Private/LightSceneInfo.cpp
@@ -136,7 +136,7 @@ void FLightSceneInfo::AddToScene()
  */
 void FLightSceneInfo::CreateLightPrimitiveInteraction(const FLightSceneInfoCompact& LightSceneInfoCompact, const FPrimitiveSceneInfoCompact& PrimitiveSceneInfoCompact)
 {
-	if(	!Scene->IsPrimitiveBeingRemoved(PrimitiveSceneInfoCompact.PrimitiveSceneInfo) && LightSceneInfoCompact.AffectsPrimitive(PrimitiveSceneInfoCompact.Bounds, PrimitiveSceneInfoCompact.Proxy))
+	if(	!Scene->IsPrimitiveBeingRemoved(PrimitiveSceneInfoCompact.PrimitiveSceneInfo) && LightSceneInfoCompact.AffectsPrimitive(FBoxSphereBounds(PrimitiveSceneInfoCompact.Bounds), PrimitiveSceneInfoCompact.Proxy))
 	{
 		// create light interaction and add to light/primitive lists
 		FLightPrimitiveInteraction::Create(this,PrimitiveSceneInfoCompact.PrimitiveSceneInfo);
diff --git a/Engine/Source/Runtime/Renderer/Private/PrimitiveSceneInfo.cpp b/Engine/Source/Runtime/Renderer/Private/PrimitiveSceneInfo.cpp
--- a/Engine/Source/Runtime/Renderer/Private/PrimitiveSceneInfo.cpp
+++ b/Engine/Source/Runtime/Renderer/Private/PrimitiveSceneInfo.cpp
@@ -165,7 +165,7 @@ FPrimitiveSceneInfoCompact::FPrimitiveSceneInfoCompact(FPrimitiveSceneInfo* InPr
 {
 	PrimitiveSceneInfo = InPrimitiveSceneInfo;
 	Proxy = PrimitiveSceneInfo->Proxy;
-	Bounds = PrimitiveSceneInfo->Proxy->GetBounds();
+	Bounds = FCompactBoxSphereBounds(PrimitiveSceneInfo->Proxy->GetBounds());
 	MinDrawDistance = PrimitiveSceneInfo->Proxy->GetMinDrawDistance();
 	MaxDrawDistance = PrimitiveSceneInfo->Proxy->GetMaxDrawDistance();
 
diff --git a/Engine/Source/Runtime/Renderer/Private/ShadowSetup.cpp b/Engine/Source/Runtime/Renderer/Private/ShadowSetup.cpp
--- a/Engine/Source/Runtime/Renderer/Private/ShadowSetup.cpp
+++ b/Engine/Source/Runtime/Renderer/Private/ShadowSetup.cpp
@@ -1528,7 +1528,7 @@ bool FProjectedShadowInfo::ShouldDrawStaticMeshes_AnyThread(
 	{
 		const int32 ForcedLOD = CurrentView.Family->EngineShowFlags.LOD ? (GetCVarForceLODShadow_AnyThread() != -1 ? GetCVarForceLODShadow_AnyThread() : GetCVarForceLOD_AnyThread()) : -1;
 
-		FLODMask ShadowLODToRender = CalcAndUpdateLODToRender(CurrentView, PrimitiveSceneInfoCompact.Bounds, PrimitiveSceneInfo, ForcedLOD);
+		FLODMask ShadowLODToRender = CalcAndUpdateLODToRender(CurrentView, FBoxSphereBounds(PrimitiveSceneInfoCompact.Bounds), PrimitiveSceneInfo, ForcedLOD);
 
 		if (WholeSceneDirectionalShadow)
 		{
@@ -1802,7 +1802,7 @@ uint64 FProjectedShadowInfo::AddSubjectPrimitive_AnyThread(
 
 	if (!ReceiverPrimitives.Contains(PrimitiveSceneInfo)
 		// Far cascade only casts from primitives marked for it
-		&& TestPrimitiveFarCascadeConditions(PrimitiveSceneInfoCompact.Proxy->CastsFarShadow(), PrimitiveSceneInfoCompact.Bounds))
+		&& TestPrimitiveFarCascadeConditions(PrimitiveSceneInfoCompact.Proxy->CastsFarShadow(), FBoxSphereBounds(PrimitiveSceneInfoCompact.Bounds)))
 	{
 		FViewInfo* CurrentView;
 		const bool bWholeSceneDirectionalShadow = IsWholeSceneDirectionalShadow();
@@ -1924,7 +1924,7 @@ uint64 FProjectedShadowInfo::AddSubjectPrimitive_AnyThread(
 			{
 				if (bWholeSceneShadow)
 				{
-					const FBoxSphereBounds& Bounds = PrimitiveSceneInfoCompact.Bounds;
+					const FCompactBoxSphereBounds& Bounds = PrimitiveSceneInfoCompact.Bounds;
 					const float DistanceSquared = (Bounds.Origin - CurrentView->ShadowViewMatrices.GetViewOrigin()).SizeSquared();
 					const float LODScaleSquared = FMath::Square(CurrentView->LODDistanceFactor);
 					const bool bDrawShadowDepth = FMath::Square(Bounds.SphereRadius) > FMath::Square(GMinScreenRadiusForShadowCaster) * DistanceSquared * LODScaleSquared;
@@ -4253,7 +4253,7 @@ struct FGatherShadowPrimitivesPacket
 	void FilterPrimitiveForShadows(FDynamicShadowsTaskData& TaskData, const FPrimitiveSceneInfoCompact& PrimitiveSceneInfoCompact)
 	{
 		const FPrimitiveFlagsCompact& PrimitiveFlagsCompact = PrimitiveSceneInfoCompact.PrimitiveFlagsCompact;
-		const FBoxSphereBounds& PrimitiveBounds = PrimitiveSceneInfoCompact.Bounds;
+		const FBoxSphereBounds PrimitiveBounds(PrimitiveSceneInfoCompact.Bounds);
 		FPrimitiveSceneInfo* PrimitiveSceneInfo = PrimitiveSceneInfoCompact.PrimitiveSceneInfo;
 		const FPrimitiveSceneProxy* PrimitiveProxy = PrimitiveSceneInfoCompact.Proxy;
 
diff --git a/Engine/Source/Runtime/Renderer/Private/ShadowSetupMobile.cpp b/Engine/Source/Runtime/Renderer/Private/ShadowSetupMobile.cpp
--- a/Engine/Source/Runtime/Renderer/Private/ShadowSetupMobile.cpp
+++ b/Engine/Source/Runtime/Renderer/Private/ShadowSetupMobile.cpp
@@ -149,7 +149,7 @@ static bool MobileDetermineStaticMeshesCSMVisibilityStateInner(
 
 	FPrimitiveSceneInfo* RESTRICT	PrimitiveSceneInfo = PrimitiveSceneInfoCompact.PrimitiveSceneInfo;
 	FPrimitiveSceneProxy* RESTRICT	PrimitiveProxy = PrimitiveSceneInfoCompact.Proxy;
-	const FBoxSphereBounds&			PrimitiveBounds = PrimitiveSceneInfoCompact.Bounds;
+	const FBoxSphereBounds			PrimitiveBounds(PrimitiveSceneInfoCompact.Bounds);
 	bool bFoundCSMReceiver = false;
 	if (PrimitiveProxy->WillEverBeLit() && PrimitiveProxy->ShouldReceiveMobileCSMShadows()
 		&& (PrimitiveProxy->GetLightingChannelMask() & LightProxy->GetLightingChannelMask()) != 0)
diff --git a/Engine/Source/Runtime/Renderer/Public/PrimitiveSceneInfo.h b/Engine/Source/Runtime/Renderer/Public/PrimitiveSceneInfo.h
--- a/Engine/Source/Runtime/Renderer/Public/PrimitiveSceneInfo.h
+++ b/Engine/Source/Runtime/Renderer/Public/PrimitiveSceneInfo.h
@@ -178,7 +178,7 @@ class FPrimitiveSceneInfoCompact
 public:
 	FPrimitiveSceneInfo* PrimitiveSceneInfo;
 	FPrimitiveSceneProxy* Proxy;
-	FBoxSphereBounds Bounds;
+	FCompactBoxSphereBounds Bounds;
 	float MinDrawDistance;
 	float MaxDrawDistance;
 	/** Used for precomputed visibility */
@@ -724,7 +724,7 @@ struct FPrimitiveOctreeSemantics
 
 	typedef TInlineAllocator<MaxElementsPerLeaf> ElementAllocator;
 
-	FORCEINLINE static const FBoxSphereBounds& GetBoundingBox(const FPrimitiveSceneInfoCompact& PrimitiveSceneInfoCompact)
+	FORCEINLINE static const FCompactBoxSphereBounds& GetBoundingBox(const FPrimitiveSceneInfoCompact& PrimitiveSceneInfoCompact)
 	{
 		return PrimitiveSceneInfoCompact.Bounds;
 	}
