[Insights] When tracing to file all relative paths should go to Profliling dir

When starting a trace to a file, all relative paths should go to the Profiling directory, on Windows even paths that starts with '/'. Also cleaned up includes.

diff --git a/Engine/Source/Runtime/Core/Private/ProfilingDebugging/TraceAuxiliary.cpp b/Engine/Source/Runtime/Core/Private/ProfilingDebugging/TraceAuxiliary.cpp
--- a/Engine/Source/Runtime/Core/Private/ProfilingDebugging/TraceAuxiliary.cpp
+++ b/Engine/Source/Runtime/Core/Private/ProfilingDebugging/TraceAuxiliary.cpp
@@ -1,9 +1,8 @@
 // Copyright Epic Games, Inc. All Rights Reserved.
 
 #include "ProfilingDebugging/TraceAuxiliary.h"
-#include "ProfilingDebugging/CpuProfilerTrace.h"
-#include "ProfilingDebugging/CallstackTrace.h"
 #include "Trace/Trace.h"
+#include "CoreGlobals.h"
 
 #if PLATFORM_WINDOWS
 #	include "Windows/AllowWindowsPlatformTypes.h"
@@ -15,18 +14,13 @@
 #	define WITH_UNREAL_TRACE_LAUNCH (PLATFORM_DESKTOP && !UE_BUILD_SHIPPING && !IS_PROGRAM)
 #endif
 
-#include "CoreGlobals.h"
-#include "Misc/Paths.h"
-#include "Misc/Fork.h"
-
 #if WITH_UNREAL_TRACE_LAUNCH
 #	include "Misc/Parse.h"
 #endif
 
-#include <atomic>
-
 #if UE_TRACE_ENABLED
 
+#include <atomic>
 #include "BuildSettings.h"
 #include "Containers/Array.h"
 #include "Containers/Map.h"
@@ -34,22 +28,23 @@
 #include "Containers/UnrealString.h"
 #include "HAL/FileManager.h"
 #include "HAL/IConsoleManager.h"
-#include "HAL/PlatformMisc.h"
 #include "Misc/App.h"
+#include "Misc/CString.h"
 #include "Misc/ConfigCacheIni.h"
 #include "Misc/CoreDelegates.h"
-#include "Misc/CString.h"
 #include "Misc/DateTime.h"
+#include "Misc/Fork.h"
+#include "Misc/PathViews.h"
+#include "Misc/Paths.h"
 #include "Modules/ModuleManager.h"
-#include "ProfilingDebugging/CountersTrace.h"
-#include "ProfilingDebugging/MiscTrace.h"
-#include "ProfilingDebugging/PlatformFileTrace.h"
-#include "ProfilingDebugging/PlatformEvents.h"
-#include "ProfilingDebugging/MemoryTrace.h"
+#include "ProfilingDebugging/CallstackTrace.h"
+#include "ProfilingDebugging/CpuProfilerTrace.h"
 #include "ProfilingDebugging/CsvProfiler.h"
+#include "ProfilingDebugging/MemoryTrace.h"
+#include "ProfilingDebugging/PlatformEvents.h"
 #include "String/ParseTokens.h"
-#include "Templates/UnrealTemplate.h"
 #include "Templates/Invoke.h"
+#include "Templates/UnrealTemplate.h"
 #include "Trace/Trace.inl"
 
 ////////////////////////////////////////////////////////////////////////////////
@@ -398,63 +393,78 @@ bool FTraceAuxiliaryImpl::SendToHost(const TCHAR* Host, const FTraceAuxiliary::F
 }
 
 ////////////////////////////////////////////////////////////////////////////////
-bool FTraceAuxiliaryImpl::WriteToFile(const TCHAR* Path, const FTraceAuxiliary::FLogCategoryAlias& LogCategory)
+bool FTraceAuxiliaryImpl::WriteToFile(const TCHAR* InPath, const FTraceAuxiliary::FLogCategoryAlias& LogCategory)
 {
+	const FStringView Path(InPath);
+	
 	// Default file name functor 
 	auto GetDefaultName = [] {return FDateTime::Now().ToString(TEXT("%Y%m%d_%H%M%S.utrace"));};
 	
-	if (Path == nullptr || *Path == '\0')
+	if (Path.IsEmpty())
 	{
 		const FString Name = GetDefaultName();
 		return WriteToFile(*Name, LogCategory);
 	}
 
 	FString WritePath;
-
-	// If there's no slash in the path, we'll put it in the profiling directory
-	if (FCString::Strchr(Path, '\\') == nullptr && FCString::Strchr(Path, '/') == nullptr)
+	// Relative paths go to the profiling directory
+	if (FPathViews::IsRelativePath(Path))
+	{
+		WritePath = FPaths::Combine(FPaths::ProfilingDir(), InPath);
+	}
+#if PLATFORM_WINDOWS
+	// On windows we treat paths starting with '/' as relative, except double slash which is a network path
+	else if (FPathViews::IsSeparator(Path[0]) && !(Path.Len() > 1 && FPathViews::IsSeparator(Path[1])))
 	{
-		WritePath = FPaths::ProfilingDir();
-		WritePath += Path;
+		WritePath = FPaths::Combine(FPaths::ProfilingDir(), InPath);
 	}
+#endif
 	else
 	{
-		WritePath = Path;
+		WritePath = InPath;
 	}
 
 	// If a directory is specified, add the default trace file name
-	if (WritePath.EndsWith(TEXT("/")) || WritePath.EndsWith(TEXT("\\")))
+	if (FPathViews::GetCleanFilename(WritePath).IsEmpty())
 	{
-		WritePath += GetDefaultName();
+		WritePath = FPaths::Combine(WritePath, GetDefaultName());
 	}
 
 	// The user may not have provided a suitable extension
-	if (!WritePath.EndsWith(".utrace"))
+	if (FPathViews::GetExtension(WritePath) != TEXT("utrace"))
 	{
-		WritePath += ".utrace";
+		WritePath = FPaths::SetExtension(WritePath, TEXT(".utrace"));
 	}
 
+	// Finally make sure the path is platform friendly
 	IFileManager& FileManager = IFileManager::Get();
+	FString NativePath = FileManager.ConvertToAbsolutePathForExternalAppForWrite(*WritePath);
 
 	// Ensure we can write the trace file appropriately
-	const FString WriteDir = FPaths::GetPath(WritePath);
+	const FString WriteDir = FPaths::GetPath(NativePath);
 	if (!FileManager.MakeDirectory(*WriteDir, true))
 	{
 		UE_LOG_REF(LogCategory, Warning, TEXT("Failed to create directory '%s'"), *WriteDir);
 		return false;
 	}
 
-	if (!bTruncateFile && FileManager.FileExists(*WritePath))
+	if (!bTruncateFile && FileManager.FileExists(*NativePath))
 	{
-		UE_LOG_REF(LogCategory, Warning, TEXT("Trace file '%s' already exists"), *WritePath);
+		UE_LOG_REF(LogCategory, Warning, TEXT("Trace file '%s' already exists"), *NativePath);
 		return false;
 	}
 
 	// Finally, tell trace to write the trace to a file.
-	FString NativePath = FileManager.ConvertToAbsolutePathForExternalAppForWrite(*WritePath);
 	if (!UE::Trace::WriteTo(*NativePath))
 	{
-		UE_LOG_REF(LogCategory, Warning, TEXT("Unable to trace to file '%s'"), *WritePath);
+		if (FPathViews::Equals(NativePath, WritePath))
+		{
+			UE_LOG_REF(LogCategory, Warning, TEXT("Unable to trace to file '%s'"), *NativePath);
+		}
+		else
+		{
+			UE_LOG_REF(LogCategory, Warning, TEXT("Unable to trace to file '%s' (transformed from '%s')"), *NativePath, *WritePath);
+		}
 		return false;
 	}
 
