FEM - Using the Concurrent Linear Allocator in multiple TaskGraph related systems to reduce Memory fragmentation, hitches and provide a slight performance boost.

diff --git a/Engine/Source/Runtime/Core/Private/Async/TaskGraph.cpp b/Engine/Source/Runtime/Core/Private/Async/TaskGraph.cpp
--- a/Engine/Source/Runtime/Core/Private/Async/TaskGraph.cpp
+++ b/Engine/Source/Runtime/Core/Private/Async/TaskGraph.cpp
@@ -2420,15 +2420,6 @@ bool FTaskGraphInterface::IsMultithread()
 	return FPlatformProcess::SupportsMultithreading() || (FForkProcessHelper::IsForkedMultithreadInstance() && GAllowTaskGraphForkMultithreading);
 }
 
-
-// Statics and some implementations from FBaseGraphTask and FGraphEvent
-
-static FBaseGraphTask::TSmallTaskAllocator TheSmallTaskAllocator;
-FBaseGraphTask::TSmallTaskAllocator& FBaseGraphTask::GetSmallTaskAllocator()
-{
-	return TheSmallTaskAllocator;
-}
-
 #if !(UE_BUILD_SHIPPING || UE_BUILD_TEST)
 void FBaseGraphTask::LogPossiblyInvalidSubsequentsTask(const TCHAR* TaskName)
 {
diff --git a/Engine/Source/Runtime/Core/Public/Async/LocalWorkQueue.h b/Engine/Source/Runtime/Core/Public/Async/LocalWorkQueue.h
--- a/Engine/Source/Runtime/Core/Public/Async/LocalWorkQueue.h
+++ b/Engine/Source/Runtime/Core/Public/Async/LocalWorkQueue.h
@@ -3,7 +3,9 @@
 #pragma once
 #include "Async/Fundamental/Task.h"
 #include "Async/Fundamental/Scheduler.h"
+#include "Async/TaskGraphInterfaces.h"
 #include "Experimental/Containers/FAAArrayQueue.h"
+#include "Experimental/ConcurrentLinearAllocator.h"
 #include "Templates/RefCounting.h"
 #include <atomic>
 
@@ -43,7 +45,7 @@ class TLocalWorkQueue
 {
 	UE_NONCOPYABLE(TLocalWorkQueue);
 
-	struct FInternalData : public FThreadSafeRefCountedObject
+	struct FInternalData : public TConcurrentLinearObject<FInternalData, FTaskGraphBlockAllocationTag>, public FThreadSafeRefCountedObject
 	{
 		FAAArrayQueue<TaskType> TaskQueue;	
 		std::atomic_int ActiveWorkers {0};
diff --git a/Engine/Source/Runtime/Core/Public/Async/TaskGraphInterfaces.h b/Engine/Source/Runtime/Core/Public/Async/TaskGraphInterfaces.h
--- a/Engine/Source/Runtime/Core/Public/Async/TaskGraphInterfaces.h
+++ b/Engine/Source/Runtime/Core/Public/Async/TaskGraphInterfaces.h
@@ -21,6 +21,7 @@
 #include "HAL/LowLevelMemTracker.h"
 #include "Templates/RefCounting.h"
 #include "Containers/LockFreeFixedSizeAllocator.h"
+#include "Experimental/ConcurrentLinearAllocator.h"
 #include "Misc/MemStack.h"
 #include "Templates/Atomic.h"
 
@@ -449,6 +450,17 @@ public:
 	static void BroadcastSlow_OnlyUseForSpecialPurposes(bool bDoTaskThreads, bool bDoBackgroundThreads, TFunction<void(ENamedThreads::Type CurrentThread)>& Callback);
 };
 
+struct FTaskGraphBlockAllocationTag : FDefaultBlockAllocationTag
+{
+	static constexpr uint32 BlockSize = 64 * 1024;
+	static constexpr bool AllowOversizedBlocks = false;
+	static constexpr bool RequiresAccurateSize = false;
+	static constexpr bool InlineBlockAllocation = true;
+	static constexpr const TCHAR* TagName = TEXT("TaskGraphLinear");
+
+	using Allocator = TBlockAllocationCache<BlockSize, FAlignedAllocator>;
+};
+
 /** 
  *	Base class for all tasks. 
  *	Tasks go through a very specific life stage progression, and this is verified.
@@ -456,15 +468,6 @@ public:
 
 class FBaseGraphTask
 {
-public:
-
-	// Allocator for small tasks.
-	enum
-	{
-		/** Total size in bytes for a small task that will use the custom allocator **/
-		SMALL_TASK_SIZE = 256 + sizeof(LowLevelTasks::FTask)
-	};
-	typedef TLockFreeFixedSizeAllocator_TLSCache<SMALL_TASK_SIZE, PLATFORM_CACHE_LINE_SIZE, FNoopCounter, true> TSmallTaskAllocator;
 protected:
 	/** 
 	 *	Constructor
@@ -516,9 +519,6 @@ protected:
 	static void CORE_API LogPossiblyInvalidSubsequentsTask(const TCHAR* TaskName);
 #endif
 
-	/** Singleton to retrieve the small task allocator **/
-	static CORE_API TSmallTaskAllocator& GetSmallTaskAllocator();
-
 	/** 
 	 *	An indication that a prerequisite has been completed. Reduces the number of prerequisites by one and if no prerequisites are outstanding, it queues the task for execution.
 	 *	@param CurrentThread; provides the index of the thread we are running on. This is handy for submitting new taks. Can be ENamedThreads::AnyThread if the current thread is unknown.
@@ -860,7 +860,7 @@ public:
  *	Embeds a user defined task, as exemplified above, for doing the work and provides the functionality for setting up and handling prerequisites and subsequents
  **/
 template<typename TTask>
-class TGraphTask final : public FBaseGraphTask
+class TGraphTask final : public TConcurrentLinearObject<TGraphTask<TTask>, FTaskGraphBlockAllocationTag>, public FBaseGraphTask
 {
 public:
 	/** 
@@ -924,11 +924,6 @@ public:
 	static FConstructor CreateTask(const FGraphEventArray* Prerequisites = NULL, ENamedThreads::Type CurrentThreadIfKnown = ENamedThreads::AnyThread)
 	{
 		int32 NumPrereq = Prerequisites ? Prerequisites->Num() : 0;
-		if (sizeof(TGraphTask) <= FBaseGraphTask::SMALL_TASK_SIZE)
-		{
-			void *Mem = FBaseGraphTask::GetSmallTaskAllocator().Allocate();
-			return FConstructor(new (Mem) TGraphTask(TTask::GetSubsequentsMode() == ESubsequentsMode::FireAndForget ? NULL : FGraphEvent::CreateGraphEvent(), NumPrereq), Prerequisites, CurrentThreadIfKnown);
-		}
 		return FConstructor(new TGraphTask(TTask::GetSubsequentsMode() == ESubsequentsMode::FireAndForget ? NULL : FGraphEvent::CreateGraphEvent(), NumPrereq), Prerequisites, CurrentThreadIfKnown);
 	}
 
@@ -998,15 +993,7 @@ private:
 
 	void DeleteTask() final override
 	{
-		if (sizeof(TGraphTask) <= FBaseGraphTask::SMALL_TASK_SIZE)
-		{
-			this->TGraphTask::~TGraphTask();
-			FBaseGraphTask::GetSmallTaskAllocator().Free(this);
-		}
-		else
-		{
-			delete this;
-		}
+		delete this;
 	}
 
 	// Internals 
@@ -1112,11 +1099,6 @@ private:
 	**/
 	static FConstructor CreateTask(FGraphEventRef SubsequentsToAssume, const FGraphEventArray* Prerequisites = NULL, ENamedThreads::Type CurrentThreadIfKnown = ENamedThreads::AnyThread)
 	{
-		if (sizeof(TGraphTask) <= FBaseGraphTask::SMALL_TASK_SIZE)
-		{
-			void *Mem = FBaseGraphTask::GetSmallTaskAllocator().Allocate();
-			return FConstructor(new (Mem) TGraphTask(SubsequentsToAssume, Prerequisites ? Prerequisites->Num() : 0), Prerequisites, CurrentThreadIfKnown);
-		}
 		return FConstructor(new TGraphTask(SubsequentsToAssume, Prerequisites ? Prerequisites->Num() : 0), Prerequisites, CurrentThreadIfKnown);
 	}
 
diff --git a/Engine/Source/Runtime/Core/Public/Misc/QueuedThreadPoolWrapper.h b/Engine/Source/Runtime/Core/Public/Misc/QueuedThreadPoolWrapper.h
--- a/Engine/Source/Runtime/Core/Public/Misc/QueuedThreadPoolWrapper.h
+++ b/Engine/Source/Runtime/Core/Public/Misc/QueuedThreadPoolWrapper.h
@@ -10,6 +10,7 @@
 #include "Misc/IQueuedWork.h"
 #include "Async/Fundamental/Scheduler.h"
 #include "Experimental/Containers/FAAArrayQueue.h"
+#include "Experimental/ConcurrentLinearAllocator.h"
 
 #include <atomic>
 
@@ -232,7 +233,7 @@ private:
 class CORE_API FQueuedLowLevelThreadPool final : public FQueuedThreadPool
 {
 	/* Internal data of the scheduler used for cancellation */
-	struct FQueuedWorkInternalData : IQueuedWorkInternalData
+	struct FQueuedWorkInternalData : TConcurrentLinearObject<FQueuedWorkInternalData, FTaskGraphBlockAllocationTag>, IQueuedWorkInternalData
 	{
 		LowLevelTasks::FTask Task;
 
