FEM - upgrade parallelFor to use more a efficient LowLevel implementation.

diff --git a/Engine/Source/Runtime/Core/Public/Async/ParallelFor.h b/Engine/Source/Runtime/Core/Public/Async/ParallelFor.h
--- a/Engine/Source/Runtime/Core/Public/Async/ParallelFor.h
+++ b/Engine/Source/Runtime/Core/Public/Async/ParallelFor.h
@@ -6,18 +6,21 @@
 
 #pragma once
 
+#include <atomic>
 #include "CoreTypes.h"
 #include "Misc/AssertionMacros.h"
 #include "Math/UnrealMathUtility.h"
 #include "Templates/Function.h"
 #include "Templates/SharedPointer.h"
+#include "Templates/RefCounting.h"
 #include "HAL/ThreadSafeCounter.h"
 #include "Stats/Stats.h"
 #include "Async/TaskGraphInterfaces.h"
 #include "Misc/App.h"
 #include "Misc/Fork.h"
-
-#include <atomic>
+#include "Async/Fundamental/Task.h"
+#include "Async/Fundamental/Scheduler.h"
+#include "Experimental/ConcurrentLinearAllocator.h"
 
 extern CORE_API int32 GParallelForBackgroundYieldingTimeoutMs;
 
@@ -306,8 +309,17 @@ namespace ParallelForImpl
 		return NumThreadTasks;
 	}
 
-	template<typename FunctionType, typename ContextType>
-	inline void ParallelForInternal(int32 Num, FunctionType Body, EParallelForFlags Flags, const TArrayView<ContextType>& Contexts)
+	/** 
+		*	General purpose parallel for that uses the taskgraph
+		*	@param Num; number of calls of Body; Body(0), Body(1)....Body(Num - 1)
+		*	@param Body; Function to call from multiple threads
+		*	@param CurrentThreadWorkToDoBeforeHelping; The work is performed on the main thread before it starts helping with the ParallelFor proper
+		*	@param Flags; Used to customize the behavior of the ParallelFor if needed.
+		*   @param Contexts; Optional per thread contexts to accumulate data concurrently.
+		*	Notes: Please add stats around to calls to parallel for and within your lambda as appropriate. Do not clog the task graph with long running tasks or tasks that block.
+	**/
+	template<typename FunctionType, typename PreWorkType, typename ContextType>
+	inline void OldParallelForInternal(int32 Num, FunctionType Body, PreWorkType CurrentThreadWorkToDoBeforeHelping, EParallelForFlags Flags, const TArrayView<ContextType>& Contexts)
 	{
 		SCOPE_CYCLE_COUNTER(STAT_ParallelFor);
 		check(Num >= 0);
@@ -315,6 +327,8 @@ namespace ParallelForImpl
 		int32 AnyThreadTasks = GetNumberOfThreadTasks(Num, Flags);
 		if (!AnyThreadTasks)
 		{
+			// do the prework
+			CurrentThreadWorkToDoBeforeHelping();
 			// no threads, just do it and return
 			for (int32 Index = 0; Index < Num; Index++)
 			{				
@@ -329,6 +343,8 @@ namespace ParallelForImpl
 		TParallelForData<FunctionType>* DataPtr = new TParallelForData<FunctionType>(Num, AnyThreadTasks + 1, (Num > AnyThreadTasks + 1) && bPumpRenderingThread, Body, Flags);
 		TSharedRef<TParallelForData<FunctionType>, ESPMode::ThreadSafe> Data = MakeShareable(DataPtr);
 		TGraphTask<TParallelForTask<FunctionType, ContextType>>::CreateTask().ConstructAndDispatchWhenReady(Contexts, 1, Data, DesiredThread, AnyThreadTasks - 1);
+		// do the prework
+		CurrentThreadWorkToDoBeforeHelping();
 		// this thread can help too and this is important to prevent deadlock on recursion 
 		if (!Data->Process(Contexts, 0, 0, Data, DesiredThread, true))
 		{
@@ -359,78 +375,269 @@ namespace ParallelForImpl
 		// Data must live on until all of the tasks are cleared which might be long after this function exits
 	}
 
-	template<typename FunctionType>
-	inline void ParallelForInternal(int32 Num, FunctionType Body, EParallelForFlags Flags)
-	{
-		ParallelForInternal<FunctionType, nullptr_t>(Num, Body, Flags, TArrayView<nullptr_t>());
-	}
-	
 	/** 
 		*	General purpose parallel for that uses the taskgraph
 		*	@param Num; number of calls of Body; Body(0), Body(1)....Body(Num - 1)
 		*	@param Body; Function to call from multiple threads
 		*	@param CurrentThreadWorkToDoBeforeHelping; The work is performed on the main thread before it starts helping with the ParallelFor proper
 		*	@param Flags; Used to customize the behavior of the ParallelFor if needed.
+		*   @param Contexts; Optional per thread contexts to accumulate data concurrently.
 		*	Notes: Please add stats around to calls to parallel for and within your lambda as appropriate. Do not clog the task graph with long running tasks or tasks that block.
 	**/
-	template<typename FunctionType>
-	inline void ParallelForWithPreWorkInternal(int32 Num, FunctionType Body, TFunctionRef<void()> CurrentThreadWorkToDoBeforeHelping, EParallelForFlags Flags = EParallelForFlags::None)
+	template<typename BodyType, typename PreWorkType, typename ContextType>
+	inline void NewParallelForInternal(int32 Num, BodyType Body, PreWorkType CurrentThreadWorkToDoBeforeHelping, EParallelForFlags Flags, const TArrayView<ContextType>& Contexts)
 	{
 		SCOPE_CYCLE_COUNTER(STAT_ParallelFor);
+		check(Num >= 0);
 
-		int32 AnyThreadTasks = 0;
+		//single threaded mode
 		const bool bIsMultithread = FApp::ShouldUseThreadingForPerformance() || FForkProcessHelper::IsForkedMultithreadInstance();
-		if ((Flags & EParallelForFlags::ForceSingleThread) == EParallelForFlags::None && bIsMultithread)
-		{
-			AnyThreadTasks = FMath::Min<int32>(FTaskGraphInterface::Get().GetNumWorkerThreads(), Num);
-		}
-		if (!AnyThreadTasks)
+		if ((Num <= 1) || ((Flags & EParallelForFlags::ForceSingleThread) == EParallelForFlags::ForceSingleThread) || !bIsMultithread)
 		{
 			// do the prework
 			CurrentThreadWorkToDoBeforeHelping();
 			// no threads, just do it and return
-			for (int32 Index = 0; Index < Num; Index++)
+			for(int32 Index = 0; Index < Num; Index++)
 			{
-				Body(Index);
+				CallBody(Body, Contexts, 0, Index);
 			}
 			return;
 		}
-		check(Num);
 
-		const ENamedThreads::Type DesiredThread = GetBestDesiredThread(Flags);
+		//calculate the number of workers
+		int32 NumWorkers = int32(LowLevelTasks::FScheduler::Get().GetNumWorkers());
+		if (!LowLevelTasks::FScheduler::Get().IsWorkerThread())
+		{
+			NumWorkers++; //named threads help with the work
+		}
+		NumWorkers = FMath::Min(NumWorkers, Num);
+		checkSlow(NumWorkers > 0);
+
+		//calculate the batch sizes
+		int32 BatchSize = 1;
+		int32 NumBatches = Num;
+		bool bIsUnbalanced = (Flags & EParallelForFlags::Unbalanced) == EParallelForFlags::Unbalanced;
+		if (!bIsUnbalanced)
+		{
+			for (int32 Div = 6; Div; Div--)
+			{
+				if (Num >= (NumWorkers * Div))
+				{	
+					BatchSize = FMath::DivideAndRoundUp<int32>(Num, (NumWorkers * Div));
+					NumBatches = FMath::DivideAndRoundUp<int32>(Num, BatchSize);
+
+					if (NumBatches >= NumWorkers)
+					{
+						break;
+					}
+				}
+			}
+		}
+		NumWorkers--; //Decrement one because this function will work on it locally
+		checkSlow(BatchSize * NumBatches >= Num);
+
+		//Try to inherit the Priority from the caller
+		// Anything scheduled by the task graph is latency sensitive because it might impact the frame rate. Anything else is not (i.e. Worker / Background threads).
+		const ETaskTag LatencySensitiveTasks = 
+			ETaskTag::EStaticInit | 
+			ETaskTag::EGameThread | 
+			ETaskTag::ESlateThread | 
+#if !UE_AUDIO_THREAD_AS_PIPE
+			ETaskTag::EAudioThread | 
+#endif
+			ETaskTag::ERenderingThread | 
+			ETaskTag::ERhiThread;
+
+		const bool bBackgroundPriority = (Flags & EParallelForFlags::BackgroundPriority) != EParallelForFlags::None;
+		const bool bIsLatencySensitive = (FTaskTagScope::GetCurrentTag() & LatencySensitiveTasks) != ETaskTag::ENone;
+
+		LowLevelTasks::ETaskPriority Priority = LowLevelTasks::ETaskPriority::Inherit;
+		if (bIsLatencySensitive && !bBackgroundPriority)
+		{
+			Priority =  LowLevelTasks::ETaskPriority::High;
+		}
+		else if (bBackgroundPriority)
+		{
+			Priority = LowLevelTasks::ETaskPriority::BackgroundNormal;
+		}
+		
+		//shared data between tasks
+		struct alignas(PLATFORM_CACHE_LINE_SIZE) FParallelForData : public TConcurrentLinearObject<FParallelForData, FTaskGraphBlockAllocationTag>, public FThreadSafeRefCountedObject
+		{
+			FParallelForData(int32 InNum, int32 InBatchSize, int32 InNumBatches, const TArrayView<ContextType>& InContexts, const BodyType& InBody, FEventRef& InFinishedSignal)
+				: Num(InNum)
+				, BatchSize(InBatchSize)
+				, NumBatches(InNumBatches)
+				, Contexts(InContexts)
+				, Body(InBody)
+				, FinishedSignal(InFinishedSignal)
+			{
+				IncompleteBatches.store(NumBatches, std::memory_order_relaxed);
+			}
+			std::atomic_int BatchItem  { 0 };
+			std::atomic_int IncompleteBatches { 0 };
+			int32 Num;
+			int32 BatchSize;
+			int32 NumBatches;
+			const TArrayView<ContextType>& Contexts;
+			const BodyType& Body;
+			FEventRef& FinishedSignal;
+		};
+		using FDataHandle = TRefCountPtr<FParallelForData>;
+
+		//each task has an executor.
+		class FParallelExecutor : public TConcurrentLinearObject<FParallelExecutor, FTaskGraphBlockAllocationTag>
+		{
+			LowLevelTasks::FTask Task;
+			FDataHandle Data;
+			int32 WorkerIndex;
+
+		public:
+			inline FParallelExecutor(FDataHandle&& InData, int32 InWorkerIndex) 
+				: Data(MoveTemp(InData))
+				, WorkerIndex(InWorkerIndex)
+			{
+			}
+
+			FParallelExecutor(const FParallelExecutor&) = delete;
+			inline FParallelExecutor(FParallelExecutor&& Other) 
+				: Data(MoveTemp(Other.Data))
+				, WorkerIndex(Other.WorkerIndex)
+			{
+			}
+
+			inline bool Run(const bool bIsMaster = true, const bool bIsBackgroundPriority = false) noexcept
+			{
+				FMemMark Mark(FMemStack::Get());
+				TRACE_CPUPROFILER_EVENT_SCOPE(ParallelFor);
+
+				auto Now = [] { return FTimespan::FromSeconds(FPlatformTime::Seconds()); };
+				FTimespan Start = FTimespan::MinValue();
+				FTimespan YieldingThreshold;
+				if (bIsBackgroundPriority)
+				{
+					Start = Now();
+					YieldingThreshold = FTimespan::FromMilliseconds(FMath::Max(0, GParallelForBackgroundYieldingTimeoutMs));
+				}
+
+				const int32 Num = Data->Num;
+				const int32 BatchSize = Data->BatchSize;
+				const int32 NumBatches = Data->NumBatches;
+				const TArrayView<ContextType>& Contexts = Data->Contexts;
+				const BodyType& Body = Data->Body;
+
+				const bool bSaveLastBlockForMaster = (Num > NumBatches);
+				for(;;)
+				{
+					int32 BatchIndex = Data->BatchItem.fetch_add(1, std::memory_order_relaxed);
+					
+					//save the last block for the master to safe an event
+					if (bSaveLastBlockForMaster && BatchIndex >= NumBatches - 1)
+					{
+						if (!bIsMaster)
+						{
+							return false;
+						}
+						BatchIndex = (NumBatches - 1);
+					}
+
+					int32 StartIndex = BatchIndex * BatchSize;
+					int32 EndIndex = FMath::Min<int32>(StartIndex + BatchSize, Num);
+					for (int32 Index = StartIndex; Index < EndIndex; Index++)
+					{
+						CallBody(Body, Contexts, WorkerIndex, Index);
+					}
+
+					// we need to decrement IncompleteBatches when processing a Batch because we need to know if we are the last one
+					// so that if the main thread is the last one we can avoid an FEvent call.
+					if (StartIndex < Num && Data->IncompleteBatches.fetch_sub(1, std::memory_order_relaxed) == 1)
+					{
+						if (!bIsMaster)
+						{
+							Data->FinishedSignal->Trigger();
+						}
+						return true;
+					}
+					else if (EndIndex >= Num)
+					{
+						return false;
+					}
+					else if (!bIsBackgroundPriority)
+					{
+						continue;
+					}
+
+					auto PassedTime = [Start, &Now]() { return Now() - Start; };
+					if (PassedTime() > YieldingThreshold)
+					{
+						//abort and reschedule to give higher priority tasks a chance to run
+						FParallelExecutor::LaunchTask(Task.GetPriority(), MoveTemp(Data), WorkerIndex);
+						return false;
+					}
+				}
+			}
+
+			static void LaunchTask(LowLevelTasks::ETaskPriority InPriority, FDataHandle&& InData, int32 InWorkerIndex)
+			{
+				using FTaskHandle = TUniquePtr<FParallelExecutor>;
+				FTaskHandle TaskHandle = MakeUnique<FParallelExecutor>(MoveTemp(InData), InWorkerIndex);
+				LowLevelTasks::FTask& Task = TaskHandle->Task;
+				Task.Init(TEXT("FParallelExecutor"), InPriority, [Priority = InPriority, TaskHandle = MoveTemp(TaskHandle)]()
+				{
+					TaskHandle->Run(false, Priority >= LowLevelTasks::ETaskPriority::BackgroundNormal);
+				});
+				verify(LowLevelTasks::TryLaunch(Task, LowLevelTasks::EQueuePreference::GlobalQueuePreference));
+			}
+		};
+
+		//launch all the worker tasks
+		FEventRef FinishedSignal { EEventMode::ManualReset };
+		FDataHandle Data = new FParallelForData(Num, BatchSize, NumBatches, Contexts, Body, FinishedSignal);
+		for (int32 Worker = 0; Worker < NumWorkers; Worker++)
+		{
+			FParallelExecutor::LaunchTask(Priority, FDataHandle(Data), Worker);
+		}
 
-		TParallelForData<FunctionType>* DataPtr = new TParallelForData<FunctionType>(Num, AnyThreadTasks, false, Body, Flags);
-		TSharedRef<TParallelForData<FunctionType>, ESPMode::ThreadSafe> Data = MakeShareable(DataPtr);
-		TGraphTask<TParallelForTask<FunctionType>>::CreateTask().ConstructAndDispatchWhenReady(1, Data, DesiredThread, AnyThreadTasks - 1);
 		// do the prework
 		CurrentThreadWorkToDoBeforeHelping();
-		// this thread can help too and this is important to prevent deadlock on recursion 
-		if (!Data->Process(0, 0, Data, DesiredThread, true))
+
+		//help with the parallel-for to prevent deadlocks
+		FParallelExecutor LocalExecutor(MoveTemp(Data), NumWorkers);
+		const bool bFinishedLast = LocalExecutor.Run();
+
+		if(!bFinishedLast)
 		{
-			if ((Flags & EParallelForFlags::PumpRenderingThread) != EParallelForFlags::None && IsInRenderingThread())
+			const bool bPumpRenderingThread  = (Flags & EParallelForFlags::PumpRenderingThread) != EParallelForFlags::None;
+			if (bPumpRenderingThread && IsInActualRenderingThread())
 			{
-				while (!Data->Event->Wait(1))
+				// FinishedSignal waits here if some other thread finishes the last item
+				// Data must live on until all of the tasks are cleared which might be long after this function exits
+				while (!FinishedSignal->Wait(1))
 				{
 					FTaskGraphInterface::Get().ProcessThreadUntilIdle(ENamedThreads::GetRenderThread_Local());
 				}
 			}
 			else
 			{
-				Data->Event->Wait();
+				// FinishedSignal waits here if some other thread finishes the last item
+				// Data must live on until all of the tasks are cleared which might be long after this function exits
+				FinishedSignal->Wait();
 			}
-			check(Data->bTriggered);
+		}
+		checkSlow(Data->BatchItem.load(std::memory_order_relaxed) * Data->BatchSize >= Data->Num);
+	}
+
+	template<typename FunctionType, typename PreWorkType, typename ContextType>
+	inline void ParallelForInternal(int32 Num, FunctionType Body, PreWorkType CurrentThreadWorkToDoBeforeHelping, EParallelForFlags Flags, const TArrayView<ContextType>& Contexts)
+	{
+		if (LowLevelTasks::FScheduler::Get().GetNumWorkers())
+		{
+			NewParallelForInternal<FunctionType, PreWorkType, ContextType>(Num, Body, CurrentThreadWorkToDoBeforeHelping, Flags, Contexts);
 		}
 		else
 		{
-			check(!Data->bTriggered);
+			OldParallelForInternal<FunctionType, PreWorkType, ContextType>(Num, Body, CurrentThreadWorkToDoBeforeHelping, Flags, Contexts);
 		}
-		check(Data->NumCompleted.GetValue() == Data->Num);
-
-#if DO_CHECK
-		Data->bExited.store(true, std::memory_order_relaxed);
-#endif
-
-		// Data must live on until all of the tasks are cleared which might be long after this function exits
 	}
 }
 
@@ -443,9 +650,9 @@ namespace ParallelForImpl
 **/
 inline void ParallelFor(int32 Num, TFunctionRef<void(int32)> Body, bool bForceSingleThread, bool bPumpRenderingThread=false)
 {
-	ParallelForImpl::ParallelForInternal(Num, Body,
+	ParallelForImpl::ParallelForInternal(Num, Body, [](){},
 		(bForceSingleThread ? EParallelForFlags::ForceSingleThread : EParallelForFlags::None) | 
-		(bPumpRenderingThread ? EParallelForFlags::PumpRenderingThread : EParallelForFlags::None));
+		(bPumpRenderingThread ? EParallelForFlags::PumpRenderingThread : EParallelForFlags::None), TArrayView<nullptr_t>());
 }
 
 /**
@@ -458,7 +665,7 @@ inline void ParallelFor(int32 Num, TFunctionRef<void(int32)> Body, bool bForceSi
 template<typename FunctionType>
 inline void ParallelForTemplate(int32 Num, const FunctionType& Body, EParallelForFlags Flags = EParallelForFlags::None)
 {
-	ParallelForImpl::ParallelForInternal(Num, Body, Flags);
+	ParallelForImpl::ParallelForInternal(Num, Body, [](){}, Flags, TArrayView<nullptr_t>());
 }
 /** 
 	*	General purpose parallel for that uses the taskgraph for unbalanced tasks
@@ -472,7 +679,7 @@ inline void ParallelForTemplate(int32 Num, const FunctionType& Body, EParallelFo
 **/
 inline void ParallelFor(int32 Num, TFunctionRef<void(int32)> Body, EParallelForFlags Flags = EParallelForFlags::None)
 {
-	ParallelForImpl::ParallelForInternal(Num, Body, Flags);
+	ParallelForImpl::ParallelForInternal(Num, Body, [](){}, Flags, TArrayView<nullptr_t>());
 }
 
 /** 
@@ -485,9 +692,9 @@ inline void ParallelFor(int32 Num, TFunctionRef<void(int32)> Body, EParallelForF
 **/
 inline void ParallelForWithPreWork(int32 Num, TFunctionRef<void(int32)> Body, TFunctionRef<void()> CurrentThreadWorkToDoBeforeHelping, bool bForceSingleThread, bool bPumpRenderingThread = false)
 {
-	ParallelForImpl::ParallelForWithPreWorkInternal(Num, Body, CurrentThreadWorkToDoBeforeHelping,
+	ParallelForImpl::ParallelForInternal(Num, Body, CurrentThreadWorkToDoBeforeHelping,
 		(bForceSingleThread ? EParallelForFlags::ForceSingleThread : EParallelForFlags::None) |
-		(bPumpRenderingThread ? EParallelForFlags::PumpRenderingThread : EParallelForFlags::None));
+		(bPumpRenderingThread ? EParallelForFlags::PumpRenderingThread : EParallelForFlags::None), TArrayView<nullptr_t>());
 }
 
 /** 
@@ -500,7 +707,7 @@ inline void ParallelForWithPreWork(int32 Num, TFunctionRef<void(int32)> Body, TF
 **/
 inline void ParallelForWithPreWork(int32 Num, TFunctionRef<void(int32)> Body, TFunctionRef<void()> CurrentThreadWorkToDoBeforeHelping, EParallelForFlags Flags = EParallelForFlags::None)
 {
-	ParallelForImpl::ParallelForWithPreWorkInternal(Num, Body, CurrentThreadWorkToDoBeforeHelping, Flags);
+	ParallelForImpl::ParallelForInternal(Num, Body, CurrentThreadWorkToDoBeforeHelping, Flags, TArrayView<nullptr_t>());
 }
 
 /** 
@@ -527,7 +734,7 @@ inline void ParallelForWithTaskContext(TArray<ContextType, ContextAllocatorType>
 		{
 			new(&OutContexts[ContextIndex]) ContextType(ContextConstructor(ContextIndex, NumContexts));
 		}
-		ParallelForImpl::ParallelForInternal(Num, Body, Flags, TArrayView<ContextType>(OutContexts));
+		ParallelForImpl::ParallelForInternal(Num, Body, [](){}, Flags, TArrayView<ContextType>(OutContexts));
 	}
 }
 
@@ -549,6 +756,6 @@ inline void ParallelForWithTaskContext(TArray<ContextType, ContextAllocatorType>
 		const int32 NumContexts = ParallelForImpl::GetNumberOfThreadTasks(Num, Flags) + 1;
 		OutContexts.Reset();
 		OutContexts.AddDefaulted(NumContexts);
-		ParallelForImpl::ParallelForInternal(Num, Body, Flags, TArrayView<ContextType>(OutContexts));
+		ParallelForImpl::ParallelForInternal(Num, Body, [](){}, Flags, TArrayView<ContextType>(OutContexts));
 	}
 }
\ No newline at end of file
