FEM - fix USE_MALLOC_BINNED3 not uniformly defined for all platforms

diff --git a/Engine/Source/Runtime/Core/Public/Experimental/ConcurrentLinearAllocator.h b/Engine/Source/Runtime/Core/Public/Experimental/ConcurrentLinearAllocator.h
--- a/Engine/Source/Runtime/Core/Public/Experimental/ConcurrentLinearAllocator.h
+++ b/Engine/Source/Runtime/Core/Public/Experimental/ConcurrentLinearAllocator.h
@@ -9,6 +9,16 @@
 #include "Templates/UniquePtr.h"
 #include "Templates/UnrealTypeTraits.h"
 
+#ifdef USE_MALLOC_BINNED3
+#if USE_MALLOC_BINNED3
+#define SUPPORTS_VERY_LARGE_ALIGNMENTS 0
+#else
+#define SUPPORTS_VERY_LARGE_ALIGNMENTS 1
+#endif
+#else
+#define SUPPORTS_VERY_LARGE_ALIGNMENTS 1
+#endif
+
 #if __has_include(<sanitizer/asan_interface.h>)
 #include <sanitizer/asan_interface.h>
 #if defined(__SANITIZE_ADDRESS__)
@@ -130,7 +140,7 @@ struct FDefaultBlockAllocationTag
 template<typename BlockAllocationTag>
 class TConcurrentLinearAllocator
 {
-	static constexpr bool SupportsFastPath = ((BlockAllocationTag::BlockSize <= (64 * 1024)) || !USE_MALLOC_BINNED3) //MallocBinned3 only supports Alignment up to (64 * 1024)
+	static constexpr bool SupportsFastPath = ((BlockAllocationTag::BlockSize <= (64 * 1024)) || SUPPORTS_VERY_LARGE_ALIGNMENTS) //MallocBinned3 only supports Alignment up to (64 * 1024)
 		&& FMath::IsPowerOfTwo(BlockAllocationTag::BlockSize) //Aligndown only works with Pow2
 		&& !IS_ASAN_ENABLED && !BlockAllocationTag::RequiresAccurateSize // Only enabled when not using ASAN or when the Size of an allocation is not required 
 		&& BlockAllocationTag::Allocator::SupportsAlignment; //The allocator needs to align the BlockHeader by Blocksize to find the BlockHeader
