TTuple's per-element constructor is now conditionally explicit, allowing braced init: e.g. TTuple<float, const TCHAR*> Func() { return { 3.14f, TEXT("Hello") }; }

diff --git a/Engine/Source/Runtime/Core/Public/Templates/Tuple.h b/Engine/Source/Runtime/Core/Public/Templates/Tuple.h
--- a/Engine/Source/Runtime/Core/Public/Templates/Tuple.h
+++ b/Engine/Source/Runtime/Core/Public/Templates/Tuple.h
@@ -662,7 +662,7 @@ namespace UE::Core::Private::Tuple
 	};
 
 	template <typename... Given, typename... Deduced>
-	std::enable_if_t<TAnd<TIsConstructible<Given, Deduced&&>...>::Value> ConstructibleConceptCheck(Deduced&&...);
+	std::enable_if_t<std::conjunction_v<std::is_constructible<Given, Deduced&&>...>> ConstructibleConceptCheck(Deduced&&...);
 
 	template <typename... Given, typename... Deduced>
 	decltype(ConceptCheckingHelper((DeclVal<Given>() = DeclVal<Deduced&&>(), 0)...)) AssignableConceptCheck(Deduced&&...);
@@ -675,14 +675,36 @@ private:
 	typedef UE::Core::Private::Tuple::TTupleBase<TMakeIntegerSequence<uint32, sizeof...(Types)>, Types...> Super;
 
 public:
+#ifdef __cpp_conditional_explicit
 	template <
 		typename... ArgTypes,
 		decltype(UE::Core::Private::Tuple::ConstructibleConceptCheck<Types...>(DeclVal<ArgTypes&&>()...))* = nullptr
 	>
+	explicit(!std::conjunction_v<std::is_convertible<ArgTypes&&, Types>...>>) TTuple(ArgTypes&&... Args)
+		: Super(UE::Core::Private::Tuple::ForwardingConstructor, Forward<ArgTypes>(Args)...)
+	{
+	}
+#else
+	template <
+		typename... ArgTypes,
+		decltype(UE::Core::Private::Tuple::ConstructibleConceptCheck<Types...>(DeclVal<ArgTypes&&>()...))* = nullptr,
+		std::enable_if_t<std::conjunction_v<std::is_convertible<ArgTypes&&, Types>...>>* = nullptr
+	>
+	TTuple(ArgTypes&&... Args)
+		: Super(UE::Core::Private::Tuple::ForwardingConstructor, Forward<ArgTypes>(Args)...)
+	{
+	}
+
+	template <
+		typename... ArgTypes,
+		decltype(UE::Core::Private::Tuple::ConstructibleConceptCheck<Types...>(DeclVal<ArgTypes&&>()...))* = nullptr,
+		std::enable_if_t<!std::conjunction_v<std::is_convertible<ArgTypes&&, Types>...>>* = nullptr
+	>
 	explicit TTuple(ArgTypes&&... Args)
 		: Super(UE::Core::Private::Tuple::ForwardingConstructor, Forward<ArgTypes>(Args)...)
 	{
 	}
+#endif
 
 	template <
 		typename... OtherTypes,
