[Insights] Always start Trace worker thread, even if SupportsMultithreading is false.

Trace normally runs with a worker thread to write data regularly, however when `FPlatformProcess:SupportsMultithreading` was false we did not start the worker thread, instead we relied on EndFrame delegate to manually pump Trace. When tracing is very intense (e.g. memory tracing) could lead to data building up internally, up to 1 GB of data by the time the first frame callback came. We have recently added a special sync type packet in the stream to deal with edge cases related to late-connect, so analysis stalled until those came trough. However since the backlog of data has to be transmitted first, those sync packet ended up after the initial data. Analysis had to (very inefficiently) get trough that 1Gb of data before the first sync packet came through and events starts to appear.

With this change trace always start the worker thread (if we're intending to fork) regardless of the SupportsMultithreading setting.

diff --git a/Engine/Source/Runtime/Core/Private/ProfilingDebugging/TraceAuxiliary.cpp b/Engine/Source/Runtime/Core/Private/ProfilingDebugging/TraceAuxiliary.cpp
--- a/Engine/Source/Runtime/Core/Private/ProfilingDebugging/TraceAuxiliary.cpp
+++ b/Engine/Source/Runtime/Core/Private/ProfilingDebugging/TraceAuxiliary.cpp
@@ -1197,7 +1197,7 @@ void FTraceAuxiliary::Initialize(const TCHAR* CommandLine)
 	// Always register end frame updates. This path is short circuited if a worker thread
 	// exists.
 	GTraceAuxiliary.StartEndFramePump();
-	if (FPlatformProcess::SupportsMultithreading() && !FForkProcessHelper::IsForkRequested())
+	if (!FForkProcessHelper::IsForkRequested())
 	{
 		GTraceAuxiliary.StartWorkerThread();
 	}
