FrameRate: add support for parsing FFrameRates from strings that end with drop frame tokens "ndf" or "df"

This allows parsing a frame rate from an FPS string that denotes timecode drop frame-ness with
either "ndf" for non-drop frame or "df" for drop frame as a suffix. Drop frame timecode is only
supported for NTSC_30 (29.97fps) and NTSC_60 (59.94fps) frame rates, so the only FPS strings
that include the "df" token that will parse successfully are "29.97df" and "59.94df". Any other FPS
value that includes the "df" token will result in an error. All NTSC FPS values (23.976, 29.97, and 59.94)
as well as any integer FPS value can include the "ndf" token.

Note however that the presence or absence of the "ndf" and "df" tokens in the string will have
no effect on the resulting FFrameRate. This simply ensures that strings that include those tokens
can be parsed into the correct frame rate.

diff --git a/Engine/Source/Runtime/Core/Private/Misc/FrameRate.cpp b/Engine/Source/Runtime/Core/Private/Misc/FrameRate.cpp
--- a/Engine/Source/Runtime/Core/Private/Misc/FrameRate.cpp
+++ b/Engine/Source/Runtime/Core/Private/Misc/FrameRate.cpp
@@ -12,9 +12,13 @@ namespace
 {
 	struct FSeconds         { static const TCHAR* const Moniker; };
 	struct FFramesPerSecond { static const TCHAR* const Moniker; };
+	struct FNonDropFrame    { static const TCHAR* const Moniker; };
+	struct FDropFrame       { static const TCHAR* const Moniker; };
 
 	const TCHAR* const FFramesPerSecond::Moniker = TEXT("fps");
 	const TCHAR* const FSeconds::Moniker         = TEXT("s");
+	const TCHAR* const FNonDropFrame::Moniker    = TEXT("ndf");
+	const TCHAR* const FDropFrame::Moniker       = TEXT("df");
 
 	const double NTSC_24_FPS = 23.976;
 	const double NTSC_30_FPS = 29.97;
@@ -29,6 +33,8 @@ namespace
 DEFINE_EXPRESSION_NODE_TYPE(FFrameRate, 0x4EDAA92F, 0xB75E4B9E, 0xB7E0ABC2, 0x8D981FCB)
 DEFINE_EXPRESSION_NODE_TYPE(FSeconds, 0x3DC5F60D, 0x934E4753, 0xA80CD6D0, 0xE9EB4640)
 DEFINE_EXPRESSION_NODE_TYPE(FFramesPerSecond, 0x8423B4AE, 0x2FF64795, 0xA7EFFAC0, 0xC560531A)
+DEFINE_EXPRESSION_NODE_TYPE(FNonDropFrame, 0xEE67D9FA, 0x435C4991, 0xDF530089, 0x8B3019D0)
+DEFINE_EXPRESSION_NODE_TYPE(FDropFrame, 0x56AF8BD1, 0x431DF28A, 0x8ECC1EB7, 0xC00C5E5D)
 
 const float FFrameTime::MaxSubframe = 0.99999994f;
 
@@ -45,12 +51,16 @@ public:
 		TokenDefinitions.IgnoreWhitespace();
 		TokenDefinitions.DefineToken(&ConsumeSymbol<FSeconds>);
 		TokenDefinitions.DefineToken(&ConsumeSymbol<FFramesPerSecond>);
+		TokenDefinitions.DefineToken(&ConsumeSymbol<FNonDropFrame>);
+		TokenDefinitions.DefineToken(&ConsumeSymbol<FDropFrame>);
 		TokenDefinitions.DefineToken(&ConsumeSymbol<FForwardSlash>);
 		TokenDefinitions.DefineToken(&ConsumeLocalizedNumberWithAgnosticFallback);
 
 		Grammar.DefineBinaryOperator<FForwardSlash>(1);
 		Grammar.DefinePostUnaryOperator<FSeconds>();
 		Grammar.DefinePostUnaryOperator<FFramesPerSecond>();
+		Grammar.DefinePostUnaryOperator<FNonDropFrame>();
+		Grammar.DefinePostUnaryOperator<FDropFrame>();
 
 		JumpTable.MapPostUnary<FSeconds>(
 			[](double In) -> FExpressionResult
@@ -63,6 +73,15 @@ public:
 			}
 		);
 		JumpTable.MapPostUnary<FFramesPerSecond>([](double In) -> FExpressionResult { return MakeFrameRateFromFPS(In); });
+		JumpTable.MapPostUnary<FNonDropFrame>([](double In) -> FExpressionResult { return MakeFrameRateFromFPS(In); });
+		JumpTable.MapPostUnary<FDropFrame>([](double In) -> FExpressionResult {
+			if (!FMath::IsNearlyEqual(In, NTSC_30_FPS, DOUBLE_KINDA_SMALL_NUMBER) && !FMath::IsNearlyEqual(In, NTSC_60_FPS, DOUBLE_KINDA_SMALL_NUMBER))
+			{
+				return MakeError(FText::Format(LOCTEXT("InvalidDropFrameFPS", "Drop frame not supported for FPS specified: {0}.\nDrop frame timecode is only supported for NTSC_30 ({1}fps) and NTSC_60 ({2}fps) frame rates."), In, NTSC_30_FPS, NTSC_60_FPS));
+			}
+
+			return MakeFrameRateFromFPS(In);
+		});
 		JumpTable.MapBinary<FForwardSlash>([](double A, double B) -> FExpressionResult { return MakeFrameRate(A, B); });
 	}
 
