FEM - Alignment issues on mac which broke UHT

diff --git a/Engine/Source/Runtime/Core/Public/Experimental/ConcurrentLinearAllocator.h b/Engine/Source/Runtime/Core/Public/Experimental/ConcurrentLinearAllocator.h
--- a/Engine/Source/Runtime/Core/Public/Experimental/ConcurrentLinearAllocator.h
+++ b/Engine/Source/Runtime/Core/Public/Experimental/ConcurrentLinearAllocator.h
@@ -51,11 +51,17 @@ struct FOsAllocator
 
 struct FAlignedAllocator
 {
-	static constexpr bool SupportsAlignment = true;
+	static constexpr bool SupportsAlignment = !(PLATFORM_MAC || PLATFORM_IOS || PLATFORM_TVOS); //TODO: UHT on Mac OS returned false alignment.
 
 	FORCEINLINE static void* Malloc(SIZE_T Size, uint32 Alignment)
 	{
+#if DO_CHECK
+		void* Ret = AnsiMalloc(Size, Alignment);
+		check(IsAligned(Ret, Alignment));
+		return Ret;
+#else
 		return AnsiMalloc(Size, Alignment);
+#endif
 	}
 
 	FORCEINLINE static void Free(void* Pointer, SIZE_T Size)
