Fix double path seperators in NotAllowedUnattendedBugReports path.

diff --git a/Engine/Source/Runtime/Core/Private/GenericPlatform/GenericPlatformCrashContext.cpp b/Engine/Source/Runtime/Core/Private/GenericPlatform/GenericPlatformCrashContext.cpp
--- a/Engine/Source/Runtime/Core/Private/GenericPlatform/GenericPlatformCrashContext.cpp
+++ b/Engine/Source/Runtime/Core/Private/GenericPlatform/GenericPlatformCrashContext.cpp
@@ -434,7 +434,7 @@ void FGenericCrashContext::InitializeFromConfig()
 	// Write a marker file to disk indicating the user has allowed unattended crash reports being
 	// sent. This allows us to submit reports for crashes during static initialization when user
 	// settings are not available. 
-	FString MarkerFilePath = FString::Printf(TEXT("%s/NotAllowedUnattendedBugReports"), FPlatformProcess::ApplicationSettingsDir());
+	FString MarkerFilePath = FPaths::Combine(FPlatformProcess::ApplicationSettingsDir(), TEXT("NotAllowedUnattendedBugReports"));
 	if (!NCached::UserSettings.bSendUnattendedBugReports)
 	{
 		TUniquePtr<IFileHandle> File(FPlatformFileManager::Get().GetPlatformFile().OpenWrite(*MarkerFilePath));
diff --git a/Engine/Source/Runtime/Core/Private/Windows/WindowsPlatformCrashContext.cpp b/Engine/Source/Runtime/Core/Private/Windows/WindowsPlatformCrashContext.cpp
--- a/Engine/Source/Runtime/Core/Private/Windows/WindowsPlatformCrashContext.cpp
+++ b/Engine/Source/Runtime/Core/Private/Windows/WindowsPlatformCrashContext.cpp
@@ -701,7 +701,7 @@ int32 ReportCrashForMonitor(
 		// Crashes before config system is ready (e.g. during static init) we cannot know the user
 		// settings for sending unattended. We check for the existense of the marker file from previous
 		// sessions, otherwise we cannot send a report at all.
-		FString NotAllowedUnattendedBugReportMarkerPath = FString::Printf(TEXT("%s/NotAllowedUnattendedBugReports"), FWindowsPlatformProcess::ApplicationSettingsDir());
+		FString NotAllowedUnattendedBugReportMarkerPath = FPaths::Combine(FWindowsPlatformProcess::ApplicationSettingsDir(), TEXT("NotAllowedUnattendedBugReports"));
 		if (::PathFileExistsW(*NotAllowedUnattendedBugReportMarkerPath))
 		{
 			bSendUnattendedBugReports = false;
