Read proc smaps_rollup in FUnixPlatformMemory::GetExtendedStats() if it exists

smaps_rollup was added in v4.14. Details on this page:
  https://kernelnewbies.org/Linux_4.14

It should be a constant ~1ms vs 6.8ms (and far higher) when reading entire smaps file
  (Tested on my 5.15 kernel)

smaps_rollup info:
  https://patchwork.kernel.org/project/linux-fsdevel/patch/20170812022148.178293-1-dancol[REDACTED_HANDLE]/

Smaps info:
> FPlatformMemory::GetExtendedStats() is slow and took 28ms in my local tests.
> Moreover, it can be even slower because the time it takes is a (hopefully
> linear) function of the number of memory allocations.

diff --git a/Engine/Source/Runtime/Core/Private/Unix/UnixPlatformMemory.cpp b/Engine/Source/Runtime/Core/Private/Unix/UnixPlatformMemory.cpp
--- a/Engine/Source/Runtime/Core/Private/Unix/UnixPlatformMemory.cpp
+++ b/Engine/Source/Runtime/Core/Private/Unix/UnixPlatformMemory.cpp
@@ -828,37 +828,80 @@ FPlatformMemoryStats FUnixPlatformMemory::GetStats()
 
 FExtendedPlatformMemoryStats FUnixPlatformMemory::GetExtendedStats()
 {
-	FExtendedPlatformMemoryStats MemoryStats;
-
-	// More /proc "API" :/
-	MemoryStats.Shared_Clean = 0;
-	MemoryStats.Shared_Dirty = 0;
-	MemoryStats.Private_Clean = 0;
-	MemoryStats.Private_Dirty = 0;
-	if (FILE* ProcSMaps = fopen("/proc/self/smaps", "r"))
+	const ANSICHAR Shared_CleanStr[] = "Shared_Clean:";
+	const ANSICHAR Shared_DirtyStr[] = "Shared_Dirty:";
+	const ANSICHAR Private_CleanStr[] = "Private_Clean:";
+	const ANSICHAR Private_DirtyStr[] = "Private_Dirty:";
+	FExtendedPlatformMemoryStats MemoryStats = { 0 };
+
+	// ~ 1.06ms per call on my Threadripper 3990X w/ Debian Testing 5.15.0-2-amd64
+	if (FILE* ProcSMapsRollup = fopen("/proc/self/smaps_rollup", "r"))
+	{
+		struct
+		{
+			const ANSICHAR* Name;
+			uint32 NameLen;
+			SIZE_T* Addr;
+		} SMapsFields[] =
+		{
+			{ Shared_CleanStr,  sizeof(Shared_CleanStr) - 1,  &MemoryStats.Shared_Clean },
+			{ Shared_DirtyStr,  sizeof(Shared_DirtyStr) - 1,  &MemoryStats.Shared_Dirty },
+			{ Private_CleanStr, sizeof(Private_CleanStr) - 1, &MemoryStats.Private_Clean },
+			{ Private_DirtyStr, sizeof(Private_DirtyStr) - 1, &MemoryStats.Private_Dirty },
+		};
+		const uint32 NumFields = UE_ARRAY_COUNT(SMapsFields);
+		uint32 FieldsFound = 0;
+
+		while (FieldsFound < NumFields)
+		{
+			ANSICHAR LineBuffer[256];
+			ANSICHAR *Line = fgets(LineBuffer, UE_ARRAY_COUNT(LineBuffer), ProcSMapsRollup);
+
+			if (Line == nullptr)
+			{
+				// eof or an error
+				break;
+			}
+
+			for (uint32 i = 0; i < NumFields; i++)
+			{
+				if (!FCStringAnsi::Strncmp(SMapsFields[i].Name, Line, SMapsFields[i].NameLen))
+				{
+					*SMapsFields[i].Addr = atoll(Line + SMapsFields[i].NameLen) * 1024ULL;
+					FieldsFound++;
+					break;
+				}
+			}
+		}
+
+		fclose(ProcSMapsRollup);
+	}
+	// ~ 6.8ms per call on my Threadripper 3990X w/ Debian Testing 5.15.0-2-amd64 in TestPAL.
+	// Potentially far higher though.
+	else if (FILE* ProcSMaps = fopen("/proc/self/smaps", "r"))
 	{
 		do
 		{
-			char LineBuffer[256] = { 0 };
-			char *Line = fgets(LineBuffer, UE_ARRAY_COUNT(LineBuffer), ProcSMaps);
+			ANSICHAR LineBuffer[256] = { 0 };
+			ANSICHAR *Line = fgets(LineBuffer, UE_ARRAY_COUNT(LineBuffer), ProcSMaps);
 			if (Line == nullptr)
 			{
 				break;	// eof or an error
 			}
 
-			if (strstr(Line, "Shared_Clean:") == Line)
+			if (strstr(Line, Shared_CleanStr) == Line)
 			{
 				MemoryStats.Shared_Clean += UnixPlatformMemory::GetBytesFromStatusLine(Line);
 			}
-			else if (strstr(Line, "Shared_Dirty:") == Line)
+			else if (strstr(Line, Shared_DirtyStr) == Line)
 			{
 				MemoryStats.Shared_Dirty += UnixPlatformMemory::GetBytesFromStatusLine(Line);
 			}
-			if (strstr(Line, "Private_Clean:") == Line)
+			if (strstr(Line, Private_CleanStr) == Line)
 			{
 				MemoryStats.Private_Clean += UnixPlatformMemory::GetBytesFromStatusLine(Line);
 			}
-			else if (strstr(Line, "Private_Dirty:") == Line)
+			else if (strstr(Line, Private_DirtyStr) == Line)
 			{
 				MemoryStats.Private_Dirty += UnixPlatformMemory::GetBytesFromStatusLine(Line);
 			}
