Fix failing checkSlow

diff --git a/Engine/Source/Runtime/Core/Public/Async/ParallelFor.h b/Engine/Source/Runtime/Core/Public/Async/ParallelFor.h
--- a/Engine/Source/Runtime/Core/Public/Async/ParallelFor.h
+++ b/Engine/Source/Runtime/Core/Public/Async/ParallelFor.h
@@ -506,6 +506,11 @@ namespace ParallelForImpl
 			{
 			}
 
+			inline const FDataHandle& GetData() const
+			{
+				return Data;
+			}
+
 			inline bool Run(const bool bIsMaster = true, const bool bIsBackgroundPriority = false) noexcept
 			{
 				FMemMark Mark(FMemStack::Get());
@@ -624,7 +629,7 @@ namespace ParallelForImpl
 				FinishedSignal->Wait();
 			}
 		}
-		checkSlow(Data->BatchItem.load(std::memory_order_relaxed) * Data->BatchSize >= Data->Num);
+		checkSlow(LocalExecutor.GetData()->BatchItem.load(std::memory_order_relaxed) * LocalExecutor.GetData()->BatchSize >= LocalExecutor.GetData()->Num);
 	}
 
 	template<typename FunctionType, typename PreWorkType, typename ContextType>
