- Restored Mac/WindowsPlatformProperties::SupportsMeshLODStreaming() to check for HasEditorOnlyData, because it was causing cooked client meshes to crash when loaded by the editor. This may be a temporary fix

diff --git a/Engine/Source/Runtime/Core/Public/Mac/MacPlatformProperties.h b/Engine/Source/Runtime/Core/Public/Mac/MacPlatformProperties.h
--- a/Engine/Source/Runtime/Core/Public/Mac/MacPlatformProperties.h
+++ b/Engine/Source/Runtime/Core/Public/Mac/MacPlatformProperties.h
@@ -130,7 +130,7 @@ struct FMacPlatformProperties
 
 	static FORCEINLINE bool SupportsMeshLODStreaming()
 	{
-		return !IsServerOnly();
+		return !IsServerOnly() && !HasEditorOnlyData();
 	}
 
 	static FORCEINLINE bool SupportsVirtualTextureStreaming()
diff --git a/Engine/Source/Runtime/Core/Public/Windows/WindowsPlatformProperties.h b/Engine/Source/Runtime/Core/Public/Windows/WindowsPlatformProperties.h
--- a/Engine/Source/Runtime/Core/Public/Windows/WindowsPlatformProperties.h
+++ b/Engine/Source/Runtime/Core/Public/Windows/WindowsPlatformProperties.h
@@ -85,7 +85,7 @@ struct FWindowsPlatformProperties
 
 	static FORCEINLINE bool SupportsMeshLODStreaming()
 	{
-		return !IsServerOnly();
+		return !IsServerOnly() && !HasEditorOnlyData();
 	}
 
 	static FORCEINLINE bool SupportsGrayscaleSRGB()
