[Insights]
  - Updated help for trace console commands.
  - Fixed comments and code style issues in TraceAuxiliary.

diff --git a/Engine/Source/Runtime/Core/Private/ProfilingDebugging/TraceAuxiliary.cpp b/Engine/Source/Runtime/Core/Private/ProfilingDebugging/TraceAuxiliary.cpp
--- a/Engine/Source/Runtime/Core/Private/ProfilingDebugging/TraceAuxiliary.cpp
+++ b/Engine/Source/Runtime/Core/Private/ProfilingDebugging/TraceAuxiliary.cpp
@@ -234,7 +234,7 @@ void FTraceAuxiliaryImpl::RemoveChannel(const TCHAR* Name)
 ////////////////////////////////////////////////////////////////////////////////
 bool FTraceAuxiliaryImpl::Connect(ETraceConnectType Type, const TCHAR* Parameter, const FTraceAuxiliary::FLogCategoryAlias& LogCategory)
 {
-	// Connect/write to file. But only if we're not already sending/writing
+	// Connect/write to file, but only if we're not already sending/writing.
 	bool bConnected = UE::Trace::IsTracing();
 	if (!bConnected)
 	{
@@ -248,7 +248,7 @@ bool FTraceAuxiliaryImpl::Connect(ETraceConnectType Type, const TCHAR* Parameter
 			else
 			{
 				UE_LOG_REF(LogCategory, Error, TEXT("Trace failed to connect (trace server: %s)!"), Parameter ? Parameter : TEXT(""));
-			}	
+			}
 		}
 
 		else if (Type == ETraceConnectType::File)
@@ -261,7 +261,7 @@ bool FTraceAuxiliaryImpl::Connect(ETraceConnectType Type, const TCHAR* Parameter
 			else
 			{
 				UE_LOG_REF(LogCategory, Error, TEXT("Trace failed to connect (file: \"%s\")!"), Parameter ? Parameter : TEXT(""));
-			}	
+			}
 		}
 	}
 
@@ -396,10 +396,10 @@ bool FTraceAuxiliaryImpl::SendToHost(const TCHAR* Host, const FTraceAuxiliary::F
 bool FTraceAuxiliaryImpl::WriteToFile(const TCHAR* InPath, const FTraceAuxiliary::FLogCategoryAlias& LogCategory)
 {
 	const FStringView Path(InPath);
-	
-	// Default file name functor 
-	auto GetDefaultName = [] {return FDateTime::Now().ToString(TEXT("%Y%m%d_%H%M%S.utrace"));};
-	
+
+	// Default file name functor
+	auto GetDefaultName = [] { return FDateTime::Now().ToString(TEXT("%Y%m%d_%H%M%S.utrace")); };
+
 	if (Path.IsEmpty())
 	{
 		const FString Name = GetDefaultName();
@@ -570,8 +570,8 @@ static void TraceAuxiliaryConnectEpilogue()
 	// Give the user some feedback that everything's underway.
 	TStringBuilder<128> Channels;
 	GTraceAuxiliary.GetActiveChannelsString(Channels);
-	
-	UE_LOG(LogConsoleResponse, Log, TEXT("Active channels: %s"), Channels.ToString());
+
+	UE_LOG(LogConsoleResponse, Display, TEXT("Enabled channels: %s"), Channels.ToString());
 }
 
 ////////////////////////////////////////////////////////////////////////////////
@@ -579,7 +579,7 @@ static void TraceAuxiliarySend(const TArray<FString>& Args)
 {
 	if (Args.Num() < 1)
 	{
-		UE_LOG(LogConsoleResponse, Warning, TEXT("No host name given; Trace.Send <Host> [ChannelSet]"));
+		UE_LOG(LogConsoleResponse, Warning, TEXT("No host name given. Usage: Trace.Send <Host> [ChannelSet]"));
 		return;
 	}
 
@@ -594,30 +594,38 @@ static void TraceAuxiliarySend(const TArray<FString>& Args)
 ////////////////////////////////////////////////////////////////////////////////
 static void TraceAuxiliaryFile(const TArray<FString>& Args)
 {
+	const TCHAR* Filepath = nullptr;
 	const TCHAR* Channels = nullptr;
-	const TCHAR* File = nullptr;
+
 	if (Args.Num() == 2)
 	{
-		File = *Args[0];
+		Filepath = *Args[0];
 		Channels = *Args[1];
 	}
 	else if (Args.Num() == 1)
 	{
-		// Try to detect if the first argument is a file or path
-		if (!FCString::Strchr(*Args[0], TEXT('/')) && !FCString::Strchr(*Args[0], TEXT('\\'))
-			&& !FCString::Strchr(*Args[0], TEXT('.')) && !FCString::Strchr(*Args[0], TEXT(':')) )
+		// Try to detect if the first argument is a file path.
+		if (!FCString::Strchr(*Args[0], TEXT('/')) ||
+			!FCString::Strchr(*Args[0], TEXT('\\')) ||
+			!FCString::Strchr(*Args[0], TEXT('.')) ||
+			!FCString::Strchr(*Args[0], TEXT(':')))
 		{
-			Channels = *Args[0];
+			Filepath = *Args[0];
 		}
 		else
 		{
-			File = *Args[0];
+			Channels = *Args[0];
 		}
 	}
-	
+	else if (Args.Num() > 2)
+	{
+		UE_LOG(LogConsoleResponse, Warning, TEXT("Invalid arguments. Usage: Trace.File [Path] [ChannelSet]"));
+		return;
+	}
+
 	FTraceAuxiliary::Options Opts;
 	Opts.bNoWorkerThread = true;
-	if (FTraceAuxiliary::Start(FTraceAuxiliary::EConnectionType::File, File, Channels, &Opts, LogConsoleResponse))
+	if (FTraceAuxiliary::Start(FTraceAuxiliary::EConnectionType::File, Filepath, Channels, &Opts, LogConsoleResponse))
 	{
 		TraceAuxiliaryConnectEpilogue();
 	}
@@ -633,29 +641,29 @@ static void TraceAuxiliaryStart(const TArray<FString>& Args)
 ////////////////////////////////////////////////////////////////////////////////
 static void TraceAuxiliaryStop()
 {
-	UE_LOG(LogConsoleResponse, Log, TEXT("Tracing stopped."));
+	UE_LOG(LogConsoleResponse, Display, TEXT("Tracing stopped."));
 	GTraceAuxiliary.Stop();
 }
 
 ////////////////////////////////////////////////////////////////////////////////
 static void TraceAuxiliaryPause()
 {
-	UE_LOG(LogConsoleResponse, Log, TEXT("Tracing paused"));
+	UE_LOG(LogConsoleResponse, Display, TEXT("Tracing paused."));
 	GTraceAuxiliary.PauseChannels();
 }
 
 ////////////////////////////////////////////////////////////////////////////////
 static void TraceAuxiliaryResume()
 {
-	UE_LOG(LogConsoleResponse, Log, TEXT("Tracing resumed"));
+	UE_LOG(LogConsoleResponse, Display, TEXT("Tracing resumed."));
 	GTraceAuxiliary.ResumeChannels();
 }
 
 ////////////////////////////////////////////////////////////////////////////////
 static void TraceAuxiliaryStatus()
 {
-	UE_LOG(LogConsoleResponse, Display, TEXT("Trace status -------------------------------"));
-	
+	UE_LOG(LogConsoleResponse, Display, TEXT("Trace status ----------------------------------------------------------"));
+
 	// Status of data connection
 	const TCHAR* Dest = GTraceAuxiliary.GetDest();
 	TStringBuilder<256> ConnectionStr;
@@ -668,59 +676,58 @@ static void TraceAuxiliaryStatus()
 		ConnectionStr.Append(TEXT("Not tracing"));
 	}
 	UE_LOG(LogConsoleResponse, Display, TEXT("- Connection: %s"), ConnectionStr.ToString());
-	
+
 	// Stats
 	UE::Trace::FStatistics Stats;
 	UE::Trace::GetStatistics(Stats);
-	constexpr float MbDiv = 1024*1024;
+	constexpr double MiB = 1.0 / (1024.0 * 1024.0);
 	UE_LOG(LogConsoleResponse, Display, TEXT("- Memory used: %.02f MiB"),
-		float(Stats.MemoryUsed) / MbDiv);
-	UE_LOG(LogConsoleResponse, Display, TEXT("- Important Events cache: %.02f MiB, waste: %0.02f MiB"),
-		float(Stats.CacheUsed) / MbDiv,
-		float(Stats.CacheWaste) / MbDiv);
-	UE_LOG(LogConsoleResponse, Display, TEXT("- Sent: %.2f MiB"),
-		float(Stats.BytesSent) / MbDiv);
-	
+		double(Stats.MemoryUsed) * MiB);
+	// Stats.CacheWaste is actually the total size of allocated buffers.
+	UE_LOG(LogConsoleResponse, Display, TEXT("- Important Events cache: %.02f MiB (%.02f MiB used + %0.02f MiB waste)"),
+		double(Stats.CacheWaste) * MiB,
+		double(Stats.CacheUsed) * MiB,
+		double(Stats.CacheWaste - Stats.CacheUsed) * MiB);
+	UE_LOG(LogConsoleResponse, Display, TEXT("- Sent: %.02f MiB"),
+		double(Stats.BytesSent) * MiB);
+
 	// Channels
 	struct EnumerateType
 	{
-		TStringBuilder<128> EnabledChannelsStr;
-		uint32 EnabledChannelsLastLineBreak = 0;
-		TStringBuilder<1024> DisabledChannelsStr;
-		uint32 DisabledChannelsLastLineBreak = 0;
+		TStringBuilder<512> ChannelsStr;
+		uint32 Count = 0;
+#if WITH_EDITOR
+		uint32 LineLen = 50;
+#else
+		uint32 LineLen = 20;
+#endif
 
-		void InsertLineBreaks()
+		void AddChannel(const FAnsiStringView& NameView)
 		{
-			if ((EnabledChannelsStr.Len() - EnabledChannelsLastLineBreak) > 100)
+			if (Count++ > 0)
 			{
-				EnabledChannelsStr << TEXT("\n");
-				EnabledChannelsLastLineBreak = DisabledChannelsStr.Len();
+				ChannelsStr << TEXT(", ");
+				LineLen += 2;
 			}
-			if ((DisabledChannelsStr.Len() - DisabledChannelsLastLineBreak) > 100)
+			if (LineLen + NameView.Len() > 100)
 			{
-				DisabledChannelsStr << TEXT("\n");
-				DisabledChannelsLastLineBreak = DisabledChannelsStr.Len();
+				ChannelsStr << TEXT("\n    ");
+				LineLen = 4;
 			}
+			ChannelsStr << NameView;
+			LineLen += NameView.Len();
 		}
-	} EnumerateUserData;
+	} EnumerateUserData[2];
 	UE::Trace::EnumerateChannels([](const ANSICHAR* Name, bool bEnabled, void* User)
 	{
 		EnumerateType* EnumerateUser = static_cast<EnumerateType*>(User);
 		FAnsiStringView NameView = FAnsiStringView(Name).LeftChop(7); // Remove "Channel" suffix
-		if (bEnabled)
-		{
-			EnumerateUser->EnabledChannelsStr << NameView << TEXT(", ");
-		}
-		else
-		{
-			EnumerateUser->DisabledChannelsStr << NameView << TEXT(", ");
-		}
-		EnumerateUser->InsertLineBreaks();
-	}, &EnumerateUserData);
-	UE_LOG(LogConsoleResponse, Display, TEXT("- Enabled channels:\n %s"), EnumerateUserData.EnabledChannelsStr.ToString());
-	UE_LOG(LogConsoleResponse, Display, TEXT("- Available channels:\n %s"), EnumerateUserData.DisabledChannelsStr.ToString());
-	
-	UE_LOG(LogConsoleResponse, Display, TEXT("-------------------------------------------"));
+		EnumerateUser[bEnabled ? 0 : 1].AddChannel(NameView);
+	}, EnumerateUserData);
+	UE_LOG(LogConsoleResponse, Display, TEXT("- Enabled channels: %s"), EnumerateUserData[0].Count == 0 ? TEXT("<none>") : EnumerateUserData[0].ChannelsStr.ToString());
+	UE_LOG(LogConsoleResponse, Display, TEXT("- Available channels: %s"), EnumerateUserData[1].ChannelsStr.ToString());
+
+	UE_LOG(LogConsoleResponse, Display, TEXT("-----------------------------------------------------------------------"));
 }
 
 ////////////////////////////////////////////////////////////////////////////////
@@ -749,7 +756,7 @@ static void TraceAuxiliaryDisableChannels(const TArray<FString>& Args)
 	{
 		GTraceAuxiliary.DisableChannels(*Args[0]);
 	}
-	
+
 	TStringBuilder<128> EnabledChannels;
 	GTraceAuxiliary.GetActiveChannelsString(EnabledChannels);
 	UE_LOG(LogConsoleResponse, Display, TEXT("Enabled channels: %s"), EnabledChannels.ToString());
@@ -758,17 +765,18 @@ static void TraceAuxiliaryDisableChannels(const TArray<FString>& Args)
 ////////////////////////////////////////////////////////////////////////////////
 static FAutoConsoleCommand TraceAuxiliarySendCmd(
 	TEXT("Trace.Send"),
-	TEXT("Send trace data to the trace store; Trace.Send <Host> [ChannelSet]"
-		" where ChannelSet is either comma-separated list of trace channels, a Config/Trace.ChannelPresets key, or optional."),
+	TEXT("<Host> [ChannelSet] - Starts tracing to a trace store."
+		" <Host> is the IP address or hostname of the trace store."
+		" ChannelSet is comma-separated list of trace channels/presets to be enabled."
+	),
 	FConsoleCommandWithArgsDelegate::CreateStatic(TraceAuxiliarySend)
 );
 
 ////////////////////////////////////////////////////////////////////////////////
 static FAutoConsoleCommand TraceAuxiliaryStartCmd(
 	TEXT("Trace.Start"),
-	TEXT(
-		"(Deprecated: Use Trace.File) Begin tracing profiling events to a file; Trace.Start [ChannelSet]"
-		" where ChannelSet is either comma-separated list of trace channels, a Config/Trace.ChannelPresets key, or optional."
+	TEXT("[ChannelSet] - (Deprecated: Use Trace.File instead.) Starts tracing to a file."
+		" ChannelSet is comma-separated list of trace channels/presets to be enabled."
 	),
 	FConsoleCommandWithArgsDelegate::CreateStatic(TraceAuxiliaryStart)
 );
@@ -776,10 +784,9 @@ static FAutoConsoleCommand TraceAuxiliaryStartCmd(
 ////////////////////////////////////////////////////////////////////////////////
 static FAutoConsoleCommand TraceAuxiliaryFileCmd(
 	TEXT("Trace.File"),
-	TEXT(
-		"Begin tracing profiling events to a file; Trace.File [Filepath.utrace] [ChannelSet]"
-		" where ChannelSet is either comma-separated list of trace channels, a Config/Trace.ChannelPresets key, or optional."
-		" Either Filepath or ChannelSet can be excluded."
+	TEXT("[Path] [ChannelSet] - Starts tracing to a file."
+		" ChannelSet is comma-separated list of trace channels/presets to be enabled."
+		" Either Path or ChannelSet can be excluded."
 	),
 	FConsoleCommandWithArgsDelegate::CreateStatic(TraceAuxiliaryFile)
 );
@@ -787,45 +794,47 @@ static FAutoConsoleCommand TraceAuxiliaryFileCmd(
 ////////////////////////////////////////////////////////////////////////////////
 static FAutoConsoleCommand TraceAuxiliaryStopCmd(
 	TEXT("Trace.Stop"),
-	TEXT("Stops tracing profiling events"),
+	TEXT("Stops tracing profiling events."),
 	FConsoleCommandDelegate::CreateStatic(TraceAuxiliaryStop)
 );
 
 ////////////////////////////////////////////////////////////////////////////////
 static FAutoConsoleCommand TraceAuxiliaryPauseCmd(
 	TEXT("Trace.Pause"),
-	TEXT("Pauses all trace channels currently sending events"),
+	TEXT("Pauses all trace channels currently sending events."),
 	FConsoleCommandDelegate::CreateStatic(TraceAuxiliaryPause)
 );
 
 ////////////////////////////////////////////////////////////////////////////////
 static FAutoConsoleCommand TraceAuxiliaryResumeCmd(
 	TEXT("Trace.Resume"),
-	TEXT("Resume tracing that was previously paused"),
+	TEXT("Resumes tracing that was previously paused (re-enables the paused channels)."),
 	FConsoleCommandDelegate::CreateStatic(TraceAuxiliaryResume)
 );
 
 ////////////////////////////////////////////////////////////////////////////////
 static FAutoConsoleCommand TraceAuxiliaryStatusCmd(
 	TEXT("Trace.Status"),
-	TEXT("Print Trace status to console"),
+	TEXT("Prints Trace status to console."),
 	FConsoleCommandDelegate::CreateStatic(TraceAuxiliaryStatus)
 );
 
 ////////////////////////////////////////////////////////////////////////////////
 static FAutoConsoleCommand TraceAuxiliaryChannelEnableCmd(
 	TEXT("Trace.Enable"),
-	TEXT("Enable a set of channels; Trace.Enable [ChannelSet]"
-		" where ChannelSet is either comma-separated list of trace channels, a Config/Trace.ChannelPresets key, or optional."),
+	TEXT("[ChannelSet] - Enables a set of channels."
+		" ChannelSet is comma-separated list of trace channels/presets to be enabled."
+	),
 	FConsoleCommandWithArgsDelegate::CreateStatic(TraceAuxiliaryEnableChannels)
 );
 
 ////////////////////////////////////////////////////////////////////////////////
 static FAutoConsoleCommand TraceAuxiliaryChannelDisableCmd(
 	TEXT("Trace.Disable"),
-	TEXT("Disable a set of channels; Trace.Disable [ChannelSet]"
-		" where ChannelSet is either comma-separated list of trace channels, a Config/Trace.ChannelPresets key, or optional."
-		" If no channels are specified all channels are disabled."),
+	TEXT("[ChannelSet] - Disables a set of channels."
+		" ChannelSet is comma-separated list of trace channels/presets to be disabled."
+		" If no channel is specified, it disables all channels."
+	),
 	FConsoleCommandWithArgsDelegate::CreateStatic(TraceAuxiliaryDisableChannels)
 );
 
@@ -940,7 +949,7 @@ static void LaunchUnrealTraceInternal(const TCHAR* CommandLine)
 		execve(*BinPath, Args, environ);
 		_exit(0x80 | (errno & 0x7f));
 	}
-	
+
 	int32 WaitStatus = 0;
 	do
 	{
@@ -996,7 +1005,7 @@ static bool StartFromCommandlineArguments(const TCHAR* CommandLine)
 	{
 		Channels = GDefaultChannels;
 	}
-	
+
 	// By default, if any channels are enabled we trace to memory.
 	FTraceAuxiliary::EConnectionType Type = FTraceAuxiliary::EConnectionType::None;
 
@@ -1058,7 +1067,7 @@ bool FTraceAuxiliary::Start(EConnectionType Type, const TCHAR* Target, const TCH
 		UE_LOG_REF(LogCategory, Error, TEXT("Unable to start trace, already tracing to %s"), GTraceAuxiliary.GetDest());
 		return false;
 	}
-	
+
 	// Make sure the worker thread is started unless explicitly opt out.
 	if (!Options || !Options->bNoWorkerThread)
 	{
@@ -1072,10 +1081,10 @@ bool FTraceAuxiliary::Start(EConnectionType Type, const TCHAR* Target, const TCH
 		GTraceAuxiliary.AddCommandlineChannels(Channels);
 		GTraceAuxiliary.EnableCommandlineChannels();
 	}
-	
+
 	if (Options)
 	{
-		// Truncation is only valid when tracing to file and filename is set
+		// Truncation is only valid when tracing to file and filename is set.
 		if (Options->bTruncateFile && Type == EConnectionType::File && Target != nullptr)
 		{
 			GTraceAuxiliary.SetTruncateFile(Options->bTruncateFile);
@@ -1088,7 +1097,7 @@ bool FTraceAuxiliary::Start(EConnectionType Type, const TCHAR* Target, const TCH
 	}
 	else if(Type == EConnectionType::Network)
 	{
-		return GTraceAuxiliary.Connect(ETraceConnectType::Network, Target, LogCategory);	
+		return GTraceAuxiliary.Connect(ETraceConnectType::Network, Target, LogCategory);
 	}
 #endif
 	return false;
@@ -1149,8 +1158,8 @@ void FTraceAuxiliary::Initialize(const TCHAR* CommandLine)
 	MemoryTrace_InitializeLate();
 #endif
 
-	// Trace out information about this session. This is done before initialisation
-	// so that it is always sent (all channels are enabled prior to initialisation)
+	// Trace out information about this session. This is done before initialisation,
+	// so that it is always sent (all channels are enabled prior to initialisation).
 	const TCHAR* BranchName = BuildSettings::GetBranchName();
 	const TCHAR* BuildVersion = BuildSettings::GetBuildVersion();
 	constexpr uint32 PlatformLen = UE_ARRAY_COUNT(PREPROCESSOR_TO_STRING(UBT_COMPILED_PLATFORM)) - 1;
@@ -1194,8 +1203,7 @@ void FTraceAuxiliary::Initialize(const TCHAR* CommandLine)
 	}
 	UE::Trace::Initialize(Desc);
 
-	// Always register end frame updates. This path is short circuited if a worker thread
-	// exists.
+	// Always register end frame updates. This path is short circuited if a worker thread exists.
 	GTraceAuxiliary.StartEndFramePump();
 	if (!FForkProcessHelper::IsForkRequested())
 	{
@@ -1206,7 +1214,7 @@ void FTraceAuxiliary::Initialize(const TCHAR* CommandLine)
 	CallstackTrace_Create(GMalloc);
 	CallstackTrace_Initialize();
 
-	// By default use 1msec for stack sampling interval
+	// By default use 1 msec for stack sampling interval.
 	uint32 Microseconds = 1000;
 	FParse::Value(CommandLine, TEXT("-samplinginterval="), Microseconds);
 	PlatformEvents_Init(Microseconds);
@@ -1214,7 +1222,7 @@ void FTraceAuxiliary::Initialize(const TCHAR* CommandLine)
 #if CSV_PROFILER
 	FCoreDelegates::OnEndFrame.AddRaw(&GTraceAuxiliary, &FTraceAuxiliaryImpl::UpdateCsvStats);
 #endif
-	
+
 	FModuleManager::Get().OnModulesChanged().AddLambda([](FName Name, EModuleChangeReason Reason)
 	{
 		if (Reason == EModuleChangeReason::ModuleLoaded)
@@ -1222,7 +1230,7 @@ void FTraceAuxiliary::Initialize(const TCHAR* CommandLine)
 			GTraceAuxiliary.EnableCommandlineChannels();
 		}
 	});
-	
+
 	UE::Trace::ThreadRegister(TEXT("GameThread"), FPlatformTLS::GetCurrentThreadId(), -1);
 	
 	// Workaround for the fact that even if StartFromCommandlineArguments will enable channels
@@ -1243,15 +1251,15 @@ void FTraceAuxiliary::InitializePresets(const TCHAR* CommandLine)
 		GTraceAuxiliary.AddCommandlineChannels(*Parameter);
 		GTraceAuxiliary.EnableCommandlineChannels();
 	}
-#endif 
+#endif
 }
 
 ////////////////////////////////////////////////////////////////////////////////
 void FTraceAuxiliary::Shutdown()
 {
 #if UE_TRACE_ENABLED
-	// make sure all platform event functionality has shut down as on some
-	// platforms it impacts whole system, even if application has terminated
+	// Make sure all platform event functionality has shut down as on some
+	// platforms it impacts whole system, even if application has terminated.
 	PlatformEvents_Stop();
 #endif
 }
diff --git a/Engine/Source/Runtime/Core/Public/ProfilingDebugging/TraceAuxiliary.h b/Engine/Source/Runtime/Core/Public/ProfilingDebugging/TraceAuxiliary.h
--- a/Engine/Source/Runtime/Core/Public/ProfilingDebugging/TraceAuxiliary.h
+++ b/Engine/Source/Runtime/Core/Public/ProfilingDebugging/TraceAuxiliary.h
@@ -6,13 +6,10 @@
 #include "HAL/Platform.h"
 #include "Logging/LogMacros.h"
 
-////////////////////////////////////////////////////////////////////////////////
-	
 ////////////////////////////////////////////////////////////////////////////////
 class CORE_API FTraceAuxiliary
 {
 public:
-
 	// In no logging configurations all log categories are of type FNoLoggingCategory, which has no relation with
 	// FLogCategoryBase. In order to not need to conditionally set the argument alias the type here.
 #if NO_LOGGING
@@ -20,7 +17,7 @@ public:
 #else
 	typedef FLogCategoryBase FLogCategoryAlias;
 #endif
-	
+
 	enum class EConnectionType
 	{
 		/**
@@ -40,12 +37,12 @@ public:
 
 	struct Options
 	{
-		/** When set trace will not start a worker thread, instead it is updated from end frame delegate. */
+		/** When set, trace will not start a worker thread, instead it is updated from end frame delegate. */
 		bool bNoWorkerThread = false;
-		/** When set the target file will be truncated if it already exists */ 
+		/** When set, the target file will be truncated if it already exists. */
 		bool bTruncateFile = false;
 	};
-	
+
 	/**
 	 * Start tracing to a target (network connection or file) with an active set of channels. If a connection is
 	 * already active this call does nothing.
@@ -79,13 +76,13 @@ public:
 	 * @param CommandLine to use for initializing
 	 */
 	static void Initialize(const TCHAR* CommandLine);
-	
+
 	/**
 	 * Initialize channels that use the config driven presets.
 	 * @param CommandLine to use for initializing
 	 */
 	static void InitializePresets(const TCHAR* CommandLine);
-	
+
 	/**
 	 * Shut down Trace systems.
 	 */
@@ -103,4 +100,3 @@ public:
 	 */
 	static void EnableChannels();
 };
-
