[Insights] Clear thread buffer for Cpu scopes

In order to avoid scopes ending up after thread end (rare) clear the thread buffer pointer. In that case a new thread buffer will created, specifically for that scope.

diff --git a/Engine/Source/Runtime/Core/Private/ProfilingDebugging/CpuProfilerTrace.cpp b/Engine/Source/Runtime/Core/Private/ProfilingDebugging/CpuProfilerTrace.cpp
--- a/Engine/Source/Runtime/Core/Private/ProfilingDebugging/CpuProfilerTrace.cpp
+++ b/Engine/Source/Runtime/Core/Private/ProfilingDebugging/CpuProfilerTrace.cpp
@@ -79,6 +79,10 @@ struct FCpuProfilerTraceInternal
 		virtual ~FThreadBuffer()
 		{
 			UE_TRACE_LOG(CpuProfiler, EndThread, CpuChannel);
+			// Clear the thread buffer pointer. In the rare event of there being scopes in the destructors of other
+			// FTLSAutoCleanup instances. In that case a new buffer is created for that event only. There is no way of
+			// controlling the order of destruction for FTLSAutoCleanup types.
+			ThreadBuffer = nullptr;
 		}
 
 		uint64 LastCycle = 0;
