FEM - harden some of the atomics to corellate with what x86 would do.

diff --git a/Engine/Source/Runtime/Core/Public/Async/Fundamental/LocalQueue.h b/Engine/Source/Runtime/Core/Public/Async/Fundamental/LocalQueue.h
--- a/Engine/Source/Runtime/Core/Public/Async/Fundamental/LocalQueue.h
+++ b/Engine/Source/Runtime/Core/Public/Async/Fundamental/LocalQueue.h
@@ -30,11 +30,11 @@ protected:
 		checkSlow(Item != uintptr_t(ESlotState::Taken));
 
 		uint32 Idx = (Head + 1) % NumItems;
-		uintptr_t Slot = ItemSlots[Idx].Value.load(std::memory_order_relaxed);
+		uintptr_t Slot = ItemSlots[Idx].Value.load(std::memory_order_acquire);
 
 		if (Slot == uintptr_t(ESlotState::Free))
 		{
-			ItemSlots[Idx].Value.store(Item, std::memory_order_relaxed);
+			ItemSlots[Idx].Value.store(Item, std::memory_order_release);
 			Head++;
 			checkSlow(Head % NumItems == Idx);
 			return true;		
@@ -46,9 +46,9 @@ protected:
 	inline bool Get(uintptr_t& Item)
 	{
 		uint32 Idx = Head % NumItems;
-		uintptr_t Slot = ItemSlots[Idx].Value.load(std::memory_order_relaxed);
+		uintptr_t Slot = ItemSlots[Idx].Value.load(std::memory_order_acquire);
 
-		if (Slot > uintptr_t(ESlotState::Taken) && ItemSlots[Idx].Value.compare_exchange_strong(Slot, uintptr_t(ESlotState::Free), std::memory_order_relaxed))
+		if (Slot > uintptr_t(ESlotState::Taken) && ItemSlots[Idx].Value.compare_exchange_strong(Slot, uintptr_t(ESlotState::Free), std::memory_order_acq_rel))
 		{
 			Head--;
 			checkSlow((Head + 1) % NumItems == Idx);
@@ -63,25 +63,25 @@ protected:
 	{
 		do
 		{
-			uint32 IdxVer = Tail.load(std::memory_order_relaxed);
+			uint32 IdxVer = Tail.load(std::memory_order_acquire);
 			uint32 Idx = IdxVer % NumItems;
-			uintptr_t Slot = ItemSlots[Idx].Value.load(std::memory_order_relaxed);
+			uintptr_t Slot = ItemSlots[Idx].Value.load(std::memory_order_acquire);
 
 			if (Slot == uintptr_t(ESlotState::Free))
 			{
 				return false;
 			}
-			else if (Slot != uintptr_t(ESlotState::Taken) && ItemSlots[Idx].Value.compare_exchange_weak(Slot, uintptr_t(ESlotState::Taken), std::memory_order_relaxed))
+			else if (Slot != uintptr_t(ESlotState::Taken) && ItemSlots[Idx].Value.compare_exchange_weak(Slot, uintptr_t(ESlotState::Taken), std::memory_order_acq_rel))
 			{
-				if(IdxVer == Tail.load(std::memory_order_relaxed))
+				if(IdxVer == Tail.load(std::memory_order_acquire))
 				{
-					uint32 Prev = Tail.fetch_add(1, std::memory_order_relaxed);
+					uint32 Prev = Tail.fetch_add(1, std::memory_order_release);
 					checkSlow(Prev % NumItems == Idx);
-					ItemSlots[Idx].Value.store(uintptr_t(ESlotState::Free), std::memory_order_relaxed);
+					ItemSlots[Idx].Value.store(uintptr_t(ESlotState::Free), std::memory_order_release);
 					Item = Slot;
 					return true;
 				}
-				ItemSlots[Idx].Value.store(Slot, std::memory_order_relaxed);
+				ItemSlots[Idx].Value.store(Slot, std::memory_order_release);
 			}
 		} while(true);
 	}
@@ -410,7 +410,7 @@ private:
 			FLocalQueueCollection* Previous = Hazard.Get();
 			Copy->LocalQueues = Previous->LocalQueues;	
 			Copy->LocalQueues.Add(QueueToAdd);
-			if (!QueueCollection.compare_exchange_strong(Previous, Copy, std::memory_order_release, std::memory_order_relaxed))
+			if (!QueueCollection.compare_exchange_strong(Previous, Copy, std::memory_order_acq_rel, std::memory_order_relaxed))
 			{
 				continue;
 			}
@@ -434,7 +434,7 @@ private:
 			FLocalQueueCollection* Previous = Hazard.Get();
 			Copy->LocalQueues = Previous->LocalQueues;	
 			verifySlow(Copy->LocalQueues.Remove(QueueToRemove) == 1);
-			if (!QueueCollection.compare_exchange_strong(Previous, Copy, std::memory_order_release, std::memory_order_relaxed))
+			if (!QueueCollection.compare_exchange_strong(Previous, Copy, std::memory_order_acq_rel, std::memory_order_relaxed))
 			{
 				continue;
 			}
diff --git a/Engine/Source/Runtime/Core/Public/Experimental/Containers/HazardPointer.h b/Engine/Source/Runtime/Core/Public/Experimental/Containers/HazardPointer.h
--- a/Engine/Source/Runtime/Core/Public/Experimental/Containers/HazardPointer.h
+++ b/Engine/Source/Runtime/Core/Public/Experimental/Containers/HazardPointer.h
@@ -302,8 +302,8 @@ public:
 		H* HazardPointer;
 		do
 		{
-			HazardPointer = (H*)Record->SetHazard(Hazard->load(std::memory_order_relaxed));
-		} while (HazardPointer != Hazard->load(std::memory_order_relaxed));
+			HazardPointer = (H*)Record->SetHazard(Hazard->load(std::memory_order_acquire));
+		} while (HazardPointer != Hazard->load(std::memory_order_acquire));
 		return HazardPointer;
 	}
 
@@ -317,4 +317,4 @@ template<typename H>
 THazardPointer<H, false> MakeHazardPointer(std::atomic<H*>& InHazard, FHazardPointerCollection& Collection)
 {
 	return {InHazard, Collection};
-}
\ No newline at end of file
+}
