Change LWC tile size to OLD_WORLD_MAX
Allow fractional tile values for camera matrices, gives better precision when sharing tile between WorldToView and ViewToWorld

diff --git a/Engine/Source/Runtime/Core/Private/Misc/LargeWorldRenderPosition.cpp b/Engine/Source/Runtime/Core/Private/Misc/LargeWorldRenderPosition.cpp
--- a/Engine/Source/Runtime/Core/Private/Misc/LargeWorldRenderPosition.cpp
+++ b/Engine/Source/Runtime/Core/Private/Misc/LargeWorldRenderPosition.cpp
@@ -6,8 +6,10 @@
 #if !UE_LARGE_WORLD_COORDINATES_DISABLED
 
 #define UE_LWC_RENDER_TILE_SIZE_MIN (262144.0)	// LWC_TODO: Make this smaller?
-#define UE_LWC_RENDER_TILE_SIZE (FMath::Max(UE_LWC_RENDER_TILE_SIZE_MIN,  HALF_WORLD_MAX / 16777216.0))	// LWC_TODO: Tweak render tile divisor to maximise range.
-static_assert(UE_LWC_RENDER_TILE_SIZE <= 1048576.0, "Current WORLD_MAX is too large and is likely to adversely affect world space coordinate precision within shaders!");	// LWC_TODO: Set this to something reasonable.
+
+#define UE_LWC_RENDER_TILE_SIZE UE_OLD_WORLD_MAX
+//#define UE_LWC_RENDER_TILE_SIZE UE_OLD_WORLD_MAX (FMath::Max(UE_LWC_RENDER_TILE_SIZE_MIN,  HALF_WORLD_MAX / 16777216.0))	// LWC_TODO: Tweak render tile divisor to maximise range.
+//static_assert(UE_LWC_RENDER_TILE_SIZE <= 1048576.0, "Current WORLD_MAX is too large and is likely to adversely affect world space coordinate precision within shaders!");	// LWC_TODO: Set this to something reasonable.
 
 //#define UE_LWC_RENDER_TILE_SIZE (262144.0)
 //#define UE_LWC_RENDER_TILE_SIZE (1048576.0)
@@ -16,21 +18,7 @@ static_assert(UE_LWC_RENDER_TILE_SIZE <= 1048576.0, "Current WORLD_MAX is too la
 // This is the max size we allow for LWC offsets relative to the tile
 // Value chosen to ensure sufficient precision when stored in single precision float
 // Normally offsets should be within +/-TileSizeDivideBy2, but we often rebase multiple quantities off a single tile origin
-#define UE_LWC_RENDER_MAX_OFFSET (2097152.0)
-
-static FMatrix44f SafeCastMatrix(const FMatrix& Matrix)
-{
-	const FVector Origin = Matrix.GetOrigin();
-	const double OriginX = FMath::Abs(Origin.X);
-	const double OriginY = FMath::Abs(Origin.Y);
-	const double OriginZ = FMath::Abs(Origin.Z);
-	if (OriginX > UE_LWC_RENDER_MAX_OFFSET || OriginY > UE_LWC_RENDER_MAX_OFFSET || OriginZ > UE_LWC_RENDER_MAX_OFFSET)
-	{
-		ensure(false);
-	}
-
-	return FMatrix44f(Matrix);
-}
+#define UE_LWC_RENDER_MAX_OFFSET UE_OLD_HALF_WORLD_MAX
 
 double FLargeWorldRenderScalar::GetTileSize()
 {
@@ -53,30 +41,48 @@ FVector3f FLargeWorldRenderScalar::GetTileFor(FVector InPosition)
 	return LWCTile;
 }
 
-FMatrix44f FLargeWorldRenderPosition::MakeToRelativeWorldMatrix(const FMatrix& ToWorld) const
+FMatrix44f FLargeWorldRenderScalar::SafeCastMatrix(const FMatrix& Matrix)
 {
-	return SafeCastMatrix(ToWorld * FTranslationMatrix(-GetTileOffset()));
+	const double OriginMax = UE_LWC_RENDER_MAX_OFFSET;
+
+	const FVector Origin = Matrix.GetOrigin();
+	const double OriginX = FMath::Abs(Origin.X);
+	const double OriginY = FMath::Abs(Origin.Y);
+	const double OriginZ = FMath::Abs(Origin.Z);
+	if (OriginX > OriginMax || OriginY > OriginMax || OriginZ > OriginMax)
+	{
+		ensure(false);
+	}
+
+	return FMatrix44f(Matrix);
+}
+
+FMatrix44f FLargeWorldRenderScalar::MakeToRelativeWorldMatrix(const FVector Origin, const FMatrix& ToWorld)
+{
+	return SafeCastMatrix(ToWorld * FTranslationMatrix(-Origin));
+}
+
+FMatrix44f FLargeWorldRenderScalar::MakeFromRelativeWorldMatrix(const FVector Origin, const FMatrix& FromWorld)
+{
+	return SafeCastMatrix(FTranslationMatrix(Origin) * FromWorld);
 }
 
-FMatrix44f FLargeWorldRenderPosition::MakeClampedToRelativeWorldMatrix(const FMatrix& ToWorld) const
+FMatrix44f FLargeWorldRenderScalar::MakeClampedToRelativeWorldMatrix(const FVector Origin, const FMatrix& ToWorld)
 {
+	const double OriginMax = UE_LWC_RENDER_MAX_OFFSET;
+
 	// Clamp the relative matrix, avoid allowing the relative translation to get too far away from the tile origin
-	const FVector RelativeOrigin = ToWorld.GetOrigin() - GetTileOffset();
+	const FVector RelativeOrigin = ToWorld.GetOrigin() - Origin;
 	FVector ClampedRelativeOrigin = RelativeOrigin;
-	ClampedRelativeOrigin.X = FMath::Clamp(ClampedRelativeOrigin.X, -UE_LWC_RENDER_MAX_OFFSET, UE_LWC_RENDER_MAX_OFFSET);
-	ClampedRelativeOrigin.Y = FMath::Clamp(ClampedRelativeOrigin.Y, -UE_LWC_RENDER_MAX_OFFSET, UE_LWC_RENDER_MAX_OFFSET);
-	ClampedRelativeOrigin.Z = FMath::Clamp(ClampedRelativeOrigin.Z, -UE_LWC_RENDER_MAX_OFFSET, UE_LWC_RENDER_MAX_OFFSET);
+	ClampedRelativeOrigin.X = FMath::Clamp(ClampedRelativeOrigin.X, -OriginMax, OriginMax);
+	ClampedRelativeOrigin.Y = FMath::Clamp(ClampedRelativeOrigin.Y, -OriginMax, OriginMax);
+	ClampedRelativeOrigin.Z = FMath::Clamp(ClampedRelativeOrigin.Z, -OriginMax, OriginMax);
 
 	FMatrix ClampedToRelativeWorld(ToWorld);
 	ClampedToRelativeWorld.SetOrigin(ClampedRelativeOrigin);
 	return FMatrix44f(ClampedToRelativeWorld);
 }
 
-FMatrix44f FLargeWorldRenderPosition::MakeFromRelativeWorldMatrix(const FMatrix& FromWorld) const
-{
-	return SafeCastMatrix(FTranslationMatrix(GetTileOffset()) * FromWorld);
-}
-
 void FLargeWorldRenderScalar::Validate(double InAbsolute)
 {
 	const double Tolerance = 0.01; // TODO LWC - How precise do we need to be?
diff --git a/Engine/Source/Runtime/Core/Public/Misc/LargeWorldRenderPosition.h b/Engine/Source/Runtime/Core/Public/Misc/LargeWorldRenderPosition.h
--- a/Engine/Source/Runtime/Core/Public/Misc/LargeWorldRenderPosition.h
+++ b/Engine/Source/Runtime/Core/Public/Misc/LargeWorldRenderPosition.h
@@ -14,6 +14,31 @@ public:
 	CORE_API static double GetTileSize();
 	CORE_API static FVector3f GetTileFor(FVector InPosition);
 
+	static double MakeTile(double InValue)
+	{
+		return FMath::FloorToDouble(InValue / GetTileSize() + 0.5);
+	}
+
+	static double MakeQuantizedTile(double InValue, double InQuantization)
+	{
+		return FMath::FloorToDouble((InValue / GetTileSize()) * InQuantization + 0.5) / InQuantization;
+	}
+
+	static FVector MakeTile(const FVector& InValue)
+	{
+		return FVector(MakeTile(InValue.X), MakeTile(InValue.Y), MakeTile(InValue.Z));
+	}
+
+	static FVector MakeQuantizedTile(const FVector& InValue, double InQuantization)
+	{
+		return FVector(MakeQuantizedTile(InValue.X, InQuantization), MakeQuantizedTile(InValue.Y, InQuantization), MakeQuantizedTile(InValue.Z, InQuantization));
+	}
+
+	CORE_API static FMatrix44f SafeCastMatrix(const FMatrix& Matrix);
+	CORE_API static FMatrix44f MakeToRelativeWorldMatrix(const FVector Origin, const FMatrix& ToWorld);
+	CORE_API static FMatrix44f MakeFromRelativeWorldMatrix(const FVector Origin, const FMatrix& FromWorld);
+	CORE_API static FMatrix44f MakeClampedToRelativeWorldMatrix(const FVector Origin, const FMatrix& ToWorld);
+
 	float GetTile() const { return (float)Tile; }
 	float GetOffset() const { return (float)Offset; }
 	double GetTileAsDouble() const { return Tile; }
@@ -27,7 +52,7 @@ public:
 	FLargeWorldRenderScalar(double InAbsolute)
 	{
 		// Tiles are centered on the origin
-		Tile = FMath::FloorToDouble(InAbsolute / GetTileSize() + 0.5);
+		Tile = MakeTile(InAbsolute);
 		Offset = InAbsolute - GetTileOffset();
 		Validate(InAbsolute);
 	}
@@ -68,16 +93,6 @@ public:
 	FVector3f GetOffset() const { return FVector3f(X.GetOffset(), Y.GetOffset(), Z.GetOffset()); }
 	FVector GetAbsolute() const { return FVector(X.GetAbsolute(), Y.GetAbsolute(), Z.GetAbsolute()); }
 
-#if !UE_LARGE_WORLD_COORDINATES_DISABLED
-	CORE_API FMatrix44f MakeToRelativeWorldMatrix(const FMatrix& ToWorld) const;
-	CORE_API FMatrix44f MakeClampedToRelativeWorldMatrix(const FMatrix& ToWorld) const;
-	CORE_API FMatrix44f MakeFromRelativeWorldMatrix(const FMatrix& FromWorld) const;
-#else
-	FMatrix44f MakeToRelativeWorldMatrix(const FMatrix& ToWorld) const { return ToWorld; }
-	FMatrix44f MakeClampedToRelativeWorldMatrix(const FMatrix& ToWorld) const { return ToWorld; }
-	FMatrix44f MakeFromRelativeWorldMatrix(const FMatrix& FromWorld) const { return FromWorld; }
-#endif
-
 	FLargeWorldRenderPosition(const FVector3f& InWorldPosition)
 		: X(InWorldPosition.X)
 		, Y(InWorldPosition.Y)
diff --git a/Engine/Source/Runtime/Engine/Private/SceneRelativeViewMatrices.cpp b/Engine/Source/Runtime/Engine/Private/SceneRelativeViewMatrices.cpp
--- a/Engine/Source/Runtime/Engine/Private/SceneRelativeViewMatrices.cpp
+++ b/Engine/Source/Runtime/Engine/Private/SceneRelativeViewMatrices.cpp
@@ -18,18 +18,26 @@ FRelativeViewMatrices FRelativeViewMatrices::Create(const FViewMatrices& Matrice
 
 FRelativeViewMatrices FRelativeViewMatrices::Create(const FInitializer& Initializer)
 {
-	const FLargeWorldRenderPosition AbsoluteOrigin(Initializer.ViewToWorld.GetOrigin());
+	const double TileSize = FLargeWorldRenderScalar::GetTileSize();
+
+	// We allow a fractional tile position here
+	// Tile offset is applied beofre WorldToView transform, but after ViewToWorld transform
+	// This means that if we use a regular tile offset, the remaining offset in the relative matrices may become too large (>TileSize/2)
+	// Allowing for a fractional tile lets us use the same value for both matrices
+	// Quantization factor is somewhat arbitrary...controls distribution of precision between tile fraction and relative offset
+	const FVector ViewOrigin = Initializer.ViewToWorld.GetOrigin();
+	const FVector ViewOriginTile = FLargeWorldRenderScalar::MakeQuantizedTile(ViewOrigin, 8.0);
 
 	FRelativeViewMatrices Result;
-	Result.TilePosition = AbsoluteOrigin.GetTile();
-	Result.RelativeWorldToView = AbsoluteOrigin.MakeFromRelativeWorldMatrix(Initializer.WorldToView);
-	Result.ViewToRelativeWorld = AbsoluteOrigin.MakeToRelativeWorldMatrix(Initializer.ViewToWorld);
+	Result.TilePosition = ViewOriginTile;
+	Result.RelativeWorldToView = FLargeWorldRenderScalar::MakeFromRelativeWorldMatrix(ViewOriginTile * TileSize, Initializer.WorldToView);
+	Result.ViewToRelativeWorld = FLargeWorldRenderScalar::MakeToRelativeWorldMatrix(ViewOriginTile * TileSize, Initializer.ViewToWorld);
 	Result.ViewToClip = FMatrix44f(Initializer.ViewToClip);
 	Result.ClipToView = FMatrix44f(Initializer.ClipToView);
 	Result.RelativeWorldToClip = Result.RelativeWorldToView * Result.ViewToClip;
 	Result.ClipToRelativeWorld = Result.ClipToView * Result.ViewToRelativeWorld;
 
-	Result.PrevViewToRelativeWorld = AbsoluteOrigin.MakeClampedToRelativeWorldMatrix(Initializer.PrevViewToWorld);
+	Result.PrevViewToRelativeWorld = FLargeWorldRenderScalar::MakeClampedToRelativeWorldMatrix(ViewOriginTile * TileSize, Initializer.PrevViewToWorld);
 	Result.PrevClipToView = FMatrix44f(Initializer.PrevClipToView);
 	Result.PrevClipToRelativeWorld = Result.PrevClipToView * Result.PrevViewToRelativeWorld;
 	return Result;
diff --git a/Engine/Source/Runtime/Engine/Public/PrimitiveUniformShaderParameters.h b/Engine/Source/Runtime/Engine/Public/PrimitiveUniformShaderParameters.h
--- a/Engine/Source/Runtime/Engine/Public/PrimitiveUniformShaderParameters.h
+++ b/Engine/Source/Runtime/Engine/Public/PrimitiveUniformShaderParameters.h
@@ -262,7 +262,7 @@ public:
 		Parameters.TilePosition = AbsoluteWorldPosition.GetTile();
 
 		{
-			Parameters.LocalToRelativeWorld = AbsoluteWorldPosition.MakeToRelativeWorldMatrix(AbsoluteLocalToWorld);
+			Parameters.LocalToRelativeWorld = FLargeWorldRenderScalar::MakeToRelativeWorldMatrix(TilePositionOffset, AbsoluteLocalToWorld);
 			Parameters.RelativeWorldToLocal = Parameters.LocalToRelativeWorld.Inverse();
 		}
 
@@ -271,7 +271,7 @@ public:
 
 		if (bHasPreviousLocalToWorld)
 		{
-			Parameters.PreviousLocalToRelativeWorld = AbsoluteWorldPosition.MakeClampedToRelativeWorldMatrix(AbsolutePreviousLocalToWorld);
+			Parameters.PreviousLocalToRelativeWorld = FLargeWorldRenderScalar::MakeClampedToRelativeWorldMatrix(TilePositionOffset, AbsolutePreviousLocalToWorld);
 			Parameters.PreviousRelativeWorldToLocal = Parameters.PreviousLocalToRelativeWorld.Inverse();
 		}
 		else
diff --git a/Engine/Source/Runtime/Renderer/Private/GPUScene.cpp b/Engine/Source/Runtime/Renderer/Private/GPUScene.cpp
--- a/Engine/Source/Runtime/Renderer/Private/GPUScene.cpp
+++ b/Engine/Source/Runtime/Renderer/Private/GPUScene.cpp
@@ -472,7 +472,7 @@ struct FUploadDataSourceAdapterScenePrimitives
 
 		InstanceUploadInfo.LastUpdateSceneFrameNumber = SceneFrameNumber;
 		InstanceUploadInfo.PrimitiveID = GetPrimitiveID(Scene, PrimitiveID);
-		InstanceUploadInfo.PrimitiveToWorld = AbsoluteOrigin.MakeToRelativeWorldMatrix(LocalToWorld);
+		InstanceUploadInfo.PrimitiveToWorld = FLargeWorldRenderScalar::MakeToRelativeWorldMatrix(AbsoluteOrigin.GetTileOffset(), LocalToWorld);
 
 		{
 			bool bHasPrecomputedVolumetricLightmap{};
@@ -481,7 +481,7 @@ struct FUploadDataSourceAdapterScenePrimitives
 
 			FMatrix PreviousLocalToWorld;
 			Scene.GetPrimitiveUniformShaderParameters_RenderThread(PrimitiveSceneInfo, bHasPrecomputedVolumetricLightmap, PreviousLocalToWorld, SingleCaptureIndex, bOutputVelocity);
-			InstanceUploadInfo.PrevPrimitiveToWorld = AbsoluteOrigin.MakeClampedToRelativeWorldMatrix(PreviousLocalToWorld);;
+			InstanceUploadInfo.PrevPrimitiveToWorld = FLargeWorldRenderScalar::MakeClampedToRelativeWorldMatrix(AbsoluteOrigin.GetTileOffset(), PreviousLocalToWorld);;
 		}
 
 		InstanceUploadInfo.PayloadDataFlags = PrimitiveSceneProxy->GetPayloadDataFlags();
@@ -2366,7 +2366,7 @@ void FGPUSceneCompactInstanceData::Init(const FScene* Scene, int32 PrimitiveId)
 	{
 		const FMatrix LocalToWorld = Scene->PrimitiveTransforms[PrimitiveId];
 		const FLargeWorldRenderPosition AbsoluteOrigin(LocalToWorld.GetOrigin());
-		LocalToRelativeWorld = AbsoluteOrigin.MakeToRelativeWorldMatrix(LocalToWorld);
+		LocalToRelativeWorld = FLargeWorldRenderScalar::MakeToRelativeWorldMatrix(AbsoluteOrigin.GetTileOffset(), LocalToWorld);
 	}
 	InstanceOriginAndId		= LocalToRelativeWorld.GetOrigin();
 	InstanceTransform1		= LocalToRelativeWorld.GetScaledAxis(EAxis::X);
