Add new API for LWC disabled case

diff --git a/Engine/Source/Runtime/Core/Public/Misc/LargeWorldRenderPosition.h b/Engine/Source/Runtime/Core/Public/Misc/LargeWorldRenderPosition.h
--- a/Engine/Source/Runtime/Core/Public/Misc/LargeWorldRenderPosition.h
+++ b/Engine/Source/Runtime/Core/Public/Misc/LargeWorldRenderPosition.h
@@ -68,6 +68,17 @@ private:
 public:
 	CORE_API static float GetTileSize() { return 0.0f; }
 	CORE_API static FVector3f GetTileFor(FVector InPosition) { return FVector3f::ZeroVector; };
+
+	static double MakeTile(double InValue) { return 0.0; }
+	static double MakeQuantizedTile(double InValue, double InQuantization) { return 0.0; }
+	static FVector MakeTile(const FVector& InValue) { return FVector::ZeroVector; }
+	static FVector MakeQuantizedTile(const FVector& InValue, double InQuantization) { return FVector::ZeroVector; }
+
+	static FMatrix44f SafeCastMatrix(const FMatrix& Matrix) { return Matrix; }
+	static FMatrix44f MakeToRelativeWorldMatrix(const FVector Origin, const FMatrix& ToWorld) { return ToWorld; }
+	static FMatrix44f MakeFromRelativeWorldMatrix(const FVector Origin, const FMatrix& FromWorld) { return FromWorld; }
+	static FMatrix44f MakeClampedToRelativeWorldMatrix(const FVector Origin, const FMatrix& ToWorld) { return ToWorld; }
+
 	float GetTile() const { return 0.0f; }
 	float GetOffset() const { return Offset; }
 	double GetTileAsDouble() const { return 0.0; }
