Fix merge to work without UE_OLD_WORLD_MAX define

diff --git a/Engine/Source/Runtime/Core/Private/Misc/LargeWorldRenderPosition.cpp b/Engine/Source/Runtime/Core/Private/Misc/LargeWorldRenderPosition.cpp
--- a/Engine/Source/Runtime/Core/Private/Misc/LargeWorldRenderPosition.cpp
+++ b/Engine/Source/Runtime/Core/Private/Misc/LargeWorldRenderPosition.cpp
@@ -7,7 +7,7 @@
 
 #define UE_LWC_RENDER_TILE_SIZE_MIN (262144.0)	// LWC_TODO: Make this smaller?
 
-#define UE_LWC_RENDER_TILE_SIZE UE_OLD_WORLD_MAX
+#define UE_LWC_RENDER_TILE_SIZE 2097152.0
 //#define UE_LWC_RENDER_TILE_SIZE UE_OLD_WORLD_MAX (FMath::Max(UE_LWC_RENDER_TILE_SIZE_MIN,  HALF_WORLD_MAX / 16777216.0))	// LWC_TODO: Tweak render tile divisor to maximise range.
 //static_assert(UE_LWC_RENDER_TILE_SIZE <= 1048576.0, "Current WORLD_MAX is too large and is likely to adversely affect world space coordinate precision within shaders!");	// LWC_TODO: Set this to something reasonable.
 
@@ -18,7 +18,7 @@
 // This is the max size we allow for LWC offsets relative to the tile
 // Value chosen to ensure sufficient precision when stored in single precision float
 // Normally offsets should be within +/-TileSizeDivideBy2, but we often rebase multiple quantities off a single tile origin
-#define UE_LWC_RENDER_MAX_OFFSET UE_OLD_HALF_WORLD_MAX
+#define UE_LWC_RENDER_MAX_OFFSET (2097152.0*0.5)
 
 double FLargeWorldRenderScalar::GetTileSize()
 {
