Back out changelist 18699168

diff --git a/Engine/Source/Runtime/Core/Public/Experimental/Containers/RobinHoodHashTable.h b/Engine/Source/Runtime/Core/Public/Experimental/Containers/RobinHoodHashTable.h
--- a/Engine/Source/Runtime/Core/Public/Experimental/Containers/RobinHoodHashTable.h
+++ b/Engine/Source/Runtime/Core/Public/Experimental/Containers/RobinHoodHashTable.h
@@ -208,7 +208,6 @@ namespace RobinHoodHashTable_Private
 			FreeList.Reset(1);
 			//push a dummy
 			FreeList.Push(FSpan{ INDEX_NONE, INDEX_NONE });
-			NumElements = 0;
 		}
 
 		int Num() const
