FEM - enable aligned malloc for after it was accitently disabled a while back.

diff --git a/Engine/Source/Runtime/Core/Private/HAL/MallocAnsi.cpp b/Engine/Source/Runtime/Core/Private/HAL/MallocAnsi.cpp
--- a/Engine/Source/Runtime/Core/Private/HAL/MallocAnsi.cpp
+++ b/Engine/Source/Runtime/Core/Private/HAL/MallocAnsi.cpp
@@ -9,9 +9,9 @@
 #include "Math/UnrealMathUtility.h"
 #include "Templates/AlignmentTemplates.h"
 
-#if PLATFORM_UNIX
+#if PLATFORM_UNIX || PLATFORM_SWITCH
 	#include <malloc.h>
-#endif // PLATFORM_UNIX
+#endif // PLATFORM_UNIX || PLATFORM_SWITCH
 
 #if PLATFORM_IOS
 	#include "mach/mach.h"
@@ -23,7 +23,7 @@
 
 void* AnsiMalloc(SIZE_T Size, uint32 Alignment)
 {
-#if USE_ALIGNED_MALLOC
+#if PLATFORM_USES__ALIGNED_MALLOC
 	void* Result = _aligned_malloc( Size, Alignment );
 #elif PLATFORM_USE_ANSI_POSIX_MALLOC
 	void* Result;
@@ -49,20 +49,20 @@ void* AnsiMalloc(SIZE_T Size, uint32 Alignment)
 
 static SIZE_T AnsiGetAllocationSize(void* Original)
 {
-#if	USE_ALIGNED_MALLOC
+#if	PLATFORM_USES__ALIGNED_MALLOC
 	return _aligned_msize(Original, 16, 0); // Assumes alignment of 16
 #elif PLATFORM_USE_ANSI_POSIX_MALLOC || PLATFORM_USE_ANSI_MEMALIGN
 	return malloc_usable_size(Original);
 #else
 	return *((SIZE_T*)((uint8*)Original - sizeof(void*) - sizeof(SIZE_T)));
-#endif // USE_ALIGNED_MALLOC
+#endif // PLATFORM_USES__ALIGNED_MALLOC
 }
 
 void* AnsiRealloc(void* Ptr, SIZE_T NewSize, uint32 Alignment)
 {
 	void* Result;
 
-#if USE_ALIGNED_MALLOC
+#if PLATFORM_USES__ALIGNED_MALLOC
 	if (Ptr && NewSize)
 	{
 		Result = _aligned_realloc(Ptr, NewSize, Alignment);
@@ -129,7 +129,7 @@ void* AnsiRealloc(void* Ptr, SIZE_T NewSize, uint32 Alignment)
 
 void AnsiFree(void* Ptr)
 {
-#if USE_ALIGNED_MALLOC
+#if PLATFORM_USES__ALIGNED_MALLOC
 	_aligned_free(Ptr);
 #elif PLATFORM_USE_ANSI_POSIX_MALLOC || PLATFORM_USE_ANSI_MEMALIGN
 	free(Ptr);
diff --git a/Engine/Source/Runtime/Core/Public/Android/AndroidPlatform.h b/Engine/Source/Runtime/Core/Public/Android/AndroidPlatform.h
--- a/Engine/Source/Runtime/Core/Public/Android/AndroidPlatform.h
+++ b/Engine/Source/Runtime/Core/Public/Android/AndroidPlatform.h
@@ -60,6 +60,7 @@ typedef FAndroidTypes FPlatformTypes;
 #define PLATFORM_SUPPORTS_EARLY_MOVIE_PLAYBACK		1 // movies will start before engine is initalized
 #define PLATFORM_SUPPORTS_GEOMETRY_SHADERS			0
 #define PLATFORM_SUPPORTS_VIRTUAL_TEXTURE_STREAMING	1
+#define PLATFORM_USE_ANSI_POSIX_MALLOC				1
 #define PLATFORM_SUPPORTS_LANDSCAPE_VISUAL_MESH_LOD_STREAMING 1
 
 #define PLATFORM_CODE_SECTION(Name)					__attribute__((section(Name)))
diff --git a/Engine/Source/Runtime/Core/Public/Experimental/ConcurrentLinearAllocator.h b/Engine/Source/Runtime/Core/Public/Experimental/ConcurrentLinearAllocator.h
--- a/Engine/Source/Runtime/Core/Public/Experimental/ConcurrentLinearAllocator.h
+++ b/Engine/Source/Runtime/Core/Public/Experimental/ConcurrentLinearAllocator.h
@@ -51,13 +51,13 @@ struct FOsAllocator
 
 struct FAlignedAllocator
 {
-	static constexpr bool SupportsAlignment = !(PLATFORM_MAC || PLATFORM_IOS || PLATFORM_TVOS); //TODO: UHT on Mac OS returned false alignment.
+	static constexpr bool SupportsAlignment = true;
 
 	FORCEINLINE static void* Malloc(SIZE_T Size, uint32 Alignment)
 	{
 #if DO_CHECK
 		void* Ret = AnsiMalloc(Size, Alignment);
-		check(!SupportsAlignment || IsAligned(Ret, Alignment));
+		check(IsAligned(Ret, Alignment));
 		return Ret;
 #else
 		return AnsiMalloc(Size, Alignment);
diff --git a/Engine/Source/Runtime/Core/Public/HAL/MallocAnsi.h b/Engine/Source/Runtime/Core/Public/HAL/MallocAnsi.h
--- a/Engine/Source/Runtime/Core/Public/HAL/MallocAnsi.h
+++ b/Engine/Source/Runtime/Core/Public/HAL/MallocAnsi.h
@@ -5,14 +5,6 @@
 #include "CoreTypes.h"
 #include "HAL/UnrealMemory.h"
 
-#ifndef USE_ALIGNED_MALLOC
-	#if PLATFORM_MAC || PLATFORM_IOS
-		#define USE_ALIGNED_MALLOC 1
-	#else
-		#define USE_ALIGNED_MALLOC 0
-	#endif
-#endif
-
 CORE_API void* AnsiMalloc(SIZE_T Size, uint32 Alignment);
 CORE_API void* AnsiRealloc(void* Ptr, SIZE_T NewSize, uint32 Alignment);
 CORE_API void AnsiFree(void* Ptr);
diff --git a/Engine/Source/Runtime/Core/Public/HAL/Platform.h b/Engine/Source/Runtime/Core/Public/HAL/Platform.h
--- a/Engine/Source/Runtime/Core/Public/HAL/Platform.h
+++ b/Engine/Source/Runtime/Core/Public/HAL/Platform.h
@@ -486,6 +486,10 @@
 	#define PLATFORM_USE_ANSI_POSIX_MALLOC						0
 #endif
 
+#ifndef PLATFORM_USES__ALIGNED_MALLOC
+	#define PLATFORM_USES__ALIGNED_MALLOC						0
+#endif
+
 #ifndef PLATFORM_USE_ANSI_MEMALIGN
 	#define PLATFORM_USE_ANSI_MEMALIGN							0
 #endif
diff --git a/Engine/Source/Runtime/Core/Public/IOS/IOSPlatform.h b/Engine/Source/Runtime/Core/Public/IOS/IOSPlatform.h
--- a/Engine/Source/Runtime/Core/Public/IOS/IOSPlatform.h
+++ b/Engine/Source/Runtime/Core/Public/IOS/IOSPlatform.h
@@ -56,6 +56,7 @@ typedef FIOSPlatformTypes FPlatformTypes;
 #define PLATFORM_SUPPORTS_EARLY_MOVIE_PLAYBACK			1 // movies will start before engine is initalized
 #define PLATFORM_USE_FULL_TASK_GRAPH					0 // @todo platplug: not platplug, but should investigate soon anyway
 #define PLATFORM_IS_ANSI_MALLOC_THREADSAFE				1
+#define PLATFORM_USE_ANSI_POSIX_MALLOC					1
 
 // on iOS we now perform offline symbolication as it's significantly faster. Requires bGenerateCrashReportSymbols=true in the ini file.
 #define	PLATFORM_RUNTIME_MALLOCPROFILER_SYMBOLICATION	0	
diff --git a/Engine/Source/Runtime/Core/Public/Mac/MacPlatform.h b/Engine/Source/Runtime/Core/Public/Mac/MacPlatform.h
--- a/Engine/Source/Runtime/Core/Public/Mac/MacPlatform.h
+++ b/Engine/Source/Runtime/Core/Public/Mac/MacPlatform.h
@@ -84,6 +84,7 @@ typedef FMacPlatformTypes FPlatformTypes;
 #define PLATFORM_HAS_BSD_SOCKET_FEATURE_MSG_DONTWAIT	1
 #define PLATFORM_IS_ANSI_MALLOC_THREADSAFE				1
 #define PLATFORM_SUPPORTS_VIRTUAL_TEXTURE_STREAMING		1
+#define PLATFORM_USE_ANSI_POSIX_MALLOC					1
 
 #define PLATFORM_RHITHREAD_DEFAULT_BYPASS				0
 
diff --git a/Engine/Source/Runtime/Core/Public/Unix/UnixPlatform.h b/Engine/Source/Runtime/Core/Public/Unix/UnixPlatform.h
--- a/Engine/Source/Runtime/Core/Public/Unix/UnixPlatform.h
+++ b/Engine/Source/Runtime/Core/Public/Unix/UnixPlatform.h
@@ -69,6 +69,7 @@ typedef FUnixPlatformTypes FPlatformTypes;
 #define PLATFORM_ALLOW_ALLOCATIONS_IN_FASYNCWRITER_SERIALIZEBUFFERTOARCHIVE 0
 #define PLATFORM_RHITHREAD_DEFAULT_BYPASS				0
 #define PLATFORM_SUPPORTS_VIRTUAL_TEXTURE_STREAMING		1
+#define PLATFORM_USE_ANSI_POSIX_MALLOC					1
 
 #if PLATFORM_CPU_X86_FAMILY
 	#define PLATFORM_BREAK()							__asm__ volatile("int $0x03")
diff --git a/Engine/Source/Runtime/Core/Public/Windows/WindowsPlatform.h b/Engine/Source/Runtime/Core/Public/Windows/WindowsPlatform.h
--- a/Engine/Source/Runtime/Core/Public/Windows/WindowsPlatform.h
+++ b/Engine/Source/Runtime/Core/Public/Windows/WindowsPlatform.h
@@ -90,6 +90,7 @@ typedef FWindowsPlatformTypes FPlatformTypes;
 #define PLATFORM_SUPPORTS_VIRTUAL_TEXTURE_STREAMING			1
 #define PLATFORM_SUPPORTS_VARIABLE_RATE_SHADING				1
 #define PLATFORM_SUPPORTS_MESH_SHADERS						1
+#define PLATFORM_USES__ALIGNED_MALLOC						1
 
 #if WITH_EDITOR
 #define PLATFORM_FILE_READER_BUFFER_SIZE					(256*1024)
