FEM - fixing a very small leak in the HazardRecordList and an unitizialized value in the WorkStealingQueue

diff --git a/Engine/Source/Runtime/Core/Private/Experimental/Containers/HazardPointer.cpp b/Engine/Source/Runtime/Core/Private/Experimental/Containers/HazardPointer.cpp
--- a/Engine/Source/Runtime/Core/Private/Experimental/Containers/HazardPointer.cpp
+++ b/Engine/Source/Runtime/Core/Private/Experimental/Containers/HazardPointer.cpp
@@ -40,6 +40,7 @@ CORE_API FHazardPointerCollection::FHazardRecord* FHazardPointerCollection::Grow
 	lst->Next.store(p, std::memory_order_release);
 	TotalNumHazardRecords += HazardChunkSize; //keep count of HazardPointers, there should not be too many (maybe 2 per thread)
 
+	HazardRecordBlocks.Add(p);
 	checkSlow(p->Records[0].GetHazard() == nullptr);
 	return &p->Records[0];
 }
diff --git a/Engine/Source/Runtime/Core/Public/Async/Fundamental/LocalQueue.h b/Engine/Source/Runtime/Core/Public/Async/Fundamental/LocalQueue.h
--- a/Engine/Source/Runtime/Core/Public/Async/Fundamental/LocalQueue.h
+++ b/Engine/Source/Runtime/Core/Public/Async/Fundamental/LocalQueue.h
@@ -89,7 +89,7 @@ protected:
 private:
 	struct FAlignedElement
 	{
-		alignas(PLATFORM_CACHE_LINE_SIZE * 2) std::atomic<uintptr_t> Value;
+		alignas(PLATFORM_CACHE_LINE_SIZE * 2) std::atomic<uintptr_t> Value = {};
 	};
 
 	alignas(PLATFORM_CACHE_LINE_SIZE * 2) uint32 Head { ~0u };
