StringCast: Add a version that takes a static buffer size to remove an allocation.
#rnx

diff --git a/Engine/Source/Runtime/Core/Private/HAL/LowLevelMemTracker.cpp b/Engine/Source/Runtime/Core/Private/HAL/LowLevelMemTracker.cpp
--- a/Engine/Source/Runtime/Core/Private/HAL/LowLevelMemTracker.cpp
+++ b/Engine/Source/Runtime/Core/Private/HAL/LowLevelMemTracker.cpp
@@ -3733,7 +3733,7 @@ namespace LLMPrivate
 	*/
 	void FLLMCsvWriter::Write(FStringView Text)
 	{
-		Archive->Serialize(const_cast<ANSICHAR*>(StringCast<ANSICHAR, TCHAR>(Text.GetData(), Text.Len()).Get()), Text.Len() * sizeof(ANSICHAR));
+		Archive->Serialize(const_cast<ANSICHAR*>(StringCast<ANSICHAR>(Text.GetData(), Text.Len()).Get()), Text.Len() * sizeof(ANSICHAR));
 	}
 
 	void FLLMCsvWriter::CreateArchive()
diff --git a/Engine/Source/Runtime/Core/Public/Containers/StringConv.h b/Engine/Source/Runtime/Core/Public/Containers/StringConv.h
--- a/Engine/Source/Runtime/Core/Public/Containers/StringConv.h
+++ b/Engine/Source/Runtime/Core/Public/Containers/StringConv.h
@@ -1016,10 +1016,17 @@ typedef TStringPointer<wchar_t, TCHAR> FWCharToTCHAR;
 /**
  * StringCast example usage:
  *
- * void Func(const FString& Str)
+ * void Func(const FString& Str, const TCHAR* MediumLong, FStringView MediumLongSV)
  * {
- *     auto Src = StringCast<ANSICHAR>();
+ *     // Basic version
+ *     auto Src = StringCast<ANSICHAR>(*Str);
  *     const ANSICHAR* Ptr = Src.Get(); // Ptr is a pointer to an ANSICHAR representing the potentially-converted string data.
+ *     // Avoid calling strlen
+ *     auto NoStrlenString = StringCast<ANSICHAR>(MediumLongSV.GetData(), MediumLongSV.Len());
+ *     // Specify the static allocation size when maximum length is known
+ *     auto NoAllocSrc = StringCast<ANSICHAR, 1024>(MediumLong);
+ *     // Avoid calling strlen and also specify the allocation size
+ *     auto NoAllocNoStrlenSrc = StringCast<ANSICHAR, 1024>(MediumLongSV.GetData(), MediumLongSV.Len());
  * }
  *
  */
@@ -1029,7 +1036,7 @@ typedef TStringPointer<wchar_t, TCHAR> FWCharToTCHAR;
  *
  * @param Str The null-terminated source string to convert.
  */
-template <typename To, typename From>
+template <typename To, int32 DefaultConversionSize = DEFAULT_STRING_CONVERSION_SIZE, typename From>
 FORCEINLINE auto StringCast(const From* Str)
 {
 	if constexpr (FPlatformString::IsCharEncodingCompatibleWith<From, To>())
@@ -1038,7 +1045,7 @@ FORCEINLINE auto StringCast(const From* Str)
 	}
 	else
 	{
-		return TStringConversion<TStringConvert<From, To>>(Str);
+		return TStringConversion<TStringConvert<From, To>, DefaultConversionSize>(Str);
 	}
 }
 
@@ -1048,7 +1055,7 @@ FORCEINLINE auto StringCast(const From* Str)
  * @param Str The source string to convert, not necessarily null-terminated.
  * @param Len The number of From elements in Str.
  */
-template <typename To, typename From>
+template <typename To, int32 DefaultConversionSize = DEFAULT_STRING_CONVERSION_SIZE, typename From>
 FORCEINLINE auto StringCast(const From* Str, int32 Len)
 {
 	if constexpr (FPlatformString::IsCharEncodingCompatibleWith<From, To>())
@@ -1057,7 +1064,7 @@ FORCEINLINE auto StringCast(const From* Str, int32 Len)
 	}
 	else
 	{
-		return TStringConversion<TStringConvert<From, To>>(Str, Len);
+		return TStringConversion<TStringConvert<From, To>, DefaultConversionSize>(Str, Len);
 	}
 }
 
