TRefCountPtr: Add move operator= from other templated type, so that TRefCountPtr<const Foo*> P = TRefCountPtr<Foo*> G will do a move operation.
#rnx

diff --git a/Engine/Source/Runtime/Core/Public/Templates/RefCounting.h b/Engine/Source/Runtime/Core/Public/Templates/RefCounting.h
--- a/Engine/Source/Runtime/Core/Public/Templates/RefCounting.h
+++ b/Engine/Source/Runtime/Core/Public/Templates/RefCounting.h
@@ -240,6 +240,20 @@ public:
 		return *this;
 	}
 
+	template<typename MoveReferencedType>
+	TRefCountPtr& operator=(TRefCountPtr<MoveReferencedType>&& InPtr)
+	{
+		// InPtr is a different type (or we would have called the other operator), so we need not test &InPtr != this
+		ReferencedType* OldReference = Reference;
+		Reference = InPtr.Reference;
+		InPtr.Reference = nullptr;
+		if (OldReference)
+		{
+			OldReference->Release();
+		}
+		return *this;
+	}
+
 	FORCEINLINE ReferencedType* operator->() const
 	{
 		return Reference;
