Undo change to Mac to use ansi posix malloc as malloc_usable_size is not available in OSX11.1
#rnx

diff --git a/Engine/Source/Runtime/Core/Public/Mac/MacPlatform.h b/Engine/Source/Runtime/Core/Public/Mac/MacPlatform.h
--- a/Engine/Source/Runtime/Core/Public/Mac/MacPlatform.h
+++ b/Engine/Source/Runtime/Core/Public/Mac/MacPlatform.h
@@ -84,7 +84,6 @@ typedef FMacPlatformTypes FPlatformTypes;
 #define PLATFORM_HAS_BSD_SOCKET_FEATURE_MSG_DONTWAIT	1
 #define PLATFORM_IS_ANSI_MALLOC_THREADSAFE				1
 #define PLATFORM_SUPPORTS_VIRTUAL_TEXTURE_STREAMING		1
-#define PLATFORM_USE_ANSI_POSIX_MALLOC					1
 
 #define PLATFORM_RHITHREAD_DEFAULT_BYPASS				0
 
