Undo change to IOS to use ansi posix malloc as malloc_usable_size is not available
#rnx

diff --git a/Engine/Source/Runtime/Core/Public/IOS/IOSPlatform.h b/Engine/Source/Runtime/Core/Public/IOS/IOSPlatform.h
--- a/Engine/Source/Runtime/Core/Public/IOS/IOSPlatform.h
+++ b/Engine/Source/Runtime/Core/Public/IOS/IOSPlatform.h
@@ -56,7 +56,6 @@ typedef FIOSPlatformTypes FPlatformTypes;
 #define PLATFORM_SUPPORTS_EARLY_MOVIE_PLAYBACK			1 // movies will start before engine is initalized
 #define PLATFORM_USE_FULL_TASK_GRAPH					0 // @todo platplug: not platplug, but should investigate soon anyway
 #define PLATFORM_IS_ANSI_MALLOC_THREADSAFE				1
-#define PLATFORM_USE_ANSI_POSIX_MALLOC					1
 
 // on iOS we now perform offline symbolication as it's significantly faster. Requires bGenerateCrashReportSymbols=true in the ini file.
 #define	PLATFORM_RUNTIME_MALLOCPROFILER_SYMBOLICATION	0	
