FEM - Move GetActiveTask as it is a Property of the Task.

diff --git a/Engine/Source/Runtime/Core/Private/Async/Fundamental/Scheduler.cpp b/Engine/Source/Runtime/Core/Private/Async/Fundamental/Scheduler.cpp
--- a/Engine/Source/Runtime/Core/Private/Async/Fundamental/Scheduler.cpp
+++ b/Engine/Source/Runtime/Core/Private/Async/Fundamental/Scheduler.cpp
@@ -13,7 +13,7 @@ namespace LowLevelTasks
 	DEFINE_LOG_CATEGORY(LowLevelTasks);
 
 	thread_local FSchedulerTls::FLocalQueueType* FSchedulerTls::LocalQueue = nullptr;
-	thread_local FTask* FSchedulerTls::ActiveTask = nullptr;
+	thread_local FTask* FTask::ActiveTask = nullptr;
 	thread_local FSchedulerTls* FSchedulerTls::ActiveScheduler = nullptr;
 	thread_local FSchedulerTls::EWorkerType FSchedulerTls::WorkerType = FSchedulerTls::EWorkerType::None;
 	thread_local uint32 FSchedulerTls::BusyWaitingDepth = 0;
@@ -25,7 +25,7 @@ namespace LowLevelTasks
 		RegisteredLocalQueue = FSchedulerTls::LocalQueue == nullptr;
 		if (RegisteredLocalQueue)
 		{
-			bool bPermitBackgroundWork = FSchedulerTls::PermitBackgroundWork();
+			bool bPermitBackgroundWork = FTask::PermitBackgroundWork();
 			FSchedulerTls::LocalQueue = FSchedulerTls::FLocalQueueType::AllocateLocalQueue(Scheduler.QueueRegistry, bPermitBackgroundWork ? ELocalQueueType::EBackground : ELocalQueueType::EForeground);
 		}
 	}
@@ -34,7 +34,7 @@ namespace LowLevelTasks
 	{
 		if (RegisteredLocalQueue)
 		{
-			bool bPermitBackgroundWork = FSchedulerTls::PermitBackgroundWork();
+			bool bPermitBackgroundWork = FTask::PermitBackgroundWork();
 			FSchedulerTls::FLocalQueueType::DeleteLocalQueue(FSchedulerTls::LocalQueue, bPermitBackgroundWork ? ELocalQueueType::EBackground : ELocalQueueType::EForeground);
 			FSchedulerTls::LocalQueue = nullptr;
 		}
@@ -217,11 +217,6 @@ namespace LowLevelTasks
 		}
 	}
 
-	const FTask* FSchedulerTls::GetActiveTask() 
-	{
-		return ActiveTask;
-	}
-
 	bool FSchedulerTls::IsWorkerThread() const
 	{
 		return WorkerType != FSchedulerTls::EWorkerType::None && ActiveScheduler == this;
@@ -241,31 +236,6 @@ namespace LowLevelTasks
 		return ~0;
 	}
 
-	void FTask::InheritParentData(const TCHAR*& DebugName, ETaskPriority& Priority)
-	{
-		const FTask* ActiveTask = FSchedulerTls::GetActiveTask();
-		if (ActiveTask != nullptr)
-		{
-			if(DebugName == nullptr)
-			{
-				DebugName = ActiveTask->GetDebugName();
-			}
-			if (Priority == ETaskPriority::Inherit)
-			{
-				Priority = ActiveTask->GetPriority();
-			}
-			UserData = ActiveTask->GetUserData();
-		}
-		else
-		{
-			if (Priority == ETaskPriority::Inherit)
-			{
-				Priority = ETaskPriority::Default;
-			}
-			UserData = nullptr;
-		}
-	}
-
 	template<FTask* (FSchedulerTls::FLocalQueueType::*DequeueFunction)(bool, bool), bool bIsBusyWaiting>
 	bool FScheduler::TryExecuteTaskFrom(FSchedulerTls::FLocalQueueType* Queue, FSchedulerTls::FQueueRegistry::FOutOfWork& OutOfWork, bool bPermitBackgroundWork, bool bDisableThrottleStealing)
 	{
@@ -287,8 +257,8 @@ namespace LowLevelTasks
 						WakeUpWorker(!bPermitBackgroundWork);
 					}
 				}
-				FTask* OldTask = FSchedulerTls::ActiveTask;
-				FSchedulerTls::ActiveTask = Task;
+				FTask* OldTask = FTask::ActiveTask;
+				FTask::ActiveTask = Task;
 				{
 #if UE_TASK_TRACE_ENABLED
 					if (!UE_TRACE_CHANNELEXPR_IS_ENABLED(TaskTrace::TaskChannel))
@@ -300,7 +270,7 @@ namespace LowLevelTasks
 #endif
 					Task->ExecuteTask();
 				}
-				FSchedulerTls::ActiveTask = OldTask;
+				FTask::ActiveTask = OldTask;
 				return true;
 			}
 			return false;
@@ -406,7 +376,7 @@ namespace LowLevelTasks
 		uint32 WaitCount = 0;
 		bool HasWokenEmergencyWorker = false;
 		const bool bIsBackgroundWorker = FSchedulerTls::IsBackgroundWorker();
-		bool bPermitBackgroundWork = FSchedulerTls::PermitBackgroundWork() || ForceAllowBackgroundWork;
+		bool bPermitBackgroundWork = FTask::PermitBackgroundWork() || ForceAllowBackgroundWork;
 		FSchedulerTls::FQueueRegistry::FOutOfWork OutOfWork = QueueRegistry.GetOutOfWorkScope(bPermitBackgroundWork ? ELocalQueueType::EBackground : ELocalQueueType::EForeground);
 		while (true)
 		{
diff --git a/Engine/Source/Runtime/Core/Private/Async/TaskGraph.cpp b/Engine/Source/Runtime/Core/Private/Async/TaskGraph.cpp
--- a/Engine/Source/Runtime/Core/Private/Async/TaskGraph.cpp
+++ b/Engine/Source/Runtime/Core/Private/Async/TaskGraph.cpp
@@ -2064,7 +2064,7 @@ private:
 
 	bool IsCurrentThreadKnown() final override
 	{
-		return FPlatformTLS::GetTlsValue(PerThreadIDTLSSlot) != nullptr || LowLevelTasks::FScheduler::GetActiveTask() != nullptr;
+		return FPlatformTLS::GetTlsValue(PerThreadIDTLSSlot) != nullptr || LowLevelTasks::FTask::GetActiveTask() != nullptr;
 	}
 
 	ENamedThreads::Type GetCurrentThreadIfKnown(bool bLocalQueue) final override
@@ -2308,7 +2308,7 @@ private:
 		}
 		else 
 		{
-			const LowLevelTasks::FTask* Task = LowLevelTasks::FScheduler::GetActiveTask();
+			const LowLevelTasks::FTask* Task = LowLevelTasks::FTask::GetActiveTask();
 			if (Task != nullptr)
 			{
 				ENamedThreads::Type ThreadConversion[int(LowLevelTasks::ETaskPriority::Count)] = { ENamedThreads::HighThreadPriority, ENamedThreads::NormalThreadPriority, ENamedThreads::BackgroundThreadPriority, ENamedThreads::BackgroundThreadPriority, ENamedThreads::BackgroundThreadPriority };
diff --git a/Engine/Source/Runtime/Core/Public/Async/Fundamental/Scheduler.h b/Engine/Source/Runtime/Core/Public/Async/Fundamental/Scheduler.h
--- a/Engine/Source/Runtime/Core/Public/Async/Fundamental/Scheduler.h
+++ b/Engine/Source/Runtime/Core/Public/Async/Fundamental/Scheduler.h
@@ -100,7 +100,6 @@ namespace LowLevelTasks
 			Foreground,
 		};
 
-		static thread_local FTask* ActiveTask;
 		static thread_local FSchedulerTls* ActiveScheduler;
 		static thread_local FLocalQueueType* LocalQueue;
 		static thread_local EWorkerType WorkerType;
@@ -110,9 +109,6 @@ namespace LowLevelTasks
 	public:
 		CORE_API bool IsWorkerThread() const;
 
-		//get the active task if any
-		CORE_API static const FTask* GetActiveTask();
-
 		// returns true if the current thread execution is in the context of busy-waiting
 		CORE_API static bool IsBusyWaiting();
 
@@ -120,11 +116,6 @@ namespace LowLevelTasks
 		CORE_API static uint32 GetAffinityIndex();
 
 	protected:
-		inline static bool PermitBackgroundWork()
-		{
-			return ActiveTask && ActiveTask->IsBackgroundTask();
-		}
-
 		inline static bool IsBackgroundWorker()
 		{
 			return WorkerType == EWorkerType::Background;
@@ -136,8 +127,6 @@ namespace LowLevelTasks
 		UE_NONCOPYABLE(FScheduler);
 		static constexpr uint32 WorkerSpinCycles = 53;
 
-		static thread_local FTask* ActiveTask;
-
 		static CORE_API FScheduler Singleton;
 
 		// using 16 bytes here because it fits the vtable and one additional pointer
diff --git a/Engine/Source/Runtime/Core/Public/Async/Fundamental/Task.h b/Engine/Source/Runtime/Core/Public/Async/Fundamental/Task.h
--- a/Engine/Source/Runtime/Core/Public/Async/Fundamental/Task.h
+++ b/Engine/Source/Runtime/Core/Public/Async/Fundamental/Task.h
@@ -203,6 +203,8 @@ namespace LowLevelTasks
 		friend class FScheduler;
 		UE_NONCOPYABLE(FTask); //means non movable
 
+		static thread_local FTask* ActiveTask;
+
 	public: //Public Interface
 		//means the task is completed and this taskhandle can be recycled
 		inline bool IsCompleted(std::memory_order MemoryOrder = std::memory_order_seq_cst) const
@@ -225,7 +227,12 @@ namespace LowLevelTasks
 			return State == ETaskState::Ready || State == ETaskState::CanceledAndReady; 
 		}
 
-		
+		//get the currently active task if any
+		CORE_API static const FTask* GetActiveTask()
+		{
+			return ActiveTask;
+		}
+
 		//try to cancel the task if it has not been launched yet and ExecuteTaskOnSuccess is true the continuation will run immediately.
 		inline bool TryCancel(bool ExecuteTaskOnSuccess = true);
 
@@ -268,16 +275,47 @@ namespace LowLevelTasks
 		inline ~FTask();
 
 	private: //Interface of the Scheduler
+		inline static bool PermitBackgroundWork()
+		{
+			return ActiveTask && ActiveTask->IsBackgroundTask();
+		}
+
 		inline bool TryPrepareLaunch();
 		//after calling this function the task can be considered dead
 		inline void ExecuteTask();
-		CORE_API void InheritParentData(const TCHAR*& DebugName, ETaskPriority& Priority);
+		inline void InheritParentData(ETaskPriority& Priority);
 	};
 
    /******************
 	* IMPLEMENTATION *
 	******************/
 
+	inline ETaskPriority FTask::GetPriority() const 
+	{ 
+		return PackedData.load(std::memory_order_relaxed).GetPriority(); 
+	}
+
+	inline void FTask::InheritParentData(ETaskPriority& Priority)
+	{
+		const FTask* LocalActiveTask = FTask::GetActiveTask();
+		if (LocalActiveTask != nullptr)
+		{
+			if (Priority == ETaskPriority::Inherit)
+			{
+				Priority = LocalActiveTask->GetPriority();
+			}
+			UserData = LocalActiveTask->GetUserData();
+		}
+		else
+		{
+			if (Priority == ETaskPriority::Inherit)
+			{
+				Priority = ETaskPriority::Default;
+			}
+			UserData = nullptr;
+		}
+	}
+
 	template<typename TRunnable, typename TContinuation>
 	inline void FTask::Init(const TCHAR* InDebugName, ETaskPriority InPriority, TRunnable&& InRunnable, TContinuation&& InContinuation, bool bAllowBusyWaiting)
 	{
@@ -301,7 +339,7 @@ namespace LowLevelTasks
 				return true;
 			}
 		};
-		InheritParentData(InDebugName, InPriority);
+		InheritParentData(InPriority);
 		PackedData.store(FPackedData(InDebugName, InPriority, ETaskState::Ready, bAllowBusyWaiting), std::memory_order_release);
 	}
 
@@ -326,7 +364,7 @@ namespace LowLevelTasks
 				return true;
 			}
 		};
-		InheritParentData(InDebugName, InPriority);
+		InheritParentData(InPriority);
 		PackedData.store(FPackedData(InDebugName, InPriority, ETaskState::Ready, bAllowBusyWaiting), std::memory_order_release);
 	}
 
@@ -397,11 +435,6 @@ namespace LowLevelTasks
 		PackedData.store(FPackedData(LocalPackedData, Canceled ? ETaskState::CanceledAndCompleted : ETaskState::Completed), std::memory_order_seq_cst);
 	}
 
-	inline ETaskPriority FTask::GetPriority() const 
-	{ 
-		return PackedData.load(std::memory_order_relaxed).GetPriority(); 
-	}
-
 	inline const TCHAR* FTask::GetDebugName() const
 	{
 		return PackedData.load(std::memory_order_relaxed).GetDebugName(); 
diff --git a/Engine/Source/Runtime/Core/Public/Async/LocalWorkQueue.h b/Engine/Source/Runtime/Core/Public/Async/LocalWorkQueue.h
--- a/Engine/Source/Runtime/Core/Public/Async/LocalWorkQueue.h
+++ b/Engine/Source/Runtime/Core/Public/Async/LocalWorkQueue.h
@@ -61,7 +61,7 @@ public:
 	{
 		if (Priority == LowLevelTasks::ETaskPriority::Count)
 		{
-			const LowLevelTasks::FTask* ActiveTask = LowLevelTasks::FScheduler::Get().GetActiveTask();
+			const LowLevelTasks::FTask* ActiveTask = LowLevelTasks::FTask::GetActiveTask();
 			if (ActiveTask)
 			{
 				Priority = ActiveTask->GetPriority();
