FEM - Linux Desktop Platform is still inlining (DLL)EXPORT marked functions, therefor we prevent inlining on all Destop Platforms

diff --git a/Engine/Source/Runtime/Core/Private/Async/Fundamental/Scheduler.cpp b/Engine/Source/Runtime/Core/Private/Async/Fundamental/Scheduler.cpp
--- a/Engine/Source/Runtime/Core/Private/Async/Fundamental/Scheduler.cpp
+++ b/Engine/Source/Runtime/Core/Private/Async/Fundamental/Scheduler.cpp
@@ -217,6 +217,13 @@ namespace LowLevelTasks
 		}
 	}
 
+#if PLATFORM_DESKTOP
+	const FTask* FTask::GetActiveTask()
+	{
+		return ActiveTask;
+	}
+#endif
+
 	bool FSchedulerTls::IsWorkerThread() const
 	{
 		return WorkerType != FSchedulerTls::EWorkerType::None && ActiveScheduler == this;
diff --git a/Engine/Source/Runtime/Core/Public/Async/Fundamental/Task.h b/Engine/Source/Runtime/Core/Public/Async/Fundamental/Task.h
--- a/Engine/Source/Runtime/Core/Public/Async/Fundamental/Task.h
+++ b/Engine/Source/Runtime/Core/Public/Async/Fundamental/Task.h
@@ -227,11 +227,15 @@ namespace LowLevelTasks
 			return State == ETaskState::Ready || State == ETaskState::CanceledAndReady; 
 		}
 
+#if PLATFORM_DESKTOP
 		//get the currently active task if any
-		CORE_API static const FTask* GetActiveTask()
+		CORE_API static const FTask* GetActiveTask();
+#else
+		FORCEINLINE const FTask* GetActiveTask()
 		{
 			return ActiveTask;
 		}
+#endif
 
 		//try to cancel the task if it has not been launched yet and ExecuteTaskOnSuccess is true the continuation will run immediately.
 		inline bool TryCancel(bool ExecuteTaskOnSuccess = true);
