FEM - missing static

diff --git a/Engine/Source/Runtime/Core/Public/Async/Fundamental/Task.h b/Engine/Source/Runtime/Core/Public/Async/Fundamental/Task.h
--- a/Engine/Source/Runtime/Core/Public/Async/Fundamental/Task.h
+++ b/Engine/Source/Runtime/Core/Public/Async/Fundamental/Task.h
@@ -231,7 +231,7 @@ namespace LowLevelTasks
 		//get the currently active task if any
 		CORE_API static const FTask* GetActiveTask();
 #else
-		FORCEINLINE const FTask* GetActiveTask()
+		FORCEINLINE static const FTask* GetActiveTask()
 		{
 			return ActiveTask;
 		}
