Fixed float convertion warning

diff --git a/Engine/Source/Runtime/Core/Private/FramePro/FrameProProfiler.cpp b/Engine/Source/Runtime/Core/Private/FramePro/FrameProProfiler.cpp
--- a/Engine/Source/Runtime/Core/Private/FramePro/FrameProProfiler.cpp
+++ b/Engine/Source/Runtime/Core/Private/FramePro/FrameProProfiler.cpp
@@ -222,7 +222,7 @@ static void SendCPUStats()
 #endif // PLATFORM_ANDROID
 }
 
-float GFrameProCPUStatsUpdateRate = 0.00100; // default to 1000Hz (ie every frame)
+float GFrameProCPUStatsUpdateRate = 0.00100f; // default to 1000Hz (ie every frame)
 static FAutoConsoleVariableRef CVarFrameProCPUStatsUpdateRate(
 	TEXT("framepro.CPUStatsUpdateRate"),
 	GFrameProCPUStatsUpdateRate,
