LWC: Enable ISPC for LWC double types. Submitted on behalf of jeff.rous[REDACTED_HANDLE]

diff --git a/Engine/Plugins/Experimental/ChaosCloth/Source/Chaos/Private/ChaosCloth/ChaosClothingSimulation.cpp b/Engine/Plugins/Experimental/ChaosCloth/Source/Chaos/Private/ChaosCloth/ChaosClothingSimulation.cpp
--- a/Engine/Plugins/Experimental/ChaosCloth/Source/Chaos/Private/ChaosCloth/ChaosClothingSimulation.cpp
+++ b/Engine/Plugins/Experimental/ChaosCloth/Source/Chaos/Private/ChaosCloth/ChaosClothingSimulation.cpp
@@ -65,7 +65,14 @@
 #include "Chaos/VelocityField.h"
 #endif
 
+#if INTEL_ISPC && UE_LARGE_WORLD_COORDINATES_DISABLED
+namespace ispc { typedef FVector FVector3f; typedef FVector4 FVector4f; } // Need this because ISPC doesn't expose typedefs in header and optimized function uses mixed types
+#endif
+
 #if INTEL_ISPC && !UE_BUILD_SHIPPING
+static_assert(sizeof(ispc::FVector3f) == sizeof(FVector3f), "sizeof(ispc::FVector3f) != sizeof(FVector3f)");
+static_assert(sizeof(ispc::FVector4f) == sizeof(FQuat4f), "sizeof(ispc::FVector4f) != sizeof(Chaos::FQuat4f))");
+
 bool bChaos_GetSimData_ISPC_Enabled = true;
 FAutoConsoleVariableRef CVarChaosGetSimDataISPCEnabled(TEXT("p.Chaos.GetSimData.ISPC"), bChaos_GetSimData_ISPC_Enabled, TEXT("Whether to use ISPC optimizations when getting simulation data"));
 #endif
@@ -670,11 +677,22 @@ void FClothingSimulation::GetSimulationData(
 		if (bRealTypeCompatibleWithISPC && bChaos_GetSimData_ISPC_Enabled)
 		{
 #if INTEL_ISPC
+			const FVector3f Location = FVector3f(LocalSpaceLocation);
+			const FQuat4f Rotation = FQuat4f(ReferenceSpaceTransform.GetRotation());
+			const FVector3f Translation = FVector3f(ReferenceSpaceTransform.GetTranslation());
+			const FVector3f Scale3D = FVector3f(ReferenceSpaceTransform.GetScale3D());
+
+			// ISPC is assuming float input here
+			check(sizeof(FVector3f) == Data.Positions.GetTypeSize());
+			check(sizeof(FVector3f) == Data.Normals.GetTypeSize());
+
 			ispc::GetClothingSimulationData(
-				(ispc::FVector*)Data.Positions.GetData(),
-				(ispc::FVector*)Data.Normals.GetData(),
-				(ispc::FTransform&)ReferenceSpaceTransform,
-				(ispc::FVector&)LocalSpaceLocation,
+				(ispc::FVector3f*)Data.Positions.GetData(),
+				(ispc::FVector3f*)Data.Normals.GetData(),
+				(ispc::FVector4f&)Rotation,
+				(ispc::FVector3f&)Translation,
+				(ispc::FVector3f&)Scale3D,
+				(ispc::FVector3f&)Location,
 				Data.Positions.Num());
 #endif
 		}
diff --git a/Engine/Plugins/Experimental/ChaosCloth/Source/Chaos/Private/ChaosCloth/ChaosClothingSimulation.ispc b/Engine/Plugins/Experimental/ChaosCloth/Source/Chaos/Private/ChaosCloth/ChaosClothingSimulation.ispc
--- a/Engine/Plugins/Experimental/ChaosCloth/Source/Chaos/Private/ChaosCloth/ChaosClothingSimulation.ispc
+++ b/Engine/Plugins/Experimental/ChaosCloth/Source/Chaos/Private/ChaosCloth/ChaosClothingSimulation.ispc
@@ -6,32 +6,33 @@
 #include "Math/Quat.isph"
 #include "Math/Transform.isph"
 
-export void GetClothingSimulationData(uniform FVector Positions[],
-									uniform FVector Normals[],
-									const uniform FTransform &OwnerTransform,
-									const uniform FVector &LocalSpaceLocation,
+export void GetClothingSimulationData(uniform FVector3f Positions[],
+									uniform FVector3f Normals[],
+									const uniform FVector4f &Rotation,
+									const uniform FVector3f &Translation,
+									const uniform FVector3f &Scale3D,
+									const uniform FVector3f &LocalSpaceLocation,
 									const uniform int NumPositions)
 {
-	const uniform FVector4 RotationInverse = QuatInverse(OwnerTransform.Rotation);
-	const uniform FVector Translation = SetVector(OwnerTransform.Translation);
-	const uniform FVector SafeReciprocal = SetVector(TransformGetSafeScaleReciprocal(OwnerTransform.Scale3D, SMALL_NUMBER));
+	const uniform FVector4f RotationInverse = QuatInverse(Rotation);
+	const uniform FVector3f SafeReciprocal = TransformGetSafeScaleReciprocal(Scale3D, FLOAT_SMALL_NUMBER);
 
 	foreach(i = 0 ... NumPositions)
 	{
 		const uniform int Index = extract(i, 0);
-		const FVector Position = VectorLoad(&Positions[Index]);
-		const FVector Normal = VectorLoad(&Normals[Index]);
+		const FVector3f Position = VectorLoad(&Positions[Index]);
+		const FVector3f Normal = VectorLoad(&Normals[Index]);
 
-		const FVector InputPositionVector = Position + LocalSpaceLocation;
-		const FVector InputNormalVector = Normal * -1.0f;
+		const FVector3f InputPositionVector = Position + LocalSpaceLocation;
+		const FVector3f InputNormalVector = Normal * -FLOAT_ONE;
 
-		FVector TranslatedPosition = InputPositionVector - Translation;
+		FVector3f TranslatedPosition = InputPositionVector - Translation;
 
-		const FVector VRPosition = VectorQuaternionRotateVector(RotationInverse, TranslatedPosition);
-		const FVector VRNormal = VectorQuaternionRotateVector(RotationInverse, InputNormalVector);
+		const FVector3f VRPosition = VectorQuaternionRotateVector(RotationInverse, TranslatedPosition);
+		const FVector3f VRNormal = VectorQuaternionRotateVector(RotationInverse, InputNormalVector);
 
-		const FVector OutPosition = VRPosition * SafeReciprocal;
-		const FVector OutNormal = VRNormal * SafeReciprocal;
+		const FVector3f OutPosition = VRPosition * SafeReciprocal;
+		const FVector3f OutNormal = VRNormal * SafeReciprocal;
 
 		VectorStore(&Positions[Index], OutPosition);
 		VectorStore(&Normals[Index], OutNormal);
diff --git a/Engine/Plugins/Experimental/ChaosCloth/Source/Chaos/Private/ChaosCloth/ChaosClothingSimulationMesh.cpp b/Engine/Plugins/Experimental/ChaosCloth/Source/Chaos/Private/ChaosCloth/ChaosClothingSimulationMesh.cpp
--- a/Engine/Plugins/Experimental/ChaosCloth/Source/Chaos/Private/ChaosCloth/ChaosClothingSimulationMesh.cpp
+++ b/Engine/Plugins/Experimental/ChaosCloth/Source/Chaos/Private/ChaosCloth/ChaosClothingSimulationMesh.cpp
@@ -12,9 +12,19 @@
 #include "ChaosClothingSimulationMesh.ispc.generated.h"
 #endif
 
+#if INTEL_ISPC && UE_LARGE_WORLD_COORDINATES_DISABLED
+namespace ispc { typedef FVector FVector3f; typedef FMatrix FMatrix44f; } // Need this because ISPC doesn't expose typedefs in header and optimized function uses mixed types
+#endif
+
 #if INTEL_ISPC && !UE_BUILD_SHIPPING
-bool bChaos_SkinPhysicsMesh_ISPC_Enabled = bChaos_SkinPhysicsMesh_ISPC_Enable;
+bool bChaos_SkinPhysicsMesh_ISPC_Enabled = true;
 FAutoConsoleVariableRef CVarChaosSkinPhysicsMeshISPCEnabled(TEXT("p.Chaos.SkinPhysicsMesh.ISPC"), bChaos_SkinPhysicsMesh_ISPC_Enabled, TEXT("Whether to use ISPC optimizations on skinned physics meshes"));
+
+static_assert(sizeof(ispc::FVector) == sizeof(Chaos::FVec3), "sizeof(ispc::FVector) != sizeof(FVec3)");
+static_assert(sizeof(ispc::FVector3f) == sizeof(FVector3f), "sizeof(ispc::FVector3f) != sizeof(FVector3f)");
+static_assert(sizeof(ispc::FMatrix44f) == sizeof(FMatrix44f), "sizeof(ispc::FMatrix44f) != sizeof(FMatrix44f)");
+static_assert(sizeof(ispc::FTransform) == sizeof(FTransform), "sizeof(ispc::FTransform) != sizeof(FTransform)");
+static_assert(sizeof(ispc::FClothVertBoneData) == sizeof(FClothVertBoneData), "sizeof(ispc::FClothVertBoneData) != sizeof(Chaos::FClothVertBoneData)");
 #endif
 
 DECLARE_CYCLE_STAT(TEXT("Chaos Cloth Skin Physics Mesh"), STAT_ChaosClothSkinPhysicsMesh, STATGROUP_ChaosCloth);
@@ -258,12 +268,12 @@ void FClothingSimulationMesh::SkinPhysicsMesh(int32 LODIndex, const FVec3& Local
 		ispc::SkinPhysicsMesh(
 			(ispc::FVector*)OutPositions,
 			(ispc::FVector*)OutNormals,
-			(ispc::FVector*)PhysicalMeshData.Vertices.GetData(),
-			(ispc::FVector*)PhysicalMeshData.Normals.GetData(),
+			(ispc::FVector3f*)PhysicalMeshData.Vertices.GetData(),
+			(ispc::FVector3f*)PhysicalMeshData.Normals.GetData(),
 			(ispc::FClothVertBoneData*)PhysicalMeshData.BoneData.GetData(),
 			BoneMap,
-			(ispc::FMatrix*)BoneMatrices,
-			(ispc::FTransform&)ComponentToLocalSpace,
+			(ispc::FMatrix44f*)BoneMatrices,
+			(ispc::FTransform&)ComponentToLocalSpaceReal,
 			NumPoints);
 	}
 	else
diff --git a/Engine/Plugins/Experimental/ChaosCloth/Source/Chaos/Private/ChaosCloth/ChaosClothingSimulationMesh.ispc b/Engine/Plugins/Experimental/ChaosCloth/Source/Chaos/Private/ChaosCloth/ChaosClothingSimulationMesh.ispc
--- a/Engine/Plugins/Experimental/ChaosCloth/Source/Chaos/Private/ChaosCloth/ChaosClothingSimulationMesh.ispc
+++ b/Engine/Plugins/Experimental/ChaosCloth/Source/Chaos/Private/ChaosCloth/ChaosClothingSimulationMesh.ispc
@@ -5,7 +5,6 @@
 #include "Math/Vector.isph"
 #include "Math/Matrix.isph"
 #include "Math/Transform.isph"
-#include "Chaos/Matrix33.isph"
 
 #define MaxTotalInfluences 12
 
@@ -16,50 +15,55 @@ struct FClothVertBoneData
 	float BoneWeights[MaxTotalInfluences];
 };
 
-static inline void AddInfluence(uniform FVector& OutPosition, uniform FVector& OutNormal, const uniform FVector& RefParticle, const uniform FVector& RefNormal, const uniform FMatrix *uniform BoneMatrix, const uniform float Weight)
+static inline void AddInfluence(uniform FVector4f& OutPosition, uniform FVector4f& OutNormal, const uniform FVector3f& RefParticle, const uniform FVector3f& RefNormal, const uniform FMatrix44f *uniform BoneMatrix, const uniform float Weight)
 {
-	const uniform FVector4 M0M1M2 = *((uniform FVector4 *uniform)&BoneMatrix->M[0]);
+	const uniform FVector4f M0M1M2 = *((uniform FVector4f *uniform)&BoneMatrix->M[0]);
 
-	const uniform FVector4 VTempXP = RefParticle.V[0] * M0M1M2;
-	const uniform FVector4 VTempXN = RefNormal.V[0] * M0M1M2;
+	const uniform FVector4f VTempXP = RefParticle.V[0] * M0M1M2;
+	const uniform FVector4f VTempXN = RefNormal.V[0] * M0M1M2;
 
-	const uniform FVector4 M4M5M6 = *((uniform FVector4 *uniform)&BoneMatrix->M[4]);
+	const uniform FVector4f M4M5M6 = *((uniform FVector4f *uniform)&BoneMatrix->M[4]);
 
-	const uniform FVector4 VTempYP = RefParticle.V[1] * M4M5M6 + VTempXP;
-	const uniform FVector4 VTempYN = RefNormal.V[1] * M4M5M6 + VTempXN;
+	const uniform FVector4f VTempYP = RefParticle.V[1] * M4M5M6 + VTempXP;
+	const uniform FVector4f VTempYN = RefNormal.V[1] * M4M5M6 + VTempXN;
 
-	const uniform FVector4 M8M9M10 = *((uniform FVector4 *uniform)&BoneMatrix->M[8]);
+	const uniform FVector4f M8M9M10 = *((uniform FVector4f *uniform)&BoneMatrix->M[8]);
 
-	const uniform FVector4 VTempZP = RefParticle.V[2] * M8M9M10 + VTempYP;
-	const uniform FVector4 Norm = RefNormal.V[2] * M8M9M10 + VTempYN;
+	const uniform FVector4f VTempZP = RefParticle.V[2] * M8M9M10 + VTempYP;
+	const uniform FVector4f Norm = RefNormal.V[2] * M8M9M10 + VTempYN;
 
-	const uniform FVector4 VTempWP = *((uniform FVector4 *uniform)&BoneMatrix->M[12]);
-	const uniform FVector4 Pos = VTempZP + VTempWP;
+	const uniform FVector4f VTempWP = *((uniform FVector4f *uniform)&BoneMatrix->M[12]);
+	const uniform FVector4f Pos = VTempZP + VTempWP;
 
-	OutPosition = OutPosition + SetVector(Pos.V[0], Pos.V[1], Pos.V[2]) * Weight;
-	OutNormal = OutNormal + SetVector(Norm.V[0], Norm.V[1], Norm.V[2]) * Weight;
+	OutPosition = OutPosition + Pos * Weight;
+	OutNormal = OutNormal + Norm * Weight;
+}
+
+static inline uniform FVector GetSafeNormal(const uniform FVector &V)
+{
+	//We want N / ||N|| and to avoid inf
+	//So we want N / ||N|| < 1 / eps => N eps < ||N||, but this is clearly true for all eps < 1 and N > 0
+	const uniform FReal SizeSqr = VectorSizeSquared(V);
+	return VectorSelect(SizeSqr <= REAL_MIN, ForwardVector, V / sqrt(SizeSqr));
 }
 
 export void SkinPhysicsMesh(uniform FVector OutPositions[],
 							uniform FVector OutNormals[],
-							const uniform FVector Vertices[],
-							const uniform FVector Normals[],
+							const uniform FVector3f Vertices[],
+							const uniform FVector3f Normals[],
 							const uniform FClothVertBoneData BoneData[],
 							const uniform int32 BoneMap[],
-							const uniform FMatrix BoneMatrices[],
+							const uniform FMatrix44f BoneMatrices[],
 							const uniform FTransform &ComponentToLocalSpace,
 							const uniform int32 NumPoints)
 {
 	for(uniform int i = 0; i < NumPoints; ++i)
 	{
-		const uniform FVector RefParticle = Vertices[i];
-		const uniform FVector RefNormal = Normals[i];
+		const uniform FVector3f RefParticle = Vertices[i];
+		const uniform FVector3f RefNormal = Normals[i];
 
-		uniform FVector OutPosition = OutPositions[i];
-		uniform FVector OutNormal = OutNormals[i];
-
-		OutPosition = SetVector(0.0f, 0.0f, 0.0f);
-		OutNormal = SetVector(0.0f, 0.0f, 0.0f);
+		uniform FVector4f OutPosition = FloatZero;
+		uniform FVector4f OutNormal = FloatZero;
 
 		const uniform int NumInfluences = BoneData[i].NumInfluences;
 
@@ -68,11 +72,17 @@ export void SkinPhysicsMesh(uniform FVector OutPositions[],
 			const uniform int BoneMatrixIndex = BoneMap[BoneData[i].BoneIndices[j]];
 			const uniform float Weight = BoneData[i].BoneWeights[j];
 
-			const uniform FMatrix *uniform BoneMatrix = &BoneMatrices[BoneMatrixIndex];
+			const uniform FMatrix44f *uniform BoneMatrix = &BoneMatrices[BoneMatrixIndex];
 			AddInfluence(OutPosition, OutNormal, RefParticle, RefNormal, BoneMatrix, Weight);
 		}
 
-		OutPositions[i] = TransformPosition(ComponentToLocalSpace, OutPosition);
-		OutNormals[i] = VectorNormalize(TransformVector(ComponentToLocalSpace, OutNormal));
+		const uniform FVector4 ConvertedPosition = ConvertVector4fTo4Native(OutPosition);
+		const uniform FVector4 ConvertedNormal = ConvertVector4fTo4Native(OutNormal);
+
+		const uniform FVector Position = SetVector(ConvertedPosition);
+		const uniform FVector Normal = SetVector(ConvertedNormal);
+
+		OutPositions[i] = TransformPosition(ComponentToLocalSpace, Position);
+		OutNormals[i] = GetSafeNormal(TransformVector(ComponentToLocalSpace, Normal));
 	}
 }
diff --git a/Engine/Plugins/Experimental/ChaosCloth/Source/Chaos/Private/ChaosCloth/ChaosClothingSimulationSolver.cpp b/Engine/Plugins/Experimental/ChaosCloth/Source/Chaos/Private/ChaosCloth/ChaosClothingSimulationSolver.cpp
--- a/Engine/Plugins/Experimental/ChaosCloth/Source/Chaos/Private/ChaosCloth/ChaosClothingSimulationSolver.cpp
+++ b/Engine/Plugins/Experimental/ChaosCloth/Source/Chaos/Private/ChaosCloth/ChaosClothingSimulationSolver.cpp
@@ -53,6 +53,8 @@ bool bChaos_PreSimulationTransforms_ISPC_Enabled = true;
 FAutoConsoleVariableRef CVarChaosPreSimulationTransformsISPCEnabled(TEXT("p.Chaos.PreSimulationTransforms.ISPC"), bChaos_PreSimulationTransforms_ISPC_Enabled, TEXT("Whether to use ISPC optimizations in ApplySimulationTransforms"));
 bool bChaos_CalculateBounds_ISPC_Enabled = bChaos_CalculateBounds_ISPC_Enable;  // Disabled by default
 FAutoConsoleVariableRef CVarChaosCalculateBoundsISPCEnabled(TEXT("p.Chaos.CalculateBounds.ISPC"), bChaos_CalculateBounds_ISPC_Enabled, TEXT("Whether to use ISPC optimizations in CalculateBounds"));
+
+static_assert(sizeof(ispc::FVector) == sizeof(Chaos::FVec3), "sizeof(ispc::FVector) != sizeof(FVec3)");
 #endif
 
 FAutoConsoleVariableRef CVarChaosClothSolverUseImprovedTimeStepSmoothing(TEXT("p.ChaosCloth.Solver.UseImprovedTimeStepSmoothing"), bChaosClothSolverUseImprovedTimeStepSmoothing, TEXT("Use the time step smoothing on input forces only rather than on the entire cloth solver, in order to avoid miscalculating velocities."));
@@ -962,7 +964,7 @@ FBoxSphereBounds FClothingSimulationSolver::CalculateBounds() const
 						Offset,
 						Range);
 
-					TAABB<float, 3> NewAABB(NewMin, NewMax);
+					FAABB3 NewAABB(NewMin, NewMax);
 					BoundingBox = NewAABB;
 				});
 		}
diff --git a/Engine/Plugins/Experimental/ChaosCloth/Source/Chaos/Private/ChaosCloth/ChaosClothingSimulationSolver.ispc b/Engine/Plugins/Experimental/ChaosCloth/Source/Chaos/Private/ChaosCloth/ChaosClothingSimulationSolver.ispc
--- a/Engine/Plugins/Experimental/ChaosCloth/Source/Chaos/Private/ChaosCloth/ChaosClothingSimulationSolver.ispc
+++ b/Engine/Plugins/Experimental/ChaosCloth/Source/Chaos/Private/ChaosCloth/ChaosClothingSimulationSolver.ispc
@@ -60,13 +60,13 @@ export void CalculateBounds(uniform FVector &MMin,
 	MMax = VectorReduceMax(Max);
 }
 
-export void CalculateSquaredRadius(uniform float &SquaredRadius,
+export void CalculateSquaredRadius(uniform FReal &SquaredRadius,
 									const uniform FVector &Center,
 									const uniform FVector ParticlesX[],
 									const uniform int32 Offset,
 									const uniform int32 Range)
 {
-	varying float vSquaredRadius = SquaredRadius;
+	varying FReal vSquaredRadius = SquaredRadius;
 
 	foreach(i = Offset ... Range)
 	{
diff --git a/Engine/Plugins/Experimental/ChaosCloth/Source/Chaos/Public/ChaosCloth/ChaosClothingSimulationMesh.h b/Engine/Plugins/Experimental/ChaosCloth/Source/Chaos/Public/ChaosCloth/ChaosClothingSimulationMesh.h
--- a/Engine/Plugins/Experimental/ChaosCloth/Source/Chaos/Public/ChaosCloth/ChaosClothingSimulationMesh.h
+++ b/Engine/Plugins/Experimental/ChaosCloth/Source/Chaos/Public/ChaosCloth/ChaosClothingSimulationMesh.h
@@ -78,7 +78,7 @@ constexpr bool bChaos_SkinPhysicsMesh_ISPC_Enable = true;
 #if !INTEL_ISPC
 const bool bChaos_SkinPhysicsMesh_ISPC_Enabled = false;
 #elif UE_BUILD_SHIPPING
-const bool bChaos_SkinPhysicsMesh_ISPC_Enabled = bChaos_SkinPhysicsMesh_ISPC_Enable;
+const bool bChaos_SkinPhysicsMesh_ISPC_Enabled = true;
 #else
 extern bool bChaos_SkinPhysicsMesh_ISPC_Enabled;
 #endif
diff --git a/Engine/Plugins/FX/Niagara/Source/Niagara/Private/NiagaraDataInterfaceVectorField.cpp b/Engine/Plugins/FX/Niagara/Source/Niagara/Private/NiagaraDataInterfaceVectorField.cpp
--- a/Engine/Plugins/FX/Niagara/Source/Niagara/Private/NiagaraDataInterfaceVectorField.cpp
+++ b/Engine/Plugins/FX/Niagara/Source/Niagara/Private/NiagaraDataInterfaceVectorField.cpp
@@ -29,6 +29,7 @@ static const FName GetVectorFieldDimensionsName("FieldDimensions");
 static const FName GetVectorFieldBoundsName("FieldBounds");
 
 #if INTEL_ISPC
+static_assert(sizeof(ispc::FVector) == sizeof(FVector), "sizeof(ispc::FVector) != sizeof(FVector)");
 
 #if UE_BUILD_SHIPPING
 const bool GNiagaraVectorFieldUseIspc = true;
diff --git a/Engine/Plugins/FX/Niagara/Source/Niagara/Private/NiagaraDataInterfaceVectorField.ispc b/Engine/Plugins/FX/Niagara/Source/Niagara/Private/NiagaraDataInterfaceVectorField.ispc
--- a/Engine/Plugins/FX/Niagara/Source/Niagara/Private/NiagaraDataInterfaceVectorField.ispc
+++ b/Engine/Plugins/FX/Niagara/Source/Niagara/Private/NiagaraDataInterfaceVectorField.ispc
@@ -15,7 +15,7 @@ inline varying FVector ReadHalfVector(const uniform FHalfVector Data[], int Data
 	#pragma ignore warning(perf)
 	const FHalfVector Sample = Data[SampleIndex >= 0 && SampleIndex <= DataSize ? SampleIndex : 0];
 
-	return SetVector(half_to_float_fast(Sample.V[0]), half_to_float_fast(Sample.V[1]), half_to_float_fast(Sample.V[2]));
+	return SetVector((FReal)(half_to_float_fast(Sample.V[0])), (FReal)(half_to_float_fast(Sample.V[1])), (FReal)(half_to_float_fast(Sample.V[2])));
 }
 
 export void SampleVectorField(const uniform float XParam[], const uniform float YParam[], const uniform float ZParam[],
@@ -25,15 +25,15 @@ export void SampleVectorField(const uniform float XParam[], const uniform float
 							const uniform FVector &Size, const uniform FVector &TilingAxes, const uniform int NumInstances)
 {
 	const uniform FVector SizeMinusOne = Size - OneVector;
-	const uniform FVector OneOverSize = 1.0f / Size;
+	const uniform FVector OneOverSize = ONE / Size;
 
 	foreach(InstanceIdx = 0 ... NumInstances)
 	{
 		// Position in Volume Space
 		FVector Pos = SetVector(
-			LoadFloat(XParam, InstanceIdx, bXParamConst),
-			LoadFloat(YParam, InstanceIdx, bYParamConst),
-			LoadFloat(ZParam, InstanceIdx, bZParamConst));
+			(FReal)(LoadFloat(XParam, InstanceIdx, bXParamConst)),
+			(FReal)(LoadFloat(YParam, InstanceIdx, bYParamConst)),
+			(FReal)(LoadFloat(ZParam, InstanceIdx, bZParamConst)));
 
 		// Normalize position
 		Pos = (Pos - MinBounds) * OneOverBoundSize;
diff --git a/Engine/Source/Programs/UnrealBuildTool/Configuration/TargetRules.cs b/Engine/Source/Programs/UnrealBuildTool/Configuration/TargetRules.cs
--- a/Engine/Source/Programs/UnrealBuildTool/Configuration/TargetRules.cs
+++ b/Engine/Source/Programs/UnrealBuildTool/Configuration/TargetRules.cs
@@ -2255,12 +2255,12 @@ namespace UnrealBuildTool
 
 		public bool bCompileISPC
 		{
-			get { return false; /*Inner.bCompileISPC && GlobalDefinitions.Contains("UE_LARGE_WORLD_COORDINATES_DISABLED=1");*/ }   // LWC_TODO: Temporarily disable ISPC when LWC is turned on (note this is now forced - we ignore LARGE_WORLD_COORDINATES_DISABLED).
+			get { return Inner.bCompileISPC; }
 		}
 
 		public bool bLWCDisabled 
 		{
-			get { return false; /*GlobalDefinitions.Contains("UE_LARGE_WORLD_COORDINATES_DISABLED=1");*/ }
+			get { return GlobalDefinitions.Contains("UE_LARGE_WORLD_COORDINATES_DISABLED=1"); }
 		}
 		
 		public bool bCompilePython
diff --git a/Engine/Source/Programs/UnrealBuildTool/Configuration/UEBuildModuleCPP.cs b/Engine/Source/Programs/UnrealBuildTool/Configuration/UEBuildModuleCPP.cs
--- a/Engine/Source/Programs/UnrealBuildTool/Configuration/UEBuildModuleCPP.cs
+++ b/Engine/Source/Programs/UnrealBuildTool/Configuration/UEBuildModuleCPP.cs
@@ -486,7 +486,7 @@ namespace UnrealBuildTool
 			// Generate ISPC headers first so C++ can consume them
 			if (InputFiles.ISPCFiles.Count > 0)
 			{
-				CreateHeadersForISPC(ToolChain, CompileEnvironment, InputFiles.ISPCFiles, IntermediateDirectory, Graph);
+				CreateHeadersForISPC(ToolChain, ModuleCompileEnvironment, InputFiles.ISPCFiles, IntermediateDirectory, Graph);
 			}
 
 			// Compile any module interfaces
@@ -586,7 +586,7 @@ namespace UnrealBuildTool
 			// Compile ISPC files directly
 			if (InputFiles.ISPCFiles.Count > 0)
 			{
-				LinkInputFiles.AddRange(ToolChain.CompileISPCFiles(CompileEnvironment, InputFiles.ISPCFiles, IntermediateDirectory, Graph).ObjectFiles);
+				LinkInputFiles.AddRange(ToolChain.CompileISPCFiles(ModuleCompileEnvironment, InputFiles.ISPCFiles, IntermediateDirectory, Graph).ObjectFiles);
 			}
 
 			// Compile C files directly. Do not use a PCH here, because a C++ PCH is not compatible with C source files.
diff --git a/Engine/Source/Programs/UnrealBuildTool/ToolChain/UEToolChain.cs b/Engine/Source/Programs/UnrealBuildTool/ToolChain/UEToolChain.cs
--- a/Engine/Source/Programs/UnrealBuildTool/ToolChain/UEToolChain.cs
+++ b/Engine/Source/Programs/UnrealBuildTool/ToolChain/UEToolChain.cs
@@ -492,6 +492,16 @@ namespace UnrealBuildTool
 					Arguments.Add(String.Format("-I\"{0}\"", SystemIncludePath));
 				}
 
+				// Preprocessor definitions.
+				foreach (string Definition in CompileEnvironment.Definitions)
+				{
+					// TODO: Causes ISPC compiler to generate a spurious warning about the universal character set
+					if(!Definition.Contains("\\\\U"))
+					{
+						Arguments.Add(String.Format("-D\"{0}\"", Definition));
+					}
+				}
+
 				// Generate the included header dependency list
 				if (CompileEnvironment.bGenerateDependenciesFile)
 				{
@@ -689,7 +699,11 @@ namespace UnrealBuildTool
 				// Preprocessor definitions.
 				foreach (string Definition in CompileEnvironment.Definitions)
 				{
-					Arguments.Add(String.Format("-D\"{0}\"", Definition));
+					// TODO: Causes ISPC compiler to generate a spurious warning about the universal character set
+					if(!Definition.Contains("\\\\U"))
+					{
+						Arguments.Add(String.Format("-D\"{0}\"", Definition));
+					}
 				}
 
 				// Consume the included header dependency list
diff --git a/Engine/Source/Runtime/Core/Public/Math/Aossoa.isph b/Engine/Source/Runtime/Core/Public/Math/Aossoa.isph
--- a/Engine/Source/Runtime/Core/Public/Math/Aossoa.isph
+++ b/Engine/Source/Runtime/Core/Public/Math/Aossoa.isph
@@ -52,6 +52,13 @@ static const varying int vAOS32 = {0, 1, 2, 3, 8, 9, 10, 11};
 static const varying int vAOS33 = {4, 5, 6, 7, 12, 13, 14, 15};
 static const varying int vAOS34 = {2, 5, 8, 11, 14, -1, -1, -1};
 static const varying int vAOS35 = {0, 1, 2, 3, 4, 9, 12, 15};
+
+static const varying int vAOS301 = {0, 1, 6, 7, 2, 3, 8, 9};
+static const varying int vAOS311 = {4, 5, 10, 11, 6, 7, 12, 13};
+static const varying int vAOS321 = {0, 1, 2, 3, 8, 9, 10, 11};
+static const varying int vAOS331 = {4, 5, 6, 7, 12, 13, 14, 15};
+static const varying int vAOS341 = {4, 5, 10, 11, -1, -1, -1, -1};
+static const varying int vAOS351 = {0, 1, 2, 3, 8, 9, 14, 15};
 #elif TARGET_WIDTH == 16
 static const varying int vAOS30 = {0, 3, 6, 9, 12, 15, 18, 21, 1, 4, 7, 10, 13, 16, 19, 22};
 static const varying int vAOS31 = {8, 11, 14, 17, 20, 23, 26, 29, 9, 12, 15, 18, 21, 24, 27, 30};
@@ -63,6 +70,48 @@ static const varying int vAOS35 = {0, 1, 2, 3, 4, 5, 6, 7, 8, 9, 16, 19, 22, 25,
 #error "No implementation for this target"
 #endif
 
+#if TARGET_WIDTH == 8
+#define AOS_TO_SOA3_ISPC_64BIT(T)                                                                                      \
+    unmasked inline void aos_to_soa3_ispc(uniform T src[], varying T *uniform v0, varying T *uniform v1,               \
+                                        varying T *uniform v2)                                                         \
+    {                                                                                                                  \
+        const uniform float *uniform srcflt = (const uniform float *uniform)src;                                       \
+        const varying float src0 = srcflt[programIndex];                                                               \
+        const varying float src1 = srcflt[programCount + programIndex];                                                \
+        const varying float src2 = srcflt[2 * programCount + programIndex];                                            \
+        const varying float src3 = srcflt[3 * programCount + programIndex];                                            \
+        const varying float src4 = srcflt[4 * programCount + programIndex];                                            \
+        const varying float src5 = srcflt[5 * programCount + programIndex];                                            \
+                                                                                                                       \
+        const varying float t0 = shuffle(src0, src1, vAOS301);                                                         \
+        const varying float t1 = shuffle(src1, src2, vAOS311);                                                         \
+        const varying float s0 = shuffle(src3, src4, vAOS301);                                                         \
+        const varying float s1 = shuffle(src4, src5, vAOS311);                                                         \
+        const varying float t2 = shuffle(src0, src1, vAOS341);                                                         \
+        const varying float u2 = shuffle(t2, src2, vAOS351);                                                           \
+        const varying float s2 = shuffle(src3, src4, vAOS341);                                                         \
+                                                                                                                       \
+        const varying float u0 = shuffle(t0, t1, vAOS321);                                                             \
+        const varying float w0 = shuffle(s0, s1, vAOS321);                                                             \
+        uniform T<TARGET_WIDTH> x0;                                                                                    \
+        *((varying float *uniform)&x0) = u0;                                                                           \
+        *((varying float *uniform)&x0[TARGET_WIDTH/2]) = w0;                                                           \
+        *v0 = *((varying T *uniform)&x0);                                                                              \
+                                                                                                                       \
+        const varying float u1 = shuffle(t0, t1, vAOS331);                                                             \
+        const varying float w1 = shuffle(s0, s1, vAOS331);                                                             \
+        uniform T<TARGET_WIDTH> x1;                                                                                    \
+        *((varying float *uniform)&x1) = u1;                                                                           \
+        *((varying float *uniform)&x1[TARGET_WIDTH/2]) = w1;                                                           \
+        *v1 = *((varying T *uniform)&x1);                                                                              \
+                                                                                                                       \
+        const varying float w2 = shuffle(s2, src5, vAOS351);                                                           \
+        uniform T<TARGET_WIDTH> x2;                                                                                    \
+        *((varying float *uniform)&x2) = u2;                                                                           \
+        *((varying float *uniform)&x2[TARGET_WIDTH/2]) = w2;                                                           \
+        *v2 = *((varying T *uniform)&x2);                                                                              \
+    }
+#endif
 #define AOS_TO_SOA3_ISPC(T)                                                                                            \
     unmasked inline void aos_to_soa3_ispc(uniform T src[], varying T *uniform v0, varying T *uniform v1,               \
                                         varying T *uniform v2)                                                         \
@@ -83,8 +132,13 @@ static const varying int vAOS35 = {0, 1, 2, 3, 4, 5, 6, 7, 8, 9, 16, 19, 22, 25,
 AOS_TO_SOA3_ISPC(int16)
 AOS_TO_SOA3_ISPC(int32)
 AOS_TO_SOA3_ISPC(float)
-AOS_TO_SOA3_ISPC(int64)
+#if TARGET_WIDTH == 8
+AOS_TO_SOA3_ISPC_64BIT(double)
+AOS_TO_SOA3_ISPC_64BIT(int64)
+#else
 AOS_TO_SOA3_ISPC(double)
+AOS_TO_SOA3_ISPC(int64)
+#endif
 
 
 #if TARGET_WIDTH == 4
diff --git a/Engine/Source/Runtime/Core/Public/Math/Matrix.isph b/Engine/Source/Runtime/Core/Public/Math/Matrix.isph
--- a/Engine/Source/Runtime/Core/Public/Math/Matrix.isph
+++ b/Engine/Source/Runtime/Core/Public/Math/Matrix.isph
@@ -9,6 +9,7 @@
 #define EXPLICIT_MATRIX 0
 #endif
 
+#if UE_LARGE_WORLD_COORDINATES_DISABLED
 struct FMatrix
 {
 #if EXPLICIT_MATRIX == 0
@@ -18,6 +19,30 @@ struct FMatrix
 #endif
 };
 
+struct FMatrix44d
+{
+	double M[16];
+};
+
+typedef FMatrix FMatrix44f;
+#else
+struct FMatrix44f
+{
+#if EXPLICIT_MATRIX == 0
+	float M[16];
+#else
+	float<16> M;
+#endif
+};
+
+struct FMatrix
+{
+	double M[16];
+};
+
+typedef FMatrix FMatrix44d;
+#endif
+
 #ifndef EXPLICIT_MATRIX3x4
 #define EXPLICIT_MATRIX3x4 0
 #endif
@@ -84,10 +109,10 @@ inline FMatrix operator*(const FMatrix &A, const FMatrix &B)
 
 	for (uniform unsigned int m = 0; m < 4; m++) 
 	{
-		varying float Sum;
+		varying FReal Sum;
 		for (uniform unsigned int k = 0; k < 4; k++) 
 		{
-			Sum = 0.0f;
+			Sum = 0;
 			for (uniform unsigned int n = 0; n < 4; n++) 
 			{
 				Sum += A.M[m * 4 + n] * B.M[n * 4 + k];
@@ -106,10 +131,10 @@ inline FMatrix operator*(const FMatrix &A, const uniform FMatrix &B)
 
 	for (uniform unsigned int m = 0; m < 4; m++) 
 	{
-		varying float Sum;
+		varying FReal Sum;
 		for (uniform unsigned int k = 0; k < 4; k++) 
 		{
-			Sum = 0.0f;
+			Sum = 0;
 			for (uniform unsigned int n = 0; n < 4; n++) 
 			{
 				Sum += A.M[m * 4 + n] * B.M[n * 4 + k];
@@ -126,12 +151,12 @@ inline FMatrix operator*(const FMatrix &A, const uniform FMatrix &B)
 inline void MatrixRemoveScaling(uniform FMatrix &M)
 {
 	// For each row, find magnitude, and if its non-zero re-scale so its unit length.
-	const uniform float SquareSum0 = (M.M[0] * M.M[0]) + (M.M[1] * M.M[1]) + (M.M[2] * M.M[2]);
-	const uniform float SquareSum1 = (M.M[4] * M.M[4]) + (M.M[5] * M.M[5]) + (M.M[6] * M.M[6]);
-	const uniform float SquareSum2 = (M.M[8] * M.M[8]) + (M.M[9] * M.M[9]) + (M.M[10] * M.M[10]);
-	const uniform float Scale0 = select( SquareSum0 - SMALL_NUMBER >= 0.f, InvSqrt(SquareSum0), 1.0f );
-	const uniform float Scale1 = select( SquareSum1 - SMALL_NUMBER >= 0.f, InvSqrt(SquareSum1), 1.0f );
-	const uniform float Scale2 = select( SquareSum2 - SMALL_NUMBER >= 0.f, InvSqrt(SquareSum2), 1.0f );
+	const uniform FReal SquareSum0 = (M.M[0] * M.M[0]) + (M.M[1] * M.M[1]) + (M.M[2] * M.M[2]);
+	const uniform FReal SquareSum1 = (M.M[4] * M.M[4]) + (M.M[5] * M.M[5]) + (M.M[6] * M.M[6]);
+	const uniform FReal SquareSum2 = (M.M[8] * M.M[8]) + (M.M[9] * M.M[9]) + (M.M[10] * M.M[10]);
+	const uniform FReal Scale0 = select( SquareSum0 - SMALL_NUMBER >= ZERO, InvSqrt(SquareSum0), ONE );
+	const uniform FReal Scale1 = select( SquareSum1 - SMALL_NUMBER >= ZERO, InvSqrt(SquareSum1), ONE );
+	const uniform FReal Scale2 = select( SquareSum2 - SMALL_NUMBER >= ZERO, InvSqrt(SquareSum2), ONE );
 	M.M[0] *= Scale0;
 	M.M[1] *= Scale0;
 	M.M[2] *= Scale0;
@@ -155,15 +180,15 @@ inline void MatrixGetScaledAxes(const uniform FMatrix &M, uniform FVector &X, un
 	Z = SetVector(M.M[8], M.M[9], M.M[10]);
 }
 
-inline void MatrixScaleAxis(uniform FMatrix &M, const uniform int Axis, const uniform float ScalingFactor)
+inline void MatrixScaleAxis(uniform FMatrix &M, const uniform int Axis, const uniform FReal ScalingFactor)
 {
-	varying float S0;
+	varying FReal S0;
 	*((uniform FVector *uniform)&S0) = *((uniform FVector *uniform)&M.M[(4*Axis)]);
 
 #if EXPLICIT_VECTOR == 1
 	// Using explicit vector aligns to 16 bytes. Save the 4th component when pulling from Matrix type
-	const uniform float Saved = extract(S0, 3);
-	S0 = insert(S0, 3, 0.0f);
+	const uniform FReal Saved = extract(S0, 3);
+	S0 = insert(S0, 3, 0);
 #endif
 
 	S0 = S0 * ScalingFactor;
@@ -176,7 +201,7 @@ inline void MatrixScaleAxis(uniform FMatrix &M, const uniform int Axis, const un
 	*((uniform FVector *uniform)&M.M[(4*Axis)]) = *((uniform FVector *uniform)&S0);
 }
 
-inline uniform bool IsAnyMatrixScaledAxesNearlyZero(const uniform FMatrix &M, const uniform float Tolerance)
+inline uniform bool IsAnyMatrixScaledAxesNearlyZero(const uniform FMatrix &M, const uniform FReal Tolerance)
 {
 	uniform FVector X, Y, Z;
 	MatrixGetScaledAxes(M, X, Y, Z);
@@ -189,7 +214,7 @@ inline uniform bool IsAnyMatrixScaledAxesNearlyZero(const uniform FMatrix &M, co
 	return false;
 }
 
-inline uniform bool IsAllMatrixScaledAxesNearlyZero(const uniform FMatrix &M, const uniform float Tolerance)
+inline uniform bool IsAllMatrixScaledAxesNearlyZero(const uniform FMatrix &M, const uniform FReal Tolerance)
 {
 	uniform FVector X, Y, Z;
 	MatrixGetScaledAxes(M, X, Y, Z);
@@ -245,7 +270,7 @@ static inline uniform FVector4 Mat2MulAdj(const uniform FVector4& vec1, const un
 			VectorMultiply(VectorSwizzle(vec1, 1,0,3,2), VectorSwizzle(vec2, 2,1,2,1)));
 }
 
-inline uniform float MatrixDeterminant(const uniform FMatrix& M)
+inline uniform FReal MatrixDeterminant(const uniform FMatrix& M)
 {
 	const uniform FVector4 *uniform MVec = (const uniform FVector4 *uniform) &M;
 
@@ -267,7 +292,7 @@ inline uniform float MatrixDeterminant(const uniform FMatrix& M)
 	const uniform FVector4 A_B = Mat2AdjMul(A, B);
 
 	// |M| = |A|*|D| + |B|*|C|
-	uniform float detM = (detSub.V[0] * detSub.V[3]) + (detSub.V[1] * detSub.V[2]);
+	uniform FReal detM = (detSub.V[0] * detSub.V[3]) + (detSub.V[1] * detSub.V[2]);
 
 	// tr((A#B)(D#C))
 	uniform FVector4 tr = VectorMultiply(A_B, VectorSwizzle(D_C, 0,2,1,3));
@@ -334,9 +359,9 @@ inline uniform FMatrix VectorMatrixInverse(const uniform FMatrix& M)
 
 	uniform FMatrix R = MatrixIdentity;
 
-	if(detM.V[0] != 0.0f)
+	if(detM.V[0] != ZERO)
 	{
-		const uniform FVector4 adjSignMask = SetVector4(1.f, -1.f, -1.f, 1.f);
+		const uniform FVector4 adjSignMask = SetVector4(ONE, -ONE, -ONE, ONE);
 		// (1/|M|, -1/|M|, -1/|M|, 1/|M|)
 		const uniform FVector4 rDetM = VectorDivide(adjSignMask, detM);
 
@@ -375,18 +400,28 @@ inline uniform FMatrix MatrixInverse(const uniform FMatrix& M)
 	return Result;
 }
 
-inline uniform FVector4 VectorTransformVector(const uniform FVector4 &VecP, const uniform FMatrix &Matrix)
+inline uniform FVector4 VectorTransformVector(const uniform FVector4 &VecP, const uniform FMatrix &M)
 {
-	uniform FVector4 Result;
-    const uniform FVector4 *uniform M = (const uniform FVector4 *uniform) &Matrix;
+	uniform FVector4 VTempX, VTempY, VTempZ, VTempW;
+
+	// Splat x,y,z and w
+	VTempX = VectorReplicate(VecP, 0);
+	VTempY = VectorReplicate(VecP, 1);
+	VTempZ = VectorReplicate(VecP, 2);
+	VTempW = VectorReplicate(VecP, 3);
 
 	// Mul by the matrix
-	Result = VectorMultiply(M[0], VecP.V[0]);
-	Result = VectorMultiplyAdd(M[1], VecP.V[1], Result);
-	Result = VectorMultiplyAdd(M[2], VecP.V[2], Result);
-	Result = VectorMultiplyAdd(M[3], VecP.V[3], Result);
+	VTempX = VectorMultiply(VTempX, SetVector4(M.M[0], M.M[1], M.M[2], M.M[3]));
+	VTempY = VectorMultiply(VTempY, SetVector4(M.M[4], M.M[5], M.M[6], M.M[7]));
+	VTempZ = VectorMultiply(VTempZ, SetVector4(M.M[8], M.M[9], M.M[10], M.M[11]));
+	VTempW = VectorMultiply(VTempW, SetVector4(M.M[12], M.M[13], M.M[14], M.M[15]));
 
-	return Result;
+	// Add them all together
+	VTempX = VectorAdd(VTempX, VTempY);
+	VTempZ = VectorAdd(VTempZ, VTempW);
+	VTempX = VectorAdd(VTempX, VTempZ);
+
+	return VTempX;
 }
 
 inline FVector4 VectorTransformVector(const FVector4 &VecP, const FMatrix &M)
@@ -438,9 +473,32 @@ inline FVector MatrixTransformPosition(const FVector &P, const FMatrix &M)
 }
 
 // Calculate homogeneous transform. W component assumed to be 1.0
-inline FVector MatrixTransformPosition(const FVector &P, const uniform FMatrix &M)
+inline FVector3d MatrixTransformPosition(const FVector3d &P, const uniform FMatrix44d &M)
 {
-	FVector VTempX, VTempY, VTempZ;
+	FVector3d VTempX, VTempY, VTempZ;
+
+	// Splat x,y,z
+	VTempX = SetVector(P.V[0], P.V[0], P.V[0]);
+	VTempY = SetVector(P.V[1], P.V[1], P.V[1]);
+	VTempZ = SetVector(P.V[2], P.V[2], P.V[2]);
+
+	// Mul by the matrix
+	VTempX = VTempX * SetVector(M.M[0], M.M[1], M.M[2]);
+	VTempY = VTempY * SetVector(M.M[4], M.M[5], M.M[6]);
+	VTempZ = VTempZ * SetVector(M.M[8], M.M[9], M.M[10]);
+	const uniform FVector3d VTempW = SetVector(M.M[12], M.M[13], M.M[14]);
+
+	// Add them all together
+	VTempX = VTempX + VTempY;
+	VTempZ = VTempZ + VTempW;
+	VTempX = VTempX + VTempZ;
+
+	return VTempX;
+}
+
+inline FVector3f MatrixTransformPosition(const FVector3f &P, const uniform FMatrix44f &M)
+{
+	FVector3f VTempX, VTempY, VTempZ;
 
 	// Splat x,y,z
 	VTempX = SetVector(P.V[0], P.V[0], P.V[0]);
@@ -451,7 +509,7 @@ inline FVector MatrixTransformPosition(const FVector &P, const uniform FMatrix &
 	VTempX = VTempX * SetVector(M.M[0], M.M[1], M.M[2]);
 	VTempY = VTempY * SetVector(M.M[4], M.M[5], M.M[6]);
 	VTempZ = VTempZ * SetVector(M.M[8], M.M[9], M.M[10]);
-	const uniform FVector VTempW = SetVector(M.M[12], M.M[13], M.M[14]);
+	const uniform FVector3f VTempW = SetVector(M.M[12], M.M[13], M.M[14]);
 
 	// Add them all together
 	VTempX = VTempX + VTempY;
@@ -483,7 +541,7 @@ inline FVector MatrixTransformVector(const FVector &P, const FMatrix &M)
 inline uniform FVector MatrixInverseTransformVector(const uniform FMatrix &M, const uniform FVector &V)
 {
 	const uniform FMatrix InvSelf = VectorMatrixInverse(M); // FMatrix::InverseFast();
-	return SetVector(VectorTransformVector(SetVector4(V, 0.0f), InvSelf));
+	return SetVector(VectorTransformVector(SetVector4(V, ZERO), InvSelf));
 }
 
 inline uniform FMatrix MatrixReduceAdd(const varying FMatrix &M)
diff --git a/Engine/Source/Runtime/Core/Public/Math/Quat.isph b/Engine/Source/Runtime/Core/Public/Math/Quat.isph
--- a/Engine/Source/Runtime/Core/Public/Math/Quat.isph
+++ b/Engine/Source/Runtime/Core/Public/Math/Quat.isph
@@ -13,25 +13,25 @@ inline uniform FVector4 MatrixToQuat(const uniform FMatrix &M)
 	// for now, if you convert to matrix from 0 scale and convert back, you'll lose rotation. Don't do that.
 	if (IsAnyMatrixScaledAxesNearlyZero(M, KINDA_SMALL_NUMBER))
 	{
-		return Float0001;
+		return Vector0001;
 	}
 
 	//const MeReal *const t = (MeReal *) tm;
-	uniform float	s;
+	uniform FReal	s;
 
 	// Check diagonal (trace)
-	const uniform float tr = M.M[0] + M.M[5] + M.M[10];
+	const uniform FReal tr = M.M[0] + M.M[5] + M.M[10];
 
-	if (tr > 0.0f)
+	if (tr > ZERO)
 	{
 		uniform FVector4 qt;
-		const uniform float InvS = 1.0f / sqrt(tr + 1.f);
-		s = 0.5f * InvS;
+		const uniform FReal InvS = ONE / sqrt(tr + 1);
+		s = ONE_HALF * InvS;
 
 		qt.V[0] = (M.M[6] - M.M[9]) * s;
 		qt.V[1] = (M.M[8] - M.M[2]) * s;
 		qt.V[2] = (M.M[1] - M.M[4]) * s;
-		qt.V[3] = 0.5f * (1.f / InvS);
+		qt.V[3] = ONE_HALF * (ONE / InvS);
 
 		return qt;
 	}
@@ -54,14 +54,14 @@ inline uniform FVector4 MatrixToQuat(const uniform FMatrix &M)
 		const uniform int j = nxt[i];
 		const uniform int k = nxt[j];
 
-		s = M.M[4*i+i] - M.M[4*j+j] - M.M[4*k+k] + 1.0f;
+		s = M.M[4*i+i] - M.M[4*j+j] - M.M[4*k+k] + ONE;
 
-		const uniform float InvS = 1.0f / sqrt(s);
-		s = 0.5f * InvS;
+		const uniform FReal InvS = ONE / sqrt(s);
+		s = ONE_HALF * InvS;
 
 		uniform FVector4 qt;
 
-		qt.V[i] = 0.5f * (1.f / InvS);
+		qt.V[i] = ONE_HALF * (ONE / InvS);
 		qt.V[j] = (M.M[4*i+j] + M.M[4*j+i]) * s;
 		qt.V[k] = (M.M[4*i+k] + M.M[4*k+i]) * s;
 		qt.V[3] = (M.M[4*j+k] - M.M[4*k+j]) * s;
@@ -101,14 +101,19 @@ inline uniform FVector4 VectorQuaternionMultiply2( const uniform FVector4& Quat1
 }
 
 // Normalize quaternion ( result = (Q.Q >= 1e-8) ? (Q / |Q|) : (0,0,0,1) )
-inline FVector4 VectorNormalizeQuaternion(const FVector4& UnnormalizedQuat)
+inline FVector4d VectorNormalizeQuaternion(const FVector4d& UnnormalizedQuat)
+{
+	return VectorNormalizeSafe(UnnormalizedQuat, Double0001);
+}
+
+inline FVector4f VectorNormalizeQuaternion(const FVector4f& UnnormalizedQuat)
 {
 	return VectorNormalizeSafe(UnnormalizedQuat, Float0001);
 }
 
 inline uniform FVector4 VectorNormalizeQuaternion(const uniform FVector4& UnnormalizedQuat)
 {
-	return VectorNormalizeSafe2(UnnormalizedQuat, Float0001);
+	return VectorNormalizeSafe2(UnnormalizedQuat, Vector0001);
 }
 
 inline FVector4 QuatInverse(const FVector4 &Quat)
@@ -116,25 +121,38 @@ inline FVector4 QuatInverse(const FVector4 &Quat)
 	return Quat * QINV_SIGN_MASK;
 }
 
-inline uniform FVector4 QuatInverse(const uniform FVector4 &Quat)
+inline uniform FVector4d QuatInverse(const uniform FVector4d &Quat)
 {
-	return Quat * QINV_SIGN_MASK;
+	return Quat * DOUBLE_QINV_SIGN_MASK;
+}
+
+inline uniform FVector4f QuatInverse(const uniform FVector4f &Quat)
+{
+	return Quat * FLOAT_QINV_SIGN_MASK;
 }
 
-inline FVector4 QuatFastLerp(const FVector4& A, const FVector4& B, const float Alpha)
+inline FVector4d QuatFastLerp(const FVector4d& A, const FVector4d& B, const double Alpha)
+{
+	// To ensure the 'shortest route', we make sure the dot product between the both rotations is positive.
+	const double DotResult = VectorDot(A, B);
+	const double Bias = select(DotResult >= 0.d, 1.d, -1.d);
+	return (B * Alpha) + (A * (Bias * (1.d - Alpha)));
+}
+
+inline FVector4f QuatFastLerp(const FVector4f& A, const FVector4f& B, const float Alpha)
 {
 	// To ensure the 'shortest route', we make sure the dot product between the both rotations is positive.
 	const float DotResult = VectorDot(A, B);
-	const float Bias = select(DotResult >= 0.f, 1.0f, -1.0f);
+	const float Bias = select(DotResult >= 0.f, 1.f, -1.f);
 	return (B * Alpha) + (A * (Bias * (1.f - Alpha)));
 }
 
-inline uniform FVector4 QuatFastLerp(const uniform FVector4& A, const uniform FVector4& B, const uniform float Alpha)
+inline uniform FVector4 QuatFastLerp(const uniform FVector4& A, const uniform FVector4& B, const uniform FReal Alpha)
 {
 	// To ensure the 'shortest route', we make sure the dot product between the both rotations is positive.
-	const uniform float DotResult = VectorDot(A, B);
-	const uniform float Bias = select(DotResult >= 0.f, 1.0f, -1.0f);
-	return (B * Alpha) + (A * (Bias * (1.f - Alpha)));
+	const uniform FReal DotResult = VectorDot(A, B);
+	const uniform FReal Bias = select(DotResult >= ZERO, ONE, -ONE);
+	return (B * Alpha) + (A * (Bias * (ONE - Alpha)));
 }
 
 // A and B are quaternions.  The result is A + (|A.B| >= 0 ? 1 : -1) * B
@@ -145,8 +163,8 @@ inline uniform FVector4 VectorAccumulateQuaternionShortestPath(const uniform FVe
 	//     const float Bias = (|A.B| >= 0 ? 1 : -1)
 	//     return A + B * Bias;
 	const uniform FVector4 RotationDot = VectorDot4(A, B);
-	const uniform FVector4 QuatRotationDirMask = VectorCompareGE(RotationDot, FloatZero);
-	const uniform FVector4 NegativeB = FloatZero - B;
+	const uniform FVector4 QuatRotationDirMask = VectorCompareGE(RotationDot, VectorZero);
+	const uniform FVector4 NegativeB = VectorZero - B;
 	const uniform FVector4 BiasTimesB = VectorSelect(QuatRotationDirMask, B, NegativeB);
 	return A + BiasTimesB;
 }
@@ -240,7 +258,7 @@ inline uniform FVector VectorQuaternionRotateVector(const uniform FVector4& Quat
 	// V' = V + w*(T) + (Q x T)
 
 	const uniform FVector Q = SetVector(Quat.V[0], Quat.V[1], Quat.V[2]);
-	const uniform FVector T = VectorCross(Q, V) * 2.f;
+	const uniform FVector T = VectorCross(Q, V) * 2;
 	const uniform FVector Result = V + (T * Quat.V[3]) + VectorCross(Q, T);
 	return Result;
 }
@@ -255,12 +273,12 @@ inline FVector VectorQuaternionRotateVector(const FVector4& Quat, const FVector&
 	// V' = V + w*(T) + (Q x T)
 
 	const FVector Q = SetVector(Quat.V[0], Quat.V[1], Quat.V[2]);
-	const FVector T = VectorCross(Q, V) * 2.f;
+	const FVector T = VectorCross(Q, V) * 2;
 	const FVector Result = V + (T * Quat.V[3]) + VectorCross(Q, T);
 	return Result;
 }
 
-inline FVector VectorQuaternionRotateVector(const uniform FVector4& Quat, const FVector& V)
+inline FVector3d VectorQuaternionRotateVector(const uniform FVector4d& Quat, const FVector3d& V)
 {
 	// http://people.csail.mit.edu/bkph/articles/Quaternions.pdf
 	// V' = V + 2w(Q x V) + (2Q x (Q x V))
@@ -269,9 +287,24 @@ inline FVector VectorQuaternionRotateVector(const uniform FVector4& Quat, const
 	// T = 2(Q x V);
 	// V' = V + w*(T) + (Q x T)
 
-	const uniform FVector Q = SetVector(Quat.V[0], Quat.V[1], Quat.V[2]);
-	const FVector T = VectorCross(Q, V) * 2.f;
-	const FVector Result = V + (T * Quat.V[3]) + VectorCross(Q, T);
+	const uniform FVector3d Q = SetVector(Quat.V[0], Quat.V[1], Quat.V[2]);
+	const FVector3d T = VectorCross(Q, V) * 2;
+	const FVector3d Result = V + (T * Quat.V[3]) + VectorCross(Q, T);
+	return Result;
+}
+
+inline FVector3f VectorQuaternionRotateVector(const uniform FVector4f& Quat, const FVector3f& V)
+{
+	// http://people.csail.mit.edu/bkph/articles/Quaternions.pdf
+	// V' = V + 2w(Q x V) + (2Q x (Q x V))
+	// refactor:
+	// V' = V + w(2(Q x V)) + (Q x (2(Q x V)))
+	// T = 2(Q x V);
+	// V' = V + w*(T) + (Q x T)
+
+	const uniform FVector3f Q = SetVector(Quat.V[0], Quat.V[1], Quat.V[2]);
+	const FVector3f T = VectorCross(Q, V) * 2;
+	const FVector3f Result = V + (T * Quat.V[3]) + VectorCross(Q, T);
 	return Result;
 }
 
@@ -286,17 +319,17 @@ inline uniform FVector8 VectorQuaternionRotateVector(const uniform FVector8& Qua
 	return Rotated;
 }
 
-inline uniform FVector4 VectorLerpQuat(const uniform FVector4& A, const uniform FVector4& B, const uniform float Alpha)
+inline uniform FVector4 VectorLerpQuat(const uniform FVector4& A, const uniform FVector4& B, const uniform FReal Alpha)
 {
 	// Blend rotation
 	//     To ensure the 'shortest route', we make sure the dot product between the both rotations is positive.
 	//     const float Bias = (|A.B| >= 0 ? 1 : -1)
 	//     Rotation = (B * Alpha) + (A * (Bias * (1.f - Alpha)));
-	const uniform float OneMinusAlpha = 1.0f - Alpha;
+	const uniform FReal OneMinusAlpha = ONE - Alpha;
 
 	const uniform FVector4 RotationDot = VectorDot4(A, B);
-	const uniform FVector4 QuatRotationDirMask = VectorCompareGE(RotationDot, FloatZero);
-	const uniform FVector4 NegativeA = VectorSubtract(FloatZero, A);
+	const uniform FVector4 QuatRotationDirMask = VectorCompareGE(RotationDot, VectorZero);
+	const uniform FVector4 NegativeA = VectorSubtract(VectorZero, A);
 	const uniform FVector4 BiasTimesA = VectorSelect(QuatRotationDirMask, A, NegativeA);
 	const uniform FVector4 BTimesWeight = VectorMultiply(B, Alpha);
 	const uniform FVector4 UnnormalizedResult = VectorMultiplyAdd(BiasTimesA, OneMinusAlpha, BTimesWeight);
@@ -306,8 +339,8 @@ inline uniform FVector4 VectorLerpQuat(const uniform FVector4& A, const uniform
 
 inline uniform FVector4 VectorQuaternionEnforceShortestArcWith(const uniform FVector4& A, const uniform FVector4& B)
 {
-	const uniform float DotResult = VectorDot(B, A);
-	const uniform float Bias = select(DotResult >= 0.f, 1.0f, -1.0f);
+	const uniform FReal DotResult = VectorDot(B, A);
+	const uniform FReal Bias = select(DotResult >= ZERO, ONE, -ONE);
 
 	return A * Bias;
 }
diff --git a/Engine/Source/Runtime/Core/Public/Math/Transform.isph b/Engine/Source/Runtime/Core/Public/Math/Transform.isph
--- a/Engine/Source/Runtime/Core/Public/Math/Transform.isph
+++ b/Engine/Source/Runtime/Core/Public/Math/Transform.isph
@@ -21,7 +21,7 @@ struct FTransform8
 	FVector8 Scale3D;
 };
 
-inline uniform FTransform operator*(const uniform FTransform &A, const uniform float F)
+inline uniform FTransform operator*(const uniform FTransform &A, const uniform FReal F)
 {
 	uniform FTransform Result;
 
@@ -66,7 +66,7 @@ inline uniform FMatrix ToMatrixWithScale(const uniform FTransform &T)
 	const uniform FVector4 yy2_xx2_xx2 = VectorSwizzle(RotationXX2YY2ZZ2, 1, 0, 0, 0);
 	const uniform FVector4 zz2_zz2_yy2 = VectorSwizzle(RotationXX2YY2ZZ2, 2, 2, 1, 0);
 	const uniform FVector4 DiagonalSum = VectorAdd(yy2_xx2_xx2, zz2_zz2_yy2);
-	const uniform FVector4 Diagonals = VectorSubtract(FloatOne, DiagonalSum);
+	const uniform FVector4 Diagonals = VectorSubtract(VectorOne, DiagonalSum);
 	const uniform FVector4 OutDiagonals = VectorMultiply(Diagonals, T.Scale3D);
 
 	// Grouping the non-diagonal elements in the rotation block by operations:
@@ -93,7 +93,7 @@ inline uniform FMatrix ToMatrixWithScale(const uniform FTransform &T)
 	const uniform FVector4 Subtracts = VectorSubtract(RotBase, RotOffset);
 	const uniform FVector4 OutSubtracts = VectorMultiply(Subtracts , Scale3DYZXW);
 
-	const uniform FVector4 DiagonalsXYZ_W0 = SetVector4(OutDiagonals.V[0], OutDiagonals.V[1], OutDiagonals.V[2], 0.0f);
+	const uniform FVector4 DiagonalsXYZ_W0 = SetVector4(OutDiagonals.V[0], OutDiagonals.V[1], OutDiagonals.V[2], ZERO);
 
 	// OutMatrix.M[0][0] = (1.0f - (yy2 + zz2)) * Scale.X;    // Diagonal.X
 	// OutMatrix.M[0][1] = (xy2 + wz2) * Scale.X;             // Adds.X
@@ -122,15 +122,15 @@ inline uniform FMatrix ToMatrixWithScale(const uniform FTransform &T)
 	// OutMatrix.M[3][1] = Translation.Y;
 	// OutMatrix.M[3][2] = Translation.Z;
 	// OutMatrix.M[3][3] = 1.0f;
-	const uniform FVector4 Row3 = SetVector4(T.Translation.V[0], T.Translation.V[1], T.Translation.V[2], 1.0f);
+	const uniform FVector4 Row3 = SetVector4(T.Translation.V[0], T.Translation.V[1], T.Translation.V[2], ONE);
 
 	return SetMatrix(Row0, Row1, Row2, Row3);
 }
 
 static inline uniform bool AnyHasNegativeScale(const uniform FVector4 &InScale3D, const uniform FVector4 &InOtherScale3D)
 {
-	varying float S0 = 0.0f;
-	varying float S1 = 0.0f;
+	varying FReal S0 = 0;
+	varying FReal S1 = 0;
 	*((uniform FVector4 *uniform)&S0) = *((uniform FVector4 *uniform)&InScale3D);
 	*((uniform FVector4 *uniform)&S1) = *((uniform FVector4 *uniform)&InOtherScale3D);
 
@@ -160,7 +160,7 @@ static inline void ConstructTransformFromMatrixWithDesiredScale(const uniform FM
 	// instead of using Scale in matrix multiplication
 	// it's a question of between RemoveScaling vs using FTransform to move translation
 	const uniform FVector Translation = MatrixGetOrigin(M);
-	OutTransform.Translation = SetVector4(Translation.V[0], Translation.V[1], Translation.V[2], 0.0f);
+	OutTransform.Translation = SetVector4(Translation.V[0], Translation.V[1], Translation.V[2], 0);
 
 	OutTransform.Scale3D = DesiredScale;
 }
@@ -210,15 +210,15 @@ inline void TransformMultiply(uniform FTransform &OutTransform, const uniform FT
 	}
 }
 
-inline void ScaleTranslation(uniform FTransform &Transform, const uniform float Scale)
+inline void ScaleTranslation(uniform FTransform &Transform, const uniform FReal Scale)
 {
-	const uniform FVector4 VectorScale = SetVector4(Scale, Scale, Scale, 0.0f);
+	const uniform FVector4 VectorScale = SetVector4(Scale, Scale, Scale, 0);
 	Transform.Translation = VectorMultiply(Transform.Translation, VectorScale);
 }
 
 inline uniform FVector4 TransformVector(const uniform FTransform &T, const uniform FVector4& V)
 {
-	const uniform FVector4 InputVectorW0 = SetVector4(V.V[0], V.V[1], V.V[2], 0.0f);
+	const uniform FVector4 InputVectorW0 = SetVector4(V.V[0], V.V[1], V.V[2], 0);
 
 	//RotatedVec = Q.Rotate(Scale*V.X, Scale*V.Y, Scale*V.Z, 0.f)
 	const uniform FVector4 ScaledVec = VectorMultiply(T.Scale3D, InputVectorW0);
@@ -227,7 +227,7 @@ inline uniform FVector4 TransformVector(const uniform FTransform &T, const unifo
 
 inline uniform FVector TransformVector(const uniform FTransform &T, const uniform FVector& V)
 {
-	const uniform FVector4 InputVectorW0 = SetVector4(V, 0.0f);
+	const uniform FVector4 InputVectorW0 = SetVector4(V, 0);
 
 	//RotatedVec = Q.Rotate(Scale*V.X, Scale*V.Y, Scale*V.Z, 0.f)
 	const uniform FVector4 ScaledVec = VectorMultiply(T.Scale3D, InputVectorW0);
@@ -253,7 +253,7 @@ inline FVector TransformVector(const uniform FTransform &T, const FVector& V)
 
 inline uniform FVector TransformVectorNoScale(const uniform FTransform &T, const uniform FVector& V)
 {
-	const uniform FVector4 InputVectorW0 = SetVector4(V, 0.0f);
+	const uniform FVector4 InputVectorW0 = SetVector4(V, 0);
 
 	//RotatedVec = Q.Rotate(Scale*V.X, Scale*V.Y, Scale*V.Z, 0.f)
 	const uniform FVector4 RotatedVec = VectorQuaternionRotateVector(T.Rotation, InputVectorW0);
@@ -262,7 +262,7 @@ inline uniform FVector TransformVectorNoScale(const uniform FTransform &T, const
 
 inline FVector TransformPosition(const FTransform &T, const FVector& V)
 {
-	FVector4 InputVectorW0 = SetVector4(V.V[0], V.V[1], V.V[2], 0.0f);
+	FVector4 InputVectorW0 = SetVector4(V.V[0], V.V[1], V.V[2], 0);
 
 	//Transform using QST is following
 	//QST(P) = Q.Rotate(S*P) + T where Q = quaternion, S = scale, T = translation
@@ -278,7 +278,7 @@ inline FVector TransformPosition(const FTransform &T, const FVector& V)
 
 inline FVector TransformPosition(const uniform FTransform &T, const FVector& V)
 {
-	FVector4 InputVectorW0 = SetVector4(V.V[0], V.V[1], V.V[2], 0.0f);
+	FVector4 InputVectorW0 = SetVector4(V.V[0], V.V[1], V.V[2], 0);
 
 	//Transform using QST is following
 	//QST(P) = Q.Rotate(S*P) + T where Q = quaternion, S = scale, T = translation
@@ -294,7 +294,7 @@ inline FVector TransformPosition(const uniform FTransform &T, const FVector& V)
 
 inline uniform FVector TransformPosition(const uniform FTransform &T, const uniform FVector& V)
 {
-	uniform FVector4 InputVectorW0 = SetVector4(V.V[0], V.V[1], V.V[2], 0.0f);
+	uniform FVector4 InputVectorW0 = SetVector4(V.V[0], V.V[1], V.V[2], 0);
 
 	//Transform using QST is following
 	//QST(P) = Q.Rotate(S*P) + T where Q = quaternion, S = scale, T = translation
@@ -336,6 +336,20 @@ inline static uniform FVector8 TransformPosition(const uniform FTransform8 &T, c
 	return TranslatedVec;
 }
 
+inline static FVector TransformPositionNoScale(const uniform FTransform &T, const FVector& V)
+{
+	FVector4 InputVectorW0 = SetVector4(V.V[0], V.V[1], V.V[2], 0);
+
+	//Transform using QST is following
+	//QST(P) = Q.Rotate(S*P) + T where Q = quaternion, S = 1.0f, T = translation
+
+	//RotatedVec = Q.Rotate(V.X, V.Y, V.Z, 0.f)
+	const FVector4 RotatedVec = VectorQuaternionRotateVector(T.Rotation, InputVectorW0);
+	const FVector4 TranslatedVec = RotatedVec + T.Translation;
+
+	return SetVector(TranslatedVec.V[0], TranslatedVec.V[1], TranslatedVec.V[2]);
+}
+
 /** 
  * mathematically if you have 0 scale, it should be infinite, 
  * however, in practice if you have 0 scale, and relative transform doesn't make much sense 
@@ -343,17 +357,27 @@ inline static uniform FVector8 TransformPosition(const uniform FTransform8 &T, c
  * also returning BIG_NUMBER causes sequential NaN issues by multiplying 
  * so we hardcode as 0
  */
-inline uniform FVector TransformGetSafeScaleReciprocal(const uniform FVector& InScale, const uniform float Tolerance)
+inline uniform FVector3d TransformGetSafeScaleReciprocal(const uniform FVector3d& InScale, const uniform double Tolerance)
+{
+	varying double Scale = 0;
+	*((uniform FVector3d *uniform)&Scale) = *((uniform FVector3d *uniform)&InScale);
+
+	const varying double SafeReciprocalScale = abs(Scale) <= Tolerance ? 0 : rcp(Scale);
+
+	return *((uniform FVector3d *uniform)&SafeReciprocalScale);
+}
+
+inline uniform FVector3f TransformGetSafeScaleReciprocal(const uniform FVector3f& InScale, const uniform float Tolerance)
 {
-	varying float Scale = 0.0f;
-	*((uniform FVector *uniform)&Scale) = *((uniform FVector *uniform)&InScale);
+	varying float Scale = 0;
+	*((uniform FVector3f *uniform)&Scale) = *((uniform FVector3f *uniform)&InScale);
 
-	const varying float SafeReciprocalScale = abs(Scale) <= Tolerance ? 0.0f : rcp(Scale);
+	const varying float SafeReciprocalScale = abs(Scale) <= Tolerance ? 0 : rcp(Scale);
 
-	return *((uniform FVector *uniform)&SafeReciprocalScale);
+	return *((uniform FVector3f *uniform)&SafeReciprocalScale);
 }
 
-inline uniform FVector4 TransformGetSafeScaleReciprocal(const uniform FVector4& InScale, const uniform float Tolerance)
+inline uniform FVector4 TransformGetSafeScaleReciprocal(const uniform FVector4& InScale, const uniform FReal Tolerance)
 {
 	// SafeReciprocalScale.X = (InScale.X == 0) ? 0.f : 1/InScale.X; // same for YZW
 
@@ -366,7 +390,7 @@ inline uniform FVector4 TransformGetSafeScaleReciprocal(const uniform FVector4&
 	//const VectorRegister ScaleZeroMask = VectorCompareEQ(InScale, VectorZero());
 
 	// VectorRegister( for each bit i: Mask[i] ? Vec1[i] : Vec2[i] )
-	const uniform FVector4 SafeReciprocalScale = VectorSelect(ScaleZeroMask, FloatZero, ReciprocalScale);
+	const uniform FVector4 SafeReciprocalScale = VectorSelect(ScaleZeroMask, VectorZero, ReciprocalScale);
 
 	return SafeReciprocalScale;
 }
