LWC: Enable ISPC for LWC double types. Submitted on behalf of jeff.rous[REDACTED_HANDLE]

diff --git a/Engine/Source/Runtime/Core/Public/Math/Vector.isph b/Engine/Source/Runtime/Core/Public/Math/Vector.isph
--- a/Engine/Source/Runtime/Core/Public/Math/Vector.isph
+++ b/Engine/Source/Runtime/Core/Public/Math/Vector.isph
@@ -3,14 +3,20 @@
 #ifndef __VECTOR_ISPH__
 #define __VECTOR_ISPH__
 
+#ifndef UE_LARGE_WORLD_COORDINATES_DISABLED
+#include "Misc/LargeWorldCoordinates.h"
+#endif
+
 #include "Math/Aossoa.isph"
 #include "Math/Soaaos.isph"
+#include "Math/Scalar.isph"
 
 // Opt-in for explicit vector structs. Can assist with alignment and codegen in some cases.
 #ifndef EXPLICIT_VECTOR2
 #define EXPLICIT_VECTOR2 0
 #endif
 
+#if UE_LARGE_WORLD_COORDINATES_DISABLED
 struct FVector2
 {
 #if EXPLICIT_VECTOR2 == 0
@@ -20,10 +26,31 @@ struct FVector2
 #endif
 };
 
+struct FVector2d
+{
+	double V[2];
+};
+#else
+struct FVector2f
+{
+#if EXPLICIT_VECTOR2 == 0
+	float V[2];
+#else
+	float<2> V;
+#endif
+};
+
+struct FVector2
+{
+	double V[2];
+};
+#endif
+
 #ifndef EXPLICIT_VECTOR
 #define EXPLICIT_VECTOR 0
 #endif
 
+#if UE_LARGE_WORLD_COORDINATES_DISABLED
 struct FVector
 {
 #if EXPLICIT_VECTOR == 0
@@ -33,10 +60,31 @@ struct FVector
 #endif
 };
 
+struct FVector3d
+{
+	double V[3];
+};
+#else
+struct FVector3f
+{
+#if EXPLICIT_VECTOR == 0
+	float V[3];
+#else
+	float<3> V;
+#endif
+};
+
+struct FVector
+{
+	double V[3];
+};
+#endif
+
 #ifndef EXPLICIT_VECTOR4
 #define EXPLICIT_VECTOR4 0
 #endif
 
+#if UE_LARGE_WORLD_COORDINATES_DISABLED
 struct FVector4
 {
 #if EXPLICIT_VECTOR4 == 0
@@ -46,6 +94,26 @@ struct FVector4
 #endif
 };
 
+struct FVector4d
+{
+	double V[4];
+};
+#else
+struct FVector4f
+{
+#if EXPLICIT_VECTOR4 == 0
+	float V[4];
+#else
+	float<4> V;
+#endif
+};
+
+struct FVector4
+{
+	double V[4];
+};
+#endif
+
 struct FVector8
 {
 #if EXPLICIT_VECTOR8 == 0
@@ -94,61 +162,187 @@ struct FIntVector
 #endif
 };
 
-#define SMALL_NUMBER (1.e-8f)
-#define KINDA_SMALL_NUMBER	(1.e-4f)
-#define INDEX_NONE -1
-#define SIGN_BIT (0x80000000)
-#define FLT_MAX 3.402823466e+38F
-#define FLT_MIN 1.175494351e-38F
-
-static const uniform struct FVector4 QMULTI_SIGN_MASK0 = { {1.f, -1.f, 1.f, -1.f} };
-static const uniform struct FVector4 QMULTI_SIGN_MASK1 = { {1.f, 1.f, -1.f, -1.f} };
-static const uniform struct FVector4 QMULTI_SIGN_MASK2 = { {-1.f, 1.f, 1.f, -1.f} };
-static const uniform struct FVector4 QINV_SIGN_MASK = { {-1.f, -1.f, -1.f, 1.f} };
-static const uniform struct FVector4 FloatZero = { {0.f, 0.f, 0.f, 0.f} };
-static const uniform struct FVector4 FloatOne = { {1.f, 1.f, 1.f, 1.f} };
-static const uniform struct FVector4 FloatMinusOne = { {-1.f, -1.f, -1.f, -1.f} };
-static const uniform struct FVector4 Float0001 = { {0.0f, 0.0f, 0.0f, 1.0f} };
-static const uniform struct FVector4 Float1110 = { {1.0f, 1.0f, 1.0f, 0.0f} };
-static const uniform struct FVector4 FloatOneHalf = { {0.5f, 0.5f, 0.5f, 0.5f} };
-static const uniform struct FVector4 SmallLengthThreshold = { {1.e-8f, 1.e-8f, 1.e-8f, 1.e-8f} };
-
-static const uniform struct FVector4 Pi = { {PI, PI, PI, PI} };
-static const uniform struct FVector4 TwoPi = { {2.0f*PI, 2.0f*PI, 2.0f*PI, 2.0f*PI} };
-static const uniform struct FVector4 PiByTwo = { {0.5f*PI, 0.5f*PI, 0.5f*PI, 0.5f*PI} };
-static const uniform struct FVector4 PiByFour = { {0.25f*PI, 0.25f*PI, 0.25f*PI, 0.25f*PI} };
-static const uniform struct FVector4 OneOverPi = { {1.0f / PI, 1.0f / PI, 1.0f / PI, 1.0f / PI} };
-static const uniform struct FVector4 OneOverTwoPi = { {1.0f / (2.0f*PI), 1.0f / (2.0f*PI), 1.0f / (2.0f*PI), 1.0f / (2.0f*PI)} };
-
-static const uniform struct FVector4 Float255 = { {255.0f, 255.0f, 255.0f, 255.0f} };
-static const uniform struct FVector4 Float127 = { {127.0f, 127.0f, 127.0f, 127.0f} };
-static const uniform struct FVector4 FloatNeg127 = { {-127.0f, -127.0f, -127.0f, -127.0f} };
-static const uniform struct FVector4 Float360 = { {360.f, 360.f, 360.f, 360.f} };
-static const uniform struct FVector4 Float180 = { {180.f, 180.f, 180.f, 180.f} };
-static const uniform struct FVector4 FloatNonFractional = { {8388608.f, 8388608.f, 8388608.f, 8388608.f} };
-
-static const uniform struct FVector4 DEG_TO_RAD ={ {PI/(180.f), PI/(180.f), PI/(180.f), PI/(180.f) } };
-static const uniform struct FVector4 DEG_TO_RAD_HALF = { {(PI/180.f)*0.5f, (PI/180.f)*0.5f, (PI/180.f)*0.5f, (PI/180.f)*0.5f} };
-static const uniform struct FVector4 RAD_TO_DEG = { {(180.f)/PI, (180.f)/PI, (180.f)/PI, (180.f)/PI} };
-
-static const uniform struct FVector4 UpVector4 = { {0.0f, 0.0f, 1.0f, 0.0f} };
-static const uniform struct FVector4 ForwardVector4 = { {1.0f, 0.0f, 0.0f, 0.0f} };
-static const uniform struct FVector4 RightVector4 = { {0.0f, 1.0f, 0.0f, 0.0f} };
-
-static const uniform struct FVector ZeroVector = { {0.0f, 0.0f, 0.0f} };
-static const uniform struct FVector OneVector = { {1.0f, 1.0f, 1.0f} };
-static const uniform struct FVector HalfVector = { {0.5f, 0.5f, 0.5f} };
-static const uniform struct FVector UpVector = { {0.0f, 0.0f, 1.0f} };
-static const uniform struct FVector DownVector = { {0.0f, 0.0f, -1.0f} };
-static const uniform struct FVector ForwardVector = { {1.0f, 0.0f, 0.0f} };
-static const uniform struct FVector BackwardVector = { {-1.0f, 0.0f, 0.0f} };
-static const uniform struct FVector RightVector = { {0.0f, 1.0f, 0.0f} };
-static const uniform struct FVector LeftVector = { {0.0f, -1.0f, 0.0f} };
+// Define types
+#if UE_LARGE_WORLD_COORDINATES_DISABLED
+typedef FVector2 FVector2f;
+typedef FVector FVector3f;
+typedef FVector4 FVector4f;
+#else
+typedef FVector2 FVector2d;
+typedef FVector FVector3d;
+typedef FVector4 FVector4d;
+#endif
+
+static const uniform struct FVector4f FLOAT_QMULTI_SIGN_MASK0 = { {1.f, -1.f, 1.f, -1.f} };
+static const uniform struct FVector4f FLOAT_QMULTI_SIGN_MASK1 = { {1.f, 1.f, -1.f, -1.f} };
+static const uniform struct FVector4f FLOAT_QMULTI_SIGN_MASK2 = { {-1.f, 1.f, 1.f, -1.f} };
+static const uniform struct FVector4f FLOAT_QINV_SIGN_MASK = { {-1.f, -1.f, -1.f, 1.f} };
+static const uniform struct FVector4f FloatZero = { {0.f, 0.f, 0.f, 0.f} };
+static const uniform struct FVector4f FloatOne = { {1.f, 1.f, 1.f, 1.f} };
+static const uniform struct FVector4f FloatMinusOne = { {-1.f, -1.f, -1.f, -1.f} };
+static const uniform struct FVector4f Float0001 = { {0.0f, 0.0f, 0.0f, 1.0f} };
+static const uniform struct FVector4f Float1110 = { {1.0f, 1.0f, 1.0f, 0.0f} };
+static const uniform struct FVector4f FloatOneHalf = { {0.5f, 0.5f, 0.5f, 0.5f} };
+static const uniform struct FVector4f FloatSmallLengthThreshold = { {FLOAT_SMALL_NUMBER, FLOAT_SMALL_NUMBER, FLOAT_SMALL_NUMBER, FLOAT_SMALL_NUMBER} };
+
+static const uniform struct FVector4d DOUBLE_QMULTI_SIGN_MASK0 = { {1.d, -1.d, 1.d, -1.d} };
+static const uniform struct FVector4d DOUBLE_QMULTI_SIGN_MASK1 = { {1.d, 1.d, -1.d, -1.d} };
+static const uniform struct FVector4d DOUBLE_QMULTI_SIGN_MASK2 = { {-1.d, 1.d, 1.d, -1.d} };
+static const uniform struct FVector4d DOUBLE_QINV_SIGN_MASK = { {-1.d, -1.d, -1.d, 1.d} };
+static const uniform struct FVector4d DoubleZero = { {0.d, 0.d, 0.d, 0.d} };
+static const uniform struct FVector4d DoubleOne = { {1.d, 1.d, 1.d, 1.d} };
+static const uniform struct FVector4d DoubleMinusOne = { {-1.d, -1.d, -1.d, -1.d} };
+static const uniform struct FVector4d Double0001 = { {0.0d, 0.0d, 0.0d, 1.0d} };
+static const uniform struct FVector4d Double1110 = { {1.0d, 1.0d, 1.0d, 0.0d} };
+static const uniform struct FVector4d DoubleOneHalf = { {0.5d, 0.5d, 0.5d, 0.5d} };
+static const uniform struct FVector4d DoubleSmallLengthThreshold = { {DOUBLE_SMALL_NUMBER, DOUBLE_SMALL_NUMBER, DOUBLE_SMALL_NUMBER, DOUBLE_SMALL_NUMBER} };
+
+static const uniform struct FVector4f FloatPi = { {FLOAT_PI, FLOAT_PI, FLOAT_PI, FLOAT_PI} };
+static const uniform struct FVector4f FloatTwoPi = { {2.0f*FLOAT_PI, 2.0f*FLOAT_PI, 2.0f*FLOAT_PI, 2.0f*FLOAT_PI} };
+static const uniform struct FVector4f FloatPiByTwo = { {0.5f*FLOAT_PI, 0.5f*FLOAT_PI, 0.5f*FLOAT_PI, 0.5f*FLOAT_PI} };
+static const uniform struct FVector4f FloatPiByFour = { {0.25f*FLOAT_PI, 0.25f*FLOAT_PI, 0.25f*FLOAT_PI, 0.25f*FLOAT_PI} };
+static const uniform struct FVector4f FloatOneOverPi = { {1.0f / FLOAT_PI, 1.0f / FLOAT_PI, 1.0f / FLOAT_PI, 1.0f / FLOAT_PI} };
+static const uniform struct FVector4f FloatOneOverTwoPi = { {1.0f / (2.0f*FLOAT_PI), 1.0f / (2.0f*FLOAT_PI), 1.0f / (2.0f*FLOAT_PI), 1.0f / (2.0f*FLOAT_PI)} };
+
+static const uniform struct FVector4d DoublePi = { {DOUBLE_PI, DOUBLE_PI, DOUBLE_PI, DOUBLE_PI} };
+static const uniform struct FVector4d DoubleTwoPi = { {2.0d*DOUBLE_PI, 2.0d*DOUBLE_PI, 2.0d*DOUBLE_PI, 2.0d*DOUBLE_PI} };
+static const uniform struct FVector4d DoublePiByTwo = { {0.5d*DOUBLE_PI, 0.5d*DOUBLE_PI, 0.5d*DOUBLE_PI, 0.5d*DOUBLE_PI} };
+static const uniform struct FVector4d DoublePiByFour = { {0.25d*DOUBLE_PI, 0.25d*DOUBLE_PI, 0.25d*DOUBLE_PI, 0.25d*DOUBLE_PI} };
+static const uniform struct FVector4d DoubleOneOverPi = { {1.0d / DOUBLE_PI, 1.0d / DOUBLE_PI, 1.0d / DOUBLE_PI, 1.0d / DOUBLE_PI} };
+static const uniform struct FVector4d DoubleOneOverTwoPi = { {1.0d / (2.0d*DOUBLE_PI), 1.0d / (2.0d*DOUBLE_PI), 1.0d / (2.0d*DOUBLE_PI), 1.0d / (2.0d*DOUBLE_PI)} };
+
+static const uniform struct FVector4f Float255 = { {255.0f, 255.0f, 255.0f, 255.0f} };
+static const uniform struct FVector4f Float127 = { {127.0f, 127.0f, 127.0f, 127.0f} };
+static const uniform struct FVector4f FloatNeg127 = { {-127.0f, -127.0f, -127.0f, -127.0f} };
+static const uniform struct FVector4f Float360 = { {360.f, 360.f, 360.f, 360.f} };
+static const uniform struct FVector4f Float180 = { {180.f, 180.f, 180.f, 180.f} };
+static const uniform struct FVector4f FloatNonFractional = { {FLOAT_NON_FRACTIONAL, FLOAT_NON_FRACTIONAL, FLOAT_NON_FRACTIONAL, FLOAT_NON_FRACTIONAL} };
+
+static const uniform struct FVector4d Double255 = { {255.0d, 255.0d, 255.0d, 255.0d} };
+static const uniform struct FVector4d Double127 = { {127.0d, 127.0d, 127.0d, 127.0d} };
+static const uniform struct FVector4d DoubleNeg127 = { {-127.0d, -127.0d, -127.0d, -127.0d} };
+static const uniform struct FVector4d Double360 = { {360.d, 360.d, 360.d, 360.d} };
+static const uniform struct FVector4d Double180 = { {180.d, 180.d, 180.d, 180.d} };
+static const uniform struct FVector4d DoubleNonFractional = { {DOUBLE_NON_FRACTIONAL, DOUBLE_NON_FRACTIONAL, DOUBLE_NON_FRACTIONAL, DOUBLE_NON_FRACTIONAL} };
+
+static const uniform struct FVector4f FLOAT_DEG_TO_RAD ={ {PI/(180.f), PI/(180.f), PI/(180.f), PI/(180.f)} };
+static const uniform struct FVector4f FLOAT_DEG_TO_RAD_HALF = { {(PI/180.f)*0.5f, (PI/180.f)*0.5f, (PI/180.f)*0.5f, (PI/180.f)*0.5f} };
+static const uniform struct FVector4f FLOAT_RAD_TO_DEG = { {(180.f)/PI, (180.f)/PI, (180.f)/PI, (180.f)/PI} };
+
+static const uniform struct FVector4d DOUBLE_DEG_TO_RAD ={ {DOUBLE_PI/(180.d), DOUBLE_PI/(180.d), DOUBLE_PI/(180.d), DOUBLE_PI/(180.d)} };
+static const uniform struct FVector4d DOUBLE_DEG_TO_RAD_HALF = { {(DOUBLE_PI/180.d)*0.5d, (DOUBLE_PI/180.d)*0.5d, (DOUBLE_PI/180.d)*0.5d, (DOUBLE_PI/180.d)*0.5d} };
+static const uniform struct FVector4d DOUBLE_RAD_TO_DEG = { {(180.d)/DOUBLE_PI, (180.d)/DOUBLE_PI, (180.d)/DOUBLE_PI, (180.d)/DOUBLE_PI} };
+
+static const uniform struct FVector4f FloatUpVector4 = { {0.0f, 0.0f, 1.0f, 0.0f} };
+static const uniform struct FVector4f FloatForwardVector4 = { {1.0f, 0.0f, 0.0f, 0.0f} };
+static const uniform struct FVector4f FloatRightVector4 = { {0.0f, 1.0f, 0.0f, 0.0f} };
+
+static const uniform struct FVector4d DoubleUpVector4 = { {0.0d, 0.0d, 1.0d, 0.0d} };
+static const uniform struct FVector4d DoubleForwardVector4 = { {1.0d, 0.0d, 0.0d, 0.0d} };
+static const uniform struct FVector4d DoubleRightVector4 = { {0.0d, 1.0d, 0.0d, 0.0d} };
+
+static const uniform struct FVector3f FloatZeroVector = { {0.0f, 0.0f, 0.0f} };
+static const uniform struct FVector3f FloatOneVector = { {1.0f, 1.0f, 1.0f} };
+static const uniform struct FVector3f FloatHalfVector = { {0.5f, 0.5f, 0.5f} };
+static const uniform struct FVector3f FloatUpVector = { {0.0f, 0.0f, 1.0f} };
+static const uniform struct FVector3f FloatDownVector = { {0.0f, 0.0f, -1.0f} };
+static const uniform struct FVector3f FloatForwardVector = { {1.0f, 0.0f, 0.0f} };
+static const uniform struct FVector3f FloatBackwardVector = { {-1.0f, 0.0f, 0.0f} };
+static const uniform struct FVector3f FloatRightVector = { {0.0f, 1.0f, 0.0f} };
+static const uniform struct FVector3f FloatLeftVector = { {0.0f, -1.0f, 0.0f} };
+
+static const uniform struct FVector3d DoubleZeroVector = { {0.0d, 0.0d, 0.0d} };
+static const uniform struct FVector3d DoubleOneVector = { {1.0d, 1.0d, 1.0d} };
+static const uniform struct FVector3d DoubleHalfVector = { {0.5d, 0.5d, 0.5d} };
+static const uniform struct FVector3d DoubleUpVector = { {0.0d, 0.0d, 1.0d} };
+static const uniform struct FVector3d DoubleDownVector = { {0.0d, 0.0d, -1.0d} };
+static const uniform struct FVector3d DoubleForwardVector = { {1.0d, 0.0d, 0.0d} };
+static const uniform struct FVector3d DoubleBackwardVector = { {-1.0d, 0.0d, 0.0d} };
+static const uniform struct FVector3d DoubleRightVector = { {0.0d, 1.0d, 0.0d} };
+static const uniform struct FVector3d DoubleLeftVector = { {0.0d, -1.0d, 0.0d} };
 
 static const uniform struct FVector8 UpVector8 = { {0.0f, 0.0f, 1.0f, 0.0f, 0.0f, 0.0f, -1.0f, 0.0f} };
 static const uniform struct FVector8 ForwardVector8 = { {1.0f, 0.0f, 0.0f, 0.0f, -1.0f, 0.0f, 0.0f, 0.0f} };
 static const uniform struct FVector8 RightVector8 = { {0.0f, 1.0f, 0.0f, 0.0f, 0.0f, -1.0f, 0.0f, 0.0f} };
 
+// Define vector type constants
+#if UE_LARGE_WORLD_COORDINATES_DISABLED
+
+#define VectorZero				FloatZero
+#define VectorOne				FloatOne
+#define VectorMinusOne			FloatMinusOne
+#define Vector0001				Float0001
+#define Vector1110				Float1110
+#define VectorOneHalf			FloatOneHalf
+#define SmallLengthThreshold	FloatSmallLengthThreshold
+
+#define VectorPi				FloatPi
+#define VectorTwoPi				FloatTwoPi
+#define VectorPiByTwo			FloatPiByTwo
+#define VectorPiByFour			FloatPiByFour
+#define VectorOneOverPi			FloatOneOverPi
+#define VectorOneOverTwoPi		FloatOneOverTwoPi
+
+#define Vector255				Float255
+#define Vector127				Float127
+#define VectorNeg127			FloatNeg127
+#define Vector360				Float360
+#define Vector180				Float180
+#define VectorNonFractional		FloatNonFractional
+
+#define UpVector4				FloatUpVector4
+#define ForwardVector4			FloatForwardVector4
+#define RightVector4			FloatRightVector4
+
+#define ZeroVector				FloatZeroVector
+#define OneVector				FloatOneVector
+#define HalfVector				FloatHalfVector
+#define UpVector				FloatUpVector
+#define DownVector				FloatDownVector
+#define ForwardVector			FloatForwardVector
+#define BackwardVector			FloatBackwardVector
+#define RightVector				FloatRightVector
+#define LeftVector				FloatLeftVector
+
+#else
+
+#define VectorZero				DoubleZero
+#define VectorOne				DoubleOne
+#define VectorMinusOne			DoubleMinusOne
+#define Vector0001				Double0001
+#define Vector1110				Double1110
+#define VectorOneHalf			DoubleOneHalf
+#define SmallLengthThreshold	DoubleSmallLengthThreshold
+
+#define VectorPi				DoublePi
+#define VectorTwoPi				DoubleTwoPi
+#define VectorPiByTwo			DoublePiByTwo
+#define VectorPiByFour			DoublePiByFour
+#define VectorOneOverPi			DoubleOneOverPi
+#define VectorOneOverTwoPi		DoubleOneOverTwoPi
+
+#define Vector255				Double255
+#define Vector127				Double127
+#define VectorNeg127			DoubleNeg127
+#define Vector360				Double360
+#define Vector180				Double180
+#define VectorNonFractional		DoubleNonFractional
+
+#define UpVector4				DoubleUpVector4
+#define ForwardVector4			DoubleForwardVector4
+#define RightVector4			DoubleRightVector4
+
+#define ZeroVector				DoubleZeroVector
+#define OneVector				DoubleOneVector
+#define HalfVector				DoubleHalfVector
+#define UpVector				DoubleUpVector
+#define DownVector				DoubleDownVector
+#define ForwardVector			DoubleForwardVector
+#define BackwardVector			DoubleBackwardVector
+#define RightVector				DoubleRightVector
+#define LeftVector				DoubleLeftVector
+#endif
+
 #if (defined(ISPC_TARGET_AVX2) || defined(ISPC_TARGET_AVX512SKX))
 #define HW_GATHER_SUPPORTED 1
 #else
@@ -257,7 +451,7 @@ inline varying FVector2 VectorLoad(const uniform FVector2 *uniform SrcPtr)
 
 	if(((1<<TARGET_WIDTH)-1 ^ lanemask()) == 0)
 	{
-		aos_to_soa2_ispc((uniform float * uniform)SrcPtr, &Result.V[0], &Result.V[1]);
+		aos_to_soa2_ispc((uniform FReal * uniform)SrcPtr, &Result.V[0], &Result.V[1]);
 	}
 	else
 	{
@@ -272,7 +466,7 @@ inline void VectorStore(uniform FVector2 *uniform DstPtr, const varying FVector2
 {
 	if(((1<<TARGET_WIDTH)-1 ^ lanemask()) == 0)
 	{
-		soa_to_aos2_ispc(V.V[0], V.V[1], (uniform float *uniform)DstPtr);
+		soa_to_aos2_ispc(V.V[0], V.V[1], (uniform FReal *uniform)DstPtr);
 	}
 	else
 	{
@@ -281,9 +475,26 @@ inline void VectorStore(uniform FVector2 *uniform DstPtr, const varying FVector2
 	}
 }
 
-inline varying FVector VectorLoad(const uniform FVector *uniform SrcPtr)
+inline varying FVector3d VectorLoad(const uniform FVector3d *uniform SrcPtr)
+{
+	varying FVector3d Result;
+
+	if(((1<<TARGET_WIDTH)-1 ^ lanemask()) == 0)
+	{
+		aos_to_soa3_ispc((uniform double * uniform)SrcPtr, &Result.V[0], &Result.V[1], &Result.V[2]);
+	}
+	else
+	{
+		#pragma ignore warning(perf)
+		Result = SrcPtr[programIndex];
+	}
+
+	return Result;
+}
+
+inline varying FVector3f VectorLoad(const uniform FVector3f *uniform SrcPtr)
 {
-	varying FVector Result;
+	varying FVector3f Result;
 
 	if(((1<<TARGET_WIDTH)-1 ^ lanemask()) == 0)
 	{
@@ -298,7 +509,20 @@ inline varying FVector VectorLoad(const uniform FVector *uniform SrcPtr)
 	return Result;
 }
 
-inline void VectorStore(uniform FVector *uniform DstPtr, const varying FVector &V)
+inline void VectorStore(uniform FVector3d *uniform DstPtr, const varying FVector3d &V)
+{
+	if(((1<<TARGET_WIDTH)-1 ^ lanemask()) == 0)
+	{
+		soa_to_aos3_ispc(V.V[0], V.V[1], V.V[2], (uniform double *uniform)DstPtr);
+	}
+	else
+	{
+		#pragma ignore warning(perf)
+		DstPtr[programIndex] = V;
+	}
+}
+
+inline void VectorStore(uniform FVector3f *uniform DstPtr, const varying FVector3f &V)
 {
 	if(((1<<TARGET_WIDTH)-1 ^ lanemask()) == 0)
 	{
@@ -317,7 +541,7 @@ inline varying FVector4 VectorLoad(const uniform FVector4 *uniform SrcPtr)
 
 	if(((1<<TARGET_WIDTH)-1 ^ lanemask()) == 0)
 	{
-		aos_to_soa4_ispc((uniform float * uniform)SrcPtr, &Result.V[0], &Result.V[1], &Result.V[2], &Result.V[3]);
+		aos_to_soa4_ispc((uniform FReal * uniform)SrcPtr, &Result.V[0], &Result.V[1], &Result.V[2], &Result.V[3]);
 	}
 	else
 	{
@@ -332,7 +556,7 @@ inline void VectorStore(uniform FVector4 *uniform DstPtr, const varying FVector4
 {
 	if(((1<<TARGET_WIDTH)-1 ^ lanemask()) == 0)
 	{
-		soa_to_aos4_ispc(V.V[0], V.V[1], V.V[2], V.V[3], (uniform float *uniform)DstPtr);
+		soa_to_aos4_ispc(V.V[0], V.V[1], V.V[2], V.V[3], (uniform FReal *uniform)DstPtr);
 	}
 	else
 	{
@@ -402,7 +626,7 @@ inline varying FVector4 VectorGather(const uniform FVector4 *varying SrcPtr)
 			DstPtr[15] = *Src[15];
 #endif
 
-			aos_to_soa4_ispc((uniform float * uniform)&Result, &Result.V[0], &Result.V[1], &Result.V[2], &Result.V[3]);
+			aos_to_soa4_ispc((uniform FReal * uniform)&Result, &Result.V[0], &Result.V[1], &Result.V[2], &Result.V[3]);
 		}
 	}
 	else
@@ -421,7 +645,7 @@ inline void VectorScatter(uniform FVector4 *varying DstPtr, const varying FVecto
 		unmasked
 		{
 			uniform FVector4 Result[programCount];
-			soa_to_aos4_ispc(V.V[0], V.V[1], V.V[2], V.V[3], (uniform float *uniform)&Result);
+			soa_to_aos4_ispc(V.V[0], V.V[1], V.V[2], V.V[3], (uniform FReal *uniform)&Result);
 
 			uniform FVector4 *uniform Dst[programCount];
 
@@ -459,9 +683,9 @@ inline void VectorScatter(uniform FVector4 *varying DstPtr, const varying FVecto
 	}
 }
 
-inline FVector4 operator+(const FVector4 &A, const FVector4 &B)
+inline FVector4d operator+(const FVector4d &A, const FVector4d &B)
 {
-	FVector4 Result;
+	FVector4d Result;
 
 	for(uniform int i = 0; i < 4; i++)
 	{
@@ -471,15 +695,62 @@ inline FVector4 operator+(const FVector4 &A, const FVector4 &B)
 	return Result;
 }
 
-inline uniform FVector4 operator+(const uniform FVector4 &A, const uniform FVector4 &B)
+inline FVector4f operator+(const FVector4f &A, const FVector4f &B)
+{
+	FVector4f Result;
+
+	for(uniform int i = 0; i < 4; i++)
+	{
+		Result.V[i] = A.V[i] + B.V[i];
+	}
+
+	return Result;
+}
+
+inline uniform FVector4d operator+(const uniform FVector4d &A, const uniform FVector4d &B)
+{
+	varying double S0, S1, Result;
+	*((uniform FVector4d *uniform)&S0) = *((uniform FVector4d *uniform)&A);
+	*((uniform FVector4d *uniform)&S1) = *((uniform FVector4d *uniform)&B);
+
+	Result = S0 + S1;
+
+	return *((uniform FVector4d *uniform)&Result);
+}
+
+inline uniform FVector4f operator+(const uniform FVector4f &A, const uniform FVector4f &B)
 {
 	varying float S0, S1, Result;
-	*((uniform FVector4 *uniform)&S0) = *((uniform FVector4 *uniform)&A);
-	*((uniform FVector4 *uniform)&S1) = *((uniform FVector4 *uniform)&B);
+	*((uniform FVector4f *uniform)&S0) = *((uniform FVector4f *uniform)&A);
+	*((uniform FVector4f *uniform)&S1) = *((uniform FVector4f *uniform)&B);
 
 	Result = S0 + S1;
 
-	return *((uniform FVector4 *uniform)&Result);
+	return *((uniform FVector4f *uniform)&Result);
+}
+
+inline uniform FVector4d operator+(const uniform FVector4d &A, const uniform FVector4f &B)
+{
+	varying double S0, Result;
+	varying float S1;
+	*((uniform FVector4d *uniform)&S0) = *((uniform FVector4d *uniform)&A);
+	*((uniform FVector4f *uniform)&S1) = *((uniform FVector4f *uniform)&B);
+
+	Result = S0 + S1;
+
+	return *((uniform FVector4d *uniform)&Result);
+}
+
+inline uniform FVector4d operator+(const uniform FVector4f &A, const uniform FVector4d &B)
+{
+	varying float S0;
+	varying double S1, Result;
+	*((uniform FVector4f *uniform)&S0) = *((uniform FVector4f *uniform)&A);
+	*((uniform FVector4d *uniform)&S1) = *((uniform FVector4d *uniform)&B);
+
+	Result = S0 + S1;
+
+	return *((uniform FVector4d *uniform)&Result);
 }
 
 inline FVector4 operator+(const FVector4 &A, const uniform FVector4 &B)
@@ -494,7 +765,7 @@ inline FVector4 operator+(const FVector4 &A, const uniform FVector4 &B)
 	return Result;
 }
 
-inline FVector4 operator+(const FVector4 &A, const float F)
+inline FVector4 operator+(const FVector4 &A, const FReal F)
 {
 	FVector4 Result;
 
@@ -506,7 +777,7 @@ inline FVector4 operator+(const FVector4 &A, const float F)
 	return Result;
 }
 
-inline FVector4 operator+(const float F, const FVector4 &A)
+inline FVector4 operator+(const FReal F, const FVector4 &A)
 {
 	FVector4 Result;
 
@@ -518,15 +789,27 @@ inline FVector4 operator+(const float F, const FVector4 &A)
 	return Result;
 }
 
-inline uniform FVector4 operator+(const uniform FVector4 &A, const uniform float F)
+inline uniform FVector4 operator+(const uniform FVector4 &A, const uniform FReal F)
 {
 	const uniform FVector4 FVec = {{F,F,F,F}};
 	return A + FVec;
 }
 
-inline FVector4 operator*(const FVector4 &A, const FVector4 &B)
+inline FVector4d operator*(const FVector4d &A, const FVector4d &B)
 {
-	FVector4 Result;
+	FVector4d Result;
+
+	for(uniform int i = 0; i < 4; i++)
+	{
+		Result.V[i] = A.V[i] * B.V[i];
+	}
+
+	return Result;
+}
+
+inline FVector4f operator*(const FVector4f &A, const FVector4f &B)
+{
+	FVector4f Result;
 
 	for(uniform int i = 0; i < 4; i++)
 	{
@@ -548,9 +831,21 @@ inline FVector4 operator*(const uniform FVector4 &A, const FVector4 &B)
 	return Result;
 }
 
-inline FVector4 operator*(const FVector4 &A, const uniform FVector4 &B)
+inline FVector4d operator*(const FVector4d &A, const uniform FVector4d &B)
 {
-	FVector4 Result;
+	FVector4d Result;
+
+	for(uniform int i = 0; i < 4; i++)
+	{
+		Result.V[i] = A.V[i] * B.V[i];
+	}
+
+	return Result;
+}
+
+inline FVector4f operator*(const FVector4f &A, const uniform FVector4f &B)
+{
+	FVector4f Result;
 
 	for(uniform int i = 0; i < 4; i++)
 	{
@@ -560,20 +855,43 @@ inline FVector4 operator*(const FVector4 &A, const uniform FVector4 &B)
 	return Result;
 }
 
-inline uniform FVector4 operator*(const uniform FVector4 &A, const uniform FVector4 &B)
+inline uniform FVector4d operator*(const uniform FVector4d &A, const uniform FVector4d &B)
+{
+	varying double S0, S1, Result;
+	*((uniform FVector4d *uniform)&S0) = *((uniform FVector4d *uniform)&A);
+	*((uniform FVector4d *uniform)&S1) = *((uniform FVector4d *uniform)&B);
+
+	Result = S0 * S1;
+
+	return *((uniform FVector4d *uniform)&Result);
+}
+
+inline uniform FVector4f operator*(const uniform FVector4f &A, const uniform FVector4f &B)
 {
 	varying float S0, S1, Result;
-	*((uniform FVector4 *uniform)&S0) = *((uniform FVector4 *uniform)&A);
-	*((uniform FVector4 *uniform)&S1) = *((uniform FVector4 *uniform)&B);
+	*((uniform FVector4f *uniform)&S0) = *((uniform FVector4f *uniform)&A);
+	*((uniform FVector4f *uniform)&S1) = *((uniform FVector4f *uniform)&B);
 
 	Result = S0 * S1;
 
-	return *((uniform FVector4 *uniform)&Result);
+	return *((uniform FVector4f *uniform)&Result);
 }
 
-inline FVector4 operator*(const FVector4 &A, const float F)
+inline FVector4d operator*(const FVector4d &A, const double F)
 {
-	FVector4 Result;
+	FVector4d Result;
+
+	for(uniform int i = 0; i < 4; i++)
+	{
+		Result.V[i] = A.V[i] * F;
+	}
+
+	return Result;
+}
+
+inline FVector4f operator*(const FVector4f &A, const float F)
+{
+	FVector4f Result;
 
 	for(uniform int i = 0; i < 4; i++)
 	{
@@ -583,7 +901,7 @@ inline FVector4 operator*(const FVector4 &A, const float F)
 	return Result;
 }
 
-inline FVector4 operator*(const FVector4 &A, const uniform float F)
+inline FVector4 operator*(const FVector4 &A, const uniform FReal F)
 {
 	FVector4 Result;
 
@@ -595,7 +913,7 @@ inline FVector4 operator*(const FVector4 &A, const uniform float F)
 	return Result;
 }
 
-inline FVector4 operator*(const uniform float F, const FVector4 &A)
+inline FVector4 operator*(const uniform FReal F, const FVector4 &A)
 {
 	FVector4 Result;
 
@@ -607,15 +925,27 @@ inline FVector4 operator*(const uniform float F, const FVector4 &A)
 	return Result;
 }
 
-inline uniform FVector4 operator*(const uniform FVector4 &A, const uniform float F)
+inline uniform FVector4d operator*(const uniform FVector4d &A, const uniform double F)
 {
-	const uniform FVector4 FVec = {{F,F,F,F}};
+	const uniform FVector4d FVec = {{F,F,F,F}};
 	return A * FVec;
 }
 
-inline uniform FVector4 operator*(const uniform float F, const uniform FVector4 &A)
+inline uniform FVector4f operator*(const uniform FVector4f &A, const uniform float F)
 {
-	const uniform FVector4 FVec = {{F,F,F,F}};
+	const uniform FVector4f FVec = {{F,F,F,F}};
+	return A * FVec;
+}
+
+inline uniform FVector4d operator*(const uniform double F, const uniform FVector4d &A)
+{
+	const uniform FVector4d FVec = {{F,F,F,F}};
+	return FVec * A;
+}
+
+inline uniform FVector4f operator*(const uniform float F, const uniform FVector4f &A)
+{
+	const uniform FVector4f FVec = {{F,F,F,F}};
 	return FVec * A;
 }
 
@@ -657,7 +987,7 @@ inline FVector4 operator/(const uniform FVector4 &A, const FVector4 &B)
 
 inline uniform FVector4 operator/(const uniform FVector4 &A, const uniform FVector4 &B)
 {
-	varying float S0, S1, Result;
+	varying FReal S0, S1, Result;
 	*((uniform FVector4 *uniform)&S0) = *((uniform FVector4 *uniform)&A);
 	*((uniform FVector4 *uniform)&S1) = *((uniform FVector4 *uniform)&B);
 
@@ -666,21 +996,33 @@ inline uniform FVector4 operator/(const uniform FVector4 &A, const uniform FVect
 	return *((uniform FVector4 *uniform)&Result);
 }
 
-inline uniform FVector4 operator/(const uniform FVector4 &A, const uniform float F)
+inline uniform FVector4 operator/(const uniform FVector4 &A, const uniform FReal F)
 {
 	const uniform FVector4 FVec = {{F,F,F,F}};
 	return A / FVec;
 }
 
-inline uniform FVector4 operator/(const uniform float F, const uniform FVector4 &A)
+inline uniform FVector4 operator/(const uniform FReal F, const uniform FVector4 &A)
 {
 	const uniform FVector4 FVec = {{F,F,F,F}};
 	return FVec / A;
 }
 
-inline FVector4 operator-(const FVector4 &A, const FVector4 &B)
+inline FVector4d operator-(const FVector4d &A, const FVector4d &B)
 {
-	FVector4 Result;
+	FVector4d Result;
+
+	for(uniform int i = 0; i < 4; i++)
+	{
+		Result.V[i] = A.V[i] - B.V[i];
+	}
+
+	return Result;
+}
+
+inline FVector4f operator-(const FVector4f &A, const FVector4f &B)
+{
+	FVector4f Result;
 
 	for(uniform int i = 0; i < 4; i++)
 	{
@@ -690,15 +1032,26 @@ inline FVector4 operator-(const FVector4 &A, const FVector4 &B)
 	return Result;
 }
 
-inline uniform FVector4 operator-(const uniform FVector4 &A, const uniform FVector4 &B)
+inline uniform FVector4d operator-(const uniform FVector4d &A, const uniform FVector4d &B)
+{
+	varying double S0, S1, Result;
+	*((uniform FVector4d *uniform)&S0) = *((uniform FVector4d *uniform)&A);
+	*((uniform FVector4d *uniform)&S1) = *((uniform FVector4d *uniform)&B);
+
+	Result = S0 - S1;
+
+	return *((uniform FVector4d *uniform)&Result);
+}
+
+inline uniform FVector4f operator-(const uniform FVector4f &A, const uniform FVector4f &B)
 {
 	varying float S0, S1, Result;
-	*((uniform FVector4 *uniform)&S0) = *((uniform FVector4 *uniform)&A);
-	*((uniform FVector4 *uniform)&S1) = *((uniform FVector4 *uniform)&B);
+	*((uniform FVector4f *uniform)&S0) = *((uniform FVector4f *uniform)&A);
+	*((uniform FVector4f *uniform)&S1) = *((uniform FVector4f *uniform)&B);
 
 	Result = S0 - S1;
 
-	return *((uniform FVector4 *uniform)&Result);
+	return *((uniform FVector4f *uniform)&Result);
 }
 
 inline FVector4 operator-(const FVector4 &A, const uniform FVector4 &B)
@@ -713,9 +1066,9 @@ inline FVector4 operator-(const FVector4 &A, const uniform FVector4 &B)
 	return Result;
 }
 
-inline FVector4 operator-(const uniform FVector4 &A, const FVector4 &B)
+inline FVector4d operator-(const uniform FVector4d &A, const FVector4d &B)
 {
-	FVector4 Result;
+	FVector4d Result;
 
 	for(uniform int i = 0; i < 4; i++)
 	{
@@ -725,7 +1078,19 @@ inline FVector4 operator-(const uniform FVector4 &A, const FVector4 &B)
 	return Result;
 }
 
-inline FVector4 operator-(const FVector4 &A, const float F)
+inline FVector4f operator-(const uniform FVector4f &A, const FVector4f &B)
+{
+	FVector4f Result;
+
+	for(uniform int i = 0; i < 4; i++)
+	{
+		Result.V[i] = A.V[i] - B.V[i];
+	}
+
+	return Result;
+}
+
+inline FVector4 operator-(const FVector4 &A, const FReal F)
 {
 	FVector4 Result;
 
@@ -737,7 +1102,7 @@ inline FVector4 operator-(const FVector4 &A, const float F)
 	return Result;
 }
 
-inline FVector4 operator-(const FVector4 &A, const uniform float F)
+inline FVector4 operator-(const FVector4 &A, const uniform FReal F)
 {
 	FVector4 Result;
 
@@ -749,7 +1114,7 @@ inline FVector4 operator-(const FVector4 &A, const uniform float F)
 	return Result;
 }
 
-inline FVector4 operator-(const uniform float F, const FVector4 &A)
+inline FVector4 operator-(const uniform FReal F, const FVector4 &A)
 {
 	FVector4 Result;
 
@@ -761,57 +1126,125 @@ inline FVector4 operator-(const uniform float F, const FVector4 &A)
 	return Result;
 }
 
-inline uniform FVector4 operator-(const uniform FVector4 &A, const uniform float F)
+inline uniform FVector4 operator-(const uniform FVector4 &A, const uniform FReal F)
 {
 	const uniform FVector4 FVec = {{F,F,F,F}};
 	return A - FVec;
 }
 
-inline uniform FVector4 operator-(const uniform float F, const uniform FVector4 &A)
+inline uniform FVector4 operator-(const uniform FReal F, const uniform FVector4 &A)
 {
 	const uniform FVector4 FVec = {{F,F,F,F}};
 	return FVec - A;
 }
 
-inline FVector4 SetVector4(const float X, const float Y, const float Z, const float W)
+inline FVector4d SetVector4(const double X, const double Y, const double Z, const double W)
+{
+	const FVector4d Result = {{X, Y, Z, W}};
+	return Result;
+}
+
+inline FVector4f SetVector4(const float X, const float Y, const float Z, const float W)
+{
+	const FVector4f Result = {{X, Y, Z, W}};
+	return Result;
+}
+
+inline uniform FVector4d SetVector4(const uniform double X, const uniform double Y, const uniform double Z, const uniform double W)
 {
-	const FVector4 Result = {{X, Y, Z, W}};
+	const uniform FVector4d Result = {{X, Y, Z, W}};
 	return Result;
 }
 
-inline uniform FVector4 SetVector4(const uniform float X, const uniform float Y, const uniform float Z, const uniform float W)
+inline uniform FVector4f SetVector4(const uniform float X, const uniform float Y, const uniform float Z, const uniform float W)
 {
-	const uniform FVector4 Result = {{X, Y, Z, W}};
+	const uniform FVector4f Result = {{X, Y, Z, W}};
 	return Result;
 }
 
-inline FVector4 SetVector4(const FVector &V, const float W)
+inline FVector4 SetVector4(const FVector &V, const FReal W)
 {
 	const FVector4 Result = {{V.V[0], V.V[1], V.V[2], W}};
 	return Result;
 }
 
-inline uniform FVector4 SetVector4(const uniform FVector &V, const uniform float W)
+inline uniform FVector4 SetVector4(const uniform FVector &V, const uniform FReal W)
 {
 	const uniform FVector4 Result = {{V.V[0], V.V[1], V.V[2], W}};
 	return Result;
 }
 
-inline uniform FVector4 SetVector4(const uniform float F)
+inline uniform FVector4 SetVector4(const uniform FReal F)
 {
 	const uniform FVector4 Result = {{F, F, F, F}};
 	return Result;
 }
 
-inline varying FVector4 SetVector4(const varying float F)
+inline varying FVector4 SetVector4(const varying FReal F)
 {
 	const varying FVector4 Result = {{F, F, F, F}};
 	return Result;
 }
 
-inline FVector operator+(const FVector &A, const FVector &B)
+inline uniform FVector4d ConvertVector4fTo4d(const uniform FVector4f &V)
 {
-	FVector Result;
+	varying float S0;
+	*((uniform FVector4f *uniform)&S0) = *((uniform FVector4f *uniform)&V);
+
+	const varying double Result = (double)S0;
+
+	return *((uniform FVector4d *uniform)&Result);
+}
+
+inline uniform FVector4f ConvertVector4dTo4f(const uniform FVector4d &V)
+{
+	varying double S0;
+	*((uniform FVector4d *uniform)&S0) = *((uniform FVector4d *uniform)&V);
+
+	const varying float Result = (float)S0;
+
+	return *((uniform FVector4f *uniform)&Result);
+}
+
+inline uniform FVector4 ConvertVector4dTo4Native(const uniform FVector4d &V)
+{
+#if UE_LARGE_WORLD_COORDINATES_DISABLED
+	return ConvertVector4dTo4f(V);
+#else
+	return V;
+#endif
+}
+
+inline uniform FVector4 ConvertVector4fTo4Native(const uniform FVector4f &V)
+{
+#if UE_LARGE_WORLD_COORDINATES_DISABLED
+	return V;
+#else
+	return ConvertVector4fTo4d(V);
+#endif
+}
+
+inline uniform FVector4d ConvertVector4NativeTo4d(const uniform FVector4 &V)
+{
+#if UE_LARGE_WORLD_COORDINATES_DISABLED
+	return ConvertVector4fTo4d(V);
+#else
+	return V;
+#endif
+}
+
+inline uniform FVector4f ConvertVector4NativeTo4f(const uniform FVector4 &V)
+{
+#if UE_LARGE_WORLD_COORDINATES_DISABLED
+	return V;
+#else
+	return ConvertVector4dTo4f(V);
+#endif
+}
+
+inline FVector3d operator+(const FVector3d &A, const FVector3d &B)
+{
+	FVector3d Result;
 
 	Result.V[0] = A.V[0] + B.V[0];
 	Result.V[1] = A.V[1] + B.V[1];
@@ -820,20 +1253,53 @@ inline FVector operator+(const FVector &A, const FVector &B)
 	return Result;
 }
 
-inline uniform FVector operator+(const uniform FVector &A, const uniform FVector &B)
+inline FVector3f operator+(const FVector3f &A, const FVector3f &B)
 {
-	uniform FVector4 S0, S1, Result;
-	*((uniform FVector *uniform)&S0) = *((uniform FVector *uniform)&A);
-	*((uniform FVector *uniform)&S1) = *((uniform FVector *uniform)&B);
+	FVector3f Result;
+
+	Result.V[0] = A.V[0] + B.V[0];
+	Result.V[1] = A.V[1] + B.V[1];
+	Result.V[2] = A.V[2] + B.V[2];
+
+	return Result;
+}
+
+inline uniform FVector3d operator+(const uniform FVector3d &A, const uniform FVector3d &B)
+{
+	uniform FVector4d S0, S1, Result;
+	*((uniform FVector3d *uniform)&S0) = *((uniform FVector3d *uniform)&A);
+	*((uniform FVector3d *uniform)&S1) = *((uniform FVector3d *uniform)&B);
 
 	Result = S0 + S1;
 
-	return *((uniform FVector *uniform)&Result);
+	return *((uniform FVector3d *uniform)&Result);
 }
 
-inline FVector operator+(const FVector &A, const uniform FVector &B)
+inline uniform FVector3f operator+(const uniform FVector3f &A, const uniform FVector3f &B)
 {
-	FVector Result;
+	uniform FVector4f S0, S1, Result;
+	*((uniform FVector3f *uniform)&S0) = *((uniform FVector3f *uniform)&A);
+	*((uniform FVector3f *uniform)&S1) = *((uniform FVector3f *uniform)&B);
+
+	Result = S0 + S1;
+
+	return *((uniform FVector3f *uniform)&Result);
+}
+
+inline FVector3d operator+(const FVector3d &A, const uniform FVector3d &B)
+{
+	FVector3d Result;
+
+	Result.V[0] = A.V[0] + B.V[0];
+	Result.V[1] = A.V[1] + B.V[1];
+	Result.V[2] = A.V[2] + B.V[2];
+
+	return Result;
+}
+
+inline FVector3f operator+(const FVector3f &A, const uniform FVector3f &B)
+{
+	FVector3f Result;
 
 	Result.V[0] = A.V[0] + B.V[0];
 	Result.V[1] = A.V[1] + B.V[1];
@@ -853,7 +1319,7 @@ inline FVector operator+(const uniform FVector &A, const FVector &B)
 	return Result;
 }
 
-inline FVector operator+(const FVector &A, const float F)
+inline FVector operator+(const FVector &A, const FReal F)
 {
 	FVector Result;
 
@@ -864,15 +1330,15 @@ inline FVector operator+(const FVector &A, const float F)
 	return Result;
 }
 
-inline uniform FVector operator+(const uniform FVector &A, const uniform float F)
+inline uniform FVector operator+(const uniform FVector &A, const uniform FReal F)
 {
 	const uniform FVector FVec = {{F,F,F}};
 	return A + FVec;
 }
 
-inline FVector operator*(const FVector &A, const FVector &B)
+inline FVector3d operator*(const FVector3d &A, const FVector3d &B)
 {
-	FVector Result;
+	FVector3d Result;
 
 	Result.V[0] = A.V[0] * B.V[0];
 	Result.V[1] = A.V[1] * B.V[1];
@@ -881,20 +1347,42 @@ inline FVector operator*(const FVector &A, const FVector &B)
 	return Result;
 }
 
-inline uniform FVector operator*(const uniform FVector &A, const uniform FVector &B)
+inline FVector3f operator*(const FVector3f &A, const FVector3f &B)
 {
-	uniform FVector4 S0, S1, Result;
-	*((uniform FVector *uniform)&S0) = *((uniform FVector *uniform)&A);
-	*((uniform FVector *uniform)&S1) = *((uniform FVector *uniform)&B);
+	FVector3f Result;
+
+	Result.V[0] = A.V[0] * B.V[0];
+	Result.V[1] = A.V[1] * B.V[1];
+	Result.V[2] = A.V[2] * B.V[2];
+
+	return Result;
+}
+
+inline uniform FVector3d operator*(const uniform FVector3d &A, const uniform FVector3d &B)
+{
+	uniform FVector4d S0, S1, Result;
+	*((uniform FVector3d *uniform)&S0) = *((uniform FVector3d *uniform)&A);
+	*((uniform FVector3d *uniform)&S1) = *((uniform FVector3d *uniform)&B);
 
 	Result = S0 * S1;
 
-	return *((uniform FVector *uniform)&Result);
+	return *((uniform FVector3d *uniform)&Result);
 }
 
-inline FVector operator*(const FVector &A, const uniform FVector &B)
+inline uniform FVector3f operator*(const uniform FVector3f &A, const uniform FVector3f &B)
 {
-	FVector Result;
+	uniform FVector4f S0, S1, Result;
+	*((uniform FVector3f *uniform)&S0) = *((uniform FVector3f *uniform)&A);
+	*((uniform FVector3f *uniform)&S1) = *((uniform FVector3f *uniform)&B);
+
+	Result = S0 * S1;
+
+	return *((uniform FVector3f *uniform)&Result);
+}
+
+inline FVector3d operator*(const FVector3d &A, const uniform FVector3d &B)
+{
+	FVector3d Result;
 
 	Result.V[0] = A.V[0] * B.V[0];
 	Result.V[1] = A.V[1] * B.V[1];
@@ -903,9 +1391,31 @@ inline FVector operator*(const FVector &A, const uniform FVector &B)
 	return Result;
 }
 
-inline FVector operator*(const uniform FVector &A, const FVector &B)
+inline FVector3f operator*(const FVector3f &A, const uniform FVector3f &B)
 {
-	FVector Result;
+	FVector3f Result;
+
+	Result.V[0] = A.V[0] * B.V[0];
+	Result.V[1] = A.V[1] * B.V[1];
+	Result.V[2] = A.V[2] * B.V[2];
+
+	return Result;
+}
+
+inline FVector3d operator*(const uniform FVector3d &A, const FVector3d &B)
+{
+	FVector3d Result;
+
+	Result.V[0] = A.V[0] * B.V[0];
+	Result.V[1] = A.V[1] * B.V[1];
+	Result.V[2] = A.V[2] * B.V[2];
+
+	return Result;
+}
+
+inline FVector3f operator*(const uniform FVector3f &A, const FVector3f &B)
+{
+	FVector3f Result;
 
 	Result.V[0] = A.V[0] * B.V[0];
 	Result.V[1] = A.V[1] * B.V[1];
@@ -925,9 +1435,9 @@ inline FVector operator*(const uniform FVector &A, const FIntVector &B)
 	return Result;
 }
 
-inline FVector operator*(const FVector &A, const float F)
+inline FVector3d operator*(const FVector3d &A, const double F)
 {
-	FVector Result;
+	FVector3d Result;
 
 	Result.V[0] = A.V[0] * F;
 	Result.V[1] = A.V[1] * F;
@@ -936,9 +1446,42 @@ inline FVector operator*(const FVector &A, const float F)
 	return Result;
 }
 
-inline FVector operator*(const uniform FVector &A, const varying float F)
+inline FVector3f operator*(const FVector3f &A, const float F)
 {
-	FVector Result;
+	FVector3f Result;
+
+	Result.V[0] = A.V[0] * F;
+	Result.V[1] = A.V[1] * F;
+	Result.V[2] = A.V[2] * F;
+
+	return Result;
+}
+
+inline FVector3d operator*(const uniform FVector3d &A, const varying double F)
+{
+	FVector3d Result;
+
+	Result.V[0] = A.V[0] * F;
+	Result.V[1] = A.V[1] * F;
+	Result.V[2] = A.V[2] * F;
+
+	return Result;
+}
+
+inline FVector3f operator*(const uniform FVector3f &A, const varying float F)
+{
+	FVector3f Result;
+
+	Result.V[0] = A.V[0] * F;
+	Result.V[1] = A.V[1] * F;
+	Result.V[2] = A.V[2] * F;
+
+	return Result;
+}
+
+inline FVector3d operator*(const FVector3d &A, const uniform double F)
+{
+	FVector3d Result;
 
 	Result.V[0] = A.V[0] * F;
 	Result.V[1] = A.V[1] * F;
@@ -947,32 +1490,87 @@ inline FVector operator*(const uniform FVector &A, const varying float F)
 	return Result;
 }
 
-inline FVector operator*(const float F, const FVector &A)
+inline FVector3f operator*(const FVector3f &A, const uniform float F)
+{
+	FVector3f Result;
+
+	Result.V[0] = A.V[0] * F;
+	Result.V[1] = A.V[1] * F;
+	Result.V[2] = A.V[2] * F;
+
+	return Result;
+}
+
+inline FVector operator*(const FReal F, const FVector &A)
+{
+	FVector Result;
+
+	Result.V[0] = F * A.V[0];
+	Result.V[1] = F * A.V[1];
+	Result.V[2] = F * A.V[2];
+
+	return Result;
+}
+
+inline uniform FVector operator*(const uniform FVector &A, const uniform FReal F)
+{
+	const uniform FVector FVec = {{F,F,F}};
+	return A * FVec;
+}
+
+inline uniform FVector operator*(const uniform FReal F, const uniform FVector &A)
+{
+	const uniform FVector FVec = {{F,F,F}};
+	return FVec * A;
+}
+
+inline FVector3d operator-(const FVector3d &A, const FVector3d &B)
+{
+	FVector3d Result;
+
+	Result.V[0] = A.V[0] - B.V[0];
+	Result.V[1] = A.V[1] - B.V[1];
+	Result.V[2] = A.V[2] - B.V[2];
+
+	return Result;
+}
+
+inline FVector3f operator-(const FVector3f &A, const FVector3f &B)
+{
+	FVector3f Result;
+
+	Result.V[0] = A.V[0] - B.V[0];
+	Result.V[1] = A.V[1] - B.V[1];
+	Result.V[2] = A.V[2] - B.V[2];
+
+	return Result;
+}
+
+inline uniform FVector3d operator-(const uniform FVector3d &A, const uniform FVector3d &B)
 {
-	FVector Result;
+	uniform FVector4d S0, S1, Result;
+	*((uniform FVector3d *uniform)&S0) = *((uniform FVector3d *uniform)&A);
+	*((uniform FVector3d *uniform)&S1) = *((uniform FVector3d *uniform)&B);
 
-	Result.V[0] = F * A.V[0];
-	Result.V[1] = F * A.V[1];
-	Result.V[2] = F * A.V[2];
+	Result = S0 - S1;
 
-	return Result;
+	return *((uniform FVector3d *uniform)&Result);
 }
 
-inline uniform FVector operator*(const uniform FVector &A, const uniform float F)
+inline uniform FVector3f operator-(const uniform FVector3f &A, const uniform FVector3f &B)
 {
-	const uniform FVector FVec = {{F,F,F}};
-	return A * FVec;
-}
+	uniform FVector4f S0, S1, Result;
+	*((uniform FVector3f *uniform)&S0) = *((uniform FVector3f *uniform)&A);
+	*((uniform FVector3f *uniform)&S1) = *((uniform FVector3f *uniform)&B);
 
-inline uniform FVector operator*(const uniform float F, const uniform FVector &A)
-{
-	const uniform FVector FVec = {{F,F,F}};
-	return FVec * A;
+	Result = S0 - S1;
+
+	return *((uniform FVector3f *uniform)&Result);
 }
 
-inline FVector operator-(const FVector &A, const FVector &B)
+inline FVector3d operator-(const FVector3d &A, const uniform FVector3d &B)
 {
-	FVector Result;
+	FVector3d Result;
 
 	Result.V[0] = A.V[0] - B.V[0];
 	Result.V[1] = A.V[1] - B.V[1];
@@ -981,20 +1579,9 @@ inline FVector operator-(const FVector &A, const FVector &B)
 	return Result;
 }
 
-inline uniform FVector operator-(const uniform FVector &A, const uniform FVector &B)
-{
-	uniform FVector4 S0, S1, Result;
-	*((uniform FVector *uniform)&S0) = *((uniform FVector *uniform)&A);
-	*((uniform FVector *uniform)&S1) = *((uniform FVector *uniform)&B);
-
-	Result = S0 - S1;
-
-	return *((uniform FVector *uniform)&Result);
-}
-
-inline FVector operator-(const FVector &A, const uniform FVector &B)
+inline FVector3f operator-(const FVector3f &A, const uniform FVector3f &B)
 {
-	FVector Result;
+	FVector3f Result;
 
 	Result.V[0] = A.V[0] - B.V[0];
 	Result.V[1] = A.V[1] - B.V[1];
@@ -1014,7 +1601,7 @@ inline FVector operator-(const uniform FVector &A, const FVector &B)
 	return Result;
 }
 
-inline FVector operator-(const FVector &A, const float F)
+inline FVector operator-(const FVector &A, const FReal F)
 {
 	FVector Result;
 
@@ -1025,19 +1612,19 @@ inline FVector operator-(const FVector &A, const float F)
 	return Result;
 }
 
-inline uniform FVector operator-(const uniform FVector &A, const uniform float F)
+inline uniform FVector operator-(const uniform FVector &A, const uniform FReal F)
 {
 	const uniform FVector FVec = {{F,F,F}};
 	return A - FVec;
 }
 
-inline uniform FVector operator-(const uniform float F, const uniform FVector &A)
+inline uniform FVector operator-(const uniform FReal F, const uniform FVector &A)
 {
 	const uniform FVector FVec = {{F,F,F}};
 	return FVec - A;
 }
 
-inline FVector operator-(const float F, const FVector &A)
+inline FVector operator-(const FReal F, const FVector &A)
 {
 	FVector Result;
 
@@ -1081,7 +1668,7 @@ inline FVector operator/(const FVector &A, const uniform FVector &B)
 	return Result;
 }
 
-inline FVector operator/(const FVector &A, const float F)
+inline FVector operator/(const FVector &A, const FReal F)
 {
 	FVector Result;
 
@@ -1092,37 +1679,49 @@ inline FVector operator/(const FVector &A, const float F)
 	return Result;
 }
 
-inline uniform FVector operator/(const uniform FVector &A, const uniform float F)
+inline uniform FVector operator/(const uniform FVector &A, const uniform FReal F)
 {
 	const uniform FVector FVec = {{F,F,F}};
 	return A / FVec;
 }
 
-inline uniform FVector operator/(const uniform float F, const uniform FVector &A)
+inline uniform FVector operator/(const uniform FReal F, const uniform FVector &A)
 {
 	const uniform FVector FVec = {{F,F,F}};
 	return FVec / A;
 }
 
-inline FVector SetVector(const float X, const float Y, const float Z)
+inline FVector3d SetVector(const double X, const double Y, const double Z)
+{
+	const FVector3d Result = {{X, Y, Z}};
+	return Result;
+}
+
+inline FVector3f SetVector(const float X, const float Y, const float Z)
+{
+	const FVector3f Result = {{X, Y, Z}};
+	return Result;
+}
+
+inline uniform FVector3d SetVector(const uniform double X, const uniform double Y, const uniform double Z)
 {
-	const FVector Result = {{X, Y, Z}};
+	const uniform FVector3d Result = {{X, Y, Z}};
 	return Result;
 }
 
-inline uniform FVector SetVector(const uniform float X, const uniform float Y, const uniform float Z)
+inline uniform FVector3f SetVector(const uniform float X, const uniform float Y, const uniform float Z)
 {
-	const uniform FVector Result = {{X, Y, Z}};
+	const uniform FVector3f Result = {{X, Y, Z}};
 	return Result;
 }
 
-inline uniform FVector SetVector(const uniform float F)
+inline uniform FVector SetVector(const uniform FReal F)
 {
 	const uniform FVector Result = {{F, F, F}};
 	return Result;
 }
 
-inline FVector SetVector(const float F)
+inline FVector SetVector(const FReal F)
 {
 	const FVector Result = {{F, F, F}};
 	return Result;
@@ -1134,12 +1733,74 @@ inline FVector SetVector(const FVector4 &Vec)
 	return Result;
 }
 
-inline uniform FVector SetVector(const uniform FVector4 &Vec)
+inline uniform FVector3f SetVector(const uniform FVector4f &Vec)
+{
+	const uniform FVector3f Result = {{Vec.V[0], Vec.V[1], Vec.V[2]}};
+	return Result;
+}
+
+inline uniform FVector3d SetVector(const uniform FVector4d &Vec)
 {
-	const uniform FVector Result = {{Vec.V[0], Vec.V[1], Vec.V[2]}};
+	const uniform FVector3d Result = {{Vec.V[0], Vec.V[1], Vec.V[2]}};
 	return Result;
 }
 
+inline uniform FVector3d ConvertVector3fTo3d(const uniform FVector3f &V)
+{
+	varying float S0;
+	*((uniform FVector3f *uniform)&S0) = *((uniform FVector3f *uniform)&V);
+
+	const varying double Result = (double)S0;
+
+	return *((uniform FVector3d *uniform)&Result);
+}
+
+inline uniform FVector3f ConvertVector3dTo3f(const uniform FVector3d &V)
+{
+	varying double S0;
+	*((uniform FVector3d *uniform)&S0) = *((uniform FVector3d *uniform)&V);
+
+	const varying float Result = (float)S0;
+
+	return *((uniform FVector3f *uniform)&Result);
+}
+
+inline uniform FVector ConvertVector3dTo3Native(const uniform FVector3d &V)
+{
+#if UE_LARGE_WORLD_COORDINATES_DISABLED
+	return ConvertVector3dTo3f(V);
+#else
+	return V;
+#endif
+}
+
+inline uniform FVector ConvertVector3fTo3Native(const uniform FVector3f &V)
+{
+#if UE_LARGE_WORLD_COORDINATES_DISABLED
+	return V;
+#else
+	return ConvertVector3fTo3d(V);
+#endif
+}
+
+inline uniform FVector3d ConvertVector3NativeTo3d(const uniform FVector &V)
+{
+#if UE_LARGE_WORLD_COORDINATES_DISABLED
+	return ConvertVector3fTo3d(V);
+#else
+	return V;
+#endif
+}
+
+inline uniform FVector3f ConvertVector3NativeTo3f(const uniform FVector &V)
+{
+#if UE_LARGE_WORLD_COORDINATES_DISABLED
+	return V;
+#else
+	return ConvertVector3dTo3f(V);
+#endif
+}
+
 inline uniform FVector8 operator+(const uniform FVector8 &A, const uniform FVector8 &B)
 {
 	varying float S0, S1, Result;
@@ -1401,30 +2062,35 @@ inline FIntVector SetIntVector(const FVector &A)
 	return Result;
 }
 
-inline uniform bool VectorIsAnyNearlyZero(const uniform FVector &V, const uniform float Tolerance)
+inline uniform bool VectorIsAnyNearlyZero(const uniform FVector &V, const uniform FReal Tolerance)
 {
-	varying float S0 = FLT_MAX;
+	varying FReal S0 = FLT_MAX;
 
 	*((uniform FVector *uniform)&S0) = *((uniform FVector *uniform)&V);
 
 	return any(abs(S0) <= Tolerance);
 }
 
-inline uniform bool VectorIsAllNearlyZero(const uniform FVector &V, const uniform float Tolerance)
+inline uniform bool VectorIsAllNearlyZero(const uniform FVector &V, const uniform FReal Tolerance)
 {
-	varying float S0 = FLT_MIN;
+	varying FReal S0 = FLT_MIN;
 
 	*((uniform FVector *uniform)&S0) = *((uniform FVector *uniform)&V);
 
 	return all(abs(S0) <= Tolerance);
 }
 
-inline FVector VectorLerp(const FVector& A, const FVector& B, const float Alpha)
+inline FVector3d VectorLerp(const FVector3d& A, const FVector3d& B, const double Alpha)
+{
+	return (A + (B-A) * Alpha);
+}
+
+inline FVector3f VectorLerp(const FVector3f& A, const FVector3f& B, const float Alpha)
 {
 	return (A + (B-A) * Alpha);
 }
 
-inline uniform FVector VectorLerp(const uniform FVector& A, const uniform FVector& B, const uniform float Alpha)
+inline uniform FVector VectorLerp(const uniform FVector& A, const uniform FVector& B, const uniform FReal Alpha)
 {
 	return (A + (B-A) * Alpha);
 }
@@ -1503,7 +2169,7 @@ inline uniform FVector4 VectorMultiply(const uniform FVector4 &A, const uniform
 	return A * B;
 }
 
-inline uniform FVector4 VectorMultiply(const uniform FVector4 &A, const uniform float F)
+inline uniform FVector4 VectorMultiply(const uniform FVector4 &A, const uniform FReal F)
 {
 	return A * F;
 }
@@ -1559,7 +2225,7 @@ inline FVector4 VectorMultiplyAdd(const uniform FVector4& A, const FVector4 &B,
 	return Result;
 }
 
-inline FVector4 VectorMultiplyAdd(const float F, const FVector4 &B, const FVector4 &C)
+inline FVector4 VectorMultiplyAdd(const FReal F, const FVector4 &B, const FVector4 &C)
 {
 	FVector4 Result;
 
@@ -1573,7 +2239,7 @@ inline FVector4 VectorMultiplyAdd(const float F, const FVector4 &B, const FVecto
 
 inline uniform FVector4 VectorMultiplyAdd(const uniform FVector4& A, const uniform FVector4 &B, const uniform FVector4 &C)
 {
-	varying float S0, S1, S2, Result;
+	varying FReal S0, S1, S2, Result;
 	*((uniform FVector4 *uniform)&S0) = *((uniform FVector4 *uniform)&A);
 	*((uniform FVector4 *uniform)&S1) = *((uniform FVector4 *uniform)&B);
 	*((uniform FVector4 *uniform)&S2) = *((uniform FVector4 *uniform)&C);
@@ -1583,7 +2249,7 @@ inline uniform FVector4 VectorMultiplyAdd(const uniform FVector4& A, const unifo
 	return *((uniform FVector4 *uniform)&Result);
 }
 
-inline FVector4 VectorMultiplyAdd(const FVector4& A, const uniform float B, const FVector4 &C)
+inline FVector4 VectorMultiplyAdd(const FVector4& A, const uniform FReal B, const FVector4 &C)
 {
 	FVector4 Result;
 
@@ -1595,9 +2261,9 @@ inline FVector4 VectorMultiplyAdd(const FVector4& A, const uniform float B, cons
 	return Result;
 }
 
-inline uniform FVector4 VectorMultiplyAdd(const uniform FVector4& A, const uniform float S1, const uniform FVector4 &C)
+inline uniform FVector4 VectorMultiplyAdd(const uniform FVector4& A, const uniform FReal S1, const uniform FVector4 &C)
 {
-	varying float S0, S2, Result;
+	varying FReal S0, S2, Result;
 	*((uniform FVector4 *uniform)&S0) = *((uniform FVector4 *uniform)&A);
 	*((uniform FVector4 *uniform)&S2) = *((uniform FVector4 *uniform)&C);
 
@@ -1606,9 +2272,9 @@ inline uniform FVector4 VectorMultiplyAdd(const uniform FVector4& A, const unifo
 	return *((uniform FVector4 *uniform)&Result);
 }
 
-inline uniform FVector4 VectorMultiplyAdd(const uniform float S0, const uniform FVector4 &B, const uniform FVector4 &C)
+inline uniform FVector4 VectorMultiplyAdd(const uniform FReal S0, const uniform FVector4 &B, const uniform FVector4 &C)
 {
-	varying float S1, S2, Result;
+	varying FReal S1, S2, Result;
 	*((uniform FVector4 *uniform)&S1) = *((uniform FVector4 *uniform)&B);
 	*((uniform FVector4 *uniform)&S2) = *((uniform FVector4 *uniform)&C);
 
@@ -1643,7 +2309,7 @@ inline FVector4 VectorAbs(const FVector4& A)
 
 inline uniform FVector4 VectorAbs(const uniform FVector4& A)
 {
-	varying float S0, Result;
+	varying FReal S0, Result;
 	*((uniform FVector4 *uniform)&S0) = *((uniform FVector4 *uniform)&A);
 
 	Result = abs(S0);
@@ -1679,7 +2345,7 @@ inline bool IsVectorLessEqual(const FVector &A, const uniform FVector &B)
 inline uniform bool IsVectorGreaterEqual(const uniform FVector4 &A, const uniform FVector4 &B)
 {
 	varying bool b = true;
-	varying float S0, S1;
+	varying FReal S0, S1;
 	*((uniform FVector4 *uniform)&S0) = *((uniform FVector4 *uniform)&A);
 	*((uniform FVector4 *uniform)&S1) = *((uniform FVector4 *uniform)&B);
 
@@ -1693,14 +2359,26 @@ inline bool IsVectorGreaterEqual(const FVector &A, const uniform FVector &B)
 	return A.V[0] >= B.V[0] && A.V[1] >= B.V[1] && A.V[2] >= B.V[2];
 }
 
-inline FVector4 VectorCompareGE(const FVector4 &A, const uniform FVector4 &B)
+inline FVector4d VectorCompareGE(const FVector4d &A, const uniform FVector4d &B)
 {
-	FVector4 Result;
+	FVector4d Result;
+
+	Result.V[0] = select(A.V[0] >= B.V[0], doublebits(DOUBLE_NAN), doublebits(0));
+	Result.V[1] = select(A.V[1] >= B.V[1], doublebits(DOUBLE_NAN), doublebits(0));
+	Result.V[2] = select(A.V[2] >= B.V[2], doublebits(DOUBLE_NAN), doublebits(0));
+	Result.V[3] = select(A.V[3] >= B.V[3], doublebits(DOUBLE_NAN), doublebits(0));
+
+	return Result;
+}
+
+inline FVector4f VectorCompareGE(const FVector4f &A, const uniform FVector4f &B)
+{
+	FVector4f Result;
 
-	Result.V[0] = select(A.V[0] >= B.V[0], floatbits(0xFFFFFFFF), floatbits(0));
-	Result.V[1] = select(A.V[1] >= B.V[1], floatbits(0xFFFFFFFF), floatbits(0));
-	Result.V[2] = select(A.V[2] >= B.V[2], floatbits(0xFFFFFFFF), floatbits(0));
-	Result.V[3] = select(A.V[3] >= B.V[3], floatbits(0xFFFFFFFF), floatbits(0));
+	Result.V[0] = select(A.V[0] >= B.V[0], floatbits(FLOAT_NAN), floatbits(0));
+	Result.V[1] = select(A.V[1] >= B.V[1], floatbits(FLOAT_NAN), floatbits(0));
+	Result.V[2] = select(A.V[2] >= B.V[2], floatbits(FLOAT_NAN), floatbits(0));
+	Result.V[3] = select(A.V[3] >= B.V[3], floatbits(FLOAT_NAN), floatbits(0));
 
 	return Result;
 }
@@ -1709,31 +2387,31 @@ inline FVector4 VectorCompareGE(const uniform FVector4 &A, const FVector4 &B)
 {
 	FVector4 Result;
 
-	Result.V[0] = select(A.V[0] >= B.V[0], floatbits(0xFFFFFFFF), floatbits(0));
-	Result.V[1] = select(A.V[1] >= B.V[1], floatbits(0xFFFFFFFF), floatbits(0));
-	Result.V[2] = select(A.V[2] >= B.V[2], floatbits(0xFFFFFFFF), floatbits(0));
-	Result.V[3] = select(A.V[3] >= B.V[3], floatbits(0xFFFFFFFF), floatbits(0));
+	Result.V[0] = select(A.V[0] >= B.V[0], REAL_BITS(NAN), REAL_BITS(0));
+	Result.V[1] = select(A.V[1] >= B.V[1], REAL_BITS(NAN), REAL_BITS(0));
+	Result.V[2] = select(A.V[2] >= B.V[2], REAL_BITS(NAN), REAL_BITS(0));
+	Result.V[3] = select(A.V[3] >= B.V[3], REAL_BITS(NAN), REAL_BITS(0));
 
 	return Result;
 }
 
 inline uniform FVector4 VectorCompareGE(const uniform FVector4 &A, const uniform FVector4 &B)
 {
-	varying float S0, S1, Result;
+	varying FReal S0, S1, Result;
 	*((uniform FVector4 *uniform)&S0) = *((uniform FVector4 *uniform)&A);
 	*((uniform FVector4 *uniform)&S1) = *((uniform FVector4 *uniform)&B);
 
-	Result = S0 >= S1 ? floatbits(0xFFFFFFFF) : floatbits(0);
+	Result = S0 >= S1 ? REAL_BITS(NAN) : REAL_BITS(0);
 
 	return *((uniform FVector4 *uniform)&Result);
 }
 
-inline uniform FVector4 VectorCompareGE(const uniform float F, const uniform FVector4 &B)
+inline uniform FVector4 VectorCompareGE(const uniform FReal F, const uniform FVector4 &B)
 {
-	varying float S1, Result;
+	varying FReal S1, Result;
 	*((uniform FVector4 *uniform)&S1) = *((uniform FVector4 *uniform)&B);
 
-	Result = F >= S1 ? floatbits(0xFFFFFFFF) : floatbits(0);
+	Result = F >= S1 ? REAL_BITS(NAN) : REAL_BITS(0);
 
 	return *((uniform FVector4 *uniform)&Result);
 }
@@ -1749,18 +2427,13 @@ inline uniform FVector VectorCompareGE(const uniform FVector& A, const uniform F
 	return *((uniform FVector *uniform)&Result);
 }
 
-inline uniform float CompareGE(const uniform float &A, const uniform float &B)
-{
-	return A >= B ? floatbits(0xFFFFFFFF) : floatbits(0);
-}
-
 inline uniform FVector4 VectorCompareGT(const uniform FVector4 &A, const uniform FVector4 &B)
 {
-	varying float S0, S1, Result;
+	varying FReal S0, S1, Result;
 	*((uniform FVector4 *uniform)&S0) = *((uniform FVector4 *uniform)&A);
 	*((uniform FVector4 *uniform)&S1) = *((uniform FVector4 *uniform)&B);
 
-	Result = S0 > S1 ? floatbits(0xFFFFFFFF) : floatbits(0);
+	Result = S0 > S1 ? REAL_BITS(NAN) : REAL_BITS(0);
 
 	return *((uniform FVector4 *uniform)&Result);
 }
@@ -1790,7 +2463,7 @@ inline FVector4 VectorMax(const FVector4& V1, const FVector4& V2)
 
 inline uniform FVector4 VectorMax(const uniform FVector4& V1, const uniform FVector4& V2)
 {
-	varying float S0, S1, Result;
+	varying FReal S0, S1, Result;
 	*((uniform FVector4 *uniform)&S0) = *((uniform FVector4 *uniform)&V1);
 	*((uniform FVector4 *uniform)&S1) = *((uniform FVector4 *uniform)&V2);
 
@@ -1846,7 +2519,7 @@ inline FVector4 VectorMin(const FVector4& V1, const FVector4& V2)
 
 inline uniform FVector4 VectorMin(const uniform FVector4& V1, const uniform FVector4& V2)
 {
-	varying float S0, S1, Result;
+	varying FReal S0, S1, Result;
 	*((uniform FVector4 *uniform)&S0) = *((uniform FVector4 *uniform)&V1);
 	*((uniform FVector4 *uniform)&S1) = *((uniform FVector4 *uniform)&V2);
 
@@ -1888,22 +2561,22 @@ inline uniform FVector8 VectorMin(const uniform FVector8 &V1, const uniform FVec
 	return *((uniform FVector8 *uniform)&Result);
 }
 
-inline uniform float VectorSum(const uniform FVector4 &V)
+inline uniform FReal VectorSum(const uniform FVector4 &V)
 {
 	return V.V[0] + V.V[1] + V.V[2] + V.V[3];
 }
 
-inline float VectorSum(const FVector4 &V)
+inline FReal VectorSum(const FVector4 &V)
 {
 	return V.V[0] + V.V[1] + V.V[2] + V.V[3];
 }
 
-inline uniform float VectorSum(const uniform FVector &V)
+inline uniform FReal VectorSum(const uniform FVector &V)
 {
 	return V.V[0] + V.V[1] + V.V[2];
 }
 
-inline float VectorSum(const FVector &A)
+inline FReal VectorSum(const FVector &A)
 {
 	return A.V[0] + A.V[1] + A.V[2];
 }
@@ -1922,7 +2595,7 @@ inline FVector4 VectorReciprocal(const FVector4& V)
 
 inline uniform FVector4 VectorReciprocal(const uniform FVector4& V)
 {
-	varying float S0, Result;
+	varying FReal S0, Result;
 	*((uniform FVector4 *uniform)&S0) = *((uniform FVector4 *uniform)&V);
 
 	Result = rcp_fast(S0);
@@ -1990,9 +2663,9 @@ inline uniform FVector4 VectorReciprocalAccurate(const uniform FVector4& V)
 	return x2;
 }
 
-inline FVector4 VectorReciprocalSqrt(const FVector4& V)
+inline FVector4f VectorReciprocalSqrt(const FVector4f& V)
 {
-	FVector4 Result;
+	FVector4f Result;
 
 	Result.V[0] = rsqrt_fast(V.V[0]);
 	Result.V[1] = rsqrt_fast(V.V[1]);
@@ -2002,121 +2675,39 @@ inline FVector4 VectorReciprocalSqrt(const FVector4& V)
 	return Result;
 }
 
-inline uniform FVector4 VectorReciprocalSqrt(const uniform FVector4& V)
-{
-	varying float S0, Result;
-	*((uniform FVector4 *uniform)&S0) = *((uniform FVector4 *uniform)&V);
-
-	Result = rsqrt_fast(S0);
-
-	return *((uniform FVector4 *uniform)&Result);
-}
-
-inline float InvSqrt(const float F)
+inline FVector4d VectorReciprocalSqrt(const FVector4d& V)
 {
-	// Performs two passes of Newton-Raphson iteration on the hardware estimate
-	//    v^-0.5 = x
-	// => x^2 = v^-1
-	// => 1/(x^2) = v
-	// => F(x) = x^-2 - v
-	//    F'(x) = -2x^-3
-
-	//    x1 = x0 - F(x0)/F'(x0)
-	// => x1 = x0 + 0.5 * (x0^-2 - Vec) * x0^3
-	// => x1 = x0 + 0.5 * (x0 - Vec * x0^3)
-	// => x1 = x0 + x0 * (0.5 - 0.5 * Vec * x0^2)
-	//
-	// This final form has one more operation than the legacy factorization (X1 = 0.5*X0*(3-(Y*X0)*X0)
-	// but retains better accuracy (namely InvSqrt(1) = 1 exactly).
-
-	static const uniform float fOneHalf = 0.5f;
-	float X0, X1, X2, FOver2;
-
-	X0 = rsqrt_fast(F);	// 1/sqrt estimate (12 bits)
-	FOver2 = F * fOneHalf;
+	FVector4d Result;
 
-	// 1st Newton-Raphson iteration
-	X1 = X0 * X0;
-	X1 = fOneHalf - (FOver2 * X1);
-	X1 = X0 + (X0 * X1);
+	Result.V[0] = 1.0d / sqrt(V.V[0]);
+	Result.V[1] = 1.0d / sqrt(V.V[1]);
+	Result.V[2] = 1.0d / sqrt(V.V[2]);
+	Result.V[3] = 1.0d / sqrt(V.V[3]);
 
-	// 2nd Newton-Raphson iteration
-	X2 = X1 * X1;
-	X2 = fOneHalf - (FOver2 * X2);
-	X2 = X1 + (X1 * X2);
-
-	return X2;
+	return Result;
 }
 
-inline uniform float InvSqrt(const uniform float F)
+inline uniform FVector4f VectorReciprocalSqrt(const uniform FVector4f& V)
 {
-	// Performs two passes of Newton-Raphson iteration on the hardware estimate
-	//    v^-0.5 = x
-	// => x^2 = v^-1
-	// => 1/(x^2) = v
-	// => F(x) = x^-2 - v
-	//    F'(x) = -2x^-3
-
-	//    x1 = x0 - F(x0)/F'(x0)
-	// => x1 = x0 + 0.5 * (x0^-2 - Vec) * x0^3
-	// => x1 = x0 + 0.5 * (x0 - Vec * x0^3)
-	// => x1 = x0 + x0 * (0.5 - 0.5 * Vec * x0^2)
-	//
-	// This final form has one more operation than the legacy factorization (X1 = 0.5*X0*(3-(Y*X0)*X0)
-	// but retains better accuracy (namely InvSqrt(1) = 1 exactly).
-
-	static const uniform float fOneHalf = 0.5f;
-	uniform float X0, X1, X2, FOver2;
-
-	X0 = rsqrt_fast(F);	// 1/sqrt estimate (12 bits)
-	FOver2 = F * fOneHalf;
-
-	// 1st Newton-Raphson iteration
-	X1 = X0 * X0;
-	X1 = fOneHalf - (FOver2 * X1);
-	X1 = X0 + (X0 * X1);
+	varying float S0, Result;
+	*((uniform FVector4f *uniform)&S0) = *((uniform FVector4f *uniform)&V);
 
-	// 2nd Newton-Raphson iteration
-	X2 = X1 * X1;
-	X2 = fOneHalf - (FOver2 * X2);
-	X2 = X1 + (X1 * X2);
+	Result = rsqrt_fast(S0);
 
-	return X2;
+	return *((uniform FVector4f *uniform)&Result);
 }
 
-inline FVector4 VectorReciprocalSqrtAccurate(const FVector4& V)
+inline uniform FVector4d VectorReciprocalSqrt(const uniform FVector4d& V)
 {
-	// Perform two passes of Newton-Raphson iteration on the hardware estimate
-	//    v^-0.5 = x
-	// => x^2 = v^-1
-	// => 1/(x^2) = v
-	// => F(x) = x^-2 - v
-	//    F'(x) = -2x^-3
-
-	//    x1 = x0 - F(x0)/F'(x0)
-	// => x1 = x0 + 0.5 * (x0^-2 - Vec) * x0^3
-	// => x1 = x0 + 0.5 * (x0 - Vec * x0^3)
-	// => x1 = x0 + x0 * (0.5 - 0.5 * Vec * x0^2)
-
-	const FVector4 VecDivBy2 = VectorMultiply(V, FloatOneHalf);
-
-	// Initial estimate
-	const FVector4 x0 = VectorReciprocalSqrt(V);
+	varying double S0, Result;
+	*((uniform FVector4d *uniform)&S0) = *((uniform FVector4d *uniform)&V);
 
-	// First iteration
-	FVector4 x1 = VectorMultiply(x0, x0);
-	x1 = VectorSubtract(FloatOneHalf, VectorMultiply(VecDivBy2, x1));
-	x1 = VectorMultiplyAdd(x0, x1, x0);
-
-	// Second iteration
-	FVector4 x2 = VectorMultiply(x1, x1);
-	x2 = VectorSubtract(FloatOneHalf, VectorMultiply(VecDivBy2, x2));
-	x2 = VectorMultiplyAdd(x1, x2, x1);
+	Result = 1.0d / sqrt(S0);
 
-	return x2;
+	return *((uniform FVector4d *uniform)&Result);
 }
 
-inline uniform FVector4 VectorReciprocalSqrtAccurate(const uniform FVector4& V)
+inline FVector4f VectorReciprocalSqrtAccurate(const FVector4f& V)
 {
 	// Perform two passes of Newton-Raphson iteration on the hardware estimate
 	//    v^-0.5 = x
@@ -2130,35 +2721,37 @@ inline uniform FVector4 VectorReciprocalSqrtAccurate(const uniform FVector4& V)
 	// => x1 = x0 + 0.5 * (x0 - Vec * x0^3)
 	// => x1 = x0 + x0 * (0.5 - 0.5 * Vec * x0^2)
 
-	const uniform FVector4 VecDivBy2 = VectorMultiply(V, FloatOneHalf);
+	const FVector4f VecDivBy2 = V * FloatOneHalf;
 
 	// Initial estimate
-	const uniform FVector4 x0 = VectorReciprocalSqrt(V);
+	const FVector4f x0 = VectorReciprocalSqrt(V);
 
 	// First iteration
-	uniform FVector4 x1 = VectorMultiply(x0, x0);
-	x1 = VectorSubtract(FloatOneHalf, VectorMultiply(VecDivBy2, x1));
-	x1 = VectorMultiplyAdd(x0, x1, x0);
+	FVector4f x1 = x0 * x0;
+	x1 = FloatOneHalf - (VecDivBy2 * x1);
+	x1 = x0 * x1 + x0;
 
 	// Second iteration
-	uniform FVector4 x2 = VectorMultiply(x1, x1);
-	x2 = VectorSubtract(FloatOneHalf, VectorMultiply(VecDivBy2, x2));
-	x2 = VectorMultiplyAdd(x1, x2, x1);
+	FVector4f x2 = x1 * x1;
+	x2 = FloatOneHalf - (VecDivBy2 * x2);
+	x2 = x1 * x2 + x1;
 
 	return x2;
 }
 
-inline uniform FVector VectorReciprocalSqrtAccurate(const uniform FVector& V)
+inline FVector4d VectorReciprocalSqrtAccurate(const FVector4d& V)
 {
-	uniform FVector4 S0, Result;
-	*((uniform FVector *uniform)&S0) = *((uniform FVector *uniform)&V);
+	FVector4d Result;
 
-	Result = VectorReciprocalSqrtAccurate(S0);
+	Result.V[0] = 1.0d / sqrt(V.V[0]);
+	Result.V[1] = 1.0d / sqrt(V.V[1]);
+	Result.V[2] = 1.0d / sqrt(V.V[2]);
+	Result.V[3] = 1.0d / sqrt(V.V[3]);
 
-	return *((uniform FVector *uniform)&Result);
+	return Result;
 }
 
-inline uniform float ReciprocalSqrtAccurate(const uniform float V)
+inline uniform FVector4f VectorReciprocalSqrtAccurate(const uniform FVector4f& V)
 {
 	// Perform two passes of Newton-Raphson iteration on the hardware estimate
 	//    v^-0.5 = x
@@ -2172,54 +2765,42 @@ inline uniform float ReciprocalSqrtAccurate(const uniform float V)
 	// => x1 = x0 + 0.5 * (x0 - Vec * x0^3)
 	// => x1 = x0 + x0 * (0.5 - 0.5 * Vec * x0^2)
 
-	const uniform float DivBy2 = V * 0.5f;
+	const uniform FVector4f VecDivBy2 = V * FloatOneHalf;
 
 	// Initial estimate
-	const uniform float x0 = rsqrt_fast(V);
+	const uniform FVector4f x0 = VectorReciprocalSqrt(V);
 
 	// First iteration
-	uniform float x1 = x0 * x0;
-	x1 = 0.5f - (DivBy2 * x1);
+	uniform FVector4f x1 = x0 * x0;
+	x1 = FloatOneHalf - (VecDivBy2 * x1);
 	x1 = x0 * x1 + x0;
 
 	// Second iteration
-	uniform float x2 = x1 * x1;
-	x2 = 0.5f - (DivBy2 * x2);
+	uniform FVector4f x2 = x1 * x1;
+	x2 = FloatOneHalf - (VecDivBy2 * x2);
 	x2 = x1 * x2 + x1;
 
 	return x2;
 }
 
-inline float ReciprocalSqrtAccurate(const float V)
+inline uniform FVector4d VectorReciprocalSqrtAccurate(const uniform FVector4d& V)
 {
-	// Perform two passes of Newton-Raphson iteration on the hardware estimate
-	//    v^-0.5 = x
-	// => x^2 = v^-1
-	// => 1/(x^2) = v
-	// => F(x) = x^-2 - v
-	//    F'(x) = -2x^-3
-
-	//    x1 = x0 - F(x0)/F'(x0)
-	// => x1 = x0 + 0.5 * (x0^-2 - Vec) * x0^3
-	// => x1 = x0 + 0.5 * (x0 - Vec * x0^3)
-	// => x1 = x0 + x0 * (0.5 - 0.5 * Vec * x0^2)
+	varying double S0, Result;
+	*((uniform FVector4d *uniform)&S0) = *((uniform FVector4d *uniform)&V);
 
-	const float DivBy2 = V * 0.5f;
+	Result = 1.0d / sqrt(S0);
 
-	// Initial estimate
-	const float x0 = rsqrt_fast(V);
+	return *((uniform FVector4d *uniform)&Result);
+}
 
-	// First iteration
-	float x1 = x0 * x0;
-	x1 = 0.5f - (DivBy2 * x1);
-	x1 = x0 * x1 + x0;
+inline uniform FVector VectorReciprocalSqrtAccurate(const uniform FVector& V)
+{
+	uniform FVector4 S0, Result;
+	*((uniform FVector *uniform)&S0) = *((uniform FVector *uniform)&V);
 
-	// Second iteration
-	float x2 = x1 * x1;
-	x2 = 0.5f - (DivBy2 * x2);
-	x2 = x1 * x2 + x1;
+	Result = VectorReciprocalSqrtAccurate(S0);
 
-	return x2;
+	return *((uniform FVector *uniform)&Result);
 }
 
 /**
@@ -2254,22 +2835,42 @@ inline uniform FVector8 VectorReplicate(const uniform FVector8 &Vec, const unifo
  * @param W			Index for which component to use for W (literal 0-3)
  * @return			The swizzled vector
  */
-inline FVector4 VectorSwizzle(const FVector4 &Vec, const uniform int X, const uniform int Y, const uniform int Z, const uniform int W)
+inline FVector4d VectorSwizzle(const FVector4d &Vec, const uniform int X, const uniform int Y, const uniform int Z, const uniform int W)
+{
+	return SetVector4(Vec.V[X], Vec.V[Y], Vec.V[Z], Vec.V[W]);
+}
+
+inline FVector4f VectorSwizzle(const FVector4f &Vec, const uniform int X, const uniform int Y, const uniform int Z, const uniform int W)
 {
 	return SetVector4(Vec.V[X], Vec.V[Y], Vec.V[Z], Vec.V[W]);
 }
 
-inline uniform FVector4 VectorSwizzle(const uniform FVector4 &Vec, const uniform int X, const uniform int Y, const uniform int Z, const uniform int W)
+inline uniform FVector4d VectorSwizzle(const uniform FVector4d &Vec, const uniform int X, const uniform int Y, const uniform int Z, const uniform int W)
 {
 	return SetVector4(Vec.V[X], Vec.V[Y], Vec.V[Z], Vec.V[W]);
 }
 
-inline FVector VectorSwizzle(const FVector &Vec, const uniform int X, const uniform int Y, const uniform int Z)
+inline uniform FVector4f VectorSwizzle(const uniform FVector4f &Vec, const uniform int X, const uniform int Y, const uniform int Z, const uniform int W)
+{
+	return SetVector4(Vec.V[X], Vec.V[Y], Vec.V[Z], Vec.V[W]);
+}
+
+inline FVector3d VectorSwizzle(const FVector3d &Vec, const uniform int X, const uniform int Y, const uniform int Z)
+{
+	return SetVector(Vec.V[X], Vec.V[Y], Vec.V[Z]);
+}
+
+inline FVector3f VectorSwizzle(const FVector3f &Vec, const uniform int X, const uniform int Y, const uniform int Z)
 {
 	return SetVector(Vec.V[X], Vec.V[Y], Vec.V[Z]);
 }
 
-inline uniform FVector VectorSwizzle(const uniform FVector &Vec, const uniform int X, const uniform int Y, const uniform int Z)
+inline uniform FVector3d VectorSwizzle(const uniform FVector3d &Vec, const uniform int X, const uniform int Y, const uniform int Z)
+{
+	return SetVector(Vec.V[X], Vec.V[Y], Vec.V[Z]);
+}
+
+inline uniform FVector3f VectorSwizzle(const uniform FVector3f &Vec, const uniform int X, const uniform int Y, const uniform int Z)
 {
 	return SetVector(Vec.V[X], Vec.V[Y], Vec.V[Z]);
 }
@@ -2347,23 +2948,43 @@ inline FVector VectorCross(const FVector& Vec1, const FVector& Vec2)
 	return VectorSwizzle(Tmp2, 1,2,0);
 }
 
-inline FVector VectorCross(const uniform FVector& Vec1, const FVector& Vec2)
+inline FVector3d VectorCross(const uniform FVector3d& Vec1, const FVector3d& Vec2)
 {
-	FVector Tmp0 = VectorSwizzle(Vec2, 1,2,0);
-	const uniform FVector Tmp1 = VectorSwizzle(Vec1, 1,2,0);
+	FVector3d Tmp0 = VectorSwizzle(Vec2, 1,2,0);
+	const uniform FVector3d Tmp1 = VectorSwizzle(Vec1, 1,2,0);
+	Tmp0 = Tmp0 * Vec1;
+	const FVector3d Tmp3 = Tmp1 * Vec2;
+	const FVector3d Tmp2 = Tmp0 - Tmp3;
+	return VectorSwizzle(Tmp2, 1,2,0);
+}
+
+inline FVector3f VectorCross(const uniform FVector3f& Vec1, const FVector3f& Vec2)
+{
+	FVector3f Tmp0 = VectorSwizzle(Vec2, 1,2,0);
+	const uniform FVector3f Tmp1 = VectorSwizzle(Vec1, 1,2,0);
+	Tmp0 = Tmp0 * Vec1;
+	const FVector3f Tmp3 = Tmp1 * Vec2;
+	const FVector3f Tmp2 = Tmp0 - Tmp3;
+	return VectorSwizzle(Tmp2, 1,2,0);
+}
+
+inline uniform FVector3d VectorCross(const uniform FVector3d& Vec1, const uniform FVector3d& Vec2)
+{
+	uniform FVector3d Tmp0 = VectorSwizzle(Vec2, 1,2,0);
+	uniform FVector3d Tmp1 = VectorSwizzle(Vec1, 1,2,0);
 	Tmp0 = Tmp0 * Vec1;
-	const FVector Tmp3 = Tmp1 * Vec2;
-	const FVector Tmp2 = Tmp0 - Tmp3;
+	Tmp1 = Tmp1 * Vec2;
+	const uniform FVector3d Tmp2 = Tmp0 - Tmp1;
 	return VectorSwizzle(Tmp2, 1,2,0);
 }
 
-inline uniform FVector VectorCross(const uniform FVector& Vec1, const uniform FVector& Vec2)
+inline uniform FVector3f VectorCross(const uniform FVector3f& Vec1, const uniform FVector3f& Vec2)
 {
-	uniform FVector Tmp0 = VectorSwizzle(Vec2, 1,2,0);
-	uniform FVector Tmp1 = VectorSwizzle(Vec1, 1,2,0);
+	uniform FVector3f Tmp0 = VectorSwizzle(Vec2, 1,2,0);
+	uniform FVector3f Tmp1 = VectorSwizzle(Vec1, 1,2,0);
 	Tmp0 = Tmp0 * Vec1;
 	Tmp1 = Tmp1 * Vec2;
-	const uniform FVector Tmp2 = Tmp0 - Tmp1;
+	const uniform FVector3f Tmp2 = Tmp0 - Tmp1;
 	return VectorSwizzle(Tmp2, 1,2,0);
 }
 
@@ -2385,9 +3006,19 @@ inline uniform FVector8 VectorCross(const uniform FVector8& Vec1, const uniform
  * @param Vec2	2nd vector
  * @return		d = dot4(Vec1.xyzw, Vec2.xyzw), VectorRegister( d, d, d, d )
  */
-inline FVector4 VectorDot4( const FVector4& Vec1, const FVector4& Vec2 )
+inline FVector4d VectorDot4( const FVector4d& Vec1, const FVector4d& Vec2 )
 {
-	FVector4 Temp1, Temp2;
+	FVector4d Temp1, Temp2;
+	Temp1 = Vec1 * Vec2;
+	Temp2 = VectorSwizzle(Temp1, 2,3,0,1);	// (Z,W,X,Y).
+	Temp1 = Temp1 + Temp2; // (X*X + Z*Z, Y*Y + W*W, Z*Z + X*X, W*W + Y*Y)
+	Temp2 = VectorSwizzle(Temp1, 1,2,3,0); // Rotate left 4 bytes (Y,Z,W,X).
+	return Temp1 + Temp2; // (X*X + Z*Z + Y*Y + W*W, Y*Y + W*W + Z*Z + X*X, Z*Z + X*X + W*W + Y*Y, W*W + Y*Y + X*X + Z*Z)
+}
+
+inline FVector4f VectorDot4( const FVector4f& Vec1, const FVector4f& Vec2 )
+{
+	FVector4f Temp1, Temp2;
 	Temp1 = Vec1 * Vec2;
 	Temp2 = VectorSwizzle(Temp1, 2,3,0,1);	// (Z,W,X,Y).
 	Temp1 = Temp1 + Temp2; // (X*X + Z*Z, Y*Y + W*W, Z*Z + X*X, W*W + Y*Y)
@@ -2405,9 +3036,19 @@ inline uniform FVector4 VectorDot4( const uniform FVector4& Vec1, const uniform
 	return Temp1 + Temp2; // (X*X + Z*Z + Y*Y + W*W, Y*Y + W*W + Z*Z + X*X, Z*Z + X*X + W*W + Y*Y, W*W + Y*Y + X*X + Z*Z)
 }
 
-inline float VectorDot( const FVector4& Vec1, const FVector4& Vec2 )
+inline double VectorDot( const FVector4d& Vec1, const FVector4d& Vec2 )
+{
+	FVector4d Temp1, Temp2;
+	Temp1 = Vec1 * Vec2;
+	Temp2 = VectorSwizzle(Temp1, 2,3,0,1);	// (Z,W,X,Y).
+	Temp1 = Temp1 + Temp2; // (X*X + Z*Z, Y*Y + W*W, Z*Z + X*X, W*W + Y*Y)
+	Temp2 = VectorSwizzle(Temp1, 1,2,3,0); // Rotate left 4 bytes (Y,Z,W,X).
+	return Temp1.V[0] + Temp2.V[0]; // (X*X + Z*Z + Y*Y + W*W)
+}
+
+inline float VectorDot( const FVector4f& Vec1, const FVector4f& Vec2 )
 {
-	FVector4 Temp1, Temp2;
+	FVector4f Temp1, Temp2;
 	Temp1 = Vec1 * Vec2;
 	Temp2 = VectorSwizzle(Temp1, 2,3,0,1);	// (Z,W,X,Y).
 	Temp1 = Temp1 + Temp2; // (X*X + Z*Z, Y*Y + W*W, Z*Z + X*X, W*W + Y*Y)
@@ -2415,7 +3056,7 @@ inline float VectorDot( const FVector4& Vec1, const FVector4& Vec2 )
 	return Temp1.V[0] + Temp2.V[0]; // (X*X + Z*Z + Y*Y + W*W)
 }
 
-inline uniform float VectorDot( const uniform FVector4& Vec1, const uniform FVector4& Vec2 )
+inline uniform FReal VectorDot( const uniform FVector4& Vec1, const uniform FVector4& Vec2 )
 {
 	uniform FVector4 Temp1, Temp2;
 	Temp1 = Vec1 * Vec2;
@@ -2425,19 +3066,25 @@ inline uniform float VectorDot( const uniform FVector4& Vec1, const uniform FVec
 	return Temp1.V[0] + Temp2.V[0]; // (X*X + Z*Z + Y*Y + W*W)
 }
 
-inline uniform float VectorDot( const uniform FVector &Vec1, const uniform FVector &Vec2 )
+inline uniform FReal VectorDot( const uniform FVector &Vec1, const uniform FVector &Vec2 )
 {
 	const uniform FVector Temp = Vec1 * Vec2;
 	return Temp.V[0] + Temp.V[1] + Temp.V[2];
 }
 
-inline varying float VectorDot( const varying FVector &Vec1, const varying FVector &Vec2 )
+inline varying double VectorDot( const varying FVector3d &Vec1, const varying FVector3d &Vec2 )
 {
-	const varying FVector Temp = Vec1 * Vec2;
+	const varying FVector3d Temp = Vec1 * Vec2;
+	return Temp.V[0] + Temp.V[1] + Temp.V[2];
+}
+
+inline varying float VectorDot( const varying FVector3f &Vec1, const varying FVector3f &Vec2 )
+{
+	const varying FVector3f Temp = Vec1 * Vec2;
 	return Temp.V[0] + Temp.V[1] + Temp.V[2];
 }
 
-inline varying float VectorDot( const varying FVector &Vec1, const uniform FVector &Vec2 )
+inline varying FReal VectorDot( const varying FVector &Vec1, const uniform FVector &Vec2 )
 {
 	const varying FVector Temp = Vec1 * Vec2;
 	return Temp.V[0] + Temp.V[1] + Temp.V[2];
@@ -2445,12 +3092,11 @@ inline varying float VectorDot( const varying FVector &Vec1, const uniform FVect
 
 inline uniform FVector4 VectorBitwiseAnd(const uniform FVector4 &Vec1, const uniform FVector4 &Vec2)
 {
-	varying unsigned int S0, S1;
-	varying float Result;
+	varying FReal S0, S1, Result;
 	*((uniform FVector4 *uniform)&S0) = *((uniform FVector4 *uniform)&Vec1);
 	*((uniform FVector4 *uniform)&S1) = *((uniform FVector4 *uniform)&Vec2);
 
-	Result = floatbits(S0 & S1);
+	Result = REAL_BITS(intbits(S0) & intbits(S1));
 
 	return *((uniform FVector4 *uniform)&Result);
 }
@@ -2466,13 +3112,12 @@ inline uniform FVector VectorBitwiseAnd(const uniform FVector& A, const uniform
 	return *((uniform FVector *uniform)&Result);
 }
 
-inline uniform FVector4 VectorBitwiseAnd(const uniform FVector4 &Vec1, const uniform unsigned int I)
+inline uniform FVector4 VectorBitwiseAnd(const uniform FVector4 &Vec1, const uniform UIntType I)
 {
-	varying unsigned int S0;
-	varying float Result;
+	varying FReal S0, Result;
 	*((uniform FVector4 *uniform)&S0) = *((uniform FVector4 *uniform)&Vec1);
 
-	Result = floatbits(S0 & I);
+	Result = REAL_BITS(intbits(S0) & I);
 
 	return *((uniform FVector4 *uniform)&Result);
 }
@@ -2489,12 +3134,11 @@ inline uniform FVector VectorBitwiseAnd(const uniform FVector& A, const uniform
 
 inline uniform FVector4 VectorBitwiseOr(const uniform FVector4 &Vec1, const uniform FVector4 &Vec2)
 {
-	varying unsigned int S0, S1;
-	varying float Result;
+	varying FReal S0, S1, Result;
 	*((uniform FVector4 *uniform)&S0) = *((uniform FVector4 *uniform)&Vec1);
 	*((uniform FVector4 *uniform)&S1) = *((uniform FVector4 *uniform)&Vec2);
 
-	Result = floatbits(S0 | S1);
+	Result = REAL_BITS(intbits(S0) | intbits(S1));
 
 	return *((uniform FVector4 *uniform)&Result);
 }
@@ -2512,12 +3156,11 @@ inline uniform FVector VectorBitwiseOr(const uniform FVector& A, const uniform F
 
 inline uniform FVector4 VectorBitwiseXor(const uniform FVector4 &Vec1, const uniform FVector4 &Vec2)
 {
-	varying unsigned int S0, S1;
-	varying float Result;
+	varying FReal S0, S1, Result;
 	*((uniform FVector4 *uniform)&S0) = *((uniform FVector4 *uniform)&Vec1);
 	*((uniform FVector4 *uniform)&S1) = *((uniform FVector4 *uniform)&Vec2);
 
-	Result = floatbits(S0 ^ S1);
+	Result = REAL_BITS(intbits(S0) ^ intbits(S1));
 
 	return *((uniform FVector4 *uniform)&Result);
 }
@@ -2546,17 +3189,29 @@ inline FVector4 VectorSelect(const FVector4& Mask, const FVector4& Vec1, const F
 {
 	FVector4 Result;
 
-	Result.V[0] = floatbits(intbits(Vec2.V[0]) ^ (intbits(Mask.V[0]) & (intbits(Vec1.V[0]) ^ intbits(Vec2.V[0]))));
-	Result.V[1] = floatbits(intbits(Vec2.V[1]) ^ (intbits(Mask.V[1]) & (intbits(Vec1.V[1]) ^ intbits(Vec2.V[1]))));
-	Result.V[2] = floatbits(intbits(Vec2.V[2]) ^ (intbits(Mask.V[2]) & (intbits(Vec1.V[2]) ^ intbits(Vec2.V[2]))));
-	Result.V[3] = floatbits(intbits(Vec2.V[3]) ^ (intbits(Mask.V[3]) & (intbits(Vec1.V[3]) ^ intbits(Vec2.V[3]))));
+	Result.V[0] = REAL_BITS(intbits(Vec2.V[0]) ^ (intbits(Mask.V[0]) & (intbits(Vec1.V[0]) ^ intbits(Vec2.V[0]))));
+	Result.V[1] = REAL_BITS(intbits(Vec2.V[1]) ^ (intbits(Mask.V[1]) & (intbits(Vec1.V[1]) ^ intbits(Vec2.V[1]))));
+	Result.V[2] = REAL_BITS(intbits(Vec2.V[2]) ^ (intbits(Mask.V[2]) & (intbits(Vec1.V[2]) ^ intbits(Vec2.V[2]))));
+	Result.V[3] = REAL_BITS(intbits(Vec2.V[3]) ^ (intbits(Mask.V[3]) & (intbits(Vec1.V[3]) ^ intbits(Vec2.V[3]))));
 
 	return Result;
 }
 
-inline FVector4 VectorSelect(const FVector4& Mask, const FVector4& Vec1, const uniform FVector4& Vec2 )
+inline FVector4d VectorSelect(const FVector4d& Mask, const FVector4d& Vec1, const uniform FVector4d& Vec2 )
 {
-	FVector4 Result;
+	FVector4d Result;
+
+	Result.V[0] = doublebits(intbits(Vec2.V[0]) ^ (intbits(Mask.V[0]) & (intbits(Vec1.V[0]) ^ intbits(Vec2.V[0]))));
+	Result.V[1] = doublebits(intbits(Vec2.V[1]) ^ (intbits(Mask.V[1]) & (intbits(Vec1.V[1]) ^ intbits(Vec2.V[1]))));
+	Result.V[2] = doublebits(intbits(Vec2.V[2]) ^ (intbits(Mask.V[2]) & (intbits(Vec1.V[2]) ^ intbits(Vec2.V[2]))));
+	Result.V[3] = doublebits(intbits(Vec2.V[3]) ^ (intbits(Mask.V[3]) & (intbits(Vec1.V[3]) ^ intbits(Vec2.V[3]))));
+
+	return Result;
+}
+
+inline FVector4f VectorSelect(const FVector4f& Mask, const FVector4f& Vec1, const uniform FVector4f& Vec2 )
+{
+	FVector4f Result;
 
 	Result.V[0] = floatbits(intbits(Vec2.V[0]) ^ (intbits(Mask.V[0]) & (intbits(Vec1.V[0]) ^ intbits(Vec2.V[0]))));
 	Result.V[1] = floatbits(intbits(Vec2.V[1]) ^ (intbits(Mask.V[1]) & (intbits(Vec1.V[1]) ^ intbits(Vec2.V[1]))));
@@ -2570,10 +3225,10 @@ inline FVector4 VectorSelect(const FVector4& Mask, const uniform FVector4& Vec1,
 {
 	FVector4 Result;
 
-	Result.V[0] = floatbits(intbits(Vec2.V[0]) ^ (intbits(Mask.V[0]) & (intbits(Vec1.V[0]) ^ intbits(Vec2.V[0]))));
-	Result.V[1] = floatbits(intbits(Vec2.V[1]) ^ (intbits(Mask.V[1]) & (intbits(Vec1.V[1]) ^ intbits(Vec2.V[1]))));
-	Result.V[2] = floatbits(intbits(Vec2.V[2]) ^ (intbits(Mask.V[2]) & (intbits(Vec1.V[2]) ^ intbits(Vec2.V[2]))));
-	Result.V[3] = floatbits(intbits(Vec2.V[3]) ^ (intbits(Mask.V[3]) & (intbits(Vec1.V[3]) ^ intbits(Vec2.V[3]))));
+	Result.V[0] = REAL_BITS(intbits(Vec2.V[0]) ^ (intbits(Mask.V[0]) & (intbits(Vec1.V[0]) ^ intbits(Vec2.V[0]))));
+	Result.V[1] = REAL_BITS(intbits(Vec2.V[1]) ^ (intbits(Mask.V[1]) & (intbits(Vec1.V[1]) ^ intbits(Vec2.V[1]))));
+	Result.V[2] = REAL_BITS(intbits(Vec2.V[2]) ^ (intbits(Mask.V[2]) & (intbits(Vec1.V[2]) ^ intbits(Vec2.V[2]))));
+	Result.V[3] = REAL_BITS(intbits(Vec2.V[3]) ^ (intbits(Mask.V[3]) & (intbits(Vec1.V[3]) ^ intbits(Vec2.V[3]))));
 
 	return Result;
 }
@@ -2634,18 +3289,27 @@ inline uniform FVector VectorSelect(uniform bool Mask, const uniform FVector &Ve
 
 inline uniform FVector4 VectorSign(const uniform FVector4& X)
 {
-	const uniform FVector4 Mask = VectorCompareGE(X, (FloatZero));
-	return VectorSelect(Mask, FloatOne, FloatMinusOne);
+	const uniform FVector4 Mask = VectorCompareGE(X, (VectorZero));
+	return VectorSelect(Mask, VectorOne, VectorMinusOne);
 }
 
 // Returns ((Vector dot Vector) >= 1e-8) ? (Vector / |Vector|) : DefaultValue
 // Uses accurate 1/sqrt, not the estimate
-inline FVector4 VectorNormalizeSafe( const FVector4& Vector, const uniform FVector4& DefaultValue)
+inline FVector4d VectorNormalizeSafe( const FVector4d& Vector, const uniform FVector4d& DefaultValue)
+{
+	const FVector4d SquareSum = VectorDot4(Vector, Vector);
+	const FVector4d NonZeroMask = VectorCompareGE(SquareSum, DoubleSmallLengthThreshold);
+	const FVector4d InvLength = VectorReciprocalSqrtAccurate(SquareSum);
+	const FVector4d NormalizedVector = InvLength * Vector;
+	return VectorSelect(NonZeroMask, NormalizedVector, DefaultValue);
+}
+
+inline FVector4f VectorNormalizeSafe( const FVector4f& Vector, const uniform FVector4f& DefaultValue)
 {
-	const FVector4 SquareSum = VectorDot4(Vector, Vector);
-	const FVector4 NonZeroMask = VectorCompareGE(SquareSum, SmallLengthThreshold);
-	const FVector4 InvLength = VectorReciprocalSqrtAccurate(SquareSum);
-	const FVector4 NormalizedVector = InvLength * Vector;
+	const FVector4f SquareSum = VectorDot4(Vector, Vector);
+	const FVector4f NonZeroMask = VectorCompareGE(SquareSum, FloatSmallLengthThreshold);
+	const FVector4f InvLength = VectorReciprocalSqrtAccurate(SquareSum);
+	const FVector4f NormalizedVector = InvLength * Vector;
 	return VectorSelect(NonZeroMask, NormalizedVector, DefaultValue);
 }
 
@@ -2660,9 +3324,9 @@ inline uniform FVector4 VectorNormalizeSafe( const uniform FVector4& Vector, con
 
 inline uniform FVector4 VectorNormalizeSafe2( const uniform FVector4& Vector, const uniform FVector4& DefaultValue)
 {
-	const uniform float SquareSum = VectorDot(Vector, Vector);
-	const uniform float NonZeroMask = CompareGE(SquareSum, SMALL_NUMBER);
-	const uniform float InvLength = ReciprocalSqrtAccurate(SquareSum);
+	const uniform FReal SquareSum = VectorDot(Vector, Vector);
+	const uniform FReal NonZeroMask = VectorCompareGE(SquareSum, SMALL_NUMBER);
+	const uniform FReal InvLength = ReciprocalSqrtAccurate(SquareSum);
 	const uniform FVector4 NormalizedVector = Vector * InvLength;
 	const uniform FVector4 VNonZeroMask = SetVector4(NonZeroMask, NonZeroMask, NonZeroMask, NonZeroMask);
 	return VectorSelect(VNonZeroMask, NormalizedVector, DefaultValue);
@@ -2670,46 +3334,46 @@ inline uniform FVector4 VectorNormalizeSafe2( const uniform FVector4& Vector, co
 
 inline FVector VectorNormalize(const varying FVector& Vector)
 {
-	const varying float SquareSum = VectorDot(Vector, Vector);
+	const varying FReal SquareSum = VectorDot(Vector, Vector);
 	const varying bool bNonZeroMask = SquareSum > SMALL_NUMBER;
-	const varying float Scale = ReciprocalSqrtAccurate(SquareSum);
+	const varying FReal Scale = ReciprocalSqrtAccurate(SquareSum);
 	return VectorSelect(bNonZeroMask, Vector * Scale, Vector);
 }
 
 inline uniform FVector VectorNormalize(const uniform FVector& Vector)
 {
-	const uniform float SquareSum = VectorDot(Vector, Vector);
+	const uniform FReal SquareSum = VectorDot(Vector, Vector);
 	const uniform bool bNonZeroMask = SquareSum > SMALL_NUMBER;
-	const uniform float Scale = ReciprocalSqrtAccurate(SquareSum);
+	const uniform FReal Scale = ReciprocalSqrtAccurate(SquareSum);
 	return VectorSelect(bNonZeroMask, Vector * Scale, Vector);
 }
 
 inline uniform FVector4 VectorTruncate(const uniform FVector4 &X)
 {
-	varying float S0, Result;
+	varying FReal S0, Result;
 	*((uniform FVector4 *uniform)&S0) = *((uniform FVector4 *uniform)&X);
 
-	Result = (float)((int)S0);
+	Result = trunc(S0);
 
 	return *((uniform FVector4 *uniform)&Result);
 }
 
 inline uniform FVector4 VectorNegate(const uniform FVector4 &A)
 {
-	varying float S0, Result;
+	varying FReal S0, Result;
 	*((uniform FVector4 *uniform)&S0) = *((uniform FVector4 *uniform)&A);
 
-	Result = floatbits(intbits(S0) ^ SIGN_BIT);
+	Result = REAL_BITS(intbits(S0) ^ SIGN_BIT);
 
 	return *((uniform FVector4 *uniform)&Result);
 }
 
 inline uniform FVector VectorNegate(const uniform FVector &A)
 {
-	varying float S0, Result;
+	uniform FVector4 S0, Result;
 	*((uniform FVector *uniform)&S0) = *((uniform FVector *uniform)&A);
 
-	Result = floatbits(intbits(S0) ^ SIGN_BIT);
+	Result = VectorNegate(S0);
 
 	return *((uniform FVector *uniform)&Result);
 }
@@ -2718,7 +3382,7 @@ inline uniform FVector4 VectorMod(const uniform FVector4& X, const uniform FVect
 {
 	const uniform FVector4 Div = VectorDivide(X, Y);
 	// Floats where abs(f) >= 2^23 have no fractional portion, and larger values would overflow VectorTruncate.
-	const uniform FVector4 NoFractionMask = VectorCompareGE(VectorAbs(Div), FloatNonFractional);
+	const uniform FVector4 NoFractionMask = VectorCompareGE(VectorAbs(Div), VectorNonFractional);
 	const uniform FVector4 Temp = VectorSelect(NoFractionMask, Div, VectorTruncate(Div));
 	const uniform FVector4 Result = VectorSubtract(X, VectorMultiply(Y, Temp));
 	// Clamp to [-AbsY, AbsY] because of possible failures for very large numbers (>1e10) due to precision loss.
@@ -2740,26 +3404,31 @@ inline void VectorSinCos(uniform FVector4 &VSinAngles, uniform FVector4 &VCosAng
 	// Note the round(), not truncate(). In this case round() can round halfway cases using round-to-nearest-even OR round-to-nearest.
 
 	// Quotient = round(A/2pi)
-	uniform FVector4 Quotient = VectorMultiply(VAngles, OneOverTwoPi);
+	uniform FVector4 Quotient = VectorMultiply(VAngles, VectorOneOverTwoPi);
 	Quotient = VectorTruncate(Quotient); // round to nearest even is the default rounding mode but that's fine here.
 	// X = A - 2pi * Quotient
-	uniform FVector4 X = VectorSubtract(VAngles, VectorMultiply(TwoPi, Quotient));
+	uniform FVector4 X = VectorSubtract(VAngles, VectorMultiply(VectorTwoPi, Quotient));
 
 	// Map in [-pi/2,pi/2]
 	uniform FVector4 sign = VectorBitwiseAnd(X, SIGN_BIT);
-	const uniform FVector4 c = VectorBitwiseOr(Pi, sign);  // pi when x >= 0, -pi when x < 0
+	const uniform FVector4 c = VectorBitwiseOr(VectorPi, sign);  // pi when x >= 0, -pi when x < 0
 	const uniform FVector4 absx = VectorAbs(X);
 	const uniform FVector4 rflx = VectorSubtract(c, X);
-	const uniform FVector4 comp = VectorCompareGT(absx, PiByTwo);
+	const uniform FVector4 comp = VectorCompareGT(absx, VectorPiByTwo);
 	X = VectorSelect(comp, rflx, X);
-	sign = VectorSelect(comp, FloatMinusOne, FloatOne);
+	sign = VectorSelect(comp, VectorMinusOne, VectorOne);
 
 	const uniform FVector4 XSquared = VectorMultiply(X, X);
 
 	// 11-degree minimax approximation
 	//*ScalarSin = (((((-2.3889859e-08f * y2 + 2.7525562e-06f) * y2 - 0.00019840874f) * y2 + 0.0083333310f) * y2 - 0.16666667f) * y2 + 1.0f) * y;
-	static const uniform FVector4 SinCoeff0 = {{ 1.0f, -0.16666667f, 0.0083333310f, -0.00019840874f }};
-	static const uniform FVector4 SinCoeff1 = {{ 2.7525562e-06f, -2.3889859e-08f, /*unused*/ 0.f, /*unused*/ 0.f }};
+#if UE_LARGE_WORLD_COORDINATES_DISABLED
+	static const uniform FVector4f SinCoeff0 = {{ 1.0f, -0.16666667f, 0.0083333310f, -0.00019840874f }};
+	static const uniform FVector4f SinCoeff1 = {{ 2.7525562e-06f, -2.3889859e-08f, /*unused*/ 0.f, /*unused*/ 0.f }};
+#else
+	static const uniform FVector4d SinCoeff0 = {{ 1.0d, -0.16666667d, 0.0083333310d, -0.00019840874d }};
+	static const uniform FVector4d SinCoeff1 = {{ 2.7525562d-06, -2.3889859d-08, /*unused*/ 0.d, /*unused*/ 0.d }};
+#endif
 
 	uniform FVector4 S;
 	S = VectorReplicate(SinCoeff1, 1);
@@ -2772,8 +3441,13 @@ inline void VectorSinCos(uniform FVector4 &VSinAngles, uniform FVector4 &VCosAng
 
 	// 10-degree minimax approximation
 	//*ScalarCos = sign * (((((-2.6051615e-07f * y2 + 2.4760495e-05f) * y2 - 0.0013888378f) * y2 + 0.041666638f) * y2 - 0.5f) * y2 + 1.0f);
-	static const uniform FVector4 CosCoeff0 = {{ 1.0f, -0.5f, 0.041666638f, -0.0013888378f }};
-	static const uniform FVector4 CosCoeff1 = {{ 2.4760495e-05f, -2.6051615e-07f, /*unused*/ 0.f, /*unused*/ 0.f }};
+#if UE_LARGE_WORLD_COORDINATES_DISABLED
+	static const uniform FVector4f CosCoeff0 = {{ 1.0f, -0.5f, 0.041666638f, -0.0013888378f }};
+	static const uniform FVector4f CosCoeff1 = {{ 2.4760495e-05f, -2.6051615e-07f, /*unused*/ 0.f, /*unused*/ 0.f }};
+#else
+	static const uniform FVector4d CosCoeff0 = {{ 1.0d, -0.5d, 0.041666638d, -0.0013888378d }};
+	static const uniform FVector4d CosCoeff1 = {{ 2.4760495d-05, -2.6051615d-07, /*unused*/ 0.d, /*unused*/ 0.d }};
+#endif
 
 	uniform FVector4 C;
 	C = VectorReplicate(CosCoeff1, 1);
@@ -2785,51 +3459,51 @@ inline void VectorSinCos(uniform FVector4 &VSinAngles, uniform FVector4 &VCosAng
 	VCosAngles = VectorMultiply(C, sign);
 }
 
-inline FVector4 VectorLerp(const FVector4 &A, const FVector4 &B, const float Alpha)
+inline FVector4 VectorLerp(const FVector4 &A, const FVector4 &B, const FReal Alpha)
 {
 	const FVector4 Delta = VectorSubtract(B, A);
 	return VectorMultiplyAdd(Alpha, Delta, A);
 }
 
-inline uniform FVector4 VectorLerp(const uniform FVector4 &A, const uniform FVector4 &B, const uniform float Alpha)
+inline uniform FVector4 VectorLerp(const uniform FVector4 &A, const uniform FVector4 &B, const uniform FReal Alpha)
 {
 	const uniform FVector4 Delta = VectorSubtract(B, A);
 	return VectorMultiplyAdd(Alpha, Delta, A);
 }
 
-inline float VectorSizeSquared(const FVector4 &A)
+inline FReal VectorSizeSquared(const FVector4 &A)
 {
 	return VectorDot(A, A);
 }
 
-inline uniform float VectorSizeSquared(const uniform FVector4 &A)
+inline uniform FReal VectorSizeSquared(const uniform FVector4 &A)
 {
 	return VectorDot(A, A);
 }
 
-inline float VectorSizeSquared(const FVector &A)
+inline FReal VectorSizeSquared(const FVector &A)
 {
 	return VectorDot(A, A);
 }
 
-inline uniform float VectorSizeSquared(const uniform FVector &A)
+inline uniform FReal VectorSizeSquared(const uniform FVector &A)
 {
 	return VectorDot(A, A);
 }
 
-inline uniform float VectorSize(const uniform FVector& A)
+inline uniform FReal VectorSize(const uniform FVector& A)
 {
 	return sqrt(VectorSizeSquared(A));
 }
 
-inline float VectorSize(const varying FVector &A)
+inline FReal VectorSize(const varying FVector &A)
 {
 	return sqrt(VectorDot(A, A));
 }
 
 inline uniform FVector VectorGetSafeNormal(const uniform FVector &Vector)
 {
-	const uniform float SquareSum = VectorSizeSquared(Vector);
+	const uniform FReal SquareSum = VectorSizeSquared(Vector);
 
 	// Not sure if it's safe to add tolerance in there. Might introduce too many errors
 	if(SquareSum == 1.f)
@@ -2840,7 +3514,7 @@ inline uniform FVector VectorGetSafeNormal(const uniform FVector &Vector)
 	{
 		return ZeroVector;
 	}
-	const uniform float Scale = InvSqrt(SquareSum);
+	const uniform FReal Scale = InvSqrt(SquareSum);
 	return Vector * Scale;
 }
 
@@ -2889,7 +3563,7 @@ inline FVector VectorGetSafeNormal(const FVector &V)
 {
 	//We want N / ||N|| and to avoid inf
 	//So we want N / ||N|| < 1 / eps => N eps < ||N||, but this is clearly true for all eps < 1 and N > 0
-	float SizeSqr = VectorSizeSquared(V);
+	FReal SizeSqr = VectorSizeSquared(V);
 	return VectorSelect(SizeSqr <= FLT_MIN, ForwardVector, V / sqrt(SizeSqr));
 }
 
