FEM - The Linear Allocator for oversized TaskDelegates (which we currently don't have any but that could enable more use cases for the future)

diff --git a/Engine/Source/Runtime/Core/Public/Async/Fundamental/Task.h b/Engine/Source/Runtime/Core/Public/Async/Fundamental/Task.h
--- a/Engine/Source/Runtime/Core/Public/Async/Fundamental/Task.h
+++ b/Engine/Source/Runtime/Core/Public/Async/Fundamental/Task.h
@@ -429,7 +429,6 @@ namespace LowLevelTasks
 	inline void FTask::ExecuteTask()
 	{
 		checkSlow(Runnable.IsSet());
-		alignas(PLATFORM_CACHE_LINE_SIZE) 
 		FTaskDelegate LocalRunnable;
 		bool Canceled = Runnable.CallAndMove(LocalRunnable);
 		checkSlow(!Runnable.IsSet());
diff --git a/Engine/Source/Runtime/Core/Public/Async/Fundamental/TaskDelegate.h b/Engine/Source/Runtime/Core/Public/Async/Fundamental/TaskDelegate.h
--- a/Engine/Source/Runtime/Core/Public/Async/Fundamental/TaskDelegate.h
+++ b/Engine/Source/Runtime/Core/Public/Async/Fundamental/TaskDelegate.h
@@ -1,6 +1,7 @@
 // Copyright Epic Games, Inc. All Rights Reserved.
 
 #pragma once
+#include "Experimental/ConcurrentLinearAllocator.h"
 #include "Templates/IsInvocable.h"
 #include "Misc/ScopeExit.h"
 #include "Misc/Launder.h"
@@ -8,6 +9,17 @@
 
 namespace LowLevelTasks
 {
+	struct FLowLevelTasksBlockAllocationTag : FDefaultBlockAllocationTag
+	{
+		static constexpr uint32 BlockSize = 64 * 1024;
+		static constexpr bool AllowOversizedBlocks = true;
+		static constexpr bool RequiresAccurateSize = false;
+		static constexpr bool InlineBlockAllocation = true;
+		static constexpr const TCHAR* TagName = TEXT("LowLevelTasksLinear");
+
+		using Allocator = TBlockAllocationCache<BlockSize, FAlignedAllocator>;
+	};
+
 	namespace TTaskDelegate_Impl
 	{
 		template<typename ReturnType>
@@ -195,7 +207,8 @@ namespace LowLevelTasks
 				static_assert(TIsSame<ReturnType, decltype(Callable())>::Value, "TCallableType return type does not match");
 				static_assert(sizeof(TTaskDelegateImpl<TCallableType, true>) == sizeof(TTaskDelegateBase), "Size must match the Baseclass");
 				TCallableType** HeapPtr = reinterpret_cast<TCallableType**>(InlineData);
-				*HeapPtr = new TCallableType(Forward<CallableT>(Callable));
+				*HeapPtr = reinterpret_cast<TCallableType*>(TConcurrentLinearAllocator<FLowLevelTasksBlockAllocationTag>::template Malloc<alignof(TCallableType)>(sizeof(TCallableType)));
+				new (*HeapPtr) TCallableType(Forward<CallableT>(Callable));
 			}
 
 			inline void Move(TTaskDelegateBase& DstWrapper, void* DstData, void* SrcData, uint32 DestInlineSize) override
@@ -222,7 +235,10 @@ namespace LowLevelTasks
 			void Destroy(void* InlineData) override
 			{
 				TCallableType* HeapPtr = reinterpret_cast<TCallableType*>(*reinterpret_cast<void**>(InlineData));
-				delete HeapPtr;
+				// We need a typedef here because VC won't compile the destructor call below if ElementType itself has a member called ElementType
+				using DestructorType = TCallableType;
+				HeapPtr->DestructorType::~TCallableType();
+				TConcurrentLinearAllocator<FLowLevelTasksBlockAllocationTag>::Free(HeapPtr);
 			}
 
 			bool IsHeapAllocated() const override 
