Timecode: add helper function for determining drop frame-ness of frame rate strings

FTimecode is the primary place where we determine whether drop frame timecode format is
supported and/or should be enabled for particular frame rates. The only drop frame frame
rates that are supported are NTSC_30 at 29.97 FPS and NTSC_60 at 59.94 FPS, and the
only way to compactly represent both the rate and its drop frame-ness are with the strings
"29.97df" and "59.94df", respectively, so this change adds an
FTimecode::IsValidDropFormatTimecodeRate() function that returns true if the given string
is either of these.

diff --git a/Engine/Source/Runtime/Core/Public/Misc/Timecode.h b/Engine/Source/Runtime/Core/Public/Misc/Timecode.h
--- a/Engine/Source/Runtime/Core/Public/Misc/Timecode.h
+++ b/Engine/Source/Runtime/Core/Public/Misc/Timecode.h
@@ -283,6 +283,19 @@ public:
 	/** By default, should we generate a timecode in drop frame format when the frame rate does support it. */
 	static CORE_API bool UseDropFormatTimecodeByDefaultWhenSupported();
 
+	/**
+	 * Returns true if the given frame rate string represents a supported drop frame timecode frame rate, or false otherwise.
+	 * 
+	 * Drop frame timecode is only supported for NTSC_30 (29.97 FPS) or NTSC_60 (59.94 FPS) frame rates.
+	 */
+	static bool IsValidDropFormatTimecodeRate(const FString& InRateString)
+	{
+		const FString NTSC_30_DF = "29.97df";
+		const FString NTSC_60_DF = "59.94df";
+
+		return (InRateString == NTSC_30_DF || InRateString == NTSC_60_DF);
+	}
+
 	/**
 	 * Get the Qualified Timecode formatted in HH:MM:SS:FF or HH:MM:SS;FF depending on if this represents drop-frame timecode or not.
 	 * @param bForceSignDisplay - Forces the timecode to be prepended with a positive or negative sign.
