FEM - added missing OnLowLevelFree call

diff --git a/Engine/Source/Runtime/Core/Public/Experimental/ConcurrentLinearAllocator.h b/Engine/Source/Runtime/Core/Public/Experimental/ConcurrentLinearAllocator.h
--- a/Engine/Source/Runtime/Core/Public/Experimental/ConcurrentLinearAllocator.h
+++ b/Engine/Source/Runtime/Core/Public/Experimental/ConcurrentLinearAllocator.h
@@ -40,6 +40,7 @@ struct FOsAllocator
 
 	FORCEINLINE static void* Malloc(SIZE_T Size, uint32 Alignment)
 	{
+		//LLM_SCOPED_PAUSE_TRACKING(ELLMAllocType::System);
 		return FPlatformMemory::BinnedAllocFromOS(Size);
 	}
 
@@ -55,6 +56,7 @@ struct FAlignedAllocator
 
 	FORCEINLINE static void* Malloc(SIZE_T Size, uint32 Alignment)
 	{
+		//LLM_SCOPED_PAUSE_TRACKING(ELLMAllocType::FMalloc);
 #if DO_CHECK
 		void* Ret = AnsiMalloc(Size, Alignment);
 		check(IsAligned(Ret, Alignment));
@@ -380,6 +382,7 @@ public:
 		{
 			if constexpr (SupportsFastPath)
 			{
+				LLM_IF_ENABLED(FLowLevelMemTracker::Get().OnLowLevelFree(ELLMTracker::Default, Pointer));
 				FBlockHeader* Header = reinterpret_cast<FBlockHeader*>(AlignDown(Pointer, BlockAllocationTag::BlockSize));
 
 				// this deletes complete blocks when the last allocation is freed
