Added a few implementations of the ANSI overload of FExternalProfiler::StartScopedEvent. These both support ANSI natively and so we don't need to go through an expensive TCHAR conversion if that is what the build is configured to. This should make it viable to use the ANSI external profiler on all platforms.

diff --git a/Engine/Source/Runtime/Core/Private/ProfilingDebugging/Windows/PIXExternalProfiler.cpp b/Engine/Source/Runtime/Core/Private/ProfilingDebugging/Windows/PIXExternalProfiler.cpp
--- a/Engine/Source/Runtime/Core/Private/ProfilingDebugging/Windows/PIXExternalProfiler.cpp
+++ b/Engine/Source/Runtime/Core/Private/ProfilingDebugging/Windows/PIXExternalProfiler.cpp
@@ -83,6 +83,15 @@ public:
 		}
 	}
 
+	/** Starts a scoped event specific to the profiler. */
+	void StartScopedEvent(const struct FColor& Color, const ANSICHAR* Text) final
+	{
+		if (IsEnabled())
+		{
+			PIXBeginEvent(Color.DWColor(), Text);
+		}
+	}
+
 	/** Ends a scoped event specific to the profiler. */
 	void EndScopedEvent() final
 	{
diff --git a/Engine/Source/Runtime/Core/Private/ProfilingDebugging/Windows/SuperluminalExternalProfiler.cpp b/Engine/Source/Runtime/Core/Private/ProfilingDebugging/Windows/SuperluminalExternalProfiler.cpp
--- a/Engine/Source/Runtime/Core/Private/ProfilingDebugging/Windows/SuperluminalExternalProfiler.cpp
+++ b/Engine/Source/Runtime/Core/Private/ProfilingDebugging/Windows/SuperluminalExternalProfiler.cpp
@@ -57,6 +57,12 @@ public:
 		PerformanceAPI_BeginEvent_Wide(Text, nullptr, UPerformanceAPI_Pack(Color.R, Color.G, Color.B));
 	}
 
+	/** Starts a scoped event specific to the profiler. */
+	virtual void StartScopedEvent(const struct FColor& Color, const ANSICHAR* Text) override
+	{
+		PerformanceAPI_BeginEvent(Text, nullptr, UPerformanceAPI_Pack(Color.R, Color.G, Color.B));
+	}
+
 	/** Ends a scoped event specific to the profiler. */
 	virtual void EndScopedEvent() override
 	{
