FEM - include malloc.h with PLATFORM_USE_ANSI_POSIX_MALLOC because it has the definition for malloc_usable_size

diff --git a/Engine/Source/Runtime/Core/Private/HAL/MallocAnsi.cpp b/Engine/Source/Runtime/Core/Private/HAL/MallocAnsi.cpp
--- a/Engine/Source/Runtime/Core/Private/HAL/MallocAnsi.cpp
+++ b/Engine/Source/Runtime/Core/Private/HAL/MallocAnsi.cpp
@@ -9,9 +9,9 @@
 #include "Math/UnrealMathUtility.h"
 #include "Templates/AlignmentTemplates.h"
 
-#if PLATFORM_UNIX || PLATFORM_SWITCH
+#if PLATFORM_USE_ANSI_POSIX_MALLOC
 	#include <malloc.h>
-#endif // PLATFORM_UNIX || PLATFORM_SWITCH
+#endif
 
 #if PLATFORM_IOS
 	#include "mach/mach.h"
