FEM - remove unique pointer from parallel for for some small improvement of MVP

diff --git a/Engine/Source/Runtime/Core/Public/Async/ParallelFor.h b/Engine/Source/Runtime/Core/Public/Async/ParallelFor.h
--- a/Engine/Source/Runtime/Core/Public/Async/ParallelFor.h
+++ b/Engine/Source/Runtime/Core/Public/Async/ParallelFor.h
@@ -464,7 +464,7 @@ namespace ParallelForImpl
 		//shared data between tasks
 		struct alignas(PLATFORM_CACHE_LINE_SIZE) FParallelForData : public TConcurrentLinearObject<FParallelForData, FTaskGraphBlockAllocationTag>, public FThreadSafeRefCountedObject
 		{
-			FParallelForData(int32 InNum, int32 InBatchSize, int32 InNumBatches, const TArrayView<ContextType>& InContexts, const BodyType& InBody, FEventRef& InFinishedSignal)
+			FParallelForData(int32 InNum, int32 InBatchSize, int32 InNumBatches, int32 InNumWorkers, const TArrayView<ContextType>& InContexts, const BodyType& InBody, FEventRef& InFinishedSignal)
 				: Num(InNum)
 				, BatchSize(InBatchSize)
 				, NumBatches(InNumBatches)
@@ -473,6 +473,7 @@ namespace ParallelForImpl
 				, FinishedSignal(InFinishedSignal)
 			{
 				IncompleteBatches.store(NumBatches, std::memory_order_relaxed);
+				Tasks.AddDefaulted(InNumWorkers);
 			}
 			std::atomic_int BatchItem  { 0 };
 			std::atomic_int IncompleteBatches { 0 };
@@ -482,20 +483,22 @@ namespace ParallelForImpl
 			const TArrayView<ContextType>& Contexts;
 			const BodyType& Body;
 			FEventRef& FinishedSignal;
+			TArray<LowLevelTasks::FTask, TConcurrentLinearArrayAllocator<FTaskGraphBlockAllocationTag>> Tasks;
 		};
 		using FDataHandle = TRefCountPtr<FParallelForData>;
 
 		//each task has an executor.
-		class FParallelExecutor : public TConcurrentLinearObject<FParallelExecutor, FTaskGraphBlockAllocationTag>
+		class FParallelExecutor 
 		{
-			LowLevelTasks::FTask Task;
 			FDataHandle Data;
 			int32 WorkerIndex;
+			LowLevelTasks::ETaskPriority Priority;
 
 		public:
-			inline FParallelExecutor(FDataHandle&& InData, int32 InWorkerIndex) 
+			inline FParallelExecutor(FDataHandle&& InData, int32 InWorkerIndex, LowLevelTasks::ETaskPriority InPriority) 
 				: Data(MoveTemp(InData))
 				, WorkerIndex(InWorkerIndex)
+				, Priority(InPriority)
 			{
 			}
 
@@ -503,15 +506,24 @@ namespace ParallelForImpl
 			inline FParallelExecutor(FParallelExecutor&& Other) 
 				: Data(MoveTemp(Other.Data))
 				, WorkerIndex(Other.WorkerIndex)
+				, Priority(Other.Priority)
 			{
 			}
 
+			~FParallelExecutor()
+			{
+				if (Data.IsValid() && WorkerIndex >= 0)
+				{
+					FParallelExecutor::LaunchTask(MoveTemp(Data), WorkerIndex, Priority);
+				}
+			}
+
 			inline const FDataHandle& GetData() const
 			{
 				return Data;
 			}
 
-			inline bool Run(const bool bIsMaster = true, const bool bIsBackgroundPriority = false) noexcept
+			inline bool operator()(const bool bIsMaster = false) noexcept
 			{
 				FMemMark Mark(FMemStack::Get());
 				TRACE_CPUPROFILER_EVENT_SCOPE(ParallelFor);
@@ -519,6 +531,8 @@ namespace ParallelForImpl
 				auto Now = [] { return FTimespan::FromSeconds(FPlatformTime::Seconds()); };
 				FTimespan Start = FTimespan::MinValue();
 				FTimespan YieldingThreshold;
+				const bool bIsBackgroundPriority = !bIsMaster && (Priority >= LowLevelTasks::ETaskPriority::BackgroundNormal);
+
 				if (bIsBackgroundPriority)
 				{
 					Start = Now();
@@ -541,6 +555,7 @@ namespace ParallelForImpl
 					{
 						if (!bIsMaster)
 						{
+							WorkerIndex = -1;
 							return false;
 						}
 						BatchIndex = (NumBatches - 1);
@@ -561,10 +576,12 @@ namespace ParallelForImpl
 						{
 							Data->FinishedSignal->Trigger();
 						}
+						WorkerIndex = -1;
 						return true;
 					}
 					else if (EndIndex >= Num)
 					{
+						WorkerIndex = -1;
 						return false;
 					}
 					else if (!bIsBackgroundPriority)
@@ -575,40 +592,34 @@ namespace ParallelForImpl
 					auto PassedTime = [Start, &Now]() { return Now() - Start; };
 					if (PassedTime() > YieldingThreshold)
 					{
-						//abort and reschedule to give higher priority tasks a chance to run
-						FParallelExecutor::LaunchTask(Task.GetPriority(), MoveTemp(Data), WorkerIndex);
+						//abort and reschedule (in the destructor as WorkerIndex is larger_eq Zero) to give higher priority tasks a chance to run
 						return false;
 					}
 				}
 			}
 
-			static void LaunchTask(LowLevelTasks::ETaskPriority InPriority, FDataHandle&& InData, int32 InWorkerIndex)
+			static void LaunchTask(FDataHandle&& InData, int32 InWorkerIndex, LowLevelTasks::ETaskPriority InPriority)
 			{
-				using FTaskHandle = TUniquePtr<FParallelExecutor>;
-				FTaskHandle TaskHandle = MakeUnique<FParallelExecutor>(MoveTemp(InData), InWorkerIndex);
-				LowLevelTasks::FTask& Task = TaskHandle->Task;
-				Task.Init(TEXT("FParallelExecutor"), InPriority, [Priority = InPriority, TaskHandle = MoveTemp(TaskHandle)]()
-				{
-					TaskHandle->Run(false, Priority >= LowLevelTasks::ETaskPriority::BackgroundNormal);
-				});
+				LowLevelTasks::FTask& Task = InData->Tasks[InWorkerIndex];
+				Task.Init(TEXT("FParallelExecutor"), InPriority, FParallelExecutor(MoveTemp(InData), InWorkerIndex, InPriority));
 				verify(LowLevelTasks::TryLaunch(Task, LowLevelTasks::EQueuePreference::GlobalQueuePreference));
 			}
 		};
 
 		//launch all the worker tasks
 		FEventRef FinishedSignal { EEventMode::ManualReset };
-		FDataHandle Data = new FParallelForData(Num, BatchSize, NumBatches, Contexts, Body, FinishedSignal);
+		FDataHandle Data = new FParallelForData(Num, BatchSize, NumBatches, NumWorkers, Contexts, Body, FinishedSignal);
 		for (int32 Worker = 0; Worker < NumWorkers; Worker++)
 		{
-			FParallelExecutor::LaunchTask(Priority, FDataHandle(Data), Worker);
+			FParallelExecutor::LaunchTask(FDataHandle(Data), Worker, Priority);
 		}
 
 		// do the prework
 		CurrentThreadWorkToDoBeforeHelping();
 
 		//help with the parallel-for to prevent deadlocks
-		FParallelExecutor LocalExecutor(MoveTemp(Data), NumWorkers);
-		const bool bFinishedLast = LocalExecutor.Run();
+		FParallelExecutor LocalExecutor(MoveTemp(Data), NumWorkers, Priority);
+		const bool bFinishedLast = LocalExecutor(true);
 
 		if(!bFinishedLast)
 		{
