Adjust logging levels for managed storage wrapper to provide better information around what/when happens during init without being spammy.
Managed storage wrapper - in the case of platform API errors when doing a copy or move, check for the existence of files and don't assume that they are there or not there just because the API call failed.
review [REDACTED_HANDLE]

diff --git a/Engine/Source/Runtime/Core/Public/HAL/IPlatformFileManagedStorageWrapper.h b/Engine/Source/Runtime/Core/Public/HAL/IPlatformFileManagedStorageWrapper.h
--- a/Engine/Source/Runtime/Core/Public/HAL/IPlatformFileManagedStorageWrapper.h
+++ b/Engine/Source/Runtime/Core/Public/HAL/IPlatformFileManagedStorageWrapper.h
@@ -19,7 +19,7 @@
 #include "Algo/IndexOf.h"
 #include <atomic>
 
-DECLARE_LOG_CATEGORY_EXTERN(LogPlatformFileManagedStorage, Display, All);
+DECLARE_LOG_CATEGORY_EXTERN(LogPlatformFileManagedStorage, Log, All);
 
 namespace ManagedStorageInternal
 {
@@ -207,7 +207,7 @@ public:
 			FWriteScopeLock ScopeLock(FileSizesLock);
 			FileSizes.AddByHash(KeyHash, Filename, FileSize);
 
-			UE_LOG(LogPlatformFileManagedStorage, Log, TEXT("File %s is added to category %s"), *Filename, *CategoryName);
+			UE_LOG(LogPlatformFileManagedStorage, Verbose, TEXT("File %s is added to category %s"), *Filename, *CategoryName);
 		}
 
 		return Result;
@@ -224,7 +224,7 @@ public:
 
 			UsedQuota -= OldSize;
 
-			UE_LOG(LogPlatformFileManagedStorage, Log, TEXT("File %s is removed from category %s"), *Filename, *CategoryName);
+			UE_LOG(LogPlatformFileManagedStorage, Verbose, TEXT("File %s is removed from category %s"), *Filename, *CategoryName);
 
 			return true;
 		}
@@ -390,10 +390,12 @@ public:
 
 			for (const FString& RootDir : FPersistentStorageManager::Get().GetRootDirectories())
 			{
-				UE_LOG(LogPlatformFileManagedStorage, Log, TEXT("Scan directory %s"), *RootDir);
+				UE_LOG(LogPlatformFileManagedStorage, Display, TEXT("Scan directory %s"), *RootDir);
 
 				IPlatformFile::GetPlatformPhysical().IterateDirectoryRecursively(*RootDir, Visitor);
 			}
+
+			UE_LOG(LogPlatformFileManagedStorage, Display, TEXT("Done scanning root directories"));
 		});
 
 		bInitialized = true;
@@ -866,13 +868,13 @@ public:
 		FManagedStorageScopeFileLock ScopeFileLockTo(ManagedTo);
 		FManagedStorageScopeFileLock ScopeFileLockFrom(ManagedFrom);
 
-		int64 Size = this->FileSize(From);
-		if (Size < 0)
+		const int64 SizeFrom = this->FileSize(From);
+		if (SizeFrom < 0)
 		{
 			return false;
 		}
 
-		EPersistentStorageManagerFileSizeFlags Result = Manager.AddOrUpdateFile(ManagedTo, Size, EPersistentStorageManagerFileSizeFlags::OnlyUpdateIfLess | EPersistentStorageManagerFileSizeFlags::RespectQuota);
+		EPersistentStorageManagerFileSizeFlags Result = Manager.AddOrUpdateFile(ManagedTo, SizeFrom, EPersistentStorageManagerFileSizeFlags::OnlyUpdateIfLess | EPersistentStorageManagerFileSizeFlags::RespectQuota);
 		if (EnumHasAnyFlags(Result, EPersistentStorageManagerFileSizeFlags::RespectQuota))
 		{
 			UE_LOG(LogPlatformFileManagedStorage, Error, TEXT("Failed to move file to %s.  The target category of the destination has reach quota limit in peristent storage."), To);
@@ -883,11 +885,37 @@ public:
 		if (bSuccess)
 		{
 			Manager.RemoveFileFromManager(ManagedFrom);
-			Manager.AddOrUpdateFile(ManagedTo, Size);
+			Manager.AddOrUpdateFile(ManagedTo, SizeFrom);
 		}
 		else
 		{
-			Manager.RemoveFileFromManager(ManagedTo);
+			// On some implementations MoveFile can operate across volumes, so don't make assumptions about the state of
+			// of the file system in the case of failure.
+
+			if (ManagedFrom && !this->FileExists(From))
+			{
+				Manager.RemoveFileFromManager(ManagedFrom);
+			}
+
+			if (ManagedTo)
+			{
+				if (!this->FileExists(To))
+				{
+					Manager.RemoveFileFromManager(ManagedTo);
+				}
+				else
+				{
+					const int64 SizeTo = this->FileSize(To);
+					if (ensureAlways(SizeTo >= 0))
+					{
+						Manager.AddOrUpdateFile(ManagedTo, SizeTo);
+					}
+					else
+					{
+						Manager.RemoveFileFromManager(ManagedTo);
+					}
+				}
+			}
 		}
 
 		return bSuccess;
@@ -954,13 +982,13 @@ public:
 		FManagedStorageScopeFileLock ScopeFileLockTo(ManagedTo);
 		FManagedStorageScopeFileLock ScopeFileLockFrom(ManagedFrom);
 
-		int64 Size = this->FileSize(From);
-		if (Size < 0)
+		const int64 SizeFrom = this->FileSize(From);
+		if (SizeFrom < 0)
 		{
 			return false;
 		}
 
-		EPersistentStorageManagerFileSizeFlags Result = Manager.AddOrUpdateFile(ManagedTo, Size, EPersistentStorageManagerFileSizeFlags::OnlyUpdateIfLess | EPersistentStorageManagerFileSizeFlags::RespectQuota);
+		EPersistentStorageManagerFileSizeFlags Result = Manager.AddOrUpdateFile(ManagedTo, SizeFrom, EPersistentStorageManagerFileSizeFlags::OnlyUpdateIfLess | EPersistentStorageManagerFileSizeFlags::RespectQuota);
 		if (EnumHasAnyFlags(Result, EPersistentStorageManagerFileSizeFlags::RespectQuota))
 		{
 			UE_LOG(LogPlatformFileManagedStorage, Error, TEXT("Failed to copy file to %s.  The category of the destination has reach quota limit in peristent storage."), To);
@@ -970,9 +998,9 @@ public:
 		bool bSuccess = BaseClass::CopyFile(To, From, ReadFlags, WriteFlags);
 		if (bSuccess)
 		{
-			Manager.AddOrUpdateFile(ManagedTo, Size);
+			Manager.AddOrUpdateFile(ManagedTo, SizeFrom);
 		}
-		else
+		else if(ManagedTo)
 		{
 			if (!this->FileExists(To))
 			{
@@ -980,14 +1008,14 @@ public:
 			}
 			else
 			{
-				Size = this->FileSize(To);
-				if (Size < 0)
+				const int64 SizeTo = this->FileSize(To);
+				if (ensureAlways(SizeTo >= 0))
 				{
-					Manager.RemoveFileFromManager(ManagedTo);
+					Manager.AddOrUpdateFile(ManagedTo, SizeTo);
 				}
 				else
 				{
-					Manager.AddOrUpdateFile(ManagedTo, Size);
+					Manager.RemoveFileFromManager(ManagedTo);
 				}
 			}
 		}
