Fix VectorMatrixMultiply not existing if -nopch is used
#rnx

diff --git a/Engine/Source/Runtime/Core/Private/Math/UnrealMathSSE.cpp b/Engine/Source/Runtime/Core/Private/Math/UnrealMathSSE.cpp
--- a/Engine/Source/Runtime/Core/Private/Math/UnrealMathSSE.cpp
+++ b/Engine/Source/Runtime/Core/Private/Math/UnrealMathSSE.cpp
@@ -1,9 +1,10 @@
 // Copyright Epic Games, Inc. All Rights Reserved.
 
 // This file implements non-inlined functions from UnrealMathSSE.h
-#if defined(UE_USING_UNREALMATH_SSE) && (UE_USING_UNREALMATH_SSE != 0)
-
 #include "Math/UnrealMathSSE.h"
+
+#if (UE_USING_UNREALMATH_SSE != 0)
+
 #include "Math/Matrix.h"
 
 void VectorMatrixMultiply(FMatrix44f* Result, const FMatrix44f* Matrix1, const FMatrix44f* Matrix2)
diff --git a/Engine/Source/Runtime/Core/Public/Math/UnrealMathSSE.h b/Engine/Source/Runtime/Core/Public/Math/UnrealMathSSE.h
--- a/Engine/Source/Runtime/Core/Public/Math/UnrealMathSSE.h
+++ b/Engine/Source/Runtime/Core/Public/Math/UnrealMathSSE.h
@@ -2,6 +2,12 @@
 
 #pragma once
 
+#include "CoreTypes.h"
+
+#if PLATFORM_ENABLE_VECTORINTRINSICS
+
+struct VectorRegisterConstInit {};
+
 // Define to pull in UnrealMathSSE.cpp, since conditions for this file are complex in VectorRegister.h
 #define UE_USING_UNREALMATH_SSE 1
 
@@ -3376,3 +3382,9 @@ FORCEINLINE VectorRegister4Int VectorFloatToInt(const VectorRegister4Double& A)
 * @return		VectorRegister4Int(*Ptr, *Ptr, *Ptr, *Ptr)
 */
 #define VectorIntLoad1( Ptr )	_mm_set1_epi32(*(Ptr))
+
+#else
+
+#define UE_USING_UNREALMATH_SSE 0
+
+#endif
diff --git a/Engine/Source/Runtime/Core/Public/Math/UnrealMathVectorConstants.h b/Engine/Source/Runtime/Core/Public/Math/UnrealMathVectorConstants.h
--- a/Engine/Source/Runtime/Core/Public/Math/UnrealMathVectorConstants.h
+++ b/Engine/Source/Runtime/Core/Public/Math/UnrealMathVectorConstants.h
@@ -2,6 +2,8 @@
 
 #pragma once
 
+#include "Math/UnrealMathUtility.h"
+
 ///////////////////////////////////////////////////////////////////////////////////////////////////////////////////
 // constexpr floating point vector constant creation functions that bypass SIMD intrinsic setters
 //
diff --git a/Engine/Source/Runtime/Core/Public/Math/VectorRegister.h b/Engine/Source/Runtime/Core/Public/Math/VectorRegister.h
--- a/Engine/Source/Runtime/Core/Public/Math/VectorRegister.h
+++ b/Engine/Source/Runtime/Core/Public/Math/VectorRegister.h
@@ -8,18 +8,20 @@
 // If defined, allow double->float conversion in some VectorStore functions.
 #define SUPPORT_DOUBLE_TO_FLOAT_VECTOR_CONVERSION 1
 
-struct VectorRegisterConstInit {};
-
 // Platform specific vector intrinsics include.
 #if WITH_DIRECTXMATH
+struct VectorRegisterConstInit {};
 #include "Math/UnrealMathDirectX.h"
 #elif PLATFORM_ENABLE_VECTORINTRINSICS_NEON
+struct VectorRegisterConstInit {};
 #include "Math/UnrealMathNeon.h"
 #elif defined(__cplusplus_cli) && !PLATFORM_HOLOLENS
+struct VectorRegisterConstInit {};
 #include "Math/UnrealMathFPU.h" // there are compile issues with UnrealMathSSE in managed mode, so use the FPU version
 #elif PLATFORM_ENABLE_VECTORINTRINSICS
 #include "Math/UnrealMathSSE.h"
 #else
+struct VectorRegisterConstInit {};
 #include "Math/UnrealMathFPU.h"
 #endif
 
