Non-SSE platforms compile fix
#rnx

diff --git a/Engine/Source/Runtime/Core/Private/Math/UnrealMathSSE.cpp b/Engine/Source/Runtime/Core/Private/Math/UnrealMathSSE.cpp
--- a/Engine/Source/Runtime/Core/Private/Math/UnrealMathSSE.cpp
+++ b/Engine/Source/Runtime/Core/Private/Math/UnrealMathSSE.cpp
@@ -1,9 +1,11 @@
 // Copyright Epic Games, Inc. All Rights Reserved.
 
 // This file implements non-inlined functions from UnrealMathSSE.h
+#if WITH_DIRECTXMATH && !PLATFORM_ENABLE_VECTORINTRINSICS_NEON && !(defined(__cplusplus_cli) && !PLATFORM_HOLOLENS) && PLATFORM_ENABLE_VECTORINTRINSICS
 #include "Math/UnrealMathSSE.h"
+#endif
 
-#if (UE_USING_UNREALMATH_SSE != 0)
+#if defined(UE_USING_UNREALMATH_SSE) && (UE_USING_UNREALMATH_SSE != 0)
 
 #include "Math/Matrix.h"
 
