Added a projection to Algo::Compare

Projections are available with Algo::CompareBy, and Algo::CompareByPredicate has been renamed to Algo::Compare to reduce confusion.

diff --git a/Engine/Source/Runtime/Core/Private/Tests/String/ParseLinesTest.cpp b/Engine/Source/Runtime/Core/Private/Tests/String/ParseLinesTest.cpp
--- a/Engine/Source/Runtime/Core/Private/Tests/String/ParseLinesTest.cpp
+++ b/Engine/Source/Runtime/Core/Private/Tests/String/ParseLinesTest.cpp
@@ -19,7 +19,7 @@ bool StringParseLinesTest::RunTest(const FString& Parameters)
 	{
 		TArray<FStringView, TInlineAllocator<8>> ResultLines;
 		UE::String::ParseLines(View, ResultLines, Options);
-		if (!Algo::CompareByPredicate(ResultLines, ExpectedLines, TEqualTo<>()))
+		if (!Algo::Compare(ResultLines, ExpectedLines))
 		{
 			TStringBuilder<512> Error;
 			Error << TEXTVIEW("UE::String::ParseLines failed to parse \"") << FString(View).ReplaceCharWithEscapedChar() << TEXTVIEW("\" result {");
diff --git a/Engine/Source/Runtime/Core/Private/Tests/String/ParseTokensTest.cpp b/Engine/Source/Runtime/Core/Private/Tests/String/ParseTokensTest.cpp
--- a/Engine/Source/Runtime/Core/Private/Tests/String/ParseTokensTest.cpp
+++ b/Engine/Source/Runtime/Core/Private/Tests/String/ParseTokensTest.cpp
@@ -29,7 +29,7 @@ bool StringParseTokensByStringTest::RunTest(const FString& Parameters)
 		{
 			UE::String::ParseTokensMultiple(View, Delimiters, ResultTokens, Options);
 		}
-		if (!Algo::CompareByPredicate(ResultTokens, ExpectedTokens, TEqualTo<>()))
+		if (!Algo::Compare(ResultTokens, ExpectedTokens))
 		{
 			TStringBuilder<512> Error;
 			Error << TEXTVIEW("UE::String::ParseTokens[Multiple] failed to parse \"") << View << TEXTVIEW("\" with delimiters {");
@@ -110,7 +110,7 @@ bool StringParseTokensByCharTest::RunTest(const FString& Parameters)
 		{
 			UE::String::ParseTokensMultiple(View, Delimiters, ResultTokens, Options);
 		}
-		if (!Algo::CompareByPredicate(ResultTokens, ExpectedTokens, TEqualTo<>()))
+		if (!Algo::Compare(ResultTokens, ExpectedTokens))
 		{
 			TStringBuilder<512> Error;
 			Error << TEXTVIEW("UE::String::ParseTokens[Multiple] failed to parse \"") << View << TEXTVIEW("\" with delimiters {");
diff --git a/Engine/Source/Runtime/Core/Public/Algo/Compare.h b/Engine/Source/Runtime/Core/Public/Algo/Compare.h
--- a/Engine/Source/Runtime/Core/Public/Algo/Compare.h
+++ b/Engine/Source/Runtime/Core/Public/Algo/Compare.h
@@ -3,20 +3,15 @@
 #pragma once
 
 #include "CoreTypes.h"
+#include "Templates/EqualTo.h"
+#include "Templates/IdentityFunctor.h"
+#include "Templates/Invoke.h"
 #include "Templates/UnrealTemplate.h"
 
-namespace Algo
+namespace Algo::Private
 {
-	/**
-	* Compares two contiguous containers using operator== to compare pairs of elements.
-	*
-	* @param  InputA     Container of elements that are used as the first argument to operator==.
-	* @param  InputB     Container of elements that are used as the second argument to operator==.
-	*
-	* @return Whether the containers are the same size and operator== returned true for every pair of elements.
-	*/
-	template <typename InAT, typename InBT>
-	constexpr bool Compare(InAT&& InputA, InBT&& InputB)
+	template <typename InAT, typename InBT, typename ProjectionT, typename PredicateT>
+	constexpr bool Compare(InAT&& InputA, InBT&& InputB, ProjectionT Projection, PredicateT Predicate)
 	{
 		const SIZE_T SizeA = GetNum(InputA);
 		const SIZE_T SizeB = GetNum(InputB);
@@ -31,7 +26,7 @@ namespace Algo
 
 		for (SIZE_T Count = SizeA; Count; --Count)
 		{
-			if (!(*A++ == *B++))
+			if (!Invoke(Predicate, Invoke(Projection, *A++), Invoke(Projection, *B++)))
 			{
 				return false;
 			}
@@ -39,38 +34,74 @@ namespace Algo
 
 		return true;
 	}
+}
 
+namespace Algo
+{
 	/**
-	* Compares two contiguous containers using a predicate to compare pairs of elements.
-	*
-	* @param  InputA     Container of elements that are used as the first argument to the predicate.
-	* @param  InputB     Container of elements that are used as the second argument to the predicate.
-	* @param  Predicate  Condition which returns true for elements which are deemed equal.
-	*
-	* @return Whether the containers are the same size and the predicate returned true for every pair of elements.
-	*/
-	template <typename InAT, typename InBT, typename PredicateT>
-	constexpr bool CompareByPredicate(InAT&& InputA, InBT&& InputB, PredicateT Predicate)
+	 * Compares two contiguous containers using operator== to compare pairs of elements.
+	 *
+	 * @param  InputA     Container of elements that are used as the first argument to operator==.
+	 * @param  InputB     Container of elements that are used as the second argument to operator==.
+	 *
+	 * @return Whether the containers are the same size and operator== returned true for every pair of elements.
+	 */
+	template <typename InAT, typename InBT>
+	constexpr bool Compare(InAT&& InputA, InBT&& InputB)
 	{
-		const SIZE_T SizeA = GetNum(InputA);
-		const SIZE_T SizeB = GetNum(InputB);
+		return Private::Compare(Forward<InAT>(InputA), Forward<InBT>(InputB), FIdentityFunctor(), TEqualTo<>());
+	}
 
-		if (SizeA != SizeB)
-		{
-			return false;
-		}
+	/**
+	 * Compares two contiguous containers using a predicate to compare pairs of elements.
+	 *
+	 * @param  InputA     Container of elements that are used as the first argument to the predicate.
+	 * @param  InputB     Container of elements that are used as the second argument to the predicate.
+	 * @param  Predicate  Condition which returns true for elements which are deemed equal.
+	 *
+	 * @return Whether the containers are the same size and the predicate returned true for every pair of elements.
+	 */
+	template <typename InAT, typename InBT, typename PredicateT>
+	constexpr bool Compare(InAT&& InputA, InBT&& InputB, PredicateT Predicate)
+	{
+		return Private::Compare(Forward<InAT>(InputA), Forward<InBT>(InputB), FIdentityFunctor(), MoveTemp(Predicate));
+	}
 
-		auto* A = GetData(InputA);
-		auto* B = GetData(InputB);
+	template <typename InAT, typename InBT, typename PredicateT>
+	UE_DEPRECATED(5.0, "CompareByPredicate has been renamed to Compare.")
+	constexpr bool CompareByPredicate(InAT&& InputA, InBT&& InputB, PredicateT Predicate)
+	{
+		return Compare(Forward<InAT>(InputA), Forward<InBT>(InputB), MoveTemp(Predicate));
+	}
 
-		for (SIZE_T Count = SizeA; Count; --Count)
-		{
-			if (!Invoke(Predicate, *A++, *B++))
-			{
-				return false;
-			}
-		}
+	/**
+	 * Compares two contiguous containers using operator== to compare pairs of projected elements.
+	 *
+	 * @param  InputA     Container of elements that are used as the first argument to operator==.
+	 * @param  InputB     Container of elements that are used as the second argument to operator==.
+	 * @param  Projection Projection to apply to the elements before comparing them.
+	 *
+	 * @return Whether the containers are the same size and operator== returned true for every pair of elements.
+	 */
+	template <typename InAT, typename InBT, typename ProjectionT>
+	constexpr bool CompareBy(InAT&& InputA, InBT&& InputB, ProjectionT Projection)
+	{
+		return Private::Compare(Forward<InAT>(InputA), Forward<InBT>(InputB), MoveTemp(Projection), TEqualTo<>());
+	}
 
-		return true;
+	/**
+	 * Compares two contiguous containers using a predicate to compare pairs of projected elements.
+	 *
+	 * @param  InputA     Container of elements that are used as the first argument to the predicate.
+	 * @param  InputB     Container of elements that are used as the second argument to the predicate.
+	 * @param  Projection Projection to apply to the elements before comparing them.
+	 * @param  Predicate  Condition which returns true for elements which are deemed equal.
+	 *
+	 * @return Whether the containers are the same size and the predicate returned true for every pair of elements.
+	 */
+	template <typename InAT, typename InBT, typename ProjectionT, typename PredicateT>
+	constexpr bool CompareBy(InAT&& InputA, InBT&& InputB, ProjectionT Projection, PredicateT Predicate)
+	{
+		return Private::Compare(Forward<InAT>(InputA), Forward<InBT>(InputB), MoveTemp(Projection), MoveTemp(Predicate));
 	}
 }
