Final(?) fix for -nopch linker issues but also not breaking other platforms
#rnx

diff --git a/Engine/Source/Runtime/Core/Private/Math/UnrealMathSSE.cpp b/Engine/Source/Runtime/Core/Private/Math/UnrealMathSSE.cpp
--- a/Engine/Source/Runtime/Core/Private/Math/UnrealMathSSE.cpp
+++ b/Engine/Source/Runtime/Core/Private/Math/UnrealMathSSE.cpp
@@ -1,11 +1,9 @@
 // Copyright Epic Games, Inc. All Rights Reserved.
 
 // This file implements non-inlined functions from UnrealMathSSE.h
-#if WITH_DIRECTXMATH && !PLATFORM_ENABLE_VECTORINTRINSICS_NEON && !(defined(__cplusplus_cli) && !PLATFORM_HOLOLENS) && PLATFORM_ENABLE_VECTORINTRINSICS
 #include "Math/UnrealMathSSE.h"
-#endif
 
-#if defined(UE_USING_UNREALMATH_SSE) && (UE_USING_UNREALMATH_SSE != 0)
+#if (UE_USING_UNREALMATH_SSE != 0)
 
 #include "Math/Matrix.h"
 
diff --git a/Engine/Source/Runtime/Core/Public/Math/UnrealMathSSE.h b/Engine/Source/Runtime/Core/Public/Math/UnrealMathSSE.h
--- a/Engine/Source/Runtime/Core/Public/Math/UnrealMathSSE.h
+++ b/Engine/Source/Runtime/Core/Public/Math/UnrealMathSSE.h
@@ -4,7 +4,7 @@
 
 #include "CoreTypes.h"
 
-#if PLATFORM_ENABLE_VECTORINTRINSICS
+#if !WITH_DIRECTXMATH && !PLATFORM_ENABLE_VECTORINTRINSICS_NEON && !(defined(__cplusplus_cli) && !PLATFORM_HOLOLENS) && PLATFORM_ENABLE_VECTORINTRINSICS
 
 struct VectorRegisterConstInit {};
 
