Trace - fix issue tracing to file in root drive directory

diff --git a/Engine/Source/Runtime/Core/Private/ProfilingDebugging/TraceAuxiliary.cpp b/Engine/Source/Runtime/Core/Private/ProfilingDebugging/TraceAuxiliary.cpp
--- a/Engine/Source/Runtime/Core/Private/ProfilingDebugging/TraceAuxiliary.cpp
+++ b/Engine/Source/Runtime/Core/Private/ProfilingDebugging/TraceAuxiliary.cpp
@@ -442,10 +442,13 @@ bool FTraceAuxiliaryImpl::WriteToFile(const TCHAR* InPath, const FTraceAuxiliary
 
 	// Ensure we can write the trace file appropriately
 	const FString WriteDir = FPaths::GetPath(NativePath);
-	if (!FileManager.MakeDirectory(*WriteDir, true))
+	if (!FPaths::IsDrive(WriteDir))
 	{
-		UE_LOG_REF(LogCategory, Warning, TEXT("Failed to create directory '%s'"), *WriteDir);
-		return false;
+		if (!FileManager.MakeDirectory(*WriteDir, true))
+		{
+			UE_LOG_REF(LogCategory, Warning, TEXT("Failed to create directory '%s'"), *WriteDir);
+			return false;
+		}
 	}
 
 	if (!bTruncateFile && FileManager.FileExists(*NativePath))
