Support for non-default constructable types in TPromise

diff --git a/Engine/Source/Editor/Experimental/HordeExecutor/Private/ContentAddressableStorage.cpp b/Engine/Source/Editor/Experimental/HordeExecutor/Private/ContentAddressableStorage.cpp
--- a/Engine/Source/Editor/Experimental/HordeExecutor/Private/ContentAddressableStorage.cpp
+++ b/Engine/Source/Editor/Experimental/HordeExecutor/Private/ContentAddressableStorage.cpp
@@ -168,7 +168,7 @@ namespace UE::RemoteExecution
 
 		TSharedPtr<TPromise<EStatusCode>> ReturnPromise = MakeShared<TPromise<EStatusCode>>();
 		Request->OnProcessRequestComplete().BindLambda([ReturnPromise](FHttpRequestPtr Req, FHttpResponsePtr HttpResponse, bool bSucceeded) {
-			ReturnPromise->EmplaceValue((EHttpResponseCodes::Type)HttpResponse->GetResponseCode());
+			ReturnPromise->EmplaceValue((EStatusCode)HttpResponse->GetResponseCode());
 			});
 		Request->ProcessRequest();
 		return ReturnPromise->GetFuture();
diff --git a/Engine/Source/Runtime/Core/Public/Async/Future.h b/Engine/Source/Runtime/Core/Public/Async/Future.h
--- a/Engine/Source/Runtime/Core/Public/Async/Future.h
+++ b/Engine/Source/Runtime/Core/Public/Async/Future.h
@@ -144,6 +144,14 @@ public:
 		: FFutureState()
 	{ }
 
+	~TFutureState()
+	{
+		if (IsComplete())
+		{
+			DestructItem(Result.GetTypedPtr());
+		}
+	}
+
 	/**
 	 * Create and initialize a new instance with a callback.
 	 *
@@ -168,7 +176,7 @@ public:
 			WaitFor(FTimespan::MaxValue());
 		}
 
-		return Result;
+		return *Result.GetTypedPtr();
 	}
 
 	/**
@@ -181,15 +189,14 @@ public:
 	void EmplaceResult(ArgTypes&&... Args)
 	{
 		check(!IsComplete());
-
-		Result = InternalResultType(Forward<ArgTypes>(Args)...);
+		new(Result.GetTypedPtr()) InternalResultType(Forward<ArgTypes>(Args)...);
 		MarkComplete();
 	}
 
 private:
 
 	/** Holds the asynchronous result. */
-	InternalResultType Result;
+	TTypeCompatibleBytes<InternalResultType> Result;
 };
 
 /* TFuture
