Fix FMallocLinked destroying the link list by setting the head to nullptr even if there was a proper ->Next pointer

diff --git a/Engine/Source/Runtime/Core/Private/Unix/UnixForkPageProtector.cpp b/Engine/Source/Runtime/Core/Private/Unix/UnixForkPageProtector.cpp
--- a/Engine/Source/Runtime/Core/Private/Unix/UnixForkPageProtector.cpp
+++ b/Engine/Source/Runtime/Core/Private/Unix/UnixForkPageProtector.cpp
@@ -396,17 +396,6 @@ void FMallocLinked::Free(void* Ptr)
 			}
 		}
 
-		if (CurrentBlock->FreeList && CurrentBlock->FreeList->FreeSize == CurrentBlock->BlockSize)
-		{
-			if (Blocks == CurrentBlock)
-			{
-				Blocks = nullptr;
-			}
-
-			PreviousBlock->Next = CurrentBlock->Next;
-			munmap(CurrentBlock, CurrentBlock->BlockSize + sizeof(FBlock));
-		}
-
 		if (bPtrFreed)
 		{
 			return;
