StringBuilder: Delete operator<< for incompatible character types to avoid them using implicit conversions to integers

#rnx

diff --git a/Engine/Source/Runtime/Core/Public/Misc/StringBuilder.h b/Engine/Source/Runtime/Core/Public/Misc/StringBuilder.h
--- a/Engine/Source/Runtime/Core/Public/Misc/StringBuilder.h
+++ b/Engine/Source/Runtime/Core/Public/Misc/StringBuilder.h
@@ -391,10 +391,14 @@ inline auto operator<<(TStringBuilderBase<CharType>& Builder, CharRangeType&& St
 }
 
 inline FAnsiStringBuilderBase&		operator<<(FAnsiStringBuilderBase& Builder, ANSICHAR Char)							{ return Builder.AppendChar(Char); }
+inline FAnsiStringBuilderBase&		operator<<(FAnsiStringBuilderBase& Builder, UTF8CHAR Char) = delete;
+inline FAnsiStringBuilderBase&		operator<<(FAnsiStringBuilderBase& Builder, WIDECHAR Char) = delete;
 inline FWideStringBuilderBase&		operator<<(FWideStringBuilderBase& Builder, ANSICHAR Char)							{ return Builder.AppendChar(Char); }
+inline FWideStringBuilderBase&		operator<<(FWideStringBuilderBase& Builder, UTF8CHAR Char) = delete;
 inline FWideStringBuilderBase&		operator<<(FWideStringBuilderBase& Builder, WIDECHAR Char)							{ return Builder.AppendChar(Char); }
 inline FUtf8StringBuilderBase&		operator<<(FUtf8StringBuilderBase& Builder, ANSICHAR Char)							{ return Builder.AppendChar(UTF8CHAR(Char)); }
 inline FUtf8StringBuilderBase&		operator<<(FUtf8StringBuilderBase& Builder, UTF8CHAR Char)							{ return Builder.AppendChar(Char); }
+inline FUtf8StringBuilderBase&		operator<<(FUtf8StringBuilderBase& Builder, WIDECHAR Char) = delete;
 
 // Prefer using << instead of += as operator+= is only intended for mechanical FString -> FStringView replacement.
 inline FStringBuilderBase&			operator+=(FStringBuilderBase& Builder, ANSICHAR Char)								{ return Builder.AppendChar(Char); }
