Insert a space between the currency code and the numeric value to improve readability when formatting a foreign monetary value. When a currency is not from the current locale, the currency formatter may prepend or append the ISO currency code to disambiguate. This may cause readability issues.

For example, 1500 Japanese yen will be formatted as "¥1,500" in ja-JP without the ISO currency code (because it's the local currency of that locale); however it will be formatted as "1 500 JPY" in es-ES, and as "JPY1 500"in es-MX. To improve the readability of the last case, we insert a no-break space between the currency code and the numeric value, so that it can be formatted as "JPY 1 500".

#tests tested in locales en-US, es-ES, es-MX, es-UY, and ja-JP with currencies USD, CAD, JPY, MXN, and UYU. Visually confirmed that they all look as expected.

diff --git a/Engine/Source/Runtime/Core/Private/Internationalization/ICUCulture.cpp b/Engine/Source/Runtime/Core/Private/Internationalization/ICUCulture.cpp
--- a/Engine/Source/Runtime/Core/Private/Internationalization/ICUCulture.cpp
+++ b/Engine/Source/Runtime/Core/Private/Internationalization/ICUCulture.cpp
@@ -1145,7 +1145,33 @@ const FDecimalNumberFormattingRules& FICUCultureImplementation::GetCurrencyForma
 		CurrencyFormatterForCulture->setCurrency(ICUCurrencyCode.getBuffer());
 	}
 
-	const FDecimalNumberFormattingRules NewUECurrencyFormattingRules = ExtractNumberFormattingRulesFromICUDecimalFormatter(ICULocale, *CurrencyFormatterForCulture);
+	auto FixPrefixFormat = [](FString& InOutCurrencyPrefix)
+	{
+		const int32 PrefixLen = InOutCurrencyPrefix.Len();
+		if (PrefixLen >= 3 && ICUUtilities::IsValidCurencyCodeCharacter(InOutCurrencyPrefix[PrefixLen - 1]) && ICUUtilities::IsValidCurencyCodeCharacter(InOutCurrencyPrefix[PrefixLen - 2]) && ICUUtilities::IsValidCurencyCodeCharacter(InOutCurrencyPrefix[PrefixLen - 3]))
+		{
+			InOutCurrencyPrefix.AppendChar(TEXT('\u00A0')); // No-break space
+		}
+	};
+
+	auto FixSuffixFormat = [](FString& InOutCurrencySuffix)
+	{
+		const int32 SuffixLen = InOutCurrencySuffix.Len();
+		if (SuffixLen >= 3 && ICUUtilities::IsValidCurencyCodeCharacter(InOutCurrencySuffix[0]) && ICUUtilities::IsValidCurencyCodeCharacter(InOutCurrencySuffix[1]) && ICUUtilities::IsValidCurencyCodeCharacter(InOutCurrencySuffix[2]))
+		{
+			InOutCurrencySuffix.InsertAt(0, TEXT('\u00A0')); // No-break space
+		}
+	};
+
+	FDecimalNumberFormattingRules NewUECurrencyFormattingRules = ExtractNumberFormattingRulesFromICUDecimalFormatter(ICULocale, *CurrencyFormatterForCulture);
+	
+	// When a currency is from a different locale, it is common for ICU to disambiguate by prefixing or postfixing the numeric value
+	// with ISO currency code. To ensure the readability in all cases we will add a no-break space between the currency code prefix/suffix 
+	// and the numeric value, if the prefix/suffix would directly touch the numeric value (eg, to avoid a result like "JPY1 500").
+	FixPrefixFormat(NewUECurrencyFormattingRules.PositivePrefixString);
+	FixPrefixFormat(NewUECurrencyFormattingRules.NegativePrefixString);
+	FixSuffixFormat(NewUECurrencyFormattingRules.PositiveSuffixString);
+	FixSuffixFormat(NewUECurrencyFormattingRules.NegativeSuffixString);
 
 	if (bUseDefaultFormattingRules)
 	{
@@ -1155,7 +1181,7 @@ const FDecimalNumberFormattingRules& FICUCultureImplementation::GetCurrencyForma
 
 			if (!UECurrencyFormattingRules.IsValid())
 			{
-				UECurrencyFormattingRules = MakeShared<FDecimalNumberFormattingRules, ESPMode::ThreadSafe>(NewUECurrencyFormattingRules);
+				UECurrencyFormattingRules = MakeShared<FDecimalNumberFormattingRules, ESPMode::ThreadSafe>(MoveTemp(NewUECurrencyFormattingRules));
 			}
 		}
 
@@ -1172,7 +1198,7 @@ const FDecimalNumberFormattingRules& FICUCultureImplementation::GetCurrencyForma
 			return *FoundUEAlternateCurrencyFormattingRules;
 		}
 
-		FoundUEAlternateCurrencyFormattingRules = MakeShared<FDecimalNumberFormattingRules>(NewUECurrencyFormattingRules);
+		FoundUEAlternateCurrencyFormattingRules = MakeShared<FDecimalNumberFormattingRules>(MoveTemp(NewUECurrencyFormattingRules));
 		UEAlternateCurrencyFormattingRules.Add(SanitizedCurrencyCode, FoundUEAlternateCurrencyFormattingRules);
 		return *FoundUEAlternateCurrencyFormattingRules;
 	}
diff --git a/Engine/Source/Runtime/Core/Private/Internationalization/ICUUtilities.cpp b/Engine/Source/Runtime/Core/Private/Internationalization/ICUUtilities.cpp
--- a/Engine/Source/Runtime/Core/Private/Internationalization/ICUUtilities.cpp
+++ b/Engine/Source/Runtime/Core/Private/Internationalization/ICUUtilities.cpp
@@ -398,7 +398,7 @@ namespace ICUUtilities
 			{
 				if (InChar != 0)
 				{
-					const bool bIsValid = (InChar >= TEXT('A') && InChar <= TEXT('Z')) || (InChar >= TEXT('a') && InChar <= TEXT('z'));
+					const bool bIsValid = IsValidCurencyCodeCharacter(InChar);
 					return !bIsValid || ++NumValidChars > 3;
 				}
 				return false;
diff --git a/Engine/Source/Runtime/Core/Private/Internationalization/ICUUtilities.h b/Engine/Source/Runtime/Core/Private/Internationalization/ICUUtilities.h
--- a/Engine/Source/Runtime/Core/Private/Internationalization/ICUUtilities.h
+++ b/Engine/Source/Runtime/Core/Private/Internationalization/ICUUtilities.h
@@ -105,5 +105,11 @@ namespace ICUUtilities
 
 	/** Sanitize the given currency code so that it is safe to use with ICU */
 	FString SanitizeCurrencyCode(const FString& InCurrencyCode);
+
+	/** Is the given character valid as part of currency code to be used with ICU? */
+	static bool IsValidCurencyCodeCharacter(const TCHAR InChar)
+	{
+		return (InChar >= TEXT('A') && InChar <= TEXT('Z')) || (InChar >= TEXT('a') && InChar <= TEXT('z'));
+	};
 }
 #endif
