#ue5 - Removed VectorLoad(TVector<T>*) and VectorStore(TVector<T>*). Users should prefer VectorLoadFloat3/VectorStoreFloat3 to be correct, as the normal versions would read or write past the end. Fixed up a few engine uses of this that were not correct, but they were luckily not causing problems.

Also removed VectorLoad(TVector2*) and replaced with float* at the single call location in the engine (this usage was correct here).

diff --git a/Engine/Plugins/Runtime/GeometryCache/Source/GeometryCache/Private/GeometryCacheSceneProxy.cpp b/Engine/Plugins/Runtime/GeometryCache/Source/GeometryCache/Private/GeometryCacheSceneProxy.cpp
--- a/Engine/Plugins/Runtime/GeometryCache/Source/GeometryCache/Private/GeometryCacheSceneProxy.cpp
+++ b/Engine/Plugins/Runtime/GeometryCache/Source/GeometryCache/Private/GeometryCacheSceneProxy.cpp
@@ -1033,8 +1033,9 @@ void FGeometryCacheSceneProxy::FrameUpdate() const
 						int32 Index = 0;
 						for (; Index + 1 < NumVerts; Index += 2)
 						{
-							VectorRegister4Float InterpolatedUVx2 = VectorMultiplyAdd(	VectorLoad(&UVAPtr[Index]), WeightA,
-																VectorMultiply(		VectorLoad(&UVBPtr[Index]), WeightB));
+							VectorRegister4Float InterpolatedUVx2 = VectorMultiplyAdd(	VectorLoad(&UVAPtr[Index].X),
+																						WeightA,
+																						VectorMultiply(VectorLoad(&UVBPtr[Index].X), WeightB));
 							VectorStore(InterpolatedUVx2, &(InterpolatedUVsPtr[Index].X));
 						}
 
diff --git a/Engine/Source/Runtime/Core/Public/Math/Bounds.h b/Engine/Source/Runtime/Core/Public/Math/Bounds.h
--- a/Engine/Source/Runtime/Core/Public/Math/Bounds.h
+++ b/Engine/Source/Runtime/Core/Public/Math/Bounds.h
@@ -21,8 +21,8 @@ struct FBounds
 
 	FORCEINLINE FBounds& operator+=( const FVector3f& Other )
 	{
-		VectorStoreAligned( VectorMin( VectorLoadAligned( &Min ), VectorLoad( &Other ) ), &Min );
-		VectorStoreAligned( VectorMax( VectorLoadAligned( &Max ), VectorLoad( &Other ) ), &Max );
+		VectorStoreAligned( VectorMin( VectorLoadAligned( &Min ), VectorLoadFloat3( &Other ) ), &Min );
+		VectorStoreAligned( VectorMax( VectorLoadAligned( &Max ), VectorLoadFloat3( &Other ) ), &Max );
 		return *this;
 	}
 
diff --git a/Engine/Source/Runtime/Core/Public/Math/UnrealMathVectorCommon.h b/Engine/Source/Runtime/Core/Public/Math/UnrealMathVectorCommon.h
--- a/Engine/Source/Runtime/Core/Public/Math/UnrealMathVectorCommon.h
+++ b/Engine/Source/Runtime/Core/Public/Math/UnrealMathVectorCommon.h
@@ -41,18 +41,6 @@ FORCEINLINE TVectorRegisterType<T> VectorLoad(const UE::Math::TVector4<T>* Ptr)
 	return VectorLoad((const T*)(Ptr));
 }
 
-template<typename T>
-FORCEINLINE TVectorRegisterType<T> VectorLoad(const UE::Math::TVector<T>* Ptr)
-{
-	return VectorLoad((const T*)(Ptr));
-}
-
-template<typename T>
-FORCEINLINE TVectorRegisterType<T> VectorLoad(const UE::Math::TVector2<T>* Ptr)
-{
-	return VectorLoad((const T*)(Ptr));
-}
-
 FORCEINLINE VectorRegister4Float VectorLoad(const VectorRegister4Float* Ptr)
 {
 	return VectorLoad((const float*)(Ptr));
@@ -272,18 +260,6 @@ FORCEINLINE void VectorStore(const VectorRegister4Double& Vec, float* Dst)
 }
 #endif
 
-template<typename T>
-FORCEINLINE void VectorStore(const VectorRegister4Float& Vec, UE::Math::TVector<T>* Dst)
-{
-	VectorStore(Vec, (T*)Dst);
-}
-
-template<typename T>
-FORCEINLINE void VectorStore(const VectorRegister4Double& Vec, UE::Math::TVector<T>* Dst)
-{
-	VectorStore(Vec, (T*)Dst);
-}
-
 template<typename T>
 FORCEINLINE void VectorStore(const VectorRegister4Float& Vec, UE::Math::TVector4<T>* Dst)
 {
diff --git a/Engine/Source/Runtime/Engine/Private/HierarchicalInstancedStaticMesh.cpp b/Engine/Source/Runtime/Engine/Private/HierarchicalInstancedStaticMesh.cpp
--- a/Engine/Source/Runtime/Engine/Private/HierarchicalInstancedStaticMesh.cpp
+++ b/Engine/Source/Runtime/Engine/Private/HierarchicalInstancedStaticMesh.cpp
@@ -1158,8 +1158,8 @@ static FORCEINLINE_DEBUGGABLE bool CullNode(const FFoliageCullInstanceParams& Pa
 		checkSlow(Params.ViewFrustumLocal.PermutedPlanes.Num() == 4);
 
 		//@todo, once we have more than one mesh per tree, these should be aligned
-		VectorRegister BoxMin = VectorLoad(&Node.BoundMin);
-		VectorRegister BoxMax = VectorLoad(&Node.BoundMax);
+		VectorRegister BoxMin = VectorLoadFloat3(&Node.BoundMin);
+		VectorRegister BoxMax = VectorLoadFloat3(&Node.BoundMax);
 
 		VectorRegister BoxDiff = VectorSubtract(BoxMax,BoxMin);
 		VectorRegister BoxSum = VectorAdd(BoxMax,BoxMin);
diff --git a/Engine/Source/Runtime/Engine/Private/SkeletalRenderCPUSkin.cpp b/Engine/Source/Runtime/Engine/Private/SkeletalRenderCPUSkin.cpp
--- a/Engine/Source/Runtime/Engine/Private/SkeletalRenderCPUSkin.cpp
+++ b/Engine/Source/Runtime/Engine/Private/SkeletalRenderCPUSkin.cpp
@@ -946,8 +946,8 @@ static void SkinVertexSection(
 			// carry over the W component (sign of basis determinant) 
 			DstNormals[2] = VectorMultiplyAdd( VECTOR_0001, SrcNormals[2], DstNormals[2] );
 
-			// Write to 16-byte aligned memory:
-			VectorStore( DstNormals[0], &DestVertex->Position );
+			// Write to memory:
+			VectorStoreFloat3( DstNormals[0], &DestVertex->Position );
 			Pack3( DstNormals[1], &DestVertex->TangentX.Vector.Packed );
 			Pack4( DstNormals[2], &DestVertex->TangentZ.Vector.Packed );
 			VectorResetFloatRegisters(); // Need to call this to be able to use regular floating point registers again after Pack().
