Re-enable FPayloadId tests and change the filter to EAutomationTestFlags::EngineFilter to prevent them running on every start up.

#rnx

- When FPayloadId::ToString was switched to LexToString we changed from outputting an empty string for invalid ids to an id that is all zeros. However the stringbuilder path was not updated to do the same and so the tests failed.

diff --git a/Engine/Source/Runtime/Core/Private/Tests/Virtualization/PayloadIdTests.cpp b/Engine/Source/Runtime/Core/Private/Tests/Virtualization/PayloadIdTests.cpp
--- a/Engine/Source/Runtime/Core/Private/Tests/Virtualization/PayloadIdTests.cpp
+++ b/Engine/Source/Runtime/Core/Private/Tests/Virtualization/PayloadIdTests.cpp
@@ -8,12 +8,12 @@
 
 #include "Misc/AutomationTest.h"
 
-#if WITH_DEV_AUTOMATION_TESTS && 0
+#if WITH_DEV_AUTOMATION_TESTS
 
 namespace UE::Virtualization
 {
 	
-constexpr const uint32 TestFlags = EAutomationTestFlags::ApplicationContextMask | EAutomationTestFlags::SmokeFilter;
+constexpr const uint32 TestFlags = EAutomationTestFlags::ApplicationContextMask | EAutomationTestFlags::EngineFilter;
 
 /** This test ensures that ::IsValid returns the correct value dending on how the FPayloadId was created */
 IMPLEMENT_SIMPLE_AUTOMATION_TEST(FPayloadIdTestValidity, TEXT("System.Core.Virtualization.PayloadId.Validity"), TestFlags)
@@ -310,9 +310,9 @@ bool FPayloadIdTestStringBuilder::RunTest(const FString& Parameters)
 
 	// Test again both TWideStringBuilder and TAnsiStringBuilder
 #define TEST_PAYLOAD_ID(BufferToTest) {	TWideStringBuilder<128> WideBuilder; WideBuilder << BufferToTest; \
-										TestEqual(TEXT("TWideStringBuilder << and ToString to produce the same result"), *WideBuilder, LexToString(BufferToTest));\
+										TestEqual(TEXT("TWideStringBuilder << and LexToString to produce the same result"), *WideBuilder, LexToString(BufferToTest));\
 										TAnsiStringBuilder<128> AnsiBuilder; AnsiBuilder << BufferToTest; \
-										TestEqual(TEXT("TAnsiBuilder << and ToString to produce the same result"), *AnsiBuilder, LexToString(BufferToTest)); }
+										TestEqual(TEXT("TAnsiBuilder << and LexToString to produce the same result"), *AnsiBuilder, LexToString(BufferToTest)); }
 
 	TEST_PAYLOAD_ID(EmptyId);
 	TEST_PAYLOAD_ID(BufferId0);
diff --git a/Engine/Source/Runtime/Core/Public/Virtualization/PayloadId.h b/Engine/Source/Runtime/Core/Public/Virtualization/PayloadId.h
--- a/Engine/Source/Runtime/Core/Public/Virtualization/PayloadId.h
+++ b/Engine/Source/Runtime/Core/Public/Virtualization/PayloadId.h
@@ -97,10 +97,7 @@ inline uint32 GetTypeHash(const FPayloadId& Hash)
 template <typename CharType>
 TStringBuilderBase<CharType>& operator<<(TStringBuilderBase<CharType>& Builder, const FPayloadId& Id)
 {
-	if (Id.IsValid())
-	{
-		Builder << Id.Identifier;
-	}
+	Builder << Id.Identifier;
 
 	return Builder;
 }
