Make FMatrix44f <-> FMatrix44d conversion casts explicit, including all specializaitons

diff --git a/Engine/Source/Runtime/Core/Public/Math/ClipProjectionMatrix.h b/Engine/Source/Runtime/Core/Public/Math/ClipProjectionMatrix.h
--- a/Engine/Source/Runtime/Core/Public/Math/ClipProjectionMatrix.h
+++ b/Engine/Source/Runtime/Core/Public/Math/ClipProjectionMatrix.h
@@ -28,6 +28,10 @@ public:
 	 */
 	TClipProjectionMatrix( const TMatrix<T>& SrcProjMat, const TPlane<T>& Plane );
 
+	// Conversion to other type.
+	template<typename FArg, TEMPLATE_REQUIRES(!TIsSame<T, FArg>::Value)>
+	explicit TClipProjectionMatrix(const TClipProjectionMatrix<FArg>& From) : TMatrix<T>(From) {}
+	
 private:
 	/** return sign of a number */
 	FORCEINLINE T sgn(T a );
diff --git a/Engine/Source/Runtime/Core/Public/Math/InverseRotationMatrix.h b/Engine/Source/Runtime/Core/Public/Math/InverseRotationMatrix.h
--- a/Engine/Source/Runtime/Core/Public/Math/InverseRotationMatrix.h
+++ b/Engine/Source/Runtime/Core/Public/Math/InverseRotationMatrix.h
@@ -22,6 +22,10 @@ public:
 	 * @param Rot rotation
 	 */
 	TInverseRotationMatrix(const TRotator<T>& Rot);
+	
+	// Conversion to other type.
+    template<typename FArg, TEMPLATE_REQUIRES(!TIsSame<T, FArg>::Value)>
+    explicit TInverseRotationMatrix(const TInverseRotationMatrix<FArg>& From) : TMatrix<T>(From) {}	
 };
 
 template<typename T>
diff --git a/Engine/Source/Runtime/Core/Public/Math/Matrix.h b/Engine/Source/Runtime/Core/Public/Math/Matrix.h
--- a/Engine/Source/Runtime/Core/Public/Math/Matrix.h
+++ b/Engine/Source/Runtime/Core/Public/Math/Matrix.h
@@ -425,9 +425,9 @@ public:
 		Dest[11] = Src[14]; // [3][2]
 	}
 
-	// Conversion to other type. TODO: explicit!
+	// Conversion to other type.
 	template<typename FArg, TEMPLATE_REQUIRES(!TIsSame<T, FArg>::Value)>
-	TMatrix(const TMatrix<FArg>& From)
+	explicit TMatrix(const TMatrix<FArg>& From)
 	{
 		// TODO: SIMD this?
 		M[0][0] = (T)From.M[0][0]; M[0][1] = (T)From.M[0][1]; M[0][2] = (T)From.M[0][2]; M[0][3] = (T)From.M[0][3];
@@ -510,17 +510,21 @@ inline FArchive& operator<<(FArchive& Ar, TMatrix<double>& M)
 }
 
 template<typename T>
-struct TBasisVectorMatrix : TMatrix<T>
+struct TBasisVectorMatrix : public TMatrix<T>
 {
 	using TMatrix<T>::M;
-	
+
 	// Create Basis matrix from 3 axis vectors and the origin
 	TBasisVectorMatrix(const TVector<T>& XAxis,const TVector<T>& YAxis,const TVector<T>& ZAxis,const TVector<T>& Origin);
+
+	// Conversion to other type.
+	template<typename FArg, TEMPLATE_REQUIRES(!TIsSame<T, FArg>::Value)>
+	explicit TBasisVectorMatrix(const TBasisVectorMatrix<FArg>& From) : TMatrix<T>(From) {}
 };
 
 
 template<typename T>
-struct TLookFromMatrix : TMatrix<T>
+struct TLookFromMatrix : public TMatrix<T>
 {
 	using TMatrix<T>::M;
 
@@ -531,11 +535,15 @@ struct TLookFromMatrix : TMatrix<T>
 	* Always use this instead of e.g., FLookAtMatrix(Pos, Pos + Dir,...);
 	*/
 	TLookFromMatrix(const TVector<T>& EyePosition, const TVector<T>& LookDirection, const TVector<T>& UpVector);
+
+	// Conversion to other type.
+	template<typename FArg, TEMPLATE_REQUIRES(!TIsSame<T, FArg>::Value)>
+	explicit TLookFromMatrix(const TLookFromMatrix<FArg>& From) : TMatrix<T>(From) {}
 };
 
 
 template<typename T>
-struct TLookAtMatrix : TLookFromMatrix<T>
+struct TLookAtMatrix : public TLookFromMatrix<T>
 {
 	using TLookFromMatrix<T>::M;
 
@@ -546,6 +554,10 @@ struct TLookAtMatrix : TLookFromMatrix<T>
 	* This does the same thing as D3DXMatrixLookAtLH.
 	*/
 	TLookAtMatrix(const TVector<T>& EyePosition, const TVector<T>& LookAtPosition, const TVector<T>& UpVector);
+
+	// Conversion to other type.
+	template<typename FArg, TEMPLATE_REQUIRES(!TIsSame<T, FArg>::Value)>
+	explicit TLookAtMatrix(const TLookAtMatrix<FArg>& From) : TLookFromMatrix<T>(From) {}
 };
 
 
diff --git a/Engine/Source/Runtime/Core/Public/Math/MirrorMatrix.h b/Engine/Source/Runtime/Core/Public/Math/MirrorMatrix.h
--- a/Engine/Source/Runtime/Core/Public/Math/MirrorMatrix.h
+++ b/Engine/Source/Runtime/Core/Public/Math/MirrorMatrix.h
@@ -24,6 +24,10 @@ public:
 	 * @param Plane source plane for mirroring (assumed normalized)
 	 */
 	TMirrorMatrix( const TPlane<T>& Plane );
+
+	// Conversion to other type.
+	template<typename FArg, TEMPLATE_REQUIRES(!TIsSame<T, FArg>::Value)>
+	explicit TMirrorMatrix(const TMirrorMatrix<FArg>& From) : TMatrix<T>(From) {}
 };
 
 template<typename T>
diff --git a/Engine/Source/Runtime/Core/Public/Math/OrthoMatrix.h b/Engine/Source/Runtime/Core/Public/Math/OrthoMatrix.h
--- a/Engine/Source/Runtime/Core/Public/Math/OrthoMatrix.h
+++ b/Engine/Source/Runtime/Core/Public/Math/OrthoMatrix.h
@@ -25,6 +25,10 @@ public:
 	 * @param ZOffset offset in the Z axis
 	 */
 	TOrthoMatrix(T Width,T Height,T ZScale,T ZOffset);
+	
+	// Conversion to other type.
+	template<typename FArg, TEMPLATE_REQUIRES(!TIsSame<T, FArg>::Value)>
+	explicit TOrthoMatrix(const TOrthoMatrix<FArg>& From) : TMatrix<T>(From) {}
 };
 
 
@@ -34,6 +38,10 @@ struct TReversedZOrthoMatrix : public TMatrix<T>
 public:
 	TReversedZOrthoMatrix(T Width,T Height,T ZScale,T ZOffset);
 	TReversedZOrthoMatrix(T Left, T Right, T Bottom, T Top, T ZScale, T ZOffset);
+
+	// Conversion to other type.
+	template<typename FArg, TEMPLATE_REQUIRES(!TIsSame<T, FArg>::Value)>
+	explicit TReversedZOrthoMatrix(const TReversedZOrthoMatrix<FArg>& From) : TMatrix<T>(From) {}
 };
 
 template<typename T>
diff --git a/Engine/Source/Runtime/Core/Public/Math/PerspectiveMatrix.h b/Engine/Source/Runtime/Core/Public/Math/PerspectiveMatrix.h
--- a/Engine/Source/Runtime/Core/Public/Math/PerspectiveMatrix.h
+++ b/Engine/Source/Runtime/Core/Public/Math/PerspectiveMatrix.h
@@ -53,6 +53,10 @@ public:
 	 * @note that the FOV you pass in is actually half the FOV, unlike most perspective matrix functions (D3DXMatrixPerspectiveFovLH).
 	 */
 	TPerspectiveMatrix(T HalfFOV, T Width, T Height, T MinZ);
+
+	// Conversion to other type.
+	template<typename FArg, TEMPLATE_REQUIRES(!TIsSame<T, FArg>::Value)>
+	explicit TPerspectiveMatrix(const TPerspectiveMatrix<FArg>& From) : TMatrix<T>(From) {}
 };
 
 
@@ -63,6 +67,10 @@ public:
 	TReversedZPerspectiveMatrix(T HalfFOVX, T HalfFOVY, T MultFOVX, T MultFOVY, T MinZ, T MaxZ);
 	TReversedZPerspectiveMatrix(T HalfFOV, T Width, T Height, T MinZ, T MaxZ);
 	TReversedZPerspectiveMatrix(T HalfFOV, T Width, T Height, T MinZ);
+	
+	// Conversion to other type.
+	template<typename FArg, TEMPLATE_REQUIRES(!TIsSame<T, FArg>::Value)>
+	explicit TReversedZPerspectiveMatrix(const TReversedZPerspectiveMatrix<FArg>& From) : TMatrix<T>(From) {}
 };
 
 
diff --git a/Engine/Source/Runtime/Core/Public/Math/QuatRotationTranslationMatrix.h b/Engine/Source/Runtime/Core/Public/Math/QuatRotationTranslationMatrix.h
--- a/Engine/Source/Runtime/Core/Public/Math/QuatRotationTranslationMatrix.h
+++ b/Engine/Source/Runtime/Core/Public/Math/QuatRotationTranslationMatrix.h
@@ -26,6 +26,10 @@ public:
 	*/
 	TQuatRotationTranslationMatrix(const TQuat<T>& Q, const TVector<T>& Origin);
 
+	// Conversion to other type.
+	template<typename FArg, TEMPLATE_REQUIRES(!TIsSame<T, FArg>::Value)>
+	explicit TQuatRotationTranslationMatrix(const TQuatRotationTranslationMatrix<FArg>& From) : TMatrix<T>(From) {}
+
 	/** Matrix factory. Return an FMatrix so we don't have type conversion issues in expressions. */
 	static TMatrix<T> Make(const TQuat<T>& Q, const TVector<T>& Origin)
 	{
@@ -51,6 +55,10 @@ public:
 	{
 	}
 
+	// Conversion to other type.
+	template<typename FArg, TEMPLATE_REQUIRES(!TIsSame<T, FArg>::Value)>
+	explicit TQuatRotationMatrix(const TQuatRotationMatrix<FArg>& From) : TQuatRotationTranslationMatrix<T>(From) {}
+	
 	/** Matrix factory. Return an FMatrix so we don't have type conversion issues in expressions. */
 	static TMatrix<T> Make(const TQuat<T>& Q)
 	{
diff --git a/Engine/Source/Runtime/Core/Public/Math/RotationAboutPointMatrix.h b/Engine/Source/Runtime/Core/Public/Math/RotationAboutPointMatrix.h
--- a/Engine/Source/Runtime/Core/Public/Math/RotationAboutPointMatrix.h
+++ b/Engine/Source/Runtime/Core/Public/Math/RotationAboutPointMatrix.h
@@ -27,6 +27,10 @@ public:
 	 */
 	TRotationAboutPointMatrix(const TRotator<T>& Rot, const TVector<T>& Origin);
 
+	// Conversion to other type.
+	template<typename FArg, TEMPLATE_REQUIRES(!TIsSame<T, FArg>::Value)>
+	explicit TRotationAboutPointMatrix(const TRotationAboutPointMatrix<FArg>& From) : TRotationTranslationMatrix<T>(From) {}	
+	
 	/** Matrix factory. Return an TMatrix<T> so we don't have type conversion issues in expressions. */
 	static TMatrix<T> Make(const TRotator<T>& Rot, const TVector<T>& Origin)
 	{
diff --git a/Engine/Source/Runtime/Core/Public/Math/RotationTranslationMatrix.h b/Engine/Source/Runtime/Core/Public/Math/RotationTranslationMatrix.h
--- a/Engine/Source/Runtime/Core/Public/Math/RotationTranslationMatrix.h
+++ b/Engine/Source/Runtime/Core/Public/Math/RotationTranslationMatrix.h
@@ -26,6 +26,10 @@ public:
 	 */
 	TRotationTranslationMatrix(const TRotator<T>& Rot, const TVector<T>& Origin);
 
+	// Conversion to other type.
+	template<typename FArg, TEMPLATE_REQUIRES(!TIsSame<T, FArg>::Value)>
+	explicit TRotationTranslationMatrix(const TRotationTranslationMatrix<FArg>& From) : TMatrix<T>(From) {}
+	
 	/** Matrix factory. Return an TMatrix<T> so we don't have type conversion issues in expressions. */
 	static TMatrix<T> Make(const TRotator<T>& Rot, const TVector<T>& Origin)
 	{
diff --git a/Engine/Source/Runtime/Core/Public/Math/ScaleMatrix.h b/Engine/Source/Runtime/Core/Public/Math/ScaleMatrix.h
--- a/Engine/Source/Runtime/Core/Public/Math/ScaleMatrix.h
+++ b/Engine/Source/Runtime/Core/Public/Math/ScaleMatrix.h
@@ -28,6 +28,10 @@ public:
 	 */
 	TScaleMatrix( const TVector<T>& Scale );
 
+	// Conversion to other type.
+	template<typename FArg, TEMPLATE_REQUIRES(!TIsSame<T, FArg>::Value)>
+	explicit TScaleMatrix(const TScaleMatrix<FArg>& From) : TMatrix<T>(From) {}
+	
 	/** Matrix factory. Return an TMatrix<T> so we don't have type conversion issues in expressions. */
 	static TMatrix<T> Make(T Scale)
 	{
diff --git a/Engine/Source/Runtime/Core/Public/Math/ScaleRotationTranslationMatrix.h b/Engine/Source/Runtime/Core/Public/Math/ScaleRotationTranslationMatrix.h
--- a/Engine/Source/Runtime/Core/Public/Math/ScaleRotationTranslationMatrix.h
+++ b/Engine/Source/Runtime/Core/Public/Math/ScaleRotationTranslationMatrix.h
@@ -26,6 +26,10 @@ public:
 	 * @param Origin translation to apply
 	 */
 	TScaleRotationTranslationMatrix(const TVector<T>& Scale, const TRotator<T>& Rot, const TVector<T>& Origin);
+	
+	// Conversion to other type.
+	template<typename FArg, TEMPLATE_REQUIRES(!TIsSame<T, FArg>::Value)>
+	explicit TScaleRotationTranslationMatrix(const TScaleRotationTranslationMatrix<FArg>& From) : TMatrix<T>(From) {}
 };
 
 namespace
diff --git a/Engine/Source/Runtime/Core/Public/Math/TranslationMatrix.h b/Engine/Source/Runtime/Core/Public/Math/TranslationMatrix.h
--- a/Engine/Source/Runtime/Core/Public/Math/TranslationMatrix.h
+++ b/Engine/Source/Runtime/Core/Public/Math/TranslationMatrix.h
@@ -19,6 +19,10 @@ public:
 	/** Constructor translation matrix based on given vector */
 	TTranslationMatrix(const TVector<T>& Delta);
 
+	// Conversion to other type.
+	template<typename FArg, TEMPLATE_REQUIRES(!TIsSame<T, FArg>::Value)>
+	explicit TTranslationMatrix(const TTranslationMatrix<FArg>& From) : TMatrix<T>(From) {}
+	
 	/** Matrix factory. Return an FMatrix so we don't have type conversion issues in expressions. */
 	static TMatrix<T> Make(TVector<T> const& Delta)
 	{
