[Insights] Report ThreadStack size correctly.

diff --git a/Engine/Source/Runtime/Core/Private/Windows/WindowsRunnableThread.h b/Engine/Source/Runtime/Core/Private/Windows/WindowsRunnableThread.h
--- a/Engine/Source/Runtime/Core/Private/Windows/WindowsRunnableThread.h
+++ b/Engine/Source/Runtime/Core/Private/Windows/WindowsRunnableThread.h
@@ -145,9 +145,11 @@ protected:
 			// add in the thread size, since it's allocated in a black box we can't track
 			// note: I don't see any accounting for this when threads are destroyed
 			const uint64 FakeAddress = uint64(InRunnable) | (uint64(1u) << 48);
-			MemoryTrace_Alloc(FakeAddress, InStackSize, 4);
+			// Size of zero indicates using default thread stack size. This is 1MB unless overridden in .def file.
+			const uint64 Size = InStackSize > 0 ? InStackSize : 1 * 1024 * 1024;
+			MemoryTrace_Alloc(FakeAddress, Size, 4);
 			MemoryTrace_MarkAllocAsHeap(FakeAddress, EMemoryTraceRootHeap::SystemMemory);
-			MemoryTrace_Alloc(FakeAddress, InStackSize, 4);
+			MemoryTrace_Alloc(FakeAddress, Size, 4);
 			LLM(FLowLevelMemTracker::Get().OnLowLevelAlloc(ELLMTracker::Default, nullptr, InStackSize));
 			LLM(FLowLevelMemTracker::Get().OnLowLevelAlloc(ELLMTracker::Platform, nullptr, InStackSize));
 
