Fixed FMemoryWriter64 failing to write >2GB

diff --git a/Engine/Source/Runtime/Core/Public/Serialization/MemoryWriter.h b/Engine/Source/Runtime/Core/Public/Serialization/MemoryWriter.h
--- a/Engine/Source/Runtime/Core/Public/Serialization/MemoryWriter.h
+++ b/Engine/Source/Runtime/Core/Public/Serialization/MemoryWriter.h
@@ -19,6 +19,7 @@ template <int IndexSize>
 class TMemoryWriter : public FMemoryArchive
 {
 	static_assert(IndexSize == 32 || IndexSize == 64, "Only 32-bit and 64-bit index sizes supported");
+	using IndexSizeType = typename TBitsToSizeType<IndexSize>::Type;
 
 public:
 	TMemoryWriter( TArray<uint8, TSizedDefaultAllocator<IndexSize>>& InBytes, bool bIsPersistent = false, bool bSetOffset = false, const FName InArchiveName = NAME_None )
@@ -42,21 +43,21 @@ public:
 			const int64 NewArrayCount = Bytes.Num() + NumBytesToAdd;
 			if constexpr (IndexSize == 32)
 			{
-				if( NewArrayCount >= MAX_int32 )
+				if (NewArrayCount >= MAX_int32)
 				{
-					UE_LOG( LogSerialization, Fatal, TEXT( "FMemoryWriter does not support data larger than 2GB. Archive name: %s." ), *ArchiveName.ToString() );
+					UE_LOG(LogSerialization, Fatal, TEXT("FMemoryWriter does not support data larger than 2GB. Archive name: %s."), *ArchiveName.ToString());
 				}
 			}
 
-			Bytes.AddUninitialized( (int32)NumBytesToAdd );
+			Bytes.AddUninitialized((IndexSizeType)NumBytesToAdd);
 		}
 
 		check((Offset + Num) <= Bytes.Num());
 		
 		if( Num )
 		{
-			FMemory::Memcpy( &Bytes[(int32)Offset], Data, Num );
-			Offset+=Num;
+			FMemory::Memcpy(&Bytes[(IndexSizeType)Offset], Data, Num);
+			Offset += Num;
 		}
 	}
 	/**
