StringView: Fixed calls to FPlatformString with null pointers

#rnx

diff --git a/Engine/Source/Runtime/Core/Private/Tests/Containers/StringViewTest.cpp b/Engine/Source/Runtime/Core/Private/Tests/Containers/StringViewTest.cpp
--- a/Engine/Source/Runtime/Core/Private/Tests/Containers/StringViewTest.cpp
+++ b/Engine/Source/Runtime/Core/Private/Tests/Containers/StringViewTest.cpp
@@ -316,6 +316,20 @@ bool FStringViewTestEquality::RunTest(const FString& Parameters)
 	TestFalse(TEXT("Equality(92)"), AnsiViewUpper.Equals(WideStringLiteralSrc, ESearchCase::CaseSensitive));
 	TestTrue(TEXT("Equality(93)"), AnsiViewUpper.Equals(WideStringLiteralSrc, ESearchCase::IgnoreCase));
 
+	// Test equality of empty strings
+	{
+		const TCHAR* EmptyLiteral = TEXT("");
+		const TCHAR* NonEmptyLiteral = TEXT("ABC");
+		FStringView EmptyView;
+		FStringView NonEmptyView = TEXTVIEW("ABC");
+		TestTrue(TEXT("Equals(94)"), EmptyView.Equals(EmptyLiteral));
+		TestTrue(TEXT("Equals(95)"), !EmptyView.Equals(NonEmptyLiteral));
+		TestTrue(TEXT("Equals(96)"), !NonEmptyView.Equals(EmptyLiteral));
+		TestTrue(TEXT("Equals(97)"), EmptyView.Equals(EmptyView));
+		TestTrue(TEXT("Equals(98)"), !EmptyView.Equals(NonEmptyView));
+		TestTrue(TEXT("Equals(99)"), !NonEmptyView.Equals(EmptyView));
+	}
+
 	// Test types convertible to a string view
 	static_assert(TIsSame<bool, decltype(FAnsiStringView().Equals(FString()))>::Value, "Error with Equals");
 	static_assert(TIsSame<bool, decltype(FWideStringView().Equals(FString()))>::Value, "Error with Equals");
@@ -391,6 +405,20 @@ bool FStringViewTestComparisonCaseSensitive::RunTest(const FString& Parameters)
 		TestTrue(TEXT("ComparisonCaseInsensitive(18)"), ViewLongUpper.Compare(WideStringLiteralLowerShort, ESearchCase::CaseSensitive) < 0);
 	}
 
+	// Test comparisons of empty strings
+	{
+		const TCHAR* EmptyLiteral = TEXT("");
+		const TCHAR* NonEmptyLiteral = TEXT("ABC");
+		FStringView EmptyView;
+		FStringView NonEmptyView = TEXTVIEW("ABC");
+		TestTrue(TEXT("ComparisonEmpty(19)"), EmptyView.Compare(EmptyLiteral) == 0);
+		TestTrue(TEXT("ComparisonEmpty(20)"), EmptyView.Compare(NonEmptyLiteral) < 0);
+		TestTrue(TEXT("ComparisonEmpty(21)"), NonEmptyView.Compare(EmptyLiteral) > 0);
+		TestTrue(TEXT("ComparisonEmpty(22)"), EmptyView.Compare(EmptyView) == 0);
+		TestTrue(TEXT("ComparisonEmpty(23)"), EmptyView.Compare(NonEmptyView) < 0);
+		TestTrue(TEXT("ComparisonEmpty(24)"), NonEmptyView.Compare(EmptyView) > 0);
+	}
+
 	// Test types convertible to a string view
 	static_assert(TIsSame<int32, decltype(FAnsiStringView().Compare(FString()))>::Value, "Error with Compare");
 	static_assert(TIsSame<int32, decltype(FWideStringView().Compare(FString()))>::Value, "Error with Compare");
diff --git a/Engine/Source/Runtime/Core/Public/Containers/StringView.h b/Engine/Source/Runtime/Core/Public/Containers/StringView.h
--- a/Engine/Source/Runtime/Core/Public/Containers/StringView.h
+++ b/Engine/Source/Runtime/Core/Public/Containers/StringView.h
@@ -497,6 +497,11 @@ template <typename OtherCharType,
 	std::enable_if_t<TIsCharType<OtherCharType>::Value>*>
 inline bool TStringView<CharType>::Equals(const OtherCharType* const Other, ESearchCase::Type SearchCase) const
 {
+	if (IsEmpty())
+	{
+		return !*Other;
+	}
+
 	if (SearchCase == ESearchCase::CaseSensitive)
 	{
 		return FPlatformString::Strncmp(GetData(), Other, Len()) == 0 && !Other[Len()];
@@ -516,6 +521,11 @@ inline int32 TStringView<CharType>::Compare(TStringView<OtherCharType> OtherView
 	const int32 OtherLen = OtherView.Len();
 	const int32 MinLen = FMath::Min(SelfLen, OtherLen);
 
+	if (!SelfLen || !OtherLen)
+	{
+		return SelfLen - OtherLen;
+	}
+
 	int Result;
 	if (SearchCase == ESearchCase::CaseSensitive)
 	{
@@ -538,6 +548,11 @@ template <typename CharType>
 template <typename OtherCharType>
 inline int32 TStringView<CharType>::Compare(const OtherCharType* const Other, ESearchCase::Type SearchCase) const
 {
+	if (IsEmpty())
+	{
+		return -!!*Other;
+	}
+
 	int Result;
 	if (SearchCase == ESearchCase::CaseSensitive)
 	{
@@ -658,7 +673,7 @@ inline auto operator<(CharRangeType&& Lhs, TStringView<CharType> Rhs)
 template <typename CharType>
 inline bool operator==(TStringView<CharType> Lhs, const CharType* Rhs)
 {
-	return TCString<CharType>::Strnicmp(Lhs.GetData(), Rhs, Lhs.Len()) == 0 && !Rhs[Lhs.Len()];
+	return Lhs.Equals(Rhs, ESearchCase::IgnoreCase);
 }
 
 template <typename CharType>
