FEM - introduce Parameters to TTaskDelegate in an effort to eventually reduce the size of FTask. This allows us to remove the this capture inside of FTask::Init. As an added benefit the nested Lambdas might also optimize better because of the much smaller body.

diff --git a/Engine/Source/Runtime/Core/Public/Async/Fundamental/ReserveScheduler.h b/Engine/Source/Runtime/Core/Public/Async/Fundamental/ReserveScheduler.h
--- a/Engine/Source/Runtime/Core/Public/Async/Fundamental/ReserveScheduler.h
+++ b/Engine/Source/Runtime/Core/Public/Async/Fundamental/ReserveScheduler.h
@@ -15,7 +15,7 @@ namespace LowLevelTasks
 		~FReserveScheduler();
 
 	public:
-		using FConditional = TTaskDelegate<16, bool>;
+		using FConditional = TTaskDelegate<bool(), 16>;
 
 	private:
 		struct alignas(8) FYieldedWork
diff --git a/Engine/Source/Runtime/Core/Public/Async/Fundamental/Scheduler.h b/Engine/Source/Runtime/Core/Public/Async/Fundamental/Scheduler.h
--- a/Engine/Source/Runtime/Core/Public/Async/Fundamental/Scheduler.h
+++ b/Engine/Source/Runtime/Core/Public/Async/Fundamental/Scheduler.h
@@ -130,7 +130,7 @@ namespace LowLevelTasks
 		static CORE_API FScheduler Singleton;
 
 		// using 16 bytes here because it fits the vtable and one additional pointer
-		using FConditional = TTaskDelegate<16, bool>;
+		using FConditional = TTaskDelegate<bool(), 16>;
 
 	private: 
 		//the FLocalQueueInstaller installs a LocalQueue into the current thread
diff --git a/Engine/Source/Runtime/Core/Public/Async/Fundamental/Task.h b/Engine/Source/Runtime/Core/Public/Async/Fundamental/Task.h
--- a/Engine/Source/Runtime/Core/Public/Async/Fundamental/Task.h
+++ b/Engine/Source/Runtime/Core/Public/Async/Fundamental/Task.h
@@ -7,6 +7,8 @@
 #include "CoreTypes.h"
 #include <atomic>
 
+#define LOWLEVEL_TASK_SIZE PLATFORM_CACHE_LINE_SIZE
+
 namespace LowLevelTasks
 {
 	DECLARE_LOG_CATEGORY_EXTERN(LowLevelTasks, Log, All);
@@ -97,17 +99,22 @@ namespace LowLevelTasks
 		Type* Value;
 
 	public:
-		TDeleter(Type* InValue) : Value(InValue)
+		inline TDeleter(Type* InValue) : Value(InValue)
 		{
 		}
 
-		TDeleter(const TDeleter&) = delete;
-		TDeleter(TDeleter&& Other) : Value(Other.Value)
+		inline TDeleter(const TDeleter&) = delete;
+		inline TDeleter(TDeleter&& Other) : Value(Other.Value)
 		{
 			Other.Value = nullptr;
 		}
 
-		~TDeleter() 
+		inline Type* GetValue() const
+		{
+			return Value;
+		}
+
+		inline ~TDeleter() 
 		{
 			if(Value)
 			{
@@ -180,16 +187,13 @@ namespace LowLevelTasks
 		};
 
 	private:
-		using FTaskDelegate = TTaskDelegate<PLATFORM_CACHE_LINE_SIZE - sizeof(FPackedData) - sizeof(void*), bool>;
+		using FTaskDelegate = TTaskDelegate<void(bool), LOWLEVEL_TASK_SIZE - sizeof(FPackedData) - sizeof(void*)>;
 		FTaskDelegate Runnable;
 		mutable void* UserData = nullptr;
 		std::atomic<FPackedData> PackedData { FPackedData() };
 
 	private:
-		FTaskBase()
-		{
-			static_assert(sizeof(FTaskBase) == PLATFORM_CACHE_LINE_SIZE, "Require FTaskBase to be cacheline size");
-		}
+		FTaskBase() = default;
 
 	private:
 		inline bool IsCanceled() const { return PackedData.load(std::memory_order_relaxed).GetState() == ETaskState::Canceled; }
@@ -272,10 +276,7 @@ namespace LowLevelTasks
 		void SetUserData(void* NewUserData) const { UserData = NewUserData; }
 
 	public:
-		FTask()
-		{
-			static_assert(sizeof(FTask) == PLATFORM_CACHE_LINE_SIZE, "Require FTask to be cacheline size");
-		}
+		FTask() = default;
 		inline ~FTask();
 
 	private: //Interface of the Scheduler
@@ -325,23 +326,13 @@ namespace LowLevelTasks
 	{
 		checkf(IsCompleted(), TEXT("State: %d"), PackedData.load(std::memory_order_relaxed).GetState());
 		checkSlow(!Runnable.IsSet());
-		Runnable = [this, LocalRunnable = Forward<TRunnable>(InRunnable), LocalContinuation = Forward<TContinuation>(InContinuation)]() mutable -> bool
+		Runnable = [LocalRunnable = Forward<TRunnable>(InRunnable), LocalContinuation = Forward<TContinuation>(InContinuation)](const bool NotCanceled)
 		{
-			checkSlow(IsScheduled() || IsCanceled());
-			FPackedData LocalPackedData = PackedData.load(std::memory_order_relaxed);
-			FPackedData ScheduledState(LocalPackedData, ETaskState::Scheduled);
-			if (PackedData.compare_exchange_strong(ScheduledState, FPackedData(LocalPackedData, ETaskState::Running), std::memory_order_relaxed))
+			if (NotCanceled)
 			{
 				LocalRunnable();
-				LocalContinuation();
-				return false;
-			}
-			else
-			{
-				checkSlow(IsCanceled());
-				LocalContinuation();
-				return true;
 			}
+			LocalContinuation();
 		};
 		InheritParentData(InPriority);
 		PackedData.store(FPackedData(InDebugName, InPriority, ETaskState::Ready, bAllowBusyWaiting), std::memory_order_release);
@@ -352,20 +343,11 @@ namespace LowLevelTasks
 	{
 		checkf(IsCompleted(), TEXT("State: %d"), PackedData.load(std::memory_order_relaxed).GetState());
 		checkSlow(!Runnable.IsSet());
-		Runnable = [this, LocalRunnable = Forward<TRunnable>(InRunnable)]() mutable -> bool
+		Runnable = [LocalRunnable = Forward<TRunnable>(InRunnable)](const bool NotCanceled)
 		{
-			checkSlow(IsScheduled() || IsCanceled());
-			FPackedData LocalPackedData = PackedData.load(std::memory_order_relaxed);
-			FPackedData ScheduledState(LocalPackedData, ETaskState::Scheduled);
-			if (PackedData.compare_exchange_strong(ScheduledState, FPackedData(LocalPackedData, ETaskState::Running), std::memory_order_relaxed))
+			if (NotCanceled)
 			{
 				LocalRunnable();
-				return false;
-			}
-			else
-			{
-				checkSlow(IsCanceled());
-				return true;
 			}
 		};
 		InheritParentData(InPriority);
@@ -430,12 +412,18 @@ namespace LowLevelTasks
 	{
 		checkSlow(Runnable.IsSet());
 		FTaskDelegate LocalRunnable;
-		bool Canceled = Runnable.CallAndMove(LocalRunnable);
+
+		checkSlow(IsScheduled() || IsCanceled());
+		FPackedData LocalPackedData = PackedData.load(std::memory_order_relaxed);
+		FPackedData ScheduledState(LocalPackedData, ETaskState::Scheduled);
+		const bool NotCanceled = PackedData.compare_exchange_strong(ScheduledState, FPackedData(LocalPackedData, ETaskState::Running), std::memory_order_relaxed);
+		
+		Runnable.CallAndMove(LocalRunnable, NotCanceled);
+
 		checkSlow(!Runnable.IsSet());
 		//do not access the task again after this call
 		//as by defitition the task can be considered dead
-		FPackedData LocalPackedData = PackedData.load(std::memory_order_relaxed);
-		PackedData.store(FPackedData(LocalPackedData, Canceled ? ETaskState::CanceledAndCompleted : ETaskState::Completed), std::memory_order_seq_cst);
+		PackedData.store(FPackedData(LocalPackedData, NotCanceled ? ETaskState::Completed : ETaskState::CanceledAndCompleted), std::memory_order_seq_cst);
 	}
 
 	inline const TCHAR* FTask::GetDebugName() const
diff --git a/Engine/Source/Runtime/Core/Public/Async/Fundamental/TaskDelegate.h b/Engine/Source/Runtime/Core/Public/Async/Fundamental/TaskDelegate.h
--- a/Engine/Source/Runtime/Core/Public/Async/Fundamental/TaskDelegate.h
+++ b/Engine/Source/Runtime/Core/Public/Async/Fundamental/TaskDelegate.h
@@ -36,13 +36,16 @@ namespace LowLevelTasks
 	}
 	//version of TUniqueFunction<ReturnType()> that is less wasteful with it's memory
 	//this class might be removed when TUniqueFunction<ReturnType()> is fixed
-	template<uint32 TotalSize = PLATFORM_CACHE_LINE_SIZE, typename ReturnType = void>
-	class alignas(8) TTaskDelegate
+	template<typename = void(), uint32 = PLATFORM_CACHE_LINE_SIZE>
+	class TTaskDelegate;
+
+	template<uint32 TotalSize , typename ReturnType, typename... ParamTypes>
+	class alignas(8) TTaskDelegate<ReturnType(ParamTypes...), TotalSize>
 	{
-		template<uint32, typename>
+		template<typename, uint32>
 		friend class TTaskDelegate;
 
-		using ThisClass = TTaskDelegate<TotalSize, ReturnType>;
+		using ThisClass = TTaskDelegate<ReturnType(ParamTypes...), TotalSize>;
 		struct TTaskDelegateBase
 		{
 			virtual void Move(TTaskDelegateBase&, void*, void*, uint32) 
@@ -50,13 +53,13 @@ namespace LowLevelTasks
 				checkNoEntry(); //if we get to this place than the compiler is optimizing our vtable lookup and ignores our undefined/shallow copy of the memory
 			};
 
-			virtual ReturnType Call(void*) const 
+			virtual ReturnType Call(void*, ParamTypes...) const 
 			{ 
 				checkNoEntry(); //if we get to this place than the compiler is optimizing our vtable lookup and ignores our undefined/shallow copy of the memory 
 				return TTaskDelegate_Impl::MakeDummyValue<ReturnType>(); 
 			};
 
-			virtual ReturnType CallAndMove(ThisClass&, void*, uint32) 
+			virtual ReturnType CallAndMove(ThisClass&, void*, uint32, ParamTypes...) 
 			{
 				checkNoEntry(); //if we get to this place than the compiler is optimizing our vtable lookup and ignores our undefined/shallow copy of the memory 
 				return TTaskDelegate_Impl::MakeDummyValue<ReturnType>(); 
@@ -92,13 +95,13 @@ namespace LowLevelTasks
 			{
 			}
 
-			ReturnType Call(void*) const override 
+			ReturnType Call(void*, ParamTypes...) const override 
 			{ 
 				checkf(false, TEXT("trying to Call a dummy TaskDelegate"));
 				return TTaskDelegate_Impl::MakeDummyValue<ReturnType>(); 
 			}
 
-			ReturnType CallAndMove(ThisClass&, void*, uint32) override 
+			ReturnType CallAndMove(ThisClass&, void*, uint32, ParamTypes...) override 
 			{ 
 				checkf(false, TEXT("trying to Call a dummy TaskDelegate"));
 				return TTaskDelegate_Impl::MakeDummyValue<ReturnType>(); 
@@ -133,8 +136,8 @@ namespace LowLevelTasks
 			template<typename CallableT>
 			inline TTaskDelegateImpl(CallableT&& Callable, void* InlineData)
 			{
-				static_assert(TIsInvocable<TCallableType>::Value, "TCallableType is not invocable");
-				static_assert(TIsSame<ReturnType, decltype(Callable())>::Value, "TCallableType return type does not match");
+				static_assert(TIsInvocable<TCallableType, ParamTypes...>::Value, "TCallableType is not invocable");
+				static_assert(TIsSame<ReturnType, decltype(Callable(UE::Core::Private::IsInvocable::DeclVal<ParamTypes>()...))>::Value, "TCallableType return type does not match");
 				static_assert(sizeof(TTaskDelegateImpl<TCallableType, false>) == sizeof(TTaskDelegateBase), "Size must match the Baseclass");
 				new (InlineData) TCallableType(Forward<CallableT>(Callable));
 			}
@@ -153,19 +156,19 @@ namespace LowLevelTasks
 				new (this) TTaskDelegateDummy();
 			}
 
-			inline ReturnType Call(void* InlineData) const override
+			inline ReturnType Call(void* InlineData, ParamTypes... Params) const override
 			{
 				TCallableType* LocalPtr = reinterpret_cast<TCallableType*>(InlineData);
-				return Invoke(*LocalPtr);
+				return Invoke(*LocalPtr, Params...);
 			}
 
-			ReturnType CallAndMove(ThisClass& Destination, void* InlineData, uint32 DestInlineSize) override
+			ReturnType CallAndMove(ThisClass& Destination, void* InlineData, uint32 DestInlineSize, ParamTypes... Params) override
 			{
 				ON_SCOPE_EXIT
 				{
 					Move(Destination.CallableWrapper, Destination.InlineStorage, InlineData, DestInlineSize);
 				};
-				return Call(InlineData);		
+				return Call(InlineData, Params...);		
 			}
 
 			void Destroy(void* InlineData) override
@@ -203,8 +206,8 @@ namespace LowLevelTasks
 			template<typename CallableT>
 			inline TTaskDelegateImpl(CallableT&& Callable, void* InlineData)
 			{
-				static_assert(TIsInvocable<TCallableType>::Value, "TCallableType is not invocable");
-				static_assert(TIsSame<ReturnType, decltype(Callable())>::Value, "TCallableType return type does not match");
+				static_assert(TIsInvocable<TCallableType, ParamTypes...>::Value, "TCallableType is not invocable");
+				static_assert(TIsSame<ReturnType, decltype(Callable(UE::Core::Private::IsInvocable::DeclVal<ParamTypes>()...))>::Value, "TCallableType return type does not match");
 				static_assert(sizeof(TTaskDelegateImpl<TCallableType, true>) == sizeof(TTaskDelegateBase), "Size must match the Baseclass");
 				TCallableType** HeapPtr = reinterpret_cast<TCallableType**>(InlineData);
 				*HeapPtr = reinterpret_cast<TCallableType*>(TConcurrentLinearAllocator<FLowLevelTasksBlockAllocationTag>::template Malloc<alignof(TCallableType)>(sizeof(TCallableType)));
@@ -217,19 +220,19 @@ namespace LowLevelTasks
 				new (this) TTaskDelegateDummy();
 			}
 
-			inline ReturnType Call(void* InlineData) const override
+			inline ReturnType Call(void* InlineData, ParamTypes... Params) const override
 			{
 				TCallableType* HeapPtr = reinterpret_cast<TCallableType*>(*reinterpret_cast<void* const*>(InlineData));
-				return Invoke(*HeapPtr);
+				return Invoke(*HeapPtr, Params...);
 			}
 
-			ReturnType CallAndMove(ThisClass& Destination, void* InlineData, uint32 DestInlineSize) override
+			ReturnType CallAndMove(ThisClass& Destination, void* InlineData, uint32 DestInlineSize, ParamTypes... Params) override
 			{
 				ON_SCOPE_EXIT
 				{
 					Move(Destination.CallableWrapper, Destination.InlineStorage, InlineData, DestInlineSize);
 				};
-				return Call(InlineData);			
+				return Call(InlineData, Params...);			
 			}
 
 			void Destroy(void* InlineData) override
@@ -266,10 +269,10 @@ namespace LowLevelTasks
 		}
 
 		template<uint32 SourceTotalSize>
-		TTaskDelegate(const TTaskDelegate<SourceTotalSize, ReturnType>&) = delete;
+		TTaskDelegate(const TTaskDelegate<ReturnType(ParamTypes...), SourceTotalSize>&) = delete;
 
 		template<uint32 SourceTotalSize>
-		TTaskDelegate(TTaskDelegate<SourceTotalSize, ReturnType>&& Other)
+		TTaskDelegate(TTaskDelegate<ReturnType(ParamTypes...), SourceTotalSize>&& Other)
 		{
 			Other.GetWrapper()->Move(CallableWrapper, InlineStorage, Other.InlineStorage, InlineStorageSize);
 		}
@@ -293,23 +296,23 @@ namespace LowLevelTasks
 			GetWrapper()->Destroy(InlineStorage);
 		}
 
-		ReturnType operator()() const
+		ReturnType operator()(ParamTypes... Params) const
 		{
-			return GetWrapper()->Call(InlineStorage);
+			return GetWrapper()->Call(InlineStorage, Params...);
 		}
 
 		template<uint32 DestTotalSize>
-		ReturnType CallAndMove(TTaskDelegate<DestTotalSize, ReturnType>& Destination)
+		ReturnType CallAndMove(TTaskDelegate<ReturnType(ParamTypes...), DestTotalSize>& Destination, ParamTypes... Params)
 		{
 			checkSlow(!Destination.IsSet());
-			return GetWrapper()->CallAndMove(Destination, InlineStorage, TTaskDelegate<DestTotalSize, ReturnType>::InlineStorageSize);
+			return GetWrapper()->CallAndMove(Destination, InlineStorage, TTaskDelegate<ReturnType(ParamTypes...), DestTotalSize>::InlineStorageSize, Params...);
 		}
 
 		template<uint32 SourceTotalSize>
-		ThisClass& operator= (const TTaskDelegate<SourceTotalSize, ReturnType>&) = delete;
+		ThisClass& operator= (const TTaskDelegate<ReturnType(ParamTypes...), SourceTotalSize>&) = delete;
 
 		template<uint32 SourceTotalSize>
-		ThisClass& operator= (TTaskDelegate<SourceTotalSize, ReturnType>&& Other)
+		ThisClass& operator= (TTaskDelegate<ReturnType(ParamTypes...), SourceTotalSize>&& Other)
 		{
 			GetWrapper()->Destroy(InlineStorage);
 			Other.GetWrapper()->Move(CallableWrapper, InlineStorage, Other.InlineStorage, InlineStorageSize);
diff --git a/Engine/Source/Runtime/Core/Public/Async/ParallelFor.h b/Engine/Source/Runtime/Core/Public/Async/ParallelFor.h
--- a/Engine/Source/Runtime/Core/Public/Async/ParallelFor.h
+++ b/Engine/Source/Runtime/Core/Public/Async/ParallelFor.h
@@ -491,7 +491,7 @@ namespace ParallelForImpl
 		class FParallelExecutor 
 		{
 			FDataHandle Data;
-			int32 WorkerIndex;
+			mutable int32 WorkerIndex;
 			LowLevelTasks::ETaskPriority Priority;
 
 		public:
@@ -523,7 +523,7 @@ namespace ParallelForImpl
 				return Data;
 			}
 
-			inline bool operator()(const bool bIsMaster = false) noexcept
+			inline bool operator()(const bool bIsMaster = false) const noexcept
 			{
 				FMemMark Mark(FMemStack::Get());
 				TRACE_CPUPROFILER_EVENT_SCOPE(ParallelFor);
