Adding OnPreFork to LLM to avoid sharing the csv tracker with child processes

diff --git a/Engine/Source/Runtime/Core/Private/HAL/LowLevelMemTracker.cpp b/Engine/Source/Runtime/Core/Private/HAL/LowLevelMemTracker.cpp
--- a/Engine/Source/Runtime/Core/Private/HAL/LowLevelMemTracker.cpp
+++ b/Engine/Source/Runtime/Core/Private/HAL/LowLevelMemTracker.cpp
@@ -229,11 +229,13 @@ namespace LLMPrivate
 
 		void Publish(FLowLevelMemTracker& LLMRef, const FTrackerTagSizeMap& TagSizes, const FTagData* OverrideTrackedTotalTagData, const FTagData* OverrideUntaggedTagData, int64 TrackedTotal, bool bTrackPeaks);
 
+		void OnPreFork();
+
 	private:
 		void Write(FStringView Text);
 		static const TCHAR* GetTrackerCsvName(ELLMTracker InTracker);
 
-		void CreateArchive();
+		bool CreateArchive();
 		bool UpdateColumns(const FTrackerTagSizeMap& TagSizes);
 		void WriteHeader(const FTagData* OverrideTrackedTotalTagData, const FTagData* OverrideUntaggedTagData);
 		void AddRow(FLowLevelMemTracker& LLMRef, const FTrackerTagSizeMap& TagSizes, const FTagData* OverrideTrackedTotalName, const FTagData* OverrideUntaggedName, int64 TrackedTotal, bool bTrackPeaks);
@@ -370,6 +372,8 @@ namespace LLMPrivate
 		void PublishTrace(bool bTrackPeaks);
 		void PublishCsvProfiler(bool bTrackPeaks);
 
+		void OnPreFork();
+
 		struct FLowLevelAllocInfo
 		{
 		public:
@@ -724,6 +728,20 @@ void FLowLevelMemTracker::Clear()
 	bInitialisedTracking = false;
 }
 
+void FLowLevelMemTracker::OnPreFork()
+{
+	using namespace UE::LLMPrivate;
+
+	if (!bIsDisabled)
+	{
+		FLLMTracker& DefaultTracker = *GetTracker(ELLMTracker::Default);
+		FLLMTracker& PlatformTracker = *GetTracker(ELLMTracker::Platform);
+
+		DefaultTracker.OnPreFork();
+		PlatformTracker.OnPreFork();
+	}
+}
+
 void FLowLevelMemTracker::UpdateStatsPerFrame(const TCHAR* LogName)
 {
 	if (bIsDisabled)
@@ -3188,6 +3206,11 @@ namespace LLMPrivate
 		CsvProfilerWriter.Clear();
 	}
 
+	void FLLMTracker::OnPreFork()
+	{
+		CsvWriter.OnPreFork();
+	}
+
 	void FLLMTracker::SetTotalTags(const FTagData* InOverrideUntaggedTagData, const FTagData* InOverrideTrackedTotalTagData)
 	{
 		OverrideUntaggedTagData = InOverrideUntaggedTagData;
@@ -3692,6 +3715,16 @@ namespace LLMPrivate
 		ExistingColumns.Empty();
 	}
 
+	void FLLMCsvWriter::OnPreFork()
+	{
+		if (Archive)
+		{
+			Archive->Flush();
+			delete Archive;
+			Archive = nullptr;
+		}
+	}
+
 	void FLLMCsvWriter::SetTracker(ELLMTracker InTracker)
 	{
 		Tracker = InTracker;
@@ -3699,18 +3732,20 @@ namespace LLMPrivate
 
 	void FLLMCsvWriter::Publish(FLowLevelMemTracker& LLMRef, const FTrackerTagSizeMap& TagSizes, const FTagData* OverrideTrackedTotalTagData, const FTagData* OverrideUntaggedTagData, int64 TrackedTotal, bool bTrackPeaks)
 	{
-		double Now = FPlatformTime::Seconds();
+		const double Now = FPlatformTime::Seconds();
+
 		if ((FLowLevelMemTracker::Get().bPublishSingleFrame == false) && 
 			(Now - LastWriteTime < (double)CVarLLMWriteInterval.GetValueOnAnyThread()))
 		{
 			return;
 		}
+
 		LastWriteTime = Now;
 
+		const bool bCreatedArchive = CreateArchive();
+		const bool bColumnsUpdated = UpdateColumns(TagSizes);
 
-		CreateArchive();
-		bool bColumnsUpdated = UpdateColumns(TagSizes);
-		if (bColumnsUpdated)
+		if (bCreatedArchive || bColumnsUpdated)
 		{
 			// The column names are written at the start of the archive; when they change we seek back to the start of the file and rewrite the column names.
 			WriteHeader(OverrideTrackedTotalTagData, OverrideUntaggedTagData);
@@ -3737,11 +3772,11 @@ namespace LLMPrivate
 		Archive->Serialize(const_cast<ANSICHAR*>(StringCast<ANSICHAR>(Text.GetData(), Text.Len()).Get()), Text.Len() * sizeof(ANSICHAR));
 	}
 
-	void FLLMCsvWriter::CreateArchive()
+	bool FLLMCsvWriter::CreateArchive()
 	{
 		if (Archive)
 		{
-			return;
+			return false;
 		}
 
 		// create the csv file
@@ -3768,6 +3803,8 @@ namespace LLMPrivate
 		// create space for column titles that are filled in as we get them
 		Write(FString::ChrN(CVarLLMHeaderMaxSize.GetValueOnAnyThread(), ' '));
 		Write(TEXT("\n"));
+
+		return true;
 	}
 
 	bool FLLMCsvWriter::UpdateColumns(const FTrackerTagSizeMap& TagSizes)
diff --git a/Engine/Source/Runtime/Core/Public/HAL/LowLevelMemTracker.h b/Engine/Source/Runtime/Core/Public/HAL/LowLevelMemTracker.h
--- a/Engine/Source/Runtime/Core/Public/HAL/LowLevelMemTracker.h
+++ b/Engine/Source/Runtime/Core/Public/HAL/LowLevelMemTracker.h
@@ -566,6 +566,8 @@ public:
 	// Publishes the active LLM stats in the active frame, useful for single targeted LLM snapshots
 	void PublishDataSingleFrame();
 
+	void OnPreFork();
+
 private:
 	FLowLevelMemTracker();
 
