LWC: Double NetSerialize support for FVector/FVector2D.
Backwards compatibility with replays from float based builds.

diff --git a/Engine/Source/Runtime/Core/Private/Math/UnrealMath.cpp b/Engine/Source/Runtime/Core/Private/Math/UnrealMath.cpp
--- a/Engine/Source/Runtime/Core/Private/Math/UnrealMath.cpp
+++ b/Engine/Source/Runtime/Core/Private/Math/UnrealMath.cpp
@@ -1128,7 +1128,20 @@ namespace Math
 			}
 		}
 
-		Ar << Q.X << Q.Y << Q.Z;
+		if (Ar.EngineNetVer() >= HISTORY_SERIALIZE_DOUBLE_VECTORS_AS_DOUBLES)
+		{
+			Ar << Q.X << Q.Y << Q.Z;
+		}
+		else
+		{
+			checkf(Ar.IsLoading(), TEXT("float -> double conversion applied outside of load!"));
+			// Always serialize as float
+			float SX, SY, SZ;
+			Ar << SX << SY << SZ;
+			Q.X = SX;
+			Q.Y = SY;
+			Q.Z = SZ;
+		}
 
 		if (Ar.IsLoading())
 		{
diff --git a/Engine/Source/Runtime/Core/Public/Math/Vector.h b/Engine/Source/Runtime/Core/Public/Math/Vector.h
--- a/Engine/Source/Runtime/Core/Public/Math/Vector.h
+++ b/Engine/Source/Runtime/Core/Public/Math/Vector.h
@@ -10,6 +10,7 @@
 #include "Containers/UnrealString.h"
 #include "Misc/Parse.h"
 #include "Misc/LargeWorldCoordinatesSerializer.h"
+#include "Misc/NetworkVersion.h"
 #include "Math/Color.h"
 #include "Math/IntPoint.h"
 #include "Logging/LogMacros.h"
@@ -1126,13 +1127,16 @@ public:
      */
 	bool NetSerialize(FArchive& Ar, class UPackageMap* Map, bool& bOutSuccess)
 	{
-		// LWC_TODO: Serializer - NetSerialize forced to float for now to fix replays recorded with LWC disabled.
-		float SX = (float)X, SY = (float)Y, SZ = (float)Z;
-		Ar << SX;
-		Ar << SY;
-		Ar << SZ;
-		if(Ar.IsLoading())
+		if (Ar.EngineNetVer() >= HISTORY_SERIALIZE_DOUBLE_VECTORS_AS_DOUBLES)
 		{
+			Ar << X << Y << Z;
+		}
+		else
+		{
+			checkf(Ar.IsLoading(), TEXT("float -> double conversion applied outside of load!"));
+			// Always serialize as float
+			float SX, SY, SZ;
+			Ar << SX << SY << SZ;
 			X = SX;
 			Y = SY;
 			Z = SZ;
diff --git a/Engine/Source/Runtime/Core/Public/Math/Vector2D.h b/Engine/Source/Runtime/Core/Public/Math/Vector2D.h
--- a/Engine/Source/Runtime/Core/Public/Math/Vector2D.h
+++ b/Engine/Source/Runtime/Core/Public/Math/Vector2D.h
@@ -9,6 +9,7 @@
 #include "Containers/UnrealString.h"
 #include "Misc/Parse.h"
 #include "Misc/LargeWorldCoordinatesSerializer.h"
+#include "Misc/NetworkVersion.h"
 #include "Math/IntPoint.h"
 #include "Logging/LogMacros.h"
 
@@ -616,12 +617,16 @@ public:
 	*/
 	bool NetSerialize(FArchive& Ar, class UPackageMap* Map, bool& bOutSuccess)
 	{
-		// LWC_TODO: Serializer - NetSerialize forced to float for now to fix replays recorded with LWC disabled.
-		float SX = (float)X, SY = (float)Y;
-		Ar << SX;
-		Ar << SY;
-		if(Ar.IsLoading())
+		if (Ar.EngineNetVer() >= HISTORY_SERIALIZE_DOUBLE_VECTORS_AS_DOUBLES)
 		{
+			Ar << X << Y;
+		}
+		else
+		{
+			checkf(Ar.IsLoading(), TEXT("float -> double conversion applied outside of load!"));
+			// Always serialize as float
+			float SX, SY;
+			Ar << SX << SY;
 			X = SX;
 			Y = SY;
 		}
diff --git a/Engine/Source/Runtime/Core/Public/Misc/NetworkVersion.h b/Engine/Source/Runtime/Core/Public/Misc/NetworkVersion.h
--- a/Engine/Source/Runtime/Core/Public/Misc/NetworkVersion.h
+++ b/Engine/Source/Runtime/Core/Public/Misc/NetworkVersion.h
@@ -53,6 +53,7 @@ enum EEngineNetworkVersionHistory
 	HISTORY_HITRESULT_INSTANCEHANDLE = 19,			// Bump version to support FHitResult change of Actor to HitObjectHandle. This change was made in CL 14369221 but a net version wasn't added at the time.
 	HISTORY_INTERFACE_PROPERTY_SERIALIZATION = 20,	// Bump version to support net serialization of FInterfaceProperty
 	HISTORY_MONTAGE_PLAY_INST_ID_SERIALIZATION = 21,// Bump version to support net serialization of FGameplayAbilityRepAnimMontage, addition of PlayInstanceId and removal of bForcePlayBit
+	HISTORY_SERIALIZE_DOUBLE_VECTORS_AS_DOUBLES	= 22,// Bump version to support net serialization of double vector types
 	// New history items go above here.
 
 	HISTORY_ENGINENETVERSION_PLUS_ONE,
