[Insights] Auto connect would always enable default channels.

Using auto connect would always enable default channel, even if started with a specific channel. Also renamed internal type to avoid confusion with TraceLog's FChannel.

diff --git a/Engine/Source/Runtime/Core/Private/ProfilingDebugging/TraceAuxiliary.cpp b/Engine/Source/Runtime/Core/Private/ProfilingDebugging/TraceAuxiliary.cpp
--- a/Engine/Source/Runtime/Core/Private/ProfilingDebugging/TraceAuxiliary.cpp
+++ b/Engine/Source/Runtime/Core/Private/ProfilingDebugging/TraceAuxiliary.cpp
@@ -77,6 +77,7 @@ public:
 	void					GetActiveChannelsString(FStringBuilderBase& String) const;
 	void					AddCommandlineChannels(const TCHAR* ChannelList);
 	void					ResetCommandlineChannels();
+	bool					HasCommandlineChannels() const { return !CommandlineChannels.IsEmpty(); };
 	void					EnableChannels(const TCHAR* ChannelList);
 	void					DisableChannels(const TCHAR* ChannelList);
 	bool					Connect(ETraceConnectType Type, const TCHAR* Parameter, const FTraceAuxiliary::FLogCategoryAlias& LogCategory);
@@ -97,7 +98,7 @@ private:
 		Tracing,
 	};
 
-	struct FChannel
+	struct FChannelEntry
 	{
 		FString				Name;
 		bool				bActive = false;
@@ -112,7 +113,8 @@ private:
 	bool					SendToHost(const TCHAR* Host, const FTraceAuxiliary::FLogCategoryAlias& LogCategory);
 	bool					WriteToFile(const TCHAR* Path, const FTraceAuxiliary::FLogCategoryAlias& LogCategory);
 
-	TMap<uint32, FChannel>	CommandlineChannels;
+	typedef TMap<uint32, FChannelEntry, TInlineSetAllocator<128>> ChannelSet;
+	ChannelSet				CommandlineChannels;
 	FString					TraceDest;
 	EState					State = EState::Stopped;
 	bool					bTruncateFile = false;
@@ -204,7 +206,7 @@ void FTraceAuxiliaryImpl::AddChannel(const TCHAR* Name)
 		return;
 	}
 
-	FChannel& Value = CommandlineChannels.Add(Hash, {});
+	FChannelEntry& Value = CommandlineChannels.Add(Hash, {});
 	Value.Name = Name;
 
 	if (State >= EState::Tracing && !Value.bActive)
@@ -218,7 +220,7 @@ void FTraceAuxiliaryImpl::RemoveChannel(const TCHAR* Name)
 {
 	uint32 Hash = HashChannelName(Name);
 
-	FChannel Channel;
+	FChannelEntry Channel;
 	if (!CommandlineChannels.RemoveAndCopyValue(Hash, Channel))
 	{
 		return;
@@ -1290,7 +1292,7 @@ void FTraceAuxiliary::TryAutoConnect()
 		if (KnownEvent != nullptr)
 		{
 			UE_LOG(LogCore, Display, TEXT("Unreal Insights instance detected, auto-connecting to local trace server..."));
-			Start(EConnectionType::Network, TEXT("127.0.0.1"), TEXT("default"), nullptr);
+			Start(EConnectionType::Network, TEXT("127.0.0.1"), GTraceAuxiliary.HasCommandlineChannels() ? nullptr : TEXT("default"), nullptr);
 			::CloseHandle(KnownEvent);
 		}
 	#endif
