[Insights] Fixes for mismatching alloc/free events when using MiMalloc

diff --git a/Engine/Source/Runtime/Core/Private/ProfilingDebugging/MemoryAllocationTrace.cpp b/Engine/Source/Runtime/Core/Private/ProfilingDebugging/MemoryAllocationTrace.cpp
--- a/Engine/Source/Runtime/Core/Private/ProfilingDebugging/MemoryAllocationTrace.cpp
+++ b/Engine/Source/Runtime/Core/Private/ProfilingDebugging/MemoryAllocationTrace.cpp
@@ -3,6 +3,7 @@
 #include "ProfilingDebugging/MemoryAllocationTrace.h"
 #include "HAL/PlatformTime.h"
 #include "ProfilingDebugging/TagTrace.h"
+#include "ProfilingDebugging/CallstackTrace.h"
 #include "ProfilingDebugging/TraceMalloc.h"
 
 #if UE_TRACE_ENABLED
@@ -397,23 +398,23 @@ void* FTraceMalloc::Malloc(SIZE_T Count, uint32 Alignment)
 /////////////////////////////////////////////////////////////////////////////
 void* FTraceMalloc::Realloc(void* Original, SIZE_T Count, uint32 Alignment)
 {
-	void* NewPtr;
-	{
-		TGuardValue<bool> _(GDoNotTrace, true);
-		NewPtr = WrappedMalloc->Realloc(Original, Count, Alignment);
-	}
-
+	void* NewPtr = nullptr;
 #if UE_MEMORY_TRACE_ENABLED
 	const uint64 Size = Count;
 	const uint32 AlignmentPow2 = uint32(FPlatformMath::CountTrailingZeros(Alignment));
 	const uint32 Alignment_SizeLower = (AlignmentPow2 << SizeShift) | uint32(Size & ((1 << SizeShift) - 1));
 
 	UE_MEMSCOPE(TRACE_TAG);
-
+	
 	UE_TRACE_LOG(Memory, ReallocFree, MemAllocChannel)
 		<< ReallocFree.Address(uint64(Original))
 		<< ReallocFree.RootHeap(uint8(EMemoryTraceRootHeap::SystemMemory));
-
+	
+	{
+		TGuardValue<bool> _(GDoNotTrace, true);
+		NewPtr = WrappedMalloc->Realloc(Original, Count, Alignment);
+	}
+	
 	UE_TRACE_LOG(Memory, ReallocAlloc, MemAllocChannel)
 		<< ReallocAlloc.CallstackId(0)
 		<< ReallocAlloc.Address(uint64(NewPtr))
@@ -428,16 +429,16 @@ void* FTraceMalloc::Realloc(void* Original, SIZE_T Count, uint32 Alignment)
 /////////////////////////////////////////////////////////////////////////////
 void FTraceMalloc::Free(void* Original)
 {
-	{
-		TGuardValue<bool> _(GDoNotTrace, true);
-		WrappedMalloc->Free(Original);
-	}
-
 #if UE_MEMORY_TRACE_ENABLED
 	UE_TRACE_LOG(Memory, Free, MemAllocChannel)
 		<< Free.Address(uint64(Original))
 		<< Free.RootHeap(uint8(EMemoryTraceRootHeap::SystemMemory));
 #endif //UE_MEMORY_TRACE_ENABLED
+	
+	{
+		TGuardValue<bool> _(GDoNotTrace, true);
+		WrappedMalloc->Free(Original);
+	}
 }
 
 /////////////////////////////////////////////////////////////////////////////
diff --git a/Engine/Source/Runtime/Core/Private/ProfilingDebugging/Microsoft/MemoryTrace.cpp b/Engine/Source/Runtime/Core/Private/ProfilingDebugging/Microsoft/MemoryTrace.cpp
--- a/Engine/Source/Runtime/Core/Private/ProfilingDebugging/Microsoft/MemoryTrace.cpp
+++ b/Engine/Source/Runtime/Core/Private/ProfilingDebugging/Microsoft/MemoryTrace.cpp
@@ -1,6 +1,8 @@
 // Copyright Epic Games, Inc. All Rights Reserved.
 
 #include "ProfilingDebugging/MemoryTrace.h"
+
+#include "HAL/MallocMimalloc.h"
 #include "ProfilingDebugging/CallstackTrace.h"
 #include "ProfilingDebugging/TraceMalloc.h"
 
@@ -345,9 +347,11 @@ void FVirtualWinApiHooks::Initialize(bool bInLight)
 LPVOID WINAPI FVirtualWinApiHooks::VmAlloc(LPVOID Address, SIZE_T Size, DWORD Type, DWORD Protect)
 {
 	LPVOID Ret = VmAllocOrig(Address, Size, Type, Protect);
-	if (Ret != nullptr && (Type & MEM_COMMIT) && FTraceMalloc::ShouldTrace())
+	// Track any reserve for now. Going forward we need events to differentiate reserves/commits and
+	// corresponding information on frees.
+	if (Ret != nullptr && (Type & MEM_RESERVE))
 	{
-		const uint32 Callstack = CallstackTrace_GetCurrentId();
+		const uint32 Callstack = FTraceMalloc::ShouldTrace() ? CallstackTrace_GetCurrentId() : 0;
 		GAllocationTrace->Alloc(Ret, Size, 0, Callstack, EMemoryTraceRootHeap::SystemMemory);
 		GAllocationTrace->MarkAllocAsHeap(Ret, EMemoryTraceRootHeap::SystemMemory);
 	}
@@ -358,7 +362,8 @@ LPVOID WINAPI FVirtualWinApiHooks::VmAlloc(LPVOID Address, SIZE_T Size, DWORD Ty
 ////////////////////////////////////////////////////////////////////////////////
 BOOL WINAPI FVirtualWinApiHooks::VmFree(LPVOID Address, SIZE_T Size, DWORD Type)
 {
-	if (FTraceMalloc::ShouldTrace())
+	// Currently tracking any free event
+	if (Type & MEM_RELEASE)
 	{
 		GAllocationTrace->Free(Address, EMemoryTraceRootHeap::SystemMemory);
 	}
@@ -369,9 +374,9 @@ BOOL WINAPI FVirtualWinApiHooks::VmFree(LPVOID Address, SIZE_T Size, DWORD Type)
 LPVOID WINAPI FVirtualWinApiHooks::VmAllocEx(HANDLE Process, LPVOID Address, SIZE_T Size, DWORD Type, DWORD Protect)
 {
 	LPVOID Ret = VmAllocExOrig(Process, Address, Size, Type, Protect);
-	if (Process == GetCurrentProcess() && Ret != nullptr && (Type & MEM_COMMIT) && FTraceMalloc::ShouldTrace())
+	if (Process == GetCurrentProcess() && Ret != nullptr && (Type & MEM_RESERVE))
 	{
-		const uint32 Callstack = CallstackTrace_GetCurrentId();
+		const uint32 Callstack = FTraceMalloc::ShouldTrace() ? CallstackTrace_GetCurrentId() : 0;
 		GAllocationTrace->Alloc(Ret, Size, 0, Callstack);
 		GAllocationTrace->MarkAllocAsHeap(Ret, EMemoryTraceRootHeap::SystemMemory);
 	}
@@ -382,7 +387,7 @@ LPVOID WINAPI FVirtualWinApiHooks::VmAllocEx(HANDLE Process, LPVOID Address, SIZ
 ////////////////////////////////////////////////////////////////////////////////
 BOOL WINAPI FVirtualWinApiHooks::VmFreeEx(HANDLE Process, LPVOID Address, SIZE_T Size, DWORD Type)
 {
-	if (Process == GetCurrentProcess() && FTraceMalloc::ShouldTrace())
+	if (Process == GetCurrentProcess() && (Type & MEM_RELEASE))
 	{
 		GAllocationTrace->Free(Address, EMemoryTraceRootHeap::SystemMemory);
 	}
