[Insights]
  - Fixed memory tracking for an FRunnableThread to emit a free event when the thread ends.
  - Fixed memory tracking for "ProgramSize" of loaded modules (dlls), to emit a free event when dll is unloaded.

diff --git a/Engine/Source/Runtime/Core/Private/ProfilingDebugging/Microsoft/WindowsModuleDiagnostics.cpp b/Engine/Source/Runtime/Core/Private/ProfilingDebugging/Microsoft/WindowsModuleDiagnostics.cpp
--- a/Engine/Source/Runtime/Core/Private/ProfilingDebugging/Microsoft/WindowsModuleDiagnostics.cpp
+++ b/Engine/Source/Runtime/Core/Private/ProfilingDebugging/Microsoft/WindowsModuleDiagnostics.cpp
@@ -81,11 +81,8 @@ FModuleTrace* FModuleTrace::Get()
 void FModuleTrace::Initialize()
 {
 	using namespace UE::Trace;
-	
-	ProgramHeapId = MemoryTrace_HeapSpec(
-		SystemMemory, 
-		TEXT("Program"), 
-		EMemoryTraceHeapFlags::NeverFrees);
+
+	ProgramHeapId = MemoryTrace_HeapSpec(SystemMemory, TEXT("Module"), EMemoryTraceHeapFlags::None);
 
 	UE_TRACE_LOG(Diagnostics, ModuleInit, ModuleChannel, sizeof(char) * 3)
 		<< ModuleInit.SymbolFormat("pdb", 3)
@@ -202,9 +199,7 @@ void FModuleTrace::OnDllLoaded(const UNICODE_STRING& Name, UPTRINT Base)
 		<< ModuleLoad.ImageId(ImageId.GetData(), ImageId.Num());
 
 	{
-#if ENABLE_LOW_LEVEL_MEM_TRACKER
-		UE_MEMSCOPE(ELLMTag::ProgramSize);
-#endif
+		LLM(UE_MEMSCOPE(ELLMTag::ProgramSize));
 		MemoryTrace_Alloc(Base, OptionalHeader.SizeOfImage, 4 * 1024);
 		MemoryTrace_MarkAllocAsHeap(Base, ProgramHeapId);
 		MemoryTrace_Alloc(Base, OptionalHeader.SizeOfImage, 4 * 1024);
@@ -219,6 +214,8 @@ void FModuleTrace::OnDllLoaded(const UNICODE_STRING& Name, UPTRINT Base)
 ////////////////////////////////////////////////////////////////////////////////
 void FModuleTrace::OnDllUnloaded(UPTRINT Base)
 {
+	MemoryTrace_Free(Base);
+
 	UE_TRACE_LOG(Diagnostics, ModuleUnload, ModuleChannel)
 		<< ModuleUnload.Base(uint64(Base));
 
diff --git a/Engine/Source/Runtime/Core/Private/Windows/WindowsRunnableThread.cpp b/Engine/Source/Runtime/Core/Private/Windows/WindowsRunnableThread.cpp
--- a/Engine/Source/Runtime/Core/Private/Windows/WindowsRunnableThread.cpp
+++ b/Engine/Source/Runtime/Core/Private/Windows/WindowsRunnableThread.cpp
@@ -145,6 +145,15 @@ uint32 FRunnableThreadWin::Run()
 		// Allow any allocated resources to be cleaned up
 		Runnable->Exit();
 
+#if UE_MEMORY_TRACE_ENABLED || ENABLE_LOW_LEVEL_MEM_TRACKER
+		const uint64 FakeAddress = uint64(this) ^ (1ull << 47);
+		MemoryTrace_Free(FakeAddress);
+		LLM(FLowLevelMemTracker::Get().OnLowLevelFree(ELLMTracker::Default, (const void*)FakeAddress));
+		LLM(FLowLevelMemTracker::Get().OnLowLevelFree(ELLMTracker::Platform, (const void*)FakeAddress));
+#endif // UE_MEMORY_TRACE_ENABLED || ENABLE_LOW_LEVEL_MEM_TRACKER
+
+		Runnable = nullptr; // skips the Runnable->Stop() if Kill() is further called
+
 #if STATS
 		FThreadStats::Shutdown();
 #endif
diff --git a/Engine/Source/Runtime/Core/Private/Windows/WindowsRunnableThread.h b/Engine/Source/Runtime/Core/Private/Windows/WindowsRunnableThread.h
--- a/Engine/Source/Runtime/Core/Private/Windows/WindowsRunnableThread.h
+++ b/Engine/Source/Runtime/Core/Private/Windows/WindowsRunnableThread.h
@@ -85,6 +85,13 @@ public:
 		if (Runnable)
 		{
 			Runnable->Stop();
+
+#if UE_MEMORY_TRACE_ENABLED || ENABLE_LOW_LEVEL_MEM_TRACKER
+			const uint64 FakeAddress = uint64(this) ^ (1ull << 47);
+			MemoryTrace_Free(FakeAddress);
+			LLM(FLowLevelMemTracker::Get().OnLowLevelFree(ELLMTracker::Default, (const void*)FakeAddress));
+			LLM(FLowLevelMemTracker::Get().OnLowLevelFree(ELLMTracker::Platform, (const void*)FakeAddress));
+#endif // UE_MEMORY_TRACE_ENABLED || ENABLE_LOW_LEVEL_MEM_TRACKER
 		}
 
 		if (bShouldWait == true)
@@ -140,18 +147,20 @@ protected:
 
 		// Create the new thread
 		{
+#if UE_MEMORY_TRACE_ENABLED || ENABLE_LOW_LEVEL_MEM_TRACKER
 			LLM_SCOPE(ELLMTag::ThreadStack);
 			LLM_PLATFORM_SCOPE(ELLMTag::ThreadStackPlatform);
-			// add in the thread size, since it's allocated in a black box we can't track
-			// note: I don't see any accounting for this when threads are destroyed
-			const uint64 FakeAddress = uint64(InRunnable) | (uint64(1u) << 48);
-			// Size of zero indicates using default thread stack size. This is 1MB unless overridden in .def file.
-			const uint64 Size = InStackSize > 0 ? InStackSize : 1 * 1024 * 1024;
-			MemoryTrace_Alloc(FakeAddress, Size, 4);
+			const uint64 FakeAddress = uint64(this) ^ (1ull << 47);
+			constexpr uint64 DefaultStackSize = 1024ull * 1024ull; // 1 MB, unless overridden in the .def file
+			constexpr uint32 DefaultAlignment = 64 * 1024; // 64 KB, typical system's allocation granularity
+			// Size of zero indicates using default thread stack size.
+			const uint64 ActualStackSize = (InStackSize == 0) ? DefaultStackSize : uint64(InStackSize);
+			MemoryTrace_Alloc(FakeAddress, ActualStackSize, DefaultAlignment);
 			MemoryTrace_MarkAllocAsHeap(FakeAddress, EMemoryTraceRootHeap::SystemMemory);
-			MemoryTrace_Alloc(FakeAddress, Size, 4);
-			LLM(FLowLevelMemTracker::Get().OnLowLevelAlloc(ELLMTracker::Default, nullptr, InStackSize));
-			LLM(FLowLevelMemTracker::Get().OnLowLevelAlloc(ELLMTracker::Platform, nullptr, InStackSize));
+			MemoryTrace_Alloc(FakeAddress, ActualStackSize, DefaultAlignment);
+			LLM(FLowLevelMemTracker::Get().OnLowLevelAlloc(ELLMTracker::Default, (const void*)FakeAddress, ActualStackSize));
+			LLM(FLowLevelMemTracker::Get().OnLowLevelAlloc(ELLMTracker::Platform, (const void*)FakeAddress, ActualStackSize));
+#endif // UE_MEMORY_TRACE_ENABLED || ENABLE_LOW_LEVEL_MEM_TRACKER
 
 			// Create the thread as suspended, so we can ensure ThreadId is initialized and the thread manager knows about the thread before it runs.
 			Thread = CreateThread(NULL, InStackSize, _ThreadProc, this, STACK_SIZE_PARAM_IS_A_RESERVATION | CREATE_SUSPENDED, (::DWORD*)&ThreadID);
@@ -160,6 +169,13 @@ protected:
 		// If it fails, clear all the vars
 		if (Thread == NULL)
 		{
+#if UE_MEMORY_TRACE_ENABLED || ENABLE_LOW_LEVEL_MEM_TRACKER
+			const uint64 FakeAddress = uint64(this) ^ (1ull << 47);
+			MemoryTrace_Free(FakeAddress);
+			LLM(FLowLevelMemTracker::Get().OnLowLevelFree(ELLMTracker::Default, (const void*)FakeAddress));
+			LLM(FLowLevelMemTracker::Get().OnLowLevelFree(ELLMTracker::Platform, (const void*)FakeAddress));
+#endif // UE_MEMORY_TRACE_ENABLED || ENABLE_LOW_LEVEL_MEM_TRACKER
+
 			Runnable = nullptr;
 		}
 		else
