Added default param when faking a forked process.

diff --git a/Engine/Source/Runtime/Core/Public/Misc/Fork.h b/Engine/Source/Runtime/Core/Public/Misc/Fork.h
--- a/Engine/Source/Runtime/Core/Public/Misc/Fork.h
+++ b/Engine/Source/Runtime/Core/Public/Misc/Fork.h
@@ -44,7 +44,7 @@ public:
 	/**
 	 * Sets the forked child process flag and index given to this child process
 	 */
-	static void SetIsForkedChildProcess(uint16 ChildIndex);
+	static void SetIsForkedChildProcess(uint16 ChildIndex=1);
 
 	/**
 	* Returns the unique index of this forked child process. Index 0 is for the master server
