Fix for spurious FileCache/CompressionBlockSize warning:
-- When not using iostore containers
-- When not using compression

diff --git a/Engine/Source/Runtime/Core/Private/FileCache/FileCache.cpp b/Engine/Source/Runtime/Core/Private/FileCache/FileCache.cpp
--- a/Engine/Source/Runtime/Core/Private/FileCache/FileCache.cpp
+++ b/Engine/Source/Runtime/Core/Private/FileCache/FileCache.cpp
@@ -349,7 +349,8 @@ FFileCache::FFileCache(int32 NumSlots)
 		UE_LOG(LogStreamingFileCache, Warning, TEXT("	Check your IoStore compression block size (Project Settings -> 'Package Compression Commandline Options'"));
 		UE_LOG(LogStreamingFileCache, Warning, TEXT("	and your File Cache block size (fc.BlockSize cvar). They should match!"));
 	}
-	else if (IoStoreCompressionBlockSize != CacheSlotID::GetSize())
+	else if (IoStoreCompressionBlockSize && // If we load without any iostore containers (i.e. editor) we don't need to warn.
+		IoStoreCompressionBlockSize != CacheSlotID::GetSize())
 	{
 		UE_LOG(LogStreamingFileCache, Warning, TEXT("IoStore containers have a different block sizes than FileCache (%d vs %d)!"), IoStoreCompressionBlockSize, CacheSlotID::GetSize());
 		UE_LOG(LogStreamingFileCache, Warning, TEXT("	Check your IoStore compression block size (Project Settings -> 'Package Compression Commandline Options'"));
diff --git a/Engine/Source/Runtime/PakFile/Private/IoDispatcherFileBackend.cpp b/Engine/Source/Runtime/PakFile/Private/IoDispatcherFileBackend.cpp
--- a/Engine/Source/Runtime/PakFile/Private/IoDispatcherFileBackend.cpp
+++ b/Engine/Source/Runtime/PakFile/Private/IoDispatcherFileBackend.cpp
@@ -899,7 +899,11 @@ TIoStatusOr<FIoContainerHeader> FFileIoStoreReader::ReadContainerHeader() const
 	const uint64 RawOffset = CompressionBlockEntry->GetOffset() % ContainerFile.PartitionSize;
 
 #if !UE_BUILD_SHIPPING
-	FileCache_PostIoStoreCompressionBlockSize(CompressionBlockSize, Partition.FilePath);
+	// Check for flag - compressed containers still have CompressionBlockSize != 0 and CompressionMethod "None".
+	if (EnumHasAnyFlags(ContainerFile.ContainerFlags, EIoContainerFlags::Compressed)) 
+	{
+		FileCache_PostIoStoreCompressionBlockSize(CompressionBlockSize, Partition.FilePath);
+	}
 #endif
 
 	FIoBuffer IoBuffer(Align(Size, FAES::AESBlockSize));
