#ue5 - Document VectorCross() for VectorRegister implementations noting that input W is expected to be 0, and the result W will remain 0 if so.

diff --git a/Engine/Source/Runtime/Core/Public/Math/UnrealMathFPU.h b/Engine/Source/Runtime/Core/Public/Math/UnrealMathFPU.h
--- a/Engine/Source/Runtime/Core/Public/Math/UnrealMathFPU.h
+++ b/Engine/Source/Runtime/Core/Public/Math/UnrealMathFPU.h
@@ -1163,11 +1163,11 @@ FORCEINLINE VectorRegister4Double VectorBitwiseXor(const VectorRegister4Double&
 
 
 /**
- * Calculates the cross product of two vectors (XYZ components). W is set to 0.
+ * Calculates the cross product of two vectors (XYZ components). W of the input should be 0, and will remain 0.
  *
  * @param Vec1	1st vector
  * @param Vec2	2nd vector
- * @return		cross(Vec1.xyz, Vec2.xyz). W is set to 0.
+ * @return		cross(Vec1.xyz, Vec2.xyz). W of the input should be 0, and will remain 0.
  */
 FORCEINLINE VectorRegister4Float VectorCross(const VectorRegister4Float& Vec1, const VectorRegister4Float& Vec2)
 {
diff --git a/Engine/Source/Runtime/Core/Public/Math/UnrealMathNeon.h b/Engine/Source/Runtime/Core/Public/Math/UnrealMathNeon.h
--- a/Engine/Source/Runtime/Core/Public/Math/UnrealMathNeon.h
+++ b/Engine/Source/Runtime/Core/Public/Math/UnrealMathNeon.h
@@ -1535,11 +1535,11 @@ FORCEINLINE void VectorDeinterleave(VectorRegister4Double& RESTRICT OutEvens, Ve
 }
 
 /**
- * Calculates the cross product of two vectors (XYZ components). W is set to 0.
+ * Calculates the cross product of two vectors (XYZ components). W of the input should be 0, and will remain 0.
  *
  * @param Vec1	1st vector
  * @param Vec2	2nd vector
- * @return		cross(Vec1.xyz, Vec2.xyz). W is set to 0.
+ * @return		cross(Vec1.xyz, Vec2.xyz). W of the input should be 0, and will remain 0.
  */
 FORCEINLINE VectorRegister4Float VectorCross( const VectorRegister4Float& Vec1, const VectorRegister4Float& Vec2 )
 {
diff --git a/Engine/Source/Runtime/Core/Public/Math/UnrealMathSSE.h b/Engine/Source/Runtime/Core/Public/Math/UnrealMathSSE.h
--- a/Engine/Source/Runtime/Core/Public/Math/UnrealMathSSE.h
+++ b/Engine/Source/Runtime/Core/Public/Math/UnrealMathSSE.h
@@ -1847,11 +1847,11 @@ FORCEINLINE VectorRegister4Double VectorBitwiseXor(const VectorRegister4Double&
 }
 
 /**
- * Calculates the cross product of two vectors (XYZ components). W is set to 0.
+ * Calculates the cross product of two vectors (XYZ components). W of the input should be 0, and will remain 0.
  *
  * @param Vec1	1st vector
  * @param Vec2	2nd vector
- * @return		cross(Vec1.xyz, Vec2.xyz). W is set to 0.
+ * @return		cross(Vec1.xyz, Vec2.xyz). W of the input should be 0, and will remain 0.
  */
 FORCEINLINE VectorRegister4Float VectorCross( const VectorRegister4Float& Vec1, const VectorRegister4Float& Vec2 )
 {
