#ue5 - Avoid repeated loads of PersistentVectorRegisterType in vectorized Transform where possible.

diff --git a/Engine/Source/Runtime/Core/Public/Math/TransformVectorized.h b/Engine/Source/Runtime/Core/Public/Math/TransformVectorized.h
--- a/Engine/Source/Runtime/Core/Public/Math/TransformVectorized.h
+++ b/Engine/Source/Runtime/Core/Public/Math/TransformVectorized.h
@@ -712,11 +712,12 @@ private:
 
 	FORCEINLINE bool Private_RotationEquals( const TransformVectorRegister& InRotation, const FReal ToleranceScalar = KINDA_SMALL_NUMBER) const
 	{			
+		const TransformVectorRegister MyRotation = Rotation; // Load from persistent value, to avoid repeated loads.
 		const TransformVectorRegister Tolerance = VectorLoadFloat1(&ToleranceScalar);
 		// !( (FMath::Abs(X-Q.X) > Tolerance) || (FMath::Abs(Y-Q.Y) > Tolerance) || (FMath::Abs(Z-Q.Z) > Tolerance) || (FMath::Abs(W-Q.W) > Tolerance) )
-		const TransformVectorRegister RotationSub = VectorAbs(VectorSubtract(Rotation, InRotation));		
+		const TransformVectorRegister RotationSub = VectorAbs(VectorSubtract(MyRotation, InRotation));
 		// !( (FMath::Abs(X+Q.X) > Tolerance) || (FMath::Abs(Y+Q.Y) > Tolerance) || (FMath::Abs(Z+Q.Z) > Tolerance) || (FMath::Abs(W+Q.W) > Tolerance) )
-		const TransformVectorRegister RotationAdd = VectorAbs(VectorAdd(Rotation, InRotation));
+		const TransformVectorRegister RotationAdd = VectorAbs(VectorAdd(MyRotation, InRotation));
 		return !VectorAnyGreaterThan(RotationSub, Tolerance) || !VectorAnyGreaterThan(RotationAdd, Tolerance);
 	}
 
@@ -1295,8 +1296,9 @@ private:
 		check( IsRotationNormalized() );
 #endif		
 
-		const TransformVectorRegister RotationX2Y2Z2 = VectorAdd(Rotation, Rotation);	// x2, y2, z2
-		const TransformVectorRegister RotationXX2YY2ZZ2 = VectorMultiply(RotationX2Y2Z2, Rotation);	// xx2, yy2, zz2		
+		const TransformVectorRegister MyRotation = Rotation; // Load from persistent value, to avoid repeated loads.
+		const TransformVectorRegister RotationX2Y2Z2 = VectorAdd(MyRotation, MyRotation);	// x2, y2, z2
+		const TransformVectorRegister RotationXX2YY2ZZ2 = VectorMultiply(RotationX2Y2Z2, MyRotation);	// xx2, yy2, zz2		
 
 		// The diagonal terms of the rotation matrix are:
 		//   (1 - (yy2 + zz2)) * scale
@@ -1316,11 +1318,11 @@ private:
 
 		// RotBase = x*y2, y*z2, x*z2
 		// RotOffset = w*z2, w*x2, w*y2
-		const TransformVectorRegister x_y_x = VectorSwizzle(Rotation, 0, 1, 0, 0);
+		const TransformVectorRegister x_y_x = VectorSwizzle(MyRotation, 0, 1, 0, 0);
 		const TransformVectorRegister y2_z2_z2 = VectorSwizzle(RotationX2Y2Z2, 1, 2, 2, 0);
 		const TransformVectorRegister RotBase = VectorMultiply(x_y_x, y2_z2_z2);
 
-		const TransformVectorRegister w_w_w = VectorReplicate(Rotation, 3);
+		const TransformVectorRegister w_w_w = VectorReplicate(MyRotation, 3);
 		const TransformVectorRegister z2_x2_y2 = VectorSwizzle(RotationX2Y2Z2, 2, 0, 1, 0);
 		const TransformVectorRegister RotOffset = VectorMultiply(w_w_w, z2_x2_y2);
 
@@ -1339,8 +1341,9 @@ private:
 		// Make sure Rotation is normalized when we turn it into a matrix.
 		ensure( IsRotationNormalized() );
 #endif		
-		const TransformVectorRegister RotationX2Y2Z2 = VectorAdd(Rotation, Rotation);	// x2, y2, z2
-		const TransformVectorRegister RotationXX2YY2ZZ2 = VectorMultiply(RotationX2Y2Z2, Rotation);	// xx2, yy2, zz2		
+		const TransformVectorRegister MyRotation = Rotation; // Load from persistent value, to avoid repeated loads.
+		const TransformVectorRegister RotationX2Y2Z2 = VectorAdd(MyRotation, MyRotation);	// x2, y2, z2
+		const TransformVectorRegister RotationXX2YY2ZZ2 = VectorMultiply(RotationX2Y2Z2, MyRotation);	// xx2, yy2, zz2		
 
 		// The diagonal terms of the rotation matrix are:
 		//   (1 - (yy2 + zz2))
@@ -1359,11 +1362,11 @@ private:
 
 		// RotBase = x*y2, y*z2, x*z2
 		// RotOffset = w*z2, w*x2, w*y2
-		const TransformVectorRegister x_y_x = VectorSwizzle(Rotation, 0, 1, 0, 0);
+		const TransformVectorRegister x_y_x = VectorSwizzle(MyRotation, 0, 1, 0, 0);
 		const TransformVectorRegister y2_z2_z2 = VectorSwizzle(RotationX2Y2Z2, 1, 2, 2, 0);
 		const TransformVectorRegister RotBase = VectorMultiply(x_y_x, y2_z2_z2);
 
-		const TransformVectorRegister w_w_w = VectorReplicate(Rotation, 3);
+		const TransformVectorRegister w_w_w = VectorReplicate(MyRotation, 3);
 		const TransformVectorRegister z2_x2_y2 = VectorSwizzle(RotationX2Y2Z2, 2, 0, 1, 0);
 		const TransformVectorRegister RotOffset = VectorMultiply(w_w_w, z2_x2_y2);
 
