[Insights]
  - Fixed parsing of Trace.Start/Trace.File console command params to correctly detect "file path" and "list of channels" parameters.
  - Fixed Trace.Disable to not save anymore the previous enabled channels as "paused" preset.
  - Fixed Trace::Connect (Tarce.Send, Trace.File) to not resume paused channels anymore.

diff --git a/Engine/Source/Runtime/Core/Private/ProfilingDebugging/TraceAuxiliary.cpp b/Engine/Source/Runtime/Core/Private/ProfilingDebugging/TraceAuxiliary.cpp
--- a/Engine/Source/Runtime/Core/Private/ProfilingDebugging/TraceAuxiliary.cpp
+++ b/Engine/Source/Runtime/Core/Private/ProfilingDebugging/TraceAuxiliary.cpp
@@ -141,19 +141,33 @@ void FTraceAuxiliaryImpl::ResetCommandlineChannels()
 ////////////////////////////////////////////////////////////////////////////////
 void FTraceAuxiliaryImpl::EnableChannels(const TCHAR* ChannelList)
 {
-	ForEachChannel(ChannelList, true, &FTraceAuxiliaryImpl::EnableChannel);
+	if (ChannelList)
+	{
+		ForEachChannel(ChannelList, true, &FTraceAuxiliaryImpl::EnableChannel);
+	}
 }
 
 ////////////////////////////////////////////////////////////////////////////////
 void FTraceAuxiliaryImpl::DisableChannels(const TCHAR* ChannelList)
 {
-	ForEachChannel(ChannelList, true, &FTraceAuxiliaryImpl::DisableChannel);
+	if (ChannelList)
+	{
+		ForEachChannel(ChannelList, true, &FTraceAuxiliaryImpl::DisableChannel);
+	}
+	else
+	{
+		// Disable all channels.
+		TStringBuilder<128> EnabledChannels;
+		GetActiveChannelsString(EnabledChannels);
+		ForEachChannel(EnabledChannels.ToString(), true, &FTraceAuxiliaryImpl::DisableChannel);
+	}
 }
 
 ////////////////////////////////////////////////////////////////////////////////
 template<typename T>
 void FTraceAuxiliaryImpl::ForEachChannel(const TCHAR* ChannelList, bool bResolvePresets, T Callable)
 {
+	check(ChannelList);
 	UE::String::ParseTokens(ChannelList, TEXT(","), [this, bResolvePresets, Callable] (const FStringView& Token)
 	{
 		TCHAR Name[80];
@@ -191,7 +205,7 @@ uint32 FTraceAuxiliaryImpl::HashChannelName(const TCHAR* Name)
 	for (const TCHAR* c = Name; *c; ++c)
 	{
 		uint32 LowerC = *c | 0x20;
-        Hash = ((Hash << 5) + Hash) + LowerC;
+		Hash = ((Hash << 5) + Hash) + LowerC;
 	}
 	return Hash;
 }
@@ -272,8 +286,6 @@ bool FTraceAuxiliaryImpl::Connect(ETraceConnectType Type, const TCHAR* Parameter
 		return false;
 	}
 
-	ResumeChannels();
-
 	State = EState::Tracing;
 	return true;
 }
@@ -334,22 +346,21 @@ void FTraceAuxiliaryImpl::DisableChannel(const TCHAR* Channel)
 ////////////////////////////////////////////////////////////////////////////////
 void FTraceAuxiliaryImpl::ResumeChannels()
 {
-	// There was paused preset apply that
-	if (!PausedPreset.IsEmpty())
-	{
-		ForEachChannel(*PausedPreset, false, &FTraceAuxiliaryImpl::EnableChannel);
-	}
+	// Enable channels from the "paused" preset.
+	ForEachChannel(*PausedPreset, false, &FTraceAuxiliaryImpl::EnableChannel);
 }
 
 ////////////////////////////////////////////////////////////////////////////////
 void FTraceAuxiliaryImpl::PauseChannels()
 {
-	// Collect current set of active channels
 	TStringBuilder<128> EnabledChannels;
 	GetActiveChannelsString(EnabledChannels);
+
+	// Save the list of enabled channels as the current "paused" preset.
+	// The "paused" preset can only be used in the Trace.Resume command / API.
 	PausedPreset = EnabledChannels.ToString();
 
-	// Disable those channels
+	// Disable all "paused" channels.
 	ForEachChannel(*PausedPreset, true, &FTraceAuxiliaryImpl::DisableChannel);
 }
 
@@ -610,10 +621,10 @@ static void TraceAuxiliaryFile(const TArray<FString>& Args)
 	else if (Args.Num() == 1)
 	{
 		// Try to detect if the first argument is a file path.
-		if (!FCString::Strchr(*Args[0], TEXT('/')) ||
-			!FCString::Strchr(*Args[0], TEXT('\\')) ||
-			!FCString::Strchr(*Args[0], TEXT('.')) ||
-			!FCString::Strchr(*Args[0], TEXT(':')))
+		if (FCString::Strchr(*Args[0], TEXT('/')) ||
+			FCString::Strchr(*Args[0], TEXT('\\')) ||
+			FCString::Strchr(*Args[0], TEXT('.')) ||
+			FCString::Strchr(*Args[0], TEXT(':')))
 		{
 			Filepath = *Args[0];
 		}
@@ -755,7 +766,7 @@ static void TraceAuxiliaryDisableChannels(const TArray<FString>& Args)
 {
 	if (Args.Num() == 0)
 	{
-		GTraceAuxiliary.PauseChannels();
+		GTraceAuxiliary.DisableChannels(nullptr);
 	}
 	else
 	{
