[Insights] Fixed memory tracking for "ProgramSize" of loaded modules (dlls), to emit a free event when dll is unloaded.

diff --git a/Engine/Source/Runtime/Core/Private/ProfilingDebugging/Microsoft/WindowsModuleDiagnostics.cpp b/Engine/Source/Runtime/Core/Private/ProfilingDebugging/Microsoft/WindowsModuleDiagnostics.cpp
--- a/Engine/Source/Runtime/Core/Private/ProfilingDebugging/Microsoft/WindowsModuleDiagnostics.cpp
+++ b/Engine/Source/Runtime/Core/Private/ProfilingDebugging/Microsoft/WindowsModuleDiagnostics.cpp
@@ -81,11 +81,8 @@ FModuleTrace* FModuleTrace::Get()
 void FModuleTrace::Initialize()
 {
 	using namespace UE::Trace;
-	
-	ProgramHeapId = MemoryTrace_HeapSpec(
-		SystemMemory, 
-		TEXT("Program"), 
-		EMemoryTraceHeapFlags::NeverFrees);
+
+	ProgramHeapId = MemoryTrace_HeapSpec(SystemMemory, TEXT("Module"), EMemoryTraceHeapFlags::None);
 
 	UE_TRACE_LOG(Diagnostics, ModuleInit, ModuleChannel, sizeof(char) * 3)
 		<< ModuleInit.SymbolFormat("pdb", 3)
@@ -219,6 +216,8 @@ void FModuleTrace::OnDllLoaded(const UNICODE_STRING& Name, UPTRINT Base)
 ////////////////////////////////////////////////////////////////////////////////
 void FModuleTrace::OnDllUnloaded(UPTRINT Base)
 {
+	MemoryTrace_Free(Base);
+
 	UE_TRACE_LOG(Diagnostics, ModuleUnload, ModuleChannel)
 		<< ModuleUnload.Base(uint64(Base));
 
