[Insights] Fixed memory tracking for an FRunnableThread to emit a free event when the thread ends (this time without resetting Runnable to nullptr when thread ends).

diff --git a/Engine/Source/Runtime/Core/Private/Windows/WindowsRunnableThread.h b/Engine/Source/Runtime/Core/Private/Windows/WindowsRunnableThread.h
--- a/Engine/Source/Runtime/Core/Private/Windows/WindowsRunnableThread.h
+++ b/Engine/Source/Runtime/Core/Private/Windows/WindowsRunnableThread.h
@@ -103,6 +103,13 @@ public:
 		CloseHandle(Thread);
 		Thread = NULL;
 
+#if UE_MEMORY_TRACE_ENABLED || ENABLE_LOW_LEVEL_MEM_TRACKER
+		const uint64 FakeAddress = uint64(this) | (1ull << 47);
+		MemoryTrace_Free(FakeAddress);
+		LLM(FLowLevelMemTracker::Get().OnLowLevelFree(ELLMTracker::Default, (const void*)FakeAddress));
+		LLM(FLowLevelMemTracker::Get().OnLowLevelFree(ELLMTracker::Platform, (const void*)FakeAddress));
+#endif // UE_MEMORY_TRACE_ENABLED || ENABLE_LOW_LEVEL_MEM_TRACKER
+
 		return bDidExitOK;
 	}
 
@@ -140,18 +147,20 @@ protected:
 
 		// Create the new thread
 		{
+#if UE_MEMORY_TRACE_ENABLED || ENABLE_LOW_LEVEL_MEM_TRACKER
 			LLM_SCOPE(ELLMTag::ThreadStack);
 			LLM_PLATFORM_SCOPE(ELLMTag::ThreadStackPlatform);
-			// add in the thread size, since it's allocated in a black box we can't track
-			// note: I don't see any accounting for this when threads are destroyed
-			const uint64 FakeAddress = uint64(InRunnable) | (uint64(1u) << 48);
-			// Size of zero indicates using default thread stack size. This is 1MB unless overridden in .def file.
-			const uint64 Size = InStackSize > 0 ? InStackSize : 1 * 1024 * 1024;
-			MemoryTrace_Alloc(FakeAddress, Size, 4);
+			const uint64 FakeAddress = uint64(this) | (1ull << 47);
+			constexpr uint64 DefaultStackSize = 1024 * 1024; // 1 MiB, unless overridden in the .def file
+			constexpr uint32 DefaultAlignment = 64 * 1024; // 64 KiB, typical system's allocation granularity
+			// Size of zero indicates using default thread stack size.
+			const uint64 ActualStackSize = (InStackSize == 0) ? DefaultStackSize : uint64(InStackSize);
+			MemoryTrace_Alloc(FakeAddress, ActualStackSize, DefaultAlignment);
 			MemoryTrace_MarkAllocAsHeap(FakeAddress, EMemoryTraceRootHeap::SystemMemory);
-			MemoryTrace_Alloc(FakeAddress, Size, 4);
-			LLM(FLowLevelMemTracker::Get().OnLowLevelAlloc(ELLMTracker::Default, nullptr, InStackSize));
-			LLM(FLowLevelMemTracker::Get().OnLowLevelAlloc(ELLMTracker::Platform, nullptr, InStackSize));
+			MemoryTrace_Alloc(FakeAddress, ActualStackSize, DefaultAlignment);
+			LLM(FLowLevelMemTracker::Get().OnLowLevelAlloc(ELLMTracker::Default, (const void*)FakeAddress, ActualStackSize));
+			LLM(FLowLevelMemTracker::Get().OnLowLevelAlloc(ELLMTracker::Platform, (const void*)FakeAddress, ActualStackSize));
+#endif // UE_MEMORY_TRACE_ENABLED || ENABLE_LOW_LEVEL_MEM_TRACKER
 
 			// Create the thread as suspended, so we can ensure ThreadId is initialized and the thread manager knows about the thread before it runs.
 			Thread = CreateThread(NULL, InStackSize, _ThreadProc, this, STACK_SIZE_PARAM_IS_A_RESERVATION | CREATE_SUSPENDED, (::DWORD*)&ThreadID);
@@ -160,6 +169,13 @@ protected:
 		// If it fails, clear all the vars
 		if (Thread == NULL)
 		{
+#if UE_MEMORY_TRACE_ENABLED || ENABLE_LOW_LEVEL_MEM_TRACKER
+			const uint64 FakeAddress = uint64(this) | (1ull << 47);
+			MemoryTrace_Free(FakeAddress);
+			LLM(FLowLevelMemTracker::Get().OnLowLevelFree(ELLMTracker::Default, (const void*)FakeAddress));
+			LLM(FLowLevelMemTracker::Get().OnLowLevelFree(ELLMTracker::Platform, (const void*)FakeAddress));
+#endif // UE_MEMORY_TRACE_ENABLED || ENABLE_LOW_LEVEL_MEM_TRACKER
+
 			Runnable = nullptr;
 		}
 		else
