Fix for backwards compat issue introduced in 4.26

diff --git a/Engine/Source/Runtime/Core/Private/HAL/ConsoleManager.cpp b/Engine/Source/Runtime/Core/Private/HAL/ConsoleManager.cpp
--- a/Engine/Source/Runtime/Core/Private/HAL/ConsoleManager.cpp
+++ b/Engine/Source/Runtime/Core/Private/HAL/ConsoleManager.cpp
@@ -116,7 +116,7 @@ public:
 	virtual void SetOnChangedCallback(const FConsoleVariableDelegate& Callback) 
 	{
 		OnChangedCallback.Remove(LegacyDelegateHandle);
-		OnChangedCallback.Add(Callback); 
+		LegacyDelegateHandle = OnChangedCallback.Add(Callback); 
 	}
 
 	/** Returns a multicast delegate with which to register. Called when this CVar changes. */
