Add ToHostLineEndings FromHostLineEndings functions for FString for converting between different platform line endings

diff --git a/Engine/Source/Runtime/Core/Private/String/LineEndings.cpp b/Engine/Source/Runtime/Core/Private/String/LineEndings.cpp
new file mode 100644
--- /dev/null
+++ b/Engine/Source/Runtime/Core/Private/String/LineEndings.cpp
@@ -0,0 +1,48 @@
+// Copyright Epic Games, Inc. All Rights Reserved.
+
+#include "String/LineEndings.h"
+
+namespace UE::String
+{
+	CORE_API FString FromHostLineEndings(const FString& InString)
+	{
+		FString New = InString;
+		FromHostLineEndingsInline(New);
+		return New;
+	}
+
+	CORE_API FString FromHostLineEndings(FString&& InString)
+	{
+		FromHostLineEndingsInline(InString);
+		return MoveTemp(InString);
+	}
+
+	CORE_API void FromHostLineEndingsInline(FString& InString)
+	{
+		InString.ReplaceInline(TEXT("\r\n"), TEXT("\n"));
+		InString.ReplaceInline(TEXT("\r"), TEXT("\n"));
+	}
+
+	CORE_API FString ToHostLineEndings(const FString& InString)
+	{
+		FString New = InString;
+		ToHostLineEndingsInline(New);
+		return New;
+	}
+
+	CORE_API FString ToHostLineEndings(FString&& InString)
+	{
+		ToHostLineEndingsInline(InString);
+		return MoveTemp(InString);
+	}
+
+	CORE_API void ToHostLineEndingsInline(FString& InString)
+	{
+		FromHostLineEndingsInline(InString);
+
+		if (TCString<TCHAR>::Strcmp(TEXT("\n"), LINE_TERMINATOR) != 0)
+		{
+			InString.ReplaceInline(TEXT("\n"), LINE_TERMINATOR);
+		}
+	}
+} // UE::String
diff --git a/Engine/Source/Runtime/Core/Public/String/LineEndings.h b/Engine/Source/Runtime/Core/Public/String/LineEndings.h
new file mode 100644
--- /dev/null
+++ b/Engine/Source/Runtime/Core/Public/String/LineEndings.h
@@ -0,0 +1,20 @@
+// Copyright Epic Games, Inc. All Rights Reserved.
+
+#pragma once
+
+#include "CoreTypes.h"
+#include "Containers/UnrealString.h"
+
+namespace UE::String
+{
+	/** Replaces all Line Endings with "\n" line terminator */
+	CORE_API FString FromHostLineEndings(const FString& InString);
+	CORE_API FString FromHostLineEndings(FString&& InString);
+	CORE_API void FromHostLineEndingsInline(FString& InString);
+
+	/** Replaces all Line Endings with the host platform line terminator */
+	CORE_API FString ToHostLineEndings(const FString& InString);
+	CORE_API FString ToHostLineEndings(FString&& InString);
+	CORE_API void ToHostLineEndingsInline(FString& InString);
+
+} // UE::String
