Clarified the comment on TIndexedContainerIterator::SetToEnd()

#rnx

diff --git a/Engine/Source/Runtime/Core/Public/Containers/Array.h b/Engine/Source/Runtime/Core/Public/Containers/Array.h
--- a/Engine/Source/Runtime/Core/Public/Containers/Array.h
+++ b/Engine/Source/Runtime/Core/Public/Containers/Array.h
@@ -138,7 +138,7 @@ public:
 		Index = 0;
 	}
 
-	/** Sets iterator to the last element. */
+	/** Sets the iterator to one past the last element. */
 	void SetToEnd()
 	{
 		Index = Container.Num();
