LLT 2.0 on preflights
- bIncludeAllTests flag on target will compile with all tests: used on LowLevelTests target to run on preflights
- Preflights now running low level tests by default, only those from LowLevelTests and its dependencies stemming from Core and Projects

diff --git a/Engine/Source/Programs/LowLevelTests/LowLevelTests.Target.cs b/Engine/Source/Programs/LowLevelTests/LowLevelTests.Target.cs
--- a/Engine/Source/Programs/LowLevelTests/LowLevelTests.Target.cs
+++ b/Engine/Source/Programs/LowLevelTests/LowLevelTests.Target.cs
@@ -15,6 +15,12 @@ public class LowLevelTestsTarget : TargetRules
 		Type = TargetType.Program;
 		LinkType = TargetLinkType.Monolithic;
 
+		// LowLevelTests to produce executable with Core tests
+		if (GetType() == typeof(LowLevelTestsTarget))
+		{
+			bIncludeAllTestsOverride = true;
+		}
+
 		bDeployAfterCompile = false;
 		bIsBuildingConsoleApplication = true;
 
diff --git a/Engine/Source/Programs/LowLevelTests/Private/SampleTest.cpp b/Engine/Source/Programs/LowLevelTests/Private/SampleTest.cpp
deleted file mode 100644
diff --git a/Engine/Source/Programs/LowLevelTests/Private/main.cpp b/Engine/Source/Programs/LowLevelTests/Private/main.cpp
deleted file mode 100644
diff --git a/Engine/Source/Programs/LowLevelTests/Tests/main.cpp b/Engine/Source/Programs/LowLevelTests/Tests/main.cpp
new file mode 100644
--- /dev/null
+++ b/Engine/Source/Programs/LowLevelTests/Tests/main.cpp
@@ -0,0 +1,11 @@
+// Copyright Epic Games, Inc. All Rights Reserved.
+#include "TestRunner.h"
+#include "RequiredProgramMainCPPInclude.h"
+
+IMPLEMENT_APPLICATION(SampleTests, "SampleTests");
+
+void Setup() {
+}
+
+void Teardown() {
+}
\ No newline at end of file
diff --git a/Engine/Source/Programs/UnrealBuildTool/Configuration/TargetRules.cs b/Engine/Source/Programs/UnrealBuildTool/Configuration/TargetRules.cs
--- a/Engine/Source/Programs/UnrealBuildTool/Configuration/TargetRules.cs
+++ b/Engine/Source/Programs/UnrealBuildTool/Configuration/TargetRules.cs
@@ -379,6 +379,18 @@ namespace UnrealBuildTool
 		[CommandLine("-AllModules")]
 		public bool bBuildAllModules = false;
 
+		/// <summary>
+		/// Decides whether to compile with all the "Tests" folders from all dependent modules.
+		/// </summary>
+		public bool bIncludeAllTests
+		{
+			get { return bIncludeAllTestsOverride ?? false; }
+		}
+		/// <summary>
+		/// Set this override to true in derived target classes to compile with all the "Tests" folders from all dependent modules.
+		/// </summary>
+		protected bool? bIncludeAllTestsOverride;
+
 		/// <summary>
 		/// Additional plugins that are built for this target type but not enabled.
 		/// </summary>
@@ -2135,6 +2147,11 @@ namespace UnrealBuildTool
 			get { return Inner.bBuildAllModules; }
 		}
 
+		public bool bIncludeAllTests
+		{
+			get { return Inner.bIncludeAllTests; }
+		}
+		
 		public IEnumerable<string> AdditionalPlugins
 		{
 			get { return Inner.AdditionalPlugins; }
diff --git a/Engine/Source/Programs/UnrealBuildTool/Configuration/TestModuleRules.cs b/Engine/Source/Programs/UnrealBuildTool/Configuration/TestModuleRules.cs
--- a/Engine/Source/Programs/UnrealBuildTool/Configuration/TestModuleRules.cs
+++ b/Engine/Source/Programs/UnrealBuildTool/Configuration/TestModuleRules.cs
@@ -62,8 +62,18 @@ namespace UnrealBuildTool
 			PrivateIncludePathModuleNames.Add("Launch");
 
 			// Tests can refer to tested module's Public and Private paths
-			PublicIncludePaths.Add(Path.Combine(TestedModule.ModuleDirectory, "Public"));
-			PrivateIncludePaths.Add(Path.Combine(TestedModule.ModuleDirectory, "Private"));
+			string ModulePublicDir = Path.Combine(TestedModule.ModuleDirectory, "Public");
+			if (System.IO.Directory.Exists(ModulePublicDir))
+			{
+				PublicIncludePaths.Add(ModulePublicDir);
+			}
+
+			string ModulePrivateDir = Path.Combine(TestedModule.ModuleDirectory, "Private");
+			if (System.IO.Directory.Exists(ModulePrivateDir))
+			{
+				PrivateIncludePaths.Add(ModulePrivateDir);
+			}
+			
 
 			// Platforms specific setup
 			if (Target.Platform == UnrealTargetPlatform.Android)
diff --git a/Engine/Source/Programs/UnrealBuildTool/Configuration/UEBuildBinary.cs b/Engine/Source/Programs/UnrealBuildTool/Configuration/UEBuildBinary.cs
--- a/Engine/Source/Programs/UnrealBuildTool/Configuration/UEBuildBinary.cs
+++ b/Engine/Source/Programs/UnrealBuildTool/Configuration/UEBuildBinary.cs
@@ -718,7 +718,7 @@ namespace UnrealBuildTool
 				{
 					// Compile each module.
 					Log.TraceVerbose("Compile module: " + Module.Name);
-					LinkInputFiles = Module.Compile(Target, ToolChain, BinaryCompileEnvironment, SpecificFilesToCompile, WorkingSet, Graph, Module.Rules.bBuildIncludeTestsFolder);
+					LinkInputFiles = Module.Compile(Target, ToolChain, BinaryCompileEnvironment, SpecificFilesToCompile, WorkingSet, Graph, Module.Rules.bBuildIncludeTestsFolder || Target.bIncludeAllTests);
 
 					// Save the module outputs. In monolithic builds, this is just the object files.
 					if (Target.LinkType == TargetLinkType.Monolithic)
diff --git a/Engine/Source/Programs/UnrealBuildTool/Configuration/UEBuildTarget.cs b/Engine/Source/Programs/UnrealBuildTool/Configuration/UEBuildTarget.cs
--- a/Engine/Source/Programs/UnrealBuildTool/Configuration/UEBuildTarget.cs
+++ b/Engine/Source/Programs/UnrealBuildTool/Configuration/UEBuildTarget.cs
@@ -4160,10 +4160,12 @@ namespace UnrealBuildTool
 				DirectoryReference ModuleDirectory = RulesObject.File.Directory;
 
 				// If we're building a test executable, include test harness dependency
-				if (!IsTestModule && ReferenceChain.Contains(typeof(TestModuleRules).Name) && Directory.Exists(RulesObject.TestsDirectory))
+				if (!IsTestModule && ReferenceChain.Contains(typeof(TestModuleRules).Name) && Directory.Exists(RulesObject.TestsDirectory) || RulesObject.Target.bIncludeAllTests)
 				{
-					// TODO: If explicit Tests module rules exists, use that one instead
-					RulesObject.PrivateDependencyModuleNames.Add("LowLevelTestsRunner");
+					if (RulesObject.Name != "LowLevelTestsRunner")
+					{
+						RulesObject.PrivateDependencyModuleNames.Add("LowLevelTestsRunner");
+					}
 				}
 
 				// Clear the bUsePrecompiled flag if we're compiling a foreign plugin; since it's treated like an engine module, it will default to true in an installed build.
diff --git a/Engine/Source/Runtime/Core/Tests/Containers/Deque.cpp b/Engine/Source/Runtime/Core/Tests/Containers/Deque.cpp
new file mode 100644
--- /dev/null
+++ b/Engine/Source/Runtime/Core/Tests/Containers/Deque.cpp
@@ -0,0 +1,930 @@
+// Copyright Epic Games, Inc. All Rights Reserved.
+
+#include "Containers/Deque.h"
+#include "CoreMinimal.h"
+#include "TestToken.h"
+#include "TestUtils.h"
+#include "TestHarness.h"
+
+namespace Deque
+{
+namespace Test
+{
+static constexpr int32 DefaultCapacity = 4;
+
+/**
+ * Emplaces the parameter Count elements into the parameter queue and pops them one by one validating FIFO ordering.
+ * This method also verifies ther correctness of: Last() and First()
+ */
+bool EmplaceLastPopFirst(TDeque<int32Token>& Deque);
+bool EmplaceLastPopFirst(TDeque<int32Token>& Deque, int32 Count);
+bool EmplaceFirstPopLast(TDeque<int32Token>& Deque);
+bool EmplaceFirstPopLast(TDeque<int32Token>& Deque, int32 Count);
+
+}  // namespace Test
+
+//---------------------------------------------------------------------------------------------------------------------
+// Deque::Test
+//---------------------------------------------------------------------------------------------------------------------
+
+bool Test::EmplaceLastPopFirst(TDeque<int32Token>& Deque)
+{
+	ensure(Deque.Max());
+	return EmplaceLastPopFirst(Deque, Deque.Max());
+}
+
+bool Test::EmplaceLastPopFirst(TDeque<int32Token>& Deque, int32 Count)
+{
+	const int32 SeedValue = FMath::RandRange(1, 999);
+	for (int32 i = 0; i < Count; ++i)
+	{
+		Deque.EmplaceLast(SeedValue + i);
+		REQUIRE(Deque.Num() == i + 1);
+		REQUIRE(Deque.First() == SeedValue);
+		REQUIRE(Deque.Last() == SeedValue + i);
+	}
+	for (int32 i = 0; i < Count; ++i)
+	{
+		REQUIRE(Deque.First() == SeedValue + i);
+		REQUIRE(Deque.Last() == SeedValue + Count - 1);
+		Deque.PopFirst();
+		REQUIRE(Deque.Num() == Count - (i + 1));
+	}
+	return true;
+}
+
+bool Test::EmplaceFirstPopLast(TDeque<int32Token>& Deque)
+{
+	ensure(Deque.Max());
+	return EmplaceFirstPopLast(Deque, Deque.Max());
+}
+
+bool Test::EmplaceFirstPopLast(TDeque<int32Token>& Deque, int32 Count)
+{
+	const int32 SeedValue = FMath::RandRange(1, 999);
+	for (int32 i = 0; i < Count; ++i)
+	{
+		Deque.EmplaceFirst(SeedValue + i);
+		REQUIRE(Deque.Num() == i + 1);
+		REQUIRE(Deque.First() == SeedValue + i);
+		REQUIRE(Deque.Last() == SeedValue);
+	}
+	for (int32 i = 0; i < Count; ++i)
+	{
+		REQUIRE(Deque.First() == SeedValue + Count - 1);
+		REQUIRE(Deque.Last() == SeedValue + i);
+		Deque.PopLast();
+		REQUIRE(Deque.Num() == Count - (i + 1));
+	}
+	return true;
+}
+
+//---------------------------------------------------------------------------------------------------------------------
+// Unit tests
+//---------------------------------------------------------------------------------------------------------------------
+
+TEST_CASE("Containers::Deque::Reserve without data", "[Deque]")
+{
+	int32Token::Reset();
+	TDeque<int32Token> Deque;
+	REQUIRE(!Deque.Max());
+	REQUIRE(!Deque.Num());
+	REQUIRE(Deque.IsEmpty());
+	Deque.Reserve(Test::DefaultCapacity);
+	REQUIRE(Deque.Max() >= Test::DefaultCapacity);
+	REQUIRE(!Deque.Num());
+	REQUIRE(Deque.IsEmpty());
+	REQUIRE(int32Token::EvenConstructionDestructionCalls(0));
+}
+
+TEST_CASE("Containers::Deque::Reserve EmplaceLast single element", "[Deque]")
+{
+	int32Token::Reset();
+	{
+		TDeque<int32Token> Deque;
+		REQUIRE(!Deque.Max());
+		REQUIRE(!Deque.Num());
+		Deque.Reserve(Test::DefaultCapacity);
+		Deque.EmplaceLast(0);
+		REQUIRE(Deque.Max() >= Test::DefaultCapacity);
+		REQUIRE(Deque.Max() < Test::DefaultCapacity * 2);
+		REQUIRE(Deque.Num() == 1);
+		Deque.Reserve(Test::DefaultCapacity * 2);
+		REQUIRE(Deque.Max() >= Test::DefaultCapacity * 2);
+		REQUIRE(Deque.Num() == 1);
+	}
+	REQUIRE(int32Token::EvenConstructionDestructionCalls(1));
+}
+
+TEST_CASE("Containers::Deque::Reset", "[Deque]")
+{
+	int32Token::Reset();
+	{
+		TDeque<int32Token> Deque;
+		Deque.Reset();	// Should be innocuous
+		REQUIRE(!Deque.Max());
+		REQUIRE(!Deque.Num());
+		REQUIRE(Deque.IsEmpty());
+		Deque.EmplaceLast(0);
+		REQUIRE(Deque.Max());
+		REQUIRE(Deque.Num() == 1);
+		REQUIRE(!Deque.IsEmpty());
+		Deque.Reset();
+		REQUIRE(Deque.Max());
+		REQUIRE(!Deque.Num());
+		REQUIRE(Deque.IsEmpty());
+	}
+	REQUIRE(int32Token::EvenConstructionDestructionCalls(1));
+}
+
+TEST_CASE("Containers::Deque::Empty", "[Deque]")
+{
+	int32Token::Reset();
+	{
+		TDeque<int32Token> Deque;
+		Deque.Empty();	// Should be innocuous
+		REQUIRE(!Deque.Max());
+		REQUIRE(!Deque.Num());
+		REQUIRE(Deque.IsEmpty());
+	}
+	REQUIRE(int32Token::EvenConstructionDestructionCalls(0));
+}
+
+TEST_CASE("Containers::Deque::Empty after single element EmplaceLast", "[Deque]")
+{
+	int32Token::Reset();
+	{
+		TDeque<int32Token> Deque;
+		Deque.EmplaceLast(0);
+		REQUIRE(Deque.Max());
+		REQUIRE(Deque.Num() == 1);
+		REQUIRE(!Deque.IsEmpty());
+		Deque.Empty();
+		REQUIRE(!Deque.Max());
+		REQUIRE(!Deque.Num());
+		REQUIRE(Deque.IsEmpty());
+	}
+	REQUIRE(int32Token::EvenConstructionDestructionCalls(1));
+}
+
+TEST_CASE("Containers::Deque::EmplaceLast single element", "[Deque]")
+{
+	int32Token::Reset();
+	{
+		TDeque<int32Token> Deque;
+		REQUIRE(!Deque.Max());
+		REQUIRE(!Deque.Num());
+		REQUIRE(Deque.IsEmpty());
+		Deque.EmplaceLast(0);
+		REQUIRE(Deque.Max());
+		REQUIRE(Deque.Num() == 1);
+		REQUIRE(!Deque.IsEmpty());
+	}
+	REQUIRE(int32Token::EvenConstructionDestructionCalls(1));
+}
+
+TEST_CASE("Containers::Deque::EmplaceLast range to capacity", "[Deque]")
+{
+	int32Token::Reset();
+	{
+		TDeque<int32Token> Deque;
+		REQUIRE(!Deque.Max());
+		REQUIRE(!Deque.Num());
+		Deque.Reserve(Test::DefaultCapacity * 10);
+		REQUIRE(Deque.Max() == Test::DefaultCapacity * 10);
+		while (Deque.Num() < Deque.Max())
+		{
+			Deque.EmplaceLast();
+		}
+		REQUIRE(Deque.Max() == Deque.Num());
+		REQUIRE(Deque.Max() == Test::DefaultCapacity * 10);
+	}
+	REQUIRE(int32Token::EvenConstructionDestructionCalls(Test::DefaultCapacity * 10));
+}
+
+TEST_CASE("Containers::Deque::EmplaceLast range past capacity", "[Deque]")
+{
+	int32Token::Reset();
+	{
+		TDeque<int32Token> Deque;
+		REQUIRE(!Deque.Max());
+		REQUIRE(!Deque.Num());
+		Deque.Reserve(Test::DefaultCapacity);
+		REQUIRE(Deque.Max() == Test::DefaultCapacity);
+		while (Deque.Num() < Deque.Max())
+		{
+			Deque.EmplaceLast();
+		}
+		REQUIRE(Deque.Max() == Deque.Num());
+		REQUIRE(Deque.Max() == Test::DefaultCapacity);
+		Deque.EmplaceLast();
+		REQUIRE(Deque.Max() > Deque.Num());
+		REQUIRE(Deque.Max() > Test::DefaultCapacity);
+	}
+	REQUIRE(int32Token::EvenConstructionDestructionCalls(Test::DefaultCapacity + 1));
+}
+
+TEST_CASE("Containers::Deque::EmplaceFirst single element", "[Deque]")
+{
+	int32Token::Reset();
+	{
+		TDeque<int32Token> Deque;
+		REQUIRE(!Deque.Max());
+		REQUIRE(!Deque.Num());
+		REQUIRE(Deque.IsEmpty());
+		Deque.EmplaceFirst(0);
+		REQUIRE(Deque.Max());
+		REQUIRE(Deque.Num() == 1);
+		REQUIRE(!Deque.IsEmpty());
+	}
+	REQUIRE(int32Token::EvenConstructionDestructionCalls(1));
+}
+
+TEST_CASE("Containers::Deque::EmplaceFirst range to capacity", "[Deque]")
+{
+	int32Token::Reset();
+	{
+		TDeque<int32Token> Deque;
+		REQUIRE(!Deque.Max());
+		REQUIRE(!Deque.Num());
+		Deque.Reserve(Test::DefaultCapacity * 10);
+		REQUIRE(Deque.Max() == Test::DefaultCapacity * 10);
+		while (Deque.Num() < Deque.Max())
+		{
+			Deque.EmplaceFirst();
+		}
+		REQUIRE(Deque.Max() == Deque.Num());
+		REQUIRE(Deque.Max() == Test::DefaultCapacity * 10);
+	}
+	REQUIRE(int32Token::EvenConstructionDestructionCalls(Test::DefaultCapacity * 10));
+}
+
+TEST_CASE("Containers::Deque::PushLast single element (implicit move)", "[Deque]")
+{
+	int32Token::Reset();
+	{
+		TDeque<int32Token> Deque;
+		REQUIRE(!Deque.Max());
+		REQUIRE(!Deque.Num());
+		REQUIRE(Deque.IsEmpty());
+		Deque.PushLast(0);	// implicit conversion from temporary
+		REQUIRE(Deque.Max());
+		REQUIRE(Deque.Num() == 1);
+		REQUIRE(!Deque.IsEmpty());
+	}
+	REQUIRE(int32Token::EvenConstructionDestructionCalls(2));
+	REQUIRE(int32Token::NumConstructorCalls() == 1);
+	REQUIRE(int32Token::NumCopyConstructorCalls() == 0);
+	REQUIRE(int32Token::NumMoveConstructorCalls() == 1);
+}
+
+TEST_CASE("Containers::Deque::PushLast single element from move", "[Deque]")
+{
+	int32Token::Reset();
+	{
+		TDeque<int32Token> Deque;
+		REQUIRE(!Deque.Max());
+		REQUIRE(!Deque.Num());
+		REQUIRE(Deque.IsEmpty());
+		int32Token TempToken;
+		Deque.PushLast(MoveTemp(TempToken));
+		REQUIRE(Deque.Max());
+		REQUIRE(Deque.Num() == 1);
+		REQUIRE(!Deque.IsEmpty());
+	}
+	REQUIRE(int32Token::EvenConstructionDestructionCalls(2));
+	REQUIRE(int32Token::NumConstructorCalls() == 1);
+	REQUIRE(int32Token::NumCopyConstructorCalls() == 0);
+	REQUIRE(int32Token::NumMoveConstructorCalls() == 1);
+}
+
+TEST_CASE("Containers::Deque::PushLast single element from copy", "[Deque]")
+{
+	int32Token::Reset();
+	{
+		TDeque<int32Token> Deque;
+		REQUIRE(!Deque.Max());
+		REQUIRE(!Deque.Num());
+		REQUIRE(Deque.IsEmpty());
+		int32Token TempToken;
+		Deque.PushLast(TempToken);
+		REQUIRE(Deque.Max());
+		REQUIRE(Deque.Num() == 1);
+		REQUIRE(!Deque.IsEmpty());
+	}
+	REQUIRE(int32Token::EvenConstructionDestructionCalls(2));
+	REQUIRE(int32Token::NumConstructorCalls() == 1);
+	REQUIRE(int32Token::NumCopyConstructorCalls() == 1);
+	REQUIRE(int32Token::NumMoveConstructorCalls() == 0);
+}
+
+TEST_CASE("Containers::Deque::PushFirst single element (implicit move)", "[Deque]")
+{
+	int32Token::Reset();
+	{
+		TDeque<int32Token> Deque;
+		REQUIRE(!Deque.Max());
+		REQUIRE(!Deque.Num());
+		REQUIRE(Deque.IsEmpty());
+		Deque.PushFirst(0);	 // implicit conversion from temporary
+		REQUIRE(Deque.Max());
+		REQUIRE(Deque.Num() == 1);
+		REQUIRE(!Deque.IsEmpty());
+	}
+	REQUIRE(int32Token::EvenConstructionDestructionCalls(2));
+	REQUIRE(int32Token::NumConstructorCalls() == 1);
+	REQUIRE(int32Token::NumCopyConstructorCalls() == 0);
+	REQUIRE(int32Token::NumMoveConstructorCalls() == 1);
+}
+
+TEST_CASE("Containers::Deque::PushFirst single element from move", "[Deque]")
+{
+	int32Token::Reset();
+	{
+		TDeque<int32Token> Deque;
+		REQUIRE(!Deque.Max());
+		REQUIRE(!Deque.Num());
+		REQUIRE(Deque.IsEmpty());
+		int32Token TempToken;
+		Deque.PushFirst(MoveTemp(TempToken));
+		REQUIRE(Deque.Max());
+		REQUIRE(Deque.Num() == 1);
+		REQUIRE(!Deque.IsEmpty());
+	}
+	REQUIRE(int32Token::EvenConstructionDestructionCalls(2));
+	REQUIRE(int32Token::NumConstructorCalls() == 1);
+	REQUIRE(int32Token::NumCopyConstructorCalls() == 0);
+	REQUIRE(int32Token::NumMoveConstructorCalls() == 1);
+}
+
+TEST_CASE("Containers::Deque::PushFirst single element from copy", "[Deque]")
+{
+	int32Token::Reset();
+	{
+		TDeque<int32Token> Deque;
+		REQUIRE(!Deque.Max());
+		REQUIRE(!Deque.Num());
+		REQUIRE(Deque.IsEmpty());
+		int32Token TempToken;
+		Deque.PushFirst(TempToken);
+		REQUIRE(Deque.Max());
+		REQUIRE(Deque.Num() == 1);
+		REQUIRE(!Deque.IsEmpty());
+	}
+	REQUIRE(int32Token::EvenConstructionDestructionCalls(2));
+	REQUIRE(int32Token::NumConstructorCalls() == 1);
+	REQUIRE(int32Token::NumCopyConstructorCalls() == 1);
+	REQUIRE(int32Token::NumMoveConstructorCalls() == 0);
+}
+
+void PopOne(TDeque<int32Token>& Deque)
+{
+	Deque.PopFirst();
+}
+
+TEST_CASE("Containers::Deque::EmplaceLast/PopFirst single element", "[Deque]")
+{
+	int32Token::Reset();
+	{
+		TDeque<int32Token> Deque;
+		REQUIRE(!Deque.Max());
+		REQUIRE(!Deque.Num());
+		REQUIRE(Deque.IsEmpty());
+		Deque.EmplaceLast(0);
+		REQUIRE(Deque.Max());
+		REQUIRE(Deque.Num() == 1);
+		REQUIRE(!Deque.IsEmpty());
+		PopOne(Deque);
+		REQUIRE(Deque.Max());
+		REQUIRE(!Deque.Num());
+		REQUIRE(Deque.IsEmpty());
+	}
+	REQUIRE(int32Token::EvenConstructionDestructionCalls(1));
+}
+
+TEST_CASE("Containers::Deque::EmplaceLast/PopFirst single element multiple times causing head/tail wrap around", "[Deque]")
+{
+	int32Token::Reset();
+	{
+		TDeque<int32Token> Deque;
+		Deque.Reserve(Test::DefaultCapacity);
+		REQUIRE(Deque.Max() == Test::DefaultCapacity);
+		for (int32 i = 0; i < Test::DefaultCapacity * 2; ++i)
+		{
+			REQUIRE(Test::EmplaceLastPopFirst(Deque, 1));  // Rotates head and tail
+			REQUIRE(Deque.Max() == Deque.Max());
+		}
+	}
+	REQUIRE(int32Token::EvenConstructionDestructionCalls(Test::DefaultCapacity * 2));
+}
+
+TEST_CASE("Containers::Deque::EmplaceLast/PopFirst range without head/tail wrap around", "[Deque]")
+{
+	int32Token::Reset();
+	{
+		TDeque<int32Token> Deque;
+		Deque.Reserve(Test::DefaultCapacity);
+		REQUIRE(Test::EmplaceLastPopFirst(Deque));
+		REQUIRE(Deque.Max() == Test::DefaultCapacity);
+	}
+	REQUIRE(int32Token::EvenConstructionDestructionCalls(Test::DefaultCapacity));
+}
+
+TEST_CASE("Containers::Deque::EmplaceLast/PopFirst range with reallocation without head/tail wrap around", "[Deque]")
+{
+	int32Token::Reset();
+	{
+		TDeque<int32Token> Deque;
+		Deque.Reserve(Test::DefaultCapacity);
+		REQUIRE(Deque.Max() == Test::DefaultCapacity);
+		REQUIRE(Test::EmplaceLastPopFirst(Deque, Test::DefaultCapacity + 1));
+		REQUIRE(Deque.Max() > Test::DefaultCapacity);
+	}
+	REQUIRE(int32Token::EvenConstructionDestructionCalls(Test::DefaultCapacity + 1));
+}
+
+TEST_CASE("Containers::Deque::EmplaceLast/PopFirst range with head/tail wrap around", "[Deque]")
+{
+	int32Token::Reset();
+	{
+		TDeque<int32Token> Deque;
+		Deque.Reserve(Test::DefaultCapacity);
+		for (int32 i = 0; i < Test::DefaultCapacity; ++i)
+		{
+			REQUIRE(Test::EmplaceLastPopFirst(Deque, Test::DefaultCapacity - 1));  // Rotates head and tail
+			REQUIRE(Deque.Max() == Test::DefaultCapacity);
+		}
+	}
+	REQUIRE(int32Token::EvenConstructionDestructionCalls(Test::DefaultCapacity * (Test::DefaultCapacity - 1)));
+}
+
+TEST_CASE("Containers::Deque::EmplaceFirst/PopLast single element", "[Deque]")
+{
+	int32Token::Reset();
+	{
+		TDeque<int32Token> Deque;
+		REQUIRE(!Deque.Max());
+		REQUIRE(!Deque.Num());
+		REQUIRE(Deque.IsEmpty());
+		Deque.EmplaceFirst(0);
+		REQUIRE(Deque.Max());
+		REQUIRE(Deque.Num() == 1);
+		REQUIRE(!Deque.IsEmpty());
+		Deque.PopLast();
+		REQUIRE(Deque.Max());
+		REQUIRE(!Deque.Num());
+		REQUIRE(Deque.IsEmpty());
+	}
+	REQUIRE(int32Token::EvenConstructionDestructionCalls(1));
+}
+
+TEST_CASE("Containers::Deque::EmplaceFirst/PopLast range", "[Deque]")
+{
+	int32Token::Reset();
+	{
+		TDeque<int32Token> Deque;
+		Deque.Reserve(Test::DefaultCapacity);
+		REQUIRE(Test::EmplaceFirstPopLast(Deque));
+		REQUIRE(Deque.Max() == Test::DefaultCapacity);
+	}
+	REQUIRE(int32Token::EvenConstructionDestructionCalls(Test::DefaultCapacity));
+}
+
+TEST_CASE("Containers::Deque::EmplaceFirst/PopLast range with reallocation", "[Deque]")
+{
+	int32Token::Reset();
+	{
+		TDeque<int32Token> Deque;
+		Deque.Reserve(Test::DefaultCapacity);
+		REQUIRE(Deque.Max() == Test::DefaultCapacity);
+		REQUIRE(Test::EmplaceFirstPopLast(Deque, Test::DefaultCapacity + 1));
+		REQUIRE(Deque.Max() > Test::DefaultCapacity);
+	}
+	REQUIRE(int32Token::EvenConstructionDestructionCalls(Test::DefaultCapacity + 1));
+}
+
+TEST_CASE("Containers::Deque::TryPopFirst", "[Deque]")
+{
+	int32Token::Reset();
+	{
+		TDeque<int32Token> Deque;
+		Deque.Reserve(Test::DefaultCapacity);
+		REQUIRE(Deque.Max() == Test::DefaultCapacity);
+		for (int32 i = 0; i < Test::DefaultCapacity; ++i)
+		{
+			Deque.EmplaceLast(i);
+			REQUIRE(Deque.Max() == Test::DefaultCapacity);
+		}
+		int32Token Result;
+		int CheckValue = 0;
+		while (Deque.TryPopFirst(Result))
+		{
+			REQUIRE(*Result == CheckValue++);
+		}
+		REQUIRE(CheckValue == Test::DefaultCapacity);
+	}
+	REQUIRE(int32Token::EvenConstructionDestructionCalls(Test::DefaultCapacity + 1));  // + 1 for Result
+}
+
+TEST_CASE("Containers::Deque::TryPopFirst with reallocation", "[Deque]")
+{
+	int32Token::Reset();
+	{
+		TDeque<int32Token> Deque;
+		Deque.Reserve(Test::DefaultCapacity);
+		REQUIRE(Deque.Max() == Test::DefaultCapacity);
+		for (int32 i = 0; i < Test::DefaultCapacity; ++i)
+		{
+			Deque.EmplaceLast(i);
+			REQUIRE(Deque.Max() == Test::DefaultCapacity);
+		}
+		Deque.EmplaceLast(Test::DefaultCapacity);
+		REQUIRE(Deque.Max() > Test::DefaultCapacity);
+
+		int32Token Result;
+		int32 CheckValue = 0;
+		while (Deque.TryPopFirst(Result))
+		{
+			REQUIRE(*Result == CheckValue++);
+		}
+		REQUIRE(CheckValue == Test::DefaultCapacity + 1);
+	}
+	REQUIRE(int32Token::EvenConstructionDestructionCalls(Test::DefaultCapacity + 2));
+}
+
+TEST_CASE("Containers::Deque::TryPopLast", "[Deque]")
+{
+	int32Token::Reset();
+	{
+		TDeque<int32Token> Deque;
+		Deque.Reserve(Test::DefaultCapacity);
+		REQUIRE(Deque.Max() == Test::DefaultCapacity);
+		for (int32 i = 0; i < Test::DefaultCapacity; ++i)
+		{
+			Deque.EmplaceFirst(i);
+			REQUIRE(Deque.Max() == Test::DefaultCapacity);
+		}
+		int32Token Result;
+		int CheckValue = 0;
+		while (Deque.TryPopLast(Result))
+		{
+			REQUIRE(*Result == CheckValue++);
+		}
+		REQUIRE(CheckValue == Test::DefaultCapacity);
+	}
+	REQUIRE(int32Token::EvenConstructionDestructionCalls(Test::DefaultCapacity + 1));  // + 1 for Result
+}
+
+TEST_CASE("Containers::Deque::TryPopLast with reallocation", "[Deque]")
+{
+	int32Token::Reset();
+	{
+		TDeque<int32Token> Deque;
+		Deque.Reserve(Test::DefaultCapacity);
+		REQUIRE(Deque.Max() == Test::DefaultCapacity);
+		for (int32 i = 0; i < Test::DefaultCapacity; ++i)
+		{
+			Deque.EmplaceFirst(i);
+			REQUIRE(Deque.Max() == Test::DefaultCapacity);
+		}
+		Deque.EmplaceFirst(Test::DefaultCapacity);
+		REQUIRE(Deque.Max() > Test::DefaultCapacity);
+
+		int32Token Result;
+		int32 CheckValue = 0;
+		while (Deque.TryPopLast(Result))
+		{
+			REQUIRE(*Result == CheckValue++);
+		}
+		REQUIRE(CheckValue == Test::DefaultCapacity + 1);
+	}
+	REQUIRE(int32Token::EvenConstructionDestructionCalls(Test::DefaultCapacity + 2));
+}
+
+TEST_CASE("Containers::Deque::Comparison simple", "[Deque]")
+{
+	int32Token::Reset();
+	{
+		TDeque<int32Token> Deque, DestQueue;
+		Deque.Reserve(Test::DefaultCapacity);
+		DestQueue.Reserve(Test::DefaultCapacity);
+		REQUIRE(Deque.Max() == Test::DefaultCapacity);
+		for (int32 i = 0; i < Test::DefaultCapacity; ++i)
+		{
+			Deque.EmplaceLast(i);
+			DestQueue.EmplaceLast(i);
+			REQUIRE(Deque.Max() == Test::DefaultCapacity);
+		}
+		REQUIRE(Deque == DestQueue);
+	}
+	REQUIRE(int32Token::EvenConstructionDestructionCalls(Test::DefaultCapacity * 2));
+}
+
+TEST_CASE("Containers::Deque::Comparison with head/tail wrap around", "[Deque]")
+{
+	int32Token::Reset();
+	{
+		TDeque<int32Token> Deque, DestQueue;
+		Deque.Reserve(Test::DefaultCapacity);
+		DestQueue.Reserve(Test::DefaultCapacity);
+		REQUIRE(Deque.Max() == Test::DefaultCapacity);
+		for (int32 i = 0; i < Test::DefaultCapacity; ++i)
+		{
+			Deque.EmplaceLast(i);
+			REQUIRE(Deque.Max() == Test::DefaultCapacity);
+		}
+		for (int32 Iteration = 0; Iteration < Test::DefaultCapacity; ++Iteration)
+		{
+			REQUIRE(Test::EmplaceLastPopFirst(DestQueue, 1));  // Rotates head and tail
+			for (int32 i = 0; i < Test::DefaultCapacity; ++i)
+			{
+				DestQueue.EmplaceLast(i);
+				REQUIRE(DestQueue.Max() == Test::DefaultCapacity);
+			}
+			REQUIRE(Deque == DestQueue);
+			DestQueue.Reset();
+		}
+	}
+	REQUIRE(int32Token::EvenConstructionDestructionCalls(Test::DefaultCapacity * 6));
+}
+
+TEST_CASE("Containers::Deque::Copy simple", "[Deque]")
+{
+	int32Token::Reset();
+	{
+		TDeque<int32Token> Deque;
+		Deque.Reserve(Test::DefaultCapacity);
+		REQUIRE(Deque.Max() == Test::DefaultCapacity);
+		for (int32 i = 0; i < Test::DefaultCapacity; ++i)
+		{
+			Deque.EmplaceLast(i);
+			REQUIRE(Deque.Max() == Test::DefaultCapacity);
+		}
+		TDeque<int32Token> DestQueue(Deque);
+		REQUIRE(Deque == DestQueue);
+	}
+	REQUIRE(int32Token::EvenConstructionDestructionCalls(Test::DefaultCapacity * 2));
+}
+
+TEST_CASE("Containers::Deque::Copy with head/tail wrap around", "[Deque]")
+{
+	int32Token::Reset();
+	{
+		TDeque<int32Token> Deque;
+		Deque.Reserve(Test::DefaultCapacity);
+		REQUIRE(Deque.Max() == Test::DefaultCapacity);
+		for (int32 Iteration = 0; Iteration < Test::DefaultCapacity; ++Iteration)
+		{
+			REQUIRE(Test::EmplaceLastPopFirst(Deque, 1));  // Rotates head and tail
+			for (int32 i = 0; i < Test::DefaultCapacity; ++i)
+			{
+				Deque.EmplaceLast(i);
+				REQUIRE(Deque.Max() == Test::DefaultCapacity);
+			}
+			TDeque<int32Token> DestQueue(Deque);
+			REQUIRE(Deque == DestQueue);
+			Deque.Reset();
+		}
+	}
+	REQUIRE(int32Token::EvenConstructionDestructionCalls(Test::DefaultCapacity * 9));
+}
+
+TEST_CASE("Containers::Deque::Copy variable size with head/tail wrap around", "[Deque]")
+{
+	int32Token::Reset();
+	{
+		TDeque<int32Token> Deque;
+		Deque.Reserve(Test::DefaultCapacity);
+		REQUIRE(Deque.Max() == Test::DefaultCapacity);
+		for (int32 Iteration = 0; Iteration < Test::DefaultCapacity; ++Iteration)
+		{
+			for (int32 Size = 1; Size <= Test::DefaultCapacity; ++Size)
+			{
+				REQUIRE(Test::EmplaceLastPopFirst(Deque, 1));  // Rotates head and tail
+				const int32 SeedValue = FMath::RandRange(1, 999);
+				for (int32 i = 0; i < Size; ++i)
+				{
+					Deque.EmplaceLast(SeedValue + i);
+					REQUIRE(Deque.Max() == Test::DefaultCapacity);
+				}
+				TDeque<int32Token> DestQueue(Deque);
+				REQUIRE(Deque == DestQueue);
+				REQUIRE(DestQueue.Max() <= Test::DefaultCapacity);
+				Deque.Reset();
+			}
+		}
+	}
+	REQUIRE(int32Token::EvenConstructionDestructionCalls());
+}
+
+TEST_CASE("Containers::Deque::Move simple", "[Deque]")
+{
+	int32Token::Reset();
+	{
+		TDeque<int32Token> Deque;
+		Deque.Reserve(Test::DefaultCapacity);
+		REQUIRE(Deque.Max() == Test::DefaultCapacity);
+		for (int32 i = 0; i < Test::DefaultCapacity; ++i)
+		{
+			Deque.EmplaceLast(i);
+			REQUIRE(Deque.Max() == Test::DefaultCapacity);
+		}
+		TDeque<int32Token> DestQueue(MoveTemp(Deque));
+		REQUIRE(Deque.IsEmpty());
+		int32Token Result;
+		int CheckValue = 0;
+		while (Deque.TryPopFirst(Result))
+		{
+			REQUIRE(*Result == CheckValue++);
+		}
+	}
+	REQUIRE(int32Token::EvenConstructionDestructionCalls(Test::DefaultCapacity + 1));
+}
+
+TEST_CASE("Containers::Deque::Move with head/tail wrap around", "[Deque]")
+{
+	int32Token::Reset();
+	{
+		TDeque<int32Token> Deque;
+		Deque.Reserve(Test::DefaultCapacity);
+		REQUIRE(Deque.Max() == Test::DefaultCapacity);
+		for (int32 Iteration = 0; Iteration < Test::DefaultCapacity; ++Iteration)
+		{
+			REQUIRE(Test::EmplaceLastPopFirst(Deque, 1));  // Rotates head and tail
+			for (int32 i = 0; i < Test::DefaultCapacity; ++i)
+			{
+				Deque.EmplaceLast(i);
+				REQUIRE(Deque.Max() == Test::DefaultCapacity);
+			}
+			TDeque<int32Token> DestQueue(MoveTemp(Deque));
+			REQUIRE(Deque.IsEmpty());
+			int32Token Result;
+			int32 CheckValue = 0;
+			while (Deque.TryPopFirst(Result))
+			{
+				REQUIRE(*Result == CheckValue++);
+			}
+		}
+	}
+	REQUIRE(int32Token::EvenConstructionDestructionCalls(Test::DefaultCapacity * 6));
+}
+
+TEST_CASE("Containers::Deque::Move variable size with head/tail wrap around", "[Deque]")
+{
+	int32Token::Reset();
+	{
+		TDeque<int32Token> Deque;
+		Deque.Reserve(Test::DefaultCapacity);
+		REQUIRE(Deque.Max() == Test::DefaultCapacity);
+		for (int32 Iteration = 0; Iteration < Test::DefaultCapacity; ++Iteration)
+		{
+			for (int32 Size = 1; Size <= Test::DefaultCapacity; ++Size)
+			{
+				REQUIRE(Test::EmplaceLastPopFirst(Deque, 1));  // Rotates head and tail
+				const int32 SeedValue = FMath::RandRange(1, 999);
+				for (int32 i = 0; i < Size; ++i)
+				{
+					Deque.EmplaceLast(SeedValue + i);
+					REQUIRE(Deque.Max() == Test::DefaultCapacity);
+				}
+				TDeque<int32Token> DestQueue(MoveTemp(Deque));
+				REQUIRE(Deque.IsEmpty());
+				int32Token Result;
+				int32 CheckValue = SeedValue;
+				while (Deque.TryPopFirst(Result))
+				{
+					REQUIRE(*Result == CheckValue++);
+				}
+			}
+		}
+	}
+	REQUIRE(int32Token::EvenConstructionDestructionCalls());
+}
+
+TEST_CASE("Containers::Deque::Iteration without head/tail wrap around", "[Deque]")
+{
+	int32Token::Reset();
+	{
+		TDeque<int32Token> Deque;
+		Deque.Reserve(Test::DefaultCapacity);
+		REQUIRE(Deque.Max() == Test::DefaultCapacity);
+		for (int32 i = 0; i < Test::DefaultCapacity; ++i)
+		{
+			Deque.EmplaceLast(i);
+			REQUIRE(Deque.Max() == Test::DefaultCapacity);
+		}
+		for (int32 i = 0; i < Test::DefaultCapacity; ++i)
+		{
+			REQUIRE(Deque[i] == i);
+		}
+		int32 CheckValue = 0;
+		for (const auto& Value : Deque)
+		{
+			REQUIRE(Value == CheckValue++);
+		}
+	}
+	REQUIRE(int32Token::EvenConstructionDestructionCalls(Test::DefaultCapacity));
+}
+
+TEST_CASE("Containers::Deque::Iteration with head/tail wrap around", "[Deque]")
+{
+	int32Token::Reset();
+	{
+		TDeque<int32Token> Deque;
+		Deque.Reserve(Test::DefaultCapacity);
+		REQUIRE(Deque.Max() == Test::DefaultCapacity);
+		for (int32 Iteration = 0; Iteration < Test::DefaultCapacity; ++Iteration)
+		{
+			REQUIRE(Test::EmplaceLastPopFirst(Deque, 1));  // Rotates head and tail
+			for (int32 i = 0; i < Test::DefaultCapacity; ++i)
+			{
+				Deque.EmplaceLast(i);
+				REQUIRE(Deque.Max() == Test::DefaultCapacity);
+			}
+			for (int32 i = 0; i < Test::DefaultCapacity; ++i)
+			{
+				REQUIRE(Deque[i] == i);
+			}
+			int32 CheckValue = 0;
+			for (const auto& Value : Deque)
+			{
+				REQUIRE(Value == CheckValue++);
+			}
+			Deque.Reset();
+		}
+	}
+	REQUIRE(int32Token::EvenConstructionDestructionCalls(Test::DefaultCapacity * (Test::DefaultCapacity + 1)));
+}
+
+TEST_CASE("Containers::Deque::Iterator arithmetic", "[Deque]")
+{
+	TDeque<int32Token> Deque;
+	Deque.Reserve(Test::DefaultCapacity);
+	REQUIRE(Deque.Max() == Test::DefaultCapacity);
+	for (int32 Iteration = 0; Iteration < Test::DefaultCapacity; ++Iteration)
+	{
+		REQUIRE(Test::EmplaceLastPopFirst(Deque, 1));  // Rotates head and tail
+		Deque.EmplaceLast(13);
+		Deque.EmplaceLast(42);
+		Deque.EmplaceLast(19);
+
+		auto It = Deque.begin();
+		REQUIRE(**It == 13);
+		REQUIRE(*It == int32Token(13));
+		REQUIRE(It->Value == 13);
+		auto It2 = It++;
+		REQUIRE(It2 != It);
+		REQUIRE(It2->Value == 13);
+		REQUIRE(**It == 42);
+		REQUIRE(*It == int32Token(42));
+		REQUIRE(It->Value == 42);
+
+		Deque.Reset();
+	}
+}
+
+TEST_CASE("Containers::Deque::Construct from std initializer_list", "[Deque]")
+{
+	int32Token::Reset();
+	{
+		TDeque<int32Token> Deque({0, 1, 2, 3, 4, 5});
+		int32Token Result;
+		int32 CheckValue = 0;
+		while (Deque.TryPopFirst(Result))
+		{
+			REQUIRE(*Result == CheckValue++);
+		}
+		REQUIRE(CheckValue == 6);
+	}
+	REQUIRE(int32Token::EvenConstructionDestructionCalls(6 * 2 + 1));
+}
+
+TEST_CASE("Containers::Deque::Construct from empty std initializer_list", "[Deque]")
+{
+	int32Token::Reset();
+	{
+		TDeque<int32Token> Deque({});
+		REQUIRE(Deque.IsEmpty());
+	}
+	REQUIRE(int32Token::EvenConstructionDestructionCalls(0));
+}
+
+TEST_CASE("Containers::Deque::Assign from std initializer_list", "[Deque]")
+{
+	int32Token::Reset();
+	{
+		TDeque<int32Token> Deque;
+		Deque.EmplaceLast(0);
+		Deque = {0, 1, 2, 3, 4, 5};
+		int32Token Result;
+		int32 CheckValue = 0;
+		while (Deque.TryPopFirst(Result))
+		{
+			REQUIRE(*Result == CheckValue++);
+		}
+		REQUIRE(CheckValue == 6);
+	}
+	REQUIRE(int32Token::EvenConstructionDestructionCalls(6 * 2 + 2));
+}
+
+}  // namespace Deque
diff --git a/Engine/Source/Runtime/Core/Tests/Containers/TestToken.cpp b/Engine/Source/Runtime/Core/Tests/Containers/TestToken.cpp
new file mode 100644
--- /dev/null
+++ b/Engine/Source/Runtime/Core/Tests/Containers/TestToken.cpp
@@ -0,0 +1,110 @@
+// Copyright Epic Games, Inc. All Rights Reserved.
+
+#include "TestToken.h"
+#include "TestHarness.h"
+
+TEST_CASE("Containers::TToken::TTestToken default", "[Container][TTestToken]")
+{
+	int32Token::Reset();
+	{
+		int32Token Token;
+		REQUIRE(*Token == 0);
+	}
+	REQUIRE(int32Token::EvenConstructionDestructionCalls());
+	REQUIRE(int32Token::NumConstructorCalls() == 1);
+	REQUIRE(int32Token::NumCopyConstructorCalls() == 0);
+	REQUIRE(int32Token::NumMoveConstructorCalls() == 0);
+	REQUIRE(int32Token::NumConstructionCalls() == 1);
+	REQUIRE(int32Token::NumDestructionCalls() == 1);
+}
+
+TEST_CASE("Containers::TToken::TTestToken explicit constructor", "[Container][TTestToken]")
+{
+	int32Token::Reset();
+	{
+		int32Token Token = int32Token(1);
+		REQUIRE(*Token == 1);
+	}
+	REQUIRE(int32Token::EvenConstructionDestructionCalls(1));
+	REQUIRE(int32Token::NumCopyCalls() == 0);
+	REQUIRE(int32Token::NumMoveCalls() == 0);
+	REQUIRE(int32Token::NumConstructorCalls() == 1);
+	REQUIRE(int32Token::NumCopyConstructorCalls() == 0);
+	REQUIRE(int32Token::NumMoveConstructorCalls() == 0);
+	REQUIRE(int32Token::NumCopyAssignmentCalls() == 0);
+	REQUIRE(int32Token::NumMoveAssignmentCalls() == 0);
+}
+
+TEST_CASE("Containers::TToken::TTestToken copy constructor", "[Container][TTestToken]")
+{
+	int32Token::Reset();
+	{
+		int32Token TempToken(2);
+		int32Token Token = TempToken;
+		REQUIRE(*Token == 2);
+	}
+	REQUIRE(int32Token::EvenConstructionDestructionCalls(2));
+	REQUIRE(int32Token::NumCopyCalls() == 1);
+	REQUIRE(int32Token::NumMoveCalls() == 0);
+	REQUIRE(int32Token::NumConstructorCalls() == 1);
+	REQUIRE(int32Token::NumCopyConstructorCalls() == 1);
+	REQUIRE(int32Token::NumMoveConstructorCalls() == 0);
+	REQUIRE(int32Token::NumCopyAssignmentCalls() == 0);
+	REQUIRE(int32Token::NumMoveAssignmentCalls() == 0);
+}
+
+TEST_CASE("Containers::TToken::TTestToken move constructor", "[Container][TTestToken]")
+{
+	int32Token::Reset();
+	{
+		int32Token TempToken(3);
+		int32Token Token = MoveTemp(TempToken);
+		REQUIRE(*Token == 3);
+	}
+	REQUIRE(int32Token::EvenConstructionDestructionCalls(2));
+	REQUIRE(int32Token::NumCopyCalls() == 0);
+	REQUIRE(int32Token::NumMoveCalls() == 1);
+	REQUIRE(int32Token::NumConstructorCalls() == 1);
+	REQUIRE(int32Token::NumCopyConstructorCalls() == 0);
+	REQUIRE(int32Token::NumMoveConstructorCalls() == 1);
+	REQUIRE(int32Token::NumCopyAssignmentCalls() == 0);
+	REQUIRE(int32Token::NumMoveAssignmentCalls() == 0);
+}
+
+TEST_CASE("Containers::TToken::TTestToken copy assignment", "[Container][TTestToken]")
+{
+	int32Token::Reset();
+	{
+		int32Token TempToken(4);
+		int32Token Token;
+		Token = TempToken;
+		REQUIRE(*Token == 4);
+	}
+	REQUIRE(int32Token::EvenConstructionDestructionCalls(2));
+	REQUIRE(int32Token::NumCopyCalls() == 1);
+	REQUIRE(int32Token::NumMoveCalls() == 0);
+	REQUIRE(int32Token::NumConstructorCalls() == 2);
+	REQUIRE(int32Token::NumCopyConstructorCalls() == 0);
+	REQUIRE(int32Token::NumMoveConstructorCalls() == 0);
+	REQUIRE(int32Token::NumCopyAssignmentCalls() == 1);
+	REQUIRE(int32Token::NumMoveAssignmentCalls() == 0);
+}
+
+TEST_CASE("Containers::TToken::TTestToken move assignment", "[Container][TTestToken]")
+{
+	int32Token::Reset();
+	{
+		int32Token TempToken(5);
+		int32Token Token;
+		Token = MoveTemp(TempToken);
+		REQUIRE(*Token == 5);
+	}
+	REQUIRE(int32Token::EvenConstructionDestructionCalls(2));
+	REQUIRE(int32Token::NumCopyCalls() == 0);
+	REQUIRE(int32Token::NumMoveCalls() == 1);
+	REQUIRE(int32Token::NumConstructorCalls() == 2);
+	REQUIRE(int32Token::NumCopyConstructorCalls() == 0);
+	REQUIRE(int32Token::NumMoveConstructorCalls() == 0);
+	REQUIRE(int32Token::NumCopyAssignmentCalls() == 0);
+	REQUIRE(int32Token::NumMoveAssignmentCalls() == 1);
+}
diff --git a/Engine/Source/Runtime/Core/Tests/Containers/TestToken.h b/Engine/Source/Runtime/Core/Tests/Containers/TestToken.h
new file mode 100644
--- /dev/null
+++ b/Engine/Source/Runtime/Core/Tests/Containers/TestToken.h
@@ -0,0 +1,154 @@
+// Copyright Epic Games, Inc. All Rights Reserved.
+
+#pragma once
+
+#include "CoreTypes.h"
+#include "Templates/UnrealTemplate.h"
+
+template <typename InValueType>
+class TTestToken
+{
+	static int32 ConstructorCalls;
+	static int32 DestructorCalls;
+	static int32 CopyConstructorCalls;
+	static int32 MoveConstructorCalls;
+	static int32 CopyAssignmentCalls;
+	static int32 MoveAssignmentCalls;
+
+public:
+	using ValueType = InValueType;
+
+	TTestToken(const ValueType& value)
+	{
+		ConstructorCalls++;
+		Value = value;
+	}
+
+	TTestToken() noexcept : Value(ValueType{}) { ConstructorCalls++; }
+
+	TTestToken(TTestToken&& Other) noexcept : Value(MoveTemp(Other.Value))
+	{
+		ensure(&Other != this);
+		MoveConstructorCalls++;
+	}
+
+	TTestToken(const TTestToken& Other) : Value(Other.Value)
+	{
+		ensure(&Other != this);
+		CopyConstructorCalls++;
+	}
+
+	~TTestToken()
+	{
+		DestructorCalls++;
+	}
+
+	const ValueType& operator*() const { return Value; }
+
+	ValueType& operator*() { return Value; }
+
+	const ValueType* operator->() const { return &Value; }
+
+	ValueType* operator->() { return &Value; }
+
+	TTestToken& operator=(const TTestToken& Other)
+	{
+		ensure(&Other != this);
+		Value = Other.Value;
+		CopyAssignmentCalls++;
+		return *this;
+	}
+
+	TTestToken& operator=(TTestToken&& Other) noexcept
+	{
+		ensure(&Other != this);
+		Value = MoveTemp(Other.Value);
+		MoveAssignmentCalls++;
+		return *this;
+	}
+
+	bool operator<(const TTestToken& Other) const { return Value < Other.Value; }
+
+	static int32 NumConstructionCalls() { return ConstructorCalls + CopyConstructorCalls + MoveConstructorCalls; }
+
+	static int32 NumConstructorCalls() { return ConstructorCalls; }
+
+	static int32 NumCopyConstructorCalls() { return CopyConstructorCalls; }
+
+	static int32 NumMoveConstructorCalls() { return MoveConstructorCalls; }
+
+	static int32 ConstructionDestructionCallDifference() { return NumConstructionCalls() - DestructorCalls; }
+
+	static int32 NumCopyCalls() { return CopyConstructorCalls + CopyAssignmentCalls; }
+
+	static int32 NumCopyAssignmentCalls() { return CopyAssignmentCalls; }
+
+	static int32 NumMoveCalls() { return MoveConstructorCalls + MoveAssignmentCalls; }
+
+	static int32 NumMoveAssignmentCalls() { return MoveAssignmentCalls; }
+
+	static int32 NumDestructionCalls() { return DestructorCalls; }
+
+	static bool EvenConstructionDestructionCalls() { return ConstructionDestructionCallDifference() == 0; }
+
+	static bool EvenConstructionDestructionCalls(int32 ExpectedNum)
+	{
+		return EvenConstructionDestructionCalls() && DestructorCalls == ExpectedNum;
+	}
+
+	static void Reset()
+	{
+		ConstructorCalls = 0;
+		DestructorCalls = 0;
+		CopyConstructorCalls = 0;
+		MoveConstructorCalls = 0;
+		CopyAssignmentCalls = 0;
+		MoveAssignmentCalls = 0;
+	}
+
+	ValueType Value;
+};
+
+template <typename InValueType>
+int32 TTestToken<InValueType>::ConstructorCalls = 0;
+
+template <typename InValueType>
+int32 TTestToken<InValueType>::DestructorCalls = 0;
+
+template <typename InValueType>
+int32 TTestToken<InValueType>::CopyConstructorCalls = 0;
+
+template <typename InValueType>
+int32 TTestToken<InValueType>::MoveConstructorCalls = 0;
+
+template <typename InValueType>
+int32 TTestToken<InValueType>::CopyAssignmentCalls = 0;
+
+template <typename InValueType>
+int32 TTestToken<InValueType>::MoveAssignmentCalls = 0;
+
+using int32Token = TTestToken<int32>;
+
+template <typename InValueType>
+inline bool operator==(const TTestToken<InValueType>& left, const TTestToken<InValueType>& right)
+{
+	return left.Value == right.Value;
+}
+
+template <typename InValueType>
+inline bool operator==(const TTestToken<InValueType>& left, const InValueType& right)
+{
+	return left.Value == right;
+}
+
+template <typename InValueType>
+inline bool operator!=(const TTestToken<InValueType>& left, const TTestToken<InValueType>& right)
+{
+	return !operator==(left, right);
+}
+
+template <typename InValueType>
+inline bool operator!=(const TTestToken<InValueType>& left, const InValueType& right)
+{
+	return !operator==(left, right);
+}
diff --git a/Engine/Source/Runtime/Core/Tests/Containers/TestUtils.cpp b/Engine/Source/Runtime/Core/Tests/Containers/TestUtils.cpp
new file mode 100644
--- /dev/null
+++ b/Engine/Source/Runtime/Core/Tests/Containers/TestUtils.cpp
@@ -0,0 +1,22 @@
+// Copyright Epic Games, Inc. All Rights Reserved.
+
+#include "TestUtils.h"
+#include <numeric>
+
+namespace Test
+{
+	constexpr int32 DefaultRandomSeed = 283281277;
+}
+
+void Test::ResetRandomSeed()
+{
+	FMath::SRandInit(DefaultRandomSeed);
+}
+
+TArray<int32> Test::MakeIndexArray(int32 Size)
+{
+	TArray<int32> Indices;
+	Indices.SetNum(Size);
+	std::iota(Indices.begin(), Indices.end(), 0);
+	return Indices;
+}
\ No newline at end of file
diff --git a/Engine/Source/Runtime/Core/Tests/Containers/TestUtils.h b/Engine/Source/Runtime/Core/Tests/Containers/TestUtils.h
new file mode 100644
--- /dev/null
+++ b/Engine/Source/Runtime/Core/Tests/Containers/TestUtils.h
@@ -0,0 +1,59 @@
+// Copyright Epic Games, Inc. All Rights Reserved.
+
+#pragma once
+
+#include "CoreMinimal.h"
+#include "Math/UnrealMathUtility.h"
+
+namespace Test
+{
+void ResetRandomSeed();
+
+/**
+ * Returns an array with the parameter number of sequentially increased index values starting from zero.
+ * This method is equivalent to using std::iota with a zero starting value in a sequential container.
+ **/
+TArray<int32> MakeIndexArray(int32 Size);
+
+/**
+ * Returns an array with a subset of the elements in the parameter input array. The subset size is expressed
+ * by the parameter size which must be smaller than the number of elements in the input array.
+ **/
+template <typename T>
+TArray<T> MakeRandomSubset(const TArray<T>& InArray, int32 Size)
+{
+	TArray<T> Result;
+	if (InArray.Num() > 1 && Size < InArray.Num())
+	{
+		Result.Reset(Size);
+		while (Result.Num() < Size)
+		{
+			const int32 i = FMath::RandRange(0, InArray.Num() - 1);
+			if (Result.Find(InArray[i]) == INDEX_NONE)
+			{
+				Result.Add(InArray[i]);
+			}
+		}
+	}
+	return Result;
+}
+
+/**
+ * Performs a number of randomized swaps in the parameter array equivalent to its number of elements.
+ **/
+template <typename T>
+void Shuffle(TArray<T>& OutArray)
+{
+	if (OutArray.Num() > 1)
+	{
+		for (int32 i = 0; i < OutArray.Num(); i++)
+		{
+			const int32 j = FMath::RandRange(i, OutArray.Num() - 1);
+			if (i != j)
+			{
+				OutArray.Swap(i, j);
+			}
+		}
+	}
+}
+}  // namespace Test
\ No newline at end of file
diff --git a/Engine/Source/Runtime/Core/Tests/Math/MathTests.cpp b/Engine/Source/Runtime/Core/Tests/Math/MathTests.cpp
new file mode 100644
--- /dev/null
+++ b/Engine/Source/Runtime/Core/Tests/Math/MathTests.cpp
@@ -0,0 +1,65 @@
+// Copyright Epic Games, Inc. All Rights Reserved.
+#include "CoreMinimal.h"
+#include "Math/Color.h"
+#include "TestHarness.h"
+
+TEST_CASE("Math::Colour::FLinearColor", "[Core][Math][Smoke]")
+{
+	SECTION("FLinearColor")
+	{
+		FLinearColor black(FLinearColor::Black);
+		FLinearColor white(FLinearColor::White);
+		FLinearColor red(FLinearColor::Red);
+
+		REQUIRE(red.R == 1.0);
+		REQUIRE(red.G == 0.0);
+		REQUIRE(red.B == 0.0);
+		REQUIRE(red.A == 1.0);
+		REQUIRE(red == FLinearColor::Red);
+		REQUIRE(red != FLinearColor::Green);
+
+		REQUIRE(white.IsAlmostBlack() == false);
+		REQUIRE(black.IsAlmostBlack() == true);
+
+		FLinearColor yellow(FColor::FromHex(FString("FFFF00FF")));
+		FLinearColor green(FColor::FromHex(FString("00FF00FF")));
+		FLinearColor blue(FColor::FromHex(FString("0000FFFF")));
+
+		REQUIRE(yellow == FLinearColor::Yellow);
+		REQUIRE(green == FLinearColor::Green);
+		REQUIRE(blue == FLinearColor::Blue);
+		REQUIRE(blue != yellow);
+		REQUIRE(green != yellow);
+		REQUIRE(yellow == yellow);
+	}
+}
+
+TEST_CASE("Math::Colour::FColor", "[Core][Math][Smoke]")
+{
+	SECTION("FColor")
+	{
+		FColor black(FColor::Black);
+		FColor white(FColor::White);
+		FColor red(FColor::Red);
+
+		REQUIRE(red.R == 0xFF);
+		REQUIRE(red.G == 0x00);
+		REQUIRE(red.B == 0x00);
+		REQUIRE(red.A == 0xFF);
+		REQUIRE(red == FColor::Red);
+		REQUIRE(red != FColor::Green);
+
+		FColor yellow(FColor::FromHex(FString("FFFF00FF")));
+		FColor green(FColor::FromHex(FString("00FF00FF")));
+		FColor blue(FColor::FromHex(FString("0000FFFF")));
+
+		REQUIRE(yellow == FColor::Yellow);
+		REQUIRE(green == FColor::Green);
+		REQUIRE(blue == FColor::Blue);
+		REQUIRE(blue != yellow);
+		REQUIRE(green != yellow);
+		REQUIRE(yellow == yellow);
+
+		REQUIRE(1 == 1);
+	}
+}
