Remove UE_LARGE_WORLD_COORDINATES_DISABLED usage from ispc

diff --git a/Engine/Source/Runtime/Core/Public/Math/Matrix.isph b/Engine/Source/Runtime/Core/Public/Math/Matrix.isph
--- a/Engine/Source/Runtime/Core/Public/Math/Matrix.isph
+++ b/Engine/Source/Runtime/Core/Public/Math/Matrix.isph
@@ -9,23 +9,6 @@
 #define EXPLICIT_MATRIX 0
 #endif
 
-#if UE_LARGE_WORLD_COORDINATES_DISABLED
-struct FMatrix
-{
-#if EXPLICIT_MATRIX == 0
-	float M[16];
-#else
-	float<16> M;
-#endif
-};
-
-struct FMatrix44d
-{
-	double M[16];
-};
-
-typedef FMatrix FMatrix44f;
-#else
 struct FMatrix44f
 {
 #if EXPLICIT_MATRIX == 0
@@ -41,7 +24,6 @@ struct FMatrix
 };
 
 typedef FMatrix FMatrix44d;
-#endif
 
 #ifndef EXPLICIT_MATRIX3x4
 #define EXPLICIT_MATRIX3x4 0
diff --git a/Engine/Source/Runtime/Core/Public/Math/Scalar.isph b/Engine/Source/Runtime/Core/Public/Math/Scalar.isph
--- a/Engine/Source/Runtime/Core/Public/Math/Scalar.isph
+++ b/Engine/Source/Runtime/Core/Public/Math/Scalar.isph
@@ -3,24 +3,12 @@
 #ifndef __SCALAR_ISPH__
 #define __SCALAR_ISPH__
 
-#ifndef UE_LARGE_WORLD_COORDINATES_DISABLED
-#define UE_LARGE_WORLD_COORDINATES_DISABLED 0
-#endif
-
 // Define types
-#if UE_LARGE_WORLD_COORDINATES_DISABLED
-typedef float FReal;
-typedef uint32 UIntType;
-typedef int32 IntType;
-
-#define REAL_BITS floatbits
-#else
 typedef double FReal;
 typedef uint64 UIntType;
 typedef int64 IntType;
 
 #define REAL_BITS doublebits
-#endif
 
 #define FLOAT_PI					(3.1415926535897932f)
 #define FLOAT_SMALL_NUMBER			(1.e-8f)
@@ -52,30 +40,6 @@ typedef int64 IntType;
 
 #undef  PI
 // Define type constants
-#if UE_LARGE_WORLD_COORDINATES_DISABLED
-#define PI						FLOAT_PI
-#define SMALL_NUMBER			FLOAT_SMALL_NUMBER
-#define KINDA_SMALL_NUMBER		FLOAT_KINDA_SMALL_NUMBER
-#define SIGN_BIT				FLOAT_SIGN_BIT
-#define REAL_MAX				FLT_MAX
-#define REAL_MIN				FLT_MIN
-#define NON_FRACTIONAL			FLOAT_NON_FRACTIONAL
-#define NAN						FLOAT_NAN
-#define ZERO					FLOAT_ZERO
-#define ONE_HALF				FLOAT_ONE_HALF
-#define ONE						FLOAT_ONE
-#define TWO						FLOAT_TWO
-
-#define QMULTI_SIGN_MASK0		FLOAT_QMULTI_SIGN_MASK0
-#define QMULTI_SIGN_MASK1		FLOAT_QMULTI_SIGN_MASK1
-#define QMULTI_SIGN_MASK2		FLOAT_QMULTI_SIGN_MASK2
-#define QINV_SIGN_MASK			FLOAT_QINV_SIGN_MASK
-
-#define DEG_TO_RAD				FLOAT_DEG_TO_RAD
-#define DEG_TO_RAD_HALF			FLOAT_DEG_TO_RAD_HALF
-#define RAD_TO_DEG				FLOAT_RAD_TO_DEG
-
-#else
 
 #define PI						DOUBLE_PI
 #define SMALL_NUMBER			DOUBLE_SMALL_NUMBER
@@ -98,7 +62,6 @@ typedef int64 IntType;
 #define DEG_TO_RAD				DOUBLE_DEG_TO_RAD
 #define DEG_TO_RAD_HALF			DOUBLE_DEG_TO_RAD_HALF
 #define RAD_TO_DEG				DOUBLE_RAD_TO_DEG
-#endif
 
 // _mm_cmpge_pd
 inline uniform double VectorCompareGE(const uniform double A, const uniform double B)
diff --git a/Engine/Source/Runtime/Core/Public/Math/Vector.isph b/Engine/Source/Runtime/Core/Public/Math/Vector.isph
--- a/Engine/Source/Runtime/Core/Public/Math/Vector.isph
+++ b/Engine/Source/Runtime/Core/Public/Math/Vector.isph
@@ -3,10 +3,6 @@
 #ifndef __VECTOR_ISPH__
 #define __VECTOR_ISPH__
 
-#ifndef UE_LARGE_WORLD_COORDINATES_DISABLED
-#define UE_LARGE_WORLD_COORDINATES_DISABLED 0
-#endif
-
 #include "Math/Aossoa.isph"
 #include "Math/Soaaos.isph"
 #include "Math/Scalar.isph"
@@ -16,21 +12,6 @@
 #define EXPLICIT_VECTOR2 0
 #endif
 
-#if UE_LARGE_WORLD_COORDINATES_DISABLED
-struct FVector2
-{
-#if EXPLICIT_VECTOR2 == 0
-	float V[2];
-#else
-	float<2> V;
-#endif
-};
-
-struct FVector2d
-{
-	double V[2];
-};
-#else
 struct FVector2f
 {
 #if EXPLICIT_VECTOR2 == 0
@@ -44,27 +25,11 @@ struct FVector2
 {
 	double V[2];
 };
-#endif
 
 #ifndef EXPLICIT_VECTOR
 #define EXPLICIT_VECTOR 0
 #endif
 
-#if UE_LARGE_WORLD_COORDINATES_DISABLED
-struct FVector
-{
-#if EXPLICIT_VECTOR == 0
-	float V[3];
-#else
-	float<3> V;
-#endif
-};
-
-struct FVector3d
-{
-	double V[3];
-};
-#else
 struct FVector3f
 {
 #if EXPLICIT_VECTOR == 0
@@ -78,27 +43,11 @@ struct FVector
 {
 	double V[3];
 };
-#endif
 
 #ifndef EXPLICIT_VECTOR4
 #define EXPLICIT_VECTOR4 0
 #endif
 
-#if UE_LARGE_WORLD_COORDINATES_DISABLED
-struct FVector4
-{
-#if EXPLICIT_VECTOR4 == 0
-	float V[4];
-#else
-	float<4> V;
-#endif
-};
-
-struct FVector4d
-{
-	double V[4];
-};
-#else
 struct FVector4f
 {
 #if EXPLICIT_VECTOR4 == 0
@@ -112,7 +61,6 @@ struct FVector4
 {
 	double V[4];
 };
-#endif
 
 struct FVector8
 {
@@ -163,15 +111,9 @@ struct FIntVector
 };
 
 // Define types
-#if UE_LARGE_WORLD_COORDINATES_DISABLED
-typedef FVector2 FVector2f;
-typedef FVector FVector3f;
-typedef FVector4 FVector4f;
-#else
 typedef FVector2 FVector2d;
 typedef FVector FVector3d;
 typedef FVector4 FVector4d;
-#endif
 
 static const uniform struct FVector4f FLOAT_QMULTI_SIGN_MASK0 = { {1.f, -1.f, 1.f, -1.f} };
 static const uniform struct FVector4f FLOAT_QMULTI_SIGN_MASK1 = { {1.f, 1.f, -1.f, -1.f} };
@@ -266,45 +208,6 @@ static const uniform struct FVector8 ForwardVector8 = { {1.0f, 0.0f, 0.0f, 0.0f,
 static const uniform struct FVector8 RightVector8 = { {0.0f, 1.0f, 0.0f, 0.0f, 0.0f, -1.0f, 0.0f, 0.0f} };
 
 // Define vector type constants
-#if UE_LARGE_WORLD_COORDINATES_DISABLED
-
-#define VectorZero				FloatZero
-#define VectorOne				FloatOne
-#define VectorMinusOne			FloatMinusOne
-#define Vector0001				Float0001
-#define Vector1110				Float1110
-#define VectorOneHalf			FloatOneHalf
-#define SmallLengthThreshold	FloatSmallLengthThreshold
-
-#define VectorPi				FloatPi
-#define VectorTwoPi				FloatTwoPi
-#define VectorPiByTwo			FloatPiByTwo
-#define VectorPiByFour			FloatPiByFour
-#define VectorOneOverPi			FloatOneOverPi
-#define VectorOneOverTwoPi		FloatOneOverTwoPi
-
-#define Vector255				Float255
-#define Vector127				Float127
-#define VectorNeg127			FloatNeg127
-#define Vector360				Float360
-#define Vector180				Float180
-#define VectorNonFractional		FloatNonFractional
-
-#define UpVector4				FloatUpVector4
-#define ForwardVector4			FloatForwardVector4
-#define RightVector4			FloatRightVector4
-
-#define ZeroVector				FloatZeroVector
-#define OneVector				FloatOneVector
-#define HalfVector				FloatHalfVector
-#define UpVector				FloatUpVector
-#define DownVector				FloatDownVector
-#define ForwardVector			FloatForwardVector
-#define BackwardVector			FloatBackwardVector
-#define RightVector				FloatRightVector
-#define LeftVector				FloatLeftVector
-
-#else
 
 #define VectorZero				DoubleZero
 #define VectorOne				DoubleOne
@@ -341,7 +244,6 @@ static const uniform struct FVector8 RightVector8 = { {0.0f, 1.0f, 0.0f, 0.0f, 0
 #define BackwardVector			DoubleBackwardVector
 #define RightVector				DoubleRightVector
 #define LeftVector				DoubleLeftVector
-#endif
 
 #if (defined(ISPC_TARGET_AVX2) || defined(ISPC_TARGET_AVX512SKX))
 #define HW_GATHER_SUPPORTED 1
@@ -1208,38 +1110,22 @@ inline uniform FVector4f ConvertVector4dTo4f(const uniform FVector4d &V)
 
 inline uniform FVector4 ConvertVector4dTo4Native(const uniform FVector4d &V)
 {
-#if UE_LARGE_WORLD_COORDINATES_DISABLED
-	return ConvertVector4dTo4f(V);
-#else
 	return V;
-#endif
 }
 
 inline uniform FVector4 ConvertVector4fTo4Native(const uniform FVector4f &V)
 {
-#if UE_LARGE_WORLD_COORDINATES_DISABLED
-	return V;
-#else
 	return ConvertVector4fTo4d(V);
-#endif
 }
 
 inline uniform FVector4d ConvertVector4NativeTo4d(const uniform FVector4 &V)
 {
-#if UE_LARGE_WORLD_COORDINATES_DISABLED
-	return ConvertVector4fTo4d(V);
-#else
 	return V;
-#endif
 }
 
 inline uniform FVector4f ConvertVector4NativeTo4f(const uniform FVector4 &V)
 {
-#if UE_LARGE_WORLD_COORDINATES_DISABLED
-	return V;
-#else
 	return ConvertVector4dTo4f(V);
-#endif
 }
 
 inline FVector3d operator+(const FVector3d &A, const FVector3d &B)
@@ -1767,38 +1653,22 @@ inline uniform FVector3f ConvertVector3dTo3f(const uniform FVector3d &V)
 
 inline uniform FVector ConvertVector3dTo3Native(const uniform FVector3d &V)
 {
-#if UE_LARGE_WORLD_COORDINATES_DISABLED
-	return ConvertVector3dTo3f(V);
-#else
 	return V;
-#endif
 }
 
 inline uniform FVector ConvertVector3fTo3Native(const uniform FVector3f &V)
 {
-#if UE_LARGE_WORLD_COORDINATES_DISABLED
-	return V;
-#else
 	return ConvertVector3fTo3d(V);
-#endif
 }
 
 inline uniform FVector3d ConvertVector3NativeTo3d(const uniform FVector &V)
 {
-#if UE_LARGE_WORLD_COORDINATES_DISABLED
-	return ConvertVector3fTo3d(V);
-#else
 	return V;
-#endif
 }
 
 inline uniform FVector3f ConvertVector3NativeTo3f(const uniform FVector &V)
 {
-#if UE_LARGE_WORLD_COORDINATES_DISABLED
-	return V;
-#else
 	return ConvertVector3dTo3f(V);
-#endif
 }
 
 inline uniform FVector8 operator+(const uniform FVector8 &A, const uniform FVector8 &B)
@@ -3422,13 +3292,8 @@ inline void VectorSinCos(uniform FVector4 &VSinAngles, uniform FVector4 &VCosAng
 
 	// 11-degree minimax approximation
 	//*ScalarSin = (((((-2.3889859e-08f * y2 + 2.7525562e-06f) * y2 - 0.00019840874f) * y2 + 0.0083333310f) * y2 - 0.16666667f) * y2 + 1.0f) * y;
-#if UE_LARGE_WORLD_COORDINATES_DISABLED
-	static const uniform FVector4f SinCoeff0 = {{ 1.0f, -0.16666667f, 0.0083333310f, -0.00019840874f }};
-	static const uniform FVector4f SinCoeff1 = {{ 2.7525562e-06f, -2.3889859e-08f, /*unused*/ 0.f, /*unused*/ 0.f }};
-#else
 	static const uniform FVector4d SinCoeff0 = {{ 1.0d, -0.16666667d, 0.0083333310d, -0.00019840874d }};
 	static const uniform FVector4d SinCoeff1 = {{ 2.7525562d-06, -2.3889859d-08, /*unused*/ 0.d, /*unused*/ 0.d }};
-#endif
 
 	uniform FVector4 S;
 	S = VectorReplicate(SinCoeff1, 1);
@@ -3441,13 +3306,8 @@ inline void VectorSinCos(uniform FVector4 &VSinAngles, uniform FVector4 &VCosAng
 
 	// 10-degree minimax approximation
 	//*ScalarCos = sign * (((((-2.6051615e-07f * y2 + 2.4760495e-05f) * y2 - 0.0013888378f) * y2 + 0.041666638f) * y2 - 0.5f) * y2 + 1.0f);
-#if UE_LARGE_WORLD_COORDINATES_DISABLED
-	static const uniform FVector4f CosCoeff0 = {{ 1.0f, -0.5f, 0.041666638f, -0.0013888378f }};
-	static const uniform FVector4f CosCoeff1 = {{ 2.4760495e-05f, -2.6051615e-07f, /*unused*/ 0.f, /*unused*/ 0.f }};
-#else
 	static const uniform FVector4d CosCoeff0 = {{ 1.0d, -0.5d, 0.041666638d, -0.0013888378d }};
 	static const uniform FVector4d CosCoeff1 = {{ 2.4760495d-05, -2.6051615d-07, /*unused*/ 0.d, /*unused*/ 0.d }};
-#endif
 
 	uniform FVector4 C;
 	C = VectorReplicate(CosCoeff1, 1);
diff --git a/Engine/Source/Runtime/Core/Public/Math/WideVector.isph b/Engine/Source/Runtime/Core/Public/Math/WideVector.isph
--- a/Engine/Source/Runtime/Core/Public/Math/WideVector.isph
+++ b/Engine/Source/Runtime/Core/Public/Math/WideVector.isph
@@ -5,17 +5,6 @@
 
 #include "Math/Vector.isph"
 
-#if UE_LARGE_WORLD_COORDINATES_DISABLED
-struct WideFVector4
-{
-	float V[programCount];
-};
-
-struct WideFVector4d
-{
-	double V[programCount];
-};
-#else
 struct WideFVector4f
 {
 	float V[programCount];
@@ -25,30 +14,10 @@ struct WideFVector4
 {
 	double V[programCount];
 };
-#endif
 
 // Define types
-#if UE_LARGE_WORLD_COORDINATES_DISABLED
-typedef WideFVector4 WideFVector4f;
-#else
 typedef WideFVector4 WideFVector4d;
-#endif
 
-#if UE_LARGE_WORLD_COORDINATES_DISABLED
-#if TARGET_WIDTH == 4
-static const uniform struct WideFVector4f WFloatZero = { {0.0f, 0.0f, 0.0f, 0.0f} };
-static const uniform struct WideFVector4f WFloat0001 = { {0.0f, 0.0f, 0.0f, 1.0f} };
-static const uniform struct WideFVector4f WSmallLengthThreshold = { {1.e-8f, 1.e-8f, 1.e-8f, 1.e-8f} };
-#elif TARGET_WIDTH == 8
-static const uniform struct WideFVector4f WFloatZero = { {0.0f, 0.0f, 0.0f, 0.0f, 0.0f, 0.0f, 0.0f, 0.0f} };
-static const uniform struct WideFVector4f WFloat0001 = { {0.0f, 0.0f, 0.0f, 1.0f, 0.0f, 0.0f, 0.0f, 1.0f} };
-static const uniform struct WideFVector4f WSmallLengthThreshold = { {1.e-8f, 1.e-8f, 1.e-8f, 1.e-8f, 1.e-8f, 1.e-8f, 1.e-8f, 1.e-8f} };
-#elif TARGET_WIDTH == 16
-static const uniform struct WideFVector4f WFloatZero = { {0.0f, 0.0f, 0.0f, 0.0f, 0.0f, 0.0f, 0.0f, 0.0f, 0.0f, 0.0f, 0.0f, 0.0f, 0.0f, 0.0f, 0.0f, 0.0f} };
-static const uniform struct WideFVector4f WFloat0001 = { {0.0f, 0.0f, 0.0f, 1.0f, 0.0f, 0.0f, 0.0f, 1.0f, 0.0f, 0.0f, 0.0f, 1.0f, 0.0f, 0.0f, 0.0f, 1.0f} };
-static const uniform struct WideFVector4f WSmallLengthThreshold = { {1.e-8f, 1.e-8f, 1.e-8f, 1.e-8f, 1.e-8f, 1.e-8f, 1.e-8f, 1.e-8f, 1.e-8f, 1.e-8f, 1.e-8f, 1.e-8f, 1.e-8f, 1.e-8f, 1.e-8f, 1.e-8f} };
-#endif
-#else
 #if TARGET_WIDTH == 4
 static const uniform struct WideFVector4d WFloatZero = { {0.0d, 0.0d, 0.0d, 0.0d} };
 static const uniform struct WideFVector4d WFloat0001 = { {0.0d, 0.0d, 0.0d, 1.0d} };
@@ -62,7 +31,6 @@ static const uniform struct WideFVector4d WFloatZero = { {0.0d, 0.0d, 0.0d, 0.0d
 static const uniform struct WideFVector4d WFloat0001 = { {0.0d, 0.0d, 0.0d, 1.0d, 0.0d, 0.0d, 0.0d, 1.0d, 0.0d, 0.0d, 0.0d, 1.0d, 0.0d, 0.0d, 0.0d, 1.0d} };
 static const uniform struct WideFVector4d WSmallLengthThreshold = { {1.d-8, 1.d-8, 1.d-8, 1.d-8, 1.d-8, 1.d-8, 1.d-8, 1.d-8, 1.d-8, 1.d-8, 1.d-8, 1.d-8, 1.d-8, 1.d-8, 1.d-8, 1.d-8} };
 #endif
-#endif // UE_LARGE_WORLD_COORDINATES_DISABLED
 
 unmasked inline uniform WideFVector4 operator+(const uniform WideFVector4 &A, const uniform WideFVector4 &B)
 {
diff --git a/Engine/Source/Runtime/Experimental/Chaos/Private/Chaos/AABB.ispc b/Engine/Source/Runtime/Experimental/Chaos/Private/Chaos/AABB.ispc
--- a/Engine/Source/Runtime/Experimental/Chaos/Private/Chaos/AABB.ispc
+++ b/Engine/Source/Runtime/Experimental/Chaos/Private/Chaos/AABB.ispc
@@ -11,25 +11,14 @@ inline static void GrowToInclude(uniform FVector4 &Min, uniform FVector4 &Max, c
 	Max = VectorMax(Max, V);
 }
 
-inline static void GrowToInclude(uniform FVector8 &MinMax, const uniform FVector8& V)
-{
-	const uniform FVector8 VReverse = SetVector8(V.V[4], V.V[5], V.V[6], V.V[7], V.V[0], V.V[1], V.V[2], V.V[3]);
-
-	const uniform FVector8 Min = VectorMin(VectorMin(MinMax, V), VReverse);
-	const uniform FVector8 Max = VectorMax(VectorMax(MinMax, V), VReverse);
-
-	MinMax = SetVector8(Min.V[0], Min.V[1], Min.V[2], Min.V[3], Max.V[4], Max.V[5], Max.V[6], Max.V[7]);
-}
-
-// Three permutations need supporting - FTransform<float>/AABB<float>, FTransform<double>/AABB<double>, FTransform<double>/AABB<float>.
-// TransformAABB handles the first two based on LWC enablement. Mixed handles the third case when LWC is enabled.
+// Two permutations need supporting - FTransform<double>/AABB<double> and FTransform<double>/AABB<float>.
+// TransformAABB handles the first. Mixed handles the second.
 export void TransformedAABB(const uniform FTransform &SpaceTransform, const uniform FVector &MMin, const uniform FVector &MMax, uniform FVector &NewMin, uniform FVector &NewMax)
 {
 	const uniform FVector4 Min = SetVector4(MMin.V[0], MMin.V[1], MMin.V[2], ZERO);
 	const uniform FVector4 Max = SetVector4(MMax.V[0], MMax.V[1], MMax.V[2], ZERO);
 	const uniform FVector4 CurrentExtents = Max - Min;
 
-#if UE_LARGE_WORLD_COORDINATES_DISABLED == 0 || TARGET_WIDTH == 4
 	uniform FVector4 MinToNewSpace = TransformPosition(SpaceTransform, Min);
 	uniform FVector4 MaxToNewSpace = MinToNewSpace;
 	GrowToInclude(MinToNewSpace, MaxToNewSpace, TransformPosition(SpaceTransform, Max));
@@ -43,27 +32,8 @@ export void TransformedAABB(const uniform FTransform &SpaceTransform, const unif
 
 	NewMin = SetVector(MinToNewSpace);
 	NewMax = SetVector(MaxToNewSpace);
-#else
-	const uniform FVector8 MinMax = SetVector8(Min, Max);
-	const uniform FVector8 CurrentExtentsWide = SetVector8(CurrentExtents);
-	const uniform FTransform8 SpaceTransformWide = SetTransform8(SpaceTransform);
-	const uniform FVector8 MinMaxTransform = TransformPosition(SpaceTransformWide, MinMax);
-	const uniform FVector8 ForwardVectorTransform = TransformPosition(SpaceTransformWide, MinMax + ForwardVector8 * CurrentExtentsWide);
-	const uniform FVector8 RightVectorTransform = TransformPosition(SpaceTransformWide, MinMax + RightVector8 * CurrentExtentsWide);
-	const uniform FVector8 UpVectorTransform = TransformPosition(SpaceTransformWide, MinMax + UpVector8 * CurrentExtentsWide);
-
-	uniform FVector8 MinMaxToNewSpace = MinMaxTransform;
-	GrowToInclude(MinMaxToNewSpace, MinMaxTransform);
-	GrowToInclude(MinMaxToNewSpace, ForwardVectorTransform);
-	GrowToInclude(MinMaxToNewSpace, RightVectorTransform);
-	GrowToInclude(MinMaxToNewSpace, UpVectorTransform);
-
-	NewMin = SetVector(MinMaxToNewSpace.V[0], MinMaxToNewSpace.V[1], MinMaxToNewSpace.V[2]);
-	NewMax = SetVector(MinMaxToNewSpace.V[4], MinMaxToNewSpace.V[5], MinMaxToNewSpace.V[6]);
-#endif
 }
 
-#if UE_LARGE_WORLD_COORDINATES_DISABLED == 0
 export void TransformedAABBMixed(const uniform FTransform &SpaceTransform, const uniform FVector3f &MMin, const uniform FVector3f &MMax, uniform FVector3f &NewMin, uniform FVector3f &NewMax)
 {
 	const uniform FVector Min3 = ConvertVector3fTo3Native(MMin);
@@ -89,4 +59,3 @@ export void TransformedAABBMixed(const uniform FTransform &SpaceTransform, const
 	NewMin = SetVector(MinToNewSpace4f);
 	NewMax = SetVector(MaxToNewSpace4f);
 }
-#endif
