Updated Low Level Tests
Removed existing SampleTest
Added an example Catch 2 Test for Math Colors into Core Module
Coverted Cbor tests to Catch2
Moved TDeque and TToken tests into Core Module

diff --git a/Engine/Source/Programs/LowLevelTests/LowLevelTests.Build.cs b/Engine/Source/Programs/LowLevelTests/LowLevelTests.Build.cs
--- a/Engine/Source/Programs/LowLevelTests/LowLevelTests.Build.cs
+++ b/Engine/Source/Programs/LowLevelTests/LowLevelTests.Build.cs
@@ -46,6 +46,7 @@ public class LowLevelTests : ModuleRules
 			PrivateDependencyModuleNames.AddRange(
 				new string[] {
 				"Core",
+				"Cbor",
 				"Projects",
 				"LowLevelTestsRunner"
 				});
diff --git a/Engine/Source/Runtime/Cbor/Private/Tests/CborTests.cpp b/Engine/Source/Runtime/Cbor/Private/Tests/CborTests.cpp
deleted file mode 100644
diff --git a/Engine/Source/Runtime/Cbor/Tests/CborTests.cpp b/Engine/Source/Runtime/Cbor/Tests/CborTests.cpp
new file mode 100644
--- /dev/null
+++ b/Engine/Source/Runtime/Cbor/Tests/CborTests.cpp
@@ -0,0 +1,390 @@
+// Copyright Epic Games, Inc. All Rights Reserved.
+
+#include "CoreMinimal.h"
+#include "Misc/AutomationTest.h"
+#include "Serialization/MemoryReader.h"
+#include "Serialization/MemoryWriter.h"
+
+#include "CborReader.h"
+#include "CborWriter.h"
+#include "CoreGlobals.h"
+#include "TestHarness.h"
+
+TEST_CASE("Core::FCbor::Smoke Test", "[Core][Smoke]")
+{
+	// Run the test twices, once with little-endian encoding, once with big-endian encoding.
+	auto RunWithEndiannessFn = [](ECborEndianness Endianness)
+	{
+		// Create the Writer
+		TArray<uint8> Bytes;
+		TUniquePtr<FArchive> OutputStream = MakeUnique<FMemoryWriter>(Bytes);
+		FCborWriter Writer(OutputStream.Get(), Endianness);
+
+		// Create the Reader
+		TUniquePtr<FArchive> InputStream = MakeUnique<FMemoryReader>(Bytes);
+		FCborReader Reader(InputStream.Get(), Endianness);
+
+		int64 TestInt = 0;
+		FCborContext Context;
+
+		// Positive Integer Item
+		Writer.WriteValue(TestInt);
+		CHECK(Reader.ReadNext(Context) == true);
+		CHECK(Context.MajorType() == ECborCode::Uint);
+		CHECK(Context.AsUInt() == TestInt);
+		CHECK(Context.AsInt() == TestInt);
+	
+		TestInt = 1;
+		Writer.WriteValue(TestInt);
+		CHECK(Reader.ReadNext(Context) == true);
+		CHECK(Context.MajorType() == ECborCode::Uint);
+		CHECK(Context.AsUInt() == TestInt);
+		CHECK(Context.AsInt() == TestInt);
+
+		TestInt = 10;
+		Writer.WriteValue(TestInt);
+		CHECK(Reader.ReadNext(Context) == true);
+		CHECK(Context.MajorType() == ECborCode::Uint);
+		CHECK(Context.AsUInt() == TestInt);
+		CHECK(Context.AsInt() == TestInt);
+
+		TestInt = 23;
+		Writer.WriteValue(TestInt);
+		CHECK(Reader.ReadNext(Context) == true);
+		CHECK(Context.MajorType() == ECborCode::Uint);
+		CHECK(Context.AsUInt() == TestInt);
+		CHECK(Context.AsInt() == TestInt);
+
+		TestInt = 24;
+		Writer.WriteValue(TestInt);
+		CHECK(Reader.ReadNext(Context) == true);
+		CHECK(Context.MajorType() == ECborCode::Uint);
+		CHECK(Context.AdditionalValue() == ECborCode::Value_1Byte);
+		CHECK(Context.AsUInt() == TestInt);
+		CHECK(Context.AsInt() == TestInt);
+
+		TestInt = 1000;
+		Writer.WriteValue(TestInt);
+		CHECK(Reader.ReadNext(Context) == true);
+		CHECK(Context.MajorType() == ECborCode::Uint);
+		CHECK(Context.AdditionalValue() == ECborCode::Value_2Bytes);
+
+		CHECK(Context.AsUInt() == TestInt);
+		CHECK(Context.AsInt() == TestInt);
+
+		TestInt = 3000000000;
+		Writer.WriteValue(TestInt);
+		CHECK(Reader.ReadNext(Context) == true);
+		CHECK(Context.MajorType() == ECborCode::Uint);
+		CHECK(Context.AdditionalValue() == ECborCode::Value_4Bytes);
+
+		CHECK(Context.AsUInt() == TestInt);
+		CHECK(Context.AsInt() == TestInt);
+
+		TestInt = 9223372036854775807;
+		Writer.WriteValue(TestInt);
+		CHECK(Reader.ReadNext(Context) == true);
+		CHECK(Context.MajorType() == ECborCode::Uint);
+		CHECK(Context.AdditionalValue() == ECborCode::Value_8Bytes);
+		CHECK(Context.AsUInt() == TestInt);
+		CHECK(Context.AsInt() == TestInt);
+
+		// Negative numbers
+
+		TestInt = -1;
+		Writer.WriteValue(TestInt);
+		CHECK(Reader.ReadNext(Context) == true);
+		CHECK(Context.MajorType() == ECborCode::Int);
+		CHECK(Context.AsInt() == TestInt);
+
+		TestInt = -23;
+		Writer.WriteValue(TestInt);
+		CHECK(Reader.ReadNext(Context) == true);
+		CHECK(Context.MajorType() == ECborCode::Int);
+		CHECK(Context.AsInt() == TestInt);
+
+		TestInt = -25;
+		Writer.WriteValue(TestInt);
+		CHECK(Reader.ReadNext(Context) == true);
+		CHECK(Context.MajorType() == ECborCode::Int);
+		CHECK(Context.AdditionalValue() == ECborCode::Value_1Byte);
+		CHECK(Context.AsInt() == TestInt);
+
+		TestInt = -1000;
+		Writer.WriteValue(TestInt);
+		CHECK(Reader.ReadNext(Context) == true);
+		CHECK(Context.MajorType() == ECborCode::Int);
+		CHECK(Context.AdditionalValue() == ECborCode::Value_2Bytes);
+		CHECK(Context.AsInt() == TestInt);
+
+		TestInt = -3000000000LL;
+		Writer.WriteValue(TestInt);
+		CHECK(Reader.ReadNext(Context) == true);
+		CHECK(Context.MajorType() == ECborCode::Int);
+		CHECK(Context.AdditionalValue() == ECborCode::Value_4Bytes);
+		CHECK(Context.AsInt() == TestInt);
+
+		TestInt = -92233720368547758LL; //-9223372036854775807LL;
+		Writer.WriteValue(TestInt);
+		CHECK(Reader.ReadNext(Context) == true);
+		CHECK(Context.MajorType() == ECborCode::Int);
+		CHECK(Context.AdditionalValue() == ECborCode::Value_8Bytes);
+		CHECK(Context.AsInt() == TestInt);
+
+		// Bool
+
+		bool TestBool = false;
+		Writer.WriteValue(TestBool);
+		CHECK(Reader.ReadNext(Context) == true);
+		CHECK(Context.MajorType() == ECborCode::Prim);
+		CHECK(Context.AdditionalValue() == ECborCode::False);
+		CHECK(Context.AsBool() == TestBool);
+
+		TestBool = true;
+		Writer.WriteValue(TestBool);
+		CHECK(Reader.ReadNext(Context) == true);
+		CHECK(Context.MajorType() == ECborCode::Prim);
+		CHECK(Context.AdditionalValue() == ECborCode::True);
+		CHECK(Context.AsBool() == TestBool);
+
+		// Float
+
+		float TestFloat = 3.14159265f;
+		Writer.WriteValue(TestFloat);
+		CHECK(Reader.ReadNext(Context) == true);
+		CHECK(Context.MajorType() == ECborCode::Prim);
+		CHECK(Context.AdditionalValue() == ECborCode::Value_4Bytes);
+		CHECK(Context.AsFloat() == TestFloat);
+
+		// Double
+
+		double TestDouble = 3.14159265; // 3.4028234663852886e+38;
+		Writer.WriteValue(TestDouble);
+		CHECK(Reader.ReadNext(Context) == true);
+		CHECK(Context.MajorType() == ECborCode::Prim);
+		CHECK(Context.AdditionalValue() == ECborCode::Value_8Bytes);
+		CHECK(Context.AsDouble() == TestDouble);
+
+		// String
+
+		FString TestString(TEXT("ANSIString"));
+
+		Writer.WriteValue(TestString);
+		CHECK(Reader.ReadNext(Context) == true);
+		CHECK(Context.MajorType() == ECborCode::TextString);
+		CHECK(Context.AsString() == TestString);
+
+		TestString = TEXT("\u3042\u308A\u304C\u3068\u3046");
+		Writer.WriteValue(TestString);
+		CHECK(Reader.ReadNext(Context) == true);
+		CHECK(Context.MajorType() == ECborCode::TextString);
+		CHECK(Context.AsString() == TestString);
+
+		// C String
+		char TestCString[] = "Potato";
+
+		Writer.WriteValue(TestCString, (sizeof(TestCString) / sizeof(char)) - 1); // do not count the null terminating character
+		CHECK(Reader.ReadNext(Context) == true);
+		CHECK(Context.MajorType() == ECborCode::ByteString);
+		CHECK(TCString<char>::Strcmp(Context.AsCString(), TestCString) == 0);
+
+		// Byte String. (with '\0' in the middle)
+		uint8 ByteString[] = {static_cast<uint8>(-1), static_cast<uint8>(-55), static_cast<uint8>(-128), 0, 1, 15, 127};
+		Writer.WriteValue(ByteString, sizeof(ByteString)/sizeof(uint8));
+		CHECK(Reader.ReadNext(Context) == true);
+		CHECK(Context.MajorType() == ECborCode::ByteString);
+		CHECK(FMemory::Memcmp(ByteString, Context.AsByteArray().GetData(), sizeof(ByteString)/sizeof(uint8)) == 0);
+		CHECK(Context.AsByteArray().Num() == sizeof(ByteString)/sizeof(uint8));
+
+		// Array
+		TArray<int64> IntArray { 0, 1, -1, 10, -1000, -3000000000LL, 240, -24 };
+		Writer.WriteContainerStart(ECborCode::Array, IntArray.Num());
+		for (int64 Val : IntArray)
+		{
+			Writer.WriteValue(Val);
+		}
+		// Array start & length
+		CHECK(Reader.ReadNext(Context) == true);
+		CHECK(Context.MajorType() == ECborCode::Array);
+		CHECK(Context.AsLength() == IntArray.Num());
+
+		for (int64 Val : IntArray)
+		{
+			CHECK(Reader.ReadNext(Context) == true);
+			CHECK(Context.AsInt() == Val);
+		}
+
+		// Read array end, report length 0 on finite container
+		// although the array wasn't written as indefinite,
+		// the reader will emit a virtual break token to notify the container end
+		CHECK(Reader.ReadNext(Context) == true);
+		CHECK(Context.IsBreak());
+		CHECK(Context.AsLength() == 0);
+
+		// Indefinite Array
+		Writer.WriteContainerStart(ECborCode::Array, -1);
+		for (int64 Val : IntArray)
+		{
+			Writer.WriteValue(Val);
+		}
+		Writer.WriteContainerEnd();
+
+		// Array start & length
+		CHECK(Reader.ReadNext(Context) == true);
+		CHECK(Context.MajorType() == ECborCode::Array);
+		CHECK(Context.IsIndefiniteContainer());
+		CHECK(Context.AsLength() == 0);
+
+		for (int64 Val : IntArray)
+		{
+			CHECK(Reader.ReadNext(Context) == true);
+			CHECK(Context.AsInt() == Val);
+		}
+
+		// Read array end, report length 
+		// although the array wasn't written as indefinite,
+		// the reader will emit a virtual break token to notify the container end
+		CHECK(Reader.ReadNext(Context) == true);
+		CHECK(Context.IsBreak());
+		CHECK(Context.AsLength() == IntArray.Num());
+
+		// Map
+		TMap<FString, FString> StringMap = { {TEXT("Apple"), TEXT("Orange")}, {TEXT("Potato"), TEXT("Tomato")}, {TEXT("Meat"), TEXT("Treat")} };
+		Writer.WriteContainerStart(ECborCode::Map, StringMap.Num());
+
+		for (const auto& Pair : StringMap)
+		{
+			Writer.WriteValue(Pair.Key);
+			Writer.WriteValue(Pair.Value);
+		}
+
+		// Map start & length
+		CHECK(Reader.ReadNext(Context) == true);
+		CHECK(Context.MajorType() == ECborCode::Map);
+		CHECK(Context.AsLength() == StringMap.Num() * 2);
+
+		for (const auto& Pair : StringMap)
+		{
+			CHECK(Reader.ReadNext(Context) == true);
+			CHECK(Context.AsString() == Pair.Key);
+			CHECK(Reader.ReadNext(Context) == true);
+			CHECK(Context.AsString() == Pair.Value);
+		}
+
+		// Read map end 
+		// although the array wasn't written as indefinite,
+		// the reader will emit a virtual break token to notify the container end
+		CHECK(Reader.ReadNext(Context) == true);
+		CHECK(Context.IsBreak());
+
+		CHECK(Reader.ReadNext(Context) == false);
+		CHECK(Context.RawCode() == ECborCode::StreamEnd);
+		return true;
+	};
+
+	RunWithEndiannessFn(ECborEndianness::LittleEndian);
+	RunWithEndiannessFn(ECborEndianness::BigEndian);
+
+	// Ensure that setting the endianness does something.
+	{
+		// Create the big endian writer
+		TArray<uint8> BytesBE;
+		TUniquePtr<FArchive> OutputStreamBE = MakeUnique<FMemoryWriter>(BytesBE);
+		FCborWriter WriterBE(OutputStreamBE.Get(), ECborEndianness::BigEndian);
+
+		// Create the little endian writer
+		TArray<uint8> BytesLE;
+		TUniquePtr<FArchive> OutputStreamLE = MakeUnique<FMemoryWriter>(BytesLE);
+		FCborWriter WriterLE(OutputStreamLE.Get(), ECborEndianness::LittleEndian);
+
+		// Write the same values to both streams and ensure the resulting streams are not identical.
+		WriterBE.WriteValue(0x1122334455667788ull);
+		WriterLE.WriteValue(0x1122334455667788ull);
+		CHECK(BytesBE != BytesLE);
+	}
+
+	// Ensure the 'Platform' endianness is correctly handled.
+	{
+		// Create the writer using the platform endianness.
+		TArray<uint8> BytesPlatform;
+		TUniquePtr<FArchive> OutputStreamPlatform = MakeUnique<FMemoryWriter>(BytesPlatform);
+		FCborWriter WriterPlatform(OutputStreamPlatform.Get(), ECborEndianness::Platform);
+
+		// Create the writer using the same endianness as the platform is expected to deduce to.
+		TArray<uint8> Bytes;
+		TUniquePtr<FArchive> OutputStream = MakeUnique<FMemoryWriter>(Bytes);
+		FCborWriter Writer(OutputStream.Get(), PLATFORM_LITTLE_ENDIAN != 0 ? ECborEndianness::LittleEndian : ECborEndianness::BigEndian);
+
+		// Write the same values to both streams and ensure the resulting streams are identical.
+		WriterPlatform.WriteValue((int64)0xDEADBEEFDEADBEEF);
+		Writer.WriteValue((int64)0xDEADBEEFDEADBEEF);
+		CHECK(BytesPlatform == Bytes);
+	}
+}
+
+TEST_CASE("Core::FCbor::ByteSwap Performance Test", "[Core][Perf]")
+{
+	//  NOTE: The most expensive value to write in CBOR is the double as it needs to swap 8 bytes all the time. Integers can be encoded in 1, 2, 4 or 8 bytes depending on their value, but not double.
+	//        String are not considered because they are written in UTF8. The test ensure this is not significantly longer to write/read while swapping the bytes.
+
+	const int32 ReservedByteCount = 512 * 1024 * 1024; // Seems like 512 MB is big enough to draw a worst case scenario. If 512 MB doesn't show a significant performance cost, then real life case should not either.
+	const int32 WriteCount = ReservedByteCount / (sizeof(double) + sizeof(FCborHeader)); // Each double has its own header.
+	
+	auto RunSample = [ReservedByteCount, WriteCount](ECborEndianness Endianness)
+	{
+		TArray<uint8> Bytes;
+		Bytes.Reserve(ReservedByteCount);
+		TUniquePtr<FArchive> OutputStream = MakeUnique<FMemoryWriter>(Bytes);
+		FCborWriter Writer(OutputStream.Get(), Endianness);
+
+		double Value = 1.0;
+
+		FDateTime WriteStartTime = FDateTime::UtcNow();
+		for (int32 I = 0; I < WriteCount; ++I)
+		{
+			Writer.WriteValue(Value);
+			Value += I + I * 0.5;
+		}
+		FTimespan WriteSpan = FDateTime::UtcNow() - WriteStartTime;
+		CHECK(Bytes.Num() <= ReservedByteCount);
+
+		TUniquePtr<FArchive> InputStream = MakeUnique<FMemoryReader>(Bytes);
+		FCborReader Reader(InputStream.Get(), Endianness);
+		FCborContext Ctx;
+		FDateTime ReadStartTime = FDateTime::UtcNow();
+		while (Reader.ReadNext(Ctx))
+		{
+			; // Just consume.
+		}
+		FTimespan ReadSpan = FDateTime::UtcNow() - ReadStartTime;
+
+		return MakeTuple(WriteSpan, ReadSpan);
+	};
+
+	ECborEndianness PlatformEndianness = PLATFORM_LITTLE_ENDIAN ? ECborEndianness::LittleEndian : ECborEndianness::BigEndian;
+	ECborEndianness SwapEndianness = PLATFORM_LITTLE_ENDIAN ? ECborEndianness::BigEndian : ECborEndianness::LittleEndian;
+
+	TTuple<FTimespan, FTimespan> PlatformEndianessTimes1 = RunSample(PlatformEndianness);
+	TTuple<FTimespan, FTimespan> SwapEndianessTimes1 = RunSample(SwapEndianness);
+
+	TTuple<FTimespan, FTimespan> PlatformEndianessTimes2 = RunSample(PlatformEndianness);
+	TTuple<FTimespan, FTimespan> SwapEndianessTimes2 = RunSample(SwapEndianness);
+
+	TTuple<FTimespan, FTimespan> PlatformEndianessTimes3 = RunSample(PlatformEndianness);
+	TTuple<FTimespan, FTimespan> SwapEndianessTimes3 = RunSample(SwapEndianness);
+
+	// Average the times.
+	FTimespan AvgWritePlatformEndianess = (PlatformEndianessTimes1.Get<0>() + PlatformEndianessTimes2.Get<0>() + PlatformEndianessTimes3.Get<0>()) / 3;
+	FTimespan AvgReadPlatformEndianess  = (PlatformEndianessTimes1.Get<1>() + PlatformEndianessTimes2.Get<1>() + PlatformEndianessTimes3.Get<1>()) / 3;
+	FTimespan AvgWriteSwapEndianess = (SwapEndianessTimes1.Get<0>() + SwapEndianessTimes2.Get<0>() + SwapEndianessTimes3.Get<0>()) / 3;
+	FTimespan AvgReadSwapEndianess  = (SwapEndianessTimes1.Get<1>() + SwapEndianessTimes2.Get<1>() + SwapEndianessTimes3.Get<1>()) / 3;
+
+	// Ratio.
+	double WriteRatio = static_cast<double>(AvgWriteSwapEndianess.GetTicks()) / static_cast<double>(AvgWritePlatformEndianess.GetTicks());
+	double ReadRatio = static_cast<double>(AvgReadSwapEndianess.GetTicks()) / static_cast<double>(AvgReadPlatformEndianess.GetTicks());
+
+	// The ratio is usually around 1 +/- 0.08 as we don't measure significant performance change, but to account for the testing machine workload, be safe, use a large enough margin.
+	double Margin = 0.5;
+	CHECK(( WriteRatio >= 1.0 - Margin && WriteRatio <= 1.0 + Margin)==true);
+	CHECK(( ReadRatio  >= 1.0 - Margin && ReadRatio  <= 1.0 + Margin)==true);
+}
\ No newline at end of file
diff --git a/Engine/Source/Runtime/Core/Tests/Containers/Deque.cpp b/Engine/Source/Runtime/Core/Tests/Containers/Deque.cpp
--- a/Engine/Source/Runtime/Core/Tests/Containers/Deque.cpp
+++ b/Engine/Source/Runtime/Core/Tests/Containers/Deque.cpp
@@ -83,7 +83,7 @@ bool Test::EmplaceFirstPopLast(TDeque<int32Token>& Deque, int32 Count)
 // Unit tests
 //---------------------------------------------------------------------------------------------------------------------
 
-TEST_CASE("Containers::Deque::Reserve without data", "[Deque]")
+TEST_CASE("Core::Containers::TDeque::Reserve without data", "[Core][Containers][Smoke]")
 {
 	int32Token::Reset();
 	TDeque<int32Token> Deque;
@@ -97,7 +97,7 @@ TEST_CASE("Containers::Deque::Reserve without data", "[Deque]")
 	REQUIRE(int32Token::EvenConstructionDestructionCalls(0));
 }
 
-TEST_CASE("Containers::Deque::Reserve EmplaceLast single element", "[Deque]")
+TEST_CASE("Core::Containers::TDeque::Reserve EmplaceLast single element", "[Core][Containers][Smoke]")
 {
 	int32Token::Reset();
 	{
@@ -116,7 +116,7 @@ TEST_CASE("Containers::Deque::Reserve EmplaceLast single element", "[Deque]")
 	REQUIRE(int32Token::EvenConstructionDestructionCalls(1));
 }
 
-TEST_CASE("Containers::Deque::Reset", "[Deque]")
+TEST_CASE("Core::Containers::TDeque::Reset", "[Core][Containers][Smoke]")
 {
 	int32Token::Reset();
 	{
@@ -137,7 +137,7 @@ TEST_CASE("Containers::Deque::Reset", "[Deque]")
 	REQUIRE(int32Token::EvenConstructionDestructionCalls(1));
 }
 
-TEST_CASE("Containers::Deque::Empty", "[Deque]")
+TEST_CASE("Core::Containers::TDeque::Empty", "[Core][Containers][Smoke]")
 {
 	int32Token::Reset();
 	{
@@ -150,7 +150,7 @@ TEST_CASE("Containers::Deque::Empty", "[Deque]")
 	REQUIRE(int32Token::EvenConstructionDestructionCalls(0));
 }
 
-TEST_CASE("Containers::Deque::Empty after single element EmplaceLast", "[Deque]")
+TEST_CASE("Core::Containers::TDeque::Empty after single element EmplaceLast", "[Core][Containers][Smoke]")
 {
 	int32Token::Reset();
 	{
@@ -167,7 +167,7 @@ TEST_CASE("Containers::Deque::Empty after single element EmplaceLast", "[Deque]"
 	REQUIRE(int32Token::EvenConstructionDestructionCalls(1));
 }
 
-TEST_CASE("Containers::Deque::EmplaceLast single element", "[Deque]")
+TEST_CASE("Core::Containers::TDeque::EmplaceLast single element", "[Core][Containers][Smoke]")
 {
 	int32Token::Reset();
 	{
@@ -183,7 +183,7 @@ TEST_CASE("Containers::Deque::EmplaceLast single element", "[Deque]")
 	REQUIRE(int32Token::EvenConstructionDestructionCalls(1));
 }
 
-TEST_CASE("Containers::Deque::EmplaceLast range to capacity", "[Deque]")
+TEST_CASE("Core::Containers::TDeque::EmplaceLast range to capacity", "[Core][Containers][Smoke]")
 {
 	int32Token::Reset();
 	{
@@ -202,7 +202,7 @@ TEST_CASE("Containers::Deque::EmplaceLast range to capacity", "[Deque]")
 	REQUIRE(int32Token::EvenConstructionDestructionCalls(Test::DefaultCapacity * 10));
 }
 
-TEST_CASE("Containers::Deque::EmplaceLast range past capacity", "[Deque]")
+TEST_CASE("Core::Containers::TDeque::EmplaceLast range past capacity", "[Core][Containers][Smoke]")
 {
 	int32Token::Reset();
 	{
@@ -224,7 +224,7 @@ TEST_CASE("Containers::Deque::EmplaceLast range past capacity", "[Deque]")
 	REQUIRE(int32Token::EvenConstructionDestructionCalls(Test::DefaultCapacity + 1));
 }
 
-TEST_CASE("Containers::Deque::EmplaceFirst single element", "[Deque]")
+TEST_CASE("Core::Containers::TDeque::EmplaceFirst single element", "[Core][Containers][Smoke]")
 {
 	int32Token::Reset();
 	{
@@ -240,7 +240,7 @@ TEST_CASE("Containers::Deque::EmplaceFirst single element", "[Deque]")
 	REQUIRE(int32Token::EvenConstructionDestructionCalls(1));
 }
 
-TEST_CASE("Containers::Deque::EmplaceFirst range to capacity", "[Deque]")
+TEST_CASE("Core::Containers::TDeque::EmplaceFirst range to capacity", "[Core][Containers][Smoke]")
 {
 	int32Token::Reset();
 	{
@@ -259,7 +259,7 @@ TEST_CASE("Containers::Deque::EmplaceFirst range to capacity", "[Deque]")
 	REQUIRE(int32Token::EvenConstructionDestructionCalls(Test::DefaultCapacity * 10));
 }
 
-TEST_CASE("Containers::Deque::PushLast single element (implicit move)", "[Deque]")
+TEST_CASE("Core::Containers::TDeque::PushLast single element (implicit move)", "[Core][Containers][Smoke]")
 {
 	int32Token::Reset();
 	{
@@ -278,7 +278,7 @@ TEST_CASE("Containers::Deque::PushLast single element (implicit move)", "[Deque]
 	REQUIRE(int32Token::NumMoveConstructorCalls() == 1);
 }
 
-TEST_CASE("Containers::Deque::PushLast single element from move", "[Deque]")
+TEST_CASE("Core::Containers::TDeque::PushLast single element from move", "[Core][Containers][Smoke]")
 {
 	int32Token::Reset();
 	{
@@ -298,7 +298,7 @@ TEST_CASE("Containers::Deque::PushLast single element from move", "[Deque]")
 	REQUIRE(int32Token::NumMoveConstructorCalls() == 1);
 }
 
-TEST_CASE("Containers::Deque::PushLast single element from copy", "[Deque]")
+TEST_CASE("Core::Containers::TDeque::PushLast single element from copy", "[Core][Containers][Smoke]")
 {
 	int32Token::Reset();
 	{
@@ -318,7 +318,7 @@ TEST_CASE("Containers::Deque::PushLast single element from copy", "[Deque]")
 	REQUIRE(int32Token::NumMoveConstructorCalls() == 0);
 }
 
-TEST_CASE("Containers::Deque::PushFirst single element (implicit move)", "[Deque]")
+TEST_CASE("Core::Containers::TDeque::PushFirst single element (implicit move)", "[Core][Containers][Smoke]")
 {
 	int32Token::Reset();
 	{
@@ -337,7 +337,7 @@ TEST_CASE("Containers::Deque::PushFirst single element (implicit move)", "[Deque
 	REQUIRE(int32Token::NumMoveConstructorCalls() == 1);
 }
 
-TEST_CASE("Containers::Deque::PushFirst single element from move", "[Deque]")
+TEST_CASE("Core::Containers::TDeque::PushFirst single element from move", "[Core][Containers][Smoke]")
 {
 	int32Token::Reset();
 	{
@@ -357,7 +357,7 @@ TEST_CASE("Containers::Deque::PushFirst single element from move", "[Deque]")
 	REQUIRE(int32Token::NumMoveConstructorCalls() == 1);
 }
 
-TEST_CASE("Containers::Deque::PushFirst single element from copy", "[Deque]")
+TEST_CASE("Core::Containers::TDeque::PushFirst single element from copy", "[Core][Containers][Smoke]")
 {
 	int32Token::Reset();
 	{
@@ -382,7 +382,7 @@ void PopOne(TDeque<int32Token>& Deque)
 	Deque.PopFirst();
 }
 
-TEST_CASE("Containers::Deque::EmplaceLast/PopFirst single element", "[Deque]")
+TEST_CASE("Core::Containers::TDeque::EmplaceLast/PopFirst single element", "[Core][Containers][Smoke]")
 {
 	int32Token::Reset();
 	{
@@ -402,7 +402,7 @@ TEST_CASE("Containers::Deque::EmplaceLast/PopFirst single element", "[Deque]")
 	REQUIRE(int32Token::EvenConstructionDestructionCalls(1));
 }
 
-TEST_CASE("Containers::Deque::EmplaceLast/PopFirst single element multiple times causing head/tail wrap around", "[Deque]")
+TEST_CASE("Core::Containers::TDeque::EmplaceLast/PopFirst single element multiple times causing head/tail wrap around", "[Core][Containers][Smoke]")
 {
 	int32Token::Reset();
 	{
@@ -418,7 +418,7 @@ TEST_CASE("Containers::Deque::EmplaceLast/PopFirst single element multiple times
 	REQUIRE(int32Token::EvenConstructionDestructionCalls(Test::DefaultCapacity * 2));
 }
 
-TEST_CASE("Containers::Deque::EmplaceLast/PopFirst range without head/tail wrap around", "[Deque]")
+TEST_CASE("Core::Containers::TDeque::EmplaceLast/PopFirst range without head/tail wrap around", "[Core][Containers][Smoke]")
 {
 	int32Token::Reset();
 	{
@@ -430,7 +430,7 @@ TEST_CASE("Containers::Deque::EmplaceLast/PopFirst range without head/tail wrap
 	REQUIRE(int32Token::EvenConstructionDestructionCalls(Test::DefaultCapacity));
 }
 
-TEST_CASE("Containers::Deque::EmplaceLast/PopFirst range with reallocation without head/tail wrap around", "[Deque]")
+TEST_CASE("Core::Containers::TDeque::EmplaceLast/PopFirst range with reallocation without head/tail wrap around", "[Core][Containers][Smoke]")
 {
 	int32Token::Reset();
 	{
@@ -443,7 +443,7 @@ TEST_CASE("Containers::Deque::EmplaceLast/PopFirst range with reallocation witho
 	REQUIRE(int32Token::EvenConstructionDestructionCalls(Test::DefaultCapacity + 1));
 }
 
-TEST_CASE("Containers::Deque::EmplaceLast/PopFirst range with head/tail wrap around", "[Deque]")
+TEST_CASE("Core::Containers::TDeque::EmplaceLast/PopFirst range with head/tail wrap around", "[Core][Containers][Smoke]")
 {
 	int32Token::Reset();
 	{
@@ -458,7 +458,7 @@ TEST_CASE("Containers::Deque::EmplaceLast/PopFirst range with head/tail wrap aro
 	REQUIRE(int32Token::EvenConstructionDestructionCalls(Test::DefaultCapacity * (Test::DefaultCapacity - 1)));
 }
 
-TEST_CASE("Containers::Deque::EmplaceFirst/PopLast single element", "[Deque]")
+TEST_CASE("Core::Containers::TDeque::EmplaceFirst/PopLast single element", "[Core][Containers][Smoke]")
 {
 	int32Token::Reset();
 	{
@@ -478,7 +478,7 @@ TEST_CASE("Containers::Deque::EmplaceFirst/PopLast single element", "[Deque]")
 	REQUIRE(int32Token::EvenConstructionDestructionCalls(1));
 }
 
-TEST_CASE("Containers::Deque::EmplaceFirst/PopLast range", "[Deque]")
+TEST_CASE("Core::Containers::TDeque::EmplaceFirst/PopLast range", "[Core][Containers][Smoke]")
 {
 	int32Token::Reset();
 	{
@@ -490,7 +490,7 @@ TEST_CASE("Containers::Deque::EmplaceFirst/PopLast range", "[Deque]")
 	REQUIRE(int32Token::EvenConstructionDestructionCalls(Test::DefaultCapacity));
 }
 
-TEST_CASE("Containers::Deque::EmplaceFirst/PopLast range with reallocation", "[Deque]")
+TEST_CASE("Core::Containers::TDeque::EmplaceFirst/PopLast range with reallocation", "[Core][Containers][Smoke]")
 {
 	int32Token::Reset();
 	{
@@ -503,7 +503,7 @@ TEST_CASE("Containers::Deque::EmplaceFirst/PopLast range with reallocation", "[D
 	REQUIRE(int32Token::EvenConstructionDestructionCalls(Test::DefaultCapacity + 1));
 }
 
-TEST_CASE("Containers::Deque::TryPopFirst", "[Deque]")
+TEST_CASE("Core::Containers::TDeque::TryPopFirst", "[Core][Containers][Smoke]")
 {
 	int32Token::Reset();
 	{
@@ -526,7 +526,7 @@ TEST_CASE("Containers::Deque::TryPopFirst", "[Deque]")
 	REQUIRE(int32Token::EvenConstructionDestructionCalls(Test::DefaultCapacity + 1));  // + 1 for Result
 }
 
-TEST_CASE("Containers::Deque::TryPopFirst with reallocation", "[Deque]")
+TEST_CASE("Core::Containers::TDeque::TryPopFirst with reallocation", "[Core][Containers][Smoke]")
 {
 	int32Token::Reset();
 	{
@@ -552,7 +552,7 @@ TEST_CASE("Containers::Deque::TryPopFirst with reallocation", "[Deque]")
 	REQUIRE(int32Token::EvenConstructionDestructionCalls(Test::DefaultCapacity + 2));
 }
 
-TEST_CASE("Containers::Deque::TryPopLast", "[Deque]")
+TEST_CASE("Core::Containers::TDeque::TryPopLast", "[Core][Containers][Smoke]")
 {
 	int32Token::Reset();
 	{
@@ -575,7 +575,7 @@ TEST_CASE("Containers::Deque::TryPopLast", "[Deque]")
 	REQUIRE(int32Token::EvenConstructionDestructionCalls(Test::DefaultCapacity + 1));  // + 1 for Result
 }
 
-TEST_CASE("Containers::Deque::TryPopLast with reallocation", "[Deque]")
+TEST_CASE("Core::Containers::TDeque::TryPopLast with reallocation", "[Core][Containers][Smoke]")
 {
 	int32Token::Reset();
 	{
@@ -601,7 +601,7 @@ TEST_CASE("Containers::Deque::TryPopLast with reallocation", "[Deque]")
 	REQUIRE(int32Token::EvenConstructionDestructionCalls(Test::DefaultCapacity + 2));
 }
 
-TEST_CASE("Containers::Deque::Comparison simple", "[Deque]")
+TEST_CASE("Core::Containers::TDeque::Comparison simple", "[Core][Containers][Smoke]")
 {
 	int32Token::Reset();
 	{
@@ -620,7 +620,7 @@ TEST_CASE("Containers::Deque::Comparison simple", "[Deque]")
 	REQUIRE(int32Token::EvenConstructionDestructionCalls(Test::DefaultCapacity * 2));
 }
 
-TEST_CASE("Containers::Deque::Comparison with head/tail wrap around", "[Deque]")
+TEST_CASE("Core::Containers::TDeque::Comparison with head/tail wrap around", "[Core][Containers][Smoke]")
 {
 	int32Token::Reset();
 	{
@@ -648,7 +648,7 @@ TEST_CASE("Containers::Deque::Comparison with head/tail wrap around", "[Deque]")
 	REQUIRE(int32Token::EvenConstructionDestructionCalls(Test::DefaultCapacity * 6));
 }
 
-TEST_CASE("Containers::Deque::Copy simple", "[Deque]")
+TEST_CASE("Core::Containers::TDeque::Copy simple", "[Core][Containers][Smoke]")
 {
 	int32Token::Reset();
 	{
@@ -666,7 +666,7 @@ TEST_CASE("Containers::Deque::Copy simple", "[Deque]")
 	REQUIRE(int32Token::EvenConstructionDestructionCalls(Test::DefaultCapacity * 2));
 }
 
-TEST_CASE("Containers::Deque::Copy with head/tail wrap around", "[Deque]")
+TEST_CASE("Core::Containers::TDeque::Copy with head/tail wrap around", "[Core][Containers][Smoke]")
 {
 	int32Token::Reset();
 	{
@@ -689,7 +689,7 @@ TEST_CASE("Containers::Deque::Copy with head/tail wrap around", "[Deque]")
 	REQUIRE(int32Token::EvenConstructionDestructionCalls(Test::DefaultCapacity * 9));
 }
 
-TEST_CASE("Containers::Deque::Copy variable size with head/tail wrap around", "[Deque]")
+TEST_CASE("Core::Containers::TDeque::Copy variable size with head/tail wrap around", "[Core][Containers][Smoke]")
 {
 	int32Token::Reset();
 	{
@@ -717,7 +717,7 @@ TEST_CASE("Containers::Deque::Copy variable size with head/tail wrap around", "[
 	REQUIRE(int32Token::EvenConstructionDestructionCalls());
 }
 
-TEST_CASE("Containers::Deque::Move simple", "[Deque]")
+TEST_CASE("Core::Containers::TDeque::Move simple", "[Core][Containers][Smoke]")
 {
 	int32Token::Reset();
 	{
@@ -741,7 +741,7 @@ TEST_CASE("Containers::Deque::Move simple", "[Deque]")
 	REQUIRE(int32Token::EvenConstructionDestructionCalls(Test::DefaultCapacity + 1));
 }
 
-TEST_CASE("Containers::Deque::Move with head/tail wrap around", "[Deque]")
+TEST_CASE("Core::Containers::TDeque::Move with head/tail wrap around", "[Core][Containers][Smoke]")
 {
 	int32Token::Reset();
 	{
@@ -769,7 +769,7 @@ TEST_CASE("Containers::Deque::Move with head/tail wrap around", "[Deque]")
 	REQUIRE(int32Token::EvenConstructionDestructionCalls(Test::DefaultCapacity * 6));
 }
 
-TEST_CASE("Containers::Deque::Move variable size with head/tail wrap around", "[Deque]")
+TEST_CASE("Core::Containers::TDeque::Move variable size with head/tail wrap around", "[Core][Containers][Smoke]")
 {
 	int32Token::Reset();
 	{
@@ -801,7 +801,7 @@ TEST_CASE("Containers::Deque::Move variable size with head/tail wrap around", "[
 	REQUIRE(int32Token::EvenConstructionDestructionCalls());
 }
 
-TEST_CASE("Containers::Deque::Iteration without head/tail wrap around", "[Deque]")
+TEST_CASE("Core::Containers::TDeque::Iteration without head/tail wrap around", "[Core][Containers][Smoke]")
 {
 	int32Token::Reset();
 	{
@@ -826,7 +826,7 @@ TEST_CASE("Containers::Deque::Iteration without head/tail wrap around", "[Deque]
 	REQUIRE(int32Token::EvenConstructionDestructionCalls(Test::DefaultCapacity));
 }
 
-TEST_CASE("Containers::Deque::Iteration with head/tail wrap around", "[Deque]")
+TEST_CASE("Core::Containers::TDeque::Iteration with head/tail wrap around", "[Core][Containers][Smoke]")
 {
 	int32Token::Reset();
 	{
@@ -856,7 +856,7 @@ TEST_CASE("Containers::Deque::Iteration with head/tail wrap around", "[Deque]")
 	REQUIRE(int32Token::EvenConstructionDestructionCalls(Test::DefaultCapacity * (Test::DefaultCapacity + 1)));
 }
 
-TEST_CASE("Containers::Deque::Iterator arithmetic", "[Deque]")
+TEST_CASE("Core::Containers::TDeque::Iterator arithmetic", "[Core][Containers][Smoke]")
 {
 	TDeque<int32Token> Deque;
 	Deque.Reserve(Test::DefaultCapacity);
@@ -883,7 +883,7 @@ TEST_CASE("Containers::Deque::Iterator arithmetic", "[Deque]")
 	}
 }
 
-TEST_CASE("Containers::Deque::Construct from std initializer_list", "[Deque]")
+TEST_CASE("Core::Containers::TDeque::Construct from std initializer_list", "[Core][Containers][Smoke]")
 {
 	int32Token::Reset();
 	{
@@ -899,7 +899,7 @@ TEST_CASE("Containers::Deque::Construct from std initializer_list", "[Deque]")
 	REQUIRE(int32Token::EvenConstructionDestructionCalls(6 * 2 + 1));
 }
 
-TEST_CASE("Containers::Deque::Construct from empty std initializer_list", "[Deque]")
+TEST_CASE("Core::Containers::TDeque::Construct from empty std initializer_list", "[Core][Containers][Smoke]")
 {
 	int32Token::Reset();
 	{
@@ -909,7 +909,7 @@ TEST_CASE("Containers::Deque::Construct from empty std initializer_list", "[Dequ
 	REQUIRE(int32Token::EvenConstructionDestructionCalls(0));
 }
 
-TEST_CASE("Containers::Deque::Assign from std initializer_list", "[Deque]")
+TEST_CASE("Core::Containers::TDeque::Assign from std initializer_list", "[Core][Containers][Smoke]")
 {
 	int32Token::Reset();
 	{
diff --git a/Engine/Source/Runtime/Core/Tests/Containers/TestToken.cpp b/Engine/Source/Runtime/Core/Tests/Containers/TestToken.cpp
--- a/Engine/Source/Runtime/Core/Tests/Containers/TestToken.cpp
+++ b/Engine/Source/Runtime/Core/Tests/Containers/TestToken.cpp
@@ -3,7 +3,7 @@
 #include "TestToken.h"
 #include "TestHarness.h"
 
-TEST_CASE("Containers::TToken::TTestToken default", "[Container][TTestToken]")
+TEST_CASE("Core::Containers::TToken::Default", "[Core][Containers][Smoke]")
 {
 	int32Token::Reset();
 	{
@@ -18,7 +18,7 @@ TEST_CASE("Containers::TToken::TTestToken default", "[Container][TTestToken]")
 	REQUIRE(int32Token::NumDestructionCalls() == 1);
 }
 
-TEST_CASE("Containers::TToken::TTestToken explicit constructor", "[Container][TTestToken]")
+TEST_CASE("Core::Containers::TToken::Explicit constructor", "[Core][Containers][Smoke]")
 {
 	int32Token::Reset();
 	{
@@ -35,7 +35,7 @@ TEST_CASE("Containers::TToken::TTestToken explicit constructor", "[Container][TT
 	REQUIRE(int32Token::NumMoveAssignmentCalls() == 0);
 }
 
-TEST_CASE("Containers::TToken::TTestToken copy constructor", "[Container][TTestToken]")
+TEST_CASE("Core::Containers::TToken::Copy constructor", "[Core][Containers][Smoke]")
 {
 	int32Token::Reset();
 	{
@@ -53,7 +53,7 @@ TEST_CASE("Containers::TToken::TTestToken copy constructor", "[Container][TTestT
 	REQUIRE(int32Token::NumMoveAssignmentCalls() == 0);
 }
 
-TEST_CASE("Containers::TToken::TTestToken move constructor", "[Container][TTestToken]")
+TEST_CASE("Core::Containers::TToken::Move constructor", "[Core][Containers][Smoke]")
 {
 	int32Token::Reset();
 	{
@@ -71,7 +71,7 @@ TEST_CASE("Containers::TToken::TTestToken move constructor", "[Container][TTestT
 	REQUIRE(int32Token::NumMoveAssignmentCalls() == 0);
 }
 
-TEST_CASE("Containers::TToken::TTestToken copy assignment", "[Container][TTestToken]")
+TEST_CASE("Core::Containers::TToken::Copy assignment", "[Core][Containers][Smoke]")
 {
 	int32Token::Reset();
 	{
@@ -90,7 +90,7 @@ TEST_CASE("Containers::TToken::TTestToken copy assignment", "[Container][TTestTo
 	REQUIRE(int32Token::NumMoveAssignmentCalls() == 0);
 }
 
-TEST_CASE("Containers::TToken::TTestToken move assignment", "[Container][TTestToken]")
+TEST_CASE("Core::Containers::TToken::Move assignment", "[Core][Containers][Smoke]")
 {
 	int32Token::Reset();
 	{
diff --git a/Engine/Source/Runtime/Core/Tests/Math/MathTests.cpp b/Engine/Source/Runtime/Core/Tests/Math/MathTests.cpp
--- a/Engine/Source/Runtime/Core/Tests/Math/MathTests.cpp
+++ b/Engine/Source/Runtime/Core/Tests/Math/MathTests.cpp
@@ -3,7 +3,7 @@
 #include "Math/Color.h"
 #include "TestHarness.h"
 
-TEST_CASE("Math::Colour::FLinearColor", "[Core][Math][Smoke]")
+TEST_CASE("Core::Math::FLinearColor::Smoke Test", "[Core][Math][Smoke]")
 {
 	SECTION("FLinearColor")
 	{
@@ -34,7 +34,7 @@ TEST_CASE("Math::Colour::FLinearColor", "[Core][Math][Smoke]")
 	}
 }
 
-TEST_CASE("Math::Colour::FColor", "[Core][Math][Smoke]")
+TEST_CASE("Core::Math::FColor::Smoke Test", "[Core][Math][Smoke]")
 {
 	SECTION("FColor")
 	{
