Fix non unity build error for Unreal Insights

diff --git a/Engine/Source/Runtime/Core/Private/ProfilingDebugging/TraceAuxiliary.cpp b/Engine/Source/Runtime/Core/Private/ProfilingDebugging/TraceAuxiliary.cpp
--- a/Engine/Source/Runtime/Core/Private/ProfilingDebugging/TraceAuxiliary.cpp
+++ b/Engine/Source/Runtime/Core/Private/ProfilingDebugging/TraceAuxiliary.cpp
@@ -15,7 +15,8 @@
 #endif
 
 #if WITH_UNREAL_TRACE_LAUNCH
-#	include "Misc/Parse.h"
+#include "Misc/Parse.h"
+#include "Misc/Paths.h"
 #endif
 
 #if UE_TRACE_ENABLED
