String: Extended ParseTokens to support ANSICHAR and UTF8CHAR

#rnx

diff --git a/Engine/Source/Runtime/Core/Private/String/ParseTokens.cpp b/Engine/Source/Runtime/Core/Private/String/ParseTokens.cpp
--- a/Engine/Source/Runtime/Core/Private/String/ParseTokens.cpp
+++ b/Engine/Source/Runtime/Core/Private/String/ParseTokens.cpp
@@ -9,13 +9,14 @@
 #include "Containers/Array.h"
 #include "Containers/BitArray.h"
 
-namespace UE::String
+namespace UE::String::Private
 {
 
+template <typename CharType>
 inline static void ParseTokensVisitToken(
-	const TFunctionRef<void (FStringView)>& Visitor,
+	const TFunctionRef<void (TStringView<CharType>)>& Visitor,
 	const EParseTokensOptions Options,
-	FStringView Token)
+	TStringView<CharType> Token)
 {
 	if (EnumHasAnyFlags(Options, EParseTokensOptions::Trim))
 	{
@@ -28,31 +29,32 @@ inline static void ParseTokensVisitToken(
 }
 
 /** Parse tokens with one single-character delimiter. */
+template <typename CharType>
 inline static void ParseTokens1Delim1Char(
-	const FStringView View,
-	const TCHAR Delimiter,
-	const TFunctionRef<void (FStringView)> Visitor,
+	const TStringView<CharType> View,
+	const CharType Delimiter,
+	const TFunctionRef<void (TStringView<CharType>)> Visitor,
 	const EParseTokensOptions Options)
 {
-	const TCHAR* ViewIt = View.GetData();
-	const TCHAR* const ViewEnd = ViewIt + View.Len();
-	const TCHAR* NextToken = ViewIt;
+	const CharType* ViewIt = View.GetData();
+	const CharType* const ViewEnd = ViewIt + View.Len();
+	const CharType* NextToken = ViewIt;
 
 	if (EnumHasAnyFlags(Options, EParseTokensOptions::IgnoreCase))
 	{
-		const TCHAR LowerDelimiter = FChar::ToLower(Delimiter);
+		const CharType LowerDelimiter = TChar<CharType>::ToLower(Delimiter);
 		for (;;)
 		{
 			if (ViewIt == ViewEnd)
 			{
 				break;
 			}
-			if (FChar::ToLower(*ViewIt) != Delimiter)
+			if (TChar<CharType>::ToLower(*ViewIt) != Delimiter)
 			{
 				++ViewIt;
 				continue;
 			}
-			ParseTokensVisitToken(Visitor, Options, FStringView(NextToken, UE_PTRDIFF_TO_INT32(ViewIt - NextToken)));
+			ParseTokensVisitToken(Visitor, Options, TStringView<CharType>(NextToken, UE_PTRDIFF_TO_INT32(ViewIt - NextToken)));
 			NextToken = ++ViewIt;
 		}
 	}
@@ -69,41 +71,42 @@ inline static void ParseTokens1Delim1Char(
 				++ViewIt;
 				continue;
 			}
-			ParseTokensVisitToken(Visitor, Options, FStringView(NextToken, UE_PTRDIFF_TO_INT32(ViewIt - NextToken)));
+			ParseTokensVisitToken(Visitor, Options, TStringView<CharType>(NextToken, UE_PTRDIFF_TO_INT32(ViewIt - NextToken)));
 			NextToken = ++ViewIt;
 		}
 	}
 
-	ParseTokensVisitToken(Visitor, Options, FStringView(NextToken, UE_PTRDIFF_TO_INT32(ViewIt - NextToken)));
+	ParseTokensVisitToken(Visitor, Options, TStringView<CharType>(NextToken, UE_PTRDIFF_TO_INT32(ViewIt - NextToken)));
 }
 
 /** Parse tokens with multiple single-character Basic Latin delimiters. */
+template <typename CharType>
 inline static void ParseTokensNDelim1CharBasicLatin(
-	const FStringView View,
-	const TConstArrayView<TCHAR> Delimiters,
-	const TFunctionRef<void (FStringView)> Visitor,
+	const TStringView<CharType> View,
+	const TConstArrayView<CharType> Delimiters,
+	const TFunctionRef<void (TStringView<CharType>)> Visitor,
 	const EParseTokensOptions Options)
 {
 	TBitArray<> DelimiterMask(false, 128);
 	if (EnumHasAnyFlags(Options, EParseTokensOptions::IgnoreCase))
 	{
-		for (TCHAR Delimiter : Delimiters)
+		for (CharType Delimiter : Delimiters)
 		{
-			DelimiterMask[FChar::ToUnsigned(FChar::ToLower(Delimiter))] = true;
-			DelimiterMask[FChar::ToUnsigned(FChar::ToUpper(Delimiter))] = true;
+			DelimiterMask[TChar<CharType>::ToUnsigned(TChar<CharType>::ToLower(Delimiter))] = true;
+			DelimiterMask[TChar<CharType>::ToUnsigned(TChar<CharType>::ToUpper(Delimiter))] = true;
 		}
 	}
 	else
 	{
-		for (TCHAR Delimiter : Delimiters)
+		for (CharType Delimiter : Delimiters)
 		{
-			DelimiterMask[FChar::ToUnsigned(Delimiter)] = true;
+			DelimiterMask[TChar<CharType>::ToUnsigned(Delimiter)] = true;
 		}
 	}
 
-	const TCHAR* ViewIt = View.GetData();
-	const TCHAR* const ViewEnd = ViewIt + View.Len();
-	const TCHAR* NextToken = ViewIt;
+	const CharType* ViewIt = View.GetData();
+	const CharType* const ViewEnd = ViewIt + View.Len();
+	const CharType* NextToken = ViewIt;
 
 	for (;;)
 	{
@@ -117,45 +120,46 @@ inline static void ParseTokensNDelim1CharBasicLatin(
 			++ViewIt;
 			continue;
 		}
-		ParseTokensVisitToken(Visitor, Options, FStringView(NextToken, UE_PTRDIFF_TO_INT32(ViewIt - NextToken)));
+		ParseTokensVisitToken(Visitor, Options, TStringView<CharType>(NextToken, UE_PTRDIFF_TO_INT32(ViewIt - NextToken)));
 		NextToken = ++ViewIt;
 	}
 
-	ParseTokensVisitToken(Visitor, Options, FStringView(NextToken, UE_PTRDIFF_TO_INT32(ViewIt - NextToken)));
+	ParseTokensVisitToken(Visitor, Options, TStringView<CharType>(NextToken, UE_PTRDIFF_TO_INT32(ViewIt - NextToken)));
 }
 
 /** Parse tokens with multiple single-character delimiters in the Basic Multilingual Plane. */
+template <typename CharType>
 inline static void ParseTokensNDelim1Char(
-	const FStringView View,
-	const TConstArrayView<TCHAR> Delimiters,
-	TFunctionRef<void (FStringView)> Visitor,
+	const TStringView<CharType> View,
+	const TConstArrayView<CharType> Delimiters,
+	TFunctionRef<void (TStringView<CharType>)> Visitor,
 	const EParseTokensOptions Options)
 {
-	if (Algo::AllOf(Delimiters, [](TCHAR Delimiter) { return Delimiter < 128; }))
+	if (Algo::AllOf(Delimiters, [](CharType Delimiter) { return Delimiter < 128; }))
 	{
 		return ParseTokensNDelim1CharBasicLatin(View, Delimiters, MoveTemp(Visitor), Options);
 	}
 
-	const TCHAR* ViewIt = View.GetData();
-	const TCHAR* const ViewEnd = ViewIt + View.Len();
-	const TCHAR* NextToken = ViewIt;
+	const CharType* ViewIt = View.GetData();
+	const CharType* const ViewEnd = ViewIt + View.Len();
+	const CharType* NextToken = ViewIt;
 
 	if (EnumHasAnyFlags(Options, EParseTokensOptions::IgnoreCase))
 	{
-		TArray<TCHAR, TInlineAllocator<16>> LowerDelimiters;
-		Algo::Transform(Delimiters, LowerDelimiters, FChar::ToLower);
+		TArray<CharType, TInlineAllocator<16>> LowerDelimiters;
+		Algo::Transform(Delimiters, LowerDelimiters, TChar<CharType>::ToLower);
 		for (;;)
 		{
 			if (ViewIt == ViewEnd)
 			{
 				break;
 			}
-			if (!Algo::Find(Delimiters, FChar::ToLower(*ViewIt)))
+			if (!Algo::Find(Delimiters, TChar<CharType>::ToLower(*ViewIt)))
 			{
 				++ViewIt;
 				continue;
 			}
-			ParseTokensVisitToken(Visitor, Options, FStringView(NextToken, UE_PTRDIFF_TO_INT32(ViewIt - NextToken)));
+			ParseTokensVisitToken(Visitor, Options, TStringView<CharType>(NextToken, UE_PTRDIFF_TO_INT32(ViewIt - NextToken)));
 			NextToken = ++ViewIt;
 		}
 	}
@@ -172,19 +176,20 @@ inline static void ParseTokensNDelim1Char(
 				++ViewIt;
 				continue;
 			}
-			ParseTokensVisitToken(Visitor, Options, FStringView(NextToken, UE_PTRDIFF_TO_INT32(ViewIt - NextToken)));
+			ParseTokensVisitToken(Visitor, Options, TStringView<CharType>(NextToken, UE_PTRDIFF_TO_INT32(ViewIt - NextToken)));
 			NextToken = ++ViewIt;
 		}
 	}
 
-	ParseTokensVisitToken(Visitor, Options, FStringView(NextToken, UE_PTRDIFF_TO_INT32(ViewIt - NextToken)));
+	ParseTokensVisitToken(Visitor, Options, TStringView<CharType>(NextToken, UE_PTRDIFF_TO_INT32(ViewIt - NextToken)));
 }
 
 /** Parse tokens with multiple multi-character delimiters. */
+template <typename CharType>
 inline static void ParseTokensNDelimNChar(
-	const FStringView View,
-	const TConstArrayView<FStringView> Delimiters,
-	const TFunctionRef<void (FStringView)> Visitor,
+	const TStringView<CharType> View,
+	const TConstArrayView<TStringView<CharType>> Delimiters,
+	const TFunctionRef<void (TStringView<CharType>)> Visitor,
 	const EParseTokensOptions Options)
 {
 	// This is a naive implementation that takes time proportional to View.Len() * TotalDelimiterLen.
@@ -197,11 +202,11 @@ inline static void ParseTokensNDelimNChar(
 	const ESearchCase::Type SearchCase = EnumHasAnyFlags(Options, EParseTokensOptions::IgnoreCase) ? ESearchCase::IgnoreCase : ESearchCase::CaseSensitive;
 	for (int32 ViewIndex = 0; ViewIndex != ViewLen;)
 	{
-		const FStringView RemainingView(View.GetData() + ViewIndex, ViewLen - ViewIndex);
-		auto MatchDelimiter = [RemainingView, SearchCase](FStringView Delimiter) { return RemainingView.StartsWith(Delimiter, SearchCase); };
-		if (const FStringView* Delimiter = Algo::FindByPredicate(Delimiters, MatchDelimiter))
+		const TStringView<CharType> RemainingView(View.GetData() + ViewIndex, ViewLen - ViewIndex);
+		auto MatchDelimiter = [RemainingView, SearchCase](TStringView<CharType> Delimiter) { return RemainingView.StartsWith(Delimiter, SearchCase); };
+		if (const TStringView<CharType>* Delimiter = Algo::FindByPredicate(Delimiters, MatchDelimiter))
 		{
-			ParseTokensVisitToken(Visitor, Options, FStringView(View.GetData() + NextTokenIndex, ViewIndex - NextTokenIndex));
+			ParseTokensVisitToken(Visitor, Options, TStringView<CharType>(View.GetData() + NextTokenIndex, ViewIndex - NextTokenIndex));
 			ViewIndex += Delimiter->Len();
 			NextTokenIndex = ViewIndex;
 		}
@@ -211,16 +216,17 @@ inline static void ParseTokensNDelimNChar(
 		}
 	}
 
-	ParseTokensVisitToken(Visitor, Options, FStringView(View.GetData() + NextTokenIndex, ViewLen - NextTokenIndex));
+	ParseTokensVisitToken(Visitor, Options, TStringView<CharType>(View.GetData() + NextTokenIndex, ViewLen - NextTokenIndex));
 }
 
-void ParseTokensMultiple(
-	const FStringView View,
-	const TConstArrayView<FStringView> Delimiters,
-	TFunctionRef<void (FStringView)> Visitor,
+template <typename CharType>
+inline static void ParseTokensMultiple(
+	const TStringView<CharType> View,
+	const TConstArrayView<TStringView<CharType>> Delimiters,
+	TFunctionRef<void (TStringView<CharType>)> Visitor,
 	const EParseTokensOptions Options)
 {
-	check(Algo::NoneOf(Delimiters, &FStringView::IsEmpty));
+	check(Algo::NoneOf(Delimiters, &TStringView<CharType>::IsEmpty));
 	switch (Delimiters.Num())
 	{
 	case 0:
@@ -232,15 +238,15 @@ void ParseTokensMultiple(
 		}
 		return ParseTokensNDelimNChar(View, Delimiters, MoveTemp(Visitor), Options);
 	default:
-		if (Algo::AllOf(Delimiters, [](const FStringView& Delimiter) { return Delimiter.Len() == 1; }))
+		if (Algo::AllOf(Delimiters, [](const TStringView<CharType>& Delimiter) { return Delimiter.Len() == 1; }))
 		{
-			TArray<TCHAR, TInlineAllocator<32>> DelimiterChars;
+			TArray<CharType, TInlineAllocator<32>> DelimiterChars;
 			DelimiterChars.Reserve(Delimiters.Num());
-			for (const FStringView& Delimiter : Delimiters)
+			for (const TStringView<CharType>& Delimiter : Delimiters)
 			{
 				DelimiterChars.Add(Delimiter[0]);
 			}
-			return ParseTokensNDelim1Char(View, DelimiterChars, MoveTemp(Visitor), Options);
+			return ParseTokensNDelim1Char<CharType>(View, DelimiterChars, MoveTemp(Visitor), Options);
 		}
 		else
 		{
@@ -249,10 +255,11 @@ void ParseTokensMultiple(
 	}
 }
 
-void ParseTokensMultiple(
-	const FStringView View,
-	const TConstArrayView<TCHAR> Delimiters,
-	TFunctionRef<void (FStringView)> Visitor,
+template <typename CharType>
+inline static void ParseTokensMultiple(
+	const TStringView<CharType> View,
+	const TConstArrayView<CharType> Delimiters,
+	TFunctionRef<void (TStringView<CharType>)> Visitor,
 	const EParseTokensOptions Options)
 {
 	switch (Delimiters.Num())
@@ -266,10 +273,11 @@ void ParseTokensMultiple(
 	}
 }
 
-void ParseTokens(
-	const FStringView View,
-	const FStringView Delimiter,
-	TFunctionRef<void (FStringView)> Visitor,
+template <typename CharType>
+inline static void ParseTokens(
+	const TStringView<CharType> View,
+	const TStringView<CharType> Delimiter,
+	TFunctionRef<void (TStringView<CharType>)> Visitor,
 	const EParseTokensOptions Options)
 {
 	if (Delimiter.Len() == 1)
@@ -279,13 +287,69 @@ void ParseTokens(
 	return ParseTokensNDelimNChar(View, MakeArrayView(&Delimiter, 1), MoveTemp(Visitor), Options);
 }
 
-void ParseTokens(
-	const FStringView View,
-	const TCHAR Delimiter,
-	TFunctionRef<void (FStringView)> Visitor,
-	const EParseTokensOptions Options)
+} // UE::String::Private
+
+namespace UE::String
+{
+
+void ParseTokens(const FAnsiStringView View, const ANSICHAR Delimiter, TFunctionRef<void (FAnsiStringView)> Visitor, const EParseTokensOptions Options)
+{
+	Private::ParseTokens1Delim1Char(View, Delimiter, MoveTemp(Visitor), Options);
+}
+
+void ParseTokens(const FWideStringView View, const WIDECHAR Delimiter, TFunctionRef<void (FWideStringView)> Visitor, const EParseTokensOptions Options)
+{
+	Private::ParseTokens1Delim1Char(View, Delimiter, MoveTemp(Visitor), Options);
+}
+
+void ParseTokens(const FUtf8StringView View, const UTF8CHAR Delimiter, TFunctionRef<void (FUtf8StringView)> Visitor, const EParseTokensOptions Options)
+{
+	Private::ParseTokens1Delim1Char(View, Delimiter, MoveTemp(Visitor), Options);
+}
+
+void ParseTokens(const FAnsiStringView View, const FAnsiStringView Delimiter, TFunctionRef<void (FAnsiStringView)> Visitor, const EParseTokensOptions Options)
+{
+	Private::ParseTokens(View, Delimiter, MoveTemp(Visitor), Options);
+}
+
+void ParseTokens(const FWideStringView View, const FWideStringView Delimiter, TFunctionRef<void (FWideStringView)> Visitor, const EParseTokensOptions Options)
+{
+	Private::ParseTokens(View, Delimiter, MoveTemp(Visitor), Options);
+}
+
+void ParseTokens(const FUtf8StringView View, const FUtf8StringView Delimiter, TFunctionRef<void (FUtf8StringView)> Visitor, const EParseTokensOptions Options)
+{
+	Private::ParseTokens(View, Delimiter, MoveTemp(Visitor), Options);
+}
+
+void ParseTokensMultiple(const FAnsiStringView View, const TConstArrayView<ANSICHAR> Delimiters, TFunctionRef<void (FAnsiStringView)> Visitor, const EParseTokensOptions Options)
+{
+	Private::ParseTokensMultiple(View, Delimiters, MoveTemp(Visitor), Options);
+}
+
+void ParseTokensMultiple(const FWideStringView View, const TConstArrayView<WIDECHAR> Delimiters, TFunctionRef<void (FWideStringView)> Visitor, const EParseTokensOptions Options)
+{
+	Private::ParseTokensMultiple(View, Delimiters, MoveTemp(Visitor), Options);
+}
+
+void ParseTokensMultiple(const FUtf8StringView View, const TConstArrayView<UTF8CHAR> Delimiters, TFunctionRef<void (FUtf8StringView)> Visitor, const EParseTokensOptions Options)
+{
+	Private::ParseTokensMultiple(View, Delimiters, MoveTemp(Visitor), Options);
+}
+
+void ParseTokensMultiple(const FAnsiStringView View, const TConstArrayView<FAnsiStringView> Delimiters, TFunctionRef<void (FAnsiStringView)> Visitor, const EParseTokensOptions Options)
+{
+	Private::ParseTokensMultiple(View, Delimiters, MoveTemp(Visitor), Options);
+}
+
+void ParseTokensMultiple(const FWideStringView View, const TConstArrayView<FWideStringView> Delimiters, TFunctionRef<void (FWideStringView)> Visitor, const EParseTokensOptions Options)
+{
+	Private::ParseTokensMultiple(View, Delimiters, MoveTemp(Visitor), Options);
+}
+
+void ParseTokensMultiple(const FUtf8StringView View, const TConstArrayView<FUtf8StringView> Delimiters, TFunctionRef<void (FUtf8StringView)> Visitor, const EParseTokensOptions Options)
 {
-	return ParseTokens1Delim1Char(View, Delimiter, MoveTemp(Visitor), Options);
+	Private::ParseTokensMultiple(View, Delimiters, MoveTemp(Visitor), Options);
 }
 
 } // UE::String
diff --git a/Engine/Source/Runtime/Core/Public/String/ParseTokens.h b/Engine/Source/Runtime/Core/Public/String/ParseTokens.h
--- a/Engine/Source/Runtime/Core/Public/String/ParseTokens.h
+++ b/Engine/Source/Runtime/Core/Public/String/ParseTokens.h
@@ -36,9 +36,19 @@ ENUM_CLASS_FLAGS(EParseTokensOptions);
  * @param Options     Flags to modify the default behavior.
  */
 CORE_API void ParseTokens(
-	FStringView View,
-	TCHAR Delimiter,
-	TFunctionRef<void (FStringView)> Visitor,
+	FAnsiStringView View,
+	ANSICHAR Delimiter,
+	TFunctionRef<void (FAnsiStringView)> Visitor,
+	EParseTokensOptions Options = EParseTokensOptions::None);
+CORE_API void ParseTokens(
+	FWideStringView View,
+	WIDECHAR Delimiter,
+	TFunctionRef<void (FWideStringView)> Visitor,
+	EParseTokensOptions Options = EParseTokensOptions::None);
+CORE_API void ParseTokens(
+	FUtf8StringView View,
+	UTF8CHAR Delimiter,
+	TFunctionRef<void (FUtf8StringView)> Visitor,
 	EParseTokensOptions Options = EParseTokensOptions::None);
 
 /**
@@ -54,12 +64,30 @@ CORE_API void ParseTokens(
  */
 template <typename OutputType>
 inline void ParseTokens(
-	const FStringView View,
-	const TCHAR Delimiter,
+	const FAnsiStringView View,
+	const ANSICHAR Delimiter,
+	OutputType& Output,
+	const EParseTokensOptions Options = EParseTokensOptions::None)
+{
+	return ParseTokens(View, Delimiter, [&Output](FAnsiStringView Token) { Output.Add(Token); }, Options);
+}
+template <typename OutputType>
+inline void ParseTokens(
+	const FWideStringView View,
+	const WIDECHAR Delimiter,
+	OutputType& Output,
+	const EParseTokensOptions Options = EParseTokensOptions::None)
+{
+	return ParseTokens(View, Delimiter, [&Output](FWideStringView Token) { Output.Add(Token); }, Options);
+}
+template <typename OutputType>
+inline void ParseTokens(
+	const FUtf8StringView View,
+	const UTF8CHAR Delimiter,
 	OutputType& Output,
 	const EParseTokensOptions Options = EParseTokensOptions::None)
 {
-	return ParseTokens(View, Delimiter, [&Output](FStringView Token) { Output.Add(Token); }, Options);
+	return ParseTokens(View, Delimiter, [&Output](FUtf8StringView Token) { Output.Add(Token); }, Options);
 }
 
 /**
@@ -73,9 +101,19 @@ inline void ParseTokens(
  * @param Options     Flags to modify the default behavior.
  */
 CORE_API void ParseTokens(
-	FStringView View,
-	FStringView Delimiter,
-	TFunctionRef<void (FStringView)> Visitor,
+	FAnsiStringView View,
+	FAnsiStringView Delimiter,
+	TFunctionRef<void (FAnsiStringView)> Visitor,
+	EParseTokensOptions Options = EParseTokensOptions::None);
+CORE_API void ParseTokens(
+	FWideStringView View,
+	FWideStringView Delimiter,
+	TFunctionRef<void (FWideStringView)> Visitor,
+	EParseTokensOptions Options = EParseTokensOptions::None);
+CORE_API void ParseTokens(
+	FUtf8StringView View,
+	FUtf8StringView Delimiter,
+	TFunctionRef<void (FUtf8StringView)> Visitor,
 	EParseTokensOptions Options = EParseTokensOptions::None);
 
 /**
@@ -91,12 +129,30 @@ CORE_API void ParseTokens(
  */
 template <typename OutputType>
 inline void ParseTokens(
-	const FStringView View,
-	const FStringView Delimiter,
+	const FAnsiStringView View,
+	const FAnsiStringView Delimiter,
+	OutputType& Output,
+	const EParseTokensOptions Options = EParseTokensOptions::None)
+{
+	return ParseTokens(View, Delimiter, [&Output](FAnsiStringView Token) { Output.Add(Token); }, Options);
+}
+template <typename OutputType>
+inline void ParseTokens(
+	const FWideStringView View,
+	const FWideStringView Delimiter,
+	OutputType& Output,
+	const EParseTokensOptions Options = EParseTokensOptions::None)
+{
+	return ParseTokens(View, Delimiter, [&Output](FWideStringView Token) { Output.Add(Token); }, Options);
+}
+template <typename OutputType>
+inline void ParseTokens(
+	const FUtf8StringView View,
+	const FUtf8StringView Delimiter,
 	OutputType& Output,
 	const EParseTokensOptions Options = EParseTokensOptions::None)
 {
-	return ParseTokens(View, Delimiter, [&Output](FStringView Token) { Output.Add(Token); }, Options);
+	return ParseTokens(View, Delimiter, [&Output](FUtf8StringView Token) { Output.Add(Token); }, Options);
 }
 
 /**
@@ -110,9 +166,19 @@ inline void ParseTokens(
  * @param Options      Flags to modify the default behavior.
  */
 CORE_API void ParseTokensMultiple(
-	FStringView View,
-	TConstArrayView<TCHAR> Delimiters,
-	TFunctionRef<void (FStringView)> Visitor,
+	FAnsiStringView View,
+	TConstArrayView<ANSICHAR> Delimiters,
+	TFunctionRef<void (FAnsiStringView)> Visitor,
+	EParseTokensOptions Options = EParseTokensOptions::None);
+CORE_API void ParseTokensMultiple(
+	FWideStringView View,
+	TConstArrayView<WIDECHAR> Delimiters,
+	TFunctionRef<void (FWideStringView)> Visitor,
+	EParseTokensOptions Options = EParseTokensOptions::None);
+CORE_API void ParseTokensMultiple(
+	FUtf8StringView View,
+	TConstArrayView<UTF8CHAR> Delimiters,
+	TFunctionRef<void (FUtf8StringView)> Visitor,
 	EParseTokensOptions Options = EParseTokensOptions::None);
 
 /**
@@ -128,12 +194,30 @@ CORE_API void ParseTokensMultiple(
  */
 template <typename OutputType>
 inline void ParseTokensMultiple(
-	const FStringView View,
-	const TConstArrayView<TCHAR> Delimiters,
+	const FAnsiStringView View,
+	const TConstArrayView<ANSICHAR> Delimiters,
+	OutputType& Output,
+	const EParseTokensOptions Options = EParseTokensOptions::None)
+{
+	ParseTokensMultiple(View, Delimiters, [&Output](FAnsiStringView Token) { Output.Add(Token); }, Options);
+}
+template <typename OutputType>
+inline void ParseTokensMultiple(
+	const FWideStringView View,
+	const TConstArrayView<WIDECHAR> Delimiters,
+	OutputType& Output,
+	const EParseTokensOptions Options = EParseTokensOptions::None)
+{
+	ParseTokensMultiple(View, Delimiters, [&Output](FWideStringView Token) { Output.Add(Token); }, Options);
+}
+template <typename OutputType>
+inline void ParseTokensMultiple(
+	const FUtf8StringView View,
+	const TConstArrayView<UTF8CHAR> Delimiters,
 	OutputType& Output,
 	const EParseTokensOptions Options = EParseTokensOptions::None)
 {
-	ParseTokensMultiple(View, Delimiters, [&Output](FStringView Token) { Output.Add(Token); }, Options);
+	ParseTokensMultiple(View, Delimiters, [&Output](FUtf8StringView Token) { Output.Add(Token); }, Options);
 }
 
 /**
@@ -149,9 +233,19 @@ inline void ParseTokensMultiple(
  * @param Options      Flags to modify the default behavior.
  */
 CORE_API void ParseTokensMultiple(
-	FStringView View,
-	TConstArrayView<FStringView> Delimiters,
-	TFunctionRef<void (FStringView)> Visitor,
+	FAnsiStringView View,
+	TConstArrayView<FAnsiStringView> Delimiters,
+	TFunctionRef<void (FAnsiStringView)> Visitor,
+	EParseTokensOptions Options = EParseTokensOptions::None);
+CORE_API void ParseTokensMultiple(
+	FWideStringView View,
+	TConstArrayView<FWideStringView> Delimiters,
+	TFunctionRef<void (FWideStringView)> Visitor,
+	EParseTokensOptions Options = EParseTokensOptions::None);
+CORE_API void ParseTokensMultiple(
+	FUtf8StringView View,
+	TConstArrayView<FUtf8StringView> Delimiters,
+	TFunctionRef<void (FUtf8StringView)> Visitor,
 	EParseTokensOptions Options = EParseTokensOptions::None);
 
 /**
@@ -169,12 +263,30 @@ CORE_API void ParseTokensMultiple(
  */
 template <typename OutputType>
 inline void ParseTokensMultiple(
-	const FStringView View,
-	const TConstArrayView<FStringView> Delimiters,
+	const FAnsiStringView View,
+	const TConstArrayView<FAnsiStringView> Delimiters,
+	OutputType& Output,
+	const EParseTokensOptions Options = EParseTokensOptions::None)
+{
+	ParseTokensMultiple(View, Delimiters, [&Output](FAnsiStringView Token) { Output.Add(Token); }, Options);
+}
+template <typename OutputType>
+inline void ParseTokensMultiple(
+	const FWideStringView View,
+	const TConstArrayView<FWideStringView> Delimiters,
+	OutputType& Output,
+	const EParseTokensOptions Options = EParseTokensOptions::None)
+{
+	ParseTokensMultiple(View, Delimiters, [&Output](FWideStringView Token) { Output.Add(Token); }, Options);
+}
+template <typename OutputType>
+inline void ParseTokensMultiple(
+	const FUtf8StringView View,
+	const TConstArrayView<FUtf8StringView> Delimiters,
 	OutputType& Output,
 	const EParseTokensOptions Options = EParseTokensOptions::None)
 {
-	ParseTokensMultiple(View, Delimiters, [&Output](FStringView Token) { Output.Add(Token); }, Options);
+	ParseTokensMultiple(View, Delimiters, [&Output](FUtf8StringView Token) { Output.Add(Token); }, Options);
 }
 
 } // UE::String
