Use _NSGetArgc / _NSGetArgv in FMacPlatformMemory::BaseAllocator

These functions don't allocate memory like "NSArray *Args = [[NSProcessInfo processInfo] arguments];" does.

This fixes initialization crash in UnrealLightmass:

  * Embree has a global ctor embree::registerRegressionTest which adds to a std::vector which does an operator new.
  * That called FMemory::Malloc and FMacPlatformMemory::BaseAllocator for the first time.
  * BaseAllocator does an allocation here:
     - NSArray *Args = [[NSProcessInfo processInfo] arguments];
  * And now we're in a spot where CFAllocatorAllocate calls malloc_zone_malloc to alloc some memory where:
     - __CFAllocatorCustomMalloc --> CFAllocatorAllocate --> malloc_zone_malloc --> __CFAllocatorCustomMalloc --> rinse, lather, repeat...

And we're in recursive death.

diff --git a/Engine/Source/Runtime/Core/Private/Mac/MacPlatformMemory.cpp b/Engine/Source/Runtime/Core/Private/Mac/MacPlatformMemory.cpp
--- a/Engine/Source/Runtime/Core/Private/Mac/MacPlatformMemory.cpp
+++ b/Engine/Source/Runtime/Core/Private/Mac/MacPlatformMemory.cpp
@@ -45,6 +45,25 @@ void* CFNetwork_CFAllocatorOperatorNew_Replacement(unsigned long Size, CFAllocat
 }
 #endif // PLATFORM_MAC_X86
 
+static bool HasArg(const char *Arg)
+{
+	if (_NSGetArgc() && _NSGetArgv())
+	{
+		int Argc = *_NSGetArgc();
+		char** Argv = *_NSGetArgv();
+
+		for (int i = 1; i < Argc; ++i)
+		{
+			if (FCStringAnsi::Stricmp(Argv[i], Arg) == 0)
+			{
+				return true;
+			}
+		}
+	}
+
+	return false;
+}
+
 FMalloc* FMacPlatformMemory::BaseAllocator()
 {
 	if (FORCE_ANSI_ALLOCATOR || IS_PROGRAM)
@@ -65,34 +84,22 @@ FMalloc* FMacPlatformMemory::BaseAllocator()
 	}
 
 	// Force ANSI malloc per user preference.
-	if(getenv("UE4_FORCE_MALLOC_ANSI") != nullptr)
+	if((getenv("UE4_FORCE_MALLOC_ANSI") != nullptr) || HasArg("-ansimalloc"))
 	{
 		AllocatorToUse = EMemoryAllocatorToUse::Ansi;
 	}
-#if WITH_MALLOC_STOMP || (PLATFORM_SUPPORTS_MIMALLOC && MIMALLOC_ALLOCATOR_ALLOWED)
-	else
-	{
-		NSArray *Args = [[NSProcessInfo processInfo] arguments];
-		[Args enumerateObjectsUsingBlock:^(id object, NSUInteger idx, BOOL *stop)
-		{
 #if WITH_MALLOC_STOMP
-			if ([object isEqual:@"-stompmalloc"])
-			{
-				AllocatorToUse = EMemoryAllocatorToUse::Stomp;
-				*stop = YES;
-			}
+	else if (HasArg("-stompmalloc"))
+	{
+		AllocatorToUse = EMemoryAllocatorToUse::Stomp;
+	}
 #endif // WITH_MALLOC_STOMP
-
 #if PLATFORM_SUPPORTS_MIMALLOC && MIMALLOC_ALLOCATOR_ALLOWED
-			if ([object isEqual:@"-mimalloc"])
-			{
-				AllocatorToUse = EMemoryAllocatorToUse::Mimalloc;
-				*stop = YES;
-			}
-#endif // PLATFORM_SUPPORTS_MIMALLOC && MIMALLOC_ALLOCATOR_ALLOWED
-		}];
+	else if (HasArg("-mimalloc"))
+	{
+		AllocatorToUse = EMemoryAllocatorToUse::Mimalloc;
 	}
-#endif // WITH_MALLOC_STOMP || (PLATFORM_SUPPORTS_MIMALLOC && MIMALLOC_ALLOCATOR_ALLOWED)
+#endif // PLATFORM_SUPPORTS_MIMALLOC && MIMALLOC_ALLOCATOR_ALLOWED
 
 	// Force ANSI malloc with TSAN
 #if defined(__has_feature)
