Fix Mac compilation error with _NSGetArgc & _NSGetArgv

diff --git a/Engine/Source/Runtime/Core/Private/Mac/MacPlatformMemory.cpp b/Engine/Source/Runtime/Core/Private/Mac/MacPlatformMemory.cpp
--- a/Engine/Source/Runtime/Core/Private/Mac/MacPlatformMemory.cpp
+++ b/Engine/Source/Runtime/Core/Private/Mac/MacPlatformMemory.cpp
@@ -21,12 +21,10 @@
 
 #include <mach/vm_page_size.h>
 
-#if WITH_MALLOC_STOMP
 extern "C"
 {
 	#include <crt_externs.h> // Needed for _NSGetArgc & _NSGetArgv
 }
-#endif
 
 // Set rather to use BinnedMalloc2 for binned malloc, can be overridden below
 #define USE_MALLOC_BINNED2 (1)
