Temporary rollback of explicit cast for FVector.

diff --git a/Engine/Source/Runtime/Core/Public/Math/Vector.h b/Engine/Source/Runtime/Core/Public/Math/Vector.h
--- a/Engine/Source/Runtime/Core/Public/Math/Vector.h
+++ b/Engine/Source/Runtime/Core/Public/Math/Vector.h
@@ -1144,9 +1144,9 @@ public:
 		return true;
 	}
 
-	// Conversion from other type.
+	// Conversion from other type. TODO: explicit!
 	template<typename FArg, TEMPLATE_REQUIRES(!TIsSame<T, FArg>::Value)>
-	explicit TVector(const TVector<FArg>& From) : TVector<T>((T)From.X, (T)From.Y, (T)From.Z) {}
+	TVector(const TVector<FArg>& From) : TVector<T>((T)From.X, (T)From.Y, (T)From.Z) {}
 };
 
 /**
