Temporarily remove overloaded test that is causing the build to fail.

diff --git a/Engine/Source/Runtime/Core/Private/Tests/Math/UnrealMathTest.cpp b/Engine/Source/Runtime/Core/Private/Tests/Math/UnrealMathTest.cpp
--- a/Engine/Source/Runtime/Core/Private/Tests/Math/UnrealMathTest.cpp
+++ b/Engine/Source/Runtime/Core/Private/Tests/Math/UnrealMathTest.cpp
@@ -228,25 +228,25 @@ FORCENOINLINE bool TestFVector3Equal( const FVector3f& Vec0, const FVector3f& Ve
  *
  * @return true if equal(ish)
  */
-FORCENOINLINE bool TestFVector3Equal(const FVector3d& Vec0, const FVector3d& Vec1, double Tolerance = 0.0l)
-{
-	GScratchDouble[0] = Vec0.X;
-	GScratchDouble[1] = Vec0.Y;
-	GScratchDouble[2] = Vec0.Z;
-	GScratchDouble[3] = 0.0f;
-	GScratchDouble[4] = Vec1.X;
-	GScratchDouble[5] = Vec1.Y;
-	GScratchDouble[6] = Vec1.Z;
-	GScratchDouble[7] = 0.0f;
-	GSumDouble = 0.f;
-
-	for (int32 Component = 0; Component < 3; Component++)
-	{
-		GSumDouble += FMath::Abs<double>(GScratchDouble[Component + 0] - GScratchDouble[Component + 4]);
-	}
-	CheckPassing(GSumDouble <= Tolerance);
-	return GSumDouble <= Tolerance;
-}
+// FORCENOINLINE bool TestFVector3Equal(const FVector3d& Vec0, const FVector3d& Vec1, double Tolerance = 0.0l)
+// {
+// 	GScratchDouble[0] = Vec0.X;
+// 	GScratchDouble[1] = Vec0.Y;
+// 	GScratchDouble[2] = Vec0.Z;
+// 	GScratchDouble[3] = 0.0f;
+// 	GScratchDouble[4] = Vec1.X;
+// 	GScratchDouble[5] = Vec1.Y;
+// 	GScratchDouble[6] = Vec1.Z;
+// 	GScratchDouble[7] = 0.0f;
+// 	GSumDouble = 0.f;
+// 
+// 	for (int32 Component = 0; Component < 3; Component++)
+// 	{
+// 		GSumDouble += FMath::Abs<double>(GScratchDouble[Component + 0] - GScratchDouble[Component + 4]);
+// 	}
+// 	CheckPassing(GSumDouble <= Tolerance);
+// 	return GSumDouble <= Tolerance;
+// }
 
 FORCENOINLINE bool TestQuatsEqual(const FQuat4f& Q0, const FQuat4f& Q1, float Tolerance)
 {
