UE5 - Fix off by one in the binary search of the HazardPointer collection and similar code in other systems, which could led to the case where an element was not found even though it was present in the list. Missing an element in that process at the wrong time causes major corruption in the Scheduler.

diff --git a/Engine/Source/Runtime/Core/Private/Experimental/Containers/HazardPointer.cpp b/Engine/Source/Runtime/Core/Private/Experimental/Containers/HazardPointer.cpp
--- a/Engine/Source/Runtime/Core/Private/Experimental/Containers/HazardPointer.cpp
+++ b/Engine/Source/Runtime/Core/Private/Experimental/Containers/HazardPointer.cpp
@@ -91,7 +91,7 @@ static uint32 BinarySearch(const TArray<void*, TInlineAllocator<64>>& Hazards, v
 	}
 
 	//small size array optimization
-	int32 ArrayEnd = Index + Size;
+	int32 ArrayEnd = FMath::Min(Index + Size + 1, Hazards.Num());
 	while (Index < ArrayEnd)
 	{
 		if (Hazards[Index] == Pointer)
diff --git a/Engine/Source/Runtime/Core/Public/Experimental/Containers/RobinHoodHashTable.h b/Engine/Source/Runtime/Core/Public/Experimental/Containers/RobinHoodHashTable.h
--- a/Engine/Source/Runtime/Core/Public/Experimental/Containers/RobinHoodHashTable.h
+++ b/Engine/Source/Runtime/Core/Public/Experimental/Containers/RobinHoodHashTable.h
@@ -136,7 +136,7 @@ namespace RobinHoodHashTable_Private
 			}
 
 			//small size array optimization
-			int ArrayEnd = Index + Size;
+			int ArrayEnd = FMath::Min(Index + Size + 1, FreeList.Num());
 			while (Index < ArrayEnd)
 			{
 				if (FreeList[Index].Start < NodeIndex)
diff --git a/Engine/Source/Runtime/Core/Public/Math/GenericOctree.h b/Engine/Source/Runtime/Core/Public/Math/GenericOctree.h
--- a/Engine/Source/Runtime/Core/Public/Math/GenericOctree.h
+++ b/Engine/Source/Runtime/Core/Public/Math/GenericOctree.h
@@ -418,7 +418,7 @@ private:
 			}
 
 			//small size array optimization
-			int ArrayEnd = Index + Size;
+			int ArrayEnd = FMath::Min(Index + Size + 1, FreeList.Num());
 			while (Index < ArrayEnd)
 			{
 				if (FreeList[Index].Start < NodeIndex)
