fix FDateTime::GetJulianDay()

FDateTime::GetJulianDay() was using an integer division which result in values snaped to midnight each day.
result is:
FDateTime::FromJulianDay(x.GetJulianDay()) is different from x by up to 1 day.

diff --git a/Engine/Plugins/Runtime/SunPosition/Source/SunPosition/Private/SunPosition.cpp b/Engine/Plugins/Runtime/SunPosition/Source/SunPosition/Private/SunPosition.cpp
--- a/Engine/Plugins/Runtime/SunPosition/Source/SunPosition/Private/SunPosition.cpp
+++ b/Engine/Plugins/Runtime/SunPosition/Source/SunPosition/Private/SunPosition.cpp
@@ -30,7 +30,7 @@ void USunPositionFunctionLibrary::GetSunPosition(float Latitude, float Longitude
 	double LatitudeRad = FMath::DegreesToRadians(Latitude);
 
 	// Get the julian day (number of days since Jan 1st of the year 4713 BC)
-	double JulianDay = CalcTime.GetJulianDay() + (CalcTime.GetTimeOfDay().GetTotalHours() - TimeOffset) / 24.0;
+	double JulianDay = CalcTime.GetJulianDay();
 	double JulianCentury = (JulianDay - 2451545.0) / 36525.0;
 
 	// Get the sun's mean longitude , referred to the mean equinox of julian date
diff --git a/Engine/Source/Runtime/Core/Private/Misc/DateTime.cpp b/Engine/Source/Runtime/Core/Private/Misc/DateTime.cpp
--- a/Engine/Source/Runtime/Core/Private/Misc/DateTime.cpp
+++ b/Engine/Source/Runtime/Core/Private/Misc/DateTime.cpp
@@ -93,7 +93,7 @@ void FDateTime::GetDate(int32& OutYear, int32& OutMonth, int32& OutDay) const
 
 	int32 i, j, k, l, n;
 
-	l = FMath::FloorToInt((float)(GetJulianDay() + 0.5)) + 68569;
+	l = FMath::FloorToInt(GetDate().GetJulianDay() + 0.5) + 68569;
 	n = 4 * l / 146097;
 	l = l - (146097 * n + 3) / 4;
 	i = 4000 * (l + 1) / 1461001;
diff --git a/Engine/Source/Runtime/Core/Private/Tests/Misc/DateTimeTest.cpp b/Engine/Source/Runtime/Core/Private/Tests/Misc/DateTimeTest.cpp
--- a/Engine/Source/Runtime/Core/Private/Tests/Misc/DateTimeTest.cpp
+++ b/Engine/Source/Runtime/Core/Private/Tests/Misc/DateTimeTest.cpp
@@ -35,7 +35,7 @@ bool FDateTimeTest::RunTest(const FString& Parameters)
 	const FDateTime UnixDecimalSequence = FDateTime::FromUnixTimestamp(UnixDecimalSequenceTimestamp);
 
 	const FDateTime TestDateTime(2013, 8, 14, 12, 34, 56, 789);
-	
+
 	TestUnixEquivalent(TEXT("Testing Unix Epoch Ticks"), UnixEpoch, UnixEpochTimestamp);
 	TestUnixEquivalent(TEXT("Testing Unix Billennium Ticks"), UnixBillennium, UnixBillenniumTimestamp);
 	TestUnixEquivalent(TEXT("Testing Unix Ones Ticks"), UnixOnes, UnixOnesTimestamp);
@@ -89,7 +89,7 @@ bool FDateTimeTest::RunTest(const FString& Parameters)
 	FDateTime ParsedDateTime;
 
 	TestFalse(TEXT("Parsing an empty ISO string must fail"), FDateTime::ParseIso8601(TEXT(""), ParsedDateTime));
-	
+
 	FDateTime::ParseIso8601(TEXT("2019-05-22"), ParsedDateTime);
 	TestEqual(TEXT("Testing ISO 8601 date"), ParsedDateTime, FDateTime{2019, 5, 22});
 
@@ -102,6 +102,34 @@ bool FDateTimeTest::RunTest(const FString& Parameters)
 	FDateTime::ParseIso8601(TEXT("2019-05-20T19:41:38-01"), ParsedDateTime);
 	TestEqual(TEXT("Testing ISO 8601 with -hh timezone info"), ParsedDateTime, FDateTime{ 2019, 5, 20, 20, 41, 38 });
 
+	for (double JulianDay : {0., 1000., 1721425.5, 1721425.0, FDateTime::MinValue().GetJulianDay(), FDateTime::MaxValue().GetJulianDay()})
+	{
+		TestEqual(TEXT("convertion from/to JulianDay is stable"), FDateTime::FromJulianDay(JulianDay).GetJulianDay(), JulianDay);
+	}
+
+	FDateTime Date{ 2019, 5, 20, 18, 11, 38 };
+	FDateTime DateMidnight = Date.GetDate();
+	TestEqual(TEXT("GetDate returns a FDateTime at midnight"), DateMidnight.GetHour(), 0);
+
+	FDateTime DateNoon = DateMidnight + FTimespan::FromHours(12);
+	TestEqual(TEXT("GetDate returns a FDateTime at midnight"), DateNoon.GetHour(), 12);
+
+	double JulianDayMidnight = DateMidnight.GetJulianDay();
+	double JulianDayNoon = DateNoon.GetJulianDay();
+	TestEqual(TEXT("At midnight, fractionnal part of the JulianDay value should be 0.5"), JulianDayMidnight - (int)JulianDayMidnight, 0.5);
+	TestEqual(TEXT("A 12h timespan adds half a julianday"), JulianDayNoon - JulianDayMidnight, 0.5);
+
+	double OffsetDayCount = 12345.;
+	FDateTime OffsetDate = Date + FTimespan::FromDays(OffsetDayCount);
+	TestEqual(TEXT("An offset by a given numer of days leads to a similar JulianDay offset"), OffsetDate.GetJulianDay() - Date.GetJulianDay(), OffsetDayCount);
+
+	int32 Year, Month, Day;
+	DateMidnight.GetDate(Year, Month, Day);
+	FDateTime PreviousDay(DateMidnight.GetTicks() - 1);
+	int32 YearPrev, MonthPrev, DayPrev;
+	PreviousDay.GetDate(YearPrev, MonthPrev, DayPrev);
+	TestNotEqual(TEXT("One tick before a date at midnight leads to the previous date"), Day, DayPrev);
+
 	return true;
 }
 
diff --git a/Engine/Source/Runtime/Core/Public/Misc/DateTime.h b/Engine/Source/Runtime/Core/Public/Misc/DateTime.h
--- a/Engine/Source/Runtime/Core/Public/Misc/DateTime.h
+++ b/Engine/Source/Runtime/Core/Public/Misc/DateTime.h
@@ -332,7 +332,7 @@ public:
 	 */
 	double GetJulianDay() const
 	{
-		return (double)(1721425.5 + Ticks / ETimespan::TicksPerDay);
+		return 1721425.5 + double(Ticks / ETimespan::TicksPerDay) + GetTimeOfDay().GetTotalDays();
 	}
 
 	/**
