Chaos Cloth - Move back the post LWC solver to float for performance reason, since cloth is being simulated in local space.

diff --git a/Engine/Plugins/Experimental/ChaosCloth/Source/Chaos/Public/ChaosCloth/ChaosClothingSimulationMesh.h b/Engine/Plugins/Experimental/ChaosCloth/Source/Chaos/Public/ChaosCloth/ChaosClothingSimulationMesh.h
--- a/Engine/Plugins/Experimental/ChaosCloth/Source/Chaos/Public/ChaosCloth/ChaosClothingSimulationMesh.h
+++ b/Engine/Plugins/Experimental/ChaosCloth/Source/Chaos/Public/ChaosCloth/ChaosClothingSimulationMesh.h
@@ -2,8 +2,8 @@
 #pragma once
 
 #include "Chaos/Core.h"
+#include "Chaos/PBDSoftsEvolutionFwd.h"
 #include "Containers/ContainersFwd.h"
-#include "Chaos/Transform.h"
 
 class USkeletalMeshComponent;
 class UClothingAssetCommon;
@@ -45,27 +45,27 @@ namespace Chaos
 		bool WrapDeformLOD(
 			int32 PrevLODIndex,
 			int32 LODIndex,
-			const FVec3* Normals,
-			const FVec3* Positions,
-			FVec3* OutPositions) const;
+			const Softs::FSolverVec3* Normals,
+			const Softs::FSolverVec3* Positions,
+			Softs::FSolverVec3* OutPositions) const;
 
 		bool WrapDeformLOD(
 			int32 PrevLODIndex,
 			int32 LODIndex,
-			const FVec3* Normals,
-			const FVec3* Positions,
-			const FVec3* Velocities,
-			FVec3* OutPositions0,
-			FVec3* OutPositions1,
-			FVec3* OutVelocities) const;
+			const Softs::FSolverVec3* Normals,
+			const Softs::FSolverVec3* Positions,
+			const Softs::FSolverVec3* Velocities,
+			Softs::FSolverVec3* OutPositions0,
+			Softs::FSolverVec3* OutPositions1,
+			Softs::FSolverVec3* OutVelocities) const;
 		// ---- End of the Cloth interface ----
 
 	private:
 		void SkinPhysicsMesh(
 			int32 LODIndex,
 			const FVec3& LocalSpaceLocation,
-			FVec3* OutPositions,
-			FVec3* OutNormals) const;
+			Softs::FSolverVec3* OutPositions,
+			Softs::FSolverVec3* OutNormals) const;
 
 	private:
 		const UClothingAssetCommon* Asset;
diff --git a/Engine/Plugins/Experimental/ChaosCloth/Source/Chaos/Public/ChaosCloth/ChaosClothingSimulationSolver.h b/Engine/Plugins/Experimental/ChaosCloth/Source/Chaos/Public/ChaosCloth/ChaosClothingSimulationSolver.h
--- a/Engine/Plugins/Experimental/ChaosCloth/Source/Chaos/Public/ChaosCloth/ChaosClothingSimulationSolver.h
+++ b/Engine/Plugins/Experimental/ChaosCloth/Source/Chaos/Public/ChaosCloth/ChaosClothingSimulationSolver.h
@@ -2,6 +2,7 @@
 #pragma once
 
 #include "ChaosCloth/ChaosClothConstraints.h"
+#include "Chaos/PBDSoftsEvolutionFwd.h"
 #include "Chaos/ArrayCollectionArray.h"
 #include "Chaos/Transform.h"
 #include "Chaos/ImplicitObject.h"
@@ -13,8 +14,6 @@ class UClothingAssetCommon;
 namespace Chaos
 {
 	class FTriangleMesh;
-	class FPBDEvolution;
-	class FVelocityField;
 
 	class FClothingSimulationCloth;
 	class FClothingSimulationCollider;
@@ -34,11 +33,11 @@ namespace Chaos
 		// Disables all Cloths gravity override mechanism
 		void EnableClothGravityOverride(bool bInIsClothGravityOverrideEnabled) { bIsClothGravityOverrideEnabled = bInIsClothGravityOverrideEnabled; }
 		bool IsClothGravityOverrideEnabled() const { return bIsClothGravityOverrideEnabled; }
-		void SetGravity(const FVec3& InGravity) { Gravity = InGravity; }
-		const FVec3& GetGravity() const { return Gravity; }
+		void SetGravity(const TVec3<FRealSingle>& InGravity) { Gravity = InGravity; }
+		const TVec3<FRealSingle>& GetGravity() const { return Gravity; }
 
-		void SetWindVelocity(const FVec3& InWindVelocity, FRealSingle InLegacyWindAdaption = 0.f);
-		const FVec3& GetWindVelocity() const { return WindVelocity; }
+		void SetWindVelocity(const TVec3<FRealSingle>& InWindVelocity, FRealSingle InLegacyWindAdaption = (FRealSingle)0.);
+		const TVec3<FRealSingle>& GetWindVelocity() const { return WindVelocity; }
 
 		UE_DEPRECATED(4.27, "Use SetWindProperties instead.")
 		void SetWindFluidDensity(FRealSingle /*WindFluidDensity*/) {}
@@ -61,7 +60,7 @@ namespace Chaos
 		TConstArrayView<const FClothingSimulationCloth*> GetCloths() const { return Cloths; }
 
 		// Update solver properties before simulation.
-		void Update(FReal InDeltaTime);
+		void Update(FSolverReal InDeltaTime);
 
 		FBoxSphereBounds CalculateBounds() const;
 		// ---- End of the object management functions ----
@@ -70,19 +69,21 @@ namespace Chaos
 		int32 AddParticles(int32 NumParticles, uint32 GroupId);
 		void EnableParticles(int32 Offset, bool bEnable);
 
+		void ResetStartPose(int32 Offset, int32 NumParticles);
+
 		// Get the currrent solver time
-		FReal GetTime() const { return Time; }
-		void SetParticleMassUniform(int32 Offset, FReal UniformMass, FReal MinPerParticleMass, const FTriangleMesh& Mesh, const TFunctionRef<bool(int32)>& KinematicPredicate);
-		void SetParticleMassFromTotalMass(int32 Offset, FReal TotalMass, FReal MinPerParticleMass, const FTriangleMesh& Mesh, const TFunctionRef<bool(int32)>& KinematicPredicate);
-		void SetParticleMassFromDensity(int32 Offset, FReal Density, FReal MinPerParticleMass, const FTriangleMesh& Mesh, const TFunctionRef<bool(int32)>& KinematicPredicate);
+		FSolverReal GetTime() const { return Time; }
+		void SetParticleMassUniform(int32 Offset, FRealSingle UniformMass, FRealSingle MinPerParticleMass, const FTriangleMesh& Mesh, const TFunctionRef<bool(int32)>& KinematicPredicate);
+		void SetParticleMassFromTotalMass(int32 Offset, FRealSingle TotalMass, FRealSingle MinPerParticleMass, const FTriangleMesh& Mesh, const TFunctionRef<bool(int32)>& KinematicPredicate);
+		void SetParticleMassFromDensity(int32 Offset, FRealSingle Density, FRealSingle MinPerParticleMass, const FTriangleMesh& Mesh, const TFunctionRef<bool(int32)>& KinematicPredicate);
 
 		// Set the amount of velocity allowed to filter from the given change in reference space transform, including local simulation space.
 		void SetReferenceVelocityScale(uint32 GroupId,
 			const FRigidTransform3& OldReferenceSpaceTransform,
 			const FRigidTransform3& ReferenceSpaceTransform,
-			const FVec3& LinearVelocityScale,
-			FReal AngularVelocityScale,
-			FReal FictitiousAngularScale);
+			const TVec3<FRealSingle>& LinearVelocityScale,
+			FRealSingle AngularVelocityScale,
+			FRealSingle FictitiousAngularScale);
 
 		// Set general cloth simulation properties.
 		void SetProperties(uint32 GroupId, FRealSingle DampingCoefficient, FRealSingle CollisionThickness, FRealSingle FrictionCoefficient);
@@ -91,52 +92,38 @@ namespace Chaos
 		void SetUseCCD(uint32 GroupId, bool bUseCCD);
 
 		// Set per group gravity, used to override solver's gravity. Must be called during cloth update.
-		void SetGravity(uint32 GroupId, const FVec3& Gravity);
+		void SetGravity(uint32 GroupId, const TVec3<FRealSingle>& Gravity);
 
 		// Set per group wind velocity, used to override solver's wind velocity. Must be called during cloth update.
-		void SetWindVelocity(uint32 GroupId, const FVec3& InWindVelocity);
-
-		UE_DEPRECATED(4.27, "Use SetWindGeometry and SetWindPropertiesinstead.")
-		void SetWindVelocityField(uint32 GroupId, FRealSingle DragCoefficient, FRealSingle LiftCoefficient, const FTriangleMesh* TriangleMesh = nullptr)
-		{
-			if (TriangleMesh)
-			{
-				SetWindGeometry(GroupId, *TriangleMesh, TConstArrayView<FRealSingle>(), TConstArrayView<FRealSingle>());
-				SetWindProperties(GroupId, FVec2((FReal)DragCoefficient), FVec2((FReal)LiftCoefficient));
-			}
-			else
-			{
-				SetWindProperties(GroupId, FVec2((FReal)0.), FVec2((FReal)0.));
-			}
-		}
+		void SetWindVelocity(uint32 GroupId, const TVec3<FRealSingle>& InWindVelocity);
 
 		// Set the geometry affected by the wind.
 		void SetWindGeometry(uint32 GroupId, const FTriangleMesh& TriangleMesh, const TConstArrayView<FRealSingle>& DragMultipliers, const TConstArrayView<FRealSingle>& LiftMultipliers);
 
 		// Set the wind properties.
-		void SetWindProperties(uint32 GroupId, const FVec2& Drag, const FVec2& Lift, FReal AirDensity = (FReal)1.225e-6);
+		void SetWindProperties(uint32 GroupId, const TVec2<FRealSingle>& Drag, const TVec2<FRealSingle>& Lift, FRealSingle AirDensity = 1.225e-6f);
 
 		// Return the wind velocity field associated with a given group id.
-		const FVelocityField&  GetWindVelocityField(uint32 GroupId);
+		const Softs::FVelocityField& GetWindVelocityField(uint32 GroupId);
 
 		// Add external forces to the particles
 		void AddExternalForces(uint32 GroupId, bool bUseLegacyWind);
 
-		const FVec3* GetOldAnimationPositions(int32 Offset) const { return OldAnimationPositions.GetData() + Offset; }
-		FVec3* GetOldAnimationPositions(int32 Offset) { return OldAnimationPositions.GetData() + Offset; }
-		const FVec3* GetAnimationPositions(int32 Offset) const { return AnimationPositions.GetData() + Offset; }
-		FVec3* GetAnimationPositions(int32 Offset) { return AnimationPositions.GetData() + Offset; }
-		const FVec3* GetAnimationNormals(int32 Offset) const { return AnimationNormals.GetData() + Offset; }
-		FVec3* GetAnimationNormals(int32 Offset) { return AnimationNormals.GetData() + Offset; }
-		const FVec3* GetNormals(int32 Offset) const { return Normals.GetData() + Offset; }
-		FVec3* GetNormals(int32 Offset) { return Normals.GetData() + Offset; }
-		const FVec3* GetParticlePs(int32 Offset) const;
-		FVec3* GetParticlePs(int32 Offset);
-		const FVec3* GetParticleXs(int32 Offset) const;
-		FVec3* GetParticleXs(int32 Offset);
-		const FVec3* GetParticleVs(int32 Offset) const;
-		FVec3* GetParticleVs(int32 Offset);
-		const FReal* GetParticleInvMasses(int32 Offset) const;
+		const Softs::FSolverVec3* GetOldAnimationPositions(int32 Offset) const { return OldAnimationPositions.GetData() + Offset; }
+		Softs::FSolverVec3* GetOldAnimationPositions(int32 Offset) { return OldAnimationPositions.GetData() + Offset; }
+		const Softs::FSolverVec3* GetAnimationPositions(int32 Offset) const { return AnimationPositions.GetData() + Offset; }
+		Softs::FSolverVec3* GetAnimationPositions(int32 Offset) { return AnimationPositions.GetData() + Offset; }
+		const Softs::FSolverVec3* GetAnimationNormals(int32 Offset) const { return AnimationNormals.GetData() + Offset; }
+		Softs::FSolverVec3* GetAnimationNormals(int32 Offset) { return AnimationNormals.GetData() + Offset; }
+		const Softs::FSolverVec3* GetNormals(int32 Offset) const { return Normals.GetData() + Offset; }
+		Softs::FSolverVec3* GetNormals(int32 Offset) { return Normals.GetData() + Offset; }
+		const Softs::FSolverVec3* GetParticlePs(int32 Offset) const;
+		Softs::FSolverVec3* GetParticlePs(int32 Offset);
+		const Softs::FSolverVec3* GetParticleXs(int32 Offset) const;
+		Softs::FSolverVec3* GetParticleXs(int32 Offset);
+		const Softs::FSolverVec3* GetParticleVs(int32 Offset) const;
+		Softs::FSolverVec3* GetParticleVs(int32 Offset);
+		const Softs::FSolverReal* GetParticleInvMasses(int32 Offset) const;
 		const FClothConstraints& GetClothConstraints(int32 Offset) const { return *ClothsConstraints.FindChecked(Offset); }
 		FClothConstraints& GetClothConstraints(int32 Offset) { return *ClothsConstraints.FindChecked(Offset); }
 		// ---- End of the Cloth interface ----
@@ -145,23 +132,25 @@ namespace Chaos
 		int32 AddCollisionParticles(int32 NumCollisionParticles, uint32 GroupId, int32 RecycledOffset = 0);
 		void EnableCollisionParticles(int32 Offset, bool bEnable);
 
+		void ResetCollisionStartPose(int32 Offset, int32 NumCollisionParticles);
+
 		const int32* GetCollisionBoneIndices(int32 Offset) const { return CollisionBoneIndices.GetData() + Offset; }
 		int32* GetCollisionBoneIndices(int32 Offset) { return CollisionBoneIndices.GetData() + Offset; }
-		const FRigidTransform3* GetCollisionBaseTransforms(int32 Offset) const { return CollisionBaseTransforms.GetData() + Offset; }
-		FRigidTransform3* GetCollisionBaseTransforms(int32 Offset) { return CollisionBaseTransforms.GetData() + Offset; }
-		const FRigidTransform3* GetOldCollisionTransforms(int32 Offset) const { return OldCollisionTransforms.GetData() + Offset; }
-		FRigidTransform3* GetOldCollisionTransforms(int32 Offset) { return OldCollisionTransforms.GetData() + Offset; }
-		const FRigidTransform3* GetCollisionTransforms(int32 Offset) const { return CollisionTransforms.GetData() + Offset; }
-		FRigidTransform3* GetCollisionTransforms(int32 Offset) { return CollisionTransforms.GetData() + Offset; }
-		const FVec3* GetCollisionParticleXs(int32 Offset) const;
-		FVec3* GetCollisionParticleXs(int32 Offset);
-		const FRotation3* GetCollisionParticleRs(int32 Offset) const;
-		FRotation3* GetCollisionParticleRs(int32 Offset);
+		const Softs::FSolverRigidTransform3* GetCollisionBaseTransforms(int32 Offset) const { return CollisionBaseTransforms.GetData() + Offset; }
+		Softs::FSolverRigidTransform3* GetCollisionBaseTransforms(int32 Offset) { return CollisionBaseTransforms.GetData() + Offset; }
+		const Softs::FSolverRigidTransform3* GetOldCollisionTransforms(int32 Offset) const { return OldCollisionTransforms.GetData() + Offset; }
+		Softs::FSolverRigidTransform3* GetOldCollisionTransforms(int32 Offset) { return OldCollisionTransforms.GetData() + Offset; }
+		const Softs::FSolverRigidTransform3* GetCollisionTransforms(int32 Offset) const { return CollisionTransforms.GetData() + Offset; }
+		Softs::FSolverRigidTransform3* GetCollisionTransforms(int32 Offset) { return CollisionTransforms.GetData() + Offset; }
+		const Softs::FSolverVec3* GetCollisionParticleXs(int32 Offset) const;
+		Softs::FSolverVec3* GetCollisionParticleXs(int32 Offset);
+		const Softs::FSolverRotation3* GetCollisionParticleRs(int32 Offset) const;
+		Softs::FSolverRotation3* GetCollisionParticleRs(int32 Offset);
 		void SetCollisionGeometry(int32 Offset, int32 Index, TUniquePtr<FImplicitObject>&& Geometry);
 		const TUniquePtr<FImplicitObject>* GetCollisionGeometries(int32 Offset) const;
 		const bool* GetCollisionStatus(int32 Offset) const;
-		const TArray<FVec3>& GetCollisionContacts() const;
-		const TArray<FVec3>& GetCollisionNormals() const;
+		const TArray<Softs::FSolverVec3>& GetCollisionContacts() const;
+		const TArray<Softs::FSolverVec3>& GetCollisionNormals() const;
 		// ---- End of the Collider interface ----
 
 		// ---- Field interface ----
@@ -173,45 +162,45 @@ namespace Chaos
 		void ResetParticles();
 		void ResetCollisionParticles(int32 InCollisionParticlesOffset = 0);
 		void ApplyPreSimulationTransforms();
-		FReal SetParticleMassPerArea(int32 Offset, int32 Size, const FTriangleMesh& Mesh);
-		void ParticleMassUpdateDensity(const FTriangleMesh& Mesh, FReal Density);
-		void ParticleMassClampAndEnslave(int32 Offset, int32 Size, FReal MinPerParticleMass, const TFunctionRef<bool(int32)>& KinematicPredicate);
+		Softs::FSolverReal SetParticleMassPerArea(int32 Offset, int32 Size, const FTriangleMesh& Mesh);
+		void ParticleMassUpdateDensity(const FTriangleMesh& Mesh, Softs::FSolverReal Density);
+		void ParticleMassClampAndEnslave(int32 Offset, int32 Size, Softs::FSolverReal MinPerParticleMass, const TFunctionRef<bool(int32)>& KinematicPredicate);
 
 		// Update the solver field forces/velocities at the particles location
 		void UpdateSolverField();
 
 	private:
-		TUniquePtr<FPBDEvolution> Evolution;
+		TUniquePtr<Softs::FPBDEvolution> Evolution;
 
 		// Object arrays
 		TArray<FClothingSimulationCloth*> Cloths;
 
 		// Simulation group attributes
-		TArrayCollectionArray<FRigidTransform3> PreSimulationTransforms;  // Allow a different frame of reference for each cloth groups
-		TArrayCollectionArray<FVec3> FictitiousAngularDisplacement;  // Relative angular displacement of the reference bone that depends on the fictitious angular scale factor
+		TArrayCollectionArray<Softs::FSolverRigidTransform3> PreSimulationTransforms;  // Allow a different frame of reference for each cloth groups
+		TArrayCollectionArray<Softs::FSolverVec3> FictitiousAngularDisplacements;  // Relative angular displacement of the reference bone that depends on the fictitious angular scale factor
 
 		// Particle attributes
-		TArrayCollectionArray<FVec3> Normals;
-		TArrayCollectionArray<FVec3> OldAnimationPositions;
-		TArrayCollectionArray<FVec3> AnimationPositions;
-		TArrayCollectionArray<FVec3> AnimationNormals;
+		TArrayCollectionArray<Softs::FSolverVec3> Normals;
+		TArrayCollectionArray<Softs::FSolverVec3> OldAnimationPositions;
+		TArrayCollectionArray<Softs::FSolverVec3> AnimationPositions;
+		TArrayCollectionArray<Softs::FSolverVec3> AnimationNormals;
 
 		// Collision particle attributes
 		TArrayCollectionArray<int32> CollisionBoneIndices;
-		TArrayCollectionArray<FRigidTransform3> CollisionBaseTransforms;
-		TArrayCollectionArray<FRigidTransform3> OldCollisionTransforms;
-		TArrayCollectionArray<FRigidTransform3> CollisionTransforms;
+		TArrayCollectionArray<Softs::FSolverRigidTransform3> CollisionBaseTransforms;
+		TArrayCollectionArray<Softs::FSolverRigidTransform3> OldCollisionTransforms;
+		TArrayCollectionArray<Softs::FSolverRigidTransform3> CollisionTransforms;
 
 		// Cloth constraints
 		TMap<int32, TUniquePtr<FClothConstraints>> ClothsConstraints;
 
 		// Local space simulation
-		FVec3 OldLocalSpaceLocation;
-		FVec3 LocalSpaceLocation;  // This is used to translate between world space and simulation space. Add this to simulation space coordinates to get world space coordinates
+		FVec3 OldLocalSpaceLocation;  // This is used to translate between world space and simulation space,
+		FVec3 LocalSpaceLocation;     // add this to simulation space coordinates to get world space coordinates, must keep FReal as underlying type for LWC
 
 		// Time stepping
-		FReal Time;
-		FReal DeltaTime;
+		FSolverReal Time;
+		FSolverReal DeltaTime;
 		int32 NumIterations;
 		int32 NumSubsteps;
 
@@ -220,9 +209,9 @@ namespace Chaos
 		int32 CollisionParticlesSize;  // Number of solver only colliders
 
 		// Solver parameters
-		FVec3 Gravity;
-		FVec3 WindVelocity;
-		FReal LegacyWindAdaption;
+		TVec3<FRealSingle> Gravity;
+		TVec3<FRealSingle> WindVelocity;
+		FRealSingle LegacyWindAdaption;
 		bool bIsClothGravityOverrideEnabled;
 
 		// Field system unique to the cloth solver
diff --git a/Engine/Plugins/Experimental/ChaosClothEditor/Source/Chaos/Private/ChaosClothEditor/ChaosSimulationEditorExtender.cpp b/Engine/Plugins/Experimental/ChaosClothEditor/Source/Chaos/Private/ChaosClothEditor/ChaosSimulationEditorExtender.cpp
--- a/Engine/Plugins/Experimental/ChaosClothEditor/Source/Chaos/Private/ChaosClothEditor/ChaosSimulationEditorExtender.cpp
+++ b/Engine/Plugins/Experimental/ChaosClothEditor/Source/Chaos/Private/ChaosClothEditor/ChaosSimulationEditorExtender.cpp
@@ -16,54 +16,52 @@
 
 namespace Chaos
 {
-	struct FVisualizationOption
-	{
-		// Actual option entries
-		static const FVisualizationOption OptionData[];
-		static const uint32 Count;
-
-		// Chaos debug draw function
-		typedef void (Chaos::FClothingSimulation::*FDebugDrawFunction)(FPrimitiveDrawInterface*) const;
-		typedef void (Chaos::FClothingSimulation::*FDebugDrawTextsFunction)(FCanvas*, const FSceneView*) const;
-		FDebugDrawFunction DebugDrawFunction;
-		FDebugDrawTextsFunction DebugDrawTextsFunction;
-
-		FText DisplayName;         // Text for menu entries.
-		FText ToolTip;             // Text for menu tooltips.
-		bool bDisablesSimulation;  // Whether or not this option requires the simulation to be disabled.
-		bool bHidesClothSections;  // Hides the cloth section to avoid zfighting with the debug geometry.
-
-		// Console override
-		IConsoleVariable* const ConsoleVariable;
-
-		FVisualizationOption(FDebugDrawFunction InDebugDrawFunction, const TCHAR* ConsoleName, const FText& InDisplayName, const FText& InToolTip, bool bInDisablesSimulation = false, bool bInHidesClothSections = false)
-			: DebugDrawFunction(InDebugDrawFunction)
-			, DisplayName(InDisplayName)
-			, ToolTip(InToolTip)
-			, bDisablesSimulation(bInDisablesSimulation)
-			, bHidesClothSections(bInHidesClothSections)
-			, ConsoleVariable(IConsoleManager::Get().RegisterConsoleVariable(ConsoleName, false, *InToolTip.ToString(), 0))
-		{}
-
-		FVisualizationOption(FDebugDrawTextsFunction InDebugDrawTextsFunction, const TCHAR* ConsoleName, const FText& InDisplayName, const FText& InToolTip, bool bInDisablesSimulation = false, bool bInHidesClothSections = false)
-			: DebugDrawTextsFunction(InDebugDrawTextsFunction)
-			, DisplayName(InDisplayName)
-			, ToolTip(InToolTip)
-			, bDisablesSimulation(bInDisablesSimulation)
-			, bHidesClothSections(bInHidesClothSections)
-			, ConsoleVariable(IConsoleManager::Get().RegisterConsoleVariable(ConsoleName, false, *InToolTip.ToString(), 0))
-		{}
-
-		~FVisualizationOption()
-		{
-			IConsoleManager::Get().UnregisterConsoleObject(ConsoleVariable);
-		}
 
-		bool IsConsoleVariableEnabled() const { return ConsoleVariable->AsVariableBool()->GetValueOnGameThread(); }
-	};
-}
+struct FVisualizationOption
+{
+	// Actual option entries
+	static const FVisualizationOption OptionData[];
+	static const uint32 Count;
+
+	// Chaos debug draw function
+	typedef void (Chaos::FClothingSimulation::*FDebugDrawFunction)(FPrimitiveDrawInterface*) const;
+	typedef void (Chaos::FClothingSimulation::*FDebugDrawTextsFunction)(FCanvas*, const FSceneView*) const;
+	FDebugDrawFunction DebugDrawFunction;
+	FDebugDrawTextsFunction DebugDrawTextsFunction;
+
+	FText DisplayName;         // Text for menu entries.
+	FText ToolTip;             // Text for menu tooltips.
+	bool bDisablesSimulation;  // Whether or not this option requires the simulation to be disabled.
+	bool bHidesClothSections;  // Hides the cloth section to avoid zfighting with the debug geometry.
+
+	// Console override
+	IConsoleVariable* const ConsoleVariable;
+
+	FVisualizationOption(FDebugDrawFunction InDebugDrawFunction, const TCHAR* ConsoleName, const FText& InDisplayName, const FText& InToolTip, bool bInDisablesSimulation = false, bool bInHidesClothSections = false)
+		: DebugDrawFunction(InDebugDrawFunction)
+		, DisplayName(InDisplayName)
+		, ToolTip(InToolTip)
+		, bDisablesSimulation(bInDisablesSimulation)
+		, bHidesClothSections(bInHidesClothSections)
+		, ConsoleVariable(IConsoleManager::Get().RegisterConsoleVariable(ConsoleName, false, *InToolTip.ToString(), 0))
+	{}
+
+	FVisualizationOption(FDebugDrawTextsFunction InDebugDrawTextsFunction, const TCHAR* ConsoleName, const FText& InDisplayName, const FText& InToolTip, bool bInDisablesSimulation = false, bool bInHidesClothSections = false)
+		: DebugDrawTextsFunction(InDebugDrawTextsFunction)
+		, DisplayName(InDisplayName)
+		, ToolTip(InToolTip)
+		, bDisablesSimulation(bInDisablesSimulation)
+		, bHidesClothSections(bInHidesClothSections)
+		, ConsoleVariable(IConsoleManager::Get().RegisterConsoleVariable(ConsoleName, false, *InToolTip.ToString(), 0))
+	{}
+
+	~FVisualizationOption()
+	{
+		IConsoleManager::Get().UnregisterConsoleObject(ConsoleVariable);
+	}
 
-using namespace Chaos;
+	bool IsConsoleVariableEnabled() const { return ConsoleVariable->AsVariableBool()->GetValueOnGameThread(); }
+};
 
 const FVisualizationOption FVisualizationOption::OptionData[] = 
 {
@@ -211,6 +209,8 @@ void FSimulationEditorExtender::ShowClothSections(USkeletalMeshComponent* MeshCo
 	}
 }
 
+}  // End namespace Chaos
+
 #undef LOCTEXT_NAMESPACE
 
 #endif  // #if WITH_CHAOS
\ No newline at end of file
diff --git a/Engine/Source/Runtime/Core/Public/Math/Matrix.isph b/Engine/Source/Runtime/Core/Public/Math/Matrix.isph
--- a/Engine/Source/Runtime/Core/Public/Math/Matrix.isph
+++ b/Engine/Source/Runtime/Core/Public/Math/Matrix.isph
@@ -38,11 +38,23 @@ struct FMatrix3x4
 #endif
 };
 
-static const uniform struct FMatrix MatrixIdentity = {{ 1, 0, 0, 0, 0, 1, 0, 0, 0, 0, 1, 0, 0, 0, 0, 1}};
+static const uniform struct FMatrix44f FloatMatrixIdentity = {{ 1, 0, 0, 0, 0, 1, 0, 0, 0, 0, 1, 0, 0, 0, 0, 1}};
+static const uniform struct FMatrix44d DoubleMatrixIdentity = {{ 1, 0, 0, 0, 0, 1, 0, 0, 0, 0, 1, 0, 0, 0, 0, 1}};
 
-inline uniform FMatrix SetMatrix(const uniform FVector4 &R0, const uniform FVector4 &R1, const uniform FVector4 &R2, const uniform FVector4 &R3)
+#define MatrixIdentity DoubleMatrixIdentity
+
+inline uniform FMatrix44d SetMatrix(const uniform FVector4d &R0, const uniform FVector4d &R1, const uniform FVector4d &R2, const uniform FVector4d &R3)
 {
-	const uniform FMatrix Result = {{ R0.V[0], R0.V[1], R0.V[2], R0.V[3],
+	const uniform FMatrix44d Result = {{ R0.V[0], R0.V[1], R0.V[2], R0.V[3],
+									R1.V[0], R1.V[1], R1.V[2], R1.V[3],
+									R2.V[0], R2.V[1], R2.V[2], R2.V[3],
+									R3.V[0], R3.V[1], R3.V[2], R3.V[3] }};
+	return Result;
+}
+
+inline uniform FMatrix44f SetMatrix(const uniform FVector4f &R0, const uniform FVector4f &R1, const uniform FVector4f &R2, const uniform FVector4f &R3)
+{
+	const uniform FMatrix44f Result = {{ R0.V[0], R0.V[1], R0.V[2], R0.V[3],
 									R1.V[0], R1.V[1], R1.V[2], R1.V[3],
 									R2.V[0], R2.V[1], R2.V[2], R2.V[3],
 									R3.V[0], R3.V[1], R3.V[2], R3.V[3] }};
@@ -58,9 +70,18 @@ inline FMatrix SetMatrix(const FVector4 &R0, const FVector4 &R1, const FVector4
 	return Result;
 }
 
-inline FMatrix SetMatrix(const FVector4 &R0, const FVector4 &R1, const FVector4 &R2, const uniform FVector4 &R3)
+inline FMatrix44d SetMatrix(const FVector4d &R0, const FVector4d &R1, const FVector4d &R2, const uniform FVector4d &R3)
 {
-	const FMatrix Result = {{ R0.V[0], R0.V[1], R0.V[2], R0.V[3],
+	const FMatrix44d Result = {{ R0.V[0], R0.V[1], R0.V[2], R0.V[3],
+							R1.V[0], R1.V[1], R1.V[2], R1.V[3],
+							R2.V[0], R2.V[1], R2.V[2], R2.V[3],
+							R3.V[0], R3.V[1], R3.V[2], R3.V[3] }};
+	return Result;
+}
+
+inline FMatrix44f SetMatrix(const FVector4f &R0, const FVector4f &R1, const FVector4f &R2, const uniform FVector4f &R3)
+{
+	const FMatrix44f Result = {{ R0.V[0], R0.V[1], R0.V[2], R0.V[3],
 							R1.V[0], R1.V[1], R1.V[2], R1.V[3],
 							R2.V[0], R2.V[1], R2.V[2], R2.V[3],
 							R3.V[0], R3.V[1], R3.V[2], R3.V[3] }};
@@ -155,7 +176,14 @@ inline uniform FVector MatrixGetOrigin(const uniform FMatrix &M)
 	return SetVector(M.M[12], M.M[13], M.M[14]);
 }
 
-inline void MatrixGetScaledAxes(const uniform FMatrix &M, uniform FVector &X, uniform FVector &Y, uniform FVector &Z)
+inline void MatrixGetScaledAxes(const uniform FMatrix44d &M, uniform FVector3d &X, uniform FVector3d &Y, uniform FVector3d &Z)
+{
+	X = SetVector(M.M[0], M.M[1], M.M[2]);
+	Y = SetVector(M.M[4], M.M[5], M.M[6]);
+	Z = SetVector(M.M[8], M.M[9], M.M[10]);
+}
+
+inline void MatrixGetScaledAxes(const uniform FMatrix44f &M, uniform FVector3f &X, uniform FVector3f &Y, uniform FVector3f &Z)
 {
 	X = SetVector(M.M[0], M.M[1], M.M[2]);
 	Y = SetVector(M.M[4], M.M[5], M.M[6]);
@@ -196,9 +224,22 @@ inline uniform bool IsAnyMatrixScaledAxesNearlyZero(const uniform FMatrix &M, co
 	return false;
 }
 
-inline uniform bool IsAllMatrixScaledAxesNearlyZero(const uniform FMatrix &M, const uniform FReal Tolerance)
+inline uniform bool IsAllMatrixScaledAxesNearlyZero(const uniform FMatrix44d &M, const uniform double Tolerance)
 {
-	uniform FVector X, Y, Z;
+	uniform FVector3d X, Y, Z;
+	MatrixGetScaledAxes(M, X, Y, Z);
+
+	if(VectorIsAllNearlyZero(X, Tolerance) && VectorIsAllNearlyZero(Y, Tolerance) && VectorIsAllNearlyZero(Z, Tolerance))
+	{
+		return true;
+	}
+
+	return false;
+}
+
+inline uniform bool IsAllMatrixScaledAxesNearlyZero(const uniform FMatrix44f &M, const uniform float Tolerance)
+{
+	uniform FVector3f X, Y, Z;
 	MatrixGetScaledAxes(M, X, Y, Z);
 
 	if(VectorIsAllNearlyZero(X, Tolerance) && VectorIsAllNearlyZero(Y, Tolerance) && VectorIsAllNearlyZero(Z, Tolerance))
@@ -230,14 +271,27 @@ inline uniform FMatrix MatrixTranspose(const uniform FMatrix& M)
 // we use __m128 to represent 2x2 matrix as A = | A0  A1 |
 //                                              | A2  A3 |
 // 2x2 row major Matrix multiply A*B
-static inline uniform FVector4 Mat2Mul(const uniform FVector4& vec1, const uniform FVector4& vec2)
+static inline uniform FVector4d Mat2Mul(const uniform FVector4d& vec1, const uniform FVector4d& vec2)
+{
+	return
+		VectorAdd(VectorMultiply( vec1, VectorSwizzle(vec2, 0,3,0,3)),
+			VectorMultiply(VectorSwizzle(vec1, 1,0,3,2), VectorSwizzle(vec2, 2,1,2,1)));
+}
+static inline uniform FVector4f Mat2Mul(const uniform FVector4f& vec1, const uniform FVector4f& vec2)
 {
 	return
 		VectorAdd(VectorMultiply( vec1, VectorSwizzle(vec2, 0,3,0,3)),
 			VectorMultiply(VectorSwizzle(vec1, 1,0,3,2), VectorSwizzle(vec2, 2,1,2,1)));
 }
 // 2x2 row major Matrix adjugate multiply (A#)*B
-static inline uniform FVector4 Mat2AdjMul(const uniform FVector4& vec1, const uniform FVector4& vec2)
+static inline uniform FVector4d Mat2AdjMul(const uniform FVector4d& vec1, const uniform FVector4d& vec2)
+{
+	return
+		VectorSubtract(VectorMultiply(VectorSwizzle(vec1, 3,3,0,0), vec2),
+			VectorMultiply(VectorSwizzle(vec1, 1,1,2,2), VectorSwizzle(vec2, 2,3,0,1)));
+
+}
+static inline uniform FVector4f Mat2AdjMul(const uniform FVector4f& vec1, const uniform FVector4f& vec2)
 {
 	return
 		VectorSubtract(VectorMultiply(VectorSwizzle(vec1, 3,3,0,0), vec2),
@@ -245,39 +299,79 @@ static inline uniform FVector4 Mat2AdjMul(const uniform FVector4& vec1, const un
 
 }
 // 2x2 row major Matrix multiply adjugate A*(B#)
-static inline uniform FVector4 Mat2MulAdj(const uniform FVector4& vec1, const uniform FVector4& vec2)
+static inline uniform FVector4d Mat2MulAdj(const uniform FVector4d& vec1, const uniform FVector4d& vec2)
+{
+	return
+		VectorSubtract(VectorMultiply( vec1, VectorSwizzle(vec2, 3,0,3,0)),
+			VectorMultiply(VectorSwizzle(vec1, 1,0,3,2), VectorSwizzle(vec2, 2,1,2,1)));
+}
+static inline uniform FVector4f Mat2MulAdj(const uniform FVector4f& vec1, const uniform FVector4f& vec2)
 {
 	return
 		VectorSubtract(VectorMultiply( vec1, VectorSwizzle(vec2, 3,0,3,0)),
 			VectorMultiply(VectorSwizzle(vec1, 1,0,3,2), VectorSwizzle(vec2, 2,1,2,1)));
 }
 
-inline uniform FReal MatrixDeterminant(const uniform FMatrix& M)
+inline uniform double MatrixDeterminant(const uniform FMatrix44d& M)
+{
+	const uniform FVector4d *uniform MVec = (const uniform FVector4d *uniform) &M;
+
+	// sub matrices
+	const uniform FVector4d A = VectorShuffle(MVec[0], MVec[1], 0,1,0,1);
+	const uniform FVector4d B = VectorShuffle(MVec[0], MVec[1], 2,3,2,3);
+	const uniform FVector4d C = VectorShuffle(MVec[2], MVec[3], 0,1,0,1);
+	const uniform FVector4d D = VectorShuffle(MVec[2], MVec[3], 2,3,2,3);
+
+	// determinant as (|A| |B| |C| |D|)
+	const uniform FVector4d detSub = VectorSubtract(
+		VectorMultiply(VectorShuffle(MVec[0], MVec[2], 0,2,0,2), VectorShuffle(MVec[1], MVec[3], 1,3,1,3)),
+		VectorMultiply(VectorShuffle(MVec[0], MVec[2], 1,3,1,3), VectorShuffle(MVec[1], MVec[3], 0,2,0,2))
+	);
+
+	// D#C
+	const uniform FVector4d D_C = Mat2AdjMul(D, C);
+	// A#B
+	const uniform FVector4d A_B = Mat2AdjMul(A, B);
+
+	// |M| = |A|*|D| + |B|*|C|
+	uniform double detM = (detSub.V[0] * detSub.V[3]) + (detSub.V[1] * detSub.V[2]);
+
+	// tr((A#B)(D#C))
+	uniform FVector4d tr = VectorMultiply(A_B, VectorSwizzle(D_C, 0,2,1,3));
+	tr = tr + VectorSwizzle(tr, 2,3,0,1);
+	tr = tr + VectorSwizzle(tr, 1,2,3,0);
+	// |M| = |A|*|D| + |B|*|C| - tr((A#B)(D#C)
+	detM = detM - tr.V[0];
+
+	return detM;
+}
+
+inline uniform float MatrixDeterminant(const uniform FMatrix44f& M)
 {
-	const uniform FVector4 *uniform MVec = (const uniform FVector4 *uniform) &M;
+	const uniform FVector4f *uniform MVec = (const uniform FVector4f *uniform) &M;
 
 	// sub matrices
-	const uniform FVector4 A = VectorShuffle(MVec[0], MVec[1], 0,1,0,1);
-	const uniform FVector4 B = VectorShuffle(MVec[0], MVec[1], 2,3,2,3);
-	const uniform FVector4 C = VectorShuffle(MVec[2], MVec[3], 0,1,0,1);
-	const uniform FVector4 D = VectorShuffle(MVec[2], MVec[3], 2,3,2,3);
+	const uniform FVector4f A = VectorShuffle(MVec[0], MVec[1], 0,1,0,1);
+	const uniform FVector4f B = VectorShuffle(MVec[0], MVec[1], 2,3,2,3);
+	const uniform FVector4f C = VectorShuffle(MVec[2], MVec[3], 0,1,0,1);
+	const uniform FVector4f D = VectorShuffle(MVec[2], MVec[3], 2,3,2,3);
 
 	// determinant as (|A| |B| |C| |D|)
-	const uniform FVector4 detSub = VectorSubtract(
+	const uniform FVector4f detSub = VectorSubtract(
 		VectorMultiply(VectorShuffle(MVec[0], MVec[2], 0,2,0,2), VectorShuffle(MVec[1], MVec[3], 1,3,1,3)),
 		VectorMultiply(VectorShuffle(MVec[0], MVec[2], 1,3,1,3), VectorShuffle(MVec[1], MVec[3], 0,2,0,2))
 	);
 
 	// D#C
-	const uniform FVector4 D_C = Mat2AdjMul(D, C);
+	const uniform FVector4f D_C = Mat2AdjMul(D, C);
 	// A#B
-	const uniform FVector4 A_B = Mat2AdjMul(A, B);
+	const uniform FVector4f A_B = Mat2AdjMul(A, B);
 
 	// |M| = |A|*|D| + |B|*|C|
-	uniform FReal detM = (detSub.V[0] * detSub.V[3]) + (detSub.V[1] * detSub.V[2]);
+	uniform float detM = (detSub.V[0] * detSub.V[3]) + (detSub.V[1] * detSub.V[2]);
 
 	// tr((A#B)(D#C))
-	uniform FVector4 tr = VectorMultiply(A_B, VectorSwizzle(D_C, 0,2,1,3));
+	uniform FVector4f tr = VectorMultiply(A_B, VectorSwizzle(D_C, 0,2,1,3));
 	tr = tr + VectorSwizzle(tr, 2,3,0,1);
 	tr = tr + VectorSwizzle(tr, 1,2,3,0);
 	// |M| = |A|*|D| + |B|*|C| - tr((A#B)(D#C)
@@ -286,73 +380,151 @@ inline uniform FReal MatrixDeterminant(const uniform FMatrix& M)
 	return detM;
 }
 
-inline uniform FMatrix VectorMatrixInverse(const uniform FMatrix& M)
+inline uniform FMatrix44d VectorMatrixInverse(const uniform FMatrix44d& M)
+{
+	// use block matrix method
+	// A is a matrix, then i(A) or iA means inverse of A, A# (or A_ in code) means adjugate of A, |A| (or detA in code) is determinant, tr(A) is trace
+	const uniform FVector4d *uniform MVec = (const uniform FVector4d *uniform) &M;
+
+	// sub matrices
+	const uniform FVector4d A = VectorShuffle(MVec[0], MVec[1], 0,1,0,1);
+	const uniform FVector4d B = VectorShuffle(MVec[0], MVec[1], 2,3,2,3);
+	const uniform FVector4d C = VectorShuffle(MVec[2], MVec[3], 0,1,0,1);
+	const uniform FVector4d D = VectorShuffle(MVec[2], MVec[3], 2,3,2,3);
+
+	// determinant as (|A| |B| |C| |D|)
+	const uniform FVector4d detSub = VectorSubtract(
+		VectorMultiply(VectorShuffle(MVec[0], MVec[2], 0,2,0,2), VectorShuffle(MVec[1], MVec[3], 1,3,1,3)),
+		VectorMultiply(VectorShuffle(MVec[0], MVec[2], 1,3,1,3), VectorShuffle(MVec[1], MVec[3], 0,2,0,2))
+	);
+
+	const uniform FVector4d detA = VectorSwizzle(detSub, 0,0,0,0);
+	const uniform FVector4d detB = VectorSwizzle(detSub, 1,1,1,1);
+	const uniform FVector4d detC = VectorSwizzle(detSub, 2,2,2,2);
+	const uniform FVector4d detD = VectorSwizzle(detSub, 3,3,3,3);
+
+	// let iM = 1/|M| * | X  Y |
+	//                  | Z  W |
+
+	// D#C
+	const uniform FVector4d D_C = Mat2AdjMul(D, C);
+	// A#B
+	const uniform FVector4d A_B = Mat2AdjMul(A, B);
+	// X# = |D|A - B(D#C)
+	uniform FVector4d X_ = VectorSubtract(VectorMultiply(detD, A), Mat2Mul(B, D_C));
+	// W# = |A|D - C(A#B)
+	uniform FVector4d W_ = VectorSubtract(VectorMultiply(detA, D), Mat2Mul(C, A_B));
+
+	// |M| = |A|*|D| + ... (continue later)
+	uniform FVector4d detM = VectorMultiply(detA, detD);
+
+	// Y# = |B|C - D(A#B)#
+	uniform FVector4d Y_ = VectorSubtract(VectorMultiply(detB, C), Mat2MulAdj(D, A_B));
+	// Z# = |C|B - A(D#C)#
+	uniform FVector4d Z_ = VectorSubtract(VectorMultiply(detC, B), Mat2MulAdj(A, D_C));
+
+	// |M| = |A|*|D| + |B|*|C| ... (continue later)
+	detM = VectorAdd(detM, VectorMultiply(detB, detC));
+
+	// tr((A#B)(D#C))
+	uniform FVector4d tr = VectorMultiply(A_B, VectorSwizzle(D_C, 0,2,1,3));
+	tr = tr + VectorSwizzle(tr, 2,3,0,1);
+	tr = tr + VectorSwizzle(tr, 1,2,3,0);
+	// |M| = |A|*|D| + |B|*|C| - tr((A#B)(D#C)
+	detM = VectorSubtract(detM, tr);
+
+	uniform FMatrix44d R = DoubleMatrixIdentity;
+
+	if(detM.V[0] != DOUBLE_ZERO)
+	{
+		const uniform FVector4d adjSignMask = SetVector4(DOUBLE_ONE, -DOUBLE_ONE, -DOUBLE_ONE, DOUBLE_ONE);
+		// (1/|M|, -1/|M|, -1/|M|, 1/|M|)
+		const uniform FVector4d rDetM = VectorDivide(adjSignMask, detM);
+
+		X_ = VectorMultiply(X_, rDetM);
+		Y_ = VectorMultiply(Y_, rDetM);
+		Z_ = VectorMultiply(Z_, rDetM);
+		W_ = VectorMultiply(W_, rDetM);
+
+		uniform FVector4d *uniform RVec = (uniform FVector4d *uniform) &R;
+
+		// apply adjugate and store, here we combine adjugate shuffle and store shuffle
+		RVec[0] = VectorShuffle(X_, Y_, 3,1,3,1);
+		RVec[1] = VectorShuffle(X_, Y_, 2,0,2,0);
+		RVec[2] = VectorShuffle(Z_, W_, 3,1,3,1);
+		RVec[3] = VectorShuffle(Z_, W_, 2,0,2,0);	
+	}
+
+	return R;
+}
+
+inline uniform FMatrix44f VectorMatrixInverse(const uniform FMatrix44f& M)
 {
 	// use block matrix method
 	// A is a matrix, then i(A) or iA means inverse of A, A# (or A_ in code) means adjugate of A, |A| (or detA in code) is determinant, tr(A) is trace
-	const uniform FVector4 *uniform MVec = (const uniform FVector4 *uniform) &M;
+	const uniform FVector4f *uniform MVec = (const uniform FVector4f *uniform) &M;
 
 	// sub matrices
-	const uniform FVector4 A = VectorShuffle(MVec[0], MVec[1], 0,1,0,1);
-	const uniform FVector4 B = VectorShuffle(MVec[0], MVec[1], 2,3,2,3);
-	const uniform FVector4 C = VectorShuffle(MVec[2], MVec[3], 0,1,0,1);
-	const uniform FVector4 D = VectorShuffle(MVec[2], MVec[3], 2,3,2,3);
+	const uniform FVector4f A = VectorShuffle(MVec[0], MVec[1], 0,1,0,1);
+	const uniform FVector4f B = VectorShuffle(MVec[0], MVec[1], 2,3,2,3);
+	const uniform FVector4f C = VectorShuffle(MVec[2], MVec[3], 0,1,0,1);
+	const uniform FVector4f D = VectorShuffle(MVec[2], MVec[3], 2,3,2,3);
 
 	// determinant as (|A| |B| |C| |D|)
-	const uniform FVector4 detSub = VectorSubtract(
+	const uniform FVector4f detSub = VectorSubtract(
 		VectorMultiply(VectorShuffle(MVec[0], MVec[2], 0,2,0,2), VectorShuffle(MVec[1], MVec[3], 1,3,1,3)),
 		VectorMultiply(VectorShuffle(MVec[0], MVec[2], 1,3,1,3), VectorShuffle(MVec[1], MVec[3], 0,2,0,2))
 	);
 
-	const uniform FVector4 detA = VectorSwizzle(detSub, 0,0,0,0);
-	const uniform FVector4 detB = VectorSwizzle(detSub, 1,1,1,1);
-	const uniform FVector4 detC = VectorSwizzle(detSub, 2,2,2,2);
-	const uniform FVector4 detD = VectorSwizzle(detSub, 3,3,3,3);
+	const uniform FVector4f detA = VectorSwizzle(detSub, 0,0,0,0);
+	const uniform FVector4f detB = VectorSwizzle(detSub, 1,1,1,1);
+	const uniform FVector4f detC = VectorSwizzle(detSub, 2,2,2,2);
+	const uniform FVector4f detD = VectorSwizzle(detSub, 3,3,3,3);
 
 	// let iM = 1/|M| * | X  Y |
 	//                  | Z  W |
 
 	// D#C
-	const uniform FVector4 D_C = Mat2AdjMul(D, C);
+	const uniform FVector4f D_C = Mat2AdjMul(D, C);
 	// A#B
-	const uniform FVector4 A_B = Mat2AdjMul(A, B);
+	const uniform FVector4f A_B = Mat2AdjMul(A, B);
 	// X# = |D|A - B(D#C)
-	uniform FVector4 X_ = VectorSubtract(VectorMultiply(detD, A), Mat2Mul(B, D_C));
+	uniform FVector4f X_ = VectorSubtract(VectorMultiply(detD, A), Mat2Mul(B, D_C));
 	// W# = |A|D - C(A#B)
-	uniform FVector4 W_ = VectorSubtract(VectorMultiply(detA, D), Mat2Mul(C, A_B));
+	uniform FVector4f W_ = VectorSubtract(VectorMultiply(detA, D), Mat2Mul(C, A_B));
 
 	// |M| = |A|*|D| + ... (continue later)
-	uniform FVector4 detM = VectorMultiply(detA, detD);
+	uniform FVector4f detM = VectorMultiply(detA, detD);
 
 	// Y# = |B|C - D(A#B)#
-	uniform FVector4 Y_ = VectorSubtract(VectorMultiply(detB, C), Mat2MulAdj(D, A_B));
+	uniform FVector4f Y_ = VectorSubtract(VectorMultiply(detB, C), Mat2MulAdj(D, A_B));
 	// Z# = |C|B - A(D#C)#
-	uniform FVector4 Z_ = VectorSubtract(VectorMultiply(detC, B), Mat2MulAdj(A, D_C));
+	uniform FVector4f Z_ = VectorSubtract(VectorMultiply(detC, B), Mat2MulAdj(A, D_C));
 
 	// |M| = |A|*|D| + |B|*|C| ... (continue later)
 	detM = VectorAdd(detM, VectorMultiply(detB, detC));
 
 	// tr((A#B)(D#C))
-	uniform FVector4 tr = VectorMultiply(A_B, VectorSwizzle(D_C, 0,2,1,3));
+	uniform FVector4f tr = VectorMultiply(A_B, VectorSwizzle(D_C, 0,2,1,3));
 	tr = tr + VectorSwizzle(tr, 2,3,0,1);
 	tr = tr + VectorSwizzle(tr, 1,2,3,0);
 	// |M| = |A|*|D| + |B|*|C| - tr((A#B)(D#C)
 	detM = VectorSubtract(detM, tr);
 
-	uniform FMatrix R = MatrixIdentity;
+	uniform FMatrix44f R = FloatMatrixIdentity;
 
 	if(detM.V[0] != ZERO)
 	{
-		const uniform FVector4 adjSignMask = SetVector4(ONE, -ONE, -ONE, ONE);
+		const uniform FVector4f adjSignMask = SetVector4(FLOAT_ONE, -FLOAT_ONE, -FLOAT_ONE, FLOAT_ONE);
 		// (1/|M|, -1/|M|, -1/|M|, 1/|M|)
-		const uniform FVector4 rDetM = VectorDivide(adjSignMask, detM);
+		const uniform FVector4f rDetM = VectorDivide(adjSignMask, detM);
 
 		X_ = VectorMultiply(X_, rDetM);
 		Y_ = VectorMultiply(Y_, rDetM);
 		Z_ = VectorMultiply(Z_, rDetM);
 		W_ = VectorMultiply(W_, rDetM);
 
-		uniform FVector4 *uniform RVec = (uniform FVector4 *uniform) &R;
+		uniform FVector4f *uniform RVec = (uniform FVector4f *uniform) &R;
 
 		// apply adjugate and store, here we combine adjugate shuffle and store shuffle
 		RVec[0] = VectorShuffle(X_, Y_, 3,1,3,1);
@@ -364,15 +536,33 @@ inline uniform FMatrix VectorMatrixInverse(const uniform FMatrix& M)
 	return R;
 }
 
-inline uniform FMatrix MatrixInverse(const uniform FMatrix& M)
+inline uniform FMatrix44d MatrixInverse(const uniform FMatrix44d& M)
 {
-	uniform FMatrix Result;
+	uniform FMatrix44d Result;
+
+	// Check for zero scale matrix to invert
+	if(	IsAllMatrixScaledAxesNearlyZero(M, DOUBLE_SMALL_NUMBER) ) 
+	{
+		// just set to zero - avoids unsafe inverse of zero and duplicates what QNANs were resulting in before (scaling away all children)
+		Result = DoubleMatrixIdentity;
+	}
+	else
+	{
+		Result = VectorMatrixInverse(M);
+	}
+
+	return Result;
+}
+
+inline uniform FMatrix44f MatrixInverse(const uniform FMatrix44f& M)
+{
+	uniform FMatrix44f Result;
 
 	// Check for zero scale matrix to invert
-	if(	IsAllMatrixScaledAxesNearlyZero(M, SMALL_NUMBER) ) 
+	if(	IsAllMatrixScaledAxesNearlyZero(M, FLOAT_SMALL_NUMBER) ) 
 	{
 		// just set to zero - avoids unsafe inverse of zero and duplicates what QNANs were resulting in before (scaling away all children)
-		Result = MatrixIdentity;
+		Result = FloatMatrixIdentity;
 	}
 	else
 	{
@@ -382,9 +572,33 @@ inline uniform FMatrix MatrixInverse(const uniform FMatrix& M)
 	return Result;
 }
 
-inline uniform FVector4 VectorTransformVector(const uniform FVector4 &VecP, const uniform FMatrix &M)
+inline uniform FVector4d VectorTransformVector(const uniform FVector4d &VecP, const uniform FMatrix44d &M)
+{
+	uniform FVector4d VTempX, VTempY, VTempZ, VTempW;
+
+	// Splat x,y,z and w
+	VTempX = VectorReplicate(VecP, 0);
+	VTempY = VectorReplicate(VecP, 1);
+	VTempZ = VectorReplicate(VecP, 2);
+	VTempW = VectorReplicate(VecP, 3);
+
+	// Mul by the matrix
+	VTempX = VectorMultiply(VTempX, SetVector4(M.M[0], M.M[1], M.M[2], M.M[3]));
+	VTempY = VectorMultiply(VTempY, SetVector4(M.M[4], M.M[5], M.M[6], M.M[7]));
+	VTempZ = VectorMultiply(VTempZ, SetVector4(M.M[8], M.M[9], M.M[10], M.M[11]));
+	VTempW = VectorMultiply(VTempW, SetVector4(M.M[12], M.M[13], M.M[14], M.M[15]));
+
+	// Add them all together
+	VTempX = VectorAdd(VTempX, VTempY);
+	VTempZ = VectorAdd(VTempZ, VTempW);
+	VTempX = VectorAdd(VTempX, VTempZ);
+
+	return VTempX;
+}
+
+inline uniform FVector4f VectorTransformVector(const uniform FVector4f &VecP, const uniform FMatrix44f &M)
 {
-	uniform FVector4 VTempX, VTempY, VTempZ, VTempW;
+	uniform FVector4f VTempX, VTempY, VTempZ, VTempW;
 
 	// Splat x,y,z and w
 	VTempX = VectorReplicate(VecP, 0);
@@ -520,13 +734,29 @@ inline FVector MatrixTransformVector(const FVector &P, const FMatrix &M)
 	return VTempX + VTempY + VTempZ;
 }
 
-inline uniform FVector MatrixInverseTransformVector(const uniform FMatrix &M, const uniform FVector &V)
+inline uniform FVector3d MatrixInverseTransformVector(const uniform FMatrix44d &M, const uniform FVector3d &V)
+{
+	const uniform FMatrix44d InvSelf = VectorMatrixInverse(M); // FMatrix::InverseFast();
+	return SetVector(VectorTransformVector(SetVector4(V, DOUBLE_ZERO), InvSelf));
+}
+
+inline uniform FVector3f MatrixInverseTransformVector(const uniform FMatrix44f &M, const uniform FVector3f &V)
+{
+	const uniform FMatrix44f InvSelf = VectorMatrixInverse(M); // FMatrix::InverseFast();
+	return SetVector(VectorTransformVector(SetVector4(V, FLOAT_ZERO), InvSelf));
+}
+
+inline uniform FMatrix44d MatrixReduceAdd(const varying FMatrix44d &M)
 {
-	const uniform FMatrix InvSelf = VectorMatrixInverse(M); // FMatrix::InverseFast();
-	return SetVector(VectorTransformVector(SetVector4(V, ZERO), InvSelf));
+	return SetMatrix(
+		SetVector4(reduce_add(M.M[0]), reduce_add(M.M[1]), reduce_add(M.M[2]), reduce_add(M.M[3])),
+		SetVector4(reduce_add(M.M[4]), reduce_add(M.M[5]), reduce_add(M.M[6]), reduce_add(M.M[7])),
+		SetVector4(reduce_add(M.M[8]), reduce_add(M.M[9]), reduce_add(M.M[10]), reduce_add(M.M[11])),
+		SetVector4(reduce_add(M.M[12]), reduce_add(M.M[13]), reduce_add(M.M[14]), reduce_add(M.M[15]))
+		);
 }
 
-inline uniform FMatrix MatrixReduceAdd(const varying FMatrix &M)
+inline uniform FMatrix44f MatrixReduceAdd(const varying FMatrix44f &M)
 {
 	return SetMatrix(
 		SetVector4(reduce_add(M.M[0]), reduce_add(M.M[1]), reduce_add(M.M[2]), reduce_add(M.M[3])),
diff --git a/Engine/Source/Runtime/Core/Public/Math/Quat.isph b/Engine/Source/Runtime/Core/Public/Math/Quat.isph
--- a/Engine/Source/Runtime/Core/Public/Math/Quat.isph
+++ b/Engine/Source/Runtime/Core/Public/Math/Quat.isph
@@ -176,7 +176,7 @@ inline uniform FVector4 VectorAccumulateQuaternionShortestPath(const uniform FVe
  * @param VectorW0 Vector to rotate. W component must be zero.
  * @return Vector after rotation by Quat.
  */
-inline FVector4 VectorQuaternionRotateVector(const FVector4& Quat, const FVector4& VectorW0)
+inline FVector4d VectorQuaternionRotateVector(const FVector4d& Quat, const FVector4d& VectorW0)
 {
 	// Q * V * Q.Inverse
 	//const VectorRegister InverseRotation = VectorQuaternionInverse(Quat);
@@ -191,12 +191,36 @@ inline FVector4 VectorQuaternionRotateVector(const FVector4& Quat, const FVector
 	// T = 2(Q x V);
 	// V' = V + w*(T) + (Q x T)
 
-	const FVector4 QW = VectorReplicate(Quat, 3);
-	FVector4 T = VectorCross(Quat, VectorW0);
+	const FVector4d QW = VectorReplicate(Quat, 3);
+	FVector4d T = VectorCross(Quat, VectorW0);
 	T = VectorAdd(T, T);
-	const FVector4 VTemp0 = VectorMultiplyAdd(QW, T, VectorW0);
-	const FVector4 VTemp1 = VectorCross(Quat, T);
-	const FVector4 Rotated = VectorAdd(VTemp0, VTemp1);
+	const FVector4d VTemp0 = VectorMultiplyAdd(QW, T, VectorW0);
+	const FVector4d VTemp1 = VectorCross(Quat, T);
+	const FVector4d Rotated = VectorAdd(VTemp0, VTemp1);
+	return Rotated;
+}
+
+inline FVector4f VectorQuaternionRotateVector(const FVector4f& Quat, const FVector4f& VectorW0)
+{
+	// Q * V * Q.Inverse
+	//const VectorRegister InverseRotation = VectorQuaternionInverse(Quat);
+	//const VectorRegister Temp = VectorQuaternionMultiply2(Quat, VectorW0);
+	//const VectorRegister Rotated = VectorQuaternionMultiply2(Temp, InverseRotation);
+
+	// Equivalence of above can be shown to be:
+	// http://people.csail.mit.edu/bkph/articles/Quaternions.pdf
+	// V' = V + 2w(Q x V) + (2Q x (Q x V))
+	// refactor:
+	// V' = V + w(2(Q x V)) + (Q x (2(Q x V)))
+	// T = 2(Q x V);
+	// V' = V + w*(T) + (Q x T)
+
+	const FVector4f QW = VectorReplicate(Quat, 3);
+	FVector4f T = VectorCross(Quat, VectorW0);
+	T = VectorAdd(T, T);
+	const FVector4f VTemp0 = VectorMultiplyAdd(QW, T, VectorW0);
+	const FVector4f VTemp1 = VectorCross(Quat, T);
+	const FVector4f Rotated = VectorAdd(VTemp0, VTemp1);
 	return Rotated;
 }
 
@@ -224,7 +248,7 @@ inline FVector4 VectorQuaternionRotateVector(const uniform FVector4& Quat, const
 	return Rotated;
 }
 
-inline uniform FVector4 VectorQuaternionRotateVector(const uniform FVector4& Quat, const uniform FVector4& VectorW0)
+inline uniform FVector4d VectorQuaternionRotateVector(const uniform FVector4d& Quat, const uniform FVector4d& VectorW0)
 {
 	// Q * V * Q.Inverse
 	//const VectorRegister InverseRotation = VectorQuaternionInverse(Quat);
@@ -239,12 +263,36 @@ inline uniform FVector4 VectorQuaternionRotateVector(const uniform FVector4& Qua
 	// T = 2(Q x V);
 	// V' = V + w*(T) + (Q x T)
 
-	const uniform FVector4 QW = VectorReplicate(Quat, 3);
-	uniform FVector4 T = VectorCross(Quat, VectorW0);
+	const uniform FVector4d QW = VectorReplicate(Quat, 3);
+	uniform FVector4d T = VectorCross(Quat, VectorW0);
+	T = VectorAdd(T, T);
+	const uniform FVector4d VTemp0 = VectorMultiplyAdd(QW, T, VectorW0);
+	const uniform FVector4d VTemp1 = VectorCross(Quat, T);
+	const uniform FVector4d Rotated = VectorAdd(VTemp0, VTemp1);
+	return Rotated;
+}
+
+inline uniform FVector4f VectorQuaternionRotateVector(const uniform FVector4f& Quat, const uniform FVector4f& VectorW0)
+{
+	// Q * V * Q.Inverse
+	//const VectorRegister InverseRotation = VectorQuaternionInverse(Quat);
+	//const VectorRegister Temp = VectorQuaternionMultiply2(Quat, VectorW0);
+	//const VectorRegister Rotated = VectorQuaternionMultiply2(Temp, InverseRotation);
+
+	// Equivalence of above can be shown to be:
+	// http://people.csail.mit.edu/bkph/articles/Quaternions.pdf
+	// V' = V + 2w(Q x V) + (2Q x (Q x V))
+	// refactor:
+	// V' = V + w(2(Q x V)) + (Q x (2(Q x V)))
+	// T = 2(Q x V);
+	// V' = V + w*(T) + (Q x T)
+
+	const uniform FVector4f QW = VectorReplicate(Quat, 3);
+	uniform FVector4f T = VectorCross(Quat, VectorW0);
 	T = VectorAdd(T, T);
-	const uniform FVector4 VTemp0 = VectorMultiplyAdd(QW, T, VectorW0);
-	const uniform FVector4 VTemp1 = VectorCross(Quat, T);
-	const uniform FVector4 Rotated = VectorAdd(VTemp0, VTemp1);
+	const uniform FVector4f VTemp0 = VectorMultiplyAdd(QW, T, VectorW0);
+	const uniform FVector4f VTemp1 = VectorCross(Quat, T);
+	const uniform FVector4f Rotated = VectorAdd(VTemp0, VTemp1);
 	return Rotated;
 }
 
@@ -263,7 +311,22 @@ inline uniform FVector VectorQuaternionRotateVector(const uniform FVector4& Quat
 	return Result;
 }
 
-inline FVector VectorQuaternionRotateVector(const FVector4& Quat, const FVector& V)
+inline FVector3d VectorQuaternionRotateVector(const FVector4d& Quat, const FVector3d& V)
+{
+	// http://people.csail.mit.edu/bkph/articles/Quaternions.pdf
+	// V' = V + 2w(Q x V) + (2Q x (Q x V))
+	// refactor:
+	// V' = V + w(2(Q x V)) + (Q x (2(Q x V)))
+	// T = 2(Q x V);
+	// V' = V + w*(T) + (Q x T)
+
+	const FVector3d Q = SetVector(Quat.V[0], Quat.V[1], Quat.V[2]);
+	const FVector3d T = VectorCross(Q, V) * 2;
+	const FVector3d Result = V + (T * Quat.V[3]) + VectorCross(Q, T);
+	return Result;
+}
+
+inline FVector3f VectorQuaternionRotateVector(const FVector4f& Quat, const FVector3f& V)
 {
 	// http://people.csail.mit.edu/bkph/articles/Quaternions.pdf
 	// V' = V + 2w(Q x V) + (2Q x (Q x V))
@@ -272,9 +335,9 @@ inline FVector VectorQuaternionRotateVector(const FVector4& Quat, const FVector&
 	// T = 2(Q x V);
 	// V' = V + w*(T) + (Q x T)
 
-	const FVector Q = SetVector(Quat.V[0], Quat.V[1], Quat.V[2]);
-	const FVector T = VectorCross(Q, V) * 2;
-	const FVector Result = V + (T * Quat.V[3]) + VectorCross(Q, T);
+	const FVector3f Q = SetVector(Quat.V[0], Quat.V[1], Quat.V[2]);
+	const FVector3f T = VectorCross(Q, V) * 2;
+	const FVector3f Result = V + (T * Quat.V[3]) + VectorCross(Q, T);
 	return Result;
 }
 
diff --git a/Engine/Source/Runtime/Core/Public/Math/Transform.isph b/Engine/Source/Runtime/Core/Public/Math/Transform.isph
--- a/Engine/Source/Runtime/Core/Public/Math/Transform.isph
+++ b/Engine/Source/Runtime/Core/Public/Math/Transform.isph
@@ -7,13 +7,22 @@
 #include "Math/Quat.isph"
 #include "Math/Matrix.isph"
 
+struct FTransform3f
+{
+	FVector4f Rotation;
+	FVector4f Translation;
+	FVector4f Scale3D;
+};
+
 struct FTransform
 {
-	FVector4 Rotation;
-	FVector4 Translation;
-	FVector4 Scale3D;
+	FVector4d Rotation;
+	FVector4d Translation;
+	FVector4d Scale3D;
 };
 
+typedef FTransform FTransform3d;
+
 struct FTransform8
 {
 	FVector8 Rotation;
@@ -216,30 +225,57 @@ inline void ScaleTranslation(uniform FTransform &Transform, const uniform FReal
 	Transform.Translation = VectorMultiply(Transform.Translation, VectorScale);
 }
 
-inline uniform FVector4 TransformVector(const uniform FTransform &T, const uniform FVector4& V)
+inline uniform FVector4d TransformVector(const uniform FTransform3d &T, const uniform FVector4d& V)
 {
-	const uniform FVector4 InputVectorW0 = SetVector4(V.V[0], V.V[1], V.V[2], 0);
+	const uniform FVector4d InputVectorW0 = SetVector4(V.V[0], V.V[1], V.V[2], 0);
 
 	//RotatedVec = Q.Rotate(Scale*V.X, Scale*V.Y, Scale*V.Z, 0.f)
-	const uniform FVector4 ScaledVec = VectorMultiply(T.Scale3D, InputVectorW0);
+	const uniform FVector4d ScaledVec = VectorMultiply(T.Scale3D, InputVectorW0);
 	return VectorQuaternionRotateVector(T.Rotation, ScaledVec);
 }
 
-inline uniform FVector TransformVector(const uniform FTransform &T, const uniform FVector& V)
+inline uniform FVector4f TransformVector(const uniform FTransform3f &T, const uniform FVector4f& V)
 {
-	const uniform FVector4 InputVectorW0 = SetVector4(V, 0);
+	const uniform FVector4f InputVectorW0 = SetVector4(V.V[0], V.V[1], V.V[2], 0);
 
 	//RotatedVec = Q.Rotate(Scale*V.X, Scale*V.Y, Scale*V.Z, 0.f)
-	const uniform FVector4 ScaledVec = VectorMultiply(T.Scale3D, InputVectorW0);
-	const uniform FVector4 RotatedVec = VectorQuaternionRotateVector(T.Rotation, ScaledVec);
+	const uniform FVector4f ScaledVec = VectorMultiply(T.Scale3D, InputVectorW0);
+	return VectorQuaternionRotateVector(T.Rotation, ScaledVec);
+}
+
+inline uniform FVector3d TransformVector(const uniform FTransform3d &T, const uniform FVector3d& V)
+{
+	const uniform FVector4d InputVectorW0 = SetVector4(V, 0);
+
+	//RotatedVec = Q.Rotate(Scale*V.X, Scale*V.Y, Scale*V.Z, 0.f)
+	const uniform FVector4d ScaledVec = VectorMultiply(T.Scale3D, InputVectorW0);
+	const uniform FVector4d RotatedVec = VectorQuaternionRotateVector(T.Rotation, ScaledVec);
 	return SetVector(RotatedVec);
 }
 
-inline FVector TransformVector(const FTransform &T, const FVector& V)
+inline uniform FVector3f TransformVector(const uniform FTransform3f &T, const uniform FVector3f& V)
 {
+	const uniform FVector4f InputVectorW0 = SetVector4(V, 0);
+
 	//RotatedVec = Q.Rotate(Scale*V.X, Scale*V.Y, Scale*V.Z, 0.f)
-	const FVector ScaledVec = SetVector(T.Scale3D) * V;
-	const FVector RotatedVec = VectorQuaternionRotateVector(T.Rotation, ScaledVec);
+	const uniform FVector4f ScaledVec = VectorMultiply(T.Scale3D, InputVectorW0);
+	const uniform FVector4f RotatedVec = VectorQuaternionRotateVector(T.Rotation, ScaledVec);
+	return SetVector(RotatedVec);
+}
+
+inline FVector3d TransformVector(const FTransform3d &T, const FVector3d& V)
+{
+	//RotatedVec = Q.Rotate(Scale*V.X, Scale*V.Y, Scale*V.Z, 0.f)
+	const FVector3d ScaledVec = SetVector(T.Scale3D) * V;
+	const FVector3d RotatedVec = VectorQuaternionRotateVector(T.Rotation, ScaledVec);
+	return RotatedVec;
+}
+
+inline FVector3f TransformVector(const FTransform3f &T, const FVector3f& V)
+{
+	//RotatedVec = Q.Rotate(Scale*V.X, Scale*V.Y, Scale*V.Z, 0.f)
+	const FVector3f ScaledVec = SetVector(T.Scale3D) * V;
+	const FVector3f RotatedVec = VectorQuaternionRotateVector(T.Rotation, ScaledVec);
 	return RotatedVec;
 }
 
@@ -260,18 +296,34 @@ inline uniform FVector TransformVectorNoScale(const uniform FTransform &T, const
 	return SetVector(RotatedVec);
 }
 
-inline FVector TransformPosition(const FTransform &T, const FVector& V)
+inline FVector3d TransformPosition(const FTransform3d &T, const FVector3d& V)
 {
-	FVector4 InputVectorW0 = SetVector4(V.V[0], V.V[1], V.V[2], 0);
+	FVector4d InputVectorW0 = SetVector4(V.V[0], V.V[1], V.V[2], 0);
 
 	//Transform using QST is following
 	//QST(P) = Q.Rotate(S*P) + T where Q = quaternion, S = scale, T = translation
 
 	//RotatedVec = Q.Rotate(Scale*V.X, Scale*V.Y, Scale*V.Z, 0.f)
-	const FVector4 ScaledVec = VectorMultiply(T.Scale3D, InputVectorW0);
-	const FVector4 RotatedVec = VectorQuaternionRotateVector(T.Rotation, ScaledVec);
+	const FVector4d ScaledVec = VectorMultiply(T.Scale3D, InputVectorW0);
+	const FVector4d RotatedVec = VectorQuaternionRotateVector(T.Rotation, ScaledVec);
 
-	const FVector4 TranslatedVec = VectorAdd(RotatedVec, T.Translation);
+	const FVector4d TranslatedVec = VectorAdd(RotatedVec, T.Translation);
+
+	return SetVector(TranslatedVec);
+}
+
+inline FVector3f TransformPosition(const FTransform3f &T, const FVector3f& V)
+{
+	FVector4f InputVectorW0 = SetVector4(V.V[0], V.V[1], V.V[2], 0);
+
+	//Transform using QST is following
+	//QST(P) = Q.Rotate(S*P) + T where Q = quaternion, S = scale, T = translation
+
+	//RotatedVec = Q.Rotate(Scale*V.X, Scale*V.Y, Scale*V.Z, 0.f)
+	const FVector4f ScaledVec = VectorMultiply(T.Scale3D, InputVectorW0);
+	const FVector4f RotatedVec = VectorQuaternionRotateVector(T.Rotation, ScaledVec);
+
+	const FVector4f TranslatedVec = VectorAdd(RotatedVec, T.Translation);
 
 	return SetVector(TranslatedVec);
 }
@@ -292,32 +344,62 @@ inline FVector TransformPosition(const uniform FTransform &T, const FVector& V)
 	return SetVector(TranslatedVec);
 }
 
-inline uniform FVector TransformPosition(const uniform FTransform &T, const uniform FVector& V)
+inline uniform FVector3d TransformPosition(const uniform FTransform3d &T, const uniform FVector3d& V)
+{
+	uniform FVector4d InputVectorW0 = SetVector4(V.V[0], V.V[1], V.V[2], 0);
+
+	//Transform using QST is following
+	//QST(P) = Q.Rotate(S*P) + T where Q = quaternion, S = scale, T = translation
+
+	//RotatedVec = Q.Rotate(Scale*V.X, Scale*V.Y, Scale*V.Z, 0.f)
+	const uniform FVector4d ScaledVec = VectorMultiply(T.Scale3D, InputVectorW0);
+	const uniform FVector4d RotatedVec = VectorQuaternionRotateVector(T.Rotation, ScaledVec);
+
+	const uniform FVector4d TranslatedVec = VectorAdd(RotatedVec, T.Translation);
+
+	return SetVector(TranslatedVec);
+}
+
+inline uniform FVector3f TransformPosition(const uniform FTransform3f &T, const uniform FVector3f& V)
 {
-	uniform FVector4 InputVectorW0 = SetVector4(V.V[0], V.V[1], V.V[2], 0);
+	uniform FVector4f InputVectorW0 = SetVector4(V.V[0], V.V[1], V.V[2], 0);
 
 	//Transform using QST is following
 	//QST(P) = Q.Rotate(S*P) + T where Q = quaternion, S = scale, T = translation
 
 	//RotatedVec = Q.Rotate(Scale*V.X, Scale*V.Y, Scale*V.Z, 0.f)
-	const uniform FVector4 ScaledVec = VectorMultiply(T.Scale3D, InputVectorW0);
-	const uniform FVector4 RotatedVec = VectorQuaternionRotateVector(T.Rotation, ScaledVec);
+	const uniform FVector4f ScaledVec = VectorMultiply(T.Scale3D, InputVectorW0);
+	const uniform FVector4f RotatedVec = VectorQuaternionRotateVector(T.Rotation, ScaledVec);
 
-	const uniform FVector4 TranslatedVec = VectorAdd(RotatedVec, T.Translation);
+	const uniform FVector4f TranslatedVec = VectorAdd(RotatedVec, T.Translation);
 
 	return SetVector(TranslatedVec);
 }
 
-inline static uniform FVector4 TransformPosition(const uniform FTransform &T, const uniform FVector4& V)
+inline static uniform FVector4d TransformPosition(const uniform FTransform3d &T, const uniform FVector4d& V)
+{
+	//Transform using QST is following
+	//QST(P) = Q.Rotate(S*P) + T where Q = quaternion, S = scale, T = translation
+
+	//RotatedVec = Q.Rotate(Scale*V.X, Scale*V.Y, Scale*V.Z, 0.f)
+	const uniform FVector4d ScaledVec = VectorMultiply(T.Scale3D, V);
+	const uniform FVector4d RotatedVec = VectorQuaternionRotateVector(T.Rotation, ScaledVec);
+
+	const uniform FVector4d TranslatedVec = VectorAdd(RotatedVec, T.Translation);
+
+	return TranslatedVec;
+}
+
+inline static uniform FVector4f TransformPosition(const uniform FTransform3f &T, const uniform FVector4f& V)
 {
 	//Transform using QST is following
 	//QST(P) = Q.Rotate(S*P) + T where Q = quaternion, S = scale, T = translation
 
 	//RotatedVec = Q.Rotate(Scale*V.X, Scale*V.Y, Scale*V.Z, 0.f)
-	const uniform FVector4 ScaledVec = VectorMultiply(T.Scale3D, V);
-	const uniform FVector4 RotatedVec = VectorQuaternionRotateVector(T.Rotation, ScaledVec);
+	const uniform FVector4f ScaledVec = VectorMultiply(T.Scale3D, V);
+	const uniform FVector4f RotatedVec = VectorQuaternionRotateVector(T.Rotation, ScaledVec);
 
-	const uniform FVector4 TranslatedVec = VectorAdd(RotatedVec, T.Translation);
+	const uniform FVector4f TranslatedVec = VectorAdd(RotatedVec, T.Translation);
 
 	return TranslatedVec;
 }
diff --git a/Engine/Source/Runtime/Core/Public/Math/Vector.isph b/Engine/Source/Runtime/Core/Public/Math/Vector.isph
--- a/Engine/Source/Runtime/Core/Public/Math/Vector.isph
+++ b/Engine/Source/Runtime/Core/Public/Math/Vector.isph
@@ -467,25 +467,49 @@ inline void VectorStore(uniform FVector4 *uniform DstPtr, const varying FVector4
 	}
 }
 
-inline varying FVector2 VectorGather(const uniform FVector2 *varying SrcPtr)
+inline varying FVector2d VectorGather(const uniform FVector2d *varying SrcPtr)
 {
 	#pragma ignore warning(perf)
 	return *SrcPtr;
 }
 
-inline void VectorScatter(uniform FVector2 *varying DstPtr, const varying FVector2 &V)
+inline varying FVector2f VectorGather(const uniform FVector2f *varying SrcPtr)
+{
+	#pragma ignore warning(perf)
+	return *SrcPtr;
+}
+
+inline void VectorScatter(uniform FVector2d *varying DstPtr, const varying FVector2d &V)
+{
+	#pragma ignore warning(perf)
+	*DstPtr = V;
+}
+
+inline void VectorScatter(uniform FVector2f *varying DstPtr, const varying FVector2f &V)
 {
 	#pragma ignore warning(perf)
 	*DstPtr = V;
 }
 
-inline varying FVector VectorGather(const uniform FVector *varying SrcPtr)
+inline varying FVector3d VectorGather(const uniform FVector3d *varying SrcPtr)
 {
 	#pragma ignore warning(perf)
 	return *SrcPtr;
 }
 
-inline void VectorScatter(uniform FVector *varying DstPtr, const varying FVector &V)
+inline varying FVector3f VectorGather(const uniform FVector3f *varying SrcPtr)
+{
+	#pragma ignore warning(perf)
+	return *SrcPtr;
+}
+
+inline void VectorScatter(uniform FVector3d *varying DstPtr, const varying FVector3d &V)
+{
+	#pragma ignore warning(perf)
+	*DstPtr = V;
+}
+
+inline void VectorScatter(uniform FVector3f *varying DstPtr, const varying FVector3f &V)
 {
 	#pragma ignore warning(perf)
 	*DstPtr = V;
@@ -851,9 +875,21 @@ inline uniform FVector4f operator*(const uniform float F, const uniform FVector4
 	return FVec * A;
 }
 
-inline FVector4 operator/(const FVector4 &A, const FVector4 &B)
+inline FVector4d operator/(const FVector4d &A, const FVector4d &B)
 {
-	FVector4 Result;
+	FVector4d Result;
+
+	for(uniform int i = 0; i < 4; i++)
+	{
+		Result.V[i] = A.V[i] / B.V[i];
+	}
+
+	return Result;
+}
+
+inline FVector4f operator/(const FVector4f &A, const FVector4f &B)
+{
+	FVector4f Result;
 
 	for(uniform int i = 0; i < 4; i++)
 	{
@@ -887,15 +923,26 @@ inline FVector4 operator/(const uniform FVector4 &A, const FVector4 &B)
 	return Result;
 }
 
-inline uniform FVector4 operator/(const uniform FVector4 &A, const uniform FVector4 &B)
+inline uniform FVector4d operator/(const uniform FVector4d &A, const uniform FVector4d &B)
 {
-	varying FReal S0, S1, Result;
-	*((uniform FVector4 *uniform)&S0) = *((uniform FVector4 *uniform)&A);
-	*((uniform FVector4 *uniform)&S1) = *((uniform FVector4 *uniform)&B);
+	varying double S0, S1, Result;
+	*((uniform FVector4d *uniform)&S0) = *((uniform FVector4d *uniform)&A);
+	*((uniform FVector4d *uniform)&S1) = *((uniform FVector4d *uniform)&B);
 
 	Result = S0 / S1;
 
-	return *((uniform FVector4 *uniform)&Result);
+	return *((uniform FVector4d *uniform)&Result);
+}
+
+inline uniform FVector4f operator/(const uniform FVector4f &A, const uniform FVector4f &B)
+{
+	varying float S0, S1, Result;
+	*((uniform FVector4f *uniform)&S0) = *((uniform FVector4f *uniform)&A);
+	*((uniform FVector4f *uniform)&S1) = *((uniform FVector4f *uniform)&B);
+
+	Result = S0 / S1;
+
+	return *((uniform FVector4f *uniform)&Result);
 }
 
 inline uniform FVector4 operator/(const uniform FVector4 &A, const uniform FReal F)
@@ -1070,21 +1117,63 @@ inline FVector4 SetVector4(const FVector &V, const FReal W)
 	return Result;
 }
 
-inline uniform FVector4 SetVector4(const uniform FVector &V, const uniform FReal W)
+inline uniform FVector4d SetVector4(const uniform FVector3d &V, const uniform double W)
+{
+	const uniform FVector4d Result = {{V.V[0], V.V[1], V.V[2], W}};
+	return Result;
+}
+
+inline uniform FVector4f SetVector4(const uniform FVector3f &V, const uniform float W)
+{
+	const uniform FVector4f Result = {{V.V[0], V.V[1], V.V[2], W}};
+	return Result;
+}
+
+inline uniform FVector4d SetVector4(const uniform double F)
+{
+	const uniform FVector4d Result = {{F, F, F, F}};
+	return Result;
+}
+
+inline uniform FVector4f SetVector4(const uniform float F)
 {
-	const uniform FVector4 Result = {{V.V[0], V.V[1], V.V[2], W}};
+	const uniform FVector4f Result = {{F, F, F, F}};
 	return Result;
 }
 
-inline uniform FVector4 SetVector4(const uniform FReal F)
+inline varying FVector4d SetVector4(const varying double F)
 {
-	const uniform FVector4 Result = {{F, F, F, F}};
+	const varying FVector4d Result = {{F, F, F, F}};
 	return Result;
 }
 
-inline varying FVector4 SetVector4(const varying FReal F)
+inline varying FVector4f SetVector4(const varying float F)
 {
-	const varying FVector4 Result = {{F, F, F, F}};
+	const varying FVector4f Result = {{F, F, F, F}};
+	return Result;
+}
+
+inline FVector4d ConvertVector4fTo4d(const FVector4f &V)
+{
+	FVector4d Result;
+
+	for(uniform int i = 0; i < 4; i++)
+	{
+		Result.V[i] = (varying double)V.V[i];
+	}
+
+	return Result;
+}
+
+inline FVector4f ConvertVector4dTo4f(const FVector4d &V)
+{
+	FVector4f Result;
+
+	for(uniform int i = 0; i < 4; i++)
+	{
+		Result.V[i] = (varying float)V.V[i];
+	}
+
 	return Result;
 }
 
@@ -1108,6 +1197,26 @@ inline uniform FVector4f ConvertVector4dTo4f(const uniform FVector4d &V)
 	return *((uniform FVector4f *uniform)&Result);
 }
 
+inline FVector4 ConvertVector4dTo4Native(const FVector4d &V)
+{
+	return V;
+}
+
+inline FVector4 ConvertVector4fTo4Native(const FVector4f &V)
+{
+	return ConvertVector4fTo4d(V);
+}
+
+inline FVector4d ConvertVector4NativeTo4d(const FVector4 &V)
+{
+	return V;
+}
+
+inline FVector4f ConvertVector4NativeTo4f(const FVector4 &V)
+{
+	return ConvertVector4dTo4f(V);
+}
+
 inline uniform FVector4 ConvertVector4dTo4Native(const uniform FVector4d &V)
 {
 	return V;
@@ -1194,9 +1303,20 @@ inline FVector3f operator+(const FVector3f &A, const uniform FVector3f &B)
 	return Result;
 }
 
-inline FVector operator+(const uniform FVector &A, const FVector &B)
+inline FVector3d operator+(const uniform FVector3d &A, const FVector3d &B)
 {
-	FVector Result;
+	FVector3d Result;
+
+	Result.V[0] = A.V[0] + B.V[0];
+	Result.V[1] = A.V[1] + B.V[1];
+	Result.V[2] = A.V[2] + B.V[2];
+
+	return Result;
+}
+
+inline FVector3f operator+(const uniform FVector3f &A, const FVector3f &B)
+{
+	FVector3f Result;
 
 	Result.V[0] = A.V[0] + B.V[0];
 	Result.V[1] = A.V[1] + B.V[1];
@@ -1387,9 +1507,20 @@ inline FVector3f operator*(const FVector3f &A, const uniform float F)
 	return Result;
 }
 
-inline FVector operator*(const FReal F, const FVector &A)
+inline FVector3d operator*(const double F, const FVector3d &A)
 {
-	FVector Result;
+	FVector3d Result;
+
+	Result.V[0] = F * A.V[0];
+	Result.V[1] = F * A.V[1];
+	Result.V[2] = F * A.V[2];
+
+	return Result;
+}
+
+inline FVector3f operator*(const float F, const FVector3f &A)
+{
+	FVector3f Result;
 
 	Result.V[0] = F * A.V[0];
 	Result.V[1] = F * A.V[1];
@@ -1476,9 +1607,20 @@ inline FVector3f operator-(const FVector3f &A, const uniform FVector3f &B)
 	return Result;
 }
 
-inline FVector operator-(const uniform FVector &A, const FVector &B)
+inline FVector3d operator-(const uniform FVector3d &A, const FVector3d &B)
 {
-	FVector Result;
+	FVector3d Result;
+
+	Result.V[0] = A.V[0] - B.V[0];
+	Result.V[1] = A.V[1] - B.V[1];
+	Result.V[2] = A.V[2] - B.V[2];
+
+	return Result;
+}
+
+inline FVector3f operator-(const uniform FVector3f &A, const FVector3f &B)
+{
+	FVector3f Result;
 
 	Result.V[0] = A.V[0] - B.V[0];
 	Result.V[1] = A.V[1] - B.V[1];
@@ -1521,9 +1663,9 @@ inline FVector operator-(const FReal F, const FVector &A)
 	return Result;
 }
 
-inline FVector operator/(const FVector &A, const FVector &B)
+inline FVector3d operator/(const FVector3d &A, const FVector3d &B)
 {
-	FVector Result;
+	FVector3d Result;
 
 	Result.V[0] = A.V[0] / B.V[0];
 	Result.V[1] = A.V[1] / B.V[1];
@@ -1532,15 +1674,37 @@ inline FVector operator/(const FVector &A, const FVector &B)
 	return Result;
 }
 
-inline uniform FVector operator/(const uniform FVector &A, const uniform FVector &B)
+inline FVector3f operator/(const FVector3f &A, const FVector3f &B)
 {
-	uniform FVector4 S0, S1, Result;
-	*((uniform FVector *uniform)&S0) = *((uniform FVector *uniform)&A);
-	*((uniform FVector *uniform)&S1) = *((uniform FVector *uniform)&B);
+	FVector3f Result;
+
+	Result.V[0] = A.V[0] / B.V[0];
+	Result.V[1] = A.V[1] / B.V[1];
+	Result.V[2] = A.V[2] / B.V[2];
+
+	return Result;
+}
+
+inline uniform FVector3d operator/(const uniform FVector3d &A, const uniform FVector3d &B)
+{
+	uniform FVector4d S0, S1, Result;
+	*((uniform FVector3d *uniform)&S0) = *((uniform FVector3d *uniform)&A);
+	*((uniform FVector3d *uniform)&S1) = *((uniform FVector3d *uniform)&B);
 
 	Result = S0 / S1;
 
-	return *((uniform FVector *uniform)&Result);
+	return *((uniform FVector3d *uniform)&Result);
+}
+
+inline uniform FVector3f operator/(const uniform FVector3f &A, const uniform FVector3f &B)
+{
+	uniform FVector4f S0, S1, Result;
+	*((uniform FVector3f *uniform)&S0) = *((uniform FVector3f *uniform)&A);
+	*((uniform FVector3f *uniform)&S1) = *((uniform FVector3f *uniform)&B);
+
+	Result = S0 / S1;
+
+	return *((uniform FVector3f *uniform)&Result);
 }
 
 inline FVector operator/(const FVector &A, const uniform FVector &B)
@@ -1554,9 +1718,9 @@ inline FVector operator/(const FVector &A, const uniform FVector &B)
 	return Result;
 }
 
-inline FVector operator/(const FVector &A, const FReal F)
+inline FVector3d operator/(const FVector3d &A, const double F)
 {
-	FVector Result;
+	FVector3d Result;
 
 	Result.V[0] = A.V[0] / F;
 	Result.V[1] = A.V[1] / F;
@@ -1565,9 +1729,26 @@ inline FVector operator/(const FVector &A, const FReal F)
 	return Result;
 }
 
-inline uniform FVector operator/(const uniform FVector &A, const uniform FReal F)
+inline FVector3f operator/(const FVector3f &A, const float F)
 {
-	const uniform FVector FVec = {{F,F,F}};
+	FVector3f Result;
+
+	Result.V[0] = A.V[0] / F;
+	Result.V[1] = A.V[1] / F;
+	Result.V[2] = A.V[2] / F;
+
+	return Result;
+}
+
+inline uniform FVector3d operator/(const uniform FVector3d &A, const uniform double F)
+{
+	const uniform FVector3d FVec = {{F,F,F}};
+	return A / FVec;
+}
+
+inline uniform FVector3f operator/(const uniform FVector3f &A, const uniform float F)
+{
+	const uniform FVector3f FVec = {{F,F,F}};
 	return A / FVec;
 }
 
@@ -1613,9 +1794,15 @@ inline FVector SetVector(const FReal F)
 	return Result;
 }
 
-inline FVector SetVector(const FVector4 &Vec)
+inline FVector3d SetVector(const FVector4d &Vec)
 {
-	const FVector Result = {{Vec.V[0], Vec.V[1], Vec.V[2]}};
+	const FVector3d Result = {{Vec.V[0], Vec.V[1], Vec.V[2]}};
+	return Result;
+}
+
+inline FVector3f SetVector(const FVector4f &Vec)
+{
+	const FVector3f Result = {{Vec.V[0], Vec.V[1], Vec.V[2]}};
 	return Result;
 }
 
@@ -1631,6 +1818,30 @@ inline uniform FVector3d SetVector(const uniform FVector4d &Vec)
 	return Result;
 }
 
+inline FVector3d ConvertVector3fTo3d(const FVector3f &V)
+{
+	FVector3d Result;
+
+	for(uniform int i = 0; i < 3; i++)
+	{
+		Result.V[i] = (varying double)V.V[i];
+	}
+
+	return Result;
+}
+
+inline FVector3f ConvertVector3dTo3f(const FVector3d &V)
+{
+	FVector3f Result;
+
+	for(uniform int i = 0; i < 3; i++)
+	{
+		Result.V[i] = (varying float)V.V[i];
+	}
+
+	return Result;
+}
+
 inline uniform FVector3d ConvertVector3fTo3d(const uniform FVector3f &V)
 {
 	varying float S0;
@@ -1651,6 +1862,26 @@ inline uniform FVector3f ConvertVector3dTo3f(const uniform FVector3d &V)
 	return *((uniform FVector3f *uniform)&Result);
 }
 
+inline FVector ConvertVector3dTo3Native(const FVector3d &V)
+{
+	return V;
+}
+
+inline FVector ConvertVector3fTo3Native(const FVector3f &V)
+{
+	return ConvertVector3fTo3d(V);
+}
+
+inline FVector3d ConvertVector3NativeTo3d(const FVector &V)
+{
+	return V;
+}
+
+inline FVector3f ConvertVector3NativeTo3f(const FVector &V)
+{
+	return ConvertVector3dTo3f(V);
+}
+
 inline uniform FVector ConvertVector3dTo3Native(const uniform FVector3d &V)
 {
 	return V;
@@ -1934,18 +2165,27 @@ inline FIntVector SetIntVector(const FVector &A)
 
 inline uniform bool VectorIsAnyNearlyZero(const uniform FVector &V, const uniform FReal Tolerance)
 {
-	varying FReal S0 = FLT_MAX;
+	varying FReal S0 = REAL_MAX;
 
 	*((uniform FVector *uniform)&S0) = *((uniform FVector *uniform)&V);
 
 	return any(abs(S0) <= Tolerance);
 }
 
-inline uniform bool VectorIsAllNearlyZero(const uniform FVector &V, const uniform FReal Tolerance)
+inline uniform bool VectorIsAllNearlyZero(const uniform FVector3d &V, const uniform double Tolerance)
 {
-	varying FReal S0 = FLT_MIN;
+	varying double S0 = DBL_MIN;
 
-	*((uniform FVector *uniform)&S0) = *((uniform FVector *uniform)&V);
+	*((uniform FVector3d *uniform)&S0) = *((uniform FVector3d *uniform)&V);
+
+	return all(abs(S0) <= Tolerance);
+}
+
+inline uniform bool VectorIsAllNearlyZero(const uniform FVector3f &V, const uniform float Tolerance)
+{
+	varying float S0 = FLT_MIN;
+
+	*((uniform FVector3f *uniform)&S0) = *((uniform FVector3f *uniform)&V);
 
 	return all(abs(S0) <= Tolerance);
 }
@@ -1984,7 +2224,12 @@ inline FVector VectorClamp(const FVector &A, const uniform FVector &B, const uni
 		clamp(A.V[2], B.V[2], C.V[2]));
 }
 
-inline FVector4 VectorAdd(const FVector4 &A, const FVector4 &B)
+inline FVector4d VectorAdd(const FVector4d &A, const FVector4d &B)
+{
+	return A + B;
+}
+
+inline FVector4f VectorAdd(const FVector4f &A, const FVector4f &B)
 {
 	return A + B;
 }
@@ -1994,7 +2239,12 @@ inline FVector4 VectorAdd(const FVector4 &A, const uniform FVector4 &B)
 	return A + B;
 }
 
-inline uniform FVector4 VectorAdd(const uniform FVector4 &A, const uniform FVector4 &B)
+inline uniform FVector4d VectorAdd(const uniform FVector4d &A, const uniform FVector4d &B)
+{
+	return A + B;
+}
+
+inline uniform FVector4f VectorAdd(const uniform FVector4f &A, const uniform FVector4f &B)
 {
 	return A + B;
 }
@@ -2004,7 +2254,12 @@ inline uniform FVector8 VectorAdd(const uniform FVector8 &A, const uniform FVect
 	return A + B;
 }
 
-inline FVector4 VectorSubtract(const FVector4 &A, const FVector4 &B)
+inline FVector4d VectorSubtract(const FVector4d &A, const FVector4d &B)
+{
+	return A - B;
+}
+
+inline FVector4f VectorSubtract(const FVector4f &A, const FVector4f &B)
 {
 	return A - B;
 }
@@ -2014,12 +2269,22 @@ inline FVector4 VectorSubtract(const uniform FVector4 &A, const FVector4 &B)
 	return A - B;
 }
 
-inline uniform FVector4 VectorSubtract(const uniform FVector4 &A, const uniform FVector4 &B)
+inline uniform FVector4d VectorSubtract(const uniform FVector4d &A, const uniform FVector4d &B)
+{
+	return A - B;
+}
+
+inline uniform FVector4f VectorSubtract(const uniform FVector4f &A, const uniform FVector4f &B)
 {
 	return A - B;
 }
 
-inline FVector4 VectorMultiply(const FVector4 &A, const FVector4 &B)
+inline FVector4d VectorMultiply(const FVector4d &A, const FVector4d &B)
+{
+	return A * B;
+}
+
+inline FVector4f VectorMultiply(const FVector4f &A, const FVector4f &B)
 {
 	return A * B;
 }
@@ -2034,7 +2299,12 @@ inline FVector4 VectorMultiply(const uniform FVector4 &A, const FVector4 &B)
 	return A * B;
 }
 
-inline uniform FVector4 VectorMultiply(const uniform FVector4 &A, const uniform FVector4 &B)
+inline uniform FVector4d VectorMultiply(const uniform FVector4d &A, const uniform FVector4d &B)
+{
+	return A * B;
+}
+
+inline uniform FVector4f VectorMultiply(const uniform FVector4f &A, const uniform FVector4f &B)
 {
 	return A * B;
 }
@@ -2049,19 +2319,41 @@ inline uniform FVector8 VectorMultiply(const uniform FVector8 &A, const uniform
 	return A * B;
 }
 
-inline FVector4 VectorDivide(const FVector4 &A, const FVector4 &B)
+inline FVector4d VectorDivide(const FVector4d &A, const FVector4d &B)
 {
 	return A / B;
 }
 
-inline uniform FVector4 VectorDivide(const uniform FVector4 &A, const uniform FVector4 &B)
+inline FVector4f VectorDivide(const FVector4f &A, const FVector4f &B)
 {
 	return A / B;
 }
 
-inline FVector4 VectorMultiplyAdd(const FVector4& A, const FVector4 &B, const FVector4 &C)
+inline uniform FVector4d VectorDivide(const uniform FVector4d &A, const uniform FVector4d &B)
 {
-	FVector4 Result;
+	return A / B;
+}
+
+inline uniform FVector4f VectorDivide(const uniform FVector4f &A, const uniform FVector4f &B)
+{
+	return A / B;
+}
+
+inline FVector4d VectorMultiplyAdd(const FVector4d &A, const FVector4d &B, const FVector4d &C)
+{
+	FVector4d Result;
+
+	Result.V[0] = A.V[0] * B.V[0] + C.V[0];
+	Result.V[1] = A.V[1] * B.V[1] + C.V[1];
+	Result.V[2] = A.V[2] * B.V[2] + C.V[2];
+	Result.V[3] = A.V[3] * B.V[3] + C.V[3];
+
+	return Result;
+}
+
+inline FVector4f VectorMultiplyAdd(const FVector4f &A, const FVector4f &B, const FVector4f &C)
+{
+	FVector4f Result;
 
 	Result.V[0] = A.V[0] * B.V[0] + C.V[0];
 	Result.V[1] = A.V[1] * B.V[1] + C.V[1];
@@ -2107,16 +2399,28 @@ inline FVector4 VectorMultiplyAdd(const FReal F, const FVector4 &B, const FVecto
 	return Result;
 }
 
-inline uniform FVector4 VectorMultiplyAdd(const uniform FVector4& A, const uniform FVector4 &B, const uniform FVector4 &C)
+inline uniform FVector4d VectorMultiplyAdd(const uniform FVector4d &A, const uniform FVector4d &B, const uniform FVector4d &C)
 {
-	varying FReal S0, S1, S2, Result;
-	*((uniform FVector4 *uniform)&S0) = *((uniform FVector4 *uniform)&A);
-	*((uniform FVector4 *uniform)&S1) = *((uniform FVector4 *uniform)&B);
-	*((uniform FVector4 *uniform)&S2) = *((uniform FVector4 *uniform)&C);
+	varying double S0, S1, S2, Result;
+	*((uniform FVector4d *uniform)&S0) = *((uniform FVector4d *uniform)&A);
+	*((uniform FVector4d *uniform)&S1) = *((uniform FVector4d *uniform)&B);
+	*((uniform FVector4d *uniform)&S2) = *((uniform FVector4d *uniform)&C);
 
 	Result = S0 * S1 + S2;
 
-	return *((uniform FVector4 *uniform)&Result);
+	return *((uniform FVector4d *uniform)&Result);
+}
+
+inline uniform FVector4f VectorMultiplyAdd(const uniform FVector4f &A, const uniform FVector4f &B, const uniform FVector4f &C)
+{
+	varying float S0, S1, S2, Result;
+	*((uniform FVector4f *uniform)&S0) = *((uniform FVector4f *uniform)&A);
+	*((uniform FVector4f *uniform)&S1) = *((uniform FVector4f *uniform)&B);
+	*((uniform FVector4f *uniform)&S2) = *((uniform FVector4f *uniform)&C);
+
+	Result = S0 * S1 + S2;
+
+	return *((uniform FVector4f *uniform)&Result);
 }
 
 inline FVector4 VectorMultiplyAdd(const FVector4& A, const uniform FReal B, const FVector4 &C)
@@ -2342,9 +2646,20 @@ inline uniform FVector4 VectorMax(const uniform FVector4& V1, const uniform FVec
 	return *((uniform FVector4 *uniform)&Result);
 }
 
-inline FVector VectorMax(const FVector& V1, const FVector& V2)
+inline FVector3d VectorMax(const FVector3d& V1, const FVector3d& V2)
 {
-	FVector Result;
+	FVector3d Result;
+
+	Result.V[0] = max(V1.V[0], V2.V[0]);
+	Result.V[1] = max(V1.V[1], V2.V[1]);
+	Result.V[2] = max(V1.V[2], V2.V[2]);
+
+	return Result;
+}
+
+inline FVector3f VectorMax(const FVector3f& V1, const FVector3f& V2)
+{
+	FVector3f Result;
 
 	Result.V[0] = max(V1.V[0], V2.V[0]);
 	Result.V[1] = max(V1.V[1], V2.V[1]);
@@ -2398,9 +2713,20 @@ inline uniform FVector4 VectorMin(const uniform FVector4& V1, const uniform FVec
 	return *((uniform FVector4 *uniform)&Result);
 }
 
-inline FVector VectorMin(const FVector& V1, const FVector& V2)
+inline FVector3d VectorMin(const FVector3d& V1, const FVector3d& V2)
 {
-	FVector Result;
+	FVector3d Result;
+
+	Result.V[0] = min(V1.V[0], V2.V[0]);
+	Result.V[1] = min(V1.V[1], V2.V[1]);
+	Result.V[2] = min(V1.V[2], V2.V[2]);
+
+	return Result;
+}
+
+inline FVector3f VectorMin(const FVector3f& V1, const FVector3f& V2)
+{
+	FVector3f Result;
 
 	Result.V[0] = min(V1.V[0], V2.V[0]);
 	Result.V[1] = min(V1.V[1], V2.V[1]);
@@ -2680,12 +3006,22 @@ inline uniform FVector VectorReciprocalSqrtAccurate(const uniform FVector& V)
  * @param ElementIndex	Index (0-3) of the element to replicate
  * @return				VectorRegister( Vec[ElementIndex], Vec[ElementIndex], Vec[ElementIndex], Vec[ElementIndex] )
  */
-inline FVector4 VectorReplicate(const FVector4 &Vec, const uniform int ElementIndex)
+inline FVector4d VectorReplicate(const FVector4d &Vec, const uniform int ElementIndex)
 {
 	return SetVector4(Vec.V[ElementIndex], Vec.V[ElementIndex], Vec.V[ElementIndex], Vec.V[ElementIndex]);
 }
 
-inline uniform FVector4 VectorReplicate(const uniform FVector4 &Vec, const uniform int ElementIndex)
+inline FVector4f VectorReplicate(const FVector4f &Vec, const uniform int ElementIndex)
+{
+	return SetVector4(Vec.V[ElementIndex], Vec.V[ElementIndex], Vec.V[ElementIndex], Vec.V[ElementIndex]);
+}
+
+inline uniform FVector4d VectorReplicate(const uniform FVector4d &Vec, const uniform int ElementIndex)
+{
+	return SetVector4(Vec.V[ElementIndex], Vec.V[ElementIndex], Vec.V[ElementIndex], Vec.V[ElementIndex]);
+}
+
+inline uniform FVector4f VectorReplicate(const uniform FVector4f &Vec, const uniform int ElementIndex)
 {
 	return SetVector4(Vec.V[ElementIndex], Vec.V[ElementIndex], Vec.V[ElementIndex], Vec.V[ElementIndex]);
 }
@@ -2766,7 +3102,12 @@ inline FVector4 VectorShuffle(const FVector4 &Vec1, const FVector4 &Vec2, const
 	return SetVector4(Vec1.V[X], Vec1.V[Y], Vec2.V[Z], Vec2.V[W]);
 }
 
-inline uniform FVector4 VectorShuffle(const uniform FVector4 &Vec1, const uniform FVector4 &Vec2, const uniform int X, const uniform int Y, const uniform int Z, const uniform int W)
+inline uniform FVector4d VectorShuffle(const uniform FVector4d &Vec1, const uniform FVector4d &Vec2, const uniform int X, const uniform int Y, const uniform int Z, const uniform int W)
+{
+	return SetVector4(Vec1.V[X], Vec1.V[Y], Vec2.V[Z], Vec2.V[W]);
+}
+
+inline uniform FVector4f VectorShuffle(const uniform FVector4f &Vec1, const uniform FVector4f &Vec2, const uniform int X, const uniform int Y, const uniform int Z, const uniform int W)
 {
 	return SetVector4(Vec1.V[X], Vec1.V[Y], Vec2.V[Z], Vec2.V[W]);
 }
@@ -2778,13 +3119,23 @@ inline uniform FVector4 VectorShuffle(const uniform FVector4 &Vec1, const unifor
  * @param Vec2	2nd vector
  * @return		cross(Vec1.xyz, Vec2.xyz). W is set to 0.
  */
-inline FVector4 VectorCross(const FVector4& Vec1, const FVector4& Vec2)
+inline FVector4d VectorCross(const FVector4d& Vec1, const FVector4d& Vec2)
 {
-	FVector4 Tmp0 = VectorSwizzle(Vec2, 1,2,0,3);
-	FVector4 Tmp1 = VectorSwizzle(Vec1, 1,2,0,3);
+	FVector4d Tmp0 = VectorSwizzle(Vec2, 1,2,0,3);
+	FVector4d Tmp1 = VectorSwizzle(Vec1, 1,2,0,3);
 	Tmp0 = Tmp0 * Vec1;
 	Tmp1 = Tmp1 * Vec2;
-	FVector4 Tmp2 = Tmp0 - Tmp1;
+	FVector4d Tmp2 = Tmp0 - Tmp1;
+	return VectorSwizzle(Tmp2, 1,2,0,3);
+}
+
+inline FVector4f VectorCross(const FVector4f& Vec1, const FVector4f& Vec2)
+{
+	FVector4f Tmp0 = VectorSwizzle(Vec2, 1,2,0,3);
+	FVector4f Tmp1 = VectorSwizzle(Vec1, 1,2,0,3);
+	Tmp0 = Tmp0 * Vec1;
+	Tmp1 = Tmp1 * Vec2;
+	FVector4f Tmp2 = Tmp0 - Tmp1;
 	return VectorSwizzle(Tmp2, 1,2,0,3);
 }
 
@@ -2798,23 +3149,43 @@ inline FVector4 VectorCross(const uniform FVector4& Vec1, const FVector4& Vec2)
 	return VectorSwizzle(Tmp2, 1,2,0,3);
 }
 
-inline uniform FVector4 VectorCross(const uniform FVector4& Vec1, const uniform FVector4& Vec2)
+inline uniform FVector4d VectorCross(const uniform FVector4d& Vec1, const uniform FVector4d& Vec2)
 {
-	uniform FVector4 Tmp0 = VectorSwizzle(Vec2, 1,2,0,3);
-	uniform FVector4 Tmp1 = VectorSwizzle(Vec1, 1,2,0,3);
+	uniform FVector4d Tmp0 = VectorSwizzle(Vec2, 1,2,0,3);
+	uniform FVector4d Tmp1 = VectorSwizzle(Vec1, 1,2,0,3);
 	Tmp0 = Tmp0 * Vec1;
 	Tmp1 = Tmp1 * Vec2;
-	uniform FVector4 Tmp2 = Tmp0 - Tmp1;
+	uniform FVector4d Tmp2 = Tmp0 - Tmp1;
 	return VectorSwizzle(Tmp2, 1,2,0,3);
 }
 
-inline FVector VectorCross(const FVector& Vec1, const FVector& Vec2)
+inline uniform FVector4f VectorCross(const uniform FVector4f& Vec1, const uniform FVector4f& Vec2)
 {
-	FVector Tmp0 = VectorSwizzle(Vec2, 1,2,0);
-	FVector Tmp1 = VectorSwizzle(Vec1, 1,2,0);
+	uniform FVector4f Tmp0 = VectorSwizzle(Vec2, 1,2,0,3);
+	uniform FVector4f Tmp1 = VectorSwizzle(Vec1, 1,2,0,3);
 	Tmp0 = Tmp0 * Vec1;
 	Tmp1 = Tmp1 * Vec2;
-	const FVector Tmp2 = Tmp0 - Tmp1;
+	uniform FVector4f Tmp2 = Tmp0 - Tmp1;
+	return VectorSwizzle(Tmp2, 1,2,0,3);
+}
+
+inline FVector3d VectorCross(const FVector3d& Vec1, const FVector3d& Vec2)
+{
+	FVector3d Tmp0 = VectorSwizzle(Vec2, 1,2,0);
+	FVector3d Tmp1 = VectorSwizzle(Vec1, 1,2,0);
+	Tmp0 = Tmp0 * Vec1;
+	Tmp1 = Tmp1 * Vec2;
+	const FVector3d Tmp2 = Tmp0 - Tmp1;
+	return VectorSwizzle(Tmp2, 1,2,0);
+}
+
+inline FVector3f VectorCross(const FVector3f& Vec1, const FVector3f& Vec2)
+{
+	FVector3f Tmp0 = VectorSwizzle(Vec2, 1,2,0);
+	FVector3f Tmp1 = VectorSwizzle(Vec1, 1,2,0);
+	Tmp0 = Tmp0 * Vec1;
+	Tmp1 = Tmp1 * Vec2;
+	const FVector3f Tmp2 = Tmp0 - Tmp1;
 	return VectorSwizzle(Tmp2, 1,2,0);
 }
 
@@ -2936,9 +3307,15 @@ inline uniform FReal VectorDot( const uniform FVector4& Vec1, const uniform FVec
 	return Temp1.V[0] + Temp2.V[0]; // (X*X + Z*Z + Y*Y + W*W)
 }
 
-inline uniform FReal VectorDot( const uniform FVector &Vec1, const uniform FVector &Vec2 )
+inline uniform double VectorDot( const uniform FVector3d &Vec1, const uniform FVector3d &Vec2 )
+{
+	const uniform FVector3d Temp = Vec1 * Vec2;
+	return Temp.V[0] + Temp.V[1] + Temp.V[2];
+}
+
+inline uniform float VectorDot( const uniform FVector3f &Vec1, const uniform FVector3f &Vec2 )
 {
-	const uniform FVector Temp = Vec1 * Vec2;
+	const uniform FVector3f Temp = Vec1 * Vec2;
 	return Temp.V[0] + Temp.V[1] + Temp.V[2];
 }
 
@@ -3113,9 +3490,9 @@ inline uniform FVector VectorSelect(const uniform FVector& Mask, const uniform F
 	return VectorBitwiseXor(Vec2, VectorBitwiseAnd(Mask, VectorBitwiseXor(Vec1, Vec2)));
 }
 
-inline FVector VectorSelect(varying bool Mask, const varying FVector &Vec1, const varying FVector &Vec2)
+inline FVector3d VectorSelect(varying bool Mask, const varying FVector3d &Vec1, const varying FVector3d &Vec2)
 {
-	FVector Result;
+	FVector3d Result;
 
 	Result.V[0] = select(Mask, Vec1.V[0], Vec2.V[0]);
 	Result.V[1] = select(Mask, Vec1.V[1], Vec2.V[1]);
@@ -3124,9 +3501,9 @@ inline FVector VectorSelect(varying bool Mask, const varying FVector &Vec1, cons
 	return Result;
 }
 
-inline FVector VectorSelect(varying bool Mask, const uniform FVector &Vec1, const varying FVector &Vec2)
+inline FVector3f VectorSelect(varying bool Mask, const varying FVector3f &Vec1, const varying FVector3f &Vec2)
 {
-	FVector Result;
+	FVector3f Result;
 
 	Result.V[0] = select(Mask, Vec1.V[0], Vec2.V[0]);
 	Result.V[1] = select(Mask, Vec1.V[1], Vec2.V[1]);
@@ -3135,9 +3512,42 @@ inline FVector VectorSelect(varying bool Mask, const uniform FVector &Vec1, cons
 	return Result;
 }
 
-inline FVector VectorSelect(varying bool Mask, const varying FVector &Vec1, const uniform FVector &Vec2)
+inline FVector3d VectorSelect(varying bool Mask, const uniform FVector3d &Vec1, const varying FVector3d &Vec2)
 {
-	FVector Result;
+	FVector3d Result;
+
+	Result.V[0] = select(Mask, Vec1.V[0], Vec2.V[0]);
+	Result.V[1] = select(Mask, Vec1.V[1], Vec2.V[1]);
+	Result.V[2] = select(Mask, Vec1.V[2], Vec2.V[2]);
+
+	return Result;
+}
+
+inline FVector3f VectorSelect(varying bool Mask, const uniform FVector3f &Vec1, const varying FVector3f &Vec2)
+{
+	FVector3f Result;
+
+	Result.V[0] = select(Mask, Vec1.V[0], Vec2.V[0]);
+	Result.V[1] = select(Mask, Vec1.V[1], Vec2.V[1]);
+	Result.V[2] = select(Mask, Vec1.V[2], Vec2.V[2]);
+
+	return Result;
+}
+
+inline FVector3d VectorSelect(varying bool Mask, const varying FVector3d &Vec1, const uniform FVector3d &Vec2)
+{
+	FVector3d Result;
+
+	Result.V[0] = select(Mask, Vec1.V[0], Vec2.V[0]);
+	Result.V[1] = select(Mask, Vec1.V[1], Vec2.V[1]);
+	Result.V[2] = select(Mask, Vec1.V[2], Vec2.V[2]);
+
+	return Result;
+}
+
+inline FVector3f VectorSelect(varying bool Mask, const varying FVector3f &Vec1, const uniform FVector3f &Vec2)
+{
+	FVector3f Result;
 
 	Result.V[0] = select(Mask, Vec1.V[0], Vec2.V[0]);
 	Result.V[1] = select(Mask, Vec1.V[1], Vec2.V[1]);
@@ -3146,9 +3556,20 @@ inline FVector VectorSelect(varying bool Mask, const varying FVector &Vec1, cons
 	return Result;
 }
 
-inline uniform FVector VectorSelect(uniform bool Mask, const uniform FVector &Vec1, const uniform FVector &Vec2)
+inline uniform FVector3d VectorSelect(uniform bool Mask, const uniform FVector3d &Vec1, const uniform FVector3d &Vec2)
 {
-	uniform FVector Result;
+	uniform FVector3d Result;
+
+	Result.V[0] = select(Mask, Vec1.V[0], Vec2.V[0]);
+	Result.V[1] = select(Mask, Vec1.V[1], Vec2.V[1]);
+	Result.V[2] = select(Mask, Vec1.V[2], Vec2.V[2]);
+
+	return Result;
+}
+
+inline uniform FVector3f VectorSelect(uniform bool Mask, const uniform FVector3f &Vec1, const uniform FVector3f &Vec2)
+{
+	uniform FVector3f Result;
 
 	Result.V[0] = select(Mask, Vec1.V[0], Vec2.V[0]);
 	Result.V[1] = select(Mask, Vec1.V[1], Vec2.V[1]);
@@ -3341,22 +3762,42 @@ inline uniform FReal VectorSizeSquared(const uniform FVector4 &A)
 	return VectorDot(A, A);
 }
 
-inline FReal VectorSizeSquared(const FVector &A)
+inline double VectorSizeSquared(const FVector3d &A)
+{
+	return VectorDot(A, A);
+}
+
+inline float VectorSizeSquared(const FVector3f &A)
+{
+	return VectorDot(A, A);
+}
+
+inline uniform double VectorSizeSquared(const uniform FVector3d &A)
 {
 	return VectorDot(A, A);
 }
 
-inline uniform FReal VectorSizeSquared(const uniform FVector &A)
+inline uniform float VectorSizeSquared(const uniform FVector3f &A)
 {
 	return VectorDot(A, A);
 }
 
-inline uniform FReal VectorSize(const uniform FVector& A)
+inline uniform double VectorSize(const uniform FVector3d& A)
 {
 	return sqrt(VectorSizeSquared(A));
 }
 
-inline FReal VectorSize(const varying FVector &A)
+inline uniform float VectorSize(const uniform FVector3f& A)
+{
+	return sqrt(VectorSizeSquared(A));
+}
+
+inline double VectorSize(const varying FVector3d &A)
+{
+	return sqrt(VectorDot(A, A));
+}
+
+inline float VectorSize(const varying FVector3f &A)
 {
 	return sqrt(VectorDot(A, A));
 }
@@ -3399,7 +3840,12 @@ inline FVector4 VectorClamp(const FVector4 &A, const uniform FVector4 &B, const
 		clamp(A.V[3], B.V[3], C.V[3]));
 }
 
-inline uniform FVector VectorReduceAdd(const FVector& A)
+inline uniform FVector3d VectorReduceAdd(const FVector3d& A)
+{
+	return SetVector(reduce_add(A.V[0]), reduce_add(A.V[1]), reduce_add(A.V[2]));
+}
+
+inline uniform FVector3f VectorReduceAdd(const FVector3f& A)
 {
 	return SetVector(reduce_add(A.V[0]), reduce_add(A.V[1]), reduce_add(A.V[2]));
 }
@@ -3409,22 +3855,40 @@ inline uniform FVector4 VectorReduceAdd(const FVector4& A)
 	return SetVector4(reduce_add(A.V[0]), reduce_add(A.V[1]), reduce_add(A.V[2]), reduce_add(A.V[3]));
 }
 
-inline uniform FVector VectorReduceMin(const FVector& A)
+inline uniform FVector3d VectorReduceMin(const FVector3d& A)
 {
 	return SetVector(reduce_min(A.V[0]), reduce_min(A.V[1]), reduce_min(A.V[2]));
 }
 
-inline uniform FVector VectorReduceMax(const FVector& A)
+inline uniform FVector3f VectorReduceMin(const FVector3f& A)
+{
+	return SetVector(reduce_min(A.V[0]), reduce_min(A.V[1]), reduce_min(A.V[2]));
+}
+
+inline uniform FVector3d VectorReduceMax(const FVector3d& A)
 {
 	return SetVector(reduce_max(A.V[0]), reduce_max(A.V[1]), reduce_max(A.V[2]));
 }
 
-inline FVector VectorGetSafeNormal(const FVector &V)
+inline uniform FVector3f VectorReduceMax(const FVector3f& A)
+{
+	return SetVector(reduce_max(A.V[0]), reduce_max(A.V[1]), reduce_max(A.V[2]));
+}
+
+inline FVector3d VectorGetSafeNormal(const FVector3d &V)
+{
+	//We want N / ||N|| and to avoid inf
+	//So we want N / ||N|| < 1 / eps => N eps < ||N||, but this is clearly true for all eps < 1 and N > 0
+	double SizeSqr = VectorSizeSquared(V);
+	return VectorSelect(SizeSqr <= DBL_MIN, DoubleForwardVector, V / sqrt(SizeSqr));
+}
+
+inline FVector3f VectorGetSafeNormal(const FVector3f &V)
 {
 	//We want N / ||N|| and to avoid inf
 	//So we want N / ||N|| < 1 / eps => N eps < ||N||, but this is clearly true for all eps < 1 and N > 0
-	FReal SizeSqr = VectorSizeSquared(V);
-	return VectorSelect(SizeSqr <= FLT_MIN, ForwardVector, V / sqrt(SizeSqr));
+	float SizeSqr = VectorSizeSquared(V);
+	return VectorSelect(SizeSqr <= FLT_MIN, FloatForwardVector, V / sqrt(SizeSqr));
 }
 
 #endif
diff --git a/Engine/Source/Runtime/Experimental/Chaos/Private/Chaos/GeometryParticles.cpp b/Engine/Source/Runtime/Experimental/Chaos/Private/Chaos/GeometryParticles.cpp
--- a/Engine/Source/Runtime/Experimental/Chaos/Private/Chaos/GeometryParticles.cpp
+++ b/Engine/Source/Runtime/Experimental/Chaos/Private/Chaos/GeometryParticles.cpp
@@ -167,14 +167,20 @@ namespace Chaos
 
 
 	template <typename T, int d, EGeometryParticlesSimType SimType>
-	void TGeometryParticlesImp<T, d, SimType>::SetHandle(int32 Index, TGeometryParticleHandle<T, d>* Handle)
+	void TGeometryParticlesImp<T, d, SimType>::SetHandle(int32 Index, FGeometryParticleHandle* Handle)
 	{
 		Handle->SetSOALowLevel(this);
 		MGeometryParticleHandle[Index] = AsAlwaysSerializable(Handle);
 	}
 
 	template <>
-	void TGeometryParticlesImp<FReal, 3, EGeometryParticlesSimType::Other>::SetHandle(int32 Index, TGeometryParticleHandle<FReal, 3>* Handle)
+	void TGeometryParticlesImp<FRealSingle, 3, EGeometryParticlesSimType::Other>::SetHandle(int32 Index, FGeometryParticleHandle* Handle)
+	{
+		check(false);  // TODO: Implement EGeometryParticlesSimType::Other (cloth) particle serialization
+	}
+
+	template <>
+	void TGeometryParticlesImp<FRealDouble, 3, EGeometryParticlesSimType::Other>::SetHandle(int32 Index, FGeometryParticleHandle* Handle)
 	{
 		check(false);  // TODO: Implement EGeometryParticlesSimType::Other (cloth) particle serialization
 	}
@@ -196,7 +202,14 @@ namespace Chaos
 	}
 	
 	template<>
-	TGeometryParticlesImp<FReal, 3, EGeometryParticlesSimType::Other>* TGeometryParticlesImp<FReal, 3, EGeometryParticlesSimType::Other>::SerializationFactory(FChaosArchive& Ar, TGeometryParticlesImp<FReal, 3, EGeometryParticlesSimType::Other>* Particles)
+	TGeometryParticlesImp<FRealSingle, 3, EGeometryParticlesSimType::Other>* TGeometryParticlesImp<FRealSingle, 3, EGeometryParticlesSimType::Other>::SerializationFactory(FChaosArchive& Ar, TGeometryParticlesImp<FRealSingle, 3, EGeometryParticlesSimType::Other>* Particles)
+	{
+		check(false);  // TODO: Implement EGeometryParticlesSimType::Other (cloth) particle serialization
+		return nullptr;
+	}
+
+	template<>
+	TGeometryParticlesImp<FRealDouble, 3, EGeometryParticlesSimType::Other>* TGeometryParticlesImp<FRealDouble, 3, EGeometryParticlesSimType::Other>::SerializationFactory(FChaosArchive& Ar, TGeometryParticlesImp<FRealDouble, 3, EGeometryParticlesSimType::Other>* Particles)
 	{
 		check(false);  // TODO: Implement EGeometryParticlesSimType::Other (cloth) particle serialization
 		return nullptr;
@@ -295,5 +308,6 @@ namespace Chaos
 
 	
 	template class TGeometryParticlesImp<FReal, 3, EGeometryParticlesSimType::RigidBodySim>;
-	template class TGeometryParticlesImp<FReal, 3, EGeometryParticlesSimType::Other>;
+	template class TGeometryParticlesImp<FRealSingle, 3, EGeometryParticlesSimType::Other>;
+	template class TGeometryParticlesImp<FRealDouble, 3, EGeometryParticlesSimType::Other>;
 }
