Core Math: Augment TBox2<T> class with methods available in the TBox<T> class.

Specifically adds:
 * non-uniform expansion ExpandBy variant method that takes a TVector2<T> instead of a uniform scalar.
 * MoveTo method for translating the 2D box center to specified location

#rnx

diff --git a/Engine/Source/Runtime/Core/Public/Math/Box2D.h b/Engine/Source/Runtime/Core/Public/Math/Box2D.h
--- a/Engine/Source/Runtime/Core/Public/Math/Box2D.h
+++ b/Engine/Source/Runtime/Core/Public/Math/Box2D.h
@@ -199,6 +199,17 @@ public:
 		return TBox2<T>(Min - TVector2<T>(W, W), Max + TVector2<T>(W, W));
 	}
 
+	/**
+     * Returns a box of increased size.
+     *
+     * @param V The size to increase the volume by.
+     * @return A new bounding box.
+     */
+	TBox2<T> ExpandBy(const TVector2<T>& V) const
+	{
+		return TBox2<T>(Min - V, Max + V);
+	}
+
 	/**
 	 * Gets the box area.
 	 *
@@ -324,6 +335,18 @@ public:
 		return TBox2<T>(Min + Offset, Max + Offset);
 	}
 
+	/**
+	 * Returns a box with its center moved to the new destination.
+	 *
+	 * @param Destination The destination point to move center of box to.
+	 * @return A new bounding box.
+	 */
+	TBox2<T> MoveTo(const TVector2<T>& Destination) const
+	{
+		const TVector2<T> Offset = Destination - GetCenter();
+		return TBox2<T>(Min + Offset, Max + Offset);
+	}
+
 	/**
 	 * Get a textual representation of this box.
 	 *
