Fix !DstMatrix->ContainsNaN ensure failure when LocalToWorld transform has large scaling values

diff --git a/Engine/Source/Runtime/Core/Private/Misc/LargeWorldRenderPosition.cpp b/Engine/Source/Runtime/Core/Private/Misc/LargeWorldRenderPosition.cpp
--- a/Engine/Source/Runtime/Core/Private/Misc/LargeWorldRenderPosition.cpp
+++ b/Engine/Source/Runtime/Core/Private/Misc/LargeWorldRenderPosition.cpp
@@ -39,7 +39,7 @@ FVector3f FLargeWorldRenderScalar::GetTileFor(FVector InPosition)
 	return LWCTile;
 }
 
-FMatrix44f FLargeWorldRenderScalar::SafeCastMatrix(const FMatrix& Matrix)
+FMatrix CheckMatrixInTileOffsetRange(const FMatrix& Matrix)
 {
 	const double OriginMax = UE_LWC_RENDER_MAX_OFFSET;
 
@@ -51,21 +51,40 @@ FMatrix44f FLargeWorldRenderScalar::SafeCastMatrix(const FMatrix& Matrix)
 	{
 		ensure(false);
 	}
+	return Matrix;
+}
 
-	return FMatrix44f(Matrix);
+FMatrix44f FLargeWorldRenderScalar::SafeCastMatrix(const FMatrix& Matrix)
+{
+	return FMatrix44f(CheckMatrixInTileOffsetRange(Matrix));
 }
 
 FMatrix44f FLargeWorldRenderScalar::MakeToRelativeWorldMatrix(const FVector Origin, const FMatrix& ToWorld)
 {
-	return SafeCastMatrix(ToWorld * FTranslationMatrix(-Origin));
+	return FMatrix44f(MakeToRelativeWorldMatrixDouble(Origin, ToWorld));
+}
+
+FMatrix FLargeWorldRenderScalar::MakeToRelativeWorldMatrixDouble(const FVector Origin, const FMatrix& ToWorld)
+{
+	return CheckMatrixInTileOffsetRange(ToWorld * FTranslationMatrix(-Origin));
 }
 
 FMatrix44f FLargeWorldRenderScalar::MakeFromRelativeWorldMatrix(const FVector Origin, const FMatrix& FromWorld)
 {
-	return SafeCastMatrix(FTranslationMatrix(Origin) * FromWorld);
+	return FMatrix44f(MakeFromRelativeWorldMatrixDouble(Origin, FromWorld));
+}
+
+FMatrix FLargeWorldRenderScalar::MakeFromRelativeWorldMatrixDouble(const FVector Origin, const FMatrix& FromWorld)
+{
+	return CheckMatrixInTileOffsetRange(FTranslationMatrix(Origin) * FromWorld);
 }
 
 FMatrix44f FLargeWorldRenderScalar::MakeClampedToRelativeWorldMatrix(const FVector Origin, const FMatrix& ToWorld)
+{
+	return FMatrix44f(MakeClampedToRelativeWorldMatrixDouble(Origin, ToWorld));
+}
+
+FMatrix FLargeWorldRenderScalar::MakeClampedToRelativeWorldMatrixDouble(const FVector Origin, const FMatrix& ToWorld)
 {
 	const double OriginMax = UE_LWC_RENDER_MAX_OFFSET;
 
@@ -78,7 +97,7 @@ FMatrix44f FLargeWorldRenderScalar::MakeClampedToRelativeWorldMatrix(const FVect
 
 	FMatrix ClampedToRelativeWorld(ToWorld);
 	ClampedToRelativeWorld.SetOrigin(ClampedRelativeOrigin);
-	return FMatrix44f(ClampedToRelativeWorld);
+	return ClampedToRelativeWorld;
 }
 
 void FLargeWorldRenderScalar::Validate(double InAbsolute)
diff --git a/Engine/Source/Runtime/Core/Public/Misc/LargeWorldRenderPosition.h b/Engine/Source/Runtime/Core/Public/Misc/LargeWorldRenderPosition.h
--- a/Engine/Source/Runtime/Core/Public/Misc/LargeWorldRenderPosition.h
+++ b/Engine/Source/Runtime/Core/Public/Misc/LargeWorldRenderPosition.h
@@ -34,8 +34,11 @@ public:
 
 	CORE_API static FMatrix44f SafeCastMatrix(const FMatrix& Matrix);
 	CORE_API static FMatrix44f MakeToRelativeWorldMatrix(const FVector Origin, const FMatrix& ToWorld);
+	CORE_API static FMatrix    MakeToRelativeWorldMatrixDouble(const FVector Origin, const FMatrix& ToWorld);
 	CORE_API static FMatrix44f MakeFromRelativeWorldMatrix(const FVector Origin, const FMatrix& FromWorld);
+	CORE_API static FMatrix    MakeFromRelativeWorldMatrixDouble(const FVector Origin, const FMatrix& FromWorld);
 	CORE_API static FMatrix44f MakeClampedToRelativeWorldMatrix(const FVector Origin, const FMatrix& ToWorld);
+	CORE_API static FMatrix    MakeClampedToRelativeWorldMatrixDouble(const FVector Origin, const FMatrix& ToWorld);
 
 	float GetTile() const { return (float)Tile; }
 	float GetOffset() const { return (float)Offset; }
diff --git a/Engine/Source/Runtime/Engine/Public/PrimitiveUniformShaderParameters.h b/Engine/Source/Runtime/Engine/Public/PrimitiveUniformShaderParameters.h
--- a/Engine/Source/Runtime/Engine/Public/PrimitiveUniformShaderParameters.h
+++ b/Engine/Source/Runtime/Engine/Public/PrimitiveUniformShaderParameters.h
@@ -291,8 +291,10 @@ public:
 		Parameters.TilePosition = AbsoluteWorldPosition.GetTile();
 
 		{
-			Parameters.LocalToRelativeWorld = FLargeWorldRenderScalar::MakeToRelativeWorldMatrix(TilePositionOffset, AbsoluteLocalToWorld);
-			Parameters.RelativeWorldToLocal = Parameters.LocalToRelativeWorld.Inverse();
+			// Inverse on FMatrix44f can generate NaNs if the source matrix contains large scaling, so do it in double precision.
+			FMatrix LocalToRelativeWorld = FLargeWorldRenderScalar::MakeToRelativeWorldMatrixDouble(TilePositionOffset, AbsoluteLocalToWorld);
+			Parameters.LocalToRelativeWorld = FMatrix44f(LocalToRelativeWorld);
+			Parameters.RelativeWorldToLocal = FMatrix44f(LocalToRelativeWorld.Inverse());
 		}
 
 		Parameters.ActorRelativeWorldPosition = FVector3f(AbsoluteActorWorldPosition - TilePositionOffset);	//LWC_TODO: Precision loss
@@ -301,8 +303,10 @@ public:
 
 		if (bHasPreviousLocalToWorld)
 		{
-			Parameters.PreviousLocalToRelativeWorld = FLargeWorldRenderScalar::MakeClampedToRelativeWorldMatrix(TilePositionOffset, AbsolutePreviousLocalToWorld);
-			Parameters.PreviousRelativeWorldToLocal = Parameters.PreviousLocalToRelativeWorld.Inverse();
+			// Inverse on FMatrix44f can generate NaNs if the source matrix contains large scaling, so do it in double precision.
+			FMatrix PrevLocalToRelativeWorld = FLargeWorldRenderScalar::MakeClampedToRelativeWorldMatrixDouble(TilePositionOffset, AbsolutePreviousLocalToWorld);
+			Parameters.PreviousLocalToRelativeWorld = FMatrix44f(PrevLocalToRelativeWorld);
+			Parameters.PreviousRelativeWorldToLocal = FMatrix44f(PrevLocalToRelativeWorld.Inverse());
 		}
 		else
 		{
