Changed Blake3Hash/IoHash memory view constructors to static functions

#rnx

diff --git a/Engine/Source/Runtime/Core/Public/Hash/Blake3.h b/Engine/Source/Runtime/Core/Public/Hash/Blake3.h
--- a/Engine/Source/Runtime/Core/Public/Hash/Blake3.h
+++ b/Engine/Source/Runtime/Core/Public/Hash/Blake3.h
@@ -27,14 +27,14 @@ public:
 	/** Construct a hash from an array of 32 bytes. */
 	inline explicit FBlake3Hash(const ByteArray& Hash);
 
-	/** Construct a hash from a view of 32 bytes. */
-	inline explicit FBlake3Hash(FMemoryView Hash);
-
 	/** Construct a hash from a 64-character hex string. */
 	inline explicit FBlake3Hash(FAnsiStringView HexHash);
 	inline explicit FBlake3Hash(FWideStringView HexHash);
 	inline explicit FBlake3Hash(FUtf8StringView HexHash);
 
+	/** Construct a hash from a view of 32 bytes. */
+	inline static FBlake3Hash FromView(FMemoryView Hash);
+
 	/** Reset this to a zero hash. */
 	inline void Reset() { *this = FBlake3Hash(); }
 
@@ -96,13 +96,6 @@ inline FBlake3Hash::FBlake3Hash(const ByteArray& InHash)
 	FMemory::Memcpy(Hash, InHash, sizeof(ByteArray));
 }
 
-inline FBlake3Hash::FBlake3Hash(const FMemoryView InHash)
-{
-	checkf(InHash.GetSize() == sizeof(ByteArray),
-		TEXT("FBlake3Hash cannot be constructed from a view of %" UINT64_FMT " bytes."), InHash.GetSize());
-	FMemory::Memcpy(Hash, InHash.GetData(), sizeof(ByteArray));
-}
-
 inline FBlake3Hash::FBlake3Hash(const FAnsiStringView HexHash)
 {
 	check(HexHash.Len() == sizeof(ByteArray) * 2);
@@ -121,6 +114,15 @@ inline FBlake3Hash::FBlake3Hash(const FUtf8StringView HexHash)
 	UE::String::HexToBytes(HexHash, Hash);
 }
 
+inline FBlake3Hash FBlake3Hash::FromView(const FMemoryView InHash)
+{
+	checkf(InHash.GetSize() == sizeof(ByteArray),
+		TEXT("FBlake3Hash cannot be constructed from a view of %" UINT64_FMT " bytes."), InHash.GetSize());
+	FBlake3Hash NewHash;
+	FMemory::Memcpy(NewHash.Hash, InHash.GetData(), sizeof(ByteArray));
+	return NewHash;
+}
+
 inline bool FBlake3Hash::IsZero() const
 {
 	using UInt32Array = uint32[8];
diff --git a/Engine/Source/Runtime/Core/Public/IO/IoHash.h b/Engine/Source/Runtime/Core/Public/IO/IoHash.h
--- a/Engine/Source/Runtime/Core/Public/IO/IoHash.h
+++ b/Engine/Source/Runtime/Core/Public/IO/IoHash.h
@@ -33,9 +33,6 @@ public:
 	/** Construct a hash from an array of 20 bytes. */
 	inline explicit FIoHash(const ByteArray& Hash);
 
-	/** Construct a hash from a view of 20 bytes. */
-	inline explicit FIoHash(FMemoryView Hash);
-
 	/** Construct a hash from a BLAKE3-256 hash. */
 	inline FIoHash(const FBlake3Hash& Hash);
 
@@ -44,6 +41,9 @@ public:
 	inline explicit FIoHash(FWideStringView HexHash);
 	inline explicit FIoHash(FUtf8StringView HexHash);
 
+	/** Construct a hash from a view of 20 bytes. */
+	inline static FIoHash FromView(FMemoryView Hash);
+
 	/** Reset this to a zero hash. */
 	inline void Reset() { *this = FIoHash(); }
 
@@ -73,13 +73,6 @@ inline FIoHash::FIoHash(const ByteArray& InHash)
 	FMemory::Memcpy(Hash, InHash, sizeof(ByteArray));
 }
 
-inline FIoHash::FIoHash(const FMemoryView InHash)
-{
-	checkf(InHash.GetSize() == sizeof(ByteArray),
-		TEXT("FIoHash cannot be constructed from a view of %" UINT64_FMT " bytes."), InHash.GetSize());
-	FMemory::Memcpy(Hash, InHash.GetData(), sizeof(ByteArray));
-}
-
 inline FIoHash::FIoHash(const FBlake3Hash& InHash)
 {
 	static_assert(sizeof(ByteArray) <= sizeof(decltype(InHash.GetBytes())), "Reading too many bytes from source.");
@@ -104,6 +97,15 @@ inline FIoHash::FIoHash(const FUtf8StringView HexHash)
 	UE::String::HexToBytes(HexHash, Hash);
 }
 
+inline FIoHash FIoHash::FromView(const FMemoryView InHash)
+{
+	checkf(InHash.GetSize() == sizeof(ByteArray),
+		TEXT("FIoHash cannot be constructed from a view of %" UINT64_FMT " bytes."), InHash.GetSize());
+	FIoHash NewHash;
+	FMemory::Memcpy(NewHash.Hash, InHash.GetData(), sizeof(ByteArray));
+	return NewHash;
+}
+
 inline bool FIoHash::IsZero() const
 {
 	using UInt32Array = uint32[5];
