- Moved the MakeBinaryConfig commandlet into UnrealPak (like IoStore commandlet) for much faster operation
- Fixes for the ini console command for other platforms
- Allow for reading Config files from another project (for instance from a program like UnrelPak)

diff --git a/Engine/Source/Developer/PakFileUtilities/Private/PakFileUtilities.cpp b/Engine/Source/Developer/PakFileUtilities/Private/PakFileUtilities.cpp
--- a/Engine/Source/Developer/PakFileUtilities/Private/PakFileUtilities.cpp
+++ b/Engine/Source/Developer/PakFileUtilities/Private/PakFileUtilities.cpp
@@ -28,6 +28,7 @@
 #include "Misc/ICompressionFormat.h"
 #include "Misc/KeyChainUtilities.h"
 #include "IoStoreUtilities.h"
+#include "Interfaces/IPluginManager.h"
 
 IMPLEMENT_MODULE(FDefaultModuleImpl, PakFileUtilities);
 
@@ -36,6 +37,8 @@ IMPLEMENT_MODULE(FDefaultModuleImpl, PakFileUtilities);
 #define USE_DDC_FOR_COMPRESSED_FILES 0
 #define PAKCOMPRESS_DERIVEDDATA_VER TEXT("9493D2AB515048658AF7BE1342EC21FC")
 
+DEFINE_LOG_CATEGORY_STATIC(LogMakeBinaryConfig, Log, All);
+
 
 #define SEEK_OPT_VERBOSITY Display
 
@@ -4815,6 +4818,105 @@ void CheckAndReallocThreadPool()
 	}
 }
 
+bool MakeBinaryConfig(const TCHAR* CmdLine)
+{
+	FString OutputFile;
+	if (!FParse::Value(CmdLine, TEXT("OutputFile="), OutputFile))
+	{
+		UE_LOG(LogMakeBinaryConfig, Error, TEXT("OutputFile= parameter required"));
+		return false;
+	}
+
+	FString StagedPluginsFile;
+	if (!FParse::Value(CmdLine, TEXT("StagedPluginsFile="), StagedPluginsFile))
+	{
+		UE_LOG(LogMakeBinaryConfig, Error, TEXT("StagedPluginsFile= parameter required"));
+		return false;
+	}
+
+	FString ProjectFile;
+	if (!FParse::Value(CmdLine, TEXT("Project="), ProjectFile))
+	{
+		UE_LOG(LogMakeBinaryConfig, Error, TEXT("ProjectFile= parameter required (path to .uproject file)"));
+		return false;
+	}
+
+	FString PlatformName;
+	if (!FParse::Value(CmdLine, TEXT("Platform="), PlatformName))
+	{
+		UE_LOG(LogMakeBinaryConfig, Error, TEXT("Platform= parameter required (Ini platform name, not targetplatform name - Windows, not Win64 or WindowsClient)"));
+		return false;
+	}
+
+	FConfigCacheIni Config(EConfigCacheType::Temporary);
+	FString ProjectDir = FPaths::GetPath(ProjectFile);
+	Config.InitializeKnownConfigFiles(*PlatformName, false, *ProjectDir);
+
+	// removing for now, because this causes issues with some plugins not getting ini files merged in
+	// IPluginManager::Get().IntegratePluginsIntoConfig(Config, *GEngineIni, *PlatformName, *StagedPluginsFile);
+
+	// pull out deny list entries
+
+	TArray<FString> KeyDenyListStrings;
+	TArray<FString> SectionsDenyList;
+	GConfig->GetArray(TEXT("/Script/UnrealEd.ProjectPackagingSettings"), TEXT("IniKeyBlacklist"), KeyDenyListStrings, GGameIni);
+	GConfig->GetArray(TEXT("/Script/UnrealEd.ProjectPackagingSettings"), TEXT("IniSectionBlacklist"), SectionsDenyList, GGameIni);
+	TArray<FName> KeysDenyList;
+	for (const FString& Key : KeyDenyListStrings)
+	{
+		KeysDenyList.Add(FName(*Key));
+	}
+
+	for (const FString& Filename : Config.GetFilenames())
+	{
+		FConfigFile* File = Config.FindConfigFile(Filename);
+
+		delete File->SourceConfigFile;
+		File->SourceConfigFile = nullptr;
+
+		for (const FString& Section : SectionsDenyList)
+		{
+			File->Remove(Section);
+		}
+
+		// now go over any remaining sections and remove keys
+		for (TPair<FString, FConfigSection>& SectionPair : *File)
+		{
+			FConfigSection& Section = SectionPair.Value;
+			for (FName Key : KeysDenyList)
+			{
+				Section.Remove(Key);
+			}
+		}
+	}
+
+	// check the deny list removed itself
+	KeyDenyListStrings.Empty();
+	Config.GetArray(TEXT("/Script/UnrealEd.ProjectPackagingSettings"), TEXT("IniKeyBlacklist"), KeyDenyListStrings, GGameIni);
+	check(KeyDenyListStrings.Num() == 0);
+
+	// allow delegates to modify the config data with some tagged binary data
+	FCoreDelegates::FExtraBinaryConfigData ExtraData(Config, true);
+	FCoreDelegates::AccessExtraBinaryConfigData.Broadcast(ExtraData);
+
+	// write it all out!
+	TArray<uint8> FileContent;
+	{
+		// Use FMemoryWriter because FileManager::CreateFileWriter doesn't serialize FName as string and is not overridable
+		FMemoryWriter MemoryWriter(FileContent, true);
+
+		Config.Serialize(MemoryWriter);
+		MemoryWriter << ExtraData.Data;
+	}
+
+	if (!FFileHelper::SaveArrayToFile(FileContent, *OutputFile))
+	{
+		UE_LOG(LogMakeBinaryConfig, Error, TEXT("Failed to create binary config file '%s'"), *OutputFile);
+		return false;
+	}
+
+	return true;
+}
 
 
 /**
@@ -4852,6 +4954,11 @@ bool ExecuteUnrealPak(const TCHAR* CmdLine)
 		return CreateIoStoreContainerFiles(CmdLine) == 0;
 	}
 
+	if (NonOptionArguments.Num() && NonOptionArguments[0] == TEXT("MakeBinaryConfig"))
+	{
+		return MakeBinaryConfig(CmdLine);
+	}
+
 	FString ProjectArg;
 	if (FParse::Value(CmdLine, TEXT("-Project="), ProjectArg))
 	{
diff --git a/Engine/Source/Programs/AutomationTool/Scripts/CopyBuildToStagingDirectory.Automation.cs b/Engine/Source/Programs/AutomationTool/Scripts/CopyBuildToStagingDirectory.Automation.cs
--- a/Engine/Source/Programs/AutomationTool/Scripts/CopyBuildToStagingDirectory.Automation.cs
+++ b/Engine/Source/Programs/AutomationTool/Scripts/CopyBuildToStagingDirectory.Automation.cs
@@ -1199,8 +1199,11 @@ namespace AutomationScripts
 						// run the commandlet to generate a binary file
 						String TargetPlatformName = ThisPlatform.GetCookPlatform(Params.DedicatedServer, Params.Client);
 						FileReference OutputFile = FileReference.Combine(SC.ProjectRoot, "Intermediate", "Config", TargetPlatformName, "BinaryConfig.ini");
-						String CommandletParams = String.Format("-TargetPlatform={0} -OutputFile={1} -StagedPluginsFile={2}", TargetPlatformName, OutputFile.FullName, PluginListFile.FullName);
-						RunCommandlet(SC.RawProjectPath, Params.UnrealExe, "MakeBinaryConfig", CommandletParams);
+						String UnrealPakParams = String.Format("MakeBinaryConfig -Project=\"{0}\" -Platform={1} -OutputFile=\"{2}\" -StagedPluginsFile=\"{3}\"", SC.RawProjectPath, ConfigHierarchy.GetIniPlatformName(ThisPlatform.IniPlatformType), OutputFile.FullName, PluginListFile.FullName);
+
+						LogInformation("Running UnrealPak with arguments: {0}", UnrealPakParams);
+						RunAndLog(CmdEnv, GetUnrealPakLocation().FullName, UnrealPakParams, Options: ERunOptions.Default | ERunOptions.UTF8Output);
+
 						SC.StageFile(StagedFileType.UFS, OutputFile, StagedFileReference.Combine(SC.RelativeProjectRootForStage, "Config", OutputFile.GetFileName()));
 					}
 
diff --git a/Engine/Source/Runtime/Core/Private/Misc/ConfigCacheIni.cpp b/Engine/Source/Runtime/Core/Private/Misc/ConfigCacheIni.cpp
--- a/Engine/Source/Runtime/Core/Private/Misc/ConfigCacheIni.cpp
+++ b/Engine/Source/Runtime/Core/Private/Misc/ConfigCacheIni.cpp
@@ -2999,7 +2999,7 @@ void FConfigCacheIni::Dump(FOutputDevice& Ar, const TCHAR* BaseIniName)
 
 	for (const FConfigCacheIni::FKnownConfigFiles::FKnownConfigFile& File : KnownFiles.Files)
 	{
-		DumpFile(Ar, File.IniPath, File.IniFile);
+		DumpFile(Ar, File.IniName.ToString(), File.IniFile);
 	}
 
 	for ( TIterator It(*this); It; ++It )
@@ -3885,7 +3885,7 @@ static FString PerformExpansionReplacements(const FConfigLayerExpansion& Expansi
 	return OutString;
 }
 
-static FString PerformFinalExpansions(const FString InString, const FString& PlatformName, const TCHAR* EngineConfigDir, const TCHAR* SourceConfigDir)
+static FString PerformFinalExpansions(const FString InString, const FString& PlatformName, const TCHAR* EngineDir, const TCHAR* ProjectDir)
 {
 	static FString LastPlatform;
 	static FString PlatformExtensionEngineDir;
@@ -3899,8 +3899,8 @@ static FString PerformFinalExpansions(const FString InString, const FString& Pla
 	if (LastPlatform != PlatformName)
 	{
 		LastPlatform = PlatformName;
-		PlatformExtensionEngineDir = FPaths::Combine(*FPaths::EnginePlatformExtensionsDir(), *PlatformName);
-		PlatformExtensionProjectDir = FPaths::Combine(*FPaths::ProjectPlatformExtensionsDir(), *PlatformName);
+		PlatformExtensionEngineDir = FPaths::Combine(*FPaths::EnginePlatformExtensionsDir(), *PlatformName).Replace(*FPaths::EngineDir(), *(FString(EngineDir) + "/"));
+		PlatformExtensionProjectDir = FPaths::Combine(*FPaths::ProjectPlatformExtensionsDir(), *PlatformName).Replace(*FPaths::ProjectDir(), *(FString(ProjectDir) + "/"));
 	}
 
 	// cache some slow operations
@@ -3920,9 +3920,9 @@ static FString PerformFinalExpansions(const FString InString, const FString& Pla
 		}
 	}
 
-	FString OutString = InString.Replace(TEXT("{ENGINE}"), EngineConfigDir);
+	FString OutString = InString.Replace(TEXT("{ENGINE}"), EngineDir);
 	OutString = OutString.Replace(TEXT("{EXTENGINE}"), *PlatformExtensionEngineDir);
-	OutString = OutString.Replace(TEXT("{PROJECT}"), SourceConfigDir);
+	OutString = OutString.Replace(TEXT("{PROJECT}"), ProjectDir);
 	OutString = OutString.Replace(TEXT("{EXTPROJECT}"), *PlatformExtensionProjectDir);
 	OutString = OutString.Replace(TEXT("{PLATFORM}"), *PlatformName);
 	OutString = OutString.Replace(TEXT("{RESTRICTEDPROJECT_NFL}"), *ProjectNotForLicenseesDir);
@@ -4145,7 +4145,7 @@ void FConfigCacheIni::SerializeStateForBootstrap_Impl(FArchive& Ar)
 }
 
 
-bool FConfigCacheIni::InitializeKnownConfigFiles(const TCHAR* PlatformName, bool bDefaultEngineIniRequired)
+bool FConfigCacheIni::InitializeKnownConfigFiles(const TCHAR* PlatformName, bool bDefaultEngineIniRequired, const TCHAR* OverrideProjectDir)
 {
 	// check for scalability platform override.
 	const TCHAR* ScalabilityPlatformOverride = nullptr;
@@ -4168,8 +4168,31 @@ bool FConfigCacheIni::InitializeKnownConfigFiles(const TCHAR* PlatformName, bool
 		// allo for scalability to come from another 
 		const TCHAR* PlatformNameForThisIni = KnownIndex == (uint8)EKnownIniFile::Scalability ? ScalabilityPlatformOverride : PlatformName;
 		// load the hierarchy!
-		bool bConfigCreated = FConfigCacheIni::LoadGlobalIniFile(KnownFiles.Files[KnownIndex].IniPath, *KnownFiles.Files[KnownIndex].IniName.ToString(), PlatformName, 
-			false, bDefaultEngineIniRequired, bAllowGeneratedIniWhenCooked, bAllowRemoteConfig, *FPaths::GeneratedConfigDir(), this);
+		bool bConfigCreated;
+		
+		if (OverrideProjectDir == nullptr)
+		{
+			bConfigCreated = FConfigCacheIni::LoadGlobalIniFile(KnownFiles.Files[KnownIndex].IniPath, *KnownFiles.Files[KnownIndex].IniName.ToString(), PlatformName,
+				false, bDefaultEngineIniRequired, bAllowGeneratedIniWhenCooked, bAllowRemoteConfig, *FPaths::GeneratedConfigDir(), this);
+		}
+		else
+		{
+			FString ProjectGeneratedConfigDir = FPaths::Combine(OverrideProjectDir, TEXT("Saved"), TEXT("Config/"));
+			KnownFiles.Files[KnownIndex].IniPath = GetDestIniFilename(*KnownFiles.Files[KnownIndex].IniName.ToString(), PlatformName, *ProjectGeneratedConfigDir);
+
+			bConfigCreated = LoadExternalIniFile(
+				KnownFiles.Files[KnownIndex].IniFile,
+				*KnownFiles.Files[KnownIndex].IniName.ToString(),
+				*FPaths::EngineConfigDir(),
+				*FPaths::Combine(OverrideProjectDir, TEXT("Config/")),
+				/*bIsBaseIniName*/ true,
+				PlatformName,
+				/*bForceReload*/ false,
+				/*bWriteDestIni*/ false,
+				/*bAllowGeneratedIniWhenCooked*/false,
+				*ProjectGeneratedConfigDir);
+		}
+
 
 		// we want to return if the Engine config was successfully created (to not remove any functionality from old code)
 		if (KnownIndex == (uint8)EKnownIniFile::Engine)
@@ -4288,6 +4311,38 @@ bool FConfigCacheIni::CreateGConfigFromSaved(const TCHAR* Filename)
 
 #endif
 
+static void LoadRemainingConfigFiles()
+{
+	SCOPED_BOOT_TIMING("LoadRemainingConfigFiles");
+
+#if WITH_EDITOR
+	// load some editor specific .ini files
+
+	FConfigCacheIni::LoadGlobalIniFile(GEditorIni, TEXT("Editor"));
+
+	// Upgrade editor user settings before loading the editor per project user settings
+	FConfigManifest::MigrateEditorUserSettings();
+	FConfigCacheIni::LoadGlobalIniFile(GEditorPerProjectIni, TEXT("EditorPerProjectUserSettings"));
+
+	// Project agnostic editor ini files
+	const FString EditorSettingsDir = FPaths::EngineEditorSettingsDir();
+	FConfigCacheIni::LoadGlobalIniFile(GEditorSettingsIni, TEXT("EditorSettings"), nullptr, false, false, true, true, *EditorSettingsDir);
+	FConfigCacheIni::LoadGlobalIniFile(GEditorKeyBindingsIni, TEXT("EditorKeyBindings"), nullptr, false, false, true, true, *EditorSettingsDir);
+	FConfigCacheIni::LoadGlobalIniFile(GEditorLayoutIni, TEXT("EditorLayout"), nullptr, false, false, true, true, *EditorSettingsDir);
+
+#endif
+#if PLATFORM_DESKTOP
+	// load some desktop only .ini files
+	FConfigCacheIni::LoadGlobalIniFile(GCompatIni, TEXT("Compat"));
+	FConfigCacheIni::LoadGlobalIniFile(GLightmassIni, TEXT("Lightmass"));
+#endif
+
+	if (FParse::Param(FCommandLine::Get(), TEXT("dumpconfig")))
+	{
+		GConfig->Dump(*GLog);
+	}
+}
+
 void FConfigCacheIni::InitializeConfigSystem()
 {
 	// assign the G***Ini strings for the known ini's
@@ -4296,17 +4351,24 @@ void FConfigCacheIni::InitializeConfigSystem()
 		ENUMERATE_KNOWN_INI_FILES(ASSIGN_GLOBAL_INI_STRING);
 	#undef ASSIGN_GLOBAL_INI_STRING
 
+#if WITH_EDITOR
+	AsyncInitializeConfigForPlatforms();
+#endif
+
 #if PLATFORM_SUPPORTS_BINARYCONFIG && PRELOAD_BINARY_CONFIG
 	// attempt to load from staged binary config data
-	FString BinaryConfigFile = FPaths::Combine(FPaths::SourceConfigDir(), TEXT("BinaryConfig.ini"));
-
 	if (!FParse::Param(FCommandLine::Get(), TEXT("textconfig")) &&
-		IFileManager::Get().FileExists(*BinaryConfigFile) &&
-		FConfigCacheIni::CreateGConfigFromSaved(*BinaryConfigFile))
+		FConfigCacheIni::CreateGConfigFromSaved(nullptr))
 	{
 		// Force reload GameUserSettings because they may be saved to disk on consoles/similar platforms
 		// So the safest thing to do is to re-read the file after binary configs load.
 		FConfigCacheIni::LoadGlobalIniFile(GGameUserSettingsIni, TEXT("GameUserSettings"));
+
+#if WITH_EDITOR
+		// a cooked editor (cooked cooker more likely) can be initialized from binary config for speed, 
+		// but we still need the other files as well
+		LoadRemainingConfigFiles();
+#endif
 		return;
 	}
 #endif
@@ -4366,35 +4428,12 @@ void FConfigCacheIni::InitializeConfigSystem()
 		}
 	}
 
-#if WITH_EDITOR
-	// load some editor specific .ini files
-
-	FConfigCacheIni::LoadGlobalIniFile(GEditorIni, TEXT("Editor"));
-
-	// Upgrade editor user settings before loading the editor per project user settings
-	FConfigManifest::MigrateEditorUserSettings();
-	FConfigCacheIni::LoadGlobalIniFile(GEditorPerProjectIni, TEXT("EditorPerProjectUserSettings"));
-
-	// Project agnostic editor ini files
-	const FString EditorSettingsDir = FPaths::EngineEditorSettingsDir();
-	FConfigCacheIni::LoadGlobalIniFile(GEditorSettingsIni, TEXT("EditorSettings"), nullptr, false, false, true, true, *EditorSettingsDir);
-	FConfigCacheIni::LoadGlobalIniFile(GEditorKeyBindingsIni, TEXT("EditorKeyBindings"), nullptr, false, false, true, true, *EditorSettingsDir);
-	FConfigCacheIni::LoadGlobalIniFile(GEditorLayoutIni, TEXT("EditorLayout"), nullptr, false, false, true, true, *EditorSettingsDir);
-
-#endif
-#if PLATFORM_DESKTOP
-	// load some desktop only .ini files
-	FConfigCacheIni::LoadGlobalIniFile(GCompatIni, TEXT("Compat"));
-	FConfigCacheIni::LoadGlobalIniFile(GLightmassIni, TEXT("Lightmass"));
-#endif
+	// load editor, etc config files that won't be binaried
+	LoadRemainingConfigFiles();
 
 	// now we can make use of GConfig
 	GConfig->bIsReadyForUse = true;
 
-#if WITH_EDITOR
-	AsyncInitializeConfigForPlatforms();
-#endif
-
 	FCoreDelegates::ConfigReadyForUse.Broadcast();
 }
 
diff --git a/Engine/Source/Runtime/Core/Public/Misc/ConfigCacheIni.h b/Engine/Source/Runtime/Core/Public/Misc/ConfigCacheIni.h
--- a/Engine/Source/Runtime/Core/Public/Misc/ConfigCacheIni.h
+++ b/Engine/Source/Runtime/Core/Public/Misc/ConfigCacheIni.h
@@ -1203,10 +1203,11 @@ public:
 	 *
 	 * @param PlatformName Ini name of the platform this is loadenig for (or nullptr for current platform)
 	 * @param bDefaultEngineIniRequired True if the engine ini is required
+	 * @param OverrideProjectDir Used to generate configs for another project (like from a program making binary configs)
 	 *
 	 * return True if the engine ini was loaded
 	 */
-	bool InitializeKnownConfigFiles(const TCHAR* PlatformName, bool bDefaultEngineIniRequired);
+	bool InitializeKnownConfigFiles(const TCHAR* PlatformName, bool bDefaultEngineIniRequired, const TCHAR* OverrideProjectDir=nullptr);
 
 	/**
 	 * Returns true if the given name is one of the known configs, where the matching G****Ini property is going to match the 
diff --git a/Engine/Source/Runtime/Engine/Private/UnrealEngine.cpp b/Engine/Source/Runtime/Engine/Private/UnrealEngine.cpp
--- a/Engine/Source/Runtime/Engine/Private/UnrealEngine.cpp
+++ b/Engine/Source/Runtime/Engine/Private/UnrealEngine.cpp
@@ -9286,7 +9286,7 @@ bool UEngine::HandleConfigMemCommand( const TCHAR* Cmd, FOutputDevice& Ar )
 
 bool UEngine::HandleGetIniCommand(const TCHAR* Cmd, FOutputDevice& Ar)
 {
-	// Format: GetIni [Platform@]IniFile:Section.SubSection Key
+	// Format: GetIni [Platform@]IniFile:Section Key
 	TCHAR IniPlusSectionName[256];
 	TCHAR KeyName[256];
 
@@ -9326,23 +9326,36 @@ bool UEngine::HandleGetIniCommand(const TCHAR* Cmd, FOutputDevice& Ar)
 			else
 			{
 				TArray<FString> ConfigList;
-				FString SearchStr = IniPlusSection.Left(IniDelim) + TEXT(".ini");
-
 				ConfigSystem->GetConfigFilenames(ConfigList);
 
-				for (auto CurConfig : ConfigList)
+				// first try exact match (this helps with known files that have no .ini extension, etc)
+				FString SearchString = IniPlusSection.Left(IniDelim);
+				const FString* Result = Algo::FindByPredicate(ConfigList, [&SearchString](const FString& Test) { return FPaths::GetCleanFilename(Test) == SearchString; });
+				
+				// now try with .ini added if not already a .ini
+				if (!SearchString.EndsWith(TEXT(".ini")))
 				{
-					if (CurConfig.Contains(SearchStr, ESearchCase::IgnoreCase, ESearchDir::FromEnd))
+					SearchString += ".ini";
+					if (Result == nullptr)
 					{
-						IniName = CurConfig;
-						break;
+						Result = Algo::FindByPredicate(ConfigList, [&SearchString](const FString& Test) { return FPaths::GetCleanFilename(Test) == SearchString; });
 					}
 				}
 
-				if (IniName.IsEmpty())
+				// finally fallback to old method of substring search with .ini extension
+				if (Result == nullptr)
+				{
+					Result = Algo::FindByPredicate(ConfigList, [&SearchString](const FString& Test) { return Test.Contains(SearchString); });
+				}
+
+				if (Result == nullptr)
 				{
 					UE_SUPPRESS(LogExec, Warning,
-						Ar.Logf(TEXT("Failed to find IniFile '%s' (note: can only search loaded ini files)."), *SearchStr));
+						Ar.Logf(TEXT("Failed to find IniFile '%s' (note: can only search loaded ini files)."), *SearchString));
+				}
+				else
+				{
+					IniName = *Result;
 				}
 			}
 		}
@@ -9375,26 +9388,26 @@ bool UEngine::HandleGetIniCommand(const TCHAR* Cmd, FOutputDevice& Ar)
 			else
 			{
 				UE_SUPPRESS(LogExec, Warning,
-					Ar.Logf(TEXT("No Key specified. Command format: GetIni IniFile:Section.SubSection Key")));
+					Ar.Logf(TEXT("No Key specified. Command format: GetIni [Platform@]IniFile:Section Key")));
 			}
 		}
 		else if (IniName.IsEmpty())
 		{
 			UE_SUPPRESS(LogExec, Warning,
-				Ar.Logf(TEXT("IniFile parsing failed (%s). Command format: GetIni IniFile:Section.SubSection Key"),
+				Ar.Logf(TEXT("IniFile parsing failed (%s). Command format: GetIni [Platform@]IniFile:Section Key"),
 					IniPlusSectionName));
 		}
 		else // if (SectionName.IsEmpty())
 		{
 			UE_SUPPRESS(LogExec, Warning,
-				Ar.Logf(TEXT("Section parsing failed (%s). Command format: GetIni IniFile:Section.SubSection Key"),
+				Ar.Logf(TEXT("Section parsing failed (%s). Command format: GetIni [Platform@]IniFile:Section Key"),
 					IniPlusSectionName));
 		}
 	}
 	else
 	{
 		UE_SUPPRESS(LogExec, Warning,
-			Ar.Logf(TEXT("No Section specified. Command format: GetIni IniFile:Section.SubSection Key")))
+			Ar.Logf(TEXT("No Section specified. Command format: GetIni [Platform@]IniFile:Section Key")))
 	}
 
 	return true;
