- Fixed non-unity compile error

diff --git a/Engine/Source/Runtime/Core/Private/Misc/ConfigCacheIni.cpp b/Engine/Source/Runtime/Core/Private/Misc/ConfigCacheIni.cpp
--- a/Engine/Source/Runtime/Core/Private/Misc/ConfigCacheIni.cpp
+++ b/Engine/Source/Runtime/Core/Private/Misc/ConfigCacheIni.cpp
@@ -24,6 +24,7 @@
 #include "Serialization/MemoryWriter.h"
 #include "Serialization/LargeMemoryReader.h"
 #include "Async/Async.h"
+#include "Misc/OutputDeviceRedirector.h"
 
 #include <limits>
 
