Add named events for EventWait stats instrumented via BeginSetWaitStat/EndSetWaitStat if -csvNamedEvents is enabled.

diff --git a/Engine/Source/Runtime/Core/Private/ProfilingDebugging/CsvProfiler.cpp b/Engine/Source/Runtime/Core/Private/ProfilingDebugging/CsvProfiler.cpp
--- a/Engine/Source/Runtime/Core/Private/ProfilingDebugging/CsvProfiler.cpp
+++ b/Engine/Source/Runtime/Core/Private/ProfilingDebugging/CsvProfiler.cpp
@@ -3220,6 +3220,12 @@ void FCsvProfiler::BeginSetWaitStat(const char * StatName)
 #if RECORD_TIMESTAMPS
 	if (GCsvProfilerIsCapturing && GCsvCategoriesEnabled[CSV_CATEGORY_INDEX(Exclusive)])
 	{
+#if CSV_PROFILER_SUPPORT_NAMED_EVENTS
+		if (UNLIKELY(GCsvProfilerNamedEventsExclusive))
+		{
+			FPlatformMisc::BeginNamedEvent(FColor::Yellow, *FString::Printf(TEXT("CsvWaitStat_%s"), StringCast<TCHAR>(StatName).Get()));
+		}
+#endif
 		FCsvProfilerThreadData::Get().PushWaitStatName(StatName == nullptr ? GIgnoreWaitStatName : StatName);
 	}
 #endif
@@ -3230,6 +3236,12 @@ void FCsvProfiler::EndSetWaitStat()
 #if RECORD_TIMESTAMPS
 	if (GCsvProfilerIsCapturing && GCsvCategoriesEnabled[CSV_CATEGORY_INDEX(Exclusive)])
 	{
+#if CSV_PROFILER_SUPPORT_NAMED_EVENTS
+		if (UNLIKELY(GCsvProfilerNamedEventsExclusive))
+		{
+			FPlatformMisc::EndNamedEvent();
+		}
+#endif
 		FCsvProfilerThreadData::Get().PopWaitStatName();
 	}
 #endif
