Fix overflow in frozen memory pointers (UE-138667).

diff --git a/Engine/Plugins/Compositing/OpenColorIO/Source/OpenColorIO/Private/OpenColorIODerivedDataVersion.h b/Engine/Plugins/Compositing/OpenColorIO/Source/OpenColorIO/Private/OpenColorIODerivedDataVersion.h
--- a/Engine/Plugins/Compositing/OpenColorIO/Source/OpenColorIO/Private/OpenColorIODerivedDataVersion.h
+++ b/Engine/Plugins/Compositing/OpenColorIO/Source/OpenColorIO/Private/OpenColorIODerivedDataVersion.h
@@ -10,4 +10,4 @@ OpenColorIOShaderDerivedDataVersion.h: Derived data version for OpenColorIO shad
 
 // In case of merge conflicts with DDC versions, you MUST generate a new GUID and set this new
 // guid as version
-#define OPENCOLORIO_DERIVEDDATA_VER		TEXT("839BC82D2-160B-49FA-AAD8-D3233D678492")
\ No newline at end of file
+#define OPENCOLORIO_DERIVEDDATA_VER		TEXT("BFFE3C29-C30F-4D6E-81D4-14BDE56F16ED")
\ No newline at end of file
diff --git a/Engine/Plugins/Experimental/ComputeFramework/Source/ComputeFramework/Private/ComputeFramework/ComputeKernelDerivedDataVersion.h b/Engine/Plugins/Experimental/ComputeFramework/Source/ComputeFramework/Private/ComputeFramework/ComputeKernelDerivedDataVersion.h
--- a/Engine/Plugins/Experimental/ComputeFramework/Source/ComputeFramework/Private/ComputeFramework/ComputeKernelDerivedDataVersion.h
+++ b/Engine/Plugins/Experimental/ComputeFramework/Source/ComputeFramework/Private/ComputeFramework/ComputeKernelDerivedDataVersion.h
@@ -5,4 +5,4 @@
 #include "CoreMinimal.h"
 
 // In case of merge conflicts with DDC versions, you MUST generate a new GUID and set this new guid as version
-#define COMPUTEKERNEL_DERIVEDDATA_VER	TEXT("A66D3C08-DD56-4DFB-A188-9C0E08F15E02")
+#define COMPUTEKERNEL_DERIVEDDATA_VER	TEXT("F35F306A-EA41-4E0D-A354-06DEE736B126")
diff --git a/Engine/Plugins/FX/Niagara/Source/NiagaraShader/Private/NiagaraShaderDerivedDataVersion.h b/Engine/Plugins/FX/Niagara/Source/NiagaraShader/Private/NiagaraShaderDerivedDataVersion.h
--- a/Engine/Plugins/FX/Niagara/Source/NiagaraShader/Private/NiagaraShaderDerivedDataVersion.h
+++ b/Engine/Plugins/FX/Niagara/Source/NiagaraShader/Private/NiagaraShaderDerivedDataVersion.h
@@ -11,4 +11,4 @@ NiagaraShaderDerivedDataVersion.h: Shader derived data version for Niagara.
 // In case of merge conflicts with DDC versions, you MUST generate a new GUID and set this new
 // guid as version
 
-#define NIAGARASHADERMAP_DERIVEDDATA_VER		TEXT("2289D511-6CF9-4BC0-AFEA-EC541468E645")
+#define NIAGARASHADERMAP_DERIVEDDATA_VER		TEXT("D204CCC1-66ED-4826-85DA-EE8AB72895AB")
diff --git a/Engine/Source/Runtime/Core/Private/Serialization/MemoryImage.cpp b/Engine/Source/Runtime/Core/Private/Serialization/MemoryImage.cpp
--- a/Engine/Source/Runtime/Core/Private/Serialization/MemoryImage.cpp
+++ b/Engine/Source/Runtime/Core/Private/Serialization/MemoryImage.cpp
@@ -1670,9 +1670,9 @@ FMemoryImageSection* FMemoryImageSection::WritePointer(const FTypeLayoutDesc& St
 		DerivedTypeDesc.Name, StaticTypeDesc.Name);
 
 	FFrozenMemoryImagePtr FrozenPtr;
-	FrozenPtr.OffsetFromThis = 0; // dummy value
+	FrozenPtr.Packed = 0; // dummy value
 	// If derived type matches the static type, store INDEX_NONE to indicate the static type may be used
-	FrozenPtr.TypeIndex = bStaticTypeMatchesDerived ? INDEX_NONE : TypeDependencyIndex;
+	FrozenPtr.SetTypeIndex(bStaticTypeMatchesDerived ? INDEX_NONE : TypeDependencyIndex);
 
 	SectionPointer.Offset = WriteBytes(FrozenPtr);
 	SectionPointer.PointerOffset = OffsetToBase;
@@ -1853,12 +1853,14 @@ void FMemoryImage::Flatten(FMemoryImageResult& OutResult, bool bMergeDuplicateSe
 		const FMemoryImageSection* Section = UniqueSections[SectionIndex];
 		for (const FMemoryImageSection::FSectionPointer& Pointer : Section->Pointers)
 		{
-			const int32 OffsetToPointer = SectionOffset[SectionIndex] + Pointer.Offset;
+			const int64 OffsetToPointer = SectionOffset[SectionIndex] + Pointer.Offset;
 			const int32 RemapSectionIndex = SectionIndexRemap[Pointer.SectionIndex];
 			FFrozenMemoryImagePtr* FrozenPtr = (FFrozenMemoryImagePtr*)(OutResult.Bytes.GetData() + OffsetToPointer);
-			check(FrozenPtr->OffsetFromThis == 0);
-			const int32 OffsetFromPointer = (int32)(SectionOffset[RemapSectionIndex] + Pointer.PointerOffset) - OffsetToPointer;
-			FrozenPtr->OffsetFromThis = (OffsetFromPointer << 1) | 1;
+			check(FrozenPtr->GetOffsetFromThis() == 0 && !FrozenPtr->IsFrozen());
+			const int64 OffsetFromPointer = (int64)(SectionOffset[RemapSectionIndex] + Pointer.PointerOffset) - OffsetToPointer;
+			FrozenPtr->SetOffsetFromThis(OffsetFromPointer);
+			FrozenPtr->SetIsFrozen(true);
+			check(FrozenPtr->GetOffsetFromThis() == OffsetFromPointer);
 		}
 	}
 
@@ -1940,7 +1942,7 @@ uint32 FMemoryImageWriter::WriteBytes(const void* Data, uint32 Size)
 
 uint32 FMemoryImageWriter::WriteNullPointer()
 {
-	return WriteBytes(FFrozenMemoryImagePtr{ 0, 0u });
+	return WriteBytes(FFrozenMemoryImagePtr{ 0ull });
 }
 
 FMemoryImageWriter FMemoryImageWriter::WritePointer(const FTypeLayoutDesc& StaticTypeDesc, const FTypeLayoutDesc& DerivedTypeDesc, uint32* OutOffsetToBase)
diff --git a/Engine/Source/Runtime/Core/Public/Serialization/MemoryImage.h b/Engine/Source/Runtime/Core/Public/Serialization/MemoryImage.h
--- a/Engine/Source/Runtime/Core/Public/Serialization/MemoryImage.h
+++ b/Engine/Source/Runtime/Core/Public/Serialization/MemoryImage.h
@@ -289,33 +289,86 @@ public:
 	const class UStruct* CurrentStruct;
 };
 
+/**
+ * Value of this struct should never be a valid unfrozen pointer (i.e. a memory address). We rely on real pointers to have lowest bit(s) 0 this days for the alignment, this is checked later.
+ * Unfortunately, we cannot use bitfields as their layout might be compiler-specific, and the data for the target platform is being prepared with a different compiler during the cook.
+ */
 struct FFrozenMemoryImagePtr
 {
-	int32 OffsetFromThis;
-	int32 TypeIndex; // store 0-based index or INDEX_NONE, don't really need 32bit for this, but nothing else to store here for now
+	static constexpr uint64 bIsFrozenBits = 1;
+	static constexpr uint64 OffsetBits = 40;
+	static constexpr uint64 TypeIndexBits = 64 - OffsetBits - bIsFrozenBits;
+
+	static constexpr uint64 bIsFrozenShift = 0;
+	static constexpr uint64 TypeIndexShift = bIsFrozenBits;
+	static constexpr uint64 OffsetShift = bIsFrozenBits + TypeIndexBits;
+
+	static constexpr uint64 bIsFrozenMask = (1ULL << bIsFrozenShift);
+	static constexpr uint64 TypeIndexMask = (((1ULL << TypeIndexBits) - 1ULL) << TypeIndexShift);
+	static constexpr uint64 OffsetMask = (((1ULL << OffsetBits) - 1ULL) << OffsetShift);
+
+	uint64 Packed;
+
+	/** Whether the value is indeed a frozen pointer, must come first to avoid being set in regular pointers - which are expected to point at padded things so are never even. */
+	bool IsFrozen() const
+	{
+		return (Packed & bIsFrozenMask) != 0;
+	}
+
+	void SetIsFrozen(bool bTrue)
+	{
+		Packed = (Packed & ~bIsFrozenMask) | (bTrue ? 1 : 0);
+	}
+
+	/** Signed offset from the current position in the memory image. */
+	int64 GetOffsetFromThis() const
+	{
+		// Since the offset occupies the highest part of the int64, its sign is preserved.
+		// Not masking as there's nothing to the left of the Offset
+		static_assert(OffsetShift + OffsetBits == 64);
+		return static_cast<int64>(Packed/* & OffsetMask*/) >> OffsetShift;
+	}
+
+	void SetOffsetFromThis(int64 Offset)
+	{
+		Packed = (Packed & ~OffsetMask) | (static_cast<uint64>(Offset << OffsetShift) & OffsetMask);
+	}
+
+	/** The pointer type index in the pointer table. Does not store other negative values except for INDEX_NONE */
+	int32 GetTypeIndex() const
+	{
+		return static_cast<int32>((Packed & TypeIndexMask) >> TypeIndexShift) - 1;
+	}
+
+	void SetTypeIndex(int32 TypeIndex)
+	{
+		static_assert(INDEX_NONE == -1, "TypeIndex cannot store INDEX_NONE when it's not -1");
+		Packed = (Packed & ~TypeIndexMask) | ((static_cast<uint64>(TypeIndex + 1) << TypeIndexShift) & TypeIndexMask);
+	}
 };
 
+static_assert(sizeof(FFrozenMemoryImagePtr) == sizeof(uint64), "FFrozenMemoryImagePtr is larger than a native pointer would be");
 
 template<typename T>
 class TMemoryImagePtr
 {
 public:
-	inline bool IsFrozen() const { return Frozen.OffsetFromThis & IsFrozenMask; }
+	inline bool IsFrozen() const { return Frozen.IsFrozen(); }
 	inline bool IsValid() const { return UnfrozenPtr != nullptr; }
 	inline bool IsNull() const { return UnfrozenPtr == nullptr; }
 
-	inline TMemoryImagePtr(T* InPtr = nullptr) : UnfrozenPtr(InPtr) { check(((uint64)UnfrozenPtr & AllFlags) == 0u); }
-	inline TMemoryImagePtr(const TMemoryImagePtr<T>& InPtr) : UnfrozenPtr(InPtr.Get()) { check(((uint64)UnfrozenPtr & AllFlags) == 0u); }
-	inline TMemoryImagePtr& operator=(T* InPtr) { UnfrozenPtr = InPtr; check(((uint64)UnfrozenPtr & AllFlags) == 0u); return *this; }
-	inline TMemoryImagePtr& operator=(const TMemoryImagePtr<T>& InPtr) { UnfrozenPtr = InPtr.Get(); check(((uint64)UnfrozenPtr & AllFlags) == 0u); return *this; }
+	inline TMemoryImagePtr(T* InPtr = nullptr) : UnfrozenPtr(InPtr) { check(!Frozen.IsFrozen()); }
+	inline TMemoryImagePtr(const TMemoryImagePtr<T>& InPtr) : UnfrozenPtr(InPtr.Get()) { check(!Frozen.IsFrozen()); }
+	inline TMemoryImagePtr& operator=(T* InPtr) { UnfrozenPtr = InPtr; check(!Frozen.IsFrozen()); return *this; }
+	inline TMemoryImagePtr& operator=(const TMemoryImagePtr<T>& InPtr) { UnfrozenPtr = InPtr.Get(); check(!Frozen.IsFrozen()); return *this; }
 
 	inline ~TMemoryImagePtr() 
 	{
 
 	}
 
-	inline int32 GetFrozenOffsetFromThis() const { check(IsFrozen()); return (Frozen.OffsetFromThis >> OffsetShift); }
-	inline uint32 GetFrozenTypeIndex() const { check(IsFrozen()); return Frozen.TypeIndex; }
+	inline int64 GetFrozenOffsetFromThis() const { check(IsFrozen()); return Frozen.GetOffsetFromThis(); }
+	inline int32 GetFrozenTypeIndex() const { check(IsFrozen()); return Frozen.GetTypeIndex(); }
 
 	inline T* Get() const
 	{
@@ -356,16 +409,9 @@ public:
 private:
 	inline T* GetFrozenPtrInternal() const
 	{
-		return (T*)((char*)this + (Frozen.OffsetFromThis >> OffsetShift));
+		return (T*)((char*)this + Frozen.GetOffsetFromThis());
 	}
 
-	enum
-	{
-		IsFrozenMask = (1 << 0),
-		AllFlags = IsFrozenMask,
-		OffsetShift = 1u,
-	};
-
 protected:
 	union
 	{
@@ -550,7 +596,7 @@ public:
 	{
 		return Data.IsValid();
 	}
-	FORCEINLINE int32 GetFrozenOffsetFromThis() const { return Data.GetFrozenOffsetFromThis(); }
+	FORCEINLINE int64 GetFrozenOffsetFromThis() const { return Data.GetFrozenOffsetFromThis(); }
 
 	void ResizeAllocation(int32 PreviousNumElements, int32 NumElements, SIZE_T NumBytesPerElement, uint32 Alignment);
 	void WriteMemoryImage(FMemoryImageWriter& Writer, const FTypeLayoutDesc& TypeDesc, int32 NumAllocatedElements, uint32 Alignment) const;
diff --git a/Engine/Source/Runtime/Engine/Private/ShaderDerivedDataVersion.h b/Engine/Source/Runtime/Engine/Private/ShaderDerivedDataVersion.h
--- a/Engine/Source/Runtime/Engine/Private/ShaderDerivedDataVersion.h
+++ b/Engine/Source/Runtime/Engine/Private/ShaderDerivedDataVersion.h
@@ -10,5 +10,5 @@
 
 // In case of merge conflicts with DDC versions, you MUST generate a new GUID and set this new
 // guid as version
-#define GLOBALSHADERMAP_DERIVEDDATA_VER			TEXT("1D2DE12F326C4BB38DBC0B2CBE348802")
-#define MATERIALSHADERMAP_DERIVEDDATA_VER		TEXT("01DF57ADAE184DB89CB6C43A8A421B57")
+#define GLOBALSHADERMAP_DERIVEDDATA_VER			TEXT("3F70708DD25F41228CD7A9EDAC04AC67")
+#define MATERIALSHADERMAP_DERIVEDDATA_VER		TEXT("01614A42F67F4FFC8600AD04D15D199F")
