Reworked IVirtualizationSystem::DoPayloadsExist to be clearer about how it works.

#rnx

- Renamed to ::QueryPayloadStatuses since the output data is in the form of an array of FPayloadStatus
- The method now returns EQueryResult rather than bool, which has 'Success' as value 0, all other values indicate an error.
-- This can be expanded in the future if we start to provide more info from backends as to what actually failed.
- FPayloadStatus has been cleaned up so that 'Partial' is not 'FoundPartial' to give a better indication that the payload was found in at least one storage backend, but not all of them.

diff --git a/Engine/Source/Developer/Virtualization/Private/PackageSubmissionChecks.cpp b/Engine/Source/Developer/Virtualization/Private/PackageSubmissionChecks.cpp
--- a/Engine/Source/Developer/Virtualization/Private/PackageSubmissionChecks.cpp
+++ b/Engine/Source/Developer/Virtualization/Private/PackageSubmissionChecks.cpp
@@ -217,7 +217,7 @@ void OnPrePackageSubmission(const TArray<FString>& FilesToSubmit, TArray<FText>&
 	UE_LOG(LogVirtualization, Display, TEXT("Found %" INT64_FMT " payload(s) in %d package(s) that need to be examined for virtualization"), TotalPayloadsToCheck, Packages.Num());
 
 	TArray<FPayloadStatus> PayloadStatuses;
-	if (!System.DoPayloadsExist(AllLocalPayloads, EStorageType::Persistent, PayloadStatuses))
+	if (System.QueryPayloadStatuses(AllLocalPayloads, EStorageType::Persistent, PayloadStatuses) != EQueryResult::Success)
 	{
 		FText Message = LOCTEXT("Virtualization_DoesExistFail", "Failed to find the status of the payloads in the packages being submitted");
 		Errors.Add(Message);
diff --git a/Engine/Source/Developer/Virtualization/Private/VirtualizationManager.cpp b/Engine/Source/Developer/Virtualization/Private/VirtualizationManager.cpp
--- a/Engine/Source/Developer/Virtualization/Private/VirtualizationManager.cpp
+++ b/Engine/Source/Developer/Virtualization/Private/VirtualizationManager.cpp
@@ -484,7 +484,7 @@ FCompressedBuffer FVirtualizationManager::PullData(const FIoHash& Id)
 	return FCompressedBuffer();
 }
 
-bool FVirtualizationManager::DoPayloadsExist(TArrayView<const FIoHash> Ids, EStorageType StorageType, TArray<FPayloadStatus>& OutStatuses)
+EQueryResult FVirtualizationManager::QueryPayloadStatuses(TArrayView<const FIoHash> Ids, EStorageType StorageType, TArray<FPayloadStatus>& OutStatuses)
 {
 	TRACE_CPUPROFILER_EVENT_SCOPE(FVirtualizationManager::DoPayloadsExist);
 
@@ -492,7 +492,7 @@ bool FVirtualizationManager::DoPayloadsExist(TArrayView<const FIoHash> Ids, ESto
 
 	if (Ids.IsEmpty())
 	{
-		return true;
+		return EQueryResult::Success;
 	}
 
 	for (int32 Index = 0; Index < Ids.Num(); ++Index)
@@ -517,7 +517,7 @@ bool FVirtualizationManager::DoPayloadsExist(TArrayView<const FIoHash> Ids, ESto
 			{
 				// If a backend entirely failed we should early out and report the problem
 				OutStatuses.Reset();
-				return false;
+				return EQueryResult::Failure_Unknown;
 			}
 
 			for (int32 Index = 0; Index < Ids.Num(); ++Index)
@@ -545,12 +545,12 @@ bool FVirtualizationManager::DoPayloadsExist(TArrayView<const FIoHash> Ids, ESto
 			}
 			else
 			{
-				OutStatuses[Index] = FPayloadStatus::Partial;;
+				OutStatuses[Index] = FPayloadStatus::FoundPartial;
 			}
 		}
 	}
 
-	return true;
+	return EQueryResult::Success;
 }
 
 FPayloadActivityInfo FVirtualizationManager::GetAccumualtedPayloadActivityInfo() const
diff --git a/Engine/Source/Developer/Virtualization/Private/VirtualizationManager.h b/Engine/Source/Developer/Virtualization/Private/VirtualizationManager.h
--- a/Engine/Source/Developer/Virtualization/Private/VirtualizationManager.h
+++ b/Engine/Source/Developer/Virtualization/Private/VirtualizationManager.h
@@ -98,7 +98,7 @@ private:
 
 	virtual FCompressedBuffer PullData(const FIoHash& Id) override;
 
-	virtual bool DoPayloadsExist(TArrayView<const FIoHash> Ids, EStorageType StorageType, TArray<FPayloadStatus>& OutStatuses) override;
+	virtual EQueryResult QueryPayloadStatuses(TArrayView<const FIoHash> Ids, EStorageType StorageType, TArray<FPayloadStatus>& OutStatuses) override;
 
 	virtual FPayloadActivityInfo GetAccumualtedPayloadActivityInfo() const override;
 
diff --git a/Engine/Source/Runtime/Core/Private/Virtualization/VirtualizationSystem.cpp b/Engine/Source/Runtime/Core/Private/Virtualization/VirtualizationSystem.cpp
--- a/Engine/Source/Runtime/Core/Private/Virtualization/VirtualizationSystem.cpp
+++ b/Engine/Source/Runtime/Core/Private/Virtualization/VirtualizationSystem.cpp
@@ -44,11 +44,11 @@ public:
 		return FCompressedBuffer();
 	}
 
-	virtual bool DoPayloadsExist(TArrayView<const FIoHash> Ids, EStorageType StorageType, TArray<FPayloadStatus>& OutStatuses) override
+	virtual EQueryResult QueryPayloadStatuses(TArrayView<const FIoHash> Ids, EStorageType StorageType, TArray<FPayloadStatus>& OutStatuses) override
 	{
 		OutStatuses.Reset();
 
-		return false;
+		return EQueryResult::Failure_NotImplemented;
 	}
 
 	virtual void GetPayloadActivityInfo(GetPayloadActivityInfoFuncRef) const override
diff --git a/Engine/Source/Runtime/Core/Public/Virtualization/VirtualizationSystem.h b/Engine/Source/Runtime/Core/Public/Virtualization/VirtualizationSystem.h
--- a/Engine/Source/Runtime/Core/Public/Virtualization/VirtualizationSystem.h
+++ b/Engine/Source/Runtime/Core/Public/Virtualization/VirtualizationSystem.h
@@ -40,6 +40,21 @@ enum class EStorageType : int8
 	Persistent
 };
 
+/** 
+ * The result of a query.
+ * Success indicates that the query worked and that the results are valid and can be used.
+ * Any other value indicates that the query failed in some manner and that the results cannot be trusted and should be discarded.
+ */
+enum class EQueryResult : int8
+{
+	/** The query succeeded and the results are valid */
+	Success = 0,
+	/** The query failed with an unspecified error */
+	Failure_Unknown,
+	/** The query failed because the current virtualization system has not implemented it */
+	Failure_NotImplemented
+};
+
 /** Describes the status of a payload in regards to a backend storage system */
 enum class FPayloadStatus : int8
 {
@@ -48,7 +63,7 @@ enum class FPayloadStatus : int8
 	/** The payload was not found in any backend for the given storage type */
 	NotFound = 0,
 	/** The payload was found in at least one backend but was not found in all backends available for the given storage type */
-	Partial,
+	FoundPartial,
 	/** The payload was found in all of the backends available for the given storage type */
 	FoundAll
 };
@@ -172,8 +187,6 @@ public:
 	 */
 	virtual FCompressedBuffer PullData(const FIoHash& Id) = 0;
 
-	
-
 	/**
 	 * Query if a number of payloads exist or not in the given storage type. 
 	 * 
@@ -185,7 +198,13 @@ public:
 	 * @return	True if the operation succeeded and the contents of OutStatuses is valid. False if errors were 
 	 * 			encountered in which case the contents of OutStatuses should be ignored.
 	 */
-	virtual bool DoPayloadsExist(TArrayView<const FIoHash> Ids, EStorageType StorageType, TArray<FPayloadStatus>& OutStatuses) = 0;
+	virtual EQueryResult QueryPayloadStatuses(TArrayView<const FIoHash> Ids, EStorageType StorageType, TArray<FPayloadStatus>& OutStatuses) = 0;
+
+	UE_DEPRECATED(5.1, "Call ::QueryPayloadStatuses instead")
+	bool DoPayloadsExist(TArrayView<const FIoHash> Ids, EStorageType StorageType, TArray<FPayloadStatus>& OutStatuses)
+	{
+		return QueryPayloadStatuses(Ids, StorageType, OutStatuses) != EQueryResult::Success;
+	}
 
 	using GetPayloadActivityInfoFuncRef = TFunctionRef<void(const FString& DebugName, const FString& ConfigName, const FPayloadActivityInfo& PayloadInfo)>;
 
