[Insights] Added hooks for VirtualAlloc2 API for memory tracing.

diff --git a/Engine/Source/Runtime/Core/Private/ProfilingDebugging/Microsoft/MemoryTrace.cpp b/Engine/Source/Runtime/Core/Private/ProfilingDebugging/Microsoft/MemoryTrace.cpp
--- a/Engine/Source/Runtime/Core/Private/ProfilingDebugging/Microsoft/MemoryTrace.cpp
+++ b/Engine/Source/Runtime/Core/Private/ProfilingDebugging/Microsoft/MemoryTrace.cpp
@@ -2,7 +2,6 @@
 
 #include "ProfilingDebugging/MemoryTrace.h"
 
-#include "HAL/MallocMimalloc.h"
 #include "ProfilingDebugging/CallstackTrace.h"
 #include "ProfilingDebugging/TraceMalloc.h"
 
@@ -307,19 +306,24 @@ private:
 							FVirtualWinApiHooks();
 	static bool				bLight;
 	static LPVOID WINAPI	VmAlloc(LPVOID Address, SIZE_T Size, DWORD Type, DWORD Protect);
+	static LPVOID WINAPI	VmAllocEx(HANDLE Process, LPVOID Address, SIZE_T Size, DWORD Type, DWORD Protect);
+	static LPVOID WINAPI	VmAlloc2(HANDLE Process, LPVOID BaseAddress, SIZE_T Size, ULONG AllocationType, ULONG PageProtection, /*MEM_EXTENDED_PARAMETER* */ void* ExtendedParameters, ULONG ParameterCount);
 	static BOOL WINAPI		VmFree(LPVOID Address, SIZE_T Size, DWORD Type);
 	static BOOL WINAPI		VmFreeEx(HANDLE Process, LPVOID Address, SIZE_T Size, DWORD Type);
-	static LPVOID WINAPI	VmAllocEx(HANDLE Process, LPVOID Address, SIZE_T Size, DWORD Type, DWORD Protect);
 	static LPVOID			(WINAPI *VmAllocOrig)(LPVOID, SIZE_T, DWORD, DWORD);
 	static LPVOID			(WINAPI *VmAllocExOrig)(HANDLE, LPVOID, SIZE_T, DWORD, DWORD);
+	static LPVOID			(WINAPI *VmAlloc2Orig)(HANDLE, LPVOID, SIZE_T, ULONG, ULONG, /*MEM_EXTENDED_PARAMETER* */ void*, ULONG);
 	static BOOL				(WINAPI *VmFreeOrig)(LPVOID, SIZE_T, DWORD);
 	static BOOL				(WINAPI *VmFreeExOrig)(HANDLE, LPVOID, SIZE_T, DWORD);
+
+	typedef LPVOID(__stdcall* FnVirtualAlloc2)(HANDLE, LPVOID, SIZE_T, ULONG, ULONG, /* MEM_EXTENDED_PARAMETER* */ void*, ULONG);
 };
 
 ////////////////////////////////////////////////////////////////////////////////
 bool	FVirtualWinApiHooks::bLight;
 LPVOID	(WINAPI *FVirtualWinApiHooks::VmAllocOrig)(LPVOID, SIZE_T, DWORD, DWORD);
 LPVOID	(WINAPI *FVirtualWinApiHooks::VmAllocExOrig)(HANDLE, LPVOID, SIZE_T, DWORD, DWORD);
+LPVOID	(WINAPI *FVirtualWinApiHooks::VmAlloc2Orig)(HANDLE, LPVOID, SIZE_T, ULONG, ULONG, /*MEM_EXTENDED_PARAMETER* */ void*, ULONG);
 BOOL	(WINAPI *FVirtualWinApiHooks::VmFreeOrig)(LPVOID, SIZE_T, DWORD);
 BOOL	(WINAPI *FVirtualWinApiHooks::VmFreeExOrig)(HANDLE, LPVOID, SIZE_T, DWORD);
 
@@ -338,9 +342,35 @@ void FVirtualWinApiHooks::Initialize(bool bInLight)
 			VmAllocExOrig = Editor.Hook(VirtualAllocEx, &FVirtualWinApiHooks::VmAllocEx);
 			VmFreeExOrig = Editor.Hook(VirtualFreeEx, &FVirtualWinApiHooks::VmFreeEx);
 		}
-	}
 
-	
+#if PLATFORM_WINDOWS
+#if (NTDDI_VERSION >= NTDDI_WIN10_RS4)
+		{
+			FTextSectionEditor EditorAlloc2((void*)VirtualAlloc2);
+			VmAlloc2Orig = Editor.Hook(VirtualAlloc2, &FVirtualWinApiHooks::VmAlloc2);
+		}
+#else // NTDDI_VERSION
+		{
+			VmAlloc2Orig = nullptr;
+			HINSTANCE DllInstance;
+			DllInstance = LoadLibrary(TEXT("kernelbase.dll"));
+			if (DllInstance != NULL)
+			{
+#pragma warning(push)
+#pragma warning(disable: 4191) // 'type cast': unsafe conversion from 'FARPROC' to 'FVirtualWinApiHooks::FnVirtualAlloc2'
+				VmAlloc2Orig = (FnVirtualAlloc2)GetProcAddress(DllInstance, "VirtualAlloc2");
+#pragma warning(pop)
+				FreeLibrary(DllInstance);
+			}
+			if (VmAlloc2Orig)
+			{
+				FTextSectionEditor EditorAlloc2((void*)VmAlloc2Orig);
+				VmAlloc2Orig = Editor.Hook(VmAlloc2Orig, &FVirtualWinApiHooks::VmAlloc2);
+			}
+		}
+#endif // NTDDI_VERSION
+#endif // PLATFORM_WINDOWS
+	}
 }
 
 ////////////////////////////////////////////////////////////////////////////////
@@ -362,7 +392,6 @@ LPVOID WINAPI FVirtualWinApiHooks::VmAlloc(LPVOID Address, SIZE_T Size, DWORD Ty
 ////////////////////////////////////////////////////////////////////////////////
 BOOL WINAPI FVirtualWinApiHooks::VmFree(LPVOID Address, SIZE_T Size, DWORD Type)
 {
-	// Currently tracking any free event
 	if (Type & MEM_RELEASE)
 	{
 		GAllocationTrace->Free(Address, EMemoryTraceRootHeap::SystemMemory);
@@ -394,6 +423,20 @@ BOOL WINAPI FVirtualWinApiHooks::VmFreeEx(HANDLE Process, LPVOID Address, SIZE_T
 
 	return VmFreeExOrig(Process, Address, Size, Type);
 }
+
+////////////////////////////////////////////////////////////////////////////////
+LPVOID WINAPI FVirtualWinApiHooks::VmAlloc2(HANDLE Process, LPVOID BaseAddress, SIZE_T Size, ULONG Type, ULONG PageProtection, /*MEM_EXTENDED_PARAMETER* */ void* ExtendedParameters, ULONG ParameterCount)
+{
+	LPVOID Ret = VmAlloc2Orig(Process, BaseAddress, Size, Type, PageProtection, ExtendedParameters, ParameterCount);
+	if (Process == GetCurrentProcess() && Ret != nullptr && (Type & MEM_RESERVE))
+	{
+		const uint32 Callstack = FTraceMalloc::ShouldTrace() ? CallstackTrace_GetCurrentId() : 0;
+		GAllocationTrace->Alloc(Ret, Size, 0, Callstack);
+		GAllocationTrace->MarkAllocAsHeap(Ret, EMemoryTraceRootHeap::SystemMemory);
+	}
+
+	return Ret;
+}
 #endif // defined(PLATFORM_SUPPORTS_TRACE_WIN32_VIRTUAL_MEMORY_HOOKS)
 
 
