Added heuristic for detecting Windows 11 in WindowsPlatformMisc.cpp

diff --git a/Engine/Source/Runtime/Core/Private/Windows/WindowsPlatformMisc.cpp b/Engine/Source/Runtime/Core/Private/Windows/WindowsPlatformMisc.cpp
--- a/Engine/Source/Runtime/Core/Private/Windows/WindowsPlatformMisc.cpp
+++ b/Engine/Source/Runtime/Core/Private/Windows/WindowsPlatformMisc.cpp
@@ -173,23 +173,41 @@ int32 GetOSVersionsHelper( TCHAR* OutOSVersionLabel, int32 OSVersionLabelLength,
 			case 0:
 				if (OsVersionInfo.wProductType == VER_NT_WORKSTATION)
 				{
-					OSVersionLabel = TEXT("Windows 10");
+					// Windows 11 still reports a major version of 10 and minor of 0, so it looks
+					// like we need to use the build number as the discriminator
+					if (OsVersionInfo.dwBuildNumber >= 22000)
+					{
+						OSVersionLabel = TEXT("Windows 11");
+					}
+					else
+					{
+						OSVersionLabel = TEXT("Windows 10");
+					}
 				}
 				else
 				{
 					OSVersionLabel = TEXT("Windows Server 2019");
 				}
 
-				// For Windows 10, get the release number and append that to the string too (eg. 1709 = Fall Creators Update). There doesn't seem to be any good way to get
-				// this other than grabbing an entry from the registry.
+				// For Windows 10, get the release number and append that to the string too (eg. 1709 = Fall Creators Update). 
+				// There doesn't seem to be any good way to get this other than grabbing an entry from the registry.
+				// 
+				// The new semi-annual release scheme 20H1/20H2 etc appears to use a different key so we query that first.
 				{
-					FString ReleaseId;
-					if(FWindowsPlatformMisc::QueryRegKey(HKEY_LOCAL_MACHINE, TEXT("SOFTWARE\\Microsoft\\Windows NT\\CurrentVersion"), TEXT("ReleaseId"), ReleaseId))
+					FString DisplayVersion;
+					if (FWindowsPlatformMisc::QueryRegKey(HKEY_LOCAL_MACHINE, TEXT("SOFTWARE\\Microsoft\\Windows NT\\CurrentVersion"), TEXT("DisplayVersion"), DisplayVersion))
 					{
-						OSVersionLabel += FString::Printf(TEXT(" (Release %s)"), *ReleaseId);
+						OSVersionLabel += FString::Printf(TEXT(" (%s)"), *DisplayVersion);
+					}
+					else
+					{
+						FString ReleaseId;
+						if (FWindowsPlatformMisc::QueryRegKey(HKEY_LOCAL_MACHINE, TEXT("SOFTWARE\\Microsoft\\Windows NT\\CurrentVersion"), TEXT("ReleaseId"), ReleaseId))
+						{
+							OSVersionLabel += FString::Printf(TEXT(" (Release %s)"), *ReleaseId);
+						}
 					}
 				}
-
 				break;
 			default:
 				ErrorCode |= (int32)FWindowsOSVersionHelper::ERROR_UNKNOWNVERSION;
