REL - allow to not run the task continuation during cancellation, as it can be beneficial for some implementations to handle it with a manual Launch.

diff --git a/Engine/Source/Runtime/Core/Public/Async/Fundamental/Task.h b/Engine/Source/Runtime/Core/Public/Async/Fundamental/Task.h
--- a/Engine/Source/Runtime/Core/Public/Async/Fundamental/Task.h
+++ b/Engine/Source/Runtime/Core/Public/Async/Fundamental/Task.h
@@ -141,7 +141,7 @@ namespace LowLevelTasks
 			uintptr_t bAllowBusyWaiting : 1;
 			
 		public:
-			FPackedData() : FPackedData(nullptr,  ETaskPriority::Default, ETaskState::Completed, true)
+			FPackedData() : FPackedData(nullptr,  ETaskPriority::Count, ETaskState::Completed, true)
 			{
 				static_assert(!PLATFORM_32BITS, "32bit Platforms are not supported");
 				static_assert(uintptr_t(ETaskPriority::Count) <= (1ull << 3), "Not enough bits to store ETaskPriority");
