Fix infinite wait on async compilation when in pause mode by using an unpausable priority when forced compilation is required to wait

diff --git a/Engine/Source/Developer/DerivedDataCache/Private/DerivedDataCacheStoreAsync.cpp b/Engine/Source/Developer/DerivedDataCache/Private/DerivedDataCacheStoreAsync.cpp
--- a/Engine/Source/Developer/DerivedDataCache/Private/DerivedDataCacheStoreAsync.cpp
+++ b/Engine/Source/Developer/DerivedDataCache/Private/DerivedDataCacheStoreAsync.cpp
@@ -218,7 +218,7 @@ private:
 	{
 		switch (Priority)
 		{
-		case EPriority::Blocking: return EQueuedWorkPriority::Highest;
+		case EPriority::Blocking: return EQueuedWorkPriority::Blocking;
 		case EPriority::Highest:  return EQueuedWorkPriority::Highest;
 		case EPriority::High:     return EQueuedWorkPriority::High;
 		case EPriority::Normal:   return EQueuedWorkPriority::Normal;
diff --git a/Engine/Source/Runtime/Core/Private/Misc/QueuedThreadPoolWrapper.cpp b/Engine/Source/Runtime/Core/Private/Misc/QueuedThreadPoolWrapper.cpp
--- a/Engine/Source/Runtime/Core/Private/Misc/QueuedThreadPoolWrapper.cpp
+++ b/Engine/Source/Runtime/Core/Private/Misc/QueuedThreadPoolWrapper.cpp
@@ -201,9 +201,9 @@ void FQueuedThreadPoolWrapper::ReleaseWorkNoLock(FScheduledWork* Work)
 	WorkPool.Push(Work);
 }
 
-bool FQueuedThreadPoolWrapper::CanSchedule() const
+bool FQueuedThreadPoolWrapper::CanSchedule(EQueuedWorkPriority Priority) const
 {
-	return (MaxTaskToSchedule == -1 || MaxTaskToSchedule > 0) && CurrentConcurrency < GetMaxConcurrency();
+	return (MaxTaskToSchedule == -1 || MaxTaskToSchedule > 0 || Priority == EQueuedWorkPriority::Blocking) && CurrentConcurrency < GetMaxConcurrency();
 }
 
 FQueuedThreadPoolWrapper::FScheduledWork* FQueuedThreadPoolWrapper::AllocateWork(IQueuedWork* InnerWork, EQueuedWorkPriority Priority)
@@ -263,13 +263,13 @@ void FQueuedThreadPoolWrapper::Schedule(FScheduledWork* Work)
 	if (NextWork)
 	{
 		// Continue retracting more work until nothing can be retracted anymore or we can finally squeeze the higher priority task in
-		while (!CanSchedule() && TryRetractWorkNoLock(NextWorkPriority))
+		while (!CanSchedule(NextWorkPriority) && TryRetractWorkNoLock(NextWorkPriority))
 		{
 		}
 	}
 
 	// Schedule as many tasks we can fit
-	while (CanSchedule())
+	while (CanSchedule(NextWorkPriority))
 	{
 		EQueuedWorkPriority WorkPriority;
 		IQueuedWork* InnerWork = QueuedWork.Dequeue(&WorkPriority);
@@ -281,7 +281,7 @@ void FQueuedThreadPoolWrapper::Schedule(FScheduledWork* Work)
 			Work = AllocateWork(InnerWork, WorkPriority);
 			ScheduledWork.Add(InnerWork, Work);
 			OnScheduled(Work);
-			WrappedQueuedThreadPool->AddQueuedWork(Work, PriorityMapper(WorkPriority));
+			WrappedQueuedThreadPool->AddQueuedWork(Work, WorkPriority == EQueuedWorkPriority::Blocking ? WorkPriority : PriorityMapper(WorkPriority));
 
 			if (MaxTaskToSchedule > 0)
 			{
diff --git a/Engine/Source/Runtime/Core/Public/Misc/QueuedThreadPool.h b/Engine/Source/Runtime/Core/Public/Misc/QueuedThreadPool.h
--- a/Engine/Source/Runtime/Core/Public/Misc/QueuedThreadPool.h
+++ b/Engine/Source/Runtime/Core/Public/Misc/QueuedThreadPool.h
@@ -10,11 +10,12 @@ class IQueuedWork;
 /** Higher priority are picked up first by the task thread pool. */
 enum class EQueuedWorkPriority : uint8
 {
-	Highest = 0,
-	High = 1,
-	Normal = 2,
-	Low = 3,
-	Lowest = 4,
+	Blocking = 0,
+	Highest = 1,
+	High = 2,
+	Normal = 3,
+	Low = 4,
+	Lowest = 5,
 	Count
 };
 
@@ -22,6 +23,8 @@ inline const TCHAR* LexToString(EQueuedWorkPriority Priority)
 {
 	switch (Priority)
 	{
+	case EQueuedWorkPriority::Blocking:
+		return TEXT("Blocking");
 	case EQueuedWorkPriority::Highest:
 		return TEXT("Highest");
 	case EQueuedWorkPriority::High:
diff --git a/Engine/Source/Runtime/Core/Public/Misc/QueuedThreadPoolWrapper.h b/Engine/Source/Runtime/Core/Public/Misc/QueuedThreadPoolWrapper.h
--- a/Engine/Source/Runtime/Core/Public/Misc/QueuedThreadPoolWrapper.h
+++ b/Engine/Source/Runtime/Core/Public/Misc/QueuedThreadPoolWrapper.h
@@ -65,7 +65,7 @@ protected:
 private:
 	struct FScheduledWork;
 	FScheduledWork* AllocateWork(IQueuedWork* InnerWork, EQueuedWorkPriority Priority);
-	bool CanSchedule() const;
+	bool CanSchedule(EQueuedWorkPriority Priority) const;
 	bool Create(uint32 InNumQueuedThreads, uint32 StackSize, EThreadPriority ThreadPriority, const TCHAR* Name) override;
 	void Destroy() override;
 	void Schedule(FScheduledWork* Work = nullptr);
@@ -327,7 +327,7 @@ private:
 		InQueuedWork->InternalData = QueuedWorkInternalData;
 		EQueuedWorkPriority Priority = PriorityMapper(InPriority);
 
-		LowLevelTasks::ETaskPriority TaskPriorityMapper[int(EQueuedWorkPriority::Count)] = { LowLevelTasks::ETaskPriority::High, LowLevelTasks::ETaskPriority::BackgroundHigh, LowLevelTasks::ETaskPriority::BackgroundNormal, LowLevelTasks::ETaskPriority::BackgroundLow, LowLevelTasks::ETaskPriority::BackgroundLow };
+		LowLevelTasks::ETaskPriority TaskPriorityMapper[int(EQueuedWorkPriority::Count)] = { LowLevelTasks::ETaskPriority::High, LowLevelTasks::ETaskPriority::High, LowLevelTasks::ETaskPriority::BackgroundHigh, LowLevelTasks::ETaskPriority::BackgroundNormal, LowLevelTasks::ETaskPriority::BackgroundLow, LowLevelTasks::ETaskPriority::BackgroundLow };
 		LowLevelTasks::ETaskPriority TaskPriority = TaskPriorityMapper[int(Priority)];
 		const bool bAllowBusyWaiting = (InQueuedWork->GetQueuedWorkFlags() & EQueuedWorkFlags::DoNotRunInsideBusyWait) == EQueuedWorkFlags::None;
 
diff --git a/Engine/Source/Runtime/Engine/Private/AssetCompilingManager.cpp b/Engine/Source/Runtime/Engine/Private/AssetCompilingManager.cpp
--- a/Engine/Source/Runtime/Engine/Private/AssetCompilingManager.cpp
+++ b/Engine/Source/Runtime/Engine/Private/AssetCompilingManager.cpp
@@ -221,8 +221,8 @@ FQueuedThreadPool* FAssetCompilingManager::GetThreadPool() const
 
 		// All asset priorities will resolve to a Low priority once being scheduled.
 		// Any asset supporting being built async should be scheduled lower than Normal to let non-async stuff go first
-		// However, we let Highest priority pass-through as it to benefit from going to foreground threads when required (i.e. Game-thread is waiting on some assets)
-		GAssetThreadPool = new FMemoryBoundQueuedThreadPoolWrapper(LimitedThreadPool, -1, [](EQueuedWorkPriority Priority) { return Priority == EQueuedWorkPriority::Highest ? Priority : EQueuedWorkPriority::Low; });
+		// However, we let Highest and Blocking priority pass-through as it to benefit from going to foreground threads when required (i.e. Game-thread is waiting on some assets)
+		GAssetThreadPool = new FMemoryBoundQueuedThreadPoolWrapper(LimitedThreadPool, -1, [](EQueuedWorkPriority Priority) { return Priority <= EQueuedWorkPriority::Highest ? Priority : EQueuedWorkPriority::Low; });
 
 		AsyncCompilationHelpers::BindThreadPoolToCVar(
 			GAssetThreadPool,
diff --git a/Engine/Source/Runtime/Engine/Private/AsyncCompilationHelpers.cpp b/Engine/Source/Runtime/Engine/Private/AsyncCompilationHelpers.cpp
--- a/Engine/Source/Runtime/Engine/Private/AsyncCompilationHelpers.cpp
+++ b/Engine/Source/Runtime/Engine/Private/AsyncCompilationHelpers.cpp
@@ -121,12 +121,12 @@ namespace AsyncCompilationHelpers
 			SlowTask->MakeDialogDelayed(1.0f, false /*bShowCancelButton*/, true /*bAllowInPIE*/);
 		}
 
-		// Reschedule everything to be executed at highest priority
+		// Reschedule everything to be executed at blocking priority, it bypasses the pause mechanism to ensure forward progress since we're waiting.
 		// Important to keep using the asset compiling manager's thread pool to benefit from it's memory management and avoid OOM.
 		FQueuedThreadPool* AssetThreadPool = FAssetCompilingManager::Get().GetThreadPool();
 		for (int32 Index = 0; Index < Num; ++Index)
 		{
-			Getter(Index).Reschedule(AssetThreadPool, EQueuedWorkPriority::Highest);
+			Getter(Index).Reschedule(AssetThreadPool, EQueuedWorkPriority::Blocking);
 		}
 
 		auto FormatProgress =
diff --git a/Engine/Source/Runtime/Engine/Private/DistanceFieldAtlas.cpp b/Engine/Source/Runtime/Engine/Private/DistanceFieldAtlas.cpp
--- a/Engine/Source/Runtime/Engine/Private/DistanceFieldAtlas.cpp
+++ b/Engine/Source/Runtime/Engine/Private/DistanceFieldAtlas.cpp
@@ -603,7 +603,7 @@ void FDistanceFieldAsyncQueue::BlockUntilBuildComplete(UStaticMesh* StaticMesh,
 					}
 #endif
 
-					RescheduleBackgroundTask(Task, EQueuedWorkPriority::Highest);
+					RescheduleBackgroundTask(Task, EQueuedWorkPriority::Blocking);
 				}
 			}
 		}
@@ -656,7 +656,7 @@ void FDistanceFieldAsyncQueue::BlockUntilAllBuildsComplete()
 			FScopeLock Lock(&CriticalSection);
 			for (FAsyncDistanceFieldTask* Task : ReferencedTasks)
 			{
-				RescheduleBackgroundTask(Task, EQueuedWorkPriority::Highest);
+				RescheduleBackgroundTask(Task, EQueuedWorkPriority::Blocking);
 			}
 		}
 
diff --git a/Engine/Source/Runtime/Engine/Private/MeshCardRepresentation.cpp b/Engine/Source/Runtime/Engine/Private/MeshCardRepresentation.cpp
--- a/Engine/Source/Runtime/Engine/Private/MeshCardRepresentation.cpp
+++ b/Engine/Source/Runtime/Engine/Private/MeshCardRepresentation.cpp
@@ -569,7 +569,7 @@ void FCardRepresentationAsyncQueue::BlockUntilBuildComplete(UStaticMesh* StaticM
 					}
 #endif
 
-					RescheduleBackgroundTask(Task, EQueuedWorkPriority::Highest);
+					RescheduleBackgroundTask(Task, EQueuedWorkPriority::Blocking);
 				}
 			}
 		}
@@ -622,7 +622,7 @@ void FCardRepresentationAsyncQueue::BlockUntilAllBuildsComplete()
 			FScopeLock Lock(&CriticalSection);
 			for (FAsyncCardRepresentationTask* Task : ReferencedTasks)
 			{
-				RescheduleBackgroundTask(Task, EQueuedWorkPriority::Highest);
+				RescheduleBackgroundTask(Task, EQueuedWorkPriority::Blocking);
 			}
 		}
 
diff --git a/Engine/Source/Runtime/Engine/Private/TextureCompiler.cpp b/Engine/Source/Runtime/Engine/Private/TextureCompiler.cpp
--- a/Engine/Source/Runtime/Engine/Private/TextureCompiler.cpp
+++ b/Engine/Source/Runtime/Engine/Private/TextureCompiler.cpp
@@ -16,6 +16,7 @@
 #include "Materials/MaterialInstance.h"
 #include "Materials/Material.h"
 #include "TextureDerivedDataTask.h"
+#include "ProfilingDebugging/CountersTrace.h"
 #include "Misc/IQueuedWork.h"
 #include "Components/PrimitiveComponent.h"
 #include "AsyncCompilationHelpers.h"
@@ -56,7 +57,7 @@ namespace TextureCompilingManagerImpl
 
 	static EQueuedWorkPriority GetBoostPriority(UTexture* InTexture)
 	{
-		return (EQueuedWorkPriority)(FMath::Max((uint8)1, (uint8)GetBasePriority(InTexture)) - 1);
+		return (EQueuedWorkPriority)(FMath::Max((uint8)EQueuedWorkPriority::Highest, (uint8)GetBasePriority(InTexture)) - 1);
 	}
 
 	static void EnsureInitializedCVars()
@@ -192,8 +193,10 @@ bool FTextureCompilingManager::IsAsyncTextureCompilationEnabled() const
 	return CVarAsyncTextureStandard.AsyncCompilation.GetValueOnAnyThread() != 0;
 }
 
+TRACE_DECLARE_INT_COUNTER(QueuedTextureCompilation, TEXT("AsyncCompilation/QueuedTexture"));
 void FTextureCompilingManager::UpdateCompilationNotification()
 {
+	TRACE_COUNTER_SET(QueuedTextureCompilation, GetNumRemainingTextures());
 	Notification.Update(GetNumRemainingTextures());
 }
 
@@ -270,6 +273,8 @@ void FTextureCompilingManager::AddTextures(TArrayView<UTexture* const> InTexture
 		}
 		RegisteredTextureBuckets[TexturePriority].Emplace(Texture);
 	}
+
+	TRACE_COUNTER_SET(QueuedTextureCompilation, GetNumRemainingTextures());
 }
 
 void FTextureCompilingManager::FinishCompilation(TArrayView<UTexture* const> InTextures)
diff --git a/Engine/Source/Runtime/Engine/Private/TextureDerivedDataTask.cpp b/Engine/Source/Runtime/Engine/Private/TextureDerivedDataTask.cpp
--- a/Engine/Source/Runtime/Engine/Private/TextureDerivedDataTask.cpp
+++ b/Engine/Source/Runtime/Engine/Private/TextureDerivedDataTask.cpp
@@ -1521,7 +1521,7 @@ private:
 		case EPriority::Normal:   return EQueuedWorkPriority::Normal;
 		case EPriority::High:     return EQueuedWorkPriority::High;
 		case EPriority::Highest:  return EQueuedWorkPriority::Highest;
-		case EPriority::Blocking: return EQueuedWorkPriority::Highest;
+		case EPriority::Blocking: return EQueuedWorkPriority::Blocking;
 		default:                  return EQueuedWorkPriority::Normal;
 		}
 	}
