Implement FUnixPlatformMisc::GetPrimaryGPUBrand()

For NVIDIA, should return something like:

  "NVIDIA GeForce RTX 3080 Ti (NVIDIA UNIX x86_64 Kernel Module  470.86  Tue Oct 26 21:55:45 UTC 2021)"

And AMD Mesa:

  "AMD PCI-id: 1002-731f ([REDACTED_PHONE])"

diff --git a/Engine/Source/Runtime/Core/Private/Unix/UnixPlatformMisc.cpp b/Engine/Source/Runtime/Core/Private/Unix/UnixPlatformMisc.cpp
--- a/Engine/Source/Runtime/Core/Private/Unix/UnixPlatformMisc.cpp
+++ b/Engine/Source/Runtime/Core/Private/Unix/UnixPlatformMisc.cpp
@@ -39,6 +39,8 @@
 #include <net/if_arp.h>
 #include <linux/limits.h>
 
+#include <glob.h>
+
 #include "Modules/ModuleManager.h"
 #include "HAL/ThreadHeartBeat.h"
 
@@ -188,50 +190,6 @@ void FUnixPlatformMisc::PlatformPreInit()
 	UnixCrashReporterTracker::PreInit();
 }
 
-static FString ReadGPUFile(const char *Filename)
-{
-	FString Value;
-	FILE* FsFile = fopen(Filename, "r");
-
-	if (FsFile)
-	{
-		ANSICHAR LineBuffer[512] = { 0 };
-		ANSICHAR *Line = fgets(LineBuffer, sizeof(LineBuffer), FsFile);
-
-		if (Line)
-		{
-			Value = UTF8_TO_TCHAR(LineBuffer);
-			Value.TrimStartAndEndInline();
-		}
-
-		fclose(FsFile);
-	}
-
-	return Value;
-}
-
-static FString GetGPUInfo()
-{
-	FString GPUInfo;
-	FString Device0 = ReadGPUFile("/sys/class/drm/card0/device/device");
-	FString Device1 = ReadGPUFile("/sys/class/drm/card1/device/device");
-	FString NVIDIAVersion = ReadGPUFile("/proc/driver/nvidia/version");
-
-	if (!Device0.IsEmpty())
-	{
-		FString Vendor0 = ReadGPUFile("/sys/class/drm/card0/device/vendor");
-		GPUInfo += FString::Printf(TEXT("Card0 PCI-id: %s-%s "), *Vendor0, *Device0);
-	}
-	if (!Device1.IsEmpty())
-	{
-		FString Vendor1 = ReadGPUFile("/sys/class/drm/card1/device/vendor");
-		GPUInfo += FString::Printf(TEXT("Card1 PCI-id: %s-%s "), *Vendor1, *Device1);
-	}
-	GPUInfo += NVIDIAVersion;
-
-	return GPUInfo.TrimStartAndEnd();
-}
-
 void FUnixPlatformMisc::PlatformInit()
 {
 	// install a platform-specific signal handler
@@ -245,7 +203,7 @@ void FUnixPlatformMisc::PlatformInit()
 
 	UnixPlatForm_CheckIfKSMUsable();
 
-	FString GPUInfo = GetGPUInfo();
+	FString GPUBrandInfo = GetPrimaryGPUBrand();
 
 	UE_LOG(LogInit, Log, TEXT("Unix hardware info:"));
 	UE_LOG(LogInit, Log, TEXT(" - we are %sthe first instance of this executable"), bFirstInstance ? TEXT("") : TEXT("not "));
@@ -260,10 +218,12 @@ void FUnixPlatformMisc::PlatformInit()
 	UE_LOG(LogInit, Log, TEXT(" - Number of logical cores available for the process: %d"), FPlatformMisc::NumberOfCoresIncludingHyperthreads());
 	UnixPlatform_UpdateCacheLineSize();
 	UE_LOG(LogInit, Log, TEXT(" - Cache line size: %zu"), GCacheLineSize);
-	if (!GPUInfo.IsEmpty())
+
+	if (!GPUBrandInfo.IsEmpty())
 	{
-		UE_LOG(LogInit, Log, TEXT(" - GPU Info. %s"), *GPUInfo);
+		UE_LOG(LogInit, Log, TEXT(" - GPU Brand Info: %s"), *GPUBrandInfo);
 	}
+
 	UE_LOG(LogInit, Log, TEXT(" - Memory allocator used: %s"), GMalloc->GetDescriptiveName());
 	UE_LOG(LogInit, Log, TEXT(" - This binary is optimized with LTO: %s, PGO: %s, instrumented for PGO data collection: %s"),
 		PLATFORM_COMPILER_OPTIMIZATION_LTCG ? TEXT("yes") : TEXT("no"),
@@ -716,7 +676,7 @@ int32 FUnixPlatformMisc::NumberOfCores()
 				int NumPackages = MaxPackageId + 1;
 				int NumPairs = NumPackages * NumCores;
 
-				// AArch64 topology seems to be incompatible with the above assumptions, particularly, core_id can be all 0 while the cores themselves are obviously independent. 
+				// AArch64 topology seems to be incompatible with the above assumptions, particularly, core_id can be all 0 while the cores themselves are obviously independent.
 				// Check if num CPUs available to us is more than 2 per core (i.e. more than reasonable when hyperthreading is involved), and if so, don't trust the topology.
 				if (2 * NumCores < NumCpusAvailable)
 				{
@@ -1178,7 +1138,7 @@ uint32 FUnixPlatformMisc::GetCPUInfo()
 FString FUnixPlatformMisc::GetCPUBrand()
 {
 #if PLATFORM_HAS_CPUID
-	static TCHAR Result[64] = TEXT("NonX86CPUBrand");
+	static FString Result = TEXT("NonX86CPUBrand");
 	static bool bHaveResult = false;
 
 	if (!bHaveResult)
@@ -1202,16 +1162,144 @@ FString FUnixPlatformMisc::GetCPUBrand()
 			}
 		}
 
-		FCString::Strncpy(Result, UTF8_TO_TCHAR(BrandString), UE_ARRAY_COUNT(Result));
+		Result = UTF8_TO_TCHAR(BrandString);
+		Result.TrimStartAndEndInline();
 		bHaveResult = true;
 	}
 
-	return FString(Result);
+	return Result;
 #else
 	return FPlatformProcCPUInfo::GetCPUBrand();
 #endif // PLATFORM_HAS_CPUID
 }
 
+// From RHIVendorIdToString()
+static const TCHAR* VendorIdToString(uint32 VendorId)
+{
+	switch (VendorId)
+	{
+		case 0x1002:  return TEXT("AMD");
+		case 0x1010:  return TEXT("ImgTec");
+		case 0x10DE:  return TEXT("NVIDIA");
+		case 0x13B5:  return TEXT("ARM");
+		case 0x14E4:  return TEXT("Broadcom");
+		case 0x5143:  return TEXT("Qualcomm");
+		case 0x106B:  return TEXT("Apple");
+		case 0x8086:  return TEXT("Intel");
+		case 0x7a05:  return TEXT("Vivante");
+		case 0x1EB1:  return TEXT("VeriSilicon");
+		case 0x10003: return TEXT("Kazan");
+		case 0x10004: return TEXT("Codeplay");
+		case 0x10005: return TEXT("Mesa");
+	default:
+		break;
+	}
+
+	return TEXT("UnknownVendor");
+}
+
+static FString ReadGPUFile(const char *Filename, const char *Token = nullptr)
+{
+	FString Value;
+	FILE* FsFile = fopen(Filename, "r");
+
+	if (FsFile)
+	{
+		uint32 TokenLen = Token ? FCStringAnsi::Strlen(Token) : 0;
+
+		for (;;)
+		{
+			ANSICHAR LineBuffer[512];
+			ANSICHAR *Line = fgets(LineBuffer, sizeof(LineBuffer), FsFile);
+
+			if (!Line)
+			{
+				break;
+			}
+
+			if (!Token)
+			{
+				// Skip over the 0x prefix
+				if (Line[0] == '0' && Line[1] == 'x')
+				{
+					Line += 2;
+				}
+				Value = UTF8_TO_TCHAR(Line);
+				break;
+			}
+			if (!FCStringAnsi::Strncmp(Line, Token, TokenLen))
+			{
+				Value = UTF8_TO_TCHAR(Line) + TokenLen;
+				break;
+			}
+		}
+		Value.TrimStartAndEndInline();
+
+		fclose(FsFile);
+	}
+
+	return Value;
+}
+
+// Return primary GPU brand string
+//  This should trying to report out what GPU is in use: Nvidia/Type/Driver version.
+FString FUnixPlatformMisc::GetPrimaryGPUBrand()
+{
+	static FString ProcGPUBrandStr;
+
+	if (ProcGPUBrandStr.IsEmpty())
+	{
+		// https://download.nvidia.com/XFree86/Linux-x86_64/440.31/README/procinterface.html
+		FString NVIDIAVersion = ReadGPUFile("/proc/driver/nvidia/version", "NVRM version:");
+		if (!NVIDIAVersion.IsEmpty())
+		{
+			// Returns something ~ like "NVIDIA GeForce RTX 3080 Ti (NVIDIA UNIX x86_64 Kernel Module  470.86  Tue Oct 26 21:55:45 UTC 2021)"
+			FString GPUModel;
+			glob_t GlobResults;
+
+			// Try to read GPU Model information from:
+			//   /proc/driver/nvidia/gpus/domain:bus:device.function/information
+			if (glob("/proc/driver/nvidia/gpus/**/information", GLOB_NOSORT, nullptr, &GlobResults) == 0)
+			{
+				for (int32 Idx = 0; Idx < GlobResults.gl_pathc; Idx++)
+				{
+					GPUModel = ReadGPUFile(GlobResults.gl_pathv[Idx], "Model:");
+					if (!GPUModel.IsEmpty())
+					{
+						break;
+					}
+				}
+
+				globfree(&GlobResults);
+			}
+
+			ProcGPUBrandStr = FString::Printf(TEXT("%s (%s)"), GPUModel.IsEmpty() ? TEXT("NVIDIA") : *GPUModel, *NVIDIAVersion);
+		}
+		else
+		{
+			// Returns something ~ like "AMD PCI-id: 1002-731f (1462-3811)"
+			FString Vendor0 = ReadGPUFile("/sys/class/drm/card0/device/vendor");
+			FString Device0 = ReadGPUFile("/sys/class/drm/card0/device/device");
+			FString SubsystemVendor0 = ReadGPUFile("/sys/class/drm/card0/device/subsystem_vendor");
+			FString SubsystemDevice0 = ReadGPUFile("/sys/class/drm/card0/device/subsystem_device");
+			uint32 VendorId = strtoul(TCHAR_TO_UTF8(*Vendor0), nullptr, 16);
+			const TCHAR* VendorStr = VendorIdToString(VendorId);
+
+			ProcGPUBrandStr = VendorStr;
+			if (!Device0.IsEmpty())
+			{
+				ProcGPUBrandStr += FString::Printf(TEXT(" PCI-id: %s-%s"), *Vendor0, *Device0);
+			}
+			if (!SubsystemVendor0.IsEmpty() || !SubsystemDevice0.IsEmpty())
+			{
+				ProcGPUBrandStr += FString::Printf(TEXT(" (%s-%s)"), *SubsystemVendor0, *SubsystemDevice0);
+			}
+		}
+	}
+
+	return ProcGPUBrandStr;
+}
+
 // __builtin_popcountll() will not be compiled to use popcnt instruction unless -mpopcnt or a sufficiently recent target CPU arch is passed (which UBT doesn't by default)
 #if defined(__POPCNT__)
 	#define UE4_LINUX_NEED_TO_CHECK_FOR_POPCNT_PRESENCE				(PLATFORM_ENABLE_POPCNT_INTRINSIC)
@@ -1229,7 +1317,7 @@ bool FUnixPlatformMisc::HasNonoptionalCPUFeatures()
 #if PLATFORM_HAS_CPUID
 		int Info[4];
 		__cpuid(1, Info[0], Info[1], Info[2], Info[3]);
-	
+
 	#if UE4_LINUX_NEED_TO_CHECK_FOR_POPCNT_PRESENCE
 		bHasNonOptionalFeature = (Info[2] & (1 << 23)) != 0;
 	#endif // UE4_LINUX_NEED_TO_CHECK_FOR_POPCNT
@@ -1261,18 +1349,18 @@ bool FUnixPlatformMisc::IsDebuggerPresent()
 	}
 
 	// If a process is tracing this one then TracerPid in /proc/self/status will
-	// be the id of the tracing process. Use SignalHandler safe functions 
-	
+	// be the id of the tracing process. Use SignalHandler safe functions
+
 	int StatusFile = open("/proc/self/status", O_RDONLY);
-	if (StatusFile == -1) 
+	if (StatusFile == -1)
 	{
 		// Failed - unknown debugger status.
 		return false;
 	}
 
-	char Buffer[256];	
+	char Buffer[256];
 	ssize_t Length = read(StatusFile, Buffer, sizeof(Buffer));
-	
+
 	bool bDebugging = false;
 	const char* TracerString = "TracerPid:\t";
 	const ssize_t LenTracerString = strlen(TracerString);
@@ -1486,28 +1574,28 @@ namespace
 #if !defined(GRND_NONBLOCK)
 	#define GRND_NONBLOCK 0x0001
 #endif
-	
+
 	int SysGetRandom(void *buf, size_t buflen)
 	{
 		if (SysGetRandomSupported < 0)
 		{
 			int Ret = syscall(SYS_getrandom, buf, buflen, GRND_NONBLOCK);
-	
+
 			// If -1 is returned with ENOSYS, kernel doesn't support getrandom
 			SysGetRandomSupported = ((Ret == -1) && (errno == ENOSYS)) ? 0 : 1;
 		}
-	
+
 		return SysGetRandomSupported ?
 			syscall(SYS_getrandom, buf, buflen, GRND_NONBLOCK) : -1;
 	}
-	
+
 #else
 
 	int SysGetRandom(void *buf, size_t buflen)
 	{
 		return -1;
 	}
-	
+
 #endif // !SYS_getrandom
 }
 
diff --git a/Engine/Source/Runtime/Core/Public/Unix/UnixPlatformMisc.h b/Engine/Source/Runtime/Core/Public/Unix/UnixPlatformMisc.h
--- a/Engine/Source/Runtime/Core/Public/Unix/UnixPlatformMisc.h
+++ b/Engine/Source/Runtime/Core/Public/Unix/UnixPlatformMisc.h
@@ -114,6 +114,8 @@ struct CORE_API FUnixPlatformMisc : public FGenericPlatformMisc
 	 */
 	static uint32 GetCPUInfo();
 
+	static FString GetPrimaryGPUBrand();
+
 	static bool HasNonoptionalCPUFeatures();
 	static bool NeedsNonoptionalCPUFeaturesCheck();
 
diff --git a/Engine/Source/Runtime/VulkanRHI/Private/VulkanRHI.cpp b/Engine/Source/Runtime/VulkanRHI/Private/VulkanRHI.cpp
--- a/Engine/Source/Runtime/VulkanRHI/Private/VulkanRHI.cpp
+++ b/Engine/Source/Runtime/VulkanRHI/Private/VulkanRHI.cpp
@@ -739,6 +739,8 @@ void FVulkanDynamicRHI::SelectAndInitDevice()
 		GRHIAdapterInternalDriverVersion = FString::Printf(TEXT("%d.%d.%d (0x%X)"), VK_VERSION_MAJOR(Props.apiVersion), VK_VERSION_MINOR(Props.apiVersion), VK_VERSION_PATCH(Props.apiVersion), Props.apiVersion);
 		GRHIAdapterUserDriverVersion = FString::Printf(TEXT("%d.%d.%d (0x%X)"), VK_VERSION_MAJOR(Props.driverVersion), VK_VERSION_MINOR(Props.driverVersion), VK_VERSION_PATCH(Props.driverVersion), Props.driverVersion);
 		GRHIDeviceId = Props.deviceID;
+
+		UE_LOG(LogVulkanRHI, Display, TEXT("'%s' User Driver Version = %s"), *GRHIAdapterName, *GRHIAdapterUserDriverVersion);
 	}
 
 	GRHIPersistentThreadGroupCount = 1440; // TODO: Revisit based on vendor/adapter/perf query
