Incorrect shader format null RHI was used on Mac.

According to DataDrivenPlatformInfo.ini, Mac only supports shader platforms SP_METAL_SM5, SP_METAL_MRT_MAC, and SP_METAL_MACES3_1. The SP_METAL platforms was used when creating a null RHI on Mac however and resulting in an assert of an unsupported shader platform when running a dedicated server (and -nullrhi one presumes although that remains untested).

#rnx

diff --git a/Engine/Source/Runtime/Core/Public/Mac/MacPlatformMisc.h b/Engine/Source/Runtime/Core/Public/Mac/MacPlatformMisc.h
--- a/Engine/Source/Runtime/Core/Public/Mac/MacPlatformMisc.h
+++ b/Engine/Source/Runtime/Core/Public/Mac/MacPlatformMisc.h
@@ -53,7 +53,7 @@ struct CORE_API FMacPlatformMisc : public FApplePlatformMisc
 	 */
 	FORCEINLINE static const TCHAR* GetNullRHIShaderFormat() 
 	{ 
-		return TEXT("SF_METAL"); 
+		return TEXT("SF_METAL_SM5"); 
 	}
 
 	/**
