- Fixed a StressTestUnity issue with global variable name

diff --git a/Engine/Source/Runtime/Core/Private/GenericPlatform/GenericPlatformMisc.cpp b/Engine/Source/Runtime/Core/Private/GenericPlatform/GenericPlatformMisc.cpp
--- a/Engine/Source/Runtime/Core/Private/GenericPlatform/GenericPlatformMisc.cpp
+++ b/Engine/Source/Runtime/Core/Private/GenericPlatform/GenericPlatformMisc.cpp
@@ -252,7 +252,7 @@ static FAutoConsoleCommand GLogNamedEventsCmd(
 #endif//LOG_NAMED_EVENTS
 
 /** Holds an override path if a program has special needs */
-FString OverrideProjectDir;
+static FString GOverrideProjectDir;
 
 /** Hooks for moving ClipboardCopy and ClipboardPaste into FPlatformApplicationMisc */
 CORE_API void (*ClipboardCopyShim)(const TCHAR* Text) = nullptr;
@@ -1240,7 +1240,7 @@ const TCHAR* FGenericPlatformMisc::ProjectDir()
 	if (ProjectDir.Len() == 0)
 	{
 		ProjectDir.Reserve(FPlatformMisc::GetMaxPathLength());
-		ProjectDir = OverrideProjectDir;
+		ProjectDir = GOverrideProjectDir;
 	}
 
 	if (ProjectDir.Len() == 0)
@@ -1393,7 +1393,7 @@ float FGenericPlatformMisc::GetDeviceTemperatureLevel()
 
 void FGenericPlatformMisc::SetOverrideProjectDir(const FString& InOverrideDir)
 {
-	OverrideProjectDir = InOverrideDir;
+	GOverrideProjectDir = InOverrideDir;
 }
 
 bool FGenericPlatformMisc::UseRenderThread()
