Init the thread priority to TPri_Num during thread creation to ensure that it's always set in SetThreadPriority, even if TPri_Normal is requested. On some platforms, threads inherit their priority from the creating thread rather than using a default value.

diff --git a/Engine/Source/Runtime/Core/Private/HAL/PThreadRunnableThread.h b/Engine/Source/Runtime/Core/Private/HAL/PThreadRunnableThread.h
--- a/Engine/Source/Runtime/Core/Private/HAL/PThreadRunnableThread.h
+++ b/Engine/Source/Runtime/Core/Private/HAL/PThreadRunnableThread.h
@@ -306,7 +306,7 @@ protected:
 		ThreadInitSyncEvent	= FPlatformProcess::GetSynchEventFromPool(true);
 		// A name for the thread in for debug purposes. _ThreadProc will set it.
 		ThreadName = InThreadName ? InThreadName : TEXT("Unnamed UE4");
-		ThreadPriority = InThreadPri;
+		ThreadPriority = InThreadPri; // Make sure the thread priority is correct during init
 		ThreadAffinityMask = InThreadAffinityMask;
 
 		// Create the new thread
@@ -318,7 +318,7 @@ protected:
 			ThreadInitSyncEvent->Wait((uint32)-1); // infinite wait
 
 			// set the priority
-			ThreadPriority = TPri_Normal; // set back to default as some impls check if calling syscalls is necessary 
+			ThreadPriority = TPri_Num; // set to Num (invalid) to ensure that the priority will be set below
 			SetThreadPriority(InThreadPri);
 		}
 		else // If it fails, clear all the vars
