Added child offset number to linux parent server logs, to help trace issues related to forking.

diff --git a/Engine/Source/Runtime/Core/Private/Unix/UnixPlatformProcess.cpp b/Engine/Source/Runtime/Core/Private/Unix/UnixPlatformProcess.cpp
--- a/Engine/Source/Runtime/Core/Private/Unix/UnixPlatformProcess.cpp
+++ b/Engine/Source/Runtime/Core/Private/Unix/UnixPlatformProcess.cpp
@@ -1424,7 +1424,7 @@ FGenericPlatformProcess::EWaitAndForkResult FUnixPlatformProcess::WaitAndFork()
 		FPidAndSignal(pid_t InPid, int32 InSignalValue) : Pid(InPid), SignalValue(InSignalValue) {}
 	};
 	TArray<FPidAndSignal> AllChildren;
-	AllChildren.Reserve(1024); // Sized to be big enough that it probably wont reallocte, but its not the end of the world if it does.
+	AllChildren.Reserve(1024); // Sized to be big enough that it probably wont reallocate, but its not the end of the world if it does.
 	while (!IsEngineExitRequested())
 	{
 		BeginExitIfRequested();
@@ -1435,6 +1435,10 @@ FGenericPlatformProcess::EWaitAndForkResult FUnixPlatformProcess::WaitAndFork()
 			// Sleep for a short while to avoid spamming new processes to the OS all at once
 			FPlatformProcess::Sleep(WAIT_AND_FORK_CHILD_SPAWN_DELAY);
 
+			uint16 Cookie = (SignalValue >> 16) & 0xffff;
+			uint16 ChildIdx = SignalValue & 0xffff;
+			UE_LOG(LogHAL, Log, TEXT("[Parent] WaitAndFork processing child request %04hx-%04hx."), Cookie, ChildIdx);
+
 			FMemoryStatsHolder CurrentMasterMemStats(FPlatformMemory::GetStats());
 			UE_LOG(LogHAL, Log, TEXT("MemoryStats PreFork: AvailablePhysical: %.02fMiB (%+.02fMiB), PeakPhysical: %.02fMiB, PeakVirtual: %.02fMiB"),
 				CurrentMasterMemStats.AvailablePhysical, (CurrentMasterMemStats.AvailablePhysical - PreviousMasterMemStats.AvailablePhysical),
@@ -1462,8 +1466,6 @@ FGenericPlatformProcess::EWaitAndForkResult FUnixPlatformProcess::WaitAndFork()
 			else if (ChildPID == 0)
 			{
 				// Child
-				uint16 Cookie = (SignalValue >> 16) & 0xffff;
-				uint16 ChildIdx = SignalValue & 0xffff;
 				FForkProcessHelper::SetIsForkedChildProcess(ChildIdx);
 
 				if (FPlatformMemory::HasForkPageProtectorEnabled())
@@ -1498,7 +1500,7 @@ FGenericPlatformProcess::EWaitAndForkResult FUnixPlatformProcess::WaitAndFork()
 				// Set the process name, if specified
 				if (ChildIdx > 0)
 				{
-					if (prctl(PR_SET_NAME, TCHAR_TO_UTF8(*FString::Printf(TEXT("DS-%04x-%04x"), Cookie, ChildIdx))) != 0)
+					if (prctl(PR_SET_NAME, TCHAR_TO_UTF8(*FString::Printf(TEXT("DS-%04hx-%04hx"), Cookie, ChildIdx))) != 0)
 					{
 						int ErrNo = errno;
 						UE_LOG(LogHAL, Fatal, TEXT("WaitAndFork failed to set process name with prctl! error:%d"), ErrNo);
@@ -1524,7 +1526,7 @@ FGenericPlatformProcess::EWaitAndForkResult FUnixPlatformProcess::WaitAndFork()
 					};
 					sigaction(WAIT_AND_FORK_RESPONSE_SIGNAL, &Action, nullptr);
 
-					UE_LOG(LogHAL, Log, TEXT("[Child] WaitAndFork child waiting for signal %d to proceed."), WAIT_AND_FORK_RESPONSE_SIGNAL);
+					UE_LOG(LogHAL, Log, TEXT("[Child] WaitAndFork child %04hx-%04hx waiting for signal %d to proceed."), Cookie, ChildIdx, WAIT_AND_FORK_RESPONSE_SIGNAL);
 					while (!IsEngineExitRequested() && !bResponseReceived)
 					{
 						FPlatformProcess::Sleep(1);
@@ -1534,7 +1536,7 @@ FGenericPlatformProcess::EWaitAndForkResult FUnixPlatformProcess::WaitAndFork()
 					sigaction(WAIT_AND_FORK_RESPONSE_SIGNAL, &Action, nullptr);
 				}
 
-				UE_LOG(LogHAL, Log, TEXT("[Child] WaitAndFork child process has started with pid %d."), GetCurrentProcessId());
+				UE_LOG(LogHAL, Log, TEXT("[Child] WaitAndFork child process %04hx-%04hx has started with pid %d."), Cookie, ChildIdx, GetCurrentProcessId());
 				FApp::PrintStartupLogMessages();
 
 				// Children break out of the loop and return
@@ -1546,7 +1548,7 @@ FGenericPlatformProcess::EWaitAndForkResult FUnixPlatformProcess::WaitAndFork()
 				// Parent
 				AllChildren.Emplace(ChildPID, SignalValue);
 
-				UE_LOG(LogHAL, Log, TEXT("[Parent] WaitAndFork Successfully made a child with pid %d!"), ChildPID);
+				UE_LOG(LogHAL, Log, TEXT("[Parent] WaitAndFork Successfully processed request %04hx-%04hx, made a child with pid %d! Total number of children: %d."), Cookie, ChildIdx, ChildPID, AllChildren.Num());
 			}
 		}
 		else
