TSizedInlineAllocator added as the inlined equivalent of TSizedHeapAllocator.
TInlineAlloctor changed to be an alias of TSizedInlineAllocator, and TInlineAlloctor64 added.
Fix for Pakfile creation when the uncompressed buffer size exceeds int32 capacity.

diff --git a/Engine/Source/Developer/PakFileUtilities/Private/PakFileUtilities.cpp b/Engine/Source/Developer/PakFileUtilities/Private/PakFileUtilities.cpp
--- a/Engine/Source/Developer/PakFileUtilities/Private/PakFileUtilities.cpp
+++ b/Engine/Source/Developer/PakFileUtilities/Private/PakFileUtilities.cpp
@@ -711,7 +711,7 @@ bool FCompressedFileBuffer::CompressFileToWorkingBuffer(const FPakInputPair& InF
 	const int64 PaddedEncryptedFileSize = Align(FileSize,FAES::AESBlockSize);
 	check(PaddedEncryptedFileSize >= FileSize);
 
-	TArray<uint8, TInlineAllocator<4096>> UncompressedBuffer;
+	TArray<uint8, TInlineAllocator64<4096>> UncompressedBuffer;
 	UncompressedBuffer.SetNumUninitialized(PaddedEncryptedFileSize);
 
 	// Load to buffer
diff --git a/Engine/Source/Runtime/Core/Public/Containers/ContainerAllocationPolicies.h b/Engine/Source/Runtime/Core/Public/Containers/ContainerAllocationPolicies.h
--- a/Engine/Source/Runtime/Core/Public/Containers/ContainerAllocationPolicies.h
+++ b/Engine/Source/Runtime/Core/Public/Containers/ContainerAllocationPolicies.h
@@ -639,11 +639,13 @@ struct TCanMoveBetweenAllocators<TSizedHeapAllocator<FromIndexSize>, TSizedHeapA
  * Any allocation needed beyond that causes all data to be moved into an indirect allocation.
  * It always uses DEFAULT_ALIGNMENT.
  */
-template <uint32 NumInlineElements, typename SecondaryAllocator = FDefaultAllocator>
-class TInlineAllocator
+template <uint32 NumInlineElements, int IndexSize, typename SecondaryAllocator = FDefaultAllocator>
+class TSizedInlineAllocator
 {
 public:
-	using SizeType = int32;
+	using SizeType = typename TBitsToSizeType<IndexSize>::Type;
+
+	static_assert(std::is_same_v<SizeType, typename SecondaryAllocator::SizeType>, "Secondary allocator SizeType mismatch");
 
 	enum { NeedsElementType = true };
 	enum { RequireRangeCheck = true };
@@ -780,12 +782,18 @@ public:
 	typedef void ForAnyElementType;
 };
 
-template <uint32 NumInlineElements, typename SecondaryAllocator>
-struct TAllocatorTraits<TInlineAllocator<NumInlineElements, SecondaryAllocator>> : TAllocatorTraitsBase<TInlineAllocator<NumInlineElements, SecondaryAllocator>>
+template <uint32 NumInlineElements, int IndexSize, typename SecondaryAllocator>
+struct TAllocatorTraits<TSizedInlineAllocator<NumInlineElements, IndexSize, SecondaryAllocator>> : TAllocatorTraitsBase<TSizedInlineAllocator<NumInlineElements, IndexSize, SecondaryAllocator>>
 {
 	enum { SupportsMove = TAllocatorTraits<SecondaryAllocator>::SupportsMove };
 };
 
+template <uint32 NumInlineElements, typename SecondaryAllocator = FDefaultAllocator>
+using TInlineAllocator = TSizedInlineAllocator<NumInlineElements, 32, SecondaryAllocator>;
+
+template <uint32 NumInlineElements, typename SecondaryAllocator = FDefaultAllocator64>
+using TInlineAllocator64 = TSizedInlineAllocator<NumInlineElements, 64, SecondaryAllocator>;
+
 /**
  * Implements a variant of TInlineAllocator with a secondary heap allocator that is allowed to store a pointer to its inline elements.
  * This allows caching a pointer to the elements which avoids any conditional logic in GetAllocation(), but prevents the allocator being trivially relocatable.
