Change pattern for CustomConfig to support proper platform requirements.

diff --git a/Engine/Source/Programs/UnrealBuildTool/System/ConfigHierarchy.cs b/Engine/Source/Programs/UnrealBuildTool/System/ConfigHierarchy.cs
--- a/Engine/Source/Programs/UnrealBuildTool/System/ConfigHierarchy.cs
+++ b/Engine/Source/Programs/UnrealBuildTool/System/ConfigHierarchy.cs
@@ -916,8 +916,8 @@ namespace UnrealBuildTool
 			"{PROJECT}/Config/{PLATFORM}/{PLATFORM}{TYPE}.ini",
 			// Project/Platform/GeneratedPlatform*.ini  this is reserved for files which are generated by buildmachine processes (i.e. should never be checked in)
 			"{PROJECT}/Config/{PLATFORM}/Generated{PLATFORM}{TYPE}.ini",
-			// Project/Config/Custom/CustomConfig/Platform/Platform*.ini only if CustomConfig is defined
-			"{PROJECT}/Config/Custom/{CUSTOMCONFIG}/{PLATFORM}/{PLATFORM}{TYPE}.ini",
+			// Project/Platform/Custom/CustomConfig/Platform*.ini only if CustomConfig is defined
+			"{PROJECT}/Config/{PLATFORM}/Custom/{CUSTOMCONFIG}/{PLATFORM}{TYPE}.ini",
 			// UserSettings/.../User*.ini
 			"{USERSETTINGS}/Unreal Engine/Engine/Config/User{TYPE}.ini",
 			// UserDir/.../User*.ini
diff --git a/Engine/Source/Runtime/Core/Private/Misc/ConfigCacheIni.cpp b/Engine/Source/Runtime/Core/Private/Misc/ConfigCacheIni.cpp
--- a/Engine/Source/Runtime/Core/Private/Misc/ConfigCacheIni.cpp
+++ b/Engine/Source/Runtime/Core/Private/Misc/ConfigCacheIni.cpp
@@ -629,8 +629,8 @@ struct FConfigLayer
 	{ TEXT("ProjectPlatform"),			TEXT("{PROJECT}/Config/{PLATFORM}/{PLATFORM}{TYPE}.ini") },
 	// Project/Platform/GeneratedPlatform*.ini Reserved for files generated by build process and should never be checked in 
 	{ TEXT("ProjectPlatformGenerated"),	TEXT("{PROJECT}/Config/{PLATFORM}/Generated{PLATFORM}{TYPE}.ini") },
-	// Project/Custom/CustomConfig/Platform/Platform*.ini only if CustomConfig is defined
-	{ TEXT("CustomConfigPlatform"),		TEXT("{PROJECT}/Config/Custom/{CUSTOMCONFIG}/{PLATFORM}/{PLATFORM}{TYPE}.ini"), EConfigLayerFlags::RequiresCustomConfig },
+	// Project/Platform/Custom/CustomConfig/Platform*.ini only if CustomConfig is defined
+	{ TEXT("CustomConfigPlatform"),		TEXT("{PROJECT}/Config/{PLATFORM}/Custom/{CUSTOMCONFIG}/{PLATFORM}{TYPE}.ini"), EConfigLayerFlags::RequiresCustomConfig },
 	// UserSettings/.../User*.ini
 	{ TEXT("UserSettingsDir"),			TEXT("{USERSETTINGS}Unreal Engine/Engine/Config/User{TYPE}.ini"), EConfigLayerFlags::NoExpand },
 	// UserDir/.../User*.ini
