Fix some text being untranslated in -game due to it being newly generated with a package ID

#rnx

diff --git a/Engine/Source/Runtime/Core/Private/Internationalization/TextLocalizationManager.cpp b/Engine/Source/Runtime/Core/Private/Internationalization/TextLocalizationManager.cpp
--- a/Engine/Source/Runtime/Core/Private/Internationalization/TextLocalizationManager.cpp
+++ b/Engine/Source/Runtime/Core/Private/Internationalization/TextLocalizationManager.cpp
@@ -702,7 +702,6 @@ FTextConstDisplayStringRef FTextLocalizationManager::GetDisplayString(const FTex
 	FTextConstDisplayStringPtr SourceDisplayString;
 	FDisplayStringEntry* SourceLiveEntry = nullptr;
 #if USE_STABLE_LOCALIZATION_KEYS
-	if (GIsEditor)
 	{
 		const FTextKey DisplayNamespace = TextNamespaceUtil::StripPackageNamespace(TextId.GetNamespace().GetChars());
 		if (DisplayNamespace != TextId.GetNamespace())
@@ -1200,7 +1199,6 @@ void FTextLocalizationManager::UpdateFromNative(FTextLocalizationResource&& Text
 
 #if USE_STABLE_LOCALIZATION_KEYS
 			// In builds with stable keys enabled, we have to update the display strings from the "clean" version of the text (if the sources match) as this is the only version that is translated
-			if (GIsEditor)
 			{
 				const FTextKey LiveNamespace = DisplayStringPair.Key.GetNamespace();
 				const FTextKey DisplayNamespace = TextNamespaceUtil::StripPackageNamespace(LiveNamespace.GetChars());
@@ -1313,7 +1311,6 @@ void FTextLocalizationManager::UpdateFromLocalizations(FTextLocalizationResource
 
 		// Perform any additional processing over existing entries
 #if USE_STABLE_LOCALIZATION_KEYS
-		if (GIsEditor)
 		{
 			for (auto& DisplayStringPair : DisplayStringLookupTable)
 			{
