Add a flush of threaded logs on assertion faliures

This is done by adding `GLog->PanicFlushThreadedLogs()` to `FDebug::CheckVerifyFailedImpl()`.
Without this, some of the recent messages may not be reported when the debugger stops due to an assertion failure.

diff --git a/Engine/Source/Runtime/Core/Private/Misc/AssertionMacros.cpp b/Engine/Source/Runtime/Core/Private/Misc/AssertionMacros.cpp
--- a/Engine/Source/Runtime/Core/Private/Misc/AssertionMacros.cpp
+++ b/Engine/Source/Runtime/Core/Private/Misc/AssertionMacros.cpp
@@ -513,6 +513,13 @@ void FORCENOINLINE FDebug::CheckVerifyFailedImpl(
 	FDebug::LogAssertFailedMessageImplV(Expr, File, Line, ProgramCounter, Format, Args);
 	va_end(Args);
 
+	if (GLog)
+	{
+		// Flushing the threaded logs here increases the likelihood that recent messages will be written to the log file, stdout and the debugger console.
+		// Without this, some of the recent messages may not be reported when debugger stops due to an assertion failure.
+		GLog->PanicFlushThreadedLogs();
+	}
+
 	if (!FPlatformMisc::IsDebuggerPresent())
 	{
 		FPlatformMisc::PromptForRemoteDebugging(false);
