`stat DumpFrame` incorrectly reports as cycles instead of ms.

- This was a bug introduced in 16927028.
- Add missing `else if` statement.

diff --git a/Engine/Source/Runtime/Core/Private/Stats/StatsData.cpp b/Engine/Source/Runtime/Core/Private/Stats/StatsData.cpp
--- a/Engine/Source/Runtime/Core/Private/Stats/StatsData.cpp
+++ b/Engine/Source/Runtime/Core/Private/Stats/StatsData.cpp
@@ -1844,7 +1844,7 @@ FString FStatsUtils::DebugPrint(FStatMessage const& Item)
 		{
 			Result = FString::Printf(TEXT("%.3fms"), FPlatformTime::ToMilliseconds64(Item.GetValue_int64()));
 		}
-		if (Item.NameAndInfo.GetFlag(EStatMetaFlags::IsMemory))
+		else if (Item.NameAndInfo.GetFlag(EStatMetaFlags::IsMemory))
 		{
 			Result = FString::Printf(TEXT("%.3fMB"), (double)Item.GetValue_int64() / 1024.0 / 1024.0);
 		}
