Fix deadlock in iostore describe command by implementing FIoStoreReader::Read() without using tasks

#rnx

diff --git a/Engine/Source/Runtime/Core/Private/IO/IoStore.cpp b/Engine/Source/Runtime/Core/Private/IO/IoStore.cpp
--- a/Engine/Source/Runtime/Core/Private/IO/IoStore.cpp
+++ b/Engine/Source/Runtime/Core/Private/IO/IoStore.cpp
@@ -1948,7 +1948,7 @@ public:
 
 	TFuture<TIoStatusOr<FIoBuffer>> ReadAsync(const FIoChunkId& ChunkId, const FIoReadOptions& Options) const
 	{
-		TRACE_CPUPROFILER_EVENT_SCOPE(ReadChunk);
+		TRACE_CPUPROFILER_EVENT_SCOPE(ReadChunkAsync);
 
 		struct FState
 		{
@@ -2109,7 +2109,90 @@ public:
 
 	TIoStatusOr<FIoBuffer> Read(const FIoChunkId& ChunkId, const FIoReadOptions& Options) const
 	{
-		return ReadAsync(ChunkId, Options).Get();
+		TRACE_CPUPROFILER_EVENT_SCOPE(ReadChunk);
+
+		const FIoOffsetAndLength* OffsetAndLength = Toc.GetOffsetAndLength(ChunkId);
+		if (!OffsetAndLength)
+		{
+			return FIoStatus(EIoErrorCode::NotFound, TEXT("Unknown chunk ID"));
+		}
+
+		uint64 RequestedOffset = Options.GetOffset();
+		uint64 ResolvedOffset = OffsetAndLength->GetOffset() + RequestedOffset;
+		uint64 ResolvedSize = 0;
+		if (RequestedOffset <= OffsetAndLength->GetLength())
+		{
+			ResolvedSize = FMath::Min(Options.GetSize(), OffsetAndLength->GetLength() - RequestedOffset);
+		}
+
+		const FIoStoreTocResource& TocResource = Toc.GetTocResource();
+		const uint64 CompressionBlockSize = TocResource.Header.CompressionBlockSize;
+		FIoBuffer UncompressedBuffer(ResolvedSize);
+		TArray<uint8> CompressedBuffer;
+		CompressedBuffer.SetNumUninitialized(static_cast<int32>(CompressionBlockSize));
+		int32 FirstBlockIndex = int32(ResolvedOffset / CompressionBlockSize);
+		int32 LastBlockIndex = int32((Align(ResolvedOffset + ResolvedSize, CompressionBlockSize) - 1) / CompressionBlockSize);
+		uint64 UncompressedDestinationOffset = 0;
+		uint64 OffsetInBlock = ResolvedOffset % CompressionBlockSize;
+		uint64 RemainingSize = ResolvedSize;
+		for (int32 BlockIndex = FirstBlockIndex; BlockIndex <= LastBlockIndex; ++BlockIndex)
+		{
+			const FIoStoreTocCompressedBlockEntry& CompressionBlock = TocResource.CompressionBlocks[BlockIndex];
+			uint32 RawSize = Align(CompressionBlock.GetCompressedSize(), FAES::AESBlockSize);
+			if (uint32(CompressedBuffer.Num()) < RawSize)
+			{
+				CompressedBuffer.SetNumUninitialized(RawSize);
+			}
+			
+			int32 PartitionIndex = int32(CompressionBlock.GetOffset() / TocResource.Header.PartitionSize);
+			int64 PartitionOffset = int64(CompressionBlock.GetOffset() % TocResource.Header.PartitionSize);
+			TUniquePtr<IAsyncReadRequest> ReadRequest(ContainerFileHandles[PartitionIndex]->ReadRequest(PartitionOffset, RawSize, AIOP_Normal, nullptr, CompressedBuffer.GetData()));
+			{
+				TRACE_CPUPROFILER_EVENT_SCOPE(WaitForIo);
+				ReadRequest->WaitCompletion();
+			}
+
+			if (EnumHasAnyFlags(TocResource.Header.ContainerFlags, EIoContainerFlags::Encrypted))
+			{
+				TRACE_CPUPROFILER_EVENT_SCOPE(Decrypt);
+				FAES::DecryptData(CompressedBuffer.GetData(), RawSize, DecryptionKey);
+			}
+
+			FName CompressionMethod = TocResource.CompressionMethods[CompressionBlock.GetCompressionMethodIndex()];
+			uint8* UncompressedDestination = UncompressedBuffer.Data() + UncompressedDestinationOffset;
+			const uint32 UncompressedSize = CompressionBlock.GetUncompressedSize();
+			if (CompressionMethod.IsNone())
+			{
+				check(UncompressedDestination + UncompressedSize - OffsetInBlock <= UncompressedBuffer.Data() + UncompressedBuffer.DataSize());
+				FMemory::Memcpy(UncompressedDestination, CompressedBuffer.GetData() + OffsetInBlock, UncompressedSize - OffsetInBlock);
+			}
+			else
+			{
+				bool bUncompressed;
+				if (OffsetInBlock || RemainingSize < UncompressedSize)
+				{
+					TArray<uint8> TempBuffer;
+					TempBuffer.SetNumUninitialized(UncompressedSize);
+					bUncompressed = FCompression::UncompressMemory(CompressionMethod, TempBuffer.GetData(), UncompressedSize, CompressedBuffer.GetData(), CompressionBlock.GetCompressedSize());
+					uint64 CopySize = FMath::Min<uint64>(UncompressedSize - OffsetInBlock, RemainingSize);
+					FMemory::Memcpy(UncompressedDestination, TempBuffer.GetData() + OffsetInBlock, CopySize);
+				}
+				else
+				{
+					check(UncompressedDestination + UncompressedSize <= UncompressedBuffer.Data() + UncompressedBuffer.DataSize());
+					bUncompressed = FCompression::UncompressMemory(CompressionMethod, UncompressedDestination, UncompressedSize, CompressedBuffer.GetData(), CompressionBlock.GetCompressedSize());
+				}
+				if (!bUncompressed)
+				{
+					return FIoStatus(EIoErrorCode::ReadError, TEXT("Failed uncompressing chunk"));
+				}
+			}
+
+			UncompressedDestinationOffset += UncompressedSize;
+			RemainingSize -= UncompressedSize;
+			OffsetInBlock = 0;
+		}
+		return UncompressedBuffer;
 	}
 
 	TIoStatusOr<FIoStoreCompressedReadResult> ReadCompressed(const FIoChunkId& ChunkId, const FIoReadOptions& Options) const
@@ -2167,6 +2250,7 @@ public:
 
 			if (EnumHasAnyFlags(TocResource.Header.ContainerFlags, EIoContainerFlags::Encrypted))
 			{
+				TRACE_CPUPROFILER_EVENT_SCOPE(Decrypt);
 				FAES::DecryptData(CompressedBuffer.GetData(), RawSize, DecryptionKey);
 			}
 			FIoStoreCompressedBlockInfo& BlockInfo = Result.Blocks.AddDefaulted_GetRef();
