Modify FUnixPlatformMemory::GetStats and GetExtendedStats to use open/read/close

fopen/fgets both allocate memory and we're potentially hitting contention when multithreaded

Also remove the freebsd code and clean up everything to use common ReadProcFields function

diff --git a/Engine/Source/Runtime/Core/Private/Unix/UnixPlatformMemory.cpp b/Engine/Source/Runtime/Core/Private/Unix/UnixPlatformMemory.cpp
--- a/Engine/Source/Runtime/Core/Private/Unix/UnixPlatformMemory.cpp
+++ b/Engine/Source/Runtime/Core/Private/Unix/UnixPlatformMemory.cpp
@@ -23,11 +23,7 @@
 #include "HAL/MallocStomp.h"
 #include "HAL/PlatformMallocCrash.h"
 
-#if PLATFORM_FREEBSD
-	#include <kvm.h>
-#else
-	#include <sys/sysinfo.h>
-#endif
+#include <sys/sysinfo.h>
 #include <sys/file.h>
 #include <sys/mman.h>
 
@@ -653,261 +649,188 @@ void FUnixPlatformMemory::FPlatformVirtualMemoryBlock::Decommit(size_t InOffset,
 	}
 }
 
+struct FProcField
+{
+	const ANSICHAR* Name = nullptr;
+	uint64* Addr = nullptr;
+};
 
-namespace UnixPlatformMemory
+static uint32 ReadProcFields(const char *FileName, FProcField *ProcFields, uint32 NumFields)
 {
-	/**
-	 * @brief Returns value in bytes from a status line
-	 * @param Line in format "Blah:  10000 kB" - needs to be writable as it will modify it
-	 * @return value in bytes (10240000, i.e. 10000 * 1024 for the above example)
-	 */
-	uint64 GetBytesFromStatusLine(char * Line)
+	uint32 NumFieldsFound = 0;
+	// 12k should be enough for the proc files we're reading for stats, with room for growth.
+	// smaps_rollup is 642 bytes, meminfo is 1428 bytes, and proc/self/status is 1503 bytes.
+	// If we run out of space and fail to find a field we should also see the warning below.
+	ANSICHAR Buffer[12 * 1024];
+
+	int Fd = open(FileName, O_RDONLY);
+	if (Fd >= 0)
 	{
-		check(Line);
-		int Len = strlen(Line);
+		ssize_t Bytes;
 
-		// Len should be long enough to hold at least " kB\n"
-		const int kSuffixLength = 4;	// " kB\n"
-		if (Len <= kSuffixLength)
+		do
 		{
-			return 0;
-		}
+			Bytes = read(Fd, Buffer, sizeof(Buffer) - 1);
+		} while (Bytes < 0 && errno == EINTR);
 
-		const int NewLen = Len - kSuffixLength;
+		close(Fd);
+		Fd = -1;
 
-		// let's check that this is indeed "kB"
-		char * Suffix = &Line[NewLen];
-		if (strcmp(Suffix, " kB\n") != 0)
+		if ((Bytes > 0) && (Bytes < UE_ARRAY_COUNT(Buffer)))
 		{
-			// Unix the kernel changed the format, huh?
-			return 0;
-		}
+			Buffer[Bytes] = 0;
 
-		// kill the kB
-		*Suffix = 0;
+			for (uint32 Idx = 0; Idx < NumFields; Idx++)
+			{
+				const ANSICHAR* Field = FCStringAnsi::Strstr(Buffer, ProcFields[Idx].Name);
+
+				if (Field)
+				{
+					uint32 NameLen = FCStringAnsi::Strlen(ProcFields[Idx].Name);
+
+					*ProcFields[Idx].Addr = atoll(Field + NameLen) * 1024ULL;
+					NumFieldsFound++;
+				}
+			}
+		}
 
-        // find the beginning of the number
-		for (const char* NumberBegin = Line; NumberBegin < Suffix; ++NumberBegin)
+		if (NumFieldsFound != NumFields)
 		{
-			if (isdigit(*NumberBegin))
+			static bool bLogOnceMissing = false;
+			if (!bLogOnceMissing)
 			{
-				return atoll(NumberBegin) * 1024ULL;
+				// Note: We can't use UE_LOG or TCHAR_TO_UTF8 since these routines may not be initialized yet and
+				// allocate memory. This function could get called via FMemory::GCreateMalloc or signal handlers
+				// and we will potentiall crash in these cases calling UE_LOG or TCHAR_TO_UTF8.
+				fprintf(stderr, "Warning: ReadProcFields: %u of %u fields found in %s.\n",
+					NumFieldsFound, NumFields, FileName);
+				fflush(stderr);
+				bLogOnceMissing = true;
 			}
 		}
-
-		// we were unable to find whitespace in front of the number
-		return 0;
 	}
+	else
+	{
+		static bool bLogOnceFileNotFound = false;
+		if (!bLogOnceFileNotFound)
+		{
+			fprintf(stderr, "Warning: ReadProcFields failed opening %s (Err %d).\n", FileName, errno);
+			fflush(stderr);
+			bLogOnceFileNotFound = true;
+		}
+	}
+
+	return NumFieldsFound;
 }
 
+// struct CORE_API FGenericPlatformMemoryStats : public FPlatformMemoryConstants
+//   uint64 AvailablePhysical; /** The amount of physical memory currently available, in bytes. */			MemAvailable (or MemFree + Cached)
+//   uint64 AvailableVirtual;  /** The amount of virtual memory currently available, in bytes. */			SwapFree
+//   uint64 UsedPhysical;      /** The amount of physical memory used by the process, in bytes. */			VmRSS
+//   uint64 PeakUsedPhysical;  /** The peak amount of physical memory used by the process, in bytes. */	VmHWM
+//   uint64 UsedVirtual;       /** Total amount of virtual memory used by the process. */					VmSize
+//   uint64 PeakUsedVirtual;   /** The peak amount of virtual memory used by the process. */				VmPeak
 FPlatformMemoryStats FUnixPlatformMemory::GetStats()
 {
-	FPlatformMemoryStats MemoryStats;	// will init from constants
-
-#if PLATFORM_FREEBSD
-
-	const FPlatformMemoryConstants& MemoryConstants = FPlatformMemory::GetConstants();
-
-	size_t size = sizeof(SIZE_T);
-
-	SIZE_T SysFreeCount = 0;
-	sysctlbyname("vm.stats.vm.v_free_count", &SysFreeCount, &size, NULL, 0);
-
-	SIZE_T SysActiveCount = 0;
-	sysctlbyname("vm.stats.vm.v_active_count", &SysActiveCount, &size, NULL, 0);
-
-	// Get swap info from kvm api
-	kvm_t* Kvm = kvm_open(NULL, "/dev/null", NULL, O_RDONLY, NULL);
-	struct kvm_swap KvmSwap;
-	kvm_getswapinfo(Kvm, &KvmSwap, 1, 0);
-	kvm_close(Kvm);
+	uint64 MemFree = 0;
+	uint64 Cached = 0;
+	FPlatformMemoryStats MemoryStats;
 
-	MemoryStats.AvailablePhysical = SysFreeCount * MemoryConstants.PageSize;
-	MemoryStats.AvailableVirtual = (KvmSwap.ksw_total - KvmSwap.ksw_used) * MemoryConstants.PageSize;
-	MemoryStats.UsedPhysical = SysActiveCount * MemoryConstants.PageSize;
-	MemoryStats.UsedVirtual = KvmSwap.ksw_used * MemoryConstants.PageSize;
-
-#else
-
-	// open to all kind of overflows, thanks to Unix approach of exposing system stats via /proc and lack of proper C API
-	// And no, sysinfo() isn't useful for this (cannot get the same value for MemAvailable through it for example).
-
-	if (FILE* FileGlobalMemStats = fopen("/proc/meminfo", "r"))
+	FProcField SMapsFields[] =
+	{
+		// An estimate of how much memory is available for starting new applications, without swapping.
+		{ "MemAvailable:", &MemoryStats.AvailablePhysical },
+		// Amount of swap space that is currently unused.
+		{ "SwapFree:",     &MemoryStats.AvailableVirtual },
+		{ "MemFree:",      &MemFree },
+		{ "Cached:",       &Cached },
+	};
+	if (ReadProcFields("/proc/meminfo", SMapsFields, UE_ARRAY_COUNT(SMapsFields)))
 	{
-		int FieldsSetSuccessfully = 0;
-		uint64 MemFree = 0, Cached = 0;
-		do
-		{
-			char LineBuffer[256] = {0};
-			char *Line = fgets(LineBuffer, UE_ARRAY_COUNT(LineBuffer), FileGlobalMemStats);
-			if (Line == nullptr)
-			{
-				break;	// eof or an error
-			}
-
-			// if we have MemAvailable, favor that (see http://git.kernel.org/cgit/linux/kernel/git/torvalds/linux.git/commit/?id=34e431b0ae398fc54ea69ff85ec700722c9da773)
-			if (strstr(Line, "MemAvailable:") == Line)
-			{
-				MemoryStats.AvailablePhysical = UnixPlatformMemory::GetBytesFromStatusLine(Line);
-				++FieldsSetSuccessfully;
-			}
-			else if (strstr(Line, "SwapFree:") == Line)
-			{
-				MemoryStats.AvailableVirtual = UnixPlatformMemory::GetBytesFromStatusLine(Line);
-				++FieldsSetSuccessfully;
-			}
-			else if (strstr(Line, "MemFree:") == Line)
-			{
-				MemFree = UnixPlatformMemory::GetBytesFromStatusLine(Line);
-				++FieldsSetSuccessfully;
-			}
-			else if (strstr(Line, "Cached:") == Line)
-			{
-				Cached = UnixPlatformMemory::GetBytesFromStatusLine(Line);
-				++FieldsSetSuccessfully;
-			}
-		}
-		while(FieldsSetSuccessfully < 4);
-
 		// if we didn't have MemAvailable (kernels < 3.14 or CentOS 6.x), use free + cached as a (bad) approximation
 		if (MemoryStats.AvailablePhysical == 0)
 		{
 			MemoryStats.AvailablePhysical = FMath::Min(MemFree + Cached, MemoryStats.TotalPhysical);
 		}
-
-		fclose(FileGlobalMemStats);
 	}
 
-	// again /proc "API" :/
-	if (FILE* ProcMemStats = fopen("/proc/self/status", "r"))
+	FProcField SMapsFields2[] =
 	{
-		int FieldsSetSuccessfully = 0;
-		do
-		{
-			char LineBuffer[256] = {0};
-			char *Line = fgets(LineBuffer, UE_ARRAY_COUNT(LineBuffer), ProcMemStats);
-			if (Line == nullptr)
-			{
-				break;	// eof or an error
-			}
-
-			if (strstr(Line, "VmPeak:") == Line)
-			{
-				MemoryStats.PeakUsedVirtual = UnixPlatformMemory::GetBytesFromStatusLine(Line);
-				++FieldsSetSuccessfully;
-			}
-			else if (strstr(Line, "VmSize:") == Line)
-			{
-				MemoryStats.UsedVirtual = UnixPlatformMemory::GetBytesFromStatusLine(Line);
-				++FieldsSetSuccessfully;
-			}
-			else if (strstr(Line, "VmHWM:") == Line)
-			{
-				MemoryStats.PeakUsedPhysical = UnixPlatformMemory::GetBytesFromStatusLine(Line);
-				++FieldsSetSuccessfully;
-			}
-			else if (strstr(Line, "VmRSS:") == Line)
-			{
-				MemoryStats.UsedPhysical = UnixPlatformMemory::GetBytesFromStatusLine(Line);
-				++FieldsSetSuccessfully;
-			}
-		}
-		while(FieldsSetSuccessfully < 4);
-
-		fclose(ProcMemStats);
+		{ "VmPeak:", &MemoryStats.PeakUsedVirtual },
+		{ "VmSize:", &MemoryStats.UsedVirtual },      // In /proc/self/statm (Field 1)
+		{ "VmHWM:",  &MemoryStats.PeakUsedPhysical },
+		{ "VmRSS:",  &MemoryStats.UsedPhysical },     // In /proc/self/statm (Field 2)
+	};
+	if (ReadProcFields("/proc/self/status", SMapsFields2, UE_ARRAY_COUNT(SMapsFields2)))
+	{
+		// sanitize stats as sometimes peak < used for some reason
+		MemoryStats.PeakUsedVirtual = FMath::Max(MemoryStats.PeakUsedVirtual, MemoryStats.UsedVirtual);
+		MemoryStats.PeakUsedPhysical = FMath::Max(MemoryStats.PeakUsedPhysical, MemoryStats.UsedPhysical);
 	}
 
-#endif // PLATFORM_FREEBSD
-
-	// sanitize stats as sometimes peak < used for some reason
-	MemoryStats.PeakUsedVirtual = FMath::Max(MemoryStats.PeakUsedVirtual, MemoryStats.UsedVirtual);
-	MemoryStats.PeakUsedPhysical = FMath::Max(MemoryStats.PeakUsedPhysical, MemoryStats.UsedPhysical);
-
 	return MemoryStats;
 }
 
 FExtendedPlatformMemoryStats FUnixPlatformMemory::GetExtendedStats()
 {
-	const ANSICHAR Shared_CleanStr[] = "Shared_Clean:";
-	const ANSICHAR Shared_DirtyStr[] = "Shared_Dirty:";
+	const ANSICHAR Shared_CleanStr[]  = "Shared_Clean:";
+	const ANSICHAR Shared_DirtyStr[]  = "Shared_Dirty:";
 	const ANSICHAR Private_CleanStr[] = "Private_Clean:";
 	const ANSICHAR Private_DirtyStr[] = "Private_Dirty:";
+
 	FExtendedPlatformMemoryStats MemoryStats = { 0 };
 
+	FProcField SMapsFields[] =
+	{
+		{ Shared_CleanStr,  (uint64 *)&MemoryStats.Shared_Clean },
+		{ Shared_DirtyStr,  (uint64 *)&MemoryStats.Shared_Dirty },
+		{ Private_CleanStr, (uint64 *)&MemoryStats.Private_Clean },
+		{ Private_DirtyStr, (uint64 *)&MemoryStats.Private_Dirty },
+	};
+
 	// ~ 1.06ms per call on my Threadripper 3990X w/ Debian Testing 5.15.0-2-amd64
-	if (FILE* ProcSMapsRollup = fopen("/proc/self/smaps_rollup", "r"))
+	if (!ReadProcFields("/proc/self/smaps_rollup", SMapsFields, UE_ARRAY_COUNT(SMapsFields)))
 	{
-		struct
+		// ~ 6.8ms per call on my Threadripper 3990X w/ Debian Testing 5.15.0-2-amd64 in TestPAL.
+		// Potentially far higher though, like 100s of ms.
+		if (FILE* ProcSMaps = fopen("/proc/self/smaps", "r"))
 		{
-			const ANSICHAR* Name;
-			uint32 NameLen;
-			SIZE_T* Addr;
-		} SMapsFields[] =
-		{
-			{ Shared_CleanStr,  sizeof(Shared_CleanStr) - 1,  &MemoryStats.Shared_Clean },
-			{ Shared_DirtyStr,  sizeof(Shared_DirtyStr) - 1,  &MemoryStats.Shared_Dirty },
-			{ Private_CleanStr, sizeof(Private_CleanStr) - 1, &MemoryStats.Private_Clean },
-			{ Private_DirtyStr, sizeof(Private_DirtyStr) - 1, &MemoryStats.Private_Dirty },
-		};
-		const uint32 NumFields = UE_ARRAY_COUNT(SMapsFields);
-		uint32 FieldsFound = 0;
-
-		while (FieldsFound < NumFields)
-		{
-			ANSICHAR LineBuffer[256];
-			ANSICHAR *Line = fgets(LineBuffer, UE_ARRAY_COUNT(LineBuffer), ProcSMapsRollup);
+			const uint32 Shared_CleanStrLen = FCStringAnsi::Strlen(Shared_CleanStr);
+			const uint32 Shared_DirtyStrLen = FCStringAnsi::Strlen(Shared_DirtyStr);
+			const uint32 Private_CleanStrLen = FCStringAnsi::Strlen(Private_CleanStr);
+			const uint32 Private_DirtyStrLen = FCStringAnsi::Strlen(Private_DirtyStr);
 
-			if (Line == nullptr)
+			for (;;)
 			{
-				// eof or an error
-				break;
-			}
+				ANSICHAR LineBuffer[256];
+				ANSICHAR *Line = fgets(LineBuffer, UE_ARRAY_COUNT(LineBuffer), ProcSMaps);
 
-			for (uint32 i = 0; i < NumFields; i++)
-			{
-				if (!FCStringAnsi::Strncmp(SMapsFields[i].Name, Line, SMapsFields[i].NameLen))
+				if (Line == nullptr)
 				{
-					*SMapsFields[i].Addr = atoll(Line + SMapsFields[i].NameLen) * 1024ULL;
-					FieldsFound++;
+					// eof or error
 					break;
 				}
-			}
-		}
 
-		fclose(ProcSMapsRollup);
-	}
-	// ~ 6.8ms per call on my Threadripper 3990X w/ Debian Testing 5.15.0-2-amd64 in TestPAL.
-	// Potentially far higher though.
-	else if (FILE* ProcSMaps = fopen("/proc/self/smaps", "r"))
-	{
-		do
-		{
-			ANSICHAR LineBuffer[256] = { 0 };
-			ANSICHAR *Line = fgets(LineBuffer, UE_ARRAY_COUNT(LineBuffer), ProcSMaps);
-			if (Line == nullptr)
-			{
-				break;	// eof or an error
-			}
-
-			if (strstr(Line, Shared_CleanStr) == Line)
-			{
-				MemoryStats.Shared_Clean += UnixPlatformMemory::GetBytesFromStatusLine(Line);
-			}
-			else if (strstr(Line, Shared_DirtyStr) == Line)
-			{
-				MemoryStats.Shared_Dirty += UnixPlatformMemory::GetBytesFromStatusLine(Line);
-			}
-			if (strstr(Line, Private_CleanStr) == Line)
-			{
-				MemoryStats.Private_Clean += UnixPlatformMemory::GetBytesFromStatusLine(Line);
-			}
-			else if (strstr(Line, Private_DirtyStr) == Line)
-			{
-				MemoryStats.Private_Dirty += UnixPlatformMemory::GetBytesFromStatusLine(Line);
+				if (!FCStringAnsi::Strncmp(Line, Shared_CleanStr, Shared_CleanStrLen))
+				{
+					MemoryStats.Shared_Clean += atoll(Line + Shared_CleanStrLen) * 1024ULL;
+				}
+				else if (!FCStringAnsi::Strncmp(Line, Shared_DirtyStr, Shared_DirtyStrLen))
+				{
+					MemoryStats.Shared_Dirty += atoll(Line + Shared_DirtyStrLen) * 1024ULL;
+				}
+				else if (!FCStringAnsi::Strncmp(Line, Private_CleanStr, Private_CleanStrLen))
+				{
+					MemoryStats.Private_Clean += atoll(Line + Private_CleanStrLen) * 1024ULL;
+				}
+				else if (!FCStringAnsi::Strncmp(Line, Private_DirtyStr, Private_DirtyStrLen))
+				{
+					MemoryStats.Private_Dirty += atoll(Line + Private_DirtyStrLen) * 1024ULL;
+				}
 			}
-		} while (!feof(ProcSMaps));
 
-		fclose(ProcSMaps);
+			fclose(ProcSMaps);
+		}
 	}
 
 	return MemoryStats;
@@ -919,28 +842,6 @@ const FPlatformMemoryConstants& FUnixPlatformMemory::GetConstants()
 
 	if( MemoryConstants.TotalPhysical == 0 )
 	{
-#if PLATFORM_FREEBSD
-
-		size_t Size = sizeof(SIZE_T);
-
-		SIZE_T SysPageCount = 0;
-		sysctlbyname("vm.stats.vm.v_page_count", &SysPageCount, &Size, NULL, 0);
-
-		SIZE_T SysPageSize = 0;
-		sysctlbyname("vm.stats.vm.v_page_size", &SysPageSize, &Size, NULL, 0);
-
-		// Get swap info from kvm api
-		kvm_t* Kvm = kvm_open(NULL, "/dev/null", NULL, O_RDONLY, NULL);
-		struct kvm_swap KvmSwap;
-		kvm_getswapinfo(Kvm, &KvmSwap, 1, 0);
-		kvm_close(Kvm);
-
-		MemoryConstants.TotalPhysical = SysPageCount * SysPageSize;
-		MemoryConstants.TotalVirtual = KvmSwap.ksw_total * SysPageSize;
-		MemoryConstants.PageSize = SysPageSize;
-
-#else
- 
 		// Gather platform memory stats.
 		struct sysinfo SysInfo;
 		unsigned long long MaxPhysicalRAMBytes = 0;
@@ -955,8 +856,6 @@ const FPlatformMemoryConstants& FUnixPlatformMemory::GetConstants()
 		MemoryConstants.TotalPhysical = MaxPhysicalRAMBytes;
 		MemoryConstants.TotalVirtual = MaxVirtualRAMBytes;
 
-#endif // PLATFORM_FREEBSD
-
 		MemoryConstants.PageSize = sysconf(_SC_PAGESIZE);
 		MemoryConstants.BinnedPageSize = FMath::Max((SIZE_T)65536, MemoryConstants.PageSize);
 		MemoryConstants.BinnedAllocationGranularity = MemoryConstants.PageSize;
